UNiagaraDataInterface::GetFunctions() improvements - part 1
-renamed GetFunctions() to be GetFunctionSignatures() which is only defined WITH_EDITORONLY_DATA
-GetFunctions is deprecated but remains calling GetFunctionsInternal if WITH_EDITORONLY_DATA for back compat
-new protected/const/virtual/WITH_EDITORONLY_DATA function GetFunctionsInternal() to implement actual collection of signatures

diff --git a/Engine/Plugins/Experimental/ChaosNiagara/Source/ChaosNiagara/Classes/NiagaraDataInterfaceChaosDestruction.h b/Engine/Plugins/Experimental/ChaosNiagara/Source/ChaosNiagara/Classes/NiagaraDataInterfaceChaosDestruction.h
--- a/Engine/Plugins/Experimental/ChaosNiagara/Source/ChaosNiagara/Classes/NiagaraDataInterfaceChaosDestruction.h
+++ b/Engine/Plugins/Experimental/ChaosNiagara/Source/ChaosNiagara/Classes/NiagaraDataInterfaceChaosDestruction.h
@@ -552,7 +552,6 @@ public:
 	virtual bool PerInstanceTick(void* PerInstanceData, FNiagaraSystemInstance* SystemInstance, float DeltaSeconds) override;
 	virtual int32 PerInstanceDataSize() const override;
 
-	virtual void GetFunctions(TArray<FNiagaraFunctionSignature>& OutFunctions)override;
 	virtual void GetVMExternalFunction(const FVMExternalFunctionBindingInfo& BindingInfo, void* InstanceData, FVMExternalFunction &OutFunc) override;
 	virtual bool Equals(const UNiagaraDataInterface* Other) const override;
 	virtual bool CanExecuteOnTarget(ENiagaraSimTarget Target) const override { return true || (Target == ENiagaraSimTarget::CPUSim); }
@@ -688,6 +687,10 @@ public:
 	virtual void ProvidePerInstanceDataForRenderThread(void* DataForRenderThread, void* PerInstanceData, const FNiagaraSystemInstanceID& SystemInstance) override;
 protected:
 
+#if WITH_EDITORONLY_DATA
+	virtual void GetFunctionsInternal(TArray<FNiagaraFunctionSignature>& OutFunctions) const override;
+#endif
+
 	void RegisterWithSolverEventManager(Chaos::FPhysicsSolver* Solver);
 	virtual bool CopyToInternal(UNiagaraDataInterface* Destination) const override;
 
diff --git a/Engine/Plugins/Experimental/ChaosNiagara/Source/ChaosNiagara/Classes/NiagaraDataInterfacePhysicsField.h b/Engine/Plugins/Experimental/ChaosNiagara/Source/ChaosNiagara/Classes/NiagaraDataInterfacePhysicsField.h
--- a/Engine/Plugins/Experimental/ChaosNiagara/Source/ChaosNiagara/Classes/NiagaraDataInterfacePhysicsField.h
+++ b/Engine/Plugins/Experimental/ChaosNiagara/Source/ChaosNiagara/Classes/NiagaraDataInterfacePhysicsField.h
@@ -56,7 +56,6 @@ public:
 	virtual void PostInitProperties() override;
 
 	/** UNiagaraDataInterface Interface */
-	virtual void GetFunctions(TArray<FNiagaraFunctionSignature>& OutFunctions) override;
 	virtual void GetVMExternalFunction(const FVMExternalFunctionBindingInfo& BindingInfo, void* InstanceData, FVMExternalFunction& OutFunc) override;
 	virtual bool CanExecuteOnTarget(ENiagaraSimTarget Target) const override { return Target == ENiagaraSimTarget::GPUComputeSim; }
 	virtual bool InitPerInstanceData(void* PerInstanceData, FNiagaraSystemInstance* SystemInstance) override;
@@ -94,6 +93,10 @@ public:
 	void GetPhysicsFieldBounds(FVectorVMExternalFunctionContext& Context);
 
 protected:
+#if WITH_EDITORONLY_DATA
+	virtual void GetFunctionsInternal(TArray<FNiagaraFunctionSignature>& OutFunctions) const override;
+#endif
+
 	/** Copy one niagara DI to this */
 	virtual bool CopyToInternal(UNiagaraDataInterface* Destination) const override;
 };
diff --git a/Engine/Plugins/Experimental/ChaosNiagara/Source/ChaosNiagara/Private/NiagaraDataInterfaceChaosDestruction.cpp b/Engine/Plugins/Experimental/ChaosNiagara/Source/ChaosNiagara/Private/NiagaraDataInterfaceChaosDestruction.cpp
--- a/Engine/Plugins/Experimental/ChaosNiagara/Source/ChaosNiagara/Private/NiagaraDataInterfaceChaosDestruction.cpp
+++ b/Engine/Plugins/Experimental/ChaosNiagara/Source/ChaosNiagara/Private/NiagaraDataInterfaceChaosDestruction.cpp
@@ -2540,8 +2540,9 @@ bool UNiagaraDataInterfaceChaosDestruction::PerInstanceTick(void* PerInstanceDat
 	return false;
 }
 
+#if WITH_EDITORONLY_DATA
 // Returns the signature of all the functions available in the data interface
-void UNiagaraDataInterfaceChaosDestruction::GetFunctions(TArray<FNiagaraFunctionSignature>& OutFunctions)
+void UNiagaraDataInterfaceChaosDestruction::GetFunctionsInternal(TArray<FNiagaraFunctionSignature>& OutFunctions) const
 {
 	using namespace NDIChaosDestructionLocal;
 	{
@@ -2864,6 +2865,7 @@ void UNiagaraDataInterfaceChaosDestruction::GetFunctions(TArray<FNiagaraFunction
 		OutFunctions.Add(Sig);
 	}
 }
+#endif
 
 DEFINE_NDI_FUNC_BINDER(UNiagaraDataInterfaceChaosDestruction, GetPosition);
 DEFINE_NDI_FUNC_BINDER(UNiagaraDataInterfaceChaosDestruction, GetNormal);
diff --git a/Engine/Plugins/Experimental/ChaosNiagara/Source/ChaosNiagara/Private/NiagaraDataInterfaceGeometryCollection.cpp b/Engine/Plugins/Experimental/ChaosNiagara/Source/ChaosNiagara/Private/NiagaraDataInterfaceGeometryCollection.cpp
--- a/Engine/Plugins/Experimental/ChaosNiagara/Source/ChaosNiagara/Private/NiagaraDataInterfaceGeometryCollection.cpp
+++ b/Engine/Plugins/Experimental/ChaosNiagara/Source/ChaosNiagara/Private/NiagaraDataInterfaceGeometryCollection.cpp
@@ -404,7 +404,8 @@ void UNiagaraDataInterfaceGeometryCollection::PostInitProperties()
 	}
 }
 
-void UNiagaraDataInterfaceGeometryCollection::GetFunctions(TArray<FNiagaraFunctionSignature>& OutFunctions)
+#if WITH_EDITORONLY_DATA
+void UNiagaraDataInterfaceGeometryCollection::GetFunctionsInternal(TArray<FNiagaraFunctionSignature>& OutFunctions) const
 {
 	{
 		FNiagaraFunctionSignature Sig;
@@ -424,6 +425,7 @@ void UNiagaraDataInterfaceGeometryCollection::GetFunctions(TArray<FNiagaraFuncti
 		OutFunctions.Add(Sig);
 	}
 }
+#endif
 
 void UNiagaraDataInterfaceGeometryCollection::GetVMExternalFunction(const FVMExternalFunctionBindingInfo& BindingInfo, void* InstanceData, FVMExternalFunction& OutFunc)
 {
diff --git a/Engine/Plugins/Experimental/ChaosNiagara/Source/ChaosNiagara/Private/NiagaraDataInterfacePhysicsField.cpp b/Engine/Plugins/Experimental/ChaosNiagara/Source/ChaosNiagara/Private/NiagaraDataInterfacePhysicsField.cpp
--- a/Engine/Plugins/Experimental/ChaosNiagara/Source/ChaosNiagara/Private/NiagaraDataInterfacePhysicsField.cpp
+++ b/Engine/Plugins/Experimental/ChaosNiagara/Source/ChaosNiagara/Private/NiagaraDataInterfacePhysicsField.cpp
@@ -226,7 +226,8 @@ void UNiagaraDataInterfacePhysicsField::PostInitProperties()
 	}
 }
 
-void UNiagaraDataInterfacePhysicsField::GetFunctions(TArray<FNiagaraFunctionSignature>& OutFunctions)
+#if WITH_EDITORONLY_DATA
+void UNiagaraDataInterfacePhysicsField::GetFunctionsInternal(TArray<FNiagaraFunctionSignature>& OutFunctions) const
 {
 	using namespace NDIPhysicsFieldLocal;
 
@@ -324,6 +325,7 @@ void UNiagaraDataInterfacePhysicsField::GetFunctions(TArray<FNiagaraFunctionSign
 		OutFunctions.Add(Sig);
 	}
 }
+#endif
 
 DEFINE_NDI_DIRECT_FUNC_BINDER(UNiagaraDataInterfacePhysicsField, SamplePhysicsVectorField);
 DEFINE_NDI_DIRECT_FUNC_BINDER(UNiagaraDataInterfacePhysicsField, SamplePhysicsScalarField);
diff --git a/Engine/Plugins/Experimental/ChaosNiagara/Source/ChaosNiagara/Public/NiagaraDataInterfaceGeometryCollection.h b/Engine/Plugins/Experimental/ChaosNiagara/Source/ChaosNiagara/Public/NiagaraDataInterfaceGeometryCollection.h
--- a/Engine/Plugins/Experimental/ChaosNiagara/Source/ChaosNiagara/Public/NiagaraDataInterfaceGeometryCollection.h
+++ b/Engine/Plugins/Experimental/ChaosNiagara/Source/ChaosNiagara/Public/NiagaraDataInterfaceGeometryCollection.h
@@ -145,7 +145,6 @@ public:
 	virtual void PostInitProperties() override;
 
 	/** UNiagaraDataInterface Interface */
-	virtual void GetFunctions(TArray<FNiagaraFunctionSignature>& OutFunctions) override;
 	virtual void GetVMExternalFunction(const FVMExternalFunctionBindingInfo& BindingInfo, void* InstanceData, FVMExternalFunction& OutFunc) override;
 	virtual bool CanExecuteOnTarget(ENiagaraSimTarget Target) const override { return Target == ENiagaraSimTarget::GPUComputeSim; }
 	virtual bool InitPerInstanceData(void* PerInstanceData, FNiagaraSystemInstance* SystemInstance) override;
@@ -170,6 +169,10 @@ public:
 	virtual void ProvidePerInstanceDataForRenderThread(void* DataForRenderThread, void* PerInstanceData, const FNiagaraSystemInstanceID& SystemInstance) override;
 
 protected:
+#if WITH_EDITORONLY_DATA
+	virtual void GetFunctionsInternal(TArray<FNiagaraFunctionSignature>& OutFunctions) const override;
+#endif
+
 	/** Copy one niagara DI to this */
 	virtual bool CopyToInternal(UNiagaraDataInterface* Destination) const override;
 };
diff --git a/Engine/Plugins/Experimental/Water/Source/Runtime/Private/NiagaraDataInterfaceWater.cpp b/Engine/Plugins/Experimental/Water/Source/Runtime/Private/NiagaraDataInterfaceWater.cpp
--- a/Engine/Plugins/Experimental/Water/Source/Runtime/Private/NiagaraDataInterfaceWater.cpp
+++ b/Engine/Plugins/Experimental/Water/Source/Runtime/Private/NiagaraDataInterfaceWater.cpp
@@ -41,7 +41,8 @@ void UNiagaraDataInterfaceWater::PostInitProperties()
 	}
 }
 
-void UNiagaraDataInterfaceWater::GetFunctions(TArray<FNiagaraFunctionSignature>& OutFunctions)
+#if WITH_EDITORONLY_DATA
+void UNiagaraDataInterfaceWater::GetFunctionsInternal(TArray<FNiagaraFunctionSignature>& OutFunctions) const
 {
 	{
 		FNiagaraFunctionSignature Sig;
@@ -79,6 +80,7 @@ void UNiagaraDataInterfaceWater::GetFunctions(TArray<FNiagaraFunctionSignature>&
 		OutFunctions.Add(Sig);
 	}
 }
+#endif
 
 DEFINE_NDI_DIRECT_FUNC_BINDER(UNiagaraDataInterfaceWater, GetWaterDataAtPoint);
 DEFINE_NDI_DIRECT_FUNC_BINDER(UNiagaraDataInterfaceWater, GetWaveParamLookupTableOffset);
diff --git a/Engine/Plugins/Experimental/Water/Source/Runtime/Public/NiagaraDataInterfaceWater.h b/Engine/Plugins/Experimental/Water/Source/Runtime/Public/NiagaraDataInterfaceWater.h
--- a/Engine/Plugins/Experimental/Water/Source/Runtime/Public/NiagaraDataInterfaceWater.h
+++ b/Engine/Plugins/Experimental/Water/Source/Runtime/Public/NiagaraDataInterfaceWater.h
@@ -15,7 +15,6 @@ public:
 	virtual void PostInitProperties() override;
 
 	/** UNiagaraDataInterface interface */
-	virtual void GetFunctions(TArray<FNiagaraFunctionSignature>& OutFunctions) override;
 	virtual void GetVMExternalFunction(const FVMExternalFunctionBindingInfo& BindingInfo, void* InstanceData, FVMExternalFunction &OutFunc) override;
 	virtual bool Equals(const UNiagaraDataInterface* Other) const override;
 	virtual bool CopyToInternal(UNiagaraDataInterface* Destination) const override;
@@ -37,6 +36,12 @@ public:
 
 	/** Sets the current water body to be used by this data interface */
 	void SetWaterBodyComponent(UWaterBodyComponent* InWaterBodyComponent) { SourceBodyComponent = InWaterBodyComponent; }
+
+protected:
+#if WITH_EDITORONLY_DATA
+	virtual void GetFunctionsInternal(TArray<FNiagaraFunctionSignature>& OutFunctions) const override;
+#endif
+
 private:
 	UPROPERTY(EditAnywhere, Category = "Water") 
 	TObjectPtr<UWaterBodyComponent> SourceBodyComponent;
diff --git a/Engine/Plugins/FX/ExampleCustomDataInterface/Source/ExampleCustomDataInterface/Private/NiagaraDataInterfaceMousePosition.cpp b/Engine/Plugins/FX/ExampleCustomDataInterface/Source/ExampleCustomDataInterface/Private/NiagaraDataInterfaceMousePosition.cpp
--- a/Engine/Plugins/FX/ExampleCustomDataInterface/Source/ExampleCustomDataInterface/Private/NiagaraDataInterfaceMousePosition.cpp
+++ b/Engine/Plugins/FX/ExampleCustomDataInterface/Source/ExampleCustomDataInterface/Private/NiagaraDataInterfaceMousePosition.cpp
@@ -141,14 +141,13 @@ void UNiagaraDataInterfaceMousePosition::PostInitProperties()
 	}
 }
 
+#if WITH_EDITORONLY_DATA
 // this lists all the functions our DI provides (currently only one)
-void UNiagaraDataInterfaceMousePosition::GetFunctions(TArray<FNiagaraFunctionSignature>& OutFunctions)
+void UNiagaraDataInterfaceMousePosition::GetFunctionsInternal(TArray<FNiagaraFunctionSignature>& OutFunctions) const
 {
 	FNiagaraFunctionSignature Sig;
 	Sig.Name = GetMousePositionName;
-#if WITH_EDITORONLY_DATA
 	Sig.Description = LOCTEXT("GetMousePositionNameFunctionDescription", "Returns the mouse position in screen space.");
-#endif
 	Sig.bMemberFunction = true;
 	Sig.AddInput(FNiagaraVariable(FNiagaraTypeDefinition(GetClass()), TEXT("MousePosition interface")));
 	Sig.AddInput(FNiagaraVariable(FNiagaraTypeDefinition::GetBoolDef(), TEXT("Normalized")));
@@ -156,6 +155,7 @@ void UNiagaraDataInterfaceMousePosition::GetFunctions(TArray<FNiagaraFunctionSig
 	Sig.AddOutput(FNiagaraVariable(FNiagaraTypeDefinition::GetFloatDef(), TEXT("PosY")), LOCTEXT("MousePosYDescription", "Returns the y coordinates in pixels or 0-1 range if normalized"));
 	OutFunctions.Add(Sig);
 }
+#endif
 
 // this provides the cpu vm with the correct function to call
 void UNiagaraDataInterfaceMousePosition::GetVMExternalFunction(const FVMExternalFunctionBindingInfo& BindingInfo, void* InstanceData, FVMExternalFunction& OutFunc)
diff --git a/Engine/Plugins/FX/ExampleCustomDataInterface/Source/ExampleCustomDataInterface/Private/NiagaraDataInterfaceMousePosition.h b/Engine/Plugins/FX/ExampleCustomDataInterface/Source/ExampleCustomDataInterface/Private/NiagaraDataInterfaceMousePosition.h
--- a/Engine/Plugins/FX/ExampleCustomDataInterface/Source/ExampleCustomDataInterface/Private/NiagaraDataInterfaceMousePosition.h
+++ b/Engine/Plugins/FX/ExampleCustomDataInterface/Source/ExampleCustomDataInterface/Private/NiagaraDataInterfaceMousePosition.h
@@ -24,7 +24,6 @@ public:
 	//UObject Interface End
 
 	//UNiagaraDataInterface Interface
-	virtual void GetFunctions(TArray<FNiagaraFunctionSignature>& OutFunctions) override;
 	virtual void GetVMExternalFunction(const FVMExternalFunctionBindingInfo& BindingInfo, void* InstanceData, FVMExternalFunction &OutFunc) override;
 	virtual bool CanExecuteOnTarget(ENiagaraSimTarget Target) const override { return true; }
 #if WITH_EDITORONLY_DATA
@@ -44,6 +43,12 @@ public:
 	//UNiagaraDataInterface Interface
 
 	void GetMousePositionVM(FVectorVMExternalFunctionContext& Context);	
+
+protected:
+#if WITH_EDITORONLY_DATA
+	virtual void GetFunctionsInternal(TArray<FNiagaraFunctionSignature>& OutFunctions) const override;
+#endif
+
 private:
 	static const FName GetMousePositionName;
 };
diff --git a/Engine/Plugins/FX/Niagara/Source/Niagara/Classes/NiagaraDataInterface.h b/Engine/Plugins/FX/Niagara/Source/Niagara/Classes/NiagaraDataInterface.h
--- a/Engine/Plugins/FX/Niagara/Source/Niagara/Classes/NiagaraDataInterface.h
+++ b/Engine/Plugins/FX/Niagara/Source/Niagara/Classes/NiagaraDataInterface.h
@@ -649,8 +649,12 @@ public:
 	*/
 	virtual int32 PerInstanceDataSize()const { return 0; }
 
-	/** Gets all the available functions for this data interface. */
-	virtual void GetFunctions(TArray<FNiagaraFunctionSignature>& OutFunctions) {}
+#if WITH_EDITORONLY_DATA
+	/**
+	Gets all the available functions for this data interface.
+	*/
+	NIAGARA_API void GetFunctionSignatures(TArray<FNiagaraFunctionSignature>& OutFunctions) const;
+#endif
 
 	/** Returns the delegate for the passed function signature. */
 	virtual void GetVMExternalFunction(const FVMExternalFunctionBindingInfo& BindingInfo, void* InstanceData, FVMExternalFunction &OutFunc) { };
@@ -820,6 +824,20 @@ public:
 protected:
 	virtual void PushToRenderThreadImpl() {}
 
+	// deprecated function for backwards compatibility.  Callers should be using GetFunctionSignatures
+	// and sub classes should be implementing GetFunctionsInternal().
+	UE_DEPRECATED(5.4, "GetFunctions() should be renamed GetFunctionsInternal() and guarded with WITH_EDITORONLY_DATA.")
+	virtual void GetFunctions(TArray<FNiagaraFunctionSignature>& OutFunctions)
+	{
+#if WITH_EDITORONLY_DATA
+		GetFunctionsInternal(OutFunctions);
+#endif
+	}
+
+#if WITH_EDITORONLY_DATA
+	virtual void GetFunctionsInternal(TArray<FNiagaraFunctionSignature>& OutFunctions) const {};
+#endif
+
 public:
 	void PushToRenderThread()
 	{
diff --git a/Engine/Plugins/FX/Niagara/Source/Niagara/Classes/NiagaraDataInterface2DArrayTexture.h b/Engine/Plugins/FX/Niagara/Source/Niagara/Classes/NiagaraDataInterface2DArrayTexture.h
--- a/Engine/Plugins/FX/Niagara/Source/Niagara/Classes/NiagaraDataInterface2DArrayTexture.h
+++ b/Engine/Plugins/FX/Niagara/Source/Niagara/Classes/NiagaraDataInterface2DArrayTexture.h
@@ -34,7 +34,6 @@ public:
 	//UObject Interface End
 
 	//UNiagaraDataInterface Interface
-	NIAGARA_API virtual void GetFunctions(TArray<FNiagaraFunctionSignature>& OutFunctions)override;
 	NIAGARA_API virtual void GetVMExternalFunction(const FVMExternalFunctionBindingInfo& BindingInfo, void* InstanceData, FVMExternalFunction &OutFunc) override;
 	virtual bool CanExecuteOnTarget(ENiagaraSimTarget Target)const override { return Target == ENiagaraSimTarget::GPUComputeSim; }
 
@@ -64,6 +63,10 @@ public:
 protected:
 	NIAGARA_API virtual bool CopyToInternal(UNiagaraDataInterface* Destination) const override;
 
+#if WITH_EDITORONLY_DATA
+	virtual void GetFunctionsInternal(TArray<FNiagaraFunctionSignature>& OutFunctions) const override;
+#endif
+
 protected:
 	static NIAGARA_API const TCHAR* TemplateShaderFilePath;
 	static NIAGARA_API const FName LoadTextureName;
diff --git a/Engine/Plugins/FX/Niagara/Source/Niagara/Classes/NiagaraDataInterfaceArray.h b/Engine/Plugins/FX/Niagara/Source/Niagara/Classes/NiagaraDataInterfaceArray.h
--- a/Engine/Plugins/FX/Niagara/Source/Niagara/Classes/NiagaraDataInterfaceArray.h
+++ b/Engine/Plugins/FX/Niagara/Source/Niagara/Classes/NiagaraDataInterfaceArray.h
@@ -49,9 +49,10 @@ struct INDIArrayProxyBase : public FNiagaraDataInterfaceProxyRW
 	END_SHADER_PARAMETER_STRUCT()
 
 	virtual ~INDIArrayProxyBase() {}
-	virtual void GetFunctions(TArray<FNiagaraFunctionSignature>& OutFunctions) const = 0;
 	virtual void GetVMExternalFunction(const FVMExternalFunctionBindingInfo& BindingInfo, void* InstanceData, FVMExternalFunction &OutFunc) = 0;
 #if WITH_EDITORONLY_DATA
+	virtual void GetFunctions(TArray<FNiagaraFunctionSignature>& OutFunctions) const = 0;
+
 	virtual void GetParameterDefinitionHLSL(const FNiagaraDataInterfaceGPUParamInfo& ParamInfo, FString& OutHLSL) const = 0;
 	virtual bool GetFunctionHLSL(const FNiagaraDataInterfaceGPUParamInfo& ParamInfo, const FNiagaraDataInterfaceGeneratedFunction& FunctionInfo, int FunctionInstanceIndex, FString& OutHLSL) const = 0;
 	virtual bool AppendCompileHash(FNiagaraCompileHashVisitor* InVisitor) const = 0;
@@ -84,7 +85,6 @@ public:
 	//UObject Interface End
 
 	//UNiagaraDataInterface Interface
-	virtual void GetFunctions(TArray<FNiagaraFunctionSignature>& OutFunctions) override { GetProxyAs<INDIArrayProxyBase>()->GetFunctions(OutFunctions); }
 	virtual void GetVMExternalFunction(const FVMExternalFunctionBindingInfo& BindingInfo, void* InstanceData, FVMExternalFunction &OutFunc) override { GetProxyAs<INDIArrayProxyBase>()->GetVMExternalFunction(BindingInfo, InstanceData, OutFunc); }
 #if WITH_EDITORONLY_DATA
 	virtual void GetParameterDefinitionHLSL(const FNiagaraDataInterfaceGPUParamInfo& ParamInfo, FString& OutHLSL) override { GetProxyAs<INDIArrayProxyBase>()->GetParameterDefinitionHLSL(ParamInfo, OutHLSL); }
@@ -121,4 +121,9 @@ public:
 	/** When greater than 0 sets the maximum number of elements the array can hold, only relevant when using operations that modify the array size. */
 	UPROPERTY(EditAnywhere, BlueprintReadOnly, AdvancedDisplay, Category="Array", meta=(ClampMin="0"))
 	int32 MaxElements;
+
+protected:
+#if WITH_EDITORONLY_DATA
+	virtual void GetFunctionsInternal(TArray<FNiagaraFunctionSignature>& OutFunctions) const override { GetProxyAs<INDIArrayProxyBase>()->GetFunctions(OutFunctions); }
+#endif
 };
diff --git a/Engine/Plugins/FX/Niagara/Source/Niagara/Classes/NiagaraDataInterfaceArrayImpl.h b/Engine/Plugins/FX/Niagara/Source/Niagara/Classes/NiagaraDataInterfaceArrayImpl.h
--- a/Engine/Plugins/FX/Niagara/Source/Niagara/Classes/NiagaraDataInterfaceArrayImpl.h
+++ b/Engine/Plugins/FX/Niagara/Source/Niagara/Classes/NiagaraDataInterfaceArrayImpl.h
@@ -682,6 +682,7 @@ struct FNDIArrayProxyImpl : public INDIArrayProxyBase
 		return ValueOut;
 	}
 
+#if WITH_EDITORONLY_DATA
 	//////////////////////////////////////////////////////////////////////////
 	// INDIArrayProxyBase
 	virtual void GetFunctions(TArray<FNiagaraFunctionSignature>& OutFunctions) const override
@@ -695,16 +696,12 @@ struct FNDIArrayProxyImpl : public INDIArrayProxyBase
 		DefaultImmutableSig.bSupportsCPU = FNDIArrayImplHelper<TArrayType>::bSupportsCPU;
 		DefaultImmutableSig.bSupportsGPU = FNDIArrayImplHelper<TArrayType>::bSupportsGPU;
 		DefaultImmutableSig.Inputs.Emplace(FNiagaraTypeDefinition(TOwnerType::StaticClass()), TEXT("Array interface"));
-#if WITH_EDITORONLY_DATA
 		DefaultImmutableSig.FunctionVersion = FNiagaraDataInterfaceArrayImplHelper::FFunctionVersion::LatestVersion;
-#endif
 		{
 			FNiagaraFunctionSignature& Sig = OutFunctions.Add_GetRef(DefaultImmutableSig);
 			Sig.Name = FNiagaraDataInterfaceArrayImplHelper::Function_LengthName;
 			Sig.Outputs.Emplace(FNiagaraTypeDefinition::GetIntDef(), TEXT("Num"));
-		#if WITH_EDITORONLY_DATA
 			Sig.Description = NSLOCTEXT("Niagara", "Array_LengthDesc", "Gets the number of elements in the array.");
-		#endif
 		}
 
 		{
@@ -712,18 +709,14 @@ struct FNDIArrayProxyImpl : public INDIArrayProxyBase
 			Sig.Name = FNiagaraDataInterfaceArrayImplHelper::Function_IsValidIndexName;
 			Sig.Inputs.Emplace(FNiagaraTypeDefinition::GetIntDef(), TEXT("Index"));
 			Sig.Outputs.Emplace(FNiagaraTypeDefinition::GetBoolDef(), TEXT("Valid"));
-		#if WITH_EDITORONLY_DATA
 			Sig.Description = NSLOCTEXT("Niagara", "Array_IsValidIndexDesc", "Tests to see if the index is valid and exists in the array.");
-		#endif
 		}
 
 		{
 			FNiagaraFunctionSignature& Sig = OutFunctions.Add_GetRef(DefaultImmutableSig);
 			Sig.Name = FNiagaraDataInterfaceArrayImplHelper::Function_LastIndexName;
 			Sig.Outputs.Emplace(FNiagaraTypeDefinition::GetIntDef(), TEXT("Index"));
-		#if WITH_EDITORONLY_DATA
 			Sig.Description = NSLOCTEXT("Niagara", "Array_LastIndexDesc", "Returns the last valid index in the array, will be -1 if no elements.");
-		#endif
 		}
 
 		{
@@ -731,9 +724,7 @@ struct FNDIArrayProxyImpl : public INDIArrayProxyBase
 			Sig.Name = FNiagaraDataInterfaceArrayImplHelper::Function_GetName;
 			Sig.Inputs.Emplace(FNiagaraTypeDefinition::GetIntDef(), TEXT("Index"));
 			Sig.Outputs.Emplace(FNDIArrayImplHelper<TArrayType>::GetTypeDefinition(), TEXT("Value"));
-		#if WITH_EDITORONLY_DATA
 			Sig.Description = NSLOCTEXT("Niagara", "Array_GetDesc", "Gets the value from the array at the given zero based index.");
-		#endif
 		}
 
 		// Mutable functions
@@ -744,9 +735,7 @@ struct FNDIArrayProxyImpl : public INDIArrayProxyBase
 			FNiagaraFunctionSignature& Sig = OutFunctions.Add_GetRef(DefaultMutableSig);
 			Sig.Name = FNiagaraDataInterfaceArrayImplHelper::Function_ClearName;
 			Sig.ModuleUsageBitmask = ENiagaraScriptUsageMask::System | ENiagaraScriptUsageMask::Emitter;
-		#if WITH_EDITORONLY_DATA
 			Sig.Description = NSLOCTEXT("Niagara", "Array_ClearDesc", "Clears the array, removing all elements");
-		#endif
 		}
 
 		{
@@ -754,9 +743,7 @@ struct FNDIArrayProxyImpl : public INDIArrayProxyBase
 			Sig.Name = FNiagaraDataInterfaceArrayImplHelper::Function_ResizeName;
 			Sig.ModuleUsageBitmask = ENiagaraScriptUsageMask::System | ENiagaraScriptUsageMask::Emitter;
 			Sig.Inputs.Emplace(FNiagaraTypeDefinition::GetIntDef(), TEXT("Num"));
-		#if WITH_EDITORONLY_DATA
 			Sig.Description = NSLOCTEXT("Niagara", "Array_ResizeDesc", "Resizes the array to the specified size, initializing new elements with the default value.");
-		#endif
 		}
 
 		{
@@ -765,10 +752,8 @@ struct FNDIArrayProxyImpl : public INDIArrayProxyBase
 			Sig.Inputs.Emplace(FNiagaraTypeDefinition::GetBoolDef(), TEXT("SkipSet"));
 			Sig.Inputs.Emplace(FNiagaraTypeDefinition::GetIntDef(), TEXT("Index"));
 			Sig.Inputs.Emplace(FNDIArrayImplHelper<TArrayType>::GetTypeDefinition(), TEXT("Value"));
-		#if WITH_EDITORONLY_DATA
 			Sig.Description = NSLOCTEXT("Niagara", "Array_SetArrayElemDesc", "Sets the value at the given zero based index (i.e the first element is 0).");
 			Sig.InputDescriptions.Add(Sig.Inputs[1], NSLOCTEXT("Niagara", "Array_SetArrayElemDesc_SkipSet", "When enabled will not set the array value."));
-		#endif
 		}
 
 		{
@@ -776,10 +761,8 @@ struct FNDIArrayProxyImpl : public INDIArrayProxyBase
 			Sig.Name = FNiagaraDataInterfaceArrayImplHelper::Function_AddName;
 			Sig.Inputs.Emplace(FNiagaraTypeDefinition::GetBoolDef(), TEXT("SkipAdd"));
 			Sig.Inputs.Emplace(FNDIArrayImplHelper<TArrayType>::GetTypeDefinition(), TEXT("Value"));
-		#if WITH_EDITORONLY_DATA
 			Sig.Description = NSLOCTEXT("Niagara", "Array_AddDesc", "Optionally add a value onto the end of the array.");
 			Sig.InputDescriptions.Add(Sig.Inputs[1], NSLOCTEXT("Niagara", "Array_AddDesc_SkipAdd", "When enabled we will not add an element to the array."));
-		#endif
 		}
 
 		{
@@ -788,11 +771,9 @@ struct FNDIArrayProxyImpl : public INDIArrayProxyBase
 			Sig.Inputs.Emplace(FNiagaraTypeDefinition::GetBoolDef(), TEXT("SkipRemove"));
 			Sig.Outputs.Emplace(FNDIArrayImplHelper<TArrayType>::GetTypeDefinition(), TEXT("Value"));
 			Sig.Outputs.Emplace(FNiagaraTypeDefinition::GetBoolDef(), TEXT("IsValid"));
-		#if WITH_EDITORONLY_DATA
 			Sig.Description = NSLOCTEXT("Niagara", "Array_RemoveLastElemDesc", "Optionally remove the last element from the array.  Returns the default value if no elements are in the array or you skip the remove.");
 			Sig.InputDescriptions.Add(Sig.Inputs[1], NSLOCTEXT("Niagara", "Array_RemoveLastElemDesc_SkipRemove", "When enabled will not remove a value from the array, the return value will therefore be invalid."));
 			Sig.OutputDescriptions.Add(Sig.Outputs[1], NSLOCTEXT("Niagara", "Array_RemoveLastElemDesc_IsValid", "True if we removed a value from the array, False if no entries or we skipped the remove."));
-		#endif
 		}
 
 		if (FNDIArrayImplHelper<TArrayType>::bSupportsAtomicOps)
@@ -838,6 +819,7 @@ struct FNDIArrayProxyImpl : public INDIArrayProxyBase
 			}
 		}
 	}
+#endif
 
 	template<typename T = FNDIArrayImplHelper<TArrayType>>
 	typename TEnableIf<T::bSupportsCPU>::Type GetVMExternalFunction_Internal(const FVMExternalFunctionBindingInfo& BindingInfo, void* InstanceData, FVMExternalFunction& OutFunc)
diff --git a/Engine/Plugins/FX/Niagara/Source/Niagara/Classes/NiagaraDataInterfaceAsyncGpuTrace.h b/Engine/Plugins/FX/Niagara/Source/Niagara/Classes/NiagaraDataInterfaceAsyncGpuTrace.h
--- a/Engine/Plugins/FX/Niagara/Source/Niagara/Classes/NiagaraDataInterfaceAsyncGpuTrace.h
+++ b/Engine/Plugins/FX/Niagara/Source/Niagara/Classes/NiagaraDataInterfaceAsyncGpuTrace.h
@@ -44,8 +44,6 @@ public:
 	NIAGARA_API virtual bool InitPerInstanceData(void* PerInstanceData, FNiagaraSystemInstance* InSystemInstance) override;
 	NIAGARA_API virtual void DestroyPerInstanceData(void* PerInstanceData, FNiagaraSystemInstance* InSystemInstance) override;
 
-	NIAGARA_API virtual void GetFunctions(TArray<FNiagaraFunctionSignature>& OutFunctions) override;
-
 	virtual bool CanExecuteOnTarget(ENiagaraSimTarget Target) const override { return Target == ENiagaraSimTarget::GPUComputeSim; }
 	NIAGARA_API virtual bool RequiresGlobalDistanceField() const override;
 	NIAGARA_API virtual bool RequiresRayTracingScene() const override;
@@ -64,6 +62,10 @@ private:
 	NIAGARA_API virtual bool Equals(const UNiagaraDataInterface* Other) const override;
 
 protected:
+#if WITH_EDITORONLY_DATA
+	virtual void GetFunctionsInternal(TArray<FNiagaraFunctionSignature>& OutFunctions) const override;
+#endif
+
 	NIAGARA_API virtual bool CopyToInternal(UNiagaraDataInterface* Destination) const override;
 	NIAGARA_API virtual void PushToRenderThreadImpl() override;
 };
diff --git a/Engine/Plugins/FX/Niagara/Source/Niagara/Classes/NiagaraDataInterfaceAudioOscilloscope.h b/Engine/Plugins/FX/Niagara/Source/Niagara/Classes/NiagaraDataInterfaceAudioOscilloscope.h
--- a/Engine/Plugins/FX/Niagara/Source/Niagara/Classes/NiagaraDataInterfaceAudioOscilloscope.h
+++ b/Engine/Plugins/FX/Niagara/Source/Niagara/Classes/NiagaraDataInterfaceAudioOscilloscope.h
@@ -117,7 +117,6 @@ public:
 	NIAGARA_API void SampleAudio(FVectorVMExternalFunctionContext& Context);
 	NIAGARA_API void GetNumChannels(FVectorVMExternalFunctionContext& Context);
 
-	NIAGARA_API virtual void GetFunctions(TArray<FNiagaraFunctionSignature>& OutFunctions)override;
 	NIAGARA_API virtual void GetVMExternalFunction(const FVMExternalFunctionBindingInfo& BindingInfo, void* InstanceData, FVMExternalFunction &OutFunc) override;
 
 	virtual bool CanExecuteOnTarget(ENiagaraSimTarget Target) const override
@@ -145,6 +144,10 @@ public:
 	
 
 protected:
+#if WITH_EDITORONLY_DATA
+	virtual void GetFunctionsInternal(TArray<FNiagaraFunctionSignature>& OutFunctions) const override;
+#endif
+
 	NIAGARA_API virtual bool CopyToInternal(UNiagaraDataInterface* Destination) const override;
 };
 
diff --git a/Engine/Plugins/FX/Niagara/Source/Niagara/Classes/NiagaraDataInterfaceAudioPlayer.h b/Engine/Plugins/FX/Niagara/Source/Niagara/Classes/NiagaraDataInterfaceAudioPlayer.h
--- a/Engine/Plugins/FX/Niagara/Source/Niagara/Classes/NiagaraDataInterfaceAudioPlayer.h
+++ b/Engine/Plugins/FX/Niagara/Source/Niagara/Classes/NiagaraDataInterfaceAudioPlayer.h
@@ -146,7 +146,6 @@ public:
 	//UObject Interface End
 
 	//UNiagaraDataInterface Interface
-	NIAGARA_API virtual void GetFunctions(TArray<FNiagaraFunctionSignature>& OutFunctions) override;
 	NIAGARA_API virtual void GetVMExternalFunction(const FVMExternalFunctionBindingInfo& BindingInfo, void* InstanceData, FVMExternalFunction &OutFunc) override;
 	NIAGARA_API virtual bool InitPerInstanceData(void* PerInstanceData, FNiagaraSystemInstance* SystemInstance) override;
 	NIAGARA_API virtual void DestroyPerInstanceData(void* PerInstanceData, FNiagaraSystemInstance* SystemInstance) override;
@@ -176,6 +175,10 @@ public:
 	NIAGARA_API virtual void SetPausedState(FVectorVMExternalFunctionContext& Context);
 
 protected:
+#if WITH_EDITORONLY_DATA
+	virtual void GetFunctionsInternal(TArray<FNiagaraFunctionSignature>& OutFunctions) const override;
+#endif
+
 	NIAGARA_API virtual bool CopyToInternal(UNiagaraDataInterface* Destination) const override;
 	
 private:
diff --git a/Engine/Plugins/FX/Niagara/Source/Niagara/Classes/NiagaraDataInterfaceAudioSpectrum.h b/Engine/Plugins/FX/Niagara/Source/Niagara/Classes/NiagaraDataInterfaceAudioSpectrum.h
--- a/Engine/Plugins/FX/Niagara/Source/Niagara/Classes/NiagaraDataInterfaceAudioSpectrum.h
+++ b/Engine/Plugins/FX/Niagara/Source/Niagara/Classes/NiagaraDataInterfaceAudioSpectrum.h
@@ -168,7 +168,6 @@ public:
 	NIAGARA_API void GetSpectrumValue(FVectorVMExternalFunctionContext& Context);
 	NIAGARA_API void GetNumChannels(FVectorVMExternalFunctionContext& Context);
 
-	NIAGARA_API virtual void GetFunctions(TArray<FNiagaraFunctionSignature>& OutFunctions) override;
 	NIAGARA_API virtual void GetVMExternalFunction(const FVMExternalFunctionBindingInfo& BindingInfo, void* InstanceData, FVMExternalFunction &OutFunc) override;
 
 	virtual bool CanExecuteOnTarget(ENiagaraSimTarget Target) const override 
@@ -195,6 +194,10 @@ public:
 	NIAGARA_API virtual void PostLoad() override;
 
 protected:
+#if WITH_EDITORONLY_DATA
+	virtual void GetFunctionsInternal(TArray<FNiagaraFunctionSignature>& OutFunctions) const override;
+#endif
+
 	NIAGARA_API virtual bool CopyToInternal(UNiagaraDataInterface* Destination) const override;
 private:
 
diff --git a/Engine/Plugins/FX/Niagara/Source/Niagara/Classes/NiagaraDataInterfaceCamera.h b/Engine/Plugins/FX/Niagara/Source/Niagara/Classes/NiagaraDataInterfaceCamera.h
--- a/Engine/Plugins/FX/Niagara/Source/Niagara/Classes/NiagaraDataInterfaceCamera.h
+++ b/Engine/Plugins/FX/Niagara/Source/Niagara/Classes/NiagaraDataInterfaceCamera.h
@@ -48,7 +48,6 @@ public:
 	//UObject Interface End
 
 	//UNiagaraDataInterface Interface
-	NIAGARA_API virtual void GetFunctions(TArray<FNiagaraFunctionSignature>& OutFunctions)override;
 	NIAGARA_API virtual void GetVMExternalFunction(const FVMExternalFunctionBindingInfo& BindingInfo, void* InstanceData, FVMExternalFunction &OutFunc) override;
 	NIAGARA_API virtual bool InitPerInstanceData(void* PerInstanceData, FNiagaraSystemInstance* SystemInstance) override;
 	virtual int32 PerInstanceDataSize() const override { return sizeof(FCameraDataInterface_InstanceData); }
@@ -79,6 +78,10 @@ public:
 	NIAGARA_API void GetCameraProperties(FVectorVMExternalFunctionContext& Context);
 
 protected:
+#if WITH_EDITORONLY_DATA
+	virtual void GetFunctionsInternal(TArray<FNiagaraFunctionSignature>& OutFunctions) const override;
+#endif
+
 	NIAGARA_API virtual bool CopyToInternal(UNiagaraDataInterface* Destination) const override;
 private:
 	static NIAGARA_API const TCHAR* TemplateShaderFilePath;
diff --git a/Engine/Plugins/FX/Niagara/Source/Niagara/Classes/NiagaraDataInterfaceCollisionQuery.h b/Engine/Plugins/FX/Niagara/Source/Niagara/Classes/NiagaraDataInterfaceCollisionQuery.h
--- a/Engine/Plugins/FX/Niagara/Source/Niagara/Classes/NiagaraDataInterfaceCollisionQuery.h
+++ b/Engine/Plugins/FX/Niagara/Source/Niagara/Classes/NiagaraDataInterfaceCollisionQuery.h
@@ -46,7 +46,6 @@ public:
 	NIAGARA_API virtual bool PerInstanceTickPostSimulate(void* PerInstanceData, FNiagaraSystemInstance* SystemInstance, float DeltaSeconds) override;
 	virtual int32 PerInstanceDataSize() const override { return sizeof(CQDIPerInstanceData); }
 	
-	NIAGARA_API virtual void GetFunctions(TArray<FNiagaraFunctionSignature>& OutFunctions) override;
 	NIAGARA_API virtual void GetVMExternalFunction(const FVMExternalFunctionBindingInfo& BindingInfo, void* InstanceData, FVMExternalFunction &OutFunc) override;
 	NIAGARA_API virtual void GetAssetTagsForContext(const UObject* InAsset, FGuid AssetVersion, const TArray<const UNiagaraDataInterface*>& InProperties, TMap<FName, uint32>& NumericKeys, TMap<FName, FString>& StringKeys) const override;
 
@@ -76,6 +75,11 @@ public:
 	virtual bool HasPostSimulateTick() const override { return true; }
 	virtual bool PostSimulateCanOverlapFrames() const { return false; }
 
+protected:
+#if WITH_EDITORONLY_DATA
+	virtual void GetFunctionsInternal(TArray<FNiagaraFunctionSignature>& OutFunctions) const override;
+#endif
+
 private:
 	static NIAGARA_API FCriticalSection CriticalSection;
 	UEnum* TraceChannelEnum;
diff --git a/Engine/Plugins/FX/Niagara/Source/Niagara/Classes/NiagaraDataInterfaceColorCurve.h b/Engine/Plugins/FX/Niagara/Source/Niagara/Classes/NiagaraDataInterfaceColorCurve.h
--- a/Engine/Plugins/FX/Niagara/Source/Niagara/Classes/NiagaraDataInterfaceColorCurve.h
+++ b/Engine/Plugins/FX/Niagara/Source/Niagara/Classes/NiagaraDataInterfaceColorCurve.h
@@ -61,7 +61,6 @@ public:
 	NIAGARA_API virtual void UpdateTimeRanges() override;
 	NIAGARA_API virtual TArray<float> BuildLUT(int32 NumEntries) const override;
 
-	NIAGARA_API virtual void GetFunctions(TArray<FNiagaraFunctionSignature>& OutFunctions)override;
 	NIAGARA_API virtual void GetVMExternalFunction(const FVMExternalFunctionBindingInfo& BindingInfo, void* InstanceData, FVMExternalFunction &OutFunc) override;
 
 	template<typename UseLUT>
@@ -83,6 +82,10 @@ public:
 #endif
 
 protected:
+#if WITH_EDITORONLY_DATA
+	virtual void GetFunctionsInternal(TArray<FNiagaraFunctionSignature>& OutFunctions) const override;
+#endif
+
 	NIAGARA_API virtual bool CopyToInternal(UNiagaraDataInterface* Destination) const override;
 
 private:
diff --git a/Engine/Plugins/FX/Niagara/Source/Niagara/Classes/NiagaraDataInterfaceCubeTexture.h b/Engine/Plugins/FX/Niagara/Source/Niagara/Classes/NiagaraDataInterfaceCubeTexture.h
--- a/Engine/Plugins/FX/Niagara/Source/Niagara/Classes/NiagaraDataInterfaceCubeTexture.h
+++ b/Engine/Plugins/FX/Niagara/Source/Niagara/Classes/NiagaraDataInterfaceCubeTexture.h
@@ -34,7 +34,6 @@ public:
 	//UObject Interface End
 
 	//UNiagaraDataInterface Interface
-	NIAGARA_API virtual void GetFunctions(TArray<FNiagaraFunctionSignature>& OutFunctions)override;
 	NIAGARA_API virtual void GetVMExternalFunction(const FVMExternalFunctionBindingInfo& BindingInfo, void* InstanceData, FVMExternalFunction &OutFunc) override;
 	virtual bool CanExecuteOnTarget(ENiagaraSimTarget Target)const override { return Target == ENiagaraSimTarget::GPUComputeSim; }
 
@@ -63,6 +62,9 @@ public:
 	NIAGARA_API void SetTexture(UTexture* InTexture);
 
 protected:
+#if WITH_EDITORONLY_DATA
+	virtual void GetFunctionsInternal(TArray<FNiagaraFunctionSignature>& OutFunctions) const override;
+#endif
 	NIAGARA_API virtual bool CopyToInternal(UNiagaraDataInterface* Destination) const override;
 
 protected:
diff --git a/Engine/Plugins/FX/Niagara/Source/Niagara/Classes/NiagaraDataInterfaceCurlNoise.h b/Engine/Plugins/FX/Niagara/Source/Niagara/Classes/NiagaraDataInterfaceCurlNoise.h
--- a/Engine/Plugins/FX/Niagara/Source/Niagara/Classes/NiagaraDataInterfaceCurlNoise.h
+++ b/Engine/Plugins/FX/Niagara/Source/Niagara/Classes/NiagaraDataInterfaceCurlNoise.h
@@ -34,7 +34,6 @@ public:
 	//UObject Interface End
 
 	//UNiagaraDataInterface Interface
-	NIAGARA_API virtual void GetFunctions(TArray<FNiagaraFunctionSignature>& OutFunctions) override;
 	NIAGARA_API virtual void GetVMExternalFunction(const FVMExternalFunctionBindingInfo& BindingInfo, void* InstanceData, FVMExternalFunction &OutFunc) override;
 	virtual bool CanExecuteOnTarget(ENiagaraSimTarget Target)const override { return true; }
 	NIAGARA_API virtual bool Equals(const UNiagaraDataInterface* Other) const override;
@@ -52,6 +51,9 @@ public:
 	NIAGARA_API virtual void SetShaderParameters(const FNiagaraDataInterfaceSetShaderParametersContext& Context) const override;
 
 protected:
+#if WITH_EDITORONLY_DATA
+	virtual void GetFunctionsInternal(TArray<FNiagaraFunctionSignature>& OutFunctions) const override;
+#endif
 	NIAGARA_API virtual bool CopyToInternal(UNiagaraDataInterface* Destination) const override;
 
 	NIAGARA_API virtual void PushToRenderThreadImpl() override;
diff --git a/Engine/Plugins/FX/Niagara/Source/Niagara/Classes/NiagaraDataInterfaceCurve.h b/Engine/Plugins/FX/Niagara/Source/Niagara/Classes/NiagaraDataInterfaceCurve.h
--- a/Engine/Plugins/FX/Niagara/Source/Niagara/Classes/NiagaraDataInterfaceCurve.h
+++ b/Engine/Plugins/FX/Niagara/Source/Niagara/Classes/NiagaraDataInterfaceCurve.h
@@ -40,7 +40,6 @@ public:
 	NIAGARA_API virtual void UpdateTimeRanges() override;
 	NIAGARA_API virtual TArray<float> BuildLUT(int32 NumEntries) const override;
 
-	NIAGARA_API virtual void GetFunctions(TArray<FNiagaraFunctionSignature>& OutFunctions)override;
 	NIAGARA_API virtual void GetVMExternalFunction(const FVMExternalFunctionBindingInfo& BindingInfo, void* InstanceData, FVMExternalFunction &OutFunc) override;
 
 	template<typename UseLUT>
@@ -62,6 +61,9 @@ public:
 #endif
 
 protected:
+#if WITH_EDITORONLY_DATA
+	virtual void GetFunctionsInternal(TArray<FNiagaraFunctionSignature>& OutFunctions) const override;
+#endif
 	NIAGARA_API virtual bool CopyToInternal(UNiagaraDataInterface* Destination) const override;
 
 private:
diff --git a/Engine/Plugins/FX/Niagara/Source/Niagara/Classes/NiagaraDataInterfaceExport.h b/Engine/Plugins/FX/Niagara/Source/Niagara/Classes/NiagaraDataInterfaceExport.h
--- a/Engine/Plugins/FX/Niagara/Source/Niagara/Classes/NiagaraDataInterfaceExport.h
+++ b/Engine/Plugins/FX/Niagara/Source/Niagara/Classes/NiagaraDataInterfaceExport.h
@@ -75,7 +75,6 @@ public:
 	//UObject Interface End
 
 	//UNiagaraDataInterface Interface
-	NIAGARA_API virtual void GetFunctions(TArray<FNiagaraFunctionSignature>& OutFunctions) override;
 	NIAGARA_API virtual void GetVMExternalFunction(const FVMExternalFunctionBindingInfo& BindingInfo, void* InstanceData, FVMExternalFunction &OutFunc) override;
 	NIAGARA_API virtual bool InitPerInstanceData(void* PerInstanceData, FNiagaraSystemInstance* SystemInstance) override;
 	NIAGARA_API virtual void DestroyPerInstanceData(void* PerInstanceData, FNiagaraSystemInstance* SystemInstance) override;
@@ -104,5 +103,8 @@ public:
 	NIAGARA_API virtual void ExportData(FVectorVMExternalFunctionContext& Context);
 
 protected:
+#if WITH_EDITORONLY_DATA
+	virtual void GetFunctionsInternal(TArray<FNiagaraFunctionSignature>& OutFunctions) const override;
+#endif
 	NIAGARA_API virtual bool CopyToInternal(UNiagaraDataInterface* Destination) const override;
 };
diff --git a/Engine/Plugins/FX/Niagara/Source/Niagara/Classes/NiagaraDataInterfaceGrid2DCollection.h b/Engine/Plugins/FX/Niagara/Source/Niagara/Classes/NiagaraDataInterfaceGrid2DCollection.h
--- a/Engine/Plugins/FX/Niagara/Source/Niagara/Classes/NiagaraDataInterfaceGrid2DCollection.h
+++ b/Engine/Plugins/FX/Niagara/Source/Niagara/Classes/NiagaraDataInterfaceGrid2DCollection.h
@@ -139,7 +139,6 @@ public:
 	
 	//~ UNiagaraDataInterface interface
 	// VM functionality
-	NIAGARA_API virtual void GetFunctions(TArray<FNiagaraFunctionSignature>& OutFunctions) override;
 	NIAGARA_API virtual void GetVMExternalFunction(const FVMExternalFunctionBindingInfo& BindingInfo, void* InstanceData, FVMExternalFunction &OutFunc) override;
 
 #if WITH_EDITOR	
@@ -293,6 +292,10 @@ public:
 
 	TMap<FNiagaraSystemInstanceID, FGrid2DCollectionRWInstanceData_GameThread*>& GetSystemInstancesToProxyData_GT() { return SystemInstancesToProxyData_GT; }
 protected:
+#if WITH_EDITORONLY_DATA
+	virtual void GetFunctionsInternal(TArray<FNiagaraFunctionSignature>& OutFunctions) const override;
+#endif
+
 #if WITH_EDITORONLY_DATA
 	NIAGARA_API void WriteSetHLSL(const FNiagaraDataInterfaceRWAttributeHelper &AttributeHelper, const FNiagaraDataInterfaceGPUParamInfo& ParamInfo, const FNiagaraDataInterfaceGeneratedFunction& FunctionInfo, int FunctionInstanceIndex, int32 InNumChannels, FString& OutHLSL);
 	NIAGARA_API void WriteGetHLSL(const FNiagaraDataInterfaceRWAttributeHelper& AttributeHelper, const FNiagaraDataInterfaceGPUParamInfo& ParamInfo, const FNiagaraDataInterfaceGeneratedFunction& FunctionInfo, int FunctionInstanceIndex, int32 InNumChannels, FString& OutHLSL);
diff --git a/Engine/Plugins/FX/Niagara/Source/Niagara/Classes/NiagaraDataInterfaceGrid2DCollectionReader.h b/Engine/Plugins/FX/Niagara/Source/Niagara/Classes/NiagaraDataInterfaceGrid2DCollectionReader.h
--- a/Engine/Plugins/FX/Niagara/Source/Niagara/Classes/NiagaraDataInterfaceGrid2DCollectionReader.h
+++ b/Engine/Plugins/FX/Niagara/Source/Niagara/Classes/NiagaraDataInterfaceGrid2DCollectionReader.h
@@ -33,11 +33,13 @@ public:
 	NIAGARA_API virtual bool Equals(const UNiagaraDataInterface* Other) const override;
 	NIAGARA_API virtual bool InitPerInstanceData(void* PerInstanceData, FNiagaraSystemInstance* SystemInstance) override;
 	NIAGARA_API virtual void GetEmitterDependencies(UNiagaraSystem* Asset, TArray<FVersionedNiagaraEmitter>& Dependencies) const override;
-	NIAGARA_API virtual void GetFunctions(TArray<FNiagaraFunctionSignature>& OutFunctions) override;
 	NIAGARA_API virtual bool PerInstanceTickPostSimulate(void* PerInstanceData, FNiagaraSystemInstance* SystemInstance, float DeltaSeconds) override;
 
 protected:
 	//~ UNiagaraDataInterface interface
+#if WITH_EDITORONLY_DATA
+	virtual void GetFunctionsInternal(TArray<FNiagaraFunctionSignature>& OutFunctions) const override;
+#endif
 	NIAGARA_API virtual bool CopyToInternal(UNiagaraDataInterface* Destination) const override;
 	//~ UNiagaraDataInterface interface END
 };
diff --git a/Engine/Plugins/FX/Niagara/Source/Niagara/Classes/NiagaraDataInterfaceGrid3DCollection.h b/Engine/Plugins/FX/Niagara/Source/Niagara/Classes/NiagaraDataInterfaceGrid3DCollection.h
--- a/Engine/Plugins/FX/Niagara/Source/Niagara/Classes/NiagaraDataInterfaceGrid3DCollection.h
+++ b/Engine/Plugins/FX/Niagara/Source/Niagara/Classes/NiagaraDataInterfaceGrid3DCollection.h
@@ -153,7 +153,6 @@ public:
 
 	//~ UNiagaraDataInterface interface
 	// VM functionality
-	NIAGARA_API virtual void GetFunctions(TArray<FNiagaraFunctionSignature>& OutFunctions) override;
 	NIAGARA_API virtual void GetVMExternalFunction(const FVMExternalFunctionBindingInfo& BindingInfo, void* InstanceData, FVMExternalFunction& OutFunc) override;
 
 #if WITH_EDITOR	
@@ -322,6 +321,8 @@ public:
 	TMap<FNiagaraSystemInstanceID, FGrid3DCollectionRWInstanceData_GameThread*>& GetSystemInstancesToProxyData_GT() { return SystemInstancesToProxyData_GT; }
 protected:
 #if WITH_EDITORONLY_DATA
+	virtual void GetFunctionsInternal(TArray<FNiagaraFunctionSignature>& OutFunctions) const override;
+
 	NIAGARA_API void WriteSetHLSL(const FNiagaraDataInterfaceGPUParamInfo& ParamInfo, const FNiagaraDataInterfaceGeneratedFunction& FunctionInfo, int FunctionInstanceIndex, int32 InNumChannels, FString& OutHLSL);
 	NIAGARA_API void WriteSetAtIndexHLSL(const FNiagaraDataInterfaceGPUParamInfo& ParamInfo, const FNiagaraDataInterfaceGeneratedFunction& FunctionInfo, int32 InNumChannels, FString& OutHLSL);
 	NIAGARA_API void WriteGetHLSL(const FNiagaraDataInterfaceGPUParamInfo& ParamInfo, const FNiagaraDataInterfaceGeneratedFunction& FunctionInfo, int FunctionInstanceIndex, int32 InNumChannels, FString& OutHLSL);
diff --git a/Engine/Plugins/FX/Niagara/Source/Niagara/Classes/NiagaraDataInterfaceGrid3DCollectionReader.h b/Engine/Plugins/FX/Niagara/Source/Niagara/Classes/NiagaraDataInterfaceGrid3DCollectionReader.h
--- a/Engine/Plugins/FX/Niagara/Source/Niagara/Classes/NiagaraDataInterfaceGrid3DCollectionReader.h
+++ b/Engine/Plugins/FX/Niagara/Source/Niagara/Classes/NiagaraDataInterfaceGrid3DCollectionReader.h
@@ -30,11 +30,13 @@ public:
 	NIAGARA_API virtual bool Equals(const UNiagaraDataInterface* Other) const override;
 	NIAGARA_API virtual bool InitPerInstanceData(void* PerInstanceData, FNiagaraSystemInstance* SystemInstance) override;
 	NIAGARA_API virtual void GetEmitterDependencies(UNiagaraSystem* Asset, TArray<FVersionedNiagaraEmitter>& Dependencies) const override;
-	NIAGARA_API virtual void GetFunctions(TArray<FNiagaraFunctionSignature>& OutFunctions) override;
 	NIAGARA_API virtual bool PerInstanceTickPostSimulate(void* PerInstanceData, FNiagaraSystemInstance* SystemInstance, float DeltaSeconds) override;
 
 protected:
 	//~ UNiagaraDataInterface interface
+#if WITH_EDITORONLY_DATA
+	virtual void GetFunctionsInternal(TArray<FNiagaraFunctionSignature>& OutFunctions) const override;
+#endif
 	NIAGARA_API virtual bool CopyToInternal(UNiagaraDataInterface* Destination) const override;
 	//~ UNiagaraDataInterface interface END
 };
diff --git a/Engine/Plugins/FX/Niagara/Source/Niagara/Classes/NiagaraDataInterfaceIntRenderTarget2D.h b/Engine/Plugins/FX/Niagara/Source/Niagara/Classes/NiagaraDataInterfaceIntRenderTarget2D.h
--- a/Engine/Plugins/FX/Niagara/Source/Niagara/Classes/NiagaraDataInterfaceIntRenderTarget2D.h
+++ b/Engine/Plugins/FX/Niagara/Source/Niagara/Classes/NiagaraDataInterfaceIntRenderTarget2D.h
@@ -25,7 +25,6 @@ public:
 
 	//~ UNiagaraDataInterface interface
 	virtual bool CanExecuteOnTarget(ENiagaraSimTarget Target)const override { return true; }
-	NIAGARA_API virtual void GetFunctions(TArray<FNiagaraFunctionSignature>& OutFunctions) override;
 #if WITH_EDITORONLY_DATA
 	NIAGARA_API virtual bool UpgradeFunctionCall(FNiagaraFunctionSignature& FunctionSignature) override;
 #endif
@@ -75,5 +74,8 @@ public:
 	FNiagaraUserParameterBinding RenderTargetUserParameter;
 
 protected:
+#if WITH_EDITORONLY_DATA
+	virtual void GetFunctionsInternal(TArray<FNiagaraFunctionSignature>& OutFunctions) const override;
+#endif
 	static NIAGARA_API FNiagaraVariableBase ExposedRTVar;
 };
diff --git a/Engine/Plugins/FX/Niagara/Source/Niagara/Classes/NiagaraDataInterfaceLandscape.h b/Engine/Plugins/FX/Niagara/Source/Niagara/Classes/NiagaraDataInterfaceLandscape.h
--- a/Engine/Plugins/FX/Niagara/Source/Niagara/Classes/NiagaraDataInterfaceLandscape.h
+++ b/Engine/Plugins/FX/Niagara/Source/Niagara/Classes/NiagaraDataInterfaceLandscape.h
@@ -58,7 +58,6 @@ public:
 	//UObject Interface End
 
 	//UNiagaraDataInterface Interface
-	NIAGARA_API virtual void GetFunctions(TArray<FNiagaraFunctionSignature>& OutFunctions) override;
 	virtual bool CanExecuteOnTarget(ENiagaraSimTarget Target) const override { return Target == ENiagaraSimTarget::GPUComputeSim; }
 #if WITH_EDITORONLY_DATA
 	NIAGARA_API virtual bool UpgradeFunctionCall(FNiagaraFunctionSignature& FunctionSignature) override;
@@ -84,6 +83,9 @@ public:
 	virtual bool HasPreSimulateTick() const override { return true; }
 
 protected:
+#if WITH_EDITORONLY_DATA
+	virtual void GetFunctionsInternal(TArray<FNiagaraFunctionSignature>& OutFunctions) const override;
+#endif
 	NIAGARA_API virtual bool CopyToInternal(UNiagaraDataInterface* Destination) const override;
 	NIAGARA_API void ApplyLandscape(const FNiagaraSystemInstance& SystemInstance, FNDILandscapeData_GameThread& InstanceData) const;
 	NIAGARA_API ALandscape* GetLandscape(const FNiagaraSystemInstance& SystemInstance, FNDILandscapeData_GameThread& Hint) const;
diff --git a/Engine/Plugins/FX/Niagara/Source/Niagara/Classes/NiagaraDataInterfaceMeshRendererInfo.h b/Engine/Plugins/FX/Niagara/Source/Niagara/Classes/NiagaraDataInterfaceMeshRendererInfo.h
--- a/Engine/Plugins/FX/Niagara/Source/Niagara/Classes/NiagaraDataInterfaceMeshRendererInfo.h
+++ b/Engine/Plugins/FX/Niagara/Source/Niagara/Classes/NiagaraDataInterfaceMeshRendererInfo.h
@@ -38,7 +38,6 @@ public:
 	NIAGARA_API void OnMeshRendererChanged(UNiagaraMeshRendererProperties* NewMeshRenderer);
 
 	//UNiagaraDataInterface Interface
-	NIAGARA_API virtual void GetFunctions(TArray<FNiagaraFunctionSignature>& OutFunctions) override;
 	NIAGARA_API virtual void GetVMExternalFunction(const FVMExternalFunctionBindingInfo& BindingInfo, void* InstanceData, FVMExternalFunction& OutFunc) override;
 	NIAGARA_API virtual bool Equals(const UNiagaraDataInterface* Other) const override;
 #if WITH_EDITORONLY_DATA
@@ -58,6 +57,9 @@ public:
 	//UNiagaraDataInterface Interface	
 
 protected:
+#if WITH_EDITORONLY_DATA
+	virtual void GetFunctionsInternal(TArray<FNiagaraFunctionSignature>& OutFunctions) const override;
+#endif
 	NIAGARA_API virtual bool CopyToInternal(UNiagaraDataInterface* Destination) const override;
 	NIAGARA_API virtual void PushToRenderThreadImpl() override;
 	NIAGARA_API void UpdateCachedData();
diff --git a/Engine/Plugins/FX/Niagara/Source/Niagara/Classes/NiagaraDataInterfaceNeighborGrid3D.h b/Engine/Plugins/FX/Niagara/Source/Niagara/Classes/NiagaraDataInterfaceNeighborGrid3D.h
--- a/Engine/Plugins/FX/Niagara/Source/Niagara/Classes/NiagaraDataInterfaceNeighborGrid3D.h
+++ b/Engine/Plugins/FX/Niagara/Source/Niagara/Classes/NiagaraDataInterfaceNeighborGrid3D.h
@@ -92,7 +92,6 @@ public:
 
 	//~ UNiagaraDataInterface interface
 	// VM functionality
-	NIAGARA_API virtual void GetFunctions(TArray<FNiagaraFunctionSignature>& OutFunctions) override;
 	NIAGARA_API virtual void GetVMExternalFunction(const FVMExternalFunctionBindingInfo& BindingInfo, void* InstanceData, FVMExternalFunction &OutFunc) override;
 
 	NIAGARA_API virtual bool Equals(const UNiagaraDataInterface* Other) const override;
@@ -129,6 +128,9 @@ public:
 
 protected:
 	//~ UNiagaraDataInterface interface
+#if WITH_EDITORONLY_DATA
+	virtual void GetFunctionsInternal(TArray<FNiagaraFunctionSignature>& OutFunctions) const override;
+#endif
 	NIAGARA_API virtual bool CopyToInternal(UNiagaraDataInterface* Destination) const override;
 	//~ UNiagaraDataInterface interface END	
 };
diff --git a/Engine/Plugins/FX/Niagara/Source/Niagara/Classes/NiagaraDataInterfaceOcclusion.h b/Engine/Plugins/FX/Niagara/Source/Niagara/Classes/NiagaraDataInterfaceOcclusion.h
--- a/Engine/Plugins/FX/Niagara/Source/Niagara/Classes/NiagaraDataInterfaceOcclusion.h
+++ b/Engine/Plugins/FX/Niagara/Source/Niagara/Classes/NiagaraDataInterfaceOcclusion.h
@@ -23,7 +23,6 @@ public:
 	//UObject Interface End
 
 	//UNiagaraDataInterface Interface
-	NIAGARA_API virtual void GetFunctions(TArray<FNiagaraFunctionSignature>& OutFunctions)override;
 #if WITH_EDITORONLY_DATA
 	NIAGARA_API virtual bool AppendCompileHash(FNiagaraCompileHashVisitor* InVisitor) const override;
 	NIAGARA_API virtual void GetCommonHLSL(FString& OutHLSL) override;
@@ -37,6 +36,12 @@ public:
 	virtual bool CanExecuteOnTarget(ENiagaraSimTarget Target) const override { return Target == ENiagaraSimTarget::GPUComputeSim; }
 	virtual bool RequiresDepthBuffer() const override { return true; }
 	//UNiagaraDataInterface Interface
+
+protected:
+#if WITH_EDITORONLY_DATA
+	virtual void GetFunctionsInternal(TArray<FNiagaraFunctionSignature>& OutFunctions) const override;
+#endif
+
 };
 
 struct FNiagaraDataIntefaceProxyOcclusionQuery : public FNiagaraDataInterfaceProxy
diff --git a/Engine/Plugins/FX/Niagara/Source/Niagara/Classes/NiagaraDataInterfaceParticleRead.h b/Engine/Plugins/FX/Niagara/Source/Niagara/Classes/NiagaraDataInterfaceParticleRead.h
--- a/Engine/Plugins/FX/Niagara/Source/Niagara/Classes/NiagaraDataInterfaceParticleRead.h
+++ b/Engine/Plugins/FX/Niagara/Source/Niagara/Classes/NiagaraDataInterfaceParticleRead.h
@@ -45,7 +45,6 @@ public:
 	NIAGARA_API virtual bool InitPerInstanceData(void* PerInstanceData, FNiagaraSystemInstance* SystemInstance) override;
 	NIAGARA_API virtual void DestroyPerInstanceData(void* PerInstanceData, FNiagaraSystemInstance* SystemInstance) override;
 	NIAGARA_API virtual int32 PerInstanceDataSize() const override;
-	NIAGARA_API virtual void GetFunctions(TArray<FNiagaraFunctionSignature>& OutFunctions) override;
 	NIAGARA_API virtual void GetVMExternalFunction(const FVMExternalFunctionBindingInfo& BindingInfo, void* InstanceData, FVMExternalFunction& OutFunc) override;
 	virtual bool CanExecuteOnTarget(ENiagaraSimTarget Target) const override { return true; }
 	NIAGARA_API virtual bool Equals(const UNiagaraDataInterface* Other) const override;
@@ -100,8 +99,12 @@ public:
 	NIAGARA_API void ReadIDByIndex(FVectorVMExternalFunctionContext& Context, FName AttributeToRead);
 
 protected:
-	NIAGARA_API void GetPersistentIDFunctions(TArray<FNiagaraFunctionSignature>& OutFunctions);
-	NIAGARA_API void GetIndexFunctions(TArray<FNiagaraFunctionSignature>& OutFunctions);
+#if WITH_EDITORONLY_DATA
+	virtual void GetFunctionsInternal(TArray<FNiagaraFunctionSignature>& OutFunctions) const override;
+
+	NIAGARA_API void GetPersistentIDFunctions(TArray<FNiagaraFunctionSignature>& OutFunctions) const;
+	NIAGARA_API void GetIndexFunctions(TArray<FNiagaraFunctionSignature>& OutFunctions) const;
+#endif
 	NIAGARA_API virtual bool CopyToInternal(UNiagaraDataInterface* Destination) const override;
 
 #if WITH_EDITORONLY_DATA
diff --git a/Engine/Plugins/FX/Niagara/Source/Niagara/Classes/NiagaraDataInterfacePlatformSet.h b/Engine/Plugins/FX/Niagara/Source/Niagara/Classes/NiagaraDataInterfacePlatformSet.h
--- a/Engine/Plugins/FX/Niagara/Source/Niagara/Classes/NiagaraDataInterfacePlatformSet.h
+++ b/Engine/Plugins/FX/Niagara/Source/Niagara/Classes/NiagaraDataInterfacePlatformSet.h
@@ -23,7 +23,6 @@ public:
 	//UObject Interface End
 
 	// UNiagaraDataInterface interface
-	NIAGARA_API virtual void GetFunctions(TArray<FNiagaraFunctionSignature>& OutFunctions)override;
 	NIAGARA_API virtual void GetVMExternalFunction(const FVMExternalFunctionBindingInfo& BindingInfo, void* InstanceData, FVMExternalFunction& OutFunc) override;
 	virtual bool CanExecuteOnTarget(ENiagaraSimTarget Target)const override { return Target == ENiagaraSimTarget::CPUSim; }
 
@@ -35,6 +34,9 @@ protected:
 	//~ UNiagaraDataInterface interface
 
 public:
+#if WITH_EDITORONLY_DATA
+	NIAGARA_API virtual void GetFunctionsInternal(TArray<FNiagaraFunctionSignature>& OutFunctions) const override;
+#endif
 	NIAGARA_API void IsActive(FVectorVMExternalFunctionContext& Context);
 };
 
diff --git a/Engine/Plugins/FX/Niagara/Source/Niagara/Classes/NiagaraDataInterfaceRW.h b/Engine/Plugins/FX/Niagara/Source/Niagara/Classes/NiagaraDataInterfaceRW.h
--- a/Engine/Plugins/FX/Niagara/Source/Niagara/Classes/NiagaraDataInterfaceRW.h
+++ b/Engine/Plugins/FX/Niagara/Source/Niagara/Classes/NiagaraDataInterfaceRW.h
@@ -221,7 +221,6 @@ public:
 
 	//~ UNiagaraDataInterface interface
 	// VM functionality
-	NIAGARA_API virtual void GetFunctions(TArray<FNiagaraFunctionSignature>& OutFunctions) override;
 	NIAGARA_API virtual void GetVMExternalFunction(const FVMExternalFunctionBindingInfo& BindingInfo, void* InstanceData, FVMExternalFunction &OutFunc) override;	
 
 
@@ -236,6 +235,9 @@ public:
 
 protected:
 	//~ UNiagaraDataInterface interface
+#if WITH_EDITORONLY_DATA
+	virtual void GetFunctionsInternal(TArray<FNiagaraFunctionSignature>& OutFunctions) const override;
+#endif
 	NIAGARA_API virtual bool CopyToInternal(UNiagaraDataInterface* Destination) const override;
 	//~ UNiagaraDataInterface interface END	
 };
@@ -284,7 +286,6 @@ public:
 
 	//~ UNiagaraDataInterface interface
 	// VM functionality
-	NIAGARA_API virtual void GetFunctions(TArray<FNiagaraFunctionSignature>& OutFunctions) override;
 	NIAGARA_API virtual void GetVMExternalFunction(const FVMExternalFunctionBindingInfo& BindingInfo, void* InstanceData, FVMExternalFunction &OutFunc) override;
 
 	NIAGARA_API virtual bool Equals(const UNiagaraDataInterface* Other) const override;
@@ -302,6 +303,9 @@ public:
 
 
 protected:
+#if WITH_EDITORONLY_DATA
+	virtual void GetFunctionsInternal(TArray<FNiagaraFunctionSignature>& OutFunctions) const override;
+#endif
 	//~ UNiagaraDataInterface interface
 	NIAGARA_API virtual bool CopyToInternal(UNiagaraDataInterface* Destination) const override;
 	//~ UNiagaraDataInterface interface END	
diff --git a/Engine/Plugins/FX/Niagara/Source/Niagara/Classes/NiagaraDataInterfaceRasterizationGrid3D.h b/Engine/Plugins/FX/Niagara/Source/Niagara/Classes/NiagaraDataInterfaceRasterizationGrid3D.h
--- a/Engine/Plugins/FX/Niagara/Source/Niagara/Classes/NiagaraDataInterfaceRasterizationGrid3D.h
+++ b/Engine/Plugins/FX/Niagara/Source/Niagara/Classes/NiagaraDataInterfaceRasterizationGrid3D.h
@@ -84,7 +84,6 @@ public:
 
 	//~ UNiagaraDataInterface interface
 	// VM functionality
-	NIAGARA_API virtual void GetFunctions(TArray<FNiagaraFunctionSignature>& OutFunctions) override;
 	NIAGARA_API virtual void GetVMExternalFunction(const FVMExternalFunctionBindingInfo& BindingInfo, void* InstanceData, FVMExternalFunction &OutFunc) override;
 
 	NIAGARA_API virtual bool Equals(const UNiagaraDataInterface* Other) const override;
@@ -121,6 +120,9 @@ public:
 
 protected:
 	//~ UNiagaraDataInterface interface
+#if WITH_EDITORONLY_DATA
+	virtual void GetFunctionsInternal(TArray<FNiagaraFunctionSignature>& OutFunctions) const override;
+#endif
 	NIAGARA_API virtual bool CopyToInternal(UNiagaraDataInterface* Destination) const override;
 	//~ UNiagaraDataInterface interface END	
 };
diff --git a/Engine/Plugins/FX/Niagara/Source/Niagara/Classes/NiagaraDataInterfaceRenderTarget2D.h b/Engine/Plugins/FX/Niagara/Source/Niagara/Classes/NiagaraDataInterfaceRenderTarget2D.h
--- a/Engine/Plugins/FX/Niagara/Source/Niagara/Classes/NiagaraDataInterfaceRenderTarget2D.h
+++ b/Engine/Plugins/FX/Niagara/Source/Niagara/Classes/NiagaraDataInterfaceRenderTarget2D.h
@@ -84,7 +84,6 @@ public:
 	//~ UNiagaraDataInterface interface
 	// VM functionality
 	virtual bool CanExecuteOnTarget(ENiagaraSimTarget Target)const override { return true; }
-	NIAGARA_API virtual void GetFunctions(TArray<FNiagaraFunctionSignature>& OutFunctions) override;
 #if WITH_EDITORONLY_DATA
 	NIAGARA_API virtual bool UpgradeFunctionCall(FNiagaraFunctionSignature& FunctionSignature) override;
 #endif
@@ -170,6 +169,10 @@ public:
 	FNiagaraUserParameterBinding RenderTargetUserParameter;
 
 protected:
+#if WITH_EDITORONLY_DATA
+	virtual void GetFunctionsInternal(TArray<FNiagaraFunctionSignature>& OutFunctions) const override;
+#endif
+
 	static NIAGARA_API FNiagaraVariableBase ExposedRTVar;
 
 	TMap<FNiagaraSystemInstanceID, FRenderTarget2DRWInstanceData_GameThread*> SystemInstancesToProxyData_GT;
diff --git a/Engine/Plugins/FX/Niagara/Source/Niagara/Classes/NiagaraDataInterfaceRenderTarget2DArray.h b/Engine/Plugins/FX/Niagara/Source/Niagara/Classes/NiagaraDataInterfaceRenderTarget2DArray.h
--- a/Engine/Plugins/FX/Niagara/Source/Niagara/Classes/NiagaraDataInterfaceRenderTarget2DArray.h
+++ b/Engine/Plugins/FX/Niagara/Source/Niagara/Classes/NiagaraDataInterfaceRenderTarget2DArray.h
@@ -79,7 +79,6 @@ public:
 	//~ UNiagaraDataInterface interface
 	// VM functionality
 	virtual bool CanExecuteOnTarget(ENiagaraSimTarget Target)const override { return true; }
-	NIAGARA_API virtual void GetFunctions(TArray<FNiagaraFunctionSignature>& OutFunctions) override;
 #if WITH_EDITORONLY_DATA
 	NIAGARA_API virtual bool UpgradeFunctionCall(FNiagaraFunctionSignature& FunctionSignature) override;
 #endif
@@ -147,5 +146,8 @@ public:
 	FNiagaraUserParameterBinding RenderTargetUserParameter;
 
 protected:
+#if WITH_EDITORONLY_DATA
+	virtual void GetFunctionsInternal(TArray<FNiagaraFunctionSignature>& OutFunctions) const override;
+#endif
 	static NIAGARA_API FNiagaraVariableBase ExposedRTVar;
 };
diff --git a/Engine/Plugins/FX/Niagara/Source/Niagara/Classes/NiagaraDataInterfaceRenderTargetCube.h b/Engine/Plugins/FX/Niagara/Source/Niagara/Classes/NiagaraDataInterfaceRenderTargetCube.h
--- a/Engine/Plugins/FX/Niagara/Source/Niagara/Classes/NiagaraDataInterfaceRenderTargetCube.h
+++ b/Engine/Plugins/FX/Niagara/Source/Niagara/Classes/NiagaraDataInterfaceRenderTargetCube.h
@@ -72,7 +72,6 @@ public:
 	//~ UNiagaraDataInterface interface
 	// VM functionality
 	virtual bool CanExecuteOnTarget(ENiagaraSimTarget Target)const override { return true; }
-	NIAGARA_API virtual void GetFunctions(TArray<FNiagaraFunctionSignature>& OutFunctions) override;
 #if WITH_EDITORONLY_DATA
 	NIAGARA_API virtual bool UpgradeFunctionCall(FNiagaraFunctionSignature& FunctionSignature) override;
 #endif
@@ -140,7 +139,9 @@ public:
 	FNiagaraUserParameterBinding RenderTargetUserParameter;
 
 protected:
-	//~ UNiagaraDataInterface interface END
+#if WITH_EDITORONLY_DATA
+	virtual void GetFunctionsInternal(TArray<FNiagaraFunctionSignature>& OutFunctions) const override;
+#endif
 
 	static NIAGARA_API FNiagaraVariableBase ExposedRTVar;
 };
diff --git a/Engine/Plugins/FX/Niagara/Source/Niagara/Classes/NiagaraDataInterfaceRenderTargetVolume.h b/Engine/Plugins/FX/Niagara/Source/Niagara/Classes/NiagaraDataInterfaceRenderTargetVolume.h
--- a/Engine/Plugins/FX/Niagara/Source/Niagara/Classes/NiagaraDataInterfaceRenderTargetVolume.h
+++ b/Engine/Plugins/FX/Niagara/Source/Niagara/Classes/NiagaraDataInterfaceRenderTargetVolume.h
@@ -72,7 +72,6 @@ public:
 	//~ UNiagaraDataInterface interface
 	// VM functionality
 	virtual bool CanExecuteOnTarget(ENiagaraSimTarget Target)const override { return true; }
-	NIAGARA_API virtual void GetFunctions(TArray<FNiagaraFunctionSignature>& OutFunctions) override;
 #if WITH_EDITORONLY_DATA
 	NIAGARA_API virtual bool UpgradeFunctionCall(FNiagaraFunctionSignature& FunctionSignature) override;
 #endif
@@ -154,7 +153,9 @@ public:
 	FNiagaraUserParameterBinding RenderTargetUserParameter;
 
 protected:
-	//~ UNiagaraDataInterface interface END
+#if WITH_EDITORONLY_DATA
+	virtual void GetFunctionsInternal(TArray<FNiagaraFunctionSignature>& OutFunctions) const override;
+#endif
 
 	static NIAGARA_API FNiagaraVariableBase ExposedRTVar;
 };
diff --git a/Engine/Plugins/FX/Niagara/Source/Niagara/Classes/NiagaraDataInterfaceSkeletalMesh.h b/Engine/Plugins/FX/Niagara/Source/Niagara/Classes/NiagaraDataInterfaceSkeletalMesh.h
--- a/Engine/Plugins/FX/Niagara/Source/Niagara/Classes/NiagaraDataInterfaceSkeletalMesh.h
+++ b/Engine/Plugins/FX/Niagara/Source/Niagara/Classes/NiagaraDataInterfaceSkeletalMesh.h
@@ -795,7 +795,6 @@ public:
 	virtual int32 PerInstanceDataSize()const override { return sizeof(FNDISkeletalMesh_InstanceData); }
 	virtual bool HasPreSimulateTick() const override { return true; }
 
-	NIAGARA_API virtual void GetFunctions(TArray<FNiagaraFunctionSignature>& OutFunctions)override;
 	NIAGARA_API virtual void GetVMExternalFunction(const FVMExternalFunctionBindingInfo& BindingInfo, void* InstanceData, FVMExternalFunction &OutFunc)override;
 	NIAGARA_API virtual bool Equals(const UNiagaraDataInterface* Other) const override;
 	virtual bool CanExecuteOnTarget(ENiagaraSimTarget Target)const override { return true; }
@@ -837,6 +836,10 @@ public:
 	NIAGARA_API virtual void ProvidePerInstanceDataForRenderThread(void* DataForRenderThread, void* PerInstanceData, const FNiagaraSystemInstanceID& SystemInstance) override;
 
 protected:
+#if WITH_EDITORONLY_DATA
+	NIAGARA_API virtual void GetFunctionsInternal(TArray<FNiagaraFunctionSignature>& OutFunctions) const override;
+#endif
+
 	NIAGARA_API virtual bool CopyToInternal(UNiagaraDataInterface* Destination) const override;
 
 	// Bind/unbind delegates to release references to the source actor & component.
@@ -854,8 +857,9 @@ protected:
 	//Triangle sampling
 	//Triangles are sampled a using MeshTriangleCoordinates which are composed of Triangle index and a bary centric coordinate on that triangle.
 public:
-
-	NIAGARA_API void GetTriangleSamplingFunctions(TArray<FNiagaraFunctionSignature>& OutFunctions);
+#if WITH_EDITORONLY_DATA
+	NIAGARA_API void GetTriangleSamplingFunctions(TArray<FNiagaraFunctionSignature>& OutFunctions) const;
+#endif
 	NIAGARA_API void BindTriangleSamplingFunction(const FVMExternalFunctionBindingInfo& BindingInfo, FNDISkeletalMesh_InstanceData* InstData, FVMExternalFunction &OutFunc);
 
 	template<typename FilterMode, typename AreaWeightingMode>
@@ -915,7 +919,9 @@ private:
 	//Vertex sampling done with direct vertex indices.
 public:
 
-	NIAGARA_API void GetVertexSamplingFunctions(TArray<FNiagaraFunctionSignature>& OutFunctions);
+#if WITH_EDITORONLY_DATA
+	NIAGARA_API void GetVertexSamplingFunctions(TArray<FNiagaraFunctionSignature>& OutFunctions) const;
+#endif
 	NIAGARA_API void BindVertexSamplingFunction(const FVMExternalFunctionBindingInfo& BindingInfo, FNDISkeletalMesh_InstanceData* InstData, FVMExternalFunction &OutFunc);
 
 	NIAGARA_API void IsValidVertex(FVectorVMExternalFunctionContext& Context);
@@ -961,7 +967,9 @@ private:
 	// Direct Bone + Socket Sampling
 
 public:
-	NIAGARA_API void GetSkeletonSamplingFunctions(TArray<FNiagaraFunctionSignature>& OutFunctions);
+#if WITH_EDITORONLY_DATA
+	NIAGARA_API void GetSkeletonSamplingFunctions(TArray<FNiagaraFunctionSignature>& OutFunctions) const;
+#endif
 	NIAGARA_API void BindSkeletonSamplingFunction(const FVMExternalFunctionBindingInfo& BindingInfo, FNDISkeletalMesh_InstanceData* InstData, FVMExternalFunction &OutFunc);
 
 	template<typename SkinningHandlerType, typename TransformHandlerType, typename bInterpolated>
diff --git a/Engine/Plugins/FX/Niagara/Source/Niagara/Classes/NiagaraDataInterfaceSparseVolumeTexture.h b/Engine/Plugins/FX/Niagara/Source/Niagara/Classes/NiagaraDataInterfaceSparseVolumeTexture.h
--- a/Engine/Plugins/FX/Niagara/Source/Niagara/Classes/NiagaraDataInterfaceSparseVolumeTexture.h
+++ b/Engine/Plugins/FX/Niagara/Source/Niagara/Classes/NiagaraDataInterfaceSparseVolumeTexture.h
@@ -41,7 +41,6 @@ public:
 	//UObject Interface End
 
 	//UNiagaraDataInterface Interface
-	NIAGARA_API virtual void GetFunctions(TArray<FNiagaraFunctionSignature>& OutFunctions) override;
 	NIAGARA_API virtual void GetVMExternalFunction(const FVMExternalFunctionBindingInfo& BindingInfo, void* InstanceData, FVMExternalFunction& OutFunc) override;
 	virtual bool CanExecuteOnTarget(ENiagaraSimTarget Target) const override { return Target == ENiagaraSimTarget::GPUComputeSim; }
 
@@ -71,6 +70,9 @@ public:
 	NIAGARA_API void SetTexture(USparseVolumeTexture* InSparseVolumeTexture);
 
 protected:
+#if WITH_EDITORONLY_DATA
+	NIAGARA_API virtual void GetFunctionsInternal(TArray<FNiagaraFunctionSignature>& OutFunctions) const override;
+#endif
 	NIAGARA_API virtual bool CopyToInternal(UNiagaraDataInterface* Destination) const override;
 
 protected:
diff --git a/Engine/Plugins/FX/Niagara/Source/Niagara/Classes/NiagaraDataInterfaceSpline.h b/Engine/Plugins/FX/Niagara/Source/Niagara/Classes/NiagaraDataInterfaceSpline.h
--- a/Engine/Plugins/FX/Niagara/Source/Niagara/Classes/NiagaraDataInterfaceSpline.h
+++ b/Engine/Plugins/FX/Niagara/Source/Niagara/Classes/NiagaraDataInterfaceSpline.h
@@ -191,7 +191,6 @@ public:
 	NIAGARA_API virtual void DestroyPerInstanceData(void* PerInstanceData, FNiagaraSystemInstance* SystemInstance)override;
 	NIAGARA_API virtual bool PerInstanceTick(void* PerInstanceData, FNiagaraSystemInstance* SystemInstance, float DeltaSeconds) override;
 	NIAGARA_API virtual int32 PerInstanceDataSize()const override;
-	NIAGARA_API virtual void GetFunctions(TArray<FNiagaraFunctionSignature>& OutFunctions)override;
 	NIAGARA_API virtual void GetVMExternalFunction(const FVMExternalFunctionBindingInfo& BindingInfo, void* InstanceData, FVMExternalFunction &OutFunc) override;
 	NIAGARA_API virtual bool Equals(const UNiagaraDataInterface* Other) const override;
 	virtual bool CanExecuteOnTarget(ENiagaraSimTarget Target)const override { return true; }
@@ -224,6 +223,9 @@ public:
 	}
 
 protected:
+#if WITH_EDITORONLY_DATA
+	NIAGARA_API virtual void GetFunctionsInternal(TArray<FNiagaraFunctionSignature>& OutFunctions) const override;
+#endif
 	NIAGARA_API virtual bool CopyToInternal(UNiagaraDataInterface* Destination) const override;
 
 #if WITH_EDITORONLY_DATA
diff --git a/Engine/Plugins/FX/Niagara/Source/Niagara/Classes/NiagaraDataInterfaceSpriteRendererInfo.h b/Engine/Plugins/FX/Niagara/Source/Niagara/Classes/NiagaraDataInterfaceSpriteRendererInfo.h
--- a/Engine/Plugins/FX/Niagara/Source/Niagara/Classes/NiagaraDataInterfaceSpriteRendererInfo.h
+++ b/Engine/Plugins/FX/Niagara/Source/Niagara/Classes/NiagaraDataInterfaceSpriteRendererInfo.h
@@ -31,7 +31,6 @@ public:
 	NIAGARA_API void OnSpriteRendererChanged(UNiagaraSpriteRendererProperties* NewSpriteRenderer);
 
 	//UNiagaraDataInterface Interface
-	NIAGARA_API virtual void GetFunctions(TArray<FNiagaraFunctionSignature>& OutFunctions) override;
 	NIAGARA_API virtual void GetVMExternalFunction(const FVMExternalFunctionBindingInfo& BindingInfo, void* InstanceData, FVMExternalFunction& OutFunc) override;
 	NIAGARA_API virtual bool Equals(const UNiagaraDataInterface* Other) const override;
 #if WITH_EDITORONLY_DATA
@@ -49,6 +48,9 @@ public:
 	//UNiagaraDataInterface Interface	
 
 protected:
+#if WITH_EDITORONLY_DATA
+	NIAGARA_API virtual void GetFunctionsInternal(TArray<FNiagaraFunctionSignature>& OutFunctions) const override;
+#endif
 	NIAGARA_API virtual bool CopyToInternal(UNiagaraDataInterface* Destination) const override;
 	NIAGARA_API virtual void PushToRenderThreadImpl() override;
 
diff --git a/Engine/Plugins/FX/Niagara/Source/Niagara/Classes/NiagaraDataInterfaceTexture.h b/Engine/Plugins/FX/Niagara/Source/Niagara/Classes/NiagaraDataInterfaceTexture.h
--- a/Engine/Plugins/FX/Niagara/Source/Niagara/Classes/NiagaraDataInterfaceTexture.h
+++ b/Engine/Plugins/FX/Niagara/Source/Niagara/Classes/NiagaraDataInterfaceTexture.h
@@ -35,7 +35,6 @@ public:
 	//UObject Interface End
 
 	//UNiagaraDataInterface Interface
-	NIAGARA_API virtual void GetFunctions(TArray<FNiagaraFunctionSignature>& OutFunctions)override;
 	NIAGARA_API virtual void GetVMExternalFunction(const FVMExternalFunctionBindingInfo& BindingInfo, void* InstanceData, FVMExternalFunction &OutFunc) override;
 	virtual bool CanExecuteOnTarget(ENiagaraSimTarget Target)const override { return Target==ENiagaraSimTarget::GPUComputeSim; }
 
@@ -67,6 +66,9 @@ public:
 	NIAGARA_API void SetTexture(UTexture* InTexture);
 
 protected:
+#if WITH_EDITORONLY_DATA
+	NIAGARA_API virtual void GetFunctionsInternal(TArray<FNiagaraFunctionSignature>& OutFunctions) const override;
+#endif
 	NIAGARA_API virtual bool CopyToInternal(UNiagaraDataInterface* Destination) const override;
 
 protected:
diff --git a/Engine/Plugins/FX/Niagara/Source/Niagara/Classes/NiagaraDataInterfaceVector2DCurve.h b/Engine/Plugins/FX/Niagara/Source/Niagara/Classes/NiagaraDataInterfaceVector2DCurve.h
--- a/Engine/Plugins/FX/Niagara/Source/Niagara/Classes/NiagaraDataInterfaceVector2DCurve.h
+++ b/Engine/Plugins/FX/Niagara/Source/Niagara/Classes/NiagaraDataInterfaceVector2DCurve.h
@@ -47,9 +47,6 @@ public:
 
 	NIAGARA_API virtual bool Equals(const UNiagaraDataInterface* Other) const override;
 
-
-	NIAGARA_API virtual void GetFunctions(TArray<FNiagaraFunctionSignature>& OutFunctions)override;
-
 	NIAGARA_API virtual void GetVMExternalFunction(const FVMExternalFunctionBindingInfo& BindingInfo, void* InstanceData, FVMExternalFunction &OutFunc) override;
 
 	template<typename UseLUT>
@@ -68,6 +65,9 @@ public:
 #endif
 
 protected:
+#if WITH_EDITORONLY_DATA
+	virtual void GetFunctionsInternal(TArray<FNiagaraFunctionSignature>& OutFunctions) const override;
+#endif
 	NIAGARA_API virtual bool CopyToInternal(UNiagaraDataInterface* Destination) const override;
 
 private:
diff --git a/Engine/Plugins/FX/Niagara/Source/Niagara/Classes/NiagaraDataInterfaceVector4Curve.h b/Engine/Plugins/FX/Niagara/Source/Niagara/Classes/NiagaraDataInterfaceVector4Curve.h
--- a/Engine/Plugins/FX/Niagara/Source/Niagara/Classes/NiagaraDataInterfaceVector4Curve.h
+++ b/Engine/Plugins/FX/Niagara/Source/Niagara/Classes/NiagaraDataInterfaceVector4Curve.h
@@ -60,7 +60,6 @@ public:
 	NIAGARA_API virtual void UpdateTimeRanges() override;
 	NIAGARA_API virtual TArray<float> BuildLUT(int32 NumEntries) const override;
 
-	NIAGARA_API virtual void GetFunctions(TArray<FNiagaraFunctionSignature>& OutFunctions)override;
 	NIAGARA_API virtual void GetVMExternalFunction(const FVMExternalFunctionBindingInfo& BindingInfo, void* InstanceData, FVMExternalFunction &OutFunc) override;
 
 	template<typename UseLUT>
@@ -80,6 +79,9 @@ public:
 #endif
 
 protected:
+#if WITH_EDITORONLY_DATA
+	virtual void GetFunctionsInternal(TArray<FNiagaraFunctionSignature>& OutFunctions) const override;
+#endif
 	NIAGARA_API virtual bool CopyToInternal(UNiagaraDataInterface* Destination) const override;
 
 private: 
diff --git a/Engine/Plugins/FX/Niagara/Source/Niagara/Classes/NiagaraDataInterfaceVectorCurve.h b/Engine/Plugins/FX/Niagara/Source/Niagara/Classes/NiagaraDataInterfaceVectorCurve.h
--- a/Engine/Plugins/FX/Niagara/Source/Niagara/Classes/NiagaraDataInterfaceVectorCurve.h
+++ b/Engine/Plugins/FX/Niagara/Source/Niagara/Classes/NiagaraDataInterfaceVectorCurve.h
@@ -55,7 +55,6 @@ public:
 	NIAGARA_API virtual void UpdateTimeRanges() override;
 	NIAGARA_API virtual TArray<float> BuildLUT(int32 NumEntries) const override;
 
-	NIAGARA_API virtual void GetFunctions(TArray<FNiagaraFunctionSignature>& OutFunctions)override;
 	NIAGARA_API virtual void GetVMExternalFunction(const FVMExternalFunctionBindingInfo& BindingInfo, void* InstanceData, FVMExternalFunction &OutFunc) override;
 
 	template<typename UseLUT>
@@ -77,6 +76,9 @@ public:
 #endif
 
 protected:
+#if WITH_EDITORONLY_DATA
+	virtual void GetFunctionsInternal(TArray<FNiagaraFunctionSignature>& OutFunctions) const override;
+#endif
 	NIAGARA_API virtual bool CopyToInternal(UNiagaraDataInterface* Destination) const override;
 
 private:
diff --git a/Engine/Plugins/FX/Niagara/Source/Niagara/Classes/NiagaraDataInterfaceVectorField.h b/Engine/Plugins/FX/Niagara/Source/Niagara/Classes/NiagaraDataInterfaceVectorField.h
--- a/Engine/Plugins/FX/Niagara/Source/Niagara/Classes/NiagaraDataInterfaceVectorField.h
+++ b/Engine/Plugins/FX/Niagara/Source/Niagara/Classes/NiagaraDataInterfaceVectorField.h
@@ -45,7 +45,6 @@ public:
 
 	//~ UNiagaraDataInterface interface
 	// VM functionality
-	NIAGARA_API virtual void GetFunctions(TArray<FNiagaraFunctionSignature>& OutFunctions) override;
 	NIAGARA_API virtual void GetVMExternalFunction(const FVMExternalFunctionBindingInfo& BindingInfo, void* InstanceData, FVMExternalFunction &OutFunc) override;
 	NIAGARA_API virtual bool Equals(const UNiagaraDataInterface* Other) const override;
 	NIAGARA_API virtual bool CanExecuteOnTarget(ENiagaraSimTarget Target) const override;
@@ -80,6 +79,9 @@ public:
 	NIAGARA_API FVector GetMaxBounds() const;
 protected:
 	//~ UNiagaraDataInterface interface
+#if WITH_EDITORONLY_DATA
+	NIAGARA_API virtual void GetFunctionsInternal(TArray<FNiagaraFunctionSignature>& OutFunctions) const override;
+#endif
 	NIAGARA_API virtual bool CopyToInternal(UNiagaraDataInterface* Destination) const override;
 	//~ UNiagaraDataInterface interface END
 
diff --git a/Engine/Plugins/FX/Niagara/Source/Niagara/Classes/NiagaraDataInterfaceVolumeCache.h b/Engine/Plugins/FX/Niagara/Source/Niagara/Classes/NiagaraDataInterfaceVolumeCache.h
--- a/Engine/Plugins/FX/Niagara/Source/Niagara/Classes/NiagaraDataInterfaceVolumeCache.h
+++ b/Engine/Plugins/FX/Niagara/Source/Niagara/Classes/NiagaraDataInterfaceVolumeCache.h
@@ -36,7 +36,6 @@ public:
 	NIAGARA_API virtual bool Equals(const UNiagaraDataInterface* Other) const override;
 	NIAGARA_API virtual int32 PerInstanceDataSize()const override;
 	NIAGARA_API virtual bool InitPerInstanceData(void* PerInstanceData, FNiagaraSystemInstance* SystemInstance) override;	
-	NIAGARA_API virtual void GetFunctions(TArray<FNiagaraFunctionSignature>& OutFunctions) override;
 	NIAGARA_API virtual bool PerInstanceTick(void* PerInstanceData, FNiagaraSystemInstance* SystemInstance, float DeltaSeconds) override;
 	virtual bool HasPreSimulateTick() const override { return true; }
 	NIAGARA_API virtual void GetVMExternalFunction(const FVMExternalFunctionBindingInfo& BindingInfo, void* InstanceData, FVMExternalFunction& OutFunc) override;
@@ -61,6 +60,9 @@ public:
 
 protected:
 	//~ UNiagaraDataInterface interface
+#if WITH_EDITORONLY_DATA
+	NIAGARA_API virtual void GetFunctionsInternal(TArray<FNiagaraFunctionSignature>& OutFunctions) const override;
+#endif
 	NIAGARA_API virtual bool CopyToInternal(UNiagaraDataInterface* Destination) const override;
 	//~ UNiagaraDataInterface interface END
 
diff --git a/Engine/Plugins/FX/Niagara/Source/Niagara/Classes/NiagaraDataInterfaceVolumeTexture.h b/Engine/Plugins/FX/Niagara/Source/Niagara/Classes/NiagaraDataInterfaceVolumeTexture.h
--- a/Engine/Plugins/FX/Niagara/Source/Niagara/Classes/NiagaraDataInterfaceVolumeTexture.h
+++ b/Engine/Plugins/FX/Niagara/Source/Niagara/Classes/NiagaraDataInterfaceVolumeTexture.h
@@ -34,7 +34,6 @@ public:
 	//UObject Interface End
 
 	//UNiagaraDataInterface Interface
-	NIAGARA_API virtual void GetFunctions(TArray<FNiagaraFunctionSignature>& OutFunctions)override;
 	NIAGARA_API virtual void GetVMExternalFunction(const FVMExternalFunctionBindingInfo& BindingInfo, void* InstanceData, FVMExternalFunction &OutFunc) override;
 	virtual bool CanExecuteOnTarget(ENiagaraSimTarget Target)const override { return Target == ENiagaraSimTarget::GPUComputeSim; }
 
@@ -63,6 +62,9 @@ public:
 	NIAGARA_API void SetTexture(UTexture* InTexture);
 
 protected:
+#if WITH_EDITORONLY_DATA
+	NIAGARA_API virtual void GetFunctionsInternal(TArray<FNiagaraFunctionSignature>& OutFunctions) const override;
+#endif
 	NIAGARA_API virtual bool CopyToInternal(UNiagaraDataInterface* Destination) const override;
 
 protected:
diff --git a/Engine/Plugins/FX/Niagara/Source/Niagara/Internal/DataInterface/NiagaraDataInterfaceActorComponent.h b/Engine/Plugins/FX/Niagara/Source/Niagara/Internal/DataInterface/NiagaraDataInterfaceActorComponent.h
--- a/Engine/Plugins/FX/Niagara/Source/Niagara/Internal/DataInterface/NiagaraDataInterfaceActorComponent.h
+++ b/Engine/Plugins/FX/Niagara/Source/Niagara/Internal/DataInterface/NiagaraDataInterfaceActorComponent.h
@@ -74,7 +74,6 @@ public:
 	//UNiagaraDataInterface Interface
 	virtual bool CanExecuteOnTarget(ENiagaraSimTarget Target) const override { return true; }
 
-	NIAGARA_API virtual void GetFunctions(TArray<FNiagaraFunctionSignature>& OutFunctions) override;
 	NIAGARA_API virtual void GetVMExternalFunction(const FVMExternalFunctionBindingInfo& BindingInfo, void* InstanceData, FVMExternalFunction &OutFunc) override;
 #if WITH_EDITORONLY_DATA
 	NIAGARA_API virtual bool AppendCompileHash(FNiagaraCompileHashVisitor* InVisitor) const override;
@@ -101,6 +100,11 @@ public:
 	virtual bool HasPreSimulateTick() const override { return true; }
 	//UNiagaraDataInterface Interface
 
+protected:
+#if WITH_EDITORONLY_DATA
+	NIAGARA_API virtual void GetFunctionsInternal(TArray<FNiagaraFunctionSignature>& OutFunctions) const override;
+#endif
+
 private:
 	NIAGARA_API void VMGetMatrix(FVectorVMExternalFunctionContext& Context);
 	NIAGARA_API void VMGetTransform(FVectorVMExternalFunctionContext& Context);
diff --git a/Engine/Plugins/FX/Niagara/Source/Niagara/Internal/DataInterface/NiagaraDataInterfaceDataChannelRead.h b/Engine/Plugins/FX/Niagara/Source/Niagara/Internal/DataInterface/NiagaraDataInterfaceDataChannelRead.h
--- a/Engine/Plugins/FX/Niagara/Source/Niagara/Internal/DataInterface/NiagaraDataInterfaceDataChannelRead.h
+++ b/Engine/Plugins/FX/Niagara/Source/Niagara/Internal/DataInterface/NiagaraDataInterfaceDataChannelRead.h
@@ -100,7 +100,6 @@ public:
 	//UObject Interface End
 
 	//UNiagaraDataInterface Interface
-	NIAGARA_API virtual void GetFunctions(TArray<FNiagaraFunctionSignature>& OutFunctions) override;
 	NIAGARA_API virtual void GetVMExternalFunction(const FVMExternalFunctionBindingInfo& BindingInfo, void* InstanceData, FVMExternalFunction &OutFunc) override;
 	virtual bool CanExecuteOnTarget(ENiagaraSimTarget Target) const override { return true; }
 #if WITH_EDITORONLY_DATA
@@ -154,6 +153,9 @@ public:
 	FNDIDataChannelCompiledData& GetCompiledData() { return CompiledData; }
 
 protected:
+#if WITH_EDITORONLY_DATA
+	NIAGARA_API virtual void GetFunctionsInternal(TArray<FNiagaraFunctionSignature>& OutFunctions) const override;
+#endif
 	NIAGARA_API virtual bool CopyToInternal(UNiagaraDataInterface* Destination) const override;
 	NIAGARA_API UNiagaraDataInterfaceDataChannelWrite* FindSourceDI()const;
 
diff --git a/Engine/Plugins/FX/Niagara/Source/Niagara/Internal/DataInterface/NiagaraDataInterfaceDataChannelWrite.h b/Engine/Plugins/FX/Niagara/Source/Niagara/Internal/DataInterface/NiagaraDataInterfaceDataChannelWrite.h
--- a/Engine/Plugins/FX/Niagara/Source/Niagara/Internal/DataInterface/NiagaraDataInterfaceDataChannelWrite.h
+++ b/Engine/Plugins/FX/Niagara/Source/Niagara/Internal/DataInterface/NiagaraDataInterfaceDataChannelWrite.h
@@ -132,7 +132,6 @@ public:
 	//UObject Interface End
 
 	//UNiagaraDataInterface Interface
-	NIAGARA_API virtual void GetFunctions(TArray<FNiagaraFunctionSignature>& OutFunctions) override;
 	NIAGARA_API virtual void GetVMExternalFunction(const FVMExternalFunctionBindingInfo& BindingInfo, void* InstanceData, FVMExternalFunction& OutFunc) override;
 	virtual bool CanExecuteOnTarget(ENiagaraSimTarget Target) const override { return true; }
 #if WITH_EDITORONLY_DATA
@@ -188,6 +187,9 @@ public:
 	virtual bool SimCacheReadFrame(UObject* StorageObject, int FrameA, int FrameB, float Interp, FNiagaraSystemInstance* SystemInstance, void* OptionalPerInstanceData) override;
 	virtual void SimCachePostReadFrame(void* OptionalPerInstanceData, FNiagaraSystemInstance* SystemInstance) override;
 protected:
+#if WITH_EDITORONLY_DATA
+	NIAGARA_API virtual void GetFunctionsInternal(TArray<FNiagaraFunctionSignature>& OutFunctions) const override;
+#endif
 	NIAGARA_API virtual bool CopyToInternal(UNiagaraDataInterface* Destination) const override;
 
 	UPROPERTY()
diff --git a/Engine/Plugins/FX/Niagara/Source/Niagara/Internal/DataInterface/NiagaraDataInterfaceStaticMesh.h b/Engine/Plugins/FX/Niagara/Source/Niagara/Internal/DataInterface/NiagaraDataInterfaceStaticMesh.h
--- a/Engine/Plugins/FX/Niagara/Source/Niagara/Internal/DataInterface/NiagaraDataInterfaceStaticMesh.h
+++ b/Engine/Plugins/FX/Niagara/Source/Niagara/Internal/DataInterface/NiagaraDataInterfaceStaticMesh.h
@@ -145,15 +145,6 @@ public:
 
 	NIAGARA_API virtual void ProvidePerInstanceDataForRenderThread(void* DataForRenderThread, void* PerInstanceData, const FNiagaraSystemInstanceID& SystemInstance) override;
 
-	NIAGARA_API virtual void GetFunctions(TArray<FNiagaraFunctionSignature>& OutFunctions) override;
-	NIAGARA_API void GetVertexSamplingFunctions(TArray<FNiagaraFunctionSignature>& OutFunctions, const FNiagaraFunctionSignature& BaseSignature) const;
-	NIAGARA_API void GetTriangleSamplingFunctions(TArray<FNiagaraFunctionSignature>& OutFunctions, const FNiagaraFunctionSignature& BaseSignature) const;
-	NIAGARA_API void GetSocketSamplingFunctions(TArray<FNiagaraFunctionSignature>& OutFunctions, const FNiagaraFunctionSignature& BaseSignature) const;
-	NIAGARA_API void GetSectionFunctions(TArray<FNiagaraFunctionSignature>& OutFunctions, const FNiagaraFunctionSignature& BaseSignature) const;
-	NIAGARA_API void GetMiscFunctions(TArray<FNiagaraFunctionSignature>& OutFunctions, const FNiagaraFunctionSignature& BaseSignature) const;
-	NIAGARA_API void GetUVMappingFunctions(TArray<FNiagaraFunctionSignature>& OutFunctions, const FNiagaraFunctionSignature& BaseSignature) const;
-	NIAGARA_API void GetDistanceFieldFunctions(TArray<FNiagaraFunctionSignature>& OutFunctions, const FNiagaraFunctionSignature& BaseSignature) const;
-	NIAGARA_API void GetDeprecatedFunctions(TArray<FNiagaraFunctionSignature>& OutFunctions, const FNiagaraFunctionSignature& BaseSignature) const;
 	static NIAGARA_API bool FunctionNeedsCpuAccess(FName InName);
 
 #if WITH_EDITORONLY_DATA
@@ -176,6 +167,17 @@ public:
 
 	NIAGARA_API virtual bool Equals(const UNiagaraDataInterface* Other) const override;
 protected:
+#if WITH_EDITORONLY_DATA
+	NIAGARA_API virtual void GetFunctionsInternal(TArray<FNiagaraFunctionSignature>& OutFunctions) const override;
+	NIAGARA_API void GetVertexSamplingFunctions(TArray<FNiagaraFunctionSignature>& OutFunctions, const FNiagaraFunctionSignature& BaseSignature) const;
+	NIAGARA_API void GetTriangleSamplingFunctions(TArray<FNiagaraFunctionSignature>& OutFunctions, const FNiagaraFunctionSignature& BaseSignature) const;
+	NIAGARA_API void GetSocketSamplingFunctions(TArray<FNiagaraFunctionSignature>& OutFunctions, const FNiagaraFunctionSignature& BaseSignature) const;
+	NIAGARA_API void GetSectionFunctions(TArray<FNiagaraFunctionSignature>& OutFunctions, const FNiagaraFunctionSignature& BaseSignature) const;
+	NIAGARA_API void GetMiscFunctions(TArray<FNiagaraFunctionSignature>& OutFunctions, const FNiagaraFunctionSignature& BaseSignature) const;
+	NIAGARA_API void GetUVMappingFunctions(TArray<FNiagaraFunctionSignature>& OutFunctions, const FNiagaraFunctionSignature& BaseSignature) const;
+	NIAGARA_API void GetDistanceFieldFunctions(TArray<FNiagaraFunctionSignature>& OutFunctions, const FNiagaraFunctionSignature& BaseSignature) const;
+	NIAGARA_API void GetDeprecatedFunctions(TArray<FNiagaraFunctionSignature>& OutFunctions, const FNiagaraFunctionSignature& BaseSignature) const;
+#endif
 	NIAGARA_API virtual bool CopyToInternal(UNiagaraDataInterface* Destination) const override;
 public:
 #if WITH_NIAGARA_DEBUGGER
diff --git a/Engine/Plugins/FX/Niagara/Source/Niagara/Private/DataInterface/NiagaraDataInterfaceActorComponent.cpp b/Engine/Plugins/FX/Niagara/Source/Niagara/Private/DataInterface/NiagaraDataInterfaceActorComponent.cpp
--- a/Engine/Plugins/FX/Niagara/Source/Niagara/Private/DataInterface/NiagaraDataInterfaceActorComponent.cpp
+++ b/Engine/Plugins/FX/Niagara/Source/Niagara/Private/DataInterface/NiagaraDataInterfaceActorComponent.cpp
@@ -186,7 +186,8 @@ UActorComponent* UNiagaraDataInterfaceActorComponent::ResolveComponent(FNiagaraS
 	return SourceActor.IsValid() ? SourceActor->GetRootComponent() : nullptr;
 }
 
-void UNiagaraDataInterfaceActorComponent::GetFunctions(TArray<FNiagaraFunctionSignature>& OutFunctions)
+#if WITH_EDITORONLY_DATA
+void UNiagaraDataInterfaceActorComponent::GetFunctionsInternal(TArray<FNiagaraFunctionSignature>& OutFunctions) const
 {
 	using namespace NDIActorComponentLocal;
 	FNiagaraFunctionSignature DefaultSig;
@@ -219,6 +220,7 @@ void UNiagaraDataInterfaceActorComponent::GetFunctions(TArray<FNiagaraFunctionSi
 		FunctionSignature.SetDescription(LOCTEXT("GetVelocityDesc", "Returns the velocity from the component or actor if valid."));
 	}
 }
+#endif
 
 void UNiagaraDataInterfaceActorComponent::GetVMExternalFunction(const FVMExternalFunctionBindingInfo& BindingInfo, void* InstanceData, FVMExternalFunction& OutFunc)
 {
@@ -267,7 +269,7 @@ bool UNiagaraDataInterfaceActorComponent::UpgradeFunctionCall(FNiagaraFunctionSi
 	if (FunctionSignature.FunctionVersion < FNiagaraActorDIFunctionVersion::LWCConversion)
 	{
 		TArray<FNiagaraFunctionSignature> AllFunctions;
-		GetFunctions(AllFunctions);
+		GetFunctionsInternal(AllFunctions);
 		for (const FNiagaraFunctionSignature& Sig : AllFunctions)
 		{
 			if (FunctionSignature.Name == Sig.Name)
diff --git a/Engine/Plugins/FX/Niagara/Source/Niagara/Private/DataInterface/NiagaraDataInterfaceAsyncGpuTrace.cpp b/Engine/Plugins/FX/Niagara/Source/Niagara/Private/DataInterface/NiagaraDataInterfaceAsyncGpuTrace.cpp
--- a/Engine/Plugins/FX/Niagara/Source/Niagara/Private/DataInterface/NiagaraDataInterfaceAsyncGpuTrace.cpp
+++ b/Engine/Plugins/FX/Niagara/Source/Niagara/Private/DataInterface/NiagaraDataInterfaceAsyncGpuTrace.cpp
@@ -224,7 +224,8 @@ void UNiagaraDataInterfaceAsyncGpuTrace::PostLoad()
 	MarkRenderDataDirty();
 }
 
-void UNiagaraDataInterfaceAsyncGpuTrace::GetFunctions(TArray<FNiagaraFunctionSignature>& OutFunctions)
+#if WITH_EDITORONLY_DATA
+void UNiagaraDataInterfaceAsyncGpuTrace::GetFunctionsInternal(TArray<FNiagaraFunctionSignature>& OutFunctions) const
 {
 	const FText RayTraceStartWorldDescription = LOCTEXT("TraceStartWorldDescription", "Ray starting point in world space");
 	const FText RayTraceEndWorldDescription = LOCTEXT("TraceEndWorldDescription", "Ray end point in world space");
@@ -269,11 +270,9 @@ void UNiagaraDataInterfaceAsyncGpuTrace::GetFunctions(TArray<FNiagaraFunctionSig
 		CreateRayTrace.bMemberFunction = true;
 		CreateRayTrace.bSupportsCPU = false;
 		CreateRayTrace.bExperimental = true;
-#if WITH_EDITORONLY_DATA
 		CreateRayTrace.FunctionVersion = FNiagaraAsyncGpuTraceDIFunctionVersion::LatestVersion;
 		CreateRayTrace.Description = LOCTEXT("CreateAsyncRayTraceDescription", "Creates a GPU raytrace with the result being available the following frame (index is returned)");
 		CreateRayTrace.ExperimentalMessage = ExperimentalMessage;
-#endif
 		CreateRayTrace.AddInput(FNiagaraVariable(FNiagaraTypeDefinition(GetClass()), TEXT("AsyncGpuTrace")));
 		CreateRayTrace.AddInput(FNiagaraVariable(FNiagaraTypeDefinition::GetPositionDef(), TEXT("TraceStartWorld")), RayTraceStartWorldDescription);
 		CreateRayTrace.AddInput(FNiagaraVariable(FNiagaraTypeDefinition::GetPositionDef(), TEXT("TraceEndWorld")), RayTraceEndWorldDescription);
@@ -293,11 +292,9 @@ void UNiagaraDataInterfaceAsyncGpuTrace::GetFunctions(TArray<FNiagaraFunctionSig
 		ReserveRayTrace.bMemberFunction = true;
 		ReserveRayTrace.bSupportsCPU = false;
 		ReserveRayTrace.bExperimental = true;
-#if WITH_EDITORONLY_DATA
 		ReserveRayTrace.FunctionVersion = FNiagaraAsyncGpuTraceDIFunctionVersion::LatestVersion;
 		ReserveRayTrace.Description = LOCTEXT("ReserveAsyncRayTraceDescription", "Reserves a number of ray trace request slots");
 		ReserveRayTrace.ExperimentalMessage = ExperimentalMessage;
-#endif
 		ReserveRayTrace.AddInput(FNiagaraVariable(FNiagaraTypeDefinition(GetClass()), TEXT("AsyncGpuTrace")));
 		ReserveRayTrace.AddInput(FNiagaraVariable(FNiagaraTypeDefinition::GetIntDef(), TEXT("TraceCount")), TraceCountDescription);
 		ReserveRayTrace.AddOutput(FNiagaraVariable(FNiagaraTypeDefinition::GetIntDef(), TEXT("FirstQueryID")), FirstQueryIDValidDescription);
@@ -313,11 +310,9 @@ void UNiagaraDataInterfaceAsyncGpuTrace::GetFunctions(TArray<FNiagaraFunctionSig
 		ReadRayTrace.bMemberFunction = true;
 		ReadRayTrace.bSupportsCPU = false;
 		ReadRayTrace.bExperimental = true;
-#if WITH_EDITORONLY_DATA
 		ReadRayTrace.FunctionVersion = FNiagaraAsyncGpuTraceDIFunctionVersion::LatestVersion;
 		ReadRayTrace.Description = LOCTEXT("ReadAsyncRayTraceDescription", "Reads the results of a previously enqueued GPU ray trace");
 		ReadRayTrace.ExperimentalMessage = ExperimentalMessage;
-#endif
 		ReadRayTrace.AddInput(FNiagaraVariable(FNiagaraTypeDefinition(GetClass()), TEXT("AsyncGpuTrace")));
 		ReadRayTrace.AddInput(FNiagaraVariable(FNiagaraTypeDefinition::GetIntDef(), TEXT("PreviousFrameQueryID")), PreviousFrameQueryIDDescription);
 		ReadRayTrace.AddOutput(FNiagaraVariable(FNiagaraTypeDefinition::GetBoolDef(), TEXT("CollisionValid")), CollisionValidDescription);
@@ -326,6 +321,7 @@ void UNiagaraDataInterfaceAsyncGpuTrace::GetFunctions(TArray<FNiagaraFunctionSig
 		ReadRayTrace.AddOutput(FNiagaraVariable(FNiagaraTypeDefinition::GetVec3Def(), TEXT("CollisionNormal")), CollisionNormalDescription);
 	}
 }
+#endif
 
 bool UNiagaraDataInterfaceAsyncGpuTrace::RequiresGlobalDistanceField() const
 {
@@ -357,7 +353,7 @@ bool UNiagaraDataInterfaceAsyncGpuTrace::UpgradeFunctionCall(FNiagaraFunctionSig
 	if (FunctionSignature.FunctionVersion < FNiagaraAsyncGpuTraceDIFunctionVersion::LatestVersion) // -V547
 	{
 		TArray<FNiagaraFunctionSignature> AllFunctions;
-		GetFunctions(AllFunctions);
+		GetFunctionsInternal(AllFunctions);
 		for (const FNiagaraFunctionSignature& Sig : AllFunctions)
 		{
 			if (FunctionSignature.Name == Sig.Name)
diff --git a/Engine/Plugins/FX/Niagara/Source/Niagara/Private/DataInterface/NiagaraDataInterfaceConsoleVariable.cpp b/Engine/Plugins/FX/Niagara/Source/Niagara/Private/DataInterface/NiagaraDataInterfaceConsoleVariable.cpp
--- a/Engine/Plugins/FX/Niagara/Source/Niagara/Private/DataInterface/NiagaraDataInterfaceConsoleVariable.cpp
+++ b/Engine/Plugins/FX/Niagara/Source/Niagara/Private/DataInterface/NiagaraDataInterfaceConsoleVariable.cpp
@@ -265,7 +265,8 @@ bool UNiagaraDataInterfaceConsoleVariable::PerInstanceTick(void* PerInstanceData
 	return false;
 }
 
-void UNiagaraDataInterfaceConsoleVariable::GetFunctions(TArray<FNiagaraFunctionSignature>& OutFunctions)
+#if WITH_EDITORONLY_DATA
+void UNiagaraDataInterfaceConsoleVariable::GetFunctionsInternal(TArray<FNiagaraFunctionSignature>& OutFunctions) const
 {
 	using namespace NDIConsoleVariableLocal;
 
@@ -295,6 +296,7 @@ void UNiagaraDataInterfaceConsoleVariable::GetFunctions(TArray<FNiagaraFunctionS
 		Sig.Outputs.Emplace(FNiagaraTypeDefinition::GetBoolDef(), TEXT("Value"));
 	}
 }
+#endif
 
 void UNiagaraDataInterfaceConsoleVariable::GetVMExternalFunction(const FVMExternalFunctionBindingInfo& BindingInfo, void* PerInstanceData, FVMExternalFunction& OutFunc)
 {
diff --git a/Engine/Plugins/FX/Niagara/Source/Niagara/Private/DataInterface/NiagaraDataInterfaceConsoleVariable.h b/Engine/Plugins/FX/Niagara/Source/Niagara/Private/DataInterface/NiagaraDataInterfaceConsoleVariable.h
--- a/Engine/Plugins/FX/Niagara/Source/Niagara/Private/DataInterface/NiagaraDataInterfaceConsoleVariable.h
+++ b/Engine/Plugins/FX/Niagara/Source/Niagara/Private/DataInterface/NiagaraDataInterfaceConsoleVariable.h
@@ -35,7 +35,6 @@ public:
 	virtual bool HasPreSimulateTick() const override { return true; }
 	virtual bool PerInstanceTick(void* PerInstanceData, FNiagaraSystemInstance* SystemInstance, float DeltaSeconds) override;
 
-	virtual void GetFunctions(TArray<FNiagaraFunctionSignature>& OutFunctions) override;
 	virtual void GetVMExternalFunction(const FVMExternalFunctionBindingInfo& BindingInfo, void* InstanceData, FVMExternalFunction& OutFunc) override;
 #if WITH_EDITORONLY_DATA
 	virtual bool AppendCompileHash(FNiagaraCompileHashVisitor* InVisitor) const override;
@@ -45,4 +44,10 @@ public:
 	virtual void BuildShaderParameters(FNiagaraShaderParametersBuilder& ShaderParametersBuilder) const override;
 	virtual void SetShaderParameters(const FNiagaraDataInterfaceSetShaderParametersContext& Context) const override;
 	//UNiagaraDataInterface Interface
+
+protected:
+#if WITH_EDITORONLY_DATA
+	virtual void GetFunctionsInternal(TArray<FNiagaraFunctionSignature>& OutFunctions) const override;
+#endif
+
 };
diff --git a/Engine/Plugins/FX/Niagara/Source/Niagara/Private/DataInterface/NiagaraDataInterfaceDataChannelRead.cpp b/Engine/Plugins/FX/Niagara/Source/Niagara/Private/DataInterface/NiagaraDataInterfaceDataChannelRead.cpp
--- a/Engine/Plugins/FX/Niagara/Source/Niagara/Private/DataInterface/NiagaraDataInterfaceDataChannelRead.cpp
+++ b/Engine/Plugins/FX/Niagara/Source/Niagara/Private/DataInterface/NiagaraDataInterfaceDataChannelRead.cpp
@@ -900,13 +900,15 @@ bool UNiagaraDataInterfaceDataChannelRead::CopyToInternal(UNiagaraDataInterface*
 	return false;
 }
 
-void UNiagaraDataInterfaceDataChannelRead::GetFunctions(TArray<FNiagaraFunctionSignature>& OutFunctions)
+#if WITH_EDITORONLY_DATA
+void UNiagaraDataInterfaceDataChannelRead::GetFunctionsInternal(TArray<FNiagaraFunctionSignature>& OutFunctions) const
 {
 	OutFunctions.Add(NDIDataChannelReadLocal::GetFunctionSig_Num());
 	OutFunctions.Add(NDIDataChannelReadLocal::GetFunctionSig_Read());
 	OutFunctions.Add(NDIDataChannelReadLocal::GetFunctionSig_Consume());
 	OutFunctions.Add(NDIDataChannelReadLocal::GetFunctionSig_SpawnConditional());
 }
+#endif
 
 void UNiagaraDataInterfaceDataChannelRead::GetVMExternalFunction(const FVMExternalFunctionBindingInfo& BindingInfo, void* InstanceData, FVMExternalFunction& OutFunc)
 {
@@ -1502,7 +1504,7 @@ void UNiagaraDataInterfaceDataChannelRead::GetParameterDefinitionHLSL(FNiagaraDa
 bool UNiagaraDataInterfaceDataChannelRead::UpgradeFunctionCall(FNiagaraFunctionSignature& FunctionSignature)
 {
 	TArray<FNiagaraFunctionSignature> Funcs;
-	GetFunctions(Funcs);
+	GetFunctionsInternal(Funcs);
 
 	for (const FNiagaraFunctionSignature& Func : Funcs)
 	{
diff --git a/Engine/Plugins/FX/Niagara/Source/Niagara/Private/DataInterface/NiagaraDataInterfaceDataChannelWrite.cpp b/Engine/Plugins/FX/Niagara/Source/Niagara/Private/DataInterface/NiagaraDataInterfaceDataChannelWrite.cpp
--- a/Engine/Plugins/FX/Niagara/Source/Niagara/Private/DataInterface/NiagaraDataInterfaceDataChannelWrite.cpp
+++ b/Engine/Plugins/FX/Niagara/Source/Niagara/Private/DataInterface/NiagaraDataInterfaceDataChannelWrite.cpp
@@ -714,14 +714,13 @@ bool UNiagaraDataInterfaceDataChannelWrite::CopyToInternal(UNiagaraDataInterface
 	return false;
 }
 
-void UNiagaraDataInterfaceDataChannelWrite::GetFunctions(TArray<FNiagaraFunctionSignature>& OutFunctions)
+#if WITH_EDITORONLY_DATA
+void UNiagaraDataInterfaceDataChannelWrite::GetFunctionsInternal(TArray<FNiagaraFunctionSignature>& OutFunctions) const
 {
 	{
 		FNiagaraFunctionSignature Sig;
 		Sig.Name = NDIDataChannelWriteLocal::NumName;
-#if WITH_EDITORONLY_DATA
 		Sig.Description = LOCTEXT("NumFunctionDescription", "Returns the current number of DataChannel accessible by this interface.");
-#endif
 		Sig.bMemberFunction = true;
 		Sig.bExperimental = true;
 		Sig.AddInput(FNiagaraVariable(FNiagaraTypeDefinition(GetClass()), TEXT("DataChannel interface")));
@@ -735,9 +734,7 @@ void UNiagaraDataInterfaceDataChannelWrite::GetFunctions(TArray<FNiagaraFunction
 	{
 		FNiagaraFunctionSignature Sig;
 		Sig.Name = NDIDataChannelWriteLocal::WriteName;
-#if WITH_EDITORONLY_DATA
 		Sig.Description = LOCTEXT("WriteFunctionDescription", "Writes DataChannel data at a specific index.  Values in the DataChannel that are not written here are set to their defaults. Returns success if an DataChannel was written to.");
-#endif
 		Sig.bMemberFunction = true;
 		Sig.bRequiresExecPin = true;
 		Sig.bExperimental = true;
@@ -751,9 +748,7 @@ void UNiagaraDataInterfaceDataChannelWrite::GetFunctions(TArray<FNiagaraFunction
 	{
 		FNiagaraFunctionSignature Sig;
 		Sig.Name = NDIDataChannelWriteLocal::AppendName;
-#if WITH_EDITORONLY_DATA
 		Sig.Description = LOCTEXT("AppendFunctionDescription", "Appends a new DataChannel to the end of the DataChannel array and writes the specified values. Values in the DataChannel that are not written here are set to their defaults. Returns success if an DataChannel was successfully pushed.");
-#endif
 		Sig.bMemberFunction = true;
 		Sig.bRequiresExecPin = true;
 		Sig.bExperimental = true;
@@ -764,6 +759,7 @@ void UNiagaraDataInterfaceDataChannelWrite::GetFunctions(TArray<FNiagaraFunction
 		OutFunctions.Add(Sig);
 	}
 }
+#endif
 
 void UNiagaraDataInterfaceDataChannelWrite::GetVMExternalFunction(const FVMExternalFunctionBindingInfo& BindingInfo, void* InstanceData, FVMExternalFunction& OutFunc)
 {
diff --git a/Engine/Plugins/FX/Niagara/Source/Niagara/Private/DataInterface/NiagaraDataInterfaceDebugDraw.cpp b/Engine/Plugins/FX/Niagara/Source/Niagara/Private/DataInterface/NiagaraDataInterfaceDebugDraw.cpp
--- a/Engine/Plugins/FX/Niagara/Source/Niagara/Private/DataInterface/NiagaraDataInterfaceDebugDraw.cpp
+++ b/Engine/Plugins/FX/Niagara/Source/Niagara/Private/DataInterface/NiagaraDataInterfaceDebugDraw.cpp
@@ -2303,10 +2303,12 @@ void UNiagaraDataInterfaceDebugDraw::PostInitProperties()
 	}
 }
 
+#if WITH_EDITORONLY_DATA
+
 // Codegen optimization degenerates for very long functions like GetFunctions when combined with the invocation of lots of FORCEINLINE methods.
 BEGIN_FUNCTION_BUILD_OPTIMIZATION
 
-void UNiagaraDataInterfaceDebugDraw::GetFunctions(TArray<FNiagaraFunctionSignature>& OutFunctions)
+void UNiagaraDataInterfaceDebugDraw::GetFunctionsInternal(TArray<FNiagaraFunctionSignature>& OutFunctions) const
 {
 	OutFunctions.Reserve(OutFunctions.Num() + 22);
 	const int32 FirstFunction = OutFunctions.Num();
@@ -2691,19 +2693,14 @@ void UNiagaraDataInterfaceDebugDraw::GetFunctions(TArray<FNiagaraFunctionSignatu
 		Signature.Inputs.Add(FNiagaraVariable(FNiagaraTypeDefinition::GetColorDef(), TEXT("Color")));
 	}
 
-
-
-#if WITH_EDITORONLY_DATA
 	for (int i = FirstFunction; i < OutFunctions.Num(); ++i)
 	{
 		OutFunctions[i].FunctionVersion = FNiagaraDebugDrawDIFunctionVersion::LatestVersion;
 	}
-#endif
 }
 
 END_FUNCTION_BUILD_OPTIMIZATION
 
-#if WITH_EDITORONLY_DATA
 bool UNiagaraDataInterfaceDebugDraw::UpgradeFunctionCall(FNiagaraFunctionSignature& FunctionSignature)
 {
 	bool bWasChanged = false;
@@ -2717,7 +2714,7 @@ bool UNiagaraDataInterfaceDebugDraw::UpgradeFunctionCall(FNiagaraFunctionSignatu
 	if (FunctionSignature.FunctionVersion < FNiagaraDebugDrawDIFunctionVersion::AddedSphereExtraRotation && FunctionSignature.bIsCompileTagGenerator)
 	{
 		TArray<FNiagaraFunctionSignature> Sigs;
-		GetFunctions(Sigs);
+		GetFunctionsInternal(Sigs);
 		for (const FNiagaraFunctionSignature& Sig : Sigs)
 		{
 			if (Sig.Name == FunctionSignature.Name)
diff --git a/Engine/Plugins/FX/Niagara/Source/Niagara/Private/DataInterface/NiagaraDataInterfaceDebugDraw.h b/Engine/Plugins/FX/Niagara/Source/Niagara/Private/DataInterface/NiagaraDataInterfaceDebugDraw.h
--- a/Engine/Plugins/FX/Niagara/Source/Niagara/Private/DataInterface/NiagaraDataInterfaceDebugDraw.h
+++ b/Engine/Plugins/FX/Niagara/Source/Niagara/Private/DataInterface/NiagaraDataInterfaceDebugDraw.h
@@ -28,7 +28,6 @@ public:
 
 	//UNiagaraDataInterface Interface
 	virtual bool CanExecuteOnTarget(ENiagaraSimTarget Target) const override { return true; }
-	NIAGARA_API virtual void GetFunctions(TArray<FNiagaraFunctionSignature>& OutFunctions) override;
 	NIAGARA_API virtual void GetVMExternalFunction(const FVMExternalFunctionBindingInfo& BindingInfo, void* InstanceData, FVMExternalFunction& OutFunc) override;
 
 	enum ShapeId
@@ -74,6 +73,9 @@ public:
 	NIAGARA_API virtual bool Equals(const UNiagaraDataInterface* Other) const override;
 
 protected:
+#if WITH_EDITORONLY_DATA
+	NIAGARA_API virtual void GetFunctionsInternal(TArray<FNiagaraFunctionSignature>& OutFunctions) const override;
+#endif
 	/** Copy one niagara DI to this */
 	NIAGARA_API virtual bool CopyToInternal(UNiagaraDataInterface* Destination) const override;
 	//UNiagaraDataInterface Interface
diff --git a/Engine/Plugins/FX/Niagara/Source/Niagara/Private/DataInterface/NiagaraDataInterfaceDynamicMesh.cpp b/Engine/Plugins/FX/Niagara/Source/Niagara/Private/DataInterface/NiagaraDataInterfaceDynamicMesh.cpp
--- a/Engine/Plugins/FX/Niagara/Source/Niagara/Private/DataInterface/NiagaraDataInterfaceDynamicMesh.cpp
+++ b/Engine/Plugins/FX/Niagara/Source/Niagara/Private/DataInterface/NiagaraDataInterfaceDynamicMesh.cpp
@@ -751,7 +751,8 @@ void UNiagaraDataInterfaceDynamicMesh::PostInitProperties()
 	}
 }
 
-void UNiagaraDataInterfaceDynamicMesh::GetFunctions(TArray<FNiagaraFunctionSignature>& OutFunctions)
+#if WITH_EDITORONLY_DATA
+void UNiagaraDataInterfaceDynamicMesh::GetFunctionsInternal(TArray<FNiagaraFunctionSignature>& OutFunctions) const
 {
 	using namespace NDIDynamicMeshLocal;
 
@@ -1001,6 +1002,7 @@ void UNiagaraDataInterfaceDynamicMesh::GetFunctions(TArray<FNiagaraFunctionSigna
 		Signature.SetDescription(LOCTEXT("AppendTriangleDesc", "Appends a triangle to the section.  This assumes that triangles are stored sequentially in the section, no vertex sharing, etc."));
 	}
 }
+#endif
 
 void UNiagaraDataInterfaceDynamicMesh::GetVMExternalFunction(const FVMExternalFunctionBindingInfo& BindingInfo, void* InstanceData, FVMExternalFunction& OutFunc)
 {
diff --git a/Engine/Plugins/FX/Niagara/Source/Niagara/Private/DataInterface/NiagaraDataInterfaceDynamicMesh.h b/Engine/Plugins/FX/Niagara/Source/Niagara/Private/DataInterface/NiagaraDataInterfaceDynamicMesh.h
--- a/Engine/Plugins/FX/Niagara/Source/Niagara/Private/DataInterface/NiagaraDataInterfaceDynamicMesh.h
+++ b/Engine/Plugins/FX/Niagara/Source/Niagara/Private/DataInterface/NiagaraDataInterfaceDynamicMesh.h
@@ -95,7 +95,6 @@ public:
 	//UObject Interface End
 
 	//UNiagaraDataInterface Interface
-	NIAGARA_API virtual void GetFunctions(TArray<FNiagaraFunctionSignature>& OutFunctions) override;
 	NIAGARA_API virtual void GetVMExternalFunction(const FVMExternalFunctionBindingInfo& BindingInfo, void* InstanceData, FVMExternalFunction &OutFunc) override;
 	virtual bool CanExecuteOnTarget(ENiagaraSimTarget Target) const override { return true; }
 
@@ -156,4 +155,9 @@ public:
 	static NIAGARA_API void VMSetVertexData(FVectorVMExternalFunctionContext& Context);
 
 	static NIAGARA_API void VMAppendTriangle(FVectorVMExternalFunctionContext& Context);
+
+protected:
+#if WITH_EDITORONLY_DATA
+	NIAGARA_API virtual void GetFunctionsInternal(TArray<FNiagaraFunctionSignature>& OutFunctions) const override;
+#endif
 };
diff --git a/Engine/Plugins/FX/Niagara/Source/Niagara/Private/DataInterface/NiagaraDataInterfaceEmitterProperties.cpp b/Engine/Plugins/FX/Niagara/Source/Niagara/Private/DataInterface/NiagaraDataInterfaceEmitterProperties.cpp
--- a/Engine/Plugins/FX/Niagara/Source/Niagara/Private/DataInterface/NiagaraDataInterfaceEmitterProperties.cpp
+++ b/Engine/Plugins/FX/Niagara/Source/Niagara/Private/DataInterface/NiagaraDataInterfaceEmitterProperties.cpp
@@ -188,7 +188,8 @@ bool UNiagaraDataInterfaceEmitterProperties::PerInstanceTick(void* PerInstanceDa
 	return false;
 }
 
-void UNiagaraDataInterfaceEmitterProperties::GetFunctions(TArray<FNiagaraFunctionSignature>& OutFunctions)
+#if WITH_EDITORONLY_DATA
+void UNiagaraDataInterfaceEmitterProperties::GetFunctionsInternal(TArray<FNiagaraFunctionSignature>& OutFunctions) const
 {
 	using namespace NDIEmitterPropertiesLocal;
 
@@ -246,6 +247,7 @@ void UNiagaraDataInterfaceEmitterProperties::GetFunctions(TArray<FNiagaraFunctio
 #endif
 	}
 }
+#endif
 
 void UNiagaraDataInterfaceEmitterProperties::GetVMExternalFunction(const FVMExternalFunctionBindingInfo& BindingInfo, void* PerInstanceData, FVMExternalFunction& OutFunc)
 {
diff --git a/Engine/Plugins/FX/Niagara/Source/Niagara/Private/DataInterface/NiagaraDataInterfaceEmitterProperties.h b/Engine/Plugins/FX/Niagara/Source/Niagara/Private/DataInterface/NiagaraDataInterfaceEmitterProperties.h
--- a/Engine/Plugins/FX/Niagara/Source/Niagara/Private/DataInterface/NiagaraDataInterfaceEmitterProperties.h
+++ b/Engine/Plugins/FX/Niagara/Source/Niagara/Private/DataInterface/NiagaraDataInterfaceEmitterProperties.h
@@ -51,7 +51,6 @@ public:
 	virtual bool HasPreSimulateTick() const override { return true; }
 	virtual bool PerInstanceTick(void* PerInstanceData, FNiagaraSystemInstance* SystemInstance, float DeltaSeconds) override;
 
-	virtual void GetFunctions(TArray<FNiagaraFunctionSignature>& OutFunctions) override;
 	virtual void GetVMExternalFunction(const FVMExternalFunctionBindingInfo& BindingInfo, void* InstanceData, FVMExternalFunction& OutFunc) override;
 #if WITH_EDITORONLY_DATA
 	virtual bool AppendCompileHash(FNiagaraCompileHashVisitor* InVisitor) const override;
@@ -65,6 +64,11 @@ public:
 #endif
 	//UNiagaraDataInterface Interface
 
+protected:
+#if WITH_EDITORONLY_DATA
+	virtual void GetFunctionsInternal(TArray<FNiagaraFunctionSignature>& OutFunctions) const override;
+#endif
+
 private:
 	void VMGetLocalSpace(FVectorVMExternalFunctionContext& Context);
 	void VMGetBounds(FVectorVMExternalFunctionContext& Context);
diff --git a/Engine/Plugins/FX/Niagara/Source/Niagara/Private/DataInterface/NiagaraDataInterfaceGBuffer.cpp b/Engine/Plugins/FX/Niagara/Source/Niagara/Private/DataInterface/NiagaraDataInterfaceGBuffer.cpp
--- a/Engine/Plugins/FX/Niagara/Source/Niagara/Private/DataInterface/NiagaraDataInterfaceGBuffer.cpp
+++ b/Engine/Plugins/FX/Niagara/Source/Niagara/Private/DataInterface/NiagaraDataInterfaceGBuffer.cpp
@@ -139,7 +139,8 @@ void UNiagaraDataInterfaceGBuffer::PostInitProperties()
 	}
 }
 
-void UNiagaraDataInterfaceGBuffer::GetFunctions(TArray<FNiagaraFunctionSignature>& OutFunctions)
+#if WITH_EDITORONLY_DATA
+void UNiagaraDataInterfaceGBuffer::GetFunctionsInternal(TArray<FNiagaraFunctionSignature>& OutFunctions) const
 {
 	using namespace NiagaraDataInterfaceGBufferLocal;
 
@@ -151,9 +152,7 @@ void UNiagaraDataInterfaceGBuffer::GetFunctions(TArray<FNiagaraFunctionSignature
 	{
 		FNiagaraFunctionSignature& Signature = OutFunctions.AddDefaulted_GetRef();
 		Signature.Name = Attribute.ScreenUVFunctionName;
-#if WITH_EDITORONLY_DATA
 		Signature.Description = Attribute.Description;
-#endif
 		Signature.bMemberFunction = true;
 		Signature.bRequiresContext = false;
 		Signature.bSupportsCPU = false;
@@ -163,11 +162,10 @@ void UNiagaraDataInterfaceGBuffer::GetFunctions(TArray<FNiagaraFunctionSignature
 		Signature.Inputs.Add_GetRef(FNiagaraVariable(FNiagaraTypeDefinition::GetBoolDef(), TEXT("ApplyViewportOffset"))).SetValue(true);
 		Signature.Outputs.Add(FNiagaraVariable(FNiagaraTypeDefinition::GetBoolDef(), TEXT("IsValid")));
 		Signature.Outputs.Add(FNiagaraVariable(Attribute.TypeDef, Attribute.AttributeName));
-#if WITH_EDITORONLY_DATA
 		Signature.FunctionVersion = EDIFunctionVersion::LatestVersion;
-#endif
 	}
 }
+#endif
 
 void UNiagaraDataInterfaceGBuffer::BuildShaderParameters(FNiagaraShaderParametersBuilder& ShaderParametersBuilder) const
 {
diff --git a/Engine/Plugins/FX/Niagara/Source/Niagara/Private/DataInterface/NiagaraDataInterfaceGBuffer.h b/Engine/Plugins/FX/Niagara/Source/Niagara/Private/DataInterface/NiagaraDataInterfaceGBuffer.h
--- a/Engine/Plugins/FX/Niagara/Source/Niagara/Private/DataInterface/NiagaraDataInterfaceGBuffer.h
+++ b/Engine/Plugins/FX/Niagara/Source/Niagara/Private/DataInterface/NiagaraDataInterfaceGBuffer.h
@@ -19,8 +19,6 @@ public:
 	//UObject Interface End
 
 	//UNiagaraDataInterface Interface
-	NIAGARA_API virtual void GetFunctions(TArray<FNiagaraFunctionSignature>& OutFunctions)override;
-
 	NIAGARA_API virtual void BuildShaderParameters(FNiagaraShaderParametersBuilder& ShaderParametersBuilder) const override;
 	NIAGARA_API virtual void SetShaderParameters(const FNiagaraDataInterfaceSetShaderParametersContext& Context) const override;
 
@@ -34,4 +32,9 @@ public:
 	virtual bool CanExecuteOnTarget(ENiagaraSimTarget Target) const override { return Target == ENiagaraSimTarget::GPUComputeSim; }
 	virtual bool RequiresDepthBuffer() const override { return true; }
 	//UNiagaraDataInterface Interface
+
+protected:
+#if WITH_EDITORONLY_DATA
+	NIAGARA_API virtual void GetFunctionsInternal(TArray<FNiagaraFunctionSignature>& OutFunctions) const override;
+#endif
 };
diff --git a/Engine/Plugins/FX/Niagara/Source/Niagara/Private/DataInterface/NiagaraDataInterfacePhysicsAsset.cpp b/Engine/Plugins/FX/Niagara/Source/Niagara/Private/DataInterface/NiagaraDataInterfacePhysicsAsset.cpp
--- a/Engine/Plugins/FX/Niagara/Source/Niagara/Private/DataInterface/NiagaraDataInterfacePhysicsAsset.cpp
+++ b/Engine/Plugins/FX/Niagara/Source/Niagara/Private/DataInterface/NiagaraDataInterfacePhysicsAsset.cpp
@@ -828,7 +828,8 @@ void UNiagaraDataInterfacePhysicsAsset::PostLoad()
 #endif
 }
 
-void UNiagaraDataInterfacePhysicsAsset::GetFunctions(TArray<FNiagaraFunctionSignature>& OutFunctions)
+#if WITH_EDITORONLY_DATA
+void UNiagaraDataInterfacePhysicsAsset::GetFunctionsInternal(TArray<FNiagaraFunctionSignature>& OutFunctions) const
 {
 	using namespace NDIPhysicsAssetLocal;
 	{
@@ -993,6 +994,7 @@ void UNiagaraDataInterfacePhysicsAsset::GetFunctions(TArray<FNiagaraFunctionSign
 		OutFunctions.Add(Sig);
 	}
 }
+#endif
 
 DEFINE_NDI_DIRECT_FUNC_BINDER(UNiagaraDataInterfacePhysicsAsset, GetNumBoxes);
 DEFINE_NDI_DIRECT_FUNC_BINDER(UNiagaraDataInterfacePhysicsAsset, GetNumSpheres);
diff --git a/Engine/Plugins/FX/Niagara/Source/Niagara/Private/DataInterface/NiagaraDataInterfaceSceneCapture2D.cpp b/Engine/Plugins/FX/Niagara/Source/Niagara/Private/DataInterface/NiagaraDataInterfaceSceneCapture2D.cpp
--- a/Engine/Plugins/FX/Niagara/Source/Niagara/Private/DataInterface/NiagaraDataInterfaceSceneCapture2D.cpp
+++ b/Engine/Plugins/FX/Niagara/Source/Niagara/Private/DataInterface/NiagaraDataInterfaceSceneCapture2D.cpp
@@ -377,7 +377,8 @@ bool UNiagaraDataInterfaceSceneCapture2D::Equals(const UNiagaraDataInterface* In
 		Other->AutoMoveOffsetRotation == AutoMoveOffsetRotation;
 }
 
-void UNiagaraDataInterfaceSceneCapture2D::GetFunctions(TArray<FNiagaraFunctionSignature>& OutFunctions)
+#if WITH_EDITORONLY_DATA
+void UNiagaraDataInterfaceSceneCapture2D::GetFunctionsInternal(TArray<FNiagaraFunctionSignature>& OutFunctions) const
 {
 	using namespace NDISceneCapture2DLocal;
 
@@ -502,6 +503,7 @@ void UNiagaraDataInterfaceSceneCapture2D::GetFunctions(TArray<FNiagaraFunctionSi
 		Sig.SetDescription(LOCTEXT("TextureDeprojectDepthDesc", "Deproject the depth sample from the texture into world space."));
 	}
 }
+#endif
 
 void UNiagaraDataInterfaceSceneCapture2D::GetVMExternalFunction(const FVMExternalFunctionBindingInfo& BindingInfo, void* InstanceData, FVMExternalFunction &OutFunc)
 {
diff --git a/Engine/Plugins/FX/Niagara/Source/Niagara/Private/DataInterface/NiagaraDataInterfaceSceneCapture2D.h b/Engine/Plugins/FX/Niagara/Source/Niagara/Private/DataInterface/NiagaraDataInterfaceSceneCapture2D.h
--- a/Engine/Plugins/FX/Niagara/Source/Niagara/Private/DataInterface/NiagaraDataInterfaceSceneCapture2D.h
+++ b/Engine/Plugins/FX/Niagara/Source/Niagara/Private/DataInterface/NiagaraDataInterfaceSceneCapture2D.h
@@ -105,7 +105,6 @@ public:
 	//UObject Interface End
 
 	//UNiagaraDataInterface Interface
-	NIAGARA_API virtual void GetFunctions(TArray<FNiagaraFunctionSignature>& OutFunctions) override;
 	NIAGARA_API virtual void GetVMExternalFunction(const FVMExternalFunctionBindingInfo& BindingInfo, void* InstanceData, FVMExternalFunction &OutFunc) override;
 	virtual bool CanExecuteOnTarget(ENiagaraSimTarget Target) const override { return Target==ENiagaraSimTarget::GPUComputeSim; }
 
@@ -134,6 +133,9 @@ public:
 	NIAGARA_API virtual void SetShaderParameters(const FNiagaraDataInterfaceSetShaderParametersContext& Context) const override;
 
 protected:
+#if WITH_EDITORONLY_DATA
+	NIAGARA_API virtual void GetFunctionsInternal(TArray<FNiagaraFunctionSignature>& OutFunctions) const override;
+#endif
 	NIAGARA_API virtual bool Equals(const UNiagaraDataInterface* Other) const override;
 	NIAGARA_API virtual bool CopyToInternal(UNiagaraDataInterface* Destination) const override;
 
diff --git a/Engine/Plugins/FX/Niagara/Source/Niagara/Private/DataInterface/NiagaraDataInterfaceSimCacheReader.cpp b/Engine/Plugins/FX/Niagara/Source/Niagara/Private/DataInterface/NiagaraDataInterfaceSimCacheReader.cpp
--- a/Engine/Plugins/FX/Niagara/Source/Niagara/Private/DataInterface/NiagaraDataInterfaceSimCacheReader.cpp
+++ b/Engine/Plugins/FX/Niagara/Source/Niagara/Private/DataInterface/NiagaraDataInterfaceSimCacheReader.cpp
@@ -333,7 +333,8 @@ bool UNiagaraDataInterfaceSimCacheReader::PerInstanceTick(void* PerInstanceData,
 	return false;
 }
 
-void UNiagaraDataInterfaceSimCacheReader::GetFunctions(TArray<FNiagaraFunctionSignature>& OutFunctions)
+#if WITH_EDITORONLY_DATA
+void UNiagaraDataInterfaceSimCacheReader::GetFunctionsInternal(TArray<FNiagaraFunctionSignature>& OutFunctions) const
 {
 	using namespace NDISimCacheReaderLocal;
 
@@ -383,6 +384,7 @@ void UNiagaraDataInterfaceSimCacheReader::GetFunctions(TArray<FNiagaraFunctionSi
 		Sig.Outputs.Emplace(GenFuncInfo.TypeDef, TEXT("Value"));
 	}
 }
+#endif
 
 void UNiagaraDataInterfaceSimCacheReader::GetVMExternalFunction(const FVMExternalFunctionBindingInfo& BindingInfo, void* PerInstanceData, FVMExternalFunction& OutFunc)
 {
diff --git a/Engine/Plugins/FX/Niagara/Source/Niagara/Private/DataInterface/NiagaraDataInterfaceSimCacheReader.h b/Engine/Plugins/FX/Niagara/Source/Niagara/Private/DataInterface/NiagaraDataInterfaceSimCacheReader.h
--- a/Engine/Plugins/FX/Niagara/Source/Niagara/Private/DataInterface/NiagaraDataInterfaceSimCacheReader.h
+++ b/Engine/Plugins/FX/Niagara/Source/Niagara/Private/DataInterface/NiagaraDataInterfaceSimCacheReader.h
@@ -55,7 +55,6 @@ public:
 	virtual bool HasPreSimulateTick() const override { return true; }
 	virtual bool PerInstanceTick(void* PerInstanceData, FNiagaraSystemInstance* SystemInstance, float DeltaSeconds) override;
 
-	virtual void GetFunctions(TArray<FNiagaraFunctionSignature>& OutFunctions) override;
 	virtual void GetVMExternalFunction(const FVMExternalFunctionBindingInfo& BindingInfo, void* InstanceData, FVMExternalFunction& OutFunc) override;
 #if WITH_EDITORONLY_DATA
 	virtual bool AppendCompileHash(FNiagaraCompileHashVisitor* InVisitor) const override;
@@ -66,6 +65,11 @@ public:
 	virtual void SetShaderParameters(const FNiagaraDataInterfaceSetShaderParametersContext& Context) const override;
 	//UNiagaraDataInterface Interface
 
+protected:
+#if WITH_EDITORONLY_DATA
+	virtual void GetFunctionsInternal(TArray<FNiagaraFunctionSignature>& OutFunctions) const override;
+#endif
+
 private:
 	void VMGetNumFrames(FVectorVMExternalFunctionContext& Context);
 	void VMGetNumEmitters(FVectorVMExternalFunctionContext& Context);
diff --git a/Engine/Plugins/FX/Niagara/Source/Niagara/Private/DataInterface/NiagaraDataInterfaceSimpleCounter.cpp b/Engine/Plugins/FX/Niagara/Source/Niagara/Private/DataInterface/NiagaraDataInterfaceSimpleCounter.cpp
--- a/Engine/Plugins/FX/Niagara/Source/Niagara/Private/DataInterface/NiagaraDataInterfaceSimpleCounter.cpp
+++ b/Engine/Plugins/FX/Niagara/Source/Niagara/Private/DataInterface/NiagaraDataInterfaceSimpleCounter.cpp
@@ -265,7 +265,8 @@ int32 UNiagaraDataInterfaceSimpleCounter::PerInstanceDataSize() const
 	return sizeof(FNDISimpleCounterInstanceData_GameThread);
 }
 
-void UNiagaraDataInterfaceSimpleCounter::GetFunctions(TArray<FNiagaraFunctionSignature>& OutFunctions)
+#if WITH_EDITORONLY_DATA
+void UNiagaraDataInterfaceSimpleCounter::GetFunctionsInternal(TArray<FNiagaraFunctionSignature>& OutFunctions) const
 {
 	using namespace NDISimpleCounterLocal;
 
@@ -351,6 +352,7 @@ void UNiagaraDataInterfaceSimpleCounter::GetFunctions(TArray<FNiagaraFunctionSig
 		Sig.SetDescription(LOCTEXT("DecrementDesc", "Decrements the counter by 1."));
 	}
 }
+#endif
 
 void UNiagaraDataInterfaceSimpleCounter::GetVMExternalFunction(const FVMExternalFunctionBindingInfo& BindingInfo, void* InstanceData, FVMExternalFunction& OutFunc)
 {
diff --git a/Engine/Plugins/FX/Niagara/Source/Niagara/Private/DataInterface/NiagaraDataInterfaceSimpleCounter.h b/Engine/Plugins/FX/Niagara/Source/Niagara/Private/DataInterface/NiagaraDataInterfaceSimpleCounter.h
--- a/Engine/Plugins/FX/Niagara/Source/Niagara/Private/DataInterface/NiagaraDataInterfaceSimpleCounter.h
+++ b/Engine/Plugins/FX/Niagara/Source/Niagara/Private/DataInterface/NiagaraDataInterfaceSimpleCounter.h
@@ -42,7 +42,6 @@ public:
 	NIAGARA_API virtual void DestroyPerInstanceData(void* PerInstanceData, FNiagaraSystemInstance* SystemInstance) override;
 	NIAGARA_API virtual int32 PerInstanceDataSize() const override;
 
-	NIAGARA_API virtual void GetFunctions(TArray<FNiagaraFunctionSignature>& OutFunctions) override;
 	NIAGARA_API virtual void GetVMExternalFunction(const FVMExternalFunctionBindingInfo& BindingInfo, void* InstanceData, FVMExternalFunction &OutFunc) override;
 #if WITH_EDITORONLY_DATA
 	NIAGARA_API virtual void GetParameterDefinitionHLSL(const FNiagaraDataInterfaceGPUParamInfo& ParamInfo, FString& OutHLSL) override;
@@ -76,4 +75,9 @@ public:
 	/** This is the value the counter will have when the instance is reset / created. */
 	UPROPERTY(EditAnywhere, Category = "Simple Counter")
 	int32 InitialValue = 0;
+
+protected:
+#if WITH_EDITORONLY_DATA
+	NIAGARA_API virtual void GetFunctionsInternal(TArray<FNiagaraFunctionSignature>& OutFunctions) const override;
+#endif
 };
diff --git a/Engine/Plugins/FX/Niagara/Source/Niagara/Private/DataInterface/NiagaraDataInterfaceStaticMesh.cpp b/Engine/Plugins/FX/Niagara/Source/Niagara/Private/DataInterface/NiagaraDataInterfaceStaticMesh.cpp
--- a/Engine/Plugins/FX/Niagara/Source/Niagara/Private/DataInterface/NiagaraDataInterfaceStaticMesh.cpp
+++ b/Engine/Plugins/FX/Niagara/Source/Niagara/Private/DataInterface/NiagaraDataInterfaceStaticMesh.cpp
@@ -1941,7 +1941,8 @@ void UNiagaraDataInterfaceStaticMesh::ProvidePerInstanceDataForRenderThread(void
 	}
 }
 
-void UNiagaraDataInterfaceStaticMesh::GetFunctions(TArray<FNiagaraFunctionSignature>& OutFunctions)
+#if WITH_EDITORONLY_DATA
+void UNiagaraDataInterfaceStaticMesh::GetFunctionsInternal(TArray<FNiagaraFunctionSignature>& OutFunctions) const
 {
 	using namespace NDIStaticMeshLocal;
 
@@ -1952,9 +1953,7 @@ void UNiagaraDataInterfaceStaticMesh::GetFunctions(TArray<FNiagaraFunctionSignat
 	BaseSignature.Inputs.Emplace(FNiagaraTypeDefinition(GetClass()), TEXT("StaticMesh"));
 	BaseSignature.bMemberFunction = true;
 	BaseSignature.bRequiresContext = false;
-#if WITH_EDITORONLY_DATA
 	BaseSignature.FunctionVersion = EDIFunctionVersion::LatestVersion;
-#endif
 
 	GetVertexSamplingFunctions(OutFunctions, BaseSignature);
 	GetTriangleSamplingFunctions(OutFunctions, BaseSignature);
@@ -2512,6 +2511,7 @@ void UNiagaraDataInterfaceStaticMesh::GetDeprecatedFunctions(TArray<FNiagaraFunc
 		OutFunctions.Add_GetRef(Sig).Name = Deprecated_GetTriNormalWSName;
 	}
 }
+#endif
 
 bool UNiagaraDataInterfaceStaticMesh::FunctionNeedsCpuAccess(FName InName)
 {
diff --git a/Engine/Plugins/FX/Niagara/Source/Niagara/Private/DataInterface/NiagaraDataInterfaceUObjectPropertyReader.cpp b/Engine/Plugins/FX/Niagara/Source/Niagara/Private/DataInterface/NiagaraDataInterfaceUObjectPropertyReader.cpp
--- a/Engine/Plugins/FX/Niagara/Source/Niagara/Private/DataInterface/NiagaraDataInterfaceUObjectPropertyReader.cpp
+++ b/Engine/Plugins/FX/Niagara/Source/Niagara/Private/DataInterface/NiagaraDataInterfaceUObjectPropertyReader.cpp
@@ -869,7 +869,8 @@ bool UNiagaraDataInterfaceUObjectPropertyReader::PerInstanceTick(void* PerInstan
 	return false;
 }
 
-void UNiagaraDataInterfaceUObjectPropertyReader::GetFunctions(TArray<FNiagaraFunctionSignature>& OutFunctions)
+#if WITH_EDITORONLY_DATA
+void UNiagaraDataInterfaceUObjectPropertyReader::GetFunctionsInternal(TArray<FNiagaraFunctionSignature>& OutFunctions) const
 {
 	using namespace NDIUObjectPropertyReaderLocal;
 
@@ -924,6 +925,7 @@ void UNiagaraDataInterfaceUObjectPropertyReader::GetFunctions(TArray<FNiagaraFun
 		NDI_PROPERTY_TYPES
 	#undef NDI_PROPERTY_TYPE
 }
+#endif
 
 void UNiagaraDataInterfaceUObjectPropertyReader::GetVMExternalFunction(const FVMExternalFunctionBindingInfo& BindingInfo, void* PerInstanceData, FVMExternalFunction& OutFunc)
 {
diff --git a/Engine/Plugins/FX/Niagara/Source/Niagara/Private/DataInterface/NiagaraDataInterfaceUObjectPropertyReader.h b/Engine/Plugins/FX/Niagara/Source/Niagara/Private/DataInterface/NiagaraDataInterfaceUObjectPropertyReader.h
--- a/Engine/Plugins/FX/Niagara/Source/Niagara/Private/DataInterface/NiagaraDataInterfaceUObjectPropertyReader.h
+++ b/Engine/Plugins/FX/Niagara/Source/Niagara/Private/DataInterface/NiagaraDataInterfaceUObjectPropertyReader.h
@@ -104,7 +104,6 @@ public:
 	virtual bool HasPreSimulateTick() const override { return true; }
 	virtual bool PerInstanceTick(void* PerInstanceData, FNiagaraSystemInstance* SystemInstance, float DeltaSeconds) override;
 
-	virtual void GetFunctions(TArray<FNiagaraFunctionSignature>& OutFunctions) override;
 	virtual void GetVMExternalFunction(const FVMExternalFunctionBindingInfo& BindingInfo, void* InstanceData, FVMExternalFunction& OutFunc) override;
 #if WITH_EDITORONLY_DATA
 	virtual bool AppendCompileHash(FNiagaraCompileHashVisitor* InVisitor) const override;
@@ -119,6 +118,11 @@ public:
 	UFUNCTION(BlueprintCallable, Category = Niagara)
 	static void SetUObjectReaderPropertyRemap(UNiagaraComponent* NiagaraComponent, FName UserParameterName, FName GraphName, FName RemapName);
 
+protected:
+#if WITH_EDITORONLY_DATA
+	virtual void GetFunctionsInternal(TArray<FNiagaraFunctionSignature>& OutFunctions) const override;
+#endif
+
 private:
 	void VMGetComponentTransform(FVectorVMExternalFunctionContext& Context);
 	void VMGetComponentInvTransform(FVectorVMExternalFunctionContext& Context);
diff --git a/Engine/Plugins/FX/Niagara/Source/Niagara/Private/DataInterface/NiagaraDataInterfaceVirtualTexture.cpp b/Engine/Plugins/FX/Niagara/Source/Niagara/Private/DataInterface/NiagaraDataInterfaceVirtualTexture.cpp
--- a/Engine/Plugins/FX/Niagara/Source/Niagara/Private/DataInterface/NiagaraDataInterfaceVirtualTexture.cpp
+++ b/Engine/Plugins/FX/Niagara/Source/Niagara/Private/DataInterface/NiagaraDataInterfaceVirtualTexture.cpp
@@ -174,7 +174,8 @@ bool UNiagaraDataInterfaceVirtualTexture::Equals(const UNiagaraDataInterface* Ot
 		OtherTexture->TextureUserParameter == TextureUserParameter;
 }
 
-void UNiagaraDataInterfaceVirtualTexture::GetFunctions(TArray<FNiagaraFunctionSignature>& OutFunctions)
+#if WITH_EDITORONLY_DATA
+void UNiagaraDataInterfaceVirtualTexture::GetFunctionsInternal(TArray<FNiagaraFunctionSignature>& OutFunctions) const
 {
 	using namespace NDIVirtualTextureLocal;
 
@@ -216,6 +217,7 @@ void UNiagaraDataInterfaceVirtualTexture::GetFunctions(TArray<FNiagaraFunctionSi
 		Sig.Outputs.Emplace(FNiagaraTypeDefinition::GetFloatDef(), TEXT("Mask"));
 	}
 }
+#endif
 
 void UNiagaraDataInterfaceVirtualTexture::GetVMExternalFunction(const FVMExternalFunctionBindingInfo& BindingInfo, void* InstanceData, FVMExternalFunction &OutFunc)
 {
diff --git a/Engine/Plugins/FX/Niagara/Source/Niagara/Private/DataInterface/NiagaraDataInterfaceVirtualTexture.h b/Engine/Plugins/FX/Niagara/Source/Niagara/Private/DataInterface/NiagaraDataInterfaceVirtualTexture.h
--- a/Engine/Plugins/FX/Niagara/Source/Niagara/Private/DataInterface/NiagaraDataInterfaceVirtualTexture.h
+++ b/Engine/Plugins/FX/Niagara/Source/Niagara/Private/DataInterface/NiagaraDataInterfaceVirtualTexture.h
@@ -31,7 +31,6 @@ public:
 	//UObject Interface End
 
 	//UNiagaraDataInterface Interface
-	NIAGARA_API virtual void GetFunctions(TArray<FNiagaraFunctionSignature>& OutFunctions) override;
 	NIAGARA_API virtual void GetVMExternalFunction(const FVMExternalFunctionBindingInfo& BindingInfo, void* InstanceData, FVMExternalFunction &OutFunc) override;
 	virtual bool CanExecuteOnTarget(ENiagaraSimTarget Target) const override { return Target==ENiagaraSimTarget::GPUComputeSim; }
 
@@ -63,5 +62,8 @@ public:
 	NIAGARA_API void SetTexture(URuntimeVirtualTexture* InTexture);
 
 protected:
+#if WITH_EDITORONLY_DATA
+	NIAGARA_API virtual void GetFunctionsInternal(TArray<FNiagaraFunctionSignature>& OutFunctions) const override;
+#endif
 	NIAGARA_API virtual bool CopyToInternal(UNiagaraDataInterface* Destination) const override;
 };
diff --git a/Engine/Plugins/FX/Niagara/Source/Niagara/Private/NDISkeletalMesh_BoneSampling.cpp b/Engine/Plugins/FX/Niagara/Source/Niagara/Private/NDISkeletalMesh_BoneSampling.cpp
--- a/Engine/Plugins/FX/Niagara/Source/Niagara/Private/NDISkeletalMesh_BoneSampling.cpp
+++ b/Engine/Plugins/FX/Niagara/Source/Niagara/Private/NDISkeletalMesh_BoneSampling.cpp
@@ -45,7 +45,8 @@ const FName FSkeletalMeshInterfaceHelper::RandomFilteredSocketOrBoneName("Random
 const FName FSkeletalMeshInterfaceHelper::GetFilteredSocketOrBoneCountName("GetFilteredSocketOrBoneCount");
 const FName FSkeletalMeshInterfaceHelper::GetFilteredSocketOrBoneAtName("GetFilteredSocketOrBone");
 
-void UNiagaraDataInterfaceSkeletalMesh::GetSkeletonSamplingFunctions(TArray<FNiagaraFunctionSignature>& OutFunctions)
+#if WITH_EDITORONLY_DATA
+void UNiagaraDataInterfaceSkeletalMesh::GetSkeletonSamplingFunctions(TArray<FNiagaraFunctionSignature>& OutFunctions) const
 {
 	FNiagaraFunctionSignature DefaultSig;
 	DefaultSig.Inputs.Add(FNiagaraVariable(FNiagaraTypeDefinition(GetClass()), TEXT("SkeletalMesh")));
@@ -232,6 +233,7 @@ void UNiagaraDataInterfaceSkeletalMesh::GetSkeletonSamplingFunctions(TArray<FNia
 		Sig.SetDescription(LOCTEXT("GetFilteredSocketOrBoneAtDesc", "Gets a filtered socket or bone count from the DI's list."));
 	}
 }
+#endif
 
 void UNiagaraDataInterfaceSkeletalMesh::BindSkeletonSamplingFunction(const FVMExternalFunctionBindingInfo& BindingInfo, FNDISkeletalMesh_InstanceData* InstanceData, FVMExternalFunction &OutFunc)
 {
diff --git a/Engine/Plugins/FX/Niagara/Source/Niagara/Private/NDISkeletalMesh_TriangleSampling.cpp b/Engine/Plugins/FX/Niagara/Source/Niagara/Private/NDISkeletalMesh_TriangleSampling.cpp
--- a/Engine/Plugins/FX/Niagara/Source/Niagara/Private/NDISkeletalMesh_TriangleSampling.cpp
+++ b/Engine/Plugins/FX/Niagara/Source/Niagara/Private/NDISkeletalMesh_TriangleSampling.cpp
@@ -54,7 +54,8 @@ const FName FSkeletalMeshInterfaceHelper::GetTriangleCoordInAabbName("GetTriangl
 const FName FSkeletalMeshInterfaceHelper::GetAdjacentTriangleIndexName("GetAdjacentTriangleIndex");
 const FName FSkeletalMeshInterfaceHelper::GetTriangleNeighborName("GetTriangleNeighbor");
 
-void UNiagaraDataInterfaceSkeletalMesh::GetTriangleSamplingFunctions(TArray<FNiagaraFunctionSignature>& OutFunctions)
+#if WITH_EDITORONLY_DATA
+void UNiagaraDataInterfaceSkeletalMesh::GetTriangleSamplingFunctions(TArray<FNiagaraFunctionSignature>& OutFunctions) const
 {
 	//-TODO: Remove / deprecate this function!
 	{
@@ -75,9 +76,7 @@ void UNiagaraDataInterfaceSkeletalMesh::GetTriangleSamplingFunctions(TArray<FNia
 		Sig.Outputs.Add(FNiagaraVariable(FNiagaraTypeDefinition::GetBoolDef(), TEXT("IsValid")));
 		Sig.bMemberFunction = true;
 		Sig.bRequiresContext = false;
-#if WITH_EDITORONLY_DATA
 		Sig.Description = LOCTEXT("IsValidDesc", "Determine if this tri coordinate's triangle index is valid for this mesh. Note that this only checks the mesh index buffer size and does not include any filtering settings.");
-#endif
 	}
 	
 	{
@@ -91,9 +90,7 @@ void UNiagaraDataInterfaceSkeletalMesh::GetTriangleSamplingFunctions(TArray<FNia
 		Sig.Outputs.Add(FNiagaraVariable(FNiagaraTypeDefinition::GetVec3Def(), TEXT("Tangent")));
 		Sig.bMemberFunction = true;
 		Sig.bRequiresContext = false;
-#if WITH_EDITORONLY_DATA
 		Sig.Description = LOCTEXT("GetTriangleDataDesc", "Returns bind pose triangle data.");
-#endif
 	}
 
 	{
@@ -106,9 +103,7 @@ void UNiagaraDataInterfaceSkeletalMesh::GetTriangleSamplingFunctions(TArray<FNia
 		Sig.Outputs.Emplace(FNiagaraTypeDefinition::GetIntDef(), TEXT("Index2"));
 		Sig.bMemberFunction = true;
 		Sig.bRequiresContext = false;
-#if WITH_EDITORONLY_DATA
 		Sig.Description = LOCTEXT("GetTriangleDataDesc", "Returns bind pose triangle data.");
-#endif
 	}
 
 	{
@@ -123,9 +118,7 @@ void UNiagaraDataInterfaceSkeletalMesh::GetTriangleSamplingFunctions(TArray<FNia
 		Sig.Outputs.Add(FNiagaraVariable(FNiagaraTypeDefinition::GetVec3Def(), TEXT("Tangent")));
 		Sig.bMemberFunction = true;
 		Sig.bRequiresContext = false;
-#if WITH_EDITORONLY_DATA
 		Sig.Description = LOCTEXT("GetOptionalSkinnedDataDesc", "Returns skinning dependant data for the pased MeshTriCoord in local space. All outputs are optional and you will incur zerp minimal cost if they are not connected.");
-#endif
 	}
 
 	{
@@ -140,9 +133,7 @@ void UNiagaraDataInterfaceSkeletalMesh::GetTriangleSamplingFunctions(TArray<FNia
 		Sig.Outputs.Add(FNiagaraVariable(FNiagaraTypeDefinition::GetVec3Def(), TEXT("Tangent")));
 		Sig.bMemberFunction = true;
 		Sig.bRequiresContext = false;
-#if WITH_EDITORONLY_DATA
 		Sig.Description = LOCTEXT("GetOptionalSkinnedDataWSDesc", "Returns skinning dependant data for the pased MeshTriCoord in world space. All outputs are optional and you will incur zerp minimal cost if they are not connected.");
-#endif
 	}
 
 	{
@@ -158,9 +149,7 @@ void UNiagaraDataInterfaceSkeletalMesh::GetTriangleSamplingFunctions(TArray<FNia
 		Sig.Outputs.Add(FNiagaraVariable(FNiagaraTypeDefinition::GetVec3Def(), TEXT("Tangent")));
 		Sig.bMemberFunction = true;
 		Sig.bRequiresContext = false;
-#if WITH_EDITORONLY_DATA
 		Sig.Description = LOCTEXT("GetSkinnedDataDesc", "Returns skinning dependant data for the pased MeshTriCoord in local space. Interpolates between previous and current frame. All outputs are optional and you will incur zerp minimal cost if they are not connected.");
-#endif
 	}
 
 	{
@@ -176,9 +165,7 @@ void UNiagaraDataInterfaceSkeletalMesh::GetTriangleSamplingFunctions(TArray<FNia
 		Sig.Outputs.Add(FNiagaraVariable(FNiagaraTypeDefinition::GetVec3Def(), TEXT("Tangent")));
 		Sig.bMemberFunction = true;
 		Sig.bRequiresContext = false;
-#if WITH_EDITORONLY_DATA
 		Sig.Description = LOCTEXT("GetSkinnedDataWSDesc", "Returns skinning dependant data for the pased MeshTriCoord in world space. Interpolates between previous and current frame. All outputs are optional and you will incur zerp minimal cost if they are not connected.");
-#endif
 	}
 
 	{
@@ -243,9 +230,7 @@ void UNiagaraDataInterfaceSkeletalMesh::GetTriangleSamplingFunctions(TArray<FNia
 		Sig.Outputs.Add(FNiagaraVariable(FNiagaraTypeDefinition::GetIntDef(), TEXT("Vertex 2")));
 		Sig.bMemberFunction = true;
 		Sig.bRequiresContext = false;
-#if WITH_EDITORONLY_DATA
 		Sig.Description = LOCTEXT("GetTriCoordVetsName", "Takes the TriangleIndex from a MeshTriCoord and returns the vertices for that triangle.");
-#endif
 	}
 
 	{
@@ -363,6 +348,7 @@ void UNiagaraDataInterfaceSkeletalMesh::GetTriangleSamplingFunctions(TArray<FNia
 		Sig.bSupportsCPU = false;
 	}
 }
+#endif
 
 void UNiagaraDataInterfaceSkeletalMesh::BindTriangleSamplingFunction(const FVMExternalFunctionBindingInfo& BindingInfo, FNDISkeletalMesh_InstanceData* InstanceData, FVMExternalFunction &OutFunc)
 {
diff --git a/Engine/Plugins/FX/Niagara/Source/Niagara/Private/NDISkeletalMesh_VertexSampling.cpp b/Engine/Plugins/FX/Niagara/Source/Niagara/Private/NDISkeletalMesh_VertexSampling.cpp
--- a/Engine/Plugins/FX/Niagara/Source/Niagara/Private/NDISkeletalMesh_VertexSampling.cpp
+++ b/Engine/Plugins/FX/Niagara/Source/Niagara/Private/NDISkeletalMesh_VertexSampling.cpp
@@ -40,7 +40,8 @@ const FName FSkeletalMeshInterfaceHelper::RandomFilteredVertexName("RandomFilter
 const FName FSkeletalMeshInterfaceHelper::GetFilteredVertexCountName("GetFilteredVertexCount");
 const FName FSkeletalMeshInterfaceHelper::GetFilteredVertexAtName("GetFilteredVertex");
 
-void UNiagaraDataInterfaceSkeletalMesh::GetVertexSamplingFunctions(TArray<FNiagaraFunctionSignature>& OutFunctions)
+#if WITH_EDITORONLY_DATA
+void UNiagaraDataInterfaceSkeletalMesh::GetVertexSamplingFunctions(TArray<FNiagaraFunctionSignature>& OutFunctions) const
 {
 	{
 		FNiagaraFunctionSignature Sig;
@@ -53,9 +54,7 @@ void UNiagaraDataInterfaceSkeletalMesh::GetVertexSamplingFunctions(TArray<FNiaga
 		Sig.Outputs.Add(FNiagaraVariable(FNiagaraTypeDefinition::GetVec3Def(), TEXT("Tangent")));
 		Sig.bMemberFunction = true;
 		Sig.bRequiresContext = false;
-#if WITH_EDITORONLY_DATA
 		Sig.Description = LOCTEXT("GetVertexDataDesc", "Returns bind pose for the vertex");
-#endif
 		OutFunctions.Add(Sig);
 	}
 	{
@@ -70,9 +69,7 @@ void UNiagaraDataInterfaceSkeletalMesh::GetVertexSamplingFunctions(TArray<FNiaga
 		Sig.Outputs.Add(FNiagaraVariable(FNiagaraTypeDefinition::GetVec3Def(), TEXT("Tangent")));
 		Sig.bMemberFunction = true;
 		Sig.bRequiresContext = false;
-#if WITH_EDITORONLY_DATA
 		Sig.Description = LOCTEXT("GetSkinnedDataDesc", "Returns skinning dependant data for the pased vertex in local space. All outputs are optional and you will incur zero to minimal cost if they are not connected.");
-#endif
 		OutFunctions.Add(Sig);
 	}
 
@@ -88,9 +85,7 @@ void UNiagaraDataInterfaceSkeletalMesh::GetVertexSamplingFunctions(TArray<FNiaga
 		Sig.Outputs.Add(FNiagaraVariable(FNiagaraTypeDefinition::GetVec3Def(), TEXT("Tangent")));
 		Sig.bMemberFunction = true;
 		Sig.bRequiresContext = false;
-#if WITH_EDITORONLY_DATA
 		Sig.Description = LOCTEXT("GetSkinnedDataWSDesc", "Returns skinning dependant data for the pased vertex in world space. All outputs are optional and you will incur zero to minimal cost if they are not connected.");
-#endif
 		OutFunctions.Add(Sig);
 	}
 
@@ -107,9 +102,7 @@ void UNiagaraDataInterfaceSkeletalMesh::GetVertexSamplingFunctions(TArray<FNiaga
 		Sig.Outputs.Add(FNiagaraVariable(FNiagaraTypeDefinition::GetVec3Def(), TEXT("Tangent")));
 		Sig.bMemberFunction = true;
 		Sig.bRequiresContext = false;
-#if WITH_EDITORONLY_DATA
 		Sig.Description = LOCTEXT("GetSkinnedDataDesc", "Returns skinning dependant data for the pased vertex in local space. All outputs are optional and you will incur zero to minimal cost if they are not connected.");
-#endif
 		OutFunctions.Add(Sig);
 	}
 
@@ -126,9 +119,7 @@ void UNiagaraDataInterfaceSkeletalMesh::GetVertexSamplingFunctions(TArray<FNiaga
 		Sig.Outputs.Add(FNiagaraVariable(FNiagaraTypeDefinition::GetVec3Def(), TEXT("Tangent")));
 		Sig.bMemberFunction = true;
 		Sig.bRequiresContext = false;
-#if WITH_EDITORONLY_DATA
 		Sig.Description = LOCTEXT("GetSkinnedDataWSDesc", "Returns skinning dependant data for the pased vertex in world space. All outputs are optional and you will incur zero to minimal cost if they are not connected.");
-#endif
 		OutFunctions.Add(Sig);
 	}
 
@@ -226,6 +217,7 @@ void UNiagaraDataInterfaceSkeletalMesh::GetVertexSamplingFunctions(TArray<FNiaga
 		OutFunctions.Add(Sig);
 	}
 }
+#endif
 
 void UNiagaraDataInterfaceSkeletalMesh::BindVertexSamplingFunction(const FVMExternalFunctionBindingInfo& BindingInfo, FNDISkeletalMesh_InstanceData* InstanceData, FVMExternalFunction &OutFunc)
 {
diff --git a/Engine/Plugins/FX/Niagara/Source/Niagara/Private/NiagaraDataInterface.cpp b/Engine/Plugins/FX/Niagara/Source/Niagara/Private/NiagaraDataInterface.cpp
--- a/Engine/Plugins/FX/Niagara/Source/Niagara/Private/NiagaraDataInterface.cpp
+++ b/Engine/Plugins/FX/Niagara/Source/Niagara/Private/NiagaraDataInterface.cpp
@@ -279,10 +279,24 @@ void UNiagaraDataInterface::GetFeedback(UNiagaraDataInterface* DataInterface, TA
 	DataInterface->GetFeedback(Asset, Component, Errors, Warnings, Info);
 }
 
+void UNiagaraDataInterface::GetFunctionSignatures(TArray<FNiagaraFunctionSignature>& OutFunctions) const
+{
+#if WITH_EDITORONLY_DATA
+	// Until we can eliminate GetFunctions() we just call it, and let it's base implementation call
+	// GetFunctionsInternal().  When GetFunctions() is deleted we can just call GetFunctionsInternal()
+	// directly here.
+	//GetFunctionsInternal(OutFunctions);
+
+	PRAGMA_DISABLE_DEPRECATION_WARNINGS // to be removed when GetFunctions is removed
+	const_cast<UNiagaraDataInterface*>(this)->GetFunctions(OutFunctions);
+	PRAGMA_ENABLE_DEPRECATION_WARNINGS // to be removed when GetFunctions is removed
+#endif
+}
+
 void UNiagaraDataInterface::ValidateFunction(const FNiagaraFunctionSignature& Function, TArray<FText>& OutValidationErrors)
 {
 	TArray<FNiagaraFunctionSignature> DIFuncs;
-	GetFunctions(DIFuncs);
+	GetFunctionSignatures(DIFuncs);
 
 	if (!DIFuncs.ContainsByPredicate([&](const FNiagaraFunctionSignature& Sig) { return Sig.EqualsIgnoringSpecifiers(Function); }))
 	{
diff --git a/Engine/Plugins/FX/Niagara/Source/Niagara/Private/NiagaraDataInterface2DArrayTexture.cpp b/Engine/Plugins/FX/Niagara/Source/Niagara/Private/NiagaraDataInterface2DArrayTexture.cpp
--- a/Engine/Plugins/FX/Niagara/Source/Niagara/Private/NiagaraDataInterface2DArrayTexture.cpp
+++ b/Engine/Plugins/FX/Niagara/Source/Niagara/Private/NiagaraDataInterface2DArrayTexture.cpp
@@ -104,7 +104,8 @@ bool UNiagaraDataInterface2DArrayTexture::Equals(const UNiagaraDataInterface* Ot
 		OtherTexture->TextureUserParameter == TextureUserParameter;
 }
 
-void UNiagaraDataInterface2DArrayTexture::GetFunctions(TArray<FNiagaraFunctionSignature>& OutFunctions)
+#if WITH_EDITORONLY_DATA
+void UNiagaraDataInterface2DArrayTexture::GetFunctionsInternal(TArray<FNiagaraFunctionSignature>& OutFunctions) const
 {
 	OutFunctions.Reserve(OutFunctions.Num() + 3);
 
@@ -153,6 +154,7 @@ void UNiagaraDataInterface2DArrayTexture::GetFunctions(TArray<FNiagaraFunctionSi
 		Sig.SetDescription(LOCTEXT("TextureDimsDesc", "Get the dimensions of mip 0 of the texture."));
 	}
 }
+#endif
 
 void UNiagaraDataInterface2DArrayTexture::GetVMExternalFunction(const FVMExternalFunctionBindingInfo& BindingInfo, void* InstanceData, FVMExternalFunction &OutFunc)
 {
diff --git a/Engine/Plugins/FX/Niagara/Source/Niagara/Private/NiagaraDataInterfaceAudioOscilloscope.cpp b/Engine/Plugins/FX/Niagara/Source/Niagara/Private/NiagaraDataInterfaceAudioOscilloscope.cpp
--- a/Engine/Plugins/FX/Niagara/Source/Niagara/Private/NiagaraDataInterfaceAudioOscilloscope.cpp
+++ b/Engine/Plugins/FX/Niagara/Source/Niagara/Private/NiagaraDataInterfaceAudioOscilloscope.cpp
@@ -184,9 +184,10 @@ void UNiagaraDataInterfaceAudioOscilloscope::GetNumChannels(FVectorVMExternalFun
 	}
 }
 
-void UNiagaraDataInterfaceAudioOscilloscope::GetFunctions(TArray<FNiagaraFunctionSignature>& OutFunctions)
+#if WITH_EDITORONLY_DATA
+void UNiagaraDataInterfaceAudioOscilloscope::GetFunctionsInternal(TArray<FNiagaraFunctionSignature>& OutFunctions) const
 {
-	Super::GetFunctions(OutFunctions);
+	Super::GetFunctionsInternal(OutFunctions);
 
 	{
 		FNiagaraFunctionSignature SampleAudioBufferSignature;
@@ -212,6 +213,7 @@ void UNiagaraDataInterfaceAudioOscilloscope::GetFunctions(TArray<FNiagaraFunctio
 		OutFunctions.Add(NumChannelsSignature);
 	}
 }
+#endif
 
 DEFINE_NDI_DIRECT_FUNC_BINDER(UNiagaraDataInterfaceAudioOscilloscope, SampleAudio);
 DEFINE_NDI_DIRECT_FUNC_BINDER(UNiagaraDataInterfaceAudioOscilloscope, GetNumChannels);
diff --git a/Engine/Plugins/FX/Niagara/Source/Niagara/Private/NiagaraDataInterfaceAudioPlayer.cpp b/Engine/Plugins/FX/Niagara/Source/Niagara/Private/NiagaraDataInterfaceAudioPlayer.cpp
--- a/Engine/Plugins/FX/Niagara/Source/Niagara/Private/NiagaraDataInterfaceAudioPlayer.cpp
+++ b/Engine/Plugins/FX/Niagara/Source/Niagara/Private/NiagaraDataInterfaceAudioPlayer.cpp
@@ -69,7 +69,7 @@ bool UNiagaraDataInterfaceAudioPlayer::UpgradeFunctionCall(FNiagaraFunctionSigna
 	if (FunctionSignature.FunctionVersion < FNiagaraAudioPlayerDIFunctionVersion::LatestVersion)
 	{
 		TArray<FNiagaraFunctionSignature> AllFunctions;
-		GetFunctions(AllFunctions);
+		GetFunctionsInternal(AllFunctions);
 		for (const FNiagaraFunctionSignature& Sig : AllFunctions)
 		{
 			if (FunctionSignature.Name == Sig.Name)
@@ -321,7 +321,8 @@ bool UNiagaraDataInterfaceAudioPlayer::Equals(const UNiagaraDataInterface* Other
 	return OtherPlayer->SoundToPlay == SoundToPlay && OtherPlayer->Attenuation == Attenuation && OtherPlayer->Concurrency == Concurrency && OtherPlayer->ConfigurationUserParameter == ConfigurationUserParameter && OtherPlayer->bLimitPlaysPerTick == bLimitPlaysPerTick && OtherPlayer->MaxPlaysPerTick == MaxPlaysPerTick;
 }
 
-void UNiagaraDataInterfaceAudioPlayer::GetFunctions(TArray<FNiagaraFunctionSignature>& OutFunctions)
+#if WITH_EDITORONLY_DATA
+void UNiagaraDataInterfaceAudioPlayer::GetFunctionsInternal(TArray<FNiagaraFunctionSignature>& OutFunctions) const
 {
 	FNiagaraFunctionSignature Sig;
 	Sig.Name = PlayAudioName;
@@ -528,6 +529,7 @@ void UNiagaraDataInterfaceAudioPlayer::GetFunctions(TArray<FNiagaraFunctionSigna
 	Sig.Inputs.Add(FNiagaraVariable(FNiagaraTypeDefinition::GetBoolDef(), TEXT("Pause Audio")));
 	OutFunctions.Add(Sig);
 }
+#endif
 
 DEFINE_NDI_DIRECT_FUNC_BINDER(UNiagaraDataInterfaceAudioPlayer, PlayOneShotAudio);
 DEFINE_NDI_DIRECT_FUNC_BINDER(UNiagaraDataInterfaceAudioPlayer, PlayPersistentAudio);
diff --git a/Engine/Plugins/FX/Niagara/Source/Niagara/Private/NiagaraDataInterfaceAudioSpectrum.cpp b/Engine/Plugins/FX/Niagara/Source/Niagara/Private/NiagaraDataInterfaceAudioSpectrum.cpp
--- a/Engine/Plugins/FX/Niagara/Source/Niagara/Private/NiagaraDataInterfaceAudioSpectrum.cpp
+++ b/Engine/Plugins/FX/Niagara/Source/Niagara/Private/NiagaraDataInterfaceAudioSpectrum.cpp
@@ -596,9 +596,10 @@ void UNiagaraDataInterfaceAudioSpectrum::GetNumChannels(FVectorVMExternalFunctio
 	}
 }
 
-void UNiagaraDataInterfaceAudioSpectrum::GetFunctions(TArray<FNiagaraFunctionSignature>& OutFunctions)
+#if WITH_EDITORONLY_DATA
+void UNiagaraDataInterfaceAudioSpectrum::GetFunctionsInternal(TArray<FNiagaraFunctionSignature>& OutFunctions) const
 {
-	Super::GetFunctions(OutFunctions);
+	Super::GetFunctionsInternal(OutFunctions);
 
 	{
 		FNiagaraFunctionSignature GetSpectrumSignature;
@@ -623,9 +624,8 @@ void UNiagaraDataInterfaceAudioSpectrum::GetFunctions(TArray<FNiagaraFunctionSig
 		NumChannelsSignature.bRequiresContext = false;
 		OutFunctions.Add(NumChannelsSignature);
 	}
-
-
 }
+#endif
 
 DEFINE_NDI_DIRECT_FUNC_BINDER(UNiagaraDataInterfaceAudioSpectrum, GetSpectrumValue);
 DEFINE_NDI_DIRECT_FUNC_BINDER(UNiagaraDataInterfaceAudioSpectrum, GetNumChannels);
diff --git a/Engine/Plugins/FX/Niagara/Source/Niagara/Private/NiagaraDataInterfaceCamera.cpp b/Engine/Plugins/FX/Niagara/Source/Niagara/Private/NiagaraDataInterfaceCamera.cpp
--- a/Engine/Plugins/FX/Niagara/Source/Niagara/Private/NiagaraDataInterfaceCamera.cpp
+++ b/Engine/Plugins/FX/Niagara/Source/Niagara/Private/NiagaraDataInterfaceCamera.cpp
@@ -141,14 +141,13 @@ bool UNiagaraDataInterfaceCamera::PerInstanceTick(void* PerInstanceData, FNiagar
 	return false;
 }
 
-void UNiagaraDataInterfaceCamera::GetFunctions(TArray<FNiagaraFunctionSignature>& OutFunctions)
+#if WITH_EDITORONLY_DATA
+void UNiagaraDataInterfaceCamera::GetFunctionsInternal(TArray<FNiagaraFunctionSignature>& OutFunctions) const
 {
 	FNiagaraFunctionSignature Sig;
 	Sig.Name = GetViewPropertiesName;
-#if WITH_EDITORONLY_DATA
 	Sig.Description = LOCTEXT("GetViewPropertiesDescription", "This function returns the properties of the current view. Only valid for gpu particles.");
 	Sig.FunctionVersion = FNiagaraCameraDIFunctionVersion::LatestVersion;
-#endif
 	Sig.bMemberFunction = true;
 	Sig.bRequiresContext = false;
 	Sig.bSupportsCPU = false;
@@ -169,10 +168,8 @@ void UNiagaraDataInterfaceCamera::GetFunctions(TArray<FNiagaraFunctionSignature>
 
 	Sig = FNiagaraFunctionSignature();
 	Sig.Name = GetClipSpaceTransformsName;
-#if WITH_EDITORONLY_DATA
 	Sig.Description = LOCTEXT("GetClipSpaceTransformsDescription", "This function returns the clip transforms for the current view. Only valid for gpu particles.");
 	Sig.FunctionVersion = FNiagaraCameraDIFunctionVersion::LatestVersion;
-#endif
 	Sig.bMemberFunction = true;
 	Sig.bRequiresContext = false;
 	Sig.bSupportsCPU = false;
@@ -190,10 +187,8 @@ void UNiagaraDataInterfaceCamera::GetFunctions(TArray<FNiagaraFunctionSignature>
 
 	Sig = FNiagaraFunctionSignature();
 	Sig.Name = GetViewSpaceTransformsName;
-#if WITH_EDITORONLY_DATA
 	Sig.Description = LOCTEXT("GetViewSpaceTransformsDescription", "This function returns the relevant transforms for the current view. Only valid for gpu particles.");
 	Sig.FunctionVersion = FNiagaraCameraDIFunctionVersion::LatestVersion;
-#endif
 	Sig.bMemberFunction = true;
 	Sig.bRequiresContext = false;
 	Sig.bSupportsCPU = false;
@@ -209,10 +204,8 @@ void UNiagaraDataInterfaceCamera::GetFunctions(TArray<FNiagaraFunctionSignature>
 	
 	Sig = FNiagaraFunctionSignature();
 	Sig.Name = ApplyPreViewTranslationToPosition;
-#if WITH_EDITORONLY_DATA
 	Sig.Description = LOCTEXT("ApplyPreViewTranslationToPositionDescription", "This function applies the pre view translation to a lwc positions such as Particles.Position. The resulting vector can be used with any of the 'Translated X to Y' transform matrices.");
 	Sig.FunctionVersion = FNiagaraCameraDIFunctionVersion::LatestVersion;
-#endif
 	Sig.bMemberFunction = true;
 	Sig.bRequiresContext = false;
 	Sig.bSupportsCPU = false;
@@ -223,10 +216,8 @@ void UNiagaraDataInterfaceCamera::GetFunctions(TArray<FNiagaraFunctionSignature>
 
 	Sig = FNiagaraFunctionSignature();
 	Sig.Name = RemovePreViewTranslationFromPosition;
-#if WITH_EDITORONLY_DATA
 	Sig.Description = LOCTEXT("RemovePreViewTranslationFromPositionDescription", "This function removes the pre view translation from a vector to return a lwc position. Use the result of a 'X to Translated Y' transform matrix multiplication as input vector.");
 	Sig.FunctionVersion = FNiagaraCameraDIFunctionVersion::LatestVersion;
-#endif
 	Sig.bMemberFunction = true;
 	Sig.bRequiresContext = false;
 	Sig.bSupportsCPU = false;
@@ -238,10 +229,8 @@ void UNiagaraDataInterfaceCamera::GetFunctions(TArray<FNiagaraFunctionSignature>
 
 	Sig = FNiagaraFunctionSignature();
 	Sig.Name = GetFieldOfViewName;
-#if WITH_EDITORONLY_DATA
 	Sig.Description = LOCTEXT("GetNiagaraFOVDescription", "This function returns the field of view angle (in degrees) for the active camera. For gpu particles this returns the x axis fov.");
 	Sig.FunctionVersion = FNiagaraCameraDIFunctionVersion::LatestVersion;
-#endif
 	Sig.bMemberFunction = true;
 	Sig.bRequiresContext = false;
 	Sig.AddInput(FNiagaraVariable(FNiagaraTypeDefinition(GetClass()), TEXT("Camera interface")));
@@ -251,10 +240,8 @@ void UNiagaraDataInterfaceCamera::GetFunctions(TArray<FNiagaraFunctionSignature>
 
 	Sig = FNiagaraFunctionSignature();
 	Sig.Name = GetCameraPropertiesName;
-#if WITH_EDITORONLY_DATA
 	Sig.Description = LOCTEXT("GetCameraPositionDescription", "This function returns the position of the currently active camera.");
 	Sig.FunctionVersion = FNiagaraCameraDIFunctionVersion::LatestVersion;
-#endif
 	Sig.bMemberFunction = true;
 	Sig.bRequiresContext = false;
 	Sig.AddInput(FNiagaraVariable(FNiagaraTypeDefinition(GetClass()), TEXT("Camera interface")));
@@ -266,10 +253,8 @@ void UNiagaraDataInterfaceCamera::GetFunctions(TArray<FNiagaraFunctionSignature>
 
 	Sig = FNiagaraFunctionSignature();
 	Sig.Name = QueryClosestName;
-#if WITH_EDITORONLY_DATA
 	Sig.Description = LOCTEXT("QueryClosestDescription", "This function checks the previously calculated distance of each particle and then returns true for the closest particles and false for the other ones.\nThis function needs to be paired with CalculateParticleDistancesCPU to work correctly.");
 	Sig.FunctionVersion = FNiagaraCameraDIFunctionVersion::LatestVersion;
-#endif
 	Sig.bMemberFunction = true;
 	Sig.bRequiresContext = false;
 	Sig.bSupportsGPU = false;
@@ -281,10 +266,8 @@ void UNiagaraDataInterfaceCamera::GetFunctions(TArray<FNiagaraFunctionSignature>
 
 	Sig = FNiagaraFunctionSignature();
 	Sig.Name = CalculateDistancesName;
-#if WITH_EDITORONLY_DATA
 	Sig.Description = LOCTEXT("CalculateDistancesDescription", "This function compares the particle position against the camera position and stores the result to be queried in the next frame.\nThe results can then be queried with QueryClosestParticlesCPU.");
 	Sig.FunctionVersion = FNiagaraCameraDIFunctionVersion::LatestVersion;
-#endif
 	Sig.bMemberFunction = true;
 	Sig.bRequiresContext = false;
 	Sig.bSupportsGPU = false;
@@ -296,10 +279,8 @@ void UNiagaraDataInterfaceCamera::GetFunctions(TArray<FNiagaraFunctionSignature>
 
 	Sig = FNiagaraFunctionSignature();
 	Sig.Name = GetTAAJitterName;
-#if WITH_EDITORONLY_DATA
 	Sig.Description = NSLOCTEXT("Niagara", "GetTAAJitterDescription", "This function returns the TAA jitter values of the currently active camera.");
 	Sig.FunctionVersion = FNiagaraCameraDIFunctionVersion::LatestVersion;
-#endif
 	Sig.bMemberFunction = true;
 	Sig.bRequiresContext = false;
 	Sig.Inputs.Add(FNiagaraVariable(FNiagaraTypeDefinition(GetClass()), TEXT("Camera interface")));
@@ -308,7 +289,6 @@ void UNiagaraDataInterfaceCamera::GetFunctions(TArray<FNiagaraFunctionSignature>
 	OutFunctions.Add(Sig);
 }
 
-#if WITH_EDITORONLY_DATA
 bool UNiagaraDataInterfaceCamera::GetFunctionHLSL(const FNiagaraDataInterfaceGPUParamInfo& ParamInfo, const FNiagaraDataInterfaceGeneratedFunction& FunctionInfo, int FunctionInstanceIndex, FString& OutHLSL)
 {
 	if ((FunctionInfo.DefinitionName == GetViewPropertiesName) ||
@@ -356,7 +336,7 @@ bool UNiagaraDataInterfaceCamera::UpgradeFunctionCall(FNiagaraFunctionSignature&
 	if (FunctionSignature.FunctionVersion < FNiagaraCameraDIFunctionVersion::LWCConversion)
 	{
 		TArray<FNiagaraFunctionSignature> AllFunctions;
-		GetFunctions(AllFunctions);
+		GetFunctionsInternal(AllFunctions);
 		for (const FNiagaraFunctionSignature& Sig : AllFunctions)
 		{
 			if (FunctionSignature.Name == Sig.Name)
diff --git a/Engine/Plugins/FX/Niagara/Source/Niagara/Private/NiagaraDataInterfaceCollisionQuery.cpp b/Engine/Plugins/FX/Niagara/Source/Niagara/Private/NiagaraDataInterfaceCollisionQuery.cpp
--- a/Engine/Plugins/FX/Niagara/Source/Niagara/Private/NiagaraDataInterfaceCollisionQuery.cpp
+++ b/Engine/Plugins/FX/Niagara/Source/Niagara/Private/NiagaraDataInterfaceCollisionQuery.cpp
@@ -202,17 +202,16 @@ void UNiagaraDataInterfaceCollisionQuery::GetAssetTagsForContext(const UObject*
 	Super::GetAssetTagsForContext(InAsset, AssetVersion, InProperties, NumericKeys, StringKeys);
 }
 
-void UNiagaraDataInterfaceCollisionQuery::GetFunctions(TArray<FNiagaraFunctionSignature>& OutFunctions)
+#if WITH_EDITORONLY_DATA
+void UNiagaraDataInterfaceCollisionQuery::GetFunctionsInternal(TArray<FNiagaraFunctionSignature>& OutFunctions) const
 {
 	{
 		FNiagaraFunctionSignature SigDepth;
 		SigDepth.Name = NDICollisionQueryLocal::SceneDepthName;
 		SigDepth.bMemberFunction = true;
 		SigDepth.bSupportsCPU = false;
-#if WITH_EDITORONLY_DATA
 		SigDepth.FunctionVersion = FNiagaraCollisionDIFunctionVersion::LatestVersion;
 		SigDepth.Description = LOCTEXT("SceneDepthSignatureDescription", "Projects a given world position to view space and then queries the depth buffer with that position.");
-#endif
 		const FText DepthSamplePosWorldDescription = LOCTEXT("DepthSamplePosWorldDescription", "The world position where the depth should be queried. The position gets automatically transformed to view space to query the depth buffer.");
 		const FText SceneDepthDescription = LOCTEXT("SceneDepthDescription", "If the query was successful this returns the scene depth, otherwise -1.");
 		const FText CameraPosWorldDescription = LOCTEXT("CameraPosWorldDescription", "Returns the current camera position in world space.");
@@ -235,11 +234,9 @@ void UNiagaraDataInterfaceCollisionQuery::GetFunctions(TArray<FNiagaraFunctionSi
 		SigPartialDepth.bMemberFunction = true;
 		SigPartialDepth.bSupportsCPU = false;
 		SigPartialDepth.MiscUsageBitMask = (uint16)ENiagaraScriptMiscUsageMask::UsesPartialDepthCollisionQuery;
-#if WITH_EDITORONLY_DATA
 		SigPartialDepth.FunctionVersion = FNiagaraCollisionDIFunctionVersion::LatestVersion;
 		SigPartialDepth.Description = LOCTEXT("ScenePartialDepthSignatureDescription", "Projects a given world position to view space and then queries the partial depth buffer (opaque emitter using this function are not in this depth buffer) with that position.");
 		SigPartialDepth.ExperimentalMessage = LOCTEXT("ScenePartialDepthSignatureEpxMessage", "Careful: using this node involve a scene depth copy that might slow down the base pass on some platform.");
-#endif
 		SigPartialDepth.bExperimental = true;
 		SigPartialDepth.Inputs.Add(FNiagaraVariable(FNiagaraTypeDefinition(GetClass()), TEXT("CollisionQuery")));
 		SigPartialDepth.AddInput(FNiagaraVariable(FNiagaraTypeDefinition::GetPositionDef(), TEXT("DepthSamplePosWorld")), DepthSamplePosWorldDescription);
@@ -255,10 +252,8 @@ void UNiagaraDataInterfaceCollisionQuery::GetFunctions(TArray<FNiagaraFunctionSi
 		SigCustomDepth.Name = NDICollisionQueryLocal::CustomDepthName;
 		SigCustomDepth.bMemberFunction = true;
 		SigCustomDepth.bSupportsCPU = false;
-#if WITH_EDITORONLY_DATA
 		SigCustomDepth.FunctionVersion = FNiagaraCollisionDIFunctionVersion::LatestVersion;
 		SigCustomDepth.Description = LOCTEXT("CustomDepthDescription", "Projects a given world position to view space and then queries the custom depth buffer with that position.");
-#endif
 		SigCustomDepth.Inputs.Add(FNiagaraVariable(FNiagaraTypeDefinition(GetClass()), TEXT("CollisionQuery")));
 		SigCustomDepth.AddInput(FNiagaraVariable(FNiagaraTypeDefinition::GetPositionDef(), TEXT("DepthSamplePosWorld")), DepthSamplePosWorldDescription);
 		SigCustomDepth.AddOutput(FNiagaraVariable(FNiagaraTypeDefinition::GetFloatDef(), TEXT("SceneDepth")), SceneDepthDescription);
@@ -274,10 +269,8 @@ void UNiagaraDataInterfaceCollisionQuery::GetFunctions(TArray<FNiagaraFunctionSi
 		SigMeshField.Name = NDICollisionQueryLocal::DistanceFieldName;
 		SigMeshField.bMemberFunction = true;
 		SigMeshField.bSupportsCPU = false;
-#if WITH_EDITORONLY_DATA
 		SigMeshField.FunctionVersion = FNiagaraCollisionDIFunctionVersion::LatestVersion;
 		SigMeshField.Description = LOCTEXT("DistanceFieldDescription", "Queries the global distance field for a given world position.\nPlease note that the distance field resolution gets lower the farther away the queried position is from the camera.");
-#endif
 		const FText FieldSamplePosWorldDescription = LOCTEXT("FieldSamplePosWorldDescription", "The world position where the distance field should be queried.");
 		const FText DistanceToNearestSurfaceDescription = LOCTEXT("DistanceToNearestSurfaceDescription", "If the query was successful this returns the distance to the nearest surface, otherwise returns 0.");
 		const FText FieldGradientDescription = LOCTEXT("FieldGradientDescription", "If the query was successful this returns the non-normalized direction to the nearest surface, otherwise returns (0, 0, 0).");
@@ -307,10 +300,8 @@ void UNiagaraDataInterfaceCollisionQuery::GetFunctions(TArray<FNiagaraFunctionSi
 		IssueRayTrace.bMemberFunction = true;
 		IssueRayTrace.bSupportsCPU = false;
 		IssueRayTrace.bSoftDeprecatedFunction = true;
-#if WITH_EDITORONLY_DATA
 		IssueRayTrace.FunctionVersion = FNiagaraCollisionDIFunctionVersion::LatestVersion;
 		IssueRayTrace.Description = LOCTEXT("IssueAsync_RayTraceDescription", "Enqueues a GPU raytrace with the result being available the following frame");
-#endif
 		IssueRayTrace.AddInput(FNiagaraVariable(FNiagaraTypeDefinition(GetClass()), TEXT("CollisionQuery")));
 		IssueRayTrace.AddInput(FNiagaraVariable(FNiagaraTypeDefinition::GetIntDef(), TEXT("QueryID")), QueryIDDescription);
 		IssueRayTrace.AddInput(FNiagaraVariable(FNiagaraTypeDefinition::GetPositionDef(), TEXT("TraceStartWorld")), TraceStartWorldDescription);
@@ -324,10 +315,8 @@ void UNiagaraDataInterfaceCollisionQuery::GetFunctions(TArray<FNiagaraFunctionSi
 		CreateRayTrace.bMemberFunction = true;
 		CreateRayTrace.bSupportsCPU = false;
 		CreateRayTrace.bSoftDeprecatedFunction = true;
-#if WITH_EDITORONLY_DATA
 		CreateRayTrace.FunctionVersion = FNiagaraCollisionDIFunctionVersion::LatestVersion;
 		CreateRayTrace.Description = LOCTEXT("CreateAsync_RayTraceDescription", "Creates a GPU raytrace with the result being available the following frame (index is returned)");
-#endif
 		CreateRayTrace.AddInput(FNiagaraVariable(FNiagaraTypeDefinition(GetClass()), TEXT("CollisionQuery")));
 		CreateRayTrace.AddInput(FNiagaraVariable(FNiagaraTypeDefinition::GetPositionDef(), TEXT("TraceStartWorld")), TraceStartWorldDescription);
 		CreateRayTrace.AddInput(FNiagaraVariable(FNiagaraTypeDefinition::GetPositionDef(), TEXT("TraceEndWorld")), TraceEndWorldDescription);
@@ -343,10 +332,8 @@ void UNiagaraDataInterfaceCollisionQuery::GetFunctions(TArray<FNiagaraFunctionSi
 		ReserveRayTrace.bMemberFunction = true;
 		ReserveRayTrace.bSupportsCPU = false;
 		ReserveRayTrace.bSoftDeprecatedFunction = true;
-#if WITH_EDITORONLY_DATA
 		ReserveRayTrace.FunctionVersion = FNiagaraCollisionDIFunctionVersion::LatestVersion;
 		ReserveRayTrace.Description = LOCTEXT("ReserveAsync_RayTraceDescription", "Reserves a number of ray trace request slots");
-#endif
 		ReserveRayTrace.AddInput(FNiagaraVariable(FNiagaraTypeDefinition(GetClass()), TEXT("CollisionQuery")));
 		ReserveRayTrace.AddInput(FNiagaraVariable(FNiagaraTypeDefinition::GetIntDef(), TEXT("TraceCount")), LOCTEXT("ReserveAsync_QueryIDDescription", "Number of async raytrace requests to be reserved"));
 		ReserveRayTrace.AddOutput(FNiagaraVariable(FNiagaraTypeDefinition::GetIntDef(), TEXT("FirstQueryID")), LOCTEXT("ReserveAsync_TraceChannelDescription", "The first index in the block reserved through this call"));
@@ -359,10 +346,8 @@ void UNiagaraDataInterfaceCollisionQuery::GetFunctions(TArray<FNiagaraFunctionSi
 		ReadRayTrace.bMemberFunction = true;
 		ReadRayTrace.bSupportsCPU = false;
 		ReadRayTrace.bSoftDeprecatedFunction = true;
-#if WITH_EDITORONLY_DATA
 		ReadRayTrace.FunctionVersion = FNiagaraCollisionDIFunctionVersion::LatestVersion;
 		ReadRayTrace.Description = LOCTEXT("ReadAsync_RayTraceDescription", "Reads the results of a previously enqueued GPU ray trace");
-#endif
 		ReadRayTrace.AddInput(FNiagaraVariable(FNiagaraTypeDefinition(GetClass()), TEXT("CollisionQuery")));
 		ReadRayTrace.AddInput(FNiagaraVariable(FNiagaraTypeDefinition::GetIntDef(), TEXT("PreviousFrameQueryID")), LOCTEXT("ReadAsync_PreviousFrameQueryIDDescription", "The index of the results being retrieved"));
 		ReadRayTrace.AddOutput(FNiagaraVariable(FNiagaraTypeDefinition::GetBoolDef(), TEXT("CollisionValid")), LOCTEXT("ReadAsync_CollisionValidDescription", "Returns true if a Hit was encountered"));
@@ -387,11 +372,8 @@ void UNiagaraDataInterfaceCollisionQuery::GetFunctions(TArray<FNiagaraFunctionSi
 		SigCpuSync.Name = NDICollisionQueryLocal::SyncTraceName;
 		SigCpuSync.bMemberFunction = true;
 		SigCpuSync.bSupportsGPU = false;
-#if WITH_EDITORONLY_DATA
 		SigCpuSync.FunctionVersion = FNiagaraCollisionDIFunctionVersion::LatestVersion;
 		SigCpuSync.Description = LOCTEXT("SigCpuSyncDescription", "Traces a ray against the world using a specific channel and return the first blocking hit.");
-#endif
-	
 		SigCpuSync.Inputs.Add(FNiagaraVariable(FNiagaraTypeDefinition(GetClass()), TEXT("CollisionQuery")));
 		SigCpuSync.AddInput(FNiagaraVariable(FNiagaraTypeDefinition::GetPositionDef(), TEXT("TraceStartWorld")), TraceStartWorldDescription);
 		SigCpuSync.AddInput(FNiagaraVariable(FNiagaraTypeDefinition::GetPositionDef(), TEXT("TraceEndWorld")), TraceEndWorldDescription);
@@ -410,11 +392,8 @@ void UNiagaraDataInterfaceCollisionQuery::GetFunctions(TArray<FNiagaraFunctionSi
 		SigCpuAsync.Name = NDICollisionQueryLocal::AsyncTraceName;
 		SigCpuAsync.bMemberFunction = true;
 		SigCpuAsync.bSupportsGPU = false;
-#if WITH_EDITORONLY_DATA
 		SigCpuAsync.FunctionVersion = FNiagaraCollisionDIFunctionVersion::LatestVersion;
 		SigCpuAsync.Description = LOCTEXT("SigCpuAsyncDescription", "Traces a ray against the world using a specific channel and return the first blocking hit the next frame.\nNote that this is the ASYNC version of the trace function, meaning it will not returns the result right away, but with one frame latency.");
-#endif
-	
 		SigCpuAsync.Inputs.Add(FNiagaraVariable(FNiagaraTypeDefinition(GetClass()), TEXT("CollisionQuery")));
 		SigCpuAsync.AddInput(FNiagaraVariable(FNiagaraTypeDefinition::GetIntDef(), TEXT("PreviousFrameQueryID")), LOCTEXT("TraceAsync_PreviousFrameQueryIDDescription", "The query ID returned from the last frame's async trace call.\nRegardless if it is a valid ID or not this function call with issue a new async line trace, but it will only return results with a valid ID."));
 		SigCpuAsync.AddInput(FNiagaraVariable(FNiagaraTypeDefinition::GetPositionDef(), TEXT("TraceStartWorld")), TraceStartWorldDescription);
@@ -436,9 +415,6 @@ void UNiagaraDataInterfaceCollisionQuery::GetFunctions(TArray<FNiagaraFunctionSi
 // build the shader function HLSL; function name is passed in, as it's defined per-DI; that way, configuration could change
 // the HLSL in the spirit of a static switch
 // TODO: need a way to identify each specific function here
-
-// 
-#if WITH_EDITORONLY_DATA
 bool UNiagaraDataInterfaceCollisionQuery::GetFunctionHLSL(const FNiagaraDataInterfaceGPUParamInfo& ParamInfo, const FNiagaraDataInterfaceGeneratedFunction& FunctionInfo, int FunctionInstanceIndex, FString& OutHLSL)
 {
 	if ( (FunctionInfo.DefinitionName == NDICollisionQueryLocal::SceneDepthName) ||
@@ -464,7 +440,7 @@ bool UNiagaraDataInterfaceCollisionQuery::UpgradeFunctionCall(FNiagaraFunctionSi
 	if (FunctionSignature.FunctionVersion < FNiagaraCollisionDIFunctionVersion::LatestVersion)
 	{
 		TArray<FNiagaraFunctionSignature> AllFunctions;
-		GetFunctions(AllFunctions);
+		GetFunctionsInternal(AllFunctions);
 		for (const FNiagaraFunctionSignature& Sig : AllFunctions)
 		{
 			if (FunctionSignature.Name == Sig.Name)
diff --git a/Engine/Plugins/FX/Niagara/Source/Niagara/Private/NiagaraDataInterfaceColorCurve.cpp b/Engine/Plugins/FX/Niagara/Source/Niagara/Private/NiagaraDataInterfaceColorCurve.cpp
--- a/Engine/Plugins/FX/Niagara/Source/Niagara/Private/NiagaraDataInterfaceColorCurve.cpp
+++ b/Engine/Plugins/FX/Niagara/Source/Niagara/Private/NiagaraDataInterfaceColorCurve.cpp
@@ -163,7 +163,8 @@ void UNiagaraDataInterfaceColorCurve::GetCurveData(TArray<FCurveData>& OutCurveD
 	OutCurveData.Add(FCurveData(&AlphaCurve, TEXT("Alpha"), FLinearColor(0.2f, 0.2f, 0.2f)));
 }
 
-void UNiagaraDataInterfaceColorCurve::GetFunctions(TArray<FNiagaraFunctionSignature>& OutFunctions)
+#if WITH_EDITORONLY_DATA
+void UNiagaraDataInterfaceColorCurve::GetFunctionsInternal(TArray<FNiagaraFunctionSignature>& OutFunctions) const
 {
 	FNiagaraFunctionSignature Sig;
 	Sig.Name = SampleCurveName;
@@ -181,7 +182,6 @@ void UNiagaraDataInterfaceColorCurve::GetFunctions(TArray<FNiagaraFunctionSignat
 // the HLSL in the spirit of a static switch
 // TODO: need a way to identify each specific function here
 // 
-#if WITH_EDITORONLY_DATA
 void UNiagaraDataInterfaceColorCurve::SyncCurvesToAsset()
 {
 	if (UCurveLinearColor* ColorCurve = Cast<UCurveLinearColor>(CurveAsset))
diff --git a/Engine/Plugins/FX/Niagara/Source/Niagara/Private/NiagaraDataInterfaceCubeTexture.cpp b/Engine/Plugins/FX/Niagara/Source/Niagara/Private/NiagaraDataInterfaceCubeTexture.cpp
--- a/Engine/Plugins/FX/Niagara/Source/Niagara/Private/NiagaraDataInterfaceCubeTexture.cpp
+++ b/Engine/Plugins/FX/Niagara/Source/Niagara/Private/NiagaraDataInterfaceCubeTexture.cpp
@@ -102,7 +102,8 @@ bool UNiagaraDataInterfaceCubeTexture::Equals(const UNiagaraDataInterface* Other
 		OtherTexture->TextureUserParameter == TextureUserParameter;
 }
 
-void UNiagaraDataInterfaceCubeTexture::GetFunctions(TArray<FNiagaraFunctionSignature>& OutFunctions)
+#if WITH_EDITORONLY_DATA
+void UNiagaraDataInterfaceCubeTexture::GetFunctionsInternal(TArray<FNiagaraFunctionSignature>& OutFunctions) const
 {
 	{
 		FNiagaraFunctionSignature& Sig = OutFunctions.AddDefaulted_GetRef();
@@ -129,6 +130,7 @@ void UNiagaraDataInterfaceCubeTexture::GetFunctions(TArray<FNiagaraFunctionSigna
 		Sig.Outputs.Add(FNiagaraVariable(FNiagaraTypeDefinition::GetIntDef(), TEXT("Height")));
 	}
 }
+#endif
 
 DEFINE_NDI_DIRECT_FUNC_BINDER(UNiagaraDataInterfaceCubeTexture, SampleCubeTexture)
 void UNiagaraDataInterfaceCubeTexture::GetVMExternalFunction(const FVMExternalFunctionBindingInfo& BindingInfo, void* InstanceData, FVMExternalFunction &OutFunc)
diff --git a/Engine/Plugins/FX/Niagara/Source/Niagara/Private/NiagaraDataInterfaceCurlNoise.cpp b/Engine/Plugins/FX/Niagara/Source/Niagara/Private/NiagaraDataInterfaceCurlNoise.cpp
--- a/Engine/Plugins/FX/Niagara/Source/Niagara/Private/NiagaraDataInterfaceCurlNoise.cpp
+++ b/Engine/Plugins/FX/Niagara/Source/Niagara/Private/NiagaraDataInterfaceCurlNoise.cpp
@@ -391,7 +391,8 @@ bool UNiagaraDataInterfaceCurlNoise::Equals(const UNiagaraDataInterface* Other)
 	return OtherCurlNoise->Seed == Seed;
 }
 
-void UNiagaraDataInterfaceCurlNoise::GetFunctions(TArray<FNiagaraFunctionSignature>& OutFunctions)
+#if WITH_EDITORONLY_DATA
+void UNiagaraDataInterfaceCurlNoise::GetFunctionsInternal(TArray<FNiagaraFunctionSignature>& OutFunctions) const
 {
 	FNiagaraFunctionSignature Sig;
 	Sig.Name = SampleNoiseFieldName;
@@ -404,6 +405,7 @@ void UNiagaraDataInterfaceCurlNoise::GetFunctions(TArray<FNiagaraFunctionSignatu
 
 	OutFunctions.Add(Sig);
 }
+#endif
 
 DEFINE_NDI_DIRECT_FUNC_BINDER(UNiagaraDataInterfaceCurlNoise, SampleNoiseField);
 void UNiagaraDataInterfaceCurlNoise::GetVMExternalFunction(const FVMExternalFunctionBindingInfo& BindingInfo, void* InstanceData, FVMExternalFunction &OutFunc)
diff --git a/Engine/Plugins/FX/Niagara/Source/Niagara/Private/NiagaraDataInterfaceCurve.cpp b/Engine/Plugins/FX/Niagara/Source/Niagara/Private/NiagaraDataInterfaceCurve.cpp
--- a/Engine/Plugins/FX/Niagara/Source/Niagara/Private/NiagaraDataInterfaceCurve.cpp
+++ b/Engine/Plugins/FX/Niagara/Source/Niagara/Private/NiagaraDataInterfaceCurve.cpp
@@ -94,7 +94,8 @@ void UNiagaraDataInterfaceCurve::GetCurveData(TArray<FCurveData>& OutCurveData)
 	OutCurveData.Add(FCurveData(&Curve, NAME_None, FLinearColor(1.0f, 0.05f, 0.05f)));
 }
 
-void UNiagaraDataInterfaceCurve::GetFunctions(TArray<FNiagaraFunctionSignature>& OutFunctions)
+#if WITH_EDITORONLY_DATA
+void UNiagaraDataInterfaceCurve::GetFunctionsInternal(TArray<FNiagaraFunctionSignature>& OutFunctions) const
 {
 	FNiagaraFunctionSignature Sig;
 	Sig.Name = SampleCurveName;
@@ -107,6 +108,7 @@ void UNiagaraDataInterfaceCurve::GetFunctions(TArray<FNiagaraFunctionSignature>&
 
 	OutFunctions.Add(Sig);
 }
+#endif
 
 void UNiagaraDataInterfaceCurve::UpdateTimeRanges()
 {
diff --git a/Engine/Plugins/FX/Niagara/Source/Niagara/Private/NiagaraDataInterfaceExport.cpp b/Engine/Plugins/FX/Niagara/Source/Niagara/Private/NiagaraDataInterfaceExport.cpp
--- a/Engine/Plugins/FX/Niagara/Source/Niagara/Private/NiagaraDataInterfaceExport.cpp
+++ b/Engine/Plugins/FX/Niagara/Source/Niagara/Private/NiagaraDataInterfaceExport.cpp
@@ -299,13 +299,12 @@ bool UNiagaraDataInterfaceExport::Equals(const UNiagaraDataInterface* Other) con
 		OtherTyped->GPUAllocationPerParticleSize == GPUAllocationPerParticleSize;
 }
 
-void UNiagaraDataInterfaceExport::GetFunctions(TArray<FNiagaraFunctionSignature>& OutFunctions)
+#if WITH_EDITORONLY_DATA
+void UNiagaraDataInterfaceExport::GetFunctionsInternal(TArray<FNiagaraFunctionSignature>& OutFunctions) const
 {
 	FNiagaraFunctionSignature SigOld;
 	SigOld.Name = NDIExportLocal::StoreDataName_DEPRECATED;
-#if WITH_EDITORONLY_DATA
 	SigOld.Description = NSLOCTEXT("Niagara", "ExportDataFunctionDescription", "This function takes the particle data and stores it to be exported to the registered callback handler after the simulation has ticked.");
-#endif
 	SigOld.bMemberFunction = true;
 	SigOld.bRequiresContext = false;
 	SigOld.bSupportsGPU = true;
@@ -320,9 +319,7 @@ void UNiagaraDataInterfaceExport::GetFunctions(TArray<FNiagaraFunctionSignature>
 
 	FNiagaraFunctionSignature Sig;
 	Sig.Name = NDIExportLocal::ExportDataName;
-#if WITH_EDITORONLY_DATA
 	Sig.Description = NSLOCTEXT("Niagara", "ExportDataFunctionDescription", "This function takes the particle data and stores it to be exported to the registered callback handler after the simulation has ticked.");
-#endif
 	Sig.bMemberFunction = true;
 	Sig.bRequiresContext = false;
 	Sig.bSupportsGPU = true;
@@ -336,7 +333,6 @@ void UNiagaraDataInterfaceExport::GetFunctions(TArray<FNiagaraFunctionSignature>
 	OutFunctions.Add(Sig);
 }
 
-#if WITH_EDITORONLY_DATA
 void UNiagaraDataInterfaceExport::GetParameterDefinitionHLSL(const FNiagaraDataInterfaceGPUParamInfo& ParamInfo, FString& OutHLSL)
 {
 	const TMap<FString, FStringFormatArg> TemplateArgs =
diff --git a/Engine/Plugins/FX/Niagara/Source/Niagara/Private/NiagaraDataInterfaceGrid2DCollection.cpp b/Engine/Plugins/FX/Niagara/Source/Niagara/Private/NiagaraDataInterfaceGrid2DCollection.cpp
--- a/Engine/Plugins/FX/Niagara/Source/Niagara/Private/NiagaraDataInterfaceGrid2DCollection.cpp
+++ b/Engine/Plugins/FX/Niagara/Source/Niagara/Private/NiagaraDataInterfaceGrid2DCollection.cpp
@@ -364,7 +364,7 @@
 	if (UpgradeName != NAME_None)
 	{
 		TArray<FNiagaraFunctionSignature> Sigs;
-		GetFunctions(Sigs);
+		GetFunctionsInternal(Sigs);
 
 		for (const FNiagaraFunctionSignature& Sig : Sigs)
 		{
@@ -380,7 +380,7 @@
 	if ( FunctionSignature.FunctionVersion < FNiagaraGridCollection2DDIFunctionVersion::AddReadMemberToFunctions )
 	{
 		TArray<FNiagaraFunctionSignature> Sigs;
-		GetFunctions(Sigs);
+		GetFunctionsInternal(Sigs);
 
 		for (const FNiagaraFunctionSignature& Sig : Sigs)
 		{
@@ -397,9 +397,10 @@
 }
 #endif
 
-void UNiagaraDataInterfaceGrid2DCollection::GetFunctions(TArray<FNiagaraFunctionSignature>& OutFunctions)
+#if WITH_EDITORONLY_DATA
+void UNiagaraDataInterfaceGrid2DCollection::GetFunctionsInternal(TArray<FNiagaraFunctionSignature>& OutFunctions) const
 {
-	Super::GetFunctions(OutFunctions);
+	Super::GetFunctionsInternal(OutFunctions);
 
 	int32 StartIndex = OutFunctions.Num();
 
@@ -439,9 +440,7 @@
 		Sig.bSupportsCPU = false;
 		Sig.bSupportsGPU = true;
 		Sig.bReadFunction = true;
-#if WITH_EDITORONLY_DATA
 		Sig.Description = NSLOCTEXT("Niagara", "NiagaraDataInterfaceGridColl2D_GetValueFunction", "Get the value at a specific index. Note that this is an older way of working with Grids. Consider using the SetFloat or other typed, named functions or parameter map variables with StackContext namespace instead.");
-#endif
 
 		OutFunctions.Add(Sig);
 	}
@@ -462,9 +461,7 @@
 		Sig.bSupportsCPU = false;
 		Sig.bSupportsGPU = true;
 		Sig.bReadFunction = true;
-#if WITH_EDITORONLY_DATA
 		Sig.Description = NSLOCTEXT("Niagara", "NiagaraDataInterfaceGridColl2D_PreviousValueAtIndexFunction", "Get the value at a specific index.");
-#endif
 
 		OutFunctions.Add(Sig);
 	}
@@ -488,9 +485,7 @@
 		Sig.bSoftDeprecatedFunction = true;
 		Sig.bSupportsCPU = false;
 		Sig.bSupportsGPU = true;
-#if WITH_EDITORONLY_DATA
 		Sig.Description = NSLOCTEXT("Niagara", "NiagaraDataInterfaceGridColl2D_SetValueFunction", "Set the value at a specific index. Note that this is an older way of working with Grids. Consider using the SetFloat or other typed, named functions or parameter map variables with StackContext namespace instead.");
-#endif
 		OutFunctions.Add(Sig);
 	}
 
@@ -512,9 +507,7 @@
 		Sig.bSoftDeprecatedFunction = false;
 		Sig.bSupportsCPU = false;
 		Sig.bSupportsGPU = true;
-#if WITH_EDITORONLY_DATA
 		Sig.Description = NSLOCTEXT("Niagara", "NiagaraDataInterfaceGridColl2D_SetValueAtIndexFunction", "Set the value at a specific index.");
-#endif
 		OutFunctions.Add(Sig);
 	}
 
@@ -533,9 +526,7 @@
 		Sig.ModuleUsageBitmask = ENiagaraScriptUsageMask::Particle;
 		Sig.bSupportsCPU = false;
 		Sig.bSupportsGPU = true;
-#if WITH_EDITORONLY_DATA
 		Sig.Description = NSLOCTEXT("Niagara", "NiagaraDataInterfaceGridColl2D_ClearCellFunction", "Set all attributes for a given cell to be zeroes.");
-#endif
 		OutFunctions.Add(Sig);
 	}
 
@@ -555,9 +546,7 @@
 		Sig.bSupportsCPU = false;
 		Sig.bSupportsGPU = true;
 		Sig.bReadFunction = true;
-#if WITH_EDITORONLY_DATA
 		Sig.Description = NSLOCTEXT("Niagara", "NiagaraDataInterfaceGridColl2D_CopyPreviousToCurrentForCell", "Take the previous contents of the cell and copy to the output location for the cell.");
-#endif
 		OutFunctions.Add(Sig);
 	}
 
@@ -577,9 +566,7 @@
 		Sig.ModuleUsageBitmask = ENiagaraScriptUsageMask::Particle;
 		Sig.bSupportsCPU = false;
 		Sig.bSupportsGPU = true;
-#if WITH_EDITORONLY_DATA
 		Sig.Description = NSLOCTEXT("Niagara", "NiagaraDataInterfaceGridColl2D_SetVector4", "Sets a Vector4 value on the Grid by Attribute name.");
-#endif
 		OutFunctions.Add(Sig);
 	}
 
@@ -598,9 +585,7 @@
 		Sig.bSupportsCPU = false;
 		Sig.bSupportsGPU = true;
 		Sig.bReadFunction = true;
-#if WITH_EDITORONLY_DATA
 		Sig.Description = NSLOCTEXT("Niagara", "NiagaraDataInterfaceGridColl2D_GetVector4", "Gets a Vector4 value on the Grid by Attribute name. Note that this is the value from the previous execution stage.");
-#endif
 		OutFunctions.Add(Sig);
 
 		// Add older version for back compat with CustomHLSL.
@@ -623,9 +608,7 @@
 		Sig.bSupportsCPU = false;
 		Sig.bSupportsGPU = true;
 		Sig.bReadFunction = true;
-#if WITH_EDITORONLY_DATA
 		Sig.Description = NSLOCTEXT("Niagara", "NiagaraDataInterfaceGridColl2D_SampleVector4", "Sample a Vector4 value on the Grid by Attribute name. Note that this is the value from the previous execution stage.");
-#endif
 		OutFunctions.Add(Sig);
 
 		// Add older version for back compat with CustomHLSL.
@@ -648,9 +631,7 @@
 		Sig.bSupportsCPU = false;
 		Sig.bSupportsGPU = true;
 		Sig.bReadFunction = true;
-#if WITH_EDITORONLY_DATA
 		Sig.Description = NSLOCTEXT("Niagara", "NiagaraDataInterfaceGridColl2D_CubicSampleVector4", "Cubic sample a Vector4 value on the Grid by Attribute name. Note that this is the value from the previous execution stage.");
-#endif
 		OutFunctions.Add(Sig);
 
 		// Add older version for back compat with CustomHLSL.
@@ -675,9 +656,7 @@
 		Sig.ModuleUsageBitmask = ENiagaraScriptUsageMask::Particle;
 		Sig.bSupportsCPU = false;
 		Sig.bSupportsGPU = true;
-#if WITH_EDITORONLY_DATA
 		Sig.Description = NSLOCTEXT("Niagara", "NiagaraDataInterfaceGridColl2D_SetVector3", "Sets a Vector3 value on the Grid by Attribute name.");
-#endif
 		OutFunctions.Add(Sig);
 
 		// Add older version for back compat with CustomHLSL.
@@ -701,9 +680,7 @@
 		Sig.bSupportsCPU = false;
 		Sig.bSupportsGPU = true;
 		Sig.bReadFunction = true;
-#if WITH_EDITORONLY_DATA
 		Sig.Description = NSLOCTEXT("Niagara", "NiagaraDataInterfaceGridColl2D_GetVector3", "Gets a Vector3 value on the Grid by Attribute name. Note that this is the value from the previous execution stage.");
-#endif
 		OutFunctions.Add(Sig);
 
 		// Add older version for back compat with CustomHLSL.
@@ -726,9 +703,7 @@
 		Sig.bSupportsCPU = false;
 		Sig.bSupportsGPU = true;
 		Sig.bReadFunction = true;
-#if WITH_EDITORONLY_DATA
 		Sig.Description = NSLOCTEXT("Niagara", "NiagaraDataInterfaceGridColl2D_SampleVector3", "Sample a Vector3 value on the Grid by Attribute name. Note that this is the value from the previous execution stage.");
-#endif
 		OutFunctions.Add(Sig);
 
 		// Add older version for back compat with CustomHLSL.
@@ -751,9 +726,7 @@
 		Sig.bSupportsCPU = false;
 		Sig.bSupportsGPU = true;
 		Sig.bReadFunction = true;
-#if WITH_EDITORONLY_DATA
 		Sig.Description = NSLOCTEXT("Niagara", "NiagaraDataInterfaceGridColl2D_CubicSampleVector3", "Cubic sample a Vector3 value on the Grid by Attribute name. Note that this is the value from the previous execution stage.");
-#endif
 		OutFunctions.Add(Sig);
 
 		// Add older version for back compat with CustomHLSL.
@@ -778,9 +751,7 @@
 		Sig.ModuleUsageBitmask = ENiagaraScriptUsageMask::Particle;
 		Sig.bSupportsCPU = false;
 		Sig.bSupportsGPU = true;
-#if WITH_EDITORONLY_DATA
 		Sig.Description = NSLOCTEXT("Niagara", "NiagaraDataInterfaceGridColl2D_SetVector2", "Sets a Vector2 value on the Grid by Attribute name.");
-#endif
 		OutFunctions.Add(Sig);
 
 		// Add older version for back compat with CustomHLSL.
@@ -804,9 +775,7 @@
 		Sig.bSupportsCPU = false;
 		Sig.bSupportsGPU = true;
 		Sig.bReadFunction = true;
-#if WITH_EDITORONLY_DATA
 		Sig.Description = NSLOCTEXT("Niagara", "NiagaraDataInterfaceGridColl2D_GetVector2", "Gets a Vector2 value on the Grid by Attribute name. Note that this is the value from the previous execution stage.");
-#endif
 		OutFunctions.Add(Sig);
 
 		// Add older version for back compat with CustomHLSL.
@@ -829,9 +798,7 @@
 		Sig.bSupportsCPU = false;
 		Sig.bSupportsGPU = true;
 		Sig.bReadFunction = true;
-#if WITH_EDITORONLY_DATA
 		Sig.Description = NSLOCTEXT("Niagara", "NiagaraDataInterfaceGridColl2D_SampleVector2", "Sample a Vector2 value on the Grid by Attribute name. Note that this is the value from the previous execution stage.");
-#endif
 		OutFunctions.Add(Sig);
 
 		// Add older version for back compat with CustomHLSL.
@@ -854,9 +821,7 @@
 		Sig.bSupportsCPU = false;
 		Sig.bSupportsGPU = true;
 		Sig.bReadFunction = true;
-#if WITH_EDITORONLY_DATA
 		Sig.Description = NSLOCTEXT("Niagara", "NiagaraDataInterfaceGridColl2D_CubicSampleVector2", "Cubic sample a Vector2 value on the Grid by Attribute name. Note that this is the value from the previous execution stage.");
-#endif
 		OutFunctions.Add(Sig);
 
 		// Add older version for back compat with CustomHLSL.
@@ -881,9 +846,7 @@
 		Sig.ModuleUsageBitmask = ENiagaraScriptUsageMask::Particle;
 		Sig.bSupportsCPU = false;
 		Sig.bSupportsGPU = true;
-#if WITH_EDITORONLY_DATA
 		Sig.Description = NSLOCTEXT("Niagara", "NiagaraDataInterfaceGridColl2D_SetFloat", "Sets a float value on the Grid by Attribute name.");
-#endif
 		OutFunctions.Add(Sig);
 	}
 
@@ -902,9 +865,7 @@
 		Sig.bSupportsCPU = false;
 		Sig.bSupportsGPU = true;
 		Sig.bReadFunction = true;
-#if WITH_EDITORONLY_DATA
 		Sig.Description = NSLOCTEXT("Niagara", "NiagaraDataInterfaceGridColl2D_GetFloat", "Gets a float value on the Grid by Attribute name. Note that this is the value from the previous execution stage.");
-#endif
 		OutFunctions.Add(Sig);
 
 		// Add older version for back compat with CustomHLSL.
@@ -927,9 +888,7 @@
 		Sig.bSupportsCPU = false;
 		Sig.bSupportsGPU = true;
 		Sig.bReadFunction = true;
-#if WITH_EDITORONLY_DATA
 		Sig.Description = NSLOCTEXT("Niagara", "NiagaraDataInterfaceGridColl2D_SampleFloat", "Sample a float value on the Grid by Attribute name.");
-#endif
 		OutFunctions.Add(Sig);
 
 		// Add older version for back compat with CustomHLSL.
@@ -952,9 +911,7 @@
 		Sig.bSupportsCPU = false;
 		Sig.bSupportsGPU = true;
 		Sig.bReadFunction = true;
-#if WITH_EDITORONLY_DATA
 		Sig.Description = NSLOCTEXT("Niagara", "NiagaraDataInterfaceGridColl2D_CubicSampleFloat", "Cubic sample a float value on the Grid by Attribute name.");
-#endif
 		OutFunctions.Add(Sig);
 
 		// Add older version for back compat with CustomHLSL.
@@ -1030,9 +987,7 @@
 		Sig.bMemberFunction = true;
 		Sig.bRequiresContext = false;
 		Sig.bExperimental = true;
-#if WITH_EDITORONLY_DATA
 		Sig.Description = NSLOCTEXT("Niagara", "NiagaraDataInterfaceGridColl2D_GetVector4AttributeIndex", "Gets a attribute starting index value for Vector4 on the Grid by Attribute name. Returns -1 if not found.");
-#endif
 		OutFunctions.Add(Sig);
 	}
 
@@ -1045,9 +1000,7 @@
 		Sig.bMemberFunction = true;
 		Sig.bRequiresContext = false;
 		Sig.bExperimental = true;
-#if WITH_EDITORONLY_DATA
 		Sig.Description = NSLOCTEXT("Niagara", "NiagaraDataInterfaceGridColl2D_GetVector3AttributeIndex", "Gets a attribute starting index value for Vector3 on the Grid by Attribute name. Returns -1 if not found.");
-#endif
 		OutFunctions.Add(Sig);
 	}
 
@@ -1060,9 +1013,7 @@
 		Sig.bMemberFunction = true;
 		Sig.bRequiresContext = false;
 		Sig.bExperimental = true;
-#if WITH_EDITORONLY_DATA
 		Sig.Description = NSLOCTEXT("Niagara", "NiagaraDataInterfaceGridColl2D_GetVector2AttributeIndex", "Gets a attribute starting index value for Vector2 on the Grid by Attribute name. Returns -1 if not found.");
-#endif
 		OutFunctions.Add(Sig);
 	}
 
@@ -1075,20 +1026,17 @@
 		Sig.bMemberFunction = true;
 		Sig.bRequiresContext = false;
 		Sig.bExperimental = true;
-#if WITH_EDITORONLY_DATA
 		Sig.Description = NSLOCTEXT("Niagara", "NiagaraDataInterfaceGridColl2D_GetFloatAttributeIndex", "Gets a attribute starting index value for float on the Grid by Attribute name. Returns -1 if not found.");
-#endif
 		OutFunctions.Add(Sig);
 	}
 
-#if WITH_EDITORONLY_DATA
 	for (int32 i = StartIndex; i < OutFunctions.Num(); i++)
 	{
 		FNiagaraFunctionSignature& Function = OutFunctions[i];
 		Function.FunctionVersion = FNiagaraGridCollection2DDIFunctionVersion::LatestVersion;
 	}
-#endif
 }
+#endif
 
 // #todo(dmp): expose more CPU functionality
 // #todo(dmp): ideally these would be exposed on the parent class, but we can't bind functions of parent classes but need to work on the interface
diff --git a/Engine/Plugins/FX/Niagara/Source/Niagara/Private/NiagaraDataInterfaceGrid2DCollectionReader.cpp b/Engine/Plugins/FX/Niagara/Source/Niagara/Private/NiagaraDataInterfaceGrid2DCollectionReader.cpp
--- a/Engine/Plugins/FX/Niagara/Source/Niagara/Private/NiagaraDataInterfaceGrid2DCollectionReader.cpp
+++ b/Engine/Plugins/FX/Niagara/Source/Niagara/Private/NiagaraDataInterfaceGrid2DCollectionReader.cpp
@@ -157,12 +157,12 @@ void UNiagaraDataInterfaceGrid2DCollectionReader::GetEmitterDependencies(UNiagar
 }
 
 
-
-void UNiagaraDataInterfaceGrid2DCollectionReader::GetFunctions(TArray<FNiagaraFunctionSignature>& OutFunctions)
+#if WITH_EDITORONLY_DATA
+void UNiagaraDataInterfaceGrid2DCollectionReader::GetFunctionsInternal(TArray<FNiagaraFunctionSignature>& OutFunctions) const
 {
 	int32 NumFunctionsBefore = OutFunctions.Num();
 
-	Super::GetFunctions(OutFunctions);
+	Super::GetFunctionsInternal(OutFunctions);
 	TSet<FName> FunctionsToRemove =
 	{
 		ClearCellFunctionName,
@@ -191,6 +191,7 @@ void UNiagaraDataInterfaceGrid2DCollectionReader::GetFunctions(TArray<FNiagaraFu
 		}
 	}
 }
+#endif
 
 bool UNiagaraDataInterfaceGrid2DCollectionReader::PerInstanceTickPostSimulate(void* PerInstanceData, FNiagaraSystemInstance* SystemInstance, float DeltaSeconds)
 {
diff --git a/Engine/Plugins/FX/Niagara/Source/Niagara/Private/NiagaraDataInterfaceGrid3DCollection.cpp b/Engine/Plugins/FX/Niagara/Source/Niagara/Private/NiagaraDataInterfaceGrid3DCollection.cpp
--- a/Engine/Plugins/FX/Niagara/Source/Niagara/Private/NiagaraDataInterfaceGrid3DCollection.cpp
+++ b/Engine/Plugins/FX/Niagara/Source/Niagara/Private/NiagaraDataInterfaceGrid3DCollection.cpp
@@ -400,7 +400,7 @@ bool UNiagaraDataInterfaceGrid3DCollection::UpgradeFunctionCall(FNiagaraFunction
 	if (UpgradeName != NAME_None)
 	{
 		TArray<FNiagaraFunctionSignature> Sigs;
-		GetFunctions(Sigs);
+		GetFunctionsInternal(Sigs);
 
 		for (const FNiagaraFunctionSignature& Sig : Sigs)
 		{
@@ -415,10 +415,10 @@ bool UNiagaraDataInterfaceGrid3DCollection::UpgradeFunctionCall(FNiagaraFunction
 
 	return false;
 }
-#endif
-void UNiagaraDataInterfaceGrid3DCollection::GetFunctions(TArray<FNiagaraFunctionSignature>& OutFunctions)
+
+void UNiagaraDataInterfaceGrid3DCollection::GetFunctionsInternal(TArray<FNiagaraFunctionSignature>& OutFunctions) const
 {
-	Super::GetFunctions(OutFunctions);
+	Super::GetFunctionsInternal(OutFunctions);
 
 	int32 StartIndex = OutFunctions.Num();
 
@@ -460,9 +460,7 @@ void UNiagaraDataInterfaceGrid3DCollection::GetFunctions(TArray<FNiagaraFunction
 		Sig.bSupportsCPU = false;
 		Sig.bSupportsGPU = true;
 		Sig.bReadFunction = true;
-#if WITH_EDITORONLY_DATA
 		Sig.Description = NSLOCTEXT("Niagara", "NiagaraDataInterfaceGridColl3D_GetValueFunction", "Get the value at a specific index. Note that this is an older way of working with Grids. Consider using the SetFloat or other typed, named functions or parameter map variables with StackContext namespace instead.");
-#endif
 
 		OutFunctions.Add(Sig);
 	}
@@ -484,9 +482,7 @@ void UNiagaraDataInterfaceGrid3DCollection::GetFunctions(TArray<FNiagaraFunction
 		Sig.bSupportsCPU = false;
 		Sig.bSupportsGPU = true;
 		Sig.bReadFunction = true;
-#if WITH_EDITORONLY_DATA
 		Sig.Description = NSLOCTEXT("Niagara", "NiagaraDataInterfaceGridColl2D_PreviousValueAtIndexFunction", "Get the value at a specific index.");
-#endif
 
 		OutFunctions.Add(Sig);
 	}
@@ -509,9 +505,7 @@ void UNiagaraDataInterfaceGrid3DCollection::GetFunctions(TArray<FNiagaraFunction
 		Sig.bReadFunction = true;
 		Sig.bHidden = true;
 		Sig.bSoftDeprecatedFunction = true;
-#if WITH_EDITORONLY_DATA
 		Sig.Description = NSLOCTEXT("Niagara", "NiagaraDataInterfaceGridColl2D_PreviousValueAtIndexFunction", "Get the value at a specific index.");
-#endif
 
 		OutFunctions.Add(Sig);
 	}
@@ -535,9 +529,7 @@ void UNiagaraDataInterfaceGrid3DCollection::GetFunctions(TArray<FNiagaraFunction
 		Sig.bSoftDeprecatedFunction = true;
 		Sig.bSupportsCPU = false;
 		Sig.bSupportsGPU = true;
-#if WITH_EDITORONLY_DATA
 		Sig.Description = NSLOCTEXT("Niagara", "NiagaraDataInterfaceGridColl2D_SetValueFunction", "Set the value at a specific index. Note that this is an older way of working with Grids. Consider using the SetFloat or other typed, named functions or parameter map variables with StackContext namespace instead.");
-#endif
 		OutFunctions.Add(Sig);
 	}
 
@@ -560,9 +552,7 @@ void UNiagaraDataInterfaceGrid3DCollection::GetFunctions(TArray<FNiagaraFunction
 		Sig.bSupportsCPU = false;
 		Sig.bSupportsGPU = true;
 		Sig.bHidden = true;
-#if WITH_EDITORONLY_DATA
 		Sig.Description = NSLOCTEXT("Niagara", "NiagaraDataInterfaceGridColl2D_SetValueFunction", "Set the value at a specific index. Note that this is an older way of working with Grids. Consider using the SetFloat or other typed, named functions or parameter map variables with StackContext namespace instead.");
-#endif
 		OutFunctions.Add(Sig);
 	}
 
@@ -582,9 +572,7 @@ void UNiagaraDataInterfaceGrid3DCollection::GetFunctions(TArray<FNiagaraFunction
 		Sig.ModuleUsageBitmask = ENiagaraScriptUsageMask::Particle;
 		Sig.bSupportsCPU = false;
 		Sig.bSupportsGPU = true;
-#if WITH_EDITORONLY_DATA
 		Sig.Description = NSLOCTEXT("Niagara", "NiagaraDataInterfaceGridColl2D_ClearCellFunction", "Set all attributes for a given cell to be zeroes.");
-#endif
 		OutFunctions.Add(Sig);
 	}
 
@@ -605,9 +593,7 @@ void UNiagaraDataInterfaceGrid3DCollection::GetFunctions(TArray<FNiagaraFunction
 		Sig.bSupportsCPU = false;
 		Sig.bSupportsGPU = true;
 		Sig.bReadFunction = true;
-#if WITH_EDITORONLY_DATA
 		Sig.Description = NSLOCTEXT("Niagara", "NiagaraDataInterfaceGridColl2D_CopyPreviousToCurrentForCell", "Take the previous contents of the cell and copy to the output location for the cell.");
-#endif
 		OutFunctions.Add(Sig);
 	}
 
@@ -630,9 +616,7 @@ void UNiagaraDataInterfaceGrid3DCollection::GetFunctions(TArray<FNiagaraFunction
 		Sig.bSupportsCPU = false;
 		Sig.bSupportsGPU = true;
 		Sig.bReadFunction = true;
-#if WITH_EDITORONLY_DATA
 		Sig.Description = NSLOCTEXT("Niagara", "NiagaraDataInterfaceGridColl2D_CopyMaskedPreviousToCurrentForCell", "Take the previous contents of the cell and copy to the output location for the cell.");
-#endif
 		OutFunctions.Add(Sig);
 	}
 
@@ -653,9 +637,7 @@ void UNiagaraDataInterfaceGrid3DCollection::GetFunctions(TArray<FNiagaraFunction
 		Sig.ModuleUsageBitmask = ENiagaraScriptUsageMask::Particle;
 		Sig.bSupportsCPU = false;
 		Sig.bSupportsGPU = true;
-#if WITH_EDITORONLY_DATA
 		Sig.Description = NSLOCTEXT("Niagara", "NiagaraDataInterfaceGridColl2D_SetVector4", "Sets a Vector4 value on the Grid by Attribute name.");
-#endif
 		OutFunctions.Add(Sig);
 	}
 
@@ -675,9 +657,7 @@ void UNiagaraDataInterfaceGrid3DCollection::GetFunctions(TArray<FNiagaraFunction
 		Sig.bSupportsCPU = false;
 		Sig.bSupportsGPU = true;
 		Sig.bReadFunction = true;
-#if WITH_EDITORONLY_DATA
 		Sig.Description = NSLOCTEXT("Niagara", "NiagaraDataInterfaceGridColl3D_GetVector4", "Gets a Vector4 value on the Grid by Attribute name. Note that this is the value from the previous execution stage.");
-#endif
 		OutFunctions.Add(Sig);
 
 	}
@@ -696,9 +676,7 @@ void UNiagaraDataInterfaceGrid3DCollection::GetFunctions(TArray<FNiagaraFunction
 		Sig.bSupportsCPU = false;
 		Sig.bSupportsGPU = true;
 		Sig.bReadFunction = true;
-#if WITH_EDITORONLY_DATA
 		Sig.Description = NSLOCTEXT("Niagara", "NiagaraDataInterfaceGridColl3D_SampleVector4", "Sample a Vector4 value on the Grid by Attribute name. Note that this is the value from the previous execution stage.");
-#endif
 		OutFunctions.Add(Sig);
 
 	}
@@ -717,9 +695,7 @@ void UNiagaraDataInterfaceGrid3DCollection::GetFunctions(TArray<FNiagaraFunction
 		Sig.bSupportsCPU = false;
 		Sig.bSupportsGPU = true;
 		Sig.bReadFunction = true;
-#if WITH_EDITORONLY_DATA
 		Sig.Description = NSLOCTEXT("Niagara", "NiagaraDataInterfaceGridColl3D_CubicSampleVector4", "Cubic Sample a Vector4 value on the Grid by Attribute name. Note that this is the value from the previous execution stage.");
-#endif
 		OutFunctions.Add(Sig);
 
 	}
@@ -741,9 +717,7 @@ void UNiagaraDataInterfaceGrid3DCollection::GetFunctions(TArray<FNiagaraFunction
 		Sig.ModuleUsageBitmask = ENiagaraScriptUsageMask::Particle;
 		Sig.bSupportsCPU = false;
 		Sig.bSupportsGPU = true;
-#if WITH_EDITORONLY_DATA
 		Sig.Description = NSLOCTEXT("Niagara", "NiagaraDataInterfaceGridColl3D_SetVector3", "Sets a Vector3 value on the Grid by Attribute name.");
-#endif
 		OutFunctions.Add(Sig);
 
 	}
@@ -764,9 +738,7 @@ void UNiagaraDataInterfaceGrid3DCollection::GetFunctions(TArray<FNiagaraFunction
 		Sig.bSupportsCPU = false;
 		Sig.bSupportsGPU = true;
 		Sig.bReadFunction = true;
-#if WITH_EDITORONLY_DATA
 		Sig.Description = NSLOCTEXT("Niagara", "NiagaraDataInterfaceGridColl3D_GetVector3", "Gets a Vector3 value on the Grid by Attribute name. Note that this is the value from the previous execution stage.");
-#endif
 		OutFunctions.Add(Sig);
 
 	}
@@ -785,9 +757,7 @@ void UNiagaraDataInterfaceGrid3DCollection::GetFunctions(TArray<FNiagaraFunction
 		Sig.bSupportsCPU = false;
 		Sig.bSupportsGPU = true;
 		Sig.bReadFunction = true;
-#if WITH_EDITORONLY_DATA
 		Sig.Description = NSLOCTEXT("Niagara", "NiagaraDataInterfaceGridColl3D_SampleVector3", "Sample a Vector3 value on the Grid by Attribute name. Note that this is the value from the previous execution stage.");
-#endif
 		OutFunctions.Add(Sig);
 
 	}
@@ -806,9 +776,7 @@ void UNiagaraDataInterfaceGrid3DCollection::GetFunctions(TArray<FNiagaraFunction
 		Sig.bSupportsCPU = false;
 		Sig.bSupportsGPU = true;
 		Sig.bReadFunction = true;
-#if WITH_EDITORONLY_DATA
 		Sig.Description = NSLOCTEXT("Niagara", "NiagaraDataInterfaceGridColl3D_CubicSampleVector3", "Cubic Sample a Vector3 value on the Grid by Attribute name. Note that this is the value from the previous execution stage.");
-#endif
 		OutFunctions.Add(Sig);
 
 	}
@@ -830,9 +798,7 @@ void UNiagaraDataInterfaceGrid3DCollection::GetFunctions(TArray<FNiagaraFunction
 		Sig.ModuleUsageBitmask = ENiagaraScriptUsageMask::Particle;
 		Sig.bSupportsCPU = false;
 		Sig.bSupportsGPU = true;
-#if WITH_EDITORONLY_DATA
 		Sig.Description = NSLOCTEXT("Niagara", "NiagaraDataInterfaceGridColl3D_SetVector2", "Sets a Vector2 value on the Grid by Attribute name.");
-#endif
 		OutFunctions.Add(Sig);
 
 	}
@@ -853,9 +819,7 @@ void UNiagaraDataInterfaceGrid3DCollection::GetFunctions(TArray<FNiagaraFunction
 		Sig.bSupportsCPU = false;
 		Sig.bSupportsGPU = true;
 		Sig.bReadFunction = true;
-#if WITH_EDITORONLY_DATA
 		Sig.Description = NSLOCTEXT("Niagara", "NiagaraDataInterfaceGridColl3D_GetVector2", "Gets a Vector2 value on the Grid by Attribute name. Note that this is the value from the previous execution stage.");
-#endif
 		OutFunctions.Add(Sig);
 
 	}
@@ -874,9 +838,7 @@ void UNiagaraDataInterfaceGrid3DCollection::GetFunctions(TArray<FNiagaraFunction
 		Sig.bSupportsCPU = false;
 		Sig.bSupportsGPU = true;
 		Sig.bReadFunction = true;
-#if WITH_EDITORONLY_DATA
 		Sig.Description = NSLOCTEXT("Niagara", "NiagaraDataInterfaceGridColl3D_SampleVector2", "Sample a Vector2 value on the Grid by Attribute name. Note that this is the value from the previous execution stage.");
-#endif
 		OutFunctions.Add(Sig);
 
 	}
@@ -895,9 +857,7 @@ void UNiagaraDataInterfaceGrid3DCollection::GetFunctions(TArray<FNiagaraFunction
 		Sig.bSupportsCPU = false;
 		Sig.bSupportsGPU = true;
 		Sig.bReadFunction = true;
-#if WITH_EDITORONLY_DATA
 		Sig.Description = NSLOCTEXT("Niagara", "NiagaraDataInterfaceGridColl3D_CubicSampleVector2", "Cubic Sample a Vector2 value on the Grid by Attribute name. Note that this is the value from the previous execution stage.");
-#endif
 		OutFunctions.Add(Sig);
 
 	}
@@ -919,9 +879,7 @@ void UNiagaraDataInterfaceGrid3DCollection::GetFunctions(TArray<FNiagaraFunction
 		Sig.ModuleUsageBitmask = ENiagaraScriptUsageMask::Particle;
 		Sig.bSupportsCPU = false;
 		Sig.bSupportsGPU = true;
-#if WITH_EDITORONLY_DATA
 		Sig.Description = NSLOCTEXT("Niagara", "NiagaraDataInterfaceGridColl3D_SetFloat", "Sets a float value on the Grid by Attribute name.");
-#endif
 		OutFunctions.Add(Sig);
 	}
 
@@ -941,9 +899,7 @@ void UNiagaraDataInterfaceGrid3DCollection::GetFunctions(TArray<FNiagaraFunction
 		Sig.bSupportsCPU = false;
 		Sig.bSupportsGPU = true;
 		Sig.bReadFunction = true;
-#if WITH_EDITORONLY_DATA
 		Sig.Description = NSLOCTEXT("Niagara", "NiagaraDataInterfaceGridColl3D_GetFloat", "Gets a float value on the Grid by Attribute name. Note that this is the value from the previous execution stage.");
-#endif
 		OutFunctions.Add(Sig);
 
 	}
@@ -962,9 +918,7 @@ void UNiagaraDataInterfaceGrid3DCollection::GetFunctions(TArray<FNiagaraFunction
 		Sig.bSupportsCPU = false;
 		Sig.bSupportsGPU = true;
 		Sig.bReadFunction = true;
-#if WITH_EDITORONLY_DATA
 		Sig.Description = NSLOCTEXT("Niagara", "NiagaraDataInterfaceGridColl3D_SampleFloat", "Sample a float value on the Grid by Attribute name.");
-#endif
 		OutFunctions.Add(Sig);
 
 	}
@@ -983,9 +937,7 @@ void UNiagaraDataInterfaceGrid3DCollection::GetFunctions(TArray<FNiagaraFunction
 		Sig.bSupportsCPU = false;
 		Sig.bSupportsGPU = true;
 		Sig.bReadFunction = true;
-#if WITH_EDITORONLY_DATA
 		Sig.Description = NSLOCTEXT("Niagara", "NiagaraDataInterfaceGridColl3D_CubicSampleFloat", "Cubic Sample a float value on the Grid by Attribute name.");
-#endif
 		OutFunctions.Add(Sig);
 
 	}
@@ -1119,9 +1071,7 @@ void UNiagaraDataInterfaceGrid3DCollection::GetFunctions(TArray<FNiagaraFunction
 		Sig.bMemberFunction = true;
 		Sig.bRequiresContext = false;
 		Sig.bExperimental = true;
-#if WITH_EDITORONLY_DATA
 		Sig.Description = NSLOCTEXT("Niagara", "NiagaraDataInterfaceGridColl2D_GetVector4AttributeIndex", "Gets a attribute starting index value for Vector4 on the Grid by Attribute name. Returns -1 if not found.");
-#endif
 		OutFunctions.Add(Sig);
 	}
 
@@ -1134,9 +1084,7 @@ void UNiagaraDataInterfaceGrid3DCollection::GetFunctions(TArray<FNiagaraFunction
 		Sig.bMemberFunction = true;
 		Sig.bRequiresContext = false;
 		Sig.bExperimental = true;
-#if WITH_EDITORONLY_DATA
 		Sig.Description = NSLOCTEXT("Niagara", "NiagaraDataInterfaceGridColl2D_GetVector3AttributeIndex", "Gets a attribute starting index value for Vector3 on the Grid by Attribute name. Returns -1 if not found.");
-#endif
 		OutFunctions.Add(Sig);
 	}
 
@@ -1149,9 +1097,7 @@ void UNiagaraDataInterfaceGrid3DCollection::GetFunctions(TArray<FNiagaraFunction
 		Sig.bMemberFunction = true;
 		Sig.bRequiresContext = false;
 		Sig.bExperimental = true;
-#if WITH_EDITORONLY_DATA
 		Sig.Description = NSLOCTEXT("Niagara", "NiagaraDataInterfaceGridColl2D_GetVector2AttributeIndex", "Gets a attribute starting index value for Vector2 on the Grid by Attribute name. Returns -1 if not found.");
-#endif
 		OutFunctions.Add(Sig);
 	}
 
@@ -1164,9 +1110,7 @@ void UNiagaraDataInterfaceGrid3DCollection::GetFunctions(TArray<FNiagaraFunction
 		Sig.bMemberFunction = true;
 		Sig.bRequiresContext = false;
 		Sig.bExperimental = true;
-#if WITH_EDITORONLY_DATA
 		Sig.Description = NSLOCTEXT("Niagara", "NiagaraDataInterfaceGridColl2D_GetFloatAttributeIndex", "Gets a attribute starting index value for float on the Grid by Attribute name. Returns -1 if not found.");
-#endif
 		OutFunctions.Add(Sig);
 	}
 
@@ -1187,9 +1131,7 @@ void UNiagaraDataInterfaceGrid3DCollection::GetFunctions(TArray<FNiagaraFunction
 		Sig.bSupportsCPU = false;
 		Sig.bSupportsGPU = true;
 		Sig.bReadFunction = true;
-#if WITH_EDITORONLY_DATA
 		Sig.Description = NSLOCTEXT("Niagara", "NiagaraDataInterfaceGridColl2D_PreviousValueAtIndexFunction", "Get the value at a specific index.");
-#endif
 
 		OutFunctions.Add(Sig);
 	}
@@ -1211,9 +1153,7 @@ void UNiagaraDataInterfaceGrid3DCollection::GetFunctions(TArray<FNiagaraFunction
 		Sig.bSupportsCPU = false;
 		Sig.bSupportsGPU = true;
 		Sig.bReadFunction = true;
-#if WITH_EDITORONLY_DATA
 		Sig.Description = NSLOCTEXT("Niagara", "NiagaraDataInterfaceGridColl2D_PreviousValueAtIndexFunction", "Get the value at a specific index.");
-#endif
 
 		OutFunctions.Add(Sig);
 	}
@@ -1235,9 +1175,7 @@ void UNiagaraDataInterfaceGrid3DCollection::GetFunctions(TArray<FNiagaraFunction
 		Sig.bSupportsCPU = false;
 		Sig.bSupportsGPU = true;
 		Sig.bReadFunction = true;
-#if WITH_EDITORONLY_DATA
 		Sig.Description = NSLOCTEXT("Niagara", "NiagaraDataInterfaceGridColl2D_PreviousValueAtIndexFunction", "Get the value at a specific index.");
-#endif
 
 		OutFunctions.Add(Sig);
 	}
@@ -1259,9 +1197,7 @@ void UNiagaraDataInterfaceGrid3DCollection::GetFunctions(TArray<FNiagaraFunction
 		Sig.bSupportsCPU = false;
 		Sig.bSupportsGPU = true;
 		Sig.bReadFunction = true;
-#if WITH_EDITORONLY_DATA
 		Sig.Description = NSLOCTEXT("Niagara", "NiagaraDataInterfaceGridColl2D_PreviousValueAtIndexFunction", "Get the value at a specific index.");
-#endif
 
 		OutFunctions.Add(Sig);
 	}
@@ -1427,9 +1363,7 @@ void UNiagaraDataInterfaceGrid3DCollection::GetFunctions(TArray<FNiagaraFunction
 		Sig.ModuleUsageBitmask = ENiagaraScriptUsageMask::Particle;
 		Sig.bSupportsCPU = false;
 		Sig.bSupportsGPU = true;
-#if WITH_EDITORONLY_DATA
 		Sig.Description = NSLOCTEXT("Niagara", "NiagaraDataInterfaceGridColl3D_SetFloat", "Sets a float value on the Grid by Attribute name.");
-#endif
 		OutFunctions.Add(Sig);
 	}
 
@@ -1450,9 +1384,7 @@ void UNiagaraDataInterfaceGrid3DCollection::GetFunctions(TArray<FNiagaraFunction
 		Sig.ModuleUsageBitmask = ENiagaraScriptUsageMask::Particle;
 		Sig.bSupportsCPU = false;
 		Sig.bSupportsGPU = true;
-#if WITH_EDITORONLY_DATA
 		Sig.Description = NSLOCTEXT("Niagara", "NiagaraDataInterfaceGridColl3D_SetFloat", "Sets a float value on the Grid by Attribute name.");
-#endif
 		OutFunctions.Add(Sig);
 	}
 
@@ -1473,9 +1405,7 @@ void UNiagaraDataInterfaceGrid3DCollection::GetFunctions(TArray<FNiagaraFunction
 		Sig.ModuleUsageBitmask = ENiagaraScriptUsageMask::Particle;
 		Sig.bSupportsCPU = false;
 		Sig.bSupportsGPU = true;
-#if WITH_EDITORONLY_DATA
 		Sig.Description = NSLOCTEXT("Niagara", "NiagaraDataInterfaceGridColl3D_SetFloat", "Sets a float value on the Grid by Attribute name.");
-#endif
 		OutFunctions.Add(Sig);
 	}
 
@@ -1496,20 +1426,17 @@ void UNiagaraDataInterfaceGrid3DCollection::GetFunctions(TArray<FNiagaraFunction
 		Sig.ModuleUsageBitmask = ENiagaraScriptUsageMask::Particle;
 		Sig.bSupportsCPU = false;
 		Sig.bSupportsGPU = true;
-#if WITH_EDITORONLY_DATA
 		Sig.Description = NSLOCTEXT("Niagara", "NiagaraDataInterfaceGridColl3D_SetFloat", "Sets a float value on the Grid by Attribute name.");
-#endif
 		OutFunctions.Add(Sig);
 	}
 
-#if WITH_EDITORONLY_DATA
 	for (int32 i = StartIndex; i < OutFunctions.Num(); i++)
 	{
 		FNiagaraFunctionSignature& Function = OutFunctions[i];
 		Function.FunctionVersion = FNiagaraGridCollection3DDIFunctionVersion::LatestVersion;
 	}
-#endif
 }
+#endif
 
 // #todo(dmp): expose more CPU functionality
 // #todo(dmp): ideally these would be exposed on the parent class, but we can't bind functions of parent classes but need to work on the interface
@@ -4268,4 +4195,3 @@ void UNiagaraDataInterfaceGrid3DCollection::CollectAttributesForScript(UNiagaraS
 }
 
 #undef LOCTEXT_NAMESPACE
-
diff --git a/Engine/Plugins/FX/Niagara/Source/Niagara/Private/NiagaraDataInterfaceGrid3DCollectionReader.cpp b/Engine/Plugins/FX/Niagara/Source/Niagara/Private/NiagaraDataInterfaceGrid3DCollectionReader.cpp
--- a/Engine/Plugins/FX/Niagara/Source/Niagara/Private/NiagaraDataInterfaceGrid3DCollectionReader.cpp
+++ b/Engine/Plugins/FX/Niagara/Source/Niagara/Private/NiagaraDataInterfaceGrid3DCollectionReader.cpp
@@ -167,11 +167,12 @@ void UNiagaraDataInterfaceGrid3DCollectionReader::GetEmitterDependencies(UNiagar
 	}
 }
 
-void UNiagaraDataInterfaceGrid3DCollectionReader::GetFunctions(TArray<FNiagaraFunctionSignature>& OutFunctions)
+#if WITH_EDITORONLY_DATA
+void UNiagaraDataInterfaceGrid3DCollectionReader::GetFunctionsInternal(TArray<FNiagaraFunctionSignature>& OutFunctions) const
 {
 	int32 NumFunctionsBefore = OutFunctions.Num();
 
-	Super::GetFunctions(OutFunctions);
+	Super::GetFunctionsInternal(OutFunctions);
 	TSet<FName> FunctionsToRemove = 
 	{
 		ClearCellFunctionName,
@@ -201,6 +202,7 @@ void UNiagaraDataInterfaceGrid3DCollectionReader::GetFunctions(TArray<FNiagaraFu
 		}
 	}	
 }
+#endif
 
 bool UNiagaraDataInterfaceGrid3DCollectionReader::PerInstanceTickPostSimulate(void* PerInstanceData, FNiagaraSystemInstance* SystemInstance, float DeltaSeconds)
 {
diff --git a/Engine/Plugins/FX/Niagara/Source/Niagara/Private/NiagaraDataInterfaceIntRenderTarget2D.cpp b/Engine/Plugins/FX/Niagara/Source/Niagara/Private/NiagaraDataInterfaceIntRenderTarget2D.cpp
--- a/Engine/Plugins/FX/Niagara/Source/Niagara/Private/NiagaraDataInterfaceIntRenderTarget2D.cpp
+++ b/Engine/Plugins/FX/Niagara/Source/Niagara/Private/NiagaraDataInterfaceIntRenderTarget2D.cpp
@@ -187,9 +187,10 @@ void UNiagaraDataInterfaceIntRenderTarget2D::PostInitProperties()
 	}
 }
 
-void UNiagaraDataInterfaceIntRenderTarget2D::GetFunctions(TArray<FNiagaraFunctionSignature>& OutFunctions)
+#if WITH_EDITORONLY_DATA
+void UNiagaraDataInterfaceIntRenderTarget2D::GetFunctionsInternal(TArray<FNiagaraFunctionSignature>& OutFunctions) const
 {
-	Super::GetFunctions(OutFunctions);
+	Super::GetFunctionsInternal(OutFunctions);
 
 	const int32 EmitterSystemOnlyBitmask = ENiagaraScriptUsageMask::Emitter | ENiagaraScriptUsageMask::System;
 	OutFunctions.Reserve(OutFunctions.Num() + NDIIntRenderTarget2DLocal::NumFunctions);
@@ -208,10 +209,8 @@ void UNiagaraDataInterfaceIntRenderTarget2D::GetFunctions(TArray<FNiagaraFunctio
 		Sig.bWriteFunction = true;
 		Sig.bSupportsCPU = false;
 		Sig.bSupportsGPU = true;
-#if WITH_EDITORONLY_DATA
 		Sig.Description = LOCTEXT("GetValueDesc", "Gets the value from the render target at the pixel offset");
 		Sig.FunctionVersion = NDIIntRenderTarget2DLocal::FFunctionVersion::LatestVersion;
-#endif
 	}
 	{
 		FNiagaraFunctionSignature& Sig = OutFunctions.AddDefaulted_GetRef();
@@ -228,10 +227,8 @@ void UNiagaraDataInterfaceIntRenderTarget2D::GetFunctions(TArray<FNiagaraFunctio
 		Sig.bWriteFunction = true;
 		Sig.bSupportsCPU = false;
 		Sig.bSupportsGPU = true;
-#if WITH_EDITORONLY_DATA
 		Sig.Description = LOCTEXT("SetValueDesc", "Sets the value on the render target at the pixel offset.");
 		Sig.FunctionVersion = NDIIntRenderTarget2DLocal::FFunctionVersion::LatestVersion;
-#endif
 	}
 	{
 		FNiagaraFunctionSignature& Sig = OutFunctions.AddDefaulted_GetRef();
@@ -250,10 +247,8 @@ void UNiagaraDataInterfaceIntRenderTarget2D::GetFunctions(TArray<FNiagaraFunctio
 		Sig.bWriteFunction = true;
 		Sig.bSupportsCPU = false;
 		Sig.bSupportsGPU = true;
-#if WITH_EDITORONLY_DATA
 		Sig.Description = LOCTEXT("AtomicAddDesc", "Atomic min the value to the pixel at the offset, returns the current & previous values.  This opertion is thread safe.");
 		Sig.FunctionVersion = NDIIntRenderTarget2DLocal::FFunctionVersion::LatestVersion;
-#endif
 	}
 	////static const FName AtomicAndFunctionName("AtomicAnd");
 	{
@@ -274,10 +269,8 @@ void UNiagaraDataInterfaceIntRenderTarget2D::GetFunctions(TArray<FNiagaraFunctio
 		Sig.bWriteFunction = true;
 		Sig.bSupportsCPU = false;
 		Sig.bSupportsGPU = true;
-#if WITH_EDITORONLY_DATA
 		Sig.Description = LOCTEXT("AtomicCASFunctionDesc", "Compares the pixel value against the comparison value, if they are equal the value is replaced.  Original Value is the pixel value before the operation completes.  This opertion is thread safe.");
 		Sig.FunctionVersion = NDIIntRenderTarget2DLocal::FFunctionVersion::LatestVersion;
-#endif
 	}
 	////static const FName AtomicCSFunctionName("AtomicCompareStore");
 	////static const FName AtomicExchangeFunctionName("AtomicExchange");
@@ -298,10 +291,8 @@ void UNiagaraDataInterfaceIntRenderTarget2D::GetFunctions(TArray<FNiagaraFunctio
 		Sig.bWriteFunction = true;
 		Sig.bSupportsCPU = false;
 		Sig.bSupportsGPU = true;
-#if WITH_EDITORONLY_DATA
 		Sig.Description = LOCTEXT("AtomicMaxDesc", "Atomic max the value to the pixel at the offset, returns the current & previous values.  This opertion is thread safe.");
 		Sig.FunctionVersion = NDIIntRenderTarget2DLocal::FFunctionVersion::LatestVersion;
-#endif
 	}
 	{
 		FNiagaraFunctionSignature& Sig = OutFunctions.AddDefaulted_GetRef();
@@ -320,10 +311,8 @@ void UNiagaraDataInterfaceIntRenderTarget2D::GetFunctions(TArray<FNiagaraFunctio
 		Sig.bWriteFunction = true;
 		Sig.bSupportsCPU = false;
 		Sig.bSupportsGPU = true;
-#if WITH_EDITORONLY_DATA
 		Sig.Description = LOCTEXT("AtomicMinDesc", "Atomic min the value to the pixel at the offset, returns the current & previous values.  This opertion is thread safe.");
 		Sig.FunctionVersion = NDIIntRenderTarget2DLocal::FFunctionVersion::LatestVersion;
-#endif
 	}
 	////static const FName AtomicOrFunctionName("AtomicOr");
 	////static const FName AtomicXorunctionName("AtomicXor");
@@ -339,10 +328,8 @@ void UNiagaraDataInterfaceIntRenderTarget2D::GetFunctions(TArray<FNiagaraFunctio
 		Sig.bRequiresContext = false;
 		Sig.bSupportsCPU = true;
 		Sig.bSupportsGPU = true;
-#if WITH_EDITORONLY_DATA
 		Sig.Description = LOCTEXT("GetSizeDesc", "Gets the size of the rendertarget");
 		Sig.FunctionVersion = NDIIntRenderTarget2DLocal::FFunctionVersion::LatestVersion;
-#endif
 	}
 	{
 		FNiagaraFunctionSignature& Sig = OutFunctions.AddDefaulted_GetRef();
@@ -358,10 +345,8 @@ void UNiagaraDataInterfaceIntRenderTarget2D::GetFunctions(TArray<FNiagaraFunctio
 		Sig.bRequiresExecPin = true;
 		Sig.bSupportsCPU = true;
 		Sig.bSupportsGPU = false;
-#if WITH_EDITORONLY_DATA
 		Sig.Description = LOCTEXT("SetSizeDesc", "Sets the size of the rendertarget");
 		Sig.FunctionVersion = NDIIntRenderTarget2DLocal::FFunctionVersion::LatestVersion;
-#endif
 	}
 
 	{
@@ -376,10 +361,8 @@ void UNiagaraDataInterfaceIntRenderTarget2D::GetFunctions(TArray<FNiagaraFunctio
 		Sig.bRequiresContext = false;
 		Sig.bSupportsCPU = false;
 		Sig.bSupportsGPU = true;
-#if WITH_EDITORONLY_DATA
 		Sig.Description = LOCTEXT("LinearToIndexDesc", "Converts a linear index into a pixel coordinate");
 		Sig.FunctionVersion = NDIIntRenderTarget2DLocal::FFunctionVersion::LatestVersion;
-#endif
 	}
 	{
 		FNiagaraFunctionSignature& Sig = OutFunctions.AddDefaulted_GetRef();
@@ -392,10 +375,8 @@ void UNiagaraDataInterfaceIntRenderTarget2D::GetFunctions(TArray<FNiagaraFunctio
 		Sig.bRequiresContext = false;
 		Sig.bSupportsCPU = false;
 		Sig.bSupportsGPU = true;
-#if WITH_EDITORONLY_DATA
 		Sig.Description = LOCTEXT("LinearToUVDesc", "Converts a linear index into a UV coordinate");
 		Sig.FunctionVersion = NDIIntRenderTarget2DLocal::FFunctionVersion::LatestVersion;
-#endif
 	}
 
 	{
@@ -409,10 +390,8 @@ void UNiagaraDataInterfaceIntRenderTarget2D::GetFunctions(TArray<FNiagaraFunctio
 		Sig.bRequiresContext = false;
 		Sig.bSupportsCPU = false;
 		Sig.bSupportsGPU = true;
-#if WITH_EDITORONLY_DATA
 		Sig.Description = LOCTEXT("ExecToIndexDesc", "Returns the execution index as a pixel coordinate");
 		Sig.FunctionVersion = NDIIntRenderTarget2DLocal::FFunctionVersion::LatestVersion;
-#endif
 	}
 	{
 		FNiagaraFunctionSignature& Sig = OutFunctions.AddDefaulted_GetRef();
@@ -424,14 +403,11 @@ void UNiagaraDataInterfaceIntRenderTarget2D::GetFunctions(TArray<FNiagaraFunctio
 		Sig.bRequiresContext = false;
 		Sig.bSupportsCPU = false;
 		Sig.bSupportsGPU = true;
-#if WITH_EDITORONLY_DATA
 		Sig.Description = LOCTEXT("ExecToUVDesc", "Returns the execution index as a UV coordinate");
 		Sig.FunctionVersion = NDIIntRenderTarget2DLocal::FFunctionVersion::LatestVersion;
-#endif
 	}
 }
 
-#if WITH_EDITORONLY_DATA
 bool UNiagaraDataInterfaceIntRenderTarget2D::UpgradeFunctionCall(FNiagaraFunctionSignature& FunctionSignature)
 {
 	bool bWasChanged = false;
diff --git a/Engine/Plugins/FX/Niagara/Source/Niagara/Private/NiagaraDataInterfaceLandscape.cpp b/Engine/Plugins/FX/Niagara/Source/Niagara/Private/NiagaraDataInterfaceLandscape.cpp
--- a/Engine/Plugins/FX/Niagara/Source/Niagara/Private/NiagaraDataInterfaceLandscape.cpp
+++ b/Engine/Plugins/FX/Niagara/Source/Niagara/Private/NiagaraDataInterfaceLandscape.cpp
@@ -1176,7 +1176,8 @@ bool UNiagaraDataInterfaceLandscape::Equals(const UNiagaraDataInterface* Other)
 		&& OtherLandscape->bVirtualTexturesSupported == bVirtualTexturesSupported;
 }
 
-void UNiagaraDataInterfaceLandscape::GetFunctions(TArray<FNiagaraFunctionSignature>& OutFunctions)
+#if WITH_EDITORONLY_DATA
+void UNiagaraDataInterfaceLandscape::GetFunctionsInternal(TArray<FNiagaraFunctionSignature>& OutFunctions) const
 {
 	{
 		FNiagaraFunctionSignature Sig;
@@ -1238,22 +1239,19 @@ void UNiagaraDataInterfaceLandscape::GetFunctions(TArray<FNiagaraFunctionSignatu
 		OutFunctions.Add(Sig);
 	}
 
-#if WITH_EDITORONLY_DATA
 	for (FNiagaraFunctionSignature& Sig : OutFunctions)
 	{
 		Sig.FunctionVersion = NiagaraDataInterfaceLandscape::LatestVersion;
 	}
-#endif
 }
 
-#if WITH_EDITORONLY_DATA
 bool UNiagaraDataInterfaceLandscape::UpgradeFunctionCall(FNiagaraFunctionSignature& FunctionSignature)
 {
 	// always upgrade to the latest version
 	if (FunctionSignature.FunctionVersion < NiagaraDataInterfaceLandscape::LatestVersion)
 	{
 		TArray<FNiagaraFunctionSignature> AllFunctions;
-		GetFunctions(AllFunctions);
+		GetFunctionsInternal(AllFunctions);
 		for (const FNiagaraFunctionSignature& Sig : AllFunctions)
 		{
 			if (FunctionSignature.Name == Sig.Name)
diff --git a/Engine/Plugins/FX/Niagara/Source/Niagara/Private/NiagaraDataInterfaceMeshRendererInfo.cpp b/Engine/Plugins/FX/Niagara/Source/Niagara/Private/NiagaraDataInterfaceMeshRendererInfo.cpp
--- a/Engine/Plugins/FX/Niagara/Source/Niagara/Private/NiagaraDataInterfaceMeshRendererInfo.cpp
+++ b/Engine/Plugins/FX/Niagara/Source/Niagara/Private/NiagaraDataInterfaceMeshRendererInfo.cpp
@@ -131,15 +131,14 @@ void UNiagaraDataInterfaceMeshRendererInfo::OnMeshRendererChanged(UNiagaraMeshRe
 	UpdateCachedData();
 }
 
-void UNiagaraDataInterfaceMeshRendererInfo::GetFunctions(TArray<FNiagaraFunctionSignature>& OutFunctions)
+#if WITH_EDITORONLY_DATA
+void UNiagaraDataInterfaceMeshRendererInfo::GetFunctionsInternal(TArray<FNiagaraFunctionSignature>& OutFunctions) const
 {
 	{
 		FNiagaraFunctionSignature& Signature = OutFunctions.AddDefaulted_GetRef();
 		Signature.Name = NDIMeshRendererInfoInternal::GetNumMeshesName;
-#if WITH_EDITORONLY_DATA
 		Signature.FunctionVersion = (uint32)ENDIMeshRendererInfoVersion::LatestVersion;
 		Signature.Description = LOCTEXT("GetNumMeshesInRendererDesc", "Retrieves the number of meshes on the mesh renderer by index, or -1 if the index is invalid.");
-#endif
 		Signature.bMemberFunction = true;
 		Signature.Inputs.Add(FNiagaraVariable(FNiagaraTypeDefinition(GetClass()), TEXT("MeshRendererInfo")));
 		Signature.Outputs.Add(FNiagaraVariable(FNiagaraTypeDefinition::GetIntDef(), TEXT("OutNumMeshes")));
@@ -147,10 +146,8 @@ void UNiagaraDataInterfaceMeshRendererInfo::GetFunctions(TArray<FNiagaraFunction
 	{
 		FNiagaraFunctionSignature& Signature = OutFunctions.AddDefaulted_GetRef();
 		Signature.Name = NDIMeshRendererInfoInternal::GetMeshLocalBoundsName;
-#if WITH_EDITORONLY_DATA
 		Signature.FunctionVersion = (uint32)ENDIMeshRendererInfoVersion::LatestVersion;
 		Signature.Description = LOCTEXT("GetMeshLocalBoundsDesc", "Retrieves the local bounds of the specified mesh's vertices.");
-#endif
 		Signature.bMemberFunction = true;
 		Signature.Inputs.Add(FNiagaraVariable(FNiagaraTypeDefinition(GetClass()), TEXT("MeshRendererInfo")));
 		Signature.Inputs.Add(FNiagaraVariable(FNiagaraTypeDefinition::GetIntDef(), TEXT("MeshIndex")));
@@ -161,15 +158,14 @@ void UNiagaraDataInterfaceMeshRendererInfo::GetFunctions(TArray<FNiagaraFunction
 	{
 		FNiagaraFunctionSignature& Signature = OutFunctions.AddDefaulted_GetRef();
 		Signature.Name = NDIMeshRendererInfoInternal::GetSubUVDetailsName;
-#if WITH_EDITORONLY_DATA
 		Signature.FunctionVersion = (uint32)ENDIMeshRendererInfoVersion::LatestVersion;
-#endif
 		Signature.bMemberFunction = true;
 		Signature.Inputs.Emplace(FNiagaraTypeDefinition(GetClass()), TEXT("MeshRendererInfo"));
 		Signature.Outputs.Emplace(FNiagaraTypeDefinition::GetBoolDef(), TEXT("BlendEnabled"));
 		Signature.Outputs.Emplace(FNiagaraTypeDefinition::GetVec2Def(), TEXT("SubImageSize"));
 	}
 }
+#endif
 
 void UNiagaraDataInterfaceMeshRendererInfo::GetVMExternalFunction(const FVMExternalFunctionBindingInfo& BindingInfo, void* InstanceData, FVMExternalFunction& OutFunc)
 {
diff --git a/Engine/Plugins/FX/Niagara/Source/Niagara/Private/NiagaraDataInterfaceNeighborGrid3D.cpp b/Engine/Plugins/FX/Niagara/Source/Niagara/Private/NiagaraDataInterfaceNeighborGrid3D.cpp
--- a/Engine/Plugins/FX/Niagara/Source/Niagara/Private/NiagaraDataInterfaceNeighborGrid3D.cpp
+++ b/Engine/Plugins/FX/Niagara/Source/Niagara/Private/NiagaraDataInterfaceNeighborGrid3D.cpp
@@ -68,9 +68,10 @@ UNiagaraDataInterfaceNeighborGrid3D::UNiagaraDataInterfaceNeighborGrid3D(FObject
 	Proxy.Reset(new FNiagaraDataInterfaceProxyNeighborGrid3D());	
 }
 
-void UNiagaraDataInterfaceNeighborGrid3D::GetFunctions(TArray<FNiagaraFunctionSignature>& OutFunctions)
+#if WITH_EDITORONLY_DATA
+void UNiagaraDataInterfaceNeighborGrid3D::GetFunctionsInternal(TArray<FNiagaraFunctionSignature>& OutFunctions) const
 {
-	Super::GetFunctions(OutFunctions);
+	Super::GetFunctionsInternal(OutFunctions);
 
 	{
 		FNiagaraFunctionSignature Sig;
@@ -193,9 +194,8 @@ void UNiagaraDataInterfaceNeighborGrid3D::GetFunctions(TArray<FNiagaraFunctionSi
 		Sig.bRequiresExecPin = true;
 		OutFunctions.Add(Sig);
 	}
-
-
 }
+#endif
 
 DEFINE_NDI_DIRECT_FUNC_BINDER(UNiagaraDataInterfaceNeighborGrid3D, SetNumCells);
 void UNiagaraDataInterfaceNeighborGrid3D::GetVMExternalFunction(const FVMExternalFunctionBindingInfo& BindingInfo, void* InstanceData, FVMExternalFunction &OutFunc)
diff --git a/Engine/Plugins/FX/Niagara/Source/Niagara/Private/NiagaraDataInterfaceOcclusion.cpp b/Engine/Plugins/FX/Niagara/Source/Niagara/Private/NiagaraDataInterfaceOcclusion.cpp
--- a/Engine/Plugins/FX/Niagara/Source/Niagara/Private/NiagaraDataInterfaceOcclusion.cpp
+++ b/Engine/Plugins/FX/Niagara/Source/Niagara/Private/NiagaraDataInterfaceOcclusion.cpp
@@ -55,7 +55,8 @@ void UNiagaraDataInterfaceOcclusion::PostInitProperties()
 	}
 }
 
-void UNiagaraDataInterfaceOcclusion::GetFunctions(TArray<FNiagaraFunctionSignature>& OutFunctions)
+#if WITH_EDITORONLY_DATA
+void UNiagaraDataInterfaceOcclusion::GetFunctionsInternal(TArray<FNiagaraFunctionSignature>& OutFunctions) const
 {
 	using namespace NDIOcclusionImpl;
 
@@ -146,7 +147,6 @@ void UNiagaraDataInterfaceOcclusion::GetFunctions(TArray<FNiagaraFunctionSignatu
 	}
 }
 
-#if WITH_EDITORONLY_DATA
 bool UNiagaraDataInterfaceOcclusion::AppendCompileHash(FNiagaraCompileHashVisitor* InVisitor) const
 {
 	using namespace NDIOcclusionImpl;
@@ -183,7 +183,7 @@ bool UNiagaraDataInterfaceOcclusion::UpgradeFunctionCall(FNiagaraFunctionSignatu
 	if (FunctionSignature.FunctionVersion < FNiagaraOcclusionDIFunctionVersion::LatestVersion)
 	{
 		TArray<FNiagaraFunctionSignature> AllFunctions;
-		GetFunctions(AllFunctions);
+		GetFunctionsInternal(AllFunctions);
 		for (const FNiagaraFunctionSignature& Sig : AllFunctions)
 		{
 			if (FunctionSignature.Name == Sig.Name)
diff --git a/Engine/Plugins/FX/Niagara/Source/Niagara/Private/NiagaraDataInterfaceParticleRead.cpp b/Engine/Plugins/FX/Niagara/Source/Niagara/Private/NiagaraDataInterfaceParticleRead.cpp
--- a/Engine/Plugins/FX/Niagara/Source/Niagara/Private/NiagaraDataInterfaceParticleRead.cpp
+++ b/Engine/Plugins/FX/Niagara/Source/Niagara/Private/NiagaraDataInterfaceParticleRead.cpp
@@ -498,7 +498,8 @@ int32 UNiagaraDataInterfaceParticleRead::PerInstanceDataSize() const
 	return sizeof(FNDIParticleRead_InstanceData);
 }
 
-void UNiagaraDataInterfaceParticleRead::GetFunctions(TArray<FNiagaraFunctionSignature>& OutFunctions)
+#if WITH_EDITORONLY_DATA
+void UNiagaraDataInterfaceParticleRead::GetFunctionsInternal(TArray<FNiagaraFunctionSignature>& OutFunctions) const
 {
 	// Misc functionality
 	{
@@ -509,9 +510,7 @@ void UNiagaraDataInterfaceParticleRead::GetFunctions(TArray<FNiagaraFunctionSign
 		Sig.bRequiresContext = false;
 		Sig.bExperimental = true;
 		Sig.Outputs.Emplace(FNiagaraTypeDefinition::GetBoolDef(), TEXT("IsLocalSpace"));
-	#if WITH_EDITORONLY_DATA
 		Sig.SetDescription(LOCTEXT("GetEmitterLocalSpace", "Returns if the emitter is using local space or world space."));
-	#endif
 	}
 
 	//
@@ -525,9 +524,7 @@ void UNiagaraDataInterfaceParticleRead::GetFunctions(TArray<FNiagaraFunctionSign
 		Sig.bMemberFunction = true;
 		Sig.bRequiresContext = false;
 		Sig.bExperimental = true;
-#if WITH_EDITORONLY_DATA
 		Sig.Description = NSLOCTEXT("Niagara", "NiagaraDataInterfaceParticleRead_GetNumSpawnedDesc", "Returns the current number of spawned particles this frame.  Note: This is the request spawn amount, when killing particles during spawn it will not be accurate.");
-#endif
 		OutFunctions.Add(Sig);
 	}
 
@@ -541,9 +538,7 @@ void UNiagaraDataInterfaceParticleRead::GetFunctions(TArray<FNiagaraFunctionSign
 		Sig.bMemberFunction = true;
 		Sig.bRequiresContext = false;
 		Sig.bExperimental = true;
-#if WITH_EDITORONLY_DATA
 		Sig.Description = NSLOCTEXT("Niagara", "NiagaraDataInterfaceParticleRead_GetSpawnedIDAtIndexDesc", "Returns the Niagara ID for the particle spawned at Index.");
-#endif
 		OutFunctions.Add(Sig);
 	}
 
@@ -554,24 +549,20 @@ void UNiagaraDataInterfaceParticleRead::GetFunctions(TArray<FNiagaraFunctionSign
 		Sig.Outputs.Add(FNiagaraVariable(FNiagaraTypeDefinition::GetIntDef(), TEXT("Num Particles")));
 		Sig.bMemberFunction = true;
 		Sig.bRequiresContext = false;
-#if WITH_EDITORONLY_DATA
 		Sig.Description = NSLOCTEXT("Niagara", "NiagaraDataInterfaceParticleRead_GetNumParticlesDesc", "Returns the current number of particles in the buffer.  Note: When reading from self this will be the previous frames data.");
-#endif
 		OutFunctions.Add(Sig);
 	}
 
 	GetPersistentIDFunctions(OutFunctions);
 	GetIndexFunctions(OutFunctions);
 
-#if WITH_EDITORONLY_DATA
 	for (FNiagaraFunctionSignature& Function : OutFunctions)
 	{
 		Function.FunctionVersion = FNiagaraParticleReadDIFunctionVersion::LatestVersion;
 	}
-#endif
 }
 
-void UNiagaraDataInterfaceParticleRead::GetPersistentIDFunctions(TArray<FNiagaraFunctionSignature>& OutFunctions)
+void UNiagaraDataInterfaceParticleRead::GetPersistentIDFunctions(TArray<FNiagaraFunctionSignature>& OutFunctions) const
 {
 	{
 		FNiagaraFunctionSignature Sig;
@@ -782,7 +773,7 @@ void UNiagaraDataInterfaceParticleRead::GetPersistentIDFunctions(TArray<FNiagara
 }
 
 
-void UNiagaraDataInterfaceParticleRead::GetIndexFunctions(TArray<FNiagaraFunctionSignature>& OutFunctions)
+void UNiagaraDataInterfaceParticleRead::GetIndexFunctions(TArray<FNiagaraFunctionSignature>& OutFunctions) const
 {
 
 	//
@@ -958,6 +949,7 @@ void UNiagaraDataInterfaceParticleRead::GetIndexFunctions(TArray<FNiagaraFunctio
 		OutFunctions.Add(Sig);
 	}
 }
+#endif
 
 DEFINE_NDI_DIRECT_FUNC_BINDER(UNiagaraDataInterfaceParticleRead, GetNumSpawnedParticles);
 DEFINE_NDI_DIRECT_FUNC_BINDER(UNiagaraDataInterfaceParticleRead, GetSpawnedIDAtIndex);
diff --git a/Engine/Plugins/FX/Niagara/Source/Niagara/Private/NiagaraDataInterfacePlatformSet.cpp b/Engine/Plugins/FX/Niagara/Source/Niagara/Private/NiagaraDataInterfacePlatformSet.cpp
--- a/Engine/Plugins/FX/Niagara/Source/Niagara/Private/NiagaraDataInterfacePlatformSet.cpp
+++ b/Engine/Plugins/FX/Niagara/Source/Niagara/Private/NiagaraDataInterfacePlatformSet.cpp
@@ -22,7 +22,8 @@ void UNiagaraDataInterfacePlatformSet::PostInitProperties()
 
 static FName IsActiveName(TEXT("IsActive"));
 
-void UNiagaraDataInterfacePlatformSet::GetFunctions(TArray<FNiagaraFunctionSignature>& OutFunctions)
+#if WITH_EDITORONLY_DATA
+void UNiagaraDataInterfacePlatformSet::GetFunctionsInternal(TArray<FNiagaraFunctionSignature>& OutFunctions) const
 {
 	{
 		FNiagaraFunctionSignature Sig;
@@ -36,6 +37,7 @@ void UNiagaraDataInterfacePlatformSet::GetFunctions(TArray<FNiagaraFunctionSigna
 		OutFunctions.Add(Sig);
 	}
 }
+#endif
 
 DEFINE_NDI_DIRECT_FUNC_BINDER(UNiagaraDataInterfacePlatformSet, IsActive);
 void UNiagaraDataInterfacePlatformSet::GetVMExternalFunction(const FVMExternalFunctionBindingInfo& BindingInfo, void* InstanceData, FVMExternalFunction& OutFunc)
diff --git a/Engine/Plugins/FX/Niagara/Source/Niagara/Private/NiagaraDataInterfaceRW.cpp b/Engine/Plugins/FX/Niagara/Source/Niagara/Private/NiagaraDataInterfaceRW.cpp
--- a/Engine/Plugins/FX/Niagara/Source/Niagara/Private/NiagaraDataInterfaceRW.cpp
+++ b/Engine/Plugins/FX/Niagara/Source/Niagara/Private/NiagaraDataInterfaceRW.cpp
@@ -193,7 +193,8 @@ UNiagaraDataInterfaceGrid3D::UNiagaraDataInterfaceGrid3D(FObjectInitializer cons
 {
 }
 
-void UNiagaraDataInterfaceGrid3D::GetFunctions(TArray<FNiagaraFunctionSignature>& OutFunctions)
+#if WITH_EDITORONLY_DATA
+void UNiagaraDataInterfaceGrid3D::GetFunctionsInternal(TArray<FNiagaraFunctionSignature>& OutFunctions) const
 {
 
 	{
@@ -351,6 +352,7 @@ void UNiagaraDataInterfaceGrid3D::GetFunctions(TArray<FNiagaraFunctionSignature>
 		OutFunctions.Add(Sig);
 	}
 }
+#endif
 
 void UNiagaraDataInterfaceGrid3D::GetVMExternalFunction(const FVMExternalFunctionBindingInfo& BindingInfo, void* InstanceData, FVMExternalFunction &OutFunc)
 {
@@ -636,7 +638,7 @@ void UNiagaraDataInterfaceGrid2D::Serialize(FArchive& Ar)
 void UNiagaraDataInterfaceGrid2D::ValidateFunction(const FNiagaraFunctionSignature& Function, TArray<FText>& OutValidationErrors)
 {
 	TArray<FNiagaraFunctionSignature> DIFuncs;
-	GetFunctions(DIFuncs);
+	GetFunctionsInternal(DIFuncs);
 	
 	// All the deprecated grid2d functions
 	TSet<FName> DeprecatedFunctionNames;
@@ -651,9 +653,7 @@ void UNiagaraDataInterfaceGrid2D::ValidateFunction(const FNiagaraFunctionSignatu
 	Super::ValidateFunction(Function, OutValidationErrors);
 }
 
-#endif
-
-void UNiagaraDataInterfaceGrid2D::GetFunctions(TArray<FNiagaraFunctionSignature>& OutFunctions)
+void UNiagaraDataInterfaceGrid2D::GetFunctionsInternal(TArray<FNiagaraFunctionSignature>& OutFunctions) const
 {	
 	{
 		FNiagaraFunctionSignature Sig;
@@ -831,7 +831,7 @@ void UNiagaraDataInterfaceGrid2D::GetFunctions(TArray<FNiagaraFunctionSignature>
 		OutFunctions.Add(Sig);
 	}
 }
-
+#endif
 
 void UNiagaraDataInterfaceGrid2D::GetVMExternalFunction(const FVMExternalFunctionBindingInfo& BindingInfo, void* InstanceData, FVMExternalFunction &OutFunc)
 {
diff --git a/Engine/Plugins/FX/Niagara/Source/Niagara/Private/NiagaraDataInterfaceRasterizationGrid3D.cpp b/Engine/Plugins/FX/Niagara/Source/Niagara/Private/NiagaraDataInterfaceRasterizationGrid3D.cpp
--- a/Engine/Plugins/FX/Niagara/Source/Niagara/Private/NiagaraDataInterfaceRasterizationGrid3D.cpp
+++ b/Engine/Plugins/FX/Niagara/Source/Niagara/Private/NiagaraDataInterfaceRasterizationGrid3D.cpp
@@ -73,11 +73,12 @@ UNiagaraDataInterfaceRasterizationGrid3D::UNiagaraDataInterfaceRasterizationGrid
 	Proxy.Reset(new FNiagaraDataInterfaceProxyRasterizationGrid3D());	
 }
 
-void UNiagaraDataInterfaceRasterizationGrid3D::GetFunctions(TArray<FNiagaraFunctionSignature>& OutFunctions)
+#if WITH_EDITORONLY_DATA
+void UNiagaraDataInterfaceRasterizationGrid3D::GetFunctionsInternal(TArray<FNiagaraFunctionSignature>& OutFunctions) const
 {
 	using namespace NDIRasterizationGrid3DLocal;
 
-	Super::GetFunctions(OutFunctions);
+	Super::GetFunctionsInternal(OutFunctions);
 
 	{
 		FNiagaraFunctionSignature Sig;
@@ -151,9 +152,7 @@ void UNiagaraDataInterfaceRasterizationGrid3D::GetFunctions(TArray<FNiagaraFunct
 		
 		Sig.bSupportsCPU = false;
 		Sig.bSupportsGPU = true;
-#if WITH_EDITORONLY_DATA
 		Sig.Description = NSLOCTEXT("Niagara", "NiagaraDataInterfaceGridColl2D_SetValueFunction", "Set the value at a specific index. Note that this is an older way of working with Grids. Consider using the SetFloat or other typed, named functions or parameter map variables with StackContext namespace instead.");
-#endif
 		OutFunctions.Add(Sig);
 	}
 
@@ -176,9 +175,7 @@ void UNiagaraDataInterfaceRasterizationGrid3D::GetFunctions(TArray<FNiagaraFunct
 		
 		Sig.bSupportsCPU = false;
 		Sig.bSupportsGPU = true;
-#if WITH_EDITORONLY_DATA
 		Sig.Description = NSLOCTEXT("Niagara", "NiagaraDataInterfaceGridColl2D_SetValueFunction", "Set the value at a specific index. Note that this is an older way of working with Grids. Consider using the SetFloat or other typed, named functions or parameter map variables with StackContext namespace instead.");
-#endif
 		OutFunctions.Add(Sig);
 	}
 
@@ -204,9 +201,7 @@ void UNiagaraDataInterfaceRasterizationGrid3D::GetFunctions(TArray<FNiagaraFunct
 		
 		Sig.bSoftDeprecatedFunction = true;
 		Sig.bHidden = true;
-#if WITH_EDITORONLY_DATA
 		Sig.Description = NSLOCTEXT("Niagara", "NiagaraDataInterfaceGridColl2D_SetValueFunction", "Set the value at a specific index. Note that this is an older way of working with Grids. Consider using the SetFloat or other typed, named functions or parameter map variables with StackContext namespace instead.");
-#endif
 		OutFunctions.Add(Sig);
 	}
 
@@ -229,9 +224,7 @@ void UNiagaraDataInterfaceRasterizationGrid3D::GetFunctions(TArray<FNiagaraFunct
 
 		Sig.bSupportsCPU = false;
 		Sig.bSupportsGPU = true;
-#if WITH_EDITORONLY_DATA
 		Sig.Description = NSLOCTEXT("Niagara", "NiagaraDataInterfaceGridColl2D_SetValueFunction", "Set the value at a specific index. Note that this is an older way of working with Grids. Consider using the SetFloat or other typed, named functions or parameter map variables with StackContext namespace instead.");
-#endif
 		OutFunctions.Add(Sig);
 	}
 
@@ -254,9 +247,7 @@ void UNiagaraDataInterfaceRasterizationGrid3D::GetFunctions(TArray<FNiagaraFunct
 
 		Sig.bSupportsCPU = false;
 		Sig.bSupportsGPU = true;
-#if WITH_EDITORONLY_DATA
 		Sig.Description = NSLOCTEXT("Niagara", "NiagaraDataInterfaceGridColl2D_SetValueFunction", "Set the value at a specific index. Note that this is an older way of working with Grids. Consider using the SetFloat or other typed, named functions or parameter map variables with StackContext namespace instead.");
-#endif
 		OutFunctions.Add(Sig);
 	}
 
@@ -279,9 +270,7 @@ void UNiagaraDataInterfaceRasterizationGrid3D::GetFunctions(TArray<FNiagaraFunct
 
 		Sig.bSupportsCPU = false;
 		Sig.bSupportsGPU = true;
-#if WITH_EDITORONLY_DATA
 		Sig.Description = NSLOCTEXT("Niagara", "NiagaraDataInterfaceGridColl2D_SetValueFunction", "Set the value at a specific index. Note that this is an older way of working with Grids. Consider using the SetFloat or other typed, named functions or parameter map variables with StackContext namespace instead.");
-#endif
 		OutFunctions.Add(Sig);
 	}
 
@@ -307,9 +296,7 @@ void UNiagaraDataInterfaceRasterizationGrid3D::GetFunctions(TArray<FNiagaraFunct
 
 		Sig.bSoftDeprecatedFunction = true;
 		Sig.bHidden = true;
-#if WITH_EDITORONLY_DATA
 		Sig.Description = NSLOCTEXT("Niagara", "NiagaraDataInterfaceGridColl2D_SetValueFunction", "Set the value at a specific index. Note that this is an older way of working with Grids. Consider using the SetFloat or other typed, named functions or parameter map variables with StackContext namespace instead.");
-#endif
 		OutFunctions.Add(Sig);
 	}
 
@@ -329,9 +316,7 @@ void UNiagaraDataInterfaceRasterizationGrid3D::GetFunctions(TArray<FNiagaraFunct
 		Sig.bRequiresContext = false;
 		Sig.bWriteFunction = true;
 		Sig.ModuleUsageBitmask = ENiagaraScriptUsageMask::Particle;
-#if WITH_EDITORONLY_DATA
 		Sig.Description = NSLOCTEXT("Niagara", "NiagaraDataInterfaceGridColl2D_SetValueFunction", "Set the value at a specific index. Note that this is an older way of working with Grids. Consider using the SetFloat or other typed, named functions or parameter map variables with StackContext namespace instead.");
-#endif
 		OutFunctions.Add(Sig);
 	}
 
@@ -356,9 +341,7 @@ void UNiagaraDataInterfaceRasterizationGrid3D::GetFunctions(TArray<FNiagaraFunct
 		Sig.bSupportsGPU = true;
 		Sig.bSoftDeprecatedFunction = true;
 		Sig.bHidden = true;
-#if WITH_EDITORONLY_DATA
 		Sig.Description = NSLOCTEXT("Niagara", "NiagaraDataInterfaceGridColl2D_SetValueFunction", "Set the value at a specific index. Note that this is an older way of working with Grids. Consider using the SetFloat or other typed, named functions or parameter map variables with StackContext namespace instead.");
-#endif
 		OutFunctions.Add(Sig);
 	}
 
@@ -381,9 +364,7 @@ void UNiagaraDataInterfaceRasterizationGrid3D::GetFunctions(TArray<FNiagaraFunct
 
 		Sig.bSupportsCPU = false;
 		Sig.bSupportsGPU = true;
-#if WITH_EDITORONLY_DATA
 		Sig.Description = NSLOCTEXT("Niagara", "NiagaraDataInterfaceGridColl2D_SetValueFunction", "Set the value at a specific index. Note that this is an older way of working with Grids. Consider using the SetFloat or other typed, named functions or parameter map variables with StackContext namespace instead.");
-#endif
 		OutFunctions.Add(Sig);
 	}
 
@@ -405,9 +386,7 @@ void UNiagaraDataInterfaceRasterizationGrid3D::GetFunctions(TArray<FNiagaraFunct
 		
 		Sig.bSupportsCPU = false;
 		Sig.bSupportsGPU = true;
-#if WITH_EDITORONLY_DATA
 		Sig.Description = NSLOCTEXT("Niagara", "NiagaraDataInterfaceGridColl3D_GetValueFunction", "Get the value at a specific index. Note that this is an older way of working with Grids. Consider using the SetFloat or other typed, named functions or parameter map variables with StackContext namespace instead.");
-#endif
 
 		OutFunctions.Add(Sig);
 	}
@@ -430,13 +409,12 @@ void UNiagaraDataInterfaceRasterizationGrid3D::GetFunctions(TArray<FNiagaraFunct
 		
 		Sig.bSupportsCPU = false;
 		Sig.bSupportsGPU = true;
-#if WITH_EDITORONLY_DATA
 		Sig.Description = NSLOCTEXT("Niagara", "NiagaraDataInterfaceGridColl3D_GetValueFunction", "Get the value at a specific index. Note that this is an older way of working with Grids. Consider using the SetFloat or other typed, named functions or parameter map variables with StackContext namespace instead.");
-#endif
 
 		OutFunctions.Add(Sig);
 	}
 }
+#endif // WITH_EDITORONLY_DATA
 
 void UNiagaraDataInterfaceRasterizationGrid3D::GetVMExternalFunction(const FVMExternalFunctionBindingInfo& BindingInfo, void* InstanceData, FVMExternalFunction &OutFunc)
 {
diff --git a/Engine/Plugins/FX/Niagara/Source/Niagara/Private/NiagaraDataInterfaceRenderTarget2D.cpp b/Engine/Plugins/FX/Niagara/Source/Niagara/Private/NiagaraDataInterfaceRenderTarget2D.cpp
--- a/Engine/Plugins/FX/Niagara/Source/Niagara/Private/NiagaraDataInterfaceRenderTarget2D.cpp
+++ b/Engine/Plugins/FX/Niagara/Source/Niagara/Private/NiagaraDataInterfaceRenderTarget2D.cpp
@@ -118,11 +118,12 @@ void UNiagaraDataInterfaceRenderTarget2D::PostInitProperties()
 	}
 }
 
-void UNiagaraDataInterfaceRenderTarget2D::GetFunctions(TArray<FNiagaraFunctionSignature>& OutFunctions)
+#if WITH_EDITORONLY_DATA
+void UNiagaraDataInterfaceRenderTarget2D::GetFunctionsInternal(TArray<FNiagaraFunctionSignature>& OutFunctions) const
 {
 	using namespace NDIRenderTarget2DLocal;
 
-	Super::GetFunctions(OutFunctions);
+	Super::GetFunctionsInternal(OutFunctions);
 
 	const int32 EmitterSystemOnlyBitmask = ENiagaraScriptUsageMask::Emitter | ENiagaraScriptUsageMask::System;
 	OutFunctions.Reserve(OutFunctions.Num() + 4);
@@ -235,7 +236,6 @@ void UNiagaraDataInterfaceRenderTarget2D::GetFunctions(TArray<FNiagaraFunctionSi
 	}
 }
 
-#if WITH_EDITORONLY_DATA
 bool UNiagaraDataInterfaceRenderTarget2D::UpgradeFunctionCall(FNiagaraFunctionSignature& FunctionSignature)
 {
 	using namespace NDIRenderTarget2DLocal;
diff --git a/Engine/Plugins/FX/Niagara/Source/Niagara/Private/NiagaraDataInterfaceRenderTarget2DArray.cpp b/Engine/Plugins/FX/Niagara/Source/Niagara/Private/NiagaraDataInterfaceRenderTarget2DArray.cpp
--- a/Engine/Plugins/FX/Niagara/Source/Niagara/Private/NiagaraDataInterfaceRenderTarget2DArray.cpp
+++ b/Engine/Plugins/FX/Niagara/Source/Niagara/Private/NiagaraDataInterfaceRenderTarget2DArray.cpp
@@ -106,11 +106,12 @@ void UNiagaraDataInterfaceRenderTarget2DArray::PostInitProperties()
 	}
 }
 
-void UNiagaraDataInterfaceRenderTarget2DArray::GetFunctions(TArray<FNiagaraFunctionSignature>& OutFunctions)
+#if WITH_EDITORONLY_DATA
+void UNiagaraDataInterfaceRenderTarget2DArray::GetFunctionsInternal(TArray<FNiagaraFunctionSignature>& OutFunctions) const
 {
 	using namespace NDIRenderTarget2DArrayLocal;
 
-	Super::GetFunctions(OutFunctions);
+	Super::GetFunctionsInternal(OutFunctions);
 
 	const int32 EmitterSystemOnlyBitmask = ENiagaraScriptUsageMask::Emitter | ENiagaraScriptUsageMask::System;
 	OutFunctions.Reserve(OutFunctions.Num() + 4);
@@ -231,7 +232,6 @@ void UNiagaraDataInterfaceRenderTarget2DArray::GetFunctions(TArray<FNiagaraFunct
 	}
 }
 
-#if WITH_EDITORONLY_DATA
 bool UNiagaraDataInterfaceRenderTarget2DArray::UpgradeFunctionCall(FNiagaraFunctionSignature& FunctionSignature)
 {
 	using namespace NDIRenderTarget2DArrayLocal;
diff --git a/Engine/Plugins/FX/Niagara/Source/Niagara/Private/NiagaraDataInterfaceRenderTargetCube.cpp b/Engine/Plugins/FX/Niagara/Source/Niagara/Private/NiagaraDataInterfaceRenderTargetCube.cpp
--- a/Engine/Plugins/FX/Niagara/Source/Niagara/Private/NiagaraDataInterfaceRenderTargetCube.cpp
+++ b/Engine/Plugins/FX/Niagara/Source/Niagara/Private/NiagaraDataInterfaceRenderTargetCube.cpp
@@ -103,11 +103,12 @@ void UNiagaraDataInterfaceRenderTargetCube::PostInitProperties()
 	}
 }
 
-void UNiagaraDataInterfaceRenderTargetCube::GetFunctions(TArray<FNiagaraFunctionSignature>& OutFunctions)
+#if WITH_EDITORONLY_DATA
+void UNiagaraDataInterfaceRenderTargetCube::GetFunctionsInternal(TArray<FNiagaraFunctionSignature>& OutFunctions) const
 {
 	using namespace NDIRenderTargetCubeLocal;
 
-	Super::GetFunctions(OutFunctions);
+	Super::GetFunctionsInternal(OutFunctions);
 
 	const int32 EmitterSystemOnlyBitmask = ENiagaraScriptUsageMask::Emitter | ENiagaraScriptUsageMask::System;
 	OutFunctions.Reserve(OutFunctions.Num() + 7);
@@ -202,7 +203,6 @@ void UNiagaraDataInterfaceRenderTargetCube::GetFunctions(TArray<FNiagaraFunction
 	}
 }
 
-#if WITH_EDITORONLY_DATA
 bool UNiagaraDataInterfaceRenderTargetCube::UpgradeFunctionCall(FNiagaraFunctionSignature& FunctionSignature)
 {
 	using namespace NDIRenderTargetCubeLocal;
diff --git a/Engine/Plugins/FX/Niagara/Source/Niagara/Private/NiagaraDataInterfaceRenderTargetVolume.cpp b/Engine/Plugins/FX/Niagara/Source/Niagara/Private/NiagaraDataInterfaceRenderTargetVolume.cpp
--- a/Engine/Plugins/FX/Niagara/Source/Niagara/Private/NiagaraDataInterfaceRenderTargetVolume.cpp
+++ b/Engine/Plugins/FX/Niagara/Source/Niagara/Private/NiagaraDataInterfaceRenderTargetVolume.cpp
@@ -164,11 +164,12 @@ void UNiagaraDataInterfaceRenderTargetVolume::PostInitProperties()
 	}
 }
 
-void UNiagaraDataInterfaceRenderTargetVolume::GetFunctions(TArray<FNiagaraFunctionSignature>& OutFunctions)
+#if WITH_EDITORONLY_DATA
+void UNiagaraDataInterfaceRenderTargetVolume::GetFunctionsInternal(TArray<FNiagaraFunctionSignature>& OutFunctions) const
 {
 	using namespace NDIRenderTargetVolumeLocal;
 
-	Super::GetFunctions(OutFunctions);
+	Super::GetFunctionsInternal(OutFunctions);
 
 	const int32 EmitterSystemOnlyBitmask = ENiagaraScriptUsageMask::Emitter | ENiagaraScriptUsageMask::System;
 	OutFunctions.Reserve(OutFunctions.Num() + 7);
@@ -287,7 +288,6 @@ void UNiagaraDataInterfaceRenderTargetVolume::GetFunctions(TArray<FNiagaraFuncti
 	}
 }
 
-#if WITH_EDITORONLY_DATA
 bool UNiagaraDataInterfaceRenderTargetVolume::UpgradeFunctionCall(FNiagaraFunctionSignature& FunctionSignature)
 {
 	using namespace NDIRenderTargetVolumeLocal;
diff --git a/Engine/Plugins/FX/Niagara/Source/Niagara/Private/NiagaraDataInterfaceRigidMeshCollisionQuery.cpp b/Engine/Plugins/FX/Niagara/Source/Niagara/Private/NiagaraDataInterfaceRigidMeshCollisionQuery.cpp
--- a/Engine/Plugins/FX/Niagara/Source/Niagara/Private/NiagaraDataInterfaceRigidMeshCollisionQuery.cpp
+++ b/Engine/Plugins/FX/Niagara/Source/Niagara/Private/NiagaraDataInterfaceRigidMeshCollisionQuery.cpp
@@ -1278,7 +1278,8 @@ void UNiagaraDataInterfaceRigidMeshCollisionQuery::PostLoad()
 #endif
 }
 
-void UNiagaraDataInterfaceRigidMeshCollisionQuery::GetFunctions(TArray<FNiagaraFunctionSignature>& OutFunctions)
+#if WITH_EDITORONLY_DATA
+void UNiagaraDataInterfaceRigidMeshCollisionQuery::GetFunctionsInternal(TArray<FNiagaraFunctionSignature>& OutFunctions) const
 {
 	using namespace NDIRigidMeshCollisionLocal;
 
@@ -1507,6 +1508,7 @@ void UNiagaraDataInterfaceRigidMeshCollisionQuery::GetFunctions(TArray<FNiagaraF
 		OutFunctions.Add(Sig);
 	}
 }
+#endif
 
 DEFINE_NDI_DIRECT_FUNC_BINDER(UNiagaraDataInterfaceRigidMeshCollisionQuery, FindActorsCPU);
 
diff --git a/Engine/Plugins/FX/Niagara/Source/Niagara/Private/NiagaraDataInterfaceSkeletalMesh.cpp b/Engine/Plugins/FX/Niagara/Source/Niagara/Private/NiagaraDataInterfaceSkeletalMesh.cpp
--- a/Engine/Plugins/FX/Niagara/Source/Niagara/Private/NiagaraDataInterfaceSkeletalMesh.cpp
+++ b/Engine/Plugins/FX/Niagara/Source/Niagara/Private/NiagaraDataInterfaceSkeletalMesh.cpp
@@ -2301,10 +2301,10 @@ bool UNiagaraDataInterfaceSkeletalMesh::CanEditChange(const FProperty* InPropert
 
 	return true;
 }
-
 #endif //WITH_EDITOR
 
-void UNiagaraDataInterfaceSkeletalMesh::GetFunctions(TArray<FNiagaraFunctionSignature>& OutFunctions)
+#if WITH_EDITORONLY_DATA
+void UNiagaraDataInterfaceSkeletalMesh::GetFunctionsInternal(TArray<FNiagaraFunctionSignature>& OutFunctions) const
 {
 	const int32 FirstFunction = OutFunctions.Num();
 
@@ -2326,14 +2326,14 @@ void UNiagaraDataInterfaceSkeletalMesh::GetFunctions(TArray<FNiagaraFunctionSign
 	GetVertexSamplingFunctions(OutFunctions);
 	GetSkeletonSamplingFunctions(OutFunctions);
 
-#if WITH_EDITORONLY_DATA
 	for ( int i=FirstFunction; i < OutFunctions.Num(); ++i )
 	{
 		OutFunctions[i].FunctionVersion = FNiagaraSkelMeshDIFunctionVersion::LatestVersion;
 	}
-#endif
 }
 
+#endif //WITH_EDITORONLY_DATA
+
 void UNiagaraDataInterfaceSkeletalMesh::GetVMExternalFunction(const FVMExternalFunctionBindingInfo& BindingInfo, void* InstanceData, FVMExternalFunction &OutFunc)
 {
 	FNDISkeletalMesh_InstanceData* InstData = (FNDISkeletalMesh_InstanceData*)InstanceData;
@@ -2694,7 +2694,7 @@ static const FName GetTriPositionVelocityAndNormalBinormalTangentWSName_DEPRECAT
 void UNiagaraDataInterfaceSkeletalMesh::ValidateFunction(const FNiagaraFunctionSignature& Function, TArray<FText>& OutValidationErrors)
 {
 	TArray<FNiagaraFunctionSignature> DIFuncs;
-	GetFunctions(DIFuncs);
+	GetFunctionsInternal(DIFuncs);
 
 	if (!DIFuncs.Contains(Function))
 	{
diff --git a/Engine/Plugins/FX/Niagara/Source/Niagara/Private/NiagaraDataInterfaceSparseVolumeTexture.cpp b/Engine/Plugins/FX/Niagara/Source/Niagara/Private/NiagaraDataInterfaceSparseVolumeTexture.cpp
--- a/Engine/Plugins/FX/Niagara/Source/Niagara/Private/NiagaraDataInterfaceSparseVolumeTexture.cpp
+++ b/Engine/Plugins/FX/Niagara/Source/Niagara/Private/NiagaraDataInterfaceSparseVolumeTexture.cpp
@@ -88,7 +88,8 @@ void UNiagaraDataInterfaceSparseVolumeTexture::Serialize(FArchive& Ar)
 	Super::Serialize(Ar);
 }
 
-void UNiagaraDataInterfaceSparseVolumeTexture::GetFunctions(TArray<FNiagaraFunctionSignature>& OutFunctions)
+#if WITH_EDITORONLY_DATA
+void UNiagaraDataInterfaceSparseVolumeTexture::GetFunctionsInternal(TArray<FNiagaraFunctionSignature>& OutFunctions) const
 {
 	FNiagaraFunctionSignature DefaultGpuSig;
 	DefaultGpuSig.bMemberFunction = true;
@@ -137,6 +138,7 @@ void UNiagaraDataInterfaceSparseVolumeTexture::GetFunctions(TArray<FNiagaraFunct
 		Sig.SetDescription(LOCTEXT("SparseVolumeGetNumMipLevelsDesc", "Get the number of mip levels."));
 	}
 }
+#endif
 
 void UNiagaraDataInterfaceSparseVolumeTexture::GetVMExternalFunction(const FVMExternalFunctionBindingInfo& BindingInfo, void* InstanceData, FVMExternalFunction& OutFunc)
 {
diff --git a/Engine/Plugins/FX/Niagara/Source/Niagara/Private/NiagaraDataInterfaceSpline.cpp b/Engine/Plugins/FX/Niagara/Source/Niagara/Private/NiagaraDataInterfaceSpline.cpp
--- a/Engine/Plugins/FX/Niagara/Source/Niagara/Private/NiagaraDataInterfaceSpline.cpp
+++ b/Engine/Plugins/FX/Niagara/Source/Niagara/Private/NiagaraDataInterfaceSpline.cpp
@@ -134,7 +134,8 @@ void UNiagaraDataInterfaceSpline::PostInitProperties()
 	}
 }
 
-void UNiagaraDataInterfaceSpline::GetFunctions(TArray<FNiagaraFunctionSignature>& OutFunctions)
+#if WITH_EDITORONLY_DATA
+void UNiagaraDataInterfaceSpline::GetFunctionsInternal(TArray<FNiagaraFunctionSignature>& OutFunctions) const
 {
 	FNiagaraFunctionSignature DefaultSig;
 	DefaultSig.Inputs.Add(FNiagaraVariable(FNiagaraTypeDefinition(GetClass()), TEXT("Spline")));
@@ -266,6 +267,7 @@ void UNiagaraDataInterfaceSpline::GetFunctions(TArray<FNiagaraFunctionSignature>
 		Sig.Outputs.Add(FNiagaraVariable(FNiagaraTypeDefinition::GetFloatDef(), TEXT("Length")));
 	}
 }
+#endif
 
 DEFINE_NDI_FUNC_BINDER(UNiagaraDataInterfaceSpline, SampleSplinePositionByUnitDistance);
 DEFINE_NDI_FUNC_BINDER(UNiagaraDataInterfaceSpline, SampleSplineRotationByUnitDistance);
@@ -403,7 +405,7 @@ bool UNiagaraDataInterfaceSpline::UpgradeFunctionCall(FNiagaraFunctionSignature&
 	if (FunctionSignature.FunctionVersion < FNiagaraSplineDIFunctionVersion::LWCConversion)
 	{
 		TArray<FNiagaraFunctionSignature> AllFunctions;
-		GetFunctions(AllFunctions);
+		GetFunctionsInternal(AllFunctions);
 		for (const FNiagaraFunctionSignature& Sig : AllFunctions)
 		{
 			if (FunctionSignature.Name == Sig.Name)
diff --git a/Engine/Plugins/FX/Niagara/Source/Niagara/Private/NiagaraDataInterfaceSpriteRendererInfo.cpp b/Engine/Plugins/FX/Niagara/Source/Niagara/Private/NiagaraDataInterfaceSpriteRendererInfo.cpp
--- a/Engine/Plugins/FX/Niagara/Source/Niagara/Private/NiagaraDataInterfaceSpriteRendererInfo.cpp
+++ b/Engine/Plugins/FX/Niagara/Source/Niagara/Private/NiagaraDataInterfaceSpriteRendererInfo.cpp
@@ -140,7 +140,8 @@ void UNiagaraDataInterfaceSpriteRendererInfo::OnSpriteRendererChanged(UNiagaraSp
 	MarkRenderDataDirty();
 }
 
-void UNiagaraDataInterfaceSpriteRendererInfo::GetFunctions(TArray<FNiagaraFunctionSignature>& OutFunctions)
+#if WITH_EDITORONLY_DATA
+void UNiagaraDataInterfaceSpriteRendererInfo::GetFunctionsInternal(TArray<FNiagaraFunctionSignature>& OutFunctions) const
 {
 	FNiagaraFunctionSignature DefaultImmutableSignature;
 	DefaultImmutableSignature.bMemberFunction = true;
@@ -173,6 +174,7 @@ void UNiagaraDataInterfaceSpriteRendererInfo::GetFunctions(TArray<FNiagaraFuncti
 		Sig.Outputs.Emplace(FNiagaraTypeDefinition::GetVec2Def(), TEXT("SubImageSize"));
 	}
 }
+#endif
 
 void UNiagaraDataInterfaceSpriteRendererInfo::GetVMExternalFunction(const FVMExternalFunctionBindingInfo& BindingInfo, void* InstanceData, FVMExternalFunction& OutFunc)
 {
diff --git a/Engine/Plugins/FX/Niagara/Source/Niagara/Private/NiagaraDataInterfaceTexture.cpp b/Engine/Plugins/FX/Niagara/Source/Niagara/Private/NiagaraDataInterfaceTexture.cpp
--- a/Engine/Plugins/FX/Niagara/Source/Niagara/Private/NiagaraDataInterfaceTexture.cpp
+++ b/Engine/Plugins/FX/Niagara/Source/Niagara/Private/NiagaraDataInterfaceTexture.cpp
@@ -147,7 +147,8 @@ bool UNiagaraDataInterfaceTexture::Equals(const UNiagaraDataInterface* Other) co
 		OtherTexture->TextureUserParameter == TextureUserParameter;
 }
 
-void UNiagaraDataInterfaceTexture::GetFunctions(TArray<FNiagaraFunctionSignature>& OutFunctions)
+#if WITH_EDITORONLY_DATA
+void UNiagaraDataInterfaceTexture::GetFunctionsInternal(TArray<FNiagaraFunctionSignature>& OutFunctions) const
 {
 	FNiagaraFunctionSignature DefaultGpuSig;
 	DefaultGpuSig.bMemberFunction = true;
@@ -217,6 +218,7 @@ void UNiagaraDataInterfaceTexture::GetFunctions(TArray<FNiagaraFunctionSignature
 		Sig.SetFunctionVersion(FNDITextureFunctionVersion::LatestVersion);
 	}
 }
+#endif
 
 void UNiagaraDataInterfaceTexture::GetVMExternalFunction(const FVMExternalFunctionBindingInfo& BindingInfo, void* InstanceData, FVMExternalFunction &OutFunc)
 {
diff --git a/Engine/Plugins/FX/Niagara/Source/Niagara/Private/NiagaraDataInterfaceVector2DCurve.cpp b/Engine/Plugins/FX/Niagara/Source/Niagara/Private/NiagaraDataInterfaceVector2DCurve.cpp
--- a/Engine/Plugins/FX/Niagara/Source/Niagara/Private/NiagaraDataInterfaceVector2DCurve.cpp
+++ b/Engine/Plugins/FX/Niagara/Source/Niagara/Private/NiagaraDataInterfaceVector2DCurve.cpp
@@ -138,7 +138,8 @@ void UNiagaraDataInterfaceVector2DCurve::GetCurveData(TArray<FCurveData>& OutCur
 	OutCurveData.Add(FCurveData(&YCurve, TEXT("Y"), FLinearColor(0.05f, 1.0f, 0.05f)));
 }
 
-void UNiagaraDataInterfaceVector2DCurve::GetFunctions(TArray<FNiagaraFunctionSignature>& OutFunctions)
+#if WITH_EDITORONLY_DATA
+void UNiagaraDataInterfaceVector2DCurve::GetFunctionsInternal(TArray<FNiagaraFunctionSignature>& OutFunctions) const
 {
 	FNiagaraFunctionSignature Sig;
 	Sig.Name = SampleCurveName;
@@ -156,7 +157,6 @@ void UNiagaraDataInterfaceVector2DCurve::GetFunctions(TArray<FNiagaraFunctionSig
 // the HLSL in the spirit of a static switch
 // TODO: need a way to identify each specific function here
 // 
-#if WITH_EDITORONLY_DATA
 bool UNiagaraDataInterfaceVector2DCurve::GetFunctionHLSL(const FNiagaraDataInterfaceGPUParamInfo& ParamInfo, const FNiagaraDataInterfaceGeneratedFunction& FunctionInfo, int FunctionInstanceIndex, FString& OutHLSL)
 {
 	if (FunctionInfo.DefinitionName == SampleCurveName)
diff --git a/Engine/Plugins/FX/Niagara/Source/Niagara/Private/NiagaraDataInterfaceVector4Curve.cpp b/Engine/Plugins/FX/Niagara/Source/Niagara/Private/NiagaraDataInterfaceVector4Curve.cpp
--- a/Engine/Plugins/FX/Niagara/Source/Niagara/Private/NiagaraDataInterfaceVector4Curve.cpp
+++ b/Engine/Plugins/FX/Niagara/Source/Niagara/Private/NiagaraDataInterfaceVector4Curve.cpp
@@ -160,7 +160,8 @@ void UNiagaraDataInterfaceVector4Curve::GetCurveData(TArray<FCurveData>& OutCurv
 	OutCurveData.Add(FCurveData(&WCurve, TEXT("W"), FLinearColor(0.2f, 0.2f, 0.2f)));
 }
 
-void UNiagaraDataInterfaceVector4Curve::GetFunctions(TArray<FNiagaraFunctionSignature>& OutFunctions)
+#if WITH_EDITORONLY_DATA
+void UNiagaraDataInterfaceVector4Curve::GetFunctionsInternal(TArray<FNiagaraFunctionSignature>& OutFunctions) const
 {
 	FNiagaraFunctionSignature Sig;
 	Sig.Name = SampleCurveName;
@@ -178,7 +179,6 @@ void UNiagaraDataInterfaceVector4Curve::GetFunctions(TArray<FNiagaraFunctionSign
 // the HLSL in the spirit of a static switch
 // TODO: need a way to identify each specific function here
 // 
-#if WITH_EDITORONLY_DATA
 bool UNiagaraDataInterfaceVector4Curve::GetFunctionHLSL(const FNiagaraDataInterfaceGPUParamInfo& ParamInfo, const FNiagaraDataInterfaceGeneratedFunction& FunctionInfo, int FunctionInstanceIndex, FString& OutHLSL)
 {
 	if (FunctionInfo.DefinitionName == SampleCurveName)
diff --git a/Engine/Plugins/FX/Niagara/Source/Niagara/Private/NiagaraDataInterfaceVectorCurve.cpp b/Engine/Plugins/FX/Niagara/Source/Niagara/Private/NiagaraDataInterfaceVectorCurve.cpp
--- a/Engine/Plugins/FX/Niagara/Source/Niagara/Private/NiagaraDataInterfaceVectorCurve.cpp
+++ b/Engine/Plugins/FX/Niagara/Source/Niagara/Private/NiagaraDataInterfaceVectorCurve.cpp
@@ -154,7 +154,8 @@ void UNiagaraDataInterfaceVectorCurve::GetCurveData(TArray<FCurveData>& OutCurve
 	OutCurveData.Add(FCurveData(&ZCurve, TEXT("Z"), FLinearColor(0.1f, 0.2f, 1.0f)));
 }
 
-void UNiagaraDataInterfaceVectorCurve::GetFunctions(TArray<FNiagaraFunctionSignature>& OutFunctions)
+#if WITH_EDITORONLY_DATA
+void UNiagaraDataInterfaceVectorCurve::GetFunctionsInternal(TArray<FNiagaraFunctionSignature>& OutFunctions) const
 {
 	FNiagaraFunctionSignature Sig;
 	Sig.Name = SampleCurveName;
@@ -172,7 +173,6 @@ void UNiagaraDataInterfaceVectorCurve::GetFunctions(TArray<FNiagaraFunctionSigna
 // the HLSL in the spirit of a static switch
 // TODO: need a way to identify each specific function here
 // 
-#if WITH_EDITORONLY_DATA
 void UNiagaraDataInterfaceVectorCurve::SyncCurvesToAsset()
 {
 	if (UCurveLinearColor* ColorCurve = Cast<UCurveLinearColor>(CurveAsset))
diff --git a/Engine/Plugins/FX/Niagara/Source/Niagara/Private/NiagaraDataInterfaceVectorField.cpp b/Engine/Plugins/FX/Niagara/Source/Niagara/Private/NiagaraDataInterfaceVectorField.cpp
--- a/Engine/Plugins/FX/Niagara/Source/Niagara/Private/NiagaraDataInterfaceVectorField.cpp
+++ b/Engine/Plugins/FX/Niagara/Source/Niagara/Private/NiagaraDataInterfaceVectorField.cpp
@@ -104,7 +104,8 @@ void UNiagaraDataInterfaceVectorField::PostInitProperties()
 
 /*--------------------------------------------------------------------------------------------------------------------------*/
 
-void UNiagaraDataInterfaceVectorField::GetFunctions(TArray<FNiagaraFunctionSignature>& OutFunctions)
+#if WITH_EDITORONLY_DATA
+void UNiagaraDataInterfaceVectorField::GetFunctionsInternal(TArray<FNiagaraFunctionSignature>& OutFunctions) const
 {
 	{
 		FNiagaraFunctionSignature Sig;
@@ -166,6 +167,7 @@ void UNiagaraDataInterfaceVectorField::GetFunctions(TArray<FNiagaraFunctionSigna
 		OutFunctions.Add(Sig);
 	}
 }
+#endif
 
 DEFINE_NDI_DIRECT_FUNC_BINDER(UNiagaraDataInterfaceVectorField, SampleVectorField);
 DEFINE_NDI_DIRECT_FUNC_BINDER(UNiagaraDataInterfaceVectorField, LoadVectorField);
diff --git a/Engine/Plugins/FX/Niagara/Source/Niagara/Private/NiagaraDataInterfaceVolumeCache.cpp b/Engine/Plugins/FX/Niagara/Source/Niagara/Private/NiagaraDataInterfaceVolumeCache.cpp
--- a/Engine/Plugins/FX/Niagara/Source/Niagara/Private/NiagaraDataInterfaceVolumeCache.cpp
+++ b/Engine/Plugins/FX/Niagara/Source/Niagara/Private/NiagaraDataInterfaceVolumeCache.cpp
@@ -201,7 +201,8 @@ void UNiagaraDataInterfaceVolumeCache::ProvidePerInstanceDataForRenderThread(voi
 	DataFromGT->NumCells = InstanceData->NumCells;
 }
 
-void UNiagaraDataInterfaceVolumeCache::GetFunctions(TArray<FNiagaraFunctionSignature>& OutFunctions)
+#if WITH_EDITORONLY_DATA
+void UNiagaraDataInterfaceVolumeCache::GetFunctionsInternal(TArray<FNiagaraFunctionSignature>& OutFunctions) const
 {	
 	{
 		FNiagaraFunctionSignature Sig;
@@ -315,6 +316,7 @@ void UNiagaraDataInterfaceVolumeCache::GetFunctions(TArray<FNiagaraFunctionSigna
 		OutFunctions.Add(Sig);
 	}
 }
+#endif
 
 void UNiagaraDataInterfaceVolumeCache::SetFrame(FVectorVMExternalFunctionContext& Context)
 {
diff --git a/Engine/Plugins/FX/Niagara/Source/Niagara/Private/NiagaraDataInterfaceVolumeTexture.cpp b/Engine/Plugins/FX/Niagara/Source/Niagara/Private/NiagaraDataInterfaceVolumeTexture.cpp
--- a/Engine/Plugins/FX/Niagara/Source/Niagara/Private/NiagaraDataInterfaceVolumeTexture.cpp
+++ b/Engine/Plugins/FX/Niagara/Source/Niagara/Private/NiagaraDataInterfaceVolumeTexture.cpp
@@ -102,7 +102,8 @@ bool UNiagaraDataInterfaceVolumeTexture::Equals(const UNiagaraDataInterface* Oth
 		OtherTexture->TextureUserParameter == TextureUserParameter;
 }
 
-void UNiagaraDataInterfaceVolumeTexture::GetFunctions(TArray<FNiagaraFunctionSignature>& OutFunctions)
+#if WITH_EDITORONLY_DATA
+void UNiagaraDataInterfaceVolumeTexture::GetFunctionsInternal(TArray<FNiagaraFunctionSignature>& OutFunctions) const
 {
 	{
 		FNiagaraFunctionSignature & Sig = OutFunctions.AddDefaulted_GetRef();
@@ -144,6 +145,7 @@ void UNiagaraDataInterfaceVolumeTexture::GetFunctions(TArray<FNiagaraFunctionSig
 		Sig.SetDescription(LOCTEXT("TextureDimsDesc", "Get the dimensions of mip 0 of the texture."));
 	}
 }
+#endif
 
 DEFINE_NDI_DIRECT_FUNC_BINDER(UNiagaraDataInterfaceVolumeTexture, SampleVolumeTexture)
 void UNiagaraDataInterfaceVolumeTexture::GetVMExternalFunction(const FVMExternalFunctionBindingInfo& BindingInfo, void* InstanceData, FVMExternalFunction &OutFunc)
diff --git a/Engine/Plugins/FX/Niagara/Source/Niagara/Public/NiagaraDataInterfacePhysicsAsset.h b/Engine/Plugins/FX/Niagara/Source/Niagara/Public/NiagaraDataInterfacePhysicsAsset.h
--- a/Engine/Plugins/FX/Niagara/Source/Niagara/Public/NiagaraDataInterfacePhysicsAsset.h
+++ b/Engine/Plugins/FX/Niagara/Source/Niagara/Public/NiagaraDataInterfacePhysicsAsset.h
@@ -134,7 +134,6 @@ public:
 	NIAGARA_API virtual void PostLoad() override;
 
 	/** UNiagaraDataInterface Interface */
-	NIAGARA_API virtual void GetFunctions(TArray<FNiagaraFunctionSignature>& OutFunctions) override;
 	NIAGARA_API virtual void GetVMExternalFunction(const FVMExternalFunctionBindingInfo& BindingInfo, void* InstanceData, FVMExternalFunction& OutFunc) override;
 	virtual bool CanExecuteOnTarget(ENiagaraSimTarget Target) const override { return Target == ENiagaraSimTarget::GPUComputeSim; }
 	NIAGARA_API virtual bool InitPerInstanceData(void* PerInstanceData, FNiagaraSystemInstance* SystemInstance) override;
@@ -192,6 +191,9 @@ public:
 	NIAGARA_API void GetProjectionPoint(FVectorVMExternalFunctionContext& Context);
 
 protected:
+#if WITH_EDITORONLY_DATA
+	NIAGARA_API virtual void GetFunctionsInternal(TArray<FNiagaraFunctionSignature>& OutFunctions) const override;
+#endif
 	/** Copy one niagara DI to this */
 	NIAGARA_API virtual bool CopyToInternal(UNiagaraDataInterface* Destination) const override;
 };
diff --git a/Engine/Plugins/FX/Niagara/Source/Niagara/Public/NiagaraDataInterfaceRigidMeshCollisionQuery.h b/Engine/Plugins/FX/Niagara/Source/Niagara/Public/NiagaraDataInterfaceRigidMeshCollisionQuery.h
--- a/Engine/Plugins/FX/Niagara/Source/Niagara/Public/NiagaraDataInterfaceRigidMeshCollisionQuery.h
+++ b/Engine/Plugins/FX/Niagara/Source/Niagara/Public/NiagaraDataInterfaceRigidMeshCollisionQuery.h
@@ -227,7 +227,6 @@ public:
 	virtual void PostLoad() override;
 
 	/** UNiagaraDataInterface Interface */
-	virtual void GetFunctions(TArray<FNiagaraFunctionSignature>& OutFunctions) override;
 	virtual void GetVMExternalFunction(const FVMExternalFunctionBindingInfo& BindingInfo, void* InstanceData, FVMExternalFunction& OutFunc) override;
 	virtual bool CanExecuteOnTarget(ENiagaraSimTarget Target) const override { return Target == ENiagaraSimTarget::GPUComputeSim; }
 #if WITH_NIAGARA_DEBUGGER
@@ -275,6 +274,9 @@ public:
 #endif
 
 protected:
+#if WITH_EDITORONLY_DATA
+	NIAGARA_API virtual void GetFunctionsInternal(TArray<FNiagaraFunctionSignature>& OutFunctions) const override;
+#endif
 	/** Copy one niagara DI to this */
 	virtual bool CopyToInternal(UNiagaraDataInterface* Destination) const override;
 
diff --git a/Engine/Plugins/FX/Niagara/Source/NiagaraEditor/Private/EdGraphSchema_Niagara.cpp b/Engine/Plugins/FX/Niagara/Source/NiagaraEditor/Private/EdGraphSchema_Niagara.cpp
--- a/Engine/Plugins/FX/Niagara/Source/NiagaraEditor/Private/EdGraphSchema_Niagara.cpp
+++ b/Engine/Plugins/FX/Niagara/Source/NiagaraEditor/Private/EdGraphSchema_Niagara.cpp
@@ -799,7 +799,7 @@ PRAGMA_ENABLE_DEPRECATION_WARNINGS
 		{
 			FText MenuCat = FromPinDataInterface->GetClass()->GetDisplayNameText();
 			TArray<FNiagaraFunctionSignature> Functions;
-			FromPinDataInterface->GetFunctions(Functions);
+			FromPinDataInterface->GetFunctionSignatures(Functions);
 			for (FNiagaraFunctionSignature& Sig : Functions)
 			{
 				if (Sig.bSoftDeprecatedFunction || Sig.bHidden)
@@ -2181,7 +2181,7 @@ TArray<TTuple<FString, FString, FText>> UEdGraphSchema_Niagara::GetDataInterface
 			
 			// No functions, don't add an option
 			TArray<FNiagaraFunctionSignature> FunctionSignatures;
-			DataInterface->GetFunctions(FunctionSignatures);
+			DataInterface->GetFunctionSignatures(FunctionSignatures);
 			if (FunctionSignatures.Num() == 0)
 			{
 				return FunctionPrototypes;
diff --git a/Engine/Plugins/FX/Niagara/Source/NiagaraEditor/Private/NiagaraHlslTranslator.cpp b/Engine/Plugins/FX/Niagara/Source/NiagaraEditor/Private/NiagaraHlslTranslator.cpp
--- a/Engine/Plugins/FX/Niagara/Source/NiagaraEditor/Private/NiagaraHlslTranslator.cpp
+++ b/Engine/Plugins/FX/Niagara/Source/NiagaraEditor/Private/NiagaraHlslTranslator.cpp
@@ -8135,7 +8135,7 @@ void TNiagaraHlslTranslator<GraphBridge>::ProcessCustomHlsl(const FString& InCus
 			// actual custom hlsl source. If they do, then add them to the function table that we need to map.
 			FNiagaraScriptDataInterfaceCompileInfo& Info = CompilationOutput.ScriptData.DataInterfaceInfo[OwnerIdx];
 			TArray<FNiagaraFunctionSignature> Funcs;
-			CDO->GetFunctions(Funcs);
+			CDO->GetFunctionSignatures(Funcs);
 
 			TArray<FString> SanitizedFunctionNames;
 			SanitizedFunctionNames.Reserve(Funcs.Num());
@@ -8797,7 +8797,7 @@ void TNiagaraHlslTranslator<GraphBridge>::RegisterFunctionCall(ENiagaraScriptUsa
 			if (OutSignature.bMemberFunction)
 			{
 				TArray<FNiagaraFunctionSignature> DataInterfaceFunctions;
-				CDO->GetFunctions(DataInterfaceFunctions);
+				CDO->GetFunctionSignatures(DataInterfaceFunctions);
 
 				const int32 FoundMatch = DataInterfaceFunctions.IndexOfByPredicate([&](const FNiagaraFunctionSignature& Sig) -> bool { return Sig.EqualsIgnoringSpecifiers(OutSignature); });
 				if (FoundMatch < 0)
diff --git a/Engine/Plugins/FX/Niagara/Source/NiagaraEditor/Private/NiagaraNodeCustomHlsl.cpp b/Engine/Plugins/FX/Niagara/Source/NiagaraEditor/Private/NiagaraNodeCustomHlsl.cpp
--- a/Engine/Plugins/FX/Niagara/Source/NiagaraEditor/Private/NiagaraNodeCustomHlsl.cpp
+++ b/Engine/Plugins/FX/Niagara/Source/NiagaraEditor/Private/NiagaraNodeCustomHlsl.cpp
@@ -292,7 +292,7 @@ bool UNiagaraNodeCustomHlsl::CallsImpureDataInterfaceFunctions() const
 			if (UNiagaraDataInterface* DataInterfaceClass = CastChecked<UNiagaraDataInterface>(NiagaraType.GetClass()->ClassDefaultObject))
 			{
 				TArray<FNiagaraFunctionSignature> FunctionSignatures;
-				DataInterfaceClass->GetFunctions(FunctionSignatures);
+				DataInterfaceClass->GetFunctionSignatures(FunctionSignatures);
 
 				for (const FNiagaraFunctionSignature& FunctionSignature : FunctionSignatures)
 				{
diff --git a/Engine/Plugins/FX/Niagara/Source/NiagaraEditor/Private/NiagaraNodeFunctionCall.cpp b/Engine/Plugins/FX/Niagara/Source/NiagaraEditor/Private/NiagaraNodeFunctionCall.cpp
--- a/Engine/Plugins/FX/Niagara/Source/NiagaraEditor/Private/NiagaraNodeFunctionCall.cpp
+++ b/Engine/Plugins/FX/Niagara/Source/NiagaraEditor/Private/NiagaraNodeFunctionCall.cpp
@@ -1009,7 +1009,7 @@ void UNiagaraNodeFunctionCall::UpdatePinTooltips()
 	{
 		UNiagaraDataInterface* CDO = CastChecked<UNiagaraDataInterface>(Signature.Inputs[0].GetType().GetClass()->GetDefaultObject());
 		TArray<FNiagaraFunctionSignature> AllSignatures;
-		CDO->GetFunctions(AllSignatures);
+		CDO->GetFunctionSignatures(AllSignatures);
 		for (const FNiagaraFunctionSignature& DISignature : AllSignatures)
 		{
 			if (Signature.Name != DISignature.Name)
@@ -1726,7 +1726,7 @@ void UNiagaraNodeFunctionCall::RefreshSignature()
 		{
 			UNiagaraDataInterface* CDO = CastChecked<UNiagaraDataInterface>(Signature.Inputs[0].GetType().GetClass()->GetDefaultObject());
 			TArray<FNiagaraFunctionSignature> BaseDIFuncs;
-			CDO->GetFunctions(BaseDIFuncs);
+			CDO->GetFunctionSignatures(BaseDIFuncs);
 			if (FNiagaraFunctionSignature* BaseSig = BaseDIFuncs.FindByPredicate([&](const FNiagaraFunctionSignature& CheckSig) { return Signature.Name == CheckSig.Name; }))
 			{
 				//Revert signature to base and re-add dynamic pins as new inputs and outputs.
@@ -1772,7 +1772,7 @@ bool UNiagaraNodeFunctionCall::IsBaseSignatureOfDataInterfaceFunction(const UEdG
 	{
 		UNiagaraDataInterface* CDO = CastChecked<UNiagaraDataInterface>(Signature.Inputs[0].GetType().GetClass()->GetDefaultObject());
 		TArray<FNiagaraFunctionSignature> BaseDIFuncs;
-		CDO->GetFunctions(BaseDIFuncs);
+		CDO->GetFunctionSignatures(BaseDIFuncs);
 		if (FNiagaraFunctionSignature* BaseSig = BaseDIFuncs.FindByPredicate([&](const FNiagaraFunctionSignature& CheckSig) { return Signature.Name == CheckSig.Name; }))
 		{
 			FPinCollectorArray FoundPins;
diff --git a/Engine/Plugins/FX/NiagaraMRQ/Source/NiagaraMRQ/Private/NiagaraDataInterfaceMRQ.cpp b/Engine/Plugins/FX/NiagaraMRQ/Source/NiagaraMRQ/Private/NiagaraDataInterfaceMRQ.cpp
--- a/Engine/Plugins/FX/NiagaraMRQ/Source/NiagaraMRQ/Private/NiagaraDataInterfaceMRQ.cpp
+++ b/Engine/Plugins/FX/NiagaraMRQ/Source/NiagaraMRQ/Private/NiagaraDataInterfaceMRQ.cpp
@@ -81,7 +81,8 @@ void UNiagaraDataInterfaceMRQ::PostInitProperties()
 	}
 }
 
-void UNiagaraDataInterfaceMRQ::GetFunctions(TArray<FNiagaraFunctionSignature>& OutFunctions)
+#if WITH_EDITORONLY_DATA
+void UNiagaraDataInterfaceMRQ::GetFunctionsInternal(TArray<FNiagaraFunctionSignature>& OutFunctions) const
 {
 	using namespace NDIMRQLocal;
 
@@ -95,6 +96,7 @@ void UNiagaraDataInterfaceMRQ::GetFunctions(TArray<FNiagaraFunctionSignature>& O
 	Sig.AddOutput(FNiagaraVariable(FNiagaraTypeDefinition::GetIntDef(), TEXT("TemporalSampleIndex")), LOCTEXT("TemporalSampleIndexDesc", "Temporal sample we are simulating this frame for"));
 	Sig.AddOutput(FNiagaraVariable(FNiagaraTypeDefinition::GetFloatDef(), TEXT("SequenceFPS")), LOCTEXT("SequenceFPSDesc", "The source sequence frames per second"));
 }
+#endif
 
 void UNiagaraDataInterfaceMRQ::GetVMExternalFunction(const FVMExternalFunctionBindingInfo& BindingInfo, void* InstanceData, FVMExternalFunction& OutFunc)
 {
diff --git a/Engine/Plugins/FX/NiagaraMRQ/Source/NiagaraMRQ/Private/NiagaraDataInterfaceMRQ.h b/Engine/Plugins/FX/NiagaraMRQ/Source/NiagaraMRQ/Private/NiagaraDataInterfaceMRQ.h
--- a/Engine/Plugins/FX/NiagaraMRQ/Source/NiagaraMRQ/Private/NiagaraDataInterfaceMRQ.h
+++ b/Engine/Plugins/FX/NiagaraMRQ/Source/NiagaraMRQ/Private/NiagaraDataInterfaceMRQ.h
@@ -29,7 +29,6 @@ public:
 	//UObject Interface End
 
 	//UNiagaraDataInterface Interface
-	virtual void GetFunctions(TArray<FNiagaraFunctionSignature>& OutFunctions) override;
 	virtual void GetVMExternalFunction(const FVMExternalFunctionBindingInfo& BindingInfo, void* InstanceData, FVMExternalFunction &OutFunc) override;
 	virtual bool CanExecuteOnTarget(ENiagaraSimTarget Target) const override { return true; }
 #if WITH_EDITORONLY_DATA
@@ -47,4 +46,9 @@ public:
 	virtual bool PerInstanceTick(void* PerInstanceData, FNiagaraSystemInstance* SystemInstance, float DeltaSeconds) override;
 	virtual void ProvidePerInstanceDataForRenderThread(void* DataForRenderThread, void* PerInstanceData, const FNiagaraSystemInstanceID& SystemInstance) override;
 	//UNiagaraDataInterface Interface
+
+protected:
+#if WITH_EDITORONLY_DATA
+	virtual void GetFunctionsInternal(TArray<FNiagaraFunctionSignature>& OutFunctions) const override;
+#endif
 };
diff --git a/Engine/Plugins/Runtime/HairStrands/Source/HairStrandsCore/Private/Niagara/NiagaraDataInterfaceHairStrands.cpp b/Engine/Plugins/Runtime/HairStrands/Source/HairStrandsCore/Private/Niagara/NiagaraDataInterfaceHairStrands.cpp
--- a/Engine/Plugins/Runtime/HairStrands/Source/HairStrandsCore/Private/Niagara/NiagaraDataInterfaceHairStrands.cpp
+++ b/Engine/Plugins/Runtime/HairStrands/Source/HairStrandsCore/Private/Niagara/NiagaraDataInterfaceHairStrands.cpp
@@ -1049,11 +1049,13 @@ void UNiagaraDataInterfaceHairStrands::PostInitProperties()
 	}
 }
 
+#if WITH_EDITORONLY_DATA
+
 // Codegen optimization degenerates for very long functions like GetFunctions when combined with the invokation of lots of FORCEINLINE methods.
 // We don't need this code to be particularly fast anyway. The other way to improve this code compilation time would be to split it in multiple functions.
 BEGIN_FUNCTION_BUILD_OPTIMIZATION 
 
-void UNiagaraDataInterfaceHairStrands::GetFunctions(TArray<FNiagaraFunctionSignature>& OutFunctions)
+void UNiagaraDataInterfaceHairStrands::GetFunctionsInternal(TArray<FNiagaraFunctionSignature>& OutFunctions) const
 {
 	using namespace NDIHairStrandsLocal;
 	{
@@ -2093,6 +2095,8 @@ void UNiagaraDataInterfaceHairStrands::GetFunctions(TArray<FNiagaraFunctionSigna
 }
 END_FUNCTION_BUILD_OPTIMIZATION
 
+#endif
+
 DEFINE_NDI_DIRECT_FUNC_BINDER(UNiagaraDataInterfaceHairStrands, GetNumStrands);
 DEFINE_NDI_DIRECT_FUNC_BINDER(UNiagaraDataInterfaceHairStrands, GetStrandSize);
 
diff --git a/Engine/Plugins/Runtime/HairStrands/Source/HairStrandsCore/Private/Niagara/NiagaraDataInterfacePressureGrid.cpp b/Engine/Plugins/Runtime/HairStrands/Source/HairStrandsCore/Private/Niagara/NiagaraDataInterfacePressureGrid.cpp
--- a/Engine/Plugins/Runtime/HairStrands/Source/HairStrandsCore/Private/Niagara/NiagaraDataInterfacePressureGrid.cpp
+++ b/Engine/Plugins/Runtime/HairStrands/Source/HairStrandsCore/Private/Niagara/NiagaraDataInterfacePressureGrid.cpp
@@ -49,11 +49,12 @@ UNiagaraDataInterfacePressureGrid::UNiagaraDataInterfacePressureGrid(FObjectInit
 	NumAttributes = 18;
 }
 
-void UNiagaraDataInterfacePressureGrid::GetFunctions(TArray<FNiagaraFunctionSignature>& OutFunctions)
+#if WITH_EDITORONLY_DATA
+void UNiagaraDataInterfacePressureGrid::GetFunctionsInternal(TArray<FNiagaraFunctionSignature>& OutFunctions) const
 {
 	using namespace NDIPressureGridLocal;
 
-	Super::GetFunctions(OutFunctions);
+	Super::GetFunctionsInternal(OutFunctions);
 	{
 		FNiagaraFunctionSignature Sig;
 		Sig.Name = GetNodePositionName;
@@ -195,6 +196,7 @@ void UNiagaraDataInterfacePressureGrid::GetFunctions(TArray<FNiagaraFunctionSign
 		OutFunctions.Add(Sig);
 	}
 }
+#endif
 
 DEFINE_NDI_DIRECT_FUNC_BINDER(UNiagaraDataInterfacePressureGrid, BuildDistanceField);
 DEFINE_NDI_DIRECT_FUNC_BINDER(UNiagaraDataInterfacePressureGrid, BuildDensityField);
diff --git a/Engine/Plugins/Runtime/HairStrands/Source/HairStrandsCore/Private/Niagara/NiagaraDataInterfaceVelocityGrid.cpp b/Engine/Plugins/Runtime/HairStrands/Source/HairStrandsCore/Private/Niagara/NiagaraDataInterfaceVelocityGrid.cpp
--- a/Engine/Plugins/Runtime/HairStrands/Source/HairStrandsCore/Private/Niagara/NiagaraDataInterfaceVelocityGrid.cpp
+++ b/Engine/Plugins/Runtime/HairStrands/Source/HairStrandsCore/Private/Niagara/NiagaraDataInterfaceVelocityGrid.cpp
@@ -256,7 +256,8 @@ void UNiagaraDataInterfaceVelocityGrid::PostInitProperties()
 	}
 }
 
-void UNiagaraDataInterfaceVelocityGrid::GetFunctions(TArray<FNiagaraFunctionSignature>& OutFunctions)
+#if WITH_EDITORONLY_DATA
+void UNiagaraDataInterfaceVelocityGrid::GetFunctionsInternal(TArray<FNiagaraFunctionSignature>& OutFunctions) const
 {
 	using namespace NDIVelocityGridLocal;
 
@@ -341,6 +342,7 @@ void UNiagaraDataInterfaceVelocityGrid::GetFunctions(TArray<FNiagaraFunctionSign
 		OutFunctions.Add(Sig);
 	}
 }
+#endif
 
 DEFINE_NDI_DIRECT_FUNC_BINDER(UNiagaraDataInterfaceVelocityGrid, BuildVelocityField);
 DEFINE_NDI_DIRECT_FUNC_BINDER(UNiagaraDataInterfaceVelocityGrid, SampleVelocityField);
diff --git a/Engine/Plugins/Runtime/HairStrands/Source/HairStrandsCore/Public/Niagara/NiagaraDataInterfaceHairStrands.h b/Engine/Plugins/Runtime/HairStrands/Source/HairStrandsCore/Public/Niagara/NiagaraDataInterfaceHairStrands.h
--- a/Engine/Plugins/Runtime/HairStrands/Source/HairStrandsCore/Public/Niagara/NiagaraDataInterfaceHairStrands.h
+++ b/Engine/Plugins/Runtime/HairStrands/Source/HairStrandsCore/Public/Niagara/NiagaraDataInterfaceHairStrands.h
@@ -399,7 +399,6 @@ public:
 	virtual void PostInitProperties() override;
 
 	/** Begin UNiagaraDataInterface Interface */
-	virtual void GetFunctions(TArray<FNiagaraFunctionSignature>& OutFunctions) override;
 	virtual void GetVMExternalFunction(const FVMExternalFunctionBindingInfo& BindingInfo, void* InstanceData, FVMExternalFunction &OutFunc) override;
 	virtual bool CanExecuteOnTarget(ENiagaraSimTarget Target) const override { return Target == ENiagaraSimTarget::GPUComputeSim; }
 	virtual bool InitPerInstanceData(void* PerInstanceData, FNiagaraSystemInstance* SystemInstance) override;
@@ -666,6 +665,9 @@ public:
 	void GetSampleState(FVectorVMExternalFunctionContext& Context);
 
 protected:
+#if WITH_EDITORONLY_DATA
+	virtual void GetFunctionsInternal(TArray<FNiagaraFunctionSignature>& OutFunctions) const override;
+#endif
 	/** Copy one niagara DI to this */
 	virtual bool CopyToInternal(UNiagaraDataInterface* Destination) const override;
 };
diff --git a/Engine/Plugins/Runtime/HairStrands/Source/HairStrandsCore/Public/Niagara/NiagaraDataInterfacePressureGrid.h b/Engine/Plugins/Runtime/HairStrands/Source/HairStrandsCore/Public/Niagara/NiagaraDataInterfacePressureGrid.h
--- a/Engine/Plugins/Runtime/HairStrands/Source/HairStrandsCore/Public/Niagara/NiagaraDataInterfacePressureGrid.h
+++ b/Engine/Plugins/Runtime/HairStrands/Source/HairStrandsCore/Public/Niagara/NiagaraDataInterfacePressureGrid.h
@@ -17,7 +17,6 @@ class HAIRSTRANDSCORE_API UNiagaraDataInterfacePressureGrid : public UNiagaraDat
 
 public:
 	/** UNiagaraDataInterface Interface */
-	virtual void GetFunctions(TArray<FNiagaraFunctionSignature>& OutFunctions) override;
 	virtual void GetVMExternalFunction(const FVMExternalFunctionBindingInfo& BindingInfo, void* InstanceData, FVMExternalFunction &OutFunc) override;
 	virtual bool HasPreSimulateTick() const override { return true; }
 
@@ -56,6 +55,10 @@ public:
 	/** Update the deformation gradient */
 	void UpdateDeformationGradient(FVectorVMExternalFunctionContext& Context);
 
+protected:
+#if WITH_EDITORONLY_DATA
+	virtual void GetFunctionsInternal(TArray<FNiagaraFunctionSignature>& OutFunctions) const override;
+#endif
 };
 
 /** Proxy to send data to gpu */
diff --git a/Engine/Plugins/Runtime/HairStrands/Source/HairStrandsCore/Public/Niagara/NiagaraDataInterfaceVelocityGrid.h b/Engine/Plugins/Runtime/HairStrands/Source/HairStrandsCore/Public/Niagara/NiagaraDataInterfaceVelocityGrid.h
--- a/Engine/Plugins/Runtime/HairStrands/Source/HairStrandsCore/Public/Niagara/NiagaraDataInterfaceVelocityGrid.h
+++ b/Engine/Plugins/Runtime/HairStrands/Source/HairStrandsCore/Public/Niagara/NiagaraDataInterfaceVelocityGrid.h
@@ -90,7 +90,6 @@ public:
 	virtual void PostInitProperties() override;
 
 	/** UNiagaraDataInterface Interface */
-	virtual void GetFunctions(TArray<FNiagaraFunctionSignature>& OutFunctions) override;
 	virtual void GetVMExternalFunction(const FVMExternalFunctionBindingInfo& BindingInfo, void* InstanceData, FVMExternalFunction &OutFunc) override;
 	virtual bool CanExecuteOnTarget(ENiagaraSimTarget Target) const override { return Target == ENiagaraSimTarget::GPUComputeSim; }
 	virtual bool InitPerInstanceData(void* PerInstanceData, FNiagaraSystemInstance* SystemInstance) override;
@@ -127,6 +126,9 @@ public:
 	void SetGridDimension(FVectorVMExternalFunctionContext& Context);
 
 protected:
+#if WITH_EDITORONLY_DATA
+	virtual void GetFunctionsInternal(TArray<FNiagaraFunctionSignature>& OutFunctions) const override;
+#endif
 	/** Copy one niagara DI to this */
 	virtual bool CopyToInternal(UNiagaraDataInterface* Destination) const override;
 };
