Fixes for rigid mesh DI not correctly setting the MeshScale of Sphere/Capsule bodies

diff --git a/Engine/Plugins/FX/Niagara/Source/Niagara/Private/NiagaraDataInterfaceRigidMeshCollisionQuery.cpp b/Engine/Plugins/FX/Niagara/Source/Niagara/Private/NiagaraDataInterfaceRigidMeshCollisionQuery.cpp
--- a/Engine/Plugins/FX/Niagara/Source/Niagara/Private/NiagaraDataInterfaceRigidMeshCollisionQuery.cpp
+++ b/Engine/Plugins/FX/Niagara/Source/Niagara/Private/NiagaraDataInterfaceRigidMeshCollisionQuery.cpp
@@ -247,6 +247,16 @@ void ForEachBodySetup(USkeletalMeshComponent* Component, TBodySetupPredicate Pre
 	}
 }
 
+uint32 MeshComponentHash(const USkeletalMeshComponent* SkeletalMeshComponent)
+{
+	uint32 HashResult = GetTypeHash(SkeletalMeshComponent);
+	if (SkeletalMeshComponent)
+	{
+		HashResult = HashCombine(HashResult, GetTypeHash(SkeletalMeshComponent->GetSkeletalMeshAsset()));
+	}
+	return HashResult;
+}
+
 /// End UkeletalMeshComponent
 
 /// Begin UStaticMeshComponent
@@ -268,6 +278,16 @@ void ForEachBodySetup(UStaticMeshComponent* Component, TBodySetupPredicate Predi
 	Predicate(Component, Component->GetBodySetup());
 }
 
+uint32 MeshComponentHash(const UStaticMeshComponent* StaticMeshComponent)
+{
+	uint32 HashResult = GetTypeHash(StaticMeshComponent);
+	if (StaticMeshComponent)
+	{
+		HashResult = HashCombine(HashResult, GetTypeHash(StaticMeshComponent->GetStaticMesh()));
+	}
+	return HashResult;
+}
+
 /// End UStaticMeshComponent
 
 template<typename TComponentType>
@@ -278,7 +298,7 @@ void CountCollisionPrimitives(TConstArrayView<TComponentType*> Components, TArra
 	for (TComponentType* Component : Components)
 	{
 		FNDIRigidMeshCollisionData::FComponentBodyCount& BodyCount = PerComponentCounts.AddDefaulted_GetRef();
-		BodyCount.ComponentHash = GetTypeHash(Component);
+		BodyCount.ComponentHash = MeshComponentHash(Component);
 
 		ForEachBodySetup(Component, [&](TComponentType* Component, const UBodySetup* BodySetup)
 		{
@@ -413,7 +433,7 @@ void UpdateAssetArrays(TConstArrayView<TComponentType*> Components, const FVecto
 				if (InitializeStatics)
 				{
 					OutAssetArrays->ElementExtent[SphereIndex] = FVector4f(SphereElem.Radius, 0, 0, 0);
-					OutAssetArrays->MeshScale[BoxIndex] = FVector4f(CurrMeshScale.X, CurrMeshScale.Y, CurrMeshScale.Z, 0);
+					OutAssetArrays->MeshScale[SphereIndex] = FVector4f(CurrMeshScale.X, CurrMeshScale.Y, CurrMeshScale.Z, 0);
 					OutAssetArrays->PhysicsType[SphereIndex] = (SphereElem.GetCollisionEnabled() == ECollisionEnabled::QueryAndPhysics);
 					OutAssetArrays->ComponentIdIndex[SphereIndex] = ComponentIdIndex;
 				}
@@ -431,7 +451,7 @@ void UpdateAssetArrays(TConstArrayView<TComponentType*> Components, const FVecto
 				if (InitializeStatics)
 				{
 					OutAssetArrays->ElementExtent[CapsuleIndex] = FVector4f(CapsuleElem.Radius, CapsuleElem.Length, 0, 0);
-					OutAssetArrays->MeshScale[BoxIndex] = FVector4f(CurrMeshScale.X, CurrMeshScale.Y, CurrMeshScale.Z, 0);
+					OutAssetArrays->MeshScale[CapsuleIndex] = FVector4f(CurrMeshScale.X, CurrMeshScale.Y, CurrMeshScale.Z, 0);
 					OutAssetArrays->PhysicsType[CapsuleIndex] = (CapsuleElem.GetCollisionEnabled() == ECollisionEnabled::QueryAndPhysics);
 					OutAssetArrays->ComponentIdIndex[CapsuleIndex] = ComponentIdIndex;
 				}
@@ -559,7 +579,7 @@ void RemapPreviousTransforms(
 	}
 }
 
-void UpdateInternalArrays(
+bool UpdateInternalArrays(
 	TConstArrayView<UStaticMeshComponent*> StaticMeshView,
 	TConstArrayView<USkeletalMeshComponent*> SkeletalMeshView,
 	FVector LWCTile,
@@ -567,63 +587,67 @@ void UpdateInternalArrays(
 	TArray<FNDIRigidMeshCollisionData::FComponentBodyCount>& BodyCounts,
 	FNDIRigidMeshCollisionArrays* OutAssetArrays)
 {
-	if (OutAssetArrays != nullptr && OutAssetArrays->ElementOffsets.NumElements < OutAssetArrays->MaxPrimitives)
+	if (OutAssetArrays == nullptr || OutAssetArrays->ElementOffsets.NumElements >= OutAssetArrays->MaxPrimitives)
 	{
-		// when we are in game and we don't need to worry about bodies changing for a given mesh component we can try to optimize
-		// the update by just targeting the dynamic elements (transforms) and using the original values as the previous run
+		return false;
+	}
+
+	// when we are in game and we don't need to worry about bodies changing for a given mesh component we can try to optimize
+	// the update by just targeting the dynamic elements (transforms) and using the original values as the previous run
 #if !WITH_EDITOR
-		if (!bFullUpdate)
-		{
-			// if we're updating, then copy over last frame's transforms before we generate new ones
-			Swap(OutAssetArrays->PreviousTransform, OutAssetArrays->CurrentTransform);
-			Swap(OutAssetArrays->PreviousInverse, OutAssetArrays->CurrentInverse);
+	if (!bFullUpdate)
+	{
+		// if we're updating, then copy over last frame's transforms before we generate new ones
+		Swap(OutAssetArrays->PreviousTransform, OutAssetArrays->CurrentTransform);
+		Swap(OutAssetArrays->PreviousInverse, OutAssetArrays->CurrentInverse);
 
-			uint32 BoxIndex = OutAssetArrays->ElementOffsets.BoxOffset;
-			uint32 SphereIndex = OutAssetArrays->ElementOffsets.SphereOffset;
-			uint32 CapsuleIndex = OutAssetArrays->ElementOffsets.CapsuleOffset;
+		uint32 BoxIndex = OutAssetArrays->ElementOffsets.BoxOffset;
+		uint32 SphereIndex = OutAssetArrays->ElementOffsets.SphereOffset;
+		uint32 CapsuleIndex = OutAssetArrays->ElementOffsets.CapsuleOffset;
 
-			UpdateAssetArrays<UStaticMeshComponent, false>(StaticMeshView, LWCTile, OutAssetArrays, BoxIndex, SphereIndex, CapsuleIndex);
-			UpdateAssetArrays<USkeletalMeshComponent, false>(SkeletalMeshView, LWCTile, OutAssetArrays, BoxIndex, SphereIndex, CapsuleIndex);
+		UpdateAssetArrays<UStaticMeshComponent, false>(StaticMeshView, LWCTile, OutAssetArrays, BoxIndex, SphereIndex, CapsuleIndex);
+		UpdateAssetArrays<USkeletalMeshComponent, false>(SkeletalMeshView, LWCTile, OutAssetArrays, BoxIndex, SphereIndex, CapsuleIndex);
 
-			return;
-		}
+		return true;
+	}
 #endif
 
-		TArray<FNDIRigidMeshCollisionData::FComponentBodyCount> CurrentBodyCounts;
+	TArray<FNDIRigidMeshCollisionData::FComponentBodyCount> CurrentBodyCounts;
 
-		uint32 TotalBoxCount = 0;
-		uint32 TotalSphereCount = 0;
-		uint32 TotalCapsuleCount = 0;
+	uint32 TotalBoxCount = 0;
+	uint32 TotalSphereCount = 0;
+	uint32 TotalCapsuleCount = 0;
 
-		CountCollisionPrimitives(StaticMeshView, CurrentBodyCounts, TotalBoxCount, TotalSphereCount, TotalCapsuleCount);
-		CountCollisionPrimitives(SkeletalMeshView, CurrentBodyCounts, TotalBoxCount, TotalSphereCount, TotalCapsuleCount);
+	CountCollisionPrimitives(StaticMeshView, CurrentBodyCounts, TotalBoxCount, TotalSphereCount, TotalCapsuleCount);
+	CountCollisionPrimitives(SkeletalMeshView, CurrentBodyCounts, TotalBoxCount, TotalSphereCount, TotalCapsuleCount);
 
-		if ((TotalBoxCount + TotalSphereCount + TotalCapsuleCount) >= OutAssetArrays->MaxPrimitives)
-		{
-			UE_LOG(LogRigidMeshCollision, Error, TEXT("Number of Collision DI primitives is higher than the %d limit.  Please increase it."), OutAssetArrays->MaxPrimitives);
-			return;
-		}
+	if ((TotalBoxCount + TotalSphereCount + TotalCapsuleCount) >= OutAssetArrays->MaxPrimitives)
+	{
+		UE_LOG(LogRigidMeshCollision, Error, TEXT("Number of Collision DI primitives is higher than the %d limit.  Please increase it."), OutAssetArrays->MaxPrimitives);
+		return false;
+	}
 
-		OutAssetArrays->ElementOffsets.BoxOffset = 0;
-		OutAssetArrays->ElementOffsets.SphereOffset = OutAssetArrays->ElementOffsets.BoxOffset + TotalBoxCount;
-		OutAssetArrays->ElementOffsets.CapsuleOffset = OutAssetArrays->ElementOffsets.SphereOffset + TotalSphereCount;
-		OutAssetArrays->ElementOffsets.NumElements = OutAssetArrays->ElementOffsets.CapsuleOffset + TotalCapsuleCount;
+	OutAssetArrays->ElementOffsets.BoxOffset = 0;
+	OutAssetArrays->ElementOffsets.SphereOffset = OutAssetArrays->ElementOffsets.BoxOffset + TotalBoxCount;
+	OutAssetArrays->ElementOffsets.CapsuleOffset = OutAssetArrays->ElementOffsets.SphereOffset + TotalSphereCount;
+	OutAssetArrays->ElementOffsets.NumElements = OutAssetArrays->ElementOffsets.CapsuleOffset + TotalCapsuleCount;
 
-		uint32 BoxIndex = OutAssetArrays->ElementOffsets.BoxOffset;
-		uint32 SphereIndex = OutAssetArrays->ElementOffsets.SphereOffset;
-		uint32 CapsuleIndex = OutAssetArrays->ElementOffsets.CapsuleOffset;
+	uint32 BoxIndex = OutAssetArrays->ElementOffsets.BoxOffset;
+	uint32 SphereIndex = OutAssetArrays->ElementOffsets.SphereOffset;
+	uint32 CapsuleIndex = OutAssetArrays->ElementOffsets.CapsuleOffset;
 
-		// where possible PreviousTransform & PreviousInverse should be pulled from the current values of CurrentTransform &
-		// CurrentInverse based on the remapped entries
-		RemapPreviousTransforms<true>(BodyCounts, CurrentBodyCounts, OutAssetArrays);
+	// where possible PreviousTransform & PreviousInverse should be pulled from the current values of CurrentTransform &
+	// CurrentInverse based on the remapped entries
+	RemapPreviousTransforms<true>(BodyCounts, CurrentBodyCounts, OutAssetArrays);
 
-		UpdateAssetArrays<UStaticMeshComponent, true>(StaticMeshView, LWCTile, OutAssetArrays, BoxIndex, SphereIndex, CapsuleIndex);
-		UpdateAssetArrays<USkeletalMeshComponent, true>(SkeletalMeshView, LWCTile, OutAssetArrays, BoxIndex, SphereIndex, CapsuleIndex);
+	UpdateAssetArrays<UStaticMeshComponent, true>(StaticMeshView, LWCTile, OutAssetArrays, BoxIndex, SphereIndex, CapsuleIndex);
+	UpdateAssetArrays<USkeletalMeshComponent, true>(SkeletalMeshView, LWCTile, OutAssetArrays, BoxIndex, SphereIndex, CapsuleIndex);
 
-		RemapPreviousTransforms<false>(BodyCounts, CurrentBodyCounts, OutAssetArrays);
+	RemapPreviousTransforms<false>(BodyCounts, CurrentBodyCounts, OutAssetArrays);
 
-		BodyCounts = MoveTemp(CurrentBodyCounts);
-	}
+	BodyCounts = MoveTemp(CurrentBodyCounts);
+
+	return true;
 }
 
 static bool SystemHasFindActorsFunction(UNiagaraSystem* System)
@@ -824,11 +848,6 @@ PRAGMA_ENABLE_DEPRECATION_WARNINGS
 			StaticMeshes.Sort(SortPredicateByTypeHash);
 			SkeletalMeshes.Sort(SortPredicateByTypeHash);
 
-			auto HashComponentObject = [](const UMeshComponent* MeshComponent) -> uint32
-			{
-				return GetTypeHash(MeshComponent);
-			};
-
 			auto AccumulateHash = [](uint32 Lhs, uint32 Rhs) -> uint32
 			{
 				return HashCombine(Lhs, Rhs);
@@ -837,8 +856,9 @@ PRAGMA_ENABLE_DEPRECATION_WARNINGS
 			// generate a hash with the collected components so that we can see if on subsequent frames we
 			// have a different collection, in which case we'll want to do a full rebuild of the array of transforms
 			uint32 NewHashValue = 0;
-			NewHashValue = Algo::TransformAccumulate(StaticMeshes, HashComponentObject, NewHashValue, AccumulateHash);
-			NewHashValue = Algo::TransformAccumulate(SkeletalMeshes, HashComponentObject, NewHashValue, AccumulateHash);
+
+			NewHashValue = Algo::TransformAccumulate(StaticMeshes, [](const UStaticMeshComponent* MeshComponent) -> uint32 { return MeshComponentHash(MeshComponent); }, NewHashValue, AccumulateHash);
+			NewHashValue = Algo::TransformAccumulate(SkeletalMeshes, [](const USkeletalMeshComponent* MeshComponent) -> uint32 { return MeshComponentHash(MeshComponent); }, NewHashValue, AccumulateHash);
 
 			if (NewHashValue != ComponentCollectionHash)
 			{
@@ -850,16 +870,23 @@ PRAGMA_ENABLE_DEPRECATION_WARNINGS
 		TConstArrayView<UStaticMeshComponent*> StaticMeshView = MakeArrayView(StaticMeshes.GetData(), StaticMeshes.Num());
 		TConstArrayView<USkeletalMeshComponent*> SkeletalMeshView = MakeArrayView(SkeletalMeshes.GetData(), SkeletalMeshes.Num());
 
-		UpdateInternalArrays(
+		const bool bArraysUpdated = UpdateInternalArrays(
 			StaticMeshView,
 			SkeletalMeshView,
 			FVector(SystemInstance->GetLWCTile()),
 			bRequiresFullUpdate,
 			MeshBodyCounts,
 			AssetArrays.Get());
-	}
 
-	bRequiresFullUpdate = false;
+		if (bArraysUpdated)
+		{
+			bRequiresFullUpdate = false;
+		}
+	}
+	else
+	{
+		bRequiresFullUpdate = false;
+	}
 }
 
 //------------------------------------------------------------------------------------------------------------
