Fix use after free with Grid2D for the resource name of the pooled texture.

diff --git a/Engine/Plugins/FX/Niagara/Source/Niagara/Classes/NiagaraDataInterfaceGrid2DCollection.h b/Engine/Plugins/FX/Niagara/Source/Niagara/Classes/NiagaraDataInterfaceGrid2DCollection.h
--- a/Engine/Plugins/FX/Niagara/Source/Niagara/Classes/NiagaraDataInterfaceGrid2DCollection.h
+++ b/Engine/Plugins/FX/Niagara/Source/Niagara/Classes/NiagaraDataInterfaceGrid2DCollection.h
@@ -71,6 +71,10 @@ struct FGrid2DCollectionRWInstanceData_RenderThread
 	TArray<int32> VarComponents;
 	TArray<uint32> Offsets;
 
+#if !UE_BUILD_SHIPPING
+	FString GridTextureResourceName;
+#endif
+
 #if WITH_EDITORONLY_DATA
 	bool bPreviewGrid = false;
 	FIntVector4 PreviewAttribute = FIntVector4(INDEX_NONE, INDEX_NONE, INDEX_NONE, INDEX_NONE);
diff --git a/Engine/Plugins/FX/Niagara/Source/Niagara/Private/NiagaraDataInterfaceGrid2DCollection.cpp b/Engine/Plugins/FX/Niagara/Source/Niagara/Private/NiagaraDataInterfaceGrid2DCollection.cpp
--- a/Engine/Plugins/FX/Niagara/Source/Niagara/Private/NiagaraDataInterfaceGrid2DCollection.cpp
+++ b/Engine/Plugins/FX/Niagara/Source/Niagara/Private/NiagaraDataInterfaceGrid2DCollection.cpp
@@ -2041,6 +2041,22 @@ bool UNiagaraDataInterfaceGrid2DCollection::InitPerInstanceData(void* PerInstanc
 			TargetData->Vars.Emplace(RT_InstanceData.Vars[i].GetName());
 			TargetData->VarComponents.Add(RT_InstanceData.Vars[i].GetType().GetSize() / sizeof(float));
 		}
+
+#if !UE_BUILD_SHIPPING
+		{
+			TStringBuilder<128> StringBuilder;
+			TargetData->SourceDIName.AppendString(StringBuilder);
+			StringBuilder.Append("_");
+			for (FName CurrName : TargetData->Vars)
+			{
+				CurrName.AppendString(StringBuilder);
+				StringBuilder.Append("_");
+			}
+			StringBuilder.Append("Grid2DCollection");
+			TargetData->GridTextureResourceName = StringBuilder;
+		}
+#endif // WITH_SAHDER_NAMES
+
 #if WITH_EDITORONLY_DATA
 		TargetData->bPreviewGrid = RT_InstanceData.bPreviewGrid;
 		TargetData->PreviewAttribute = RT_InstanceData.PreviewAttribute;
@@ -2665,20 +2681,11 @@ void FGrid2DCollectionRWInstanceData_RenderThread::BeginSimulate(FRDGBuilder& Gr
 
 		const FRDGTextureDesc TextureDesc = FRDGTextureDesc::Create2DArray(NumCells, PixelFormat.GetValue(), FClearValueBinding::Black, ETextureCreateFlags::ShaderResource | ETextureCreateFlags::UAV, uint16(NumAttributes));
 
-		const TCHAR* GridTextureName = TEXT("Grid2D::GridTexture");
-	#if 1
-		TStringBuilder<128> StringBuilder;
-		SourceDIName.AppendString(StringBuilder);
-		StringBuilder.Append("_");
-		for (FName CurrName : Vars)
-		{
-			CurrName.AppendString(StringBuilder);
-			StringBuilder.Append("_");
-		}
-		StringBuilder.Append("Grid2DCollection");
-		GridTextureName = *StringBuilder;
-	#endif
-		DestinationData->Initialize(GraphBuilder, GridTextureName, TextureDesc);
+#if UE_BUILD_SHIPPING
+		DestinationData->Initialize(GraphBuilder, TEXT("Grid2D::GridTexture"), TextureDesc);
+#else
+		DestinationData->Initialize(GraphBuilder, *GridTextureResourceName, TextureDesc);
+#endif
 
 		// This destination buffer will sometimes have old data in it.  Force it to clear.
 		AddClearUAVPass(GraphBuilder, DestinationData->GetOrCreateUAV(GraphBuilder), FVector4f(ForceInitToZero));
