ExternalParameter cbuffer is now packed and padded explicitly by the translator to ensure compatability across all shader platforms.  Addresses vulkan shaders being potentially broken since CL 27399055

diff --git a/Engine/Plugins/FX/Niagara/Source/Niagara/Classes/NiagaraComputeExecutionContext.h b/Engine/Plugins/FX/Niagara/Source/Niagara/Classes/NiagaraComputeExecutionContext.h
--- a/Engine/Plugins/FX/Niagara/Source/Niagara/Classes/NiagaraComputeExecutionContext.h
+++ b/Engine/Plugins/FX/Niagara/Source/Niagara/Classes/NiagaraComputeExecutionContext.h
@@ -88,6 +88,9 @@ struct FNiagaraComputeExecutionContext
 		}
 	}
 
+	int32 GetConstantBufferSize() const;
+	uint8* WriteConstantBufferInstanceData(uint8* InTargetBuffer, FNiagaraComputeInstanceData& InstanceData) const;
+
 	struct 
 	{
 		// The offset at which the GPU instance count (see FNiagaraGPUInstanceCountManager()).
diff --git a/Engine/Plugins/FX/Niagara/Source/Niagara/Private/NiagaraComputeExecutionContext.cpp b/Engine/Plugins/FX/Niagara/Source/Niagara/Private/NiagaraComputeExecutionContext.cpp
--- a/Engine/Plugins/FX/Niagara/Source/Niagara/Private/NiagaraComputeExecutionContext.cpp
+++ b/Engine/Plugins/FX/Niagara/Source/Niagara/Private/NiagaraComputeExecutionContext.cpp
@@ -256,6 +256,34 @@ void FNiagaraComputeExecutionContext::SetMultiViewPreviousDataToRender(FNiagaraD
 	MultiViewPreviousDataToRender = InMultiViewPreviousDataToRender;
 }
 
+int32 FNiagaraComputeExecutionContext::GetConstantBufferSize() const
+{
+	return Align(CombinedParamStore.GetExternalParameterSize(), SHADER_PARAMETER_STRUCT_ALIGNMENT);
+}
+
+uint8* FNiagaraComputeExecutionContext::WriteConstantBufferInstanceData(uint8* InTargetBuffer, FNiagaraComputeInstanceData& InstanceData) const
+{
+	const int32 InterpFactor = HasInterpolationParameters ? 2 : 1;
+
+	InstanceData.ExternalParamDataSize = InterpFactor * GetConstantBufferSize();
+	InstanceData.ExternalParamData = InTargetBuffer;
+
+	const TArray<uint8>& ParameterDataArray = CombinedParamStore.GetParameterDataArray();
+	const int32 SourceDataSize = CombinedParamStore.GetExternalParameterSize();
+	const int32 ConstantBufferSize = GetConstantBufferSize();
+
+	check(SourceDataSize <= ParameterDataArray.Num());
+	FMemory::Memcpy(InstanceData.ExternalParamData, ParameterDataArray.GetData(), SourceDataSize);
+
+	if (HasInterpolationParameters)
+	{
+		check(SourceDataSize + SourceDataSize <= ParameterDataArray.Num());
+		FMemory::Memcpy(InstanceData.ExternalParamData + ConstantBufferSize, ParameterDataArray.GetData() + SourceDataSize, SourceDataSize);
+	}
+
+	return InTargetBuffer + InstanceData.ExternalParamDataSize;
+}
+
 bool FNiagaraComputeInstanceData::IsOutputStage(FNiagaraDataInterfaceProxy* DIProxy, uint32 SimulationStageIndex) const
 {
 	return Context->IsOutputStage(DIProxy, SimulationStageIndex);
diff --git a/Engine/Plugins/FX/Niagara/Source/Niagara/Private/NiagaraEmitterInstance.cpp b/Engine/Plugins/FX/Niagara/Source/Niagara/Private/NiagaraEmitterInstance.cpp
--- a/Engine/Plugins/FX/Niagara/Source/Niagara/Private/NiagaraEmitterInstance.cpp
+++ b/Engine/Plugins/FX/Niagara/Source/Niagara/Private/NiagaraEmitterInstance.cpp
@@ -1544,11 +1544,10 @@ void FNiagaraEmitterInstance::Tick(float DeltaSeconds)
 			GPUExecContext->CombinedParamStore.DumpParameters();
 		}
 
-		int32 ParmSize = GPUExecContext->CombinedParamStore.GetPaddedParameterSizeInBytes();
 		// Because each context is only ran once each frame, the CBuffer layout stays constant for the lifetime duration of the CBuffer (one frame).
 
 		// @todo-threadsafety do this once during init. Should not change during runtime...
-		const uint32 ConstantBufferSize = ParmSize / (GPUExecContext->HasInterpolationParameters ? 2 : 1);
+		const uint32 ConstantBufferSize = GPUExecContext->GetConstantBufferSize();
 		if (GPUExecContext->ExternalCBufferLayoutSize != ConstantBufferSize)
 		{
 			GPUExecContext->ExternalCBufferLayoutSize = ConstantBufferSize;
diff --git a/Engine/Plugins/FX/Niagara/Source/Niagara/Private/NiagaraGPUSystemTick.cpp b/Engine/Plugins/FX/Niagara/Source/Niagara/Private/NiagaraGPUSystemTick.cpp
--- a/Engine/Plugins/FX/Niagara/Source/Niagara/Private/NiagaraGPUSystemTick.cpp
+++ b/Engine/Plugins/FX/Niagara/Source/Niagara/Private/NiagaraGPUSystemTick.cpp
@@ -168,10 +168,6 @@ void FNiagaraGPUSystemTick::Init(FNiagaraSystemInstance* InSystemInstance)
 			InstanceData->EmitterParamData = ParamDataBufferPtr;
 			ParamDataBufferPtr += InterpFactor * sizeof(FNiagaraEmitterParameters);
 
-			InstanceData->ExternalParamData = ParamDataBufferPtr;
-			InstanceData->ExternalParamDataSize = GPUContext->CombinedParamStore.GetPaddedParameterSizeInBytes();
-			ParamDataBufferPtr += InstanceData->ExternalParamDataSize;
-
 			// actually copy all of the data over
 			FMemory::Memcpy(InstanceData->EmitterParamData, &InSystemInstance->GetEmitterParameters(EmitterIdx), sizeof(FNiagaraEmitterParameters));
 			if (IncludeInterpolationParameters)
@@ -179,11 +175,11 @@ void FNiagaraGPUSystemTick::Init(FNiagaraSystemInstance* InSystemInstance)
 				FMemory::Memcpy(InstanceData->EmitterParamData + sizeof(FNiagaraEmitterParameters), &InSystemInstance->GetEmitterParameters(EmitterIdx, true), sizeof(FNiagaraEmitterParameters));
 			}
 
+			ParamDataBufferPtr = GPUContext->WriteConstantBufferInstanceData(ParamDataBufferPtr, *InstanceData);
+
 			bHasMultipleStages = InstanceData->bHasMultipleStages;
 			bHasInterpolatedParameters |= GPUContext->HasInterpolationParameters;
 
-			GPUContext->CombinedParamStore.CopyParameterDataToPaddedBuffer(InstanceData->ExternalParamData, InstanceData->ExternalParamDataSize);
-
 			// Calling PostTick will push current -> previous parameters this must be done after copying the parameter data
 			GPUContext->PostTick();
 
diff --git a/Engine/Plugins/FX/Niagara/Source/Niagara/Private/NiagaraScript.cpp b/Engine/Plugins/FX/Niagara/Source/Niagara/Private/NiagaraScript.cpp
--- a/Engine/Plugins/FX/Niagara/Source/Niagara/Private/NiagaraScript.cpp
+++ b/Engine/Plugins/FX/Niagara/Source/Niagara/Private/NiagaraScript.cpp
@@ -1886,15 +1886,10 @@ void UNiagaraScript::Serialize(FArchive& Ar)
 			if (Usage != ENiagaraScriptUsage::ParticleGPUComputeScript)
 			{
 				ExecutableData.BakeScriptLiterals(ExecutableData.ScriptLiterals);
-
-				// we only need the padding info for when we're dealing with GPU scripts (for
-				// FNiagaraScriptInstanceParameterStore::CopyParameterDataToPaddedBuffer())
-				ScriptExecutionParamStore.PaddingInfo.Empty();
 			}
 			else
 			{
 				ExecutableData.ScriptLiterals.Empty();
-				ScriptExecutionParamStore.CoalescePaddingInfo();
 			}
 
 			if (GNiagaraCompressScriptByteCode)
diff --git a/Engine/Plugins/FX/Niagara/Source/Niagara/Private/NiagaraScriptExecutionParameterStore.cpp b/Engine/Plugins/FX/Niagara/Source/Niagara/Private/NiagaraScriptExecutionParameterStore.cpp
--- a/Engine/Plugins/FX/Niagara/Source/Niagara/Private/NiagaraScriptExecutionParameterStore.cpp
+++ b/Engine/Plugins/FX/Niagara/Source/Niagara/Private/NiagaraScriptExecutionParameterStore.cpp
@@ -10,7 +10,6 @@
 FNiagaraScriptExecutionParameterStore::FNiagaraScriptExecutionParameterStore()
 	: FNiagaraParameterStore()
 	, ParameterSize(0)
-	, PaddedParameterSize(0)
 	, bInitialized(false)
 {}
 
@@ -180,78 +179,12 @@ uint32 FNiagaraScriptExecutionParameterStore::GenerateLayoutInfoInternal(TArray<
 	return InSrcOffset;
 }
 
-void FNiagaraScriptExecutionParameterStore::CoalescePaddingInfo()
-{
-	int32 PaddingIt = 1;
-
-	while (PaddingIt < PaddingInfo.Num())
-	{
-		FNiagaraScriptExecutionPaddingInfo& PreviousEntry = PaddingInfo[PaddingIt - 1];
-		const FNiagaraScriptExecutionPaddingInfo& CurrentEntry = PaddingInfo[PaddingIt];
-
-		if (((PreviousEntry.SrcOffset + PreviousEntry.SrcSize) == CurrentEntry.SrcOffset)
-			&& ((PreviousEntry.DestOffset + PreviousEntry.SrcSize) == CurrentEntry.DestOffset)
-			&& ((TNumericLimits<uint16>::Max() - PreviousEntry.SrcSize) >= CurrentEntry.SrcSize))
-		{
-			PreviousEntry.SrcSize += CurrentEntry.SrcSize;
-			PreviousEntry.DestSize += CurrentEntry.DestSize;
-			PaddingInfo.RemoveAt(PaddingIt);
-		}
-		else
-		{
-			++PaddingIt;
-		}
-	}
-}
-
-void FNiagaraScriptExecutionParameterStore::AddPaddedParamSize(const FNiagaraTypeDefinition& InParamType, uint32 InOffset)
-{
-	if (InParamType.IsDataInterface())
-	{
-		return;
-	}
-
-	uint32 NextMemberOffset = 0;
-	if (PaddingInfo.Num() != 0)
-	{
-		NextMemberOffset = PaddingInfo[PaddingInfo.Num() - 1].DestOffset + PaddingInfo[PaddingInfo.Num() - 1].DestSize;
-	}
-	GenerateLayoutInfoInternal(PaddingInfo, NextMemberOffset, InParamType.GetScriptStruct(), InOffset);
-
-	if (PaddingInfo.Num() != 0)
-	{
-		NextMemberOffset = PaddingInfo[PaddingInfo.Num() - 1].DestOffset + PaddingInfo[PaddingInfo.Num() - 1].DestSize;
-		PaddedParameterSize = Align(NextMemberOffset, SHADER_PARAMETER_STRUCT_ALIGNMENT);
-	}
-	else
-	{
-		PaddedParameterSize = 0;
-	}
-}
-
-void FNiagaraScriptExecutionParameterStore::AddAlignmentPadding()
-{
-	if (PaddingInfo.Num())
-	{
-		const auto& LastEntry = PaddingInfo.Last();
-		const uint32 CurrentOffset = LastEntry.DestOffset + LastEntry.DestSize;
-		const uint32 AlignedOffset = Align(CurrentOffset, SHADER_PARAMETER_STRUCT_ALIGNMENT);
-
-		if (CurrentOffset != AlignedOffset)
-		{
-			PaddingInfo.Emplace(GetParameterDataArray().Num(), CurrentOffset, 0, AlignedOffset - CurrentOffset);
-		}
-	}
-}
-
 void FNiagaraScriptExecutionParameterStore::InitFromOwningScript(UNiagaraScript* Script, ENiagaraSimTarget SimTarget, bool bNotifyAsDirty)
 {
 	//TEMPORARTY:
 	//We should replace the storage on the script with an FNiagaraParameterStore also so we can just copy that over here. Though that is an even bigger refactor job so this is a convenient place to break that work up.
 
 	Empty();
-	PaddedParameterSize = 0;
-	PaddingInfo.Empty();
 
 	if (Script)
 	{
@@ -275,7 +208,6 @@ void FNiagaraScriptExecutionParameterStore::AddScriptParams(UNiagaraScript* Scri
 	{
 		return;
 	}
-	PaddingInfo.Empty();
 
 	const FNiagaraVMExecutableData& ExecutableData = Script->GetVMExecutableData();
 
@@ -295,8 +227,6 @@ void FNiagaraScriptExecutionParameterStore::AddScriptParams(UNiagaraScript* Scri
 
 	if (bIsInterpolatedSpawn)
 	{
-		AddAlignmentPadding();
-
 		for (const FNiagaraVariable& Param : ExecutableData.Parameters.Parameters)
 		{
 			FNiagaraVariable PrevParam(Param.GetType(), FName(*(INTERPOLATED_PARAMETER_PREFIX + Param.GetName().ToString())));
@@ -394,33 +324,6 @@ uint32 FNiagaraScriptInstanceParameterStore::GetExternalParameterSize() const
 	return 0;
 }
 
-uint32 FNiagaraScriptInstanceParameterStore::GetPaddedParameterSizeInBytes() const
-{
-	if (const auto* ScriptStore = ScriptParameterStore.Get())
-	{
-		return ScriptStore->PaddedParameterSize;
-	}
-	return 0;
-}
-
-void FNiagaraScriptInstanceParameterStore::CopyParameterDataToPaddedBuffer(uint8* InTargetBuffer, uint32 InTargetBufferSizeInBytes) const
-{
-	if (const auto* ScriptStore = ScriptParameterStore.Get())
-	{
-		check((uint32)ScriptStore->ParameterSize <= ScriptStore->PaddedParameterSize);
-		check(InTargetBufferSizeInBytes >= ScriptStore->PaddedParameterSize);
-		FMemory::Memzero(InTargetBuffer, InTargetBufferSizeInBytes);
-		const auto& ParameterDataArray = GetParameterDataArray();
-		const uint8* SrcData = ParameterDataArray.GetData();
-		for (const auto& Padding : ScriptStore->PaddingInfo)
-		{
-			check(uint32(Padding.DestOffset + Padding.DestSize) <= InTargetBufferSizeInBytes);
-			check(uint32(Padding.SrcOffset + Padding.SrcSize) <= (uint32)ParameterDataArray.Num());
-			FMemory::Memcpy(InTargetBuffer + Padding.DestOffset, SrcData + Padding.SrcOffset, Padding.SrcSize);
-		}
-	}
-}
-
 TArrayView<const FNiagaraVariableWithOffset> FNiagaraScriptInstanceParameterStore::ReadParameterVariables() const
 {
 	if (const auto* ScriptStore = ScriptParameterStore.Get())
diff --git a/Engine/Plugins/FX/Niagara/Source/Niagara/Private/NiagaraSystemInstance.cpp b/Engine/Plugins/FX/Niagara/Source/Niagara/Private/NiagaraSystemInstance.cpp
--- a/Engine/Plugins/FX/Niagara/Source/Niagara/Private/NiagaraSystemInstance.cpp
+++ b/Engine/Plugins/FX/Niagara/Source/Niagara/Private/NiagaraSystemInstance.cpp
@@ -2468,7 +2468,9 @@ void FNiagaraSystemInstance::Tick_Concurrent(bool bEnqueueGPUTickIfNeeded)
 			{
 				if (const FNiagaraComputeExecutionContext* GPUContext = Inst.GetGPUContext())
 				{
-					TotalCombinedParamStoreSize += GPUContext->CombinedParamStore.GetPaddedParameterSizeInBytes();
+					const int32 InterpFactor = GPUContext->HasInterpolationParameters ? 2 : 1;
+
+					TotalCombinedParamStoreSize += InterpFactor * GPUContext->GetConstantBufferSize();
 					GPUParamIncludeInterpolation = GPUContext->HasInterpolationParameters || GPUParamIncludeInterpolation;
 					ActiveGPUEmitterCount++;
 				}
diff --git a/Engine/Plugins/FX/Niagara/Source/Niagara/Public/NiagaraScriptExecutionParameterStore.h b/Engine/Plugins/FX/Niagara/Source/Niagara/Public/NiagaraScriptExecutionParameterStore.h
--- a/Engine/Plugins/FX/Niagara/Source/Niagara/Public/NiagaraScriptExecutionParameterStore.h
+++ b/Engine/Plugins/FX/Niagara/Source/Niagara/Public/NiagaraScriptExecutionParameterStore.h
@@ -52,7 +52,6 @@ public:
 	/** Call this init function if you are using a Niagara parameter store within a UNiagaraScript.*/
 	void InitFromOwningScript(UNiagaraScript* Script, ENiagaraSimTarget SimTarget, bool bNotifyAsDirty);
 	void AddScriptParams(UNiagaraScript* Script, ENiagaraSimTarget SimTarget, bool bTriggerRebind);
-	void CoalescePaddingInfo();
 #endif
 
 	virtual bool AddParameter(const FNiagaraVariable& Param, bool bInitInterfaces = true, bool bTriggerRebind = true, int32* OutOffset = nullptr) override
@@ -76,10 +75,6 @@ public:
 		{
 			bAdded = FNiagaraParameterStore::AddParameter(Param, bInitInterfaces, bTriggerRebind, &NewParamOffset);
 		}
-		if (bAdded)
-		{
-			AddPaddedParamSize(Param.GetType(), NewParamOffset);
-		}
 		if (OutOffset)
 		{
 			*OutOffset = NewParamOffset;
@@ -106,8 +101,6 @@ public:
 	{
 		FNiagaraParameterStore::Empty(bClearBindings);
 		ParameterSize = 0;
-		PaddedParameterSize = 0;
-		PaddingInfo.Empty();
 		bInitialized = false;
 	}
 
@@ -115,8 +108,6 @@ public:
 	{
 		FNiagaraParameterStore::Reset(bClearBindings);
 		ParameterSize = 0;
-		PaddedParameterSize = 0;
-		PaddingInfo.Empty();
 		bInitialized = false;
 	}
 
@@ -124,12 +115,6 @@ public:
 	UPROPERTY()
 	int32 ParameterSize;
 
-	UPROPERTY()
-	uint32 PaddedParameterSize;
-
-	UPROPERTY()
-	TArray<FNiagaraScriptExecutionPaddingInfo> PaddingInfo;
-
 	UPROPERTY()
 	uint8 bInitialized : 1;
 
@@ -142,14 +127,9 @@ public:
 		SIZE_T ResourceSize = sizeof(FNiagaraScriptExecutionParameterStore) - sizeof(FNiagaraParameterStore);
 
 		ResourceSize += FNiagaraParameterStore::GetResourceSize();
-		ResourceSize += PaddingInfo.GetAllocatedSize();
 		return ResourceSize;
 	}
 
-protected:
-	void AddPaddedParamSize(const FNiagaraTypeDefinition& InParamType, uint32 InOffset);
-	void AddAlignmentPadding();
-
 private:
 
 #if WITH_EDITORONLY_DATA
@@ -208,11 +188,6 @@ public:
 	// Just the external parameters, not previous or internal...
 	uint32 GetExternalParameterSize() const;
 
-	uint32 GetPaddedParameterSizeInBytes() const;
-
-	// Helper that converts the data from the base type array internally into the padded out renderer-ready format.
-	void CopyParameterDataToPaddedBuffer(uint8* InTargetBuffer, uint32 InTargetBufferSizeInBytes) const;
-
 #if WITH_EDITORONLY_DATA
 	TArrayView<const uint8> GetScriptLiterals() const;
 #endif
diff --git a/Engine/Plugins/FX/Niagara/Source/NiagaraEditor/Private/NiagaraHlslTranslator.cpp b/Engine/Plugins/FX/Niagara/Source/NiagaraEditor/Private/NiagaraHlslTranslator.cpp
--- a/Engine/Plugins/FX/Niagara/Source/NiagaraEditor/Private/NiagaraHlslTranslator.cpp
+++ b/Engine/Plugins/FX/Niagara/Source/NiagaraEditor/Private/NiagaraHlslTranslator.cpp
@@ -1067,7 +1067,7 @@ void FNiagaraHlslTranslator::BuildConstantBuffer(ENiagaraCodeChunkMode ChunkMode
 	for (const FNiagaraVariable& Variable : T::GetVariables())
 	{
 		const FString SymbolName = GetSanitizedSymbolName(Variable.GetName().ToString(), true);
-		AddUniformChunk(SymbolName, Variable, ChunkMode, false);
+		AddChunkToConstantBuffer(SymbolName, Variable, ChunkMode);
 	}
 }
 
@@ -1858,6 +1858,8 @@ FNiagaraTranslateResults TNiagaraHlslTranslator<GraphBridge>::Translate(const FN
 			}
 		}
 
+		PackRegisteredUniformChunk(CompilationOutput.ScriptData.Parameters);
+
 		// Generate the Parameter Map HLSL definitions. We don't add to the final HLSL output here. We just build up the strings and tables
 		// that are needed later.
 		TArray<FNiagaraVariable> PrimaryDataSetOutputEntries;
@@ -4459,50 +4461,198 @@ FString FNiagaraHlslTranslator::GeneratedConstantString(FVector4 Constant)
 	return FString::Format(TEXT("float4({0}, {1}, {2}, {3})"), Args);
 }
 
-int32 FNiagaraHlslTranslator::AddUniformChunk(FString SymbolName, const FNiagaraVariable& InVariable, ENiagaraCodeChunkMode ChunkMode, bool AddPadding)
+int32 FNiagaraHlslTranslator::AddChunkToConstantBuffer(const FString& SymbolName, const FNiagaraVariable& InVariable, ENiagaraCodeChunkMode ChunkMode)
 {
+	if (const FUniformVariableInfo* UniformVariableInfo = ParamMapDefinedSystemVars.Find(InVariable.GetName()))
+	{
+		check(UniformVariableInfo->ChunkMode == static_cast<int32>(ChunkMode));
+		check(UniformVariableInfo->Variable == InVariable);
+		return UniformVariableInfo->ChunkIndex;
+	}
+
 	const FNiagaraTypeDefinition& Type = InVariable.GetType();
+	const int32 ChunkModeIndex = static_cast<int32>(ChunkMode);
 
-	int32 Ret = CodeChunks.IndexOfByPredicate(
-		[&](const FNiagaraCodeChunk& Chunk)
+	check(FNiagaraTypeHelper::IsLWCType(Type) == false);
+	int32 ChunkIndex = CodeChunks.AddDefaulted();
+	FNiagaraCodeChunk& Chunk = CodeChunks[ChunkIndex];
+	Chunk.SymbolName = GetSanitizedSymbolName(SymbolName.Replace(TEXT("."), TEXT("_")));
+	Chunk.Type = Type;
+	Chunk.Original = InVariable;
+	Chunk.Mode = ChunkMode;
+
+	ChunksByMode[ChunkModeIndex].Add(ChunkIndex);
+
+	auto& SystemVar = ParamMapDefinedSystemVars.Add(InVariable.GetName());
+	SystemVar.ChunkIndex = ChunkIndex;
+	SystemVar.ChunkMode = ChunkModeIndex;
+	SystemVar.Variable = InVariable;
+
+	return ChunkIndex;
+}
+
+int32 FNiagaraHlslTranslator::RegisterUniformChunkToPack(const FString& SymbolName, const FNiagaraVariable& InVariable, bool AddPadding, FNiagaraParameters& Parameters, TOptional<FNiagaraVariable>& ConflictingVariable)
+{
+	// we must ensure that there's a one to one relationship between symbol name and parameter.  The generated VM only
+	// knows about the symbols while the parameter stores knows about the parameters, if these mismatch, then we're going
+	// to be incorrectly addressing the constant table
+	int32 ParameterIndex = Parameters.Parameters.IndexOfByKey(InVariable);
+
+	if (ParameterIndex == INDEX_NONE)
 	{
-		return Chunk.Mode == ChunkMode && Chunk.SymbolName == SymbolName && Chunk.Type == Type;
+		// add the parameter, but first evaluate whether any of the symbols for existing parameters would conflict
+		const int32 ConflictIndex = Parameters.Parameters.IndexOfByPredicate([&](const FNiagaraVariable& ExistingParameter)
+		{
+			FNameBuilder ExistingParameterName(ExistingParameter.GetName());
+			return GetSanitizedSymbolName(ExistingParameterName).Equals(SymbolName);
+		});
+
+		if (ConflictIndex != INDEX_NONE)
+		{
+			ConflictingVariable = Parameters.Parameters[ConflictIndex];
+			return INDEX_NONE;
+		}
+		else
+		{
+			ParameterIndex = Parameters.Parameters.Add(InVariable);
+		}
 	}
-	);
+	
+	int32 ChunkIndex = AddChunkToConstantBuffer(SymbolName, InVariable, ENiagaraCodeChunkMode::Uniform);
 
-	if (Ret == INDEX_NONE)
+	if (AddPadding)
 	{
-		check(FNiagaraTypeHelper::IsLWCType(Type) == false);
-		Ret = CodeChunks.AddDefaulted();
-		FNiagaraCodeChunk& Chunk = CodeChunks[Ret];
-		Chunk.SymbolName = GetSanitizedSymbolName(SymbolName);
-		Chunk.Type = Type;
-		Chunk.Original = InVariable;
+		UniformParametersToPack.Add(ParameterIndex, ChunkIndex);
+	}
+
+	return ChunkIndex;
+}
+
+void FNiagaraHlslTranslator::PackRegisteredUniformChunk(FNiagaraParameters& Parameters)
+{
+	if (UniformParametersToPack.IsEmpty())
+	{
+		return;
+	}
+
+	// go through the parameters that have been registered, generate a layout that reduces waste and adjust the chunks
+	// that have already been added via RegisterUniformChunkToPack()
+
+	constexpr int32 SlotSizeInBytes = 4;
+	constexpr int32 SlotsPerRow = 4;
+
+	const int32 InitialParameterCount = Parameters.Parameters.Num();
+
+	TArray<int32> OrderedVariableIndices;
+	OrderedVariableIndices.Reserve(InitialParameterCount);
 
-		if (AddPadding)
+	TArray<int32> SingleIndices;
+	TArray<int32> DoubleIndices;
+	TArray<int32> TripleIndices;
+
+	int32 UniquePaddingIndex = 1;
+
+	for (int32 ParameterIt = 0; ParameterIt < InitialParameterCount; ++ParameterIt)
+	{
+		const FNiagaraVariable& Parameter = Parameters.Parameters[ParameterIt];
+		const int32 ParameterSizeInBytes = Align(Parameter.GetSizeInBytes(), SlotSizeInBytes);
+		const int32 ParameterSizeInSlots = ParameterSizeInBytes / SlotSizeInBytes;
+
+		switch (ParameterSizeInSlots % SlotsPerRow)
 		{
-			if (Type == FNiagaraTypeDefinition::GetVec2Def())
-			{
-				Chunk.Type = FNiagaraTypeDefinition::GetVec4Def();
-				Chunk.ComponentMask = TEXT(".xy");
-			}
-			else if (Type == FNiagaraTypeDefinition::GetVec3Def() || Type == FNiagaraTypeDefinition::GetPositionDef())
+			case 0:			OrderedVariableIndices.Add(ParameterIt);	break;
+			case 1:			SingleIndices.Add(ParameterIt);				break;
+			case 2:			DoubleIndices.Add(ParameterIt);				break;
+			case 3:			TripleIndices.Add(ParameterIt);				break;
+		}
+	}
+
+	auto ConditionalAddSingle = [&]() -> void
+	{
+		static const FName PaddingName = TEXT("Internal.Padding_int32");
+
+		if (SingleIndices.IsEmpty())
+		{
+			// if we're not all done adding elements, then we're going to have to add some padding
+			if (!DoubleIndices.IsEmpty() || !TripleIndices.IsEmpty())
 			{
-				Chunk.Type = FNiagaraTypeDefinition::GetVec4Def();
-				Chunk.ComponentMask = TEXT(".xyz");
+				TOptional<FNiagaraVariable> ConflictingVariable;
+
+				FName InstanceName = PaddingName;
+				InstanceName.SetNumber(UniquePaddingIndex++);
+
+				FNameBuilder InstanceNameBuilder(InstanceName);
+
+				FNiagaraVariable PaddingVariable(FNiagaraTypeDefinition::GetIntDef(), InstanceName);
+				FString SymbolName = GetSanitizedSymbolName(InstanceNameBuilder);
+
+				const int32 PaddingVariableIndex = Parameters.Parameters.Num();
+				RegisterUniformChunkToPack(SymbolName, PaddingVariable, true /*AddPadding*/, Parameters, ConflictingVariable);
+				check(!ConflictingVariable.IsSet());
+
+				OrderedVariableIndices.Add(PaddingVariableIndex);
 			}
 		}
+		else
+		{
+			OrderedVariableIndices.Add(SingleIndices.Pop());
+		}
+	};
 
-		Chunk.Mode = ChunkMode;
+	while (!TripleIndices.IsEmpty())
+	{
+		OrderedVariableIndices.Add(TripleIndices.Pop());
+		ConditionalAddSingle();
+	}
 
-		ChunksByMode[static_cast<int32>(ChunkMode)].Add(Ret);
+	while (DoubleIndices.Num() > 1)
+	{
+		OrderedVariableIndices.Add(DoubleIndices.Pop());
+		OrderedVariableIndices.Add(DoubleIndices.Pop());
+	}
 
-		auto& SystemVar = ParamMapDefinedSystemVars.Add(InVariable.GetName());
-		SystemVar.ChunkIndex = Ret;
-		SystemVar.ChunkMode = static_cast<int32>(ChunkMode);
-		SystemVar.Variable = InVariable;
+	if (!DoubleIndices.IsEmpty())
+	{
+		OrderedVariableIndices.Add(DoubleIndices.Pop());
+		ConditionalAddSingle();
+		ConditionalAddSingle();
 	}
-	return Ret;
+
+	OrderedVariableIndices.Append(SingleIndices);
+
+	// now we want to reorder the chunks for the uniform expressions based on how we've rearranged the variables here
+	const int32 ChunkModeToSort = static_cast<int32>(ENiagaraCodeChunkMode::Uniform);
+	TArray<int32>& ChunksToSort = ChunksByMode[ChunkModeToSort];
+	const int32 ChunkCount = ChunksToSort.Num();
+
+	// note that chunks that fail IsVariableInUniformBuffer will still be in the ChunksByMode, but won't actually get written
+	// out.  For now we preserve their existence in the ChunksByMode.  Make a copy of the uniform chunks so that we can find
+	// all the ones that need to get added back in
+	TSet<int32> OutstandingChunkIndices(ChunksToSort);
+	ChunksToSort.Reset();
+
+	for (int32 SortedParamIt : OrderedVariableIndices)
+	{
+		const int32 ChunkIndex = UniformParametersToPack.FindRef(SortedParamIt);
+		ChunksToSort.Add(ChunkIndex);
+		OutstandingChunkIndices.Remove(ChunkIndex);
+	}
+
+	// now append in the outstanding chunk indices
+	for (int32 OutstandingChunkIndex : OutstandingChunkIndices)
+	{
+		ChunksToSort.Add(OutstandingChunkIndex);
+	}
+
+	// and finally reorder the Parameters array as well
+	TArray<FNiagaraVariable> OrderedVariables;
+	OrderedVariables.Reserve(OrderedVariableIndices.Num());
+	for (int32 VariableIndex : OrderedVariableIndices)
+	{
+		OrderedVariables.Add(Parameters.Parameters[VariableIndex]);
+	}
+
+	Parameters.Parameters = MoveTemp(OrderedVariables);
 }
 
 int32 FNiagaraHlslTranslator::AddSourceChunk(FString SymbolName, const FNiagaraTypeDefinition& Type, bool bSanitize)
@@ -5876,7 +6026,6 @@ bool TNiagaraHlslTranslator<GraphBridge>::ParameterMapRegisterExternalConstantNa
 
 	FString VarName = InVariable.GetName().ToString();
 	FString SymbolName = GetSanitizedSymbolName(VarName);
-	FString FlattenedName = SymbolName.Replace(TEXT("."), TEXT("_"));
 	FString ParameterMapInstanceName = GetParameterMapInstanceName(InParamMapHistoryIdx);
 
 	bool bMissingParameter = false;
@@ -5903,8 +6052,6 @@ bool TNiagaraHlslTranslator<GraphBridge>::ParameterMapRegisterExternalConstantNa
 
 			if (false == ParamMapDefinedSystemVars.Contains(InVariable.GetName()))
 			{
-				FString SymbolNameDefined = FlattenedName;
-
 				if (bIsDataInterface)
 				{
 					UNiagaraDataInterface* DataInterface = nullptr;
@@ -5959,34 +6106,21 @@ bool TNiagaraHlslTranslator<GraphBridge>::ParameterMapRegisterExternalConstantNa
 
 				if (IsVariableInUniformBuffer(InVariable))
 				{
-					// we must ensure that there's a one to one relationship between symbol name and parameter.  The generated VM only
-					// knows about the symbols while the parameter stores knows about the parameters, if these mismatch, then we're going
-					// to be incorrectly addressing the constant table
-					if (!CompilationOutput.ScriptData.Parameters.FindParameter(InVariable))
-					{
-						bool AddParameter = true;
+					TOptional<FNiagaraVariable> ConflictingVariable;
 
-						// add the parameter, but first evaluate whether any of the symbols for existing parameters would conflict
-						for (const FNiagaraVariable& ExistingParameter : CompilationOutput.ScriptData.Parameters.Parameters)
-						{
-							FNameBuilder ExistingParameterName(ExistingParameter.GetName());
-							if (GetSanitizedSymbolName(ExistingParameterName).Equals(SymbolName))
-							{
-								Error(FText::Format(LOCTEXT("NonUniqueSymbolNames", "Parameters ('{0}' and '{1}') found which resolve to the same HLSL symbol name '{2}'.  These should be disambiguated."),
-									FText::FromName(InVariable.GetName()), FText::FromName(ExistingParameter.GetName()), FText::FromString(SymbolName)), InNodeForErrorReporting, InDefaultPin);
-
-								AddParameter = false;
-							}
-						}
+					const bool bApplyPadding = UNiagaraScript::IsGPUScript(CompileOptions.TargetUsage);
+					UniformChunk = RegisterUniformChunkToPack(SymbolName, InVariable, bApplyPadding, CompilationOutput.ScriptData.Parameters, ConflictingVariable);
 
-						if (AddParameter)
-						{
-							CompilationOutput.ScriptData.Parameters.Parameters.Add(InVariable);
-						}
+					if (ConflictingVariable.IsSet())
+					{
+						Error(FText::Format(LOCTEXT("NonUniqueSymbolNames", "Parameters ('{0}' and '{1}') found which resolve to the same HLSL symbol name '{2}'.  These should be disambiguated."),
+							FText::FromName(InVariable.GetName()), FText::FromName(ConflictingVariable->GetName()), FText::FromString(SymbolName)), InNodeForErrorReporting, InDefaultPin);
 					}
 				}
-
-				UniformChunk = AddUniformChunk(SymbolNameDefined, InVariable, ENiagaraCodeChunkMode::Uniform, UNiagaraScript::IsGPUScript(CompileOptions.TargetUsage));
+				else
+				{
+					UniformChunk = AddChunkToConstantBuffer(SymbolName, InVariable, ENiagaraCodeChunkMode::Uniform);
+				}
 			}
 			else
 			{
diff --git a/Engine/Plugins/FX/Niagara/Source/NiagaraEditor/Private/NiagaraHlslTranslator.h b/Engine/Plugins/FX/Niagara/Source/NiagaraEditor/Private/NiagaraHlslTranslator.h
--- a/Engine/Plugins/FX/Niagara/Source/NiagaraEditor/Private/NiagaraHlslTranslator.h
+++ b/Engine/Plugins/FX/Niagara/Source/NiagaraEditor/Private/NiagaraHlslTranslator.h
@@ -268,11 +268,11 @@ protected:
 	TMap<FNiagaraDataSetID, TMap<int32, FDataSetAccessInfo>> DataSetWriteInfo[(int32)ENiagaraDataSetAccessMode::Num];
 	TMap<FNiagaraDataSetID, int32> DataSetWriteConditionalInfo[(int32)ENiagaraDataSetAccessMode::Num];
 
-	NIAGARAEDITOR_API FString GetDataSetAccessSymbol(FNiagaraDataSetID DataSet, int32 IndexChunk, bool bRead);
-	FORCEINLINE FNiagaraDataSetID GetInstanceDataSetID()const { return FNiagaraDataSetID(TEXT("DataInstance"), ENiagaraDataSetType::ParticleData); }
-	FORCEINLINE FNiagaraDataSetID GetSystemEngineDataSetID()const { return FNiagaraDataSetID(TEXT("Engine"), ENiagaraDataSetType::ParticleData); }
-	FORCEINLINE FNiagaraDataSetID GetSystemUserDataSetID()const { return FNiagaraDataSetID(TEXT("User"), ENiagaraDataSetType::ParticleData); }
-	FORCEINLINE FNiagaraDataSetID GetSystemConstantDataSetID()const { return FNiagaraDataSetID(TEXT("Constant"), ENiagaraDataSetType::ParticleData); }
+	FString GetDataSetAccessSymbol(FNiagaraDataSetID DataSet, int32 IndexChunk, bool bRead);
+	FNiagaraDataSetID GetInstanceDataSetID()const { return FNiagaraDataSetID(TEXT("DataInstance"), ENiagaraDataSetType::ParticleData); }
+	FNiagaraDataSetID GetSystemEngineDataSetID()const { return FNiagaraDataSetID(TEXT("Engine"), ENiagaraDataSetType::ParticleData); }
+	FNiagaraDataSetID GetSystemUserDataSetID()const { return FNiagaraDataSetID(TEXT("User"), ENiagaraDataSetType::ParticleData); }
+	FNiagaraDataSetID GetSystemConstantDataSetID()const { return FNiagaraDataSetID(TEXT("Constant"), ENiagaraDataSetType::ParticleData); }
 
 	/** All functions called in the script. */
 	struct FNiagaraFunctionBody
@@ -289,15 +289,15 @@ protected:
 	TMap<FNiagaraFunctionSignature, TArray<FName>> FunctionStageReadTargets;
 	TArray<TArray<FName>> ActiveStageReadTargets;
 
-	NIAGARAEDITOR_API void GenerateFunctionCall(ENiagaraScriptUsage ScriptUsage, FNiagaraFunctionSignature& FunctionSignature, TArrayView<const int32> Inputs, TArray<int32>& Outputs);
-	NIAGARAEDITOR_API FString GetFunctionIncludeStatement(const FNiagaraCustomHlslInclude& Include) const;
-	NIAGARAEDITOR_API FString GetFunctionSignature(const FNiagaraFunctionSignature& Sig);
+	void GenerateFunctionCall(ENiagaraScriptUsage ScriptUsage, FNiagaraFunctionSignature& FunctionSignature, TArrayView<const int32> Inputs, TArray<int32>& Outputs);
+	FString GetFunctionIncludeStatement(const FNiagaraCustomHlslInclude& Include) const;
+	FString GetFunctionSignature(const FNiagaraFunctionSignature& Sig);
 
-	NIAGARAEDITOR_API void WriteDataSetContextVars(TMap<FNiagaraDataSetID, TMap<int32, FDataSetAccessInfo>>& DataSetAccessInfo, bool bRead, FString& OutHLSLOutput);
-	NIAGARAEDITOR_API void WriteDataSetStructDeclarations(TMap<FNiagaraDataSetID, TMap<int32, FDataSetAccessInfo>>& DataSetAccessInfo, bool bRead, FString& OutHLSLOutput);
-	NIAGARAEDITOR_API void DecomposeVariableAccess(UStruct* Struct, bool bRead, FString IndexSymbol, FString HLSLString);
+	void WriteDataSetContextVars(TMap<FNiagaraDataSetID, TMap<int32, FDataSetAccessInfo>>& DataSetAccessInfo, bool bRead, FString& OutHLSLOutput);
+	void WriteDataSetStructDeclarations(TMap<FNiagaraDataSetID, TMap<int32, FDataSetAccessInfo>>& DataSetAccessInfo, bool bRead, FString& OutHLSLOutput);
+	void DecomposeVariableAccess(UStruct* Struct, bool bRead, FString IndexSymbol, FString HLSLString);
 
-	NIAGARAEDITOR_API FString GetUniqueSymbolName(FName BaseName);
+	FString GetUniqueSymbolName(FName BaseName);
 
 	/** Stack of all function params. */
 	struct FFunctionContext
@@ -315,34 +315,38 @@ protected:
 	};
 	TArray<FFunctionContext> FunctionContextStack;
 	const FFunctionContext* FunctionCtx()const { return FunctionContextStack.Num() > 0 ? &FunctionContextStack.Last() : nullptr; }
-	NIAGARAEDITOR_API FString GetCallstack();
-	NIAGARAEDITOR_API TArray<FGuid> GetCallstackGuids();
+	FString GetCallstack();
+	TArray<FGuid> GetCallstackGuids();
 
-	NIAGARAEDITOR_API void EnterStatsScope(FNiagaraStatScope StatScope);
-	NIAGARAEDITOR_API void ExitStatsScope();
-	NIAGARAEDITOR_API void EnterStatsScope(FNiagaraStatScope StatScope, FString& OutHlsl);
-	NIAGARAEDITOR_API void ExitStatsScope(FString& OutHlsl);
+	void EnterStatsScope(FNiagaraStatScope StatScope);
+	void ExitStatsScope();
+	void EnterStatsScope(FNiagaraStatScope StatScope, FString& OutHlsl);
+	void ExitStatsScope(FString& OutHlsl);
 
-	NIAGARAEDITOR_API FString GeneratedConstantString(float Constant);
-	NIAGARAEDITOR_API FString GeneratedConstantString(FVector4 Constant);
+	FString GeneratedConstantString(float Constant);
+	FString GeneratedConstantString(FVector4 Constant);
 
 	/* Add a chunk that is not written to the source, only used as a source chunk for others. */
-	NIAGARAEDITOR_API int32 AddSourceChunk(FString SymbolName, const FNiagaraTypeDefinition& Type, bool bSanitize = true);
+	int32 AddSourceChunk(FString SymbolName, const FNiagaraTypeDefinition& Type, bool bSanitize = true);
 
 	/** Add a chunk defining a uniform value. */
-	NIAGARAEDITOR_API int32 AddUniformChunk(FString SymbolName, const FNiagaraVariable& InVariable, ENiagaraCodeChunkMode ChunkMode, bool AddPadding);
+	int32 AddChunkToConstantBuffer(const FString& SymbolName, const FNiagaraVariable& InVariable, ENiagaraCodeChunkMode ChunkMode);
+
+	/** Reserves a chunk for a uniform value.  The chunk isn't resolved till the call to PackRegisteredUniformChunk. */
+	int32 RegisterUniformChunkToPack(const FString& SymbolName, const FNiagaraVariable& InVariable, bool AddPadding, FNiagaraParameters& Parameters, TOptional<FNiagaraVariable>& ConflictingVariable);
+	void PackRegisteredUniformChunk(FNiagaraParameters& Parameters);
 
 	/* Add a chunk that is written to the body of the shader code. */
-	NIAGARAEDITOR_API int32 AddBodyChunk(FString SymbolName, FString Definition, const FNiagaraTypeDefinition& Type, TArray<int32>& SourceChunks, bool bDecl = true, bool bIsTerminated = true);
-	NIAGARAEDITOR_API int32 AddBodyChunk(FString SymbolName, FString Definition, const FNiagaraTypeDefinition& Type, int32 SourceChunk, bool bDecl = true, bool bIsTerminated = true);
-	NIAGARAEDITOR_API int32 AddBodyChunk(FString SymbolName, FString Definition, const FNiagaraTypeDefinition& Type, bool bDecl = true, bool bIsTerminated = true);
-	NIAGARAEDITOR_API int32 AddBodyComment(const FString& Comment);
-	NIAGARAEDITOR_API int32 AddBodyChunk(const FString& Definition);
+	int32 AddBodyChunk(FString SymbolName, FString Definition, const FNiagaraTypeDefinition& Type, TArray<int32>& SourceChunks, bool bDecl = true, bool bIsTerminated = true);
+	int32 AddBodyChunk(FString SymbolName, FString Definition, const FNiagaraTypeDefinition& Type, int32 SourceChunk, bool bDecl = true, bool bIsTerminated = true);
+	int32 AddBodyChunk(FString SymbolName, FString Definition, const FNiagaraTypeDefinition& Type, bool bDecl = true, bool bIsTerminated = true);
+	int32 AddBodyComment(const FString& Comment);
+	int32 AddBodyChunk(const FString& Definition);
 
-	NIAGARAEDITOR_API FString GetUniqueEmitterName() const;
-	NIAGARAEDITOR_API void ConvertCompileInfoToParamInfo(const FNiagaraScriptDataInterfaceCompileInfo& InCompileInfo, FNiagaraDataInterfaceGPUParamInfo& OutGPUParamInfo, TArray<FNiagaraFunctionSignature>& GeneratedFunctionSignatures);
+	FString GetUniqueEmitterName() const;
+	void ConvertCompileInfoToParamInfo(const FNiagaraScriptDataInterfaceCompileInfo& InCompileInfo, FNiagaraDataInterfaceGPUParamInfo& OutGPUParamInfo, TArray<FNiagaraFunctionSignature>& GeneratedFunctionSignatures);
 
-	NIAGARAEDITOR_API FString GetFunctionDefinitions();
+	FString GetFunctionDefinitions();
 
 public:
 	NIAGARAEDITOR_API FNiagaraHlslTranslator();
@@ -424,48 +428,48 @@ public:
 
 
 protected:
-	NIAGARAEDITOR_API bool ShouldConsiderTargetParameterMap(ENiagaraScriptUsage InUsage) const;
+	bool ShouldConsiderTargetParameterMap(ENiagaraScriptUsage InUsage) const;
 
-	NIAGARAEDITOR_API void HandleNamespacedExternalVariablesToDataSetRead(TArray<FNiagaraVariable>& InDataSetVars, FString InNamespaceStr);
+	void HandleNamespacedExternalVariablesToDataSetRead(TArray<FNiagaraVariable>& InDataSetVars, FString InNamespaceStr);
 
 	// For GPU simulations we have to special case some variables and pass them view shader parameters rather than the uniform buffer as they vary from CPU simulations
-	NIAGARAEDITOR_API bool IsVariableInUniformBuffer(const FNiagaraVariable& Variable) const;
+	bool IsVariableInUniformBuffer(const FNiagaraVariable& Variable) const;
 
-	NIAGARAEDITOR_API FString ComputeMatrixColumnAccess(const FString& Name);
-	NIAGARAEDITOR_API FString ComputeMatrixRowAccess(const FString& Name);
+	FString ComputeMatrixColumnAccess(const FString& Name);
+	FString ComputeMatrixRowAccess(const FString& Name);
 
 	// Add a raw float constant chunk
-	NIAGARAEDITOR_API int32 GetConstantDirect(float InValue);
-	NIAGARAEDITOR_API int32 GetConstantDirect(bool InValue);
-	NIAGARAEDITOR_API int32 GetConstantDirect(int InValue);
+	int32 GetConstantDirect(float InValue);
+	int32 GetConstantDirect(bool InValue);
+	int32 GetConstantDirect(int InValue);
 
-	NIAGARAEDITOR_API FNiagaraTypeDefinition GetChildType(const FNiagaraTypeDefinition& BaseType, const FName& PropertyName);
-	NIAGARAEDITOR_API FString NamePathToString(const FString& Prefix, const FNiagaraTypeDefinition& RootType, const TArray<FName>& NamePath);
-	NIAGARAEDITOR_API FString GenerateAssignment(const FNiagaraTypeDefinition& SrcType, const TArray<FName>& SrcPath, const FNiagaraTypeDefinition& DestType, const TArray<FName>& DestPath);
+	FNiagaraTypeDefinition GetChildType(const FNiagaraTypeDefinition& BaseType, const FName& PropertyName);
+	FString NamePathToString(const FString& Prefix, const FNiagaraTypeDefinition& RootType, const TArray<FName>& NamePath);
+	FString GenerateAssignment(const FNiagaraTypeDefinition& SrcType, const TArray<FName>& SrcPath, const FNiagaraTypeDefinition& DestType, const TArray<FName>& DestPath);
 
 	//Generates the code for the passed chunk.
-	NIAGARAEDITOR_API FString GetCode(FNiagaraCodeChunk& Chunk);
-	NIAGARAEDITOR_API FString GetCode(int32 ChunkIdx);
+	FString GetCode(FNiagaraCodeChunk& Chunk);
+	FString GetCode(int32 ChunkIdx);
 	//Retreives the code for this chunk being used as a source for another chunk
-	NIAGARAEDITOR_API FString GetCodeAsSource(int32 ChunkIdx);
+	FString GetCodeAsSource(int32 ChunkIdx);
 
 	// Generate a structure initializer string
 	// Returns true if we generated the structure successfully or false if we encounter something we could not handle
-	NIAGARAEDITOR_API bool GenerateStructInitializer(TStringBuilder<128>& InitializerString, UStruct* UserDefinedStruct, const void* StructData, int32 ByteOffset = 0);
+	bool GenerateStructInitializer(TStringBuilder<128>& InitializerString, UStruct* UserDefinedStruct, const void* StructData, int32 ByteOffset = 0);
 	// Convert a variable with actual data into a constant string
-	NIAGARAEDITOR_API FString GenerateConstantString(const FNiagaraVariable& Constant);
+	FString GenerateConstantString(const FNiagaraVariable& Constant);
 
 	// Takes the current script state (interpolated or not) and determines the correct context variable.
-	NIAGARAEDITOR_API FString GetParameterMapInstanceName(int32 ParamMapHistoryIdx);
+	FString GetParameterMapInstanceName(int32 ParamMapHistoryIdx);
 
 	// Register an attribute in its non namespaced form
 	virtual bool ParameterMapRegisterNamespaceAttributeVariable(const FNiagaraVariable& InVariable, int32 InParamMapHistoryIdx, int32& Output) = 0;
 
-	NIAGARAEDITOR_API bool ShouldInterpolateParameter(const FNiagaraVariable& Parameter);
+	bool ShouldInterpolateParameter(const FNiagaraVariable& Parameter);
 
-	NIAGARAEDITOR_API bool IsBulkSystemScript() const;
-	NIAGARAEDITOR_API bool IsSpawnScript() const;
-	NIAGARAEDITOR_API bool RequiresInterpolation() const;
+	bool IsBulkSystemScript() const;
+	bool IsSpawnScript() const;
+	bool RequiresInterpolation() const;
 
 	template<typename T>
 	void BuildConstantBuffer(ENiagaraCodeChunkMode ChunkMode);
@@ -494,7 +498,7 @@ protected:
 	TArray<FunctionNodeStackEntry> FunctionNodeStack;
 
 	// Synced to the System uniforms encountered for parameter maps thus far.
-	struct UniformVariableInfo
+	struct FUniformVariableInfo
 	{
 		FNiagaraVariable Variable;
 		int32 ChunkIndex;
@@ -509,7 +513,9 @@ protected:
 
 	TArray<FNiagaraVariable> ParamMapDefinedEngineVars; // Engine provided values that we want to be accessible via the parameter map but not to add as uniforms etc.
 
-	TMap<FName, UniformVariableInfo> ParamMapDefinedSystemVars; // Map from the defined constants to the uniform chunk expressing them (i.e. have we encountered before in this graph?)
+	TMap<int32 /*parameter index*/, int32 /*chunk index*/> UniformParametersToPack; // Map connecting the parameters that need to be packed for a uniform buffer with their registered chunk index
+
+	TMap<FName, FUniformVariableInfo> ParamMapDefinedSystemVars; // Map from the defined constants to the uniform chunk expressing them (i.e. have we encountered before in this graph?)
 
 	// Synced to the EmitterParameter uniforms encountered for parameter maps thus far.
 	TMap<FName, int32> ParamMapDefinedEmitterParameterVarsToUniformChunks; // Map from the variable name exposed by the emitter as a parameter to the uniform chunk expressing it (i.e. have we encountered before in this graph?)
diff --git a/Engine/Source/Runtime/Core/Private/UObject/FortniteMainBranchObjectVersion.cpp b/Engine/Source/Runtime/Core/Private/UObject/FortniteMainBranchObjectVersion.cpp
--- a/Engine/Source/Runtime/Core/Private/UObject/FortniteMainBranchObjectVersion.cpp
+++ b/Engine/Source/Runtime/Core/Private/UObject/FortniteMainBranchObjectVersion.cpp
@@ -11,7 +11,7 @@ TMap<FGuid, FGuid> FFortniteMainBranchObjectVersion::GetSystemGuids()
 	SystemGuids.Add(DevGuids.MATERIALSHADERMAP_DERIVEDDATA_VER, FGuid("CE7776E63E9E43FBB8C103EE1AB40B5D"));
 	SystemGuids.Add(DevGuids.NANITE_DERIVEDDATA_VER, FGuid("D28C003F825F4C178F698938FB0E47F7"));
 	SystemGuids.Add(DevGuids.NIAGARASHADERMAP_DERIVEDDATA_VER, FGuid("D0DF49D408F34E38A98FA097E0280EC7"));
-	SystemGuids.Add(DevGuids.Niagara_LatestScriptCompileVersion, FGuid("9042E07ACA4C4375B7D063C44FA21114"));
+	SystemGuids.Add(DevGuids.Niagara_LatestScriptCompileVersion, FGuid("EEF9B6D5121D4D119D9205D5F0E838FD"));
 	SystemGuids.Add(DevGuids.SkeletalMeshDerivedDataVersion, FGuid("518D7C51D31D9BE4A9F4B81542AC8119"));
 	SystemGuids.Add(DevGuids.STATICMESH_DERIVEDDATA_VER, FGuid("543843B103794E0E9BA4BE80FB602F79"));
 	SystemGuids.Add(DevGuids.MaterialTranslationDDCVersion, FGuid("982B2D2E6E50436A95EF4D16B1D04920"));
