- Change Neighbor grid to use pixel format capabilities vs volume atomics

diff --git a/Engine/Plugins/FX/Niagara/Source/Niagara/Classes/NiagaraDataInterfaceNeighborGrid3D.h b/Engine/Plugins/FX/Niagara/Source/Niagara/Classes/NiagaraDataInterfaceNeighborGrid3D.h
--- a/Engine/Plugins/FX/Niagara/Source/Niagara/Classes/NiagaraDataInterfaceNeighborGrid3D.h
+++ b/Engine/Plugins/FX/Niagara/Source/Niagara/Classes/NiagaraDataInterfaceNeighborGrid3D.h
@@ -104,9 +104,6 @@ public:
 #endif
 	NIAGARA_API virtual void BuildShaderParameters(FNiagaraShaderParametersBuilder& ShaderParametersBuilder) const override;
 	NIAGARA_API virtual void SetShaderParameters(const FNiagaraDataInterfaceSetShaderParametersContext& Context) const override;
-#if WITH_EDITOR
-	NIAGARA_API virtual bool ShouldCompile(EShaderPlatform ShaderPlatform) const override;
-#endif
 
 	virtual void ProvidePerInstanceDataForRenderThread(void* DataForRenderThread, void* PerInstanceData, const FNiagaraSystemInstanceID& SystemInstance) override {}
 	NIAGARA_API virtual bool InitPerInstanceData(void* PerInstanceData, FNiagaraSystemInstance* SystemInstance) override;
diff --git a/Engine/Plugins/FX/Niagara/Source/Niagara/Private/NiagaraDataInterfaceNeighborGrid3D.cpp b/Engine/Plugins/FX/Niagara/Source/Niagara/Private/NiagaraDataInterfaceNeighborGrid3D.cpp
--- a/Engine/Plugins/FX/Niagara/Source/Niagara/Private/NiagaraDataInterfaceNeighborGrid3D.cpp
+++ b/Engine/Plugins/FX/Niagara/Source/Niagara/Private/NiagaraDataInterfaceNeighborGrid3D.cpp
@@ -285,18 +285,6 @@ bool UNiagaraDataInterfaceNeighborGrid3D::Equals(const UNiagaraDataInterface* Ot
 	return OtherTyped->MaxNeighborsPerCell == MaxNeighborsPerCell;
 }
 
-#if WITH_EDITOR
-bool UNiagaraDataInterfaceNeighborGrid3D::ShouldCompile(EShaderPlatform ShaderPlatform) const
-{
-	if (!RHISupportsVolumeTextureAtomics(ShaderPlatform))
-	{
-		return false;
-	}
-
-	return UNiagaraDataInterface::ShouldCompile(ShaderPlatform);
-}
-#endif
-
 #if WITH_EDITORONLY_DATA
 bool UNiagaraDataInterfaceNeighborGrid3D::AppendCompileHash(FNiagaraCompileHashVisitor* InVisitor) const
 {
@@ -508,6 +496,11 @@ void UNiagaraDataInterfaceNeighborGrid3D::SetShaderParameters(const FNiagaraData
 
 bool UNiagaraDataInterfaceNeighborGrid3D::InitPerInstanceData(void* PerInstanceData, FNiagaraSystemInstance* SystemInstance)
 {
+	if (UE::PixelFormat::HasCapabilities(EPixelFormat::PF_R32_SINT, EPixelFormatCapabilities::TypedUAVLoad | EPixelFormatCapabilities::TypedUAVStore | EPixelFormatCapabilities::BufferAtomics) == false)
+	{
+		return false;
+	}
+
 	FNDINeighborGrid3DInstanceData_GT* InstanceData = new (PerInstanceData) FNDINeighborGrid3DInstanceData_GT();
 
 	FNiagaraDataInterfaceProxyNeighborGrid3D* RT_Proxy = GetProxyAs<FNiagaraDataInterfaceProxyNeighborGrid3D>();
