- LWC fix for sampling SM / SKM without a component
- The LWC tile offset was incorrectly applied to the component transform

diff --git a/Engine/Plugins/FX/Niagara/Source/Niagara/Classes/NiagaraDataInterfaceSkeletalMesh.h b/Engine/Plugins/FX/Niagara/Source/Niagara/Classes/NiagaraDataInterfaceSkeletalMesh.h
--- a/Engine/Plugins/FX/Niagara/Source/Niagara/Classes/NiagaraDataInterfaceSkeletalMesh.h
+++ b/Engine/Plugins/FX/Niagara/Source/Niagara/Classes/NiagaraDataInterfaceSkeletalMesh.h
@@ -674,6 +674,7 @@ struct FNDISkeletalMesh_InstanceData
 		return CachedLODData ? &CachedLODData->SkinWeightVertexBuffer : nullptr;
 	}
 
+	FTransform CalculateComponentTransform(FNiagaraSystemInstance* SystemInstance) const;
 	void UpdateFilteredSocketTransforms();
 	TArray<FTransform3f>& GetFilteredSocketsWriteBuffer() { return FilteredSocketTransforms[FilteredSocketTransformsIndex]; }
 	const TArray<FTransform3f>& GetFilteredSocketsCurrBuffer() const { return FilteredSocketTransforms[FilteredSocketTransformsIndex]; }
diff --git a/Engine/Plugins/FX/Niagara/Source/Niagara/Private/DataInterface/NiagaraDataInterfaceStaticMesh.cpp b/Engine/Plugins/FX/Niagara/Source/Niagara/Private/DataInterface/NiagaraDataInterfaceStaticMesh.cpp
--- a/Engine/Plugins/FX/Niagara/Source/Niagara/Private/DataInterface/NiagaraDataInterfaceStaticMesh.cpp
+++ b/Engine/Plugins/FX/Niagara/Source/Niagara/Private/DataInterface/NiagaraDataInterfaceStaticMesh.cpp
@@ -747,6 +747,7 @@ namespace NDIStaticMeshLocal
 		void UpdateTransforms(USceneComponent* SceneComponent, FNiagaraSystemInstance* SystemInstance)
 		{
 			SystemInstanceWorldTransform = SystemInstance->GetWorldTransform();
+			LWCTileOffset = FVector(SystemInstance->GetLWCTile()) * -FLargeWorldRenderScalar::GetTileSize();
 
 			FTransform ComponentTransform;
 			if (SceneComponent)
@@ -770,17 +771,17 @@ namespace NDIStaticMeshLocal
 				{
 					ISMTransforms.Empty();
 				}
+
+				OwnerToMeshVector = FVector3f(ComponentTransform.GetLocation() - SystemInstanceWorldTransform.GetLocation());
+				ComponentTransform.AddToTranslation(LWCTileOffset);
 			}
 			else
 			{
 				ISMTransforms.Empty();
 				ComponentTransform = SystemInstanceWorldTransform;
+				OwnerToMeshVector = FVector3f::ZeroVector;
 			}
 
-			OwnerToMeshVector = FVector3f(ComponentTransform.GetLocation() - SystemInstanceWorldTransform.GetLocation());
-			LWCTileOffset = FVector(SystemInstance->GetLWCTile()) * -FLargeWorldRenderScalar::GetTileSize();
-			ComponentTransform.AddToTranslation(LWCTileOffset);
-
 			Transform = ComponentTransform.ToMatrixWithScale();
 			TransformInverseTransposed = ComponentTransform.Inverse().ToMatrixWithScale().GetTransposed();
 			Rotation = FQuat4f(ComponentTransform.GetRotation());
diff --git a/Engine/Plugins/FX/Niagara/Source/Niagara/Private/NiagaraDataInterfaceSkeletalMesh.cpp b/Engine/Plugins/FX/Niagara/Source/Niagara/Private/NiagaraDataInterfaceSkeletalMesh.cpp
--- a/Engine/Plugins/FX/Niagara/Source/Niagara/Private/NiagaraDataInterfaceSkeletalMesh.cpp
+++ b/Engine/Plugins/FX/Niagara/Source/Niagara/Private/NiagaraDataInterfaceSkeletalMesh.cpp
@@ -1629,8 +1629,7 @@ bool FNDISkeletalMesh_InstanceData::Init(UNiagaraDataInterfaceSkeletalMesh* Inte
 	bMeshValid = Mesh != nullptr;
 	bComponentValid = SceneComponent.IsValid();
 
-	FTransform ComponentTransform = (bComponentValid ? SceneComponent->GetComponentToWorld() : SystemInstance->GetWorldTransform());
-	ComponentTransform.AddToTranslation(FVector(SystemInstance->GetLWCTile()) * -FLargeWorldRenderScalar::GetTileSize());
+	const FTransform ComponentTransform = CalculateComponentTransform(SystemInstance);
 	Transform = ComponentTransform.ToMatrixWithScale();
 	TransformInverseTransposed = Transform.Inverse().GetTransposed();
 	PrevTransform = Transform;
@@ -2145,8 +2144,7 @@ bool FNDISkeletalMesh_InstanceData::Tick(UNiagaraDataInterfaceSkeletalMesh* Inte
 		DeltaSeconds = InDeltaSeconds;
 
 		PrevTransform = Transform;
-		FTransform ComponentTransform = (SceneComponent.IsValid() ? SceneComponent->GetComponentToWorld() : SystemInstance->GetWorldTransform());
-		ComponentTransform.AddToTranslation(FVector(SystemInstance->GetLWCTile()) * -FLargeWorldRenderScalar::GetTileSize());
+		const FTransform ComponentTransform = CalculateComponentTransform(SystemInstance);
 		Transform = ComponentTransform.ToMatrixWithScale();
 		TransformInverseTransposed = Transform.Inverse().GetTransposed();
 
@@ -2163,6 +2161,20 @@ bool FNDISkeletalMesh_InstanceData::Tick(UNiagaraDataInterfaceSkeletalMesh* Inte
 	}
 }
 
+FTransform FNDISkeletalMesh_InstanceData::CalculateComponentTransform(FNiagaraSystemInstance* SystemInstance) const
+{
+	if (USceneComponent* SceneComp = SceneComponent.Get())
+	{
+		FTransform ComponentTransform = SceneComp->GetComponentToWorld();
+		ComponentTransform.AddToTranslation(FVector(SystemInstance->GetLWCTile()) * -FLargeWorldRenderScalar::GetTileSize());
+		return ComponentTransform;
+	}
+	else
+	{
+		return SystemInstance->GetWorldTransform();
+	}
+}
+
 void FNDISkeletalMesh_InstanceData::UpdateFilteredSocketTransforms()
 {
 	USkeletalMeshComponent* SkelComp = Cast<USkeletalMeshComponent>(SceneComponent.Get());
