- FNiagaraEmitterInstance refactoring in a base class for different emitter types
- FNiagaraEmitterInstanceImpl is the stateful implementation

diff --git a/Engine/Plugins/FX/Niagara/Source/Niagara/Classes/NiagaraEmitterInstance.h b/Engine/Plugins/FX/Niagara/Source/Niagara/Classes/NiagaraEmitterInstance.h
--- a/Engine/Plugins/FX/Niagara/Source/Niagara/Classes/NiagaraEmitterInstance.h
+++ b/Engine/Plugins/FX/Niagara/Source/Niagara/Classes/NiagaraEmitterInstance.h
@@ -1,8 +1,5 @@
 // Copyright Epic Games, Inc. All Rights Reserved.
 
-/*==============================================================================
-NiagaraEmitterInstance.h: Niagara emitter simulation class
-==============================================================================*/
 #pragma once
 
 #include "CoreMinimal.h"
@@ -12,261 +9,228 @@ NiagaraEmitterInstance.h: Niagara emitter simulation class
 #include "NiagaraEmitter.h"
 #include "NiagaraScriptExecutionContext.h"
 
+struct FNiagaraComputeExecutionContext;
+struct FNiagaraScriptExecutionContext;
+class UNiagaraDataInterface;
 class FNiagaraSystemInstance;
-struct FNiagaraEmitterHandle;
-class UNiagaraParameterCollection;
-class UNiagaraParameterCollectionInstance;
-class FNiagaraGpuComputeDispatchInterface;
-struct FNiagaraEmitterCompiledData;
+
+class FNiagaraEmitterInstanceImpl;
 
 /**
-* A Niagara particle simulation.
+* Base class for different emitter instances
 */
 class FNiagaraEmitterInstance
 {
 	friend class UNiagaraSimCache;
 
-private:
-	struct FEventInstanceData
-	{
-		TArray<FNiagaraScriptExecutionContext> EventExecContexts;
-		TArray<FNiagaraParameterDirectBinding<int32>> EventExecCountBindings;
-
-		TArray<FNiagaraDataSet*> UpdateScriptEventDataSets;
-		TArray<FNiagaraDataSet*> SpawnScriptEventDataSets;
-
-		TArray<bool> UpdateEventGeneratorIsSharedByIndex;
-		TArray<bool> SpawnEventGeneratorIsSharedByIndex;
-
-		/** Data required for handling events. */
-		TArray<FNiagaraEventHandlingInfo> EventHandlingInfo;
-		int32 EventSpawnTotal = 0;
-	};
-
 public:
 	explicit FNiagaraEmitterInstance(FNiagaraSystemInstance* InParentSystemInstance);
-	virtual ~FNiagaraEmitterInstance();
+	virtual ~FNiagaraEmitterInstance() {}
+
+	//~Begin: Define Emitter Interface
+	virtual void Init(int32 InEmitterIdx);
+	virtual void ResetSimulation(bool bKillExisting = true) = 0;
+	virtual void OnPooledReuse() = 0;
+	virtual bool HandleCompletion(bool bForce = false) = 0;
+
+	virtual void BindParameters(bool bExternalOnly) = 0;
+	virtual void UnbindParameters(bool bExternalOnly) = 0;
+
+	virtual void PreTick() {};
+	virtual void Tick(float DeltaSeconds) = 0;
+	//~End: Define Emitter Interface
+
+	FNiagaraSystemInstance* GetParentSystemInstance() const { return ParentSystemInstance; }
+	bool IsLocalSpace() const { return bLocalSpace; }
+	bool IsDeterministic() const { return bDeterministic; }
+	bool NeedsPartialDepthTexture() const { return bNeedsPartialDepthTexture; }
+	bool NeedsEarlyViewUniformBuffer() const { return bNeedsEarlyViewUniformBuffer; }
+	ENiagaraSimTarget GetSimTarget() const { return SimTarget; }
+#if STATS
+	TStatId GetEmitterStatID(bool bGameThread, bool bConcurrent) const;
+	TStatId GetSystemStatID(bool bGameThread, bool bConcurrent) const;
+#endif
+	FNiagaraDataSet& GetParticleData() { check(ParticleDataSet); return *ParticleDataSet; }
+	const FNiagaraDataSet& GetParticleData() const { check(ParticleDataSet); return *ParticleDataSet; }
 
-	void Init(int32 InEmitterIdx, FNiagaraSystemInstanceID SystemInstanceID);
-	void InitDITickLists();
+	FNiagaraComputeExecutionContext* GetGPUContext() const { return GPUExecContext; }
 
-	void ResetSimulation(bool bKillExisting = true);
+	ENiagaraExecutionState GetExecutionState() const { return ExecutionState; }
+	bool IsActive() const { return ExecutionState == ENiagaraExecutionState::Active; }
+	bool IsDisabled() const { return ExecutionState == ENiagaraExecutionState::Disabled; }
+	bool IsInactive() const { return ExecutionState == ENiagaraExecutionState::Inactive; }
+	bool IsComplete() const { return ExecutionState == ENiagaraExecutionState::Complete || ExecutionState == ENiagaraExecutionState::Disabled; }
 
-	void OnPooledReuse();
+	FORCEINLINE bool ShouldTick() const { return ExecutionState == ENiagaraExecutionState::Active || GetNumParticles() > 0; }
 
-	void DirtyDataInterfaces();
+	//-TODO:Stateless: Does this need to be virtual?  Can we cache the value after ticking / allocating the data / have a cache value somewhere instead?
+	NIAGARA_API virtual int32 GetNumParticles() const;
+	int32 GetTotalSpawnedParticles() const { return TotalSpawnedParticles; }
 
-	/** Replaces the binding for a single parameter collection instance. If for example the component begins to override the global instance. */
-	//void RebindParameterCollection(UNiagaraParameterCollectionInstance* OldInstance, UNiagaraParameterCollectionInstance* NewInstance);
-	void BindParameters(bool bExternalOnly);
-	void UnbindParameters(bool bExternalOnly);
+	bool AreBoundsDynamic() const { return bCachedBoundsDynamic; }
+	[[nodiscard]] FBox GetBounds() const { return CachedBounds; }
 
-	bool IsAllowedToExecute() const;
+	void SetSystemFixedBoundsOverride(FBox SystemFixedBounds);
 
-#if WITH_EDITOR
-	void TickRapidIterationParameters();
+	void SetFixedBounds(const FBox& InLocalBounds);
+	[[nodiscard]] FBox GetFixedBounds() const;
+	
+#if WITH_EDITORONLY_DATA
+	uint32 GetTickTimeCycles() const { return TickTimeCycles; }
+	NIAGARA_API bool IsDisabledFromIsolation() const;
 #endif
 
-	void PreTick();
-	void Tick(float DeltaSeconds);
-	void PostTick();
-	bool HandleCompletion(bool bForce = false);
+	[[nodiscard]] NIAGARA_API int64 GetTotalBytesUsed() const;
 
-	bool RequiresPersistentIDs() const;
+	[[nodiscard]] NIAGARA_API const FNiagaraEmitterHandle& GetEmitterHandle() const;
 
-	FORCEINLINE bool ShouldTick()const { return ExecutionState == ENiagaraExecutionState::Active || GetNumParticles() > 0; }
-
-	uint32 CalculateEventSpawnCount(const FNiagaraEventScriptProperties &EventHandlerProps, TArray<int32, TInlineAllocator<16>>& EventSpawnCounts, FNiagaraDataSet *EventSet);
+	FNiagaraParameterStore& GetRendererBoundVariables() { return RendererBindings; }
+	const FNiagaraParameterStore& GetRendererBoundVariables() const { return RendererBindings; }
 
-#if WITH_EDITOR
-	/** Potentially reads back data from the GPU which will introduce a stall and should only be used for debug purposes. */
-	NIAGARA_API void CalculateFixedBounds(const FTransform& ToWorldSpace);
-#endif
-	
 	NIAGARA_API bool GetBoundRendererValue_GT(const FNiagaraVariableBase& InBaseVar, const FNiagaraVariableBase& InSubVar, void* OutValueData) const;
 
-	FNiagaraDataSet& GetData() { return *ParticleDataSet; }
-	const FNiagaraDataSet& GetData() const { return *ParticleDataSet; }
+	[[nodiscard]] NIAGARA_API UObject* FindBinding(const FNiagaraVariable& InVariable) const;
+	[[nodiscard]] NIAGARA_API UNiagaraDataInterface* FindDataInterface(const FNiagaraVariable& InVariable) const;
+	//////////////////////////////////////////////////////////////////////////////////////////////////////////////////
 
-	FORCEINLINE bool IsActive()const { return ExecutionState == ENiagaraExecutionState::Active; }
-	FORCEINLINE bool IsDisabled()const { return ExecutionState == ENiagaraExecutionState::Disabled; }
-	FORCEINLINE bool IsInactive()const { return ExecutionState == ENiagaraExecutionState::Inactive; }
-	FORCEINLINE bool IsComplete()const { return ExecutionState == ENiagaraExecutionState::Complete || ExecutionState == ENiagaraExecutionState::Disabled; }
+	////////////////////////////////////////////////////////////////////////////////////////////////////
+	//-TODO:Stateless:
+	// Can we remove these virtual functions
+	virtual FNiagaraEmitterInstanceImpl* AsStateful() { return nullptr; }
 
-private:
-	NIAGARA_API int32 GetNumParticlesGPUInternal() const;
-public:
-	FORCEINLINE int32 GetNumParticles() const
+	const FNiagaraEmitterInstanceImpl* AsStateful() const { return const_cast<FNiagaraEmitterInstance*>(this)->AsStateful(); }
+
+	UNiagaraEmitter* GetEmitter() const { return VersionedEmitter.Emitter; }
+	const FVersionedNiagaraEmitter& GetVersionedEmitter() const { return VersionedEmitter; }
+	//FVersionedNiagaraEmitterData* GetVersionedEmitterData() const { return VersionedEmitter.GetEmitterData(); }
+
+	template<typename TAction>
+	void ForEachEnabledRenderer(TAction Func) const
 	{
-		// Note: For GPU simulations the data is latent we can not read directly from GetCurrentData() until we have passed a fence
-		// which guarantees that at least one tick has occurred inside the batcher.  The count will still technically be incorrect
-		// but hopefully adequate for a system script update.
-		if (GPUExecContext)
+		for (const UNiagaraRendererProperties* Renderer : GetRenderers())
 		{
-			return GetNumParticlesGPUInternal();
+			if (Renderer && Renderer->GetIsEnabled() && Renderer->IsSimTargetSupported(SimTarget))
+			{
+				Func(Renderer);
+			}
 		}
+	}
 
-		if ( ParticleDataSet->GetCurrentData() )
+	template<typename TAction>
+	void ForEachEnabledRenderer(TAction Func)
+	{
+		for (UNiagaraRendererProperties* Renderer : GetRenderers())
 		{
-			return ParticleDataSet->GetCurrentData()->GetNumInstances();
+			if (Renderer && Renderer->GetIsEnabled() && Renderer->IsSimTargetSupported(SimTarget))
+			{
+				Func(Renderer);
+			}
 		}
-		return 0;
 	}
 
-	FORCEINLINE int32 GetTotalSpawnedParticles()const { return TotalSpawnedParticles; }
-	FORCEINLINE const FNiagaraEmitterScalabilitySettings& GetScalabilitySettings()const { return GetCachedEmitterData()->GetScalabilitySettings(); }
-
-	NIAGARA_API const FNiagaraEmitterHandle& GetEmitterHandle() const;
-
-	FNiagaraSystemInstance* GetParentSystemInstance()const { return ParentSystemInstance; }
-
-	float NIAGARA_API GetTotalCPUTimeMS();
-	int64 NIAGARA_API GetTotalBytesUsed();
-
-	ENiagaraExecutionState GetExecutionState() { return ExecutionState; }
-	void NIAGARA_API SetExecutionState(ENiagaraExecutionState InState);
-
-	bool AreBoundsDynamic() const { return bCachedBoundsDynamic; }
-	FBox GetBounds() const { return CachedBounds; }
-
-	FNiagaraScriptExecutionContext& GetSpawnExecutionContext() { return SpawnExecContext; }
-	FNiagaraScriptExecutionContext& GetUpdateExecutionContext() { return UpdateExecContext; }
-	TArrayView<FNiagaraScriptExecutionContext> GetEventExecutionContexts();
-
-	FORCEINLINE FName GetCachedIDName()const { return CachedIDName; }
-	FORCEINLINE FVersionedNiagaraEmitter GetCachedEmitter()const { return CachedEmitter; }
-	FORCEINLINE FVersionedNiagaraEmitterData* GetCachedEmitterData()const { return CachedEmitter.GetEmitterData(); }
-
-	TArray<FNiagaraSpawnInfo>& GetSpawnInfo() { return SpawnInfos; }
-
-	NIAGARA_API bool IsReadyToRun() const;
+	[[nodiscard]] const UNiagaraRendererProperties* GetRenderer(int32 i) const { return GetRenderers()[i]; }
 
-	void Dump() const;
-
-	bool WaitForDebugInfo() const;
-
-	FNiagaraComputeExecutionContext* GetGPUContext()const
+	////////////////////////////////////////////////////////////////////////////////////////////////////
+	//-TODO:Stateless: Should we cache this data? clean up const casting
+protected:
+	virtual TConstArrayView<UNiagaraRendererProperties*> GetRenderers() const = 0;
+	TArrayView<UNiagaraRendererProperties*> GetRenderers()
 	{
-		return GPUExecContext;
+		TConstArrayView<UNiagaraRendererProperties*> Renderers = static_cast<const FNiagaraEmitterInstance*>(this)->GetRenderers();
+		return MakeArrayView(const_cast<UNiagaraRendererProperties**>(Renderers.GetData()), Renderers.Num());
 	}
 
-	void SetSystemFixedBoundsOverride(FBox SystemFixedBounds);
-	FORCEINLINE void SetFixedBounds(const FBox& InLocalBounds)
-	{
-		FRWScopeLock ScopeLock(FixedBoundsGuard, SLT_Write);
-		FixedBounds = InLocalBounds;
-	}
-	FBox GetFixedBounds() const;
-
+public:
+	////////////////////////////////////////////////////////////////////////////////////////////////////
+	//-TODO: Functionality that needs implementing for stateless
 	void SetEmitterEnable(bool bNewEnableState)
 	{
 		bAllowSpawning_GT = bNewEnableState;
 	}
 
-	NIAGARA_API UObject* FindBinding(const FNiagaraVariable& InVariable) const;
-	NIAGARA_API UNiagaraDataInterface* FindDataInterface(const FNiagaraVariable& InVariable) const;
-
-	bool HasTicked() const { return TickCount > 0;  }
-
-	const FNiagaraParameterStore& GetRendererBoundVariables() const { return RendererBindings; }
-	FNiagaraParameterStore& GetRendererBoundVariables() { return RendererBindings; }
-
-	int32 GetRandomSeed() const { return RandomSeed; }
-	int32 GetInstanceSeed() const { return InstanceSeed; }
-
-	void SetParticleComponentActive(FObjectKey ComponentKey, int32 ParticleID) const;
-
-	bool IsParticleComponentActive(FObjectKey ComponentKey, int32 ParticleID) const;
-
-private:
-	void CheckForErrors();
+	////////////////////////////////////////////////////////////////////////////////////////////////////
+	// Deprecated functionality code needs to be updated
+	UE_DEPRECATED(5.4, "Please update your code to handle execution path for different emitter types")
+	NIAGARA_API FNiagaraScriptExecutionContext& GetSpawnExecutionContext();
 
-	void BuildConstantBufferTable(
-		const FNiagaraScriptExecutionContext& ExecContext,
-		FScriptExecutionConstantBufferTable& ConstantBufferTable) const;
+	UE_DEPRECATED(5.4, "Please update your code to handle execution path for different emitter types")
+	NIAGARA_API FNiagaraScriptExecutionContext& GetUpdateExecutionContext();
 
-	/** Generate emitter bounds */
-	FBox InternalCalculateDynamicBounds(int32 ParticleCount) const;
+	UE_DEPRECATED(5.4, "Please update your code to handle execution path for different emitter types")
+	NIAGARA_API TArrayView<FNiagaraScriptExecutionContext> GetEventExecutionContexts();
 
-	/** Array of all spawn info driven by our owning emitter script. */
-	TArray<FNiagaraSpawnInfo> SpawnInfos;
+	UE_DEPRECATED(5.4, "Please update your code to handle execution path for different emitter types")
+	FORCEINLINE  FVersionedNiagaraEmitter GetCachedEmitter() const { return VersionedEmitter; }
 
-	FNiagaraScriptExecutionContext SpawnExecContext;
-	FNiagaraScriptExecutionContext UpdateExecContext;
-	FNiagaraComputeExecutionContext* GPUExecContext = nullptr;
+	UE_DEPRECATED(5.4, "Please update your code to handle execution path for different emitter types")
+	FORCEINLINE FVersionedNiagaraEmitterData* GetCachedEmitterData()const { return VersionedEmitter.GetEmitterData(); }
 
-	FNiagaraParameterDirectBinding<float> SpawnIntervalBinding;
-	FNiagaraParameterDirectBinding<float> InterpSpawnStartBinding;
-	FNiagaraParameterDirectBinding<int32> SpawnGroupBinding;
+	UE_DEPRECATED(5.4, "Please update your code to handle execution path for different emitter types")
+	NIAGARA_API TArray<FNiagaraSpawnInfo>& GetSpawnInfo();
 
-	FNiagaraParameterDirectBinding<int32> SpawnExecCountBinding;
-	FNiagaraParameterDirectBinding<int32> UpdateExecCountBinding;
-
-	TSharedPtr<const FNiagaraEmitterCompiledData> CachedEmitterCompiledData;
-	FNiagaraParameterStore RendererBindings;
-
-	TUniquePtr<FEventInstanceData> EventInstanceData;
-
-	/* Are the cached emitter bounds dynamic */
-	bool bCachedBoundsDynamic = false;
-
-	/* Emitter bounds */
-	FBox CachedBounds;
-
-	/** Optional user or VM specified bounds. */
-	mutable FRWLock FixedBoundsGuard;
-	FBox FixedBounds;
-
-	/** Cached fixed bounds of the parent system which override this Emitter Instances bounds if set. Whenever we initialize the owning SystemInstance we will reconstruct this
-	 ** EmitterInstance and the cached bounds will be unset. */
-	FBox CachedSystemFixedBounds;
-
-	FNiagaraSystemInstanceID OwnerSystemInstanceID = 0;
+	UE_DEPRECATED(5.4, "Please update your code to handle execution path for different emitter types")
+	void SetParticleComponentActive(FObjectKey ComponentKey, int32 ParticleID) const;
 
-	FNiagaraGpuComputeDispatchInterface* ComputeDispatchInterface = nullptr;
+	UE_DEPRECATED(5.4, "Please update your code to handle execution path for different emitter types")
+	bool IsParticleComponentActive(FObjectKey ComponentKey, int32 ParticleID) const;
 
-	/** particle simulation data. Must be a shared ref as various things on the RT can have direct ref to it. */
-	FNiagaraDataSet* ParticleDataSet = nullptr;
+	UE_DEPRECATED(5.4, "Please update your code to handle execution path for different emitter types")
+	NIAGARA_API bool IsReadyToRun() const;
 
-	FNiagaraSystemInstance *ParentSystemInstance = nullptr;
+	UE_DEPRECATED(5.4, "Please update your code to handle execution path for different emitter types")
+	NIAGARA_API bool HasTicked() const;
 
-	/** Raw pointer to the emitter that we're instanced from. Raw ptr should be safe here as we check for the validity of the system and it's emitters higher up before any ticking. */
-	FVersionedNiagaraEmitter CachedEmitter;
-	FName CachedIDName;
+	UE_DEPRECATED(5.4, "Please update your code to handle execution path for different emitter types")
+	NIAGARA_API void SetExecutionState(ENiagaraExecutionState InState);
 
-	/** The index of our emitter in our parent system instance. */
-	int32 EmitterIdx = INDEX_NONE;
+#if WITH_EDITOR
+	UE_DEPRECATED(5.4, "Please update your code to handle execution path for different emitter types")
+	NIAGARA_API void CalculateFixedBounds(const FTransform& ToWorldSpace);
+#endif
 
-	/* The age of the emitter*/
-	float EmitterAge = 0.0f;
+	UE_DEPRECATED(5.4, "Please update your code to handle execution path for different emitter types")
+	NIAGARA_API float GetTotalCPUTimeMS() const;
 
-	int32 RandomSeed = 0;
-	int32 InstanceSeed = FGenericPlatformMath::Rand();
-	int32 TickCount = 0;
+	UE_DEPRECATED(5.4, "Please update your code to handle execution path for different emitter types")
+	NIAGARA_API FName GetCachedIDName() const;
 
-	int32 TotalSpawnedParticles = 0;
-	
-	/* Cycles taken to process the tick. */
-	uint32 CPUTimeCycles = 0;
+	UE_DEPRECATED(5.4, "Please update your code to use GetParticleData")
+	FNiagaraDataSet& GetData() { check(ParticleDataSet); return *ParticleDataSet; }
+	UE_DEPRECATED(5.4, "Please update your code to use GetParticleData")
+	const FNiagaraDataSet& GetData() const { check(ParticleDataSet); return *ParticleDataSet; }
+	////////////////////////////////////////////////////////////////////////////////////////////////////
 
-	uint32 MaxRuntimeAllocation = 0;
+protected:
+	FNiagaraSystemInstance*				ParentSystemInstance = nullptr;
+	FNiagaraDataSet*					ParticleDataSet = nullptr;
+	FNiagaraComputeExecutionContext*	GPUExecContext = nullptr;
+	FVersionedNiagaraEmitter			VersionedEmitter = {};
 
-	int32 MaxAllocationCount = 0;
-	int32 MinOverallocation = -1;
-	int32 ReallocationCount = 0;
+	FNiagaraParameterStore				RendererBindings;
 
-	uint32 MaxInstanceCount = 0;
+	int32								EmitterIndex = INDEX_NONE;
+	int32								TotalSpawnedParticles = 0;
 
-	/* Emitter tick state */
-	ENiagaraExecutionState ExecutionState = ENiagaraExecutionState::Inactive;
+	uint8								bLocalSpace : 1 = true;
+	uint8								bDeterministic : 1 = true;
+	uint8								bNeedsPartialDepthTexture : 1 = false;
+	uint8								bNeedsEarlyViewUniformBuffer : 1 = false;
+	uint8								bAllowSpawning_GT : 1 = true;
+	uint8								bAllowSpawning_CNC : 1 = true;
+	ENiagaraSimTarget					SimTarget = ENiagaraSimTarget::CPUSim;
+	ENiagaraExecutionState				ExecutionState = ENiagaraExecutionState::Active;
 
-	/** Typical resets must be deferred until the tick as the RT could still be using the current buffer. */
-	uint32 bResetPending : 1;
+	bool								bCachedBoundsDynamic = false;
+	FBox								CachedBounds = FBox(ForceInit);
+	FBox								CachedSystemFixedBounds = FBox(ForceInit);
 
-	bool bAllowSpawning_GT = true;
-	bool bAllowSpawning_CNC = true;
+	mutable FRWLock						FixedBoundsGuard;
+	FBox								FixedBounds = FBox(ForceInit);
 
-	// This is used to keep track which particles have spawned a component. This is needed when the bOnlyCreateComponentsOnParticleSpawn flag is set in the renderer.
-	// Without this bookkeeping, the particles would lose their components when the render state is recreated or the visibility tag flips them off and on again.
-	mutable TMap<FObjectKey, TSet<int32>> ParticlesWithComponents;
+#if WITH_EDITORONLY_DATA
+	uint32								TickTimeCycles = 0;
+#endif
 };
+
+using FNiagaraEmitterInstanceRef = TSharedRef<FNiagaraEmitterInstance, ESPMode::ThreadSafe>;
+using FNiagaraEmitterInstancePtr = TSharedPtr<FNiagaraEmitterInstance, ESPMode::ThreadSafe>;
diff --git a/Engine/Plugins/FX/Niagara/Source/Niagara/Private/DataInterface/NiagaraDataInterfaceConsoleVariable.cpp b/Engine/Plugins/FX/Niagara/Source/Niagara/Private/DataInterface/NiagaraDataInterfaceConsoleVariable.cpp
--- a/Engine/Plugins/FX/Niagara/Source/Niagara/Private/DataInterface/NiagaraDataInterfaceConsoleVariable.cpp
+++ b/Engine/Plugins/FX/Niagara/Source/Niagara/Private/DataInterface/NiagaraDataInterfaceConsoleVariable.cpp
@@ -175,9 +175,9 @@ bool UNiagaraDataInterfaceConsoleVariable::InitPerInstanceData(void* PerInstance
 	if ( IsUsedWithGPUScript() )
 	{
 		// We shouldn't need to do this per init, we should be able to cache once and once only
-		for (const TSharedRef<FNiagaraEmitterInstance, ESPMode::ThreadSafe>& EmitterInstance : SystemInstance->GetEmitters())
+		for (const FNiagaraEmitterInstanceRef& EmitterInstance : SystemInstance->GetEmitters())
 		{
-			if (EmitterInstance->IsDisabled() || EmitterInstance->GetCachedEmitterData() == nullptr || EmitterInstance->GetGPUContext() == nullptr)
+			if (EmitterInstance->IsDisabled() || EmitterInstance->GetEmitter() == nullptr || EmitterInstance->GetSimTarget() != ENiagaraSimTarget::GPUComputeSim)
 			{
 				continue;
 			}
diff --git a/Engine/Plugins/FX/Niagara/Source/Niagara/Private/DataInterface/NiagaraDataInterfaceDataChannelRead.cpp b/Engine/Plugins/FX/Niagara/Source/Niagara/Private/DataInterface/NiagaraDataInterfaceDataChannelRead.cpp
--- a/Engine/Plugins/FX/Niagara/Source/Niagara/Private/DataInterface/NiagaraDataInterfaceDataChannelRead.cpp
+++ b/Engine/Plugins/FX/Niagara/Source/Niagara/Private/DataInterface/NiagaraDataInterfaceDataChannelRead.cpp
@@ -16,6 +16,8 @@
 #include "NiagaraDataChannelHandler.h"
 #include "NiagaraDataChannelManager.h"
 
+#include "NiagaraEmitterInstanceImpl.h"
+
 #include "NiagaraRenderer.h"
 #include "NiagaraGPUSystemTick.h"
 
@@ -577,7 +579,11 @@ void UNiagaraDataInterfaceDataChannelRead::PostStageTick(FNDICpuPostStageContext
 					{
 						SpawnInfo.SpawnGroup = i;
 					}
-					TargetEmitter->GetSpawnInfo().Emplace(SpawnInfo);
+					//-TODO:Stateless:
+					if (FNiagaraEmitterInstanceImpl* StatefulEmitter = TargetEmitter->AsStateful())
+					{
+						StatefulEmitter->GetSpawnInfo().Emplace(SpawnInfo);
+					}
 				}
 			}
 		}
diff --git a/Engine/Plugins/FX/Niagara/Source/Niagara/Private/DataInterface/NiagaraDataInterfaceEmitterProperties.cpp b/Engine/Plugins/FX/Niagara/Source/Niagara/Private/DataInterface/NiagaraDataInterfaceEmitterProperties.cpp
--- a/Engine/Plugins/FX/Niagara/Source/Niagara/Private/DataInterface/NiagaraDataInterfaceEmitterProperties.cpp
+++ b/Engine/Plugins/FX/Niagara/Source/Niagara/Private/DataInterface/NiagaraDataInterfaceEmitterProperties.cpp
@@ -123,9 +123,9 @@ bool UNiagaraDataInterfaceEmitterProperties::InitPerInstanceData(void* PerInstan
 
 	FInstanceData_GameThread* InstanceData_GT = new(PerInstanceData) FInstanceData_GameThread();
 	InstanceData_GT->EmitterInstance = EmitterBinding.Resolve(SystemInstance, this);
-	if (InstanceData_GT->EmitterInstance && InstanceData_GT->EmitterInstance->GetCachedEmitterData())
+	if (InstanceData_GT->EmitterInstance)
 	{
-		InstanceData_GT->bLocalSpace = InstanceData_GT->EmitterInstance->GetCachedEmitterData()->bLocalSpace;
+		InstanceData_GT->bLocalSpace = InstanceData_GT->EmitterInstance->IsLocalSpace();
 	}
 
 	if ( IsUsedWithGPUScript() )
diff --git a/Engine/Plugins/FX/Niagara/Source/Niagara/Private/DataInterface/NiagaraDataInterfaceUObjectPropertyReader.cpp b/Engine/Plugins/FX/Niagara/Source/Niagara/Private/DataInterface/NiagaraDataInterfaceUObjectPropertyReader.cpp
--- a/Engine/Plugins/FX/Niagara/Source/Niagara/Private/DataInterface/NiagaraDataInterfaceUObjectPropertyReader.cpp
+++ b/Engine/Plugins/FX/Niagara/Source/Niagara/Private/DataInterface/NiagaraDataInterfaceUObjectPropertyReader.cpp
@@ -690,9 +690,9 @@ bool UNiagaraDataInterfaceUObjectPropertyReader::InitPerInstanceData(void* PerIn
 		TArray<uint32> GpuFunctionToPropertyRemap;
 
 		// We shouldn't need to do this per init, we should be able to cache once and once only
-		for (const TSharedRef<FNiagaraEmitterInstance, ESPMode::ThreadSafe>& EmitterInstance : SystemInstance->GetEmitters())
+		for (const FNiagaraEmitterInstanceRef& EmitterInstance : SystemInstance->GetEmitters())
 		{
-			if (EmitterInstance->IsDisabled() || EmitterInstance->GetCachedEmitterData() == nullptr || EmitterInstance->GetGPUContext() == nullptr)
+			if (EmitterInstance->IsDisabled() || EmitterInstance->GetEmitter() == nullptr || EmitterInstance->GetSimTarget() != ENiagaraSimTarget::GPUComputeSim)
 			{
 				continue;
 			}
diff --git a/Engine/Plugins/FX/Niagara/Source/Niagara/Private/DataInterface/NiagaraDataInterfaceUtilities.cpp b/Engine/Plugins/FX/Niagara/Source/Niagara/Private/DataInterface/NiagaraDataInterfaceUtilities.cpp
--- a/Engine/Plugins/FX/Niagara/Source/Niagara/Private/DataInterface/NiagaraDataInterfaceUtilities.cpp
+++ b/Engine/Plugins/FX/Niagara/Source/Niagara/Private/DataInterface/NiagaraDataInterfaceUtilities.cpp
@@ -6,6 +6,7 @@
 #include "NiagaraConstants.h"
 #include "NiagaraDataInterface.h"
 #include "NiagaraEmitterInstance.h"
+#include "NiagaraEmitterInstanceImpl.h"
 #include "NiagaraSystemImpl.h"
 #include "NiagaraSystemInstance.h"
 #include "NiagaraSystemSimulation.h"
@@ -243,22 +244,23 @@ void FNiagaraDataInterfaceUtilities::ForEachDataInterface(FNiagaraSystemInstance
 	ForEachDataInterface(SystemSimulation->GetSpawnExecutionContext()->Parameters, Action, SearchOptions);
 	ForEachDataInterface(SystemSimulation->GetUpdateExecutionContext()->Parameters, Action, SearchOptions);
 
-	for (const TSharedRef<FNiagaraEmitterInstance, ESPMode::ThreadSafe>& Emitter : SystemInstance->GetEmitters())
+	for (const FNiagaraEmitterInstanceRef& Emitter : SystemInstance->GetEmitters())
 	{
-		if ( Emitter->IsDisabled() )
+		//-TODO:Stateless:
+		FNiagaraEmitterInstanceImpl* StatefulEmitter = Emitter->AsStateful();
+		if (!StatefulEmitter || Emitter->IsDisabled())
 		{
 			continue;
 		}
 
-		ForEachDataInterface(Emitter->GetSpawnExecutionContext().Parameters, Action, SearchOptions);
-		ForEachDataInterface(Emitter->GetUpdateExecutionContext().Parameters, Action, SearchOptions);
-		for (int32 i=0; i < Emitter->GetEventExecutionContexts().Num(); i++)
+		ForEachDataInterface(StatefulEmitter->GetSpawnExecutionContext().Parameters, Action, SearchOptions);
+		ForEachDataInterface(StatefulEmitter->GetUpdateExecutionContext().Parameters, Action, SearchOptions);
+		for (const FNiagaraScriptExecutionContext& EventExecContext : StatefulEmitter->GetEventExecutionContexts())
 		{
-			ForEachDataInterface(Emitter->GetEventExecutionContexts()[i].Parameters, Action, SearchOptions);
+			ForEachDataInterface(EventExecContext.Parameters, Action, SearchOptions);
 		}
 
-		FVersionedNiagaraEmitterData* EmitterData = Emitter->GetCachedEmitterData();
-		if (EmitterData && EmitterData->SimTarget == ENiagaraSimTarget::GPUComputeSim && Emitter->GetGPUContext())
+		if (Emitter->GetEmitter() && Emitter->GetSimTarget() == ENiagaraSimTarget::GPUComputeSim && Emitter->GetGPUContext())
 		{
 			ForEachDataInterface(Emitter->GetGPUContext()->CombinedParamStore, Action, SearchOptions);
 		}
@@ -285,23 +287,24 @@ void FNiagaraDataInterfaceUtilities::ForEachDataInterface(FNiagaraSystemInstance
 	ForEachDataInterface(UsageContext, SystemSimulation->GetSpawnExecutionContext()->Parameters, Action, SearchOptions);
 	ForEachDataInterface(UsageContext, SystemSimulation->GetUpdateExecutionContext()->Parameters, Action, SearchOptions);
 
-	for (const TSharedRef<FNiagaraEmitterInstance, ESPMode::ThreadSafe>& Emitter : SystemInstance->GetEmitters())
+	for (const FNiagaraEmitterInstanceRef& Emitter : SystemInstance->GetEmitters())
 	{
-		if (Emitter->IsDisabled())
+		//-TODO:Stateless:
+		FNiagaraEmitterInstanceImpl* StatefulEmitter = Emitter->AsStateful();
+		if (!StatefulEmitter || Emitter->IsDisabled())
 		{
 			continue;
 		}
-		UsageContext.OwnerObject = Emitter->GetCachedEmitter().Emitter;
+		UsageContext.OwnerObject = Emitter->GetEmitter();
 
-		ForEachDataInterface(UsageContext, Emitter->GetSpawnExecutionContext().Parameters, Action, SearchOptions);
-		ForEachDataInterface(UsageContext, Emitter->GetUpdateExecutionContext().Parameters, Action, SearchOptions);
-		for (int32 i = 0; i < Emitter->GetEventExecutionContexts().Num(); i++)
+		ForEachDataInterface(UsageContext, StatefulEmitter->GetSpawnExecutionContext().Parameters, Action, SearchOptions);
+		ForEachDataInterface(UsageContext, StatefulEmitter->GetUpdateExecutionContext().Parameters, Action, SearchOptions);
+		for (const FNiagaraScriptExecutionContext& EventExecContext : StatefulEmitter->GetEventExecutionContexts() )
 		{
-			ForEachDataInterface(UsageContext, Emitter->GetEventExecutionContexts()[i].Parameters, Action, SearchOptions);
+			ForEachDataInterface(UsageContext, EventExecContext.Parameters, Action, SearchOptions);
 		}
 
-		FVersionedNiagaraEmitterData* EmitterData = Emitter->GetCachedEmitterData();
-		if (EmitterData && EmitterData->SimTarget == ENiagaraSimTarget::GPUComputeSim && Emitter->GetGPUContext())
+		if (Emitter->GetEmitter() && Emitter->GetSimTarget() == ENiagaraSimTarget::GPUComputeSim && Emitter->GetGPUContext())
 		{
 			ForEachDataInterface(UsageContext, Emitter->GetGPUContext()->CombinedParamStore, Action, SearchOptions);
 		}
diff --git a/Engine/Plugins/FX/Niagara/Source/Niagara/Private/NiagaraComponentSettings.cpp b/Engine/Plugins/FX/Niagara/Source/Niagara/Private/NiagaraComponentSettings.cpp
--- a/Engine/Plugins/FX/Niagara/Source/Niagara/Private/NiagaraComponentSettings.cpp
+++ b/Engine/Plugins/FX/Niagara/Source/Niagara/Private/NiagaraComponentSettings.cpp
@@ -1,6 +1,7 @@
 // Copyright Epic Games, Inc. All Rights Reserved.
 
 #include "NiagaraComponentSettings.h"
+#include "NiagaraComputeExecutionContext.h"
 #include "NiagaraEmitter.h"
 #include "NiagaraEmitterInstance.h"
 #include "NiagaraScript.h"
diff --git a/Engine/Plugins/FX/Niagara/Source/Niagara/Private/NiagaraDataInterfaceEmitterBinding.cpp b/Engine/Plugins/FX/Niagara/Source/Niagara/Private/NiagaraDataInterfaceEmitterBinding.cpp
--- a/Engine/Plugins/FX/Niagara/Source/Niagara/Private/NiagaraDataInterfaceEmitterBinding.cpp
+++ b/Engine/Plugins/FX/Niagara/Source/Niagara/Private/NiagaraDataInterfaceEmitterBinding.cpp
@@ -15,11 +15,11 @@ FNiagaraEmitterInstance* FNiagaraDataInterfaceEmitterBinding::Resolve(const FNia
 		// If we came from a particle script our outer will be a UNiagaraEmitter
 		if (UNiagaraEmitter* OwnerEmitter = DataInterface->GetTypedOuter<UNiagaraEmitter>())
 		{
-			for (TSharedPtr<FNiagaraEmitterInstance, ESPMode::ThreadSafe> EmitterInstance : SystemInstance->GetEmitters())
+			for (const FNiagaraEmitterInstanceRef& EmitterInstance : SystemInstance->GetEmitters())
 			{
-				if (EmitterInstance->GetCachedEmitter().Emitter == OwnerEmitter)
+				if (EmitterInstance->GetEmitter() == OwnerEmitter)
 				{
-					return EmitterInstance.Get();
+					return &EmitterInstance.Get();
 				}
 			}
 		}
@@ -44,13 +44,13 @@ FNiagaraEmitterInstance* FNiagaraDataInterfaceEmitterBinding::Resolve(const FNia
 
 			if (SourceEmitterName.IsEmpty() == false)
 			{
-				for (TSharedPtr<FNiagaraEmitterInstance, ESPMode::ThreadSafe> EmitterInstance : SystemInstance->GetEmitters())
+				for (const FNiagaraEmitterInstanceRef& EmitterInstance : SystemInstance->GetEmitters())
 				{
-					if (const UNiagaraEmitter* CachedEmitter = EmitterInstance->GetCachedEmitter().Emitter)
+					if (const UNiagaraEmitter* CachedEmitter = EmitterInstance->GetEmitter())
 					{
 						if (CachedEmitter->GetUniqueEmitterName() == SourceEmitterName)
 						{
-							return EmitterInstance.Get();
+							return &EmitterInstance.Get();
 						}
 					}
 				}
@@ -76,14 +76,14 @@ FNiagaraEmitterInstance* FNiagaraDataInterfaceEmitterBinding::Resolve(const FNia
 		EmitterName.ToString(EmitterNameString);
 		FStringView EmitterNameStringView = EmitterNameString.ToView();
 
-		for (TSharedPtr<FNiagaraEmitterInstance, ESPMode::ThreadSafe> EmitterInstance : SystemInstance->GetEmitters())
+		for (const FNiagaraEmitterInstanceRef& EmitterInstance : SystemInstance->GetEmitters())
 		{
-			if (const UNiagaraEmitter* CachedEmitter = EmitterInstance->GetCachedEmitter().Emitter)
+			if (const UNiagaraEmitter* CachedEmitter = EmitterInstance->GetEmitter())
 			{
 				//-TODO: UniqueEmitterName should probably be a FName?
 				if (EmitterNameStringView.Equals(CachedEmitter->GetUniqueEmitterName()) )
 				{
-					return EmitterInstance.Get();
+					return &EmitterInstance.Get();
 				}
 			}
 		}
diff --git a/Engine/Plugins/FX/Niagara/Source/Niagara/Private/NiagaraDataInterfaceGrid2DCollectionReader.cpp b/Engine/Plugins/FX/Niagara/Source/Niagara/Private/NiagaraDataInterfaceGrid2DCollectionReader.cpp
--- a/Engine/Plugins/FX/Niagara/Source/Niagara/Private/NiagaraDataInterfaceGrid2DCollectionReader.cpp
+++ b/Engine/Plugins/FX/Niagara/Source/Niagara/Private/NiagaraDataInterfaceGrid2DCollectionReader.cpp
@@ -53,13 +53,10 @@ bool UNiagaraDataInterfaceGrid2DCollectionReader::InitPerInstanceData(void* PerI
 	FGrid2DCollectionRWInstanceData_GameThread* InstanceData = new (PerInstanceData) FGrid2DCollectionRWInstanceData_GameThread();
 	SystemInstancesToProxyData_GT.Emplace(SystemInstance->GetId(), InstanceData);
 
-	TSharedPtr<FNiagaraEmitterInstance, ESPMode::ThreadSafe> RT_EmitterInstance;
-
-
 	FNiagaraEmitterInstance* EmitterInstanceToUse = nullptr;
-	for (TSharedPtr<FNiagaraEmitterInstance, ESPMode::ThreadSafe> EmitterInstance : SystemInstance->GetEmitters())
+	for (const FNiagaraEmitterInstanceRef& EmitterInstance : SystemInstance->GetEmitters())
 	{
-		UNiagaraEmitter* Emitter = EmitterInstance->GetCachedEmitter().Emitter;
+		const UNiagaraEmitter* Emitter = EmitterInstance->GetEmitter();
 		if (Emitter == nullptr)
 		{
 			continue;
@@ -67,7 +64,7 @@ bool UNiagaraDataInterfaceGrid2DCollectionReader::InitPerInstanceData(void* PerI
 
 		if (EmitterName == Emitter->GetUniqueEmitterName())
 		{
-			EmitterInstanceToUse = EmitterInstance.Get();
+			EmitterInstanceToUse = &EmitterInstance.Get();
 			break;
 		}
 	}
diff --git a/Engine/Plugins/FX/Niagara/Source/Niagara/Private/NiagaraDataInterfaceGrid3DCollectionReader.cpp b/Engine/Plugins/FX/Niagara/Source/Niagara/Private/NiagaraDataInterfaceGrid3DCollectionReader.cpp
--- a/Engine/Plugins/FX/Niagara/Source/Niagara/Private/NiagaraDataInterfaceGrid3DCollectionReader.cpp
+++ b/Engine/Plugins/FX/Niagara/Source/Niagara/Private/NiagaraDataInterfaceGrid3DCollectionReader.cpp
@@ -52,13 +52,10 @@ bool UNiagaraDataInterfaceGrid3DCollectionReader::InitPerInstanceData(void* PerI
 	FGrid3DCollectionRWInstanceData_GameThread* InstanceData = new (PerInstanceData) FGrid3DCollectionRWInstanceData_GameThread();
 	SystemInstancesToProxyData_GT.Emplace(SystemInstance->GetId(), InstanceData);
 
-	TSharedPtr<FNiagaraEmitterInstance, ESPMode::ThreadSafe> RT_EmitterInstance;
-
-	
 	FNiagaraEmitterInstance* EmitterInstanceToUse = nullptr;
-	for (TSharedPtr<FNiagaraEmitterInstance, ESPMode::ThreadSafe> EmitterInstance : SystemInstance->GetEmitters())
+	for (const FNiagaraEmitterInstanceRef& EmitterInstance : SystemInstance->GetEmitters())
 	{
-		UNiagaraEmitter* Emitter = EmitterInstance->GetCachedEmitter().Emitter;
+		const UNiagaraEmitter* Emitter = EmitterInstance->GetEmitter();
 		if (Emitter == nullptr)
 		{
 			continue;
@@ -66,7 +63,7 @@ bool UNiagaraDataInterfaceGrid3DCollectionReader::InitPerInstanceData(void* PerI
 
 		if (EmitterName == Emitter->GetUniqueEmitterName())
 		{
-			EmitterInstanceToUse = EmitterInstance.Get();
+			EmitterInstanceToUse = &EmitterInstance.Get();
 			break;
 		}
 	}
diff --git a/Engine/Plugins/FX/Niagara/Source/Niagara/Private/NiagaraDataInterfaceParticleRead.cpp b/Engine/Plugins/FX/Niagara/Source/Niagara/Private/NiagaraDataInterfaceParticleRead.cpp
--- a/Engine/Plugins/FX/Niagara/Source/Niagara/Private/NiagaraDataInterfaceParticleRead.cpp
+++ b/Engine/Plugins/FX/Niagara/Source/Niagara/Private/NiagaraDataInterfaceParticleRead.cpp
@@ -443,12 +443,11 @@ bool UNiagaraDataInterfaceParticleRead::InitPerInstanceData(void* PerInstanceDat
 
 	if ( UNiagaraEmitter* OwnerEmitter = EmitterBinding.Resolve(this) )
 	{
-		for (TSharedPtr<FNiagaraEmitterInstance, ESPMode::ThreadSafe> EmitterInstance : SystemInstance->GetEmitters())
+		for (const FNiagaraEmitterInstanceRef& EmitterInstance : SystemInstance->GetEmitters())
 		{
-			const FVersionedNiagaraEmitter& CachedEmitter = EmitterInstance->GetCachedEmitter();
-			if (OwnerEmitter == CachedEmitter.Emitter)
+			if (OwnerEmitter == EmitterInstance->GetEmitter())
 			{
-				PIData->EmitterInstance = EmitterInstance.Get();
+				PIData->EmitterInstance = &EmitterInstance.Get();
 				break;
 			}
 		}
@@ -1037,7 +1036,7 @@ void UNiagaraDataInterfaceParticleRead::GetVMExternalFunction(const FVMExternalF
 	TArrayView<const FNiagaraVariableBase> EmitterVariables;
 	if (PIData->EmitterInstance)
 	{
-		EmitterVariables = PIData->EmitterInstance->GetData().GetVariables();
+		EmitterVariables = PIData->EmitterInstance->GetParticleData().GetVariables();
 	}
 
 	const FName AttributeToRead = FunctionSpecifier->Value;
@@ -1185,8 +1184,8 @@ void UNiagaraDataInterfaceParticleRead::GetVMExternalFunction(const FVMExternalF
 
 	if (!bBindSuccessful)
 	{
-		UNiagaraSystem* NiagaraSystem = PIData->SystemInstance ? PIData->SystemInstance->GetSystem() : nullptr;
-		UNiagaraEmitter* NiagaraEmitter = PIData->EmitterInstance ? PIData->EmitterInstance->GetCachedEmitter().Emitter : nullptr;
+		const UNiagaraSystem* NiagaraSystem = PIData->SystemInstance ? PIData->SystemInstance->GetSystem() : nullptr;
+		const UNiagaraEmitter* NiagaraEmitter = PIData->EmitterInstance ? PIData->EmitterInstance->GetEmitter() : nullptr;
 		UE_LOG(LogNiagara, Warning, TEXT("ParticleRead: Failed to '%s' attribute '%s' System '%s' Emitter '%s'! Check that the attribute is named correctly."), *BindingInfo.Name.ToString(), *AttributeToRead.ToString(), *GetNameSafe(NiagaraSystem), *EmitterBinding.ResolveUniqueName(this));
 	}
 }
@@ -1197,8 +1196,7 @@ void UNiagaraDataInterfaceParticleRead::VMGetLocalSpace(FVectorVMExternalFunctio
 	FNDIOutputParam<bool> OutIsLocalSpace(Context);
 
 	const FNiagaraEmitterInstance* EmitterInstance = InstData.Get()->EmitterInstance;
-	FVersionedNiagaraEmitterData* EmitterData = EmitterInstance ? EmitterInstance->GetCachedEmitterData() : nullptr;
-	const bool bIsLocalSpace = EmitterData ? EmitterData->bLocalSpace : false;
+	const bool bIsLocalSpace = EmitterInstance->IsLocalSpace();
 
 	for (int32 InstanceIdx = 0; InstanceIdx < Context.GetNumInstances(); ++InstanceIdx)
 	{
@@ -1212,7 +1210,7 @@ void UNiagaraDataInterfaceParticleRead::GetNumSpawnedParticles(FVectorVMExternal
 	VectorVM::FExternalFuncRegisterHandler<int32> OutNumSpawned(Context);
 
 	const FNiagaraEmitterInstance* EmitterInstance = InstData.Get()->EmitterInstance;
-	const FNiagaraDataBuffer* CurrentData = EmitterInstance ? EmitterInstance->GetData().GetCurrentData() : nullptr;
+	const FNiagaraDataBuffer* CurrentData = EmitterInstance ? EmitterInstance->GetParticleData().GetCurrentData() : nullptr;
 	const int32 NumSpawned = CurrentData ? CurrentData->GetNumSpawnedInstances() : 0;
 
 	for (int32 InstanceIdx = 0; InstanceIdx < Context.GetNumInstances(); ++InstanceIdx)
@@ -1229,7 +1227,7 @@ static const int32* GetSpawnedIDs(FVectorVMExternalFunctionContext& Context, FNi
 		return nullptr;
 	}
 
-	const FNiagaraDataSet& DataSet = EmitterInstance->GetData();
+	const FNiagaraDataSet& DataSet = EmitterInstance->GetParticleData();
 
 #if VECTORVM_SUPPORTS_EXPERIMENTAL && VECTORVM_SUPPORTS_LEGACY
 	if (Context.UsingExperimentalVM)
@@ -1283,7 +1281,7 @@ void UNiagaraDataInterfaceParticleRead::GetSpawnedIDAtIndex(FVectorVMExternalFun
 	int32 NumSpawned = 0;
 	const int32* SpawnedIDs = GetSpawnedIDs(Context, EmitterInstance, NumSpawned);
 
-	int32 IDAcquireTag = EmitterInstance ? EmitterInstance->GetData().GetIDAcquireTag() : 0;
+	int32 IDAcquireTag = EmitterInstance ? EmitterInstance->GetParticleData().GetIDAcquireTag() : 0;
 
 	for (int32 InstanceIdx = 0; InstanceIdx < Context.GetNumInstances(); ++InstanceIdx)
 	{
@@ -1316,7 +1314,7 @@ void UNiagaraDataInterfaceParticleRead::GetNumParticles(FVectorVMExternalFunctio
 	VectorVM::FExternalFuncRegisterHandler<int32> OutNumParticles (Context);
 
 	const FNiagaraEmitterInstance* EmitterInstance = InstData.Get()->EmitterInstance;
-	const FNiagaraDataBuffer* CurrentData = EmitterInstance ? EmitterInstance->GetData().GetCurrentData() : nullptr;
+	const FNiagaraDataBuffer* CurrentData = EmitterInstance ? EmitterInstance->GetParticleData().GetCurrentData() : nullptr;
 	const int32 NumParticles = CurrentData ? CurrentData->GetNumInstances() : 0;
 
 	for (int32 InstanceIdx = 0; InstanceIdx < Context.GetNumInstances(); ++InstanceIdx)
@@ -1333,7 +1331,7 @@ void UNiagaraDataInterfaceParticleRead::GetParticleIndex(FVectorVMExternalFuncti
 	VectorVM::FExternalFuncRegisterHandler<int32> OutIndex(Context);
 
 	const FNiagaraEmitterInstance* EmitterInstance = InstData.Get()->EmitterInstance;
-	const FNiagaraDataBuffer* CurrentData = EmitterInstance ? EmitterInstance->GetData().GetCurrentData() : nullptr;
+	const FNiagaraDataBuffer* CurrentData = EmitterInstance ? EmitterInstance->GetParticleData().GetCurrentData() : nullptr;
 
 	if (!CurrentData)
 	{
@@ -1344,7 +1342,7 @@ void UNiagaraDataInterfaceParticleRead::GetParticleIndex(FVectorVMExternalFuncti
 		return;
 	}
 
-	const auto IDData = FNiagaraDataSetAccessor<FNiagaraID>::CreateReader(EmitterInstance->GetData(), ParticleReadIDName);
+	const auto IDData = FNiagaraDataSetAccessor<FNiagaraID>::CreateReader(EmitterInstance->GetParticleData(), ParticleReadIDName);
 	const TArray<int32>& IDTable = CurrentData->GetIDTable();
 
 	for (int32 InstanceIdx = 0; InstanceIdx < Context.GetNumInstances(); ++InstanceIdx)
@@ -1377,7 +1375,7 @@ void UNiagaraDataInterfaceParticleRead::GetParticleIndexFromIDTable(FVectorVMExt
 	FNDIOutputParam<int32> OutParticleIndex(Context);
 
 	const FNiagaraEmitterInstance* EmitterInstance = InstData.Get()->EmitterInstance;
-	const FNiagaraDataBuffer* CurrentData = EmitterInstance ? EmitterInstance->GetData().GetCurrentData() : nullptr;
+	const FNiagaraDataBuffer* CurrentData = EmitterInstance ? EmitterInstance->GetParticleData().GetCurrentData() : nullptr;
 
 	if ( !CurrentData )
 	{
@@ -1443,7 +1441,7 @@ FORCEINLINE void ReadWithCheck(FVectorVMExternalFunctionContext& Context, FName
 	bool bWriteDummyData = true;
 	if (FNiagaraEmitterInstance* EmitterInstance = Params.GetEmitterInstance())
 	{
-		const FNiagaraDataBuffer* CurrentData = EmitterInstance->GetData().GetCurrentData();//TODO: We should really be grabbing this in the instance data tick and adding a read ref to it.
+		const FNiagaraDataBuffer* CurrentData = EmitterInstance->GetParticleData().GetCurrentData();//TODO: We should really be grabbing this in the instance data tick and adding a read ref to it.
 		if (CurrentData && EmitterInstance->GetGPUContext() == nullptr)
 		{
 			const TArray<int32>& IDTable = CurrentData->GetIDTable();
@@ -1451,8 +1449,8 @@ FORCEINLINE void ReadWithCheck(FVectorVMExternalFunctionContext& Context, FName
 
 			if (IDTable.Num() > 0)
 			{
-				const auto ValueData = FNiagaraDataSetAccessor<T>::CreateReader(EmitterInstance->GetData(), AttributeToRead);
-				const auto IDData = FNiagaraDataSetAccessor<FNiagaraID>::CreateReader(EmitterInstance->GetData(), ParticleReadIDName);
+				const auto ValueData = FNiagaraDataSetAccessor<T>::CreateReader(EmitterInstance->GetParticleData(), AttributeToRead);
+				const auto IDData = FNiagaraDataSetAccessor<FNiagaraID>::CreateReader(EmitterInstance->GetParticleData(), ParticleReadIDName);
 
 				if (IDData.IsValid() && ValueData.IsValid())
 				{
@@ -1583,12 +1581,12 @@ FORCEINLINE void ReadByIndexWithCheck(FVectorVMExternalFunctionContext& Context,
 	bool bWriteDummyData = true;
 	if (FNiagaraEmitterInstance* EmitterInstance = Params.GetEmitterInstance())
 	{
-		const FNiagaraDataBuffer* CurrentData = EmitterInstance->GetData().GetCurrentData();//TODO: We should really be grabbing these during instance data tick and adding a read ref. Releasing that on PostTick.
+		const FNiagaraDataBuffer* CurrentData = EmitterInstance->GetParticleData().GetCurrentData();//TODO: We should really be grabbing these during instance data tick and adding a read ref. Releasing that on PostTick.
 		if (CurrentData && CurrentData->GetNumInstances() > 0 && EmitterInstance->GetGPUContext() == nullptr)
 		{
 			int32 NumSourceInstances = (int32)CurrentData->GetNumInstances();
 
-			const auto ValueData = FNiagaraDataSetAccessor<T>::CreateReader(EmitterInstance->GetData(), AttributeToRead);
+			const auto ValueData = FNiagaraDataSetAccessor<T>::CreateReader(EmitterInstance->GetParticleData(), AttributeToRead);
 			if (ValueData.IsValid())
 			{
 				bWriteDummyData = false;
diff --git a/Engine/Plugins/FX/Niagara/Source/Niagara/Private/NiagaraDataSetReadback.cpp b/Engine/Plugins/FX/Niagara/Source/Niagara/Private/NiagaraDataSetReadback.cpp
--- a/Engine/Plugins/FX/Niagara/Source/Niagara/Private/NiagaraDataSetReadback.cpp
+++ b/Engine/Plugins/FX/Niagara/Source/Niagara/Private/NiagaraDataSetReadback.cpp
@@ -4,6 +4,7 @@
 #include "NiagaraComputeExecutionContext.h"
 #include "NiagaraEmitterHandle.h"
 #include "NiagaraEmitterInstance.h"
+#include "NiagaraEmitterInstanceImpl.h"
 #include "NiagaraSystemInstance.h"
 #include "NiagaraGpuComputeDispatchInterface.h"
 #include "NiagaraGPUInstanceCountManager.h"
@@ -25,7 +26,7 @@ void FNiagaraDataSetReadback::EnqueueReadback(FNiagaraEmitterInstance* EmitterIn
 	check(IsReady());
 
 	SourceName = EmitterInstance->GetEmitterHandle().GetName();
-	DataSet.Init(&EmitterInstance->GetData().GetCompiledData());
+	DataSet.Init(&EmitterInstance->GetParticleData().GetCompiledData());
 	if ( FNiagaraComputeExecutionContext* GPUExecContext = EmitterInstance->GetGPUContext() )
 	{
 		ParameterStore = GPUExecContext->CombinedParamStore;
@@ -43,16 +44,20 @@ void FNiagaraDataSetReadback::EnqueueReadback(FNiagaraEmitterInstance* EmitterIn
 	}
 	else
 	{
-		const FNiagaraDataSet& SourceDataSet = EmitterInstance->GetData();
+		const FNiagaraDataSet& SourceDataSet = EmitterInstance->GetParticleData();
 		if (FNiagaraDataBuffer* SourceDataBuffer = SourceDataSet.GetCurrentData())
 		{
-			EmitterInstance->GetData().CopyTo(DataSet, 0, SourceDataBuffer->GetNumInstances());
+			EmitterInstance->GetParticleData().CopyTo(DataSet, 0, SourceDataBuffer->GetNumInstances());
 		}
 		else
 		{
 			DataSet.CopyFromGPUReadback(nullptr, nullptr, nullptr, 0, 0, 0, 0, 0);
 		}
-		ParameterStore = EmitterInstance->GetUpdateExecutionContext().Parameters;
+		//-TODO:Stateless
+		if (FNiagaraEmitterInstanceImpl* StatefulEmitterInstance = EmitterInstance->AsStateful())
+		{
+			ParameterStore = StatefulEmitterInstance->GetUpdateExecutionContext().Parameters;
+		}
 	}
 }
 
diff --git a/Engine/Plugins/FX/Niagara/Source/Niagara/Private/NiagaraDebugHud.cpp b/Engine/Plugins/FX/Niagara/Source/Niagara/Private/NiagaraDebugHud.cpp
--- a/Engine/Plugins/FX/Niagara/Source/Niagara/Private/NiagaraDebugHud.cpp
+++ b/Engine/Plugins/FX/Niagara/Source/Niagara/Private/NiagaraDebugHud.cpp
@@ -9,6 +9,7 @@
 #include "NiagaraDataSetDebugAccessor.h"
 #include "NiagaraDataSetReadback.h"
 #include "NiagaraEmitterInstance.h"
+#include "NiagaraEmitterInstanceImpl.h"
 #include "NiagaraFunctionLibrary.h"
 #include "NiagaraGpuComputeDispatchInterface.h"
 #include "NiagaraMeshRendererProperties.h"
@@ -724,25 +725,24 @@ namespace NiagaraDebugLocal
 		{
 			TStringBuilder<128> GpuLatentBuilder;
 
-			for (const TSharedRef<FNiagaraEmitterInstance, ESPMode::ThreadSafe>& EmitterInstance : SystemInstance->GetEmitters())
+			for (const FNiagaraEmitterInstanceRef& EmitterInstance : SystemInstance->GetEmitters())
 			{
-				FVersionedNiagaraEmitterData* EmitterData = EmitterInstance->GetCachedEmitterData();
-				UNiagaraEmitter* NiagaraEmitter = EmitterInstance->GetCachedEmitter().Emitter;
-				if (EmitterData == nullptr || NiagaraEmitter == nullptr)
+				UNiagaraEmitter* NiagaraEmitter = EmitterInstance->GetEmitter();
+				if (NiagaraEmitter == nullptr)
 				{
 					continue;
 				}
 
 				bool bLowLatencyFailed = false;
-				EmitterData->ForEachEnabledRenderer(
-					[&](UNiagaraRendererProperties* RenderProperties)
+				EmitterInstance->ForEachEnabledRenderer(
+					[&](const UNiagaraRendererProperties* RenderProperties)
 					{
 						ENiagaraRendererGpuTranslucentLatency RequestedLatency = ENiagaraRendererGpuTranslucentLatency::ProjectDefault;
-						if (UNiagaraMeshRendererProperties* MeshRenderProperties = Cast<UNiagaraMeshRendererProperties>(RenderProperties))
+						if (const UNiagaraMeshRendererProperties* MeshRenderProperties = Cast<const UNiagaraMeshRendererProperties>(RenderProperties))
 						{
 							RequestedLatency = MeshRenderProperties->GpuTranslucentLatency;
 						}
-						else if (UNiagaraSpriteRendererProperties* SpriteRendererProperties = Cast<UNiagaraSpriteRendererProperties>(RenderProperties))
+						else if (const UNiagaraSpriteRendererProperties* SpriteRendererProperties = Cast<const UNiagaraSpriteRendererProperties>(RenderProperties))
 						{
 							RequestedLatency = SpriteRendererProperties->GpuTranslucentLatency;
 						}
@@ -1129,9 +1129,9 @@ void FNiagaraDebugHud::GatherSystemInfo()
 			if( NiagaraComponent)
 			{
 				check(SystemInstance);
-				for (const TSharedRef<FNiagaraEmitterInstance, ESPMode::ThreadSafe>& EmitterInstance : SystemInstance->GetEmitters())
+				for (const FNiagaraEmitterInstanceRef& EmitterInstance : SystemInstance->GetEmitters())
 				{
-					UNiagaraEmitter* NiagaraEmitter = EmitterInstance->GetCachedEmitter().Emitter;
+					const UNiagaraEmitter* NiagaraEmitter = EmitterInstance->GetEmitter();
 					if (NiagaraEmitter == nullptr)
 					{
 						continue;
@@ -1256,7 +1256,7 @@ const FNiagaraDataSet* FNiagaraDebugHud::GetParticleDataSet(FNiagaraSystemInstan
 		return nullptr;
 	}
 
-	return &EmitterInstance->GetData();
+	return &EmitterInstance->GetParticleData();
 }
 
 FNiagaraDebugHud::FValidationErrorInfo& FNiagaraDebugHud::GetValidationErrorInfo(UNiagaraComponent* NiagaraComponent)
@@ -2539,7 +2539,7 @@ void FNiagaraDebugHud::DrawValidation(class FNiagaraWorldManager* WorldManager,
 
 		if (Settings.bValidateParticleDataBuffers)
 		{
-			auto& EmitterHandles = SystemInstance->GetEmitters();
+			TConstArrayView<FNiagaraEmitterInstanceRef> EmitterHandles = SystemInstance->GetEmitters();
 			for ( int32 iEmitter=0; iEmitter < EmitterHandles.Num(); ++iEmitter)
 			{
 				FNiagaraEmitterInstance* EmitterInstance = &EmitterHandles[iEmitter].Get();
@@ -2566,7 +2566,7 @@ void FNiagaraDebugHud::DrawValidation(class FNiagaraWorldManager* WorldManager,
 					[&](const FNiagaraVariableBase& Variable, int32 InstanceIndex, int32 ComponentIndex)
 					{
 						auto& ValidationError = GetValidationErrorInfo(NiagaraComponent);
-						const FName EmitterName(*EmitterInstance->GetCachedEmitter().Emitter->GetUniqueEmitterName());
+						const FName EmitterName(*EmitterInstance->GetEmitter()->GetUniqueEmitterName());
 						ValidationError.ParticleVariablesWithErrors.FindOrAdd(EmitterName).AddUnique(Variable.GetName());
 					}
 				);
@@ -2761,7 +2761,7 @@ void FNiagaraDebugHud::DrawComponents(FNiagaraWorldManager* WorldManager, UCanva
 				FSceneView* SceneView = Canvas->SceneView;
 
 				const FTransform& SystemTransform = SystemInstance->GetWorldTransform();
-				const bool bParticlesLocalSpace = EmitterInstance->GetCachedEmitterData()->bLocalSpace;
+				const bool bParticlesLocalSpace = EmitterInstance->IsLocalSpace();
 				//const float ClipRadius = Settings.bUseParticleDisplayRadius ? 1.0f : 0.0f;
 				const float ParticleDisplayCenterRadiusSq = Settings.bUseParticleDisplayCenterRadius ? (Settings.ParticleDisplayCenterRadius * Settings.ParticleDisplayCenterRadius) : 0.0f;
 				const float ParticleDisplayClipNearPlane = Settings.bUseParticleDisplayClip ? float(FMath::Max(Settings.ParticleDisplayClip.X, 0.0)) : 0.0f;
@@ -2932,9 +2932,9 @@ void FNiagaraDebugHud::DrawComponents(FNiagaraWorldManager* WorldManager, UCanva
 					int32 ActiveEmitters = 0;
 					int32 TotalEmitters = 0;
 					int32 ActiveParticles = 0;
-					for (const TSharedRef<FNiagaraEmitterInstance, ESPMode::ThreadSafe>& EmitterInstance : SystemInstance->GetEmitters())
+					for (const FNiagaraEmitterInstanceRef& EmitterInstance : SystemInstance->GetEmitters())
 					{
-						UNiagaraEmitter* NiagaraEmitter = EmitterInstance->GetCachedEmitter().Emitter;
+						const UNiagaraEmitter* NiagaraEmitter = EmitterInstance->GetEmitter();
 						if (NiagaraEmitter == nullptr)
 						{
 							continue;
@@ -2998,14 +2998,16 @@ void FNiagaraDebugHud::DrawComponents(FNiagaraWorldManager* WorldManager, UCanva
 							for (const FParameterStoreVariable& ParameterStoreVariable : CachedVariables.ParameterStoreVariables)
 							{
 								FNiagaraParameterStore* ParameterStore = nullptr;
+								//-TODO:Stateless:
+								FNiagaraEmitterInstanceImpl* StatefulEmitter = SystemInstance->Emitters.IsValidIndex(ParameterStoreVariable.EmitterIndex) ? SystemInstance->Emitters[ParameterStoreVariable.EmitterIndex]->AsStateful() : nullptr;
 								switch (ParameterStoreVariable.StoreLocation)
 								{
 									case EParameterStoreLocation::UserOverride:			ParameterStore = SystemInstance->GetOverrideParameters(); break;
 									case EParameterStoreLocation::SystemSpawn:			ParameterStore = &SystemSimulation->GetSpawnExecutionContext()->Parameters; break;
 									case EParameterStoreLocation::SystemUpdate:			ParameterStore = &SystemSimulation->GetUpdateExecutionContext()->Parameters; break;
-									case EParameterStoreLocation::EmitterSpawn:			ParameterStore = &SystemInstance->Emitters[ParameterStoreVariable.EmitterIndex]->GetSpawnExecutionContext().Parameters; break;
-									case EParameterStoreLocation::EmitterUpdate:		ParameterStore = &SystemInstance->Emitters[ParameterStoreVariable.EmitterIndex]->GetUpdateExecutionContext().Parameters; break;
-									case EParameterStoreLocation::EmitterRenderBindings:ParameterStore = &SystemInstance->Emitters[ParameterStoreVariable.EmitterIndex]->GetRendererBoundVariables(); break;
+									case EParameterStoreLocation::EmitterSpawn:			ParameterStore = StatefulEmitter ? &StatefulEmitter->GetSpawnExecutionContext().Parameters : nullptr; break;
+									case EParameterStoreLocation::EmitterUpdate:		ParameterStore = StatefulEmitter ? &StatefulEmitter->GetUpdateExecutionContext().Parameters : nullptr; break;
+									case EParameterStoreLocation::EmitterRenderBindings:ParameterStore = StatefulEmitter ? &StatefulEmitter->GetRendererBoundVariables() : nullptr; break;
 								}
 								if (ParameterStore == nullptr)
 								{
@@ -3079,7 +3081,7 @@ void FNiagaraDebugHud::DrawComponents(FNiagaraWorldManager* WorldManager, UCanva
 									continue;
 								}
 
-								StringBuilder.Appendf(TEXT("Emitter (%s)\n"), *EmitterInstance->GetCachedEmitter().Emitter->GetUniqueEmitterName());
+								StringBuilder.Appendf(TEXT("Emitter (%s)\n"), *EmitterInstance->GetEmitter()->GetUniqueEmitterName());
 								const uint32 NumParticles = Settings.bUseMaxParticlesToDisplay ? FMath::Min((uint32)Settings.MaxParticlesToDisplay, DataBuffer->GetNumInstances()) : DataBuffer->GetNumInstances();
 								for (uint32 iInstance = 0; iInstance < NumParticles; ++iInstance)
 								{
