- FNiagaraEmitterInstance refactoring in a base class for different emitter types

diff --git a/Engine/Plugins/FX/Niagara/Source/Niagara/Private/NiagaraRendererComponents.cpp b/Engine/Plugins/FX/Niagara/Source/Niagara/Private/NiagaraRendererComponents.cpp
--- a/Engine/Plugins/FX/Niagara/Source/Niagara/Private/NiagaraRendererComponents.cpp
+++ b/Engine/Plugins/FX/Niagara/Source/Niagara/Private/NiagaraRendererComponents.cpp
@@ -5,6 +5,7 @@
 #include "Components/PrimitiveComponent.h"
 #include "NiagaraDataSet.h"
 #include "NiagaraEmitterInstance.h"
+#include "NiagaraEmitterInstanceImpl.h"
 #include "NiagaraStats.h"
 #include "NiagaraComponentRendererProperties.h"
 #include "NiagaraSystem.h"
@@ -387,7 +388,7 @@ void FNiagaraRendererComponents::PostSystemTick_GameThread(const UNiagaraRendere
 	}
 
 #if WITH_EDITORONLY_DATA
-	if (SystemInstance->GetIsolateEnabled() && !Emitter->GetEmitterHandle().IsIsolated())
+	if (Emitter->IsDisabledFromIsolation())
 	{
 		ResetComponentPool(true);
 		return;
@@ -401,8 +402,16 @@ void FNiagaraRendererComponents::PostSystemTick_GameThread(const UNiagaraRendere
 		return;
 	}
 
+
+	//-TODO:Stateless:
+	const FNiagaraEmitterInstanceImpl* StatefulEmitter = Emitter->AsStateful();
+	if (StatefulEmitter == nullptr)
+	{
+		return;
+	}
+
 	const double CurrentTime = AttachComponent->GetWorld()->GetRealTimeSeconds();
-	const FNiagaraDataSet& Data = Emitter->GetData();
+	const FNiagaraDataSet& Data = Emitter->GetParticleData();
 	const FNiagaraDataBuffer& ParticleData = Data.GetCurrentDataChecked();
 	FNiagaraDataSetReaderInt32<FNiagaraBool> EnabledAccessor = FNiagaraDataSetAccessor<FNiagaraBool>::CreateReader(Data, Properties->EnabledBinding.GetDataSetBindableVariable().GetName());
 	FNiagaraDataSetReaderInt32<int32> VisTagAccessor = FNiagaraDataSetAccessor<int32>::CreateReader(Data, Properties->RendererVisibilityTagBinding.GetDataSetBindableVariable().GetName());
@@ -509,7 +518,7 @@ void FNiagaraRendererComponents::PostSystemTick_GameThread(const UNiagaraRendere
 			if (PoolIndex == -1 && Properties->bCreateComponentFirstParticleFrame)
 			{
 				// Don't allow this particle to acquire a component unless it was just spawned or had a component assigned to it previously
-				bool bIsNewlySpawnedParticle = Emitter->IsParticleComponentActive(ComponentKey, ParticleID) || ParticleIndex >= ParticleData.GetNumInstances() - ParticleData.GetNumSpawnedInstances();
+				bool bIsNewlySpawnedParticle = StatefulEmitter->IsParticleComponentActive(ComponentKey, ParticleID) || ParticleIndex >= ParticleData.GetNumInstances() - ParticleData.GetNumSpawnedInstances();
 				if (!bIsNewlySpawnedParticle)
 				{
 					continue;
@@ -589,7 +598,7 @@ void FNiagaraRendererComponents::PostSystemTick_GameThread(const UNiagaraRendere
 		}
 
 		FComponentPoolEntry& PoolEntry = ComponentPool[PoolIndex];
-		FNiagaraLWCConverter LwcConverter = SystemInstance->GetLWCConverter(Emitter->GetCachedEmitter().GetEmitterData()->bLocalSpace);
+		FNiagaraLWCConverter LwcConverter = SystemInstance->GetLWCConverter(Emitter->IsLocalSpace());
 		TickPropertyBindings(Properties, SceneComponent, Data, ParticleIndex, PoolEntry, LwcConverter);
 
 		// Activate the component.
@@ -612,7 +621,7 @@ void FNiagaraRendererComponents::PostSystemTick_GameThread(const UNiagaraRendere
 		PoolEntry.LastActiveTime = CurrentTime;
 		if (Properties->bCreateComponentFirstParticleFrame)
 		{
-			Emitter->SetParticleComponentActive(ComponentKey, ParticleID);
+			StatefulEmitter->SetParticleComponentActive(ComponentKey, ParticleID);
 		}
 
 		++ComponentCount;
diff --git a/Engine/Plugins/FX/Niagara/Source/Niagara/Private/NiagaraRendererDecals.cpp b/Engine/Plugins/FX/Niagara/Source/Niagara/Private/NiagaraRendererDecals.cpp
--- a/Engine/Plugins/FX/Niagara/Source/Niagara/Private/NiagaraRendererDecals.cpp
+++ b/Engine/Plugins/FX/Niagara/Source/Niagara/Private/NiagaraRendererDecals.cpp
@@ -133,7 +133,7 @@ FNiagaraDynamicDataBase* FNiagaraRendererDecals::GenerateDynamicData(const FNiag
 	using namespace NiagaraRendererDecalsLocal;
 
 	// Get DataToRender
-	const FNiagaraDataSet& DataSet = Emitter->GetData();
+	const FNiagaraDataSet& DataSet = Emitter->GetParticleData();
 	const FNiagaraDataBuffer* DataToRender = DataSet.GetCurrentData();
 	FNiagaraSystemInstance* SystemInstance = Emitter->GetParentSystemInstance();
 	if (DataToRender == nullptr || SystemInstance == nullptr || !IsRendererEnabled(InProperties, Emitter))
diff --git a/Engine/Plugins/FX/Niagara/Source/Niagara/Private/NiagaraRendererLights.cpp b/Engine/Plugins/FX/Niagara/Source/Niagara/Private/NiagaraRendererLights.cpp
--- a/Engine/Plugins/FX/Niagara/Source/Niagara/Private/NiagaraRendererLights.cpp
+++ b/Engine/Plugins/FX/Niagara/Source/Niagara/Private/NiagaraRendererLights.cpp
@@ -77,7 +77,7 @@ FNiagaraDynamicDataBase* FNiagaraRendererLights::GenerateDynamicData(const FNiag
 
 	//Bail if we don't have the required attributes to render this emitter.
 	const UNiagaraLightRendererProperties* Properties = CastChecked<const UNiagaraLightRendererProperties>(InProperties);
-	const FNiagaraDataSet& Data = Emitter->GetData();
+	const FNiagaraDataSet& Data = Emitter->GetParticleData();
 	const FNiagaraDataBuffer* DataToRender = Data.GetCurrentData();
 	if (DataToRender == nullptr || Emitter->GetParentSystemInstance() == nullptr)
 	{
diff --git a/Engine/Plugins/FX/Niagara/Source/Niagara/Private/NiagaraRendererMeshes.cpp b/Engine/Plugins/FX/Niagara/Source/Niagara/Private/NiagaraRendererMeshes.cpp
--- a/Engine/Plugins/FX/Niagara/Source/Niagara/Private/NiagaraRendererMeshes.cpp
+++ b/Engine/Plugins/FX/Niagara/Source/Niagara/Private/NiagaraRendererMeshes.cpp
@@ -146,7 +146,7 @@ FNiagaraRendererMeshes::FNiagaraRendererMeshes(ERHIFeatureLevel::Type FeatureLev
 		LockedAxis.Set(0.0f, 0.0f, 1.0f);
 	}
 
-	const FNiagaraDataSet& Data = Emitter->GetData();
+	const FNiagaraDataSet& Data = Emitter->GetParticleData();
 
 	int32 FloatOffset;
 	int32 HalfOffset;
@@ -484,7 +484,7 @@ bool FNiagaraRendererMeshes::CalculateMeshUsed(FParticleMeshRenderData& Particle
 		return true;
 	}
 
-	if (ParticleMeshRenderData.DynamicDataMesh->GetSimTarget() != ENiagaraSimTarget::CPUSim)
+	if (SimTarget != ENiagaraSimTarget::CPUSim)
 	{
 		return true;
 	}
@@ -2032,7 +2032,7 @@ FNiagaraDynamicDataBase* FNiagaraRendererMeshes::GenerateDynamicData(const FNiag
 	}
 
 	
-	FNiagaraDataBuffer* DataToRender = Emitter->GetData().GetCurrentData();
+	FNiagaraDataBuffer* DataToRender = Emitter->GetParticleData().GetCurrentData();
 	if (!DataToRender || 
 		Meshes.Num() == 0 ||
 		(SourceMode == ENiagaraRendererSourceDataMode::Particles && SimTarget != ENiagaraSimTarget::GPUComputeSim && DataToRender->GetNumInstances() == 0))
diff --git a/Engine/Plugins/FX/Niagara/Source/Niagara/Private/NiagaraRendererProperties.cpp b/Engine/Plugins/FX/Niagara/Source/Niagara/Private/NiagaraRendererProperties.cpp
--- a/Engine/Plugins/FX/Niagara/Source/Niagara/Private/NiagaraRendererProperties.cpp
+++ b/Engine/Plugins/FX/Niagara/Source/Niagara/Private/NiagaraRendererProperties.cpp
@@ -416,12 +416,12 @@ void UNiagaraRendererProperties::ChangeToPositionBinding(FNiagaraVariableAttribu
 
 bool UNiagaraRendererProperties::UpdateMaterialStaticParameters(const FNiagaraRendererMaterialParameters& MaterialParameters, UMaterialInstanceConstant* MIC)
 {
-	UNiagaraEmitter* NiagaraEmitter = GetTypedOuter<UNiagaraEmitter>();
 	UNiagaraSystem* NiagaraSystem = GetTypedOuter<UNiagaraSystem>();
-	if (NiagaraEmitter == nullptr || NiagaraSystem == nullptr)
+	if (NiagaraSystem == nullptr)
 	{
 		return false;
 	}
+	UNiagaraEmitter* NiagaraEmitter = GetTypedOuter<UNiagaraEmitter>();
 
 	TArray<FMaterialParameterInfo> AllStaticSwitchParameterInfos;
 	{
@@ -444,7 +444,10 @@ bool UNiagaraRendererProperties::UpdateMaterialStaticParameters(const FNiagaraRe
 					}
 
 					FNiagaraVariableBase ResolvedStaticVariable = StaticVariable;
-					ResolvedStaticVariable.ReplaceRootNamespace(NiagaraEmitter->GetUniqueEmitterName(), FNiagaraConstants::EmitterNamespaceString);
+					if (NiagaraEmitter)
+					{
+						ResolvedStaticVariable.ReplaceRootNamespace(NiagaraEmitter->GetUniqueEmitterName(), FNiagaraConstants::EmitterNamespaceString);
+					}
 					if (ResolvedStaticVariable.GetName() != ParameterBinding.StaticVariableName)
 					{
 						continue;
diff --git a/Engine/Plugins/FX/Niagara/Source/Niagara/Private/NiagaraRendererRibbons.cpp b/Engine/Plugins/FX/Niagara/Source/Niagara/Private/NiagaraRendererRibbons.cpp
--- a/Engine/Plugins/FX/Niagara/Source/Niagara/Private/NiagaraRendererRibbons.cpp
+++ b/Engine/Plugins/FX/Niagara/Source/Niagara/Private/NiagaraRendererRibbons.cpp
@@ -695,7 +695,7 @@ FNiagaraRendererRibbons::FNiagaraRendererRibbons(ERHIFeatureLevel::Type FeatureL
 	const UNiagaraRibbonRendererProperties* Properties = CastChecked<const UNiagaraRibbonRendererProperties>(InProps);
 
 	int32 IgnoredFloatOffset, IgnoredHalfOffset;
-	Emitter->GetData().GetVariableComponentOffsets(Properties->RibbonIdBinding.GetDataSetBindableVariable(), IgnoredFloatOffset, RibbonIDParamDataSetOffset, IgnoredHalfOffset);
+	Emitter->GetParticleData().GetVariableComponentOffsets(Properties->RibbonIdBinding.GetDataSetBindableVariable(), IgnoredFloatOffset, RibbonIDParamDataSetOffset, IgnoredHalfOffset);
 
 	// Check we actually have ribbon id if we claim we do
 	check(!GenerationConfig.HasRibbonIDs() || RibbonIDParamDataSetOffset != INDEX_NONE);
@@ -911,7 +911,7 @@ FNiagaraDynamicDataBase* FNiagaraRendererRibbons::GenerateDynamicData(const FNia
 			return nullptr;
 		}
 
-		FNiagaraDataBuffer* DataToRender = Emitter->GetData().GetCurrentData();		
+		FNiagaraDataBuffer* DataToRender = Emitter->GetParticleData().GetCurrentData();
 		if(SimTarget == ENiagaraSimTarget::GPUComputeSim || (DataToRender != nullptr && DataToRender->GetNumInstances() > 1))
 		{
 			DynamicData = new FNiagaraDynamicDataRibbon(Emitter);
@@ -936,12 +936,16 @@ FNiagaraDynamicDataBase* FNiagaraRendererRibbons::GenerateDynamicData(const FNia
 			
 			DynamicData->bUseGPUInit = bIsGPUSystem || bWantsGPUInit;
 			DynamicData->bIsGPUSystem = bIsGPUSystem;
-			DynamicData->MaxAllocationCount = Emitter->GetData().GetMaxAllocationCount();
-			DynamicData->MaxAllocatedCountEstimate = FMath::Min<uint32>(Emitter->GetCachedEmitterData()->GetMaxParticleCountEstimate(), DynamicData->MaxAllocationCount);
+			DynamicData->MaxAllocationCount = Emitter->GetParticleData().GetMaxAllocationCount();
+			DynamicData->MaxAllocatedCountEstimate = 0;
+			if (FVersionedNiagaraEmitterData* EmitterData = Emitter->GetVersionedEmitter().GetEmitterData())
+			{
+				DynamicData->MaxAllocatedCountEstimate = FMath::Min<uint32>(EmitterData->GetMaxParticleCountEstimate(), DynamicData->MaxAllocatedCountEstimate);
+			}
 			
 			if (!DynamicData->bUseGPUInit)
 			{
-				const FNiagaraGenerationInputDataCPUAccessors CPUData(Properties, Emitter->GetData());
+				const FNiagaraGenerationInputDataCPUAccessors CPUData(Properties, Emitter->GetParticleData());
 				
 				DynamicData->GenerationOutput = MakeShared<FNiagaraRibbonCPUGeneratedVertexData>();
 
diff --git a/Engine/Plugins/FX/Niagara/Source/Niagara/Private/NiagaraRendererSprites.cpp b/Engine/Plugins/FX/Niagara/Source/Niagara/Private/NiagaraRendererSprites.cpp
--- a/Engine/Plugins/FX/Niagara/Source/Niagara/Private/NiagaraRendererSprites.cpp
+++ b/Engine/Plugins/FX/Niagara/Source/Niagara/Private/NiagaraRendererSprites.cpp
@@ -131,7 +131,7 @@ FNiagaraRendererSprites::FNiagaraRendererSprites(ERHIFeatureLevel::Type FeatureL
 	else
 	{
 		int32 FloatOffset, HalfOffset;
-		const FNiagaraDataSet& Data = Emitter->GetData();
+		const FNiagaraDataSet& Data = Emitter->GetParticleData();
 		Data.GetVariableComponentOffsets(Properties->RendererVisibilityTagBinding.GetDataSetBindableVariable(), FloatOffset, RendererVisTagOffset, HalfOffset);
 		bVisTagInParamStore = false;
 		bEnableCulling |= RendererVisTagOffset != INDEX_NONE;
@@ -1206,7 +1206,7 @@ FNiagaraDynamicDataBase *FNiagaraRendererSprites::GenerateDynamicData(const FNia
 			}
 		}
 
-		FNiagaraDataBuffer* DataToRender = Emitter->GetData().GetCurrentData();
+		FNiagaraDataBuffer* DataToRender = Emitter->GetParticleData().GetCurrentData();
 		if(SimTarget == ENiagaraSimTarget::GPUComputeSim || (DataToRender != nullptr &&  (SourceMode == ENiagaraRendererSourceDataMode::Emitter || (SourceMode == ENiagaraRendererSourceDataMode::Particles && DataToRender->GetNumInstances() > 0))))
 		{
 			DynamicData = new FNiagaraDynamicDataSprites(Emitter);
diff --git a/Engine/Plugins/FX/Niagara/Source/Niagara/Private/NiagaraRendererVolumes.cpp b/Engine/Plugins/FX/Niagara/Source/Niagara/Private/NiagaraRendererVolumes.cpp
--- a/Engine/Plugins/FX/Niagara/Source/Niagara/Private/NiagaraRendererVolumes.cpp
+++ b/Engine/Plugins/FX/Niagara/Source/Niagara/Private/NiagaraRendererVolumes.cpp
@@ -154,7 +154,7 @@ FNiagaraDynamicDataBase* FNiagaraRendererVolumes::GenerateDynamicData(const FNia
 			return nullptr;
 		}
 
-		FNiagaraDataBuffer* DataToRender = Emitter->GetData().GetCurrentData();
+		FNiagaraDataBuffer* DataToRender = Emitter->GetParticleData().GetCurrentData();
 		if (!DataToRender || DataToRender->GetNumInstances() == 0)
 		{
 			return nullptr;
diff --git a/Engine/Plugins/FX/Niagara/Source/Niagara/Private/NiagaraSimCache.cpp b/Engine/Plugins/FX/Niagara/Source/Niagara/Private/NiagaraSimCache.cpp
--- a/Engine/Plugins/FX/Niagara/Source/Niagara/Private/NiagaraSimCache.cpp
+++ b/Engine/Plugins/FX/Niagara/Source/Niagara/Private/NiagaraSimCache.cpp
@@ -375,7 +375,7 @@ bool UNiagaraSimCache::BeginWrite(FNiagaraSimCacheCreateParameters InCreateParam
 		static const FName NAME_SystemExecutionState("System.ExecutionState");
 		CreateParameters.ExplicitCaptureAttributes.AddUnique(NAME_SystemExecutionState);
 
-		for (const TSharedRef<FNiagaraEmitterInstance, ESPMode::ThreadSafe>& EmitterInstance : Helper.SystemInstance->GetEmitters())
+		for (const FNiagaraEmitterInstanceRef& EmitterInstance : Helper.SystemInstance->GetEmitters())
 		{
 			const FNiagaraParameterStore& RendererBindings = EmitterInstance->GetRendererBoundVariables();
 			for ( const FNiagaraVariableWithOffset& Variable : RendererBindings.ReadParameterVariables() )
@@ -637,7 +637,7 @@ bool UNiagaraSimCache::WriteFrame(UNiagaraComponent* NiagaraComponent, FNiagaraS
 	{
 		FNiagaraSimCacheEmitterFrame& CacheEmitterFrame = CacheFrame.EmitterData[i];
 		FNiagaraEmitterInstance& EmitterInstance = Helper.SystemInstance->GetEmitters()[i].Get();
-		FNiagaraDataBuffer* EmitterCurrentData = EmitterInstance.GetData().GetCurrentData();
+		FNiagaraDataBuffer* EmitterCurrentData = EmitterInstance.GetParticleData().GetCurrentData();
 		if (EmitterInstance.IsComplete() || !EmitterCurrentData)
 		{
 			continue;
@@ -986,11 +986,11 @@ bool UNiagaraSimCache::ReadFrame(int32 FrameIndex, float FrameFraction, FNiagara
 		const FNiagaraSimCacheEmitterFrame& CacheEmitterFrameB = CacheFrames[NextFrameIndex].EmitterData[i];
 		if (CacheLayout.EmitterLayouts[i].SimTarget == ENiagaraSimTarget::GPUComputeSim)
 		{
-			Helper.ReadDataBufferGPU(FrameFraction, FrameDeltaSeconds, SimDeltaSeconds, RebaseTransform, EmitterInstance, CacheLayout.EmitterLayouts[i], CacheEmitterFrame.ParticleDataBuffers, CacheEmitterFrameB.ParticleDataBuffers, EmitterInstance.GetData(), PendingCommandsInFlight);
+			Helper.ReadDataBufferGPU(FrameFraction, FrameDeltaSeconds, SimDeltaSeconds, RebaseTransform, EmitterInstance, CacheLayout.EmitterLayouts[i], CacheEmitterFrame.ParticleDataBuffers, CacheEmitterFrameB.ParticleDataBuffers, EmitterInstance.GetParticleData(), PendingCommandsInFlight);
 		}
 		else
 		{
-			Helper.ReadDataBuffer(FrameFraction, FrameDeltaSeconds, SimDeltaSeconds, RebaseTransform, CacheLayout.EmitterLayouts[i], CacheEmitterFrame.ParticleDataBuffers, CacheEmitterFrameB.ParticleDataBuffers, EmitterInstance.GetData());
+			Helper.ReadDataBuffer(FrameFraction, FrameDeltaSeconds, SimDeltaSeconds, RebaseTransform, CacheLayout.EmitterLayouts[i], CacheEmitterFrame.ParticleDataBuffers, CacheEmitterFrameB.ParticleDataBuffers, EmitterInstance.GetParticleData());
 		}
 	}
 
diff --git a/Engine/Plugins/FX/Niagara/Source/Niagara/Private/NiagaraSystem.cpp b/Engine/Plugins/FX/Niagara/Source/Niagara/Private/NiagaraSystem.cpp
--- a/Engine/Plugins/FX/Niagara/Source/Niagara/Private/NiagaraSystem.cpp
+++ b/Engine/Plugins/FX/Niagara/Source/Niagara/Private/NiagaraSystem.cpp
@@ -30,6 +30,7 @@
 #include "NiagaraTrace.h"
 #include "NiagaraTypes.h"
 #include "NiagaraWorldManager.h"
+
 #include "Algo/RemoveIf.h"
 #include "Algo/StableSort.h"
 #include "HAL/LowLevelMemTracker.h"
@@ -385,7 +386,10 @@ TArray<INiagaraParameterDefinitionsSubscriber*> UNiagaraSystem::GetOwnedParamete
 	TArray<INiagaraParameterDefinitionsSubscriber*> OutSubscribers;
 	for (const FNiagaraEmitterHandle& EmitterHandle : GetEmitterHandles())
 	{
-		OutSubscribers.Add(EmitterHandle.GetInstance().Emitter);
+		if ( EmitterHandle.GetInstance().Emitter )
+		{
+			OutSubscribers.Add(EmitterHandle.GetInstance().Emitter);
+		}
 	}
 	return OutSubscribers;
 }
@@ -1958,6 +1962,11 @@ void UNiagaraSystem::ComputeRenderersDrawOrder()
 
 	for (const FNiagaraEmitterHandle& EmitterHandle : EmitterHandles)
 	{
+		if (!EmitterHandle.GetIsEnabled())
+		{
+			continue;
+		}
+
 		if (FVersionedNiagaraEmitterData* EmitterData = EmitterHandle.GetEmitterData())
 		{
 			EmitterData->ForEachEnabledRenderer(
@@ -2009,17 +2018,22 @@ void UNiagaraSystem::CacheFromCompiledData()
 	FNameBuilder ExecutionStateNameBuilder;
 	for (int32 i=0; i < EmitterHandles.Num(); ++i)
 	{
-		FNiagaraEmitterHandle& Handle = EmitterHandles[i];
-		FVersionedNiagaraEmitterData* EmitterData = Handle.GetEmitterData();
-		if (Handle.GetIsEnabled() && EmitterData)
+		FNiagaraEmitterHandle& EmitterHandle = EmitterHandles[i];
+		FNiagaraDataSetAccessor<ENiagaraExecutionState>& EmitterExecutionState = EmitterExecutionStateAccessors.AddDefaulted_GetRef();
+		if (!EmitterHandle.GetIsEnabled())
+		{
+			continue;
+		}
+
+		if (FVersionedNiagaraEmitterData* EmitterData = EmitterHandle.GetEmitterData())
 		{
 			// Cache system instance accessors
 			ExecutionStateNameBuilder.Reset();
-			ExecutionStateNameBuilder << Handle.GetInstance().Emitter->GetUniqueEmitterName();
+			ExecutionStateNameBuilder << EmitterHandle.GetInstance().Emitter->GetUniqueEmitterName();
 			ExecutionStateNameBuilder << TEXT(".ExecutionState");
 			const FName ExecutionStateName(ExecutionStateNameBuilder);
 
-			EmitterExecutionStateAccessors.AddDefaulted_GetRef().Init(SystemDataSet, ExecutionStateName);
+			EmitterExecutionState.Init(SystemDataSet, ExecutionStateName);
 
 			// Cache emitter data set accessors, for things like bounds, etc
 			const FNiagaraDataSetCompiledData* DataSetCompiledData = nullptr;
@@ -2037,10 +2051,10 @@ void UNiagaraSystem::CacheFromCompiledData()
 					MaxDeltaTime = MaxDeltaTime.IsSet() ? FMath::Min(MaxDeltaTime.GetValue(), NiagaraSettings->MaxDeltaTimePerTick) : NiagaraSettings->MaxDeltaTimePerTick;
 				}
 			}
-			Handle.GetInstance().Emitter->ConditionalPostLoad();
-			EmitterData->CacheFromCompiledData(DataSetCompiledData, *Handle.GetInstance().Emitter);
+			EmitterHandle.GetInstance().Emitter->ConditionalPostLoad();
+			EmitterData->CacheFromCompiledData(DataSetCompiledData, *EmitterHandle.GetInstance().Emitter);
 
-			PSOPrecacheEvents.Append(EmitterData->PrecacheComputePSOs(*Handle.GetInstance().Emitter));
+			PSOPrecacheEvents.Append(EmitterData->PrecacheComputePSOs(*EmitterHandle.GetInstance().Emitter));
 
 			// Allow data interfaces to cache static buffers
 			UNiagaraScript* NiagaraEmitterScripts[] =
@@ -2069,17 +2083,13 @@ void UNiagaraSystem::CacheFromCompiledData()
 						{
 							if (DataInterfaceInfo.bIsInternal == false)
 							{
-								DataInterfaceGpuUsage.Add(DataInterfaceInfo.ParameterStoreVariable.GetName());
+								DataInterfaceGpuUsage.Add(DataInterfaceInfo.ParameterStoreVariable.GetName()); 
 							}
 						}
 					}
 				}
 			}
 		}
-		else
-		{
-			EmitterExecutionStateAccessors.AddDefaulted();
-		}
 	}
 
 	// If we had any precache events to wait on add a task to ensure the system does not start until they are ready
@@ -2987,10 +2997,11 @@ void UNiagaraSystem::PrepareRapidIterationParametersForCompilation()
 	for (int32 i = 0; i < GetEmitterHandles().Num(); i++)
 	{
 		const FNiagaraEmitterHandle& Handle = GetEmitterHandle(i);
-		if (Handle.GetIsEnabled()) // Don't pull in the emitter if it isn't going to be used.
+		FVersionedNiagaraEmitterData* EmitterData = Handle.GetEmitterData();
+		if (Handle.GetIsEnabled() && EmitterData) // Don't pull in the emitter if it isn't going to be used.
 		{
 			TArray<UNiagaraScript*> EmitterScripts;
-			Handle.GetEmitterData()->GetScripts(EmitterScripts, false, true);
+			EmitterData->GetScripts(EmitterScripts, false, true);
 
 			for (int32 ScriptIdx = 0; ScriptIdx < EmitterScripts.Num(); ScriptIdx++)
 			{
diff --git a/Engine/Plugins/FX/Niagara/Source/Niagara/Private/NiagaraSystemInstance.cpp b/Engine/Plugins/FX/Niagara/Source/Niagara/Private/NiagaraSystemInstance.cpp
--- a/Engine/Plugins/FX/Niagara/Source/Niagara/Private/NiagaraSystemInstance.cpp
+++ b/Engine/Plugins/FX/Niagara/Source/Niagara/Private/NiagaraSystemInstance.cpp
@@ -7,6 +7,7 @@
 #include "NiagaraComputeExecutionContext.h"
 #include "NiagaraDataInterface.h"
 #include "NiagaraEmitterInstance.h"
+#include "NiagaraEmitterInstanceImpl.h"
 #include "NiagaraGPUSystemTick.h"
 #include "NiagaraStats.h"
 #include "NiagaraParameterCollection.h"
@@ -243,10 +244,15 @@ void FNiagaraSystemInstance::SetActualExecutionState(ENiagaraExecutionState InSt
 			// This may happen *before* we've successfully pulled data off of a
 			// simulation run. This means that we need to synchronize the execution
 			// states upon activation.
-			for (int32 EmitterIdx = 0; EmitterIdx < Emitters.Num(); ++EmitterIdx)
+			for (const FNiagaraEmitterInstanceRef& EmitterRef : Emitters)
 			{
-				FNiagaraEmitterInstance& EmitterInst = Emitters[EmitterIdx].Get();
-				EmitterInst.SetExecutionState(ENiagaraExecutionState::Active);
+				//-TODO:Stateless: Avoid Casting
+				FNiagaraEmitterInstanceImpl* Emitter = EmitterRef->AsStateful();
+				if (Emitter == nullptr)
+				{
+					continue;
+				}
+				Emitter->SetExecutionState(ENiagaraExecutionState::Active);
 			}
 		}
 	}
@@ -257,8 +263,14 @@ void FNiagaraSystemInstance::Dump()const
 	if (GetSystemSimulation())
 	{
 		GetSystemSimulation()->DumpInstance(this);
-		for (auto& Emitter : Emitters)
+		for (const FNiagaraEmitterInstanceRef& EmitterRef : Emitters)
 		{
+			//-TODO:Stateless: Avoid Casting
+			FNiagaraEmitterInstanceImpl* Emitter = EmitterRef->AsStateful();
+			if (Emitter == nullptr)
+			{
+				continue;
+			}
 			Emitter->Dump();
 		}
 	}
@@ -388,9 +400,15 @@ bool FNiagaraSystemInstance::QueryCaptureResults(const FGuid& RequestId, TArray<
 
 			if (bWaitForGPU)
 			{
-				for (const TSharedRef<FNiagaraEmitterInstance, ESPMode::ThreadSafe>& CachedEmitter : Emitters)
+				for (const FNiagaraEmitterInstanceRef& EmitterRef : Emitters)
 				{
-					CachedEmitter->WaitForDebugInfo();
+					//-TODO:Stateless: Avoid Casting
+					FNiagaraEmitterInstanceImpl* Emitter = EmitterRef->AsStateful();
+					if (Emitter == nullptr)
+					{
+						continue;
+					}
+					Emitter->WaitForDebugInfo();
 				}
 				return false;
 			}
@@ -655,9 +673,9 @@ void FNiagaraSystemInstance::Complete(bool bExternalCompletion)
 		SetActualExecutionState(ENiagaraExecutionState::Complete);
 		SetRequestedExecutionState(ENiagaraExecutionState::Complete);
 
-		for (const TSharedRef<FNiagaraEmitterInstance, ESPMode::ThreadSafe>& Simulation : Emitters)
+		for (const FNiagaraEmitterInstanceRef& Emitter : Emitters)
 		{
-			Simulation->HandleCompletion(true);
+			Emitter->HandleCompletion(true);
 		}
 	}
 	else
@@ -698,7 +716,7 @@ void FNiagaraSystemInstance::OnPooledReuse(UWorld& NewWorld)
 	FixedBounds_CNC.Init();
 	PreviousLocation.Reset();
 
-	for (auto&& Emitter : Emitters)
+	for (const FNiagaraEmitterInstanceRef& Emitter : Emitters)
 	{
 		Emitter->OnPooledReuse();
 	}
@@ -920,9 +938,9 @@ void FNiagaraSystemInstance::ResetInternal(bool bResetSimulations)
 		return;
 	}
 
-	for (const TSharedRef<FNiagaraEmitterInstance, ESPMode::ThreadSafe>& Simulation : Emitters)
+	for (const FNiagaraEmitterInstanceRef& Emitter : Emitters)
 	{
-		Simulation->ResetSimulation(bResetSimulations);
+		Emitter->ResetSimulation(bResetSimulations);
 	}
 
 #if WITH_EDITOR
@@ -976,15 +994,20 @@ bool FNiagaraSystemInstance::IsReadyToRun() const
 	}
 
 	// check emitters
-	bool bAllReadyToRun = true;
-	for (const TSharedRef<FNiagaraEmitterInstance, ESPMode::ThreadSafe>& Simulation : Emitters)
+	for (const FNiagaraEmitterInstanceRef& EmitterRef : Emitters)
 	{
-		if (!Simulation->IsReadyToRun())
+		//-TODO:Stateless: Avoid Casting
+		FNiagaraEmitterInstanceImpl* Emitter = EmitterRef->AsStateful();
+		if (Emitter == nullptr)
 		{
-			bAllReadyToRun = false;
+			continue;
+		}
+		if (!Emitter->IsReadyToRun())
+		{
+			return false;
 		}
 	}
-	return bAllReadyToRun;
+	return true;
 }
 
 void FNiagaraSystemInstance::ReInitInternal()
@@ -1140,7 +1163,7 @@ void FNiagaraSystemInstance::Cleanup()
 // 	OldInstance->GetParameterStore().Unbind(&InstanceParameters);
 // 	NewInstance->GetParameterStore().Bind(&InstanceParameters);
 //
-// 	for (TSharedRef<FNiagaraEmitterInstance, ESPMode::ThreadSafe> Simulation : Emitters)
+// 	for (FNiagaraEmitterInstanceRef& Simulation : Emitters)
 // 	{
 // 		Simulation->RebindParameterCollection(OldInstance, NewInstance);
 // 	}
@@ -1165,9 +1188,9 @@ void FNiagaraSystemInstance::BindParameters()
 		UE_LOG(LogNiagara, Warning, TEXT("OverrideParameters is null.  Component(%s) System(%s)"), *GetFullNameSafe(AttachComponent.Get()), *GetFullNameSafe(GetSystem()));
 	}
 
-	for (const TSharedRef<FNiagaraEmitterInstance, ESPMode::ThreadSafe>& Simulation : Emitters)
+	for (const FNiagaraEmitterInstanceRef& Emitter : Emitters)
 	{
-		Simulation->BindParameters(bAlreadyBound);
+		Emitter->BindParameters(bAlreadyBound);
 	}
 
 	bAlreadyBound = true;
@@ -1191,9 +1214,51 @@ void FNiagaraSystemInstance::UnbindParameters(bool bFromComplete)
 	}
 
 	bAlreadyBound = bFromComplete && bAlreadyBound;
-	for (const TSharedRef<FNiagaraEmitterInstance, ESPMode::ThreadSafe>& Simulation : Emitters)
+	for (const FNiagaraEmitterInstanceRef& Emitter : Emitters)
 	{
-		Simulation->UnbindParameters(bFromComplete);
+		Emitter->UnbindParameters(bFromComplete);
+	}
+}
+
+void FNiagaraSystemInstance::BindToParameterStore(FNiagaraParameterStore& ParameterStore)
+{
+	if (ParameterStore.IsEmpty())
+	{
+		return;
+	}
+
+	if (SystemSimulation.IsValid())
+	{
+		SystemSimulation->GetSpawnExecutionContext()->Parameters.Bind(&ParameterStore);
+		SystemSimulation->GetUpdateExecutionContext()->Parameters.Bind(&ParameterStore);
+	}
+
+	InstanceParameters.Bind(&ParameterStore);
+
+	if (OverrideParameters)
+	{
+		OverrideParameters->Bind(&ParameterStore);
+	}
+}
+
+void FNiagaraSystemInstance::UnbindFromParameterStore(FNiagaraParameterStore& ParameterStore)
+{
+	if (ParameterStore.IsEmpty())
+	{
+		return;
+	}
+
+	if (SystemSimulation.IsValid())
+	{
+		SystemSimulation->GetSpawnExecutionContext()->Parameters.Unbind(&ParameterStore);
+		SystemSimulation->GetUpdateExecutionContext()->Parameters.Unbind(&ParameterStore);
+	}
+
+	InstanceParameters.Unbind(&ParameterStore);
+
+	if (OverrideParameters)
+	{
+		OverrideParameters->Unbind(&ParameterStore);
 	}
 }
 
@@ -1230,10 +1295,9 @@ bool FNiagaraSystemInstance::RequiresGlobalDistanceField() const
 		return false;
 	}
 
-	for (const TSharedRef<FNiagaraEmitterInstance, ESPMode::ThreadSafe>& Emitter : Emitters)
+	for (const FNiagaraEmitterInstanceRef& Emitter : Emitters)
 	{
-		FNiagaraComputeExecutionContext* GPUContext = Emitter->GetGPUContext();
-		if (GPUContext)
+		if (FNiagaraComputeExecutionContext* GPUContext = Emitter->GetGPUContext())
 		{
 			for (UNiagaraDataInterface* DataInterface : GPUContext->CombinedParamStore.GetDataInterfaces())
 			{
@@ -1255,7 +1319,7 @@ bool FNiagaraSystemInstance::RequiresDepthBuffer() const
 		return false;
 	}
 
-	for (const TSharedRef<FNiagaraEmitterInstance, ESPMode::ThreadSafe>& Emitter : Emitters)
+	for (const FNiagaraEmitterInstanceRef& Emitter : Emitters)
 	{
 		FNiagaraComputeExecutionContext* GPUContext = Emitter->GetGPUContext();
 		if (GPUContext)
@@ -1280,7 +1344,7 @@ bool FNiagaraSystemInstance::RequiresEarlyViewData() const
 		return false;
 	}
 
-	for (const TSharedRef<FNiagaraEmitterInstance, ESPMode::ThreadSafe>& Emitter : Emitters)
+	for (const FNiagaraEmitterInstanceRef& Emitter : Emitters)
 	{
 		FNiagaraComputeExecutionContext* GPUContext = Emitter->GetGPUContext();
 		if (GPUContext)
@@ -1305,10 +1369,9 @@ bool FNiagaraSystemInstance::RequiresViewUniformBuffer() const
 		return false;
 	}
 
-	for (const TSharedRef<FNiagaraEmitterInstance, ESPMode::ThreadSafe>& EmitterHandle : Emitters)
+	for (const FNiagaraEmitterInstanceRef& EmitterInstance : Emitters)
 	{
-		FVersionedNiagaraEmitterData* EmitterData = EmitterHandle->GetCachedEmitterData();
-		if (EmitterData && EmitterHandle->GetGPUContext() && EmitterData->RequiresViewUniformBuffer())
+		if (EmitterInstance->GetGPUContext() && EmitterInstance->NeedsEarlyViewUniformBuffer())
 		{
 			return true;
 		}
@@ -1324,7 +1387,7 @@ bool FNiagaraSystemInstance::RequiresRayTracingScene() const
 		return false;
 	}
 
-	for (const TSharedRef<FNiagaraEmitterInstance, ESPMode::ThreadSafe>& Emitter : Emitters)
+	for (const FNiagaraEmitterInstanceRef& Emitter : Emitters)
 	{
 		FNiagaraComputeExecutionContext* GPUContext = Emitter->GetGPUContext();
 		if (GPUContext)
@@ -1466,31 +1529,36 @@ void FNiagaraSystemInstance::InitDataInterfaces()
 	CalcInstDataSize(SystemSimulation->GetUpdateExecutionContext()->Parameters, true, false, true);
 
 	//Iterate over interfaces to get size for table and clear their interface bindings.
-	for (const TSharedRef<FNiagaraEmitterInstance, ESPMode::ThreadSafe>& Simulation : Emitters)
+	for (const FNiagaraEmitterInstanceRef& EmitterRef : Emitters)
 	{
-		FNiagaraEmitterInstance& Sim = Simulation.Get();
-		if (Sim.IsDisabled())
+		//-TODO:Stateless: Avoid Casting
+		FNiagaraEmitterInstanceImpl* Emitter = EmitterRef->AsStateful();
+		if (Emitter == nullptr)
 		{
 			continue;
 		}
 
-		FVersionedNiagaraEmitterData* EmitterData = Sim.GetCachedEmitterData();
-		const bool bGPUSimulation = EmitterData && (EmitterData->SimTarget == ENiagaraSimTarget::GPUComputeSim);
+		if (Emitter->IsDisabled())
+		{
+			continue;
+		}
 
-		CalcInstDataSize(Sim.GetSpawnExecutionContext().Parameters, !bGPUSimulation, bGPUSimulation, false);
-		CalcInstDataSize(Sim.GetUpdateExecutionContext().Parameters, !bGPUSimulation, bGPUSimulation, false);
-		for (int32 i = 0; i < Sim.GetEventExecutionContexts().Num(); i++)
+		const bool bGPUSimulation = Emitter->GetSimTarget() == ENiagaraSimTarget::GPUComputeSim;
+
+		CalcInstDataSize(Emitter->GetSpawnExecutionContext().Parameters, !bGPUSimulation, bGPUSimulation, false);
+		CalcInstDataSize(Emitter->GetUpdateExecutionContext().Parameters, !bGPUSimulation, bGPUSimulation, false);
+		for (int32 i = 0; i < Emitter->GetEventExecutionContexts().Num(); i++)
 		{
-			CalcInstDataSize(Sim.GetEventExecutionContexts()[i].Parameters, !bGPUSimulation, bGPUSimulation, false);
+			CalcInstDataSize(Emitter->GetEventExecutionContexts()[i].Parameters, !bGPUSimulation, bGPUSimulation, false);
 		}
 
-		if (EmitterData && EmitterData->SimTarget == ENiagaraSimTarget::GPUComputeSim && Sim.GetGPUContext())
+		if (Emitter->GetGPUContext() && Emitter->GetSimTarget() == ENiagaraSimTarget::GPUComputeSim)
 		{
-			CalcInstDataSize(Sim.GetGPUContext()->CombinedParamStore, !bGPUSimulation, bGPUSimulation, false);
+			CalcInstDataSize(Emitter->GetGPUContext()->CombinedParamStore, !bGPUSimulation, bGPUSimulation, false);
 		}
 
 		//Also force a rebind while we're here.
-		Sim.DirtyDataInterfaces();
+		Emitter->DirtyDataInterfaces();
 	}
 
 	DataInterfaceInstanceData.SetNumUninitialized(InstanceDataSize);
@@ -1571,17 +1639,16 @@ void FNiagaraSystemInstance::InitDataInterfaces()
 
 	if (GetSystem()->NeedsGPUContextInitForDataInterfaces())
 	{
-		for (const TSharedRef<FNiagaraEmitterInstance, ESPMode::ThreadSafe>& Simulation : Emitters)
+		for (const FNiagaraEmitterInstanceRef& Emitter : Emitters)
 		{
-			FNiagaraEmitterInstance& Sim = Simulation.Get();
-			if (Sim.IsDisabled())
+			if (Emitter->IsDisabled())
 			{
 				continue;
 			}
 
-			if (Sim.GetCachedEmitterData()->SimTarget == ENiagaraSimTarget::GPUComputeSim && Sim.GetGPUContext())
+			if (Emitter->GetGPUContext() && Emitter->GetSimTarget() == ENiagaraSimTarget::GPUComputeSim)
 			{
-				Sim.GetGPUContext()->OptionalContexInit(this);
+				Emitter->GetGPUContext()->OptionalContexInit(this);
 			}
 		}
 	}
@@ -1590,9 +1657,15 @@ void FNiagaraSystemInstance::InitDataInterfaces()
 	//If needed, the system script execution will use these to perform per instance pre and post tick operations on our DIs.
 	SystemSpawnDIStageTickHandler.Init(SystemSimulation->GetSpawnExecutionContext()->Script, this);
 	SystemUpdateDIStageTickHandler.Init(SystemSimulation->GetUpdateExecutionContext()->Script, this);
-	for (const TSharedRef<FNiagaraEmitterInstance, ESPMode::ThreadSafe>& Sim : Emitters)
+	for (const FNiagaraEmitterInstanceRef& EmitterRef : Emitters)
 	{
-		Sim->InitDITickLists();
+		//-TODO:Stateless: Avoid Casting
+		FNiagaraEmitterInstanceImpl* Emitter = EmitterRef->AsStateful();
+		if (Emitter == nullptr)
+		{
+			continue;
+		}
+		Emitter->InitDITickLists();
 	}
 }
 
@@ -1617,23 +1690,28 @@ void FNiagaraSystemInstance::ResolveUserDataInterfaceBindings()
 		}
 	};
 	
-	for (const TSharedRef<FNiagaraEmitterInstance, ESPMode::ThreadSafe>& EmitterSimulation : Emitters)
+	for (const FNiagaraEmitterInstanceRef& EmitterRef : Emitters)
 	{
-		FNiagaraEmitterInstance& EmitterSim = EmitterSimulation.Get();
+		//-TODO:Stateless: Avoid Casting
+		FNiagaraEmitterInstanceImpl* Emitter = EmitterRef->AsStateful();
+		if (Emitter == nullptr)
+		{
+			continue;
+		}
 
-		if (EmitterSim.GetSpawnExecutionContext().Script != nullptr)
+		if (Emitter->GetSpawnExecutionContext().Script != nullptr)
 		{
-			ResolveUserDIs(EmitterSim.GetSpawnExecutionContext().Parameters, EmitterSim.GetSpawnExecutionContext().Script);
+			ResolveUserDIs(Emitter->GetSpawnExecutionContext().Parameters, Emitter->GetSpawnExecutionContext().Script);
 		}
-		if (EmitterSim.GetUpdateExecutionContext().Script != nullptr)
+		if (Emitter->GetUpdateExecutionContext().Script != nullptr)
 		{
-			ResolveUserDIs(EmitterSim.GetUpdateExecutionContext().Parameters, EmitterSim.GetUpdateExecutionContext().Script);
+			ResolveUserDIs(Emitter->GetUpdateExecutionContext().Parameters, Emitter->GetUpdateExecutionContext().Script);
 		}
-		if (EmitterSim.GetGPUContext() != nullptr && EmitterSim.GetGPUContext()->GPUScript != nullptr)
+		if (Emitter->GetGPUContext() != nullptr && Emitter->GetGPUContext()->GPUScript != nullptr)
 		{
-			ResolveUserDIs(EmitterSim.GetGPUContext()->CombinedParamStore, EmitterSim.GetGPUContext()->GPUScript);
+			ResolveUserDIs(Emitter->GetGPUContext()->CombinedParamStore, Emitter->GetGPUContext()->GPUScript);
 		}
-		for (FNiagaraScriptExecutionContext& EventContext : EmitterSim.GetEventExecutionContexts())
+		for (FNiagaraScriptExecutionContext& EventContext : Emitter->GetEventExecutionContexts())
 		{
 			if (EventContext.Script != nullptr)
 			{
@@ -1693,12 +1771,18 @@ void FNiagaraSystemInstance::TickDataInterfaces(float DeltaSeconds, bool bPostSi
 	if (bRebindVMFuncs)
 	{
 		// Dirty data interfaces for emitters
-		for (const TSharedRef<FNiagaraEmitterInstance, ESPMode::ThreadSafe>& Simulation : Emitters)
+		for (const FNiagaraEmitterInstanceRef& EmitterRef : Emitters)
 		{
-			FNiagaraEmitterInstance& Sim = Simulation.Get();
-			if (!Sim.IsDisabled())
+			//-TODO:Stateless: Avoid Casting
+			FNiagaraEmitterInstanceImpl* Emitter = EmitterRef->AsStateful();
+			if (Emitter == nullptr)
 			{
-				Sim.DirtyDataInterfaces();
+				continue;
+			}
+
+			if (!Emitter->IsDisabled())
+			{
+				Emitter->DirtyDataInterfaces();
 			}
 		}
 
@@ -1923,8 +2007,9 @@ void FNiagaraSystemInstance::TickInstanceParameters_GameThread(float DeltaSecond
 	{
 		auto& CurrentEmitterParameters = EditEmitterParameters(i);
 
-		const TSharedRef<FNiagaraEmitterInstance, ESPMode::ThreadSafe>& Emitter = Emitters[i];
-		if (Emitter->GetExecutionState() != ENiagaraExecutionState::Disabled)
+		//-TODO:Stateless: Avoid Casting
+		FNiagaraEmitterInstanceImpl* Emitter = Emitters[i]->AsStateful();
+		if (Emitter && Emitter->GetExecutionState() != ENiagaraExecutionState::Disabled)
 		{
 			int32 NumParticles = Emitter->GetNumParticles();
 			CurrentSystemParameters.NumParticles += NumParticles;
@@ -2067,15 +2152,19 @@ void FNiagaraSystemInstance::InitEmitters()
 		Emitters.Reserve(NumEmitters);
 		for (int32 EmitterIdx=0; EmitterIdx < NumEmitters; ++EmitterIdx)
 		{
-			//-TODO: We should not create emitter instances for disable emitters
-			const bool EmitterEnabled = EmitterHandles[EmitterIdx].GetIsEnabled();
+			const FNiagaraEmitterHandle& EmitterHandle = EmitterHandles[EmitterIdx];
+
+			//-TODO: We should not create emitter instances for disabled emitters
+			const bool EmitterEnabled = EmitterHandle.GetIsEnabled();
+
+			Emitters.Emplace(MakeShared<FNiagaraEmitterInstanceImpl, ESPMode::ThreadSafe>(this));
 
-			TSharedRef<FNiagaraEmitterInstance, ESPMode::ThreadSafe> Sim = MakeShared<FNiagaraEmitterInstance, ESPMode::ThreadSafe>(this);
+			FNiagaraEmitterInstanceRef EmitterInstance = Emitters.Last();
 
 			if (System->bFixedBounds)
 			{
 				// be sure to set the system bounds first so that we can bypass work in the initialization of the emitter
-				Sim->SetSystemFixedBoundsOverride(System->GetFixedBounds());
+				EmitterInstance->SetSystemFixedBoundsOverride(System->GetFixedBounds());
 			}
 			else if (EmitterEnabled)
 			{
@@ -2088,18 +2177,14 @@ void FNiagaraSystemInstance::InitEmitters()
 				}
 			}
 
-			Sim->Init(EmitterIdx, ID);
-			Emitters.Add(Sim);
+			EmitterInstance->Init(EmitterIdx);
 
 			if (EmitterEnabled)
 			{
 				// Only set bHasGPUEmitters if we allow compute shaders on the platform
 				if (bAllowComputeShaders)
 				{
-					if (const FVersionedNiagaraEmitterData* EmitterData = Sim->GetCachedEmitterData())
-					{
-						bHasGPUEmitters |= EmitterData->SimTarget == ENiagaraSimTarget::GPUComputeSim;
-					}
+					bHasGPUEmitters |= EmitterInstance->GetSimTarget() == ENiagaraSimTarget::GPUComputeSim;
 				}
 			}
 		}
@@ -2320,10 +2405,9 @@ void FNiagaraSystemInstance::WaitForConcurrentTickAndFinalize(bool bEnsureComple
 bool FNiagaraSystemInstance::HandleCompletion()
 {
 	bool bEmittersCompleteOrDisabled = true;
-	for (const TSharedRef<FNiagaraEmitterInstance, ESPMode::ThreadSafe>&it : Emitters)
+	for (const FNiagaraEmitterInstanceRef& Emitter : Emitters)
 	{
-		FNiagaraEmitterInstance& Inst = *it;
-		bEmittersCompleteOrDisabled &= Inst.HandleCompletion();
+		bEmittersCompleteOrDisabled &= Emitter->HandleCompletion();
 	}
 
 	bool bCompletedAlready = IsComplete();
@@ -2382,11 +2466,17 @@ void FNiagaraSystemInstance::Tick_GameThread(float DeltaSeconds)
 	
 #if WITH_EDITOR
 	// We need to tick the rapid iteration parameters when in the editor
-	for (auto& EmitterInstance : Emitters)
+	for (const FNiagaraEmitterInstanceRef& EmitterRef : Emitters)
 	{
-		if (EmitterInstance->ShouldTick())
+		//-TODO:Stateless: Avoid Casting
+		FNiagaraEmitterInstanceImpl* Emitter = EmitterRef->AsStateful();
+		if (Emitter == nullptr)
+		{
+			continue;
+		}
+		if (Emitter->ShouldTick())
 		{
-			EmitterInstance->TickRapidIterationParameters();
+			Emitter->TickRapidIterationParameters();
 		}
 	}
 #endif
@@ -2423,8 +2513,8 @@ void FNiagaraSystemInstance::Tick_Concurrent(bool bEnqueueGPUTickIfNeeded)
 	bool bHasTickingEmitters = false;
 	for (const FNiagaraEmitterExecutionIndex& EmitterExecIdx : EmitterExecutionOrder)
 	{
-		FNiagaraEmitterInstance& Inst = Emitters[EmitterExecIdx.EmitterIndex].Get();
-		if (Inst.ShouldTick())
+		const FNiagaraEmitterInstance& Emitter = Emitters[EmitterExecIdx.EmitterIndex].Get();
+		if (Emitter.ShouldTick())
 		{
 			bHasTickingEmitters = true;
 			EmittersShouldTick.SetRange(EmitterExecIdx.EmitterIndex, 1, true);
@@ -2442,8 +2532,8 @@ void FNiagaraSystemInstance::Tick_Concurrent(bool bEnqueueGPUTickIfNeeded)
 	{
 		if (EmittersShouldTick[EmitterExecIdx.EmitterIndex])
 		{
-			FNiagaraEmitterInstance& Inst = Emitters[EmitterExecIdx.EmitterIndex].Get();
-			Inst.PreTick();
+			FNiagaraEmitterInstance& Emitter = Emitters[EmitterExecIdx.EmitterIndex].Get();
+			Emitter.PreTick();
 		}
 	}
 
@@ -2452,27 +2542,30 @@ void FNiagaraSystemInstance::Tick_Concurrent(bool bEnqueueGPUTickIfNeeded)
 	// now tick all emitters
 	for (const FNiagaraEmitterExecutionIndex& EmitterExecIdx : EmitterExecutionOrder)
 	{
-		FNiagaraEmitterInstance& Inst = Emitters[EmitterExecIdx.EmitterIndex].Get();
+		FNiagaraEmitterInstance& Emitter = Emitters[EmitterExecIdx.EmitterIndex].Get();
 		if (EmittersShouldTick[EmitterExecIdx.EmitterIndex])
 		{
-			Inst.Tick(CachedDeltaSeconds);
+			Emitter.Tick(CachedDeltaSeconds);
 		}
 
-		FVersionedNiagaraEmitterData* EmitterData = Inst.GetCachedEmitterData();
-		if (EmitterData && EmitterData->SimTarget == ENiagaraSimTarget::GPUComputeSim && !Inst.IsComplete())
+		//-TODO:Stateless: Avoid Casting
+		if (!Emitter.IsComplete() && Emitter.GetSimTarget() == ENiagaraSimTarget::GPUComputeSim)
 		{
-			// Handle edge case where an emitter was set to inactive on the first frame by scalability
-			// Since it will not tick we should not execute a GPU tick for it, this test must be symeterical with FNiagaraGPUSystemTick::Init
-			const bool bIsInactive = (Inst.GetExecutionState() == ENiagaraExecutionState::Inactive) || (Inst.GetExecutionState() == ENiagaraExecutionState::InactiveClear);
-			if (Inst.HasTicked() || !bIsInactive)
+			if (FNiagaraEmitterInstanceImpl* StatefulEmitter = Emitter.AsStateful())
 			{
-				if (const FNiagaraComputeExecutionContext* GPUContext = Inst.GetGPUContext())
+				// Handle edge case where an emitter was set to inactive on the first frame by scalability
+				// Since it will not tick we should not execute a GPU tick for it, this test must be symeterical with FNiagaraGPUSystemTick::Init
+				const bool bIsInactive = (StatefulEmitter->GetExecutionState() == ENiagaraExecutionState::Inactive) || (StatefulEmitter->GetExecutionState() == ENiagaraExecutionState::InactiveClear);
+				if (StatefulEmitter->HasTicked() || !bIsInactive)
 				{
-					const int32 InterpFactor = GPUContext->HasInterpolationParameters ? 2 : 1;
+					if (const FNiagaraComputeExecutionContext* GPUContext = StatefulEmitter->GetGPUContext())
+					{
+						const int32 InterpFactor = GPUContext->HasInterpolationParameters ? 2 : 1;
 
-					TotalCombinedParamStoreSize += InterpFactor * GPUContext->GetConstantBufferSize();
-					GPUParamIncludeInterpolation = GPUContext->HasInterpolationParameters || GPUParamIncludeInterpolation;
-					ActiveGPUEmitterCount++;
+						TotalCombinedParamStoreSize += InterpFactor * GPUContext->GetConstantBufferSize();
+						GPUParamIncludeInterpolation = GPUContext->HasInterpolationParameters || GPUParamIncludeInterpolation;
+						ActiveGPUEmitterCount++;
+					}
 				}
 			}
 		}
@@ -2712,13 +2805,13 @@ void FNiagaraSystemInstance::DestroyDataInterfaceInstanceData()
 	GPUDataInterfaces.Empty();
 }
 
-TSharedPtr<FNiagaraEmitterInstance, ESPMode::ThreadSafe> FNiagaraSystemInstance::GetSimulationForHandle(const FNiagaraEmitterHandle& EmitterHandle)
+FNiagaraEmitterInstancePtr FNiagaraSystemInstance::GetSimulationForHandle(const FNiagaraEmitterHandle& EmitterHandle) const
 {
-	for (const TSharedRef<FNiagaraEmitterInstance, ESPMode::ThreadSafe>& Sim : Emitters)
+	for (const FNiagaraEmitterInstanceRef& Emitter : Emitters)
 	{
-		if(Sim->GetEmitterHandle().GetId() == EmitterHandle.GetId())
+		if ( Emitter->GetEmitterHandle().GetId() == EmitterHandle.GetId() )
 		{
-			return Sim;
+			return Emitter;
 		}
 	}
 	return nullptr;
@@ -2755,7 +2848,7 @@ FBox FNiagaraSystemInstance::GetSystemFixedBounds() const
 
 void FNiagaraSystemInstance::SetEmitterFixedBounds(FName EmitterName, const FBox& InLocalBounds)
 {
-	for (const TSharedRef<FNiagaraEmitterInstance, ESPMode::ThreadSafe>& Emitter : Emitters)
+	for (const FNiagaraEmitterInstanceRef& Emitter : Emitters)
 	{
 		if ( Emitter->GetEmitterHandle().GetName() == EmitterName )
 		{
@@ -2770,7 +2863,7 @@ void FNiagaraSystemInstance::SetEmitterFixedBounds(FName EmitterName, const FBox
 
 FBox FNiagaraSystemInstance::GetEmitterFixedBounds(FName EmitterName) const
 {
-	for (const TSharedRef<FNiagaraEmitterInstance, ESPMode::ThreadSafe>& Emitter : Emitters)
+	for (const FNiagaraEmitterInstanceRef& Emitter : Emitters)
 	{
 		if (Emitter->GetEmitterHandle().GetName() == EmitterName)
 		{
diff --git a/Engine/Plugins/FX/Niagara/Source/Niagara/Private/NiagaraSystemInstanceController.cpp b/Engine/Plugins/FX/Niagara/Source/Niagara/Private/NiagaraSystemInstanceController.cpp
--- a/Engine/Plugins/FX/Niagara/Source/Niagara/Private/NiagaraSystemInstanceController.cpp
+++ b/Engine/Plugins/FX/Niagara/Source/Niagara/Private/NiagaraSystemInstanceController.cpp
@@ -107,31 +107,29 @@ void FNiagaraSystemInstanceController::GetUsedMaterials(TArray<UMaterialInterfac
 		return;
 	}
 
-	for (TSharedRef<FNiagaraEmitterInstance, ESPMode::ThreadSafe> EmitterInst : SystemInstance->GetEmitters())
+	for (const FNiagaraEmitterInstanceRef& EmitterInst : SystemInstance->GetEmitters())
 	{
-		if (FVersionedNiagaraEmitterData* EmitterData = EmitterInst->GetCachedEmitterData())
-		{
-			EmitterData->ForEachEnabledRenderer(
-				[&](UNiagaraRendererProperties* Properties)
-				{
-					TArray<UMaterialInterface*> Mats;
-					Properties->GetUsedMaterials(&EmitterInst.Get(), Mats);
+		EmitterInst->ForEachEnabledRenderer(
+			[&](const UNiagaraRendererProperties* Properties)
+			{
+				bool bCreateMidsForUsedMaterials = Properties->NeedsMIDsForMaterials();
+				TArray<UMaterialInterface*> Mats;
+				Properties->GetUsedMaterials(&EmitterInst.Get(), Mats);
 
-					if (Properties->NeedsMIDsForMaterials())
+				if (Properties->NeedsMIDsForMaterials())
+				{
+					for (const FMaterialOverride& Override : EmitterMaterials)
 					{
-						for (const FMaterialOverride& Override : EmitterMaterials)
+						if (Override.EmitterRendererProperty == Properties && Mats.IsValidIndex(Override.MaterialSubIndex))
 						{
-							if (Override.EmitterRendererProperty == Properties && Mats.IsValidIndex(Override.MaterialSubIndex))
-							{
-								Mats[Override.MaterialSubIndex] = Override.Material;
-							}
+							Mats[Override.MaterialSubIndex] = Override.Material;
 						}
 					}
-
-					OutMaterials.Append(Mats);
 				}
-			);
-		}
+
+				OutMaterials.Append(Mats);
+			}
+		);
 	}
 }
 
@@ -142,51 +140,48 @@ void FNiagaraSystemInstanceController::GetMaterialStreamingInfo(FNiagaraMaterial
 		return;
 	}
 
-	for (const TSharedRef<FNiagaraEmitterInstance, ESPMode::ThreadSafe>& EmitterInst : SystemInstance->GetEmitters())
+	for (const FNiagaraEmitterInstanceRef& EmitterInst : SystemInstance->GetEmitters())
 	{
-		if (FVersionedNiagaraEmitterData* EmitterData = EmitterInst->GetCachedEmitterData())
-		{
-			EmitterData->ForEachEnabledRenderer(
-				[&](UNiagaraRendererProperties* Properties)
+		EmitterInst->ForEachEnabledRenderer(
+			[&](UNiagaraRendererProperties* Properties)
+			{
+				TArray<UMaterialInterface*> UsedMaterials;
+				Properties->GetUsedMaterials(&EmitterInst.Get(), UsedMaterials);
+				if (UsedMaterials.Num() == 0)
 				{
-					TArray<UMaterialInterface*> UsedMaterials;
-					Properties->GetUsedMaterials(&EmitterInst.Get(), UsedMaterials);
-					if (UsedMaterials.Num() == 0)
-					{
-						return;
-					}
+					return;
+				}
 
-					if (Properties->NeedsMIDsForMaterials())
+				if (Properties->NeedsMIDsForMaterials())
+				{
+					for (const FMaterialOverride& Override : EmitterMaterials)
 					{
-						for (const FMaterialOverride& Override : EmitterMaterials)
+						if (Override.EmitterRendererProperty == Properties && UsedMaterials.IsValidIndex(Override.MaterialSubIndex))
 						{
-							if (Override.EmitterRendererProperty == Properties && UsedMaterials.IsValidIndex(Override.MaterialSubIndex))
-							{
-								UsedMaterials[Override.MaterialSubIndex] = Override.Material;
-							}
+							UsedMaterials[Override.MaterialSubIndex] = Override.Material;
 						}
 					}
+				}
 
-					const float StreamingScale = Properties->GetMaterialStreamingScale();
-					for (UMaterialInterface* UsedMaterial : UsedMaterials)
+				const float StreamingScale = Properties->GetMaterialStreamingScale();
+				for (UMaterialInterface* UsedMaterial : UsedMaterials)
+				{
+					if (UsedMaterial == nullptr)
 					{
-						if (UsedMaterial == nullptr)
-						{
-							continue;
-						}
+						continue;
+					}
 
-						if (FNiagaraMaterialAndScale* Existing = OutMaterialAndScales.FindByPredicate([UsedMaterial](const FNiagaraMaterialAndScale& Existing) { return Existing.Material == UsedMaterial; }))
-						{
-							Existing->Scale = FMath::Max(Existing->Scale, StreamingScale);
-						}
-						else
-						{
-							OutMaterialAndScales.Emplace(UsedMaterial, StreamingScale);
-						}
+					if (FNiagaraMaterialAndScale* Existing = OutMaterialAndScales.FindByPredicate([UsedMaterial](const FNiagaraMaterialAndScale& Existing) { return Existing.Material == UsedMaterial; }))
+					{
+						Existing->Scale = FMath::Max(Existing->Scale, StreamingScale);
+					}
+					else
+					{
+						OutMaterialAndScales.Emplace(UsedMaterial, StreamingScale);
 					}
 				}
-			);
-		}
+			}
+		);
 	}
 }
 
@@ -197,15 +192,14 @@ void FNiagaraSystemInstanceController::GetStreamingMeshInfo(const FBoxSphereBoun
 		return;
 	}
 
-	for (TSharedRef<FNiagaraEmitterInstance, ESPMode::ThreadSafe> EmitterInst : SystemInstance->GetEmitters())
+	for (const FNiagaraEmitterInstanceRef& EmitterInst : SystemInstance->GetEmitters())
 	{
-		if (FVersionedNiagaraEmitterData* EmitterData = EmitterInst->GetCachedEmitter().GetEmitterData())
-		{
-			EmitterData->ForEachEnabledRenderer([&](UNiagaraRendererProperties* Properties)
+		EmitterInst->ForEachEnabledRenderer(
+			[&](const UNiagaraRendererProperties* Properties)
 			{
 				Properties->GetStreamingMeshInfo(OwnerBounds, &EmitterInst.Get(), OutStreamingRenderAssets);
-			});
-		}
+			}
+		);
 	}
 }
 
@@ -231,9 +225,9 @@ void FNiagaraSystemInstanceController::DebugDump(bool bFullDump)
 
 		if (!SystemInstance->IsComplete())
 		{
-			for (TSharedRef<FNiagaraEmitterInstance, ESPMode::ThreadSafe> Emitter : SystemInstance->GetEmitters())
+			for (const FNiagaraEmitterInstanceRef& Emitter : SystemInstance->GetEmitters())
 			{
-				if ( Emitter->GetCachedEmitter().Emitter != nullptr )
+				if ( Emitter->GetEmitter() != nullptr )
 				{
 					UE_LOG(LogNiagara, Log, TEXT("\tEmitter '%s' ExecutionState(%s) NumParticles(%d)"), *Emitter->GetEmitterHandle().GetUniqueInstanceName(), *ExecutionStateEnum->GetNameStringByIndex((int32)Emitter->GetExecutionState()), Emitter->GetNumParticles());
 				}
@@ -254,12 +248,9 @@ SIZE_T FNiagaraSystemInstanceController::GetTotalBytesUsed() const
 	SIZE_T Size = 0;
 	if (IsValid())
 	{
-		for (const TSharedRef<FNiagaraEmitterInstance, ESPMode::ThreadSafe>& Emitter : GetSystemInstance_Unsafe()->GetEmitters())
+		for (const FNiagaraEmitterInstanceRef& Emitter : GetSystemInstance_Unsafe()->GetEmitters())
 		{
-			if (Emitter->GetCachedEmitter().Emitter != nullptr)
-			{
-				Size += Emitter->GetTotalBytesUsed();
-			}
+			Size += Emitter->GetTotalBytesUsed();
 		}
 	}
 
@@ -286,55 +277,52 @@ void FNiagaraSystemInstanceController::UpdateEmitterMaterials()
 		for (int32 i = 0; i < SystemInstance->GetEmitters().Num(); i++)
 		{
 			FNiagaraEmitterInstance* EmitterInst = &SystemInstance->GetEmitters()[i].Get();
-			if (FVersionedNiagaraEmitterData* EmitterData = EmitterInst->GetCachedEmitterData())
-			{
-				EmitterData->ForEachEnabledRenderer(
-					[&](UNiagaraRendererProperties* Properties)
+			EmitterInst->ForEachEnabledRenderer(
+				[&](const UNiagaraRendererProperties* Properties)
+				{
+					// Nothing to do if we don't create MIDs for this material
+					if (!Properties->NeedsMIDsForMaterials())
 					{
-						// Nothing to do if we don't create MIDs for this material
-						if (!Properties->NeedsMIDsForMaterials())
-						{
-							return;
-						}
+						return;
+					}
 
-						TArray<UMaterialInterface*> UsedMaterials;
-						Properties->GetUsedMaterials(EmitterInst, UsedMaterials);
+					TArray<UMaterialInterface*> UsedMaterials;
+					Properties->GetUsedMaterials(EmitterInst, UsedMaterials);
 
-						uint32 MaterialIndex = 0;
-						for (UMaterialInterface*& ExistingMaterial : UsedMaterials)
+					uint32 MaterialIndex = 0;
+					for (UMaterialInterface*& ExistingMaterial : UsedMaterials)
+					{
+						if (ExistingMaterial)
 						{
-							if (ExistingMaterial)
+							if (UMaterialInstanceDynamic* MID = Cast<UMaterialInstanceDynamic>(ExistingMaterial))
 							{
-								if (UMaterialInstanceDynamic* MID = Cast<UMaterialInstanceDynamic>(ExistingMaterial))
+								if (EmitterMaterials.FindByPredicate([&](const FMaterialOverride& ExistingOverride) -> bool { return (ExistingOverride.Material == ExistingMaterial) && (ExistingOverride.EmitterRendererProperty == Properties) && (ExistingOverride.MaterialSubIndex == MaterialIndex); }))
 								{
-									if (EmitterMaterials.FindByPredicate([&](const FMaterialOverride& ExistingOverride) -> bool { return (ExistingOverride.Material == ExistingMaterial) && (ExistingOverride.EmitterRendererProperty == Properties) && (ExistingOverride.MaterialSubIndex == MaterialIndex); }))
-									{
-										// It's a MID we've previously created and are managing. Recreate it by grabbing the parent.
-										// TODO: Are there cases where we don't always have to recreate it?
-										ExistingMaterial = MID->Parent;
-									}
-									else
-									{
-										// If we get here this is an external MID so do not create a new one
-										continue;
-									}
+									// It's a MID we've previously created and are managing. Recreate it by grabbing the parent.
+									// TODO: Are there cases where we don't always have to recreate it?
+									ExistingMaterial = MID->Parent;
 								}
+								else
+								{
+									// If we get here this is an external MID so do not create a new one
+									continue;
+								}
+							}
 
-								// Create a new MID
-								//UE_LOG(LogNiagara, Log, TEXT("Create Dynamic Material for component %s"), *GetPathName());
-								ExistingMaterial = UMaterialInstanceDynamic::Create(ExistingMaterial, nullptr);
-								FMaterialOverride Override;
-								Override.Material = ExistingMaterial;
-								Override.EmitterRendererProperty = Properties;
-								Override.MaterialSubIndex = MaterialIndex;
+							// Create a new MID
+							//UE_LOG(LogNiagara, Log, TEXT("Create Dynamic Material for component %s"), *GetPathName());
+							ExistingMaterial = UMaterialInstanceDynamic::Create(ExistingMaterial, nullptr);
+							FMaterialOverride Override;
+							Override.Material = ExistingMaterial;
+							Override.EmitterRendererProperty = Properties;
+							Override.MaterialSubIndex = MaterialIndex;
 
-								NewEmitterMaterials.Add(Override);
-							}
-							++MaterialIndex;
+							NewEmitterMaterials.Add(Override);
 						}
+						++MaterialIndex;
 					}
-				);
-			}
+				}
+			);
 		}
 	}
 	EmitterMaterials = MoveTemp(NewEmitterMaterials);
diff --git a/Engine/Plugins/FX/Niagara/Source/Niagara/Private/NiagaraSystemRenderData.cpp b/Engine/Plugins/FX/Niagara/Source/Niagara/Private/NiagaraSystemRenderData.cpp
--- a/Engine/Plugins/FX/Niagara/Source/Niagara/Private/NiagaraSystemRenderData.cpp
+++ b/Engine/Plugins/FX/Niagara/Source/Niagara/Private/NiagaraSystemRenderData.cpp
@@ -169,20 +169,18 @@ void FNiagaraSystemRenderData::GenerateSetDynamicDataCommands(FSetDynamicDataCom
 	for (int32 i = 0; i < SystemInstance->GetEmitters().Num(); i++)
 	{
 		FNiagaraEmitterInstance* EmitterInst = &SystemInstance->GetEmitters()[i].Get();
-		FVersionedNiagaraEmitterData* EmitterData = EmitterInst->GetCachedEmitterData();
+#if STATS
+		TStatId EmitterStatID = EmitterInst->GetEmitterStatID(true, true);
+		FScopeCycleCounter EmitterStatCounter(EmitterStatID);
+#endif
 
-		if (EmitterData == nullptr)
+		if (!EmitterInst->GetEmitterHandle().GetIsEnabled())
 		{
 			continue;
 		}
 
-#if STATS
-		TStatId EmitterStatID = EmitterInst->GetCachedEmitter().Emitter->GetStatID(true, true);
-		FScopeCycleCounter EmitterStatCounter(EmitterStatID);
-#endif
-
-		EmitterData->ForEachEnabledRenderer(
-			[&](UNiagaraRendererProperties* Properties)
+		EmitterInst->ForEachEnabledRenderer(
+			[&](const UNiagaraRendererProperties* Properties)
 			{
 				FNiagaraRenderer* Renderer = EmitterRenderers_GT[RendererIndex];
 				FNiagaraDynamicDataBase* NewData = nullptr;
@@ -191,8 +189,7 @@ void FNiagaraSystemRenderData::GenerateSetDynamicDataCommands(FSetDynamicDataCom
 				{
 					bool bRendererEditorEnabled = true;
 #if WITH_EDITORONLY_DATA
-					const FNiagaraEmitterHandle& Handle = System->GetEmitterHandle(i);
-					bRendererEditorEnabled = (!SystemInstance->GetIsolateEnabled() || Handle.IsIsolated());
+					bRendererEditorEnabled = EmitterInst->IsDisabledFromIsolation() == false;
 #endif
 					if (bRendererEditorEnabled && !EmitterInst->IsComplete() && !SystemInstance->IsComplete())
 					{
@@ -272,16 +269,19 @@ void FNiagaraSystemRenderData::RecacheRenderers(const FNiagaraSystemInstance& Sy
 
 	bAnyMotionBlurEnabled = false;
 	bAnyHeterogeneousVolumesEnabled = false;
-	for (TSharedRef<const FNiagaraEmitterInstance, ESPMode::ThreadSafe> EmitterInst : SystemInstance.GetEmitters())
+	for (const FNiagaraEmitterInstanceRef& EmitterInst : SystemInstance.GetEmitters())
 	{
-		if (FVersionedNiagaraEmitterData* EmitterData = EmitterInst->GetCachedEmitterData())
+		if (!EmitterInst->GetEmitterHandle().GetIsEnabled())
 		{
-			EmitterData->ForEachEnabledRenderer(
-				[&](UNiagaraRendererProperties* Properties)
+			continue;
+		}
+
+		EmitterInst->ForEachEnabledRenderer(
+			[&](UNiagaraRendererProperties* Properties)
 			{
 				//We can skip creation of the renderer if the current quality level doesn't support it. If the quality level changes all systems are fully reinitialized.
 				FNiagaraRenderer* NewRenderer = nullptr;
-				if (Properties->GetIsActive() && EmitterInst->GetData().IsInitialized() && !EmitterInst->IsDisabled())
+				if (Properties->GetIsActive() && EmitterInst->GetParticleData().IsInitialized() && !EmitterInst->IsDisabled())
 				{
 					NewRenderer = Properties->CreateEmitterRenderer(FeatureLevel, &EmitterInst.Get(), Controller);
 					if (Properties->MotionVectorSetting != ENiagaraRendererMotionVectorSetting::Disable)
@@ -295,8 +295,7 @@ void FNiagaraSystemRenderData::RecacheRenderers(const FNiagaraSystemInstance& Sy
 				}
 				EmitterRenderers_GT.Add(NewRenderer);
 			}
-			);
-		}
+		);
 	}
 
 	// If we have renderers then the draw order on the system should match, when compiling the number of renderers can be zero
@@ -337,11 +336,9 @@ void FNiagaraSystemRenderData::PostTickRenderers(const FNiagaraSystemInstance& S
 				const FNiagaraEmitterInstance& EmitterInst = SystemInstance.GetEmitters()[ExecIdx.EmitterIndex].Get();
 				if ( !EmitterInst.IsComplete() )
 				{
-					FVersionedNiagaraEmitterData* EmitterData = EmitterInst.GetCachedEmitterData();
-					FNiagaraRenderer* EmitterRenderer = EmitterRenderers_GT[ExecIdx.SystemRendererIndex];
-					if (EmitterData && EmitterRenderer)
+					if ( FNiagaraRenderer* EmitterRenderer = EmitterRenderers_GT[ExecIdx.SystemRendererIndex] )
 					{
-						const UNiagaraRendererProperties* RendererProperties = EmitterData->GetRenderers()[ExecIdx.EmitterRendererIndex];
+						const UNiagaraRendererProperties* RendererProperties = EmitterInst.GetRenderer(ExecIdx.EmitterRendererIndex);
 						EmitterRenderer->PostSystemTick_GameThread(RendererProperties, &EmitterInst);
 					}
 				}
@@ -376,7 +373,7 @@ void FNiagaraSystemRenderData::OnSystemComplete(const FNiagaraSystemInstance& Sy
 	}
 
 	const UNiagaraSystem* System = SystemInstance.GetSystem();
-	const TArray<TSharedRef<FNiagaraEmitterInstance, ESPMode::ThreadSafe>>& EmitterInstances = SystemInstance.GetEmitters();
+	TConstArrayView<FNiagaraEmitterInstanceRef> EmitterInstances = SystemInstance.GetEmitters();
 	if (EmitterInstances.Num() == 0)
 	{
 		return;
@@ -387,11 +384,9 @@ void FNiagaraSystemRenderData::OnSystemComplete(const FNiagaraSystemInstance& Sy
 		if (EmitterRenderers_GT.IsValidIndex(ExecIdx.SystemRendererIndex) && ensure(EmitterInstances.IsValidIndex(ExecIdx.EmitterIndex)) )
 		{
 			const FNiagaraEmitterInstance& EmitterInstance = EmitterInstances[ExecIdx.EmitterIndex].Get();
-			FVersionedNiagaraEmitterData* EmitterData = EmitterInstance.GetCachedEmitterData();
-			FNiagaraRenderer* EmitterRenderer = EmitterRenderers_GT[ExecIdx.SystemRendererIndex];
-			if (EmitterData && EmitterRenderer)
+			if (FNiagaraRenderer* EmitterRenderer = EmitterRenderers_GT[ExecIdx.SystemRendererIndex])
 			{
-				const UNiagaraRendererProperties* RendererProperties = EmitterData->GetRenderers()[ExecIdx.EmitterRendererIndex];
+				const UNiagaraRendererProperties* RendererProperties = EmitterInstance.GetRenderer(ExecIdx.EmitterRendererIndex);
 				EmitterRenderer->OnSystemComplete_GameThread(RendererProperties, &EmitterInstance);
 			}
 		}
diff --git a/Engine/Plugins/FX/Niagara/Source/Niagara/Private/NiagaraSystemSimulation.cpp b/Engine/Plugins/FX/Niagara/Source/Niagara/Private/NiagaraSystemSimulation.cpp
--- a/Engine/Plugins/FX/Niagara/Source/Niagara/Private/NiagaraSystemSimulation.cpp
+++ b/Engine/Plugins/FX/Niagara/Source/Niagara/Private/NiagaraSystemSimulation.cpp
@@ -2,6 +2,7 @@
 
 #include "NiagaraSystemSimulation.h"
 #include "NiagaraEmitterInstance.h"
+#include "NiagaraEmitterInstanceImpl.h"
 #include "NiagaraModule.h"
 #include "NiagaraTypes.h"
 #include "NiagaraParameterCollection.h"
@@ -1545,21 +1546,25 @@ void FNiagaraSystemSimulation::SimCachePostTick_Concurrent(float DeltaSeconds, c
 	{
 		FNiagaraSystemInstance& SystemInstance = *SystemInstances[iSystemInstance];
 
-		TArray<TSharedRef<FNiagaraEmitterInstance, ESPMode::ThreadSafe>>& Emitters = SystemInstance.GetEmitters();
+		TArrayView<FNiagaraEmitterInstanceRef> Emitters = SystemInstance.GetEmitters();
 		for (int32 iEmitter=0; iEmitter < Emitters.Num(); ++iEmitter)
 		{
 			FNiagaraEmitterInstance& EmitterInstance = Emitters[iEmitter].Get();
-			if ( EmitterInstance.IsComplete() )
+			if (EmitterInstance.IsComplete())
 			{
 				continue;
 			}
 
-			ENiagaraExecutionState State = EmitterExecutionStateAccessors[iEmitter].GetReader(MainDataSet).GetSafe(iSystemInstance, ENiagaraExecutionState::Disabled);
-			EmitterInstance.SetExecutionState(State);
-
 			//DataSetToEmitterSpawnParameters[iEmitter].DataSetToParameterStore(EmitterInstance.GetSpawnExecutionContext().Parameters, MainDataSet, iSystemInstance);
 			//DataSetToEmitterUpdateParameters[iEmitter].DataSetToParameterStore(EmitterInstance.GetUpdateExecutionContext().Parameters, MainDataSet, iSystemInstance);
 			DataSetToEmitterRendererParameters[iEmitter].DataSetToParameterStore(EmitterInstance.GetRendererBoundVariables(), MainDataSet, iSystemInstance);
+
+			//-TODO:Stateless:
+			if ( FNiagaraEmitterInstanceImpl* StatefulEmitter = EmitterInstance.AsStateful() )
+			{
+				const ENiagaraExecutionState State = EmitterExecutionStateAccessors[iEmitter].GetReader(MainDataSet).GetSafe(iSystemInstance, ENiagaraExecutionState::Disabled);
+				StatefulEmitter->SetExecutionState(State);
+			}
 		}
 	}
 }
@@ -1770,17 +1775,16 @@ void FNiagaraSystemSimulation::PrepareForSystemSimulate(FNiagaraSystemSimulation
 			UpdateInstanceParameterToDataSetBinding.ParameterStoreToDataSet(InstParameters, UpdateInstanceParameterDataSet, SystemIndex);
 		}
 
-		FNiagaraConstantBufferToDataSetBinding::CopyToDataSets(
-			Context.System->GetSystemCompiledData(), *Inst, SpawnInstanceParameterDataSet, UpdateInstanceParameterDataSet, SystemIndex);
+		FNiagaraConstantBufferToDataSetBinding::CopyToDataSets(Context.System->GetSystemCompiledData(), *Inst, SpawnInstanceParameterDataSet, UpdateInstanceParameterDataSet, SystemIndex);
 
 		//TODO: Find good way to check that we're not using any instance parameter data interfaces in the system scripts here.
 		//In that case we need to solo and will never get here.
 
-		TArray<TSharedRef<FNiagaraEmitterInstance, ESPMode::ThreadSafe>>& Emitters = Inst->GetEmitters();
+		TArrayView<FNiagaraEmitterInstanceRef> Emitters = Inst->GetEmitters();
 		for (int32 EmitterIdx = 0; EmitterIdx < Emitters.Num(); ++EmitterIdx)
 		{
 			FNiagaraEmitterInstance& EmitterInst = Emitters[EmitterIdx].Get();
-			if ( (EmitterExecutionStateAccessors.Num() > EmitterIdx) )
+			if ( EmitterExecutionStateAccessors.IsValidIndex(EmitterIdx) )
 			{
 				EmitterExecutionStateAccessors[EmitterIdx].GetWriter(Context.DataSet).SetSafe(SystemIndex, EmitterInst.GetExecutionState());
 			}
@@ -1996,65 +2000,69 @@ void FNiagaraSystemSimulation::TransferSystemSimResults(FNiagaraSystemSimulation
 		if (!SystemInst->IsDisabled())
 		{
 			//Now pull data out of the simulation and drive the emitters with it.
-			TArray<TSharedRef<FNiagaraEmitterInstance, ESPMode::ThreadSafe>>& Emitters = SystemInst->GetEmitters();
+			TArrayView<FNiagaraEmitterInstanceRef> Emitters = SystemInst->GetEmitters();
 			for (int32 EmitterIdx = 0; EmitterIdx < Emitters.Num(); ++EmitterIdx)
 			{
-				FNiagaraEmitterInstance& EmitterInst = Emitters[EmitterIdx].Get();
+				FNiagaraEmitterInstance& EmitterInstance = Emitters[EmitterIdx].Get();
 
 				//Early exit before we set the state as if we're complete or disabled we should never let the emitter turn itself back. It needs to be reset/reinited manually.
-				if (EmitterInst.IsComplete())
+				if (EmitterInstance.IsComplete())
 				{
 					continue;
 				}
 
 				check(Emitters.Num() > EmitterIdx);
 
-				ENiagaraExecutionState State = EmitterExecutionStateAccessors[EmitterIdx].GetReader(Context.DataSet).GetSafe(SystemIndex, ENiagaraExecutionState::Disabled);
-				EmitterInst.SetExecutionState(State);
+				DataSetToEmitterRendererParameters[EmitterIdx].DataSetToParameterStore(EmitterInstance.GetRendererBoundVariables(), Context.DataSet, SystemIndex);
 
-				TConstArrayView<FNiagaraDataSetAccessor<FNiagaraSpawnInfo>> EmitterSpawnInfoAccessors = System->GetEmitterSpawnInfoAccessors(EmitterIdx);
-				TArray<FNiagaraSpawnInfo>& EmitterInstSpawnInfos = EmitterInst.GetSpawnInfo();
-				for (int32 SpawnInfoIdx = 0; SpawnInfoIdx < EmitterSpawnInfoAccessors.Num(); ++SpawnInfoIdx)
+				//-TODO:Stateless:
+				if ( FNiagaraEmitterInstanceImpl* StatefulEmitter = EmitterInstance.AsStateful() )
 				{
-					if (SpawnInfoIdx < EmitterInstSpawnInfos.Num())
-					{
-						EmitterInstSpawnInfos[SpawnInfoIdx] = EmitterSpawnInfoAccessors[SpawnInfoIdx].GetReader(Context.DataSet).Get(SystemIndex);
-					}
-					else
+					ENiagaraExecutionState State = EmitterExecutionStateAccessors[EmitterIdx].GetReader(Context.DataSet).GetSafe(SystemIndex, ENiagaraExecutionState::Disabled);
+					StatefulEmitter->SetExecutionState(State);
+
+					TConstArrayView<FNiagaraDataSetAccessor<FNiagaraSpawnInfo>> EmitterSpawnInfoAccessors = System->GetEmitterSpawnInfoAccessors(EmitterIdx);
+					TArray<FNiagaraSpawnInfo>& EmitterInstSpawnInfos = StatefulEmitter->GetSpawnInfo();
+					for (int32 SpawnInfoIdx = 0; SpawnInfoIdx < EmitterSpawnInfoAccessors.Num(); ++SpawnInfoIdx)
 					{
-						ensure(SpawnInfoIdx < EmitterInstSpawnInfos.Num());
+						if (SpawnInfoIdx < EmitterInstSpawnInfos.Num())
+						{
+							EmitterInstSpawnInfos[SpawnInfoIdx] = EmitterSpawnInfoAccessors[SpawnInfoIdx].GetReader(Context.DataSet).Get(SystemIndex);
+						}
+						else
+						{
+							ensure(SpawnInfoIdx < EmitterInstSpawnInfos.Num());
+						}
 					}
-				}
 
-				//TODO: Any other fixed function stuff like this?
+					//TODO: Any other fixed function stuff like this?
 
-				FNiagaraScriptExecutionContext& SpawnContext = EmitterInst.GetSpawnExecutionContext();
-				DataSetToEmitterSpawnParameters[EmitterIdx].DataSetToParameterStore(SpawnContext.Parameters, Context.DataSet, SystemIndex);
+					FNiagaraScriptExecutionContext& SpawnContext = StatefulEmitter->GetSpawnExecutionContext();
+					DataSetToEmitterSpawnParameters[EmitterIdx].DataSetToParameterStore(SpawnContext.Parameters, Context.DataSet, SystemIndex);
 
-				FNiagaraScriptExecutionContext& UpdateContext = EmitterInst.GetUpdateExecutionContext();
-				DataSetToEmitterUpdateParameters[EmitterIdx].DataSetToParameterStore(UpdateContext.Parameters, Context.DataSet, SystemIndex);
+					FNiagaraScriptExecutionContext& UpdateContext = StatefulEmitter->GetUpdateExecutionContext();
+					DataSetToEmitterUpdateParameters[EmitterIdx].DataSetToParameterStore(UpdateContext.Parameters, Context.DataSet, SystemIndex);
 
-				FNiagaraComputeExecutionContext* GPUContext = EmitterInst.GetGPUContext();
-				if (GPUContext)
-				{
-					DataSetToEmitterGPUParameters[EmitterIdx].DataSetToParameterStore(GPUContext->CombinedParamStore, Context.DataSet, SystemIndex);
-				}
-
-				TArrayView<FNiagaraScriptExecutionContext> EventContexts = EmitterInst.GetEventExecutionContexts();
-				for (int32 EventIdx = 0; EventIdx < EventContexts.Num(); ++EventIdx)
-				{
-					FNiagaraScriptExecutionContext& EventContext = EventContexts[EventIdx];
-					if (DataSetToEmitterEventParameters[EmitterIdx].Num() > EventIdx)
+					FNiagaraComputeExecutionContext* GPUContext = StatefulEmitter->GetGPUContext();
+					if (GPUContext)
 					{
-						DataSetToEmitterEventParameters[EmitterIdx][EventIdx].DataSetToParameterStore(EventContext.Parameters, Context.DataSet, SystemIndex);
+						DataSetToEmitterGPUParameters[EmitterIdx].DataSetToParameterStore(GPUContext->CombinedParamStore, Context.DataSet, SystemIndex);
 					}
-					else
+
+					TArrayView<FNiagaraScriptExecutionContext> EventContexts = StatefulEmitter->GetEventExecutionContexts();
+					for (int32 EventIdx = 0; EventIdx < EventContexts.Num(); ++EventIdx)
 					{
-						UE_LOG(LogNiagara, Log, TEXT("Skipping DataSetToEmitterEventParameters because EventIdx is out-of-bounds. %d of %d"), EventIdx, DataSetToEmitterEventParameters[EmitterIdx].Num());
+						FNiagaraScriptExecutionContext& EventContext = EventContexts[EventIdx];
+						if (DataSetToEmitterEventParameters[EmitterIdx].Num() > EventIdx)
+						{
+							DataSetToEmitterEventParameters[EmitterIdx][EventIdx].DataSetToParameterStore(EventContext.Parameters, Context.DataSet, SystemIndex);
+						}
+						else
+						{
+							UE_LOG(LogNiagara, Log, TEXT("Skipping DataSetToEmitterEventParameters because EventIdx is out-of-bounds. %d of %d"), EventIdx, DataSetToEmitterEventParameters[EmitterIdx].Num());
+						}
 					}
 				}
-
-				DataSetToEmitterRendererParameters[EmitterIdx].DataSetToParameterStore(EmitterInst.GetRendererBoundVariables(), Context.DataSet, SystemIndex);
 			}
 		}
 	}
@@ -2190,7 +2198,7 @@ void FNiagaraSystemSimulation::InitParameterDataSetBindings(FNiagaraSystemInstan
 		SpawnInstanceParameterToDataSetBinding.Init(SpawnInstanceParameterDataSet, SystemInst->GetInstanceParameters());
 		UpdateInstanceParameterToDataSetBinding.Init(UpdateInstanceParameterDataSet, SystemInst->GetInstanceParameters());
 
-		TArray<TSharedRef<FNiagaraEmitterInstance, ESPMode::ThreadSafe>>& Emitters = SystemInst->GetEmitters();
+		TArrayView<FNiagaraEmitterInstanceRef> Emitters = SystemInst->GetEmitters();
 		const int32 EmitterCount = Emitters.Num();
 
 		DataSetToEmitterSpawnParameters.SetNum(EmitterCount);
@@ -2201,34 +2209,38 @@ void FNiagaraSystemSimulation::InitParameterDataSetBindings(FNiagaraSystemInstan
 
 		for (int32 EmitterIdx = 0; EmitterIdx < EmitterCount; ++EmitterIdx)
 		{
-			FNiagaraEmitterInstance& EmitterInst = Emitters[EmitterIdx].Get();
-			if (EmitterInst.IsDisabled())
+			FNiagaraEmitterInstance& EmitterInstance = Emitters[EmitterIdx].Get();
+			if (EmitterInstance.IsDisabled())
 			{
 				continue;
 			}
 
-			FNiagaraScriptExecutionContext& SpawnContext = EmitterInst.GetSpawnExecutionContext();
-			DataSetToEmitterSpawnParameters[EmitterIdx].Init(MainDataSet, SpawnContext.Parameters);
-
-			FNiagaraScriptExecutionContext& UpdateContext = EmitterInst.GetUpdateExecutionContext();
-			DataSetToEmitterUpdateParameters[EmitterIdx].Init(MainDataSet, UpdateContext.Parameters);
+			DataSetToEmitterRendererParameters[EmitterIdx].Init(MainDataSet, EmitterInstance.GetRendererBoundVariables());
 
-			FNiagaraComputeExecutionContext* GPUContext = EmitterInst.GetGPUContext();
-			if (GPUContext)
+			//-TODO:Stateless:
+			if (FNiagaraEmitterInstanceImpl* StatefulEmitter = EmitterInstance.AsStateful())
 			{
-				DataSetToEmitterGPUParameters[EmitterIdx].Init(MainDataSet, GPUContext->CombinedParamStore);
-			}
+				FNiagaraScriptExecutionContext& SpawnContext = StatefulEmitter->GetSpawnExecutionContext();
+				DataSetToEmitterSpawnParameters[EmitterIdx].Init(MainDataSet, SpawnContext.Parameters);
 
-			DataSetToEmitterRendererParameters[EmitterIdx].Init(MainDataSet, EmitterInst.GetRendererBoundVariables());
+				FNiagaraScriptExecutionContext& UpdateContext = StatefulEmitter->GetUpdateExecutionContext();
+				DataSetToEmitterUpdateParameters[EmitterIdx].Init(MainDataSet, UpdateContext.Parameters);
 
-			TArrayView<FNiagaraScriptExecutionContext> EventContexts = EmitterInst.GetEventExecutionContexts();
-			const int32 EventCount = EventContexts.Num();
-			DataSetToEmitterEventParameters[EmitterIdx].SetNum(EventCount);
+				FNiagaraComputeExecutionContext* GPUContext = StatefulEmitter->GetGPUContext();
+				if (GPUContext)
+				{
+					DataSetToEmitterGPUParameters[EmitterIdx].Init(MainDataSet, GPUContext->CombinedParamStore);
+				}
 
-			for (int32 EventIdx = 0; EventIdx < EventCount; ++EventIdx)
-			{
-				FNiagaraScriptExecutionContext& EventContext = EventContexts[EventIdx];
-				DataSetToEmitterEventParameters[EmitterIdx][EventIdx].Init(MainDataSet, EventContext.Parameters);
+				TArrayView<FNiagaraScriptExecutionContext> EventContexts = StatefulEmitter->GetEventExecutionContexts();
+				const int32 EventCount = EventContexts.Num();
+				DataSetToEmitterEventParameters[EmitterIdx].SetNum(EventCount);
+
+				for (int32 EventIdx = 0; EventIdx < EventCount; ++EventIdx)
+				{
+					FNiagaraScriptExecutionContext& EventContext = EventContexts[EventIdx];
+					DataSetToEmitterEventParameters[EmitterIdx][EventIdx].Init(MainDataSet, EventContext.Parameters);
+				}
 			}
 		}
 	}
@@ -2271,8 +2283,9 @@ void FNiagaraConstantBufferToDataSetBinding::CopyToDataSets(
 		ApplyOffsets(CompiledData.UpdateInstanceOwnerBinding, OwnerParameters, UpdateDataSet, DataSestInstanceIndex);
 	}
 
-	const auto& Emitters = SystemInstance.GetEmitters();
-	const int32 EmitterCount = Emitters.Num();
+	//-TODO:Stateless: Remove when we have them inside EmitterHandles
+	TConstArrayView<FNiagaraEmitterInstanceRef> Emitters = SystemInstance.GetEmitters();
+	const int32 EmitterCount = FMath::Min(Emitters.Num(), CompiledData.SpawnInstanceEmitterBindings.Num());
 
 	for (int32 EmitterIdx = 0; EmitterIdx < EmitterCount; ++EmitterIdx)
 	{
diff --git a/Engine/Plugins/FX/Niagara/Source/Niagara/Public/NiagaraGpuComputeDispatchInterface.h b/Engine/Plugins/FX/Niagara/Source/Niagara/Public/NiagaraGpuComputeDispatchInterface.h
--- a/Engine/Plugins/FX/Niagara/Source/Niagara/Public/NiagaraGpuComputeDispatchInterface.h
+++ b/Engine/Plugins/FX/Niagara/Source/Niagara/Public/NiagaraGpuComputeDispatchInterface.h
@@ -29,6 +29,8 @@ class FNiagaraSystemGpuComputeProxy;
 class FNiagaraGpuComputeDispatchInterface : public FFXSystemInterface
 {
 public:
+	DECLARE_EVENT_OneParam(FNiagaraGpuComputeDispatchInterface, FOnPostPreRenderEvent, FRDGBuilder&);
+
 	static NIAGARA_API FNiagaraGpuComputeDispatchInterface* Get(class UWorld* World);
 	static NIAGARA_API FNiagaraGpuComputeDispatchInterface* Get(class FSceneInterface* Scene);
 	static NIAGARA_API FNiagaraGpuComputeDispatchInterface* Get(class FFXSystemInterface* FXSceneInterface);
@@ -180,6 +182,11 @@ public:
 	FORCEINLINE void MultiGPUResourceModified(FRHICommandList& RHICmdList, FRHITexture* Texture, bool bRequiredForSimulation, bool bRequiredForRendering) const {}
 #endif
 
+	/** Event that broadcast before any rendering work is prepared / executed for Niagara. */
+	FOnPostPreRenderEvent& GetOnPreRenderEvent() { return OnPreRenderEvent; }
+	/** Event that broadcast after all rendering for Niagara is complete. */
+	FOnPostPreRenderEvent& GetOnPostRenderEvent() { return OnPostRenderEvent; }
+
 protected:
 	EShaderPlatform							ShaderPlatform;
 	ERHIFeatureLevel::Type					FeatureLevel;
@@ -199,4 +206,7 @@ protected:
 	bool									bIsOutsideSceneRenderer = false;
 	bool									bIsFirstViewFamily = true;
 	bool									bIsLastViewFamily = true;
+
+	FOnPostPreRenderEvent					OnPreRenderEvent;
+	FOnPostPreRenderEvent					OnPostRenderEvent;
 };
diff --git a/Engine/Plugins/FX/Niagara/Source/Niagara/Public/NiagaraParameterBinding.h b/Engine/Plugins/FX/Niagara/Source/Niagara/Public/NiagaraParameterBinding.h
--- a/Engine/Plugins/FX/Niagara/Source/Niagara/Public/NiagaraParameterBinding.h
+++ b/Engine/Plugins/FX/Niagara/Source/Niagara/Public/NiagaraParameterBinding.h
@@ -144,6 +144,8 @@ struct FNiagaraParameterBindingWithValue : public FNiagaraParameterBinding
 		return Value;
 	}
 
+	TConstArrayView<uint8> GetDefaultValueArray() const { return DefaultValue; }
+
 	bool operator==(const FNiagaraParameterBindingWithValue& Other) const
 	{
 		return
diff --git a/Engine/Plugins/FX/Niagara/Source/Niagara/Public/NiagaraRenderer.h b/Engine/Plugins/FX/Niagara/Source/Niagara/Public/NiagaraRenderer.h
--- a/Engine/Plugins/FX/Niagara/Source/Niagara/Public/NiagaraRenderer.h
+++ b/Engine/Plugins/FX/Niagara/Source/Niagara/Public/NiagaraRenderer.h
@@ -30,13 +30,10 @@ struct FNiagaraDynamicDataBase
 	NIAGARA_API explicit FNiagaraDynamicDataBase(const FNiagaraEmitterInstance* InEmitter);
 	NIAGARA_API virtual ~FNiagaraDynamicDataBase();
 
-	FNiagaraDynamicDataBase() = delete;
-	FNiagaraDynamicDataBase(FNiagaraDynamicDataBase& Other) = delete;
-	FNiagaraDynamicDataBase& operator=(const FNiagaraDynamicDataBase& Other) = delete;
+	UE_NONCOPYABLE(FNiagaraDynamicDataBase);
 
 	NIAGARA_API bool IsGpuLowLatencyTranslucencyEnabled() const;
 	NIAGARA_API FNiagaraDataBuffer* GetParticleDataToRender(bool bIsLowLatencyTranslucent = false) const;
-	FORCEINLINE ENiagaraSimTarget GetSimTarget() const { return SimTarget; }
 	FORCEINLINE FMaterialRelevance GetMaterialRelevance() const { return MaterialRelevance; }
 
 	FORCEINLINE void SetMaterialRelevance(FMaterialRelevance NewRelevance) { MaterialRelevance = NewRelevance; }
@@ -47,7 +44,6 @@ struct FNiagaraDynamicDataBase
 
 protected:
 	FMaterialRelevance MaterialRelevance;
-	ENiagaraSimTarget SimTarget;
 	FNiagaraSystemInstanceID SystemInstanceID;
 
 	FNiagaraDataBufferRef CPUParticleData;
diff --git a/Engine/Plugins/FX/Niagara/Source/Niagara/Public/NiagaraSystemInstance.h b/Engine/Plugins/FX/Niagara/Source/Niagara/Public/NiagaraSystemInstance.h
--- a/Engine/Plugins/FX/Niagara/Source/Niagara/Public/NiagaraSystemInstance.h
+++ b/Engine/Plugins/FX/Niagara/Source/Niagara/Public/NiagaraSystemInstance.h
@@ -12,6 +12,7 @@
 #include "NiagaraDataInterface.h"
 
 class FNiagaraWorldManager;
+class FNiagaraEmitterInstanceImpl;
 class FNiagaraSystemInstance;
 class FNiagaraSystemSimulation;
 class FNiagaraGpuComputeDispatchInterface;
@@ -135,6 +136,12 @@ public:
 	NIAGARA_API void BindParameters();
 	NIAGARA_API void UnbindParameters(bool bFromComplete = false);
 
+	// Bindings Override Parameters / Instance Parameters / System Simulation Parameters to the provded parameter store
+	// I.e. binds all relevant parameters from us and parent, does not do children (i.e. emitters)
+	void BindToParameterStore(FNiagaraParameterStore& ParameterStore);
+	// Unbinds the parameters that would be bound via BindToParameterStore
+	void UnbindFromParameterStore(FNiagaraParameterStore& ParameterStore);
+
 	FORCEINLINE FNiagaraParameterStore& GetInstanceParameters() { return InstanceParameters; }
 	NIAGARA_API FNiagaraLWCConverter GetLWCConverter(bool bLocalSpaceEmitter = false) const;
 	NIAGARA_API FTransform GetLWCSimToWorld(bool bLocalSpaceEmitter = false) const;
@@ -231,15 +238,16 @@ public:
 	FORCEINLINE bool IsDisabled()const { return ActualExecutionState == ENiagaraExecutionState::Disabled; }
 
 	/** Gets the simulation for the supplied emitter handle. */
-	NIAGARA_API TSharedPtr<FNiagaraEmitterInstance, ESPMode::ThreadSafe> GetSimulationForHandle(const FNiagaraEmitterHandle& EmitterHandle);
+	NIAGARA_API FNiagaraEmitterInstancePtr GetSimulationForHandle(const FNiagaraEmitterHandle& EmitterHandle) const;
 
 	FORCEINLINE UWorld* GetWorld() const { return World; }
 	FORCEINLINE UNiagaraSystem* GetSystem() const { return System; }
 	FORCEINLINE USceneComponent* GetAttachComponent() { return AttachComponent.Get(); }
 	FORCEINLINE FNiagaraUserRedirectionParameterStore* GetOverrideParameters() { return OverrideParameters; }
 	FORCEINLINE const FNiagaraUserRedirectionParameterStore* GetOverrideParameters() const { return OverrideParameters; }
-	FORCEINLINE TArray<TSharedRef<FNiagaraEmitterInstance, ESPMode::ThreadSafe> > &GetEmitters() { return Emitters; }
-	FORCEINLINE const TArray<TSharedRef<FNiagaraEmitterInstance, ESPMode::ThreadSafe> >& GetEmitters() const { return Emitters; }
+	[[nodiscard]] NIAGARA_API TArrayView<FNiagaraEmitterInstanceRef> GetEmitters() { return Emitters; }
+	[[nodiscard]] NIAGARA_API TConstArrayView<FNiagaraEmitterInstanceRef> GetEmitters() const { return Emitters; }
+
 	FORCEINLINE const FBox& GetLocalBounds() const { return LocalBounds;  }
 	FORCEINLINE const FVector3f& GetLWCTile() const { return LWCTile;  }
 	NIAGARA_API TConstArrayView<FNiagaraEmitterExecutionIndex> GetEmitterExecutionOrder() const;
@@ -505,7 +513,7 @@ private:
 	int32 WarmupTickCount = -1;
 	float WarmupTickDelta = 0;
 	
-	TArray< TSharedRef<FNiagaraEmitterInstance, ESPMode::ThreadSafe> > Emitters;
+	TArray<FNiagaraEmitterInstanceRef> Emitters;
 
 	FOnPostTick OnPostTickDelegate;
 	FOnComplete OnCompleteDelegate;
diff --git a/Engine/Plugins/FX/Niagara/Source/NiagaraEditor/Private/NiagaraBakerRenderer.cpp b/Engine/Plugins/FX/Niagara/Source/NiagaraEditor/Private/NiagaraBakerRenderer.cpp
--- a/Engine/Plugins/FX/Niagara/Source/NiagaraEditor/Private/NiagaraBakerRenderer.cpp
+++ b/Engine/Plugins/FX/Niagara/Source/NiagaraEditor/Private/NiagaraBakerRenderer.cpp
@@ -549,7 +549,7 @@ void FNiagaraBakerRenderer::RenderParticleAttribute(UTextureRenderTarget2D* Rend
 	
 	for ( const auto& EmitterInstance : SystemInstance->GetEmitters() )
 	{
-		UNiagaraEmitter* NiagaraEmitter = EmitterInstance->GetCachedEmitter().Emitter;
+		const UNiagaraEmitter* NiagaraEmitter = EmitterInstance->GetEmitter();
 		if ( !NiagaraEmitter || (NiagaraEmitter->GetUniqueEmitterName() != EmitterName) )
 		{
 			continue;
@@ -560,7 +560,7 @@ void FNiagaraBakerRenderer::RenderParticleAttribute(UTextureRenderTarget2D* Rend
 			return;
 		}
 	
-		const FNiagaraDataSet& ParticleDataSet = EmitterInstance->GetData();
+		const FNiagaraDataSet& ParticleDataSet = EmitterInstance->GetParticleData();
 		const FNiagaraDataBuffer* ParticleDataBuffer = ParticleDataSet.GetCurrentData();
 		FNiagaraDataSetReaderInt32<int32> UniqueIDAccessor = FNiagaraDataSetAccessor<int32>::CreateReader(ParticleDataSet, FName("UniqueID"));
 		if ( !ParticleDataBuffer || !UniqueIDAccessor.IsValid() )
diff --git a/Engine/Plugins/FX/Niagara/Source/NiagaraEditor/Private/NiagaraCompiler.cpp b/Engine/Plugins/FX/Niagara/Source/NiagaraEditor/Private/NiagaraCompiler.cpp
--- a/Engine/Plugins/FX/Niagara/Source/NiagaraEditor/Private/NiagaraCompiler.cpp
+++ b/Engine/Plugins/FX/Niagara/Source/NiagaraEditor/Private/NiagaraCompiler.cpp
@@ -1026,13 +1026,14 @@ TSharedPtr<FNiagaraCompileRequestDataBase, ESPMode::ThreadSafe> FNiagaraEditorMo
 		{
 			const FNiagaraEmitterHandle& Handle = System->GetEmitterHandle(i);
 			TSharedPtr<FNiagaraCompileRequestData, ESPMode::ThreadSafe> EmitterPtr = MakeShared<FNiagaraCompileRequestData, ESPMode::ThreadSafe>();
-			EmitterPtr->EmitterUniqueName = Handle.GetInstance().Emitter->GetUniqueEmitterName();
+			EmitterPtr->EmitterUniqueName = Handle.GetUniqueInstanceName();
 			EmitterPtr->EmitterID = FNiagaraEmitterID(i);
 			EmitterPtr->SourceName = BasePtr->SourceName;
-			EmitterPtr->Source = Cast<UNiagaraScriptSource>(Handle.GetEmitterData()->GraphSource);
+			//-TODO:Stateless:
+			EmitterPtr->Source = Handle.GetEmitterData() ? Cast<UNiagaraScriptSource>(Handle.GetEmitterData()->GraphSource) : nullptr;
+			//-TODO:Stateless:
 			EmitterPtr->bUseRapidIterationParams = BasePtr->bUseRapidIterationParams;
 			EmitterPtr->bDisableDebugSwitches = BasePtr->bDisableDebugSwitches;
-			//EmitterPtr->bSimulationStagesEnabled = Handle.GetInstance()->bSimulationStagesEnabled;
 			EmitterPtr->SharedCompileDataInterfaceData = BasePtr->SharedCompileDataInterfaceData;
 			BasePtr->EmitterData.Add(EmitterPtr);
 			EmitterNames.Add(Handle.GetUniqueInstanceName());
@@ -1064,7 +1065,9 @@ TSharedPtr<FNiagaraCompileRequestDataBase, ESPMode::ThreadSafe> FNiagaraEditorMo
 		{
 			const FNiagaraEmitterHandle& Handle = System->GetEmitterHandle(i);
 			FCompileConstantResolver ConstantResolver(Handle.GetInstance(), ENiagaraScriptUsage::EmitterSpawnScript);
-			if (Handle.GetIsEnabled()) // Don't pull in the emitter if it isn't going to be used.
+			//-TODO:Stateless:
+			if (Handle.GetIsEnabled() && Handle.GetEmitterData()) // Don't pull in the emitter if it isn't going to be used.
+			//-TODO:Stateless:
 			{
 				TSharedPtr<FNiagaraGraphCachedDataBase, ESPMode::ThreadSafe> CachedTraversalEmitterData = Handle.GetInstance().Emitter->GetCachedTraversalData(Handle.GetInstance().Version);
 				TArray<FNiagaraVariable> StaticVariablesFromEmitter = StaticVariablesFromSystem;
@@ -1197,7 +1200,9 @@ TSharedPtr<FNiagaraCompileRequestDataBase, ESPMode::ThreadSafe> FNiagaraEditorMo
 			{
 				const FNiagaraEmitterHandle& Handle = System->GetEmitterHandle(i);
 				FCompileConstantResolver ConstantResolver(Handle.GetInstance(), ENiagaraScriptUsage::EmitterSpawnScript);
-				if (Handle.GetIsEnabled()) // Don't pull in the emitter if it isn't going to be used.
+				//-TODO:Stateless:
+				if (Handle.GetIsEnabled() && Handle.GetEmitterData()) // Don't pull in the emitter if it isn't going to be used.
+					//-TODO:Stateless:
 				{
 					TArray<UNiagaraScript*> EmitterScripts;
 					Handle.GetEmitterData()->GetScripts(EmitterScripts, false);
@@ -1352,14 +1357,16 @@ TSharedPtr<FNiagaraCompileRequestDuplicateDataBase, ESPMode::ThreadSafe> FNiagar
 		{
 			const FNiagaraEmitterHandle& Handle = OwningSystem->GetEmitterHandle(i);
 			TSharedPtr<FNiagaraCompileRequestDuplicateData, ESPMode::ThreadSafe> EmitterPtr = MakeShared<FNiagaraCompileRequestDuplicateData, ESPMode::ThreadSafe>();
-			EmitterPtr->EmitterUniqueName = Handle.GetInstance().Emitter->GetUniqueEmitterName();
+			EmitterPtr->EmitterUniqueName = Handle.GetUniqueInstanceName();
 			EmitterPtr->EmitterID = FNiagaraEmitterID(i);
 			EmitterPtr->ValidUsages = BasePtr->ValidUsages;
 			EmitterPtr->SharedSourceGraphToDuplicatedGraphsMap = BasePtr->SharedSourceGraphToDuplicatedGraphsMap;
 			EmitterPtr->SharedNameToDuplicatedDataInterfaceMap = BasePtr->SharedNameToDuplicatedDataInterfaceMap;
 			EmitterPtr->SharedDataInterfaceClassToDuplicatedCDOMap = BasePtr->SharedDataInterfaceClassToDuplicatedCDOMap;
 			//EmitterPtr->bSimulationStagesEnabled = Handle.GetInstance()->bSimulationStagesEnabled;
-			if (Handle.GetIsEnabled() && (OwningEmitter == nullptr || OwningEmitter == Handle.GetInstance().Emitter)) // Don't need to copy the graph if we aren't going to use it.
+			//-TODO:Stateless:
+			if (Handle.GetIsEnabled() && Handle.GetInstance().Emitter && (OwningEmitter == nullptr || OwningEmitter == Handle.GetInstance().Emitter)) // Don't need to copy the graph if we aren't going to use it.
+			//-TODO:Stateless:
 			{
 				EmitterPtr->DeepCopyGraphs(Handle.GetInstance());
 			}
@@ -1395,7 +1402,9 @@ TSharedPtr<FNiagaraCompileRequestDuplicateDataBase, ESPMode::ThreadSafe> FNiagar
 			TArray<FNiagaraVariable> EncounterableEmitterVariables;
 			OwningSystemRequestData->GetDependentRequest(i)->GatherPreCompiledVariables(FString(), EncounterableEmitterVariables);
 
-			if (Handle.GetIsEnabled() && (OwningEmitter == nullptr || OwningEmitter == Handle.GetInstance().Emitter))
+			//-TODO:Stateless:
+			if (Handle.GetIsEnabled() && Handle.GetInstance().Emitter && (OwningEmitter == nullptr || OwningEmitter == Handle.GetInstance().Emitter))
+			//-TODO:Stateless:
 			{
 				TSharedPtr<FNiagaraGraphCachedDataBase, ESPMode::ThreadSafe> CachedTraversalEmitterData = Handle.GetInstance().Emitter->GetCachedTraversalData(Handle.GetInstance().Version);
 				TArray<FNiagaraVariable> StaticVariablesFromEmitter = StaticVariablesFromSystem;
diff --git a/Engine/Plugins/FX/Niagara/Source/NiagaraEditor/Private/ViewModels/NiagaraEmitterViewModel.cpp b/Engine/Plugins/FX/Niagara/Source/NiagaraEditor/Private/ViewModels/NiagaraEmitterViewModel.cpp
--- a/Engine/Plugins/FX/Niagara/Source/NiagaraEditor/Private/ViewModels/NiagaraEmitterViewModel.cpp
+++ b/Engine/Plugins/FX/Niagara/Source/NiagaraEditor/Private/ViewModels/NiagaraEmitterViewModel.cpp
@@ -3,6 +3,7 @@
 #include "ViewModels/NiagaraEmitterViewModel.h"
 #include "NiagaraEmitter.h"
 #include "NiagaraEmitterInstance.h"
+#include "NiagaraEmitterInstanceImpl.h"
 #include "NiagaraScriptSourceBase.h"
 #include "ViewModels/NiagaraSystemViewModel.h"
 #include "ViewModels/NiagaraScriptViewModel.h"
@@ -254,14 +255,16 @@ FText FNiagaraEmitterViewModel::GetStatsText() const
 {
 	if (Simulation.IsValid())
 	{
-		TSharedPtr<FNiagaraEmitterInstance, ESPMode::ThreadSafe> SimInstance = Simulation.Pin();
+		FNiagaraEmitterInstancePtr SimInstance = Simulation.Pin();
 		if (SimInstance.IsValid())
 		{
 			static const FNumberFormattingOptions FractionalFormatOptions = FNumberFormattingOptions()
 				.SetMinimumFractionalDigits(3)
 				.SetMaximumFractionalDigits(3);
 
-			if (!SimInstance->IsReadyToRun() || SimInstance->GetParentSystemInstance()->GetSystem()->HasOutstandingCompilationRequests())
+			FNiagaraEmitterInstanceImpl* StatefulSimInstance = SimInstance->AsStateful();
+			const bool bSystemCompiling = SimInstance->GetParentSystemInstance()->GetSystem()->HasOutstandingCompilationRequests();
+			if (bSystemCompiling || (StatefulSimInstance && !StatefulSimInstance->IsReadyToRun()))
 			{
 				return LOCTEXT("PendingCompile", "Compilation in progress...");
 			}
@@ -297,9 +300,10 @@ FText FNiagaraEmitterViewModel::GetStatsText() const
 				else
 				{
 					constexpr double Megabyte = 1024 * 1024;
+					const float TickTimeMS = FPlatformTime::ToMilliseconds(SimInstance->GetTickTimeCycles());
 					return FText::Format(StatsFormat,
 						FText::AsNumber(SimInstance->GetNumParticles()),
-						FText::AsNumber(SimInstance->GetTotalCPUTimeMS(), &FractionalFormatOptions),
+						FText::AsNumber(TickTimeMS, &FractionalFormatOptions),
 						FText::AsNumber(static_cast<double>(SimInstance->GetTotalBytesUsed()) / Megabyte, &FractionalFormatOptions),
 						ExecutionStateEnum->GetDisplayNameTextByValue(static_cast<int32>(SimInstance->GetExecutionState())));
 				}
@@ -326,14 +330,24 @@ UNiagaraSummaryViewViewModel* FNiagaraEmitterViewModel::GetSummaryHierarchyViewM
 
 const UNiagaraEmitterEditorData& FNiagaraEmitterViewModel::GetEditorData() const
 {
-	check(EmitterWeakPtr.Emitter.IsValid());
-	return *Cast<UNiagaraEmitterEditorData>(EmitterWeakPtr.GetEmitterData()->GetEditorData());
+	if (EmitterWeakPtr.Emitter.IsValid())
+	{
+		return *CastChecked<UNiagaraEmitterEditorData>(EmitterWeakPtr.GetEmitterData()->GetEditorData());
+	}
+	//-TODO:Stateless:Figure this out
+	return *GetDefault<UNiagaraEmitterEditorData>();
+	//-TODO:Stateless:Figure this out
 }
 
 UNiagaraEmitterEditorData& FNiagaraEmitterViewModel::GetEditorData()
 {
-	check(EmitterWeakPtr.IsValid());
-	return *Cast<UNiagaraEmitterEditorData>(EmitterWeakPtr.GetEmitterData()->GetEditorData());
+	if (EmitterWeakPtr.IsValid())
+	{
+		return *CastChecked<UNiagaraEmitterEditorData>(EmitterWeakPtr.GetEmitterData()->GetEditorData());
+	}
+	//-TODO:Stateless:Figure this out
+	return *GetMutableDefault<UNiagaraEmitterEditorData>();
+	//-TODO:Stateless:Figure this out
 }
 
 void FNiagaraEmitterViewModel::GetEmitterMessageStores(TArray<FNiagaraMessageSourceAndStore>& OutMessageStores)
