- FNiagaraEmitterInstance refactoring in a base class for different emitter types

diff --git a/Engine/Plugins/FX/Niagara/Source/NiagaraEditor/Private/ViewModels/NiagaraSystemViewModel.cpp b/Engine/Plugins/FX/Niagara/Source/NiagaraEditor/Private/ViewModels/NiagaraSystemViewModel.cpp
--- a/Engine/Plugins/FX/Niagara/Source/NiagaraEditor/Private/ViewModels/NiagaraSystemViewModel.cpp
+++ b/Engine/Plugins/FX/Niagara/Source/NiagaraEditor/Private/ViewModels/NiagaraSystemViewModel.cpp
@@ -26,6 +26,7 @@
 #include "NiagaraEmitterEditorData.h"
 #include "NiagaraEmitterFactoryNew.h"
 #include "NiagaraEmitterHandle.h"
+#include "NiagaraEmitterInstanceImpl.h"
 #include "NiagaraGraph.h"
 #include "NiagaraMessageManager.h"
 #include "NiagaraMessageUtilities.h"
@@ -68,7 +69,6 @@
 #include "ViewModels/NiagaraParameterPanelViewModel.h"
 #include "Widgets/Notifications/SNotificationList.h"
 
-
 DECLARE_CYCLE_STAT(TEXT("Niagara - SystemViewModel - CompileSystem"), STAT_NiagaraEditor_SystemViewModel_CompileSystem, STATGROUP_NiagaraEditor);
 
 #define LOCTEXT_NAMESPACE "NiagaraSystemViewModel"
@@ -579,7 +579,7 @@ TSharedPtr<FNiagaraEmitterHandleViewModel> FNiagaraSystemViewModel::AddEmitter(U
 	return NewEmitterHandleViewModel;
 }
 
-NIAGARAEDITOR_API TSharedPtr<FNiagaraEmitterHandleViewModel> FNiagaraSystemViewModel::AddEmptyEmitter()
+TSharedPtr<FNiagaraEmitterHandleViewModel> FNiagaraSystemViewModel::AddEmptyEmitter()
 {
 	UNiagaraEmitter* EmptyEmitter = NewObject<UNiagaraEmitter>(GetTransientPackage());
 	bool bAddDefaultModulesAndRenderers = false;
@@ -590,7 +590,6 @@ NIAGARAEDITOR_API TSharedPtr<FNiagaraEmitterHandleViewModel> FNiagaraSystemViewM
 	EmptyEmitter->SetFlags(RF_Transactional);
 	return AddEmitter(FVersionedNiagaraEmitter(EmptyEmitter, FGuid()));
 }
-
 TSharedPtr<FNiagaraEmitterHandleViewModel> FNiagaraSystemViewModel::AddEmitter(const FVersionedNiagaraEmitter& VersionedEmitter)
 {
 	return AddEmitter(*VersionedEmitter.Emitter, VersionedEmitter.Version);
@@ -1396,11 +1395,14 @@ TArray<UNiagaraGraph*> FNiagaraSystemViewModel::GetAllGraphs()
 		for (const TSharedRef<FNiagaraEmitterHandleViewModel>& EmitterHandleViewModel : EmitterHandleViewModels)
 		{
 			FNiagaraEmitterHandle* EmitterHandle = EmitterHandleViewModel->GetEmitterHandle();
-			if (EmitterHandle == nullptr)
+			//-TODO:Stateless:
+			FVersionedNiagaraEmitterData* EmitterData = EmitterHandle ? EmitterHandle->GetEmitterData() : nullptr;
+			if (EmitterData == nullptr)
+			//-TODO:Stateless:
 			{
 				continue;
 			}
-			UNiagaraGraph* Graph = Cast<UNiagaraScriptSource>(EmitterHandle->GetEmitterData()->GraphSource)->NodeGraph;
+			UNiagaraGraph* Graph = Cast<UNiagaraScriptSource>(EmitterData->GraphSource)->NodeGraph;
 			if (Graph)
 			{
 				OutGraphs.Add(Graph);
@@ -1410,9 +1412,12 @@ TArray<UNiagaraGraph*> FNiagaraSystemViewModel::GetAllGraphs()
 	else
 	{
 		FNiagaraEmitterHandle* EmitterHandle = GetEmitterHandleViewModels()[0]->GetEmitterHandle();
-		if (EmitterHandle != nullptr)
+		//-TODO:Stateless:
+		FVersionedNiagaraEmitterData* EmitterData = EmitterHandle ? EmitterHandle->GetEmitterData() : nullptr;
+		if (EmitterData != nullptr)
+		//-TODO:Stateless:
 		{
-			OutGraphs.Add(Cast<UNiagaraScriptSource>(EmitterHandle->GetEmitterData()->GraphSource)->NodeGraph);
+			OutGraphs.Add(Cast<UNiagaraScriptSource>(EmitterData->GraphSource)->NodeGraph);
 		}
 	}
 
@@ -1605,11 +1610,16 @@ void FNiagaraSystemViewModel::SendLastCompileMessageJobs() const
 	for (const FNiagaraEmitterHandle& Handle : EmitterHandles)
 	{
 		FVersionedNiagaraEmitter EmitterInSystem = Handle.GetInstance();
-		TArray<UNiagaraScript*> EmitterScripts;
-		EmitterInSystem.GetEmitterData()->GetScripts(EmitterScripts, false);
-		for (UNiagaraScript* EmitterScript : EmitterScripts)
+		//-TODO:Stateless:
+		if (FVersionedNiagaraEmitterData* EmitterData = EmitterInSystem.GetEmitterData())
+		//-TODO:Stateless:
 		{
-			ScriptsToGetCompileEventsFrom.Add(FNiagaraScriptAndOwningScriptNameString(EmitterScript, EmitterInSystem.Emitter->GetUniqueEmitterName()));
+			TArray<UNiagaraScript*> EmitterScripts;
+			EmitterData->GetScripts(EmitterScripts, false);
+			for (UNiagaraScript* EmitterScript : EmitterScripts)
+			{
+				ScriptsToGetCompileEventsFrom.Add(FNiagaraScriptAndOwningScriptNameString(EmitterScript, EmitterInSystem.Emitter->GetUniqueEmitterName()));
+			}
 		}
 	}
 
@@ -1757,7 +1767,9 @@ void FNiagaraSystemViewModel::TickCompileStatus()
 					ScriptsToCheckForStatus.Add(System->GetSystemUpdateScript());
 					for (const FNiagaraEmitterHandle& EmitterHandle : System->GetEmitterHandles())
 					{
-						if (EmitterHandle.GetIsEnabled())
+						//-TODO:Stateless:
+						if (EmitterHandle.GetIsEnabled() && EmitterHandle.GetEmitterData())
+						//-TODO:Stateless:
 						{
 							EmitterHandle.GetEmitterData()->GetScripts(ScriptsToCheckForStatus, true);
 						}
@@ -2046,7 +2058,7 @@ void FNiagaraSystemViewModel::RefreshEmitterHandleViewModels()
 	for (i = 0; i < GetSystem().GetNumEmitters(); ++i)
 	{
 		FNiagaraEmitterHandle* EmitterHandle = &GetSystem().GetEmitterHandle(i);
-		TSharedPtr<FNiagaraEmitterInstance, ESPMode::ThreadSafe> Simulation = SystemInstance ? SystemInstance->GetSimulationForHandle(*EmitterHandle) : nullptr;
+		FNiagaraEmitterInstancePtr Simulation = SystemInstance ? SystemInstance->GetSimulationForHandle(*EmitterHandle) : nullptr;
 		ValidEmitterHandleIds.Add(EmitterHandle->GetId());
 
 		TSharedPtr<FNiagaraEmitterHandleViewModel> ViewModel;
@@ -2991,11 +3003,14 @@ void FNiagaraSystemViewModel::UpdateEmitterFixedBounds()
 		}
 		FNiagaraEmitterHandle* SelectedEmitterHandle = EmitterHandleViewModel->GetEmitterHandle();
 		check(SelectedEmitterHandle);
-		for (TSharedRef<FNiagaraEmitterInstance, ESPMode::ThreadSafe>& EmitterInst : SystemInstance->GetEmitters())
+		for (const FNiagaraEmitterInstanceRef& EmitterInst : SystemInstance->GetEmitters())
 		{
 			if (&EmitterInst->GetEmitterHandle() == SelectedEmitterHandle && !EmitterInst->IsComplete())
 			{
-				EmitterInst->CalculateFixedBounds(PreviewComponent->GetComponentToWorld().Inverse());
+				if (FNiagaraEmitterInstanceImpl* EmitterInstImpl = EmitterInst->AsStateful())
+				{
+					EmitterInstImpl->CalculateFixedBounds(PreviewComponent->GetComponentToWorld().Inverse());
+				}
 			}
 		}
 	}
@@ -3247,6 +3262,12 @@ void FNiagaraSystemViewModel::RefreshAssetMessages()
 		{
 			FVersionedNiagaraEmitter VersionedEmitter = EmitterHandle.GetInstance();
 			UNiagaraEmitter* Emitter = VersionedEmitter.Emitter.Get();
+			//-TODO:Stateless:
+			if (Emitter == nullptr)
+			{
+				continue;
+			}	
+		//-TODO:Stateless:
 			PublishMessages(FNiagaraMessageSourceAndStore(*Emitter, Emitter->GetMessageStore()));
 			MessageSourceObjectKeys.Add(FObjectKey(Emitter));
 			VersionedEmitter.GetEmitterData()->GetScripts(Scripts, false);
@@ -3260,13 +3281,17 @@ void FNiagaraSystemViewModel::RefreshAssetMessages()
 		if (ensureMsgf(EmitterHandles.Num() == 1, TEXT("There was not exactly 1 Emitter Handle for the SystemViewModel in Emitter edit mode!")))
 		{
 			FVersionedNiagaraEmitter VersionedEmitter = EmitterHandles[0].GetInstance();
-			UNiagaraEmitter* Emitter = VersionedEmitter.Emitter.Get();
-			PublishMessages(FNiagaraMessageSourceAndStore(*Emitter, Emitter->GetMessageStore()));
-			MessageSourceObjectKeys.Add(FObjectKey(Emitter));
+			//-TODO:Stateless:
+			if (UNiagaraEmitter* Emitter = VersionedEmitter.Emitter.Get())
+			//-TODO:Stateless:
+			{
+				PublishMessages(FNiagaraMessageSourceAndStore(*Emitter, Emitter->GetMessageStore()));
+				MessageSourceObjectKeys.Add(FObjectKey(Emitter));
 
-			TArray<UNiagaraScript*> Scripts;
-			VersionedEmitter.GetEmitterData()->GetScripts(Scripts, false);
-			MessageSourceObjectKeys.Append(PublishScriptMessages(Scripts));
+				TArray<UNiagaraScript*> Scripts;
+				VersionedEmitter.GetEmitterData()->GetScripts(Scripts, false);
+				MessageSourceObjectKeys.Append(PublishScriptMessages(Scripts));
+			}
 		}
 	}
 	else
diff --git a/Engine/Plugins/FX/Niagara/Source/NiagaraEditor/Private/Widgets/SNiagaraSpreadsheetView.cpp b/Engine/Plugins/FX/Niagara/Source/NiagaraEditor/Private/Widgets/SNiagaraSpreadsheetView.cpp
--- a/Engine/Plugins/FX/Niagara/Source/NiagaraEditor/Private/Widgets/SNiagaraSpreadsheetView.cpp
+++ b/Engine/Plugins/FX/Niagara/Source/NiagaraEditor/Private/Widgets/SNiagaraSpreadsheetView.cpp
@@ -1048,12 +1048,11 @@ void SNiagaraSpreadsheetView::HandleTimeChange()
 					FName EntryName = NAME_None;
 					if (i != UISystemUpdate)
 					{
-						auto EmitterInstances = SystemInstance->GetEmitters();
-						for (auto EmitterInstance : EmitterInstances)
+						for (const FNiagaraEmitterInstanceRef& EmitterInstance : SystemInstance->GetEmitters())
 						{
-							if (SelectedEmitterHandle->GetEmitterHandle()->GetInstance() == EmitterInstance->GetCachedEmitter())
+							if (SelectedEmitterHandle->GetEmitterHandle()->GetInstance() == EmitterInstance->GetVersionedEmitter())
 							{
-								EntryName = EmitterInstance->GetCachedIDName();
+								EntryName = EmitterInstance->GetEmitterHandle().GetIdName();
 							}
 						}
 
diff --git a/Engine/Plugins/Runtime/GeometryCache/Source/GeometryCache/Private/NiagaraRendererGeometryCache.cpp b/Engine/Plugins/Runtime/GeometryCache/Source/GeometryCache/Private/NiagaraRendererGeometryCache.cpp
--- a/Engine/Plugins/Runtime/GeometryCache/Source/GeometryCache/Private/NiagaraRendererGeometryCache.cpp
+++ b/Engine/Plugins/Runtime/GeometryCache/Source/GeometryCache/Private/NiagaraRendererGeometryCache.cpp
@@ -109,7 +109,7 @@ int32 GetGeometryCacheIndex(TOptional<int32> DefaultCacheIndex, bool bCreateRand
 		return INDEX_NONE;
 	}
 	int32 CacheIndex = DefaultCacheIndex.Get(INDEX_NONE);
-	if (CacheIndex == INDEX_NONE && Emitter->GetCachedEmitterData()->bDeterminism)
+	if (CacheIndex == INDEX_NONE && Emitter->IsDeterministic())
 	{
 		int32 Seed = Properties->bAssignComponentsOnParticleID ? ParticleID : ParticleIndex;
 		FRandomStream RandomStream = FRandomStream(Seed * 907633515U); // multiply the seed, otherwise we get very poor randomness for small seeds
@@ -171,7 +171,7 @@ UGeometryCacheComponent* FNiagaraRendererGeometryCache::CreateOrGetPooledCompone
 			GeometryComponent->SetLooping(Properties->bIsLooping);
 			GeometryComponent->SetManualTick(true); // we want to tick the component with the delta time of the niagara sim
 
-			if (Emitter->GetCachedEmitterData()->bLocalSpace)
+			if (Emitter->IsLocalSpace())
 			{
 				GeometryComponent->SetAbsolute(false, false, false);
 			}
@@ -228,7 +228,7 @@ void FNiagaraRendererGeometryCache::PostSystemTick_GameThread(const UNiagaraRend
 	}
 
 #if WITH_EDITORONLY_DATA
-	if (SystemInstance->GetIsolateEnabled() && !Emitter->GetEmitterHandle().IsIsolated())
+	if (Emitter->IsDisabledFromIsolation())
 	{
 		ResetComponentPool(true);
 		return;
@@ -258,12 +258,12 @@ void FNiagaraRendererGeometryCache::PostSystemTick_GameThread(const UNiagaraRend
 	const float DefaultElapsedTime = ParameterStore.GetParameterValueOrDefault(Properties->ElapsedTimeBinding.GetParamMapBindableVariable(), 0.0f);
 	const float CurrentTime = AttachComponent->GetWorld()->GetRealTimeSeconds();
 
-	const FNiagaraLWCConverter LwcConverter = SystemInstance->GetLWCConverter(Emitter->GetCachedEmitterData()->bLocalSpace);
+	const FNiagaraLWCConverter LwcConverter = SystemInstance->GetLWCConverter(Emitter->IsLocalSpace());
 	const bool bIsRendererEnabled = IsRendererEnabled(InProperties, Emitter);
 
 	if (Properties->SourceMode == ENiagaraRendererSourceDataMode::Particles)
 	{
-		const FNiagaraDataSet& Data = Emitter->GetData();
+		const FNiagaraDataSet& Data = Emitter->GetParticleData();
 		const FNiagaraDataBuffer& ParticleData = Data.GetCurrentDataChecked();
 		FNiagaraDataSetReaderInt32<FNiagaraBool> EnabledReader = Properties->EnabledAccessor.GetReader(Data);
 		FNiagaraDataSetReaderInt32<int32> VisTagReader = Properties->VisTagAccessor.GetReader(Data);
diff --git a/Engine/Plugins/Runtime/HairStrands/Source/HairStrandsCore/Private/Niagara/NiagaraDataInterfaceHairStrands.cpp b/Engine/Plugins/Runtime/HairStrands/Source/HairStrandsCore/Private/Niagara/NiagaraDataInterfaceHairStrands.cpp
--- a/Engine/Plugins/Runtime/HairStrands/Source/HairStrandsCore/Private/Niagara/NiagaraDataInterfaceHairStrands.cpp
+++ b/Engine/Plugins/Runtime/HairStrands/Source/HairStrandsCore/Private/Niagara/NiagaraDataInterfaceHairStrands.cpp
@@ -955,7 +955,7 @@ void UNiagaraDataInterfaceHairStrands::SimCachePostReadFrame(void* OptionalPerIn
 {
 	for (TSharedRef<FNiagaraEmitterInstance, ESPMode::ThreadSafe> EmitterInst : SystemInstance->GetEmitters())
 	{
-		FNiagaraDataSet& EmitterDatas  = EmitterInst->GetData();
+		FNiagaraDataSet& EmitterDatas  = EmitterInst->GetParticleData();
 
 		const FNiagaraVariable NodePositionVariable(FNiagaraTypeDefinition::GetVec3Def(), FName("NodePosition"));
 		const int32 NodePositionIndex = EmitterDatas.GetCompiledData().Variables.IndexOfByKey(NodePositionVariable);
