Collection of fixes for niagara compilation
-application of DI's to the target script are not pulled from the compile request, instead the target script itself is used (before this was only the case with the DDC path)
-CompileForEdit made protected
-static variables are sorted for determinism to deal with whether the data is being pulled from the script or entirely populated by the parameter map
-scopes have been added to track timing of AsyncTasks
-better handling of failures with GPU scripts

diff --git a/Engine/Plugins/FX/Niagara/Source/Niagara/Classes/NiagaraScript.h b/Engine/Plugins/FX/Niagara/Source/Niagara/Classes/NiagaraScript.h
--- a/Engine/Plugins/FX/Niagara/Source/Niagara/Classes/NiagaraScript.h
+++ b/Engine/Plugins/FX/Niagara/Source/Niagara/Classes/NiagaraScript.h
@@ -1240,6 +1240,8 @@ public:
 
 	/** Ensure that the Script and its source graph are marked out of sync.*/
 	NIAGARA_API void MarkScriptAndSourceDesynchronized(FString Reason, const FGuid& VersionGuid);
+
+	NIAGARA_API bool IsScriptShaderSynchronized(const FGuid& VersionGuid = FGuid()) const;
 	
 	/** Request a synchronous compile for the script, possibly forcing it to compile.*/
 	NIAGARA_API void RequestCompile(const FGuid& ScriptVersion, bool bForceCompile = false);
diff --git a/Engine/Plugins/FX/Niagara/Source/Niagara/Classes/NiagaraSystem.h b/Engine/Plugins/FX/Niagara/Source/Niagara/Classes/NiagaraSystem.h
--- a/Engine/Plugins/FX/Niagara/Source/Niagara/Classes/NiagaraSystem.h
+++ b/Engine/Plugins/FX/Niagara/Source/Niagara/Classes/NiagaraSystem.h
@@ -501,6 +501,8 @@ public:
 	NIAGARA_API bool UsesEmitter(const FVersionedNiagaraEmitter& VersionedEmitter) const;
 	NIAGARA_API bool UsesScript(const UNiagaraScript* Script)const; 
 	NIAGARA_API void ForceGraphToRecompileOnNextCheck();
+	NIAGARA_API void SetCompileForEdit(bool bNewCompileForEdit);
+	FORCEINLINE bool GetCompileForEdit() const { return bCompileForEdit; }
 
 	static NIAGARA_API void RequestCompileForEmitter(const FVersionedNiagaraEmitter& InEmitter);
 	static NIAGARA_API void RecomputeExecutionOrderForEmitter(const FVersionedNiagaraEmitter& InEmitter);
@@ -515,10 +517,6 @@ public:
 	FORCEINLINE void SetBakeOutRapidIterationOnCook(bool bBakeOut) { bBakeOutRapidIteration = bBakeOut; bBakeOutRapidIterationOnCook = bBakeOut; }
 	FORCEINLINE void SetTrimAttributesOnCook(bool bTrim) { bTrimAttributes = bTrim; bTrimAttributesOnCook = bTrim; }
 
-	/** When enabled we compile for the edit path, which can result in slower system performance but faster editor responsiveness. */
-	UPROPERTY(transient)
-	uint32 bCompileForEdit : 1;
-
 protected:
 	/** When enable constant values are baked into the scripts while editing the system, this will increase iteration times but improve performance. */
 	UPROPERTY(EditAnywhere, AdvancedDisplay, Category = "Performance", meta=(DisplayName="Bake Rapid Iteration Parameters During Edit"))
@@ -554,6 +552,9 @@ protected:
 	/* When set the system needs to compile before it can be activated. */
 	uint32 bNeedsRequestCompile : 1;
 
+	/** When enabled we compile for the edit path, which can result in slower system performance but faster editor responsiveness. */
+	uint32 bCompileForEdit : 1;
+
 public:
 	/** Subscriptions to definitions of parameters. */
 	UPROPERTY()
diff --git a/Engine/Plugins/FX/Niagara/Source/Niagara/Private/NiagaraCompilationTypes.cpp b/Engine/Plugins/FX/Niagara/Source/Niagara/Private/NiagaraCompilationTypes.cpp
--- a/Engine/Plugins/FX/Niagara/Source/Niagara/Private/NiagaraCompilationTypes.cpp
+++ b/Engine/Plugins/FX/Niagara/Source/Niagara/Private/NiagaraCompilationTypes.cpp
@@ -298,14 +298,21 @@ public:
 		//	parameters just because things are disabled because the editor will forget the settings.
 		constexpr bool bAllowParameterRemoval = false;
 
+		auto GetValidTargetScript = [](UNiagaraScript* Script) -> UNiagaraScript*
+		{
+			return ::IsValid(Script) ? Script : nullptr;
+		};
+
 		for (FNiagaraSystemAsyncCompileResults::FCompileResultMap::TConstIterator ResultIt(CompileResults.CompileResultMap);
 			ResultIt;
 			++ResultIt)
 		{
-			UNiagaraScript* TargetScript = ResultIt->Key;
 			const FNiagaraScriptAsyncCompileData& ScriptCompileData = ResultIt->Value;
 
-			TargetScript->ApplyRapidIterationParameters(ScriptCompileData.RapidIterationParameters, bAllowParameterRemoval);
+			if (UNiagaraScript* TargetScript = GetValidTargetScript(ResultIt->Key))
+			{
+				TargetScript->ApplyRapidIterationParameters(ScriptCompileData.RapidIterationParameters, bAllowParameterRemoval);
+			}
 		}
 
 		// Now that the above code says they are all complete, go ahead and resolve them all at once.
@@ -313,55 +320,51 @@ public:
 			ResultIt;
 			++ResultIt)
 		{
-			UNiagaraScript* TargetScript = ResultIt->Key;
-			const FNiagaraScriptAsyncCompileData& ScriptCompileData = ResultIt->Value;
-
-			// because our compilation process includes the generation of rapid iteration parameters and static
-			// variables we need to generate the ExecutableDataId
-			// if we dirtied any RI parameters then we need to regenerate our CompilationId
-			FNiagaraVMExecutableDataId UpdatedCompileId;
-			TargetScript->ComputeVMCompilationId(UpdatedCompileId, FGuid());
-
-			if (ScriptCompileData.ExeData.IsValid())
+			if (UNiagaraScript* TargetScript = GetValidTargetScript(ResultIt->Key))
 			{
-				TMap<FName, UNiagaraDataInterface*> ObjectNameMap;
+				const FNiagaraScriptAsyncCompileData& ScriptCompileData = ResultIt->Value;
 
-				if (ScriptCompileData.bFromDerivedDataCache)
+				if (ScriptCompileData.ExeData.IsValid())
 				{
-					// if the data was pulled from the DDC then we'll need to generate the map from the source
-					const UNiagaraScriptSourceBase* ScriptSource = TargetScript->GetLatestSource();
-					ObjectNameMap = ScriptSource->ComputeObjectNameMap(*Options.System, TargetScript->GetUsage(), TargetScript->GetUsageId(), ScriptCompileData.UniqueEmitterName);
-				}
-				else
-				{
-					// if we actually generated our data will include a name map that we can use
-					Algo::Transform(ScriptCompileData.NamedDataInterfaces, ObjectNameMap, [](const TMap<FName, TObjectPtr<UNiagaraDataInterface>>::ElementType& Element)
+					// because our compilation process includes the generation of rapid iteration parameters and static
+					// variables we need to generate the ExecutableDataId
+					// if we dirtied any RI parameters then we need to regenerate our CompilationId
+					FNiagaraVMExecutableDataId UpdatedCompileId;
+					TargetScript->ComputeVMCompilationId(UpdatedCompileId, FGuid());
+
+					TMap<FName, UNiagaraDataInterface*> ObjectNameMap;
+
+					// The original implementation would generate DI references from the compilation data (unless things were pulled from
+					// the DDC).  We will always pull the data from the target scripts so that we can avoid any weird caching issues with
+					// our digested graphs, but it also should ensure more consistent behavior (the compilation should only depend on the
+					// aspects of DI that impact compilation, other changes that could have been made by the user shouldn't be erased
+					// when the compilation results are applied.
+					if (const UNiagaraScriptSourceBase* ScriptSource = TargetScript->GetLatestSource())
 					{
-						return TMap<FName, UNiagaraDataInterface*>::ElementType(Element.Key, Element.Value);
-					});
-				}
+						ObjectNameMap = ScriptSource->ComputeObjectNameMap(*Options.System, TargetScript->GetUsage(), TargetScript->GetUsageId(), ScriptCompileData.UniqueEmitterName);
+					}
 
-				constexpr bool bApplyRapidIterationParameters = false;
-				TargetScript->SetVMCompilationResults(
-					UpdatedCompileId,
-					*ScriptCompileData.ExeData,
-					ScriptCompileData.UniqueEmitterName,
-					ObjectNameMap,
-					bApplyRapidIterationParameters);
+					constexpr bool bApplyRapidIterationParameters = false;
+					TargetScript->SetVMCompilationResults(
+						UpdatedCompileId,
+						*ScriptCompileData.ExeData,
+						ScriptCompileData.UniqueEmitterName,
+						ObjectNameMap,
+						bApplyRapidIterationParameters);
 
-				if (!ScriptCompileData.CompiledShaders.IsEmpty())
-				{
-					for (const FNiagaraCompiledShaderInfo& ShaderMapInfo : ScriptCompileData.CompiledShaders)
+					if (!ScriptCompileData.CompiledShaders.IsEmpty())
 					{
-						TargetScript->SetComputeCompilationResults(
-							ShaderMapInfo.TargetPlatform,
-							ShaderMapInfo.ShaderPlatform,
-							ShaderMapInfo.FeatureLevel,
-							ScriptCompileData.ExeData->ShaderScriptParametersMetadata,
-							ShaderMapInfo.CompiledShader,
-							ShaderMapInfo.CompilationErrors);
+						for (const FNiagaraCompiledShaderInfo& ShaderMapInfo : ScriptCompileData.CompiledShaders)
+						{
+							TargetScript->SetComputeCompilationResults(
+								ShaderMapInfo.TargetPlatform,
+								ShaderMapInfo.ShaderPlatform,
+								ShaderMapInfo.FeatureLevel,
+								ScriptCompileData.ExeData->ShaderScriptParametersMetadata,
+								ShaderMapInfo.CompiledShader,
+								ShaderMapInfo.CompilationErrors);
+						}
 					}
-
 				}
 			}
 		}
diff --git a/Engine/Plugins/FX/Niagara/Source/Niagara/Private/NiagaraDebugHud.cpp b/Engine/Plugins/FX/Niagara/Source/Niagara/Private/NiagaraDebugHud.cpp
--- a/Engine/Plugins/FX/Niagara/Source/Niagara/Private/NiagaraDebugHud.cpp
+++ b/Engine/Plugins/FX/Niagara/Source/Niagara/Private/NiagaraDebugHud.cpp
@@ -1037,7 +1037,7 @@ void FNiagaraDebugHud::GatherSystemInfo()
 			SystemDebugInfo.SystemName = GetNameSafe(FXComponent->GetFXSystemAsset());
 		}
 	#if WITH_EDITORONLY_DATA
-		SystemDebugInfo.bCompileForEdit = NiagaraComponent ? NiagaraComponent->GetAsset()->bCompileForEdit : false;
+		SystemDebugInfo.bCompileForEdit = NiagaraComponent ? NiagaraComponent->GetAsset()->GetCompileForEdit() : false;
 	#endif
 		SystemDebugInfo.bShowInWorld = Settings.bSystemFilterEnabled && SystemDebugInfo.SystemName.MatchesWildcard(Settings.SystemFilter);
 		SystemDebugInfo.bPassesSystemFilter = !Settings.bSystemFilterEnabled || SystemDebugInfo.SystemName.MatchesWildcard(Settings.SystemFilter);
@@ -2255,7 +2255,7 @@ void FNiagaraDebugHud::DrawGpuComputeOverriew(class FNiagaraWorldManager* WorldM
 		const bool bShowDetailed = Settings.bSystemFilterEnabled && OwnerSystem->GetName().MatchesWildcard(Settings.SystemFilter);
 		SystemIt.Value().bShowDetailed = bShowDetailed;
 #if WITH_EDITORONLY_DATA
-		SystemIt.Value().bCompileForEdit = OwnerSystem->bCompileForEdit;
+		SystemIt.Value().bCompileForEdit = OwnerSystem->GetCompileForEdit();
 #endif
 		bHasDetailedView |= bShowDetailed;
 		bHasSimpleView |= !bShowDetailed;
diff --git a/Engine/Plugins/FX/Niagara/Source/Niagara/Private/NiagaraEmitter.cpp b/Engine/Plugins/FX/Niagara/Source/Niagara/Private/NiagaraEmitter.cpp
--- a/Engine/Plugins/FX/Niagara/Source/Niagara/Private/NiagaraEmitter.cpp
+++ b/Engine/Plugins/FX/Niagara/Source/Niagara/Private/NiagaraEmitter.cpp
@@ -2161,9 +2161,21 @@ bool FVersionedNiagaraEmitterData::AreAllScriptAndSourcesSynchronized() const
 		}
 	}
 
-	if (SimTarget == ENiagaraSimTarget::GPUComputeSim && GPUComputeScript->IsCompilable() && !GPUComputeScript->AreScriptAndSourceSynchronized())
+	if (SimTarget == ENiagaraSimTarget::GPUComputeSim && GPUComputeScript->IsCompilable())
 	{
-		return false;
+		if (!GPUComputeScript->AreScriptAndSourceSynchronized())
+		{
+			return false;
+		}
+
+		if (UNiagaraScript::AreGpuScriptsCompiledBySystem())
+		{
+			// need to also check if the shader resource is available for the GPUScript
+			if (!GPUComputeScript->IsScriptShaderSynchronized())
+			{
+				return false;
+			}
+		}
 	}
 
 	return true;
diff --git a/Engine/Plugins/FX/Niagara/Source/Niagara/Private/NiagaraScript.cpp b/Engine/Plugins/FX/Niagara/Source/Niagara/Private/NiagaraScript.cpp
--- a/Engine/Plugins/FX/Niagara/Source/Niagara/Private/NiagaraScript.cpp
+++ b/Engine/Plugins/FX/Niagara/Source/Niagara/Private/NiagaraScript.cpp
@@ -2607,6 +2607,26 @@ void UNiagaraScript::MarkScriptAndSourceDesynchronized(FString Reason, const FGu
 	}
 }
 
+bool UNiagaraScript::IsScriptShaderSynchronized(const FGuid& VersionGuid) const
+{
+	if (!AllowShaderCompiling())
+	{
+		return false;
+	}
+
+	if (!IsGPUScript())
+	{
+		return true;
+	}
+
+	if (!ScriptResource.IsValid() || !ScriptResource->IsShaderMapComplete())
+	{
+		return false;
+	}
+
+	return true;
+}
+
 bool UNiagaraScript::HandleVariableRenames(const TMap<FNiagaraVariable, FNiagaraVariable>& OldToNewVars, const FString& UniqueEmitterName)
 {
 	bool bConvertedAnything = false;
diff --git a/Engine/Plugins/FX/Niagara/Source/Niagara/Private/NiagaraSystem.cpp b/Engine/Plugins/FX/Niagara/Source/Niagara/Private/NiagaraSystem.cpp
--- a/Engine/Plugins/FX/Niagara/Source/Niagara/Private/NiagaraSystem.cpp
+++ b/Engine/Plugins/FX/Niagara/Source/Niagara/Private/NiagaraSystem.cpp
@@ -163,7 +163,6 @@ UNiagaraSystem::UNiagaraSystem(const FObjectInitializer& ObjectInitializer)
 #if WITH_EDITORONLY_DATA
 , LibraryVisibility(ENiagaraScriptLibraryVisibility::Unexposed)
 , TemplateSpecification(ENiagaraScriptTemplateSpecification::None)
-, bCompileForEdit(false)
 , bBakeOutRapidIteration(false)
 , bBakeOutRapidIterationOnCook(true)
 , bTrimAttributes(false)
@@ -172,6 +171,7 @@ UNiagaraSystem::UNiagaraSystem(const FObjectInitializer& ObjectInitializer)
 , bDisableDebugSwitches(false)
 , bDisableDebugSwitchesOnCook(true)
 , bNeedsRequestCompile(false)
+, bCompileForEdit(false)
 #endif
 , bSupportLargeWorldCoordinates(true)
 , bDisableExperimentalVM(false)
@@ -2521,6 +2521,15 @@ void UNiagaraSystem::ForceGraphToRecompileOnNextCheck()
 	}
 }
 
+void UNiagaraSystem::SetCompileForEdit(bool bNewCompileForEdit)
+{
+	if (bNewCompileForEdit != bCompileForEdit)
+	{
+		bCompileForEdit = bNewCompileForEdit;
+		bNeedsRequestCompile = true;
+	}
+}
+
 void UNiagaraSystem::WaitForCompilationComplete(bool bIncludingGPUShaders, bool bShowProgress)
 {
 	LLM_SCOPE(ELLMTag::Niagara);
diff --git a/Engine/Plugins/FX/Niagara/Source/NiagaraEditor/Private/NiagaraCompilationCopy.cpp b/Engine/Plugins/FX/Niagara/Source/NiagaraEditor/Private/NiagaraCompilationCopy.cpp
--- a/Engine/Plugins/FX/Niagara/Source/NiagaraEditor/Private/NiagaraCompilationCopy.cpp
+++ b/Engine/Plugins/FX/Niagara/Source/NiagaraEditor/Private/NiagaraCompilationCopy.cpp
@@ -425,6 +425,15 @@ void FNiagaraPrecompileData::FinishPrecompile(
 		}
 	}
 
+	{
+		auto VariableLess = [](const FNiagaraVariableBase& Lhs, const FNiagaraVariableBase& Rhs) -> bool
+		{
+			return Lhs.GetName().LexicalLess(Rhs.GetName());
+		};
+
+		StaticVariables.StableSort(VariableLess);
+	}
+
 	if (NumSimStageNodes)
 	{
 		CompileSimStageData.Reserve(NumSimStageNodes);
diff --git a/Engine/Plugins/FX/Niagara/Source/NiagaraEditor/Private/NiagaraCompilationTasks.cpp b/Engine/Plugins/FX/Niagara/Source/NiagaraEditor/Private/NiagaraCompilationTasks.cpp
--- a/Engine/Plugins/FX/Niagara/Source/NiagaraEditor/Private/NiagaraCompilationTasks.cpp
+++ b/Engine/Plugins/FX/Niagara/Source/NiagaraEditor/Private/NiagaraCompilationTasks.cpp
@@ -208,7 +208,19 @@ namespace NiagaraCompilationTasksImpl
 		}
 
 		FMemoryReaderView Ar(SharedBuffer.GetView(), true);
+
+		// Read the archive version from the header of the payload so we can setup our reader with the right version
+		FPackageFileVersion ArchiveVersion;
+		Ar << ArchiveVersion;
+
+		if (!ArchiveVersion.IsCompatible(GOldestLoadablePackageFileUEVersion))
+		{
+			UE_LOG(LogNiagaraEditor, Display, TEXT("Failed to validate FNiagaraVMExecutableData received from DDC, rejecting!  Reasons:\nDeprecated object version"));
+			return false;
+		}
+
 		FObjectAndNameAsStringProxyArchive SafeAr(Ar, false);
+		SafeAr.SetUEVer(ArchiveVersion);
 		OutExecData.SerializeData(SafeAr, true);
 
 		FString ValidationErrors;
@@ -237,6 +249,11 @@ namespace NiagaraCompilationTasksImpl
 
 		TArray<uint8> BinaryData;
 		FMemoryWriter Ar(BinaryData, true);
+
+		// include the archive version into the payload since we're using struct serialization for the ExecData
+		FPackageFileVersion ArchiveVersion = GPackageFileUEVersion;
+		Ar << ArchiveVersion;
+
 		FObjectAndNameAsStringProxyArchive SafeAr(Ar, false);
 		InExecData.SerializeData(SafeAr, true);
 
@@ -487,7 +504,6 @@ void FNiagaraSystemCompilationTask::Abort()
 	CompileCompletionEvent.Trigger();
 }
 
-
 const FNiagaraSystemCompilationTask::FEmitterInfo* FNiagaraSystemCompilationTask::FSystemInfo::EmitterInfoBySourceEmitter(int32 InSourceEmitterIndex) const
 {
 	return EmitterInfo.FindByPredicate([InSourceEmitterIndex](const FEmitterInfo& Info) -> bool
@@ -496,6 +512,14 @@ const FNiagaraSystemCompilationTask::FEmitterInfo* FNiagaraSystemCompilationTask
 	});
 }
 
+FNiagaraSystemCompilationTask::FEmitterInfo* FNiagaraSystemCompilationTask::FSystemInfo::EmitterInfoBySourceEmitter(int32 InSourceEmitterIndex)
+{
+	return EmitterInfo.FindByPredicate([InSourceEmitterIndex](const FEmitterInfo& Info) -> bool
+	{
+		return Info.SourceEmitterIndex == InSourceEmitterIndex;
+	});
+}
+
 FNiagaraSystemCompilationTask::FCompileGroupInfo::FCompileGroupInfo(int32 InSourceEmitterIndex)
 : SourceEmitterIndex(InSourceEmitterIndex)
 {}
@@ -517,6 +541,8 @@ bool FNiagaraSystemCompilationTask::FCompileGroupInfo::HasOutstandingCompileTask
 
 void FNiagaraSystemCompilationTask::FCompileGroupInfo::InstantiateCompileGraph(const FNiagaraSystemCompilationTask& ParentTask)
 {
+	TRACE_CPUPROFILER_EVENT_SCOPE(NiagaraAsyncTask_InstantiateGraph);
+
 	CompilationCopy = MakeShared<FNiagaraCompilationCopyData, ESPMode::ThreadSafe>();
 	FNiagaraCompilationCopyData& BasePtr = *CompilationCopy.Get();
 
@@ -720,6 +746,8 @@ void FNiagaraSystemCompilationTask::FCompileTaskInfo::CollectNamedDataInterfaces
 
 void FNiagaraSystemCompilationTask::FCompileTaskInfo::Translate(FNiagaraSystemCompilationTask* SystemCompileTask, const FCompileGroupInfo& GroupInfo)
 {
+	TRACE_CPUPROFILER_EVENT_SCOPE(NiagaraAsyncTask_Translate);
+
 	COOK_STAT(NiagaraSystemCookStats::ScriptTranslationCount++);
 
 	const FNiagaraPrecompileData* PrecompileData = nullptr;
@@ -789,7 +817,7 @@ void FNiagaraSystemCompilationTask::FCompileTaskInfo::IssueCompileVm(FNiagaraSys
 	NiagaraCompilationTasksImpl::BuildDebugGroupName(*PrecompileData, CompileOptions, DebugGroupName);
 
 	ScriptCompiler = MakeUnique<FHlslNiagaraCompiler>();
-	ScriptCompilationJobId = ScriptCompiler->CompileScriptVM(DebugGroupName, CompileOptions, TranslateResults, TranslateOutput, TranslatedHlsl);
+	ScriptCompilationJobId = ScriptCompiler->CompileScriptVM(DebugGroupName, CompileOptions, TranslateResults, TranslateOutput, TranslatedHlsl, SystemCompileTask->NiagaraShaderType);
 }
 
 void FNiagaraSystemCompilationTask::FCompileTaskInfo::IssueTranslateGpu(FNiagaraSystemCompilationTask* SystemCompileTask, const FCompileGroupInfo& GroupInfo)
@@ -817,8 +845,11 @@ void FNiagaraSystemCompilationTask::FCompileTaskInfo::IssueCompileGpu(FNiagaraSy
 {
 	COOK_STAT(NiagaraSystemCookStats::GpuScriptCompileCount++);
 
-	// we are going to populate a FNiagaraShaderMap.  This involves it's own little DDC hop as we store
-	// the shader in a separate bucket
+	// if we don't have a valid shader file, then there's no point in issuing a compilation
+	if (!TranslateResults.bHLSLGenSucceeded)
+	{
+		return;
+	}
 
 	const FNiagaraPrecompileData* PrecompileData = nullptr;
 	const FNiagaraCompilationCopyData* CompilationCopyData = nullptr;
@@ -861,6 +892,8 @@ TOptional<FNiagaraCompileResults> FNiagaraSystemCompilationTask::FCompileTaskInf
 /** Returns true if the task has valid results */
 bool FNiagaraSystemCompilationTask::FCompileTaskInfo::RetrieveCompilationResult(bool bWait)
 {
+	TRACE_CPUPROFILER_EVENT_SCOPE(NiagaraAsyncTask_RetrieveResult);
+
 	check(CompileResultsReadyEvent.IsValid() && CompileResultsReadyEvent->IsValid());
 
 	if (CompileResultsReadyEvent->IsCompleted())
@@ -869,7 +902,6 @@ bool FNiagaraSystemCompilationTask::FCompileTaskInfo::RetrieveCompilationResult(
 	}
 	check(!ExeData.IsValid());
 	check(ScriptCompiler.IsValid());
-	check(ScriptCompilationJobId != INDEX_NONE);
 
 	// in cases where asynchronous shader compiling isn't allowed we'll simply block execution here till the task
 	// is complete
@@ -948,7 +980,6 @@ void FNiagaraSystemCompilationTask::FCompileTaskInfo::RetrieveTranslateResult()
 
 bool FNiagaraSystemCompilationTask::FCompileTaskInfo::RetrieveShaderMap(bool bWait)
 {
-	check(ShaderMapCompiler.IsValid());
 	check(CompileResultsReadyEvent.IsValid() && CompileResultsReadyEvent->IsValid());
 
 	if (CompileResultsReadyEvent->IsCompleted())
@@ -956,6 +987,13 @@ bool FNiagaraSystemCompilationTask::FCompileTaskInfo::RetrieveShaderMap(bool bWa
 		return true;
 	}
 
+	// if we don't have a shader map compiler, then just return the translation results
+	if (!ShaderMapCompiler.IsValid())
+	{
+		RetrieveTranslateResult();
+		return true;
+	}
+
 	ExeData = ShaderMapCompiler->ReadScriptMetaData();
 
 	if (ShaderMapCompiler->ProcessCompileResults(bWait))
@@ -969,6 +1007,8 @@ bool FNiagaraSystemCompilationTask::FCompileTaskInfo::RetrieveShaderMap(bool bWa
 
 void FNiagaraSystemCompilationTask::FCompileTaskInfo::GenerateOptimizedVMByteCode(FNiagaraSystemCompilationTask* SystemCompileTask, const FCompileGroupInfo& GroupInfo)
 {
+	TRACE_CPUPROFILER_EVENT_SCOPE(NiagaraAsyncTask_OptimizeVM);
+
 #if VECTORVM_SUPPORTS_LEGACY
 	const bool bExperimentalVMDisabled = CompileOptions.AdditionalDefines.Contains(FNiagaraCompileOptions::ExperimentalVMDisabled);
 
@@ -1024,8 +1064,9 @@ void FNiagaraSystemCompilationTask::FCompileTaskInfo::ProcessCompilation(FNiagar
 	}
 
 	// we need to update our ComputeShaderTasks with the generated shader map so that we can put the data to the DDC
-	if (ScriptCompileType == EScriptCompileType::CompileForGpu)
+	if (ShaderMapCompiler.IsValid())
 	{
+		check(ScriptCompileType == EScriptCompileType::CompileForGpu);
 		for (int32 ComputeShaderTaskIndex : ComputeShaderTaskIndices)
 		{
 			FCompileComputeShaderTaskInfo& ShaderTaskInfo = SystemCompileTask->CompileComputeShaderTasks[ComputeShaderTaskIndex];
@@ -1283,6 +1324,8 @@ void FNiagaraSystemCompilationTask::FDispatchDataCachePutRequests::Launch(FNiaga
 
 void FNiagaraSystemCompilationTask::DigestSystemInfo()
 {
+	TRACE_CPUPROFILER_EVENT_SCOPE(NiagaraAsyncTask_DigestSystem);
+
 	using namespace NiagaraCompilationTasksImpl;
 
 	FNiagaraDigestDatabase& DigestDatabase = FNiagaraDigestDatabase::Get();
@@ -1415,7 +1458,7 @@ void FNiagaraSystemCompilationTask::DigestShaderInfo(const ITargetPlatform* InTa
 	NiagaraShaderType = InShaderType;
 }
 
-void FNiagaraSystemCompilationTask::AddScript(int32 EmitterIndex, UNiagaraScript* Script, const FNiagaraVMExecutableDataId& CompileId, bool bRequiresCompilation, TConstArrayView<FShaderCompileRequest> ShaderRequests)
+void FNiagaraSystemCompilationTask::AddScript(int32 SourceEmitterIndex, UNiagaraScript* Script, const FNiagaraVMExecutableDataId& CompileId, bool bRequiresCompilation, TConstArrayView<FShaderCompileRequest> ShaderRequests)
 {
 	if (!Script)
 	{
@@ -1427,10 +1470,10 @@ void FNiagaraSystemCompilationTask::AddScript(int32 EmitterIndex, UNiagaraScript
 		FVersionedNiagaraScriptData* ScriptData = Script->GetScriptData(CompileId.ScriptVersionID);
 		if (ensure(ScriptData))
 		{
-			FCompileGroupInfo* GroupInfo = GetCompileGroupInfo(EmitterIndex);
+			FCompileGroupInfo* GroupInfo = GetCompileGroupInfo(SourceEmitterIndex);
 			if (!GroupInfo)
 			{
-				GroupInfo = &CompileGroups.Emplace_GetRef(EmitterIndex);
+				GroupInfo = &CompileGroups.Emplace_GetRef(SourceEmitterIndex);
 			}
 
 			TArray<ENiagaraScriptUsage> DuplicateUsages;
@@ -1471,10 +1514,8 @@ void FNiagaraSystemCompilationTask::AddScript(int32 EmitterIndex, UNiagaraScript
 
 			if (TaskInfo.ScriptCompileType == EScriptCompileType::CompileForGpu)
 			{
-				check(SystemInfo.EmitterInfo.IsValidIndex(EmitterIndex));
-				const FEmitterInfo& EmitterInfo = SystemInfo.EmitterInfo[EmitterIndex];
-
-				if (EmitterInfo.Enabled)
+				const FEmitterInfo* EmitterInfo = SystemInfo.EmitterInfoBySourceEmitter(SourceEmitterIndex);
+				if (EmitterInfo && EmitterInfo->Enabled)
 				{
 					TaskInfo.ComputeShaderTaskIndices.Reserve(ShaderRequests.Num());
 					for (const FShaderCompileRequest& ShaderRequest : ShaderRequests)
@@ -1498,9 +1539,9 @@ void FNiagaraSystemCompilationTask::AddScript(int32 EmitterIndex, UNiagaraScript
 
 	FVersionedNiagaraEmitterData* EmitterData = nullptr;
 
-	if (EmitterIndex != INDEX_NONE)
+	if (SourceEmitterIndex != INDEX_NONE)
 	{
-		EmitterData = System_GT->GetEmitterHandle(EmitterIndex).GetEmitterData();
+		EmitterData = System_GT->GetEmitterHandle(SourceEmitterIndex).GetEmitterData();
 	}
 
 	FScriptInfo& Info = DigestedScriptInfo.Add(Script);
@@ -1533,7 +1574,7 @@ void FNiagaraSystemCompilationTask::AddScript(int32 EmitterIndex, UNiagaraScript
 
 	Info.Usage = Script->GetUsage();
 	Info.UsageId = Script->GetUsageId();
-	Info.SourceEmitterIndex = EmitterIndex;
+	Info.SourceEmitterIndex = SourceEmitterIndex;
 	for (UNiagaraScript* DependentScript : NiagaraCompilationTasksImpl::FindDependentScripts(System_GT.Get(), EmitterData, Script))
 	{
 		Info.DependentScripts.AddUnique(DependentScript);
@@ -1707,6 +1748,8 @@ void FNiagaraSystemCompilationTask::WaitTillCachePutCompletion()
 
 void FNiagaraSystemCompilationTask::Precompile()
 {
+	TRACE_CPUPROFILER_EVENT_SCOPE(NiagaraAsyncTask_Precompile);
+
 	SystemPrecompileData = MakeShared<FNiagaraPrecompileData, ESPMode::ThreadSafe>();
 
 	SystemPrecompileData->SharedCompileDataInterfaceData = MakeShared<TArray<FNiagaraPrecompileData::FCompileDataInterfaceData>>();
@@ -1910,6 +1953,7 @@ struct FNiagaraSystemCompilationTask::FCollectStaticVariablesTaskBuilder
 				BuilderTask->CompilationTask = &CompilationTask;
 				BuilderTask->ConstantResolver = EmitterInfo.ConstantResolver;
 				BuilderTask->GraphContext = EmitterGraph;
+				BuilderTask->UniqueEmitterName = EmitterInfo.UniqueEmitterName;
 
 				// the initial static variables needs to come from gathering data from the Emitter
 				// but also what has been calculated in the SystemInfo
@@ -1954,6 +1998,7 @@ struct FNiagaraSystemCompilationTask::FCollectStaticVariablesTaskBuilder
 		const FNiagaraSystemCompilationTask* CompilationTask = nullptr;
 		FNiagaraFixedConstantResolver ConstantResolver;
 		const FNiagaraCompilationGraph* GraphContext = nullptr;
+		FString UniqueEmitterName;
 		TArray<const FNiagaraCompilationNodeOutput*> OutputNodes;
 		TArray<FNiagaraVariable> InitialStaticVariables;
 		TArray<FName> StageNames;
@@ -2292,9 +2337,10 @@ UE::Tasks::FTask FNiagaraSystemCompilationTask::BuildRapidIterationParametersAsy
 			}, CollectStaticVariableTasks.FindRef(EmitterInfo ? EmitterInfo->DigestedEmitterIndex : INDEX_NONE)));
 		}
 
-		// lastly we need a single task to copy over parameters between dependent scripts
+		// We need a single task to copy over parameters between dependent scripts and make sure that the static variable
+		// list is up to date
 		TMap<TObjectKey<UNiagaraScript>, FScriptInfo>* ScriptInfoMapPtr = &DigestedScriptInfo;
-		return Launch(UE_SOURCE_LOCATION, [ScriptInfoMapPtr]
+		return Launch(UE_SOURCE_LOCATION, [ScriptInfoMapPtr, this]
 		{
 			for (TMap<TObjectKey<UNiagaraScript>, FScriptInfo>::ElementType& CurrentIt : *ScriptInfoMapPtr)
 			{
@@ -2304,6 +2350,25 @@ UE::Tasks::FTask FNiagaraSystemCompilationTask::BuildRapidIterationParametersAsy
 					FScriptInfo& DstScriptInfo = ScriptInfoMapPtr->FindChecked(DependentScriptKey);
 					SrcScriptInfo.RapidIterationParameters.CopyParametersTo(DstScriptInfo.RapidIterationParameters, false, FNiagaraParameterStore::EDataInterfaceCopyMethod::None);
 				}
+
+				// make sure that all rapid iteration parameters on the script have made it to the system/emitter info
+				for (const FNiagaraVariableWithOffset& ScriptRI : SrcScriptInfo.RapidIterationParameters.ReadParameterVariables())
+				{
+					if (ScriptRI.GetType().IsStatic())
+					{
+						FNiagaraVariable StaticVariable = FNiagaraVariable(ScriptRI);
+						StaticVariable.SetData(SrcScriptInfo.RapidIterationParameters.GetParameterData(ScriptRI.Offset));
+
+						if (SrcScriptInfo.SourceEmitterIndex == INDEX_NONE)
+						{
+							SystemInfo.StaticVariableResults.AddUnique(StaticVariable);
+						}
+						else if (FEmitterInfo* EmitterInfo = SystemInfo.EmitterInfoBySourceEmitter(SrcScriptInfo.SourceEmitterIndex))
+						{
+							EmitterInfo->StaticVariableResults.AddUnique(StaticVariable);
+						}
+					}
+				}
 			}
 		}, PendingTasks);
 	}
@@ -2359,7 +2424,7 @@ void FNiagaraSystemCompilationTask::IssueCompilationTasks()
 {
 	using namespace UE::Tasks;
 
-	TArray<FTask> IssueCompilationTasks;
+	TArray<FTask> IssuedCompilationTasks;
 	TArray<FTaskEvent> CompilationTasks;
 
 	if (HasOutstandingCompileTasks())
@@ -2391,7 +2456,7 @@ void FNiagaraSystemCompilationTask::IssueCompilationTasks()
 						CompileTask.CompileResultsReadyEvent = MakeUnique<FTaskEvent>(UE_SOURCE_LOCATION);
 						CompileTask.CompileResultsProcessedEvent = MakeUnique<FTaskEvent>(UE_SOURCE_LOCATION);
 
-						IssueCompilationTasks.Add(Launch(UE_SOURCE_LOCATION, [this, &GroupInfo, &CompileTask]
+						IssuedCompilationTasks.Add(Launch(UE_SOURCE_LOCATION, [this, &GroupInfo, &CompileTask]
 						{
 							COOK_STAT(auto Timer = NiagaraSystemCookStats::UsageStats.TimeSyncWork(); Timer.TrackCyclesOnly(););
 
@@ -2442,7 +2507,7 @@ void FNiagaraSystemCompilationTask::IssueCompilationTasks()
 	Launch(UE_SOURCE_LOCATION, [this]
 	{
 		CompilationState = EState::WaitingForProcessing;
-	}, IssueCompilationTasks);
+	}, IssuedCompilationTasks);
 
 	Launch(UE_SOURCE_LOCATION, [this]
 	{
@@ -2465,6 +2530,8 @@ bool FNiagaraSystemCompilationTask::HasOutstandingCompileTasks() const
 
 UE::Tasks::FTask FNiagaraSystemCompilationTask::BeginTasks()
 {
+	TRACE_CPUPROFILER_EVENT_SCOPE(NiagaraAsyncTask_BeginTasks);
+
 	using namespace UE::Tasks;
 
 	COOK_STAT(auto Timer = NiagaraSystemCookStats::UsageStats.TimeSyncWork(); Timer.TrackCyclesOnly(););
diff --git a/Engine/Plugins/FX/Niagara/Source/NiagaraEditor/Private/NiagaraCompilationTasks.h b/Engine/Plugins/FX/Niagara/Source/NiagaraEditor/Private/NiagaraCompilationTasks.h
--- a/Engine/Plugins/FX/Niagara/Source/NiagaraEditor/Private/NiagaraCompilationTasks.h
+++ b/Engine/Plugins/FX/Niagara/Source/NiagaraEditor/Private/NiagaraCompilationTasks.h
@@ -105,6 +105,7 @@ struct FNiagaraSystemCompilationTask
 		bool bDisableDebugSwitches;
 
 		const FEmitterInfo* EmitterInfoBySourceEmitter(int32 InSourceEmitterIndex) const;
+		FEmitterInfo* EmitterInfoBySourceEmitter(int32 InSourceEmitterIndex);
 	};
 
 	struct FCompileGroupInfo
@@ -231,6 +232,7 @@ struct FNiagaraSystemCompilationTask
 	FNiagaraShaderType* NiagaraShaderType = nullptr;
 
 	int32 TasksAwaitingDDCGetResults = 0;
+	bool bCompileForEdit = false;
 
 	void Tick();
 	bool Poll(FNiagaraSystemAsyncCompileResults& Results) const;
@@ -249,8 +251,6 @@ private:
 	FNiagaraSystemCompilationTask() = delete;
 
 	bool GetStageName(int32 EmitterIndex, const FNiagaraCompilationNodeOutput* OutputNode, FName& OutStageName) const;
-	TArray<FNiagaraVariable> CollectStaticVariables(const FScriptInfo& ScriptInfo, const FNiagaraCompilationGraph& ScriptGraph, TConstArrayView<FNiagaraVariable> InitialStaticVariables) const;
-	void CollectStaticVariables();
 
 	TWeakObjectPtr<UNiagaraSystem> System_GT;
 
diff --git a/Engine/Plugins/FX/Niagara/Source/NiagaraEditor/Private/NiagaraCompiler.cpp b/Engine/Plugins/FX/Niagara/Source/NiagaraEditor/Private/NiagaraCompiler.cpp
--- a/Engine/Plugins/FX/Niagara/Source/NiagaraEditor/Private/NiagaraCompiler.cpp
+++ b/Engine/Plugins/FX/Niagara/Source/NiagaraEditor/Private/NiagaraCompiler.cpp
@@ -1925,102 +1925,85 @@ int32 FHlslNiagaraCompiler::CompileScript(const FStringView GroupName, const FNi
 	return JobID;
 }
 
-uint32 FHlslNiagaraCompiler::CompileScriptVM(const FStringView GroupName, const FNiagaraCompileOptions& InOptions, const FNiagaraTranslateResults& InTranslateResults, const FNiagaraTranslatorOutput& TranslatorOutput, const FString& TranslatedHLSL)
+uint32 FHlslNiagaraCompiler::CompileScriptVM(const FStringView GroupName, const FNiagaraCompileOptions& InOptions, const FNiagaraTranslateResults& InTranslateResults, const FNiagaraTranslatorOutput& TranslatorOutput, const FString& TranslatedHLSL, FNiagaraShaderType* NiagaraShaderType)
 {
 	check(!InOptions.IsGpuScript() || !UNiagaraScript::IsParticleScript(InOptions.TargetUsage));
 
 	CompileResults.Data = MakeShared<FNiagaraVMExecutableData>();
 
-	CompileResults.Data->LastHlslTranslation = TEXT("");
-
-	FShaderCompilerInput Input;
-	Input.Target = FShaderTarget(SF_Compute, SP_PCD3D_SM5);
-	Input.VirtualSourceFilePath = TEXT("/Plugin/FX/Niagara/Private/NiagaraEmitterInstanceShader.usf");
-	Input.EntryPointName = TEXT("SimulateMain");
-	Input.Environment.SetDefine(TEXT("VM_SIMULATION"), 1);
-	Input.Environment.SetDefine(TEXT("COMPUTESHADER"), 1);
-	Input.Environment.SetDefine(TEXT("PIXELSHADER"), 0);
-	Input.Environment.SetDefine(TEXT("DOMAINSHADER"), 0);
-	Input.Environment.SetDefine(TEXT("HULLSHADER"), 0);
-	Input.Environment.SetDefine(TEXT("VERTEXSHADER"), 0);
-	Input.Environment.SetDefine(TEXT("GEOMETRYSHADER"), 0);
-	Input.Environment.SetDefine(TEXT("MESHSHADER"), 0);
-	Input.Environment.SetDefine(TEXT("AMPLIFICATIONSHADER"), 0);
-	Input.Environment.IncludeVirtualPathToContentsMap.Add(TEXT("/Engine/Generated/NiagaraEmitterInstance.ush"), TranslatedHLSL);
-	Input.DebugInfoFlags = GShaderCompilingManager->GetDumpShaderDebugInfoFlags();
-	Input.DumpDebugInfoRootPath = GShaderCompilingManager->GetAbsoluteShaderDebugInfoDirectory() / TEXT("VM");
-	Input.DebugGroupName = GroupName;
-	Input.DebugExtension.Empty();
-	Input.DumpDebugInfoPath.Empty();
-
-	CompileResults.DumpDebugInfoPath = Input.DumpDebugInfoPath;
-
-	uint32 JobID = FShaderCommonCompileJob::GetNextJobId();
-	CompilationJob = MakeUnique<FNiagaraCompilerJob>();
-	CompilationJob->TranslatorOutput = TranslatorOutput;
-
-	CompileResults.bVMSucceeded = (CompilationJob->TranslatorOutput.Errors.Len() == 0) && (TranslatedHLSL.Len() > 0) && !InTranslateResults.NumErrors;
-
+	CompileResults.bVMSucceeded = (TranslatorOutput.Errors.Len() == 0) && (TranslatedHLSL.Len() > 0) && !InTranslateResults.NumErrors;
 	CompileResults.AppendCompileEvents(MakeArrayView(InTranslateResults.CompileEvents));
 	CompileResults.Data->LastCompileEvents.Append(InTranslateResults.CompileEvents);
 	CompileResults.Data->ExternalDependencies = InTranslateResults.CompileDependencies;
 	CompileResults.Data->CompileTags = InTranslateResults.CompileTags;
 	CompileResults.Data->CompileTagsEditorOnly = InTranslateResults.CompileTagsEditorOnly;
+	CompileResults.Data->LastHlslTranslation = TranslatedHLSL;
+	CompileResults.DumpDebugInfoPath.Reset();
+
+	CompilationJob = MakeUnique<FNiagaraCompilerJob>();
+	CompilationJob->TranslatorOutput = TranslatorOutput;
 	CompilationJob->TranslatorOutput.ScriptData.LastHlslTranslation = TranslatedHLSL;
 	CompilationJob->TranslatorOutput.ScriptData.ExternalDependencies = InTranslateResults.CompileDependencies;
 	CompilationJob->TranslatorOutput.ScriptData.CompileTags = InTranslateResults.CompileTags;
 	CompilationJob->TranslatorOutput.ScriptData.CompileTagsEditorOnly = InTranslateResults.CompileTagsEditorOnly;
+	CompilationJob->StartTime = FPlatformTime::Seconds();
+	CompilationJob->CompileResults = CompileResults;
 
-	bool bJobScheduled = false;
-	if (CompileResults.bVMSucceeded)
-	{
-		SCOPE_CYCLE_COUNTER(STAT_NiagaraEditor_HlslCompiler_CompileShader_VectorVM);
-		CompilationJob->StartTime = FPlatformTime::Seconds();
-
-		FShaderType* NiagaraShaderType = nullptr;
-		for (TLinkedList<FShaderType*>::TIterator ShaderTypeIt(FShaderType::GetTypeList()); ShaderTypeIt; ShaderTypeIt.Next())
-		{
-			if (FNiagaraShaderType* ShaderType = ShaderTypeIt->GetNiagaraShaderType())
+	if (CompileResults.bVMSucceeded && NiagaraShaderType)
+	{
+		const uint32 JobID = FShaderCommonCompileJob::GetNextJobId();
+		TRefCountPtr<FShaderCompileJob> Job = GShaderCompilingManager->PrepareShaderCompileJob(JobID, FShaderCompileJobKey(NiagaraShaderType), EShaderCompileJobPriority::Normal);
+		if (Job)
+		{
+			Job->Input.Target = FShaderTarget(SF_Compute, SP_PCD3D_SM5);
+			Job->Input.VirtualSourceFilePath = TEXT("/Plugin/FX/Niagara/Private/NiagaraEmitterInstanceShader.usf");
+			Job->Input.EntryPointName = TEXT("SimulateMain");
+			Job->Input.Environment.SetDefine(TEXT("VM_SIMULATION"), 1);
+			Job->Input.Environment.SetDefine(TEXT("COMPUTESHADER"), 1);
+			Job->Input.Environment.SetDefine(TEXT("PIXELSHADER"), 0);
+			Job->Input.Environment.SetDefine(TEXT("DOMAINSHADER"), 0);
+			Job->Input.Environment.SetDefine(TEXT("HULLSHADER"), 0);
+			Job->Input.Environment.SetDefine(TEXT("VERTEXSHADER"), 0);
+			Job->Input.Environment.SetDefine(TEXT("GEOMETRYSHADER"), 0);
+			Job->Input.Environment.SetDefine(TEXT("MESHSHADER"), 0);
+			Job->Input.Environment.SetDefine(TEXT("AMPLIFICATIONSHADER"), 0);
+			Job->Input.Environment.IncludeVirtualPathToContentsMap.Add(TEXT("/Engine/Generated/NiagaraEmitterInstance.ush"), TranslatedHLSL);
+			Job->Input.DebugInfoFlags = GShaderCompilingManager->GetDumpShaderDebugInfoFlags();
+			Job->Input.DumpDebugInfoRootPath = GShaderCompilingManager->GetAbsoluteShaderDebugInfoDirectory() / TEXT("VM");
+			Job->Input.DumpDebugInfoPath = CompileResults.DumpDebugInfoPath;
+			Job->Input.DebugGroupName = GroupName;
+			Job->Input.DebugExtension.Empty();
+			Job->Input.ShaderFormat = FName(TEXT("VVM_1_0"));
+
+			//TODO: This is normally invoked by GlobalBeginCompileShader, which is not called in this path. Should it be?
+			if (const IShaderFormat* VVMShaderFormat = GetTargetPlatformManagerRef().FindShaderFormat(Job->Input.ShaderFormat))
 			{
-				NiagaraShaderType = ShaderType;
-				break;
+				VVMShaderFormat->ModifyShaderCompilerInput(Job->Input);
 			}
-		}
-		if (NiagaraShaderType)
-		{
-			TRefCountPtr<FShaderCompileJob> Job = GShaderCompilingManager->PrepareShaderCompileJob(JobID, FShaderCompileJobKey(NiagaraShaderType), EShaderCompileJobPriority::Normal);
-			if (Job)
+
+			if (GNiagaraSkipVectorVMBackendOptimizations != 0)
 			{
-				TArray<FShaderCommonCompileJobPtr> NewJobs;
-				CompilationJob->ShaderCompileJob = Job;
-				Input.ShaderFormat = FName(TEXT("VVM_1_0"));
-				if (GNiagaraSkipVectorVMBackendOptimizations != 0)
-				{
-					Input.Environment.CompilerFlags.Add(CFLAG_SkipOptimizations);
-				}
-				Job->Input = Input;
-				NewJobs.Add(FShaderCommonCompileJobPtr(Job));
+				Job->Input.Environment.CompilerFlags.Add(CFLAG_SkipOptimizations);
+			}
 
-				GShaderCompilingManager->SubmitJobs(NewJobs, FString(), FString());
+			if (GShaderCompilingManager->GetDumpShaderDebugInfo() == FShaderCompilingManager::EDumpShaderDebugInfo::Always)
+			{
+				Job->Input.DumpDebugInfoPath = GShaderCompilingManager->CreateShaderDebugInfoPath(Job->Input);
+				CompileResults.DumpDebugInfoPath = Job->Input.DumpDebugInfoPath;
 			}
-			bJobScheduled = true;
-		}
-	}
-	CompileResults.Data->LastHlslTranslation = TranslatedHLSL;
 
-	if (!bJobScheduled)
-	{
-		CompileResults.Data->ByteCode.Reset();
-		CompileResults.Data->Attributes.Empty();
-		CompileResults.Data->Parameters.Empty();
-		CompileResults.Data->InternalParameters.Empty();
-		CompileResults.Data->DataInterfaceInfo.Empty();
-		CompileResults.Data->UObjectInfos.Empty();
+			CompilationJob->ShaderCompileJob = Job;
 
+			TArray<FShaderCommonCompileJobPtr> NewJobs;
+			NewJobs.Emplace(Job);
+
+			GShaderCompilingManager->SubmitJobs(NewJobs, FString(), FString());
+
+			return JobID;
+		}
 	}
-	CompilationJob->CompileResults = CompileResults;
 
-	return JobID;
+	return INDEX_NONE;
 }
 
 int32 FHlslNiagaraCompiler::CreateShaderIntermediateData(const FStringView GroupName, const FNiagaraCompileOptions& InOptions, const FNiagaraTranslateResults& InTranslateResults, const FNiagaraTranslatorOutput& TranslatorOutput, const FString& TranslatedHLSL)
diff --git a/Engine/Plugins/FX/Niagara/Source/NiagaraEditor/Private/NiagaraCompiler.h b/Engine/Plugins/FX/Niagara/Source/NiagaraEditor/Private/NiagaraCompiler.h
--- a/Engine/Plugins/FX/Niagara/Source/NiagaraEditor/Private/NiagaraCompiler.h
+++ b/Engine/Plugins/FX/Niagara/Source/NiagaraEditor/Private/NiagaraCompiler.h
@@ -45,7 +45,7 @@ public:
 	NIAGARAEDITOR_API virtual int32 CompileScript(const class FNiagaraCompileRequestData* InCompileRequest, const FNiagaraCompileOptions& InOptions, const FNiagaraTranslateResults& InTranslateResults, FNiagaraTranslatorOutput* TranslatorOutput, FString& TranslatedHLSL) override;
 
 	NIAGARAEDITOR_API virtual int32 CompileScript(const FStringView GroupName, const FNiagaraCompileOptions& InOptions, const FNiagaraTranslateResults& InTranslateResults, const FNiagaraTranslatorOutput& TranslatorOutput, const FString& TranslatedHLSL) override;
-	NIAGARAEDITOR_API virtual uint32 CompileScriptVM(const FStringView GroupName, const FNiagaraCompileOptions& InOptions, const FNiagaraTranslateResults& InTranslateResults, const FNiagaraTranslatorOutput& TranslatorOutput, const FString& TranslatedHLSL);
+	NIAGARAEDITOR_API virtual uint32 CompileScriptVM(const FStringView GroupName, const FNiagaraCompileOptions& InOptions, const FNiagaraTranslateResults& InTranslateResults, const FNiagaraTranslatorOutput& TranslatorOutput, const FString& TranslatedHLSL, FNiagaraShaderType* NiagaraShaderType);
 	NIAGARAEDITOR_API virtual int32 CreateShaderIntermediateData(const FStringView GroupName, const FNiagaraCompileOptions& InOptions, const FNiagaraTranslateResults& InTranslateResults, const FNiagaraTranslatorOutput& TranslatorOutput, const FString& TranslatedHLSL);
 	virtual TOptional<FNiagaraCompileResults> GetCompileResult(int32 JobID, bool bWait = false) override;
 
diff --git a/Engine/Plugins/FX/Niagara/Source/NiagaraEditor/Private/NiagaraGraphDigest.cpp b/Engine/Plugins/FX/Niagara/Source/NiagaraEditor/Private/NiagaraGraphDigest.cpp
--- a/Engine/Plugins/FX/Niagara/Source/NiagaraEditor/Private/NiagaraGraphDigest.cpp
+++ b/Engine/Plugins/FX/Niagara/Source/NiagaraEditor/Private/NiagaraGraphDigest.cpp
@@ -2778,9 +2778,9 @@ void FNiagaraCompilationNodeFunctionCall::Compile(FTranslator* Translator, TArra
 
 const FNiagaraCompilationInputPin* FNiagaraCompilationNodeFunctionCall::FindStaticSwitchInputPin(FName VariableName) const
 {
-	if (ensure(OwningGraph))
+	if (CalledGraph)
 	{
-		for (const FNiagaraVariableBase& StaticSwitchInput : OwningGraph->StaticSwitchInputs)
+		for (const FNiagaraVariableBase& StaticSwitchInput : CalledGraph->StaticSwitchInputs)
 		{
 			if (StaticSwitchInput.GetName() == VariableName)
 			{
diff --git a/Engine/Plugins/FX/Niagara/Source/NiagaraEditor/Private/ViewModels/NiagaraSystemViewModel.cpp b/Engine/Plugins/FX/Niagara/Source/NiagaraEditor/Private/ViewModels/NiagaraSystemViewModel.cpp
--- a/Engine/Plugins/FX/Niagara/Source/NiagaraEditor/Private/ViewModels/NiagaraSystemViewModel.cpp
+++ b/Engine/Plugins/FX/Niagara/Source/NiagaraEditor/Private/ViewModels/NiagaraSystemViewModel.cpp
@@ -166,7 +166,7 @@ void FNiagaraSystemViewModel::Initialize(UNiagaraSystem& InSystem, FNiagaraSyste
 
 	if (CanImpactCompileForEdit())
 	{
-		System->bCompileForEdit = true;
+		System->SetCompileForEdit(true);
 	}
 	
 	if (bIsForDataProcessingOnly == false)
@@ -262,7 +262,7 @@ void FNiagaraSystemViewModel::Cleanup()
 	{
 		if (CanImpactCompileForEdit())
 		{
-			System->bCompileForEdit = false;
+			System->SetCompileForEdit(false);
 			System->InvalidateActiveCompiles();
 			System->RequestCompile(false);
 		}
@@ -1048,8 +1048,8 @@ void FNiagaraSystemViewModel::NotifyPreSave()
 	// we want to save without compile for edit turned on for best results, but make sure to turn in back on again after save is done
 	if (CanImpactCompileForEdit())
 	{
-		check(System->bCompileForEdit);
-		System->bCompileForEdit = false;
+		check(System->GetCompileForEdit());
+		System->SetCompileForEdit(false);
 		System->RequestCompile(false);
 		System->WaitForCompilationComplete(true);
 	}
@@ -1062,8 +1062,8 @@ void FNiagaraSystemViewModel::NotifyPostSave()
 	// we compile for edit again after having it turned off for PreSave. This should fetch DDC data and must not update the ChangeID
 	if (CanImpactCompileForEdit())
 	{
-		check(!System->bCompileForEdit);
-		System->bCompileForEdit = true;
+		check(!System->GetCompileForEdit());
+		System->SetCompileForEdit(true);
 		System->RequestCompile(false);
 		System->WaitForCompilationComplete(true);
 	}
diff --git a/Engine/Source/Runtime/Core/Private/UObject/FortniteMainBranchObjectVersion.cpp b/Engine/Source/Runtime/Core/Private/UObject/FortniteMainBranchObjectVersion.cpp
--- a/Engine/Source/Runtime/Core/Private/UObject/FortniteMainBranchObjectVersion.cpp
+++ b/Engine/Source/Runtime/Core/Private/UObject/FortniteMainBranchObjectVersion.cpp
@@ -11,7 +11,7 @@ TMap<FGuid, FGuid> FFortniteMainBranchObjectVersion::GetSystemGuids()
 	SystemGuids.Add(DevGuids.MATERIALSHADERMAP_DERIVEDDATA_VER, FGuid("CE7776E63E9E43FBB8C103EE1AB40B5D"));
 	SystemGuids.Add(DevGuids.NANITE_DERIVEDDATA_VER, FGuid("D28C003F825F4C178F698938FB0E47F7"));
 	SystemGuids.Add(DevGuids.NIAGARASHADERMAP_DERIVEDDATA_VER, FGuid("D0DF49D408F34E38A98FA097E0280EC7"));
-	SystemGuids.Add(DevGuids.Niagara_LatestScriptCompileVersion, FGuid("A36D165ED85D4FF694FA4053966ACFEB"));
+	SystemGuids.Add(DevGuids.Niagara_LatestScriptCompileVersion, FGuid("CE33E29C6268469E8E18F1D96A55DFDA"));
 	SystemGuids.Add(DevGuids.SkeletalMeshDerivedDataVersion, FGuid("518D7C51D31D9BE4A9F4B81542AC8119"));
 	SystemGuids.Add(DevGuids.STATICMESH_DERIVEDDATA_VER, FGuid("543843B103794E0E9BA4BE80FB602F79"));
 	SystemGuids.Add(DevGuids.MaterialTranslationDDCVersion, FGuid("982B2D2E6E50436A95EF4D16B1D04920"));
