- Ensure OverrideParameters are populated for BP templates
- Without this Data Interface user parameters may be missing when modified inside a construction script running -game

diff --git a/Engine/Plugins/FX/Niagara/Source/Niagara/Private/NiagaraComponent.cpp b/Engine/Plugins/FX/Niagara/Source/Niagara/Private/NiagaraComponent.cpp
--- a/Engine/Plugins/FX/Niagara/Source/Niagara/Private/NiagaraComponent.cpp
+++ b/Engine/Plugins/FX/Niagara/Source/Niagara/Private/NiagaraComponent.cpp
@@ -1980,6 +1980,12 @@ void UNiagaraComponent::PostApplyToComponent()
 	Super::PostApplyToComponent();
 
 #if WITH_EDITOR
+	// Ensure update the override parameters as they may be incomplete
+	// Without this calls to set override data interface values can fail as the DI may not be present in the override parameter store
+	if (Asset != nullptr)
+	{
+		CopyParametersFromAsset(false);
+	}
 	ApplyOverridesToParameterStore();
 	FixDataInterfaceOuters();
 #endif
