- Display Shader Statistics in Niagara Script Stats

diff --git a/Engine/Plugins/FX/Niagara/Source/NiagaraEditor/Private/ViewModels/NiagaraScriptStatsViewModel.cpp b/Engine/Plugins/FX/Niagara/Source/NiagaraEditor/Private/ViewModels/NiagaraScriptStatsViewModel.cpp
--- a/Engine/Plugins/FX/Niagara/Source/NiagaraEditor/Private/ViewModels/NiagaraScriptStatsViewModel.cpp
+++ b/Engine/Plugins/FX/Niagara/Source/NiagaraEditor/Private/ViewModels/NiagaraScriptStatsViewModel.cpp
@@ -73,7 +73,12 @@ namespace NiagaraScriptStatsLocal
 				const FNiagaraVMExecutableData& VMData = Script->GetVMExecutableData();
 				if (VMData.LastOpCount > 0)
 				{
-					OutStatus.Appendf(TEXT("VM: %s = %u\n"), *Script->GetName(), VMData.LastOpCount);
+					if (!OutStatus.IsEmpty())
+					{
+						OutStatus.AppendChar(TEXT('\n'));
+					}
+					OutStatus.Appendf(TEXT("VM: %s\n"), *Script->GetName());
+					OutStatus.Appendf(TEXT("- Op Count = %u\n"), VMData.LastOpCount);
 				}
 			}
 		}
@@ -196,16 +201,24 @@ namespace NiagaraScriptStatsLocal
 								FNiagaraShaderRef Shader = ShaderScript->GetShaderGameThread(iSimStageIndex);
 								if ( Shader.IsValid() )
 								{
-									ResultsString.Appendf(TEXT("GPU: %s = "), *SimulationStageMetaData[iSimStageIndex].SimulationStageName.ToString());
-									if (Shader->GetNumInstructions() == 0)
+									if (!ResultsString.IsEmpty())
 									{
-										ResultsString.Append(TEXT("n/a"));
+										ResultsString.AppendChar(TEXT('\n'));
 									}
-									else
+									ResultsString.Appendf(TEXT("GPU: %s\n"), *SimulationStageMetaData[iSimStageIndex].SimulationStageName.ToString());
+									if (Shader->GetNumInstructions() > 0)
+									{
+										ResultsString.Appendf(TEXT("- Num Instructions = %d\n"), Shader->GetNumInstructions());
+									}
+
+									TStringBuilder<2048> ShaderStatsBuilder;
+									const FShader::FShaderStatisticMap ShaderStats = Shader->GetShaderStatistics();
+									for (const auto& ShaderStat : ShaderStats)
 									{
-										ResultsString.AppendInt(Shader->GetNumInstructions());
+										ShaderStatsBuilder << TEXT("- ") << ShaderStat.Key << TEXT(" = ");
+										Visit([&ShaderStatsBuilder](auto& StoredValue) { ShaderStatsBuilder << StoredValue << "\n"; }, ShaderStat.Value);
 									}
-									ResultsString.AppendChar(TEXT('\n'));
+									ResultsString.Append(ShaderStatsBuilder.ToString());
 								}
 							}
 						}
