- Move Niagara stat groups into Niagara category

diff --git a/Engine/Plugins/Experimental/ChaosNiagara/Source/ChaosNiagara/Private/NiagaraDataInterfaceChaosDestruction.cpp b/Engine/Plugins/Experimental/ChaosNiagara/Source/ChaosNiagara/Private/NiagaraDataInterfaceChaosDestruction.cpp
--- a/Engine/Plugins/Experimental/ChaosNiagara/Source/ChaosNiagara/Private/NiagaraDataInterfaceChaosDestruction.cpp
+++ b/Engine/Plugins/Experimental/ChaosNiagara/Source/ChaosNiagara/Private/NiagaraDataInterfaceChaosDestruction.cpp
@@ -25,7 +25,7 @@
 #define LOCTEXT_NAMESPACE "ChaosNiagaraDestructionDataInterface"
 //#pragma optimize("", off)
 
-DECLARE_STATS_GROUP(TEXT("ChaosNiagara"), STATGROUP_ChaosNiagara, STATCAT_Advanced);
+DECLARE_STATS_GROUP(TEXT("ChaosNiagara"), STATGROUP_ChaosNiagara, STATCAT_Niagara);
 DECLARE_CYCLE_STAT(TEXT("CollisionCallback"), STAT_CollisionCallback, STATGROUP_ChaosNiagara);
 DECLARE_CYCLE_STAT(TEXT("TrailingCallback"), STAT_NiagaraTrailingCallback, STATGROUP_ChaosNiagara);
 DECLARE_CYCLE_STAT(TEXT("BreakingCallback"), STAT_NiagaraBreakingCallback, STATGROUP_ChaosNiagara);
diff --git a/Engine/Plugins/FX/Niagara/Source/Niagara/Classes/NiagaraScript.h b/Engine/Plugins/FX/Niagara/Source/Niagara/Classes/NiagaraScript.h
--- a/Engine/Plugins/FX/Niagara/Source/Niagara/Classes/NiagaraScript.h
+++ b/Engine/Plugins/FX/Niagara/Source/Niagara/Classes/NiagaraScript.h
@@ -30,7 +30,7 @@ class UNiagaraConvertInPlaceUtilityBase;
 
 #define NIAGARA_SCRIPT_COMPILE_LOGGING_MEDIUM
 
-DECLARE_STATS_GROUP(TEXT("Niagara Detailed"), STATGROUP_NiagaraDetailed, STATCAT_Advanced);
+DECLARE_STATS_GROUP(TEXT("Niagara Detailed"), STATGROUP_NiagaraDetailed, STATCAT_Niagara);
 
 /** Defines what will happen to unused attributes when a script is run. */
 UENUM()
diff --git a/Engine/Plugins/FX/Niagara/Source/Niagara/Private/NiagaraRendererRibbons.cpp b/Engine/Plugins/FX/Niagara/Source/Niagara/Private/NiagaraRendererRibbons.cpp
--- a/Engine/Plugins/FX/Niagara/Source/Niagara/Private/NiagaraRendererRibbons.cpp
+++ b/Engine/Plugins/FX/Niagara/Source/Niagara/Private/NiagaraRendererRibbons.cpp
@@ -35,7 +35,7 @@ DECLARE_CYCLE_STAT(TEXT("Generate Indices GPU [RT]"), STAT_NiagaraRenderRibbonsG
 DECLARE_CYCLE_STAT(TEXT("Generate Vertices CPU [GT]"), STAT_NiagaraRenderRibbonsGenVerticesCPU, STATGROUP_Niagara);
 DECLARE_CYCLE_STAT(TEXT("Generate Vertices GPU [RT]"), STAT_NiagaraRenderRibbonsGenVerticesGPU, STATGROUP_Niagara);
 
-DECLARE_STATS_GROUP(TEXT("NiagaraRibbons"), STATGROUP_NiagaraRibbons, STATCAT_Advanced);
+DECLARE_STATS_GROUP(TEXT("NiagaraRibbons"), STATGROUP_NiagaraRibbons, STATCAT_Niagara);
 
 DECLARE_CYCLE_STAT(TEXT("Generate Vertices GPU - Sort [RT]"), STAT_NiagaraRenderRibbonsGenVerticesSortGPU, STATGROUP_NiagaraRibbons);
 DECLARE_CYCLE_STAT(TEXT("Generate Vertices GPU - InitialSort [RT]"), STAT_NiagaraRenderRibbonsGenVerticesInitialSortGPU, STATGROUP_NiagaraRibbons);
diff --git a/Engine/Plugins/FX/Niagara/Source/Niagara/Private/NiagaraStats.h b/Engine/Plugins/FX/Niagara/Source/Niagara/Private/NiagaraStats.h
--- a/Engine/Plugins/FX/Niagara/Source/Niagara/Private/NiagaraStats.h
+++ b/Engine/Plugins/FX/Niagara/Source/Niagara/Private/NiagaraStats.h
@@ -5,7 +5,7 @@
 
 #include "Stats/Stats.h"
 
-DECLARE_STATS_GROUP(TEXT("Niagara"), STATGROUP_Niagara, STATCAT_Advanced);
+DECLARE_STATS_GROUP(TEXT("Niagara"), STATGROUP_Niagara, STATCAT_Niagara);
 DECLARE_DWORD_COUNTER_STAT(TEXT("# CPU Particles"), STAT_NiagaraNumParticles, STATGROUP_Niagara);
 DECLARE_CYCLE_STAT(TEXT("Constant Setup [CNC]"), STAT_NiagaraConstants, STATGROUP_Niagara);
 DECLARE_CYCLE_STAT(TEXT("Emitter Tick [CNC]"), STAT_NiagaraTick, STATGROUP_Niagara);
@@ -16,16 +16,16 @@ DECLARE_MEMORY_STAT(TEXT("Niagara GPU data interface memory"), STAT_NiagaraGPUDa
 DECLARE_MEMORY_STAT(TEXT("Niagara vertex buffer memory"), STAT_NiagaraVBMemory, STATGROUP_Niagara);
 DECLARE_MEMORY_STAT(TEXT("Niagara RenderTarget memory"), STAT_NiagaraRenderTargetMemory, STATGROUP_Niagara);
 
-DECLARE_STATS_GROUP(TEXT("NiagaraOverview"), STATGROUP_NiagaraOverview, STATCAT_Advanced);
+DECLARE_STATS_GROUP(TEXT("NiagaraOverview"), STATGROUP_NiagaraOverview, STATCAT_Niagara);
 DECLARE_CYCLE_STAT(TEXT("GT Total"), STAT_NiagaraOverview_GT, STATGROUP_NiagaraOverview);
 DECLARE_CYCLE_STAT(TEXT("GT Concurrent Total"), STAT_NiagaraOverview_GT_CNC, STATGROUP_NiagaraOverview);
 DECLARE_CYCLE_STAT(TEXT("RT Total"), STAT_NiagaraOverview_RT, STATGROUP_NiagaraOverview);
 DECLARE_CYCLE_STAT(TEXT("RT Concurrent Total"), STAT_NiagaraOverview_RT_CNC, STATGROUP_NiagaraOverview);
 
-DECLARE_STATS_GROUP(TEXT("Niagara Systems"), STATGROUP_NiagaraSystems, STATCAT_Advanced);
-DECLARE_STATS_GROUP(TEXT("Niagara Emitters"), STATGROUP_NiagaraEmitters, STATCAT_Advanced);
+DECLARE_STATS_GROUP(TEXT("Niagara Systems"), STATGROUP_NiagaraSystems, STATCAT_Niagara);
+DECLARE_STATS_GROUP(TEXT("Niagara Emitters"), STATGROUP_NiagaraEmitters, STATCAT_Niagara);
 
-DECLARE_STATS_GROUP(TEXT("Niagara System Counts"), STATGROUP_NiagaraSystemCounts, STATCAT_Advanced);
+DECLARE_STATS_GROUP(TEXT("Niagara System Counts"), STATGROUP_NiagaraSystemCounts, STATCAT_Niagara);
 
 DECLARE_DWORD_COUNTER_STAT(TEXT("Total System Instances"), STAT_TotalNiagaraSystemInstances, STATGROUP_NiagaraSystemCounts);
 DECLARE_DWORD_COUNTER_STAT(TEXT("Total Solo System Instances"), STAT_TotalNiagaraSystemInstancesSolo, STATGROUP_NiagaraSystemCounts);
diff --git a/Engine/Plugins/FX/Niagara/Source/Niagara/Public/NiagaraDataChannel.h b/Engine/Plugins/FX/Niagara/Source/Niagara/Public/NiagaraDataChannel.h
--- a/Engine/Plugins/FX/Niagara/Source/Niagara/Public/NiagaraDataChannel.h
+++ b/Engine/Plugins/FX/Niagara/Source/Niagara/Public/NiagaraDataChannel.h
@@ -46,7 +46,7 @@ When GPU->GPU is supported, the frame dependency issue have more meaning for GPU
 #include "UObject/UObjectIterator.h"
 #include "NiagaraDataChannel.generated.h"
 
-DECLARE_STATS_GROUP(TEXT("Niagara Data Channels"), STATGROUP_NiagaraDataChannels, STATCAT_Advanced);
+DECLARE_STATS_GROUP(TEXT("Niagara Data Channels"), STATGROUP_NiagaraDataChannels, STATCAT_Niagara);
 
 class FRHICommandListImmediate;
 
diff --git a/Engine/Plugins/FX/Niagara/Source/NiagaraEditor/Public/NiagaraEditorModule.h b/Engine/Plugins/FX/Niagara/Source/NiagaraEditor/Public/NiagaraEditorModule.h
--- a/Engine/Plugins/FX/Niagara/Source/NiagaraEditor/Public/NiagaraEditorModule.h
+++ b/Engine/Plugins/FX/Niagara/Source/NiagaraEditor/Public/NiagaraEditorModule.h
@@ -52,7 +52,7 @@ class UNiagaraParameterCollection;
 struct FNiagaraSystemAsyncCompileResults;
 class ITargetPlatform;
 
-DECLARE_STATS_GROUP(TEXT("Niagara Editor"), STATGROUP_NiagaraEditor, STATCAT_Advanced);
+DECLARE_STATS_GROUP(TEXT("Niagara Editor"), STATGROUP_NiagaraEditor, STATCAT_Niagara);
 
 extern NIAGARAEDITOR_API int32 GbShowNiagaraDeveloperWindows;
 extern NIAGARAEDITOR_API int32 GbPreloadSelectablePluginAssetsOnDemand;
