- default to 512 res
- set landscape as collider
- default bottom depth constant at -25
- new slightly faster way of querying implicits avoids some matrix math
- remove and consolidate stages in sim to remove dispatches

#virtualized

diff --git a/Engine/Plugins/FX/Niagara/Shaders/Private/NiagaraDataInterfaceRigidMeshCollisionQuery.ush b/Engine/Plugins/FX/Niagara/Shaders/Private/NiagaraDataInterfaceRigidMeshCollisionQuery.ush
--- a/Engine/Plugins/FX/Niagara/Shaders/Private/NiagaraDataInterfaceRigidMeshCollisionQuery.ush
+++ b/Engine/Plugins/FX/Niagara/Shaders/Private/NiagaraDataInterfaceRigidMeshCollisionQuery.ush
@@ -175,6 +175,11 @@ float3 GetLocalPosition_{ParameterName}(in float3 WorldPosition, in int ElementI
 	return PreviousLocal + TimeFraction * (CurrentLocal-PreviousLocal);
 }
 
+float3 GetCurrentLocalPosition_{ParameterName}(in float3 WorldPosition, in int ElementIndex )
+{
+	return mul(GetInverseTransform_{ParameterName}(ElementIndex), float4(WorldPosition,1.0)).xyz;	
+}
+
 // Given a world space position (WorldPosition) compute the sphere closest point (position,normal,velocity)
 int GetClosestElement_{ParameterName}(in float3 WorldPosition, out float3 OutClosestPosition, out float3 OutClosestNormal, out float OutMinDistance, in float TimeFraction)
 {	
@@ -218,6 +223,49 @@ int GetClosestElement_{ParameterName}(in float3 WorldPosition, out float3 OutClo
 	return ElementIndex;
 }
 
+// Given a world space position (WorldPosition) compute the sphere closest point (position,normal,velocity)
+int GetClosestElementSimple_{ParameterName}(in float3 WorldPosition, out float3 OutClosestPosition, out float OutMinDistance)
+{	
+	float MinDistance = MAX_DISTANCE;
+	int ElementIndex = -1;
+	float3 CollisionPosition = float3(0,0,0);
+	float3 CollisionNormal = float3(0,0,0);
+
+	const int SpheresBegin = {ParameterName}_ElementOffsets[SPHERE_INDEX];
+	const int SpheresEnd = {ParameterName}_ElementOffsets[SPHERE_INDEX+1];
+	for (int SphereIndex = SpheresBegin; SphereIndex < SpheresEnd; ++SphereIndex)
+	{
+		const float3 LocalPosition = GetCurrentLocalPosition_{ParameterName}(WorldPosition,SphereIndex);
+
+		GetSphereProjection_{ParameterName}(LocalPosition, float3(0,0,0), {ParameterName}_ElementExtentBuffer[SphereIndex].x, SphereIndex,
+				CollisionPosition, CollisionNormal, ElementIndex, MinDistance);
+	}
+
+	const int CapsulesBegin = {ParameterName}_ElementOffsets[CAPSULE_INDEX];
+	const int CapsulesEnd = {ParameterName}_ElementOffsets[CAPSULE_INDEX+1];
+	for (int CapsuleIndex = CapsulesBegin; CapsuleIndex < CapsulesEnd; ++CapsuleIndex)
+	{
+		const float3 LocalPosition = GetCurrentLocalPosition_{ParameterName}(WorldPosition,CapsuleIndex);
+
+		GetCapsuleProjection_{ParameterName}(LocalPosition, {ParameterName}_ElementExtentBuffer[CapsuleIndex].xy, CapsuleIndex,
+				CollisionPosition, CollisionNormal, ElementIndex, MinDistance);
+	}
+
+	const int BoxesBegin = {ParameterName}_ElementOffsets[BOX_INDEX];
+	const int BoxesEnd = {ParameterName}_ElementOffsets[BOX_INDEX+1];
+	for (int BoxIndex = BoxesBegin; BoxIndex < BoxesEnd; ++BoxIndex)
+	{
+		const float3 LocalPosition = GetCurrentLocalPosition_{ParameterName}(WorldPosition,BoxIndex);
+
+		GetBoxProjection_{ParameterName}(LocalPosition, {ParameterName}_ElementExtentBuffer[BoxIndex].xyz, BoxIndex,
+				CollisionPosition, CollisionNormal, ElementIndex, MinDistance);
+	}
+
+	OutClosestPosition = CollisionPosition;	
+	OutMinDistance = MinDistance;
+	return ElementIndex;
+}
+
 /* -----------------------------------------------------------------
  * Get the closest point to the static mesh 
  * -----------------------------------------------------------------
@@ -394,6 +442,24 @@ void GetClosestPoint_{ParameterName}(in float3 WorldPosition, in float DeltaTime
 	}
 }
 
+// Given a world space position (WorldPosition) compute the static mesh closest point (position,normal,velocity)
+void GetClosestPointSimple_{ParameterName}(in float3 WorldPosition, in float DeltaTime, out float OutClosestDistance, out float3 OutClosestVelocity)
+{	
+	float3 CollisionPosition = float3(0,0,0);
+	OutClosestDistance = 0.0;
+
+	const int ElementIndex = GetClosestElementSimple_{ParameterName}(WorldPosition,CollisionPosition,OutClosestDistance);
+
+	OutClosestVelocity = float3(0,0,0);
+
+	if (ElementIndex != -1)
+	{
+		const float3 PreviousPosition = mul(GetPreviousTransform_{ParameterName}(ElementIndex), float4(CollisionPosition,1.0)).xyz;
+		const float3 CurrentPosition = mul(GetCurrentTransform_{ParameterName}(ElementIndex), float4(CollisionPosition,1.0)).xyz;
+		OutClosestVelocity = ( CurrentPosition - PreviousPosition ) / DeltaTime;
+	}
+}
+
 // Given a world space position (WorldPosition) compute the static mesh closest distance
 void GetClosestDistance_{ParameterName}(in float3 WorldPosition, in float TimeFraction, out float OutClosestDistance)
 {	
diff --git a/Engine/Plugins/FX/Niagara/Source/Niagara/Private/NiagaraDataInterfaceRigidMeshCollisionQuery.cpp b/Engine/Plugins/FX/Niagara/Source/Niagara/Private/NiagaraDataInterfaceRigidMeshCollisionQuery.cpp
--- a/Engine/Plugins/FX/Niagara/Source/Niagara/Private/NiagaraDataInterfaceRigidMeshCollisionQuery.cpp
+++ b/Engine/Plugins/FX/Niagara/Source/Niagara/Private/NiagaraDataInterfaceRigidMeshCollisionQuery.cpp
@@ -94,6 +94,7 @@ static const FName GetElementPointName(TEXT("GetElementPoint"));
 static const FName GetElementPointMeshDistanceFieldNoNormalName(TEXT("GetElementPointMeshDistanceFieldNoNormal"));
 static const FName GetElementDistanceName(TEXT("GetElementDistance"));
 static const FName GetClosestPointName(TEXT("GetClosestPoint"));
+static const FName GetClosestPointSimpleName(TEXT("GetClosestPointSimple"));
 static const FName GetClosestDistanceName(TEXT("GetClosestDistance"));
 static const FName GetClosestPointMeshDistanceFieldName(TEXT("GetClosestPointMeshDistanceField"));
 static const FName GetClosestPointMeshDistanceFieldAccurateName(TEXT("GetClosestPointMeshDistanceFieldAccurate"));
@@ -1388,6 +1389,24 @@ void UNiagaraDataInterfaceRigidMeshCollisionQuery::GetFunctionsInternal(TArray<F
 
 		OutFunctions.Add(Sig);
 	}
+
+	{
+		FNiagaraFunctionSignature Sig;
+		Sig.Name = GetClosestPointSimpleName;
+		Sig.SetDescription(LOCTEXT("GetClosestPointSimpleDescription", "Given a world space position, computes the static mesh's closest point. Also returns normal and velocity for that point."));
+		Sig.SetFunctionVersion(FNiagaraRigidMeshCollisionDIFunctionVersion::LatestVersion);
+		Sig.bSupportsGPU = true;
+		Sig.bSupportsCPU = false;
+		Sig.bMemberFunction = true;
+		Sig.Inputs.Add(FNiagaraVariable(FNiagaraTypeDefinition(GetClass()), TEXT("Collision DI")));
+		Sig.Inputs.Add(FNiagaraVariable(FNiagaraTypeDefinition::GetPositionDef(), TEXT("World Position")));
+		Sig.Inputs.Add(FNiagaraVariable(FNiagaraTypeDefinition::GetFloatDef(), TEXT("Delta Time")));	
+		Sig.Outputs.Add(FNiagaraVariable(FNiagaraTypeDefinition::GetFloatDef(), TEXT("Closest Distance")));		
+		Sig.Outputs.Add(FNiagaraVariable(FNiagaraTypeDefinition::GetVec3Def(), TEXT("Closest Velocity")));
+
+		OutFunctions.Add(Sig);
+	}
+
 	{
 		FNiagaraFunctionSignature Sig;
 		Sig.Name = GetClosestElementName;
@@ -1563,6 +1582,7 @@ bool UNiagaraDataInterfaceRigidMeshCollisionQuery::GetFunctionHLSL(const FNiagar
 		(FunctionInfo.DefinitionName == GetNumCapsulesName) ||
 		(FunctionInfo.DefinitionName == GetNumSpheresName) ||
 		(FunctionInfo.DefinitionName == GetClosestPointName) ||
+		(FunctionInfo.DefinitionName == GetClosestPointSimpleName) ||		
 		(FunctionInfo.DefinitionName == GetClosestElementName) ||
 		(FunctionInfo.DefinitionName == GetElementPointName) ||
 		(FunctionInfo.DefinitionName == GetElementPointMeshDistanceFieldNoNormalName) ||
