- Optional data interface comparison for sim cache compare

diff --git a/Engine/Plugins/FX/Niagara/Source/Niagara/Classes/NiagaraSimCacheCompare.h b/Engine/Plugins/FX/Niagara/Source/Niagara/Classes/NiagaraSimCacheCompare.h
--- a/Engine/Plugins/FX/Niagara/Source/Niagara/Classes/NiagaraSimCacheCompare.h
+++ b/Engine/Plugins/FX/Niagara/Source/Niagara/Classes/NiagaraSimCacheCompare.h
@@ -19,6 +19,12 @@ public:
 	/** Should we exclude known none deterministic variables, i.e. ID / MaterialRandom. */
 	bool bExcludeNoneDeterministicVariables = true;
 
+	/**
+	Should we include data interface as part of the comparison.
+	Note: If the data interface does not implement the comparison function an error will be generated.
+	*/
+	bool bIncludeDataInterfaces = true;
+
 	/** Add an attribute to exclude from the comparison. */
 	NIAGARA_API void AddExcludeAttribute(FNiagaraVariableBase Variable);
 
diff --git a/Engine/Plugins/FX/Niagara/Source/Niagara/Classes/NiagaraSimCacheCustomStorageInterface.h b/Engine/Plugins/FX/Niagara/Source/Niagara/Classes/NiagaraSimCacheCustomStorageInterface.h
--- a/Engine/Plugins/FX/Niagara/Source/Niagara/Classes/NiagaraSimCacheCustomStorageInterface.h
+++ b/Engine/Plugins/FX/Niagara/Source/Niagara/Classes/NiagaraSimCacheCustomStorageInterface.h
@@ -57,6 +57,12 @@ public:
 	*/
 	virtual void SimCachePostReadFrame(void* OptionalPerInstanceData, FNiagaraSystemInstance* SystemInstance) {}
 
+	/**
+	Called to compare a frame between two separate simulation cache storages
+	This will be called on the CDO object since we do not have the actual data interface.
+	*/
+	virtual bool SimCacheCompareFrame(UObject* LhsStorageObject, UObject* RhsStorageObject, int FrameIndex, TOptional<float> Tolerance, FString& OutErrors) { OutErrors = TEXT("Compare not implemented"); return false; }
+
 	/**
 	This function allows you to preserve a list of attributes when building a renderer only cache.
 	The UsageContext will be either a UNiagaraSystem or a UNiagaraEmitter and can be used to scope your variables accordingly.
diff --git a/Engine/Plugins/FX/Niagara/Source/Niagara/Private/NiagaraSimCacheCompare.cpp b/Engine/Plugins/FX/Niagara/Source/Niagara/Private/NiagaraSimCacheCompare.cpp
--- a/Engine/Plugins/FX/Niagara/Source/Niagara/Private/NiagaraSimCacheCompare.cpp
+++ b/Engine/Plugins/FX/Niagara/Source/Niagara/Private/NiagaraSimCacheCompare.cpp
@@ -3,6 +3,7 @@
 #include "NiagaraSimCacheCompare.h"
 #include "NiagaraConstants.h"
 #include "NiagaraModule.h"
+#include "NiagaraSimCacheCustomStorageInterface.h"
 
 namespace NiagaraSimCacheCompareInternal
 {
@@ -230,6 +231,83 @@ bool FNiagaraSimCacheCompare::Compare(const UNiagaraSimCache& LhsCache, const UN
 		}
 	}
 
+	// Compare data interfaces
+	if (bEqual && bIncludeDataInterfaces)
+	{
+		// Generate set of DIs to compare
+		const TArray<FNiagaraVariableBase> DataInterfacesToCompare = LhsCache.GetStoredDataInterfaces();
+		{
+			TArray<FNiagaraVariableBase> ExpectedDataInterfaces = RhsCache.GetStoredDataInterfaces();
+			for (const FNiagaraVariableBase& DIVariable : DataInterfacesToCompare)
+			{
+				if (ExpectedDataInterfaces.Remove(DIVariable) == 0)
+				{
+					AddError(OutDifferences, FString::Printf(TEXT("DataInterface(%s) was in Lhs and not Rhs cache."), *DIVariable.GetName().ToString()));
+					bEqual = false;
+				}
+			}
+
+			for (const FNiagaraVariableBase& DIVariable : ExpectedDataInterfaces)
+			{
+				AddError(OutDifferences, FString::Printf(TEXT("DataInterface(%s) was in Rhs and not Lhs cache."), *DIVariable.GetName().ToString()));
+				bEqual = false;
+			}
+		}
+
+		// Compare data interfaces
+		if (bEqual)
+		{
+			for (const FNiagaraVariableBase& DIVariable : DataInterfacesToCompare)
+			{
+				if (AttributesToExclude_SystemScript.Contains(DIVariable))
+				{
+					continue;
+				}
+
+				UClass* Class = DIVariable.GetType().GetClass();
+				INiagaraSimCacheCustomStorageInterface* DataInterfaceCDO = Class ? Cast<INiagaraSimCacheCustomStorageInterface>(DIVariable.GetType().GetClass()->GetDefaultObject()) : nullptr;
+				if (!DataInterfaceCDO)
+				{
+					AddError(OutDifferences, FString::Printf(TEXT("DataInterface(%s) could not find CDO for Class(%s)."), *DIVariable.GetName().ToString(), *GetNameSafe(Class)));
+					bEqual = false;
+					break;
+				}
+
+				TOptional<float> Tolerance;
+				if (float* UserTolerance = VariableToFloatTolerances_SystemScript.Find(DIVariable))
+				{
+					Tolerance.Emplace(*UserTolerance);
+				}
+
+				UObject* LhsStorageObject = LhsCache.GetDataInterfaceStorageObject(DIVariable);
+				UObject* RhsStorageObject = RhsCache.GetDataInterfaceStorageObject(DIVariable);
+				if (LhsStorageObject && RhsStorageObject)
+				{
+					for (int FrameIndex=0; FrameIndex < LhsCache.GetNumFrames(); FrameIndex++)
+					{
+						FString DIErrors;
+						if (!DataInterfaceCDO->SimCacheCompareFrame(LhsStorageObject, RhsStorageObject, FrameIndex, Tolerance, DIErrors))
+						{
+							AddError(OutDifferences, FString::Printf(TEXT("DataInterface(%s) frame (%d) is not equal - %s."), *DIVariable.GetName().ToString(), FrameIndex, *DIErrors));
+							bEqual = false;
+						}
+					}
+				}
+				else
+				{
+					AddError(OutDifferences, FString::Printf(TEXT("DataInterface(%s) storage was not found in Rhs(%p) Lhs(%p) cache."), *DIVariable.GetName().ToString(), LhsStorageObject, RhsStorageObject));
+					bEqual = false;
+				}
+
+				if (!bEqual)
+				{
+					break;
+				}
+			}
+		}
+	}
+
+
 	return bEqual;
 }
 
