- Simple Counter SimCache Support
- Change test to use sim cache

#virtualized

diff --git a/Engine/Plugins/FX/Niagara/Source/Niagara/Private/DataInterface/NiagaraDataInterfaceSimpleCounter.cpp b/Engine/Plugins/FX/Niagara/Source/Niagara/Private/DataInterface/NiagaraDataInterfaceSimpleCounter.cpp
--- a/Engine/Plugins/FX/Niagara/Source/Niagara/Private/DataInterface/NiagaraDataInterfaceSimpleCounter.cpp
+++ b/Engine/Plugins/FX/Niagara/Source/Niagara/Private/DataInterface/NiagaraDataInterfaceSimpleCounter.cpp
@@ -472,6 +472,124 @@ void UNiagaraDataInterfaceSimpleCounter::PushToRenderThreadImpl()
 	}
 }
 
+UObject* UNiagaraDataInterfaceSimpleCounter::SimCacheBeginWrite(UObject* SimCache, FNiagaraSystemInstance* NiagaraSystemInstance, const void* OptionalPerInstanceData, FNiagaraSimCacheFeedbackContext& FeedbackContext) const
+{
+	UNDISimpleCounterSimCacheData* CacheData = NewObject<UNDISimpleCounterSimCacheData>(SimCache);
+	return CacheData;
+}
+
+bool UNiagaraDataInterfaceSimpleCounter::SimCacheWriteFrame(UObject* StorageObject, int FrameIndex, FNiagaraSystemInstance* SystemInstance, const void* OptionalPerInstanceData, FNiagaraSimCacheFeedbackContext& FeedbackContext) const
+{
+	check(OptionalPerInstanceData && StorageObject);
+
+	UNDISimpleCounterSimCacheData* CacheData = CastChecked<UNDISimpleCounterSimCacheData>(StorageObject);
+
+	const int32 ValueOffset = FrameIndex * 2;
+	const int32 ExpectedValues = ValueOffset + 2;
+	if (CacheData->Values.Num() < ExpectedValues)
+	{
+		CacheData->Values.AddDefaulted(ExpectedValues - CacheData->Values.Num());
+	}
+
+	const FNDISimpleCounterInstanceData_GameThread* InstanceData_GT = static_cast<const FNDISimpleCounterInstanceData_GameThread*>(OptionalPerInstanceData);
+	CacheData->Values[ValueOffset + 0] = InstanceData_GT->Counter;
+	CacheData->Values[ValueOffset + 1] = InstanceData_GT->Counter;
+
+	if (IsUsedWithGPUScript())
+	{
+		ENQUEUE_RENDER_COMMAND(FNDISimpleCounter_CacheWriteFrame)
+		(
+			[Proxy_RT=GetProxyAs<FNDISimpleCounterProxy>(), InstanceID=SystemInstance->GetId(), ComputeInterface=SystemInstance->GetComputeDispatchInterface(), CacheData, ValueOffset](FRHICommandListImmediate& RHICmdList)
+			{
+				const FNDISimpleCounterInstanceData_RenderThread* InstanceData_RT = Proxy_RT->PerInstanceData_RenderThread.Find(InstanceID);
+				if (!InstanceData_RT || InstanceData_RT->CountOffset == INDEX_NONE)
+				{
+					return;
+				}
+
+				const FNiagaraGPUInstanceCountManager& CountManager = ComputeInterface->GetGPUInstanceCounterManager();
+				FNiagaraGpuReadbackManager* ReadbackManager = ComputeInterface->GetGpuReadbackManager();
+
+				ReadbackManager->EnqueueReadback(
+					RHICmdList,
+					CountManager.GetInstanceCountBuffer().Buffer,
+					InstanceData_RT->CountOffset * sizeof(uint32), sizeof(int32),
+					[CacheData, ValueOffset](TConstArrayView<TPair<void*, uint32>> ReadbackData)
+					{
+						CacheData->Values[ValueOffset + 1] = *static_cast<const int32*>(ReadbackData[0].Key);
+					}
+				);
+				ReadbackManager->WaitCompletion(RHICmdList);
+			}
+		);
+		FlushRenderingCommands();
+	}
+	return true;
+}
+
+bool UNiagaraDataInterfaceSimpleCounter::SimCacheReadFrame(UObject* StorageObject, int FrameA, int FrameB, float Interp, FNiagaraSystemInstance* SystemInstance, void* OptionalPerInstanceData)
+{
+	check(OptionalPerInstanceData && StorageObject);
+
+	const int32 ValueOffset = FrameA * 2;
+	const int32 ExpectedValues = ValueOffset + 2;
+
+	UNDISimpleCounterSimCacheData* CacheData = CastChecked<UNDISimpleCounterSimCacheData>(StorageObject);
+	if (CacheData->Values.Num() < ExpectedValues)
+	{
+		return false;
+	}
+
+	// Set CPU Data
+	FNDISimpleCounterInstanceData_GameThread* InstanceData_GT = static_cast<FNDISimpleCounterInstanceData_GameThread*>(OptionalPerInstanceData);
+	InstanceData_GT->Counter = CacheData->Values[ValueOffset + 0];
+
+	// Set GPU Data
+	if (IsUsedWithGPUScript())
+	{
+		ENQUEUE_RENDER_COMMAND(FNDISimpleCounter_PushToRender)
+		(
+			[Proxy_RT=GetProxyAs<FNDISimpleCounterProxy>(), InstanceID=SystemInstance->GetId(), NewValue=CacheData->Values[ValueOffset + 1]](FRHICommandListImmediate& RHICmdList)
+			{
+				if ( FNDISimpleCounterInstanceData_RenderThread* InstanceData_RT=Proxy_RT->PerInstanceData_RenderThread.Find(InstanceID) )
+				{
+					InstanceData_RT->CountValue = NewValue;
+				}
+			}
+		);
+	}
+
+	return true;
+}
+
+bool UNiagaraDataInterfaceSimpleCounter::SimCacheCompareFrame(UObject* LhsStorageObject, UObject* RhsStorageObject, int FrameIndex, TOptional<float> Tolerance, FString& OutErrors)
+{
+	UNDISimpleCounterSimCacheData* LhsCacheData = CastChecked<UNDISimpleCounterSimCacheData>(LhsStorageObject);
+	UNDISimpleCounterSimCacheData* RhsCacheData = CastChecked<UNDISimpleCounterSimCacheData>(RhsStorageObject);
+
+	const int32 ValueOffset = FrameIndex * 2;
+	const int32 ExpectedValues = ValueOffset + 2;
+
+	if (LhsCacheData->Values.Num() != RhsCacheData->Values.Num())
+	{
+		OutErrors = FString::Printf(TEXT("Number of LhsValues(%d) != RhsValues(%d)"), LhsCacheData->Values.Num(), RhsCacheData->Values.Num());
+		return false;
+	}
+
+	if ((LhsCacheData->Values[ValueOffset + 0] != RhsCacheData->Values[ValueOffset + 0]) ||
+		(LhsCacheData->Values[ValueOffset + 1] != RhsCacheData->Values[ValueOffset + 1]))
+	{
+		OutErrors = FString::Printf(
+			TEXT("Cpu|Gpu values to do not match LhsValues(%d|%d) != RhsValues(%d|%d)"),
+			LhsCacheData->Values[ValueOffset + 0], LhsCacheData->Values[ValueOffset + 1],
+			RhsCacheData->Values[ValueOffset + 0], RhsCacheData->Values[ValueOffset + 1]
+		);
+		return false;
+	}
+
+	return true;
+}
+
 void UNiagaraDataInterfaceSimpleCounter::UpdateDIProxy()
 {
 	FNDISimpleCounterProxy* Proxy_GT = GetProxyAs<FNDISimpleCounterProxy>();
diff --git a/Engine/Plugins/FX/Niagara/Source/Niagara/Private/DataInterface/NiagaraDataInterfaceSimpleCounter.h b/Engine/Plugins/FX/Niagara/Source/Niagara/Private/DataInterface/NiagaraDataInterfaceSimpleCounter.h
--- a/Engine/Plugins/FX/Niagara/Source/Niagara/Private/DataInterface/NiagaraDataInterfaceSimpleCounter.h
+++ b/Engine/Plugins/FX/Niagara/Source/Niagara/Private/DataInterface/NiagaraDataInterfaceSimpleCounter.h
@@ -3,9 +3,12 @@
 
 #include "NiagaraCommon.h"
 #include "NiagaraShared.h"
+#include "NiagaraSimCacheCustomStorageInterface.h"
+
 #include "VectorVM.h"
 #include "NiagaraDataInterfaceRW.h"
 #include "HAL/PlatformAtomics.h"
+
 #include "NiagaraDataInterfaceSimpleCounter.generated.h"
 
 struct FNiagaraDataInterfaceGeneratedFunction;
@@ -15,7 +18,7 @@ Thread safe counter starts at the initial value on start / reset.
 When operating between CPU & GPU ensure you set the appropriate sync mode.
 */
 UCLASS(EditInlineNew, Category = "Counting", CollapseCategories, meta = (DisplayName = "Simple Counter"), MinimalAPI)
-class UNiagaraDataInterfaceSimpleCounter : public UNiagaraDataInterfaceRWBase
+class UNiagaraDataInterfaceSimpleCounter : public UNiagaraDataInterfaceRWBase, public INiagaraSimCacheCustomStorageInterface
 {
 	GENERATED_UCLASS_BODY()
 
@@ -56,6 +59,13 @@ public:
 	NIAGARA_API virtual void PushToRenderThreadImpl() override;
 	// UNiagaraDataInterface Interface End
 
+	//~ INiagaraSimCacheCustomStorageInterface interface BEGIN
+	virtual UObject* SimCacheBeginWrite(UObject* SimCache, FNiagaraSystemInstance* NiagaraSystemInstance, const void* OptionalPerInstanceData, FNiagaraSimCacheFeedbackContext& FeedbackContext) const override;
+	virtual bool SimCacheWriteFrame(UObject* StorageObject, int FrameIndex, FNiagaraSystemInstance* SystemInstance, const void* OptionalPerInstanceData, FNiagaraSimCacheFeedbackContext& FeedbackContext) const override;
+	virtual bool SimCacheReadFrame(UObject* StorageObject, int FrameA, int FrameB, float Interp, FNiagaraSystemInstance* SystemInstance, void* OptionalPerInstanceData) override;
+	virtual bool SimCacheCompareFrame(UObject* LhsStorageObject, UObject* RhsStorageObject, int FrameIndex, TOptional<float> Tolerance, FString& OutErrors) override;
+	//~ UNiagaraDataInterface interface END
+
 	void UpdateDIProxy();
 
 	// VM functions
@@ -81,3 +91,13 @@ protected:
 	NIAGARA_API virtual void GetFunctionsInternal(TArray<FNiagaraFunctionSignature>& OutFunctions) const override;
 #endif
 };
+
+UCLASS(MinimalAPI)
+class UNDISimpleCounterSimCacheData : public UObject
+{
+	GENERATED_BODY()
+
+public:
+	UPROPERTY()
+	TArray<int32> Values;
+};
