Fix warning V1078: An empty container is iterated. The loop will not be executed.

diff --git a/Engine/Plugins/FX/Niagara/Source/Niagara/Private/NiagaraDebugHud.cpp b/Engine/Plugins/FX/Niagara/Source/Niagara/Private/NiagaraDebugHud.cpp
--- a/Engine/Plugins/FX/Niagara/Source/Niagara/Private/NiagaraDebugHud.cpp
+++ b/Engine/Plugins/FX/Niagara/Source/Niagara/Private/NiagaraDebugHud.cpp
@@ -3363,7 +3363,6 @@ bool FNiagaraDebugHUDStatsListener::Tick()
 	}
 
 #if WITH_PER_SYSTEM_PARTICLE_PERF_STATS
-	TArray<TWeakObjectPtr<const UFXSystemAsset>,TInlineAllocator<8>> ToRemove;
 	for (auto& HUDStatsPair : SystemStats)
 	{
 		TWeakObjectPtr<const UFXSystemAsset> WeakSystem = HUDStatsPair.Key;
@@ -3401,11 +3400,6 @@ bool FNiagaraDebugHUDStatsListener::Tick()
 			Stats->ResetGT();
 		}
 	}
-
-	for (auto& Remove : ToRemove)
-	{
-		SystemStats.Remove(Remove);
-	}
 #endif
 
 	return true;
@@ -3458,7 +3452,6 @@ void FNiagaraDebugHUDStatsListener::TickRT()
 	}
 
 #if WITH_PER_SYSTEM_PARTICLE_PERF_STATS
-	TArray<TWeakObjectPtr<const UFXSystemAsset>, TInlineAllocator<8>> ToRemove;
 	for (auto& HUDStatsPair : SystemStats)
 	{
 		TWeakObjectPtr<const UFXSystemAsset> WeakSystem = HUDStatsPair.Key;
@@ -3505,11 +3498,6 @@ void FNiagaraDebugHUDStatsListener::TickRT()
 			Stats->ResetRT();
 		}
 	}
-
-	for (auto& Remove : ToRemove)
-	{
-		SystemStats.Remove(Remove);
-	}
 #endif
 }
 #endif//WITH_PARTICLE_PERF_STATS
