Added parent emitter path to the open parent emitter tooltip

diff --git a/Engine/Plugins/FX/Niagara/Source/NiagaraEditorWidgets/Private/SNiagaraOverviewStackNode.cpp b/Engine/Plugins/FX/Niagara/Source/NiagaraEditorWidgets/Private/SNiagaraOverviewStackNode.cpp
--- a/Engine/Plugins/FX/Niagara/Source/NiagaraEditorWidgets/Private/SNiagaraOverviewStackNode.cpp
+++ b/Engine/Plugins/FX/Niagara/Source/NiagaraEditorWidgets/Private/SNiagaraOverviewStackNode.cpp
@@ -1012,17 +1012,27 @@ FReply SNiagaraOverviewStackNode::OpenParentEmitter()
 
 FText SNiagaraOverviewStackNode::OpenParentEmitterTooltip() const
 {
+	FString TooltipText = LOCTEXT("OpenAndFocusParentEmitterToolTip", "Open and Focus Parent Emitter").ToString();
 	TSharedPtr<FNiagaraEmitterHandleViewModel> EmitterHandleViewModel = EmitterHandleViewModelWeak.Pin();
+	
 	if (EmitterHandleViewModel.IsValid() && EmitterHandleViewModel->GetEmitterViewModel()->HasParentEmitter() && EmitterHandleViewModel->GetEmitterViewModel()->GetParentEmitter().Emitter->IsVersioningEnabled())
 	{
 		FVersionedNiagaraEmitter ParentEmitter = EmitterHandleViewModel->GetEmitterViewModel()->GetParentEmitter();
 		FText ParentName = FText::FromString(ParentEmitter.Emitter->GetUniqueEmitterName());
 		if (FVersionedNiagaraEmitterData* EmitterData = ParentEmitter.GetEmitterData())
 		{
-			return FText::Format(LOCTEXT("OpenAndFocusVersionedParentToolTip", "Open and Focus Parent Emitter:\n{0} - v{1}.{2}"), ParentName, EmitterData->Version.MajorVersion, EmitterData->Version.MinorVersion);
+			TooltipText.Append(TEXT(":\n{0} - v{1}.{2}"));
+			TooltipText = FText::Format(FText::FromString(TooltipText), ParentName, EmitterData->Version.MajorVersion, EmitterData->Version.MinorVersion).ToString();
 		}
 	}
-	return LOCTEXT("OpenAndFocusParentEmitterToolTip", "Open and Focus Parent Emitter");
+
+	if(EmitterHandleViewModel.IsValid() && EmitterHandleViewModel->GetEmitterViewModel()->HasParentEmitter())
+	{
+		TooltipText.Append(TEXT("\n{0}"));
+		TooltipText = FText::Format(FText::FromString(TooltipText), EmitterHandleViewModel->GetEmitterViewModel()->GetParentPathNameText()).ToString();
+	}
+	
+	return FText::FromString(TooltipText);
 }
 
 EVisibility SNiagaraOverviewStackNode::GetOpenParentEmitterVisibility() const
