- SimCache storage & visualizer for Niagara Arrays

diff --git a/Engine/Plugins/FX/Niagara/Source/Niagara/Classes/NiagaraDataInterfaceArray.h b/Engine/Plugins/FX/Niagara/Source/Niagara/Classes/NiagaraDataInterfaceArray.h
--- a/Engine/Plugins/FX/Niagara/Source/Niagara/Classes/NiagaraDataInterfaceArray.h
+++ b/Engine/Plugins/FX/Niagara/Source/Niagara/Classes/NiagaraDataInterfaceArray.h
@@ -6,6 +6,7 @@
 #include "NiagaraCommon.h"
 #include "NiagaraShared.h"
 #include "NiagaraShader.h"
+#include "NiagaraSimCacheCustomStorageInterface.h"
 #include "ShaderParameterUtils.h"
 #include "RHIUtilities.h"
 
@@ -13,6 +14,8 @@
 
 template<typename TArrayType, class TOwnerType>
 struct FNDIArrayProxyImpl;
+class UNDIArraySimCacheData;
+struct FNDIArraySimCacheDataFrame;
 
 #define NDIARRAY_GENERATE_BODY(CLASSNAME, TYPENAME, MEMBERNAME) \
 	using FProxyType = FNDIArrayProxyImpl<TYPENAME, CLASSNAME>; \
@@ -61,6 +64,10 @@ struct INDIArrayProxyBase : public FNiagaraDataInterfaceProxyRW
 #if WITH_NIAGARA_DEBUGGER
 	virtual void DrawDebugHud(FNDIDrawDebugHudContext& DebugHudContext) const = 0;
 #endif
+	virtual bool SimCacheWriteFrame(UNDIArraySimCacheData* CacheData, int FrameIndex, FNiagaraSystemInstance* SystemInstance) const = 0;
+	virtual bool SimCacheReadFrame(UNDIArraySimCacheData* CacheData, int FrameIndex, FNiagaraSystemInstance* SystemInstance) = 0;
+	virtual FString SimCacheVisualizerRead(UNDIArraySimCacheData* CacheData, FNDIArraySimCacheDataFrame& FrameData, int Element) const = 0;
+
 	virtual bool CopyToInternal(INDIArrayProxyBase* Destination) const = 0;
 	virtual bool Equals(const INDIArrayProxyBase* Other) const = 0;
 	virtual int32 PerInstanceDataSize() const = 0;
@@ -71,7 +78,7 @@ struct INDIArrayProxyBase : public FNiagaraDataInterfaceProxyRW
 };
 
 UCLASS(abstract, EditInlineNew, MinimalAPI)
-class UNiagaraDataInterfaceArray : public UNiagaraDataInterfaceRWBase
+class UNiagaraDataInterfaceArray : public UNiagaraDataInterfaceRWBase, public INiagaraSimCacheCustomStorageInterface
 {
 	GENERATED_UCLASS_BODY()
 
@@ -111,6 +118,14 @@ public:
 	NIAGARA_API virtual void SetShaderParameters(const FNiagaraDataInterfaceSetShaderParametersContext& Context) const override;
 	//UNiagaraDataInterface Interface
 
+	//~ INiagaraSimCacheCustomStorageInterface interface BEGIN
+	virtual UObject* SimCacheBeginWrite(UObject* SimCache, FNiagaraSystemInstance* NiagaraSystemInstance, const void* OptionalPerInstanceData, FNiagaraSimCacheFeedbackContext& FeedbackContext) const override;
+	virtual bool SimCacheWriteFrame(UObject* StorageObject, int FrameIndex, FNiagaraSystemInstance* SystemInstance, const void* OptionalPerInstanceData, FNiagaraSimCacheFeedbackContext& FeedbackContext) const override;
+	//virtual bool SimCacheEndWrite(UObject* StorageObject) const override;
+	virtual bool SimCacheReadFrame(UObject* StorageObject, int FrameA, int FrameB, float Interp, FNiagaraSystemInstance* SystemInstance, void* OptionalPerInstanceData) override;
+	//~ INiagaraSimCacheCustomStorageInterface interface END
+	NIAGARA_API FString SimCacheVisualizerRead(UNDIArraySimCacheData* CacheData, FNDIArraySimCacheDataFrame& FrameData, int Element) const;
+
 	/** ReadWrite lock to ensure safe access to the underlying array. */
 	FRWLock ArrayRWGuard;
 
@@ -127,3 +142,34 @@ protected:
 	virtual void GetFunctionsInternal(TArray<FNiagaraFunctionSignature>& OutFunctions) const override { GetProxyAs<INDIArrayProxyBase>()->GetFunctions(OutFunctions); }
 #endif
 };
+
+USTRUCT()
+struct FNDIArraySimCacheDataFrame
+{
+	GENERATED_BODY()
+		
+	UPROPERTY()
+	int32 NumElements = 0;
+
+	UPROPERTY()
+	int32 DataOffset = INDEX_NONE;
+};
+
+UCLASS(MinimalAPI)
+class UNDIArraySimCacheData : public UObject
+{
+	GENERATED_BODY()
+
+public:
+	UPROPERTY()
+	TArray<FNDIArraySimCacheDataFrame> CpuFrameData;
+
+	UPROPERTY()
+	TArray<FNDIArraySimCacheDataFrame> GpuFrameData;
+
+	UPROPERTY()
+	TArray<uint8> BufferData;
+
+	// Internal use only
+	int32 FindOrAddData(TConstArrayView<uint8> ArrayData);
+};
diff --git a/Engine/Plugins/FX/Niagara/Source/Niagara/Classes/NiagaraDataInterfaceArrayImpl.h b/Engine/Plugins/FX/Niagara/Source/Niagara/Classes/NiagaraDataInterfaceArrayImpl.h
--- a/Engine/Plugins/FX/Niagara/Source/Niagara/Classes/NiagaraDataInterfaceArrayImpl.h
+++ b/Engine/Plugins/FX/Niagara/Source/Niagara/Classes/NiagaraDataInterfaceArrayImpl.h
@@ -238,7 +238,7 @@ struct FNDIArrayInstanceData_RenderThread
 	}
 
 	template<typename T = FNDIArrayImplHelper<TArrayType>>
-	typename TEnableIf<T::bSupportsGPU>::Type UpdateDataImpl(FRHICommandList& RHICmdList, TArray<TArrayType>& InArrayData)
+	typename TEnableIf<T::bSupportsGPU>::Type UpdateDataImpl(FRHICommandList& RHICmdList, TConstArrayView<TArrayType> InArrayData)
 	{
 		const int32 NewNumElements = FMath::Max(DefaultElements, InArrayData.Num());
 
@@ -624,8 +624,6 @@ struct FNDIArrayProxyImpl : public INDIArrayProxyBase
 		return ToValueOut;
 	}
 
-	//////////////////////////////////////////////////////////////////////////
-	// VM accessors to ensure we maintain per correctness for shared data interfaces
 	void SetInstanceArrayData(FNiagaraSystemInstanceID InstanceID, const TArray<TArrayType>& InArrayData)
 	{
 		if ( FNDIArrayInstanceData_GameThread<TArrayType>* InstanceData = PerInstanceData_GameThread.FindRef(InstanceID) )
@@ -636,52 +634,6 @@ struct FNDIArrayProxyImpl : public INDIArrayProxyBase
 		}
 	}
 
-	TArray<TArrayType> GetArrayData(FNiagaraSystemInstanceID InstanceID)
-	{
-		TArray<TArrayType> ArrayDataOut;
-		if (FNDIArrayInstanceData_GameThread<TArrayType>* InstanceData = PerInstanceData_GameThread.FindRef(InstanceID))
-		{
-			FReadArrayRef ArrayData(Owner, InstanceData);
-			ArrayDataOut = ArrayData.GetArray();
-		}
-		return ArrayDataOut;
-	}
-
-	void SetArrayValue(FNiagaraSystemInstanceID InstanceID, int Index, const TArrayType& Value, bool bSizeToFit)
-	{
-		if (FNDIArrayInstanceData_GameThread<TArrayType>* InstanceData = PerInstanceData_GameThread.FindRef(InstanceID))
-		{
-			FWriteArrayRef ArrayData(Owner, InstanceData);
-			if (!ArrayData.GetArray().IsValidIndex(Index))
-			{
-				if (!bSizeToFit)
-				{
-					return;
-				}
-				ArrayData.GetArray().AddDefaulted(Index + 1 - ArrayData.GetArray().Num());
-			}
-
-			ArrayData.GetArray()[Index] = Value;
-			InstanceData->bIsRenderDirty |= bShouldSyncToGpu;
-		}
-	}
-
-	TArrayType GetArrayValue(FNiagaraSystemInstanceID InstanceID, int Index)
-	{
-		TArrayType ValueOut = TArrayType(FNDIArrayImplHelper<TArrayType>::GetDefaultValue());
-
-		if (FNDIArrayInstanceData_GameThread<TArrayType>* InstanceData = PerInstanceData_GameThread.FindRef(InstanceID))
-		{
-			FReadArrayRef ArrayData(Owner, InstanceData);
-			if (!ArrayData.GetArray().IsValidIndex(Index))
-			{
-				ValueOut = ArrayData.GetArray()[Index];
-			}
-		}
-
-		return ValueOut;
-	}
-
 #if WITH_EDITORONLY_DATA
 	//////////////////////////////////////////////////////////////////////////
 	// INDIArrayProxyBase
@@ -1033,6 +985,162 @@ struct FNDIArrayProxyImpl : public INDIArrayProxyBase
 	}
 #endif
 
+	virtual bool SimCacheWriteFrame(UNDIArraySimCacheData* CacheData, int FrameIndex, FNiagaraSystemInstance* SystemInstance) const override
+	{
+		const FNiagaraSystemInstanceID InstanceID = SystemInstance->GetId();
+		FNDIArrayInstanceData_GameThread<TArrayType>* InstanceData = PerInstanceData_GameThread.FindRef(InstanceID);
+		if (InstanceData == nullptr)
+		{
+			return false;
+		}
+
+		// Write CPU Data
+		{
+			FReadArrayRef ArrayData(Owner, InstanceData);
+			if (CacheData->CpuFrameData.Num() <= FrameIndex)
+			{
+				CacheData->CpuFrameData.AddDefaulted(FrameIndex + 1 - CacheData->CpuFrameData.Num());
+			}
+			FNDIArraySimCacheDataFrame& FrameData = CacheData->CpuFrameData[FrameIndex];
+			FrameData.NumElements = ArrayData.GetArray().Num();
+			FrameData.DataOffset = CacheData->FindOrAddData(
+				MakeArrayView(
+					reinterpret_cast<const uint8*>(ArrayData.GetArray().GetData()),
+					ArrayData.GetArray().Num() * ArrayData.GetArray().GetTypeSize()
+				)
+			);
+		}
+
+		// Write GPU Data
+		if (FNDIArrayImplHelper<TArrayType>::bSupportsGPU && Owner->IsUsedWithGPUScript())
+		{
+			FNiagaraGpuComputeDispatchInterface* ComputeInterface = FNiagaraGpuComputeDispatchInterface::Get(SystemInstance->GetWorld());
+			ENQUEUE_RENDER_COMMAND(NDIArray_SimCacheWrite)(
+				[Proxy_RT=this, InstanceID, CacheData, FrameIndex, ComputeInterface](FRHICommandListImmediate& RHICmdList)
+				{
+					const FNDIArrayInstanceData_RenderThread<TArrayType>* InstanceData_RT = Proxy_RT->PerInstanceData_RenderThread.Find(InstanceID);
+					if (InstanceData_RT == nullptr || InstanceData_RT->ArrayNumBytes == 0)
+					{
+						return;
+					}
+
+					if (CacheData->GpuFrameData.Num() <= FrameIndex)
+					{
+						CacheData->GpuFrameData.AddDefaulted(FrameIndex + 1 - CacheData->GpuFrameData.Num());
+					}
+					FNDIArraySimCacheDataFrame& FrameData = CacheData->GpuFrameData[FrameIndex];
+
+					TArray<FNiagaraGpuReadbackManager::FBufferRequest, TInlineAllocator<2>> BufferRequests;
+					TArray<FRHITransitionInfo, TInlineAllocator<2>> TransitionsBefore;
+					TArray<FRHITransitionInfo, TInlineAllocator<2>> TransitionsAfter;
+
+					BufferRequests.Emplace(InstanceData_RT->ArrayBuffer, 0, InstanceData_RT->ArrayNumBytes);
+					TransitionsBefore.Emplace(InstanceData_RT->ArrayBuffer, ERHIAccess::UAVCompute, ERHIAccess::CopySrc);
+					TransitionsAfter.Emplace(InstanceData_RT->ArrayBuffer, ERHIAccess::CopySrc, ERHIAccess::UAVCompute);
+
+					if (InstanceData_RT->IsReadOnly())
+					{
+						FrameData.NumElements = InstanceData_RT->NumElements;
+					}
+					else
+					{
+						const FNiagaraGPUInstanceCountManager& CountManager = ComputeInterface->GetGPUInstanceCounterManager();
+						BufferRequests.Emplace(CountManager.GetInstanceCountBuffer().Buffer, uint32(InstanceData_RT->CountOffset * sizeof(uint32)), sizeof(uint32));
+						TransitionsBefore.Emplace(CountManager.GetInstanceCountBuffer().UAV, ERHIAccess::UAVCompute, ERHIAccess::CopySrc);
+						TransitionsAfter.Emplace(CountManager.GetInstanceCountBuffer().UAV, ERHIAccess::CopySrc, ERHIAccess::UAVCompute);
+					}
+
+					FNiagaraGpuReadbackManager* ReadbackManager = ComputeInterface->GetGpuReadbackManager();
+					RHICmdList.Transition(TransitionsBefore);
+					ReadbackManager->EnqueueReadbacks(
+						RHICmdList,
+						BufferRequests,
+						[CacheData, &FrameData, bReadOnly=InstanceData_RT->IsReadOnly()](TConstArrayView<TPair<void*, uint32>> ReadbackData)
+						{
+							if (!bReadOnly)
+							{
+								FrameData.NumElements = *reinterpret_cast<const uint32*>(ReadbackData[1].Key);
+							}
+							if ( FrameData.NumElements > 0 )
+							{
+								TArray<TArrayType> ArrayData;
+								ArrayData.AddUninitialized(FrameData.NumElements);
+								FNDIArrayImplHelper<TArrayType>::CopyGpuToCpuMemory(ArrayData.GetData(), reinterpret_cast<const TVMArrayType*>(ReadbackData[0].Key), FrameData.NumElements);
+
+								FrameData.DataOffset = CacheData->FindOrAddData(
+									MakeArrayView(
+										reinterpret_cast<const uint8*>(ArrayData.GetData()),
+										ArrayData.Num() * ArrayData.GetTypeSize()
+									)
+								);
+							}
+						}
+					);
+					RHICmdList.Transition(TransitionsAfter);
+				}
+			);
+
+			FlushRenderingCommands();
+		}
+		return true;
+	}
+
+	virtual bool SimCacheReadFrame(UNDIArraySimCacheData* CacheData, int FrameIndex, FNiagaraSystemInstance* SystemInstance) override
+	{
+		const FNiagaraSystemInstanceID InstanceID = SystemInstance->GetId();
+		FNDIArrayInstanceData_GameThread<TArrayType>* InstanceData = PerInstanceData_GameThread.FindRef(InstanceID);
+		if (InstanceData == nullptr)
+		{
+			return false;
+		}
+
+		// Read CPU Data
+		if (CacheData->CpuFrameData.IsValidIndex(FrameIndex))
+		{
+			FWriteArrayRef ArrayData(Owner, InstanceData);
+			const FNDIArraySimCacheDataFrame& FrameData = CacheData->CpuFrameData[FrameIndex];
+			ArrayData.GetArray().SetNumUninitialized(FrameData.NumElements);
+			if (FrameData.NumElements > 0 )
+			{
+				check(FrameData.DataOffset != INDEX_NONE);
+				FMemory::Memcpy(ArrayData.GetArray().GetData(), CacheData->BufferData.GetData() + FrameData.DataOffset, FrameData.NumElements * sizeof(TArrayType));
+			}
+		}
+
+		// Read GPU Data
+		if (FNDIArrayImplHelper<TArrayType>::bSupportsGPU && Owner->IsUsedWithGPUScript())
+		{
+			if (CacheData->GpuFrameData.IsValidIndex(FrameIndex))
+			{
+				FNiagaraGpuComputeDispatchInterface* ComputeInterface = FNiagaraGpuComputeDispatchInterface::Get(SystemInstance->GetWorld());
+				ENQUEUE_RENDER_COMMAND(NDIArray_SimCacheWrite)(
+					[Proxy_RT=this, InstanceID, CacheData, FrameIndex, ComputeInterface](FRHICommandListImmediate& RHICmdList)
+					{
+						if ( FNDIArrayInstanceData_RenderThread<TArrayType>* InstanceData_RT = Proxy_RT->PerInstanceData_RenderThread.Find(InstanceID) )
+						{
+							const FNDIArraySimCacheDataFrame& FrameData = CacheData->GpuFrameData[FrameIndex];
+							TConstArrayView<TArrayType> ArrayData(reinterpret_cast<const TArrayType*>(CacheData->BufferData.GetData() + FrameData.DataOffset), FrameData.NumElements);
+							InstanceData_RT->UpdateDataImpl(RHICmdList, ArrayData);
+						}
+					}
+				);
+			}
+		}
+		return true;
+	}
+
+	virtual FString SimCacheVisualizerRead(UNDIArraySimCacheData* CacheData, FNDIArraySimCacheDataFrame& FrameData, int Element) const override
+	{
+		FString OutValue;
+		if (Element < FrameData.NumElements)
+		{
+			TArrayType Value;
+			FMemory::Memcpy(&Value, CacheData->BufferData.GetData() + FrameData.DataOffset + (sizeof(TArrayType) * Element), sizeof(TArrayType));
+			FNDIArrayImplHelper<TArrayType>::AppendValueToString(Value, OutValue);
+		}
+		return OutValue;
+	}
+
 	virtual bool CopyToInternal(INDIArrayProxyBase* InDestination) const override
 	{
 		auto Destination = static_cast<FNDIArrayProxyImpl<TArrayType, TOwnerType>*>(InDestination);
diff --git a/Engine/Plugins/FX/Niagara/Source/Niagara/Private/NiagaraDataInterfaceArray.cpp b/Engine/Plugins/FX/Niagara/Source/Niagara/Private/NiagaraDataInterfaceArray.cpp
--- a/Engine/Plugins/FX/Niagara/Source/Niagara/Private/NiagaraDataInterfaceArray.cpp
+++ b/Engine/Plugins/FX/Niagara/Source/Niagara/Private/NiagaraDataInterfaceArray.cpp
@@ -3,6 +3,7 @@
 #include "NiagaraDataInterfaceArray.h"
 #include "NiagaraCompileHashVisitor.h"
 #include "NiagaraShaderParametersBuilder.h"
+#include "NiagaraSystemInstance.h"
 
 #include UE_INLINE_GENERATED_CPP_BY_NAME(NiagaraDataInterfaceArray)
 
@@ -93,5 +94,52 @@ void UNiagaraDataInterfaceArray::SetShaderParameters(const FNiagaraDataInterface
 	ArrayProxy.SetShaderParameters(ShaderParameters, Context.GetSystemInstanceID());
 }
 
+UObject* UNiagaraDataInterfaceArray::SimCacheBeginWrite(UObject* SimCache, FNiagaraSystemInstance* NiagaraSystemInstance, const void* OptionalPerInstanceData, FNiagaraSimCacheFeedbackContext& FeedbackContext) const
+{
+	UNDIArraySimCacheData* CacheData = NewObject<UNDIArraySimCacheData>(SimCache);
+	return CacheData;
+}
+
+bool UNiagaraDataInterfaceArray::SimCacheWriteFrame(UObject* StorageObject, int FrameIndex, FNiagaraSystemInstance* SystemInstance, const void* OptionalPerInstanceData, FNiagaraSimCacheFeedbackContext& FeedbackContext) const
+{
+	check(OptionalPerInstanceData && StorageObject);
+	UNDIArraySimCacheData* CacheData = CastChecked<UNDIArraySimCacheData>(StorageObject);
+
+	const INDIArrayProxyBase* ArrayProxy = GetProxyAs<INDIArrayProxyBase>();
+	return ArrayProxy->SimCacheWriteFrame(CacheData, FrameIndex, SystemInstance);
+}
+
+bool UNiagaraDataInterfaceArray::SimCacheReadFrame(UObject* StorageObject, int FrameA, int FrameB, float Interp, FNiagaraSystemInstance* SystemInstance, void* OptionalPerInstanceData)
+{
+	check(OptionalPerInstanceData && StorageObject);
+	UNDIArraySimCacheData* CacheData = CastChecked<UNDIArraySimCacheData>(StorageObject);
+
+	INDIArrayProxyBase* ArrayProxy = GetProxyAs<INDIArrayProxyBase>();
+	return ArrayProxy->SimCacheReadFrame(CacheData, FrameA, SystemInstance);
+}
+
+FString UNiagaraDataInterfaceArray::SimCacheVisualizerRead(UNDIArraySimCacheData* CacheData, FNDIArraySimCacheDataFrame& FrameData, int Element) const
+{
+	const INDIArrayProxyBase* ArrayProxy = GetProxyAs<INDIArrayProxyBase>();
+	return ArrayProxy->SimCacheVisualizerRead(CacheData, FrameData, Element);
+}
+
+int32 UNDIArraySimCacheData::FindOrAddData(TConstArrayView<uint8> ArrayData)
+{
+	const int32 Last = BufferData.Num() - ArrayData.Num();
+	for (int32 i = 0; i <= Last; ++i)
+	{
+		if (FMemory::Memcmp(BufferData.GetData() + i, ArrayData.GetData(), ArrayData.Num()) == 0)
+		{
+			return i;
+		}
+	}
+
+	int32 NewOffset = BufferData.Num();
+	BufferData.AddUninitialized(ArrayData.Num());
+	FMemory::Memcpy(BufferData.GetData() + NewOffset, ArrayData.GetData(), ArrayData.Num());
+	return NewOffset;
+}
+
 #undef LOCTEXT_NAMESPACE
 
diff --git a/Engine/Plugins/FX/Niagara/Source/NiagaraEditor/Private/Customizations/SimCache/NiagaraArraySimCacheVisualizer.cpp b/Engine/Plugins/FX/Niagara/Source/NiagaraEditor/Private/Customizations/SimCache/NiagaraArraySimCacheVisualizer.cpp
new file mode 100644
--- /dev/null
+++ b/Engine/Plugins/FX/Niagara/Source/NiagaraEditor/Private/Customizations/SimCache/NiagaraArraySimCacheVisualizer.cpp
@@ -0,0 +1,288 @@
+// Copyright Epic Games, Inc. All Rights Reserved.
+
+#include "NiagaraArraySimCacheVisualizer.h"
+
+#include "NiagaraDataInterfaceArray.h"
+
+#include "UObject/StrongObjectPtr.h"
+#include "ViewModels/NiagaraSimCacheViewModel.h"
+#include "Widgets/Input/SSpinBox.h"
+#include "Widgets/Text/STextBlock.h"
+#include "DetailLayoutBuilder.h"
+#include "SEnumCombo.h"
+
+#define LOCTEXT_NAMESPACE "NiagaraArraySimCacheVisualizer"
+
+////////////////////////////////////////////////////////////////////////////////////////////////////////////////////
+
+namespace NDIArraySimCacheVisualizer
+{
+	static const FName NAME_Index = FName("Index");
+	static const FName NAME_CpuValue = FName("Cpu Value");
+	static const FName NAME_GpuValue = FName("Gpu Value");
+
+class SVisualizerRowWidget : public SMultiColumnTableRow<TSharedPtr<int32>>
+{
+public:
+	SLATE_BEGIN_ARGS(SVisualizerRowWidget) {}
+		SLATE_ARGUMENT(TSharedPtr<FNiagaraSimCacheViewModel>,	ViewModel)
+		SLATE_ARGUMENT(UNDIArraySimCacheData*,					CacheData)
+		SLATE_ARGUMENT(TSharedPtr<int32>,						RowIndexPtr)
+		SLATE_ARGUMENT(UNiagaraDataInterfaceArray*,				DataInterface)
+	SLATE_END_ARGS()
+
+	void Construct(const FArguments& InArgs, const TSharedRef<STableViewBase>& InOwnerTableView)
+	{
+		ViewModel			= InArgs._ViewModel;
+		CacheData.Reset(InArgs._CacheData);
+		RowIndexPtr			= InArgs._RowIndexPtr;
+		DataInterface		= InArgs._DataInterface;
+
+		SMultiColumnTableRow::Construct(
+			FSuperRowType::FArguments()
+			.Style(FAppStyle::Get(), "DataTableEditor.CellListViewRow"),
+			InOwnerTableView
+		);
+	}
+
+	virtual TSharedRef<SWidget> GenerateWidgetForColumn(const FName& InColumnName) override
+	{
+		if (InColumnName == NAME_Index)
+		{
+			return SNew(STextBlock)
+				.Text(FText::Format(LOCTEXT("IndexFormat", "{0}"), FText::AsNumber(*RowIndexPtr)));
+		}
+
+		const int32 FrameIndex = ViewModel->GetFrameIndex();
+		const int32 ElementIndex = *RowIndexPtr;
+
+		//-TODO: Print the data
+		if (InColumnName == NAME_CpuValue)
+		{
+			if (CacheData->CpuFrameData.IsValidIndex(FrameIndex))
+			{
+				FString Value = DataInterface->SimCacheVisualizerRead(CacheData.Get(), CacheData->CpuFrameData[FrameIndex], ElementIndex);
+				return SNew(STextBlock)
+					.Text(FText::FromString(Value));
+			}
+		}
+
+		if (InColumnName == NAME_GpuValue)
+		{
+			if (CacheData->GpuFrameData.IsValidIndex(FrameIndex))
+			{
+				FString Value = DataInterface->SimCacheVisualizerRead(CacheData.Get(), CacheData->GpuFrameData[FrameIndex], ElementIndex);
+				return SNew(STextBlock)
+					.Text(FText::FromString(Value));
+			}
+		}
+
+		return SNullWidget::NullWidget;
+	}
+
+	TSharedPtr<FNiagaraSimCacheViewModel>	ViewModel;
+	TStrongObjectPtr<UNDIArraySimCacheData>	CacheData;
+	TSharedPtr<int32>						RowIndexPtr;
+	UNiagaraDataInterfaceArray*				DataInterface = nullptr;
+};
+
+class SSimCacheView : public SCompoundWidget
+{
+public:
+	SLATE_BEGIN_ARGS(SSimCacheView)
+	{}
+	SLATE_END_ARGS()
+
+	virtual ~SSimCacheView() override
+	{
+		ViewModel->OnViewDataChanged().RemoveAll(this);
+	}
+
+	void Construct(const FArguments& InArgs, TSharedPtr<FNiagaraSimCacheViewModel> InViewModel, UNDIArraySimCacheData* InCacheData, UNiagaraDataInterfaceArray* InDataInterface)
+	{
+		ViewModel = InViewModel;
+		CacheData.Reset(InCacheData);
+		DataInterface = InDataInterface;
+
+		ViewModel->OnViewDataChanged().AddSP(this, &SSimCacheView::UpdateRows);
+
+		HeaderRowWidget = SNew(SHeaderRow);
+		UpdateColumns();
+
+		SAssignNew(ListViewWidget, SListView<TSharedPtr<int32>>)
+			.ListItemsSource(&RowItems)
+			.OnGenerateRow(this, &SSimCacheView::MakeRowWidget)
+			.Visibility(EVisibility::Visible)
+			.SelectionMode(ESelectionMode::Single)
+			.ConsumeMouseWheel(EConsumeMouseWheel::Never)
+			.ScrollbarVisibility(EVisibility::Collapsed)
+			.HeaderRow(HeaderRowWidget);
+
+		ChildSlot
+		[
+			SNew(SVerticalBox)
+			// Display Array Information
+			+ SVerticalBox::Slot()
+			.AutoHeight()
+			[
+				SNew(SHorizontalBox)
+				+SHorizontalBox::Slot()
+				[
+					SNew(STextBlock)
+					.Text(this, &SSimCacheView::GetCpuArrayInformation)
+				]
+			]
+			+ SVerticalBox::Slot()
+			.AutoHeight()
+			[
+				SNew(SHorizontalBox)
+				+SHorizontalBox::Slot()
+				[
+					SNew(STextBlock)
+					.Text(this, &SSimCacheView::GetGpuArrayInformation)
+				]
+			]
+			// Buffer Data View
+			+ SVerticalBox::Slot()
+			.AutoHeight()
+			.Padding(8.0f, 8.0f, 0.0f, 0.0f)
+			[
+				ListViewWidget.ToSharedRef()
+			]
+		];
+	}
+
+	TSharedRef<ITableRow> MakeRowWidget(const TSharedPtr<int32> RowIndexPtr, const TSharedRef<STableViewBase>& OwnerTable) const
+	{
+		return
+			SNew(SVisualizerRowWidget, OwnerTable)
+				.ViewModel(ViewModel)
+				.CacheData(CacheData.Get())
+				.RowIndexPtr(RowIndexPtr)
+				.DataInterface(DataInterface)
+				;
+	}
+
+	void UpdateColumns(bool bRefresh = true)
+	{
+		HeaderRowWidget->ClearColumns();
+		HeaderRowWidget->AddColumn(
+			SHeaderRow::Column(NAME_Index)
+			.DefaultLabel(LOCTEXT("Index", "Index"))
+			.HAlignHeader(HAlign_Center)
+			.VAlignHeader(VAlign_Fill)
+			.HAlignCell(HAlign_Center)
+			.VAlignCell(VAlign_Center)
+			.ManualWidth(100)
+			.SortMode(EColumnSortMode::None)
+		);
+
+		HeaderRowWidget->AddColumn(
+			SHeaderRow::Column(NAME_CpuValue)
+			.DefaultLabel(LOCTEXT("Cpu Value", "Cpu Value"))
+			.HAlignHeader(HAlign_Center)
+			.VAlignHeader(VAlign_Fill)
+			.HAlignCell(HAlign_Center)
+			.VAlignCell(VAlign_Center)
+			.ManualWidth(100)
+			.SortMode(EColumnSortMode::None)
+		);
+
+		HeaderRowWidget->AddColumn(
+			SHeaderRow::Column(NAME_GpuValue)
+			.DefaultLabel(LOCTEXT("Gpu Value", "Gpu Value"))
+			.HAlignHeader(HAlign_Center)
+			.VAlignHeader(VAlign_Fill)
+			.HAlignCell(HAlign_Center)
+			.VAlignCell(VAlign_Center)
+			.ManualWidth(100)
+			.SortMode(EColumnSortMode::None)
+		);
+
+		if (ListViewWidget && bRefresh)
+		{
+			ListViewWidget->RequestListRefresh();
+		}
+	}
+
+	void UpdateRows(bool bRefresh = true)
+	{
+		const int32 FrameIndex = ViewModel->GetFrameIndex();
+
+		int NumElements = 0;
+		if (CacheData->CpuFrameData.IsValidIndex(FrameIndex))
+		{
+			NumElements = FMath::Max(NumElements, CacheData->CpuFrameData[FrameIndex].NumElements);
+		}
+		if (CacheData->GpuFrameData.IsValidIndex(FrameIndex))
+		{
+			NumElements = FMath::Max(NumElements, CacheData->GpuFrameData[FrameIndex].NumElements);
+		}
+
+		RowItems.Reset(NumElements);
+		for ( int32 i=0; i < NumElements; ++i )
+		{
+			RowItems.Emplace(MakeShared<int32>(i));
+		}
+
+		if (ListViewWidget && bRefresh)
+		{
+			ListViewWidget->RequestListRefresh();
+		}
+	}
+
+
+	FText GetCpuArrayInformation() const
+	{
+		const int32 FrameIndex = ViewModel->GetFrameIndex();
+		int NumElements = 0;
+
+		if (CacheData->CpuFrameData.IsValidIndex(FrameIndex))
+		{
+			NumElements = CacheData->CpuFrameData[FrameIndex].NumElements;
+		}
+
+		return FText::Format(LOCTEXT("CpuArrayInfoFormat", "CPU Array Num Elements = {0}"), NumElements);
+	}
+
+	FText GetGpuArrayInformation() const
+	{
+		const int32 FrameIndex = ViewModel->GetFrameIndex();
+		int NumElements = 0;
+
+		if (CacheData->GpuFrameData.IsValidIndex(FrameIndex))
+		{
+			NumElements = CacheData->GpuFrameData[FrameIndex].NumElements;
+		}
+
+		return FText::Format(LOCTEXT("CpuArrayInfoFormat", "GPU Array Num Elements = {0}"), NumElements);
+	}
+
+private:
+	TSharedPtr<FNiagaraSimCacheViewModel>		ViewModel;
+	TStrongObjectPtr<UNDIArraySimCacheData>		CacheData;
+
+	TSharedPtr<SHeaderRow>						HeaderRowWidget;
+	TSharedPtr<SListView<TSharedPtr<int32>>>	ListViewWidget;
+	TArray<TSharedPtr<int32>>					RowItems;
+	UNiagaraDataInterfaceArray*					DataInterface = nullptr;
+};
+
+} // NDIArraySimCacheVisualizer
+
+////////////////////////////////////////////////////////////////////////////////////////////////////////////////////
+
+TSharedPtr<SWidget> FNiagaraArraySimCacheVisualizer::CreateWidgetFor(UObject* InCachedData, TSharedPtr<FNiagaraSimCacheViewModel> ViewModel)
+{
+	using namespace NDIArraySimCacheVisualizer;
+
+	if (UNDIArraySimCacheData* CachedData = Cast<UNDIArraySimCacheData>(InCachedData))
+	{
+		check(ArrayDIClass);
+		UNiagaraDataInterfaceArray* DataInterface = CastChecked<UNiagaraDataInterfaceArray>(ArrayDIClass->GetDefaultObject());		
+		return SNew(SSimCacheView, ViewModel, CachedData, DataInterface);
+	}
+	return TSharedPtr<SWidget>();
+}
+
+#undef LOCTEXT_NAMESPACE
diff --git a/Engine/Plugins/FX/Niagara/Source/NiagaraEditor/Private/Customizations/SimCache/NiagaraArraySimCacheVisualizer.h b/Engine/Plugins/FX/Niagara/Source/NiagaraEditor/Private/Customizations/SimCache/NiagaraArraySimCacheVisualizer.h
new file mode 100644
--- /dev/null
+++ b/Engine/Plugins/FX/Niagara/Source/NiagaraEditor/Private/Customizations/SimCache/NiagaraArraySimCacheVisualizer.h
@@ -0,0 +1,26 @@
+// Copyright Epic Games, Inc. All Rights Reserved.
+
+#pragma once
+
+#include "NiagaraCommon.h"
+#include "Customizations/NiagaraDataInterfaceSimCacheVisualizer.h"
+
+//#include "NiagaraArraySimCacheVisualizer.generated.h"
+
+/*
+ * Provides a custom widget to show the array DI data in the sim cache in a table.
+ */
+class FNiagaraArraySimCacheVisualizer : public INiagaraDataInterfaceSimCacheVisualizer
+{
+public:
+	FNiagaraArraySimCacheVisualizer(UClass* InArrayDIClass)
+		: ArrayDIClass(InArrayDIClass)
+	{
+	}
+	virtual ~FNiagaraArraySimCacheVisualizer() override = default;
+	
+	virtual TSharedPtr<SWidget> CreateWidgetFor(UObject* CachedData, TSharedPtr<FNiagaraSimCacheViewModel> ViewModel) override;
+
+protected:
+	UClass* ArrayDIClass = nullptr;
+};
diff --git a/Engine/Plugins/FX/Niagara/Source/NiagaraEditor/Private/NiagaraEditorModule.cpp b/Engine/Plugins/FX/Niagara/Source/NiagaraEditor/Private/NiagaraEditorModule.cpp
--- a/Engine/Plugins/FX/Niagara/Source/NiagaraEditor/Private/NiagaraEditorModule.cpp
+++ b/Engine/Plugins/FX/Niagara/Source/NiagaraEditor/Private/NiagaraEditorModule.cpp
@@ -47,6 +47,7 @@
 #include "NiagaraSettings.h"
 #include "NiagaraShaderModule.h"
 
+#include "NiagaraDataInterfaceArray.h"
 #include "NiagaraDataInterfaceCurve.h"
 #include "NiagaraDataInterfaceVector2DCurve.h"
 #include "NiagaraDataInterfaceVectorCurve.h"
@@ -114,6 +115,7 @@
 #include "Customizations/NiagaraOutlinerCustomization.h"
 #include "Customizations/NiagaraSimulationStageCustomization.h"
 #include "Customizations/NiagaraDataChannelDetails.h"
+#include "Customizations/SimCache/NiagaraArraySimCacheVisualizer.h"
 #include "Customizations/SimCache/FNiagaraDataChannelSimCacheVisualizer.h"
 #include "Customizations/SimCache/NiagaraRenderTargetVolumeSimCacheVisualizer.h"
 
@@ -1429,6 +1431,13 @@ void FNiagaraEditorModule::StartupModule()
 	// Register sim cache visualizers
 	RegisterDataInterfaceCacheVisualizer(UNiagaraDataInterfaceDataChannelWrite::StaticClass(), MakeShared<FNiagaraDataChannelSimCacheVisualizer>());
 	RegisterDataInterfaceCacheVisualizer(UNiagaraDataInterfaceRenderTargetVolume::StaticClass(), MakeShared<FNiagaraRenderTargetVolumeSimCacheVisualizer>());
+	for (TObjectIterator<UClass> It; It; ++It)
+	{
+		if (It->IsChildOf(UNiagaraDataInterfaceArray::StaticClass()))
+		{
+			RegisterDataInterfaceCacheVisualizer(*It, MakeShared<FNiagaraArraySimCacheVisualizer>(*It));
+		}
+	}
 
 #if NIAGARA_PERF_BASELINES
 	UNiagaraEffectType::OnGeneratePerfBaselines().BindRaw(this, &FNiagaraEditorModule::GeneratePerfBaselines);
