- Memory Buffer Data Interface
- Allows user to allocate a block of memory to load / store floats & ints however they want

diff --git a/Engine/Plugins/FX/Niagara/Shaders/Private/NiagaraDataInterfaceMemoryBufferTemplate.ush b/Engine/Plugins/FX/Niagara/Shaders/Private/NiagaraDataInterfaceMemoryBufferTemplate.ush
new file mode 100644
--- /dev/null
+++ b/Engine/Plugins/FX/Niagara/Shaders/Private/NiagaraDataInterfaceMemoryBufferTemplate.ush
@@ -0,0 +1,230 @@
+// Copyright Epic Games, Inc. All Rights Reserved.
+
+int					{ParameterName}_NumElements;
+RWByteAddressBuffer	{ParameterName}_Buffer;
+
+/////////////////////////////////////////////////////////////////////////////
+// Helper functions
+bool ValidateOffset_{ParameterName}(inout int ElementOffset, int Size)
+{
+	bool bValid = ElementOffset >= 0 && ElementOffset <= {ParameterName}_NumElements - Size;
+	ElementOffset = clamp(ElementOffset, 0, {ParameterName}_NumElements);
+	return bValid;
+}
+
+/////////////////////////////////////////////////////////////////////////////
+// Misc Functions
+void GetNumElements_{ParameterName}(out int NumElements)
+{
+	NumElements = {ParameterName}_NumElements;
+}
+
+/////////////////////////////////////////////////////////////////////////////
+// Load operations
+void LoadInt_{ParameterName}(int ElementOffset, out bool bValid, out int OutValue)
+{
+	bValid = ValidateOffset_{ParameterName}(ElementOffset, 1);
+	const uint BufferValue = {ParameterName}_Buffer.Load(ElementOffset << 2);
+	OutValue = asint(BufferValue);
+}
+
+void LoadInt2_{ParameterName}(int ElementOffset, out bool bValid, out int OutValue0, out int OutValue1)
+{
+	bValid = ValidateOffset_{ParameterName}(ElementOffset, 2);
+	const uint2 BufferValue = {ParameterName}_Buffer.Load2(ElementOffset << 2);
+	OutValue0 = asint(BufferValue.x);
+	OutValue1 = asint(BufferValue.y);
+}
+
+void LoadInt3_{ParameterName}(int ElementOffset, out bool bValid, out int OutValue0, out int OutValue1, out int OutValue2)
+{
+	bValid = ValidateOffset_{ParameterName}(ElementOffset, 3);
+	const uint3 BufferValue = {ParameterName}_Buffer.Load3(ElementOffset << 2);
+	OutValue0 = asint(BufferValue.x);
+	OutValue1 = asint(BufferValue.y);
+	OutValue2 = asint(BufferValue.z);
+}
+
+void LoadInt4_{ParameterName}(int ElementOffset, out bool bValid, out int OutValue0, out int OutValue1, out int OutValue2, out int OutValue3)
+{
+	bValid = ValidateOffset_{ParameterName}(ElementOffset, 4);
+	const uint4 BufferValue = {ParameterName}_Buffer.Load4(ElementOffset << 2);
+	OutValue0 = asint(BufferValue.x);
+	OutValue1 = asint(BufferValue.y);
+	OutValue2 = asint(BufferValue.z);
+	OutValue3 = asint(BufferValue.w);
+}
+
+void LoadFloat_{ParameterName}(int ElementOffset, out bool bValid, out float OutValue)
+{
+	bValid = ValidateOffset_{ParameterName}(ElementOffset, 1);
+	const uint BufferValue = {ParameterName}_Buffer.Load(ElementOffset << 2);
+	OutValue = asfloat(BufferValue);
+}
+
+void LoadFloat2_{ParameterName}(int ElementOffset, out bool bValid, out float2 OutValue)
+{
+	bValid = ValidateOffset_{ParameterName}(ElementOffset, 2);
+	const uint2 BufferValue = {ParameterName}_Buffer.Load2(ElementOffset << 2);
+	OutValue = asfloat(BufferValue);
+}
+
+void LoadFloat3_{ParameterName}(int ElementOffset, out bool bValid, out float3 OutValue)
+{
+	bValid = ValidateOffset_{ParameterName}(ElementOffset, 3);
+	const uint3 BufferValue = {ParameterName}_Buffer.Load3(ElementOffset << 2);
+	OutValue = asfloat(BufferValue);
+}
+
+void LoadFloat4_{ParameterName}(int ElementOffset, out bool bValid, out float4 OutValue)
+{
+	bValid = ValidateOffset_{ParameterName}(ElementOffset, 4);
+	const uint4 BufferValue = {ParameterName}_Buffer.Load4(ElementOffset << 2);
+	OutValue = asfloat(BufferValue);
+}
+
+/////////////////////////////////////////////////////////////////////////////
+// Store operations
+void StoreInt_{ParameterName}_UEImpureCall(bool bExecute, int ElementOffset, int Value, out bool bSuccess)
+{
+	bSuccess = bExecute && ValidateOffset_{ParameterName}(ElementOffset, 1);
+	if (bExecute)
+	{
+		{ParameterName}_Buffer.Store(ElementOffset << 2, asuint(Value));
+	}
+}
+
+void StoreInt2_{ParameterName}_UEImpureCall(bool bExecute, int ElementOffset, int Value0, int Value1, out bool bSuccess)
+{
+	bSuccess = bExecute && ValidateOffset_{ParameterName}(ElementOffset, 2);
+	if (bExecute)
+	{
+		{ParameterName}_Buffer.Store(ElementOffset << 2, asuint(int2(Value0, Value1)));
+	}
+}
+
+void StoreInt3_{ParameterName}_UEImpureCall(bool bExecute, int ElementOffset, int Value0, int Value1, int Value2, out bool bSuccess)
+{
+	bSuccess = bExecute && ValidateOffset_{ParameterName}(ElementOffset, 3);
+	if (bExecute)
+	{
+		{ParameterName}_Buffer.Store(ElementOffset << 2, asuint(int3(Value0, Value1, Value2)));
+	}
+}
+
+void StoreInt4_{ParameterName}_UEImpureCall(bool bExecute, int ElementOffset, int Value0, int Value1, int Value2, int Value3, out bool bSuccess)
+{
+	bSuccess = bExecute && ValidateOffset_{ParameterName}(ElementOffset, 4);
+	if (bExecute)
+	{
+		{ParameterName}_Buffer.Store(ElementOffset << 2, asuint(int4(Value0, Value1, Value2, Value3)));
+	}
+}
+
+void StoreFloat_{ParameterName}_UEImpureCall(bool bExecute, int ElementOffset, float Value, out bool bSuccess)
+{
+	bSuccess = bExecute && ValidateOffset_{ParameterName}(ElementOffset, 1);
+	if (bExecute)
+	{
+		{ParameterName}_Buffer.Store(ElementOffset << 2, asuint(Value));
+	}
+}
+
+void StoreFloat2_{ParameterName}_UEImpureCall(bool bExecute, int ElementOffset, float2 Value, out bool bSuccess)
+{
+	bSuccess = bExecute && ValidateOffset_{ParameterName}(ElementOffset, 2);
+	if (bExecute)
+	{
+		{ParameterName}_Buffer.Store(ElementOffset << 2, asuint(Value));
+	}
+}
+
+void StoreFloat3_{ParameterName}_UEImpureCall(bool bExecute, int ElementOffset, float3 Value, out bool bSuccess)
+{
+	bSuccess = bExecute && ValidateOffset_{ParameterName}(ElementOffset, 3);
+	if (bExecute)
+	{
+		{ParameterName}_Buffer.Store(ElementOffset << 2, asuint(Value));
+	}
+}
+
+void StoreFloat4_{ParameterName}_UEImpureCall(bool bExecute, int ElementOffset, float4 Value, out bool bSuccess)
+{
+	bSuccess = bExecute && ValidateOffset_{ParameterName}(ElementOffset, 4);
+	if (bExecute)
+	{
+		{ParameterName}_Buffer.Store(ElementOffset << 2, asuint(Value));
+	}
+}
+
+/////////////////////////////////////////////////////////////////////////////
+// Atomic Ops
+void AtomicAdd_{ParameterName}_UEImpureCall(bool bExecute, int ElementOffset, int Value, out int PreviousValue, out int CurrentValue)
+{
+	if (bExecute)
+	{
+		ElementOffset = clamp(ElementOffset, 0, {ParameterName}_NumElements) << 2;
+		{ParameterName}_Buffer.InterlockedAdd(ElementOffset, Value, PreviousValue);
+		CurrentValue = PreviousValue + Value;
+	}
+}
+
+void AtomicAnd_{ParameterName}_UEImpureCall(bool bExecute, int ElementOffset, int Value, out int PreviousValue, out int CurrentValue)
+{
+	if (bExecute)
+	{
+		ElementOffset = clamp(ElementOffset, 0, {ParameterName}_NumElements) << 2;
+		{ParameterName}_Buffer.InterlockedAnd(ElementOffset, Value, PreviousValue);
+		CurrentValue = PreviousValue & Value;
+	}
+}
+
+void AtomicCompareAndExchange_{ParameterName}_UEImpureCall(bool bExecute, int ElementOffset, int Value, int Comperand, out int PreviousValue, out int CurrentValue)
+{
+	if (bExecute)
+	{
+		ElementOffset = clamp(ElementOffset, 0, {ParameterName}_NumElements) << 2;
+		{ParameterName}_Buffer.InterlockedCompareExchange(ElementOffset, Comperand, Value, PreviousValue);
+		CurrentValue = PreviousValue == Value ? Value : PreviousValue;
+	}
+}
+
+void AtomicMax_{ParameterName}_UEImpureCall(bool bExecute, int ElementOffset, int Value, out int PreviousValue, out int CurrentValue)
+{
+	if (bExecute)
+	{
+		ElementOffset = clamp(ElementOffset, 0, {ParameterName}_NumElements) << 2;
+		{ParameterName}_Buffer.InterlockedMax(ElementOffset, Value, PreviousValue);
+		CurrentValue = max(PreviousValue, Value);
+	}
+}
+
+void AtomicMin_{ParameterName}_UEImpureCall(bool bExecute, int ElementOffset, int Value, out int PreviousValue, out int CurrentValue)
+{
+	if (bExecute)
+	{
+		ElementOffset = clamp(ElementOffset, 0, {ParameterName}_NumElements) << 2;
+		{ParameterName}_Buffer.InterlockedMin(ElementOffset, Value, PreviousValue);
+		CurrentValue = min(PreviousValue, Value);
+	}
+}
+
+void AtomicOr_{ParameterName}_UEImpureCall(bool bExecute, int ElementOffset, int Value, out int PreviousValue, out int CurrentValue)
+{
+	if (bExecute)
+	{
+		ElementOffset = clamp(ElementOffset, 0, {ParameterName}_NumElements) << 2;
+		{ParameterName}_Buffer.InterlockedOr(ElementOffset, Value, PreviousValue);
+		CurrentValue = PreviousValue | Value;
+	}
+}
+
+void AtomicXor_{ParameterName}_UEImpureCall(bool bExecute, int ElementOffset, int Value, out int PreviousValue, out int CurrentValue)
+{
+	if (bExecute)
+	{
+		ElementOffset = clamp(ElementOffset, 0, {ParameterName}_NumElements) << 2;
+		{ParameterName}_Buffer.InterlockedXor(ElementOffset, Value, PreviousValue);
+		CurrentValue = PreviousValue ^ Value;
+	}
+}
diff --git a/Engine/Plugins/FX/Niagara/Source/Niagara/Classes/NiagaraDataInterface.h b/Engine/Plugins/FX/Niagara/Source/Niagara/Classes/NiagaraDataInterface.h
--- a/Engine/Plugins/FX/Niagara/Source/Niagara/Classes/NiagaraDataInterface.h
+++ b/Engine/Plugins/FX/Niagara/Source/Niagara/Classes/NiagaraDataInterface.h
@@ -1383,6 +1383,38 @@ struct FNDIInputParam<FNiagaraEmitterID>
 	FORCEINLINE void Reset() { Index.Reset(); }
 };
 
+template<>
+struct FNDIInputParam<FIntVector2>
+{
+	VectorVM::FExternalFuncInputHandler<int32> X;
+	VectorVM::FExternalFuncInputHandler<int32> Y;
+	FORCEINLINE FNDIInputParam(FVectorVMExternalFunctionContext& Context) : X(Context), Y(Context) {}
+	FORCEINLINE FIntVector2 GetAndAdvance() { return FIntVector2(X.GetAndAdvance(), Y.GetAndAdvance()); }
+	FORCEINLINE bool IsConstant() const { return X.IsConstant() && Y.IsConstant(); }
+};
+
+template<>
+struct FNDIInputParam<FIntVector3>
+{
+	VectorVM::FExternalFuncInputHandler<int32> X;
+	VectorVM::FExternalFuncInputHandler<int32> Y;
+	VectorVM::FExternalFuncInputHandler<int32> Z;
+	FORCEINLINE FNDIInputParam(FVectorVMExternalFunctionContext& Context) : X(Context), Y(Context), Z(Context) {}
+	FORCEINLINE FIntVector3 GetAndAdvance() { return FIntVector3(X.GetAndAdvance(), Y.GetAndAdvance(), Z.GetAndAdvance()); }
+	FORCEINLINE bool IsConstant() const { return X.IsConstant() && Y.IsConstant() && Z.IsConstant(); }
+};
+
+template<>
+struct FNDIInputParam<FIntVector4>
+{
+	VectorVM::FExternalFuncInputHandler<int32> X;
+	VectorVM::FExternalFuncInputHandler<int32> Y;
+	VectorVM::FExternalFuncInputHandler<int32> Z;
+	VectorVM::FExternalFuncInputHandler<int32> W;
+	FORCEINLINE FNDIInputParam(FVectorVMExternalFunctionContext& Context) : X(Context), Y(Context), Z(Context), W(Context) {}
+	FORCEINLINE FIntVector4 GetAndAdvance() { return FIntVector4(X.GetAndAdvance(), Y.GetAndAdvance(), Z.GetAndAdvance(), W.GetAndAdvance()); }
+	FORCEINLINE bool IsConstant() const { return X.IsConstant() && Y.IsConstant() && Z.IsConstant() && W.IsConstant(); }
+};
 
 //Helper to deal with types with potentially several output registers.
 template<typename T>
@@ -1572,6 +1604,54 @@ struct FNDIOutputParam<FNiagaraID>
 	}
 };
 
+template<>
+struct FNDIOutputParam<FIntVector2>
+{
+	VectorVM::FExternalFuncRegisterHandler<int32> X;
+	VectorVM::FExternalFuncRegisterHandler<int32> Y;
+	FORCEINLINE FNDIOutputParam(FVectorVMExternalFunctionContext& Context) : X(Context), Y(Context) {}
+	FORCEINLINE bool IsValid() const { return X.IsValid() || Y.IsValid(); }
+	FORCEINLINE void SetAndAdvance(FIntVector2 Val)
+	{
+		*X.GetDestAndAdvance() = Val.X;
+		*Y.GetDestAndAdvance() = Val.Y;
+	}
+};
+
+template<>
+struct FNDIOutputParam<FIntVector3>
+{
+	VectorVM::FExternalFuncRegisterHandler<int32> X;
+	VectorVM::FExternalFuncRegisterHandler<int32> Y;
+	VectorVM::FExternalFuncRegisterHandler<int32> Z;
+	FORCEINLINE FNDIOutputParam(FVectorVMExternalFunctionContext& Context) : X(Context), Y(Context), Z(Context) {}
+	FORCEINLINE bool IsValid() const { return X.IsValid() || Y.IsValid() || Z.IsValid(); }
+	FORCEINLINE void SetAndAdvance(FIntVector3 Val)
+	{
+		*X.GetDestAndAdvance() = Val.X;
+		*Y.GetDestAndAdvance() = Val.Y;
+		*Z.GetDestAndAdvance() = Val.Z;
+	}
+};
+
+template<>
+struct FNDIOutputParam<FIntVector4>
+{
+	VectorVM::FExternalFuncRegisterHandler<int32> X;
+	VectorVM::FExternalFuncRegisterHandler<int32> Y;
+	VectorVM::FExternalFuncRegisterHandler<int32> Z;
+	VectorVM::FExternalFuncRegisterHandler<int32> W;
+	FORCEINLINE FNDIOutputParam(FVectorVMExternalFunctionContext& Context) : X(Context), Y(Context), Z(Context), W(Context) {}
+	FORCEINLINE bool IsValid() const { return X.IsValid() || Y.IsValid() || Z.IsValid() || W.IsValid(); }
+	FORCEINLINE void SetAndAdvance(FIntVector4 Val)
+	{
+		*X.GetDestAndAdvance() = Val.X;
+		*Y.GetDestAndAdvance() = Val.Y;
+		*Z.GetDestAndAdvance() = Val.Z;
+		*W.GetDestAndAdvance() = Val.W;
+	}
+};
+
 class FNDI_GeneratedData
 {
 public:
diff --git a/Engine/Plugins/FX/Niagara/Source/Niagara/Internal/DataInterface/NiagaraDataInterfaceMemoryBuffer.h b/Engine/Plugins/FX/Niagara/Source/Niagara/Internal/DataInterface/NiagaraDataInterfaceMemoryBuffer.h
new file mode 100644
--- /dev/null
+++ b/Engine/Plugins/FX/Niagara/Source/Niagara/Internal/DataInterface/NiagaraDataInterfaceMemoryBuffer.h
@@ -0,0 +1,108 @@
+// Copyright Epic Games, Inc. All Rights Reserved.
+
+#pragma once
+
+#include "NiagaraCommon.h"
+#include "NiagaraDataInterface.h"
+#include "NiagaraParameterBinding.h"
+#include "NiagaraRenderableMeshInterface.h"
+#include "NiagaraSimCacheCustomStorageInterface.h"
+
+#include "Serialization/BulkData.h"
+#include "Serialization/BulkDataBuffer.h"
+#include "VectorVM.h"
+
+#include "NiagaraDataInterfaceMemoryBuffer.generated.h"
+
+struct FNiagaraSimCacheFeedbackContext;
+
+/**
+Data interface used to access a memory buffer.
+The user is responsible for how data is read / wrote.
+The DI will ensure no out of bounds access but not that the elements are of the correct type (i.e. float or int)
+Elements are considered to be 4 bytes in size.
+*/
+UCLASS(MinimalAPI, Experimental, EditInlineNew, Category = "Data Interface", meta = (DisplayName = "Memory Buffer"))
+class UNiagaraDataInterfaceMemoryBuffer : public UNiagaraDataInterface, public INiagaraSimCacheCustomStorageInterface
+{
+	GENERATED_UCLASS_BODY()
+
+public:
+	/** Default space we will allocate for the memory buffer. */
+	UPROPERTY(EditAnywhere, Category = "Byte Buffer")
+	int32 DefaultNumElements = 0;
+
+	/** How should we synhronize the data between CPU and GPU memory. */
+	UPROPERTY(EditAnywhere, Category = "Byte Buffer")
+	ENiagaraGpuSyncMode GpuSyncMode = ENiagaraGpuSyncMode::SyncCpuToGpu;
+
+	//UObject Interface
+	virtual void PostInitProperties() override;
+	//UObject Interface End
+
+	//UNiagaraDataInterface Interface
+#if WITH_EDITORONLY_DATA
+	virtual void GetFunctionsInternal(TArray<FNiagaraFunctionSignature>& OutFunctions) const override;
+#endif
+	virtual void GetVMExternalFunction(const FVMExternalFunctionBindingInfo& BindingInfo, void* InstanceData, FVMExternalFunction &OutFunc) override;
+	virtual bool CanExecuteOnTarget(ENiagaraSimTarget Target) const override { return true; }
+
+	virtual int32 PerInstanceDataSize() const override;
+	virtual bool InitPerInstanceData(void* PerInstanceData, FNiagaraSystemInstance* SystemInstance) override;
+	virtual void DestroyPerInstanceData(void* PerInstanceData, FNiagaraSystemInstance* SystemInstance) override;
+
+	virtual bool HasPreSimulateTick() const override { return true; }
+	virtual bool HasPostSimulateTick() const override { return true; }
+	virtual bool PerInstanceTickPostSimulate(void* PerInstanceData, FNiagaraSystemInstance* SystemInstance, float DeltaSeconds) override;
+
+	virtual bool Equals(const UNiagaraDataInterface* Other) const override;
+	virtual bool CopyToInternal(UNiagaraDataInterface* Destination) const override;
+
+	// GPU sim functionality
+#if WITH_EDITORONLY_DATA
+	virtual bool AppendCompileHash(FNiagaraCompileHashVisitor* InVisitor) const override;
+	virtual void GetParameterDefinitionHLSL(const FNiagaraDataInterfaceGPUParamInfo& ParamInfo, FString& OutHLSL) override;
+	virtual bool GetFunctionHLSL(const FNiagaraDataInterfaceGPUParamInfo& ParamInfo, const FNiagaraDataInterfaceGeneratedFunction& FunctionInfo, int FunctionInstanceIndex, FString& OutHLSL) override;
+#endif
+	virtual void BuildShaderParameters(FNiagaraShaderParametersBuilder& ShaderParametersBuilder) const override;
+	virtual void SetShaderParameters(const FNiagaraDataInterfaceSetShaderParametersContext& Context) const override;
+	//UNiagaraDataInterface Interface End
+
+	//~ INiagaraSimCacheCustomStorageInterface interface BEGIN
+	virtual UObject* SimCacheBeginWrite(UObject* SimCache, FNiagaraSystemInstance* NiagaraSystemInstance, const void* OptionalPerInstanceData, FNiagaraSimCacheFeedbackContext& FeedbackContext) const override;
+	virtual bool SimCacheWriteFrame(UObject* StorageObject, int FrameIndex, FNiagaraSystemInstance* SystemInstance, const void* OptionalPerInstanceData, FNiagaraSimCacheFeedbackContext& FeedbackContext) const override;
+	//virtual bool SimCacheEndWrite(UObject* StorageObject) const override;
+	virtual bool SimCacheReadFrame(UObject* StorageObject, int FrameA, int FrameB, float Interp, FNiagaraSystemInstance* SystemInstance, void* OptionalPerInstanceData) override;
+	//~ UNiagaraDataInterface interface END
+};
+
+USTRUCT()
+struct FNDIMemoryBufferSimCacheDataFrame
+{
+	GENERATED_BODY()
+		
+	UPROPERTY()
+	int32 CpuBufferSize = 0;
+
+	UPROPERTY()
+	int32 CpuDataOffset = INDEX_NONE;
+
+	UPROPERTY()
+	int32 GpuBufferSize = 0;
+
+	UPROPERTY()
+	int32 GpuDataOffset = INDEX_NONE;
+};
+
+UCLASS(MinimalAPI)
+class UNDIMemoryBufferSimCacheData : public UObject
+{
+	GENERATED_BODY()
+
+public:
+	UPROPERTY()
+	TArray<FNDIMemoryBufferSimCacheDataFrame> FrameData;
+
+	UPROPERTY()
+	TArray<uint32> BufferData;
+};
diff --git a/Engine/Plugins/FX/Niagara/Source/Niagara/Private/DataInterface/NiagaraDataInterfaceMemoryBuffer.cpp b/Engine/Plugins/FX/Niagara/Source/Niagara/Private/DataInterface/NiagaraDataInterfaceMemoryBuffer.cpp
new file mode 100644
--- /dev/null
+++ b/Engine/Plugins/FX/Niagara/Source/Niagara/Private/DataInterface/NiagaraDataInterfaceMemoryBuffer.cpp
@@ -0,0 +1,1052 @@
+// Copyright Epic Games, Inc. All Rights Reserved.
+
+#include "DataInterface/NiagaraDataInterfaceMemoryBuffer.h"
+#include "NiagaraCompileHashVisitor.h"
+#include "NiagaraDataInterfaceUtilities.h"
+#include "NiagaraMeshVertexFactory.h"
+#include "NiagaraParameterStore.h"
+#include "NiagaraGpuComputeDispatchInterface.h"
+#include "NiagaraGpuReadbackManager.h"
+#include "NiagaraSimStageData.h"
+#include "NiagaraShaderParametersBuilder.h"
+#include "NiagaraSystemInstance.h"
+#include "NiagaraWorldManager.h"
+
+#include "Materials/MaterialRenderProxy.h"
+#include "Materials/Material.h"
+#include "MaterialDomain.h"
+#include "RenderGraphBuilder.h"
+#include "RenderGraphUtils.h"
+#include "UnifiedBuffer.h"
+
+#include UE_INLINE_GENERATED_CPP_BY_NAME(NiagaraDataInterfaceMemoryBuffer)
+
+#define LOCTEXT_NAMESPACE "UNiagaraDataInterfaceMemoryBuffer"
+
+namespace NDIMemoryBufferLocal
+{
+	BEGIN_SHADER_PARAMETER_STRUCT(FShaderParameters, )
+		SHADER_PARAMETER(int32,									NumElements)
+		SHADER_PARAMETER_RDG_BUFFER_UAV(RWByteAddressBuffer,	Buffer)
+	END_SHADER_PARAMETER_STRUCT()
+
+	const TCHAR*	TemplateShaderFilePath = TEXT("/Plugin/FX/Niagara/Private/NiagaraDataInterfaceMemoryBufferTemplate.ush");
+
+	const FName NAME_SetNumElements("SetNumElements");
+	const FName NAME_GetNumElements("GetNumElements");
+
+	const FName NAME_ClearBufferInt("ClearBufferInt");
+	const FName NAME_ClearBufferFloat("ClearBufferFloat");
+
+	template<typename T>
+	T AlignMemoryBuffer(T Size)
+	{
+		return Size + T(4);
+	}
+
+	struct FNDIInstanceData_GameThread
+	{
+		bool				bNeedsGpuSync = true;		// We have a pending writes on the CPU that can be sent to the GPU
+		bool				bUsedWithCpu = false;		// Only true if we ever read / write to the buffer on the CPU, clearing for example does not qualify
+		bool				bUsedWithGpu = false;		// Are we used on the GPU
+		FRWLock				CpuBufferGuard;
+		TArray<uint32>		CpuBuffer;
+		int32				NumElements = 0;
+
+		ENiagaraGpuSyncMode	GpuSyncMode = ENiagaraGpuSyncMode::None;
+		TOptional<uint32>	GpuPendingClear;
+
+		void UpdateCpuBuffer()
+		{
+			if (bUsedWithCpu)
+			{
+				const int32 AlignedNumElements = AlignMemoryBuffer(NumElements);
+				if (AlignedNumElements != CpuBuffer.Num())
+				{
+					CpuBuffer.SetNumZeroed(AlignedNumElements);
+				}
+			}
+		}
+	};
+
+	struct FGameToRenderData
+	{
+		explicit FGameToRenderData(FNDIInstanceData_GameThread& InstanceData)
+			: NumElements(InstanceData.NumElements)
+			, GpuSyncMode(InstanceData.GpuSyncMode)
+		{
+			if (InstanceData.bNeedsGpuSync)
+			{
+				if ( FNiagaraUtilities::ShouldSyncCpuToGpu(GpuSyncMode) )
+				{
+					InitialCpuBufferData = InstanceData.CpuBuffer;
+				}
+				GpuPendingClear = InstanceData.GpuPendingClear;
+				InstanceData.GpuPendingClear.Reset();
+			}
+		}
+
+		explicit FGameToRenderData(TConstArrayView<uint32> DataToSend)
+		{
+			NumElements = DataToSend.Num();
+			InitialCpuBufferData = DataToSend;
+		}
+
+		int32				NumElements = 0;
+		ENiagaraGpuSyncMode	GpuSyncMode = ENiagaraGpuSyncMode::None;
+		TOptional<uint32>	GpuPendingClear;
+		TArray<uint32>		InitialCpuBufferData;
+	};
+
+	struct FNDIInstanceData_RenderThread
+	{
+		int32							NumElements = 0;
+		TRefCountPtr<FRDGPooledBuffer>	PooledBuffer;
+		FRDGBufferRef					RDGTransientBuffer = nullptr;
+		FRDGBufferUAVRef				RDGTransientBufferUAV = nullptr;
+
+		ENiagaraGpuSyncMode				GpuSyncMode = ENiagaraGpuSyncMode::None;
+		TArray<uint32>					InitialCpuBufferData;
+		TOptional<uint32>				GpuPendingClear;
+
+		~FNDIInstanceData_RenderThread()
+		{
+			ReleaseData();
+		}
+
+		void UpdateData(FGameToRenderData& GameToRenderData)
+		{
+			ReleaseData();
+
+			NumElements				= GameToRenderData.NumElements;
+			GpuSyncMode				= GameToRenderData.GpuSyncMode;
+			InitialCpuBufferData	= MoveTemp(GameToRenderData.InitialCpuBufferData);
+			GpuPendingClear			= GameToRenderData.GpuPendingClear;
+		}
+
+		void ReleaseData()
+		{
+			NumElements = 0;
+			PooledBuffer.SafeRelease();
+		}
+	};
+
+	struct FNDIProxy : public FNiagaraDataInterfaceProxy
+	{
+		FNDIProxy(UObject* Owner)
+			: WeakOwner(Owner)
+		{
+		}
+
+		virtual void ConsumePerInstanceDataFromGameThread(void* PerInstanceData, const FNiagaraSystemInstanceID& Instance) override { checkNoEntry(); }
+		virtual int32 PerInstanceDataPassedToRenderThreadSize() const override { return 0; }
+
+		virtual void PreStage(const FNDIGpuComputePreStageContext& Context)
+		{
+			const FNiagaraSystemInstanceID SystemInstanceID = Context.GetSystemInstanceID();
+			FNDIInstanceData_RenderThread& InstanceData = PerInstanceData_RenderThread.FindChecked(SystemInstanceID);
+			const FNiagaraSimStageData& SimStageData = Context.GetSimStageData();
+
+			if (SimStageData.bFirstStage)
+			{
+				FRDGBuilder& GraphBuilder = Context.GetGraphBuilder();
+
+				const int32 AlignedNumElements = AlignMemoryBuffer(InstanceData.NumElements);
+				FRDGBufferDesc BufferDesc = FRDGBufferDesc::CreateByteAddressDesc(AlignedNumElements << 2);
+				BufferDesc.Usage |= BUF_SourceCopy;
+				ResizeBufferIfNeeded(GraphBuilder, InstanceData.PooledBuffer, BufferDesc, TEXT("NiagaraMemoryBuffer"));
+
+				if (InstanceData.InitialCpuBufferData.Num() > 0)
+				{
+					if (ensure(InstanceData.InitialCpuBufferData.Num() == AlignedNumElements))
+					{
+						const void* InitialData = InstanceData.InitialCpuBufferData.GetData();
+						const uint64 InitialDataSize = InstanceData.InitialCpuBufferData.Num() * InstanceData.InitialCpuBufferData.GetTypeSize();
+						GraphBuilder.QueueBufferUpload(
+							GraphBuilder.RegisterExternalBuffer(InstanceData.PooledBuffer),
+							InitialData,
+							InitialDataSize,
+							[DataBuffer=MoveTemp(InstanceData.InitialCpuBufferData)](const void*) { }
+						);
+					}
+					else
+					{
+						InstanceData.InitialCpuBufferData.Empty();
+					}
+				}
+				else if (InstanceData.GpuPendingClear.IsSet())
+				{
+					AddClearUAVPass(Context.GetGraphBuilder(), GraphBuilder.CreateUAV(GraphBuilder.RegisterExternalBuffer(InstanceData.PooledBuffer)), InstanceData.GpuPendingClear.GetValue());
+				}
+				InstanceData.GpuPendingClear.Reset();
+			}
+		}
+
+		virtual void PostSimulate(const FNDIGpuComputePostSimulateContext& Context)
+		{
+			if (!Context.IsFinalPostSimulate())
+			{
+				return;
+			}
+
+			const FNiagaraSystemInstanceID SystemInstanceID = Context.GetSystemInstanceID();
+			FNDIInstanceData_RenderThread& InstanceData = PerInstanceData_RenderThread.FindChecked(SystemInstanceID);
+
+			// Should we send the buffer back over
+			if (FNiagaraUtilities::ShouldSyncGpuToCpu(InstanceData.GpuSyncMode) && InstanceData.NumElements > 0 && InstanceData.RDGTransientBuffer)
+			{
+				//-TODO: Currently we don't know if a buffer was written to, we can extract that from the used functions to reduce readback pressure
+				FNiagaraGpuReadbackManager* ReadbackManager = Context.GetComputeDispatchInterface().GetGpuReadbackManager();
+
+				const int32 AlignedNumElements = AlignMemoryBuffer(InstanceData.NumElements);
+				ReadbackManager->EnqueueReadback(
+					Context.GetGraphBuilder(),
+					InstanceData.RDGTransientBuffer,
+					0,
+					AlignedNumElements * sizeof(uint32),
+					[SystemInstanceID, WeakOwnerObject=WeakOwner, Proxy=this, ExpectedElements=AlignedNumElements](TConstArrayView<TPair<void*, uint32>> ReadbackData)
+					{
+						if (ReadbackData.Num() != 1 || (ExpectedElements << 2 != ReadbackData[0].Value))
+						{
+							// Readback error
+							return;
+						}
+
+						TArray<uint32> BufferData_RT;
+						BufferData_RT.AddUninitialized(ExpectedElements);
+						FMemory::Memcpy(BufferData_RT.GetData(), ReadbackData[0].Key, ExpectedElements * sizeof(uint32));
+
+						FNiagaraWorldManager::EnqueueGlobalDeferredCallback(
+							[SystemInstanceID, BufferData=MoveTemp(BufferData_RT), WeakOwnerObject, Proxy]() mutable
+							{
+								// FNiagaraDataInterfaceProxy do not outlive UNiagaraDataInterface so if our Object is valid so is the proxy
+								// Equally because we do not share instance IDs (monotonically increasing number) we won't ever stomp something that has 'gone away'
+								if (!WeakOwnerObject.IsValid())
+								{
+									return;
+								}
+								if (FNDIInstanceData_GameThread* InstanceData_GT = Proxy->PerInstanceData_GameThread.FindRef(SystemInstanceID))
+								{
+									FRWScopeLock ScopeLock(InstanceData_GT->CpuBufferGuard, SLT_Write);
+									if (InstanceData_GT->CpuBuffer.Num() == BufferData.Num())
+									{
+										Swap(InstanceData_GT->CpuBuffer, BufferData);
+									}
+								}
+							}
+						);
+					}
+				);
+			}
+
+			// Clear out transient pointers we cache for graph usage
+			InstanceData.RDGTransientBuffer = nullptr;
+			InstanceData.RDGTransientBufferUAV = nullptr;
+		}
+
+		TWeakObjectPtr<UObject> WeakOwner;
+		TMap<FNiagaraSystemInstanceID, FNDIInstanceData_RenderThread>	PerInstanceData_RenderThread;
+		TMap<FNiagaraSystemInstanceID, FNDIInstanceData_GameThread*>	PerInstanceData_GameThread;
+	};
+
+	//////////////////////////////////////////////////////////////////////////////////////////
+	// VM functions
+	void VMSetNumElements(FVectorVMExternalFunctionContext& Context)
+	{
+		VectorVM::FUserPtrHandler<FNDIInstanceData_GameThread> InstanceData(Context);
+		FNDIInputParam<int32>		InNumElements(Context);
+
+		const int32 MaxElements = (int32)FMath::Min<int64>(GetMaxBufferDimension(), TNumericLimits<int32>::Max());
+
+		FRWScopeLock ScopeLock(InstanceData->CpuBufferGuard, SLT_Write);
+		for (int32 i = 0; i < Context.GetNumInstances(); ++i)
+		{
+			const int32 NumElements = FMath::Clamp(InNumElements.GetAndAdvance(), 0, MaxElements);
+			InstanceData->bNeedsGpuSync |= InstanceData->NumElements != NumElements;
+			InstanceData->NumElements = NumElements;
+		}
+		InstanceData->UpdateCpuBuffer();
+	}
+
+	void VMGetNumElements(FVectorVMExternalFunctionContext& Context)
+	{
+		VectorVM::FUserPtrHandler<FNDIInstanceData_GameThread> InstanceData(Context);
+		FNDIOutputParam<int32>		OutNumElements(Context);
+
+		FRWScopeLock ScopeLock(InstanceData->CpuBufferGuard, SLT_ReadOnly);
+		for (int32 i = 0; i < Context.GetNumInstances(); ++i)
+		{
+			OutNumElements.SetAndAdvance(InstanceData->NumElements);
+		}
+	}
+
+	template<typename TType>
+	void VMLoadData(FVectorVMExternalFunctionContext& Context)
+	{
+		VectorVM::FUserPtrHandler<FNDIInstanceData_GameThread> InstanceData(Context);
+		FNDIInputParam<int32>		InElementOffset(Context);
+		FNDIOutputParam<bool>		OutSuccess(Context);
+		FNDIOutputParam<TType>		OutValue(Context);
+
+		FRWScopeLock ScopeLock(InstanceData->CpuBufferGuard, SLT_ReadOnly);
+		const int32 TypeSize = FMath::DivideAndRoundUp(sizeof(TType), sizeof(uint32));
+		const int32 ValidEnd = InstanceData->NumElements - TypeSize;
+		for (int32 i = 0; i < Context.GetNumInstances(); ++i)
+		{
+			const int32 Offset			= InElementOffset.GetAndAdvance();
+			const int32 ClampedOffset	= FMath::Clamp(Offset, 0, InstanceData->NumElements);
+			TType Value;
+			FMemory::Memcpy(&Value, &InstanceData->CpuBuffer[ClampedOffset], sizeof(TType));
+
+			OutSuccess.SetAndAdvance(Offset >= 0 && Offset <= ValidEnd);
+			OutValue.SetAndAdvance(Value);
+		}
+	}
+
+	template<typename TType>
+	void VMStoreData(FVectorVMExternalFunctionContext& Context)
+	{
+		VectorVM::FUserPtrHandler<FNDIInstanceData_GameThread> InstanceData(Context);
+		FNDIInputParam<bool>		InExecute(Context);
+		FNDIInputParam<int32>		InElementOffset(Context);
+		FNDIInputParam<TType>		InValue(Context);
+		FNDIOutputParam<bool>		OutSuccess(Context);
+
+		FRWScopeLock ScopeLock(InstanceData->CpuBufferGuard, SLT_Write);
+		InstanceData->bNeedsGpuSync = true;
+
+		const int32 TypeSize = FMath::DivideAndRoundUp(sizeof(TType), sizeof(uint32));
+		const int32 ValidEnd = InstanceData->NumElements - TypeSize;
+		for (int32 i = 0; i < Context.GetNumInstances(); ++i)
+		{
+			const bool bExecute			= InExecute.GetAndAdvance();
+			const int32 Offset			= InElementOffset.GetAndAdvance();
+			const int32 ClampedOffset	= FMath::Clamp(Offset, 0, InstanceData->NumElements);
+			const TType Value			= InValue.GetAndAdvance();
+			const bool bSuccess			= bExecute && Offset > 0 && Offset <= ValidEnd;
+			if (bExecute)
+			{
+				FMemory::Memcpy(&InstanceData->CpuBuffer[ClampedOffset], &Value, sizeof(TType));
+			}
+			OutSuccess.SetAndAdvance(bSuccess);
+		}
+	}
+
+	template<typename TType>
+	void VMClearBuffer(FVectorVMExternalFunctionContext& Context)
+	{
+		VectorVM::FUserPtrHandler<FNDIInstanceData_GameThread> InstanceData(Context);
+		FNDIInputParam<bool>	InExecute(Context);
+		FNDIInputParam<TType>	InClearValue(Context);
+
+		FRWScopeLock ScopeLock(InstanceData->CpuBufferGuard, SLT_Write);
+		InstanceData->bNeedsGpuSync = true;
+
+		bool bExecuteClear = false;
+		uint32 ClearValue = 0;
+		for (int32 i=0; i < Context.GetNumInstances(); ++i)
+		{
+			const bool bExecute = InExecute.GetAndAdvance();
+			const TType NewClearValue = InClearValue.GetAndAdvance();
+			bExecuteClear |= bExecute;
+			ClearValue = bExecute ? reinterpret_cast<const uint32&>(NewClearValue) : ClearValue;
+		}
+
+		if (bExecuteClear)
+		{
+			if (InstanceData->bUsedWithCpu)
+			{
+				for (uint32& Value : InstanceData->CpuBuffer)
+				{
+					Value = ClearValue;
+				}
+			}
+			InstanceData->GpuPendingClear = ClearValue;
+		}
+	}
+
+	struct FAtomicAdd
+	{
+		static int32 InterlockedOp(int32* Location, int32 Value) { return FPlatformAtomics::InterlockedAdd(Location, Value); }
+		static int32 RawOp(int32 PreviousValue, int32 Value) { return PreviousValue += Value; }
+	};
+
+	struct FAtomicAnd
+	{
+		static int32 InterlockedOp(int32* Location, int32 Value) { return FPlatformAtomics::InterlockedAnd(Location, Value); }
+		static int32 RawOp(int32 PreviousValue, int32 Value) { return PreviousValue &= Value; }
+	};
+
+	struct FAtomicMin
+	{
+		static int32 InterlockedOp(int32* Location, int32 Value)
+		{
+			int32 PreviousValue = *Location;
+			while (PreviousValue > Value)
+			{
+				PreviousValue = FPlatformAtomics::InterlockedCompareExchange(Location, Value, PreviousValue);
+			}
+			return PreviousValue;
+		}
+		static int32 RawOp(int32 PreviousValue, int32 Value) { return FMath::Min(PreviousValue, Value); }
+	};
+
+	struct FAtomicMax
+	{
+		static int32 InterlockedOp(int32* Location, int32 Value)
+		{
+			int32 PreviousValue = *Location;
+			while (PreviousValue < Value)
+			{
+				PreviousValue = FPlatformAtomics::InterlockedCompareExchange(Location, Value, PreviousValue);
+			}
+			return PreviousValue;
+		}
+		static int32 RawOp(int32 PreviousValue, int32 Value) { return FMath::Max(PreviousValue, Value); }
+	};
+
+	struct FAtomicOr
+	{
+		static int32 InterlockedOp(int32* Location, int32 Value) { return FPlatformAtomics::InterlockedOr(Location, Value); }
+		static int32 RawOp(int32 PreviousValue, int32 Value) { return PreviousValue |= Value; }
+	};
+
+	struct FAtomicXor
+	{
+		static int32 InterlockedOp(int32* Location, int32 Value) { return FPlatformAtomics::InterlockedXor(Location, Value); }
+		static int32 RawOp(int32 PreviousValue, int32 Value) { return PreviousValue ^= Value; }
+	};
+
+	template<typename TType, typename TOperation>
+	void VMAtomicOp(FVectorVMExternalFunctionContext& Context)
+	{
+		VectorVM::FUserPtrHandler<FNDIInstanceData_GameThread> InstanceData(Context);
+		FNDIInputParam<bool>	InExecute(Context);
+		FNDIInputParam<int32>	InOffset(Context);
+		FNDIInputParam<TType>	InValue(Context);
+
+		FNDIOutputParam<TType>	OutPreviousValue(Context);
+		FNDIOutputParam<TType>	OutCurrentValue(Context);
+
+		FRWScopeLock ScopeLock(InstanceData->CpuBufferGuard, SLT_ReadOnly);
+		InstanceData->bNeedsGpuSync = true;
+
+		for (int32 i = 0; i < Context.GetNumInstances(); ++i)
+		{
+			const bool bExecute = InExecute.GetAndAdvance();
+			const int32 Offset = InOffset.GetAndAdvance();
+			const TType Value = InValue.GetAndAdvance();
+			TType PreviousValue = TType{};
+			TType CurrentValue = TType{};
+			if (bExecute)
+			{
+				const int32 ClampedOffset = FMath::Clamp(Offset, 0, InstanceData->NumElements);
+				PreviousValue = TOperation::InterlockedOp(reinterpret_cast<int32*>(&InstanceData->CpuBuffer[ClampedOffset]), Value);
+				CurrentValue = TOperation::RawOp(PreviousValue, Value);
+			}
+		}
+	}
+
+	template<typename TType>
+	void VMAtomicCAE(FVectorVMExternalFunctionContext& Context)
+	{
+		VectorVM::FUserPtrHandler<FNDIInstanceData_GameThread> InstanceData(Context);
+		FNDIInputParam<bool>	InExecute(Context);
+		FNDIInputParam<int32>	InOffset(Context);
+		FNDIInputParam<TType>	InValue(Context);
+		FNDIInputParam<TType>	InComperand(Context);
+
+		FNDIOutputParam<TType>	OutPreviousValue(Context);
+		FNDIOutputParam<TType>	OutCurrentValue(Context);
+
+		FRWScopeLock ScopeLock(InstanceData->CpuBufferGuard, SLT_ReadOnly);
+		InstanceData->bNeedsGpuSync = true;
+
+		for (int32 i = 0; i < Context.GetNumInstances(); ++i)
+		{
+			const bool bExecute = InExecute.GetAndAdvance();
+			const int32 Offset = InOffset.GetAndAdvance();
+			const TType Value = InValue.GetAndAdvance();
+			const TType Comperand = InComperand.GetAndAdvance();
+			TType PreviousValue = TType{};
+			TType CurrentValue = TType{};
+			if (bExecute)
+			{
+				const int32 ClampedOffset = FMath::Clamp(Offset, 0, InstanceData->NumElements);
+
+				PreviousValue = FPlatformAtomics::InterlockedCompareExchange(reinterpret_cast<int32*>(&InstanceData->CpuBuffer[ClampedOffset]), Value, Comperand);
+				CurrentValue = PreviousValue == Comperand ? Value : PreviousValue;
+			}
+		}
+	}
+
+	//////////////////////////////////////////////////////////////////////////////////////////
+	// Supported load / store types
+	struct FLoadStoreTypes
+	{
+		FName					LoadFunctionName;
+		FName					StoreFunctionName;
+		FNiagaraTypeDefinition	TypeDef;
+		int32					NumTypeDefSigCopies = 0;
+
+		FVMExternalFunction		VMLoadFunction;
+		FVMExternalFunction		VMStoreFunction;
+	};
+
+	static TConstArrayView<FLoadStoreTypes> GetLoadStoreTypes()
+	{
+		static const FLoadStoreTypes LoadStoreTypes[] =
+		{
+			{"LoadFloat",	"StoreFloat",	FNiagaraTypeDefinition::GetFloatDef(),	1, FVMExternalFunction::CreateStatic(VMLoadData<float>),		FVMExternalFunction::CreateStatic(VMStoreData<float>)		},
+			{"LoadFloat2",	"StoreFloat2",	FNiagaraTypeDefinition::GetVec2Def(),	1, FVMExternalFunction::CreateStatic(VMLoadData<FVector2f>),	FVMExternalFunction::CreateStatic(VMStoreData<FVector2f>)	},
+			{"LoadFloat3",	"StoreFloat3",	FNiagaraTypeDefinition::GetVec3Def(),	1, FVMExternalFunction::CreateStatic(VMLoadData<FVector3f>),	FVMExternalFunction::CreateStatic(VMStoreData<FVector3f>)	},
+			{"LoadFloat4",	"StoreFloat4",	FNiagaraTypeDefinition::GetVec4Def(),	1, FVMExternalFunction::CreateStatic(VMLoadData<FVector4f>),	FVMExternalFunction::CreateStatic(VMStoreData<FVector4f>)	},
+			{"LoadInt",		"StoreInt",		FNiagaraTypeDefinition::GetIntDef(),	1, FVMExternalFunction::CreateStatic(VMLoadData<int32>),		FVMExternalFunction::CreateStatic(VMStoreData<int32>)		},
+			{"LoadInt2",	"StoreInt2",	FNiagaraTypeDefinition::GetIntDef(),	2, FVMExternalFunction::CreateStatic(VMLoadData<FIntVector2>),	FVMExternalFunction::CreateStatic(VMStoreData<FIntVector2>)	},
+			{"LoadInt3",	"StoreInt3",	FNiagaraTypeDefinition::GetIntDef(),	3, FVMExternalFunction::CreateStatic(VMLoadData<FIntVector3>),	FVMExternalFunction::CreateStatic(VMStoreData<FIntVector3>)	},
+			{"LoadInt4",	"StoreInt4",	FNiagaraTypeDefinition::GetIntDef(),	4, FVMExternalFunction::CreateStatic(VMLoadData<FIntVector4>),	FVMExternalFunction::CreateStatic(VMStoreData<FIntVector4>)	},
+		};
+
+		return MakeArrayView(LoadStoreTypes);
+	};
+
+	//////////////////////////////////////////////////////////////////////////////////////////
+	// Supported interlocked ops types
+	const FName NAME_AtomicAdd("AtomicAdd");
+	const FName NAME_AtomicAnd("AtomicAnd");
+	const FName NAME_AtomicCompareAndExchange("AtomicCompareAndExchange");
+	const FName NAME_AtomicMax("AtomicMax");
+	const FName NAME_AtomicMin("AtomicMin");
+	const FName NAME_AtomicOr("AtomicOr");
+	const FName NAME_AtomicXor("AtomicXor");
+
+	static TConstArrayView<TPair<FName, FVMExternalFunction>> GetAtomicOps()
+	{
+		static const TArray<TPair<FName, FVMExternalFunction>> AtomicOps =
+		{
+			{NAME_AtomicAdd,				FVMExternalFunction::CreateStatic(VMAtomicOp<int32, FAtomicAdd>)},
+			{NAME_AtomicAnd,				FVMExternalFunction::CreateStatic(VMAtomicOp<int32, FAtomicAnd>)},
+			{NAME_AtomicCompareAndExchange,	FVMExternalFunction::CreateStatic(VMAtomicCAE<int32>)},
+			{NAME_AtomicMax,				FVMExternalFunction::CreateStatic(VMAtomicOp<int32, FAtomicMax>)},
+			{NAME_AtomicMin,				FVMExternalFunction::CreateStatic(VMAtomicOp<int32, FAtomicMin>)},
+			{NAME_AtomicOr,					FVMExternalFunction::CreateStatic(VMAtomicOp<int32, FAtomicOr>)},
+			{NAME_AtomicXor,				FVMExternalFunction::CreateStatic(VMAtomicOp<int32, FAtomicXor>)},
+			{NAME_ClearBufferInt,			FVMExternalFunction::CreateStatic(VMClearBuffer<int32>)},
+			{NAME_ClearBufferFloat,			FVMExternalFunction::CreateStatic(VMClearBuffer<float>)},
+		};
+		return AtomicOps;
+	}
+
+	//////////////////////////////////////////////////////////////////////////////////////////
+	// Cache Helper
+	int32 FindOrAddCacheData(UNDIMemoryBufferSimCacheData* CacheData, TConstArrayView<uint32> CurrentData)
+	{
+		if (CurrentData.Num() == 0)
+		{
+			return INDEX_NONE;
+		}
+
+		for (const FNDIMemoryBufferSimCacheDataFrame& ExistingFrame : CacheData->FrameData)
+		{
+			if (ExistingFrame.CpuBufferSize == CurrentData.Num())
+			{
+				if (FMemory::Memcmp(CacheData->BufferData.GetData() + ExistingFrame.CpuDataOffset, CurrentData.GetData(), CurrentData.Num() * sizeof(uint32)) == 0)
+				{
+					return ExistingFrame.CpuDataOffset;
+				}
+			}
+			if (ExistingFrame.GpuBufferSize == CurrentData.Num())
+			{
+				if (FMemory::Memcmp(CacheData->BufferData.GetData() + ExistingFrame.GpuDataOffset, CurrentData.GetData(), CurrentData.Num() * sizeof(uint32)) == 0)
+				{
+					return ExistingFrame.GpuDataOffset;
+				}
+			}
+		}
+		const int32 NewOffset = CacheData->BufferData.AddUninitialized(CurrentData.Num());
+		FMemory::Memcpy(CacheData->BufferData.GetData() + NewOffset, CurrentData.GetData(), CurrentData.Num() * sizeof(uint32));
+		return NewOffset;
+	}
+}
+
+////////////////////////////////////////////////////////////////////////////////////////////////////////////////////////////////////
+
+UNiagaraDataInterfaceMemoryBuffer::UNiagaraDataInterfaceMemoryBuffer(FObjectInitializer const& ObjectInitializer)
+	: Super(ObjectInitializer)
+{
+	using namespace NDIMemoryBufferLocal;
+
+	Proxy.Reset(new FNDIProxy(this));
+}
+
+void UNiagaraDataInterfaceMemoryBuffer::PostInitProperties()
+{
+	Super::PostInitProperties();
+
+	if (HasAnyFlags(RF_ClassDefaultObject))
+	{
+		ENiagaraTypeRegistryFlags Flags = ENiagaraTypeRegistryFlags::AllowAnyVariable | ENiagaraTypeRegistryFlags::AllowParameter;
+		FNiagaraTypeRegistry::Register(FNiagaraTypeDefinition(GetClass()), Flags);
+	}
+}
+
+#if WITH_EDITORONLY_DATA
+void UNiagaraDataInterfaceMemoryBuffer::GetFunctionsInternal(TArray<FNiagaraFunctionSignature>& OutFunctions) const
+{
+	using namespace NDIMemoryBufferLocal;
+
+	FNiagaraFunctionSignature ImmutableSignature;
+	ImmutableSignature.bMemberFunction = true;
+	ImmutableSignature.bRequiresContext = false;
+	ImmutableSignature.Inputs.Add(FNiagaraVariable(FNiagaraTypeDefinition(GetClass()), TEXT("Byte Buffer")));
+
+	FNiagaraFunctionSignature MutableSignature = ImmutableSignature;
+	MutableSignature.bRequiresExecPin = true;
+
+	////////////////////////////////////////////////////////////////////////////
+	// Helper Functions
+	{
+		FNiagaraFunctionSignature& Signature = OutFunctions.Add_GetRef(MutableSignature);
+		Signature.Name = NAME_SetNumElements;
+		Signature.ModuleUsageBitmask = ENiagaraScriptUsageMask::System | ENiagaraScriptUsageMask::Emitter;
+		Signature.Inputs.Emplace(FNiagaraTypeDefinition::GetIntDef(), TEXT("NumElements"));
+	}
+	{
+		FNiagaraFunctionSignature& Signature = OutFunctions.Add_GetRef(ImmutableSignature);
+		Signature.Name = NAME_GetNumElements;
+		Signature.Outputs.Emplace(FNiagaraTypeDefinition::GetIntDef(), TEXT("NumElements"));
+	}
+
+	////////////////////////////////////////////////////////////////////////////
+	// Load Functions
+	{
+		FNiagaraFunctionSignature LoadSignature = ImmutableSignature;
+		LoadSignature.bReadFunction = true;
+		LoadSignature.Inputs.Emplace(FNiagaraTypeDefinition::GetIntDef(), TEXT("ElementOffset"));
+		LoadSignature.Outputs.Emplace(FNiagaraTypeDefinition::GetBoolDef(), TEXT("bValid"));
+
+		for (const FLoadStoreTypes& SupportedType : GetLoadStoreTypes())
+		{
+			FNiagaraFunctionSignature& Signature = OutFunctions.Add_GetRef(LoadSignature);
+			Signature.Name = SupportedType.LoadFunctionName;
+			for ( int32 i=0; i < SupportedType.NumTypeDefSigCopies; ++i )
+			{
+				Signature.Outputs.Emplace(
+					SupportedType.TypeDef,
+					SupportedType.NumTypeDefSigCopies == 1 ? TEXT("Value") : *FString::Printf(TEXT("Value%d"), i)
+				);
+			}
+		}
+	}
+
+	////////////////////////////////////////////////////////////////////////////
+	// Store Functions
+	{
+		FNiagaraFunctionSignature StoreSignature = MutableSignature;
+		StoreSignature.bWriteFunction = true;
+		StoreSignature.Inputs.Emplace_GetRef(FNiagaraTypeDefinition::GetBoolDef(), TEXT("Execute")).SetValue(true);
+		StoreSignature.Inputs.Emplace(FNiagaraTypeDefinition::GetIntDef(), TEXT("ElementOffset"));
+		StoreSignature.Outputs.Emplace(FNiagaraTypeDefinition::GetBoolDef(), TEXT("bSuccess"));
+
+		for (const FLoadStoreTypes& SupportedType : GetLoadStoreTypes())
+		{
+			FNiagaraFunctionSignature& Signature = OutFunctions.Add_GetRef(StoreSignature);
+			Signature.Name = SupportedType.StoreFunctionName;
+			for (int32 i = 0; i < SupportedType.NumTypeDefSigCopies; ++i)
+			{
+				Signature.Inputs.Emplace(
+					SupportedType.TypeDef,
+					SupportedType.NumTypeDefSigCopies == 1 ? TEXT("Value") : *FString::Printf(TEXT("Value%d"), i)
+				);
+			}
+		}
+	}
+
+	////////////////////////////////////////////////////////////////////////////
+	// Atomic Operations
+	{
+		FNiagaraFunctionSignature AtomicSignature = MutableSignature;
+		AtomicSignature.bWriteFunction = true;
+		AtomicSignature.Inputs.Emplace_GetRef(FNiagaraTypeDefinition::GetBoolDef(), TEXT("Execute")).SetValue(true);
+		AtomicSignature.Inputs.Emplace(FNiagaraTypeDefinition::GetIntDef(), TEXT("ElementOffset"));
+		AtomicSignature.Inputs.Emplace(FNiagaraTypeDefinition::GetIntDef(), TEXT("Value"));
+		AtomicSignature.Outputs.Emplace(FNiagaraTypeDefinition::GetIntDef(), TEXT("PreviousValue"));
+		AtomicSignature.Outputs.Emplace(FNiagaraTypeDefinition::GetIntDef(), TEXT("CurrentValue"));
+
+		OutFunctions.Add_GetRef(AtomicSignature).Name = NAME_AtomicAdd;
+		OutFunctions.Add_GetRef(AtomicSignature).Name = NAME_AtomicAnd;
+		OutFunctions.Add_GetRef(AtomicSignature).Name = NAME_AtomicMax;
+		OutFunctions.Add_GetRef(AtomicSignature).Name = NAME_AtomicMin;
+		OutFunctions.Add_GetRef(AtomicSignature).Name = NAME_AtomicOr;
+		OutFunctions.Add_GetRef(AtomicSignature).Name = NAME_AtomicXor;
+
+		FNiagaraFunctionSignature& CASSignature = OutFunctions.Add_GetRef(AtomicSignature);
+		CASSignature.Name = NAME_AtomicCompareAndExchange;
+		CASSignature.Inputs.Emplace(FNiagaraTypeDefinition::GetIntDef(), TEXT("Comperand"));
+	}
+
+	////////////////////////////////////////////////////////////////////////////
+	// Buffer Clear Operations
+	{
+		FNiagaraFunctionSignature ClearSignature = MutableSignature;
+		ClearSignature.ModuleUsageBitmask = ENiagaraScriptUsageMask::System | ENiagaraScriptUsageMask::Emitter;
+		ClearSignature.Inputs.Emplace_GetRef(FNiagaraTypeDefinition::GetBoolDef(), TEXT("Execute")).SetValue(true);
+		{
+			FNiagaraFunctionSignature& Signature = OutFunctions.Add_GetRef(ClearSignature);
+			Signature.Name = NAME_ClearBufferInt;
+			Signature.Inputs.Emplace(FNiagaraTypeDefinition::GetIntDef(), TEXT("ClearValue"));
+		}
+		{
+			FNiagaraFunctionSignature& Signature = OutFunctions.Add_GetRef(ClearSignature);
+			Signature.Name = NAME_ClearBufferFloat;
+			Signature.Inputs.Emplace(FNiagaraTypeDefinition::GetFloatDef(), TEXT("ClearValue"));
+		}
+	}
+}
+#endif //WITH_EDITORONLY_DATA
+
+void UNiagaraDataInterfaceMemoryBuffer::GetVMExternalFunction(const FVMExternalFunctionBindingInfo& BindingInfo, void* InstanceData, FVMExternalFunction& OutFunc)
+{
+	using namespace NDIMemoryBufferLocal;
+
+	static const TPair<FName, FVMExternalFunction> StaticBindings[] =
+	{
+		{NAME_SetNumElements, FVMExternalFunction::CreateStatic(VMSetNumElements)},
+		{NAME_GetNumElements, FVMExternalFunction::CreateStatic(VMGetNumElements)},
+	};
+
+	for (const auto& StaticBinding : StaticBindings)
+	{
+		if (StaticBinding.Key == BindingInfo.Name)
+		{
+			OutFunc = StaticBinding.Value;
+			return;
+		}
+	}
+
+	for (const FLoadStoreTypes& SupportedType : GetLoadStoreTypes())
+	{
+		if (BindingInfo.Name == SupportedType.LoadFunctionName)
+		{
+			OutFunc = SupportedType.VMLoadFunction;
+			return;
+		}
+		if (BindingInfo.Name == SupportedType.StoreFunctionName)
+		{
+			OutFunc = SupportedType.VMStoreFunction;
+			return;
+		}
+	}
+
+	for (const TPair<FName, FVMExternalFunction>& AtomicOp : GetAtomicOps())
+	{
+		if (BindingInfo.Name == AtomicOp.Key)
+		{
+			OutFunc = AtomicOp.Value;
+			return;
+		}
+	}
+}
+
+int32 UNiagaraDataInterfaceMemoryBuffer::PerInstanceDataSize() const
+{
+	using namespace NDIMemoryBufferLocal;
+	return sizeof(FNDIInstanceData_GameThread);
+}
+
+bool UNiagaraDataInterfaceMemoryBuffer::InitPerInstanceData(void* PerInstanceData, FNiagaraSystemInstance* SystemInstance)
+{
+	using namespace NDIMemoryBufferLocal;
+
+	FNDIInstanceData_GameThread* InstanceData	= new(PerInstanceData) FNDIInstanceData_GameThread();
+	InstanceData->NumElements					= DefaultNumElements;
+	InstanceData->GpuSyncMode					= GpuSyncMode;
+
+	if (IsUsedWithCPUScript())
+	{
+		FNiagaraDataInterfaceUtilities::ForEachVMFunction(
+			this,
+			SystemInstance,
+			[InstanceData](const UNiagaraScript*, const FVMExternalFunctionBindingInfo& BindingInfo)
+			{
+				for (const FLoadStoreTypes& LoadStoreType : GetLoadStoreTypes())
+				{
+					if (BindingInfo.Name == LoadStoreType.LoadFunctionName || BindingInfo.Name == LoadStoreType.StoreFunctionName)
+					{
+						InstanceData->bUsedWithCpu = true;
+						return false;
+					}
+				}
+				for (const TPair<FName, FVMExternalFunction>& AtomicOp : GetAtomicOps())
+				{
+					if (AtomicOp.Key == BindingInfo.Name)
+					{
+						InstanceData->bUsedWithCpu = true;
+						return false;
+					}
+				}
+				return true;
+			}
+		);
+	}
+
+	InstanceData->UpdateCpuBuffer();
+
+	if (IsUsedWithGPUScript())
+	{
+		InstanceData->bUsedWithGpu = true;
+
+		FNDIProxy* Proxy_GT = GetProxyAs<FNDIProxy>();
+		Proxy_GT->PerInstanceData_GameThread.Add(SystemInstance->GetId(), InstanceData);
+	}
+
+	return true;
+}
+
+void UNiagaraDataInterfaceMemoryBuffer::DestroyPerInstanceData(void* PerInstanceData, FNiagaraSystemInstance* SystemInstance)
+{
+	using namespace NDIMemoryBufferLocal;
+
+	FNDIInstanceData_GameThread* InstanceData = reinterpret_cast<FNDIInstanceData_GameThread*>(PerInstanceData);
+	if (InstanceData->bUsedWithGpu)
+	{
+		FNDIProxy* Proxy_GT = GetProxyAs<FNDIProxy>();
+		Proxy_GT->PerInstanceData_GameThread.Remove(SystemInstance->GetId());
+
+		ENQUEUE_RENDER_COMMAND(FNDIMemoryBuffer_RemoveProxy)
+		(
+			[RT_Proxy=Proxy_GT, InstanceID=SystemInstance->GetId()](FRHICommandListImmediate& CmdList)
+			{
+				RT_Proxy->PerInstanceData_RenderThread.Remove(InstanceID);
+			}
+		);
+	}
+
+	InstanceData->~FNDIInstanceData_GameThread();
+}
+
+bool UNiagaraDataInterfaceMemoryBuffer::PerInstanceTickPostSimulate(void* PerInstanceData, FNiagaraSystemInstance* SystemInstance, float DeltaSeconds)
+{
+	using namespace NDIMemoryBufferLocal;
+
+	FNDIInstanceData_GameThread* InstanceData = reinterpret_cast<FNDIInstanceData_GameThread*>(PerInstanceData);
+
+	if (InstanceData->bUsedWithGpu && InstanceData->bNeedsGpuSync)
+	{
+		ENQUEUE_RENDER_COMMAND(FNDIMemoryBuffer_UpdateData)
+		(
+			[RT_Proxy=GetProxyAs<FNDIProxy>(), RT_SystemInstanceID=SystemInstance->GetId(), RT_GameToRenderData=FGameToRenderData(*InstanceData)](FRHICommandListImmediate& RHICmdList) mutable
+			{
+				FNDIInstanceData_RenderThread& InstanceData_RT = RT_Proxy->PerInstanceData_RenderThread.FindOrAdd(RT_SystemInstanceID);
+				InstanceData_RT.UpdateData(RT_GameToRenderData);
+			}
+		);
+		InstanceData->bNeedsGpuSync = false;
+	}
+
+	return false;
+}
+
+bool UNiagaraDataInterfaceMemoryBuffer::Equals(const UNiagaraDataInterface* Other) const
+{
+	const UNiagaraDataInterfaceMemoryBuffer* OtherTyped = CastChecked<const UNiagaraDataInterfaceMemoryBuffer>(Other);
+	return
+		Super::Equals(Other) &&
+		OtherTyped->DefaultNumElements == DefaultNumElements &&
+		OtherTyped->GpuSyncMode == GpuSyncMode;
+}
+
+bool UNiagaraDataInterfaceMemoryBuffer::CopyToInternal(UNiagaraDataInterface* Destination) const
+{
+	if (!Super::CopyToInternal(Destination))
+	{
+		return false;
+	}
+
+	UNiagaraDataInterfaceMemoryBuffer* OtherTyped = CastChecked<UNiagaraDataInterfaceMemoryBuffer>(Destination);
+	OtherTyped->DefaultNumElements			= DefaultNumElements;
+	OtherTyped->GpuSyncMode					= GpuSyncMode;
+
+	return true;
+}
+
+#if WITH_EDITORONLY_DATA
+bool UNiagaraDataInterfaceMemoryBuffer::AppendCompileHash(FNiagaraCompileHashVisitor* InVisitor) const
+{
+	using namespace NDIMemoryBufferLocal;
+
+	bool bSuccess = Super::AppendCompileHash(InVisitor);
+	bSuccess &= InVisitor->UpdateShaderFile(TemplateShaderFilePath);
+	bSuccess &= InVisitor->UpdateShaderParameters<FShaderParameters>();
+	return bSuccess;
+}
+
+void UNiagaraDataInterfaceMemoryBuffer::GetParameterDefinitionHLSL(const FNiagaraDataInterfaceGPUParamInfo& ParamInfo, FString& OutHLSL)
+{
+	using namespace NDIMemoryBufferLocal;
+
+	const TMap<FString, FStringFormatArg> TemplateArgs =
+	{
+		{TEXT("ParameterName"),	ParamInfo.DataInterfaceHLSLSymbol},
+	};
+	AppendTemplateHLSL(OutHLSL, TemplateShaderFilePath, TemplateArgs);
+}
+
+bool UNiagaraDataInterfaceMemoryBuffer::GetFunctionHLSL(const FNiagaraDataInterfaceGPUParamInfo& ParamInfo, const FNiagaraDataInterfaceGeneratedFunction& FunctionInfo, int FunctionInstanceIndex, FString& OutHLSL)
+{
+	using namespace NDIMemoryBufferLocal;
+
+	static const TSet<FName> ValidGpuFunctions =
+	{
+		NAME_GetNumElements,
+		NAME_AtomicAdd,
+		NAME_AtomicAnd,
+		NAME_AtomicCompareAndExchange,
+		NAME_AtomicMax,
+		NAME_AtomicMin,
+		NAME_AtomicOr,
+		NAME_AtomicXor,
+	};
+
+	return
+		ValidGpuFunctions.Contains(FunctionInfo.DefinitionName) ||
+		GetLoadStoreTypes().ContainsByPredicate(
+			[&](const FLoadStoreTypes& SupportedType)
+			{
+				return
+					SupportedType.LoadFunctionName == FunctionInfo.DefinitionName ||
+					SupportedType.StoreFunctionName == FunctionInfo.DefinitionName;
+			}
+		);
+}
+#endif
+
+void UNiagaraDataInterfaceMemoryBuffer::BuildShaderParameters(FNiagaraShaderParametersBuilder& ShaderParametersBuilder) const
+{
+	using namespace NDIMemoryBufferLocal;
+
+	ShaderParametersBuilder.AddNestedStruct<FShaderParameters>();
+}
+
+void UNiagaraDataInterfaceMemoryBuffer::SetShaderParameters(const FNiagaraDataInterfaceSetShaderParametersContext& Context) const
+{
+	using namespace NDIMemoryBufferLocal;
+
+	FRDGBuilder& GraphBuilder = Context.GetGraphBuilder();
+	FNDIProxy& DataInterfaceProxy = Context.GetProxy<FNDIProxy>();
+	FNDIInstanceData_RenderThread& InstanceData = DataInterfaceProxy.PerInstanceData_RenderThread.FindChecked(Context.GetSystemInstanceID());
+
+	if (InstanceData.RDGTransientBuffer == nullptr)
+	{
+		InstanceData.RDGTransientBuffer = GraphBuilder.RegisterExternalBuffer(InstanceData.PooledBuffer);
+		InstanceData.RDGTransientBufferUAV = GraphBuilder.CreateUAV(InstanceData.RDGTransientBuffer);
+	}
+
+	FShaderParameters* ShaderParameters = Context.GetParameterNestedStruct<FShaderParameters>();
+	ShaderParameters->NumElements	= InstanceData.NumElements;
+	ShaderParameters->Buffer		= InstanceData.RDGTransientBufferUAV;
+}
+
+UObject* UNiagaraDataInterfaceMemoryBuffer::SimCacheBeginWrite(UObject* SimCache, FNiagaraSystemInstance* NiagaraSystemInstance, const void* OptionalPerInstanceData, FNiagaraSimCacheFeedbackContext& FeedbackContext) const
+{
+	UNDIMemoryBufferSimCacheData* CacheData = NewObject<UNDIMemoryBufferSimCacheData>(SimCache);
+	return CacheData;
+}
+
+bool UNiagaraDataInterfaceMemoryBuffer::SimCacheWriteFrame(UObject* StorageObject, int FrameIndex, FNiagaraSystemInstance* SystemInstance, const void* OptionalPerInstanceData, FNiagaraSimCacheFeedbackContext& FeedbackContext) const
+{
+	using namespace NDIMemoryBufferLocal;
+
+	check(OptionalPerInstanceData && StorageObject);
+	UNDIMemoryBufferSimCacheData* CacheData = CastChecked<UNDIMemoryBufferSimCacheData>(StorageObject);
+	if (CacheData->FrameData.Num() <= FrameIndex)
+	{
+		CacheData->FrameData.AddDefaulted(FrameIndex + 1 - CacheData->FrameData.Num());
+	}
+
+	const FNDIInstanceData_GameThread* InstanceData_GT = static_cast<const FNDIInstanceData_GameThread*>(OptionalPerInstanceData);
+	if (InstanceData_GT->CpuBuffer.Num() > 0)
+	{
+		CacheData->FrameData[FrameIndex].CpuDataOffset = FindOrAddCacheData(CacheData, InstanceData_GT->CpuBuffer);
+		CacheData->FrameData[FrameIndex].CpuBufferSize = InstanceData_GT->CpuBuffer.Num();
+	}
+
+	if (InstanceData_GT->bUsedWithGpu)
+	{
+		FNiagaraGpuComputeDispatchInterface* ComputeInterface = FNiagaraGpuComputeDispatchInterface::Get(SystemInstance->GetWorld());
+
+		ENQUEUE_RENDER_COMMAND(NDIRenderTargetVolumeUpdate)(
+			[RT_Proxy=GetProxyAs<FNDIProxy>(), InstanceID=SystemInstance->GetId(), CacheData, FrameIndex, ComputeInterface](FRHICommandListImmediate& RHICmdList)
+			{
+				const FNDIInstanceData_RenderThread* InstanceData_RT = RT_Proxy->PerInstanceData_RenderThread.Find(InstanceID);
+				if (!InstanceData_RT || !InstanceData_RT->NumElements)
+				{
+					return;
+				}
+
+				const int32 BufferSize = InstanceData_RT->NumElements * sizeof(uint32);
+
+				FNiagaraGpuReadbackManager* ReadbackManager = ComputeInterface->GetGpuReadbackManager();
+				FRDGBuilder GraphBuilder(RHICmdList);
+				ReadbackManager->EnqueueReadback(
+					GraphBuilder,
+					GraphBuilder.RegisterExternalBuffer(InstanceData_RT->PooledBuffer), 0, BufferSize,
+					[CacheData, FrameIndex](TConstArrayView<TPair<void*, uint32>> ReadbackData)
+					{
+						const int32 BufferNumElements = ReadbackData[0].Value / sizeof(uint32);
+						const uint32* BufferData = static_cast<const uint32*>(ReadbackData[0].Key);
+						CacheData->FrameData[FrameIndex].GpuDataOffset = FindOrAddCacheData(CacheData, MakeArrayView<const uint32>(BufferData, BufferNumElements));
+						CacheData->FrameData[FrameIndex].GpuBufferSize = BufferNumElements;
+					}
+				);
+				GraphBuilder.Execute();
+				ReadbackManager->WaitCompletion(RHICmdList);
+			}
+		);
+
+		FlushRenderingCommands();
+	}
+
+	return true;
+}
+
+bool UNiagaraDataInterfaceMemoryBuffer::SimCacheReadFrame(UObject* StorageObject, int FrameA, int FrameB, float Interp, FNiagaraSystemInstance* SystemInstance, void* OptionalPerInstanceData)
+{
+	using namespace NDIMemoryBufferLocal;
+
+	check(OptionalPerInstanceData && StorageObject);
+	UNDIMemoryBufferSimCacheData* CacheData = CastChecked<UNDIMemoryBufferSimCacheData>(StorageObject);
+
+	const int FrameIndex = FrameA;
+	if (!CacheData->FrameData.IsValidIndex(FrameIndex))
+	{
+		return false;
+	}
+
+	const FNDIMemoryBufferSimCacheDataFrame& FrameData = CacheData->FrameData[FrameIndex];
+
+	FNDIInstanceData_GameThread* InstanceData_GT = static_cast<FNDIInstanceData_GameThread*>(OptionalPerInstanceData);
+	if (FrameData.CpuBufferSize > 0 && InstanceData_GT->bUsedWithCpu)
+	{		
+		InstanceData_GT->NumElements = FrameData.CpuBufferSize;
+		InstanceData_GT->UpdateCpuBuffer();
+		FMemory::Memcpy(InstanceData_GT->CpuBuffer.GetData(), &CacheData->BufferData[FrameData.CpuDataOffset], FrameData.CpuBufferSize * sizeof(uint32));
+	}
+
+	if (FrameData.GpuBufferSize > 0 && InstanceData_GT->bUsedWithGpu)
+	{
+		TConstArrayView<uint32> GpuData(&CacheData->BufferData[FrameData.GpuDataOffset], FrameData.GpuBufferSize);
+		ENQUEUE_RENDER_COMMAND(FNDIMemoryBuffer_UpdateData)
+		(
+			[RT_Proxy=GetProxyAs<FNDIProxy>(), RT_SystemInstanceID=SystemInstance->GetId(), RT_GameToRenderData=FGameToRenderData(GpuData)](FRHICommandListImmediate& RHICmdList) mutable
+			{
+				FNDIInstanceData_RenderThread& InstanceData_RT = RT_Proxy->PerInstanceData_RenderThread.FindOrAdd(RT_SystemInstanceID);
+				InstanceData_RT.UpdateData(RT_GameToRenderData);
+			}
+		);
+	}
+
+	return true;
+}
+
+#undef LOCTEXT_NAMESPACE
diff --git a/Engine/Plugins/FX/Niagara/Source/NiagaraEditor/Private/Customizations/SimCache/NiagaraMemoryBufferSimCacheVisualizer.cpp b/Engine/Plugins/FX/Niagara/Source/NiagaraEditor/Private/Customizations/SimCache/NiagaraMemoryBufferSimCacheVisualizer.cpp
new file mode 100644
--- /dev/null
+++ b/Engine/Plugins/FX/Niagara/Source/NiagaraEditor/Private/Customizations/SimCache/NiagaraMemoryBufferSimCacheVisualizer.cpp
@@ -0,0 +1,434 @@
+// Copyright Epic Games, Inc. All Rights Reserved.
+
+#include "NiagaraMemoryBufferSimCacheVisualizer.h"
+
+#include "DataInterface/NiagaraDataInterfaceMemoryBuffer.h"
+
+#include "UObject/StrongObjectPtr.h"
+#include "ViewModels/NiagaraSimCacheViewModel.h"
+#include "Widgets/Input/SSpinBox.h"
+#include "Widgets/Text/STextBlock.h"
+#include "DetailLayoutBuilder.h"
+#include "SEnumCombo.h"
+
+#define LOCTEXT_NAMESPACE "NiagaraMemoryBufferSimCacheVisualizer"
+
+////////////////////////////////////////////////////////////////////////////////////////////////////////////////////
+
+namespace NDIMemoryBufferSimCacheVisualizerLocal
+{
+static const FName NAME_ElementOffset = FName("_ElementOffset");
+
+template<typename TType>
+TConstArrayView<TType> GetCurrentData(const TSharedPtr<FNiagaraSimCacheViewModel>& ViewModel, UNDIMemoryBufferSimCacheData* CacheData, ENiagaraSimTarget SimTarget)
+{
+	const int32 FrameIndex = ViewModel->GetFrameIndex();
+	if (CacheData->FrameData.IsValidIndex(FrameIndex))
+	{
+		if (SimTarget == ENiagaraSimTarget::CPUSim)
+		{
+			return MakeArrayView(
+				reinterpret_cast<const TType*>(CacheData->BufferData.GetData() + CacheData->FrameData[FrameIndex].CpuDataOffset),
+				CacheData->FrameData[FrameIndex].CpuBufferSize
+			);
+		}
+		else
+		{
+			return MakeArrayView(
+				reinterpret_cast<const TType*>(CacheData->BufferData.GetData() + CacheData->FrameData[FrameIndex].GpuDataOffset),
+				CacheData->FrameData[FrameIndex].GpuBufferSize
+			);
+		}
+	}
+	return MakeArrayView<const TType>(nullptr, 0);
+}
+
+}// namespace NDIMemoryBufferSimCacheVisualizerLocal
+
+////////////////////////////////////////////////////////////////////////////////////////////////////////////////////
+
+namespace NDIMemoryBufferSimCacheVisualizer
+{
+class SVisualizerRowWidget : public SMultiColumnTableRow<TSharedPtr<int32>>
+{
+public:
+	SLATE_BEGIN_ARGS(SVisualizerRowWidget) {}
+		SLATE_ARGUMENT(TSharedPtr<FNiagaraSimCacheViewModel>,	ViewModel)
+		SLATE_ARGUMENT(UNDIMemoryBufferSimCacheData*,			CacheData)
+		SLATE_ARGUMENT(TSharedPtr<int32>,						RowIndexPtr)
+		SLATE_ATTRIBUTE(ENiagaraSimTarget,						DisplaySimTarget)
+		SLATE_ATTRIBUTE(ENDIMemoryBufferViewType,				DisplayAsType)
+	SLATE_END_ARGS()
+
+	void Construct(const FArguments& InArgs, const TSharedRef<STableViewBase>& InOwnerTableView)
+	{
+		ViewModel			= InArgs._ViewModel;
+		CacheData.Reset(InArgs._CacheData);
+		RowIndexPtr			= InArgs._RowIndexPtr;
+		DisplaySimTarget	= InArgs._DisplaySimTarget;
+		DisplayAsType		= InArgs._DisplayAsType;
+
+		SMultiColumnTableRow::Construct(
+			FSuperRowType::FArguments()
+			.Style(FAppStyle::Get(), "DataTableEditor.CellListViewRow"),
+			InOwnerTableView
+		);
+	}
+
+	virtual TSharedRef<SWidget> GenerateWidgetForColumn(const FName& InColumnName) override
+	{
+		using namespace NDIMemoryBufferSimCacheVisualizerLocal;
+		if (InColumnName == NAME_ElementOffset)
+		{
+			return SNew(STextBlock)
+				.Text(FText::Format(LOCTEXT("ElementFormat", "Element {0}"), FText::AsNumber(*RowIndexPtr)));
+		}
+
+		const int32 Offset = *RowIndexPtr + InColumnName.GetNumber();
+
+		FText ItemText;
+		switch (DisplayAsType.Get(ENDIMemoryBufferViewType::Integer))
+		{
+			case ENDIMemoryBufferViewType::Float:
+			{
+				TConstArrayView<float> BufferData = GetCurrentData<float>(ViewModel, CacheData.Get(), DisplaySimTarget.Get(ENiagaraSimTarget::CPUSim));
+				if (BufferData.IsValidIndex(Offset))
+				{
+					const float Value = BufferData[Offset];
+					if (FMath::IsFinite(Value))
+					{
+						FNumberFormattingOptions FormatOptions;
+						FormatOptions.SetMinimumFractionalDigits(2);
+
+						ItemText = FText::AsNumber(Value, &FormatOptions);
+					}
+					else
+					{
+						ItemText = LOCTEXT("Invalid Float", "Invalid Float");
+					}
+				}
+				break;
+			}
+			case ENDIMemoryBufferViewType::Integer:
+			{
+				TConstArrayView<int32> BufferData = GetCurrentData<int32>(ViewModel, CacheData.Get(), DisplaySimTarget.Get(ENiagaraSimTarget::CPUSim));
+				if (BufferData.IsValidIndex(Offset))
+				{
+					ItemText = FText::AsNumber(BufferData[Offset]);
+				}
+				break;
+			}
+			case ENDIMemoryBufferViewType::UnsignedInteger:
+			{
+				TConstArrayView<uint32> BufferData = GetCurrentData<uint32>(ViewModel, CacheData.Get(), DisplaySimTarget.Get(ENiagaraSimTarget::CPUSim));
+				if (BufferData.IsValidIndex(Offset))
+				{
+					ItemText = FText::AsNumber(BufferData[Offset]);
+				}
+				break;
+			}
+			case ENDIMemoryBufferViewType::Hex:
+			{
+				TConstArrayView<uint32> BufferData = GetCurrentData<uint32>(ViewModel, CacheData.Get(), DisplaySimTarget.Get(ENiagaraSimTarget::CPUSim));
+				if (BufferData.IsValidIndex(Offset))
+				{
+					ItemText = FText::FromString(FString::Printf(TEXT("%08x"), BufferData[Offset]));
+				}
+				break;
+			}
+			default:
+				break;
+		}
+
+		if (!ItemText.IsEmpty())
+		{
+			return SNew(STextBlock)
+				.Text(ItemText);
+		}
+
+		return SNullWidget::NullWidget;
+	}
+
+	TSharedPtr<FNiagaraSimCacheViewModel>			ViewModel;
+	TStrongObjectPtr<UNDIMemoryBufferSimCacheData>	CacheData;
+	TAttribute<ENiagaraSimTarget>					DisplaySimTarget = ENiagaraSimTarget::CPUSim;
+	TAttribute<ENDIMemoryBufferViewType>			DisplayAsType = ENDIMemoryBufferViewType::Integer;
+	TSharedPtr<int32>								RowIndexPtr;
+};
+
+class SSimCacheView : public SCompoundWidget
+{
+public:
+	SLATE_BEGIN_ARGS(SSimCacheView)
+	{}
+	SLATE_END_ARGS()
+
+	virtual ~SSimCacheView() override
+	{
+		ViewModel->OnViewDataChanged().RemoveAll(this);
+	}
+
+	void Construct(const FArguments& InArgs, TSharedPtr<FNiagaraSimCacheViewModel> InViewModel, UNDIMemoryBufferSimCacheData* InCacheData)
+	{
+		ViewModel = InViewModel;
+		CacheData.Reset(InCacheData);
+
+		ViewModel->OnViewDataChanged().AddSP(this, &SSimCacheView::UpdateRows);
+
+		const UFNiagaraMemoryBufferSimCacheVisualizerSettings* Settings = GetDefault<UFNiagaraMemoryBufferSimCacheVisualizerSettings>();
+		DisplayColumns		= Settings->DisplayColumns;
+		DisplaySimTarget	= Settings->DisplaySimTarget;
+		DisplayAsType		= Settings->DisplayAsType;
+
+		HeaderRowWidget = SNew(SHeaderRow);
+		UpdateColumns();
+
+		SAssignNew(ListViewWidget, SListView<TSharedPtr<int32>>)
+			.ListItemsSource(&RowItems)
+			.OnGenerateRow(this, &SSimCacheView::MakeRowWidget)
+			.Visibility(EVisibility::Visible)
+			.SelectionMode(ESelectionMode::Single)
+			.ConsumeMouseWheel(EConsumeMouseWheel::Never)
+			.ScrollbarVisibility(EVisibility::Collapsed)
+			.HeaderRow(HeaderRowWidget);
+
+		ChildSlot
+		[
+			SNew(SVerticalBox)
+			// Display Columns
+			+SVerticalBox::Slot()
+			.AutoHeight()
+			[
+				SNew(SHorizontalBox)
+				+SHorizontalBox::Slot()
+				[
+					SNew(STextBlock)
+					.Text(LOCTEXT("DisplayNumElements", "Display Columns: "))
+				]
+				+SHorizontalBox::Slot()
+				.AutoWidth()
+				[
+					SNew(SSpinBox<int32>)
+					.MinDesiredWidth(50.0f)
+					.MinValue(1)
+					.MinSliderValue(1)
+					.Font(IDetailLayoutBuilder::GetDetailFont())
+					.Value(this, &SSimCacheView::GetDisplayColumns)
+					.OnValueChanged(this, &SSimCacheView::SetDisplayColumns)
+				]
+			]
+			// Display Type
+			+SVerticalBox::Slot()
+			.AutoHeight()
+			[
+				SNew(SHorizontalBox)
+				+SHorizontalBox::Slot()
+				[
+					SNew(STextBlock)
+					.Text(LOCTEXT("DisplayNumElements", "Display Type: "))
+				]
+				+SHorizontalBox::Slot()
+				.AutoWidth()
+				[
+					SNew(SEnumComboBox, StaticEnum<ENDIMemoryBufferViewType>())
+					.CurrentValue(this, &SSimCacheView::GetDisplayAsTypeInt)
+					.OnEnumSelectionChanged(this, &SSimCacheView::SetDisplayAsType)
+				]
+			]
+			// Display Target
+			+SVerticalBox::Slot()
+			.AutoHeight()
+			[
+				SNew(SHorizontalBox)
+				+SHorizontalBox::Slot()
+				[
+					SNew(STextBlock)
+					.Text(LOCTEXT("DisplayNumElements", "Display SimTarget: "))
+				]
+				+SHorizontalBox::Slot()
+				.AutoWidth()
+				[
+					SNew(SEnumComboBox, StaticEnum<ENiagaraSimTarget>())
+					.CurrentValue(this, &SSimCacheView::GetDisplaySimTargetAsInt)
+					.OnEnumSelectionChanged(this, &SSimCacheView::SetDisplaySimTarget)
+				]
+			]
+			// Number of buffer elements
+			+ SVerticalBox::Slot()
+			.AutoHeight()
+			[
+				SNew(SHorizontalBox)
+				+SHorizontalBox::Slot()
+				[
+					SNew(STextBlock)
+					.Text(LOCTEXT("BufferNumElements", "Buffer Num Elements: "))
+				]
+				+SHorizontalBox::Slot()
+				.AutoWidth()
+				[
+					SNew(STextBlock)
+					.Font(IDetailLayoutBuilder::GetDetailFont())
+					.Text(this, &SSimCacheView::GetBufferNumElementsText)
+				]
+			]
+			// Buffer Data View
+			+ SVerticalBox::Slot()
+			.AutoHeight()
+			.Padding(8.0f, 8.0f, 0.0f, 0.0f)
+			[
+				ListViewWidget.ToSharedRef()
+			]
+		];
+	}
+
+	TSharedRef<ITableRow> MakeRowWidget(const TSharedPtr<int32> RowIndexPtr, const TSharedRef<STableViewBase>& OwnerTable) const
+	{
+		return
+			SNew(SVisualizerRowWidget, OwnerTable)
+				.ViewModel(ViewModel)
+				.CacheData(CacheData.Get())
+				.RowIndexPtr(RowIndexPtr)
+				.DisplaySimTarget(this, &SSimCacheView::GetDisplaySimTarget)
+				.DisplayAsType(this, &SSimCacheView::GetDisplayAsType)
+				;
+	}
+
+	void UpdateColumns(bool bRefresh = true)
+	{
+		using namespace NDIMemoryBufferSimCacheVisualizerLocal;
+		HeaderRowWidget->ClearColumns();
+		HeaderRowWidget->AddColumn(
+			SHeaderRow::Column(NAME_ElementOffset)
+			.DefaultLabel(LOCTEXT("ElementOffset", "Element Offset"))
+			.HAlignHeader(HAlign_Center)
+			.VAlignHeader(VAlign_Fill)
+			.HAlignCell(HAlign_Center)
+			.VAlignCell(VAlign_Center)
+			.ManualWidth(100)
+			.SortMode(EColumnSortMode::None)
+		);
+
+		for ( int32 i=0; i < DisplayColumns; ++i )
+		{
+			FName ColumnName(*FString::Printf(TEXT("%d"), i), i);
+
+			HeaderRowWidget->AddColumn(
+				SHeaderRow::Column(ColumnName)
+				.DefaultLabel(FText::AsNumber(i))
+				.HAlignHeader(HAlign_Center)
+				.VAlignHeader(VAlign_Fill)
+				.HAlignCell(HAlign_Center)
+				.VAlignCell(VAlign_Center)
+				.ManualWidth(100)
+				.SortMode(EColumnSortMode::None)
+			);
+		}
+		if (ListViewWidget && bRefresh)
+		{
+			ListViewWidget->RequestListRefresh();
+		}
+	}
+
+	void UpdateRows(bool bRefresh = true)
+	{
+		using namespace NDIMemoryBufferSimCacheVisualizerLocal;
+		TConstArrayView<uint32> BufferData = GetCurrentData<uint32>(ViewModel, CacheData.Get(), DisplaySimTarget);
+		const int32 NumRows = FMath::DivideAndRoundUp(BufferData.Num(), DisplayColumns);
+		RowItems.Reset(NumRows);
+		for ( int32 i=0; i < NumRows; ++i )
+		{
+			RowItems.Emplace(MakeShared<int32>(i * DisplayColumns));
+		}
+		if (ListViewWidget && bRefresh)
+		{
+			ListViewWidget->RequestListRefresh();
+		}
+	}
+
+	FText GetBufferNumElementsText() const
+	{
+		using namespace NDIMemoryBufferSimCacheVisualizerLocal;
+		TConstArrayView<uint32> BufferData = GetCurrentData<uint32>(ViewModel, CacheData.Get(), DisplaySimTarget);
+		return FText::AsNumber(BufferData.Num());
+	}
+
+	int32 GetDisplayColumns() const { return DisplayColumns; }
+
+	ENiagaraSimTarget GetDisplaySimTarget() const { return DisplaySimTarget; }
+	int32 GetDisplaySimTargetAsInt() const { return (int32)DisplaySimTarget; }
+
+	ENDIMemoryBufferViewType GetDisplayAsType() const { return DisplayAsType; }
+	int32 GetDisplayAsTypeInt() const { return (int32)DisplayAsType; }
+
+	void SetDisplayColumns(int32 InDisplayColumns)
+	{
+		InDisplayColumns = FMath::Max(InDisplayColumns, 1);
+		if (DisplayColumns != InDisplayColumns)
+		{
+			DisplayColumns = InDisplayColumns;
+			UpdateColumns(false);
+			UpdateRows();
+
+			UFNiagaraMemoryBufferSimCacheVisualizerSettings* Settings = GetMutableDefault<UFNiagaraMemoryBufferSimCacheVisualizerSettings>();
+			Settings->DisplayColumns = DisplayColumns;
+			Settings->SaveConfig();
+		}
+	}
+
+	void SetDisplaySimTarget(int32 InNewValue, ESelectInfo::Type InInfo)
+	{
+		const ENiagaraSimTarget NewValue = (ENiagaraSimTarget)InNewValue;
+		if (DisplaySimTarget != NewValue)
+		{
+			DisplaySimTarget = NewValue;
+			UpdateRows();
+
+			UFNiagaraMemoryBufferSimCacheVisualizerSettings* Settings = GetMutableDefault<UFNiagaraMemoryBufferSimCacheVisualizerSettings>();
+			Settings->DisplaySimTarget = DisplaySimTarget;
+			Settings->SaveConfig();
+		}
+	}
+
+	void SetDisplayAsType(int32 InNewValue, ESelectInfo::Type InInfo)
+	{
+		const ENDIMemoryBufferViewType NewValue = (ENDIMemoryBufferViewType)InNewValue;
+		if (DisplayAsType != NewValue)
+		{
+			DisplayAsType = NewValue;
+			UpdateRows();
+
+			UFNiagaraMemoryBufferSimCacheVisualizerSettings* Settings = GetMutableDefault<UFNiagaraMemoryBufferSimCacheVisualizerSettings>();
+			Settings->DisplayAsType = DisplayAsType;
+			Settings->SaveConfig();
+		}
+	}
+
+private:
+	TSharedPtr<FNiagaraSimCacheViewModel>			ViewModel;
+	TStrongObjectPtr<UNDIMemoryBufferSimCacheData>	CacheData;
+
+	int32											DisplayColumns = 1;
+	ENiagaraSimTarget								DisplaySimTarget = ENiagaraSimTarget::CPUSim;
+	ENDIMemoryBufferViewType						DisplayAsType = ENDIMemoryBufferViewType::Integer;
+
+	TSharedPtr<SHeaderRow>							HeaderRowWidget;
+	TSharedPtr<SListView<TSharedPtr<int32>>>		ListViewWidget;
+	TArray<TSharedPtr<int32>>						RowItems;
+};
+
+} // NDIMemoryBufferSimCacheVisualizer
+
+////////////////////////////////////////////////////////////////////////////////////////////////////////////////////
+
+TSharedPtr<SWidget> FNiagaraMemoryBufferSimCacheVisualizer::CreateWidgetFor(UObject* InCachedData, TSharedPtr<FNiagaraSimCacheViewModel> ViewModel)
+{
+	using namespace NDIMemoryBufferSimCacheVisualizer;
+
+	if (UNDIMemoryBufferSimCacheData* CachedData = Cast<UNDIMemoryBufferSimCacheData>(InCachedData))
+	{
+		return SNew(SSimCacheView, ViewModel, CachedData);
+	}
+	return TSharedPtr<SWidget>();
+}
+
+#undef LOCTEXT_NAMESPACE
diff --git a/Engine/Plugins/FX/Niagara/Source/NiagaraEditor/Private/Customizations/SimCache/NiagaraMemoryBufferSimCacheVisualizer.h b/Engine/Plugins/FX/Niagara/Source/NiagaraEditor/Private/Customizations/SimCache/NiagaraMemoryBufferSimCacheVisualizer.h
new file mode 100644
--- /dev/null
+++ b/Engine/Plugins/FX/Niagara/Source/NiagaraEditor/Private/Customizations/SimCache/NiagaraMemoryBufferSimCacheVisualizer.h
@@ -0,0 +1,44 @@
+// Copyright Epic Games, Inc. All Rights Reserved.
+
+#pragma once
+
+#include "NiagaraCommon.h"
+#include "Customizations/NiagaraDataInterfaceSimCacheVisualizer.h"
+
+#include "NiagaraMemoryBufferSimCacheVisualizer.generated.h"
+
+UENUM()
+enum class ENDIMemoryBufferViewType
+{
+	Float,
+	Integer,
+	UnsignedInteger,
+	Hex,
+};
+
+UCLASS(config=EditorPerProjectUserSettings, defaultconfig)
+class UFNiagaraMemoryBufferSimCacheVisualizerSettings : public UObject
+{
+	GENERATED_BODY()
+
+public:
+	UPROPERTY(Config, EditAnywhere, Category = "Settings")
+	ENiagaraSimTarget DisplaySimTarget = ENiagaraSimTarget::CPUSim;
+
+	UPROPERTY(Config, EditAnywhere, Category = "Settings")
+	ENDIMemoryBufferViewType DisplayAsType = ENDIMemoryBufferViewType::Integer;
+
+	UPROPERTY(Config, EditAnywhere, Category = "Settings")
+	int32 DisplayColumns = 1;
+};
+
+/*
+ * Provides a custom widget to show the memory buffer DI data in the sim cache in a table.
+ */
+class FNiagaraMemoryBufferSimCacheVisualizer : public INiagaraDataInterfaceSimCacheVisualizer
+{
+public:
+	virtual ~FNiagaraMemoryBufferSimCacheVisualizer() override = default;
+	
+	virtual TSharedPtr<SWidget> CreateWidgetFor(UObject* CachedData, TSharedPtr<FNiagaraSimCacheViewModel> ViewModel) override;
+};
diff --git a/Engine/Plugins/FX/Niagara/Source/NiagaraEditor/Private/NiagaraEditorModule.cpp b/Engine/Plugins/FX/Niagara/Source/NiagaraEditor/Private/NiagaraEditorModule.cpp
--- a/Engine/Plugins/FX/Niagara/Source/NiagaraEditor/Private/NiagaraEditorModule.cpp
+++ b/Engine/Plugins/FX/Niagara/Source/NiagaraEditor/Private/NiagaraEditorModule.cpp
@@ -55,6 +55,7 @@
 #include "NiagaraDataInterfaceColorCurve.h"
 #include "DataInterface/NiagaraDataInterfaceDataChannelRead.h"
 #include "DataInterface/NiagaraDataInterfaceDataChannelWrite.h"
+#include "DataInterface/NiagaraDataInterfaceMemoryBuffer.h"
 #include "NiagaraDataInterfaceRenderTargetVolume.h"
 
 #include "ViewModels/NiagaraScriptViewModel.h"
@@ -117,6 +118,7 @@
 #include "Customizations/NiagaraDataChannelDetails.h"
 #include "Customizations/SimCache/NiagaraArraySimCacheVisualizer.h"
 #include "Customizations/SimCache/FNiagaraDataChannelSimCacheVisualizer.h"
+#include "Customizations/SimCache/NiagaraMemoryBufferSimCacheVisualizer.h"
 #include "Customizations/SimCache/NiagaraRenderTargetVolumeSimCacheVisualizer.h"
 
 #include "NiagaraComponent.h"
@@ -1430,6 +1432,7 @@ void FNiagaraEditorModule::StartupModule()
 
 	// Register sim cache visualizers
 	RegisterDataInterfaceCacheVisualizer(UNiagaraDataInterfaceDataChannelWrite::StaticClass(), MakeShared<FNiagaraDataChannelSimCacheVisualizer>());
+	RegisterDataInterfaceCacheVisualizer(UNiagaraDataInterfaceMemoryBuffer::StaticClass(), MakeShared<FNiagaraMemoryBufferSimCacheVisualizer>());
 	RegisterDataInterfaceCacheVisualizer(UNiagaraDataInterfaceRenderTargetVolume::StaticClass(), MakeShared<FNiagaraRenderTargetVolumeSimCacheVisualizer>());
 	for (TObjectIterator<UClass> It; It; ++It)
 	{
