Fix warning V698: Expression 'Compare(....) == -1' is incorrect. This function can return not only the value '-1', but any negative value. Consider using 'Compare(....) < 0' instead.
#rnx

diff --git a/Engine/Plugins/FX/Niagara/Source/NiagaraEditor/Private/Widgets/SNiagaraGraphActionMenu.cpp b/Engine/Plugins/FX/Niagara/Source/NiagaraEditor/Private/Widgets/SNiagaraGraphActionMenu.cpp
--- a/Engine/Plugins/FX/Niagara/Source/NiagaraEditor/Private/Widgets/SNiagaraGraphActionMenu.cpp
+++ b/Engine/Plugins/FX/Niagara/Source/NiagaraEditor/Private/Widgets/SNiagaraGraphActionMenu.cpp
@@ -146,7 +146,7 @@ bool SNiagaraGraphActionMenu::OnCompareCategoriesForEquality(const FString& Cate
 bool SNiagaraGraphActionMenu::OnCompareCategoriesForSorting(const FString& CategoryA,
 	const FString& CategoryB)
 {
-	return CategoryA.Compare(CategoryB) == -1;
+	return CategoryA.Compare(CategoryB) < 0;
 }
 
 bool SNiagaraGraphActionMenu::OnCompareItemsForEquality(const TSharedPtr<FNiagaraAction_NewNode>& ItemA, const TSharedPtr<FNiagaraAction_NewNode>& ItemB)
diff --git a/Engine/Plugins/FX/Niagara/Source/NiagaraEditorWidgets/Private/Stack/SNiagaraStackFunctionInputValue.cpp b/Engine/Plugins/FX/Niagara/Source/NiagaraEditorWidgets/Private/Stack/SNiagaraStackFunctionInputValue.cpp
--- a/Engine/Plugins/FX/Niagara/Source/NiagaraEditorWidgets/Private/Stack/SNiagaraStackFunctionInputValue.cpp
+++ b/Engine/Plugins/FX/Niagara/Source/NiagaraEditorWidgets/Private/Stack/SNiagaraStackFunctionInputValue.cpp
@@ -1470,7 +1470,7 @@ bool SNiagaraStackFunctionInputValue::OnCompareCategoriesForEquality(const FStri
 
 bool SNiagaraStackFunctionInputValue::OnCompareCategoriesForSorting(const FString& CategoryA, const FString& CategoryB)
 {
-	return CategoryA.Compare(CategoryB) == -1;
+	return CategoryA.Compare(CategoryB) < 0;
 }
 
 bool SNiagaraStackFunctionInputValue::OnCompareItemsForEquality(const TSharedPtr<FNiagaraMenuAction_Generic>& ItemA,
diff --git a/Engine/Plugins/FX/Niagara/Source/NiagaraEditorWidgets/Private/Stack/SNiagaraStackItemGroupAddMenu.cpp b/Engine/Plugins/FX/Niagara/Source/NiagaraEditorWidgets/Private/Stack/SNiagaraStackItemGroupAddMenu.cpp
--- a/Engine/Plugins/FX/Niagara/Source/NiagaraEditorWidgets/Private/Stack/SNiagaraStackItemGroupAddMenu.cpp
+++ b/Engine/Plugins/FX/Niagara/Source/NiagaraEditorWidgets/Private/Stack/SNiagaraStackItemGroupAddMenu.cpp
@@ -204,7 +204,7 @@ bool SNiagaraStackItemGroupAddMenu::OnCompareCategoriesForEquality(const FString
 
 bool SNiagaraStackItemGroupAddMenu::OnCompareCategoriesForSorting(const FString& CategoryA, const FString& CategoryB)
 {
-	return CategoryA.Compare(CategoryB) == -1;
+	return CategoryA.Compare(CategoryB) < 0;
 }
 
 bool SNiagaraStackItemGroupAddMenu::OnCompareItemsForEquality(const TSharedPtr<FNiagaraMenuAction_Generic>& ItemA,
