- Fix some duplicate localization keys

diff --git a/Engine/Plugins/FX/Niagara/Source/NiagaraEditor/Private/Customizations/SimCache/NiagaraArraySimCacheVisualizer.cpp b/Engine/Plugins/FX/Niagara/Source/NiagaraEditor/Private/Customizations/SimCache/NiagaraArraySimCacheVisualizer.cpp
--- a/Engine/Plugins/FX/Niagara/Source/NiagaraEditor/Private/Customizations/SimCache/NiagaraArraySimCacheVisualizer.cpp
+++ b/Engine/Plugins/FX/Niagara/Source/NiagaraEditor/Private/Customizations/SimCache/NiagaraArraySimCacheVisualizer.cpp
@@ -255,7 +255,7 @@ public:
 			NumElements = CacheData->GpuFrameData[FrameIndex].NumElements;
 		}
 
-		return FText::Format(LOCTEXT("CpuArrayInfoFormat", "GPU Array Num Elements = {0}"), NumElements);
+		return FText::Format(LOCTEXT("GpuArrayInfoFormat", "GPU Array Num Elements = {0}"), NumElements);
 	}
 
 private:
diff --git a/Engine/Plugins/FX/Niagara/Source/NiagaraEditor/Private/Customizations/SimCache/NiagaraMemoryBufferSimCacheVisualizer.cpp b/Engine/Plugins/FX/Niagara/Source/NiagaraEditor/Private/Customizations/SimCache/NiagaraMemoryBufferSimCacheVisualizer.cpp
--- a/Engine/Plugins/FX/Niagara/Source/NiagaraEditor/Private/Customizations/SimCache/NiagaraMemoryBufferSimCacheVisualizer.cpp
+++ b/Engine/Plugins/FX/Niagara/Source/NiagaraEditor/Private/Customizations/SimCache/NiagaraMemoryBufferSimCacheVisualizer.cpp
@@ -104,7 +104,7 @@ public:
 					}
 					else
 					{
-						ItemText = LOCTEXT("Invalid Float", "Invalid Float");
+						ItemText = LOCTEXT("InvalidFloat", "Invalid Float");
 					}
 				}
 				break;
@@ -203,7 +203,7 @@ public:
 				+SHorizontalBox::Slot()
 				[
 					SNew(STextBlock)
-					.Text(LOCTEXT("DisplayNumElements", "Display Columns: "))
+					.Text(LOCTEXT("DisplayColumns", "Display Columns: "))
 				]
 				+SHorizontalBox::Slot()
 				.AutoWidth()
@@ -225,7 +225,7 @@ public:
 				+SHorizontalBox::Slot()
 				[
 					SNew(STextBlock)
-					.Text(LOCTEXT("DisplayNumElements", "Display Type: "))
+					.Text(LOCTEXT("DisplayType", "Display Type: "))
 				]
 				+SHorizontalBox::Slot()
 				.AutoWidth()
@@ -243,7 +243,7 @@ public:
 				+SHorizontalBox::Slot()
 				[
 					SNew(STextBlock)
-					.Text(LOCTEXT("DisplayNumElements", "Display SimTarget: "))
+					.Text(LOCTEXT("DisplaySimTarget", "Display SimTarget: "))
 				]
 				+SHorizontalBox::Slot()
 				.AutoWidth()
