- Fix for emitter asset renaming
- Assets do not use the number portion of the FName and would sometimes incorrectly collide on rename which could lead to the asset file vs name being mismatched

diff --git a/Engine/Plugins/FX/Niagara/Source/Niagara/Private/NiagaraEmitter.cpp b/Engine/Plugins/FX/Niagara/Source/Niagara/Private/NiagaraEmitter.cpp
--- a/Engine/Plugins/FX/Niagara/Source/Niagara/Private/NiagaraEmitter.cpp
+++ b/Engine/Plugins/FX/Niagara/Source/Niagara/Private/NiagaraEmitter.cpp
@@ -3069,7 +3069,9 @@ bool UNiagaraEmitter::SetUniqueEmitterName(const FString& InName)
 		FString OldName = UniqueEmitterName;
 		UniqueEmitterName = InName;
 
-		if (GetName() != InName)
+		// Note: Assets don't care about the number portion so we need to compare without the number otherwise renaming can collide
+		const FString ExistingName = IsAsset() ? GetFName().GetPlainNameString() : GetName();
+		if (ExistingName != InName)
 		{
 			// Also rename the underlying uobject to keep things consistent.
 			FName UniqueObjectName = MakeUniqueObjectName(GetOuter(), StaticClass(), *InName);
