fix crash opening some emitters when EmitterInstance is null

diff --git a/Engine/Plugins/FX/Niagara/Source/Niagara/Private/NiagaraDataInterfaceParticleRead.cpp b/Engine/Plugins/FX/Niagara/Source/Niagara/Private/NiagaraDataInterfaceParticleRead.cpp
--- a/Engine/Plugins/FX/Niagara/Source/Niagara/Private/NiagaraDataInterfaceParticleRead.cpp
+++ b/Engine/Plugins/FX/Niagara/Source/Niagara/Private/NiagaraDataInterfaceParticleRead.cpp
@@ -1196,7 +1196,7 @@ void UNiagaraDataInterfaceParticleRead::VMGetLocalSpace(FVectorVMExternalFunctio
 	FNDIOutputParam<bool> OutIsLocalSpace(Context);
 
 	const FNiagaraEmitterInstance* EmitterInstance = InstData.Get()->EmitterInstance;
-	const bool bIsLocalSpace = EmitterInstance->IsLocalSpace();
+	const bool bIsLocalSpace = EmitterInstance ? EmitterInstance->IsLocalSpace() : false;
 
 	for (int32 InstanceIdx = 0; InstanceIdx < Context.GetNumInstances(); ++InstanceIdx)
 	{
