- SimCache to Json support

diff --git a/Engine/Plugins/FX/Niagara/Source/Niagara/Classes/NiagaraSimCache.h b/Engine/Plugins/FX/Niagara/Source/Niagara/Classes/NiagaraSimCache.h
--- a/Engine/Plugins/FX/Niagara/Source/Niagara/Classes/NiagaraSimCache.h
+++ b/Engine/Plugins/FX/Niagara/Source/Niagara/Classes/NiagaraSimCache.h
@@ -392,10 +392,13 @@ public:
 	// UObject Interface
 
 	/** Get the caches assigned GUID, this can be set from a user or will be auto created on first write of the cache. */
-	const FGuid& GetCacheGuid() { return CacheGuid; }
+	const FGuid& GetCacheGuid() const { return CacheGuid; }
 	/** Set the caches GUID to use. */
 	void SetCacheGuid(const FGuid& InGuid) { CacheGuid = InGuid; }
 
+	/** Get the system asset this cache was created from. */
+	TSoftObjectPtr<UNiagaraSystem> GetSystemAsset() const { return SoftNiagaraSystem; }
+
 	/** A valid cache is one that contains at least 1 frames worth of data. */
 	UFUNCTION(BlueprintCallable, Category=NiagaraSimCache, meta=(DisplayName="IsValid"))
 	bool IsCacheValid() const { return SoftNiagaraSystem.IsNull() == false; }
diff --git a/Engine/Plugins/FX/Niagara/Source/Niagara/Classes/NiagaraSimCacheJson.h b/Engine/Plugins/FX/Niagara/Source/Niagara/Classes/NiagaraSimCacheJson.h
new file mode 100644
--- /dev/null
+++ b/Engine/Plugins/FX/Niagara/Source/Niagara/Classes/NiagaraSimCacheJson.h
@@ -0,0 +1,16 @@
+// Copyright Epic Games, Inc. All Rights Reserved.
+
+#pragma  once
+
+#include "NiagaraSimCache.h"
+
+class FJsonObject;
+
+/** Experimental Json interop for Niagara Sim Caches. */
+struct FNiagaraSimCacheJson
+{
+	NIAGARA_API static void DumpToFile(const UNiagaraSimCache& SimCache, const FString& FullPath);
+
+	NIAGARA_API static TSharedPtr<FJsonObject> ToJson(const UNiagaraSimCache& SimCache);
+	NIAGARA_API static TSharedPtr<FJsonObject> EmitterFrameToJson(const UNiagaraSimCache& SimCache, int EmitterIndex, int FrameIndex);
+};
diff --git a/Engine/Plugins/FX/Niagara/Source/Niagara/Private/NiagaraSimCacheJson.cpp b/Engine/Plugins/FX/Niagara/Source/Niagara/Private/NiagaraSimCacheJson.cpp
new file mode 100644
--- /dev/null
+++ b/Engine/Plugins/FX/Niagara/Source/Niagara/Private/NiagaraSimCacheJson.cpp
@@ -0,0 +1,126 @@
+// Copyright Epic Games, Inc. All Rights Reserved.
+
+#include "NiagaraSimCacheJson.h"
+
+#include "Dom/JsonObject.h"
+#include "Dom/JsonValue.h"
+#include "Misc/FileHelper.h"
+#include "Policies/PrettyJsonPrintPolicy.h"
+#include "Serialization/JsonSerializer.h"
+#include "Serialization/JsonWriter.h"
+
+void FNiagaraSimCacheJson::DumpToFile(const UNiagaraSimCache& SimCache, const FString& FullPath)
+{
+	TSharedPtr<FJsonObject> JsonObject = ToJson(SimCache);
+	if (JsonObject.IsValid())
+	{
+		UE_LOG(LogNiagara, Warning, TEXT("Writing file to %s"), *FullPath);
+
+		FString OutputString;
+		TSharedRef<TJsonWriter<TCHAR, TPrettyJsonPrintPolicy<TCHAR>>> Writer = TJsonWriterFactory<TCHAR, TPrettyJsonPrintPolicy<TCHAR>>::Create(&OutputString);
+		FJsonSerializer::Serialize(JsonObject.ToSharedRef(), Writer);
+		FFileHelper::SaveStringToFile(OutputString, *FullPath);
+	}
+}
+
+TSharedPtr<FJsonObject> FNiagaraSimCacheJson::ToJson(const UNiagaraSimCache& SimCache)
+{
+	if (!SimCache.IsCacheValid())
+	{
+		return nullptr;
+	}
+
+	TSharedPtr<FJsonObject> JsonCacheObject = MakeShared<FJsonObject>();
+	JsonCacheObject->SetStringField(TEXT("SystemAsset"), SimCache.GetSystemAsset().ToString());
+	JsonCacheObject->SetStringField(TEXT("CacheGuid"), SimCache.GetCacheGuid().ToString());
+	JsonCacheObject->SetNumberField(TEXT("StartSeconds"), SimCache.GetStartSeconds());
+	JsonCacheObject->SetNumberField(TEXT("DurationSeconds"), SimCache.GetDurationSeconds());
+	JsonCacheObject->SetNumberField(TEXT("NumFrames"), SimCache.GetNumFrames());
+	JsonCacheObject->SetNumberField(TEXT("NumEmitters"), SimCache.GetNumEmitters());
+
+	// Write System Instance
+	{
+		TArray<TSharedPtr<FJsonValue>> JsonFrames;
+		for (int32 iFrame=0; iFrame < SimCache.GetNumFrames(); ++iFrame)
+		{
+			TSharedPtr<FJsonObject> JsonFrame = EmitterFrameToJson(SimCache, INDEX_NONE, iFrame);
+			JsonFrames.Add(TSharedPtr<FJsonValue>(new FJsonValueObject(JsonFrame)));
+		}
+		JsonCacheObject->SetArrayField(TEXT("SystemInstance"), JsonFrames);
+	}
+
+	// Write Emitter Instances
+	for ( int32 iEmitter=0; iEmitter < SimCache.GetNumEmitters(); ++iEmitter )
+	{
+		TArray<TSharedPtr<FJsonValue>> JsonFrames;
+		for (int32 iFrame = 0; iFrame < SimCache.GetNumFrames(); ++iFrame)
+		{
+			TSharedPtr<FJsonObject> JsonFrame = EmitterFrameToJson(SimCache, iEmitter, iFrame);
+			JsonFrames.Add(TSharedPtr<FJsonValue>(new FJsonValueObject(JsonFrame)));
+		}
+		JsonCacheObject->SetArrayField(SimCache.GetEmitterName(iEmitter).ToString(), JsonFrames);
+	}
+
+	return JsonCacheObject;
+}
+
+TSharedPtr<FJsonObject> FNiagaraSimCacheJson::EmitterFrameToJson(const UNiagaraSimCache& SimCache, int EmitterIndex, int FrameIndex)
+{
+	const int NumInstances = SimCache.GetEmitterNumInstances(EmitterIndex, FrameIndex);
+
+	TSharedPtr<FJsonObject> EmitterObject = MakeShared<FJsonObject>();
+	EmitterObject->SetNumberField(TEXT("NumInstances"), NumInstances);
+	if (NumInstances > 0)
+	{
+		const FName EmitterName = SimCache.GetEmitterName(EmitterIndex);
+
+		TArray<FNiagaraVariableBase> Attributes;
+		SimCache.ForEachEmitterAttribute(EmitterIndex, [&Attributes](const FNiagaraSimCacheVariable& CacheVariable) -> bool { Attributes.Emplace(CacheVariable.Variable); return true; });
+		Algo::Sort(Attributes, [](const FNiagaraVariableBase& Lhs, const FNiagaraVariableBase& Rhs) { return FNameLexicalLess()(Lhs.GetName(), Rhs.GetName()); });
+
+		TArray<TSharedPtr<FJsonValue>> AttributesObject;
+		for (const FNiagaraVariableBase& Attribute : Attributes)
+		{
+			TSharedPtr<FJsonObject> AttributeObject = MakeShared<FJsonObject>();
+			AttributeObject->SetStringField(TEXT("Name"), Attribute.GetName().ToString());
+			AttributeObject->SetStringField(TEXT("Type"), Attribute.GetType().GetName());
+
+			AttributesObject.Add(TSharedPtr<FJsonValue>(new FJsonValueObject(AttributeObject)));
+
+			TArray<float>		Floats;
+			TArray<FFloat16>	Halfs;
+			TArray<int32>		Ints;
+			SimCache.ReadAttribute(Floats, Halfs, Ints, Attribute.GetName(), EmitterName, FrameIndex);
+
+			if (Floats.Num() > 0)
+			{
+				TArray<TSharedPtr<FJsonValue>> Values;
+				for ( float v : Floats )
+				{
+					Values.Emplace(MakeShared<FJsonValueNumber>(v));
+				}
+				AttributeObject->SetArrayField(TEXT("Floats"), Values);
+			}
+			if (Halfs.Num() > 0)
+			{
+				TArray<TSharedPtr<FJsonValue>> Values;
+				for (FFloat16 v : Halfs)
+				{
+					Values.Emplace(MakeShared<FJsonValueNumber>(v));
+				}
+				AttributeObject->SetArrayField(TEXT("Halfs"), Values);
+			}
+			if (Ints.Num() > 0)
+			{
+				TArray<TSharedPtr<FJsonValue>> Values;
+				for (int32 v : Ints)
+				{
+					Values.Emplace(MakeShared<FJsonValueNumber>(v));
+				}
+				AttributeObject->SetArrayField(TEXT("Ints"), Values);
+			}
+		}
+		EmitterObject->SetArrayField(TEXT("Attributes"), AttributesObject);
+	}
+	return EmitterObject;
+}
