- Add cache comparison functionality to arrays

diff --git a/Engine/Plugins/FX/Niagara/Source/Niagara/Classes/NiagaraDataInterfaceArray.h b/Engine/Plugins/FX/Niagara/Source/Niagara/Classes/NiagaraDataInterfaceArray.h
--- a/Engine/Plugins/FX/Niagara/Source/Niagara/Classes/NiagaraDataInterfaceArray.h
+++ b/Engine/Plugins/FX/Niagara/Source/Niagara/Classes/NiagaraDataInterfaceArray.h
@@ -66,7 +66,8 @@ struct INDIArrayProxyBase : public FNiagaraDataInterfaceProxyRW
 #endif
 	virtual bool SimCacheWriteFrame(UNDIArraySimCacheData* CacheData, int FrameIndex, FNiagaraSystemInstance* SystemInstance) const = 0;
 	virtual bool SimCacheReadFrame(UNDIArraySimCacheData* CacheData, int FrameIndex, FNiagaraSystemInstance* SystemInstance) = 0;
-	virtual FString SimCacheVisualizerRead(UNDIArraySimCacheData* CacheData, FNDIArraySimCacheDataFrame& FrameData, int Element) const = 0;
+	virtual bool SimCacheCompareElement(const uint8* LhsData, const uint8* RhsData, int32 Element, float Tolerance) const = 0;
+	virtual FString SimCacheVisualizerRead(const UNDIArraySimCacheData* CacheData, const FNDIArraySimCacheDataFrame& FrameData, int Element) const = 0;
 
 	virtual bool CopyToInternal(INDIArrayProxyBase* Destination) const = 0;
 	virtual bool Equals(const INDIArrayProxyBase* Other) const = 0;
@@ -123,8 +124,9 @@ public:
 	virtual bool SimCacheWriteFrame(UObject* StorageObject, int FrameIndex, FNiagaraSystemInstance* SystemInstance, const void* OptionalPerInstanceData, FNiagaraSimCacheFeedbackContext& FeedbackContext) const override;
 	//virtual bool SimCacheEndWrite(UObject* StorageObject) const override;
 	virtual bool SimCacheReadFrame(UObject* StorageObject, int FrameA, int FrameB, float Interp, FNiagaraSystemInstance* SystemInstance, void* OptionalPerInstanceData) override;
+	virtual bool SimCacheCompareFrame(UObject* LhsStorageObject, UObject* RhsStorageObject, int FrameIndex, TOptional<float> Tolerance, FString& OutErrors) const override;
 	//~ INiagaraSimCacheCustomStorageInterface interface END
-	NIAGARA_API FString SimCacheVisualizerRead(UNDIArraySimCacheData* CacheData, FNDIArraySimCacheDataFrame& FrameData, int Element) const;
+	NIAGARA_API FString SimCacheVisualizerRead(const UNDIArraySimCacheData* CacheData, const FNDIArraySimCacheDataFrame& FrameData, int Element) const;
 
 	/** ReadWrite lock to ensure safe access to the underlying array. */
 	FRWLock ArrayRWGuard;
diff --git a/Engine/Plugins/FX/Niagara/Source/Niagara/Classes/NiagaraDataInterfaceArrayImpl.h b/Engine/Plugins/FX/Niagara/Source/Niagara/Classes/NiagaraDataInterfaceArrayImpl.h
--- a/Engine/Plugins/FX/Niagara/Source/Niagara/Classes/NiagaraDataInterfaceArrayImpl.h
+++ b/Engine/Plugins/FX/Niagara/Source/Niagara/Classes/NiagaraDataInterfaceArrayImpl.h
@@ -141,6 +141,11 @@ struct FNDIArrayImplHelperBase
 		FMemory::Memcpy(Dest, Src, NumElements * sizeof(TArrayType));
 	}
 
+	static bool IsNearlyEqual(const TArrayType& Lhs, const TArrayType& Rhs, float Tolerance)
+	{
+		return FMath::IsNearlyEqual(Lhs, Rhs, Tolerance);
+	}
+
 	//static TArrayType AtomicAdd(TArrayType* Dest, TArrayType Value) { check(false); }
 	//static TArrayType AtomicMin(TArrayType* Dest, TArrayType Value) { check(false); }
 	//static TArrayType AtomicMax(TArrayType* Dest, TArrayType Value) { check(false); }
@@ -1077,6 +1082,7 @@ struct FNDIArrayProxyImpl : public INDIArrayProxyBase
 						}
 					);
 					RHICmdList.Transition(TransitionsAfter);
+					ReadbackManager->WaitCompletion(RHICmdList);
 				}
 			);
 
@@ -1129,7 +1135,14 @@ struct FNDIArrayProxyImpl : public INDIArrayProxyBase
 		return true;
 	}
 
-	virtual FString SimCacheVisualizerRead(UNDIArraySimCacheData* CacheData, FNDIArraySimCacheDataFrame& FrameData, int Element) const override
+	virtual bool SimCacheCompareElement(const uint8* LhsData, const uint8* RhsData, int32 Element, float Tolerance) const override
+	{
+		const TArrayType* LhsArrayData = reinterpret_cast<const TArrayType*>(LhsData);
+		const TArrayType* RhsArrayData = reinterpret_cast<const TArrayType*>(RhsData);
+		return FNDIArrayImplHelper<TArrayType>::IsNearlyEqual(LhsArrayData[Element], RhsArrayData[Element], Tolerance);
+	}
+
+	virtual FString SimCacheVisualizerRead(const UNDIArraySimCacheData* CacheData, const FNDIArraySimCacheDataFrame& FrameData, int Element) const override
 	{
 		FString OutValue;
 		if (Element < FrameData.NumElements)
diff --git a/Engine/Plugins/FX/Niagara/Source/Niagara/Classes/NiagaraSimCacheCustomStorageInterface.h b/Engine/Plugins/FX/Niagara/Source/Niagara/Classes/NiagaraSimCacheCustomStorageInterface.h
--- a/Engine/Plugins/FX/Niagara/Source/Niagara/Classes/NiagaraSimCacheCustomStorageInterface.h
+++ b/Engine/Plugins/FX/Niagara/Source/Niagara/Classes/NiagaraSimCacheCustomStorageInterface.h
@@ -61,7 +61,7 @@ public:
 	Called to compare a frame between two separate simulation cache storages
 	This will be called on the CDO object since we do not have the actual data interface.
 	*/
-	virtual bool SimCacheCompareFrame(UObject* LhsStorageObject, UObject* RhsStorageObject, int FrameIndex, TOptional<float> Tolerance, FString& OutErrors) { OutErrors = TEXT("Compare not implemented"); return false; }
+	virtual bool SimCacheCompareFrame(UObject* LhsStorageObject, UObject* RhsStorageObject, int FrameIndex, TOptional<float> Tolerance, FString& OutErrors) const { OutErrors = TEXT("Compare not implemented"); return false; }
 
 	/**
 	This function allows you to preserve a list of attributes when building a renderer only cache.
diff --git a/Engine/Plugins/FX/Niagara/Source/Niagara/Private/DataInterface/NiagaraDataInterfaceSimpleCounter.cpp b/Engine/Plugins/FX/Niagara/Source/Niagara/Private/DataInterface/NiagaraDataInterfaceSimpleCounter.cpp
--- a/Engine/Plugins/FX/Niagara/Source/Niagara/Private/DataInterface/NiagaraDataInterfaceSimpleCounter.cpp
+++ b/Engine/Plugins/FX/Niagara/Source/Niagara/Private/DataInterface/NiagaraDataInterfaceSimpleCounter.cpp
@@ -562,7 +562,7 @@ bool UNiagaraDataInterfaceSimpleCounter::SimCacheReadFrame(UObject* StorageObjec
 	return true;
 }
 
-bool UNiagaraDataInterfaceSimpleCounter::SimCacheCompareFrame(UObject* LhsStorageObject, UObject* RhsStorageObject, int FrameIndex, TOptional<float> Tolerance, FString& OutErrors)
+bool UNiagaraDataInterfaceSimpleCounter::SimCacheCompareFrame(UObject* LhsStorageObject, UObject* RhsStorageObject, int FrameIndex, TOptional<float> Tolerance, FString& OutErrors) const
 {
 	UNDISimpleCounterSimCacheData* LhsCacheData = CastChecked<UNDISimpleCounterSimCacheData>(LhsStorageObject);
 	UNDISimpleCounterSimCacheData* RhsCacheData = CastChecked<UNDISimpleCounterSimCacheData>(RhsStorageObject);
diff --git a/Engine/Plugins/FX/Niagara/Source/Niagara/Private/DataInterface/NiagaraDataInterfaceSimpleCounter.h b/Engine/Plugins/FX/Niagara/Source/Niagara/Private/DataInterface/NiagaraDataInterfaceSimpleCounter.h
--- a/Engine/Plugins/FX/Niagara/Source/Niagara/Private/DataInterface/NiagaraDataInterfaceSimpleCounter.h
+++ b/Engine/Plugins/FX/Niagara/Source/Niagara/Private/DataInterface/NiagaraDataInterfaceSimpleCounter.h
@@ -63,7 +63,7 @@ public:
 	virtual UObject* SimCacheBeginWrite(UObject* SimCache, FNiagaraSystemInstance* NiagaraSystemInstance, const void* OptionalPerInstanceData, FNiagaraSimCacheFeedbackContext& FeedbackContext) const override;
 	virtual bool SimCacheWriteFrame(UObject* StorageObject, int FrameIndex, FNiagaraSystemInstance* SystemInstance, const void* OptionalPerInstanceData, FNiagaraSimCacheFeedbackContext& FeedbackContext) const override;
 	virtual bool SimCacheReadFrame(UObject* StorageObject, int FrameA, int FrameB, float Interp, FNiagaraSystemInstance* SystemInstance, void* OptionalPerInstanceData) override;
-	virtual bool SimCacheCompareFrame(UObject* LhsStorageObject, UObject* RhsStorageObject, int FrameIndex, TOptional<float> Tolerance, FString& OutErrors) override;
+	virtual bool SimCacheCompareFrame(UObject* LhsStorageObject, UObject* RhsStorageObject, int FrameIndex, TOptional<float> Tolerance, FString& OutErrors) const override;
 	//~ UNiagaraDataInterface interface END
 
 	void UpdateDIProxy();
diff --git a/Engine/Plugins/FX/Niagara/Source/Niagara/Private/NiagaraDataInterfaceArray.cpp b/Engine/Plugins/FX/Niagara/Source/Niagara/Private/NiagaraDataInterfaceArray.cpp
--- a/Engine/Plugins/FX/Niagara/Source/Niagara/Private/NiagaraDataInterfaceArray.cpp
+++ b/Engine/Plugins/FX/Niagara/Source/Niagara/Private/NiagaraDataInterfaceArray.cpp
@@ -118,7 +118,51 @@ bool UNiagaraDataInterfaceArray::SimCacheReadFrame(UObject* StorageObject, int F
 	return ArrayProxy->SimCacheReadFrame(CacheData, FrameA, SystemInstance);
 }
 
-FString UNiagaraDataInterfaceArray::SimCacheVisualizerRead(UNDIArraySimCacheData* CacheData, FNDIArraySimCacheDataFrame& FrameData, int Element) const
+bool UNiagaraDataInterfaceArray::SimCacheCompareFrame(UObject* LhsStorageObject, UObject* RhsStorageObject, int FrameIndex, TOptional<float> InTolerance, FString& OutErrors) const
+{
+	UNDIArraySimCacheData* LhsCacheData = CastChecked<UNDIArraySimCacheData>(LhsStorageObject);
+	UNDIArraySimCacheData* RhsCacheData = CastChecked<UNDIArraySimCacheData>(RhsStorageObject);
+
+	if (!LhsCacheData->CpuFrameData.IsValidIndex(FrameIndex) || !RhsCacheData->CpuFrameData.IsValidIndex(FrameIndex) ||
+		!LhsCacheData->GpuFrameData.IsValidIndex(FrameIndex) || !RhsCacheData->GpuFrameData.IsValidIndex(FrameIndex) )
+	{
+		OutErrors = TEXT("FrameIndex was not valid");
+		return false;
+	}
+
+	const INDIArrayProxyBase* ArrayProxy = GetProxyAs<INDIArrayProxyBase>();
+	const float Tolerance = InTolerance.Get(UE_SMALL_NUMBER);
+
+	auto CompareFrames =
+		[&](const FNDIArraySimCacheDataFrame& LhsFrame, const FNDIArraySimCacheDataFrame& RhsFrame, const TCHAR* SimType)
+		{
+			if (LhsFrame.NumElements != RhsFrame.NumElements)
+			{
+				OutErrors = FString::Printf(TEXT("Element Count Mismatch (%d -> %d) for %s data"), LhsCacheData->CpuFrameData[FrameIndex].NumElements, RhsCacheData->CpuFrameData[FrameIndex].NumElements, SimType);
+				return false;
+			}
+
+			const uint8* LhsArrayData = LhsCacheData->BufferData.GetData() + LhsFrame.DataOffset;
+			const uint8* RhsArrayData = RhsCacheData->BufferData.GetData() + RhsFrame.DataOffset;
+			for (int32 i=0; i < LhsFrame.NumElements; ++i)
+			{
+				if (!ArrayProxy->SimCacheCompareElement(LhsArrayData, RhsArrayData, i, Tolerance))
+				{
+					const FString LhsValue = ArrayProxy->SimCacheVisualizerRead(LhsCacheData, LhsFrame, i);
+					const FString RhsValue = ArrayProxy->SimCacheVisualizerRead(RhsCacheData, RhsFrame, i);
+					OutErrors = FString::Printf(TEXT("Element %d Mismatch (%s -> %s) for %s data"), i, *LhsValue, *RhsValue, SimType);
+					return false;
+				}
+			}
+			return true;
+		};
+
+	return
+		CompareFrames(LhsCacheData->CpuFrameData[FrameIndex], RhsCacheData->CpuFrameData[FrameIndex], TEXT("CPU")) &&
+		CompareFrames(LhsCacheData->GpuFrameData[FrameIndex], RhsCacheData->GpuFrameData[FrameIndex], TEXT("GPU"));
+}
+
+FString UNiagaraDataInterfaceArray::SimCacheVisualizerRead(const UNDIArraySimCacheData* CacheData, const FNDIArraySimCacheDataFrame& FrameData, int Element) const
 {
 	const INDIArrayProxyBase* ArrayProxy = GetProxyAs<INDIArrayProxyBase>();
 	return ArrayProxy->SimCacheVisualizerRead(CacheData, FrameData, Element);
diff --git a/Engine/Plugins/FX/Niagara/Source/Niagara/Private/NiagaraDataInterfaceArrayImplHelpers.h b/Engine/Plugins/FX/Niagara/Source/Niagara/Private/NiagaraDataInterfaceArrayImplHelpers.h
--- a/Engine/Plugins/FX/Niagara/Source/Niagara/Private/NiagaraDataInterfaceArrayImplHelpers.h
+++ b/Engine/Plugins/FX/Niagara/Source/Niagara/Private/NiagaraDataInterfaceArrayImplHelpers.h
@@ -69,6 +69,13 @@ struct FNDIArrayImplHelper<FVector2f> : public FNDIArrayImplHelperBase<FVector2f
 	{
 		OutString.Appendf(TEXT("%f, %f"), Value.X, Value.Y);
 	}
+
+	static bool IsNearlyEqual(const FVector2f& Lhs, const FVector2f& Rhs, float Tolerance)
+	{
+		return
+			FMath::IsNearlyEqual(Lhs.X, Rhs.X, Tolerance) &&
+			FMath::IsNearlyEqual(Lhs.Y, Rhs.Y, Tolerance);
+	}
 };
 
 template<>
@@ -113,6 +120,14 @@ struct FNDIArrayImplHelper<FVector3f> : public FNDIArrayImplHelperBase<FVector3f
 	{
 		OutString.Appendf(TEXT("%f, %f, %f"), Value.X, Value.Y, Value.Z);
 	}
+
+	static bool IsNearlyEqual(const FVector3f& Lhs, const FVector3f& Rhs, float Tolerance)
+	{
+		return
+			FMath::IsNearlyEqual(Lhs.X, Rhs.X, Tolerance) &&
+			FMath::IsNearlyEqual(Lhs.Y, Rhs.Y, Tolerance) &&
+			FMath::IsNearlyEqual(Lhs.Z, Rhs.Z, Tolerance);
+	}
 };
 
 template<>
@@ -135,6 +150,14 @@ struct FNDIArrayImplHelper<FNiagaraPosition> : public FNDIArrayImplHelper<FVecto
 	{
 		OutString.Appendf(TEXT("%f, %f, %f"), Value.X, Value.Y, Value.Z);
 	}
+
+	static bool IsNearlyEqual(const FNiagaraPosition& Lhs, const FNiagaraPosition& Rhs, float Tolerance)
+	{
+		return
+			FMath::IsNearlyEqual(Lhs.X, Rhs.X, Tolerance) &&
+			FMath::IsNearlyEqual(Lhs.Y, Rhs.Y, Tolerance) &&
+			FMath::IsNearlyEqual(Lhs.Z, Rhs.Z, Tolerance);
+	}
 };
 
 template<>
@@ -179,6 +202,15 @@ struct FNDIArrayImplHelper<FVector4f> : public FNDIArrayImplHelperBase<FVector4f
 	{
 		OutString.Appendf(TEXT("%f, %f, %f, %f"), Value.X, Value.Y, Value.Z, Value.W);
 	}
+
+	static bool IsNearlyEqual(const FVector4f& Lhs, const FVector4f& Rhs, float Tolerance)
+	{
+		return
+			FMath::IsNearlyEqual(Lhs.X, Rhs.X, Tolerance) &&
+			FMath::IsNearlyEqual(Lhs.Y, Rhs.Y, Tolerance) &&
+			FMath::IsNearlyEqual(Lhs.Z, Rhs.Z, Tolerance) &&
+			FMath::IsNearlyEqual(Lhs.W, Rhs.W, Tolerance);
+	}
 };
 
 template<>
@@ -202,6 +234,15 @@ struct FNDIArrayImplHelper<FLinearColor> : public FNDIArrayImplHelperBase<FLinea
 	{
 		OutString.Appendf(TEXT("%f, %f, %f, %f"), Value.R, Value.G, Value.B, Value.A);
 	}
+
+	static bool IsNearlyEqual(const FLinearColor& Lhs, const FLinearColor& Rhs, float Tolerance)
+	{
+		return
+			FMath::IsNearlyEqual(Lhs.R, Rhs.R, Tolerance) &&
+			FMath::IsNearlyEqual(Lhs.G, Rhs.G, Tolerance) &&
+			FMath::IsNearlyEqual(Lhs.B, Rhs.B, Tolerance) &&
+			FMath::IsNearlyEqual(Lhs.A, Rhs.A, Tolerance);
+	}
 };
 
 template<>
@@ -246,6 +287,15 @@ struct FNDIArrayImplHelper<FQuat4f> : public FNDIArrayImplHelperBase<FQuat4f>
 	{
 		OutString.Appendf(TEXT("%f, %f, %f, %f"), Value.X, Value.Y, Value.Z, Value.W);
 	}
+
+	static bool IsNearlyEqual(const FQuat4f& Lhs, const FQuat4f& Rhs, float Tolerance)
+	{
+		return
+			FMath::IsNearlyEqual(Lhs.X, Rhs.X, Tolerance) &&
+			FMath::IsNearlyEqual(Lhs.Y, Rhs.Y, Tolerance) &&
+			FMath::IsNearlyEqual(Lhs.Z, Rhs.Z, Tolerance) &&
+			FMath::IsNearlyEqual(Lhs.W, Rhs.W, Tolerance);
+	}
 };
 
 template<>
@@ -312,6 +362,22 @@ struct FNDIArrayImplHelper<FMatrix44f> : public FNDIArrayImplHelperBase<FMatrix4
 			}
 		}
 	}
+
+	static bool IsNearlyEqual(const FMatrix44f& Lhs, const FMatrix44f& Rhs, float Tolerance)
+	{
+		for (int32 r = 0; r < 4; ++r)
+		{
+			for (int32 c = 0; c < 4; ++c)
+			{
+				if (!FMath::IsNearlyEqual(Lhs.M[r][c], Rhs.M[r][c], Tolerance))
+				{
+					return false;
+				}
+			}
+		}
+
+		return true;
+	}
 };
 
 template<>
@@ -469,4 +535,9 @@ struct FNDIArrayImplHelper<FNiagaraID> : public FNDIArrayImplHelperBase<FNiagara
 	{
 		OutString.Appendf(TEXT("%d, %d"), Value.Index, Value.AcquireTag);
 	}
+
+	static bool IsNearlyEqual(const FNiagaraID& Lhs, const FNiagaraID& Rhs, float Tolerance)
+	{
+		return Lhs.AcquireTag == Rhs.AcquireTag && Lhs.Index == Rhs.Index;
+	}
 };
