- Fix transitions for array & simple counter
- Cleanup some old code to remove casting

diff --git a/Engine/Plugins/FX/Niagara/Source/Niagara/Classes/NiagaraDataInterfaceArrayImpl.h b/Engine/Plugins/FX/Niagara/Source/Niagara/Classes/NiagaraDataInterfaceArrayImpl.h
--- a/Engine/Plugins/FX/Niagara/Source/Niagara/Classes/NiagaraDataInterfaceArrayImpl.h
+++ b/Engine/Plugins/FX/Niagara/Source/Niagara/Classes/NiagaraDataInterfaceArrayImpl.h
@@ -263,7 +263,7 @@ struct FNDIArrayInstanceData_RenderThread
 
 			// Create Buffer
 			FRHIResourceCreateInfo CreateInfo(TEXT("NiagaraDataInterfaceArray"));
-			const EBufferUsageFlags BufferUsage = BUF_Static | BUF_ShaderResource | BUF_VertexBuffer | (IsReadOnly() ? BUF_None : BUF_UnorderedAccess | BUF_SourceCopy);
+			const EBufferUsageFlags BufferUsage = BUF_Static | BUF_ShaderResource | BUF_VertexBuffer | BUF_SourceCopy | (IsReadOnly() ? BUF_None : BUF_UnorderedAccess);
 			const ERHIAccess DefaultAccess = IsReadOnly() ? ERHIAccess::SRVCompute : ERHIAccess::UAVCompute;
 			ArrayBuffer = RHICmdList.CreateBuffer(ArrayNumBytes, BufferUsage, TypeStride, DefaultAccess, CreateInfo);
 
@@ -1039,9 +1039,11 @@ struct FNDIArrayProxyImpl : public INDIArrayProxyBase
 					TArray<FRHITransitionInfo, TInlineAllocator<2>> TransitionsBefore;
 					TArray<FRHITransitionInfo, TInlineAllocator<2>> TransitionsAfter;
 
+					const ERHIAccess DefaultAccess = InstanceData_RT->IsReadOnly() ? ERHIAccess::SRVCompute : ERHIAccess::UAVCompute;
+
 					BufferRequests.Emplace(InstanceData_RT->ArrayBuffer, 0, InstanceData_RT->ArrayNumBytes);
-					TransitionsBefore.Emplace(InstanceData_RT->ArrayBuffer, ERHIAccess::UAVCompute, ERHIAccess::CopySrc);
-					TransitionsAfter.Emplace(InstanceData_RT->ArrayBuffer, ERHIAccess::CopySrc, ERHIAccess::UAVCompute);
+					TransitionsBefore.Emplace(InstanceData_RT->ArrayBuffer, DefaultAccess, ERHIAccess::CopySrc);
+					TransitionsAfter.Emplace(InstanceData_RT->ArrayBuffer, ERHIAccess::CopySrc, DefaultAccess);
 
 					if (InstanceData_RT->IsReadOnly())
 					{
@@ -1051,8 +1053,8 @@ struct FNDIArrayProxyImpl : public INDIArrayProxyBase
 					{
 						const FNiagaraGPUInstanceCountManager& CountManager = ComputeInterface->GetGPUInstanceCounterManager();
 						BufferRequests.Emplace(CountManager.GetInstanceCountBuffer().Buffer, uint32(InstanceData_RT->CountOffset * sizeof(uint32)), sizeof(uint32));
-						TransitionsBefore.Emplace(CountManager.GetInstanceCountBuffer().UAV, ERHIAccess::UAVCompute, ERHIAccess::CopySrc);
-						TransitionsAfter.Emplace(CountManager.GetInstanceCountBuffer().UAV, ERHIAccess::CopySrc, ERHIAccess::UAVCompute);
+						TransitionsBefore.Emplace(CountManager.GetInstanceCountBuffer().UAV, FNiagaraGPUInstanceCountManager::kCountBufferDefaultState, ERHIAccess::CopySrc);
+						TransitionsAfter.Emplace(CountManager.GetInstanceCountBuffer().UAV, ERHIAccess::CopySrc, FNiagaraGPUInstanceCountManager::kCountBufferDefaultState);
 					}
 
 					FNiagaraGpuReadbackManager* ReadbackManager = ComputeInterface->GetGpuReadbackManager();
diff --git a/Engine/Plugins/FX/Niagara/Source/Niagara/Private/DataInterface/NiagaraDataInterfaceSimpleCounter.cpp b/Engine/Plugins/FX/Niagara/Source/Niagara/Private/DataInterface/NiagaraDataInterfaceSimpleCounter.cpp
--- a/Engine/Plugins/FX/Niagara/Source/Niagara/Private/DataInterface/NiagaraDataInterfaceSimpleCounter.cpp
+++ b/Engine/Plugins/FX/Niagara/Source/Niagara/Private/DataInterface/NiagaraDataInterfaceSimpleCounter.cpp
@@ -4,6 +4,7 @@
 #include "NiagaraClearCounts.h"
 #include "NiagaraCompileHashVisitor.h"
 #include "NiagaraGpuComputeDispatchInterface.h"
+#include "NiagaraGpuComputeDispatch.h"
 #include "NiagaraGpuReadbackManager.h"
 #include "NiagaraShaderParametersBuilder.h"
 #include "NiagaraSystemInstance.h"
@@ -73,16 +74,20 @@ struct FNDISimpleCounterProxy : public FNiagaraDataInterfaceProxyRW
 				const FNiagaraGPUInstanceCountManager& CounterManager = Context.GetComputeDispatchInterface().GetGPUInstanceCounterManager();
 				const FRWBuffer& CountBuffer = CounterManager.GetInstanceCountBuffer();
 
+				//-TODO: Once the count manager is ported to RDG we won't need to track this
+				const FNiagaraGpuComputeDispatch& ComputeDispatch = static_cast<const FNiagaraGpuComputeDispatch&>(Context.GetComputeDispatchInterface());
+				const ERHIAccess CountRHIAccess = ComputeDispatch.IsExecutingFirstDispatchGroup() ? FNiagaraGPUInstanceCountManager::kCountBufferDefaultState : ERHIAccess::UAVCompute;
+
 				//-TODO:RDG: Once the count buffer is a graph resource this can be changed
 				AddPass(
 					Context.GetGraphBuilder(),
 					RDG_EVENT_NAME("NiagaraSimpleCounter::PreStage"),
-					[CountBufferUAV=CountBuffer.UAV, CountOffset=InstanceData->CountOffset, CountValue=InstanceData->CountValue.GetValue()](FRHICommandListImmediate& RHICmdList)
+					[CountBufferUAV=CountBuffer.UAV, CountOffset=InstanceData->CountOffset, CountValue=InstanceData->CountValue.GetValue(), CountRHIAccess](FRHICommandListImmediate& RHICmdList)
 					{
 						const TPair<uint32, uint32> DataToClear(CountOffset, reinterpret_cast<const uint32&>(CountValue));
-						RHICmdList.Transition(FRHITransitionInfo(CountBufferUAV, ERHIAccess::UAVCompute, ERHIAccess::UAVCompute));
+						RHICmdList.Transition(FRHITransitionInfo(CountBufferUAV, CountRHIAccess, ERHIAccess::UAVCompute));
 						NiagaraClearCounts::ClearCountsUInt(RHICmdList, CountBufferUAV, MakeArrayView(&DataToClear, 1) );
-						RHICmdList.Transition(FRHITransitionInfo(CountBufferUAV, ERHIAccess::UAVCompute, ERHIAccess::UAVCompute));
+						RHICmdList.Transition(FRHITransitionInfo(CountBufferUAV, ERHIAccess::UAVCompute, CountRHIAccess));
 					}
 				);
 
diff --git a/Engine/Plugins/FX/Niagara/Source/Niagara/Private/NiagaraComputeExecutionContext.cpp b/Engine/Plugins/FX/Niagara/Source/Niagara/Private/NiagaraComputeExecutionContext.cpp
--- a/Engine/Plugins/FX/Niagara/Source/Niagara/Private/NiagaraComputeExecutionContext.cpp
+++ b/Engine/Plugins/FX/Niagara/Source/Niagara/Private/NiagaraComputeExecutionContext.cpp
@@ -2,7 +2,6 @@
 
 #include "NiagaraComputeExecutionContext.h"
 #include "NiagaraGpuComputeDispatchInterface.h"
-#include "NiagaraGpuComputeDispatch.h"
 #include "NiagaraDataInterface.h"
 #include "NiagaraScript.h"
 #include "NiagaraSystemInstance.h"
@@ -224,7 +223,7 @@ void FNiagaraComputeExecutionContext::ResetInternal(FNiagaraGpuComputeDispatchIn
 
 	if (ComputeDispatchInterface)
 	{
-		static_cast<FNiagaraGpuComputeDispatch*>(ComputeDispatchInterface)->GetGPUInstanceCounterManager().FreeEntry(CountOffset_RT);
+		ComputeDispatchInterface->GetGPUInstanceCounterManager().FreeEntry(CountOffset_RT);
 	}
 
 	CurrentNumInstances_RT = 0;
diff --git a/Engine/Plugins/FX/Niagara/Source/Niagara/Private/NiagaraDataSet.cpp b/Engine/Plugins/FX/Niagara/Source/Niagara/Private/NiagaraDataSet.cpp
--- a/Engine/Plugins/FX/Niagara/Source/Niagara/Private/NiagaraDataSet.cpp
+++ b/Engine/Plugins/FX/Niagara/Source/Niagara/Private/NiagaraDataSet.cpp
@@ -7,7 +7,6 @@
 #include "NiagaraRenderer.h"
 #include "NiagaraGPUInstanceCountManager.h"
 #include "NiagaraGpuComputeDispatchInterface.h"
-#include "NiagaraGpuComputeDispatch.h"
 #include "NiagaraShaderParticleID.h"
 #include "NiagaraShared.h"
 
@@ -1473,7 +1472,7 @@ void FScopedNiagaraDataSetGPUReadback::ReadbackData(FNiagaraGpuComputeDispatchIn
 			const uint32 BufferOffset = DataBuffer->GetGPUInstanceCountBufferOffset();
 			if (ComputeDispatchInterface && BufferOffset != INDEX_NONE)
 			{
-				FRHIBuffer* InstanceCountBuffer = static_cast<FNiagaraGpuComputeDispatch*>(ComputeDispatchInterface)->GetGPUInstanceCounterManager().GetInstanceCountBuffer().Buffer;
+				FRHIBuffer* InstanceCountBuffer = ComputeDispatchInterface->GetGPUInstanceCounterManager().GetInstanceCountBuffer().Buffer;
 
 				void* Data = RHICmdList.LockBuffer(InstanceCountBuffer, 0, (BufferOffset + 1) * sizeof(int32), RLM_ReadOnly);
 				NumInstances = reinterpret_cast<int32*>(Data)[BufferOffset];
diff --git a/Engine/Plugins/FX/Niagara/Source/Niagara/Private/NiagaraGpuComputeDispatch.cpp b/Engine/Plugins/FX/Niagara/Source/Niagara/Private/NiagaraGpuComputeDispatch.cpp
--- a/Engine/Plugins/FX/Niagara/Source/Niagara/Private/NiagaraGpuComputeDispatch.cpp
+++ b/Engine/Plugins/FX/Niagara/Source/Niagara/Private/NiagaraGpuComputeDispatch.cpp
@@ -1099,6 +1099,8 @@ void FNiagaraGpuComputeDispatch::ExecuteTicks(FRDGBuilder& GraphBuilder, TConstS
 		const bool bIsFirstGroup = &DispatchGroup == &DispatchList.DispatchGroups[0];
 		const bool bIsLastGroup = &DispatchGroup == &DispatchList.DispatchGroups.Last();
 
+		bIsExecutingFirstDispatchGroup = bIsFirstGroup;
+
 		// Consume per tick data from the game thread
 		//-TODO: This does not work currently as some senders assume the data will not be deferred processed
 		//for (FNiagaraGPUSystemTick* Tick : DispatchGroup.TicksWithPerInstanceData)
@@ -1386,6 +1388,7 @@ void FNiagaraGpuComputeDispatch::ExecuteTicks(FRDGBuilder& GraphBuilder, TConstS
 	SceneTexturesUniformParams = nullptr;
 	MobileSceneTexturesUniformParams = nullptr;
 	SubstratePublicGlobalUniformParams = nullptr;
+	bIsExecutingFirstDispatchGroup = false;
 
 	CurrentPassExternalAccessQueue.Submit(GraphBuilder);
 }
diff --git a/Engine/Plugins/FX/Niagara/Source/Niagara/Private/NiagaraGpuComputeDispatch.h b/Engine/Plugins/FX/Niagara/Source/Niagara/Private/NiagaraGpuComputeDispatch.h
--- a/Engine/Plugins/FX/Niagara/Source/Niagara/Private/NiagaraGpuComputeDispatch.h
+++ b/Engine/Plugins/FX/Niagara/Source/Niagara/Private/NiagaraGpuComputeDispatch.h
@@ -176,6 +176,10 @@ public:
 
 	virtual FNiagaraAsyncGpuTraceHelper& GetAsyncGpuTraceHelper() const override;
 
+	//-TODO: Temporary while the count buffer is not an RDG resource
+	bool IsExecutingFirstDispatchGroup() const { return bIsExecutingFirstDispatchGroup; }
+	bool bIsExecutingFirstDispatchGroup = false;
+
 private:
 	void DumpDebugFrame();
 	void UpdateInstanceCountManager(FRHICommandListImmediate& RHICmdList);
