Only register niagara data channel assets from plugins that are actually mounted at the time

#tests tested locally with LCW

diff --git a/Engine/Plugins/FX/Niagara/Source/NiagaraEditor/Private/NiagaraEditorModule.cpp b/Engine/Plugins/FX/Niagara/Source/NiagaraEditor/Private/NiagaraEditorModule.cpp
--- a/Engine/Plugins/FX/Niagara/Source/NiagaraEditor/Private/NiagaraEditorModule.cpp
+++ b/Engine/Plugins/FX/Niagara/Source/NiagaraEditor/Private/NiagaraEditorModule.cpp
@@ -2217,7 +2217,10 @@ void FNiagaraEditorModule::OnAssetRegistryLoadComplete()
 	AssetRegistry.GetAssetsByClass(UNiagaraDataChannelAsset::StaticClass()->GetClassPathName(), AllDataChannels);
 	for (FAssetData& DataChannelAsset : AllDataChannels)
 	{
-		UNiagaraDataChannelAsset* NewAsset = Cast<UNiagaraDataChannelAsset>(DataChannelAsset.GetAsset());
+		if (FPackageName::GetPackageMountPoint(DataChannelAsset.PackageName.ToString()) != NAME_None)
+		{
+			UNiagaraDataChannelAsset* NewAsset = Cast<UNiagaraDataChannelAsset>(DataChannelAsset.GetAsset());
+		}
 	}
 }
 
