Added initial analytics support for Niagara ui actions

diff --git a/Engine/Plugins/FX/Niagara/Source/Niagara/Internal/NiagaraAnalytics.cpp b/Engine/Plugins/FX/Niagara/Source/Niagara/Internal/NiagaraAnalytics.cpp
new file mode 100644
--- /dev/null
+++ b/Engine/Plugins/FX/Niagara/Source/Niagara/Internal/NiagaraAnalytics.cpp
@@ -0,0 +1,32 @@
+﻿// Copyright Epic Games, Inc. All Rights Reserved.
+
+#include "NiagaraAnalytics.h"
+#include "EngineAnalytics.h"
+#include "UObject/Object.h"
+#include "UObject/Package.h"
+
+void NiagaraAnalytics::RecordEvent(FString&& EventName, const TArray<FAnalyticsEventAttribute>& Attributes)
+{
+	if (FEngineAnalytics::IsAvailable())
+	{
+		FEngineAnalytics::GetProvider().RecordEvent(TEXT("Editor.Usage.Niagara.") + MoveTemp(EventName), Attributes);
+	}
+}
+
+void NiagaraAnalytics::RecordEvent(FString&& EventName, const FString& AttributeName, const FString& AttributeValue)
+{
+	TArray<FAnalyticsEventAttribute> Attributes;
+	Attributes.Add(FAnalyticsEventAttribute(AttributeName, AttributeValue));
+	RecordEvent(MoveTemp(EventName), Attributes);
+}
+
+bool NiagaraAnalytics::IsPluginAsset(const UObject* Obj)
+{
+	if (IsValid(Obj) && Obj->IsAsset())
+	{
+		UPackage* Package = Obj->GetPackage();
+		FString Name = Package->GetName();
+		return Package && (Name.StartsWith(TEXT("/Niagara/"), ESearchCase::CaseSensitive) || Name.StartsWith(TEXT("/NiagaraFluids/"), ESearchCase::CaseSensitive));
+	}
+	return false;
+}
diff --git a/Engine/Plugins/FX/Niagara/Source/Niagara/Internal/NiagaraAnalytics.h b/Engine/Plugins/FX/Niagara/Source/Niagara/Internal/NiagaraAnalytics.h
new file mode 100644
--- /dev/null
+++ b/Engine/Plugins/FX/Niagara/Source/Niagara/Internal/NiagaraAnalytics.h
@@ -0,0 +1,16 @@
+﻿// Copyright Epic Games, Inc. All Rights Reserved.
+
+#pragma once
+
+#include "CoreMinimal.h"
+#include "AnalyticsEventAttribute.h"
+
+namespace NiagaraAnalytics
+{
+	NIAGARA_API void RecordEvent(FString&& EventName, const TArray<FAnalyticsEventAttribute>& Attributes = TArray<FAnalyticsEventAttribute>());
+	
+	NIAGARA_API void RecordEvent(FString&& EventName, const FString& AttributeName, const FString& AttributeValue);
+
+	// Checks if the provided object is part of a Niagara plugin and can safely be reported on (e.g. the asset name)
+	NIAGARA_API bool IsPluginAsset(const UObject* Obj);
+};
\ No newline at end of file
diff --git a/Engine/Plugins/FX/Niagara/Source/Niagara/Private/NiagaraEmitter.cpp b/Engine/Plugins/FX/Niagara/Source/Niagara/Private/NiagaraEmitter.cpp
--- a/Engine/Plugins/FX/Niagara/Source/Niagara/Private/NiagaraEmitter.cpp
+++ b/Engine/Plugins/FX/Niagara/Source/Niagara/Private/NiagaraEmitter.cpp
@@ -4,6 +4,7 @@
 #include "NiagaraEmitter.h"
 
 #include "INiagaraEditorOnlyDataUtlities.h"
+#include "NiagaraAnalytics.h"
 #include "NiagaraBoundsCalculator.h"
 #include "NiagaraCustomVersion.h"
 #include "NiagaraComponentSettings.h"
@@ -2763,6 +2764,8 @@ void UNiagaraEmitter::EnableVersioning()
 	ensure(VersionData.Num() == 1);
 	bVersioningEnabled = true;	
 	ExposedVersion = VersionData[0].Version.VersionGuid;
+
+	NiagaraAnalytics::RecordEvent("Versioning.EmitterEnabled");
 }
 
 void UNiagaraEmitter::DisableVersioning(const FGuid& VersionGuidToUse)
diff --git a/Engine/Plugins/FX/Niagara/Source/Niagara/Private/NiagaraScript.cpp b/Engine/Plugins/FX/Niagara/Source/Niagara/Private/NiagaraScript.cpp
--- a/Engine/Plugins/FX/Niagara/Source/Niagara/Private/NiagaraScript.cpp
+++ b/Engine/Plugins/FX/Niagara/Source/Niagara/Private/NiagaraScript.cpp
@@ -14,6 +14,7 @@
 #include "Misc/FileHelper.h"
 #include "Misc/SecureHash.h"
 #include "Modules/ModuleManager.h"
+#include "NiagaraAnalytics.h"
 #include "NiagaraCompileHashVisitor.h"
 #include "NiagaraCustomVersion.h"
 #include "NiagaraDataInterfaceSkeletalMesh.h"
@@ -623,6 +624,21 @@ void UNiagaraScript::ExposeVersion(const FGuid& VersionGuid)
 	}
 }
 
+FString NiagaraScriptUsageToString(const ENiagaraScriptUsage& Usage)
+{
+	switch (Usage)
+	{
+	case ENiagaraScriptUsage::Module:
+		return TEXT("Module");
+	case ENiagaraScriptUsage::Function:
+		return TEXT("Function");
+	case ENiagaraScriptUsage::DynamicInput:
+		return TEXT("DynamicInput");
+	default:
+		return TEXT("Other");
+	}
+}
+
 void UNiagaraScript::EnableVersioning()
 {
 	if (bVersioningEnabled)
@@ -633,6 +649,8 @@ void UNiagaraScript::EnableVersioning()
 	ensure(VersionData.Num() == 1);
 	bVersioningEnabled = true;	
 	ExposedVersion = VersionData[0].Version.VersionGuid;
+
+	NiagaraAnalytics::RecordEvent(TEXT("Versioning.ScriptEnabled"), TEXT("Usage"), NiagaraScriptUsageToString(Usage));
 }
 
 void UNiagaraScript::DisableVersioning(const FGuid& VersionGuidToUse)
diff --git a/Engine/Plugins/FX/Niagara/Source/NiagaraEditor/Private/Customizations/NiagaraFunctionCallNodeDetails.cpp b/Engine/Plugins/FX/Niagara/Source/NiagaraEditor/Private/Customizations/NiagaraFunctionCallNodeDetails.cpp
--- a/Engine/Plugins/FX/Niagara/Source/NiagaraEditor/Private/Customizations/NiagaraFunctionCallNodeDetails.cpp
+++ b/Engine/Plugins/FX/Niagara/Source/NiagaraEditor/Private/Customizations/NiagaraFunctionCallNodeDetails.cpp
@@ -10,6 +10,7 @@
 #include "ScopedTransaction.h"
 #include "Framework/MultiBox/MultiBoxBuilder.h"
 #include "EdGraphSchema_Niagara.h"
+#include "NiagaraAnalytics.h"
 #include "NiagaraNodeCustomHlsl.h"
 #include "Layout/Margin.h"
 #include "Widgets/DeclarativeSyntaxSupport.h"
@@ -245,6 +246,15 @@ void FNiagaraFunctionCallNodeDetails::SwitchToVersion(FNiagaraAssetVersion Versi
 	UpgradeContext.bSkipPythonScript = true;
 	Node->ChangeScriptVersion(Version.VersionGuid, UpgradeContext);
 	Node->RefreshFromExternalChanges();
+
+	// analytics
+	TArray<FAnalyticsEventAttribute> Attributes;
+	Attributes.Add(FAnalyticsEventAttribute(TEXT("Type"), TEXT("FunctionNode")));
+	if (NiagaraAnalytics::IsPluginAsset(Node->FunctionScript))
+	{
+		Attributes.Add(FAnalyticsEventAttribute(TEXT("AssetName"), Node->FunctionScript->GetPackage()->GetName()));
+	}
+	NiagaraAnalytics::RecordEvent(TEXT("Versioning.ScriptVersionChanged"), Attributes);
 }
 
 bool FNiagaraFunctionCallNodeDetails::IsVersionSelected(FNiagaraAssetVersion Version) const
diff --git a/Engine/Plugins/FX/Niagara/Source/NiagaraEditor/Private/NiagaraEditorUtilities.cpp b/Engine/Plugins/FX/Niagara/Source/NiagaraEditor/Private/NiagaraEditorUtilities.cpp
--- a/Engine/Plugins/FX/Niagara/Source/NiagaraEditor/Private/NiagaraEditorUtilities.cpp
+++ b/Engine/Plugins/FX/Niagara/Source/NiagaraEditor/Private/NiagaraEditorUtilities.cpp
@@ -14,6 +14,7 @@
 #include "NiagaraCustomVersion.h"
 #include "NiagaraDataInterface.h"
 #include "IPythonScriptPlugin.h"
+#include "NiagaraAnalytics.h"
 #include "NiagaraClipboard.h"
 #include "NiagaraEditorSettings.h"
 #include "NiagaraEditorStyle.h"
@@ -3610,6 +3611,14 @@ void FNiagaraEditorUtilities::SwitchParentEmitterVersion(TSharedRef<FNiagaraEmit
 		OldSystemViewModel.Reset();
 		CollectGarbage(GARBAGE_COLLECTION_KEEPFLAGS);
 	}
+
+	// analytics
+	TArray<FAnalyticsEventAttribute> Attributes;
+	if (NiagaraAnalytics::IsPluginAsset(Parent.Emitter))
+	{
+		Attributes.Add(FAnalyticsEventAttribute(TEXT("AssetName"), Parent.Emitter->GetPackage()->GetName()));
+	}
+	NiagaraAnalytics::RecordEvent(TEXT("Versioning.EmitterVersionChanged"), Attributes);
 }
 
 bool FNiagaraParameterUtilities::DoesParameterNameMatchSearchText(FName ParameterName, const FString& SearchTextString)
diff --git a/Engine/Plugins/FX/Niagara/Source/NiagaraEditor/Private/ViewModels/Stack/NiagaraStackFunctionInput.cpp b/Engine/Plugins/FX/Niagara/Source/NiagaraEditor/Private/ViewModels/Stack/NiagaraStackFunctionInput.cpp
--- a/Engine/Plugins/FX/Niagara/Source/NiagaraEditor/Private/ViewModels/Stack/NiagaraStackFunctionInput.cpp
+++ b/Engine/Plugins/FX/Niagara/Source/NiagaraEditor/Private/ViewModels/Stack/NiagaraStackFunctionInput.cpp
@@ -7,6 +7,7 @@
 #include "EdGraphSchema_Niagara.h"
 #include "Editor.h"
 #include "INiagaraEditorTypeUtilities.h"
+#include "NiagaraAnalytics.h"
 #include "NiagaraClipboard.h"
 #include "NiagaraConstants.h"
 #include "NiagaraConvertInPlaceUtilityBase.h"
@@ -469,6 +470,18 @@ bool UNiagaraStackFunctionInput::FilterInlineChildren(const UNiagaraStackEntry&
 	return GbEnableExperimentalInlineDynamicInputs == 0 || GetInlineDisplayMode() == ENiagaraStackEntryInlineDisplayMode::None;
 }
 
+void UNiagaraStackFunctionInput::ReportScriptVersionChange() const
+{
+	// analytics
+	TArray<FAnalyticsEventAttribute> Attributes;
+	Attributes.Add(FAnalyticsEventAttribute(TEXT("Type"), TEXT("DynamicInput")));
+	if (NiagaraAnalytics::IsPluginAsset(GetDynamicInputNode()->FunctionScript))
+	{
+		Attributes.Add(FAnalyticsEventAttribute(TEXT("AssetName"), GetDynamicInputNode()->FunctionScript->GetPackage()->GetName()));
+	}
+	NiagaraAnalytics::RecordEvent(TEXT("Versioning.ScriptVersionChanged"), Attributes);
+}
+
 FText UNiagaraStackFunctionInput::GetCollapsedStateText() const
 {
 	if (IsFinalized())
@@ -657,6 +670,8 @@ UNiagaraStackEntry::FStackIssueFixDelegate UNiagaraStackFunctionInput::GetUpgrad
 			FunctionCallNode->GetNiagaraGraph()->NotifyGraphNeedsRecompile();
 			GetSystemViewModel()->ResetSystem();
 		}
+
+		ReportScriptVersionChange();
 	});
 }
 
@@ -2794,6 +2809,8 @@ void UNiagaraStackFunctionInput::ChangeScriptVersion(FGuid NewScriptVersion)
 	if (CachedSysViewModel->GetSystemStackViewModel())
 		CachedSysViewModel->GetSystemStackViewModel()->InvalidateCachedParameterUsage();
 	RefreshChildren();
+
+	ReportScriptVersionChange();
 }
 
 const UNiagaraClipboardFunctionInput* UNiagaraStackFunctionInput::ToClipboardFunctionInput(UObject* InOuter) const
diff --git a/Engine/Plugins/FX/Niagara/Source/NiagaraEditor/Private/ViewModels/Stack/NiagaraStackModuleItem.cpp b/Engine/Plugins/FX/Niagara/Source/NiagaraEditor/Private/ViewModels/Stack/NiagaraStackModuleItem.cpp
--- a/Engine/Plugins/FX/Niagara/Source/NiagaraEditor/Private/ViewModels/Stack/NiagaraStackModuleItem.cpp
+++ b/Engine/Plugins/FX/Niagara/Source/NiagaraEditor/Private/ViewModels/Stack/NiagaraStackModuleItem.cpp
@@ -3,6 +3,7 @@
 #include "ViewModels/Stack/NiagaraStackModuleItem.h"
 
 #include "NiagaraActions.h"
+#include "NiagaraAnalytics.h"
 #include "NiagaraClipboard.h"
 #include "NiagaraCommon.h"
 #include "NiagaraConstants.h"
@@ -875,9 +876,23 @@ UNiagaraStackEntry::FStackIssueFixDelegate UNiagaraStackModuleItem::GetUpgradeVe
             FunctionCallNode->GetNiagaraGraph()->NotifyGraphNeedsRecompile();
             GetSystemViewModel()->ResetSystem();
         }
+
+		ReportScriptVersionChange();
     });
 }
 
+void UNiagaraStackModuleItem::ReportScriptVersionChange() const
+{
+	// analytics
+	TArray<FAnalyticsEventAttribute> Attributes;
+	Attributes.Add(FAnalyticsEventAttribute(TEXT("Type"), TEXT("Module")));
+	if (NiagaraAnalytics::IsPluginAsset(GetModuleNode().FunctionScript))
+	{
+		Attributes.Add(FAnalyticsEventAttribute(TEXT("AssetName"), GetModuleNode().FunctionScript->GetPackage()->GetName()));
+	}
+	NiagaraAnalytics::RecordEvent(TEXT("Versioning.ScriptVersionChanged"), Attributes);
+}
+
 void UNiagaraStackModuleItem::RefreshIssues(TArray<FStackIssue>& NewIssues)
 {
 	if (GetSystemViewModel()->GetIsForDataProcessingOnly())
@@ -1507,6 +1522,8 @@ void UNiagaraStackModuleItem::ChangeScriptVersion(FGuid NewScriptVersion)
 	    FCompileConstantResolver(&GetSystemViewModel()->GetSystem(), FNiagaraStackGraphUtilities::GetOutputNodeUsage(*FunctionCallNode));
 	GetModuleNode().ChangeScriptVersion(NewScriptVersion, UpgradeContext, true);
 	Refresh();
+
+	ReportScriptVersionChange();
 }
 
 void UNiagaraStackModuleItem::SetInputValuesFromClipboardFunctionInputs(const TArray<const UNiagaraClipboardFunctionInput*>& ClipboardFunctionInputs)
diff --git a/Engine/Plugins/FX/Niagara/Source/NiagaraEditor/Public/ViewModels/Stack/NiagaraStackFunctionInput.h b/Engine/Plugins/FX/Niagara/Source/NiagaraEditor/Public/ViewModels/Stack/NiagaraStackFunctionInput.h
--- a/Engine/Plugins/FX/Niagara/Source/NiagaraEditor/Public/ViewModels/Stack/NiagaraStackFunctionInput.h
+++ b/Engine/Plugins/FX/Niagara/Source/NiagaraEditor/Public/ViewModels/Stack/NiagaraStackFunctionInput.h
@@ -408,6 +408,8 @@ private:
 
 	NIAGARAEDITOR_API bool FilterInlineChildren(const UNiagaraStackEntry& Child) const;
 
+	void ReportScriptVersionChange() const;
+
 private:
 	/** The module function call which owns this input entry. NOTE: This input might not be an input to the module function
 		call, it may be an input to a dynamic input function call which is owned by the module. */
diff --git a/Engine/Plugins/FX/Niagara/Source/NiagaraEditor/Public/ViewModels/Stack/NiagaraStackModuleItem.h b/Engine/Plugins/FX/Niagara/Source/NiagaraEditor/Public/ViewModels/Stack/NiagaraStackModuleItem.h
--- a/Engine/Plugins/FX/Niagara/Source/NiagaraEditor/Public/ViewModels/Stack/NiagaraStackModuleItem.h
+++ b/Engine/Plugins/FX/Niagara/Source/NiagaraEditor/Public/ViewModels/Stack/NiagaraStackModuleItem.h
@@ -140,6 +140,7 @@ private:
 	NIAGARAEDITOR_API void RefreshIsEnabled();
 	NIAGARAEDITOR_API void OnMessageManagerRefresh(const TArray<TSharedRef<const INiagaraMessage>>& NewMessages);
 	NIAGARAEDITOR_API FStackIssueFixDelegate GetUpgradeVersionFix();
+	void ReportScriptVersionChange() const;
 
 private:
 	UNiagaraNodeOutput* OutputNode;
