- Fix for circle occlusion and screen percent changes

diff --git a/Engine/Plugins/FX/Niagara/Shaders/Private/NiagaraDataInterfaceOcclusion.ush b/Engine/Plugins/FX/Niagara/Shaders/Private/NiagaraDataInterfaceOcclusion.ush
--- a/Engine/Plugins/FX/Niagara/Shaders/Private/NiagaraDataInterfaceOcclusion.ush
+++ b/Engine/Plugins/FX/Niagara/Shaders/Private/NiagaraDataInterfaceOcclusion.ush
@@ -138,7 +138,7 @@ void QueryOcclusionFactorWithCircle_{ParameterName}(float3 WorldPosition, float
 	const float3 PreViewPosition = LWCToFloat(LWCAdd(LWCWorldPosition, PrimaryView.PreViewTranslation));
 	const float4 ClipPosition = mul(float4(PreViewPosition, 1), PrimaryView.TranslatedWorldToClip);
 	const float2 ScreenPosition = ClipPosition.xy / ClipPosition.w;
-	const float2 ScreenUV = (ScreenPosition * float2(0.5f, -0.5f)) + 0.5f;
+	const float2 ScreenUV = (ScreenPosition * PrimaryView.ScreenPositionScaleBias.xy) + PrimaryView.ScreenPositionScaleBias.wz;
 
 	int TotalSamples = 0;
 	int UnoccludedSamples = 0;
