- Unbind OverrideParameters in PostApplyToComponent to avoid stale parameter store bindings
- They will be rebound after PostApplyToComponent from the reset context

diff --git a/Engine/Plugins/FX/Niagara/Source/Niagara/Private/NiagaraComponent.cpp b/Engine/Plugins/FX/Niagara/Source/Niagara/Private/NiagaraComponent.cpp
--- a/Engine/Plugins/FX/Niagara/Source/Niagara/Private/NiagaraComponent.cpp
+++ b/Engine/Plugins/FX/Niagara/Source/Niagara/Private/NiagaraComponent.cpp
@@ -1984,6 +1984,7 @@ void UNiagaraComponent::PostApplyToComponent()
 	// Without this calls to set override data interface values can fail as the DI may not be present in the override parameter store
 	if (Asset != nullptr)
 	{
+		OverrideParameters.UnbindAll();
 		CopyParametersFromAsset(false);
 	}
 	ApplyOverridesToParameterStore();
