Static analysis fix

diff --git a/Engine/Plugins/FX/Niagara/Source/Niagara/Internal/NiagaraAnalytics.cpp b/Engine/Plugins/FX/Niagara/Source/Niagara/Internal/NiagaraAnalytics.cpp
--- a/Engine/Plugins/FX/Niagara/Source/Niagara/Internal/NiagaraAnalytics.cpp
+++ b/Engine/Plugins/FX/Niagara/Source/Niagara/Internal/NiagaraAnalytics.cpp
@@ -22,11 +22,10 @@ void NiagaraAnalytics::RecordEvent(FString&& EventName, const FString& Attribute
 
 bool NiagaraAnalytics::IsPluginAsset(const UObject* Obj)
 {
-	if (IsValid(Obj) && Obj->IsAsset())
+	if (IsValid(Obj) && Obj->IsAsset() && Obj->GetPackage())
 	{
-		UPackage* Package = Obj->GetPackage();
-		FString Name = Package->GetName();
-		return Package && (Name.StartsWith(TEXT("/Niagara/"), ESearchCase::CaseSensitive) || Name.StartsWith(TEXT("/NiagaraFluids/"), ESearchCase::CaseSensitive));
+		FString Name = Obj->GetPackage()->GetName();
+		return Name.StartsWith(TEXT("/Niagara/"), ESearchCase::CaseSensitive) || Name.StartsWith(TEXT("/NiagaraFluids/"), ESearchCase::CaseSensitive);
 	}
 	return false;
 }
