changed auto select single action to true

diff --git a/Engine/Plugins/FX/Niagara/Source/NiagaraEditor/Public/Widgets/SNiagaraParameterMenu.cpp b/Engine/Plugins/FX/Niagara/Source/NiagaraEditor/Public/Widgets/SNiagaraParameterMenu.cpp
--- a/Engine/Plugins/FX/Niagara/Source/NiagaraEditor/Public/Widgets/SNiagaraParameterMenu.cpp
+++ b/Engine/Plugins/FX/Niagara/Source/NiagaraEditor/Public/Widgets/SNiagaraParameterMenu.cpp
@@ -47,20 +47,21 @@ void SNiagaraParameterMenu::Construct(const FArguments& InArgs)
 			.MinDesiredWidth(300)
 			.MaxDesiredHeight(700) // Set max desired height to prevent flickering bug for menu larger than screen
 			[
-			SAssignNew(GraphMenu, SGraphActionMenu)
-			.OnActionSelected(this, &SNiagaraParameterMenu::OnActionSelected)
-			.OnCollectAllActions(this, &SNiagaraParameterMenu::CollectAllActions)
-			.SortItemsRecursively(false)
-			.AlphaSortItems(false)
-			.AutoExpandActionMenu(bAutoExpandMenu)
-			.ShowFilterTextBox(true)
-			.OnGetSectionTitle(InArgs._OnGetSectionTitle)
-			.OnCreateCustomRowExpander_Static(&SNiagaraParameterMenu::CreateCustomActionExpander)
-			.OnCreateWidgetForAction_Lambda([](const FCreateWidgetForActionData* InData)
-				{
-					return SNew(SNiagaraGraphActionWidget, InData);
-				})
-			]
+				SAssignNew(GraphMenu, SGraphActionMenu)
+				.OnActionSelected(this, &SNiagaraParameterMenu::OnActionSelected)
+				.OnCollectAllActions(this, &SNiagaraParameterMenu::CollectAllActions)
+				.SortItemsRecursively(false)
+				.AlphaSortItems(false)
+				.AutoExpandActionMenu(bAutoExpandMenu)
+				.bAutomaticallySelectSingleAction(true)
+				.ShowFilterTextBox(true)
+				.OnGetSectionTitle(InArgs._OnGetSectionTitle)
+				.OnCreateCustomRowExpander_Static(&SNiagaraParameterMenu::CreateCustomActionExpander)
+				.OnCreateWidgetForAction_Lambda([](const FCreateWidgetForActionData* InData)
+					{
+						return SNew(SNiagaraGraphActionWidget, InData);
+					})
+				]
 		]
 	];
 }
