Readded emitter scalability to default view as discussed a long time ago

diff --git a/Engine/Plugins/FX/Niagara/Source/NiagaraEditor/Private/NiagaraEmitterDetailsCustomization.cpp b/Engine/Plugins/FX/Niagara/Source/NiagaraEditor/Private/NiagaraEmitterDetailsCustomization.cpp
--- a/Engine/Plugins/FX/Niagara/Source/NiagaraEditor/Private/NiagaraEmitterDetailsCustomization.cpp
+++ b/Engine/Plugins/FX/Niagara/Source/NiagaraEditor/Private/NiagaraEmitterDetailsCustomization.cpp
@@ -83,12 +83,7 @@ void FNiagaraEmitterDetails::CustomizeDetails(IDetailLayoutBuilder& InDetailLayo
 			if (ChildProperty->HasAllPropertyFlags(CPF_Edit))
 			{
 				FName Category = FName(ChildProperty->GetMetaData(TEXT("Category")));
-
-				// we display the scalability category within scalability mode, which is why we hide it here
-				if (Category != FName("Scalability"))
-				{
-					CategoryPropertyMap.FindOrAdd(Category).Add(ChildProperty);
-				}
+				CategoryPropertyMap.FindOrAdd(Category).Add(ChildProperty);
 			}
 		}
 
