- Fix for ParticlePosition previous value being incorrect
- To mitigate cost & interpolators we only provide the value for vertex shaders, not pixel shaders
- Previous data is also not generated in all cases, it falls back to how it is currently which is previous equals current
- Niagara Sprite & Mesh renderers will generate accurate data when motion vectors are set to accurate
- Without accurate the results from uising velocity are somewhat mixed

diff --git a/Engine/Plugins/FX/Niagara/Shaders/Private/NiagaraMeshVertexFactory.ush b/Engine/Plugins/FX/Niagara/Shaders/Private/NiagaraMeshVertexFactory.ush
--- a/Engine/Plugins/FX/Niagara/Shaders/Private/NiagaraMeshVertexFactory.ush
+++ b/Engine/Plugins/FX/Niagara/Shaders/Private/NiagaraMeshVertexFactory.ush
@@ -310,6 +310,7 @@ FMaterialPixelParameters GetMaterialPixelParameters(FVertexFactoryInterpolantsVS
 #if USE_PARTICLE_POSITION
 	Result.Particle.TranslatedWorldPositionAndSize.xyz = Interpolants.ParticleTranslatedWorldPosition;
 	Result.Particle.TranslatedWorldPositionAndSize.w = 1;
+	Result.Particle.PrevTranslatedWorldPositionAndSize = Result.Particle.TranslatedWorldPositionAndSize;
 #endif
 
 #if USE_PARTICLE_VELOCITY
@@ -376,6 +377,8 @@ FMaterialVertexParameters GetMaterialVertexParameters(FVertexFactoryInput Input,
 #if NEEDS_PARTICLE_POSITION
 	Result.Particle.TranslatedWorldPositionAndSize.xyz = Intermediates.ParticleTranslatedWorldPosition;
 	Result.Particle.TranslatedWorldPositionAndSize.w = 1;
+	Result.Particle.PrevTranslatedWorldPositionAndSize.xyz = LWCToFloat(LWCAdd(LWCGetOrigin(Intermediates.PrevParticleToWorld), ResolvedView.PrevPreViewTranslation));
+	Result.Particle.PrevTranslatedWorldPositionAndSize.w = 1;
 #endif
 
 #if NEEDS_PARTICLE_VELOCITY
diff --git a/Engine/Plugins/FX/Niagara/Shaders/Private/NiagaraRibbonVertexFactory.ush b/Engine/Plugins/FX/Niagara/Shaders/Private/NiagaraRibbonVertexFactory.ush
--- a/Engine/Plugins/FX/Niagara/Shaders/Private/NiagaraRibbonVertexFactory.ush
+++ b/Engine/Plugins/FX/Niagara/Shaders/Private/NiagaraRibbonVertexFactory.ush
@@ -391,6 +391,7 @@ FMaterialPixelParameters GetMaterialPixelParameters(FVertexFactoryInterpolantsVS
 
 #if USE_PARTICLE_POSITION
 	Result.Particle.TranslatedWorldPositionAndSize = Interpolants.TranslatedWorldPositionAndSize;
+	Result.Particle.PrevTranslatedWorldPositionAndSize = Result.Particle.TranslatedWorldPositionAndSize;
 #endif
 
 #if USE_PARTICLE_VELOCITY
@@ -444,6 +445,7 @@ FMaterialVertexParameters GetMaterialVertexParameters(FVertexFactoryInput Input,
 
 #if USE_PARTICLE_POSITION
 	Result.Particle.TranslatedWorldPositionAndSize = Intermediates.TranslatedWorldPositionAndSize;
+	Result.Particle.PrevTranslatedWorldPositionAndSize = Result.Particle.TranslatedWorldPositionAndSize;
 #endif
 
 #if USE_PARTICLE_VELOCITY
diff --git a/Engine/Plugins/FX/Niagara/Shaders/Private/NiagaraSpriteVertexFactory.ush b/Engine/Plugins/FX/Niagara/Shaders/Private/NiagaraSpriteVertexFactory.ush
--- a/Engine/Plugins/FX/Niagara/Shaders/Private/NiagaraSpriteVertexFactory.ush
+++ b/Engine/Plugins/FX/Niagara/Shaders/Private/NiagaraSpriteVertexFactory.ush
@@ -168,6 +168,7 @@ struct FVertexFactoryIntermediates
 	FLWCVector3 PreviousVertexWorldPosition;
 	/** Particle translated world space position and size. */
 	float4 TranslatedWorldPositionAndSize;
+	float4 PrevTranslatedWorldPositionAndSize;
 	
 #if USE_PARTICLE_LIGHTING_OFFSET
 	float3 LightingPositionOffset;
@@ -324,6 +325,7 @@ FMaterialPixelParameters GetMaterialPixelParameters(FVertexFactoryInterpolantsVS
 
 #if USE_PARTICLE_POSITION
 	Result.Particle.TranslatedWorldPositionAndSize = Interpolants.ParticleTranslatedWorldPositionAndSize;
+	Result.Particle.PrevTranslatedWorldPositionAndSize = Result.Particle.TranslatedWorldPositionAndSize;
 #endif
 
 #if USE_PARTICLE_VELOCITY
@@ -407,6 +409,7 @@ FMaterialVertexParameters GetMaterialVertexParameters(FVertexFactoryInput Input,
 
 #if USE_PARTICLE_POSITION
 	Result.Particle.TranslatedWorldPositionAndSize = Intermediates.TranslatedWorldPositionAndSize;
+	Result.Particle.PrevTranslatedWorldPositionAndSize = Result.Particle.TranslatedWorldPositionAndSize;
 #endif
 
 #if USE_PARTICLE_VELOCITY
@@ -913,16 +916,18 @@ FVertexFactoryIntermediates GetVertexFactoryIntermediates(FVertexFactoryInput In
 	Intermediates.VertexWorldPosition = LWCAdd(ParticlePosition, VertexOffset);
 		
 	// Prev vertex position (approximated from current velocity unless enabled)
+	FLWCVector3 PrevParticlePosition;
+	float3 PrevCameraOffset = CameraOffset;
 	if ( NiagaraSpriteVF.AccurateMotionVectors != 0 )
 	{
-		const FLWCVector3 PrevParticlePosition = GetNiagaraPreviousParticlePosition(ParticleID);
+		PrevParticlePosition = GetNiagaraPreviousParticlePosition(ParticleID);
 		const float PrevParticleRotation = GetNiagaraPreviousParticleRotation(ParticleID);
 		const float2 PrevParticleSize = GetNiagaraPreviousParticleSize(ParticleID);
 		const float3 PrevParticleVelocity = GetNiagaraPreviousParticleVelocity(ParticleID);		
 		const float3 PrevCustomFacing = SafeNormalize(GetNiagaraPreviousParticleFacingVector(ParticleID));
 		const float3 PrevCustomAlignment = SafeNormalize(GetNiagaraPreviousParticleAlignmentVector(ParticleID));
 		const float2 PrevPivotOffset = GetNiagaraPreviousPivotOffset(ParticleID);
-		const float3 PrevCameraOffset = SafeNormalize(LWCToFloat(LWCSubtract(ResolvedView.PrevWorldCameraOrigin, PrevParticlePosition))) * GetNiagaraPreviousCameraOffset(ParticleID);
+		PrevCameraOffset = SafeNormalize(LWCToFloat(LWCSubtract(ResolvedView.PrevWorldCameraOrigin, PrevParticlePosition))) * GetNiagaraPreviousCameraOffset(ParticleID);
 		const float3 PrevParticleDirection = SafeNormalizeWithFallback(PrevParticleVelocity, float3(0,0,1));
 			
 		float3 PrevRight, PrevUp;
@@ -942,12 +947,16 @@ FVertexFactoryIntermediates GetVertexFactoryIntermediates(FVertexFactoryInput In
 	else
 	{
 		Intermediates.PreviousVertexWorldPosition = LWCSubtract(Intermediates.VertexWorldPosition, LWCPromote(ParticleVelocity * NiagaraSpriteVF.DeltaSeconds));
+		PrevParticlePosition = LWCSubtract(Intermediates.Position, LWCPromote(ParticleVelocity * NiagaraSpriteVF.DeltaSeconds));
 	}
 
 #if USE_PARTICLE_POSITION
 	// Do we really want the camera offset here ? What about when GENERATE_SPHERICAL_PARTICLE_NORMALS
 	Intermediates.TranslatedWorldPositionAndSize.xyz = LWCToFloat(LWCAdd(Intermediates.Position, ResolvedView.PreViewTranslation)) + CameraOffset;
 	Intermediates.TranslatedWorldPositionAndSize.w = .5f * min(Size.x, Size.y);
+	
+	Intermediates.PrevTranslatedWorldPositionAndSize.xyz = LWCToFloat(LWCAdd(PrevParticlePosition, ResolvedView.PrevPreViewTranslation)) + PrevCameraOffset;
+	Intermediates.PrevTranslatedWorldPositionAndSize.w = .5f * min(Size.x, Size.y);
 #if GENERATE_SPHERICAL_PARTICLE_NORMALS
 	// Recenter the particle position for GENERATE_SPHERICAL_PARTICLE_NORMALS to use the right center.
 	Intermediates.TranslatedWorldPositionAndSize.xyz += mul(Size * (float2(.5, .5) - PivotOffset), Tangents);
diff --git a/Engine/Shaders/Private/MaterialTemplate.ush b/Engine/Shaders/Private/MaterialTemplate.ush
--- a/Engine/Shaders/Private/MaterialTemplate.ush
+++ b/Engine/Shaders/Private/MaterialTemplate.ush
@@ -201,6 +201,8 @@ struct FMaterialParticleParameters
 	half4 Color;
 	/** Particle translated world space position and size(radius). */
 	float4 TranslatedWorldPositionAndSize;
+	/** Previous Particle translated world space position and size */
+	float4 PrevTranslatedWorldPositionAndSize;
 	/** Macro UV scale and bias. */
 	half4 MacroUV;
 
diff --git a/Engine/Shaders/Private/MeshParticleVertexFactory.ush b/Engine/Shaders/Private/MeshParticleVertexFactory.ush
--- a/Engine/Shaders/Private/MeshParticleVertexFactory.ush
+++ b/Engine/Shaders/Private/MeshParticleVertexFactory.ush
@@ -336,6 +336,7 @@ FMaterialPixelParameters GetMaterialPixelParameters(FVertexFactoryInterpolantsVS
 #if USE_PARTICLE_POSITION
 	Result.Particle.TranslatedWorldPositionAndSize.xyz = Interpolants.ParticleTranslatedWorldPosition.xyz;
 	Result.Particle.TranslatedWorldPositionAndSize.w = 1;
+	Result.Particle.PrevTranslatedWorldPositionAndSize = Result.Particle.TranslatedWorldPositionAndSize;
 #endif
 
 #if USE_PARTICLE_VELOCITY
@@ -390,6 +391,7 @@ FMaterialVertexParameters GetMaterialVertexParameters(FVertexFactoryInput Input,
 
 	Result.Particle.TranslatedWorldPositionAndSize.xyz = Intermediates.ParticleTranslatedWorldPosition.xyz;
 	Result.Particle.TranslatedWorldPositionAndSize.w = 1;
+	Result.Particle.PrevTranslatedWorldPositionAndSize = Result.Particle.TranslatedWorldPositionAndSize;
 	Result.Particle.Velocity = Intermediates.ParticleVelocity;
 	Result.Particle.RelativeTime = Intermediates.RelativeTime;
 
diff --git a/Engine/Shaders/Private/ParticleGPUSpriteVertexFactory.ush b/Engine/Shaders/Private/ParticleGPUSpriteVertexFactory.ush
--- a/Engine/Shaders/Private/ParticleGPUSpriteVertexFactory.ush
+++ b/Engine/Shaders/Private/ParticleGPUSpriteVertexFactory.ush
@@ -452,6 +452,7 @@ FMaterialPixelParameters GetMaterialPixelParameters( FVertexFactoryInterpolantsV
 
 #if USE_PARTICLE_POSITION
 	Result.Particle.TranslatedWorldPositionAndSize = Interpolants.ParticleTranslatedWorldPositionAndSize;
+	Result.Particle.PrevTranslatedWorldPositionAndSize = Result.Particle.TranslatedWorldPositionAndSize;
 #endif
 
 #if USE_SPHERICAL_NORMALS
@@ -540,6 +541,7 @@ FMaterialVertexParameters GetMaterialVertexParameters(FVertexFactoryInput Input,
 
 #if USE_PARTICLE_POSITION
 	Result.Particle.TranslatedWorldPositionAndSize = float4(Intermediates.ParticleTranslatedWorldPosition, Intermediates.ParticleRadius);
+	Result.Particle.PrevTranslatedWorldPositionAndSize = Result.Particle.TranslatedWorldPositionAndSize;
 #endif
 
 	Result.Particle.Color = Intermediates.Color;
diff --git a/Engine/Shaders/Private/ParticleSpriteVertexFactory.ush b/Engine/Shaders/Private/ParticleSpriteVertexFactory.ush
--- a/Engine/Shaders/Private/ParticleSpriteVertexFactory.ush
+++ b/Engine/Shaders/Private/ParticleSpriteVertexFactory.ush
@@ -220,6 +220,7 @@ FMaterialPixelParameters GetMaterialPixelParameters(FVertexFactoryInterpolantsVS
 
 #if USE_PARTICLE_POSITION
 	Result.Particle.TranslatedWorldPositionAndSize = Interpolants.ParticleTranslatedWorldPositionAndSize;
+	Result.Particle.PrevTranslatedWorldPositionAndSize = Result.Particle.TranslatedWorldPositionAndSize;
 #endif
 
 #if USE_PARTICLE_VELOCITY
@@ -283,6 +284,7 @@ FMaterialVertexParameters GetMaterialVertexParameters(FVertexFactoryInput Input,
 
 #if USE_PARTICLE_POSITION
 	Result.Particle.TranslatedWorldPositionAndSize = Intermediates.TranslatedWorldPositionAndSize;
+	Result.Particle.PrevTranslatedWorldPositionAndSize = Result.Particle.TranslatedWorldPositionAndSize;
 #endif
 
 #if USE_PARTICLE_VELOCITY
diff --git a/Engine/Source/Runtime/Engine/Private/Materials/HLSLMaterialTranslator.cpp b/Engine/Source/Runtime/Engine/Private/Materials/HLSLMaterialTranslator.cpp
--- a/Engine/Source/Runtime/Engine/Private/Materials/HLSLMaterialTranslator.cpp
+++ b/Engine/Source/Runtime/Engine/Private/Materials/HLSLMaterialTranslator.cpp
@@ -5768,14 +5768,23 @@ int32 FHLSLMaterialTranslator::ParticlePosition(EPositionOrigin OriginType)
 		return NonVertexOrPixelShaderExpressionError();
 	}
 	bNeedsParticlePosition = true;
+
+	const TCHAR* TranslatedWP = TEXT("Parameters.Particle.TranslatedWorldPositionAndSize.xyz");
+	const TCHAR* AbsoluteWP = TEXT("LWCSubtract(Parameters.Particle.TranslatedWorldPositionAndSize.xyz, ResolvedView.PreViewTranslation)");
+	if ( bCompilingPreviousFrame && ShaderFrequency == SF_Vertex )
+	{
+		TranslatedWP = TEXT("Parameters.Particle.PrevTranslatedWorldPositionAndSize.xyz");
+		AbsoluteWP = TEXT("LWCSubtract(Parameters.Particle.PrevTranslatedWorldPositionAndSize.xyz, ResolvedView.PrevPreViewTranslation)");
+	}
+
 	if (OriginType == EPositionOrigin::CameraRelative)
 	{
-		return AddInlinedCodeChunkZeroDeriv(MCT_Float3, TEXT("Parameters.Particle.TranslatedWorldPositionAndSize.xyz"));
+		return AddInlinedCodeChunkZeroDeriv(MCT_Float3, TranslatedWP);
 	}
 	else
 	{
 		AddLWCFuncUsage(ELWCFunctionKind::Subtract);
-		const int32 Result = AddInlinedCodeChunkZeroDeriv(MCT_LWCVector3,TEXT("LWCSubtract(Parameters.Particle.TranslatedWorldPositionAndSize.xyz, ResolvedView.PreViewTranslation)"));
+		const int32 Result = AddInlinedCodeChunkZeroDeriv(MCT_LWCVector3, AbsoluteWP);
 		return CastToNonLWCIfDisabled(Result);
 	}
 }
