- Fix for NPC data interfaces not working
- Because NPC DIs do not resolve they were assumed to be internal, this meant they did not flow through the parameter stores correctly.

diff --git a/Engine/Plugins/FX/Niagara/Source/Niagara/Private/NiagaraResolveDIHelpers.cpp b/Engine/Plugins/FX/Niagara/Source/Niagara/Private/NiagaraResolveDIHelpers.cpp
--- a/Engine/Plugins/FX/Niagara/Source/Niagara/Private/NiagaraResolveDIHelpers.cpp
+++ b/Engine/Plugins/FX/Niagara/Source/Niagara/Private/NiagaraResolveDIHelpers.cpp
@@ -306,19 +306,32 @@ namespace FNiagaraResolveDIHelpers
 
 			if (ResolvedDataInterface.ResolvedDataInterface == nullptr)
 			{
-				// If the DI was not read from a parameter or couldn't be found, use the one cached during compilation, and give it an internal
-				// name to prevent it from being bound incorrectly.
-				FNameBuilder NameBuilder;
-				NameBuilder.Append(FNiagaraConstants::InternalNamespaceString);
-				NameBuilder.AppendChar(TEXT('.'));
-				ResolvedDataInterface.Name.AppendString(NameBuilder);
-
-				FNiagaraVariable InternalVariable(CachedDefaultDataInterface.Type, FName(NameBuilder.ToString()));
-				ResolvedDataInterface.ResolvedVariable = InternalVariable;
-				ResolvedDataInterface.ParameterStoreVariable = InternalVariable;
-				ResolvedDataInterface.bIsInternal = true;
-				ResolvedDataInterface.ResolvedDataInterface = CachedDefaultDataInterface.DataInterface;
-				ResolvedDataInterface.ResolvedSourceEmitterName = CachedDefaultDataInterface.SourceEmitterName;
+				// NPC variables will fail to resolve but we want to main that it's both external & not append the internal namespace
+				if ( FNiagaraVariable::IsInNameSpace(FNiagaraConstants::ParameterCollectionNamespaceString, ResolvedDataInterface.Name) )
+				{
+					FNiagaraVariable NPCVariable(CachedDefaultDataInterface.Type, ResolvedDataInterface.Name);
+					ResolvedDataInterface.ResolvedVariable = NPCVariable;
+					ResolvedDataInterface.ParameterStoreVariable = NPCVariable;
+					ResolvedDataInterface.bIsInternal = false;
+					ResolvedDataInterface.ResolvedDataInterface = CachedDefaultDataInterface.DataInterface;
+					ResolvedDataInterface.ResolvedSourceEmitterName = CachedDefaultDataInterface.SourceEmitterName;
+				}
+				else
+				{
+					// If the DI was not read from a parameter or couldn't be found, use the one cached during compilation, and give it an internal
+					// name to prevent it from being bound incorrectly.
+					FNameBuilder NameBuilder;
+					NameBuilder.Append(FNiagaraConstants::InternalNamespaceString);
+					NameBuilder.AppendChar(TEXT('.'));
+					ResolvedDataInterface.Name.AppendString(NameBuilder);
+
+					FNiagaraVariable InternalVariable(CachedDefaultDataInterface.Type, FName(NameBuilder.ToString()));
+					ResolvedDataInterface.ResolvedVariable = InternalVariable;
+					ResolvedDataInterface.ParameterStoreVariable = InternalVariable;
+					ResolvedDataInterface.bIsInternal = true;
+					ResolvedDataInterface.ResolvedDataInterface = CachedDefaultDataInterface.DataInterface;
+					ResolvedDataInterface.ResolvedSourceEmitterName = CachedDefaultDataInterface.SourceEmitterName;
+				}
 			}
 
 			ResolvedDataInterfaceIndex++;
