- Fix array remove on GPU when 0 or max elements

diff --git a/Engine/Plugins/FX/Niagara/Shaders/Private/NiagaraDataInterfaceArrayRWTemplate.ush b/Engine/Plugins/FX/Niagara/Shaders/Private/NiagaraDataInterfaceArrayRWTemplate.ush
--- a/Engine/Plugins/FX/Niagara/Shaders/Private/NiagaraDataInterfaceArrayRWTemplate.ush
+++ b/Engine/Plugins/FX/Niagara/Shaders/Private/NiagaraDataInterfaceArrayRWTemplate.ush
@@ -77,6 +77,7 @@ void RemoveLastElem_{ParameterName}_UEImpureCall(bool bSkip, out {VariableType}
 	if ( bSkip == false )
 	{
 		InterlockedAdd(RWInstanceCounts[GetCountOffset_{ParameterName}()], uint(-1), Index);
+		--Index;
 
 		bOutValid = Index < GetCapacity_{ParameterName}();
 		if ( !bOutValid )
