When resetting systems be sure to also invalidate the compilation results for child scripts.  This ensures deterministic behavior when cooking for platforms that don't require AV data, and makes sure that the resulting editor only data isn't bloated with things like HLSL source code when the data is never going to be used.

diff --git a/Engine/Plugins/FX/Niagara/Source/Niagara/Private/NiagaraSystem.cpp b/Engine/Plugins/FX/Niagara/Source/Niagara/Private/NiagaraSystem.cpp
--- a/Engine/Plugins/FX/Niagara/Source/Niagara/Private/NiagaraSystem.cpp
+++ b/Engine/Plugins/FX/Niagara/Source/Niagara/Private/NiagaraSystem.cpp
@@ -3353,6 +3353,17 @@ void UNiagaraSystem::ResetToEmptySystem()
 
 	// while we'd like to remove these as well, BP will generate warnings for missing parameters
 	//ExposedParameters = FNiagaraUserRedirectionParameterStore();
+
+	// in the case of dedicated servers, clearing out the system is sufficient to remove unnecessary data as the emitters and scripts
+	// will not be loaded on the target platform.  This is not strictly true for all platforms that don't support AV data and so we
+	// also want to be a bit more invasive with resetting the underlying objects (Emitters & Scripts).
+	ForEachScript([&](UNiagaraScript* Script)
+	{
+		if (Script)
+		{
+			Script->InvalidateCompileResults(TEXT("Parent system ResetToEmptySystem"));
+		}
+	});
 }
 
 #endif
