Fixed missing enum entry in if check

diff --git a/Engine/Plugins/FX/Niagara/Source/NiagaraEditor/Public/Widgets/SNiagaraParameterMenu.cpp b/Engine/Plugins/FX/Niagara/Source/NiagaraEditor/Public/Widgets/SNiagaraParameterMenu.cpp
--- a/Engine/Plugins/FX/Niagara/Source/NiagaraEditor/Public/Widgets/SNiagaraParameterMenu.cpp
+++ b/Engine/Plugins/FX/Niagara/Source/NiagaraEditor/Public/Widgets/SNiagaraParameterMenu.cpp
@@ -73,7 +73,7 @@ TSharedPtr<SEditableTextBox> SNiagaraParameterMenu::GetSearchBox()
 
 void SNiagaraParameterMenu::OnActionSelected(const TArray<TSharedPtr<FEdGraphSchemaAction>>& SelectedActions, ESelectInfo::Type InSelectionType)
 {
-	if (InSelectionType == ESelectInfo::OnMouseClick || InSelectionType == ESelectInfo::OnKeyPress || SelectedActions.Num() == 0)
+	if (InSelectionType == ESelectInfo::OnMouseClick || InSelectionType == ESelectInfo::Direct || InSelectionType == ESelectInfo::OnKeyPress || SelectedActions.Num() == 0)
 	{
 		for (int32 ActionIndex = 0; ActionIndex < SelectedActions.Num(); ActionIndex++)
 		{
