- LWC fix for Niagara sorting & culling

diff --git a/Engine/Plugins/FX/Niagara/Shaders/Private/NiagaraSortKeyGen.usf b/Engine/Plugins/FX/Niagara/Shaders/Private/NiagaraSortKeyGen.usf
--- a/Engine/Plugins/FX/Niagara/Shaders/Private/NiagaraSortKeyGen.usf
+++ b/Engine/Plugins/FX/Niagara/Shaders/Private/NiagaraSortKeyGen.usf
@@ -51,7 +51,6 @@ float2 CullDistanceRangeSquared;
 float4 LODScreenSize;
 float4 LocalBoundingSphere;
 float3 CullingWorldSpaceOffset;
-float3 SystemLWCTile;
 
 RWBuffer<uint> OutKeys;
 RWBuffer<int> OutParticleIndices;
@@ -143,7 +142,6 @@ bool CullParticle(in uint ParticleIndex)
 	if (CullPositionAttributeOffset != -1)
 	{
 		Pos = GetVec3(CullPositionAttributeOffset, ParticleIndex);
-		Pos = LWCHackToFloat(MakeLWCVector3(SystemLWCTile, Pos));
 	}
 
 	float4 Quat = { 0, 0, 0, 1 };
diff --git a/Engine/Plugins/FX/Niagara/Source/Niagara/Classes/NiagaraGPUSortInfo.h b/Engine/Plugins/FX/Niagara/Source/Niagara/Classes/NiagaraGPUSortInfo.h
--- a/Engine/Plugins/FX/Niagara/Source/Niagara/Classes/NiagaraGPUSortInfo.h
+++ b/Engine/Plugins/FX/Niagara/Source/Niagara/Classes/NiagaraGPUSortInfo.h
@@ -68,7 +68,6 @@ struct FNiagaraGPUSortInfo
 	int32 MeshIndex = 0;
 	FSphere LocalBSphere = FSphere(0);
 	FVector CullingWorldSpaceOffset = FVector(0, 0, 0);
-	FVector3f SystemLWCTile = FVector3f(0, 0, 0);
 	FVector2f DistanceCullRange { 0.0f, FLT_MAX };
 	FVector4f LODScreenSize = FVector4f(0.0f, 2.0f, 1.0f, 0.0f);
 	TArray<FPlane, TFixedAllocator<MaxCullPlanes>> CullPlanes;
diff --git a/Engine/Plugins/FX/Niagara/Source/Niagara/Private/NiagaraGpuComputeDispatch.cpp b/Engine/Plugins/FX/Niagara/Source/Niagara/Private/NiagaraGpuComputeDispatch.cpp
--- a/Engine/Plugins/FX/Niagara/Source/Niagara/Private/NiagaraGpuComputeDispatch.cpp
+++ b/Engine/Plugins/FX/Niagara/Source/Niagara/Private/NiagaraGpuComputeDispatch.cpp
@@ -2296,7 +2296,6 @@ void FNiagaraGpuComputeDispatch::GenerateSortKeys(FRHICommandListImmediate& RHIC
 			Params.LODScreenSize = SortInfo.LODScreenSize;
 			Params.LocalBoundingSphere = FVector4f((FVector3f)SortInfo.LocalBSphere.Center, SortInfo.LocalBSphere.W);
 			Params.CullingWorldSpaceOffset = (FVector3f)SortInfo.CullingWorldSpaceOffset;
-			Params.SystemLWCTile = SortInfo.SystemLWCTile;
 
 			Params.NumCullPlanes = 0;
 			for (const FPlane& Plane : SortInfo.CullPlanes)
diff --git a/Engine/Plugins/FX/Niagara/Source/Niagara/Private/NiagaraRendererMeshes.cpp b/Engine/Plugins/FX/Niagara/Source/Niagara/Private/NiagaraRendererMeshes.cpp
--- a/Engine/Plugins/FX/Niagara/Source/Niagara/Private/NiagaraRendererMeshes.cpp
+++ b/Engine/Plugins/FX/Niagara/Source/Niagara/Private/NiagaraRendererMeshes.cpp
@@ -617,7 +617,6 @@ void FNiagaraRendererMeshes::InitializeSortInfo(const FParticleMeshRenderData& P
 	OutSortInfo.MeshIndexAttributeOffset = ParticleMeshRenderData.MeshIndexOffset;
 	OutSortInfo.RendererVisibility = RendererVisibility;
 	OutSortInfo.DistanceCullRange = DistanceCullRange;
-	OutSortInfo.SystemLWCTile = UseLocalSpace(&SceneProxy) ? FVector3f::Zero() :SceneProxy.GetLWCRenderTile();
 
 	if (ParticleMeshRenderData.bSortCullOnGpu)
 	{
@@ -715,6 +714,18 @@ void FNiagaraRendererMeshes::InitializeSortInfo(const FParticleMeshRenderData& P
 			}
 		}
 	}
+	else
+	{
+		const FVector LWCTileOffset = FVector(SceneProxy.GetLWCRenderTile()) * FLargeWorldRenderScalar::GetTileSize();
+		OutSortInfo.ViewOrigin -= LWCTileOffset;
+		if (bEnableFrustumCulling)
+		{
+			for (FPlane& Plane : OutSortInfo.CullPlanes)
+			{
+				Plane = Plane.TranslateBy(-LWCTileOffset);
+			}
+		}
+	}
 
 	if (ParticleMeshRenderData.bNeedsCull)
 	{
@@ -2098,40 +2109,3 @@ bool FNiagaraRendererMeshes::IsMaterialValid(const UMaterialInterface* Mat)const
 
 	return bIsMaterialValid;
 }
-
-
-//////////////////////////////////////////////////////////////////////////
-// Proposed class for ensuring Niagara/Cascade components who's proxies reference render data of other objects (Materials, Meshes etc) do not have data freed from under them.
-// Our components register themselves with the referenced component which then calls InvalidateRenderDependencies() whenever it's render data is changed or when it is destroyed.
-// UNTESTED - DO NOT USE.
-struct FComponentRenderDependencyHandler
-{
-	void AddDependency(UPrimitiveComponent* Component)
-	{
-		DependentComponents.Add(Component);
-	}
-
-	void RemoveDependancy(UPrimitiveComponent* Component)
-	{
-		DependentComponents.RemoveSwap(Component);
-	}
-
-	void InvalidateRenderDependencies()
-	{
-		int32 i = DependentComponents.Num();
-		while (--i >= 0)
-		{
-			if (UPrimitiveComponent* Comp = DependentComponents[i].Get())
-			{
-				Comp->MarkRenderStateDirty();
-			}
-			else
-			{
-				DependentComponents.RemoveAtSwap(i);
-			}
-		}
-	}
-
-	TArray<TWeakObjectPtr<UPrimitiveComponent>> DependentComponents;
-};
-
diff --git a/Engine/Plugins/FX/Niagara/Source/Niagara/Private/NiagaraRendererSprites.cpp b/Engine/Plugins/FX/Niagara/Source/Niagara/Private/NiagaraRendererSprites.cpp
--- a/Engine/Plugins/FX/Niagara/Source/Niagara/Private/NiagaraRendererSprites.cpp
+++ b/Engine/Plugins/FX/Niagara/Source/Niagara/Private/NiagaraRendererSprites.cpp
@@ -417,7 +417,6 @@ void FNiagaraRendererSprites::InitializeSortInfo(FParticleSpriteRenderData& Part
 	OutSortInfo.RendererVisTagAttributeOffset = ParticleSpriteRenderData.RendererVisTagOffset;
 	OutSortInfo.RendererVisibility = RendererVisibility;
 	OutSortInfo.DistanceCullRange = DistanceCullRange;
-	OutSortInfo.SystemLWCTile = UseLocalSpace(&SceneProxy) ? FVector3f::Zero() : SceneProxy.GetLWCRenderTile();
 
 	if ( bEnableDistanceCulling )
 	{
@@ -451,6 +450,11 @@ void FNiagaraRendererSprites::InitializeSortInfo(FParticleSpriteRenderData& Part
 		OutSortInfo.ViewOrigin = SceneProxy.GetLocalToWorldInverse().TransformPosition(OutSortInfo.ViewOrigin);
 		OutSortInfo.ViewDirection = SceneProxy.GetLocalToWorld().GetTransposed().TransformVector(OutSortInfo.ViewDirection);
 	}
+	else
+	{
+		const FVector LWCTileOffset = FVector(SceneProxy.GetLWCRenderTile()) * FLargeWorldRenderScalar::GetTileSize();
+		OutSortInfo.ViewOrigin -= LWCTileOffset;
+	}
 
 	if (ParticleSpriteRenderData.bSortCullOnGpu)
 	{
