- Deprecate SortIndices as it is no longer used

diff --git a/Engine/Plugins/FX/Niagara/Source/Niagara/Public/NiagaraRenderer.h b/Engine/Plugins/FX/Niagara/Source/Niagara/Public/NiagaraRenderer.h
--- a/Engine/Plugins/FX/Niagara/Source/Niagara/Public/NiagaraRenderer.h
+++ b/Engine/Plugins/FX/Niagara/Source/Niagara/Public/NiagaraRenderer.h
@@ -68,7 +68,6 @@ struct FParticleRenderData
 class FNiagaraRenderer
 {
 public:
-
 	NIAGARA_API FNiagaraRenderer(ERHIFeatureLevel::Type FeatureLevel, const UNiagaraRendererProperties *InProps, const FNiagaraEmitterInstance* Emitter);
 	NIAGARA_API virtual ~FNiagaraRenderer();
 
@@ -92,6 +91,7 @@ public:
 	static NIAGARA_API bool IsViewRenderingOpaqueOnly(const FSceneView* View, bool bCastsVolumetricTranslucentShadow);
 	static NIAGARA_API bool AreViewsRenderingOpaqueOnly(const TArray<const FSceneView*>& Views, int32 ViewVisibilityMask, bool bCastsVolumetricTranslucentShadow);
 
+	UE_DEPRECATED(5.4, "SortIndices method should be replaced with SortAndCullIndices")
 	static NIAGARA_API void SortIndices(const struct FNiagaraGPUSortInfo& SortInfo, const FNiagaraRendererVariableInfo& SortVariable, const FNiagaraDataBuffer& Buffer, FGlobalDynamicReadBuffer::FAllocation& OutIndices);
 	static NIAGARA_API int32 SortAndCullIndices(const FNiagaraGPUSortInfo& SortInfo, const FNiagaraDataBuffer& Buffer, FGlobalDynamicReadBuffer::FAllocation& OutIndices);
 
