- Fix for performance overview values not matching sample mode

diff --git a/Engine/Plugins/FX/Niagara/Source/Niagara/Private/NiagaraDebugHud.cpp b/Engine/Plugins/FX/Niagara/Source/Niagara/Private/NiagaraDebugHud.cpp
--- a/Engine/Plugins/FX/Niagara/Source/Niagara/Private/NiagaraDebugHud.cpp
+++ b/Engine/Plugins/FX/Niagara/Source/Niagara/Private/NiagaraDebugHud.cpp
@@ -3338,6 +3338,7 @@ bool FNiagaraDebugHUDStatsListener::Tick()
 		NumFrames = 0;
 		bPushStats = true;
 	}
+
 	double GlobalAvg = 0.0;
 	double GlobalMax = 0.0;
 	for (auto& WorldStatsPair : AccumulatedWorldStats)
@@ -3347,8 +3348,16 @@ bool FNiagaraDebugHUDStatsListener::Tick()
 			FAccumulatedParticlePerfStats* Stats = WorldStatsPair.Value.Get();
 			check(Stats);
 
-			GlobalAvg += Stats->GetGameThreadStats().GetPerInstanceAvg();
-			GlobalMax += Stats->GetGameThreadStats().GetPerInstanceMax();
+			if (Settings.PerfSampleMode == ENiagaraDebugHUDPerfSampleMode::FrameTotal)
+			{
+				GlobalAvg += Stats->GetGameThreadStats().GetPerFrameAvg();
+				GlobalMax += Stats->GetGameThreadStats().GetPerFrameMax();
+			}
+			else
+			{
+				GlobalAvg += Stats->GetGameThreadStats().GetPerInstanceAvg();
+				GlobalMax += Stats->GetGameThreadStats().GetPerInstanceMax();
+			}
 
 			Stats->ResetGT();
 		}
@@ -3432,10 +3441,20 @@ void FNiagaraDebugHUDStatsListener::TickRT()
 			FAccumulatedParticlePerfStats* Stats = WorldStatsPair.Value.Get();
 			check(Stats);
 
-			RTGlobalAvg += Stats->GetRenderThreadStats().GetPerFrameAvg();
-			RTGlobalMax += Stats->GetRenderThreadStats().GetPerFrameMax();
-			GPUGlobalAvg += Stats->GetGPUStats().GetPerFrameAvgMicroseconds();
-			GPUGlobalMax += Stats->GetGPUStats().GetPerFrameMaxMicroseconds();
+			if (Settings.PerfSampleMode == ENiagaraDebugHUDPerfSampleMode::FrameTotal)
+			{
+				RTGlobalAvg += Stats->GetRenderThreadStats().GetPerFrameAvg();
+				RTGlobalMax += Stats->GetRenderThreadStats().GetPerFrameMax();
+				GPUGlobalAvg += Stats->GetGPUStats().GetPerFrameAvgMicroseconds();
+				GPUGlobalMax += Stats->GetGPUStats().GetPerFrameMaxMicroseconds();
+			}
+			else
+			{
+				RTGlobalAvg += Stats->GetRenderThreadStats().GetPerInstanceAvg();
+				RTGlobalMax += Stats->GetRenderThreadStats().GetPerInstanceMax();
+				GPUGlobalAvg += Stats->GetGPUStats().GetPerInstanceAvgMicroseconds();
+				GPUGlobalMax += Stats->GetGPUStats().GetPerInstanceMaxMicroseconds();
+			}
 
 			Stats->ResetRT();
 		}
