Cleanup source of the referenced compile hashes for Niagara's compilation Id.  Because of reordering and non-uniqueness the source strings would be lost and are generally not needed unless someone is using one of the debug options.

diff --git a/Engine/Plugins/FX/Niagara/Source/Niagara/Classes/NiagaraScript.h b/Engine/Plugins/FX/Niagara/Source/Niagara/Classes/NiagaraScript.h
--- a/Engine/Plugins/FX/Niagara/Source/Niagara/Classes/NiagaraScript.h
+++ b/Engine/Plugins/FX/Niagara/Source/Niagara/Classes/NiagaraScript.h
@@ -219,6 +219,21 @@ struct FNiagaraScriptDebuggerInfo
 	TAtomic<bool> bWritten;
 };
 
+#if WITH_EDITORONLY_DATA
+struct FNiagaraScriptHashCollector
+{
+	FNiagaraScriptHashCollector(bool bCollectHashSources);
+
+	NIAGARA_API void AddHash(const FNiagaraCompileHash& CompileHash, FStringView CompileHashSource);
+	FString BuildCompileHashSourceString(const FNiagaraCompileHash& CompileHash) const;
+
+	TArray<FNiagaraCompileHash> ReferencedHashes;
+	TMap<FNiagaraCompileHash, TArray<FString>> ReferencedHashSources;
+
+	const bool bCollectSources = false;
+};
+#endif
+
 /** Struct containing all of the data necessary to look up a NiagaraScript's VM executable results from the Derived Data Cache.*/
 USTRUCT()
 struct FNiagaraVMExecutableDataId
@@ -280,9 +295,6 @@ public:
 	/** Compile hashes of any top level scripts the script was dependent on that might trigger a recompile if they change. */
 	UPROPERTY()
 	TArray<FNiagaraCompileHash> ReferencedCompileHashes;
-
-	/** Temp storage while generating the Id. This is NOT serialized and shouldn't be used in any comparisons*/
-	TArray<FString> DebugReferencedObjects;
 #endif
 
 	/** The version of the script that was compiled. If empty then just the latest version. */
@@ -338,8 +350,7 @@ public:
 
 #if WITH_EDITORONLY_DATA
 	/** Appends string representations of this Id to a key string. */
-	NIAGARA_API void AppendKeyString(FString& KeyString, const FString& Delimiter = TEXT("_"), bool bAppendObjectForDebugging = false) const;
-
+	NIAGARA_API void AppendKeyString(FString& KeyString, const FString& Delimiter = TEXT("_"), bool bAppendObjectForDebugging = false, const FNiagaraScriptHashCollector* HashCollector = nullptr) const;
 #endif
 };
 
@@ -996,9 +1007,10 @@ public:
 	TObjectPtr<class UNiagaraScriptSourceBase>	Source_DEPRECATED;
 
 	NIAGARA_API static const FName NiagaraCustomVersionTagName;
+
+	NIAGARA_API void ComputeVMCompilationId(FNiagaraVMExecutableDataId& Id, const FGuid& VersionGuid, FNiagaraScriptHashCollector* HashCollector = nullptr) const;
 #endif
 
-	NIAGARA_API void ComputeVMCompilationId(FNiagaraVMExecutableDataId& Id, const FGuid& VersionGuid) const;
 	const FNiagaraVMExecutableDataId& GetComputedVMCompilationId() const
 	{
 #if WITH_EDITORONLY_DATA
diff --git a/Engine/Plugins/FX/Niagara/Source/Niagara/Classes/NiagaraScriptSourceBase.h b/Engine/Plugins/FX/Niagara/Source/Niagara/Classes/NiagaraScriptSourceBase.h
--- a/Engine/Plugins/FX/Niagara/Source/Niagara/Classes/NiagaraScriptSourceBase.h
+++ b/Engine/Plugins/FX/Niagara/Source/Niagara/Classes/NiagaraScriptSourceBase.h
@@ -12,7 +12,10 @@
 
 class INiagaraParameterDefinitionsSubscriber;
 class UNiagaraParameterDefinitionsBase;
+struct FNiagaraParameterStore;
+struct FNiagaraScriptHashCollector;
 class UNiagaraScriptSourceBase;
+struct FNiagaraVMExecutableDataId;
 
 struct EditorExposedVectorConstant
 {
@@ -93,8 +96,6 @@ public:
 	TArray<FNiagaraVariableBase> AdditionalVariables;
 };
 
-struct FNiagaraParameterStore;
-
 /** Runtime data for a Niagara system */
 UCLASS(MinimalAPI)
 class UNiagaraScriptSourceBase : public UObject
@@ -118,8 +119,8 @@ class UNiagaraScriptSourceBase : public UObject
 
 	virtual FGuid GetChangeID() { return FGuid(); }
 
-	virtual void RegisterVMCompilationIdDependencies(struct FNiagaraVMExecutableDataId& Id, ENiagaraScriptUsage InUsage, const FGuid& InUsageId) const {};
-	virtual void ComputeVMCompilationId(struct FNiagaraVMExecutableDataId& Id, ENiagaraScriptUsage InUsage, const FGuid& InUsageId) const {};
+	virtual void RegisterVMCompilationIdDependencies(FNiagaraScriptHashCollector& Collector, ENiagaraScriptUsage InUsage, const FGuid& InUsageId) const {};
+	virtual void ComputeVMCompilationId(FNiagaraVMExecutableDataId& Id, FNiagaraScriptHashCollector& HashCollector, ENiagaraScriptUsage InUsage, const FGuid& InUsageId) const {};
 
 #if WITH_EDITORONLY_DATA
 	// Will conditionally refresh the graph's CompileId
diff --git a/Engine/Plugins/FX/Niagara/Source/Niagara/Private/NiagaraScript.cpp b/Engine/Plugins/FX/Niagara/Source/Niagara/Private/NiagaraScript.cpp
--- a/Engine/Plugins/FX/Niagara/Source/Niagara/Private/NiagaraScript.cpp
+++ b/Engine/Plugins/FX/Niagara/Source/Niagara/Private/NiagaraScript.cpp
@@ -218,6 +218,38 @@ FNiagaraScriptDebuggerInfo::FNiagaraScriptDebuggerInfo(FName InName, ENiagaraScr
 	}
 }
 
+#if WITH_EDITORONLY_DATA
+FNiagaraScriptHashCollector::FNiagaraScriptHashCollector(bool bCollectHashSources)
+	: bCollectSources(bCollectHashSources)
+{
+}
+
+void FNiagaraScriptHashCollector::AddHash(const FNiagaraCompileHash& CompileHash, FStringView CompileHashSource)
+{
+	ReferencedHashes.AddUnique(CompileHash);
+	if (bCollectSources)
+	{
+		ReferencedHashSources.FindOrAdd(CompileHash).Emplace(CompileHashSource);
+	}
+}
+
+FString FNiagaraScriptHashCollector::BuildCompileHashSourceString(const FNiagaraCompileHash& CompileHash) const
+{
+	FString SourceString;
+
+	for (const FString& HashSource : ReferencedHashSources.FindRef(CompileHash))
+	{
+		if (!SourceString.IsEmpty())
+		{
+			SourceString.Append(TEXT(", "));
+		}
+
+		SourceString.Append(HashSource);
+	}
+
+	return SourceString;
+}
+#endif
 
 UNiagaraScriptSourceBase::UNiagaraScriptSourceBase(const FObjectInitializer& ObjectInitializer)
 	: Super(ObjectInitializer)
@@ -881,7 +913,7 @@ bool FNiagaraVMExecutableDataId::operator==(const FNiagaraVMExecutableDataId& Re
 }
 
 #if WITH_EDITORONLY_DATA
-void FNiagaraVMExecutableDataId::AppendKeyString(FString& KeyString, const FString& Delimiter, bool bAppendObjectForDebugging) const
+void FNiagaraVMExecutableDataId::AppendKeyString(FString& KeyString, const FString& Delimiter, bool bAppendObjectForDebugging, const FNiagaraScriptHashCollector* HashCollector) const
 {
 	KeyString += FString::Printf(TEXT("%d%s"), (int32)ScriptUsageType, *Delimiter);
 	KeyString += ScriptUsageTypeID.ToString();
@@ -952,9 +984,9 @@ void FNiagaraVMExecutableDataId::AppendKeyString(FString& KeyString, const FStri
 	{
 		KeyString += ReferencedCompileHashes[HashIndex].ToString();
 
-		if (bAppendObjectForDebugging && DebugReferencedObjects.Num() > HashIndex)
+		if (bAppendObjectForDebugging && HashCollector)
 		{
-			KeyString += TEXT(" [") + DebugReferencedObjects[HashIndex] + TEXT("]") ;
+			KeyString += TEXT(" [") + HashCollector->BuildCompileHashSourceString(ReferencedCompileHashes[HashIndex]) + TEXT("]");
 		}
 
 		if (HashIndex < ReferencedCompileHashes.Num() - 1)
@@ -1055,7 +1087,7 @@ FString UNiagaraScript::GetNiagaraDDCKeyString(const FGuid& ScriptVersion, const
 	return BuildNiagaraDDCKeyString(GetLastGeneratedVMId(ScriptVersion), ScriptPath);
 }
 
-void UNiagaraScript::ComputeVMCompilationId(FNiagaraVMExecutableDataId& Id, const FGuid& VersionGuid) const
+void UNiagaraScript::ComputeVMCompilationId(FNiagaraVMExecutableDataId& Id, const FGuid& VersionGuid, FNiagaraScriptHashCollector* OutHashCollector) const
 {
 	Id = FNiagaraVMExecutableDataId();
 
@@ -1076,6 +1108,8 @@ void UNiagaraScript::ComputeVMCompilationId(FNiagaraVMExecutableDataId& Id, cons
 	int32 SystemUpdateIdx = INDEX_NONE;
 	TArray<UNiagaraScript*> Scripts;
 
+	FNiagaraScriptHashCollector HashCollector(GNiagaraDumpKeyGen == 1 || (OutHashCollector ? OutHashCollector->bCollectSources : false));
+
 	if (EmitterData != nullptr)
 	{
 		UNiagaraEmitter* Emitter = Outer.Emitter;
@@ -1251,8 +1285,7 @@ void UNiagaraScript::ComputeVMCompilationId(FNiagaraVMExecutableDataId& Id, cons
 		// Append all Shader Parameter compile hashes
 		if (GNiagaraCompileHashAllDataInterfaces)
 		{
-			Id.ReferencedCompileHashes.AddUnique(NiagaraScriptInternal::CompileHashAllDataInterfaces());
-			Id.DebugReferencedObjects.Add(TEXT("AllDataInterfaceHashes"));
+			HashCollector.AddHash(NiagaraScriptInternal::CompileHashAllDataInterfaces(), TEXT("AllDataInterfaceHashes"));
 		}
 
 		// Has simulation stages
@@ -1275,8 +1308,7 @@ void UNiagaraScript::ComputeVMCompilationId(FNiagaraVMExecutableDataId& Id, cons
 			HashState.GetHash(DataHash.GetData());
 
 			FNiagaraCompileHash Hash(DataHash);
-			Id.ReferencedCompileHashes.AddUnique(Hash);
-			Id.DebugReferencedObjects.Add(TEXT("SimulationStageHeaders"));
+			HashCollector.AddHash(FNiagaraCompileHash(DataHash), TEXT("SimulationStageHeaders"));
 		}
 
 		if (IsParticleEventScript(Usage))
@@ -1369,7 +1401,7 @@ void UNiagaraScript::ComputeVMCompilationId(FNiagaraVMExecutableDataId& Id, cons
 		// We assume that latest source is ok because these are all top-level scripts in either systems or emitters.
 		if (Script->GetLatestSource())
 		{
-			Script->GetLatestSource()->RegisterVMCompilationIdDependencies(Id, Script->Usage, Script->UsageId);
+			Script->GetLatestSource()->RegisterVMCompilationIdDependencies(HashCollector, Script->Usage, Script->UsageId);
 		}
 
 		for (const FNiagaraVariable& Var : Vars)
@@ -1465,28 +1497,33 @@ void UNiagaraScript::ComputeVMCompilationId(FNiagaraVMExecutableDataId& Id, cons
 		DataHash.AddUninitialized(FSHA1::DigestSize);
 		HashState.GetHash(DataHash.GetData());
 
-		FNiagaraCompileHash Hash(DataHash);
-		Id.ReferencedCompileHashes.AddUnique(Hash);
-		Id.DebugReferencedObjects.Add(TEXT("RIParams"));
+		HashCollector.AddHash(FNiagaraCompileHash(DataHash), TEXT("RIParams"));
 	}
 
 	if (const FVersionedNiagaraScriptData* ScriptData = GetScriptData(Id.ScriptVersionID))
 	{
 		if (ScriptData->Source)
 		{
-			ScriptData->Source->ComputeVMCompilationId(Id, Usage, UsageId);
+			ScriptData->Source->ComputeVMCompilationId(Id, HashCollector, Usage, UsageId);
 		}
 	}
 
 	// Append the state of the "Fail If Not Set severity" cvar, as it may affect the LastCompileEvents which is a member of the script's VMExecutableData.
 	Id.AdditionalDefines.Add(FString::Printf(TEXT("FailIfNotSetMessageSeverity: %d"), GNiagaraTranslatorFailIfNotSetSeverity));
+	Id.ReferencedCompileHashes = MoveTemp(HashCollector.ReferencedHashes);
 
 	FNiagaraVMExecutableDataId& LastGeneratedVMId = GetLastGeneratedVMId(VersionGuid);
+
+	// normalize the order
+	Id.AdditionalDefines.Sort();
+	Id.ReferencedCompileHashes.Sort();
+	Id.AdditionalVariables.Sort([](const FNiagaraVariableBase& Lhs, const FNiagaraVariableBase& Rhs) { return Lhs.GetName().LexicalLess(Rhs.GetName()); });
+
 	if (GNiagaraDumpKeyGen == 1 && Id != LastGeneratedVMId)
 	{
 		TArray<FString> OutputByLines;
 		FString StrDump;
-		Id.AppendKeyString(StrDump, TEXT("\n"), true);
+		Id.AppendKeyString(StrDump, TEXT("\n"), true, &HashCollector);
 		StrDump.ParseIntoArrayLines(OutputByLines, false);
 
 		UE_LOG(LogNiagara, Display, TEXT("KeyGen %s\n==================\n"), *GetPathName());
@@ -1496,10 +1533,10 @@ void UNiagaraScript::ComputeVMCompilationId(FNiagaraVMExecutableDataId& Id, cons
 		}
 	}
 
-	// normalize the order
-	Id.AdditionalDefines.Sort();
-	Id.ReferencedCompileHashes.Sort();
-	Id.AdditionalVariables.Sort([](const FNiagaraVariableBase& Lhs, const FNiagaraVariableBase& Rhs) { return Lhs.GetName().LexicalLess(Rhs.GetName()); });
+	if (OutHashCollector && OutHashCollector->bCollectSources)
+	{
+		OutHashCollector->ReferencedHashSources = MoveTemp(HashCollector.ReferencedHashSources);
+	}
 
 	LastGeneratedVMId = Id;
 }
@@ -2576,7 +2613,9 @@ bool UNiagaraScript::AreScriptAndSourceSynchronized(const FGuid& VersionGuid) co
 	if (ScriptData && ScriptData->Source)
 	{
 		FNiagaraVMExecutableDataId NewId;
-		ComputeVMCompilationId(NewId, VersionGuid);
+		FNiagaraScriptHashCollector HashCollector(GEnableVerboseNiagaraChangeIdLogging > 0);
+		ComputeVMCompilationId(NewId, VersionGuid, &HashCollector);
+
 		bool bSynchronized = (NewId.IsValid() && NewId == CachedScriptVMId);
 		if (!bSynchronized && NewId.IsValid() && CachedScriptVMId.IsValid() && CachedScriptVM.IsValid())
 		{
@@ -2603,7 +2642,7 @@ bool UNiagaraScript::AreScriptAndSourceSynchronized(const FGuid& VersionGuid) co
 							{
 								UE_LOG(LogNiagara, Log, TEXT("AreScriptAndSourceSynchronized referenced compile hash %d doesn't match. %s != %s, script %s, source %s"),
 									i, *NewId.ReferencedCompileHashes[i].ToString(), *CachedScriptVMId.ReferencedCompileHashes[i].ToString(), *GetPathName(),
-									*NewId.DebugReferencedObjects[i]);
+									*HashCollector.BuildCompileHashSourceString(NewId.ReferencedCompileHashes[i]));
 							}
 						}
 					}
@@ -3025,15 +3064,10 @@ void UNiagaraScript::SetVMCompilationResults(const FNiagaraVMExecutableDataId& I
 
 	if (bApplyRapidIterationParameters)
 	{
-		const bool bClearBindings = false;
-		RapidIterationParameters.Empty(bClearBindings);
-		for (const FNiagaraVariable& Parameter : InScriptVM.BakedRapidIterationParameters)
-		{
-			const bool bInitialize = false;
-			const bool bTriggerRebind = false;
-			RapidIterationParameters.AddParameter(Parameter, bInitialize, bTriggerRebind);
-		}
-		RapidIterationParameters.TriggerOnLayoutChanged();
+		// if we are applying rapid iteration parameters then the data supplied by InScriptVM, which could have
+		// come from the DDC, could invalidate the CompileID that we had previous generated so generate it again
+		// now that we've updated everything
+		ComputeVMCompilationId(CachedScriptVMId, FGuid());
 	}
 
 	GenerateStatIDs();
diff --git a/Engine/Plugins/FX/Niagara/Source/NiagaraEditor/Private/NiagaraGraph.cpp b/Engine/Plugins/FX/Niagara/Source/NiagaraEditor/Private/NiagaraGraph.cpp
--- a/Engine/Plugins/FX/Niagara/Source/NiagaraEditor/Private/NiagaraGraph.cpp
+++ b/Engine/Plugins/FX/Niagara/Source/NiagaraEditor/Private/NiagaraGraph.cpp
@@ -3765,7 +3765,7 @@ void UNiagaraGraph::ForceGraphToRecompileOnNextCheck()
 	MarkGraphRequiresSynchronization(__FUNCTION__);
 }
 
-void UNiagaraGraph::GatherExternalDependencyData(ENiagaraScriptUsage InUsage, const FGuid& InUsageId, TArray<FNiagaraCompileHash>& InReferencedCompileHashes, TArray<FString>& InReferencedObjs)
+void UNiagaraGraph::GatherExternalDependencyData(ENiagaraScriptUsage InUsage, const FGuid& InUsageId, FNiagaraScriptHashCollector& HashCollector)
 {
 	RebuildCachedCompileIds();
 	
@@ -3776,7 +3776,7 @@ void UNiagaraGraph::GatherExternalDependencyData(ENiagaraScriptUsage InUsage, co
 		{
 			for (UNiagaraNode* Node : CachedUsageInfo[i].Traversal)
 			{
-				Node->GatherExternalDependencyData(InUsage, InUsageId, InReferencedCompileHashes, InReferencedObjs);
+				Node->GatherExternalDependencyData(InUsage, InUsageId, HashCollector);
 			}
 		}
 		// Now add any other dependency chains that we might have...
@@ -3784,17 +3784,16 @@ void UNiagaraGraph::GatherExternalDependencyData(ENiagaraScriptUsage InUsage, co
 		{
 			if (GNiagaraUseGraphHash == 1)
 			{
-				InReferencedCompileHashes.AddUnique(CachedUsageInfo[i].CompileHashFromGraph);
+				HashCollector.AddHash(CachedUsageInfo[i].CompileHashFromGraph, CachedUsageInfo[i].Traversal.Last()->GetPathName());
 			}
 			else
 			{
-				InReferencedCompileHashes.AddUnique(CachedUsageInfo[i].CompileHash);
+				HashCollector.AddHash(CachedUsageInfo[i].CompileHash, CachedUsageInfo[i].Traversal.Last()->GetPathName());
 			}
-			InReferencedObjs.Add(CachedUsageInfo[i].Traversal.Last()->GetPathName());
 
 			for (UNiagaraNode* Node : CachedUsageInfo[i].Traversal)
 			{
-				Node->GatherExternalDependencyData(InUsage, InUsageId, InReferencedCompileHashes, InReferencedObjs);
+				Node->GatherExternalDependencyData(InUsage, InUsageId, HashCollector);
 			}
 		}
 	}
diff --git a/Engine/Plugins/FX/Niagara/Source/NiagaraEditor/Private/NiagaraNodeAssignment.cpp b/Engine/Plugins/FX/Niagara/Source/NiagaraEditor/Private/NiagaraNodeAssignment.cpp
--- a/Engine/Plugins/FX/Niagara/Source/NiagaraEditor/Private/NiagaraNodeAssignment.cpp
+++ b/Engine/Plugins/FX/Niagara/Source/NiagaraEditor/Private/NiagaraNodeAssignment.cpp
@@ -333,7 +333,7 @@ void UNiagaraNodeAssignment::BuildParameterMapHistory(FNiagaraParameterMapHistor
 	Super::BuildParameterMapHistory(OutHistory, bRecursive, bFilterForCompilation);
 }
 
-void UNiagaraNodeAssignment::GatherExternalDependencyData(ENiagaraScriptUsage InUsage, const FGuid& InUsageId, TArray<FNiagaraCompileHash>& InReferencedCompileHashes, TArray<FString>& InReferencedObjs) const
+void UNiagaraNodeAssignment::GatherExternalDependencyData(ENiagaraScriptUsage InUsage, const FGuid& InUsageId, FNiagaraScriptHashCollector& HashCollector) const
 {
 	// Assignment nodes own their function graphs and therefore have no external dependencies so we override the default function behavior here to avoid 
 	// adding additional non-deterministic guids to the compile id generation which can invalid the DDC for compiled scripts, especially during emitter merging.
diff --git a/Engine/Plugins/FX/Niagara/Source/NiagaraEditor/Private/NiagaraNodeCustomHlsl.cpp b/Engine/Plugins/FX/Niagara/Source/NiagaraEditor/Private/NiagaraNodeCustomHlsl.cpp
--- a/Engine/Plugins/FX/Niagara/Source/NiagaraEditor/Private/NiagaraNodeCustomHlsl.cpp
+++ b/Engine/Plugins/FX/Niagara/Source/NiagaraEditor/Private/NiagaraNodeCustomHlsl.cpp
@@ -559,7 +559,7 @@ FNiagaraCompileHash GetFileContentHash(const FString& FileContents)
 	return FNiagaraCompileHash(DataHash);
 }
 
-void UNiagaraNodeCustomHlsl::GatherExternalDependencyData(ENiagaraScriptUsage InUsage, const FGuid& InUsageId, TArray<FNiagaraCompileHash>& InReferencedCompileHashes, TArray<FString>& InReferencedObjs) const
+void UNiagaraNodeCustomHlsl::GatherExternalDependencyData(ENiagaraScriptUsage InUsage, const FGuid& InUsageId, FNiagaraScriptHashCollector& HashCollector) const
 {
 	for (const auto& [IncludePath] : AbsoluteIncludeFilePaths)
 	{
@@ -570,8 +570,7 @@ void UNiagaraNodeCustomHlsl::GatherExternalDependencyData(ENiagaraScriptUsage In
 
 		if (FString FileContents; FFileHelper::LoadFileToString(FileContents, *IncludePath))
 		{
-			InReferencedObjs.AddUnique(IncludePath);
-			InReferencedCompileHashes.AddUnique(GetFileContentHash(FileContents));
+			HashCollector.AddHash(GetFileContentHash(FileContents), IncludePath);
 		}
 	}
 
@@ -585,8 +584,7 @@ void UNiagaraNodeCustomHlsl::GatherExternalDependencyData(ENiagaraScriptUsage In
 		FString FileContents;
 		if(LoadShaderSourceFile(*IncludePath, SP_PCD3D_SM5, &FileContents, nullptr))
 		{
-			InReferencedObjs.AddUnique(IncludePath);
-			InReferencedCompileHashes.AddUnique(GetFileContentHash(FileContents));
+			HashCollector.AddHash(GetFileContentHash(FileContents), IncludePath);
 		}
 	}
 }
diff --git a/Engine/Plugins/FX/Niagara/Source/NiagaraEditor/Private/NiagaraNodeEmitter.cpp b/Engine/Plugins/FX/Niagara/Source/NiagaraEditor/Private/NiagaraNodeEmitter.cpp
--- a/Engine/Plugins/FX/Niagara/Source/NiagaraEditor/Private/NiagaraNodeEmitter.cpp
+++ b/Engine/Plugins/FX/Niagara/Source/NiagaraEditor/Private/NiagaraNodeEmitter.cpp
@@ -528,7 +528,7 @@ void UNiagaraNodeEmitter::Compile(FTranslator* Translator, TArray<int32>& Output
 	}
 }
 
-void UNiagaraNodeEmitter::GatherExternalDependencyData(ENiagaraScriptUsage InUsage, const FGuid& InUsageId, TArray<FNiagaraCompileHash>& InReferencedCompileHashes, TArray<FString>& InReferencedObjs) const
+void UNiagaraNodeEmitter::GatherExternalDependencyData(ENiagaraScriptUsage InUsage, const FGuid& InUsageId, FNiagaraScriptHashCollector& HashCollector) const
 {
 	UNiagaraGraph* CalledGraph = GetCalledGraph();
 
@@ -537,9 +537,8 @@ void UNiagaraNodeEmitter::GatherExternalDependencyData(ENiagaraScriptUsage InUsa
 		CalledGraph->RebuildCachedCompileIds();
 		ENiagaraScriptUsage TargetUsage = InUsage == ENiagaraScriptUsage::SystemSpawnScript ? ENiagaraScriptUsage::EmitterSpawnScript : ENiagaraScriptUsage::EmitterUpdateScript;
 		FNiagaraCompileHash Hash = CalledGraph->GetCompileDataHash(TargetUsage, FGuid(0,0,0,0));
-		InReferencedCompileHashes.AddUnique(Hash);
-		InReferencedObjs.Add(CalledGraph->GetPathName());
-		CalledGraph->GatherExternalDependencyData(TargetUsage, FGuid(0, 0, 0, 0), InReferencedCompileHashes, InReferencedObjs);
+		HashCollector.AddHash(Hash, CalledGraph->GetPathName());
+		CalledGraph->GatherExternalDependencyData(TargetUsage, FGuid(0, 0, 0, 0), HashCollector);
 	}
 }
 
diff --git a/Engine/Plugins/FX/Niagara/Source/NiagaraEditor/Private/NiagaraNodeEmitter.h b/Engine/Plugins/FX/Niagara/Source/NiagaraEditor/Private/NiagaraNodeEmitter.h
--- a/Engine/Plugins/FX/Niagara/Source/NiagaraEditor/Private/NiagaraNodeEmitter.h
+++ b/Engine/Plugins/FX/Niagara/Source/NiagaraEditor/Private/NiagaraNodeEmitter.h
@@ -61,7 +61,7 @@ public:
 	UNiagaraGraph* GetCalledGraph() const;
 	
 	virtual void Compile(FTranslator* Translator, TArray<int32>& Outputs) const override;
-	virtual void GatherExternalDependencyData(ENiagaraScriptUsage InUsage, const FGuid& InUsageId, TArray<FNiagaraCompileHash>& InReferencedCompileHashes, TArray<FString>& InReferencedObjs) const override;
+	virtual void GatherExternalDependencyData(ENiagaraScriptUsage InUsage, const FGuid& InUsageId, FNiagaraScriptHashCollector& HashCollector) const override;
 
 	void SetCachedVariablesForCompilation(const FName& InUniqueName, FNiagaraEmitterID InEmitterID, UNiagaraGraph* InGraph, UNiagaraScriptSourceBase* InSource);
 
diff --git a/Engine/Plugins/FX/Niagara/Source/NiagaraEditor/Private/NiagaraNodeFunctionCall.cpp b/Engine/Plugins/FX/Niagara/Source/NiagaraEditor/Private/NiagaraNodeFunctionCall.cpp
--- a/Engine/Plugins/FX/Niagara/Source/NiagaraEditor/Private/NiagaraNodeFunctionCall.cpp
+++ b/Engine/Plugins/FX/Niagara/Source/NiagaraEditor/Private/NiagaraNodeFunctionCall.cpp
@@ -1127,7 +1127,7 @@ bool UNiagaraNodeFunctionCall::RefreshFromExternalChanges()
 	}
 }
 
-void UNiagaraNodeFunctionCall::GatherExternalDependencyData(ENiagaraScriptUsage InUsage, const FGuid& InUsageId, TArray<FNiagaraCompileHash>& InReferencedCompileHashes, TArray<FString>& InReferencedObjs) const
+void UNiagaraNodeFunctionCall::GatherExternalDependencyData(ENiagaraScriptUsage InUsage, const FGuid& InUsageId, FNiagaraScriptHashCollector& HashCollector) const
 {
 	if (HasValidScriptAndGraph())
 	{
@@ -1140,9 +1140,8 @@ void UNiagaraNodeFunctionCall::GatherExternalDependencyData(ENiagaraScriptUsage
 			FNiagaraCompileHash FoundCompileHash = FunctionGraph->GetCompileDataHash((ENiagaraScriptUsage)i, FGuid(0, 0, 0, 0));
 			if (FoundGuid.IsValid() && FoundCompileHash.IsValid())
 			{
-				InReferencedCompileHashes.AddUnique(FoundCompileHash);
-				InReferencedObjs.Add(FunctionGraph->GetPathName());
-				FunctionGraph->GatherExternalDependencyData((ENiagaraScriptUsage)i, FGuid(0, 0, 0, 0), InReferencedCompileHashes, InReferencedObjs);
+				HashCollector.AddHash(FoundCompileHash, FunctionGraph->GetPathName());
+				FunctionGraph->GatherExternalDependencyData((ENiagaraScriptUsage)i, FGuid(0, 0, 0, 0), HashCollector);
 			}
 		}
 	}
diff --git a/Engine/Plugins/FX/Niagara/Source/NiagaraEditor/Private/NiagaraNodeParameterMapGet.cpp b/Engine/Plugins/FX/Niagara/Source/NiagaraEditor/Private/NiagaraNodeParameterMapGet.cpp
--- a/Engine/Plugins/FX/Niagara/Source/NiagaraEditor/Private/NiagaraNodeParameterMapGet.cpp
+++ b/Engine/Plugins/FX/Niagara/Source/NiagaraEditor/Private/NiagaraNodeParameterMapGet.cpp
@@ -606,7 +606,7 @@ bool UNiagaraNodeParameterMapGet::CommitEditablePinName(const FText& InName, UEd
 	return false;
 }
 
-void UNiagaraNodeParameterMapGet::GatherExternalDependencyData(ENiagaraScriptUsage InUsage, const FGuid& InUsageId, TArray<FNiagaraCompileHash>& InReferencedCompileHashes, TArray<FString>& InReferencedObjs) const
+void UNiagaraNodeParameterMapGet::GatherExternalDependencyData(ENiagaraScriptUsage InUsage, const FGuid& InUsageId, FNiagaraScriptHashCollector& HashCollector) const
 {
 	// If we are referencing any parameter collections, we need to register them here... might want to speeed this up in the future 
 	// by caching any parameter collections locally.
@@ -635,8 +635,7 @@ void UNiagaraNodeParameterMapGet::GatherExternalDependencyData(ENiagaraScriptUsa
 		if (Collection)
 		{
 			FNiagaraCompileHash Hash = Collection->GetCompileHash();
-			InReferencedCompileHashes.AddUnique(Hash);
-			InReferencedObjs.Add(Collection->GetPathName());
+			HashCollector.AddHash(Hash, Collection->GetPathName());
 		}
 	}
 }
diff --git a/Engine/Plugins/FX/Niagara/Source/NiagaraEditor/Private/NiagaraNodeParameterMapGet.h b/Engine/Plugins/FX/Niagara/Source/NiagaraEditor/Private/NiagaraNodeParameterMapGet.h
--- a/Engine/Plugins/FX/Niagara/Source/NiagaraEditor/Private/NiagaraNodeParameterMapGet.h
+++ b/Engine/Plugins/FX/Niagara/Source/NiagaraEditor/Private/NiagaraNodeParameterMapGet.h
@@ -43,7 +43,7 @@ public:
 	
 	virtual void PostLoad() override;
 
-	virtual void GatherExternalDependencyData(ENiagaraScriptUsage InUsage, const FGuid& InUsageId, TArray<FNiagaraCompileHash>& InReferencedCompileHashes, TArray<FString>& InReferencedObjs) const override;
+	virtual void GatherExternalDependencyData(ENiagaraScriptUsage InUsage, const FGuid& InUsageId, FNiagaraScriptHashCollector& HashCollector) const override;
 	virtual void GetPinHoverText(const UEdGraphPin& Pin, FString& HoverTextOut) const override;
 
 	virtual FName GetNewPinDefaultNamespace() const override { return PARAM_MAP_MODULE_STR; }
diff --git a/Engine/Plugins/FX/Niagara/Source/NiagaraEditor/Private/NiagaraScriptSource.cpp b/Engine/Plugins/FX/Niagara/Source/NiagaraEditor/Private/NiagaraScriptSource.cpp
--- a/Engine/Plugins/FX/Niagara/Source/NiagaraEditor/Private/NiagaraScriptSource.cpp
+++ b/Engine/Plugins/FX/Niagara/Source/NiagaraEditor/Private/NiagaraScriptSource.cpp
@@ -32,7 +32,7 @@ UNiagaraScriptSource::UNiagaraScriptSource(const FObjectInitializer& ObjectIniti
 {
 }
 
-void UNiagaraScriptSource::RegisterVMCompilationIdDependencies(struct FNiagaraVMExecutableDataId& Id, ENiagaraScriptUsage InUsage, const FGuid& InUsageId) const
+void UNiagaraScriptSource::RegisterVMCompilationIdDependencies(FNiagaraScriptHashCollector& Collector, ENiagaraScriptUsage InUsage, const FGuid& InUsageId) const
 {
 	if (!AllowShaderCompiling())
 	{
@@ -53,8 +53,7 @@ void UNiagaraScriptSource::RegisterVMCompilationIdDependencies(struct FNiagaraVM
 			
 			if (Hash.IsValid())
 			{
-				Id.ReferencedCompileHashes.AddUnique(Hash);
-				Id.DebugReferencedObjects.Emplace(GetPathName());
+				Collector.AddHash(Hash, GetPathName());
 			}
 			else
 			{
@@ -65,7 +64,7 @@ void UNiagaraScriptSource::RegisterVMCompilationIdDependencies(struct FNiagaraVM
 }
 
 
-void UNiagaraScriptSource::ComputeVMCompilationId(FNiagaraVMExecutableDataId& Id, ENiagaraScriptUsage InUsage, const FGuid& InUsageId) const
+void UNiagaraScriptSource::ComputeVMCompilationId(FNiagaraVMExecutableDataId& Id, FNiagaraScriptHashCollector& HashCollector, ENiagaraScriptUsage InUsage, const FGuid& InUsageId) const
 {
 	if (!AllowShaderCompiling())
 	{
@@ -79,19 +78,17 @@ void UNiagaraScriptSource::ComputeVMCompilationId(FNiagaraVMExecutableDataId& Id
 	{
 		NodeGraph->RebuildCachedCompileIds();
 		Id.BaseScriptCompileHash = FNiagaraCompileHash(NodeGraph->GetCompileDataHash(InUsage, InUsageId));
-		NodeGraph->GatherExternalDependencyData(InUsage, InUsageId, Id.ReferencedCompileHashes, Id.DebugReferencedObjects);
+
+		NodeGraph->GatherExternalDependencyData(InUsage, InUsageId, HashCollector);
 	}
 
 	// Add in any referenced HLSL files.
 	FSHAHash Hash = GetShaderFileHash((TEXT("/Plugin/FX/Niagara/Private/NiagaraEmitterInstanceShader.usf")), EShaderPlatform::SP_PCD3D_SM5);
-	Id.ReferencedCompileHashes.AddUnique(FNiagaraCompileHash(Hash.Hash, sizeof(Hash.Hash)/sizeof(uint8)));
-	Id.DebugReferencedObjects.Emplace(TEXT("/Plugin/FX/Niagara/Private/NiagaraEmitterInstanceShader.usf"));
+	HashCollector.AddHash(FNiagaraCompileHash(Hash.Hash, sizeof(Hash.Hash) / sizeof(uint8)), TEXT("/Plugin/FX/Niagara/Private/NiagaraEmitterInstanceShader.usf"));
 	Hash = GetShaderFileHash((TEXT("/Plugin/FX/Niagara/Private/NiagaraShaderVersion.ush")), EShaderPlatform::SP_PCD3D_SM5);
-	Id.ReferencedCompileHashes.AddUnique(FNiagaraCompileHash(Hash.Hash, sizeof(Hash.Hash)/sizeof(uint8)));
-	Id.DebugReferencedObjects.Emplace(TEXT("/Plugin/FX/Niagara/Private/NiagaraShaderVersion.ush"));
+	HashCollector.AddHash(FNiagaraCompileHash(Hash.Hash, sizeof(Hash.Hash) / sizeof(uint8)), TEXT("/Plugin/FX/Niagara/Private/NiagaraShaderVersion.ush"));
 	Hash = GetShaderFileHash((TEXT("/Engine/Public/ShaderVersion.ush")), EShaderPlatform::SP_PCD3D_SM5);
-	Id.ReferencedCompileHashes.AddUnique(FNiagaraCompileHash(Hash.Hash, sizeof(Hash.Hash)/sizeof(uint8)));
-	Id.DebugReferencedObjects.Emplace(TEXT("/Engine/Public/ShaderVersion.ush"));
+	HashCollector.AddHash(FNiagaraCompileHash(Hash.Hash, sizeof(Hash.Hash) / sizeof(uint8)), TEXT("/Engine/Public/ShaderVersion.ush"));
 }
 
 void UNiagaraScriptSource::RefreshGraphCompileId()
diff --git a/Engine/Plugins/FX/Niagara/Source/NiagaraEditor/Public/NiagaraGraph.h b/Engine/Plugins/FX/Niagara/Source/NiagaraEditor/Public/NiagaraGraph.h
--- a/Engine/Plugins/FX/Niagara/Source/NiagaraEditor/Public/NiagaraGraph.h
+++ b/Engine/Plugins/FX/Niagara/Source/NiagaraEditor/Public/NiagaraGraph.h
@@ -252,7 +252,7 @@ class UNiagaraGraph : public UEdGraph
 	void GetAllReferencedGraphs(TArray<const UNiagaraGraph*>& Graphs) const;
 
 	/** Gather all the change ids of external references for this specific graph traversal.*/
-	void GatherExternalDependencyData(ENiagaraScriptUsage InUsage, const FGuid& InUsageId, TArray<FNiagaraCompileHash>& InReferencedCompileHashes, TArray<FString>& InReferencedObjs);
+	void GatherExternalDependencyData(ENiagaraScriptUsage InUsage, const FGuid& InUsageId, FNiagaraScriptHashCollector& HashCollector);
 
 	/** Determine if another item has been synchronized with this graph.*/
 	bool IsOtherSynchronized(const FGuid& InChangeId) const;
diff --git a/Engine/Plugins/FX/Niagara/Source/NiagaraEditor/Public/NiagaraNode.h b/Engine/Plugins/FX/Niagara/Source/NiagaraEditor/Public/NiagaraNode.h
--- a/Engine/Plugins/FX/Niagara/Source/NiagaraEditor/Public/NiagaraNode.h
+++ b/Engine/Plugins/FX/Niagara/Source/NiagaraEditor/Public/NiagaraNode.h
@@ -155,7 +155,7 @@ public:
 	NIAGARAEDITOR_API virtual void BuildParameterMapHistory(FNiagaraParameterMapHistoryBuilder& OutHistory, bool bRecursive = true, bool bFilterForCompilation = true) const;
 	
 	/** Go through all the external dependencies of this node in isolation and add them to the reference id list.*/
-	virtual void GatherExternalDependencyData(ENiagaraScriptUsage InUsage, const FGuid& InUsageId, TArray<FNiagaraCompileHash>& InReferencedCompileHashes, TArray<FString>& InReferencedObjs) const {};
+	virtual void GatherExternalDependencyData(ENiagaraScriptUsage InUsage, const FGuid& InUsageId, FNiagaraScriptHashCollector& HashCollector) const {};
 
 	/** Traces one of this node's output pins to its source output pin.
 	 *
diff --git a/Engine/Plugins/FX/Niagara/Source/NiagaraEditor/Public/NiagaraNodeAssignment.h b/Engine/Plugins/FX/Niagara/Source/NiagaraEditor/Public/NiagaraNodeAssignment.h
--- a/Engine/Plugins/FX/Niagara/Source/NiagaraEditor/Public/NiagaraNodeAssignment.h
+++ b/Engine/Plugins/FX/Niagara/Source/NiagaraEditor/Public/NiagaraNodeAssignment.h
@@ -48,7 +48,7 @@ public:
 	virtual bool AllowDynamicPins() const override { return true; }
 	NIAGARAEDITOR_API virtual bool RefreshFromExternalChanges() override;
 	NIAGARAEDITOR_API virtual void BuildParameterMapHistory(FNiagaraParameterMapHistoryBuilder& OutHistory, bool bRecursive = true, bool bFilterForCompilation = true) const override;
-	NIAGARAEDITOR_API virtual void GatherExternalDependencyData(ENiagaraScriptUsage InUsage, const FGuid& InUsageId, TArray<FNiagaraCompileHash>& InReferencedCompileHashes, TArray<FString>& InReferencedObjs) const override;
+	NIAGARAEDITOR_API virtual void GatherExternalDependencyData(ENiagaraScriptUsage InUsage, const FGuid& InUsageId, FNiagaraScriptHashCollector& HashCollector) const override;
 
 	NIAGARAEDITOR_API void AddParameter(FNiagaraVariable InVar, FString InDefaultValue);
 	NIAGARAEDITOR_API void RemoveParameter(const FNiagaraVariable& InVar);
diff --git a/Engine/Plugins/FX/Niagara/Source/NiagaraEditor/Public/NiagaraNodeCustomHlsl.h b/Engine/Plugins/FX/Niagara/Source/NiagaraEditor/Public/NiagaraNodeCustomHlsl.h
--- a/Engine/Plugins/FX/Niagara/Source/NiagaraEditor/Public/NiagaraNodeCustomHlsl.h
+++ b/Engine/Plugins/FX/Niagara/Source/NiagaraEditor/Public/NiagaraNodeCustomHlsl.h
@@ -38,7 +38,7 @@ public:
 	bool GetTokens(TArray<FStringView>& OutTokens, bool IncludeComments = true, bool IncludeWhitespace = true) const;
 
 	virtual void BuildParameterMapHistory(FNiagaraParameterMapHistoryBuilder& OutHistory, bool bRecursive = true, bool bFilterForCompilation = true) const override;
-	virtual void GatherExternalDependencyData(ENiagaraScriptUsage InUsage, const FGuid& InUsageId, TArray<FNiagaraCompileHash>& InReferencedCompileHashes, TArray<FString>& InReferencedObjs) const override;
+	virtual void GatherExternalDependencyData(ENiagaraScriptUsage InUsage, const FGuid& InUsageId, FNiagaraScriptHashCollector& HashCollector) const override;
 
 	// Replace items in the tokens array if they start with the src string or optionally src string and a namespace delimiter
 	static uint32 ReplaceExactMatchTokens(TArray<FString>& Tokens, FStringView SrcString, FStringView ReplaceString, bool bAllowNamespaceSeparation);
diff --git a/Engine/Plugins/FX/Niagara/Source/NiagaraEditor/Public/NiagaraNodeFunctionCall.h b/Engine/Plugins/FX/Niagara/Source/NiagaraEditor/Public/NiagaraNodeFunctionCall.h
--- a/Engine/Plugins/FX/Niagara/Source/NiagaraEditor/Public/NiagaraNodeFunctionCall.h
+++ b/Engine/Plugins/FX/Niagara/Source/NiagaraEditor/Public/NiagaraNodeFunctionCall.h
@@ -112,7 +112,7 @@ public:
 	virtual bool RefreshFromExternalChanges() override;
 	virtual ENiagaraNumericOutputTypeSelectionMode GetNumericOutputTypeSelectionMode() const override;
 	virtual bool CanAddToGraph(UNiagaraGraph* TargetGraph, FString& OutErrorMsg) const override;
-	virtual void GatherExternalDependencyData(ENiagaraScriptUsage InUsage, const FGuid& InUsageId, TArray<FNiagaraCompileHash>& InReferencedCompileHashes, TArray<FString>& InReferencedObjs) const override;
+	virtual void GatherExternalDependencyData(ENiagaraScriptUsage InUsage, const FGuid& InUsageId, FNiagaraScriptHashCollector& HashCollector) const override;
 	virtual void UpdateCompileHashForNode(FSHA1& HashState) const override;
 	virtual void GetNodeContextMenuActions(class UToolMenu* Menu, class UGraphNodeContextMenuContext* Context) const override;
 	virtual TSharedRef<SWidget> CreateTitleRightWidget() override;
diff --git a/Engine/Plugins/FX/Niagara/Source/NiagaraEditor/Public/NiagaraScriptSource.h b/Engine/Plugins/FX/Niagara/Source/NiagaraEditor/Public/NiagaraScriptSource.h
--- a/Engine/Plugins/FX/Niagara/Source/NiagaraEditor/Public/NiagaraScriptSource.h
+++ b/Engine/Plugins/FX/Niagara/Source/NiagaraEditor/Public/NiagaraScriptSource.h
@@ -11,6 +11,9 @@
 #include "GraphEditAction.h"
 #include "NiagaraScriptSource.generated.h"
 
+struct FNiagaraScriptHashCollector;
+struct FNiagaraVMExecutableDataId;
+
 UCLASS(MinimalAPI)
 class UNiagaraScriptSource : public UNiagaraScriptSourceBase
 {
@@ -38,8 +41,8 @@ class UNiagaraScriptSource : public UNiagaraScriptSourceBase
 	virtual FGuid GetChangeID() override;
 	FVersionedNiagaraEmitter GetOuterEmitter() const;
 
-	virtual void RegisterVMCompilationIdDependencies(struct FNiagaraVMExecutableDataId& Id, ENiagaraScriptUsage InUsage, const FGuid& InUsageId) const override;
-	virtual void ComputeVMCompilationId(struct FNiagaraVMExecutableDataId& Id, ENiagaraScriptUsage InUsage, const FGuid& InUsageId) const override;
+	virtual void RegisterVMCompilationIdDependencies(FNiagaraScriptHashCollector& Collector, ENiagaraScriptUsage InUsage, const FGuid& InUsageId) const override;
+	virtual void ComputeVMCompilationId(FNiagaraVMExecutableDataId& Id, FNiagaraScriptHashCollector& HashCollector, ENiagaraScriptUsage InUsage, const FGuid& InUsageId) const override;
 
 	// Will conditionally refresh the graph's CompileId
 	virtual void RefreshGraphCompileId() override;
