Remove inadvertent change from CL 30674610

diff --git a/Engine/Plugins/FX/Niagara/Source/Niagara/Private/NiagaraScript.cpp b/Engine/Plugins/FX/Niagara/Source/Niagara/Private/NiagaraScript.cpp
--- a/Engine/Plugins/FX/Niagara/Source/Niagara/Private/NiagaraScript.cpp
+++ b/Engine/Plugins/FX/Niagara/Source/Niagara/Private/NiagaraScript.cpp
@@ -3064,10 +3064,15 @@ void UNiagaraScript::SetVMCompilationResults(const FNiagaraVMExecutableDataId& I
 
 	if (bApplyRapidIterationParameters)
 	{
-		// if we are applying rapid iteration parameters then the data supplied by InScriptVM, which could have
-		// come from the DDC, could invalidate the CompileID that we had previous generated so generate it again
-		// now that we've updated everything
-		ComputeVMCompilationId(CachedScriptVMId, FGuid());
+		const bool bClearBindings = false;
+		RapidIterationParameters.Empty(bClearBindings);
+		for (const FNiagaraVariable& Parameter : InScriptVM.BakedRapidIterationParameters)
+		{
+			const bool bInitialize = false;
+			const bool bTriggerRebind = false;
+			RapidIterationParameters.AddParameter(Parameter, bInitialize, bTriggerRebind);
+		}
+		RapidIterationParameters.TriggerOnLayoutChanged();
 	}
 
 	GenerateStatIDs();
