Clear MeshBodyCounts when asset buffers are cleared/reset

diff --git a/Engine/Plugins/FX/Niagara/Source/Niagara/Private/NiagaraDataInterfaceRigidMeshCollisionQuery.cpp b/Engine/Plugins/FX/Niagara/Source/Niagara/Private/NiagaraDataInterfaceRigidMeshCollisionQuery.cpp
--- a/Engine/Plugins/FX/Niagara/Source/Niagara/Private/NiagaraDataInterfaceRigidMeshCollisionQuery.cpp
+++ b/Engine/Plugins/FX/Niagara/Source/Niagara/Private/NiagaraDataInterfaceRigidMeshCollisionQuery.cpp
@@ -796,6 +796,8 @@ void FNDIRigidMeshCollisionData::Init(int32 MaxNumPrimitives)
 		ReleaseBuffers();
 	}
 
+	MeshBodyCounts.Empty();
+
 	bRequiresFullUpdate = true;
 }
 
