Fixes for non-determinism for Niagara scripts.  Handles differences between getting the results from an actual compile and pulling hte results from the DDC.  In particular
-exposed variables are properly applied from the DDC
-rapid iteration parameters (including static variables) are now applied all together before SetVMCompilationResults is invoked since it ends pulling data from depenedent scripts
-non compilable scripts will now have their rapid iteration parameters (including static variables) to/from the DDC

diff --git a/Engine/Plugins/FX/Niagara/Source/Niagara/Classes/NiagaraScript.h b/Engine/Plugins/FX/Niagara/Source/Niagara/Classes/NiagaraScript.h
--- a/Engine/Plugins/FX/Niagara/Source/Niagara/Classes/NiagaraScript.h
+++ b/Engine/Plugins/FX/Niagara/Source/Niagara/Classes/NiagaraScript.h
@@ -1290,6 +1290,8 @@ public:
 
 	NIAGARA_API bool IsShaderMapCached(const ITargetPlatform* TargetPlatform, const FNiagaraShaderMapId& ShaderMapId) const;
 
+	NIAGARA_API void AssignRapidIterationParameters(const FNiagaraVMExecutableData& InScriptVM);
+
 	/** Updates the RI parameter store based on the provided variables (missing entries will be
 	    added and stale entries will be removed.  Returns true if the parameter store was modified. */
 	NIAGARA_API bool ApplyRapidIterationParameters(TConstArrayView<FNiagaraVariable> InParameters, bool bAllowRemoval);
diff --git a/Engine/Plugins/FX/Niagara/Source/Niagara/Private/NiagaraAsyncCompile.cpp b/Engine/Plugins/FX/Niagara/Source/Niagara/Private/NiagaraAsyncCompile.cpp
--- a/Engine/Plugins/FX/Niagara/Source/Niagara/Private/NiagaraAsyncCompile.cpp
+++ b/Engine/Plugins/FX/Niagara/Source/Niagara/Private/NiagaraAsyncCompile.cpp
@@ -3,7 +3,9 @@
 
 #include "NiagaraAsyncCompile.h"
 
+#include "Algo/Copy.h"
 #include "Modules/ModuleManager.h"
+#include "NiagaraConstants.h"
 #include "NiagaraEmitter.h"
 #include "NiagaraEmitterHandle.h"
 #include "NiagaraModule.h"
@@ -116,6 +118,8 @@ FNiagaraAsyncCompileTask::FNiagaraAsyncCompileTask(UNiagaraSystem* InOwningSyste
 	AssetPath = InAssetPath;
 	ScriptPair = InScriptPair;
 
+	bCompilableScript = InScriptPair.CompiledScript ? InScriptPair.CompiledScript->IsCompilable() : false;
+
 	CurrentState = ENiagaraCompilationState::CheckDDC;
 }
 
@@ -146,8 +150,16 @@ void FNiagaraAsyncCompileTask::ProcessCurrentState()
 	else if (CurrentState == ENiagaraCompilationState::StartCompileJob)
 	{
 		// start the async compile job
-		StartCompileJob();
-		MoveToState(ENiagaraCompilationState::AwaitResult);
+		if (bCompilableScript)
+		{
+			StartCompileJob();
+			MoveToState(ENiagaraCompilationState::AwaitResult);
+		}
+		else
+		{
+			ProcessNonCompilableScript();
+			MoveToState(ENiagaraCompilationState::PutToDDC);
+		}
 	}
 	else if (CurrentState == ENiagaraCompilationState::AwaitResult)
 	{
@@ -211,7 +223,7 @@ void FNiagaraAsyncCompileTask::MoveToState(ENiagaraCompilationState NewState)
 	}
 	if (NewState == ENiagaraCompilationState::PutToDDC)
 	{
-		check(CurrentState == ENiagaraCompilationState::ProcessResult);
+		check(CurrentState == ENiagaraCompilationState::ProcessResult || CurrentState == ENiagaraCompilationState::StartCompileJob);
 	}
 	if (NewState == ENiagaraCompilationState::Finished)
 	{
@@ -383,6 +395,15 @@ void FNiagaraAsyncCompileTask::OptimizeByteCode()
 #endif // VECTORVM_SUPPORTS_EXPERIMENTAL
 }
 
+void FNiagaraAsyncCompileTask::ProcessNonCompilableScript()
+{
+	ExeData = MakeShared<FNiagaraVMExecutableData>();
+	ScriptPair.CompileResults = ExeData;
+	ScriptPair.bResultsReady = true;
+	ExeData->BakedRapidIterationParameters = BakedRapidIterationParameters;
+	ScriptPair.CompiledScript->ExecToBinaryData(ScriptPair.CompiledScript, DDCOutData, *ExeData);
+}
+
 void FNiagaraAsyncCompileTask::ProcessResult()
 {
 	check(ExeData);
@@ -757,6 +778,25 @@ bool FNiagaraActiveCompilationDefault::QueryCompileComplete(const FNiagaraQueryC
 
 void FNiagaraActiveCompilationDefault::Apply(const FNiagaraQueryCompilationOptions& Options)
 {
+	if (bEvaluateParametersPending)
+	{
+		// run a first pass to apply the rapid iteration parameters across all of the tasks with results
+		for (FAsyncTaskPtr& AsyncTask : Tasks)
+		{
+			FEmitterCompiledScriptPair& EmitterCompiledScriptPair = AsyncTask->ScriptPair;
+			if (EmitterCompiledScriptPair.bResultsReady)
+			{
+				if (TSharedPtr<FNiagaraVMExecutableData> ExeData = EmitterCompiledScriptPair.CompileResults)
+				{
+					if (UNiagaraScript* CompiledScript = EmitterCompiledScriptPair.CompiledScript)
+					{
+						CompiledScript->AssignRapidIterationParameters(*ExeData);
+					}
+				}
+			}
+		}
+	}
+
 	for (FAsyncTaskPtr& AsyncTask : Tasks)
 	{
 		FEmitterCompiledScriptPair& EmitterCompiledScriptPair = AsyncTask->ScriptPair;
@@ -803,17 +843,30 @@ void FNiagaraActiveCompilationDefault::Apply(const FNiagaraQueryCompilationOptio
 			// Synchronize the variables that we actually encountered during precompile so that we can expose them to the end user.
 			FNiagaraUserRedirectionParameterStore& ExposedParameters = Options.System->GetExposedParameters();
 
-			TArray<FNiagaraVariable> OriginalExposedParams;
-			ExposedParameters.GetParameters(OriginalExposedParams);
-			TArray<FNiagaraVariable>& EncounteredExposedVars = AsyncTask->EncounteredExposedVars;
-			for (int32 i = 0; i < EncounteredExposedVars.Num(); i++)
+			auto AddUniqueExposedVariables = [&ExposedParameters](TConstArrayView<FNiagaraVariable> InVariables) -> void
 			{
-				if (OriginalExposedParams.Contains(EncounteredExposedVars[i]) == false)
+				for (const FNiagaraVariable& InVariable : InVariables)
 				{
-					// Just in case it wasn't added previously..
-					ExposedParameters.AddParameter(EncounteredExposedVars[i], true, false);
+					if (!ExposedParameters.ReadParameterVariables().Contains(InVariable))
+					{
+						// Just in case it wasn't added previously..
+						ExposedParameters.AddParameter(InVariable, true, false);
+					}
 				}
-			}
+			};
+
+			AddUniqueExposedVariables(AsyncTask->EncounteredExposedVars);
+
+			// we also grab from the data contained in the ExeData, since that might be all we have if grabbed results from the DDC rather
+			// than actually issuing a precompile
+			auto IsUserVariable = [](const FNiagaraVariableBase& InVariable) -> bool
+			{
+				return InVariable.IsInNameSpace(FStringView(PARAM_MAP_USER_STR));
+			};
+
+			TArray<FNiagaraVariable> ScriptExposedVariables;
+			Algo::CopyIf(ExeData->Parameters.Parameters, ScriptExposedVariables, IsUserVariable);
+			AddUniqueExposedVariables(ScriptExposedVariables);
 		}
 	}
 }
@@ -871,7 +924,7 @@ bool FNiagaraActiveCompilationDefault::Launch(const FNiagaraCompilationOptions&
 					Pair.CompiledScript = EmitterScript;
 
 					TSharedPtr<FNiagaraAsyncCompileTask, ESPMode::ThreadSafe> AsyncTask = MakeShared<FNiagaraAsyncCompileTask, ESPMode::ThreadSafe>(Options.System, EmitterScript->GetPathName(), Pair);
-					if (EmitterScript->IsCompilable() && !EmitterScript->AreScriptAndSourceSynchronized())
+					if (!EmitterScript->AreScriptAndSourceSynchronized())
 					{
 						// we need to compute the vmID here to check later in the ddc task before doing the precompile if anything has changed in the meantime.
 						// the compilation id was just calculated during AreScriptAndSourceSynchronized() so we reuse it here
diff --git a/Engine/Plugins/FX/Niagara/Source/Niagara/Private/NiagaraAsyncCompile.h b/Engine/Plugins/FX/Niagara/Source/Niagara/Private/NiagaraAsyncCompile.h
--- a/Engine/Plugins/FX/Niagara/Source/Niagara/Private/NiagaraAsyncCompile.h
+++ b/Engine/Plugins/FX/Niagara/Source/Niagara/Private/NiagaraAsyncCompile.h
@@ -69,6 +69,7 @@ public:
 	bool bWaitForCompileJob = false;
 	bool bUsedShaderCompilerWorker = false;
 	bool bFetchedGCObjects = false;
+	bool bCompilableScript = false;
 
 	// in order to coordinate between tasks associated with a system CheckDDC can be handled external to the task.
 	// In that case we disable this flag
@@ -107,6 +108,7 @@ public:
 	bool AwaitResult();
 	void ProcessResult();
 	void OptimizeByteCode();
+	void ProcessNonCompilableScript();
 
 	void AssignInitialCompilationId(const FNiagaraVMExecutableDataId& InitialCompilationId);
 	void UpdateCompilationId(const FNiagaraVMExecutableDataId& UpdatedCompilationId);
diff --git a/Engine/Plugins/FX/Niagara/Source/Niagara/Private/NiagaraScript.cpp b/Engine/Plugins/FX/Niagara/Source/Niagara/Private/NiagaraScript.cpp
--- a/Engine/Plugins/FX/Niagara/Source/Niagara/Private/NiagaraScript.cpp
+++ b/Engine/Plugins/FX/Niagara/Source/Niagara/Private/NiagaraScript.cpp
@@ -3064,15 +3064,10 @@ void UNiagaraScript::SetVMCompilationResults(const FNiagaraVMExecutableDataId& I
 
 	if (bApplyRapidIterationParameters)
 	{
-		const bool bClearBindings = false;
-		RapidIterationParameters.Empty(bClearBindings);
-		for (const FNiagaraVariable& Parameter : InScriptVM.BakedRapidIterationParameters)
-		{
-			const bool bInitialize = false;
-			const bool bTriggerRebind = false;
-			RapidIterationParameters.AddParameter(Parameter, bInitialize, bTriggerRebind);
-		}
-		RapidIterationParameters.TriggerOnLayoutChanged();
+		// if we are applying rapid iteration parameters then the data supplied by InScriptVM, which could have
+		// come from the DDC, could invalidate the CompileID that we had previous generated so generate it again
+		// now that we've updated everything
+		ComputeVMCompilationId(CachedScriptVMId, FGuid());
 	}
 
 	GenerateStatIDs();
@@ -3248,6 +3243,19 @@ bool UNiagaraScript::IsShaderMapCached(const ITargetPlatform* TargetPlatform, co
 	return false;
 }
 
+void UNiagaraScript::AssignRapidIterationParameters(const FNiagaraVMExecutableData& InScriptVM)
+{
+	const bool bClearBindings = false;
+	RapidIterationParameters.Empty(bClearBindings);
+	for (const FNiagaraVariable& Parameter : InScriptVM.BakedRapidIterationParameters)
+	{
+		const bool bInitialize = false;
+		const bool bTriggerRebind = false;
+		RapidIterationParameters.AddParameter(Parameter, bInitialize, bTriggerRebind);
+	}
+	RapidIterationParameters.TriggerOnLayoutChanged();
+}
+
 bool UNiagaraScript::ApplyRapidIterationParameters(TConstArrayView<FNiagaraVariable> InParameters, bool bAllowRemoval)
 {
 	bool bStoreDirty = false;
