Updated Overview Node UI for summary view

- Thumbnails are now kept in editor data instead of in the emitter itself and are post loaded over there
- Includes emitter thumbnails and ability to easily recapture them
- Fixed merge bug that discarded the child emitters "should show summary" value. Now it maintains it.
- Added widgets for sim target & world space, acting as toggles in summary view
- Cleaned up code in StackRoot to account for summary view. Now we apply a summary or default filter instead of removing and adding stack entries
- Refactored overview node to have a 'display mode' (Default & Summary). Distinct functions for widget creation that's easy to extend (stateless?) & added delegates to account for merge process (editor data is a new object after merge so need to rebind)
- Bunch of small changes for summary UI
- Removed User Parameters in the overview stack of the system node

diff --git a/Engine/Plugins/FX/Niagara/Source/Niagara/Classes/NiagaraEmitter.h b/Engine/Plugins/FX/Niagara/Source/Niagara/Classes/NiagaraEmitter.h
--- a/Engine/Plugins/FX/Niagara/Source/Niagara/Classes/NiagaraEmitter.h
+++ b/Engine/Plugins/FX/Niagara/Source/Niagara/Classes/NiagaraEmitter.h
@@ -713,7 +713,7 @@ public:
 	UPROPERTY(AssetRegistrySearchable, EditAnywhere, Category = Script)
 	FText Category;
 
-	///** Internal: The thumbnail image.*/
+	///** The thumbnail image used for the asset. This is always the latest recorded thumbnail. This can be different from the thumbnails that are saved per emitter version in collapsed view. */
 	UPROPERTY()
 	TObjectPtr<UTexture2D> ThumbnailImage;
 	
diff --git a/Engine/Plugins/FX/Niagara/Source/NiagaraEditor/Private/NiagaraEmitterEditorData.cpp b/Engine/Plugins/FX/Niagara/Source/NiagaraEditor/Private/NiagaraEmitterEditorData.cpp
--- a/Engine/Plugins/FX/Niagara/Source/NiagaraEditor/Private/NiagaraEmitterEditorData.cpp
+++ b/Engine/Plugins/FX/Niagara/Source/NiagaraEditor/Private/NiagaraEmitterEditorData.cpp
@@ -9,6 +9,7 @@
 #include "NiagaraStackEditorData.h"
 #include "ScopedTransaction.h"
 #include "NiagaraNodeAssignment.h"
+#include "Rendering/Texture2DResource.h"
 #include "ViewModels/HierarchyEditor/NiagaraSummaryViewViewModel.h"
 
 #include UE_INLINE_GENERATED_CPP_BY_NAME(NiagaraEmitterEditorData)
@@ -367,6 +368,16 @@ void UNiagaraEmitterEditorData::PostLoad_TransferSummaryDataToNewFormat()
 	}
 }
 
+void UNiagaraEmitterEditorData::PostLoad_TransferEmitterThumbnailImage(UObject* Owner)
+{
+	UNiagaraEmitter* Emitter = Cast<UNiagaraEmitter>(Owner);
+	
+	if(Emitter->ThumbnailImage != nullptr && EmitterThumbnail == nullptr)
+	{
+		SetThumbnail(Emitter->ThumbnailImage);
+	}
+}
+
 void UNiagaraEmitterEditorData::PostLoad_TransferModuleStackNotesToNewFormat(UObject* Owner)
 {
 	// since we lack graph context during post load, we do this workaround to find the emitter graph
@@ -422,6 +433,7 @@ void UNiagaraEmitterEditorData::PostLoad()
 void UNiagaraEmitterEditorData::PostLoadFromOwner(UObject* InOwner)
 {
 	PostLoad_TransferSummaryDataToNewFormat();
+	PostLoad_TransferEmitterThumbnailImage(InOwner);
 	PostLoad_TransferModuleStackNotesToNewFormat(InOwner);
 }
 
diff --git a/Engine/Plugins/FX/Niagara/Source/NiagaraEditor/Private/NiagaraOverviewNode.cpp b/Engine/Plugins/FX/Niagara/Source/NiagaraEditor/Private/NiagaraOverviewNode.cpp
--- a/Engine/Plugins/FX/Niagara/Source/NiagaraEditor/Private/NiagaraOverviewNode.cpp
+++ b/Engine/Plugins/FX/Niagara/Source/NiagaraEditor/Private/NiagaraOverviewNode.cpp
@@ -67,7 +67,7 @@ static FNiagaraEmitterHandle* FindEmitterHandleByID(UNiagaraSystem* System, cons
 	return nullptr;
 }
 
-FNiagaraEmitterHandle* UNiagaraOverviewNode::TryGetEmitterHandle()
+FNiagaraEmitterHandle* UNiagaraOverviewNode::TryGetEmitterHandle() const
 {
 	return FindEmitterHandleByID(GetOwningSystem(), GetEmitterHandleGuid());
 }
@@ -326,7 +326,7 @@ bool UNiagaraOverviewNode::CanCreateUnderSpecifiedSchema(const UEdGraphSchema* S
 	return Schema->IsA<UEdGraphSchema_NiagaraSystemOverview>();
 }
 
-UNiagaraSystem* UNiagaraOverviewNode::GetOwningSystem()
+UNiagaraSystem* UNiagaraOverviewNode::GetOwningSystem() const
 {
 	return OwningSystem;
 }
diff --git a/Engine/Plugins/FX/Niagara/Source/NiagaraEditor/Private/NiagaraScriptMergeManager.cpp b/Engine/Plugins/FX/Niagara/Source/NiagaraEditor/Private/NiagaraScriptMergeManager.cpp
--- a/Engine/Plugins/FX/Niagara/Source/NiagaraEditor/Private/NiagaraScriptMergeManager.cpp
+++ b/Engine/Plugins/FX/Niagara/Source/NiagaraEditor/Private/NiagaraScriptMergeManager.cpp
@@ -2463,6 +2463,8 @@ void FNiagaraScriptMergeManager::DiffEmitterSummary(const UNiagaraEmitterEditorD
 	{
 		return;
 	}
+
+	DiffResults.NewShouldShowSummaryViewValue = OtherEditorData->ShouldShowSummaryView();
 	
 	TArray<UNiagaraHierarchyItemBase*> AddedItemsInOther;
 	TArray<UNiagaraHierarchyItemBase*> RemovedItemsInBase;
diff --git a/Engine/Plugins/FX/Niagara/Source/NiagaraEditor/Private/Toolkits/NiagaraSystemToolkit.cpp b/Engine/Plugins/FX/Niagara/Source/NiagaraEditor/Private/Toolkits/NiagaraSystemToolkit.cpp
--- a/Engine/Plugins/FX/Niagara/Source/NiagaraEditor/Private/Toolkits/NiagaraSystemToolkit.cpp
+++ b/Engine/Plugins/FX/Niagara/Source/NiagaraEditor/Private/Toolkits/NiagaraSystemToolkit.cpp
@@ -291,6 +291,7 @@ void FNiagaraSystemToolkit::InitializeInternal(const EToolkitMode::Type Mode, co
 	SystemViewModel->OnRequestFocusTab().AddSP(this, &FNiagaraSystemToolkit::OnViewModelRequestFocusTab);
 	SystemViewModel->OnGetWorkflowMode().BindSP(this, &FNiagaraSystemToolkit::GetCurrentMode);
 	SystemViewModel->OnChangeWorkflowMode().BindSP(this, &FNiagaraSystemToolkit::SetCurrentMode);
+	SystemViewModel->OnEmitterThumbnailRequested().BindSP(this, &FNiagaraSystemToolkit::CaptureEmitterThumbnail);
 	SystemViewModel->GetDocumentViewModel()->InitializePreTabManager(SharedThis(this));
 	
 	constexpr bool bCreateDefaultStandaloneMenu = true;
@@ -469,7 +470,7 @@ void FNiagaraSystemToolkit::SetupCommands()
 
 	GetToolkitCommands()->MapAction(
 		FNiagaraEditorCommands::Get().SaveThumbnailImage,
-		FExecuteAction::CreateSP(this, &FNiagaraSystemToolkit::OnSaveThumbnailImage));
+		FExecuteAction::CreateSP(this, &FNiagaraSystemToolkit::CaptureAssetThumbnail));
 
 	GetToolkitCommands()->MapAction(
 		FNiagaraEditorCommands::Get().Apply,
@@ -567,20 +568,52 @@ TSharedPtr<FNiagaraEmitterViewModel> FNiagaraSystemToolkit::GetEditedEmitterView
 	return HasEmitter() ? SystemViewModel->GetEmitterHandleViewModels()[0]->GetEmitterViewModel() : TSharedPtr<FNiagaraEmitterViewModel>(); 
 }
 
-void FNiagaraSystemToolkit::OnSaveThumbnailImage()
+void FNiagaraSystemToolkit::CaptureAssetThumbnail() const
 {
 	if (Viewport.IsValid())
+	{		
+		Viewport->CreateThumbnail(SystemToolkitMode == ESystemToolkitMode::System ? static_cast<UObject*>(System) : Emitter, {});
+	}
+}
+
+void FNiagaraSystemToolkit::CaptureEmitterThumbnail(FGuid EmitterGuid)
+{
+	TSharedPtr<FNiagaraEmitterHandleViewModel> EditableEmitterHandleViewModel = SystemViewModel->GetEmitterHandleViewModelById(EmitterGuid);
+	if(EditableEmitterHandleViewModel)
 	{
-		Viewport->CreateThumbnail(SystemToolkitMode == ESystemToolkitMode::System ? static_cast<UObject*>(System) : Emitter);
+		if(Viewport.IsValid())
+		{
+			Viewport->CreateThumbnail(EditableEmitterHandleViewModel->GetEmitterViewModel()->GetEmitter().Emitter, EmitterGuid);
+		}
 	}
 }
 
-void FNiagaraSystemToolkit::OnThumbnailCaptured(UTexture2D* Thumbnail)
+void FNiagaraSystemToolkit::OnThumbnailCaptured(UTexture2D* Thumbnail, TOptional<FGuid> EmitterGuid)
 {
 	if (SystemToolkitMode == ESystemToolkitMode::System)
 	{
 		System->MarkPackageDirty();
-		System->ThumbnailImage = Thumbnail;
+		// in a system, we can either capture thumbnails for individual emitters, or for the system itself
+		if(EmitterGuid.IsSet())
+		{
+			if(TSharedPtr<FNiagaraEmitterHandleViewModel> EmitterHandleViewModel = GetSystemViewModel()->GetEmitterHandleViewModelById(EmitterGuid.GetValue()))
+			{
+				EmitterHandleViewModel->GetEmitterViewModel()->GetEditorData().SetThumbnail(Thumbnail);
+			}
+			
+			// since we cached the isolation state to capture only individual emitters, we restore the state here
+			GetSystemViewModel()->RestoreIsolatedEmitterState();
+
+			// we force a refresh to the render state which wouldn't occur if taking a thumbnail while paused
+			if(Viewport.IsValid())
+			{
+				Viewport->GetPreviewComponent()->MarkRenderStateDirty();
+			}
+		}
+		else
+		{
+			System->ThumbnailImage = Thumbnail;
+		}
 		// Broadcast an object property changed event to update the content browser
 		FPropertyChangedEvent EmptyPropertyChangedEvent(nullptr);
 		FCoreUObjectDelegates::OnObjectPropertyChanged.Broadcast(System, EmptyPropertyChangedEvent);
@@ -590,6 +623,7 @@ void FNiagaraSystemToolkit::OnThumbnailCaptured(UTexture2D* Thumbnail)
 		TSharedPtr<FNiagaraEmitterViewModel> EditableEmitterViewModel = SystemViewModel->GetEmitterHandleViewModels()[0]->GetEmitterViewModel();
 		UNiagaraEmitter* EditableEmitter = EditableEmitterViewModel->GetEmitter().Emitter;
 		EditableEmitter->ThumbnailImage = Thumbnail;
+		EditableEmitterViewModel->GetEditorData().SetThumbnail(Thumbnail);
 		bEmitterThumbnailUpdated = true;
 		// Broadcast an object property changed event to update the content browser
 		FPropertyChangedEvent EmptyPropertyChangedEvent(nullptr);
diff --git a/Engine/Plugins/FX/Niagara/Source/NiagaraEditor/Private/Toolkits/NiagaraSystemToolkit.h b/Engine/Plugins/FX/Niagara/Source/NiagaraEditor/Private/Toolkits/NiagaraSystemToolkit.h
--- a/Engine/Plugins/FX/Niagara/Source/NiagaraEditor/Private/Toolkits/NiagaraSystemToolkit.h
+++ b/Engine/Plugins/FX/Niagara/Source/NiagaraEditor/Private/Toolkits/NiagaraSystemToolkit.h
@@ -203,8 +203,10 @@ private:
 	void OnViewModelRequestFocusTab(FName TabName, bool bDrawAttention = false);
 	
 	const FName GetNiagaraSystemMessageLogName(UNiagaraSystem* InSystem) const;
-	void OnSaveThumbnailImage();
-	void OnThumbnailCaptured(UTexture2D* Thumbnail);
+
+	void CaptureAssetThumbnail() const;
+	void CaptureEmitterThumbnail(FGuid EmitterGuid);
+	void OnThumbnailCaptured(UTexture2D* Thumbnail, TOptional<FGuid> EmitterGuid);
 
 	void ManageVersions();
 	TSharedPtr<FNiagaraEmitterViewModel> GetEditedEmitterViewModel() const;
@@ -236,7 +238,7 @@ private:
 	ESystemToolkitWorkflowMode ActiveWorkflowMode;
 	
 	TSharedPtr<SNiagaraSystemViewport> Viewport;
-
+	
 	/* The view model for the System being edited */
 	TSharedPtr<FNiagaraSystemViewModel> SystemViewModel;
 	
diff --git a/Engine/Plugins/FX/Niagara/Source/NiagaraEditor/Private/ViewModels/NiagaraEmitterHandleViewModel.cpp b/Engine/Plugins/FX/Niagara/Source/NiagaraEditor/Private/ViewModels/NiagaraEmitterHandleViewModel.cpp
--- a/Engine/Plugins/FX/Niagara/Source/NiagaraEditor/Private/ViewModels/NiagaraEmitterHandleViewModel.cpp
+++ b/Engine/Plugins/FX/Niagara/Source/NiagaraEditor/Private/ViewModels/NiagaraEmitterHandleViewModel.cpp
@@ -16,6 +16,7 @@
 #include "ViewModels/Stack/NiagaraStackRenderItemGroup.h"
 #include "ViewModels/Stack/NiagaraStackRendererItem.h"
 #include "NiagaraMessages.h"
+#include "ViewModels/NiagaraSystemSelectionViewModel.h"
 #include "Widgets/SNiagaraDebugger.h"
 
 #define LOCTEXT_NAMESPACE "EmitterHandleViewModel"
@@ -111,6 +112,7 @@ void FNiagaraEmitterHandleViewModel::Initialize(TSharedRef<FNiagaraSystemViewMod
 	EmitterHandle = &InOwningSystemViewModel->GetSystem().GetEmitterHandle(InEmitterHandleIndex);
 	FVersionedNiagaraEmitter Emitter = EmitterHandle != nullptr ? EmitterHandle->GetInstance() : FVersionedNiagaraEmitter();
 	EmitterViewModel->Initialize(Emitter, InSimulation);
+	EmitterViewModel->OnEmitterSelectionRequested().BindSP(this, &FNiagaraEmitterHandleViewModel::RequestEmitterSelection);
 	EmitterStackViewModel->InitializeWithViewModels(InOwningSystemViewModel, this->AsShared(), FNiagaraStackViewModelOptions(false, true));
 }
 
@@ -120,6 +122,7 @@ void FNiagaraEmitterHandleViewModel::Reset()
 	OwningSystemViewModelWeak.Reset();
 	EmitterHandleIndex = INDEX_NONE;
 	EmitterHandle = nullptr;
+	EmitterViewModel->OnEmitterSelectionRequested().Unbind();
 	EmitterViewModel->Reset();
 }
 
@@ -301,6 +304,16 @@ void FNiagaraEmitterHandleViewModel::SetIsIsolated(bool bInIsIsolated)
 	}
 }
 
+void FNiagaraEmitterHandleViewModel::ToggleIsIsolated()
+{
+	GetOwningSystemViewModel()->ToggleEmitterIsolation(AsShared());
+}
+
+void FNiagaraEmitterHandleViewModel::RequestCaptureThumbnail()
+{
+	GetOwningSystemViewModel()->OnEmitterThumbnailRequested().ExecuteIfBound(GetEmitterHandle()->GetId());
+}
+
 ENiagaraSystemViewModelEditMode FNiagaraEmitterHandleViewModel::GetOwningSystemEditMode() const
 {
 	return GetOwningSystemViewModel()->GetEditMode();
@@ -345,6 +358,11 @@ void FNiagaraEmitterHandleViewModel::SetIsRenamePending(bool bInIsRenamePending)
 	bIsRenamePending = bInIsRenamePending;
 }
 
+void FNiagaraEmitterHandleViewModel::RequestEmitterSelection(bool bClearSelection)
+{
+	GetOwningSystemViewModel()->GetSelectionViewModel()->UpdateSelectedEntries({GetEmitterStackViewModel()->GetRootEntry()}, {}, bClearSelection);
+}
+
 void FNiagaraEmitterHandleViewModel::BeginDebugEmitter()
 {
 
diff --git a/Engine/Plugins/FX/Niagara/Source/NiagaraEditor/Private/ViewModels/NiagaraEmitterViewModel.cpp b/Engine/Plugins/FX/Niagara/Source/NiagaraEditor/Private/ViewModels/NiagaraEmitterViewModel.cpp
--- a/Engine/Plugins/FX/Niagara/Source/NiagaraEditor/Private/ViewModels/NiagaraEmitterViewModel.cpp
+++ b/Engine/Plugins/FX/Niagara/Source/NiagaraEditor/Private/ViewModels/NiagaraEmitterViewModel.cpp
@@ -88,6 +88,7 @@ bool FNiagaraEmitterViewModel::Initialize(const FVersionedNiagaraEmitter& InEmit
 	SummaryViewHierarchyViewModel->Initialize(AsShared());
 	SummaryViewHierarchyViewModel->OnHierarchyChanged().AddSP(this, &FNiagaraEmitterViewModel::OnSummaryViewHierarchyChanged);
 	SummaryViewHierarchyViewModel->OnHierarchyPropertiesChanged().AddSP(this, &FNiagaraEmitterViewModel::OnSummaryViewHierarchyChanged);
+	GetEditorData().OnSummaryViewStateChanged().AddSP(this, &FNiagaraEmitterViewModel::UpdateSelectionOnSummaryViewStateChanged);
 	return true;
 }
 
@@ -538,6 +539,11 @@ FNiagaraEmitterViewModel::FOnScriptParameterStoreChanged& FNiagaraEmitterViewMod
 	return OnScriptParameterStoreChangedDelegate;
 }
 
+FNiagaraEmitterViewModel::FOnEmitterSelectionRequested& FNiagaraEmitterViewModel::OnEmitterSelectionRequested()
+{
+	return OnEmitterSelectionRequestedDelegate;
+}
+
 void FNiagaraEmitterViewModel::AddScriptEventHandlers()
 {
 	if (FVersionedNiagaraEmitterData* EmitterData = EmitterWeakPtr.GetEmitterData())
@@ -624,6 +630,14 @@ void FNiagaraEmitterViewModel::OnSummaryViewHierarchyChanged()
 	GetEditorData().OnPersistentDataChanged().Broadcast();
 }
 
+void FNiagaraEmitterViewModel::UpdateSelectionOnSummaryViewStateChanged() const
+{
+	if(GetEditorData().ShouldShowSummaryView())
+	{
+		OnEmitterSelectionRequestedDelegate.ExecuteIfBound(false);
+	}
+}
+
 FString FNiagaraEmitterViewModel::GetReferencerName() const
 {
 	return TEXT("EmitterViewModel");
diff --git a/Engine/Plugins/FX/Niagara/Source/NiagaraEditor/Private/ViewModels/NiagaraSystemViewModel.cpp b/Engine/Plugins/FX/Niagara/Source/NiagaraEditor/Private/ViewModels/NiagaraSystemViewModel.cpp
--- a/Engine/Plugins/FX/Niagara/Source/NiagaraEditor/Private/ViewModels/NiagaraSystemViewModel.cpp
+++ b/Engine/Plugins/FX/Niagara/Source/NiagaraEditor/Private/ViewModels/NiagaraSystemViewModel.cpp
@@ -511,6 +511,11 @@ FNiagaraSystemViewModel::FOnChangeWorkflowMode& FNiagaraSystemViewModel::OnChang
 	return OnChangeWorkflowModeDelegate;
 }
 
+FNiagaraSystemViewModel::FOnEmitterThumbnailRequested& FNiagaraSystemViewModel::OnEmitterThumbnailRequested()
+{
+	return OnEmitterThumbnailRequestedDelegate;
+}
+
 TSharedPtr<FNiagaraEmitterHandleViewModel> FNiagaraSystemViewModel::AddEmitterFromAssetData(const FAssetData& AssetData)
 {
 	UNiagaraEmitter* Emitter = Cast<UNiagaraEmitter>(AssetData.GetAsset());
@@ -1976,6 +1981,24 @@ void FNiagaraSystemViewModel::IsolateEmitters(TArray<FGuid> EmitterHandlesIdsToI
 	GetSystem().SetIsolateEnabled(bAnyEmitterIsolated);
 }
 
+void FNiagaraSystemViewModel::CacheIsolatedEmitterState()
+{
+	CachedEmitterIsolationState.Empty();
+	for (TSharedRef<FNiagaraEmitterHandleViewModel> EmitterHandle : EmitterHandleViewModels)
+	{
+		if(EmitterHandle->GetEmitterHandle()->IsIsolated())
+		{
+			CachedEmitterIsolationState.Add(EmitterHandle->GetEmitterHandle()->GetId());
+		}
+	}
+}
+
+void FNiagaraSystemViewModel::RestoreIsolatedEmitterState()
+{
+	IsolateEmitters(CachedEmitterIsolationState);
+	CachedEmitterIsolationState.Empty();
+}
+
 void FNiagaraSystemViewModel::DisableEmitters(TArray<FGuid> EmitterHandlesIdsToDisable)
 {
 	FScopedTransaction Transaction(LOCTEXT("DisableSelectedEmitters", "Disabled selected emitters"));
diff --git a/Engine/Plugins/FX/Niagara/Source/NiagaraEditor/Private/ViewModels/Stack/NiagaraStackEmitterPropertiesGroup.cpp b/Engine/Plugins/FX/Niagara/Source/NiagaraEditor/Private/ViewModels/Stack/NiagaraStackEmitterPropertiesGroup.cpp
--- a/Engine/Plugins/FX/Niagara/Source/NiagaraEditor/Private/ViewModels/Stack/NiagaraStackEmitterPropertiesGroup.cpp
+++ b/Engine/Plugins/FX/Niagara/Source/NiagaraEditor/Private/ViewModels/Stack/NiagaraStackEmitterPropertiesGroup.cpp
@@ -197,6 +197,20 @@ const FSlateBrush* UNiagaraStackEmitterPropertiesGroup::GetIconBrush() const
 	return FAppStyle::Get().GetBrush("Icons.Details");
 }
 
+bool UNiagaraStackEmitterPropertiesGroup::SupportsSecondaryIcon() const
+{
+	if(IsFinalized() == false && GetEmitterViewModel().IsValid())
+	{
+		// we don't want to display the CPU/GPU icons in collapsed mode as we have a separate toggle for it
+		if(GetEmitterViewModel()->GetEditorData().ShouldShowSummaryView() == true)
+		{
+			return false;
+		}
+	}
+
+	return true;
+}
+
 const FSlateBrush* UNiagaraStackEmitterPropertiesGroup::GetSecondaryIconBrush() const
 {
 	if (IsFinalized() == false && GetEmitterViewModel().IsValid())
diff --git a/Engine/Plugins/FX/Niagara/Source/NiagaraEditor/Private/ViewModels/Stack/NiagaraStackEmitterSettingsGroup.cpp b/Engine/Plugins/FX/Niagara/Source/NiagaraEditor/Private/ViewModels/Stack/NiagaraStackEmitterSettingsGroup.cpp
--- a/Engine/Plugins/FX/Niagara/Source/NiagaraEditor/Private/ViewModels/Stack/NiagaraStackEmitterSettingsGroup.cpp
+++ b/Engine/Plugins/FX/Niagara/Source/NiagaraEditor/Private/ViewModels/Stack/NiagaraStackEmitterSettingsGroup.cpp
@@ -294,13 +294,11 @@ void UNiagaraStackEmitterSummaryItem::Initialize(FRequiredEntryData InRequiredEn
 
 FText UNiagaraStackEmitterSummaryItem::GetDisplayName() const
 {
-
 	return LOCTEXT("EmitterSummaryName", "Emitter Summary");
 }
 
 FText UNiagaraStackEmitterSummaryItem::GetTooltipText() const
 {
-
 	return LOCTEXT("EmitterSummaryTooltip", "Subset of parameters from the stack, summarized here for easier access.");
 }
 
@@ -353,15 +351,6 @@ TSharedPtr<IDetailTreeNode> GetSummarySectionsPropertyNode(const TArray<TSharedR
 	return GetSummarySectionsPropertyNode(ChildrenToCheck);
 }
 
-void UNiagaraStackEmitterSummaryItem::SelectSummaryNodesFromEmitterEditorDataRootNodes(TArray<TSharedRef<IDetailTreeNode>> Source, TArray<TSharedRef<IDetailTreeNode>>* Selected)
-{
-	TSharedPtr<IDetailTreeNode> SummarySectionsNode = GetSummarySectionsPropertyNode(Source);
-	if (SummarySectionsNode.IsValid())
-	{
-		Selected->Add(SummarySectionsNode.ToSharedRef());
-	}
-}
-
 void UNiagaraStackEmitterSummaryItem::OnEditButtonClicked()
 {
 	if (GetEmitterViewModel().IsValid())
@@ -380,16 +369,6 @@ TOptional<FText> UNiagaraStackEmitterSummaryItem::GetEditModeButtonTooltip() con
 	return LOCTEXT("EditSummaryViewButtonTooltip", "Summons the Summary Editor that lets you define what inputs, renderers and properties (and more) should be displayed when looking at the Emitter Summary.");
 }
 
-EVisibility UNiagaraStackEmitterSummaryItem::IsEditButtonVisible() const
-{
-	if(GetEmitterViewModel().IsValid())
-	{
-		return GetEmitterViewModel()->GetEditorData().ShouldShowSummaryView() ? EVisibility::Collapsed : EVisibility::Visible;
-	}
-
-	return EVisibility::Collapsed;
-}
-
 UNiagaraStackEmitterSummaryGroup::UNiagaraStackEmitterSummaryGroup()
 	: SummaryItem(nullptr)
 {
@@ -412,42 +391,5 @@ FText UNiagaraStackEmitterSummaryGroup::GetIconText() const
 	return FText::FromString(FString(TEXT("\xf0ca")/* fa-list-ul */));
 }
 
-void UNiagaraStackSummaryViewCollapseButton::Initialize(FRequiredEntryData InRequiredEntryData)
-{
-	Super::Initialize(InRequiredEntryData, TEXT("ShowAdvanced"));
-}
-
-FText UNiagaraStackSummaryViewCollapseButton::GetDisplayName() const
-{
-	return LOCTEXT("SummaryCollapseButtonDisplayName", "Show Advanced");
-}
-
-UNiagaraStackEntry::EStackRowStyle UNiagaraStackSummaryViewCollapseButton::GetStackRowStyle() const
-{
-	return EStackRowStyle::GroupHeader;
-}
-
-FText UNiagaraStackSummaryViewCollapseButton::GetTooltipText() const 
-{
-	return LOCTEXT("SummaryCollapseButtonTooltip", "Expand/Collapse detailed view.");
-}
-
-bool UNiagaraStackSummaryViewCollapseButton::GetIsEnabled() const
-{
-	return true;
-}
-
-void UNiagaraStackSummaryViewCollapseButton::RefreshChildrenInternal(const TArray<UNiagaraStackEntry*>& CurrentChildren, TArray<UNiagaraStackEntry*>& NewChildren, TArray<FStackIssue>& NewIssues)
-{
-	
-}
-
-
-
-
-
-
-
-
 #undef LOCTEXT_NAMESPACE
 
diff --git a/Engine/Plugins/FX/Niagara/Source/NiagaraEditor/Private/ViewModels/Stack/NiagaraStackRoot.cpp b/Engine/Plugins/FX/Niagara/Source/NiagaraEditor/Private/ViewModels/Stack/NiagaraStackRoot.cpp
--- a/Engine/Plugins/FX/Niagara/Source/NiagaraEditor/Private/ViewModels/Stack/NiagaraStackRoot.cpp
+++ b/Engine/Plugins/FX/Niagara/Source/NiagaraEditor/Private/ViewModels/Stack/NiagaraStackRoot.cpp
@@ -14,6 +14,9 @@
 #include "NiagaraSystem.h"
 #include "NiagaraSystemEditorData.h"
 #include "NiagaraEmitterEditorData.h"
+#include "NiagaraObjectSelection.h"
+#include "ViewModels/NiagaraOverviewGraphViewModel.h"
+#include "ViewModels/NiagaraSystemSelectionViewModel.h"
 #include "ViewModels/Stack/NiagaraStackSystemSettingsGroup.h"
 
 #include UE_INLINE_GENERATED_CPP_BY_NAME(NiagaraStackRoot)
@@ -22,7 +25,6 @@
 
 UNiagaraStackRoot::UNiagaraStackRoot()
 	: SystemPropertiesGroup(nullptr)
-	, SystemUserParametersGroup(nullptr)
 	, EmitterPropertiesGroup(nullptr)
 	, EmitterSummaryGroup(nullptr)
 	, EmitterSpawnGroup(nullptr)
@@ -39,7 +41,6 @@ void UNiagaraStackRoot::Initialize(FRequiredEntryData InRequiredEntryData, bool
 	bIncludeSystemInformation = bInIncludeSystemInformation;
 	bIncludeEmitterInformation = bInIncludeEmitterInformation;
 	SystemPropertiesGroup = nullptr;
-	SystemUserParametersGroup = nullptr;
 	EmitterPropertiesGroup = nullptr;
 	EmitterSummaryGroup = nullptr;
 	EmitterSpawnGroup = nullptr;
@@ -51,6 +52,15 @@ void UNiagaraStackRoot::Initialize(FRequiredEntryData InRequiredEntryData, bool
 	if (bInIncludeEmitterInformation && GetEmitterViewModel())
 	{
 		GetEmitterViewModel()->GetEditorData().OnSummaryViewStateChanged().AddUObject(this, &UNiagaraStackRoot::OnSummaryViewStateChanged);
+
+		if(GetEmitterViewModel()->GetEditorData().ShouldShowSummaryView())
+		{
+			AddSummaryViewFilter();
+		}
+		else
+		{
+			AddDefaultViewFilter();
+		}
 	}
 }
 
@@ -85,15 +95,6 @@ void UNiagaraStackRoot::RefreshChildrenInternal(const TArray<UNiagaraStackEntry*
 		SystemPropertiesGroup->Initialize(RequiredEntryData);
 	}
 
-	if (bIncludeSystemInformation && SystemUserParametersGroup == nullptr)
-	{
-		SystemUserParametersGroup = NewObject<UNiagaraStackSystemUserParametersGroup>(this);
-		FRequiredEntryData RequiredEntryData(GetSystemViewModel(), GetEmitterViewModel(),
-			FExecutionCategoryNames::System, FExecutionSubcategoryNames::Settings,
-			GetSystemViewModel()->GetEditorData().GetStackEditorData());
-		SystemUserParametersGroup->Initialize(RequiredEntryData, &GetSystemViewModel()->GetSystem(), &GetSystemViewModel()->GetSystem().GetExposedParameters());
-	}
-
 	if (bIncludeSystemInformation && SystemSpawnGroup == nullptr)
 	{
 		SystemSpawnGroup = NewObject<UNiagaraStackScriptItemGroup>(this);
@@ -126,105 +127,86 @@ void UNiagaraStackRoot::RefreshChildrenInternal(const TArray<UNiagaraStackEntry*
 		CommentCollection->Initialize(RequiredEntryData);
 	}
 	
-	bool bShouldShowSummaryView = GetEmitterViewModel()? GetEmitterViewModel()->GetEditorData().ShouldShowSummaryView() : false;
-
-	// we clear references to entries we don't want to display in summary view.
-	// Since they also won't get added to the NewChildren array, they will get finalized and eventually GC'd
-	if(bShouldShowSummaryView)
-	{
-		EmitterSpawnGroup = nullptr;
-		EmitterUpdateGroup = nullptr;
-		ParticleSpawnGroup = nullptr;
-		ParticleUpdateGroup = nullptr;
-	}
-	
-	if (bIncludeEmitterInformation && EmitterPropertiesGroup == nullptr)
-	{
-		EmitterPropertiesGroup = NewObject<UNiagaraStackEmitterPropertiesGroup>(this);
-		FRequiredEntryData RequiredEntryData(GetSystemViewModel(), GetEmitterViewModel(),
-			FExecutionCategoryNames::Emitter, FExecutionSubcategoryNames::Settings,
-			GetEmitterViewModel()->GetEditorData().GetStackEditorData());
-		EmitterPropertiesGroup->Initialize(RequiredEntryData);
-	}
-
-	if (bIncludeEmitterInformation && EmitterSummaryGroup == nullptr)
-	{
-		EmitterSummaryGroup = NewObject<UNiagaraStackEmitterSummaryGroup>(this);
-		FRequiredEntryData RequiredEntryData(GetSystemViewModel(), GetEmitterViewModel(),
-			FExecutionCategoryNames::Emitter, FExecutionSubcategoryNames::Settings,
-			GetEmitterViewModel()->GetEditorData().GetStackEditorData());
-		FText DisplayName = LOCTEXT("EmitterSummaryGroupName", "Emitter Summary");
-		FText Tooltip = LOCTEXT("EmitterSummaryTooltip", "Summary of parameters for this Emitter.");
-		EmitterSummaryGroup->Initialize(RequiredEntryData, DisplayName, Tooltip, nullptr);
-	}
-
-	if (bIncludeEmitterInformation && !bShouldShowSummaryView && EmitterSpawnGroup == nullptr)
+	if (bIncludeEmitterInformation)
 	{
-		EmitterSpawnGroup = NewObject<UNiagaraStackScriptItemGroup>(this);
-		FRequiredEntryData RequiredEntryData(GetSystemViewModel(), GetEmitterViewModel(),
-			FExecutionCategoryNames::Emitter, FExecutionSubcategoryNames::Spawn,
-			GetEmitterViewModel()->GetEditorData().GetStackEditorData());
-		FText DisplayName = LOCTEXT("EmitterSpawnGroupName", "Emitter Spawn");
-		FText ToolTip = LOCTEXT("EmitterSpawnGroupTooltip", "Occurs once at Emitter creation on the CPU. Modules in this stage should initialize defaults and/or do initial setup.\r\nModules are executed in order from top to bottom of the stack.");
-		EmitterSpawnGroup->Initialize(RequiredEntryData, DisplayName, ToolTip, GetEmitterViewModel()->GetSharedScriptViewModel(), ENiagaraScriptUsage::EmitterSpawnScript);
-	}
+		if(EmitterPropertiesGroup == nullptr)
+		{
+			EmitterPropertiesGroup = NewObject<UNiagaraStackEmitterPropertiesGroup>(this);
+			FRequiredEntryData RequiredEntryData(GetSystemViewModel(), GetEmitterViewModel(),
+				FExecutionCategoryNames::Emitter, FExecutionSubcategoryNames::Settings,
+				GetEmitterViewModel()->GetEditorData().GetStackEditorData());
+			EmitterPropertiesGroup->Initialize(RequiredEntryData);
+		}
 
-	if (bIncludeEmitterInformation && !bShouldShowSummaryView && EmitterUpdateGroup == nullptr)
-	{
-		EmitterUpdateGroup = NewObject<UNiagaraStackScriptItemGroup>(this);
-		FRequiredEntryData RequiredEntryData(GetSystemViewModel(), GetEmitterViewModel(),
-			FExecutionCategoryNames::Emitter, FExecutionSubcategoryNames::Update,
-			GetEmitterViewModel()->GetEditorData().GetStackEditorData());
-		FText DisplayName = LOCTEXT("EmitterUpdateGroupName", "Emitter Update");
-		FText ToolTip = LOCTEXT("EmitterUpdateGroupTooltip", "Occurs every Emitter tick on the CPU. Modules in this stage should compute values for parameters for Particle Update or Spawning this frame.\r\nModules are executed in order from top to bottom of the stack.");
-		EmitterUpdateGroup->Initialize(RequiredEntryData, DisplayName, ToolTip, GetEmitterViewModel()->GetSharedScriptViewModel(), ENiagaraScriptUsage::EmitterUpdateScript);
-	}
+		if(EmitterSummaryGroup == nullptr)
+		{
+			EmitterSummaryGroup = NewObject<UNiagaraStackEmitterSummaryGroup>(this);
+			FRequiredEntryData RequiredEntryData(GetSystemViewModel(), GetEmitterViewModel(),
+				FExecutionCategoryNames::Emitter, FExecutionSubcategoryNames::Settings,
+				GetEmitterViewModel()->GetEditorData().GetStackEditorData());
+			FText DisplayName = LOCTEXT("EmitterSummaryGroupName", "Emitter Summary");
+			FText Tooltip = LOCTEXT("EmitterSummaryTooltip", "Summary of parameters for this Emitter.");
+			EmitterSummaryGroup->Initialize(RequiredEntryData, DisplayName, Tooltip, nullptr);
+		}
 
-	if (bIncludeEmitterInformation && !bShouldShowSummaryView && ParticleSpawnGroup == nullptr)
-	{
-		ParticleSpawnGroup = NewObject<UNiagaraStackScriptItemGroup>(this);
-		FRequiredEntryData RequiredEntryData(GetSystemViewModel(), GetEmitterViewModel(),
-			FExecutionCategoryNames::Particle, FExecutionSubcategoryNames::Spawn,
-			GetEmitterViewModel()->GetEditorData().GetStackEditorData());
-		FText DisplayName = LOCTEXT("ParticleSpawnGroupName", "Particle Spawn");
-		FText ToolTip = LOCTEXT("ParticleSpawnGroupTooltip", "Called once per created particle. Modules in this stage should set up initial values for each particle.\r\nIf \"Use Interpolated Spawning\" is set, we will also run the Particle Update stage after the Particle Spawn stage.\r\nModules are executed in order from top to bottom of the stack.");
-		ParticleSpawnGroup->Initialize(RequiredEntryData, DisplayName, ToolTip, GetEmitterViewModel()->GetSharedScriptViewModel(), ENiagaraScriptUsage::ParticleSpawnScript);
-	}
+		if(RenderGroup == nullptr)
+		{
+			RenderGroup = NewObject<UNiagaraStackRenderItemGroup>(this);
+			FRequiredEntryData RequiredEntryData(GetSystemViewModel(), GetEmitterViewModel(),
+				FExecutionCategoryNames::Render, FExecutionSubcategoryNames::Render,
+				GetEmitterViewModel()->GetEditorData().GetStackEditorData());
+			RenderGroup->Initialize(RequiredEntryData);
+		}
+		
+		if(EmitterSpawnGroup == nullptr)
+		{
+			EmitterSpawnGroup = NewObject<UNiagaraStackScriptItemGroup>(this);
+			FRequiredEntryData RequiredEntryData(GetSystemViewModel(), GetEmitterViewModel(),
+				FExecutionCategoryNames::Emitter, FExecutionSubcategoryNames::Spawn,
+				GetEmitterViewModel()->GetEditorData().GetStackEditorData());
+			FText DisplayName = LOCTEXT("EmitterSpawnGroupName", "Emitter Spawn");
+			FText ToolTip = LOCTEXT("EmitterSpawnGroupTooltip", "Occurs once at Emitter creation on the CPU. Modules in this stage should initialize defaults and/or do initial setup.\r\nModules are executed in order from top to bottom of the stack.");
+			EmitterSpawnGroup->Initialize(RequiredEntryData, DisplayName, ToolTip, GetEmitterViewModel()->GetSharedScriptViewModel(), ENiagaraScriptUsage::EmitterSpawnScript);
+		}
 
-	if (bIncludeEmitterInformation && !bShouldShowSummaryView && ParticleUpdateGroup == nullptr)
-	{
-		ParticleUpdateGroup = NewObject<UNiagaraStackScriptItemGroup>(this);
-		FRequiredEntryData RequiredEntryData(GetSystemViewModel(), GetEmitterViewModel(),
-			FExecutionCategoryNames::Particle, FExecutionSubcategoryNames::Update,
-			GetEmitterViewModel()->GetEditorData().GetStackEditorData());
-		FText DisplayName = LOCTEXT("ParticleUpdateGroupName", "Particle Update");
-		FText ToolTip = LOCTEXT("ParticleUpdateGroupTooltip", "Called every frame per particle. Modules in this stage should update new values for this frame.\r\nModules are executed in order from top to bottom of the stack.");
-		ParticleUpdateGroup->Initialize(RequiredEntryData, DisplayName, ToolTip, GetEmitterViewModel()->GetSharedScriptViewModel(), ENiagaraScriptUsage::ParticleUpdateScript);
-	}
+		if(EmitterUpdateGroup == nullptr)
+		{
+			EmitterUpdateGroup = NewObject<UNiagaraStackScriptItemGroup>(this);
+			FRequiredEntryData RequiredEntryData(GetSystemViewModel(), GetEmitterViewModel(),
+				FExecutionCategoryNames::Emitter, FExecutionSubcategoryNames::Update,
+				GetEmitterViewModel()->GetEditorData().GetStackEditorData());
+			FText DisplayName = LOCTEXT("EmitterUpdateGroupName", "Emitter Update");
+			FText ToolTip = LOCTEXT("EmitterUpdateGroupTooltip", "Occurs every Emitter tick on the CPU. Modules in this stage should compute values for parameters for Particle Update or Spawning this frame.\r\nModules are executed in order from top to bottom of the stack.");
+			EmitterUpdateGroup->Initialize(RequiredEntryData, DisplayName, ToolTip, GetEmitterViewModel()->GetSharedScriptViewModel(), ENiagaraScriptUsage::EmitterUpdateScript);
+		}
 
-	if (bIncludeEmitterInformation && RenderGroup == nullptr)
-	{
-		RenderGroup = NewObject<UNiagaraStackRenderItemGroup>(this);
-		FRequiredEntryData RequiredEntryData(GetSystemViewModel(), GetEmitterViewModel(),
-			FExecutionCategoryNames::Render, FExecutionSubcategoryNames::Render,
-			GetEmitterViewModel()->GetEditorData().GetStackEditorData());
-		RenderGroup->Initialize(RequiredEntryData);
-	}
+		if(ParticleSpawnGroup == nullptr)
+		{
+			ParticleSpawnGroup = NewObject<UNiagaraStackScriptItemGroup>(this);
+			FRequiredEntryData RequiredEntryData(GetSystemViewModel(), GetEmitterViewModel(),
+				FExecutionCategoryNames::Particle, FExecutionSubcategoryNames::Spawn,
+				GetEmitterViewModel()->GetEditorData().GetStackEditorData());
+			FText DisplayName = LOCTEXT("ParticleSpawnGroupName", "Particle Spawn");
+			FText ToolTip = LOCTEXT("ParticleSpawnGroupTooltip", "Called once per created particle. Modules in this stage should set up initial values for each particle.\r\nIf \"Use Interpolated Spawning\" is set, we will also run the Particle Update stage after the Particle Spawn stage.\r\nModules are executed in order from top to bottom of the stack.");
+			ParticleSpawnGroup->Initialize(RequiredEntryData, DisplayName, ToolTip, GetEmitterViewModel()->GetSharedScriptViewModel(), ENiagaraScriptUsage::ParticleSpawnScript);
+		}
 
-	if (bIncludeEmitterInformation && SummaryCollapseButton == nullptr)
-	{
-		SummaryCollapseButton = NewObject<UNiagaraStackSummaryViewCollapseButton>(this);
-		FRequiredEntryData RequiredEntryData(GetSystemViewModel(), GetEmitterViewModel(),
-			FExecutionCategoryNames::Emitter, FExecutionSubcategoryNames::Settings,
-			GetEmitterViewModel()->GetEditorData().GetStackEditorData());
-		SummaryCollapseButton->Initialize(RequiredEntryData);
+		if(ParticleUpdateGroup == nullptr)
+		{
+			ParticleUpdateGroup = NewObject<UNiagaraStackScriptItemGroup>(this);
+			FRequiredEntryData RequiredEntryData(GetSystemViewModel(), GetEmitterViewModel(),
+				FExecutionCategoryNames::Particle, FExecutionSubcategoryNames::Update,
+				GetEmitterViewModel()->GetEditorData().GetStackEditorData());
+			FText DisplayName = LOCTEXT("ParticleUpdateGroupName", "Particle Update");
+			FText ToolTip = LOCTEXT("ParticleUpdateGroupTooltip", "Called every frame per particle. Modules in this stage should update new values for this frame.\r\nModules are executed in order from top to bottom of the stack.");
+			ParticleUpdateGroup->Initialize(RequiredEntryData, DisplayName, ToolTip, GetEmitterViewModel()->GetSharedScriptViewModel(), ENiagaraScriptUsage::ParticleUpdateScript);
+		}
 	}
 
 	// Populate new children
 	if (bIncludeSystemInformation)
 	{
 		NewChildren.Add(SystemPropertiesGroup);
-		NewChildren.Add(SystemUserParametersGroup);
 		NewChildren.Add(SystemSpawnGroup);
 		NewChildren.Add(SystemUpdateGroup);
 		NewChildren.Add(CommentCollection);
@@ -232,56 +214,53 @@ void UNiagaraStackRoot::RefreshChildrenInternal(const TArray<UNiagaraStackEntry*
 
 	if (bIncludeEmitterInformation)
 	{
-		NewChildren.Add(EmitterPropertiesGroup);	
 		NewChildren.Add(EmitterSummaryGroup);
-
-		if (!bShouldShowSummaryView)
+		
+		NewChildren.Add(EmitterPropertiesGroup);
+		
+		NewChildren.Add(EmitterSpawnGroup);
+		NewChildren.Add(EmitterUpdateGroup);
+
+		NewChildren.Add(ParticleSpawnGroup);
+		NewChildren.Add(ParticleUpdateGroup);
+		
+		for (const FNiagaraEventScriptProperties& EventScriptProperties : GetEmitterViewModel()->GetEmitter().GetEmitterData()->GetEventHandlers())
 		{
-			NewChildren.Add(EmitterSpawnGroup);
-			NewChildren.Add(EmitterUpdateGroup);
-
-			NewChildren.Add(ParticleSpawnGroup);
-			NewChildren.Add(ParticleUpdateGroup);
+			UNiagaraStackEventScriptItemGroup* EventHandlerGroup = FindCurrentChildOfTypeByPredicate<UNiagaraStackEventScriptItemGroup>(CurrentChildren,
+				[&](UNiagaraStackEventScriptItemGroup* CurrentEventHandlerGroup) { return CurrentEventHandlerGroup->GetScriptUsageId() == EventScriptProperties.Script->GetUsageId(); });
 
-			for (const FNiagaraEventScriptProperties& EventScriptProperties : GetEmitterViewModel()->GetEmitter().GetEmitterData()->GetEventHandlers())
+			if (EventHandlerGroup == nullptr)
 			{
-				UNiagaraStackEventScriptItemGroup* EventHandlerGroup = FindCurrentChildOfTypeByPredicate<UNiagaraStackEventScriptItemGroup>(CurrentChildren,
-					[&](UNiagaraStackEventScriptItemGroup* CurrentEventHandlerGroup) { return CurrentEventHandlerGroup->GetScriptUsageId() == EventScriptProperties.Script->GetUsageId(); });
-
-				if (EventHandlerGroup == nullptr)
-				{
-					EventHandlerGroup = NewObject<UNiagaraStackEventScriptItemGroup>(this);
-					FRequiredEntryData RequiredEntryData(GetSystemViewModel(), GetEmitterViewModel(),
-						FExecutionCategoryNames::Particle, FExecutionSubcategoryNames::Event,
-						GetEmitterViewModel()->GetEditorData().GetStackEditorData());
-					EventHandlerGroup->Initialize(RequiredEntryData, GetEmitterViewModel()->GetSharedScriptViewModel(), ENiagaraScriptUsage::ParticleEventScript, EventScriptProperties.Script->GetUsageId(), EventScriptProperties.SourceEmitterID);
-					EventHandlerGroup->SetOnModifiedEventHandlers(UNiagaraStackEventScriptItemGroup::FOnModifiedEventHandlers::CreateUObject(this, &UNiagaraStackRoot::EmitterArraysChanged));
-				}
-
-				NewChildren.Add(EventHandlerGroup);
+				EventHandlerGroup = NewObject<UNiagaraStackEventScriptItemGroup>(this);
+				FRequiredEntryData RequiredEntryData(GetSystemViewModel(), GetEmitterViewModel(),
+					FExecutionCategoryNames::Particle, FExecutionSubcategoryNames::Event,
+					GetEmitterViewModel()->GetEditorData().GetStackEditorData());
+				EventHandlerGroup->Initialize(RequiredEntryData, GetEmitterViewModel()->GetSharedScriptViewModel(), ENiagaraScriptUsage::ParticleEventScript, EventScriptProperties.Script->GetUsageId(), EventScriptProperties.SourceEmitterID);
+				EventHandlerGroup->SetOnModifiedEventHandlers(UNiagaraStackEventScriptItemGroup::FOnModifiedEventHandlers::CreateUObject(this, &UNiagaraStackRoot::EmitterArraysChanged));
 			}
 
-			for (UNiagaraSimulationStageBase* SimulationStage : GetEmitterViewModel()->GetEmitter().GetEmitterData()->GetSimulationStages())
+			NewChildren.Add(EventHandlerGroup);
+		}
+
+		for (UNiagaraSimulationStageBase* SimulationStage : GetEmitterViewModel()->GetEmitter().GetEmitterData()->GetSimulationStages())
+		{
+			UNiagaraStackSimulationStageGroup* SimulationStageGroup = FindCurrentChildOfTypeByPredicate<UNiagaraStackSimulationStageGroup>(CurrentChildren,
+				[SimulationStage](UNiagaraStackSimulationStageGroup* CurrentSimulationStageGroup) { return CurrentSimulationStageGroup->GetSimulationStage() == SimulationStage; });
+
+			if (SimulationStageGroup == nullptr)
 			{
-				UNiagaraStackSimulationStageGroup* SimulationStageGroup = FindCurrentChildOfTypeByPredicate<UNiagaraStackSimulationStageGroup>(CurrentChildren,
-					[SimulationStage](UNiagaraStackSimulationStageGroup* CurrentSimulationStageGroup) { return CurrentSimulationStageGroup->GetSimulationStage() == SimulationStage; });
-
-				if (SimulationStageGroup == nullptr)
-				{
-					SimulationStageGroup = NewObject<UNiagaraStackSimulationStageGroup>(this);
-					FRequiredEntryData RequiredEntryData(GetSystemViewModel(), GetEmitterViewModel(),
-						FExecutionCategoryNames::Particle, FExecutionSubcategoryNames::SimulationStage,
-						GetEmitterViewModel()->GetEditorData().GetStackEditorData());
-					SimulationStageGroup->Initialize(RequiredEntryData, GetEmitterViewModel()->GetSharedScriptViewModel(), SimulationStage);
-					SimulationStageGroup->SetOnModifiedSimulationStages(UNiagaraStackSimulationStageGroup::FOnModifiedSimulationStages::CreateUObject(this, &UNiagaraStackRoot::EmitterArraysChanged));
-				}
-
-				NewChildren.Add(SimulationStageGroup);
+				SimulationStageGroup = NewObject<UNiagaraStackSimulationStageGroup>(this);
+				FRequiredEntryData RequiredEntryData(GetSystemViewModel(), GetEmitterViewModel(),
+					FExecutionCategoryNames::Particle, FExecutionSubcategoryNames::SimulationStage,
+					GetEmitterViewModel()->GetEditorData().GetStackEditorData());
+				SimulationStageGroup->Initialize(RequiredEntryData, GetEmitterViewModel()->GetSharedScriptViewModel(), SimulationStage);
+				SimulationStageGroup->SetOnModifiedSimulationStages(UNiagaraStackSimulationStageGroup::FOnModifiedSimulationStages::CreateUObject(this, &UNiagaraStackRoot::EmitterArraysChanged));
 			}
-		}
 
+			NewChildren.Add(SimulationStageGroup);
+		}
+		
 		NewChildren.Add(RenderGroup);
-		NewChildren.Add(SummaryCollapseButton);
 	}
 }
 
@@ -294,9 +273,82 @@ void UNiagaraStackRoot::OnSummaryViewStateChanged()
 {
 	if (!IsFinalized())
 	{
+		if(GetEmitterViewModel().IsValid())
+		{
+			if(GetEmitterViewModel()->GetEditorData().ShouldShowSummaryView())
+			{
+				AddSummaryViewFilter();
+				RemoveDefaultViewFilter();
+			}
+			else
+			{
+				AddDefaultViewFilter();
+				RemoveSummaryViewFilter();
+			}
+		}
+
 		RefreshChildren();
+		GetSystemViewModel()->GetSelectionViewModel()->Refresh();
+	}
+}
+
+void UNiagaraStackRoot::AddDefaultViewFilter()
+{
+	if(DefaultViewFilterHandle.IsValid())
+	{
+		RemoveDefaultViewFilter();
+	}
+	
+	DefaultViewFilterHandle = AddChildFilter(FOnFilterChild::CreateUObject(this, &UNiagaraStackRoot::FilterForDefaultView));
+}
+
+void UNiagaraStackRoot::RemoveDefaultViewFilter()
+{
+	if(DefaultViewFilterHandle.IsValid())
+	{
+		RemoveChildFilter(DefaultViewFilterHandle);
+		DefaultViewFilterHandle.Reset();
+	}
+}
+
+void UNiagaraStackRoot::AddSummaryViewFilter()
+{
+	if(SummaryViewFilterHandle.IsValid())
+	{
+		RemoveSummaryViewFilter();
+	}
+	
+	SummaryViewFilterHandle = AddChildFilter(FOnFilterChild::CreateUObject(this, &UNiagaraStackRoot::FilterForSummaryView));
+}
+
+void UNiagaraStackRoot::RemoveSummaryViewFilter()
+{
+	if(SummaryViewFilterHandle.IsValid())
+	{
+		RemoveChildFilter(SummaryViewFilterHandle);
+		SummaryViewFilterHandle.Reset();
 	}
 }
 
+bool UNiagaraStackRoot::FilterForDefaultView(const UNiagaraStackEntry& NiagaraStackEntry) const
+{
+	if(NiagaraStackEntry.IsA<UNiagaraStackEmitterSummaryGroup>())
+	{
+		return false;
+	}
+
+	return true;
+}
+
+bool UNiagaraStackRoot::FilterForSummaryView(const UNiagaraStackEntry& NiagaraStackEntry) const
+{
+	if(NiagaraStackEntry.IsA<UNiagaraStackEmitterSummaryGroup>())
+	{
+		return true;
+	}
+	
+	return false;
+}
+
 #undef LOCTEXT_NAMESPACE
 
diff --git a/Engine/Plugins/FX/Niagara/Source/NiagaraEditor/Private/ViewModels/Stack/NiagaraStackSummaryViewInputCollection.cpp b/Engine/Plugins/FX/Niagara/Source/NiagaraEditor/Private/ViewModels/Stack/NiagaraStackSummaryViewInputCollection.cpp
--- a/Engine/Plugins/FX/Niagara/Source/NiagaraEditor/Private/ViewModels/Stack/NiagaraStackSummaryViewInputCollection.cpp
+++ b/Engine/Plugins/FX/Niagara/Source/NiagaraEditor/Private/ViewModels/Stack/NiagaraStackSummaryViewInputCollection.cpp
@@ -83,6 +83,16 @@ void UNiagaraStackSummaryViewCollection::RefreshForAdvancedToggle()
 	}
 }
 
+bool UNiagaraStackSummaryViewCollection::GetShouldShowInStack() const
+{
+	// if(GetEmitterViewModel().IsValid())
+	// {
+	// 	return GetEmitterViewModel()->GetEditorData().ShouldShowSummaryView();
+	// }
+
+	return false;
+}
+
 void UNiagaraStackSummaryViewCollection::RefreshChildrenInternal(const TArray<UNiagaraStackEntry*>& CurrentChildren, TArray<UNiagaraStackEntry*>& NewChildren, TArray<FStackIssue>& NewIssues)
 {	
 	TSharedPtr<FNiagaraEmitterViewModel> ViewModel = GetEmitterViewModel();
@@ -547,6 +557,9 @@ void UNiagaraStackSummaryViewCollection::RefreshChildrenInternal(const TArray<UN
 			EmtpySummaryMessage = NewObject<UNiagaraStackItemTextContent>(this);
 			EmtpySummaryMessage->Initialize(CreateDefaultChildRequiredData(), EmptyAssignmentNodeMessageText, GetStackEditorDataKey());
 		}
+		
+		//EmtpySummaryMessage->SetIsHidden(!GetShouldShowInStack());
+		
 		NewChildren.Add(EmtpySummaryMessage);	
 	}
 }
diff --git a/Engine/Plugins/FX/Niagara/Source/NiagaraEditor/Private/Widgets/SNiagaraSystemViewport.cpp b/Engine/Plugins/FX/Niagara/Source/NiagaraEditor/Private/Widgets/SNiagaraSystemViewport.cpp
--- a/Engine/Plugins/FX/Niagara/Source/NiagaraEditor/Private/Widgets/SNiagaraSystemViewport.cpp
+++ b/Engine/Plugins/FX/Niagara/Source/NiagaraEditor/Private/Widgets/SNiagaraSystemViewport.cpp
@@ -238,7 +238,6 @@ void FNiagaraSystemViewportClient::Draw(FViewport* InViewport, FCanvas* Canvas)
 
 	if (bCaptureScreenShot && ScreenShotOwner.IsValid() && OnScreenShotCaptured.IsBound())
 	{
-
 		int32 SrcWidth = InViewport->GetSizeXY().X;
 		int32 SrcHeight = InViewport->GetSizeXY().Y;
 		// Read the contents of the viewport into an array.
@@ -260,6 +259,7 @@ void FNiagaraSystemViewportClient::Draw(FViewport* InViewport, FCanvas* Canvas)
 			UTexture2D* ThumbnailImage = FImageUtils::CreateTexture2D(ScaledWidth, ScaledHeight, ScaledBitmap, ScreenShotOwner.Get(), TEXT("ThumbnailTexture"), RF_NoFlags, Params);
 
 			OnScreenShotCaptured.Execute(ThumbnailImage);
+			NiagaraViewport->GetPreviewComponent()->MarkRenderStateDirty();
 		}
 
 		bCaptureScreenShot = false;
@@ -695,7 +695,7 @@ void SNiagaraSystemViewport::Construct(const FArguments& InArgs, TSharedRef<FNia
 	float Roll = 0.0;
 	AdvancedPreviewScene->SetLightDirection(FRotator(Pitch, Yaw, Roll));
 
-	OnThumbnailCaptured = InArgs._OnThumbnailCaptured;
+	OnThumbnailCapturedDelegate = InArgs._OnThumbnailCaptured;
 	Sequencer = InArgs._Sequencer;
 	
 	SEditorViewport::Construct( SEditorViewport::FArguments() );
@@ -733,10 +733,23 @@ SNiagaraSystemViewport::~SNiagaraSystemViewport()
 	}
 }
 
-void SNiagaraSystemViewport::CreateThumbnail(UObject* InScreenShotOwner)
+void SNiagaraSystemViewport::CreateThumbnail(UObject* InScreenShotOwner, TOptional<FGuid> InEmitterToCaptureThumbnailFor)
 {
 	if (SystemViewportClient.IsValid() && PreviewComponent != nullptr)
 	{
+		PreviewComponent->MarkRenderStateDirty();
+
+		// If we want to capture the thumbnail for a specific emitter in a system, we make sure isolation state is correctly handled
+		if(InEmitterToCaptureThumbnailFor.IsSet() && SystemViewModel.IsValid() && SystemViewModel.Pin()->GetEditMode() == ENiagaraSystemViewModelEditMode::SystemAsset)
+		{
+			if(TSharedPtr<FNiagaraEmitterHandleViewModel> EditableEmitterHandleViewModel = SystemViewModel.Pin()->GetEmitterHandleViewModelById(InEmitterToCaptureThumbnailFor.GetValue()))
+			{
+				SystemViewModel.Pin()->CacheIsolatedEmitterState();
+				EditableEmitterHandleViewModel->SetIsIsolated(true);
+			}
+		}
+		
+		EmitterToCaptureThumbnailFor = InEmitterToCaptureThumbnailFor;
 		SystemViewportClient->bCaptureScreenShot = true;
 		SystemViewportClient->ScreenShotOwner = InScreenShotOwner;
 	}
@@ -877,7 +890,8 @@ bool SNiagaraSystemViewport::IsVisible() const
 
 void SNiagaraSystemViewport::OnScreenShotCaptured(UTexture2D* ScreenShot)
 {
-	OnThumbnailCaptured.ExecuteIfBound(ScreenShot);
+	OnThumbnailCapturedDelegate.ExecuteIfBound(ScreenShot, EmitterToCaptureThumbnailFor);
+	EmitterToCaptureThumbnailFor.Reset();
 }
 
 void SNiagaraSystemViewport::BindCommands()
diff --git a/Engine/Plugins/FX/Niagara/Source/NiagaraEditor/Private/Widgets/SNiagaraSystemViewport.h b/Engine/Plugins/FX/Niagara/Source/NiagaraEditor/Private/Widgets/SNiagaraSystemViewport.h
--- a/Engine/Plugins/FX/Niagara/Source/NiagaraEditor/Private/Widgets/SNiagaraSystemViewport.h
+++ b/Engine/Plugins/FX/Niagara/Source/NiagaraEditor/Private/Widgets/SNiagaraSystemViewport.h
@@ -27,7 +27,7 @@ class UNiagaraEffectType;
 class SNiagaraSystemViewport : public SEditorViewport, public FGCObject, public ICommonEditorViewportToolbarInfoProvider
 {
 public:
-	DECLARE_DELEGATE_OneParam(FOnThumbnailCaptured, UTexture2D*);
+	DECLARE_DELEGATE_TwoParams(FOnThumbnailCaptured, UTexture2D*, TOptional<FGuid> /** An optional emitter handle guid */);
 
 public:
 	SLATE_BEGIN_ARGS( SNiagaraSystemViewport ){}
@@ -88,7 +88,8 @@ public:
 
 	bool GetDrawElement(EDrawElements Element) const;
 	void ToggleDrawElement(EDrawElements Element);
-	void CreateThumbnail(UObject* InScreenShotOwner);
+	/** By specifying an emitter guid, we can record thumbnails for individual emitters within a system.*/
+	void CreateThumbnail(UObject* InScreenShotOwner, TOptional<FGuid> EmitterToCaptureThumbnailFor);
 
 	bool IsToggleOrbitChecked() const;
 	void ToggleOrbit();
@@ -128,7 +129,9 @@ private:
 	TSharedPtr<class FAdvancedPreviewScene> AdvancedPreviewScene;
 
 	TSharedPtr<STextBlock> CompileText;
-	
+
+	TOptional<FGuid> EmitterToCaptureThumbnailFor;
+
 	/** Pointer back to the material editor tool that owns us */
 	//TWeakPtr<INiagaraSystemEditor> SystemEditorPtr;
 	
@@ -139,7 +142,7 @@ private:
 
 	uint32 DrawFlags = 0;
 
-	FOnThumbnailCaptured OnThumbnailCaptured;
+	FOnThumbnailCaptured OnThumbnailCapturedDelegate;
 
 	/** Used on tick to determine if a view transition was active so we can restore view settings at the end of it */
 	bool bIsViewTransitioning = false;
diff --git a/Engine/Plugins/FX/Niagara/Source/NiagaraEditor/Public/NiagaraEmitterEditorData.h b/Engine/Plugins/FX/Niagara/Source/NiagaraEditor/Public/NiagaraEmitterEditorData.h
--- a/Engine/Plugins/FX/Niagara/Source/NiagaraEditor/Public/NiagaraEmitterEditorData.h
+++ b/Engine/Plugins/FX/Niagara/Source/NiagaraEditor/Public/NiagaraEmitterEditorData.h
@@ -46,6 +46,9 @@ public:
 	
 	UNiagaraHierarchyRoot* GetSummaryRoot() const { return SummaryViewRoot; }
 	const TArray<UNiagaraHierarchySection*>& GetSummarySections() const;
+
+	UTexture2D* GetThumbnail() const { return EmitterThumbnail; }
+	void SetThumbnail(UTexture2D* InThumbnail) { EmitterThumbnail = InThumbnail; OnPersistentDataChanged().Broadcast(); }
 private:
 	UPROPERTY(Instanced)
 	TObjectPtr<UNiagaraStackEditorData> StackEditorData;
@@ -69,11 +72,15 @@ private:
 	UPROPERTY()
 	TObjectPtr<UNiagaraHierarchyRoot> SummaryViewRoot = nullptr;
 
+	UPROPERTY()
+	TObjectPtr<UTexture2D> EmitterThumbnail = nullptr;
+	
 	FSimpleMulticastDelegate OnSummaryViewStateChangedDelegate;
 	
 	void StackEditorDataChanged();
 
 	void PostLoad_TransferSummaryDataToNewFormat();
+	void PostLoad_TransferEmitterThumbnailImage(UObject* Owner);
 	void PostLoad_TransferModuleStackNotesToNewFormat(UObject* Owner);
 };
 
diff --git a/Engine/Plugins/FX/Niagara/Source/NiagaraEditor/Public/NiagaraOverviewNode.h b/Engine/Plugins/FX/Niagara/Source/NiagaraEditor/Public/NiagaraOverviewNode.h
--- a/Engine/Plugins/FX/Niagara/Source/NiagaraEditor/Public/NiagaraOverviewNode.h
+++ b/Engine/Plugins/FX/Niagara/Source/NiagaraEditor/Public/NiagaraOverviewNode.h
@@ -16,7 +16,7 @@ public:
 	NIAGARAEDITOR_API void Initialize(UNiagaraSystem* InOwningSystem);
 	NIAGARAEDITOR_API void Initialize(UNiagaraSystem* InOwningSystem, FGuid InEmitterHandleGuid);
 	NIAGARAEDITOR_API const FGuid GetEmitterHandleGuid() const;
-	NIAGARAEDITOR_API struct FNiagaraEmitterHandle* TryGetEmitterHandle();
+	NIAGARAEDITOR_API struct FNiagaraEmitterHandle* TryGetEmitterHandle() const;
 	
 	//~ Begin UEdGraphNode Interface
 	/** Gets the name of this node, shown in title bar */
@@ -39,7 +39,7 @@ public:
 	NIAGARAEDITOR_API virtual bool CanCreateUnderSpecifiedSchema(const UEdGraphSchema* Schema) const override;
 	//~ End UEdGraphNode Interface
 
-	NIAGARAEDITOR_API UNiagaraSystem* GetOwningSystem();
+	NIAGARAEDITOR_API UNiagaraSystem* GetOwningSystem() const;
 
 	void RequestRename() { bRenamePending = true; }
 	void RenameStarted() { bRenamePending = false; }
diff --git a/Engine/Plugins/FX/Niagara/Source/NiagaraEditor/Public/ViewModels/NiagaraEmitterHandleViewModel.h b/Engine/Plugins/FX/Niagara/Source/NiagaraEditor/Public/ViewModels/NiagaraEmitterHandleViewModel.h
--- a/Engine/Plugins/FX/Niagara/Source/NiagaraEditor/Public/ViewModels/NiagaraEmitterHandleViewModel.h
+++ b/Engine/Plugins/FX/Niagara/Source/NiagaraEditor/Public/ViewModels/NiagaraEmitterHandleViewModel.h
@@ -88,6 +88,10 @@ public:
 	/** Sets whether or not this emitter is isolated. May affect other emitters in the system. */
 	NIAGARAEDITOR_API void SetIsIsolated(bool InIsIsolated);
 
+	NIAGARAEDITOR_API void ToggleIsIsolated();
+
+	NIAGARAEDITOR_API void RequestCaptureThumbnail();
+	
 	/** Gets the check state for the is enabled check box. */
 	NIAGARAEDITOR_API ECheckBoxState GetIsEnabledCheckState() const;
 
@@ -112,6 +116,8 @@ public:
 	/** Sets whether or not this emitter handle has a rename pending. */
 	NIAGARAEDITOR_API void SetIsRenamePending(bool bInIsRenamePending);
 
+	void RequestEmitterSelection(bool bClearSelection);
+
 	/** Gets a multicast delegate which is called any time a property on the handle changes. */
 	FOnPropertyChanged& OnPropertyChanged();
 
diff --git a/Engine/Plugins/FX/Niagara/Source/NiagaraEditor/Public/ViewModels/NiagaraEmitterViewModel.h b/Engine/Plugins/FX/Niagara/Source/NiagaraEditor/Public/ViewModels/NiagaraEmitterViewModel.h
--- a/Engine/Plugins/FX/Niagara/Source/NiagaraEditor/Public/ViewModels/NiagaraEmitterViewModel.h
+++ b/Engine/Plugins/FX/Niagara/Source/NiagaraEditor/Public/ViewModels/NiagaraEmitterViewModel.h
@@ -43,6 +43,7 @@ public:
 	DECLARE_MULTICAST_DELEGATE(FOnParentRemoved);
 	DECLARE_MULTICAST_DELEGATE_TwoParams(FOnScriptGraphChanged, const FEdGraphEditAction&, const UNiagaraScript&);
 	DECLARE_MULTICAST_DELEGATE_TwoParams(FOnScriptParameterStoreChanged, const FNiagaraParameterStore&, const UNiagaraScript&);
+	DECLARE_DELEGATE_OneParam(FOnEmitterSelectionRequested, bool /* bClearSelection*/)
 
 public:
 	/** Creates a new emitter editor view model.  It must be initialized before use. */
@@ -118,6 +119,8 @@ public:
 	/** Gets a multicast delegate which is called any time a parameter store on a script owned by this emitter changes. */
 	FOnScriptParameterStoreChanged& OnScriptParameterStoreChanged();
 
+	FOnEmitterSelectionRequested& OnEmitterSelectionRequested();
+
 	/** Gets editor specific data which can be stored per emitter.  If this data hasn't been created the default version will be returned. */
 	NIAGARAEDITOR_API const UNiagaraEmitterEditorData& GetEditorData() const;
 
@@ -155,7 +158,9 @@ private:
 	void OnEmitterPropertiesChanged();
 
 	void OnSummaryViewHierarchyChanged();
-	
+
+	void UpdateSelectionOnSummaryViewStateChanged() const;
+
 	virtual FString GetReferencerName() const override;
 	virtual void AddReferencedObjects(FReferenceCollector& Collector) override;
 	
@@ -197,6 +202,8 @@ private:
 
 	FOnScriptParameterStoreChanged OnScriptParameterStoreChangedDelegate;
 
+	FOnEmitterSelectionRequested OnEmitterSelectionRequestedDelegate;
+	
 	TNiagaraViewModelManager<UNiagaraEmitter, FNiagaraEmitterViewModel>::Handle RegisteredHandle;
 
 	UEnum* ExecutionStateEnum;
diff --git a/Engine/Plugins/FX/Niagara/Source/NiagaraEditor/Public/ViewModels/NiagaraSystemViewModel.h b/Engine/Plugins/FX/Niagara/Source/NiagaraEditor/Public/ViewModels/NiagaraSystemViewModel.h
--- a/Engine/Plugins/FX/Niagara/Source/NiagaraEditor/Public/ViewModels/NiagaraSystemViewModel.h
+++ b/Engine/Plugins/FX/Niagara/Source/NiagaraEditor/Public/ViewModels/NiagaraSystemViewModel.h
@@ -135,6 +135,7 @@ public:
 	DECLARE_MULTICAST_DELEGATE_TwoParams(FOnExternalRemoveParameter, const FNiagaraVariableBase&,  UNiagaraEmitter*);
 
 	DECLARE_DELEGATE_OneParam(FOnChangeWorkflowMode, FName /* WorkflowMode */)
+	DECLARE_DELEGATE_OneParam(FOnEmitterThumbnailRequested, FGuid EmitterGuid)
 public:
 	struct FEmitterHandleToDuplicate
 	{
@@ -230,6 +231,9 @@ public:
 	
 	/** Gets the delegate that is bound to the owning toolkit's SetCurrentMode function. */
 	NIAGARAEDITOR_API FOnChangeWorkflowMode& OnChangeWorkflowMode();
+
+	/** Gets the delegate that is bound to the owning toolkit's CreateThumbnail function. */
+	FOnEmitterThumbnailRequested& OnEmitterThumbnailRequested();
 	
 	/** Adds a new emitter to the System from an emitter asset data. */
 	NIAGARAEDITOR_API TSharedPtr<FNiagaraEmitterHandleViewModel> AddEmitterFromAssetData(const FAssetData& AssetData);
@@ -332,6 +336,9 @@ public:
 	/** Isolates the supplied emitters.  This will remove all other emitters from isolation. */
 	NIAGARAEDITOR_API void IsolateEmitters(TArray<FGuid> EmitterHandlesIdsToIsolate);
 
+	void CacheIsolatedEmitterState();
+	void RestoreIsolatedEmitterState();
+	
 	/** Disable the supplied emitters. */
 	NIAGARAEDITOR_API void DisableEmitters(TArray<FGuid> EmitterHandlesIdsToDisable);
 
@@ -658,6 +665,9 @@ private:
 
 	/** A delegate that calls the toolkit's set current mode function */
 	FOnChangeWorkflowMode OnChangeWorkflowModeDelegate;
+
+	/** A delegate that calls the toolkit's Create Thumbnail function */
+	FOnEmitterThumbnailRequested OnEmitterThumbnailRequestedDelegate;
 	
 	/** A delegate which is used to generate the content for the add menu in sequencer. */
 	FOnGetAddMenuContent OnGetSequencerAddMenuContent;
@@ -719,6 +729,8 @@ private:
 	/** A flag indicating that a reset has been request on the next tick */
 	bool bResetRequestPending;
 
+	TArray<FGuid> CachedEmitterIsolationState;
+	
 	/** The system toolkit commands. */
 	TWeakPtr<class FUICommandList> ToolkitCommands;
 
diff --git a/Engine/Plugins/FX/Niagara/Source/NiagaraEditor/Public/ViewModels/Stack/NiagaraStackEmitterPropertiesGroup.h b/Engine/Plugins/FX/Niagara/Source/NiagaraEditor/Public/ViewModels/Stack/NiagaraStackEmitterPropertiesGroup.h
--- a/Engine/Plugins/FX/Niagara/Source/NiagaraEditor/Public/ViewModels/Stack/NiagaraStackEmitterPropertiesGroup.h
+++ b/Engine/Plugins/FX/Niagara/Source/NiagaraEditor/Public/ViewModels/Stack/NiagaraStackEmitterPropertiesGroup.h
@@ -24,7 +24,7 @@ public:
 
 	virtual EIconMode GetSupportedIconMode() const override { return EIconMode::Brush; }
 	NIAGARAEDITOR_API virtual const FSlateBrush* GetIconBrush() const override;
-	virtual bool SupportsSecondaryIcon() const override { return true; }
+	virtual bool SupportsSecondaryIcon() const override;
 	NIAGARAEDITOR_API virtual const FSlateBrush* GetSecondaryIconBrush() const override;
 	virtual bool GetCanExpandInOverview() const override { return false; }
 	virtual bool GetShouldShowInStack() const override { return false; }
diff --git a/Engine/Plugins/FX/Niagara/Source/NiagaraEditor/Public/ViewModels/Stack/NiagaraStackEmitterSettingsGroup.h b/Engine/Plugins/FX/Niagara/Source/NiagaraEditor/Public/ViewModels/Stack/NiagaraStackEmitterSettingsGroup.h
--- a/Engine/Plugins/FX/Niagara/Source/NiagaraEditor/Public/ViewModels/Stack/NiagaraStackEmitterSettingsGroup.h
+++ b/Engine/Plugins/FX/Niagara/Source/NiagaraEditor/Public/ViewModels/Stack/NiagaraStackEmitterSettingsGroup.h
@@ -73,14 +73,12 @@ public:
 	NIAGARAEDITOR_API virtual void OnEditButtonClicked() override;
 	NIAGARAEDITOR_API virtual TOptional<FText> GetEditModeButtonText() const override;
 	NIAGARAEDITOR_API virtual TOptional<FText> GetEditModeButtonTooltip() const override;
-	NIAGARAEDITOR_API virtual EVisibility IsEditButtonVisible() const override;
 
 protected:
-
 	NIAGARAEDITOR_API virtual void RefreshChildrenInternal(const TArray<UNiagaraStackEntry*>& CurrentChildren, TArray<UNiagaraStackEntry*>& NewChildren, TArray<FStackIssue>& NewIssues) override;
+	
 	virtual void ToggleShowAdvancedInternal() override;
-private:
-	NIAGARAEDITOR_API void SelectSummaryNodesFromEmitterEditorDataRootNodes(TArray<TSharedRef<IDetailTreeNode>> Source, TArray<TSharedRef<IDetailTreeNode>>* Selected);
+
 private:
 	FVersionedNiagaraEmitterWeakPtr Emitter;
 
@@ -99,37 +97,14 @@ public:
 	virtual EIconMode GetSupportedIconMode() const override { return EIconMode::Text; }
 	NIAGARAEDITOR_API virtual  FText GetIconText() const override;
 	virtual bool GetCanExpandInOverview() const override { return false; }
-	virtual bool GetShouldShowInStack() const override { return false; }
 
 protected:
-	NIAGARAEDITOR_API void RefreshChildrenInternal(const TArray<UNiagaraStackEntry*>& CurrentChildren, TArray<UNiagaraStackEntry*>& NewChildren, TArray<FStackIssue>& NewIssues) override;
-
+	virtual bool GetShouldShowInStack() const override { return false; }
+	virtual bool GetShouldShowInOverview() const override { return false; }
+	
+	NIAGARAEDITOR_API virtual void RefreshChildrenInternal(const TArray<UNiagaraStackEntry*>& CurrentChildren, TArray<UNiagaraStackEntry*>& NewChildren, TArray<FStackIssue>& NewIssues) override;
 
 private:
 	UPROPERTY()
 	TObjectPtr<UNiagaraStackEmitterSummaryItem> SummaryItem;
 };
-
-
-UCLASS(MinimalAPI)
-class UNiagaraStackSummaryViewCollapseButton : public UNiagaraStackEntry
-{
-	GENERATED_BODY()
-
-public:
-	NIAGARAEDITOR_API void Initialize(FRequiredEntryData InRequiredEntryData);
-
-	//~ UNiagaraStackEntry interface
-	NIAGARAEDITOR_API virtual FText GetDisplayName() const override;
-	NIAGARAEDITOR_API virtual EStackRowStyle GetStackRowStyle() const override;
-	NIAGARAEDITOR_API virtual FText GetTooltipText() const override;
-
-	virtual bool GetCanExpandInOverview() const override { return true; }
-	NIAGARAEDITOR_API virtual bool GetIsEnabled() const override;
-	virtual void SetIsEnabled(bool bEnabled) {}
-	virtual bool SupportsChangeEnabled() const { return false; }
-	
-protected:
-	NIAGARAEDITOR_API virtual void RefreshChildrenInternal(const TArray<UNiagaraStackEntry*>& CurrentChildren, TArray<UNiagaraStackEntry*>& NewChildren, TArray<FStackIssue>& NewIssues) override;
-
-};
diff --git a/Engine/Plugins/FX/Niagara/Source/NiagaraEditor/Public/ViewModels/Stack/NiagaraStackRoot.h b/Engine/Plugins/FX/Niagara/Source/NiagaraEditor/Public/ViewModels/Stack/NiagaraStackRoot.h
--- a/Engine/Plugins/FX/Niagara/Source/NiagaraEditor/Public/ViewModels/Stack/NiagaraStackRoot.h
+++ b/Engine/Plugins/FX/Niagara/Source/NiagaraEditor/Public/ViewModels/Stack/NiagaraStackRoot.h
@@ -23,7 +23,7 @@ class UNiagaraStackRoot : public UNiagaraStackEntry
 public:
 	NIAGARAEDITOR_API UNiagaraStackRoot();
 	
-	NIAGARAEDITOR_API void Initialize(FRequiredEntryData InRequiredEntryData, bool bInIncludeSystemInformation, bool bInIncludeEmitterInformation);
+	void Initialize(FRequiredEntryData InRequiredEntryData, bool bInIncludeSystemInformation, bool bInIncludeEmitterInformation);
 	NIAGARAEDITOR_API virtual void FinalizeInternal() override;
 
 	NIAGARAEDITOR_API virtual bool GetCanExpand() const override;
@@ -37,6 +37,11 @@ public:
 	{
 		return CommentCollection;
 	}
+
+	UNiagaraStackEmitterSummaryGroup* GetEmitterSummaryGroup() const
+	{
+		return EmitterSummaryGroup;
+	}
 	
 protected:
 	NIAGARAEDITOR_API virtual void RefreshChildrenInternal(const TArray<UNiagaraStackEntry*>& CurrentChildren, TArray<UNiagaraStackEntry*>& NewChildren, TArray<FStackIssue>& NewIssues) override;
@@ -45,13 +50,19 @@ private:
 	NIAGARAEDITOR_API void EmitterArraysChanged();
 	NIAGARAEDITOR_API void OnSummaryViewStateChanged();
 
+	void AddDefaultViewFilter();
+	void RemoveDefaultViewFilter();
+	
+	void AddSummaryViewFilter();
+	void RemoveSummaryViewFilter();
+
+	NIAGARAEDITOR_API bool FilterForDefaultView(const UNiagaraStackEntry& NiagaraStackEntry) const;
+	NIAGARAEDITOR_API bool FilterForSummaryView(const UNiagaraStackEntry& NiagaraStackEntry) const;
+
 private:
 	UPROPERTY()
 	TObjectPtr<UNiagaraStackSystemPropertiesGroup> SystemPropertiesGroup;
 
-	UPROPERTY()
-	TObjectPtr<UNiagaraStackSystemUserParametersGroup> SystemUserParametersGroup;
-
 	UPROPERTY()
 	TObjectPtr<UNiagaraStackScriptItemGroup> SystemSpawnGroup;
 
@@ -81,10 +92,10 @@ private:
 
 	UPROPERTY()
 	TObjectPtr<UNiagaraStackCommentCollection> CommentCollection;
-	
-	UPROPERTY()
-	TObjectPtr<UNiagaraStackSummaryViewCollapseButton> SummaryCollapseButton;
 
 	bool bIncludeSystemInformation;
 	bool bIncludeEmitterInformation;
+
+	FDelegateHandle DefaultViewFilterHandle;
+	FDelegateHandle SummaryViewFilterHandle;
 };
diff --git a/Engine/Plugins/FX/Niagara/Source/NiagaraEditor/Public/ViewModels/Stack/NiagaraStackSummaryViewInputCollection.h b/Engine/Plugins/FX/Niagara/Source/NiagaraEditor/Public/ViewModels/Stack/NiagaraStackSummaryViewInputCollection.h
--- a/Engine/Plugins/FX/Niagara/Source/NiagaraEditor/Public/ViewModels/Stack/NiagaraStackSummaryViewInputCollection.h
+++ b/Engine/Plugins/FX/Niagara/Source/NiagaraEditor/Public/ViewModels/Stack/NiagaraStackSummaryViewInputCollection.h
@@ -26,6 +26,7 @@ public:
 
 	void RefreshForAdvancedToggle();
 protected:
+	virtual bool GetShouldShowInStack() const override;
 	NIAGARAEDITOR_API virtual void RefreshChildrenInternal(const TArray<UNiagaraStackEntry*>& CurrentChildren, TArray<UNiagaraStackEntry*>& NewChildren, TArray<FStackIssue>& NewIssues) override;
 
 	NIAGARAEDITOR_API virtual void GetSectionsInternal(TArray<FNiagaraStackSection>& OutStackSections) const override;
diff --git a/Engine/Plugins/FX/Niagara/Source/NiagaraEditorWidgets/NiagaraEditorWidgets.Build.cs b/Engine/Plugins/FX/Niagara/Source/NiagaraEditorWidgets/NiagaraEditorWidgets.Build.cs
--- a/Engine/Plugins/FX/Niagara/Source/NiagaraEditorWidgets/NiagaraEditorWidgets.Build.cs
+++ b/Engine/Plugins/FX/Niagara/Source/NiagaraEditorWidgets/NiagaraEditorWidgets.Build.cs
@@ -34,7 +34,8 @@ public class NiagaraEditorWidgets : ModuleRules
 			"CurveEditor",
 			"DesktopPlatform",
 			"ToolWidgets",
-			"KismetWidgets"
+			"KismetWidgets",
+			"ToolMenus"
 		});
 
 		PrivateIncludePathModuleNames.AddRange(new string[] {
diff --git a/Engine/Plugins/FX/Niagara/Source/NiagaraEditorWidgets/Private/SNiagaraOverviewStack.cpp b/Engine/Plugins/FX/Niagara/Source/NiagaraEditorWidgets/Private/SNiagaraOverviewStack.cpp
--- a/Engine/Plugins/FX/Niagara/Source/NiagaraEditorWidgets/Private/SNiagaraOverviewStack.cpp
+++ b/Engine/Plugins/FX/Niagara/Source/NiagaraEditorWidgets/Private/SNiagaraOverviewStack.cpp
@@ -650,6 +650,8 @@ private:
 
 void SNiagaraOverviewStack::Construct(const FArguments& InArgs, UNiagaraStackViewModel& InStackViewModel, UNiagaraSystemSelectionViewModel& InOverviewSelectionViewModel)
 {
+	AllowedClasses = InArgs._AllowedClasses;
+	
 	StackCommandContext = MakeShared<FNiagaraStackCommandContext>();
 
 	bUpdatingOverviewSelectionFromStackSelection = false;
@@ -847,19 +849,16 @@ void SNiagaraOverviewStack::RefreshEntryList()
 	{
 		FlattenedEntryList.Empty();
 		EntryObjectKeyToParentChain.Empty();
-		TArray<UClass*> AcceptableClasses;
-		AcceptableClasses.Add(UNiagaraStackItemGroup::StaticClass());
-		AcceptableClasses.Add(UNiagaraStackItem::StaticClass());
 
 		UNiagaraStackEntry* RootEntry = StackViewModel->GetRootEntry();
 		checkf(RootEntry != nullptr, TEXT("Root entry was null."));
-		TArray<UNiagaraStackEntry*> RootChildren;
+		TArray<UNiagaraStackEntry*> RootChildren;		
 		RootEntry->GetFilteredChildren(RootChildren);
 		for (UNiagaraStackEntry* RootChild : RootChildren)
 		{
 			checkf(RootEntry != nullptr, TEXT("Root entry child was null."));
 			TArray<UNiagaraStackEntry*> ParentChain;
-			AddEntriesRecursive(*RootChild, FlattenedEntryList, AcceptableClasses, ParentChain);
+			AddEntriesRecursive(*RootChild, FlattenedEntryList, AllowedClasses, ParentChain);
 		}
 
 		bRefreshEntryListPending = false;
diff --git a/Engine/Plugins/FX/Niagara/Source/NiagaraEditorWidgets/Private/SNiagaraOverviewStack.h b/Engine/Plugins/FX/Niagara/Source/NiagaraEditorWidgets/Private/SNiagaraOverviewStack.h
--- a/Engine/Plugins/FX/Niagara/Source/NiagaraEditorWidgets/Private/SNiagaraOverviewStack.h
+++ b/Engine/Plugins/FX/Niagara/Source/NiagaraEditorWidgets/Private/SNiagaraOverviewStack.h
@@ -8,6 +8,8 @@
 #include "Widgets/Views/SListView.h"
 #include "UObject/ObjectKey.h"
 #include "Framework/Commands/UICommandList.h"
+#include "ViewModels/Stack/NiagaraStackItem.h"
+#include "ViewModels/Stack/NiagaraStackItemGroup.h"
 
 class UNiagaraStackEntry;
 class UNiagaraStackItem;
@@ -20,12 +22,14 @@ class SNiagaraOverviewStack : public SCompoundWidget
 {
 public:
 	SLATE_BEGIN_ARGS(SNiagaraOverviewStack)
+		: _AllowedClasses({UNiagaraStackItem::StaticClass(), UNiagaraStackItemGroup::StaticClass()})
 	{}
+		SLATE_ARGUMENT(TArray<UClass*>, AllowedClasses)
 	SLATE_END_ARGS();
 
 	void Construct(const FArguments& InArgs, UNiagaraStackViewModel& InStackViewModel, UNiagaraSystemSelectionViewModel& InOverviewSelectionViewModel);
 
-	~SNiagaraOverviewStack();
+	virtual	~SNiagaraOverviewStack() override;
 
 	virtual bool SupportsKeyboardFocus() const override;
 
@@ -96,6 +100,8 @@ private:
 
 	TSharedPtr<FNiagaraStackCommandContext> StackCommandContext;
 
+	TArray<UClass*> AllowedClasses;
+
 	bool bRefreshEntryListPending;
 	bool bUpdatingOverviewSelectionFromStackSelection;
 	bool bUpdatingStackSelectionFromOverviewSelection;
diff --git a/Engine/Plugins/FX/Niagara/Source/NiagaraEditorWidgets/Private/SNiagaraOverviewStackNode.cpp b/Engine/Plugins/FX/Niagara/Source/NiagaraEditorWidgets/Private/SNiagaraOverviewStackNode.cpp
--- a/Engine/Plugins/FX/Niagara/Source/NiagaraEditorWidgets/Private/SNiagaraOverviewStackNode.cpp
+++ b/Engine/Plugins/FX/Niagara/Source/NiagaraEditorWidgets/Private/SNiagaraOverviewStackNode.cpp
@@ -20,6 +20,10 @@
 #include "SGraphPanel.h"
 #include "SLevelOfDetailBranchNode.h"
 #include "SNiagaraOverviewStack.h"
+#include "Stack/SNiagaraDeterminismToggle.h"
+#include "Stack/SNiagaraLocalSpaceToggle.h"
+#include "Stack/SNiagaraSummaryViewToggle.h"
+#include "Stack/SNiagaraSimTargetToggle.h"
 #include "Stack/SNiagaraStackIssueIcon.h"
 #include "Styling/AppStyle.h"
 #include "Styling/StyleColors.h"
@@ -27,6 +31,7 @@
 #include "ViewModels/NiagaraEmitterViewModel.h"
 #include "ViewModels/NiagaraSystemSelectionViewModel.h"
 #include "ViewModels/NiagaraSystemViewModel.h"
+#include "ViewModels/Stack/NiagaraStackEmitterPropertiesGroup.h"
 #include "ViewModels/Stack/NiagaraStackItemGroup.h"
 #include "ViewModels/Stack/NiagaraStackModuleItem.h"
 #include "ViewModels/Stack/NiagaraStackRendererItem.h"
@@ -38,10 +43,15 @@
 #include "Widgets/SToolTip.h"
 #include "Widgets/SWidget.h"
 #include "Widgets/Input/SComboButton.h"
+#include "Widgets/Layout/SScaleBox.h"
+#include "Widgets/Layout/SSpacer.h"
+#include "AssetThumbnail.h"
 
 #define LOCTEXT_NAMESPACE "NiagaraOverviewStackNode"
 
-constexpr float ThumbnailSize = 24.0f;
+constexpr float RendererThumbnailSize = 24.f;
+constexpr float EmitterThumbnailSize = 200.f;
+constexpr float SummaryTitleFixedSize = 125.f;
 
 void SNiagaraOverviewStackNode::Construct(const FArguments& InArgs, UNiagaraOverviewNode* InNode)
 {
@@ -52,10 +62,10 @@ void SNiagaraOverviewStackNode::Construct(const FArguments& InArgs, UNiagaraOver
 	bIsHoveringThumbnail = false;
 	bTopContentBarRefreshPending = true;
 	CurrentIssueIndex = -1;
+	// we are reducing the left margin to better place our widgets
+	TitleBorderMargin = FMargin(2.f, 5.f, 30.f, 3.f);
 
 	EmitterHandleViewModelWeak.Reset();
-
-	TopContentBar = SNew(SHorizontalBox);
 	
 	if (OverviewStackNode->GetOwningSystem() != nullptr)
 	{
@@ -63,6 +73,8 @@ void SNiagaraOverviewStackNode::Construct(const FArguments& InArgs, UNiagaraOver
 		TSharedPtr<FNiagaraSystemViewModel> OwningSystemViewModel = NiagaraEditorModule.GetExistingViewModelForSystem(OverviewStackNode->GetOwningSystem());
 		if (OwningSystemViewModel.IsValid())
 		{
+			// if the emitter handle view models have updates, make sure we rebind the editor delegates. The previous bindings might have become invalid due to a merge
+			OwningSystemViewModel->OnEmitterHandleViewModelsChanged().AddSP(this, &SNiagaraOverviewStackNode::BindEditorDataDelegates);
 			if (OverviewStackNode->GetEmitterHandleGuid().IsValid() == false)
 			{
 				StackViewModel = OwningSystemViewModel->GetSystemStackViewModel();
@@ -92,6 +104,8 @@ void SNiagaraOverviewStackNode::Construct(const FArguments& InArgs, UNiagaraOver
 		}
 	}
 
+	BindEditorDataDelegates();
+
 	UpdateGraphNode();
 }
 
@@ -101,153 +115,80 @@ SNiagaraOverviewStackNode::~SNiagaraOverviewStackNode()
 	{
 		ScalabilityViewModel->OnScalabilityModeChanged().RemoveAll(this);
 	}
+
+	UnbindEditorDataDelegates();
 }
 
-TSharedRef<SWidget> SNiagaraOverviewStackNode::CreateTitleWidget(TSharedPtr<SNodeTitle> NodeTitle)
+void SNiagaraOverviewStackNode::BindEditorDataDelegates()
 {
-	TSharedRef<SWidget> DefaultTitle = SGraphNode::CreateTitleWidget(NodeTitle);
+	UnbindEditorDataDelegates();
 
-	if (StackViewModel == nullptr)
+	if(EmitterHandleViewModelWeak.IsValid())
 	{
-		return SNew(SHorizontalBox)
-			+ SHorizontalBox::Slot()
-			.Padding(0, 0, 5, 0)
-			[
-				SNew(STextBlock)
-				.Text(LOCTEXT("InvalidNode", "INVALID"))
-			]
-			+ SHorizontalBox::Slot()
-			[
-				DefaultTitle
-			];
+		EmitterHandleViewModelWeak.Pin()->GetEmitterViewModel()->GetEditorData().OnPersistentDataChanged().AddSP(this, &SNiagaraOverviewStackNode::RefreshEmitterThumbnailPreview);
+		EmitterHandleViewModelWeak.Pin()->GetEmitterViewModel()->GetEditorData().OnSummaryViewStateChanged().AddSP(this, &SNiagaraOverviewStackNode::UpdateGraphNode);
 	}
+}
 
-	return SNew(SHorizontalBox)
-		// Enabled checkbox
-		+ SHorizontalBox::Slot()
-		.AutoWidth()
-		.Padding(0)
-		[
-			SNew(SCheckBox)
-			.Visibility(this, &SNiagaraOverviewStackNode::GetEnabledCheckBoxVisibility)
-			.IsChecked(this, &SNiagaraOverviewStackNode::GetEnabledCheckState)
-			.OnCheckStateChanged(this, &SNiagaraOverviewStackNode::OnEnabledCheckStateChanged)
-		]
-		// Name
-		+ SHorizontalBox::Slot()
-		.Padding(3, 0, 0, 0)
-		.FillWidth(1.0f)
-		[
-			DefaultTitle
-		];
+void SNiagaraOverviewStackNode::UnbindEditorDataDelegates() const
+{
+	if(EmitterHandleViewModelWeak.IsValid() && EmitterHandleViewModelWeak.Pin()->GetEmitterHandle() != nullptr)
+	{
+		EmitterHandleViewModelWeak.Pin()->GetEmitterViewModel()->GetEditorData().OnPersistentDataChanged().RemoveAll(this);
+		EmitterHandleViewModelWeak.Pin()->GetEmitterViewModel()->GetEditorData().OnSummaryViewStateChanged().RemoveAll(this);
+	}
+}
+
+TSharedRef<SWidget> SNiagaraOverviewStackNode::CreateTitleWidget(TSharedPtr<SNodeTitle> NodeTitle)
+{
+	switch(DisplayMode)
+	{
+	case EDisplayMode::Default:
+		return CreateTitleWidget_Default(NodeTitle);
+	case EDisplayMode::Summary:
+		return CreateTitleWidget_Summary(NodeTitle);
+	default:
+		return CreateTitleWidget_Default(NodeTitle);
+	}
 }
 
 TSharedRef<SWidget> SNiagaraOverviewStackNode::CreateTitleRightWidget()
 {
-	if (StackViewModel == nullptr)
+	switch(DisplayMode)
 	{
-		return SNullWidget::NullWidget;
+	case EDisplayMode::Default:
+		return CreateTitleRightWidget_Default();
+	case EDisplayMode::Summary:
+		return CreateTitleRightWidget_Summary();
+	default:
+		return CreateTitleRightWidget_Default();
 	}
-	return SNew(SHorizontalBox)
+}
 
-		// open parent button
-		+ SHorizontalBox::Slot()
-		.AutoWidth()
-		.HAlign(HAlign_Right)
-		.VAlign(VAlign_Center)
-		.Padding(0, 0, 1, 0)
-		[
-			SNew(SButton)
-			.IsFocusable(false)
-			.ToolTipText(this, &SNiagaraOverviewStackNode::OpenParentEmitterTooltip)
-			.ButtonStyle(FAppStyle::Get(), "HoverHintOnly")
-			.ContentPadding(2)
-			.OnClicked(this, &SNiagaraOverviewStackNode::OpenParentEmitter)
-			.Visibility(this, &SNiagaraOverviewStackNode::GetOpenParentEmitterVisibility)
-			.DesiredSizeScale(FVector2D(14.0f / 30.0f, 14.0f / 30.0f)) // GoToSourceIcon is 30x30, scale down
-			.Content()
-			[
-				SNew(SImage)
-				.Image(FNiagaraEditorWidgetsStyle::Get().GetBrush("NiagaraEditor.Stack.GoToSourceIcon"))
-				.ColorAndOpacity(FSlateColor::UseForeground())
-			]
-		]
-	
-		// version selector
-		+ SHorizontalBox::Slot()
-		.HAlign(HAlign_Right)
-		.AutoWidth()
-		.Padding(1, 0, 2, 0)
-		[
-			SNew(SComboButton)
-			.HasDownArrow(false)
-			.ToolTipText(LOCTEXT("ChangeEmitterVersionToolTip", "Change the parent emitter version"))
-			.ButtonStyle(FAppStyle::Get(), "HoverHintOnly")
-			.ForegroundColor(FSlateColor::UseForeground())
-			.OnGetMenuContent(this, &SNiagaraOverviewStackNode::GetVersionSelectorDropdownMenu)
-			.ContentPadding(FMargin(2))
-			.Visibility(this, &SNiagaraOverviewStackNode::GetVersionSelectorVisibility)
-			.ButtonContent()
-			[
-				SNew(STextBlock)
-				.Font(FAppStyle::Get().GetFontStyle("FontAwesome.10"))
-				.ColorAndOpacity(this, &SNiagaraOverviewStackNode::GetVersionSelectorColor)
-				.Text(FEditorFontGlyphs::Random)
-			]
-		]
+TSharedRef<SWidget> SNiagaraOverviewStackNode::CreateNodeContentArea()
+{
+	switch(DisplayMode)
+	{
+	case EDisplayMode::Default:
+		return CreateNodeContentArea_Default();
+	case EDisplayMode::Summary:
+		return CreateNodeContentArea_Summary();
+	default:
+		return CreateNodeContentArea_Default();
+	}
+}
 
-		// scalability controls
-		+ SHorizontalBox::Slot()
-		.AutoWidth()
-		.Padding(2.f)
-		[
-			SNew(SVerticalBox)
-			+ SVerticalBox::Slot()
-			.AutoHeight()
-			.HAlign(HAlign_Center)
-			.VAlign(VAlign_Center)
-			[
-				// Toggle Button to enter & exit scalability mode
-				SNew(SCheckBox)
-				.Style(&FAppStyle::GetWidgetStyle<FCheckBoxStyle>("ToggleButtonCheckbox"))
-				.IsChecked(this, &SNiagaraOverviewStackNode::IsScalabilityModeActive)
-				.Padding(2.f)
-				.OnCheckStateChanged(this, &SNiagaraOverviewStackNode::OnScalabilityModeStateChanged)
-				[
-					SNew(SBox)
-					.WidthOverride(16.f)
-					.HeightOverride(16.f)
-					[
-						SNew(SImage)
-						.Image(FNiagaraEditorStyle::Get().GetBrush("NiagaraEditor.Scalability"))
-						.Visibility(this, &SNiagaraOverviewStackNode::GetScalabilityIndicatorVisibility)
-						.ToolTipText(FText::FormatOrdered(LOCTEXT("ScalabilityIndicatorToolTip",
-							"This {0} has scalability set up. Inspecting and editing scalability is accessible by entering Scalability Mode by clicking this or the button in the toolbar.."), EmitterHandleViewModelWeak.IsValid() ? FText::FromString("emitter") : FText::FromString("system")))
-					]
-				]
-			]
-			// Spawn Count Scale Info
-			+ SVerticalBox::Slot()
-			.AutoHeight()
-			[
-				SNew(STextBlock)
-				.Text(this, &SNiagaraOverviewStackNode::GetSpawnCountScaleText)
-				.ToolTipText(this, &SNiagaraOverviewStackNode::GetSpawnCountScaleTooltip)
-				.Visibility(this, &SNiagaraOverviewStackNode::GetSpawnCountScaleTextVisibility)
-			]
-		]
-	
-		// issue/error icon
-		+ SHorizontalBox::Slot()
-		.AutoWidth()
-		.HAlign(HAlign_Right)
-		.VAlign(VAlign_Center)
-		.Padding(0, 0, 1, 0)
-		[
-			SNew(SNiagaraStackIssueIcon, StackViewModel, StackViewModel->GetRootEntry())
-			.Visibility(this, &SNiagaraOverviewStackNode::GetIssueIconVisibility)
-			.OnClicked(this, &SNiagaraOverviewStackNode::OnCycleThroughIssues)
-		];
+TSharedRef<SWidget> SNiagaraOverviewStackNode::CreateTopContentBar()
+{
+	switch(DisplayMode)
+	{
+	case EDisplayMode::Default:
+		return CreateTopContentBar_Default();
+	case EDisplayMode::Summary:
+		return CreateTopContentBar_Summary();
+	default:
+		return CreateTopContentBar_Default();
+	}
 }
 
 FText SNiagaraOverviewStackNode::GetSpawnCountScaleText() const
@@ -374,7 +315,7 @@ FReply SNiagaraOverviewStackNode::OnCycleThroughIssues()
 	return FReply::Handled();
 }
 
-TSharedRef<SWidget> SNiagaraOverviewStackNode::CreateThumbnailWidget(UNiagaraStackEntry* InData, TSharedPtr<SWidget> InWidget, TSharedPtr<SWidget> InTooltipWidget)
+TSharedRef<SWidget> SNiagaraOverviewStackNode::CreateRendererThumbnailWidget(UNiagaraStackEntry* InData, TSharedPtr<SWidget> InWidget, TSharedPtr<SWidget> InTooltipWidget)
 {
 	TSharedPtr<SToolTip> ThumbnailTooltipWidget;
 	// If this is just text, don't constrain the size
@@ -410,6 +351,18 @@ TSharedRef<SWidget> SNiagaraOverviewStackNode::CreateThumbnailWidget(UNiagaraSta
 	return InWidget.ToSharedRef();
 }
 
+FReply SNiagaraOverviewStackNode::OnCaptureThumbnailButtonClicked()
+{
+	if(EmitterHandleViewModelWeak.IsValid())
+	{
+		EmitterHandleViewModelWeak.Pin()->RequestCaptureThumbnail();
+
+		return FReply::Handled();
+	}
+
+	return FReply::Unhandled();
+}
+
 FReply SNiagaraOverviewStackNode::OnClickedRenderingPreview(const FGeometry& InGeometry, const FPointerEvent& InEvent, UNiagaraStackEntry* InEntry)
 {
 	if (InEvent.GetEffectingButton() == EKeys::LeftMouseButton)
@@ -425,6 +378,37 @@ FReply SNiagaraOverviewStackNode::OnClickedRenderingPreview(const FGeometry& InG
 	return FReply::Unhandled();
 }
 
+FReply SNiagaraOverviewStackNode::OnPropertiesButtonClicked() const
+{
+	if(EmitterHandleViewModelWeak.IsValid())
+	{
+		TArray<UNiagaraStackEmitterPropertiesGroup*> EmitterProperties;
+		EmitterHandleViewModelWeak.Pin()->GetEmitterStackViewModel()->GetRootEntry()->GetUnfilteredChildrenOfType(EmitterProperties);
+
+		ensure(EmitterProperties.Num() == 1);
+		EmitterHandleViewModelWeak.Pin()->GetOwningSystemViewModel()->GetSelectionViewModel()->UpdateSelectedEntries({EmitterProperties[0]}, {}, true);
+	}
+
+	return FReply::Handled();
+}
+
+void SNiagaraOverviewStackNode::RefreshEmitterThumbnailPreview()
+{
+	if(EmitterHandleViewModelWeak.IsValid() && ThumbnailContainer.IsValid())
+	{
+		if(EmitterHandleViewModelWeak.Pin()->GetEmitterViewModel()->GetEditorData().GetThumbnail() != nullptr)
+		{
+			PreviewThumbnail = MakeShared<FAssetThumbnail>(EmitterHandleViewModelWeak.Pin()->GetEmitterViewModel()->GetEditorData().GetThumbnail(), 200.f, 200.f, UThumbnailManager::Get().GetSharedThumbnailPool());
+			ThumbnailContainer->SetContent(PreviewThumbnail->MakeThumbnailWidget());
+		}
+		else
+		{
+			PreviewThumbnail = MakeShared<FAssetThumbnail>(EmitterHandleViewModelWeak.Pin()->GetEmitterViewModel()->GetEmitter().Emitter, 200.f, 200.f, UThumbnailManager::Get().GetSharedThumbnailPool());
+			ThumbnailContainer->SetContent(PreviewThumbnail->MakeThumbnailWidget());
+		}
+	}
+}
+
 void SNiagaraOverviewStackNode::Tick(const FGeometry& AllottedGeometry, const double InCurrentTime, const float InDeltaTime)
 {
 	if (OverviewStackNode != nullptr)
@@ -437,7 +421,7 @@ void SNiagaraOverviewStackNode::Tick(const FGeometry& AllottedGeometry, const do
 
 		if (bTopContentBarRefreshPending)
 		{
-			FillTopContentBar();
+			CreateTopContentBar();
 			bTopContentBarRefreshPending = false;
 		}
 
@@ -458,9 +442,9 @@ void SNiagaraOverviewStackNode::OnMaterialCompiled(class UMaterialInterface* Mat
 	if (EmitterHandleViewModelWeak.IsValid())
 	{
 		bool bUsingThisMaterial = false;
-		EmitterHandleViewModelWeak.Pin()->GetRendererEntries(PreviewStackEntries);
+		EmitterHandleViewModelWeak.Pin()->GetRendererEntries(RendererPreviewStackEntries);
 		FNiagaraEmitterInstance* InInstance = EmitterHandleViewModelWeak.Pin()->GetEmitterViewModel()->GetSimulation().IsValid() ? EmitterHandleViewModelWeak.Pin()->GetEmitterViewModel()->GetSimulation().Pin().Get() : nullptr;
-		for (UNiagaraStackEntry* Entry : PreviewStackEntries)
+		for (UNiagaraStackEntry* Entry : RendererPreviewStackEntries)
 		{
 			if (UNiagaraStackRendererItem* RendererItem = Cast<UNiagaraStackRendererItem>(Entry))
 			{
@@ -481,77 +465,117 @@ void SNiagaraOverviewStackNode::OnMaterialCompiled(class UMaterialInterface* Mat
 	}
 }
 
-void SNiagaraOverviewStackNode::CreateBottomSummaryExpander()
+TSharedRef<SWidget> SNiagaraOverviewStackNode::CreateTitleWidget_Default(TSharedPtr<SNodeTitle> NodeTitle)
 {
-	UNiagaraEmitter* Emitter = EmitterHandleViewModelWeak.IsValid()? EmitterHandleViewModelWeak.Pin()->GetEmitterViewModel()->GetEmitter().Emitter : nullptr;
-	if (BottomSummaryExpander.IsValid() || !Emitter)
+	TSharedRef<SWidget> DefaultTitle = SGraphNode::CreateTitleWidget(NodeTitle);
+	DefaultTitle->SetToolTipText(TAttribute<FText>::CreateSP(NodeTitle.Get(), &SNodeTitle::GetHeadTitle));
+
+	if (StackViewModel == nullptr)
 	{
-		return;
-	}	
+		return SNew(SHorizontalBox)
+			+ SHorizontalBox::Slot()
+			.Padding(0, 0, 5, 0)
+			[
+				SNew(STextBlock)
+				.Text(LOCTEXT("InvalidNode", "INVALID"))
+			]
+			+ SHorizontalBox::Slot()
+			[
+				DefaultTitle
+			];
+	}
 	
-	SAssignNew(BottomSummaryExpander, SBorder)
-	.BorderImage(FNiagaraEditorWidgetsStyle::Get().GetBrush("NiagaraEditor.SystemOverview.NodeBackgroundBorder"))
-	.BorderBackgroundColor(FNiagaraEditorWidgetsStyle::Get().GetColor("NiagaraEditor.SystemOverview.NodeBackgroundColor"))
-	.HAlign(HAlign_Fill)
-	.VAlign(VAlign_Fill)
+	TSharedPtr<SWidget> TitleWidget = SNew(SHorizontalBox)
+	// Summary View Controls
+	+ SHorizontalBox::Slot()
+	.AutoWidth()
+	.HAlign(HAlign_Left)
+	.VAlign(VAlign_Center)
+	.Padding(2.f)
 	[
-		SNew(SButton)
-		.ButtonStyle(FAppStyle::Get(), "NoBorder")
-		.HAlign(HAlign_Center)
-		.ContentPadding(2)
-		.ToolTipText(this, &SNiagaraOverviewStackNode::GetSummaryViewCollapseTooltipText)
-		.OnClicked(this, &SNiagaraOverviewStackNode::ExpandSummaryViewClicked)
-		.IsFocusable(false)
-		.Content()
+		CreateSummaryViewToggle()
+	]
+	// Enabled checkbox
+	+ SHorizontalBox::Slot()
+	.AutoWidth()
+	.Padding(5.f, 2.f, 2.f, 2.f)
+	[
+		CreateEnabledCheckbox()
+	]
+	// Name
+	+ SHorizontalBox::Slot()
+	.Padding(3, 0, 0, 0)
+	.FillWidth(1.0f)
+	[
+		SNew(SBox)
+		.WidthOverride(SummaryTitleFixedSize)
 		[
-			// add the dropdown button for advanced properties 
-			SNew(SImage)
-				.Image(this, &SNiagaraOverviewStackNode::GetSummaryViewButtonBrush)
-		]		
+			DefaultTitle
+		]
 	];
+
+	return TitleWidget.ToSharedRef();
 }
 
-EVisibility SNiagaraOverviewStackNode::GetScalabilityIndicatorVisibility() const
+TSharedRef<SWidget> SNiagaraOverviewStackNode::CreateTitleRightWidget_Default()
 {
-	TSharedPtr<FNiagaraEmitterHandleViewModel> EmitterHandleViewModel = EmitterHandleViewModelWeak.Pin();
-
-	if(EmitterHandleViewModel.IsValid())
+	if (StackViewModel == nullptr)
 	{
-		if(FVersionedNiagaraEmitterData* EmitterData = EmitterHandleViewModel->GetEmitterHandle()->GetInstance().GetEmitterData())
-		{
-			bool bIsQualityLevelMaskSetup = EmitterData->Platforms.QualityLevelMask != INDEX_NONE;
-			bool bIsScalabilitySetup = EmitterData->ScalabilityOverrides.Overrides.Num() != 0 || (bIsQualityLevelMaskSetup && EmitterData->Platforms.QualityLevelMask != FNiagaraPlatformSet::GetFullQualityLevelMask(GetDefault<UNiagaraSettings>()->QualityLevels.Num())); 
-			return bIsScalabilitySetup ? EVisibility::Visible : EVisibility::Collapsed;
-		}
+		return SNullWidget::NullWidget;
 	}
+	return SNew(SHorizontalBox)
+	// open parent button
+	+ SHorizontalBox::Slot()
+	.AutoWidth()
+	.HAlign(HAlign_Right)
+	.VAlign(VAlign_Center)
+	.Padding(0, 0, 1, 0)
+	[
+		CreateOpenParentButton()
+	]
 
-	if(UNiagaraSystem* System = OverviewStackNode->GetOwningSystem())
-	{
-		bool bIsScalabilitySetup = System->GetOverrideScalabilitySettings();
-		return bIsScalabilitySetup ? EVisibility::Visible : EVisibility::Collapsed; 
-	}
-	
-	return EVisibility::Collapsed;
+	// version selector
+	+ SHorizontalBox::Slot()
+	.HAlign(HAlign_Right)
+	.AutoWidth()
+	.Padding(1, 0, 2, 0)
+	[
+		CreateVersionSelectorButton()	
+	]
+
+	// scalability controls
+	+ SHorizontalBox::Slot()
+	.AutoWidth()
+	[
+		CreateScalabilityControls()
+	]
+
+	// issue/error icon
+	+ SHorizontalBox::Slot()
+	.AutoWidth()
+	.HAlign(HAlign_Right)
+	.VAlign(VAlign_Center)
+	.Padding(0, 0, 1, 0)
+	[
+		SNew(SNiagaraStackIssueIcon, StackViewModel, StackViewModel->GetRootEntry())
+		.Visibility(this, &SNiagaraOverviewStackNode::GetIssueIconVisibility)
+		.OnClicked(this, &SNiagaraOverviewStackNode::OnCycleThroughIssues)
+	];
 }
 
-TSharedRef<SWidget> SNiagaraOverviewStackNode::CreateNodeContentArea()
+TSharedRef<SWidget> SNiagaraOverviewStackNode::CreateNodeContentArea_Default()
 {
 	TSharedPtr<SWidget> ContentWidget;
 	if (StackViewModel != nullptr && OverviewSelectionViewModel != nullptr)
 	{
-		ContentWidget = SNew(SBox)
-			.MaxDesiredWidth(300)
-			[
-				SNew(SNiagaraOverviewStack, *StackViewModel, *OverviewSelectionViewModel)
-			];
+		ContentWidget = SNew(SNiagaraOverviewStack, *StackViewModel, *OverviewSelectionViewModel)
+		.AllowedClasses({UNiagaraStackItemGroup::StaticClass(), UNiagaraStackItem::StaticClass()});
 	}
 	else
 	{
 		ContentWidget = SNullWidget::NullWidget;
 	}
-
-	FillTopContentBar();
-
+	
 	TSharedPtr<SVerticalBox> NodeBox;
 	
 	// NODE CONTENT AREA
@@ -566,12 +590,11 @@ TSharedRef<SWidget> SNiagaraOverviewStackNode::CreateNodeContentArea()
 			.AutoHeight()
 			.HAlign(HAlign_Fill)
 			.VAlign(VAlign_Center)
-			.Padding(2.0f, 2.0f)
+			.Padding(2.f)
 			[
-				TopContentBar.ToSharedRef()
+				CreateTopContentBar()
 			]
 			+ SVerticalBox::Slot()
-			.FillHeight(1.0f)
 			.HAlign(HAlign_Fill)
 			.VAlign(VAlign_Center)
 			.Padding(0.0f)
@@ -585,6 +608,9 @@ TSharedRef<SWidget> SNiagaraOverviewStackNode::CreateNodeContentArea()
 					SAssignNew(LeftNodeBox, SVerticalBox)
 				]
 				+SHorizontalBox::Slot()
+				.HAlign(HAlign_Fill)
+				.VAlign(VAlign_Center)
+				.FillWidth(1.f)
 				[
 					SNew(SBorder)
 					.BorderImage(FNiagaraEditorWidgetsStyle::Get().GetBrush("NiagaraEditor.SystemOverview.NodeBackgroundBorder"))
@@ -606,20 +632,6 @@ TSharedRef<SWidget> SNiagaraOverviewStackNode::CreateNodeContentArea()
 			]
 		];
 
-	CreateBottomSummaryExpander();
-
-	if (BottomSummaryExpander.IsValid())
-	{
-		NodeBox->AddSlot()
-		.AutoHeight()
-		.HAlign(HAlign_Fill)
-		.VAlign(VAlign_Center)
-		.Padding(2.0f, 2.0f)
-		[
-			BottomSummaryExpander.ToSharedRef()
-		];
-	}
-
 	TSharedRef<SWidget> DetailedContent = SNew(SLevelOfDetailBranchNode)
 		.UseLowDetailSlot(this, &SNiagaraOverviewStackNode::UseLowDetailNodeContent)
 		.LowDetail()
@@ -655,7 +667,222 @@ TSharedRef<SWidget> SNiagaraOverviewStackNode::CreateNodeContentArea()
 	];
 
 	return ScalabilityWrapper.ToSharedRef();
+}
+
+TSharedRef<SWidget> SNiagaraOverviewStackNode::CreateTopContentBar_Default()
+{
+	TSharedRef<SHorizontalBox> TopContentBar = SNew(SHorizontalBox);
 
+	if (EmitterHandleViewModelWeak.IsValid())
+	{		
+		// Isolate toggle button
+		TopContentBar->AddSlot()
+			.AutoWidth()
+			.HAlign(HAlign_Left)
+			.VAlign(VAlign_Center)
+			.Padding(2, 0, 0, 0)
+			[
+				CreateIsolateButton()
+			];
+
+		EmitterHandleViewModelWeak.Pin()->GetRendererEntries(RendererPreviewStackEntries);
+		FNiagaraEmitterInstance* InInstance = EmitterHandleViewModelWeak.Pin()->GetEmitterViewModel()->GetSimulation().IsValid() ? EmitterHandleViewModelWeak.Pin()->GetEmitterViewModel()->GetSimulation().Pin().Get() : nullptr;
+
+		FToolBarBuilder ToolBarBuilder(nullptr, FMultiBoxCustomization::None, nullptr, true);
+		ToolBarBuilder.SetLabelVisibility(EVisibility::Collapsed);
+		
+		for (int32 StackEntryIndex = 0; StackEntryIndex < RendererPreviewStackEntries.Num(); StackEntryIndex++)
+		{
+			UNiagaraStackEntry* Entry = RendererPreviewStackEntries[StackEntryIndex];
+			if (UNiagaraStackRendererItem* RendererItem = Cast<UNiagaraStackRendererItem>(Entry))
+			{
+				TArray<TSharedPtr<SWidget>> Widgets;
+				RendererItem->GetRendererProperties()->GetRendererWidgets(InInstance, Widgets, UThumbnailManager::Get().GetSharedThumbnailPool());
+				TArray<TSharedPtr<SWidget>> TooltipWidgets;
+				RendererItem->GetRendererProperties()->GetRendererTooltipWidgets(InInstance, TooltipWidgets, UThumbnailManager::Get().GetSharedThumbnailPool());
+				check(Widgets.Num() == TooltipWidgets.Num());
+				for (int32 WidgetIndex = 0; WidgetIndex < Widgets.Num(); WidgetIndex++)
+				{
+					ToolBarBuilder.AddWidget(
+						SNew(SBox)
+						.VAlign(VAlign_Center)
+						.HAlign(HAlign_Center)
+						.MinDesiredHeight(RendererThumbnailSize)
+						.MinDesiredWidth(RendererThumbnailSize)
+						.MaxDesiredHeight(RendererThumbnailSize)
+						.MaxDesiredWidth(RendererThumbnailSize)
+						.Visibility(this, &SNiagaraOverviewStackNode::GetEnabledCheckBoxVisibility)
+						[
+							CreateRendererThumbnailWidget(Entry, Widgets[WidgetIndex], TooltipWidgets[WidgetIndex])
+						]
+					);
+				}
+
+				// if we had a widget for this entry, add a separator for the next entry's widgets, except for the last entry
+				if(Widgets.Num() > 0 && StackEntryIndex < RendererPreviewStackEntries.Num() - 1)
+				{
+					ToolBarBuilder.AddSeparator();
+				}
+			}
+		}
+
+		TopContentBar->AddSlot()
+		.AutoWidth()
+		.HAlign(HAlign_Left)
+		.VAlign(VAlign_Center)
+		.MaxWidth(300.f)
+		[
+			ToolBarBuilder.MakeWidget()
+		];
+	}
+
+	return TopContentBar;
+}
+
+TSharedRef<SWidget> SNiagaraOverviewStackNode::CreateTitleWidget_Summary(TSharedPtr<SNodeTitle> NodeTitle)
+{
+	// We don't need to differ from the default currently
+	return CreateTitleWidget_Default(NodeTitle);
+}
+
+TSharedRef<SWidget> SNiagaraOverviewStackNode::CreateTitleRightWidget_Summary()
+{
+	if (StackViewModel == nullptr)
+	{
+		return SNullWidget::NullWidget;
+	}
+
+	return SNew(SHorizontalBox)
+	// issue/error icon
+	+ SHorizontalBox::Slot()
+	.AutoWidth()
+	.HAlign(HAlign_Right)
+	.VAlign(VAlign_Center)
+	.Padding(0, 0, 1, 0)
+	[
+		SNew(SNiagaraStackIssueIcon, StackViewModel, StackViewModel->GetRootEntry())
+		.Visibility(this, &SNiagaraOverviewStackNode::GetIssueIconVisibility)
+		.OnClicked(this, &SNiagaraOverviewStackNode::OnCycleThroughIssues)
+	];
+}
+
+TSharedRef<SWidget> SNiagaraOverviewStackNode::CreateNodeContentArea_Summary()
+{
+	return SNew(SVerticalBox)
+	+ SVerticalBox::Slot()
+	.AutoHeight()
+	[
+		CreateTopContentBar()
+	]
+	+ SVerticalBox::Slot()
+	.Padding(3.f)
+	[
+		CreateEmitterThumbnail()
+	];
+}
+
+TSharedRef<SWidget> SNiagaraOverviewStackNode::CreateTopContentBar_Summary()
+{
+	TSharedRef<SHorizontalBox> TopContentBar = SNew(SHorizontalBox);
+	
+	TopContentBar->AddSlot()
+	.AutoWidth()
+	.Padding(4.f, 5.f)
+	.HAlign(HAlign_Center)
+	.VAlign(VAlign_Center)
+	[
+		CreateIsolateButton()
+	];
+
+	TopContentBar->AddSlot()
+	[
+		SNew(SSpacer)
+	];
+
+	TopContentBar->AddSlot()
+	.AutoWidth()
+	.Padding(2.f, 5.f)
+	.HAlign(HAlign_Right)
+	.VAlign(VAlign_Center)
+	[
+		CreateLocalSpaceToggle()
+	];
+	
+	TopContentBar->AddSlot()
+	.AutoWidth()
+	.Padding(2.f, 5.f)
+	.HAlign(HAlign_Right)
+	.VAlign(VAlign_Center)
+	[
+		CreateDeterminismToggle()
+	];
+	
+	TopContentBar->AddSlot()
+	.AutoWidth()
+	.Padding(2.f, 5.f)
+	.HAlign(HAlign_Right)
+	.VAlign(VAlign_Center)
+	[
+		CreateSimTargetToggle()
+	];
+
+	// TopContentBar->AddSlot()
+	// .AutoWidth()
+	// .Padding(10.f, 5.f)
+	// .HAlign(HAlign_Right)
+	// .VAlign(VAlign_Center)
+	// [
+	// 	CreatePropertiesButton()
+	// ];
+
+	return TopContentBar;
+}
+
+void SNiagaraOverviewStackNode::UpdateGraphNode()
+{
+	DisplayMode = EDisplayMode::Default;
+	if(EmitterHandleViewModelWeak.IsValid())
+	{
+		DisplayMode = EmitterHandleViewModelWeak.Pin()->GetEmitterViewModel()->GetEditorData().ShouldShowSummaryView() ? EDisplayMode::Summary : DisplayMode;
+	}
+	
+	SGraphNode::UpdateGraphNode();
+}
+
+TOptional<ETextOverflowPolicy> SNiagaraOverviewStackNode::GetNameOverflowPolicy() const
+{
+	if(EmitterHandleViewModelWeak.IsValid())
+	{
+		if(DisplayMode == EDisplayMode::Summary)
+		{
+			return ETextOverflowPolicy::Ellipsis;
+		}
+	}
+
+	return {};
+}
+
+EVisibility SNiagaraOverviewStackNode::GetScalabilityIndicatorVisibility() const
+{
+	TSharedPtr<FNiagaraEmitterHandleViewModel> EmitterHandleViewModel = EmitterHandleViewModelWeak.Pin();
+
+	if(EmitterHandleViewModel.IsValid())
+	{
+		if(FVersionedNiagaraEmitterData* EmitterData = EmitterHandleViewModel->GetEmitterHandle()->GetInstance().GetEmitterData())
+		{
+			bool bIsQualityLevelMaskSetup = EmitterData->Platforms.QualityLevelMask != INDEX_NONE;
+			bool bIsScalabilitySetup = EmitterData->ScalabilityOverrides.Overrides.Num() != 0 || (bIsQualityLevelMaskSetup && EmitterData->Platforms.QualityLevelMask != FNiagaraPlatformSet::GetFullQualityLevelMask(GetDefault<UNiagaraSettings>()->QualityLevels.Num())); 
+			return bIsScalabilitySetup ? EVisibility::Visible : DisplayMode == EDisplayMode::Summary ? EVisibility::Hidden : EVisibility::Collapsed;
+		}
+	}
+
+	if(UNiagaraSystem* System = OverviewStackNode->GetOwningSystem())
+	{
+		bool bIsScalabilitySetup = System->GetOverrideScalabilitySettings();
+		return bIsScalabilitySetup ? EVisibility::Visible : EVisibility::Collapsed; 
+	}
+	
+	return EVisibility::Collapsed;
 }
 
 bool SNiagaraOverviewStackNode::UseLowDetailNodeContent() const
@@ -782,112 +1009,209 @@ FText SNiagaraOverviewStackNode::GetLowDetailNodeTitle() const
 	return LowDetailTitleCache.Value;
 }
 
-void SNiagaraOverviewStackNode::StackViewModelStructureChanged(ENiagaraStructureChangedFlags Flags)
+TSharedRef<SWidget> SNiagaraOverviewStackNode::CreateSummaryViewToggle()
 {
-	bTopContentBarRefreshPending = true;
+	return SNew(SNiagaraSummaryViewToggle, EmitterHandleViewModelWeak.Pin() ? EmitterHandleViewModelWeak.Pin() : nullptr);
 }
 
-void SNiagaraOverviewStackNode::StackViewModelDataObjectChanged(TArray<UObject*> ChangedObjects, ENiagaraDataObjectChange ChangeType)
+TSharedRef<SWidget> SNiagaraOverviewStackNode::CreateLocalSpaceToggle()
 {
-	for (UObject* ChangedObject : ChangedObjects)
-	{
-		if (ChangedObject->IsA<UNiagaraRendererProperties>())
-		{
-			bTopContentBarRefreshPending = true;
-			break;
-		}
-	}
+	return SNew(SNiagaraLocalSpaceToggle, EmitterHandleViewModelWeak.Pin() ? EmitterHandleViewModelWeak.Pin() : nullptr);
 }
 
-void SNiagaraOverviewStackNode::FillTopContentBar()
+TSharedRef<SWidget> SNiagaraOverviewStackNode::CreateEmitterThumbnail()
 {
-	if (TopContentBar.IsValid() && TopContentBar->GetChildren())
+	if(EmitterHandleViewModelWeak.IsValid())
 	{
-		TopContentBar->ClearChildren();
+		TSharedPtr<FNiagaraEmitterHandleViewModel> EmitterHandleViewModel = EmitterHandleViewModelWeak.Pin();
+
+		TSharedRef<SOverlay> PreviewOverlay = SNew(SOverlay)
+		+ SOverlay::Slot()
+		.HAlign(HAlign_Fill)
+		.VAlign(VAlign_Fill)
+		[		
+			SNew(SScaleBox)
+			.Stretch(EStretch::ScaleToFill)
+			[
+				SAssignNew(ThumbnailContainer, SBox)
+				.WidthOverride(EmitterThumbnailSize)
+				.HeightOverride(EmitterThumbnailSize)
+			]
+		]
+		+ SOverlay::Slot()
+		.HAlign(HAlign_Left)
+		.VAlign(VAlign_Bottom)
+		.Padding(3.f)
+		[
+			CreateCaptureThumbnailButton()
+		];
+
+		RefreshEmitterThumbnailPreview();
+
+		return PreviewOverlay;
 	}
-	if (EmitterHandleViewModelWeak.IsValid())
-	{		
-		// Isolate toggle button
-		TopContentBar->AddSlot()
-			.AutoWidth()
-			.HAlign(HAlign_Left)
-			.VAlign(VAlign_Center)
-			.Padding(2, 0, 0, 0)
+
+	return SNullWidget::NullWidget;
+}
+
+TSharedRef<SWidget> SNiagaraOverviewStackNode::CreateEnabledCheckbox()
+{
+	return SNew(SCheckBox)
+	.Visibility(this, &SNiagaraOverviewStackNode::GetEnabledCheckBoxVisibility)
+	.IsChecked(this, &SNiagaraOverviewStackNode::GetEnabledCheckState)
+	.OnCheckStateChanged(this, &SNiagaraOverviewStackNode::OnEnabledCheckStateChanged);
+}
+
+TSharedRef<SButton> SNiagaraOverviewStackNode::CreateIsolateButton()
+{
+	return SNew(SButton)
+	.ButtonStyle(FAppStyle::Get(), "HoverHintOnly")
+	.HAlign(HAlign_Center)
+	.ContentPadding(1)
+	.ToolTipText(this, &SNiagaraOverviewStackNode::GetToggleIsolateToolTip)
+	.OnClicked(this, &SNiagaraOverviewStackNode::OnToggleIsolateButtonClicked)
+	.Visibility(this, &SNiagaraOverviewStackNode::GetToggleIsolateVisibility)
+	.IsFocusable(false)
+	.Content()
+	[
+		SNew(SImage)
+		.Image(this, &SNiagaraOverviewStackNode::GetToggleIsolateImage)
+		.ColorAndOpacity(this, &SNiagaraOverviewStackNode::GetToggleIsolateImageColor)
+	];
+}
+
+TSharedRef<SButton> SNiagaraOverviewStackNode::CreateCaptureThumbnailButton()
+{
+	return SNew(SButton)
+	.ButtonStyle(FAppStyle::Get(), "HoverHintOnly")
+	.HAlign(HAlign_Center)
+	.ContentPadding(1)
+	.OnClicked(this, &SNiagaraOverviewStackNode::OnCaptureThumbnailButtonClicked)
+	.IsFocusable(false)
+	.ToolTipText(LOCTEXT("CaptureNewEmitterThumbnailButtonTooltip", "Capture a new thumbnail for this emitter based on the Niagara viewport"))
+	.Content()
+	[
+		SNew(SImage)
+		.Image(FAppStyle::Get().GetBrush("AssetEditor.SaveThumbnail"))
+	];
+}
+
+TSharedRef<SWidget> SNiagaraOverviewStackNode::CreateOpenParentButton()
+{
+	return SNew(SButton)
+	.IsFocusable(false)
+	.ToolTipText(this, &SNiagaraOverviewStackNode::OpenParentEmitterTooltip)
+	.ButtonStyle(FAppStyle::Get(), "HoverHintOnly")
+	.ContentPadding(2)
+	.OnClicked(this, &SNiagaraOverviewStackNode::OpenParentEmitter)
+	.Visibility(this, &SNiagaraOverviewStackNode::GetOpenParentEmitterVisibility)
+	.DesiredSizeScale(FVector2D(14.0f / 30.0f, 14.0f / 30.0f)) // GoToSourceIcon is 30x30, scale down
+	.Content()
+	[
+		SNew(SImage)
+		.Image(FNiagaraEditorWidgetsStyle::Get().GetBrush("NiagaraEditor.Stack.GoToSourceIcon"))
+		.ColorAndOpacity(FSlateColor::UseForeground())
+	];
+}
+
+TSharedRef<SWidget> SNiagaraOverviewStackNode::CreateVersionSelectorButton()
+{
+	return SNew(SComboButton)
+	.HasDownArrow(false)
+	.ToolTipText(LOCTEXT("ChangeEmitterVersionToolTip", "Change the parent emitter version"))
+	.ButtonStyle(FAppStyle::Get(), "HoverHintOnly")
+	.ForegroundColor(FSlateColor::UseForeground())
+	.OnGetMenuContent(this, &SNiagaraOverviewStackNode::GetVersionSelectorDropdownMenu)
+	.ContentPadding(FMargin(2))
+	.Visibility(this, &SNiagaraOverviewStackNode::GetVersionSelectorVisibility)
+	.ButtonContent()
+	[
+		SNew(STextBlock)
+		.Font(FAppStyle::Get().GetFontStyle("FontAwesome.10"))
+		.ColorAndOpacity(this, &SNiagaraOverviewStackNode::GetVersionSelectorColor)
+		.Text(FEditorFontGlyphs::Random)
+	];
+}
+
+TSharedRef<SWidget> SNiagaraOverviewStackNode::CreateScalabilityControls()
+{
+	return SNew(SVerticalBox)
+		+ SVerticalBox::Slot()
+		.AutoHeight()
+		.HAlign(HAlign_Center)
+		.VAlign(VAlign_Center)
+		[
+			// Toggle Button to enter & exit scalability mode
+			SNew(SCheckBox)
+			.Style(&FAppStyle::GetWidgetStyle<FCheckBoxStyle>("ToggleButtonCheckbox"))
+			.IsChecked(this, &SNiagaraOverviewStackNode::IsScalabilityModeActive)
+			.Padding(2.f)
+			.OnCheckStateChanged(this, &SNiagaraOverviewStackNode::OnScalabilityModeStateChanged)
 			[
-				SNew(SButton)
-				.ButtonStyle(FAppStyle::Get(), "HoverHintOnly")
-				.HAlign(HAlign_Center)
-				.ContentPadding(1)
-				.ToolTipText(this, &SNiagaraOverviewStackNode::GetToggleIsolateToolTip)
-				.OnClicked(this, &SNiagaraOverviewStackNode::OnToggleIsolateButtonClicked)
-				.Visibility(this, &SNiagaraOverviewStackNode::GetToggleIsolateVisibility)
-				.IsFocusable(false)
-				.Content()
+				SNew(SBox)
+				.WidthOverride(16.f)
+				.HeightOverride(16.f)
 				[
 					SNew(SImage)
-					.Image(FNiagaraEditorStyle::Get().GetBrush("NiagaraEditor.Isolate"))
-					.ColorAndOpacity(this, &SNiagaraOverviewStackNode::GetToggleIsolateImageColor)
+					.Image(FNiagaraEditorStyle::Get().GetBrush("NiagaraEditor.Scalability"))
+					.Visibility(this, &SNiagaraOverviewStackNode::GetScalabilityIndicatorVisibility)
+					.ToolTipText(FText::FormatOrdered(LOCTEXT("ScalabilityIndicatorToolTip",
+						"This {0} has scalability set up. Inspecting and editing scalability is accessible by entering Scalability Mode by clicking this or the button in the toolbar.."), EmitterHandleViewModelWeak.IsValid() ? FText::FromString("emitter") : FText::FromString("system")))
 				]
-			];
+			]
+		]
+		// Spawn Count Scale Info
+		+ SVerticalBox::Slot()
+		.AutoHeight()
+		[
+			SNew(STextBlock)
+			.Text(this, &SNiagaraOverviewStackNode::GetSpawnCountScaleText)
+			.ToolTipText(this, &SNiagaraOverviewStackNode::GetSpawnCountScaleTooltip)
+			.Visibility(this, &SNiagaraOverviewStackNode::GetSpawnCountScaleTextVisibility)
+		];
+}
 
-		EmitterHandleViewModelWeak.Pin()->GetRendererEntries(PreviewStackEntries);
-		FNiagaraEmitterInstance* InInstance = EmitterHandleViewModelWeak.Pin()->GetEmitterViewModel()->GetSimulation().IsValid() ? EmitterHandleViewModelWeak.Pin()->GetEmitterViewModel()->GetSimulation().Pin().Get() : nullptr;
+TSharedRef<SWidget> SNiagaraOverviewStackNode::CreateSimTargetToggle()
+{
+	return SNew(SNiagaraSimTargetToggle, EmitterHandleViewModelWeak.IsValid() ? EmitterHandleViewModelWeak.Pin() : nullptr);
+}
 
-		FToolBarBuilder ToolBarBuilder(nullptr, FMultiBoxCustomization::None, nullptr, true);
-		ToolBarBuilder.SetLabelVisibility(EVisibility::Collapsed);
-		
-		for (int32 StackEntryIndex = 0; StackEntryIndex < PreviewStackEntries.Num(); StackEntryIndex++)
-		{
-			UNiagaraStackEntry* Entry = PreviewStackEntries[StackEntryIndex];
-			if (UNiagaraStackRendererItem* RendererItem = Cast<UNiagaraStackRendererItem>(Entry))
-			{
-				TArray<TSharedPtr<SWidget>> Widgets;
-				RendererItem->GetRendererProperties()->GetRendererWidgets(InInstance, Widgets, UThumbnailManager::Get().GetSharedThumbnailPool());
-				TArray<TSharedPtr<SWidget>> TooltipWidgets;
-				RendererItem->GetRendererProperties()->GetRendererTooltipWidgets(InInstance, TooltipWidgets, UThumbnailManager::Get().GetSharedThumbnailPool());
-				check(Widgets.Num() == TooltipWidgets.Num());
-				for (int32 WidgetIndex = 0; WidgetIndex < Widgets.Num(); WidgetIndex++)
-				{
-					ToolBarBuilder.AddWidget(
-						SNew(SBox)
-						.VAlign(VAlign_Center)
-						.HAlign(HAlign_Center)
-						.MinDesiredHeight(ThumbnailSize)
-						.MinDesiredWidth(ThumbnailSize)
-						.MaxDesiredHeight(ThumbnailSize)
-						.MaxDesiredWidth(ThumbnailSize)
-						.Visibility(this, &SNiagaraOverviewStackNode::GetEnabledCheckBoxVisibility)
-						[
-							CreateThumbnailWidget(Entry, Widgets[WidgetIndex], TooltipWidgets[WidgetIndex])
-						]
-					);
-				}
+TSharedRef<SWidget> SNiagaraOverviewStackNode::CreateDeterminismToggle()
+{
+	return SNew(SNiagaraDeterminismToggle, EmitterHandleViewModelWeak.IsValid() ? EmitterHandleViewModelWeak.Pin() : nullptr);
+}
 
-				// if we had a widget for this entry, add a separator for the next entry's widgets, except for the last entry
-				if(Widgets.Num() > 0 && StackEntryIndex < PreviewStackEntries.Num() - 1)
-				{
-					ToolBarBuilder.AddSeparator();
-				}
-			}
-		}
+TSharedRef<SWidget> SNiagaraOverviewStackNode::CreatePropertiesButton()
+{
+	return SNew(SButton)
+	.ButtonStyle(FAppStyle::Get(), "HoverHintOnly")
+	.HAlign(HAlign_Center)
+	.ContentPadding(1)
+	.OnClicked(this, &SNiagaraOverviewStackNode::OnPropertiesButtonClicked)
+	.IsFocusable(false)
+	.ToolTipText(LOCTEXT("SelectEmitterPropertiesButtonTooltip", "Select this emitter's properties"))
+	.Content()
+	[
+		SNew(SImage)
+		.Image(FAppStyle::Get().GetBrush("Icons.Details"))
+	];
+}
 
-		TopContentBar->AddSlot()
-		.AutoWidth()
-		.HAlign(HAlign_Left)
-		.VAlign(VAlign_Center)
-		.MaxWidth(300.f)
-		[
-			ToolBarBuilder.MakeWidget()
-		];
-	}
+void SNiagaraOverviewStackNode::StackViewModelStructureChanged(ENiagaraStructureChangedFlags Flags)
+{
+	bTopContentBarRefreshPending = true;
+}
 
-	TopContentBar->AddSlot()
-		.HAlign(HAlign_Fill)
-		.VAlign(VAlign_Center)
-		.FillWidth(1.0f)
-		[
-			SNullWidget::NullWidget
-		];
+void SNiagaraOverviewStackNode::StackViewModelDataObjectChanged(TArray<UObject*> ChangedObjects, ENiagaraDataObjectChange ChangeType)
+{
+	for (UObject* ChangedObject : ChangedObjects)
+	{
+		if (ChangedObject->IsA<UNiagaraRendererProperties>())
+		{
+			bTopContentBarRefreshPending = true;
+			break;
+		}
+	}
 }
 
 EVisibility SNiagaraOverviewStackNode::GetIssueIconVisibility() const
@@ -905,6 +1229,13 @@ ECheckBoxState SNiagaraOverviewStackNode::GetEnabledCheckState() const
 	return EmitterHandleViewModelWeak.IsValid() && EmitterHandleViewModelWeak.Pin()->GetIsEnabled() ? ECheckBoxState::Checked : ECheckBoxState::Unchecked;
 }
 
+const FSlateBrush* SNiagaraOverviewStackNode::GetEnabledImage() const
+{
+	return EmitterHandleViewModelWeak.IsValid() && EmitterHandleViewModelWeak.Pin()->GetIsEnabled()
+		? FAppStyle::GetBrush("Icons.Success")
+		: FAppStyle::GetBrush("Icons.MinusCircle");
+}
+
 void SNiagaraOverviewStackNode::OnEnabledCheckStateChanged(ECheckBoxState InCheckState)
 {
 	TSharedPtr<FNiagaraEmitterHandleViewModel> EmitterHandleViewModel = EmitterHandleViewModelWeak.Pin();
@@ -914,6 +1245,16 @@ void SNiagaraOverviewStackNode::OnEnabledCheckStateChanged(ECheckBoxState InChec
 	}
 }
 
+EVisibility SNiagaraOverviewStackNode::GetShouldShowSummaryControls() const
+{
+	if(EmitterHandleViewModelWeak.IsValid() && EmitterHandleViewModelWeak.Pin()->GetEmitterViewModel()->GetEditorData().GetSummaryRoot()->GetChildren().Num() > 0)
+	{
+		return EVisibility::Visible;
+	}
+
+	return EVisibility::Collapsed;
+}
+
 FReply SNiagaraOverviewStackNode::OnToggleIsolateButtonClicked()
 {
 	TSharedPtr<FNiagaraEmitterHandleViewModel> EmitterHandleViewModel = EmitterHandleViewModelWeak.Pin();
@@ -942,6 +1283,11 @@ EVisibility SNiagaraOverviewStackNode::GetToggleIsolateVisibility() const
 		: EVisibility::Collapsed;
 }
 
+const FSlateBrush* SNiagaraOverviewStackNode::GetToggleIsolateImage() const
+{
+	return FNiagaraEditorStyle::Get().GetBrush("NiagaraEditor.Isolate");
+}
+
 FSlateColor SNiagaraOverviewStackNode::GetToggleIsolateImageColor() const
 {
 	TSharedPtr<FNiagaraEmitterHandleViewModel> EmitterHandleViewModel = EmitterHandleViewModelWeak.Pin();
@@ -1041,7 +1387,7 @@ EVisibility SNiagaraOverviewStackNode::GetOpenParentEmitterVisibility() const
 	return EmitterHandleViewModel.IsValid() && 
 		EmitterHandleViewModel->GetEmitterViewModel()->HasParentEmitter()
 		? EVisibility::Visible 
-		: EVisibility::Collapsed;
+		: DisplayMode == EDisplayMode::Summary ? EVisibility::Hidden : EVisibility::Collapsed;
 }
 
 EVisibility SNiagaraOverviewStackNode::GetVersionSelectorVisibility() const
@@ -1051,7 +1397,7 @@ EVisibility SNiagaraOverviewStackNode::GetVersionSelectorVisibility() const
 		EmitterHandleViewModel->GetEmitterViewModel()->HasParentEmitter() &&
 		EmitterHandleViewModel->GetEmitterViewModel()->GetParentEmitter().Emitter->IsVersioningEnabled()
 		? EVisibility::Visible 
-		: EVisibility::Collapsed;
+		: DisplayMode == EDisplayMode::Summary ? EVisibility::Hidden : EVisibility::Collapsed;
 }
 
 FSlateColor SNiagaraOverviewStackNode::GetVersionSelectorColor() const
@@ -1116,36 +1462,4 @@ void SNiagaraOverviewStackNode::SwitchToVersion(FNiagaraAssetVersion Version)
 	}
 }
 
-const FSlateBrush* SNiagaraOverviewStackNode::GetSummaryViewButtonBrush() const
-{
-	const UNiagaraEmitterEditorData* EditorData = EmitterHandleViewModelWeak.IsValid()? &EmitterHandleViewModelWeak.Pin()->GetEmitterViewModel()->GetEditorData() : nullptr;
-	if (BottomSummaryExpander->IsHovered())
-	{
-		return EditorData && EditorData->ShouldShowSummaryView()
-			? FAppStyle::GetBrush("DetailsView.PulldownArrow.Down.Hovered")
-			: FAppStyle::GetBrush("DetailsView.PulldownArrow.Up.Hovered");
-	}
-	else
-	{
-		return EditorData && EditorData->ShouldShowSummaryView()
-			? FAppStyle::GetBrush("DetailsView.PulldownArrow.Down")
-			: FAppStyle::GetBrush("DetailsView.PulldownArrow.Up");
-	}	
-}
-
-FText SNiagaraOverviewStackNode::GetSummaryViewCollapseTooltipText() const
-{
-	const UNiagaraEmitterEditorData* EditorData = EmitterHandleViewModelWeak.IsValid()? &EmitterHandleViewModelWeak.Pin()->GetEmitterViewModel()->GetEditorData() : nullptr;
-	return EditorData && EditorData->ShouldShowSummaryView()? LOCTEXT("HideAdvancedToolTip", "Show Full Emitter") : LOCTEXT("ShowAdvancedToolTip", "Show Emitter Summary");
-}
-
-FReply SNiagaraOverviewStackNode::ExpandSummaryViewClicked()
-{
-	if (UNiagaraEmitterEditorData* EditorData = EmitterHandleViewModelWeak.IsValid()? &EmitterHandleViewModelWeak.Pin()->GetEmitterViewModel()->GetEditorData() : nullptr)
-	{
-		EditorData->ToggleShowSummaryView();
-	}
-	return FReply::Handled();
-}
-
 #undef LOCTEXT_NAMESPACE
diff --git a/Engine/Plugins/FX/Niagara/Source/NiagaraEditorWidgets/Private/SNiagaraOverviewStackNode.h b/Engine/Plugins/FX/Niagara/Source/NiagaraEditorWidgets/Private/SNiagaraOverviewStackNode.h
--- a/Engine/Plugins/FX/Niagara/Source/NiagaraEditorWidgets/Private/SNiagaraOverviewStackNode.h
+++ b/Engine/Plugins/FX/Niagara/Source/NiagaraEditorWidgets/Private/SNiagaraOverviewStackNode.h
@@ -17,6 +17,12 @@ class UNiagaraStackEntry;
 class SNiagaraOverviewStackNode : public SGraphNode
 {
 public:
+	enum class EDisplayMode : uint8
+	{
+		Default,
+		Summary
+	};
+	
 	SLATE_BEGIN_ARGS(SNiagaraOverviewStackNode) {}
 	SLATE_END_ARGS()
 
@@ -27,13 +33,27 @@ public:
 
 	virtual ~SNiagaraOverviewStackNode() override;
 protected:
+	void BindEditorDataDelegates();
+	void UnbindEditorDataDelegates() const;
+	
 	virtual TSharedRef<SWidget> CreateTitleWidget(TSharedPtr<SNodeTitle> NodeTitle) override;
 	virtual TSharedRef<SWidget> CreateTitleRightWidget() override;
 	virtual TSharedRef<SWidget> CreateNodeContentArea() override;
+	TSharedRef<SWidget> CreateTopContentBar();
+
+	TSharedRef<SWidget> CreateTitleWidget_Default(TSharedPtr<SNodeTitle> NodeTitle);
+	TSharedRef<SWidget> CreateTitleRightWidget_Default();
+	TSharedRef<SWidget> CreateNodeContentArea_Default();
+	TSharedRef<SWidget> CreateTopContentBar_Default();
+	
+	TSharedRef<SWidget> CreateTitleWidget_Summary(TSharedPtr<SNodeTitle> NodeTitle);
+	TSharedRef<SWidget> CreateTitleRightWidget_Summary();
+	TSharedRef<SWidget> CreateNodeContentArea_Summary();
+	TSharedRef<SWidget> CreateTopContentBar_Summary();
+
 	virtual bool IsHidingPinWidgets() const override { return UseLowDetailNodeContent(); }
 	void StackViewModelStructureChanged(ENiagaraStructureChangedFlags Flags);
 	void StackViewModelDataObjectChanged(TArray<UObject*> ChangedObjects, ENiagaraDataObjectChange ChangeType);
-	void FillTopContentBar();
 	void OnMaterialCompiled(class UMaterialInterface* MaterialInterface);
 	
 	bool UseLowDetailNodeContent() const;
@@ -42,19 +62,43 @@ protected:
 	FOptionalSize GetLowDetailDesiredHeight() const;
 	FText GetLowDetailNodeTitle() const;
 
-	void CreateBottomSummaryExpander();
+	TSharedRef<SWidget> CreateEnabledCheckbox();
+	TSharedRef<SButton> CreateIsolateButton();
+	TSharedRef<SWidget> CreateEmitterThumbnail();
+	TSharedRef<SWidget> CreateSummaryViewToggle();
+	TSharedRef<SWidget> CreateLocalSpaceToggle();
+	TSharedRef<SButton> CreateCaptureThumbnailButton();
+	TSharedRef<SWidget> CreateOpenParentButton();
+	TSharedRef<SWidget> CreateVersionSelectorButton();
+	TSharedRef<SWidget> CreateScalabilityControls();
+	TSharedRef<SWidget> CreateSimTargetToggle();
+	TSharedRef<SWidget> CreateDeterminismToggle();
+	TSharedRef<SWidget> CreatePropertiesButton();
+	
+	virtual void UpdateGraphNode() override;
 private:
-	EVisibility GetScalabilityIndicatorVisibility() const;
+	void RefreshEmitterThumbnailPreview();
+
+	virtual TOptional<ETextOverflowPolicy> GetNameOverflowPolicy() const override;
 	EVisibility GetIssueIconVisibility() const;
+	
 	EVisibility GetEnabledCheckBoxVisibility() const;
 	ECheckBoxState GetEnabledCheckState() const;
+	const FSlateBrush* GetEnabledImage() const;
 	void OnEnabledCheckStateChanged(ECheckBoxState InCheckState);
-	TSharedRef<SWidget> CreateThumbnailWidget(UNiagaraStackEntry* InData, TSharedPtr<SWidget> InWidget, TSharedPtr<SWidget> InTooltipWidget);
+	EVisibility GetShouldShowSummaryControls() const;
+	TSharedRef<SWidget> CreateRendererThumbnailWidget(UNiagaraStackEntry* InData, TSharedPtr<SWidget> InWidget, TSharedPtr<SWidget> InTooltipWidget);
+	FReply OnCaptureThumbnailButtonClicked();
 	FReply OnClickedRenderingPreview(const FGeometry& InGeometry, const FPointerEvent& InEvent, class UNiagaraStackEntry* InEntry);
+	FReply OnPropertiesButtonClicked() const;
+	
 	FText GetToggleIsolateToolTip() const;
 	FReply OnToggleIsolateButtonClicked();
 	EVisibility GetToggleIsolateVisibility() const;
+	const FSlateBrush* GetToggleIsolateImage() const;
 	FSlateColor GetToggleIsolateImageColor() const;
+
+	EVisibility GetScalabilityIndicatorVisibility() const;
 	FSlateColor GetScalabilityTintAlpha() const;
 	void OnScalabilityModeChanged(bool bActive);
 	EVisibility ShowExcludedOverlay() const;
@@ -87,28 +131,26 @@ private:
 	TSharedRef<SWidget> GetVersionSelectorDropdownMenu();
 	void SwitchToVersion(FNiagaraAssetVersion Version);
 
-	const FSlateBrush* GetSummaryViewButtonBrush() const;
-	FText GetSummaryViewCollapseTooltipText() const;
-	FReply ExpandSummaryViewClicked();
-
 	UNiagaraOverviewNode* OverviewStackNode = nullptr;
 	UNiagaraStackViewModel* StackViewModel = nullptr;
 	UNiagaraSystemSelectionViewModel* OverviewSelectionViewModel = nullptr;
 	TWeakObjectPtr<UNiagaraSystemScalabilityViewModel> ScalabilityViewModel;
 	TWeakPtr<FNiagaraEmitterHandleViewModel> EmitterHandleViewModelWeak;
-	/** The top content bar houses the isolate button and the thumbnails */
-	TSharedPtr<SHorizontalBox> TopContentBar;
-	TSharedPtr<SWidget> BottomSummaryExpander;
 	TSharedPtr<SOverlay> ScalabilityWrapper;
 
-	TArray<UNiagaraStackEntry*> PreviewStackEntries;
+	TArray<UNiagaraStackEntry*> RendererPreviewStackEntries;
 	bool bIsHoveringThumbnail = false;
 	bool bTopContentBarRefreshPending = false;
 	int32 CurrentIssueIndex = 0;
 	bool bScalabilityModeActive = false;
+
+	EDisplayMode DisplayMode = EDisplayMode::Default;
 	
 	/** Cached size from when we last drew at high detail */
 	FVector2D LastHighDetailSize = FVector2D::ZeroVector;
 	int32 GeometryTickForSize = 3;
 	mutable TPair<FString, FText> LowDetailTitleCache;
+
+	TSharedPtr<SBox> ThumbnailContainer;
+	TSharedPtr<FAssetThumbnail> PreviewThumbnail;
 };
\ No newline at end of file
diff --git a/Engine/Plugins/FX/Niagara/Source/NiagaraEditorWidgets/Private/Stack/SNiagaraDeterminismToggle.cpp b/Engine/Plugins/FX/Niagara/Source/NiagaraEditorWidgets/Private/Stack/SNiagaraDeterminismToggle.cpp
new file mode 100644
--- /dev/null
+++ b/Engine/Plugins/FX/Niagara/Source/NiagaraEditorWidgets/Private/Stack/SNiagaraDeterminismToggle.cpp
@@ -0,0 +1,131 @@
+// Copyright Epic Games, Inc. All Rights Reserved.
+
+#include "Stack/SNiagaraDeterminismToggle.h"
+
+#include "NiagaraEditorStyle.h"
+#include "Styling/AppStyle.h"
+#include "Widgets/Input/SButton.h"
+#include "Widgets/Images/SImage.h"
+#include "ViewModels/NiagaraEmitterViewModel.h"
+#include "ViewModels/NiagaraEmitterHandleViewModel.h"
+#include "Widgets/Text/STextBlock.h"
+
+#define LOCTEXT_NAMESPACE "SNiagaraDeterminismToggle"
+
+void SNiagaraDeterminismToggle::Construct(const FArguments& InArgs, TSharedPtr<FNiagaraEmitterHandleViewModel> InEmitterHandleViewModel)
+{
+	EmitterHandleViewModel = InEmitterHandleViewModel;
+
+	if(EmitterHandleViewModel.IsValid() == false)
+	{
+		return;
+	}
+	
+	FProperty* MatchingProperty = nullptr;
+	for(TFieldIterator<FProperty> It(FVersionedNiagaraEmitterData::StaticStruct()); It; ++It)
+	{
+		if(*It->GetName() == GET_MEMBER_NAME_CHECKED(FVersionedNiagaraEmitterData, bDeterminism))
+		{
+			MatchingProperty = *It;
+		}
+	}
+
+	if(ensure(MatchingProperty != nullptr))
+	{
+		PropertyDescription = MatchingProperty->GetToolTipText();
+	}
+
+	FSlateFontInfo FontInfo = FSlateFontInfo(FPaths::EngineContentDir() / TEXT("Slate/Fonts/Roboto-Regular.ttf"), 9);
+	FontInfo.OutlineSettings = FFontOutlineSettings(1.f, FLinearColor::White);
+	
+	ChildSlot
+	[
+		SNew(SButton)
+		.ButtonStyle(&FAppStyle::GetWidgetStyle<FButtonStyle>("HoverHintOnly"))
+		.OnClicked(this, &SNiagaraDeterminismToggle::ToggleDeterminismInUI)
+		.ToolTipText(this, &SNiagaraDeterminismToggle::GetDeterminismToggleTooltip)
+		// Padding is tweaked to push the text a bit lower so it appears more central. Remove when exchanging for images
+		.ContentPadding(FMargin(4.f, 3.f, 4.f, 1.f))
+		[
+			SNew(STextBlock)
+			.Text(this, &SNiagaraDeterminismToggle::GetDeterminismButtonText)
+			.Font(FontInfo)
+
+			// SNew(SImage)
+			// .Image(this, &SNiagaraDeterminismToggle::GetDeterminismImage)
+		]
+	];
+}
+
+void SNiagaraDeterminismToggle::ToggleDeterminism() const
+{
+	if(EmitterHandleViewModel.IsValid() == false)
+	{
+		return;
+	}
+	
+	FProperty* MatchingProperty = nullptr;
+	for(TFieldIterator<FProperty> It(FVersionedNiagaraEmitterData::StaticStruct()); It; ++It)
+	{
+		if(*It->GetName() == GET_MEMBER_NAME_CHECKED(FVersionedNiagaraEmitterData, bDeterminism))
+		{
+			MatchingProperty = *It;
+		}
+	}
+	ensure(MatchingProperty != nullptr);
+
+	FScopedTransaction Transaction(LOCTEXT("ToggleDeterminismTransaction", "Determinism Changed"));
+	EmitterHandleViewModel.Pin()->GetEmitterViewModel()->GetEmitter().Emitter->PreEditChange(MatchingProperty);
+	
+	EmitterHandleViewModel.Pin()->GetEmitterViewModel()->GetEmitter().GetEmitterData()->bDeterminism = !EmitterHandleViewModel.Pin()->GetEmitterViewModel()->GetEmitter().GetEmitterData()->bDeterminism;
+
+	FPropertyChangedEvent PropertyChangedEvent(MatchingProperty);
+	UNiagaraSystem::RequestCompileForEmitter(EmitterHandleViewModel.Pin()->GetEmitterViewModel()->GetEmitter());
+}
+
+FReply SNiagaraDeterminismToggle::ToggleDeterminismInUI() const
+{
+	ToggleDeterminism();
+	return FReply::Handled();
+}
+
+FText SNiagaraDeterminismToggle::GetDeterminismButtonText() const
+{
+	if(EmitterHandleViewModel.IsValid() == false)
+	{
+		return FText::GetEmpty();
+	}
+
+	return EmitterHandleViewModel.Pin()->GetEmitterViewModel()->GetEmitter().GetEmitterData()->bDeterminism == true
+		? FText::FromString("D")
+		: FText::FromString("N D");
+}
+
+const FSlateBrush* SNiagaraDeterminismToggle::GetDeterminismImage() const
+{
+	if(EmitterHandleViewModel.IsValid() == false)
+	{
+		return FAppStyle::GetNoBrush();
+	}
+	
+	return EmitterHandleViewModel.Pin()->GetEmitterViewModel()->GetEmitter().GetEmitterData()->bDeterminism == true
+		? FNiagaraEditorStyle::Get().GetBrush("NiagaraEditor.Determinism.Enabled")
+		: FNiagaraEditorStyle::Get().GetBrush("NiagaraEditor.Determinism.Disabled");
+}
+
+FText SNiagaraDeterminismToggle::GetDeterminismToggleTooltip() const
+{
+	if(EmitterHandleViewModel.IsValid() == false)
+	{
+		return FText::GetEmpty();
+	}
+	
+	FText DeterminismEnabledText = EmitterHandleViewModel.Pin()->GetEmitterViewModel()->GetEmitter().GetEmitterData()->bDeterminism == true
+	   ? LOCTEXT("DeterminismEnabled", "Determinism Enabled")
+	   : LOCTEXT("DeterminismDisabled", "Determinism Disabled");
+	
+	FText Tooltip = FText::FormatOrdered(FText::FromString("{0}\n{1}"), DeterminismEnabledText, PropertyDescription);
+	return Tooltip;
+}
+
+#undef LOCTEXT_NAMESPACE
diff --git a/Engine/Plugins/FX/Niagara/Source/NiagaraEditorWidgets/Private/Stack/SNiagaraDeterminismToggle.h b/Engine/Plugins/FX/Niagara/Source/NiagaraEditorWidgets/Private/Stack/SNiagaraDeterminismToggle.h
new file mode 100644
--- /dev/null
+++ b/Engine/Plugins/FX/Niagara/Source/NiagaraEditorWidgets/Private/Stack/SNiagaraDeterminismToggle.h
@@ -0,0 +1,31 @@
+// Copyright Epic Games, Inc. All Rights Reserved.
+
+#pragma once
+
+#include "CoreMinimal.h"
+#include "Widgets/SCompoundWidget.h"
+#include "ViewModels/Stack/NiagaraStackEmitterSettingsGroup.h"
+
+class SNiagaraDeterminismToggle: public SCompoundWidget
+{
+public:
+	SLATE_BEGIN_ARGS(SNiagaraDeterminismToggle) { }
+	SLATE_END_ARGS();
+
+	void Construct(const FArguments& InArgs, TSharedPtr<FNiagaraEmitterHandleViewModel> InEmitterHandleViewModel);
+
+private:
+	void ToggleDeterminism() const;
+	FReply ToggleDeterminismInUI() const;
+
+	FText GetDeterminismButtonText() const;
+	const FSlateBrush* GetDeterminismImage() const;
+	FText GetDeterminismToggleTooltip() const;
+
+private:
+	TWeakPtr<FNiagaraEmitterHandleViewModel> EmitterHandleViewModel;
+	FText PropertyDescription;
+
+	static FSlateBrush* DeterminismBrush;
+	static FSlateBrush* NonDeterminismBrush;
+};
diff --git a/Engine/Plugins/FX/Niagara/Source/NiagaraEditorWidgets/Private/Stack/SNiagaraLocalSpaceToggle.cpp b/Engine/Plugins/FX/Niagara/Source/NiagaraEditorWidgets/Private/Stack/SNiagaraLocalSpaceToggle.cpp
new file mode 100644
--- /dev/null
+++ b/Engine/Plugins/FX/Niagara/Source/NiagaraEditorWidgets/Private/Stack/SNiagaraLocalSpaceToggle.cpp
@@ -0,0 +1,107 @@
+// Copyright Epic Games, Inc. All Rights Reserved.
+
+#include "Stack/SNiagaraLocalSpaceToggle.h"
+
+#include "Styling/AppStyle.h"
+#include "Widgets/Input/SButton.h"
+#include "Widgets/Images/SImage.h"
+#include "ViewModels/NiagaraEmitterViewModel.h"
+#include "ViewModels/NiagaraEmitterHandleViewModel.h"
+
+#define LOCTEXT_NAMESPACE "SNiagaraLocalSpaceToggle"
+
+void SNiagaraLocalSpaceToggle::Construct(const FArguments& InArgs, TSharedPtr<FNiagaraEmitterHandleViewModel> InEmitterHandleViewModel)
+{
+	EmitterHandleViewModel = InEmitterHandleViewModel;
+
+	if(EmitterHandleViewModel.IsValid() == false)
+	{
+		return;
+	}
+
+	FProperty* MatchingProperty = nullptr;
+	for(TFieldIterator<FProperty> It(FVersionedNiagaraEmitterData::StaticStruct()); It; ++It)
+	{
+		if(*It->GetName() == GET_MEMBER_NAME_CHECKED(FVersionedNiagaraEmitterData, bLocalSpace))
+		{
+			MatchingProperty = *It;
+		}
+	}
+	
+	if(ensure(MatchingProperty != nullptr))
+	{
+		PropertyDescription = MatchingProperty->GetToolTipText();
+	}
+	
+	ChildSlot
+	[
+		SNew(SButton)
+		.OnClicked(this, &SNiagaraLocalSpaceToggle::ToggleLocalSpaceForUI)
+		.ToolTipText(this, &SNiagaraLocalSpaceToggle::GetLocalSpaceToggleTooltip)
+		.ContentPadding(1.f)
+		.ButtonStyle(&FAppStyle::GetWidgetStyle<FButtonStyle>("HoverHintOnly"))
+		[
+			SNew(SImage)
+			.Image(this, &SNiagaraLocalSpaceToggle::GetLocalSpaceToggleImage)
+		]
+	];
+}
+
+void SNiagaraLocalSpaceToggle::ToggleLocalSpace() const
+{
+	if(EmitterHandleViewModel.IsValid() == false)
+	{
+		return;
+	}
+	
+	FProperty* MatchingProperty = nullptr;
+	for(TFieldIterator<FProperty> It(FVersionedNiagaraEmitterData::StaticStruct()); It; ++It)
+	{
+		if(*It->GetName() == GET_MEMBER_NAME_CHECKED(FVersionedNiagaraEmitterData, bLocalSpace))
+		{
+			MatchingProperty = *It;
+		}
+	}
+	ensure(MatchingProperty != nullptr);
+
+	FScopedTransaction Transaction(LOCTEXT("ToggleLocalSpaceTransaction", "Local/World Space Toggled"));
+	EmitterHandleViewModel.Pin()->GetEmitterViewModel()->GetEmitter().Emitter->PreEditChange(MatchingProperty);
+	
+	EmitterHandleViewModel.Pin()->GetEmitterViewModel()->GetEmitter().GetEmitterData()->bLocalSpace = !EmitterHandleViewModel.Pin()->GetEmitterViewModel()->GetEmitter().GetEmitterData()->bLocalSpace;
+
+	FPropertyChangedEvent PropertyChangedEvent(MatchingProperty);
+	EmitterHandleViewModel.Pin()->GetEmitterViewModel()->GetEmitter().Emitter->PostEditChangeProperty(PropertyChangedEvent);
+}
+
+FReply SNiagaraLocalSpaceToggle::ToggleLocalSpaceForUI() const
+{
+	ToggleLocalSpace();
+	return FReply::Handled();
+}
+
+const FSlateBrush* SNiagaraLocalSpaceToggle::GetLocalSpaceToggleImage() const
+{
+	if(EmitterHandleViewModel.IsValid() == false)
+	{
+		return FAppStyle::GetNoBrush();
+	}
+	
+	return EmitterHandleViewModel.Pin()->GetEmitterViewModel()->GetEmitter().GetEmitterData()->bLocalSpace ? FAppStyle::GetBrush("Icons.Transform") : FAppStyle::GetBrush("EditorViewport.RelativeCoordinateSystem_World");  
+}
+
+FText SNiagaraLocalSpaceToggle::GetLocalSpaceToggleTooltip() const
+{
+	if(EmitterHandleViewModel.IsValid() == false)
+	{
+		return FText::GetEmpty();
+	}
+
+	FText ToggleTooltip = EmitterHandleViewModel.Pin()->GetEmitterViewModel()->GetEmitter().GetEmitterData()->bLocalSpace
+		? LOCTEXT("ToggleLocalSpaceButtonTooltip_IsLocal", "This emitter operates in local space.")
+		: LOCTEXT("ToggleLocalSpaceButtonTooltip_IsWorld", "This emitter operates in world space.");
+	FText Tooltip = FText::FormatOrdered(FText::FromString("{0}\n{1}"), ToggleTooltip, PropertyDescription);
+	
+	return Tooltip;
+}
+
+#undef LOCTEXT_NAMESPACE
diff --git a/Engine/Plugins/FX/Niagara/Source/NiagaraEditorWidgets/Private/Stack/SNiagaraLocalSpaceToggle.h b/Engine/Plugins/FX/Niagara/Source/NiagaraEditorWidgets/Private/Stack/SNiagaraLocalSpaceToggle.h
new file mode 100644
--- /dev/null
+++ b/Engine/Plugins/FX/Niagara/Source/NiagaraEditorWidgets/Private/Stack/SNiagaraLocalSpaceToggle.h
@@ -0,0 +1,27 @@
+// Copyright Epic Games, Inc. All Rights Reserved.
+
+#pragma once
+
+#include "CoreMinimal.h"
+#include "Widgets/SCompoundWidget.h"
+#include "ViewModels/Stack/NiagaraStackEmitterSettingsGroup.h"
+
+class SNiagaraLocalSpaceToggle: public SCompoundWidget
+{
+public:
+	SLATE_BEGIN_ARGS(SNiagaraLocalSpaceToggle) { }
+	SLATE_END_ARGS();
+
+	void Construct(const FArguments& InArgs, TSharedPtr<FNiagaraEmitterHandleViewModel> InEmitterHandleViewModel);
+
+private:
+	void ToggleLocalSpace() const;
+	FReply ToggleLocalSpaceForUI() const;
+
+	const FSlateBrush* GetLocalSpaceToggleImage() const;
+	FText GetLocalSpaceToggleTooltip() const;
+
+private:
+	TWeakPtr<FNiagaraEmitterHandleViewModel> EmitterHandleViewModel;
+	FText PropertyDescription;
+};
diff --git a/Engine/Plugins/FX/Niagara/Source/NiagaraEditorWidgets/Private/Stack/SNiagaraSimTargetToggle.cpp b/Engine/Plugins/FX/Niagara/Source/NiagaraEditorWidgets/Private/Stack/SNiagaraSimTargetToggle.cpp
new file mode 100644
--- /dev/null
+++ b/Engine/Plugins/FX/Niagara/Source/NiagaraEditorWidgets/Private/Stack/SNiagaraSimTargetToggle.cpp
@@ -0,0 +1,92 @@
+// Copyright Epic Games, Inc. All Rights Reserved.
+
+#include "Stack/SNiagaraSimTargetToggle.h"
+
+#include "NiagaraEditorStyle.h"
+#include "Styling/AppStyle.h"
+#include "Widgets/Input/SButton.h"
+#include "Widgets/Images/SImage.h"
+#include "ViewModels/NiagaraEmitterViewModel.h"
+
+#define LOCTEXT_NAMESPACE "SNiagaraSimTargetToggle"
+
+void SNiagaraSimTargetToggle::Construct(const FArguments& InArgs, TSharedPtr<FNiagaraEmitterHandleViewModel> InEmitterHandleViewModel)
+{
+	EmitterHandleViewModel = InEmitterHandleViewModel;
+
+	if(EmitterHandleViewModel.IsValid() == false)
+	{
+		return;
+	}
+	
+	ChildSlot
+	[
+		SNew(SButton)
+		.ButtonStyle(&FAppStyle::GetWidgetStyle<FButtonStyle>("HoverHintOnly"))
+		.OnClicked(this, &SNiagaraSimTargetToggle::ToggleSimTargetInUI)
+		.ToolTipText(this, &SNiagaraSimTargetToggle::GetSimTargetToggleTooltip)
+		[
+			SNew(SImage)
+			.Image(this, &SNiagaraSimTargetToggle::GetSimTargetImage)
+		]
+	];
+}
+
+void SNiagaraSimTargetToggle::ToggleSimTarget() const
+{
+	if(EmitterHandleViewModel.IsValid() == false)
+	{
+		return;
+	}
+	
+	FProperty* MatchingProperty = nullptr;
+	for(TFieldIterator<FProperty> It(FVersionedNiagaraEmitterData::StaticStruct()); It; ++It)
+	{
+		if(*It->GetName() == GET_MEMBER_NAME_CHECKED(FVersionedNiagaraEmitterData, SimTarget))
+		{
+			MatchingProperty = *It;
+		}
+	}
+	ensure(MatchingProperty != nullptr);
+
+	FScopedTransaction Transaction(LOCTEXT("ToggleSimTargetTransaction", "Simulation Target Changed"));
+	EmitterHandleViewModel.Pin()->GetEmitterViewModel()->GetEmitter().Emitter->PreEditChange(MatchingProperty);
+	
+	ENiagaraSimTarget CurrentSimTarget = EmitterHandleViewModel.Pin()->GetEmitterViewModel()->GetEmitter().GetEmitterData()->SimTarget;
+	EmitterHandleViewModel.Pin()->GetEmitterViewModel()->GetEmitter().GetEmitterData()->SimTarget = CurrentSimTarget == ENiagaraSimTarget::CPUSim ?  ENiagaraSimTarget::GPUComputeSim : ENiagaraSimTarget::CPUSim;
+
+	FPropertyChangedEvent PropertyChangedEvent(MatchingProperty);
+	UNiagaraSystem::RequestCompileForEmitter(EmitterHandleViewModel.Pin()->GetEmitterViewModel()->GetEmitter());
+}
+
+FReply SNiagaraSimTargetToggle::ToggleSimTargetInUI() const
+{
+	ToggleSimTarget();
+	return FReply::Handled();
+}
+
+const FSlateBrush* SNiagaraSimTargetToggle::GetSimTargetImage() const
+{
+	if(EmitterHandleViewModel.IsValid() == false)
+	{
+		return FAppStyle::GetNoBrush();
+	}
+	
+	return EmitterHandleViewModel.Pin()->GetEmitterViewModel()->GetEmitter().GetEmitterData()->SimTarget == ENiagaraSimTarget::CPUSim
+		? FNiagaraEditorStyle::Get().GetBrush("NiagaraEditor.Stack.CPUIcon")
+		: FNiagaraEditorStyle::Get().GetBrush("NiagaraEditor.Stack.GPUIcon");
+}
+
+FText SNiagaraSimTargetToggle::GetSimTargetToggleTooltip() const
+{
+	if(EmitterHandleViewModel.IsValid() == false)
+	{
+		return FText::GetEmpty();
+	}
+
+	return EmitterHandleViewModel.Pin()->GetEmitterViewModel()->GetEmitter().GetEmitterData()->SimTarget == ENiagaraSimTarget::CPUSim
+		? LOCTEXT("ToggleSimTargetTooltip_CPU", "CPU emitter. Clicking will turn this emitter into a GPU emitter and recompile.")
+		: LOCTEXT("ToggleSimTargetTooltip_GPU", "GPU emitter. Clicking will turn this emitter into a CPU emitter and recompile.");		
+}
+
+#undef LOCTEXT_NAMESPACE
diff --git a/Engine/Plugins/FX/Niagara/Source/NiagaraEditorWidgets/Private/Stack/SNiagaraSimTargetToggle.h b/Engine/Plugins/FX/Niagara/Source/NiagaraEditorWidgets/Private/Stack/SNiagaraSimTargetToggle.h
new file mode 100644
--- /dev/null
+++ b/Engine/Plugins/FX/Niagara/Source/NiagaraEditorWidgets/Private/Stack/SNiagaraSimTargetToggle.h
@@ -0,0 +1,26 @@
+// Copyright Epic Games, Inc. All Rights Reserved.
+
+#pragma once
+
+#include "CoreMinimal.h"
+#include "Widgets/SCompoundWidget.h"
+#include "ViewModels/NiagaraEmitterHandleViewModel.h"
+
+class SNiagaraSimTargetToggle: public SCompoundWidget
+{
+public:
+	SLATE_BEGIN_ARGS(SNiagaraSimTargetToggle) { }
+	SLATE_END_ARGS();
+
+	void Construct(const FArguments& InArgs, TSharedPtr<FNiagaraEmitterHandleViewModel> InEmitterHandleViewModel);
+
+private:
+	void ToggleSimTarget() const;
+	FReply ToggleSimTargetInUI() const;
+
+	const FSlateBrush* GetSimTargetImage() const;
+	FText GetSimTargetToggleTooltip() const;
+
+private:
+	TWeakPtr<FNiagaraEmitterHandleViewModel> EmitterHandleViewModel;
+};
diff --git a/Engine/Plugins/FX/Niagara/Source/NiagaraEditorWidgets/Private/Stack/SNiagaraSummaryViewToggle.cpp b/Engine/Plugins/FX/Niagara/Source/NiagaraEditorWidgets/Private/Stack/SNiagaraSummaryViewToggle.cpp
new file mode 100644
--- /dev/null
+++ b/Engine/Plugins/FX/Niagara/Source/NiagaraEditorWidgets/Private/Stack/SNiagaraSummaryViewToggle.cpp
@@ -0,0 +1,62 @@
+// Copyright Epic Games, Inc. All Rights Reserved.
+
+#include "Stack/SNiagaraSummaryViewToggle.h"
+#include "Styling/AppStyle.h"
+#include "Widgets/Input/SButton.h"
+#include "Widgets/Images/SImage.h"
+#include "ViewModels/NiagaraEmitterViewModel.h"
+#include "Stack/SNiagaraStackInheritanceIcon.h"
+#include "ViewModels/NiagaraEmitterHandleViewModel.h"
+
+#define LOCTEXT_NAMESPACE "SNiagaraSummaryViewToggle"
+
+void SNiagaraSummaryViewToggle::Construct(const FArguments& InArgs, TSharedPtr<FNiagaraEmitterHandleViewModel> InEmitterHandleViewModel)
+{
+	EmitterHandleViewModel = InEmitterHandleViewModel;
+
+	if(EmitterHandleViewModel.IsValid() == false)
+	{
+		return;
+	}
+	
+	ChildSlot
+	[
+		SNew(SButton)
+		.OnClicked(this, &SNiagaraSummaryViewToggle::ToggleShowSummaryForUI)
+		.ToolTipText(this, &SNiagaraSummaryViewToggle::GetSummaryToggleTooltip)
+		.ContentPadding(1.f)
+		.ButtonStyle(&FAppStyle::GetWidgetStyle<FButtonStyle>("HoverHintOnly"))
+		[
+			SNew(SImage)
+			.Image(this, &SNiagaraSummaryViewToggle::GetSummaryToggleImage)
+		]
+	];
+}
+
+void SNiagaraSummaryViewToggle::ToggleShowSummary() const
+{
+	EmitterHandleViewModel.Pin()->GetEmitterViewModel()->GetEditorData().ToggleShowSummaryView();
+}
+
+FReply SNiagaraSummaryViewToggle::ToggleShowSummaryForUI() const
+{
+	ToggleShowSummary();
+	return FReply::Handled();
+}
+
+const FSlateBrush* SNiagaraSummaryViewToggle::GetSummaryToggleImage() const
+{
+	return EmitterHandleViewModel.Pin()->GetEmitterViewModel()->GetEditorData().ShouldShowSummaryView() ? FAppStyle::GetBrush("TreeArrow_Collapsed") : FAppStyle::GetBrush("TreeArrow_Expanded");  
+	// return EmitterHandleViewModel.Pin()->GetEmitterViewModel()->GetEditorData().ShouldShowSummaryView() ? FAppStyle::GetBrush("DetailsView.PulldownArrow.Down") : FAppStyle::GetBrush("DetailsView.PulldownArrow.Up");  
+}
+
+FText SNiagaraSummaryViewToggle::GetSummaryToggleTooltip() const
+{
+	return EmitterHandleViewModel.Pin()->GetEmitterViewModel()->GetEditorData().ShouldShowSummaryView()
+		? LOCTEXT("ToggleSummaryViewButtonTooltip_IsShowingSummaryView", "This emitter is displayed in summary view. To see the full emitter, return to the full view first.")
+		: LOCTEXT("ToggleSummaryViewButtonTooltip_IsShowingDefaultView", "This emitter is displayed in default view. Collapse the node to see the summary view instead.");
+			
+		
+}
+
+#undef LOCTEXT_NAMESPACE
diff --git a/Engine/Plugins/FX/Niagara/Source/NiagaraEditorWidgets/Private/Stack/SNiagaraSummaryViewToggle.h b/Engine/Plugins/FX/Niagara/Source/NiagaraEditorWidgets/Private/Stack/SNiagaraSummaryViewToggle.h
new file mode 100644
--- /dev/null
+++ b/Engine/Plugins/FX/Niagara/Source/NiagaraEditorWidgets/Private/Stack/SNiagaraSummaryViewToggle.h
@@ -0,0 +1,26 @@
+// Copyright Epic Games, Inc. All Rights Reserved.
+
+#pragma once
+
+#include "CoreMinimal.h"
+#include "Widgets/SCompoundWidget.h"
+#include "ViewModels/Stack/NiagaraStackEmitterSettingsGroup.h"
+
+class SNiagaraSummaryViewToggle: public SCompoundWidget
+{
+public:
+	SLATE_BEGIN_ARGS(SNiagaraSummaryViewToggle) { }
+	SLATE_END_ARGS();
+
+	void Construct(const FArguments& InArgs, TSharedPtr<FNiagaraEmitterHandleViewModel> InEmitterHandleViewModel);
+
+private:
+	void ToggleShowSummary() const;
+	FReply ToggleShowSummaryForUI() const;
+
+	const FSlateBrush* GetSummaryToggleImage() const;
+	FText GetSummaryToggleTooltip() const;
+
+private:
+	TWeakPtr<FNiagaraEmitterHandleViewModel> EmitterHandleViewModel;
+};
