Fixed minor padding issue

diff --git a/Engine/Plugins/FX/Niagara/Source/NiagaraEditorWidgets/Private/SNiagaraOverviewStackNode.cpp b/Engine/Plugins/FX/Niagara/Source/NiagaraEditorWidgets/Private/SNiagaraOverviewStackNode.cpp
--- a/Engine/Plugins/FX/Niagara/Source/NiagaraEditorWidgets/Private/SNiagaraOverviewStackNode.cpp
+++ b/Engine/Plugins/FX/Niagara/Source/NiagaraEditorWidgets/Private/SNiagaraOverviewStackNode.cpp
@@ -819,7 +819,7 @@ TSharedRef<SWidget> SNiagaraOverviewStackNode::CreateTopContentBar_Summary()
 	
 	TopContentBar->AddSlot()
 	.AutoWidth()
-	.Padding(2.f, 5.f)
+	.Padding(2.f, 5.f, 4.f, 5.f)
 	.HAlign(HAlign_Right)
 	.VAlign(VAlign_Center)
 	[
