- Fix Blueprintcallable vs DLL exposed inconsistencies

diff --git a/Engine/Plugins/FX/Niagara/Source/Niagara/Public/NiagaraComponent.h b/Engine/Plugins/FX/Niagara/Source/Niagara/Public/NiagaraComponent.h
--- a/Engine/Plugins/FX/Niagara/Source/Niagara/Public/NiagaraComponent.h
+++ b/Engine/Plugins/FX/Niagara/Source/Niagara/Public/NiagaraComponent.h
@@ -290,19 +290,19 @@ public:
 	NIAGARA_API void SetAsset(UNiagaraSystem* InAsset, bool bResetExistingOverrideParameters = true);
 
 	UFUNCTION(BlueprintCallable, Category = Niagara, meta = (DisplayName = "Get Niagara System Asset"))
-	UNiagaraSystem* GetAsset() const { return Asset; }
+	NIAGARA_API UNiagaraSystem* GetAsset() const { return Asset; }
 
 	UFUNCTION(BlueprintCallable, Category = Niagara)
-	void SetOcclusionQueryMode(ENiagaraOcclusionQueryMode Mode);
+	NIAGARA_API void SetOcclusionQueryMode(ENiagaraOcclusionQueryMode Mode);
 
 	UFUNCTION(BlueprintCallable, Category = Niagara)
-	ENiagaraOcclusionQueryMode GetOcclusionQueryMode() const { return OcclusionQueryMode; }
+	NIAGARA_API ENiagaraOcclusionQueryMode GetOcclusionQueryMode() const { return OcclusionQueryMode; }
 
 	UFUNCTION(BlueprintCallable, Category = Niagara, meta = (DisplayName = "Set Forced Solo Mode"))
 	NIAGARA_API void SetForceSolo(bool bInForceSolo);
 
 	UFUNCTION(BlueprintCallable, Category = Niagara, meta = (DisplayName = "Is In Forced Solo Mode"))
-	bool GetForceSolo()const { return bForceSolo; }
+	NIAGARA_API bool GetForceSolo() const { return bForceSolo; }
 
 private:
 	NIAGARA_API bool RequiresSoloMode() const;
@@ -320,7 +320,7 @@ public:
 	NIAGARA_API void SetCustomTimeDilation(float Dilation = 1.0f);
 
 	UFUNCTION(BlueprintCallable, Category = Niagara)
-	float GetCustomTimeDilation() const { return CustomTimeDilation; }
+	NIAGARA_API float GetCustomTimeDilation() const { return CustomTimeDilation; }
 
 	UFUNCTION(BlueprintCallable, Category = Niagara, meta = (DisplayName = "Get Age Update Mode"))
 	NIAGARA_API ENiagaraAgeUpdateMode GetAgeUpdateMode() const;
