Fixed missing sections in summary view (forgot to uncomment some code in previous check-in)

diff --git a/Engine/Plugins/FX/Niagara/Source/NiagaraEditor/Private/ViewModels/Stack/NiagaraStackSummaryViewInputCollection.cpp b/Engine/Plugins/FX/Niagara/Source/NiagaraEditor/Private/ViewModels/Stack/NiagaraStackSummaryViewInputCollection.cpp
--- a/Engine/Plugins/FX/Niagara/Source/NiagaraEditor/Private/ViewModels/Stack/NiagaraStackSummaryViewInputCollection.cpp
+++ b/Engine/Plugins/FX/Niagara/Source/NiagaraEditor/Private/ViewModels/Stack/NiagaraStackSummaryViewInputCollection.cpp
@@ -85,10 +85,10 @@ void UNiagaraStackSummaryViewCollection::RefreshForAdvancedToggle()
 
 bool UNiagaraStackSummaryViewCollection::GetShouldShowInStack() const
 {
-	// if(GetEmitterViewModel().IsValid())
-	// {
-	// 	return GetEmitterViewModel()->GetEditorData().ShouldShowSummaryView();
-	// }
+	if(GetEmitterViewModel().IsValid())
+	{
+		return GetEmitterViewModel()->GetEditorData().ShouldShowSummaryView();
+	}
 
 	return false;
 }
@@ -558,7 +558,7 @@ void UNiagaraStackSummaryViewCollection::RefreshChildrenInternal(const TArray<UN
 			EmtpySummaryMessage->Initialize(CreateDefaultChildRequiredData(), EmptyAssignmentNodeMessageText, GetStackEditorDataKey());
 		}
 		
-		//EmtpySummaryMessage->SetIsHidden(!GetShouldShowInStack());
+		EmtpySummaryMessage->SetIsHidden(!GetShouldShowInStack());
 		
 		NewChildren.Add(EmtpySummaryMessage);	
 	}
