Optimize UNiagaraScript::ComputeVMCompilationId()

diff --git a/Engine/Plugins/FX/Niagara/Source/Niagara/Private/NiagaraScript.cpp b/Engine/Plugins/FX/Niagara/Source/Niagara/Private/NiagaraScript.cpp
--- a/Engine/Plugins/FX/Niagara/Source/Niagara/Private/NiagaraScript.cpp
+++ b/Engine/Plugins/FX/Niagara/Source/Niagara/Private/NiagaraScript.cpp
@@ -30,6 +30,7 @@
 #include "Serialization/MemoryReader.h"
 #include "Serialization/ObjectAndNameAsStringProxyArchive.h"
 #include "Stats/Stats.h"
+#include "String/ParseTokens.h"
 #include "UObject/AssetRegistryTagsContext.h"
 #include "UObject/EditorObjectVersion.h"
 #include "UObject/FortniteMainBranchObjectVersion.h"
@@ -1097,15 +1098,13 @@ void UNiagaraScript::ComputeVMCompilationId(FNiagaraVMExecutableDataId& Id, cons
 	Id.bRequiresPersistentIDs = false;
 	Id.ScriptVersionID = IsVersioningEnabled() ? (VersionGuid.IsValid() ? VersionGuid : ExposedVersion) : FGuid();
 	
-	TArray<FNiagaraVariable> ReferencedStaticVars;
+	TArray<FNiagaraVariableBase> ReferencedStaticVars;
 	TArray<const uint8*> ReferencedStaticVarValues;
 	ENiagaraSimTarget SimTargetToBuild = ENiagaraSimTarget::CPUSim;
 	// Ideally we wouldn't want to do this but rather than push the data down
 	// from the emitter.  Checking all outers here to pick up simulation stages too.
 	FVersionedNiagaraEmitter Outer = GetOuterEmitter();
 	FVersionedNiagaraEmitterData* EmitterData = Outer.GetEmitterData();
-	int32 SystemSpawnIdx = INDEX_NONE;
-	int32 SystemUpdateIdx = INDEX_NONE;
 	TArray<UNiagaraScript*> Scripts;
 
 	FNiagaraScriptHashCollector HashCollector(GNiagaraDumpKeyGen == 1 || (OutHashCollector ? OutHashCollector->bCollectSources : false));
@@ -1115,8 +1114,8 @@ void UNiagaraScript::ComputeVMCompilationId(FNiagaraVMExecutableDataId& Id, cons
 		UNiagaraEmitter* Emitter = Outer.Emitter;
 		if (UNiagaraSystem* EmitterOwner = Cast<UNiagaraSystem>(Emitter->GetOuter()))
 		{
-			SystemSpawnIdx = Scripts.Add(EmitterOwner->GetSystemSpawnScript());
-			SystemUpdateIdx = Scripts.Add(EmitterOwner->GetSystemUpdateScript());
+			Scripts.Add(EmitterOwner->GetSystemSpawnScript());
+			Scripts.Add(EmitterOwner->GetSystemUpdateScript());
 
 			Id.bUsesRapidIterationParams = EmitterOwner->ShouldUseRapidIterationParameters();
 			Id.bDisableDebugSwitches = EmitterOwner->ShouldDisableDebugSwitches();
@@ -1331,8 +1330,8 @@ void UNiagaraScript::ComputeVMCompilationId(FNiagaraVMExecutableDataId& Id, cons
 	UObject* Obj = GetOuter();
 	if (UNiagaraSystem* System = Cast<UNiagaraSystem>(Obj))
 	{
-		SystemSpawnIdx = Scripts.Add(System->GetSystemSpawnScript());
-		SystemUpdateIdx = Scripts.Add(System->GetSystemUpdateScript());
+		Scripts.Add(System->GetSystemSpawnScript());
+		Scripts.Add(System->GetSystemUpdateScript());
 
 		Id.bUsesRapidIterationParams = System->ShouldUseRapidIterationParameters();
 		Id.bDisableDebugSwitches = System->ShouldDisableDebugSwitches();
@@ -1381,18 +1380,15 @@ void UNiagaraScript::ComputeVMCompilationId(FNiagaraVMExecutableDataId& Id, cons
 		checkf(false, TEXT("Unknown sim target type!"));
 	}
 
-
 	for (int32 ScriptIdx = 0; ScriptIdx < Scripts.Num(); ScriptIdx++)
 	{
 		UNiagaraScript* Script = Scripts[ScriptIdx];
-		TArray<FNiagaraVariable> Vars;
-		Script->RapidIterationParameters.GetParameters(Vars);		
-		bool bIsThisEmitterUsage = UNiagaraScript::IsEmitterScript(Usage);
-		bool bIsThisParticleUsage = UNiagaraScript::IsParticleScript(Usage);
+		const bool bIsThisEmitterUsage = UNiagaraScript::IsEmitterScript(Usage);
+		const bool bIsThisParticleUsage = UNiagaraScript::IsParticleScript(Usage);
 	
-		bool bIsOtherSystemUsage = UNiagaraScript::IsSystemScript(Script->Usage);
-		bool bIsOtherParticleUsage = UNiagaraScript::IsParticleScript(Script->Usage);
-		bool bIsOtherEmitterUsage = UNiagaraScript::IsEmitterScript(Script->Usage);
+		const bool bIsOtherSystemUsage = UNiagaraScript::IsSystemScript(Script->Usage);
+		const bool bIsOtherParticleUsage = UNiagaraScript::IsParticleScript(Script->Usage);
+		const bool bIsOtherEmitterUsage = UNiagaraScript::IsEmitterScript(Script->Usage);
 
 		// Emitter scripts don't depend on static variables from particle scripts.
 		if (bIsThisEmitterUsage && bIsOtherParticleUsage)
@@ -1404,14 +1400,15 @@ void UNiagaraScript::ComputeVMCompilationId(FNiagaraVMExecutableDataId& Id, cons
 			Script->GetLatestSource()->RegisterVMCompilationIdDependencies(HashCollector, Script->Usage, Script->UsageId);
 		}
 
-		for (const FNiagaraVariable& Var : Vars)
+		for (const FNiagaraVariableWithOffset& Var : Script->RapidIterationParameters.ReadParameterVariables())
 		{
 			if (Var.GetType().IsStatic())
 			{
 				if ((bIsThisEmitterUsage || bIsThisParticleUsage) && bIsOtherSystemUsage)
 				{
-					TArray<FString> SplitName;
-					Var.GetName().ToString().ParseIntoArray(SplitName, TEXT("."));
+					FNameBuilder VarNameBuilder(Var.GetName());
+					TArray<FStringView> SplitName;
+					UE::String::ParseTokens(VarNameBuilder.ToView(), TEXT("."), SplitName);
 
 					// Only include System based rapid iteration static variables if we're an emitter/particle script as system scripts can contain the amalgam of 
 					// it's own and any child emitter script rapid iteration variables.
@@ -1429,7 +1426,7 @@ void UNiagaraScript::ComputeVMCompilationId(FNiagaraVMExecutableDataId& Id, cons
 					}
 				}
 				ReferencedStaticVars.Add(Var);
-				ReferencedStaticVarValues.Add(Script->RapidIterationParameters.GetParameterData(Var));
+				ReferencedStaticVarValues.Add(Script->RapidIterationParameters.GetParameterData(Var.Offset));
 			}
 		}
 	}
@@ -1439,31 +1436,28 @@ void UNiagaraScript::ComputeVMCompilationId(FNiagaraVMExecutableDataId& Id, cons
 	
 	{
 		FSHA1 HashState;
+		auto UpdateHashState = [&](FName InName) -> void
+		{
+			FNameBuilder NameBuilder(InName);
+			FStringView NameBuilderView = NameBuilder.ToView();
+			HashState.UpdateWithString(NameBuilderView.GetData(), NameBuilderView.Len());
+		};
+
 		if (false == Id.bUsesRapidIterationParams)
 		{
-			TArray<FNiagaraVariable> Vars;
-			RapidIterationParameters.GetParameters(Vars);
 			//UE_LOG(LogNiagara, Display, TEXT("AreScriptAndSourceSynchronized %s ======================== "), *GetFullName());
-			for (int32 i = 0; i < Vars.Num(); i++)
+			for (const FNiagaraVariableWithOffset& Var : RapidIterationParameters.ReadParameterVariables())
 			{
-				if (Vars[i].IsValid() == false || Vars[i].IsDataInterface() || Vars[i].IsUObject())
+				if (Var.IsValid() == false || Var.IsDataInterface() || Var.IsUObject())
 				{
 					// Skip these types as they're invalid or don't bake out, just normal parameters get baked.
 				}
 				else
 				{
 					// Hash the name, type, and value of each parameter..
-					FString VarName = Vars[i].GetName().ToString();
-					FString VarTypeName = Vars[i].GetType().GetName();
-					HashState.UpdateWithString(*VarName, VarName.Len());
-					HashState.UpdateWithString(*VarTypeName, VarTypeName.Len());
-					TArray<uint8> DataValue;
-					DataValue.AddUninitialized(Vars[i].GetSizeInBytes());
-					if (RapidIterationParameters.CopyParameterData(Vars[i], DataValue.GetData()))
-					{
-						//UE_LOG(LogNiagara, Display, TEXT("Param %s %s %s"), *VarTypeName, *VarName, *ByteStr);
-						HashState.Update(DataValue.GetData(), Vars[i].GetType().GetSize());
-					}
+					UpdateHashState(Var.GetName());
+					UpdateHashState(Var.GetType().GetFName());
+					HashState.Update(RapidIterationParameters.GetParameterData(Var.Offset), Var.GetSizeInBytes());
 				}
 			}
 		}
@@ -1475,20 +1469,18 @@ void UNiagaraScript::ComputeVMCompilationId(FNiagaraVMExecutableDataId& Id, cons
 				UE_LOG(LogNiagara, Display, TEXT("***** Referenced Static Vars %d %s"), (int32)Usage, *GetPathName());
 			}
 
-			// Hash the name, type, and value of each parameter..
-			FString VarName = ReferencedStaticVars[i].GetName().ToString();
-			FString VarTypeName = ReferencedStaticVars[i].GetType().GetName();
-			HashState.UpdateWithString(*VarName, VarName.Len());
-			HashState.UpdateWithString(*VarTypeName, VarTypeName.Len());
-			const uint8* VarData = ReferencedStaticVarValues[i];
+			const FNiagaraVariableBase& StaticVar = ReferencedStaticVars[i];
 
-			if (VarData)
+			// Hash the name, type, and value of each parameter..
+			UpdateHashState(StaticVar.GetName());
+			UpdateHashState(StaticVar.GetType().GetFName());
+			if (const uint8* VarData = ReferencedStaticVarValues[i])
 			{
 				if (UNiagaraScript::LogCompileStaticVars > 0)
 				{
-					UE_LOG(LogNiagara, Display, TEXT("Param %s %s %d"), *VarTypeName, *VarName, (uint32)VarData[0]);
+					UE_LOG(LogNiagara, Display, TEXT("Param %s %s %d"), *StaticVar.GetType().GetName(), *StaticVar.GetName().ToString(), (uint32)VarData[0]);
 				}
-				HashState.Update(VarData, ReferencedStaticVars[i].GetType().GetSize());
+				HashState.Update(VarData, StaticVar.GetSizeInBytes());
 			}
 		}
 		HashState.Final();
