Changed Niagara node fixed title size to max title size.
Fixed title size could lead to a widening of emitter nodes when the name was short

diff --git a/Engine/Plugins/FX/Niagara/Source/NiagaraEditorWidgets/Private/SNiagaraOverviewStackNode.cpp b/Engine/Plugins/FX/Niagara/Source/NiagaraEditorWidgets/Private/SNiagaraOverviewStackNode.cpp
--- a/Engine/Plugins/FX/Niagara/Source/NiagaraEditorWidgets/Private/SNiagaraOverviewStackNode.cpp
+++ b/Engine/Plugins/FX/Niagara/Source/NiagaraEditorWidgets/Private/SNiagaraOverviewStackNode.cpp
@@ -51,7 +51,7 @@
 
 constexpr float RendererThumbnailSize = 24.f;
 constexpr float EmitterThumbnailSize = 200.f;
-constexpr float SummaryTitleFixedSize = 125.f;
+constexpr float NodeTitleMaxSize = 125.f;
 
 void SNiagaraOverviewStackNode::Construct(const FArguments& InArgs, UNiagaraOverviewNode* InNode)
 {
@@ -505,13 +505,10 @@ TSharedRef<SWidget> SNiagaraOverviewStackNode::CreateTitleWidget_Default(TShared
 	// Name
 	+ SHorizontalBox::Slot()
 	.Padding(3, 0, 0, 0)
-	.FillWidth(1.0f)
+	.AutoWidth()
+	.MaxWidth(NodeTitleMaxSize)
 	[
-		SNew(SBox)
-		.WidthOverride(SummaryTitleFixedSize)
-		[
-			DefaultTitle
-		]
+		DefaultTitle
 	];
 
 	return TitleWidget.ToSharedRef();
@@ -545,6 +542,7 @@ TSharedRef<SWidget> SNiagaraOverviewStackNode::CreateTitleRightWidget_Default()
 
 	// scalability controls
 	+ SHorizontalBox::Slot()
+	.HAlign(HAlign_Right)
 	.AutoWidth()
 	[
 		CreateScalabilityControls()
