- Experimental stateless emitters
- Enabled via a project cvar, data formats are not guaranteed to remain backwards compatible
- This is WIP lots of UI work, curves, etc, to come

diff --git a/Engine/Plugins/FX/Niagara/Shaders/Private/NiagaraQuaternionUtils.ush b/Engine/Plugins/FX/Niagara/Shaders/Private/NiagaraQuaternionUtils.ush
--- a/Engine/Plugins/FX/Niagara/Shaders/Private/NiagaraQuaternionUtils.ush
+++ b/Engine/Plugins/FX/Niagara/Shaders/Private/NiagaraQuaternionUtils.ush
@@ -163,3 +163,19 @@ float4 QuatFromMatrix(in float3 R[3])
 	}
 	return NormalizeQuat(Q);
 }
+
+float4 RotatorToQuat(float3 Rotator)
+{
+	Rotator = frac(Rotator) * 3.1415926535897932f;
+	float SP, CP, SY, CY, SR, CR;
+	sincos(Rotator.x, SR, CR);
+	sincos(Rotator.y, SP, CP);
+	sincos(Rotator.z, SY, CY);
+	
+	return float4(
+		 CR*SP*SY - SR*CP*CY,
+		-CR*SP*CY - SR*CP*SY,
+		 CR*CP*SY - SR*SP*CY,
+		 CR*CP*CY + SR*SP*SY
+	);
+}
diff --git a/Engine/Plugins/FX/Niagara/Shaders/Private/Stateless/Modules/NiagaraStatelessModule_CalculateAccurateVelocity.ush b/Engine/Plugins/FX/Niagara/Shaders/Private/Stateless/Modules/NiagaraStatelessModule_CalculateAccurateVelocity.ush
new file mode 100644
--- /dev/null
+++ b/Engine/Plugins/FX/Niagara/Shaders/Private/Stateless/Modules/NiagaraStatelessModule_CalculateAccurateVelocity.ush
@@ -0,0 +1,18 @@
+// Copyright Epic Games, Inc. All Rights Reserved.
+
+#include "../NiagaraStatelessCommon.ush"
+
+struct FStatelessModule_CalculateAccurateVelocity
+{
+};
+
+void CalculateAccurateVelocity_Initialize(inout FStatelessModule_CalculateAccurateVelocity ModuleData, inout FStatelessParticle Particle)
+{  
+}
+
+void CalculateAccurateVelocity_Simulate(in FStatelessModule_CalculateAccurateVelocity ModuleData, inout FStatelessParticle Particle)
+{  
+	Particle.Velocity			= Particle.Position - Particle.PreviousPosition;
+	Particle.PreviousVelocity	= Particle.Velocity;
+}
+	
\ No newline at end of file
diff --git a/Engine/Plugins/FX/Niagara/Shaders/Private/Stateless/Modules/NiagaraStatelessModule_InitialMeshOrientation.ush b/Engine/Plugins/FX/Niagara/Shaders/Private/Stateless/Modules/NiagaraStatelessModule_InitialMeshOrientation.ush
new file mode 100644
--- /dev/null
+++ b/Engine/Plugins/FX/Niagara/Shaders/Private/Stateless/Modules/NiagaraStatelessModule_InitialMeshOrientation.ush
@@ -0,0 +1,21 @@
+// Copyright Epic Games, Inc. All Rights Reserved.
+
+#include "../NiagaraStatelessCommon.ush"
+
+float3	InitialMeshOrientation_Rotation;
+float3	InitialMeshOrientation_RandomRangeScale;
+
+struct FStatelessModule_InitialMeshOrientation
+{
+};
+
+void InitialMeshOrientation_Initialize(inout FStatelessModule_InitialMeshOrientation ModuleData, inout FStatelessParticle Particle)
+{
+	const float3 InitialRotation = InitialMeshOrientation_Rotation + (RandomFloat3(Particle) * InitialMeshOrientation_RandomRangeScale);
+	Particle.MeshOrientation			= RotatorToQuat(InitialRotation);
+	Particle.PreviousMeshOrientation	= Particle.MeshOrientation;
+}
+	
+void InitialMeshOrientation_Simulate(in FStatelessModule_InitialMeshOrientation ModuleData, inout FStatelessParticle Particle)
+{
+}
diff --git a/Engine/Plugins/FX/Niagara/Shaders/Private/Stateless/Modules/NiagaraStatelessModule_InitializeParticle.ush b/Engine/Plugins/FX/Niagara/Shaders/Private/Stateless/Modules/NiagaraStatelessModule_InitializeParticle.ush
new file mode 100644
--- /dev/null
+++ b/Engine/Plugins/FX/Niagara/Shaders/Private/Stateless/Modules/NiagaraStatelessModule_InitializeParticle.ush
@@ -0,0 +1,68 @@
+// Copyright Epic Games, Inc. All Rights Reserved.
+
+#include "../NiagaraStatelessCommon.ush"
+
+float3	InitializeParticle_Position;
+float4	InitializeParticle_ColorScale;
+float4	InitializeParticle_ColorBias;
+float2	InitializeParticle_SpriteSizeScale;
+float2	InitializeParticle_SpriteSizeBias;
+float	InitializeParticle_SpriteRotationScale;
+float	InitializeParticle_SpriteRotationBias;
+
+float3	InitializeParticle_MeshScaleScale;
+float3	InitializeParticle_MeshScaleBias;
+
+float	InitializeParticle_RibbonWidthScale;
+float	InitializeParticle_RibbonWidthBias;
+	
+uint InitializeParticle_PositionMode;
+uint InitializeParticle_PositionTableMod;
+uint InitializeParticle_PositionTableOffset;
+uint InitializeParticle_PositionTableLength;
+
+struct FStatelessModule_InitializeParticle
+{
+};
+
+void InitializeParticle_Initialize(inout FStatelessModule_InitializeParticle ModuleData, inout FStatelessParticle Particle)
+{  
+	Particle.UniqueID		= Particle.UniqueIndex;
+	Particle.Position		= InitializeParticle_Position;
+	Particle.Color			= RandomScaleBiasFloat(Particle, InitializeParticle_ColorScale, InitializeParticle_ColorBias);
+	Particle.RibbonWidth	= RandomScaleBiasFloat(Particle, InitializeParticle_RibbonWidthScale, InitializeParticle_RibbonWidthBias);
+	Particle.SpriteSize		= RandomScaleBiasFloat(Particle, InitializeParticle_SpriteSizeScale, InitializeParticle_SpriteSizeBias);
+	Particle.SpriteRotation	= RandomScaleBiasFloat(Particle, InitializeParticle_SpriteRotationScale, InitializeParticle_SpriteRotationBias);
+	Particle.Scale			= RandomScaleBiasFloat(Particle, InitializeParticle_MeshScaleScale, InitializeParticle_MeshScaleBias);
+
+	// For testing purposes only
+	{
+		if ( InitializeParticle_PositionMode == 0 )
+		{
+			const uint Len = InitializeParticle_PositionTableLength - 1;
+			const uint Element = clamp(Particle.UniqueIndex % InitializeParticle_PositionTableMod, 0, Len);
+			Particle.Position += GetStaticFloat3(InitializeParticle_PositionTableOffset, Element);
+		}
+		else
+		{
+			const uint Len = InitializeParticle_PositionTableLength - 1;
+			const float	R0 = RandomScaleBiasFloat(Particle, Len, 0);
+			const float U = InitializeParticle_PositionMode != 2 ? frac(R0) : 0.0f;
+			const uint E0 = floor(R0);
+			const uint E1 = InitializeParticle_PositionMode == 1 ? min(E0 + 1, Len) : (RandomUInt(Particle) % InitializeParticle_PositionTableMod);
+			const float3 P0 = GetStaticFloat3(InitializeParticle_PositionTableOffset, E0);
+			const float3 P1 = GetStaticFloat3(InitializeParticle_PositionTableOffset, E1);
+			Particle.Position += lerp(P0, P1, U);
+		}
+	}
+
+	Particle.PreviousPosition		= Particle.Position;
+	Particle.PreviousRibbonWidth	= Particle.RibbonWidth;
+	Particle.PreviousSpriteSize		= Particle.SpriteSize;
+	Particle.PreviousScale			= Particle.Scale;
+	Particle.PreviousSpriteRotation	= Particle.SpriteRotation;
+}
+
+void InitializeParticle_Simulate(in FStatelessModule_InitializeParticle ModuleData, inout FStatelessParticle Particle)
+{  
+}
diff --git a/Engine/Plugins/FX/Niagara/Shaders/Private/Stateless/Modules/NiagaraStatelessModule_RotateAroundPoint.ush b/Engine/Plugins/FX/Niagara/Shaders/Private/Stateless/Modules/NiagaraStatelessModule_RotateAroundPoint.ush
new file mode 100644
--- /dev/null
+++ b/Engine/Plugins/FX/Niagara/Shaders/Private/Stateless/Modules/NiagaraStatelessModule_RotateAroundPoint.ush
@@ -0,0 +1,40 @@
+// Copyright Epic Games, Inc. All Rights Reserved.
+
+#include "../NiagaraStatelessCommon.ush"
+
+float	RotateAroundPoint_RateScale;
+float	RotateAroundPoint_RateBias;
+float	RotateAroundPoint_RadiusScale;
+float	RotateAroundPoint_RadiusBias;
+float	RotateAroundPoint_InitialPhaseScale;
+float	RotateAroundPoint_InitialPhaseBias;
+
+struct FStatelessModule_RotateAroundPoint
+{
+	float Rate;
+	float Radius;
+	float InitialPhase;
+};
+
+void RotateAroundPoint_IntegratePosition(in FStatelessModule_RotateAroundPoint ModuleData, float Age, inout float3 Position)
+{
+	const float Tau		= 3.14f * 2.0f;
+	const float3 AxisX	= float3(1.0f, 0.0f, 0.0f);
+	const float3 AxisY	= float3(0.0f, 1.0f, 0.0f);
+
+	Position += sin(ModuleData.InitialPhase + Age * ModuleData.Rate) * ModuleData.Radius * AxisX;
+	Position += cos(ModuleData.InitialPhase + Age * ModuleData.Rate) * ModuleData.Radius * AxisY;
+}
+
+void RotateAroundPoint_Initialize(inout FStatelessModule_RotateAroundPoint ModuleData, inout FStatelessParticle Particle)
+{
+	ModuleData.Rate			= RandomScaleBiasFloat(Particle, RotateAroundPoint_RateScale, RotateAroundPoint_RateBias);
+	ModuleData.Radius		= RandomScaleBiasFloat(Particle, RotateAroundPoint_RadiusScale, RotateAroundPoint_RadiusBias);
+	ModuleData.InitialPhase	= RandomScaleBiasFloat(Particle, RotateAroundPoint_InitialPhaseScale, RotateAroundPoint_InitialPhaseBias);
+}
+	
+void RotateAroundPoint_Simulate(in FStatelessModule_RotateAroundPoint ModuleData, inout FStatelessParticle Particle)
+{
+	RotateAroundPoint_IntegratePosition(ModuleData, Particle.NormalizedAge, Particle.Position);
+	RotateAroundPoint_IntegratePosition(ModuleData, Particle.PreviousNormalizedAge, Particle.PreviousPosition);
+}
diff --git a/Engine/Plugins/FX/Niagara/Shaders/Private/Stateless/Modules/NiagaraStatelessModule_ScaleColor.ush b/Engine/Plugins/FX/Niagara/Shaders/Private/Stateless/Modules/NiagaraStatelessModule_ScaleColor.ush
new file mode 100644
--- /dev/null
+++ b/Engine/Plugins/FX/Niagara/Shaders/Private/Stateless/Modules/NiagaraStatelessModule_ScaleColor.ush
@@ -0,0 +1,36 @@
+// Copyright Epic Games, Inc. All Rights Reserved.
+
+#include "../NiagaraStatelessCommon.ush"
+
+int		ScaleColor_Offset;
+int		ScaleColor_Length;
+
+struct FStatelessModule_ScaleColor
+{
+};
+
+float4 ScaleColor_GetScale(in FStatelessParticle Particle, float NormalizedAge)
+{
+	const float Offset = NormalizedAge * float(ScaleColor_Length - 1);
+	const int A = floor(Offset);
+	const int B = ceil(Offset);
+	const float U = frac(Offset);
+
+	return lerp(
+		GetStaticFloat4(ScaleColor_Offset, A),
+		GetStaticFloat4(ScaleColor_Offset, B), 
+		U
+	);
+}
+
+void ScaleColor_Initialize(inout FStatelessModule_ScaleColor ModuleData, inout FStatelessParticle Particle)
+{
+}
+	
+void ScaleColor_Simulate(in FStatelessModule_ScaleColor ModuleData, inout FStatelessParticle Particle)
+{
+	if (ScaleColor_Length > 0)
+	{
+		Particle.Color *= ScaleColor_GetScale(Particle, Particle.NormalizedAge);
+	}
+}
diff --git a/Engine/Plugins/FX/Niagara/Shaders/Private/Stateless/Modules/NiagaraStatelessModule_ScaleMeshSize.ush b/Engine/Plugins/FX/Niagara/Shaders/Private/Stateless/Modules/NiagaraStatelessModule_ScaleMeshSize.ush
new file mode 100644
--- /dev/null
+++ b/Engine/Plugins/FX/Niagara/Shaders/Private/Stateless/Modules/NiagaraStatelessModule_ScaleMeshSize.ush
@@ -0,0 +1,38 @@
+// Copyright Epic Games, Inc. All Rights Reserved.
+
+#include "../NiagaraStatelessCommon.ush"
+
+int		ScaleMeshSize_Offset;
+int		ScaleMeshSize_Length;
+
+struct FStatelessModule_ScaleMeshSize
+{
+};
+
+
+float3 ScaleMeshSize_GetScale(in FStatelessParticle Particle, float NormalizedAge)
+{
+	const float Offset = NormalizedAge * float(ScaleMeshSize_Length - 1);
+	const int A = floor(Offset);
+	const int B = ceil(Offset);
+	const float U = frac(Offset);
+
+	return lerp(
+		GetStaticFloat3(ScaleMeshSize_Offset, A),
+		GetStaticFloat3(ScaleMeshSize_Offset, B), 
+		U
+	);
+}
+	
+void ScaleMeshSize_Initialize(inout FStatelessModule_ScaleMeshSize ModuleData, inout FStatelessParticle Particle)
+{
+}
+	
+void ScaleMeshSize_Simulate(in FStatelessModule_ScaleMeshSize ModuleData, inout FStatelessParticle Particle)
+{
+	if (ScaleMeshSize_Length > 0)
+	{
+		Particle.Scale *= ScaleMeshSize_GetScale(Particle, Particle.NormalizedAge);
+		Particle.PreviousScale *= ScaleMeshSize_GetScale(Particle, Particle.PreviousNormalizedAge);
+	}
+}
diff --git a/Engine/Plugins/FX/Niagara/Shaders/Private/Stateless/Modules/NiagaraStatelessModule_ScaleMeshSizeBySpeed.ush b/Engine/Plugins/FX/Niagara/Shaders/Private/Stateless/Modules/NiagaraStatelessModule_ScaleMeshSizeBySpeed.ush
new file mode 100644
--- /dev/null
+++ b/Engine/Plugins/FX/Niagara/Shaders/Private/Stateless/Modules/NiagaraStatelessModule_ScaleMeshSizeBySpeed.ush
@@ -0,0 +1,41 @@
+// Copyright Epic Games, Inc. All Rights Reserved.
+
+#include "../NiagaraStatelessCommon.ush"
+
+float	ScaleMeshSizeBySpeed_VelocityNorm;
+int		ScaleMeshSizeBySpeed_Offset;
+int		ScaleMeshSizeBySpeed_Length;
+
+struct FStatelessModule_ScaleMeshSizeBySpeed
+{
+};
+
+float3 ScaleMeshSizeBySpeed_GetScale(in FStatelessParticle Particle, float NormSpeed)
+{
+	const float Offset = NormSpeed * float(ScaleMeshSizeBySpeed_Length - 1);
+	const int A = floor(Offset);
+	const int B = ceil(Offset);
+	const float U = frac(Offset);
+
+	return lerp(
+		GetStaticFloat3(ScaleMeshSizeBySpeed_Offset, A),
+		GetStaticFloat3(ScaleMeshSizeBySpeed_Offset, B),
+		U
+	);
+}
+
+void ScaleMeshSizeBySpeed_Initialize(inout FStatelessModule_ScaleMeshSizeBySpeed MouldeData, inout FStatelessParticle Particle)
+{
+}
+	
+void ScaleMeshSizeBySpeed_Simulate(in FStatelessModule_ScaleMeshSizeBySpeed MouldeData, inout FStatelessParticle Particle)
+{
+	if (ScaleMeshSizeBySpeed_Length > 0)
+	{
+		const float Speed = length2(Particle.Position - Particle.PreviousPosition) * Particle.InvDeltaTime;
+		const float NormSpeed = saturate(Speed * ScaleMeshSizeBySpeed_VelocityNorm);
+			
+		Particle.Scale			*= ScaleMeshSizeBySpeed_GetScale(Particle, NormSpeed);
+		Particle.PreviousScale	*= ScaleMeshSizeBySpeed_GetScale(Particle, NormSpeed);
+	}
+}
diff --git a/Engine/Plugins/FX/Niagara/Shaders/Private/Stateless/Modules/NiagaraStatelessModule_ScaleSpriteSize.ush b/Engine/Plugins/FX/Niagara/Shaders/Private/Stateless/Modules/NiagaraStatelessModule_ScaleSpriteSize.ush
new file mode 100644
--- /dev/null
+++ b/Engine/Plugins/FX/Niagara/Shaders/Private/Stateless/Modules/NiagaraStatelessModule_ScaleSpriteSize.ush
@@ -0,0 +1,37 @@
+// Copyright Epic Games, Inc. All Rights Reserved.
+
+#include "../NiagaraStatelessCommon.ush"
+
+int		ScaleSpriteSize_Offset;
+int		ScaleSpriteSize_Length;
+
+struct FStatelessModule_ScaleSpriteSize
+{
+};
+	
+float2 ScaleSpriteSize_GetScale(in FStatelessParticle Particle, float NormalizedAge)
+{
+	const float Offset = NormalizedAge * float(ScaleSpriteSize_Length - 1);
+	const int A = floor(Offset);
+	const int B = ceil(Offset);
+	const float U = frac(Offset);
+
+	return lerp(
+		GetStaticFloat2(ScaleSpriteSize_Offset, A),
+		GetStaticFloat2(ScaleSpriteSize_Offset, B), 
+		U
+	);
+}
+	
+void ScaleSpriteSize_Initialize(inout FStatelessModule_ScaleSpriteSize ModuleData, inout FStatelessParticle Particle)
+{
+}
+	
+void ScaleSpriteSize_Simulate(in FStatelessModule_ScaleSpriteSize ModuleData, inout FStatelessParticle Particle)
+{
+	if (ScaleSpriteSize_Length > 0)
+	{
+		Particle.SpriteSize *= ScaleSpriteSize_GetScale(Particle, Particle.NormalizedAge);
+		Particle.PreviousSpriteSize *= ScaleSpriteSize_GetScale(Particle, Particle.PreviousNormalizedAge);
+	}
+}
diff --git a/Engine/Plugins/FX/Niagara/Shaders/Private/Stateless/Modules/NiagaraStatelessModule_ScaleSpriteSizeBySpeed.ush b/Engine/Plugins/FX/Niagara/Shaders/Private/Stateless/Modules/NiagaraStatelessModule_ScaleSpriteSizeBySpeed.ush
new file mode 100644
--- /dev/null
+++ b/Engine/Plugins/FX/Niagara/Shaders/Private/Stateless/Modules/NiagaraStatelessModule_ScaleSpriteSizeBySpeed.ush
@@ -0,0 +1,41 @@
+// Copyright Epic Games, Inc. All Rights Reserved.
+
+#include "../NiagaraStatelessCommon.ush"
+
+float	ScaleSpriteSizeBySpeed_VelocityNorm;
+int		ScaleSpriteSizeBySpeed_Offset;
+int		ScaleSpriteSizeBySpeed_Length;
+
+struct FStatelessModule_ScaleSpriteSizeBySpeed
+{
+};
+
+float2 ScaleSpriteSizeBySpeed_GetScale(in FStatelessParticle Particle, float NormSpeed)
+{
+	const float Offset = NormSpeed * float(ScaleSpriteSizeBySpeed_Length - 1);
+	const int A = floor(Offset);
+	const int B = ceil(Offset);
+	const float U = frac(Offset);
+
+	return lerp(
+		GetStaticFloat2(ScaleSpriteSizeBySpeed_Offset, A),
+		GetStaticFloat2(ScaleSpriteSizeBySpeed_Offset, B),
+		U
+	);
+}
+	
+void ScaleSpriteSizeBySpeed_Initialize(inout FStatelessModule_ScaleSpriteSizeBySpeed ModuleData, inout FStatelessParticle Particle)
+{
+}
+	
+void ScaleSpriteSizeBySpeed_Simulate(in FStatelessModule_ScaleSpriteSizeBySpeed ModuleData, inout FStatelessParticle Particle)
+{
+	if (ScaleSpriteSizeBySpeed_Length > 0)
+	{
+		const float Speed = length2(Particle.Position - Particle.PreviousPosition) * Particle.InvDeltaTime;
+		const float NormSpeed = saturate(Speed * ScaleSpriteSizeBySpeed_VelocityNorm);
+			
+		Particle.SpriteSize			*= ScaleSpriteSizeBySpeed_GetScale(Particle, NormSpeed);
+		Particle.PreviousSpriteSize	*= ScaleSpriteSizeBySpeed_GetScale(Particle, NormSpeed);
+	}
+}
diff --git a/Engine/Plugins/FX/Niagara/Shaders/Private/Stateless/Modules/NiagaraStatelessModule_ShapeLocation.ush b/Engine/Plugins/FX/Niagara/Shaders/Private/Stateless/Modules/NiagaraStatelessModule_ShapeLocation.ush
new file mode 100644
--- /dev/null
+++ b/Engine/Plugins/FX/Niagara/Shaders/Private/Stateless/Modules/NiagaraStatelessModule_ShapeLocation.ush
@@ -0,0 +1,92 @@
+// Copyright Epic Games, Inc. All Rights Reserved.
+
+#include "../NiagaraStatelessCommon.ush"
+
+uint4	ShapeLocation_Mode;
+float4	ShapeLocation_Parameters0;
+float4	ShapeLocation_Parameters1;
+
+struct FStatelessModule_ShapeLocation
+{
+};
+
+float3 ShapeLocation_GetLocation(inout FStatelessParticle Particle)
+{
+	// ENSM_ShapePrimitive::Box | ENSM_ShapePrimitive::Plane
+	BRANCH
+	if ( ShapeLocation_Mode.x == 0 )
+	{
+		const float3 BoxScale = ShapeLocation_Parameters0.xyz;
+		const float3 BoxBias = ShapeLocation_Parameters1.xyz;
+		const bool bOnSurface = ShapeLocation_Mode.y == 1;
+		const float SurfaceScale = ShapeLocation_Parameters0.w;
+		const float SurfaceBias = ShapeLocation_Parameters1.w;
+
+		const float3 P0 = RandomFloat3(Particle);
+		if (bOnSurface)
+		{
+			//-TODO: This isn't quite what we want
+			const float3 L0 = round(P0);
+			const uint S = RandomUInt(Particle) % 3;
+			float3 Location;
+			Location.x = S != 0 ? P0.x : L0.x;
+			Location.y = S != 1 ? P0.y : L0.y;
+			Location.z = S != 2 ? P0.z : L0.z;
+			return (Location * BoxScale + BoxBias) + (P0 * SurfaceScale + SurfaceBias);
+		}
+		else
+		{
+			return P0 * BoxScale + BoxBias;
+		}
+	}
+
+	// ENSM_ShapePrimitive::Cylinder:
+	BRANCH
+	if ( ShapeLocation_Mode.x == 1 )
+	{
+		const float4 Random = RandomFloat4(Particle);
+		const float HeightScale = ShapeLocation_Parameters0.x;
+		const float HeightBias = ShapeLocation_Parameters0.y;
+		const float Radius = ShapeLocation_Parameters0.z;
+
+		return float3(
+			SafeNormalize(Random.xy - 0.5f) * Radius * Random.z,
+			Random.w * HeightScale + HeightBias
+		);
+	}
+
+	// ENSM_ShapePrimitive::Ring:
+	BRANCH
+	if ( ShapeLocation_Mode.x == 2 )
+	{
+		const float RadiusScale = ShapeLocation_Parameters0.x;
+		const float RadiusBias = ShapeLocation_Parameters0.y;
+		const float UDistributionScale = ShapeLocation_Parameters0.z;
+		const float UDistributionBias = ShapeLocation_Parameters0.w;
+
+		const float Radius = RandomScaleBiasFloat(Particle, RadiusScale, RadiusBias);
+		const float U = RandomScaleBiasFloat(Particle, UDistributionScale, UDistributionBias);
+
+		return float3(cos(U) * Radius, sin(U) * Radius, 0.0f);
+	}
+
+	// ENSM_ShapePrimitive::Sphere:
+	{
+		const float SphereScale	= ShapeLocation_Parameters0.x;
+		const float SphereBias	= ShapeLocation_Parameters0.y;
+
+		const float3 Vector = RandomUnitFloat3(Particle);
+		return Vector * RandomScaleBiasFloat(Particle, SphereScale, SphereBias);
+	}
+}
+
+void ShapeLocation_Initialize(inout FStatelessModule_ShapeLocation ModuleData, inout FStatelessParticle Particle)
+{
+	float3 Location = ShapeLocation_GetLocation(Particle);
+	Particle.Position += Location;
+	Particle.PreviousPosition += Location;
+}
+	
+void ShapeLocation_Simulate(in FStatelessModule_ShapeLocation ModuleData, inout FStatelessParticle Particle)
+{
+}
diff --git a/Engine/Plugins/FX/Niagara/Shaders/Private/Stateless/Modules/NiagaraStatelessModule_SolveVelocitiesAndForces.ush b/Engine/Plugins/FX/Niagara/Shaders/Private/Stateless/Modules/NiagaraStatelessModule_SolveVelocitiesAndForces.ush
new file mode 100644
--- /dev/null
+++ b/Engine/Plugins/FX/Niagara/Shaders/Private/Stateless/Modules/NiagaraStatelessModule_SolveVelocitiesAndForces.ush
@@ -0,0 +1,168 @@
+// Copyright Epic Games, Inc. All Rights Reserved.
+
+#include "../NiagaraStatelessCommon.ush"
+
+float	SolveVelocitiesAndForces_MassScale;
+float	SolveVelocitiesAndForces_MassBias;
+float	SolveVelocitiesAndForces_DragScale;
+float	SolveVelocitiesAndForces_DragBias;
+float3	SolveVelocitiesAndForces_VelocityScale;
+float3	SolveVelocitiesAndForces_VelocityBias;
+float3	SolveVelocitiesAndForces_WindScale;
+float3	SolveVelocitiesAndForces_WindBias;
+float3	SolveVelocitiesAndForces_AccelerationScale;
+float3	SolveVelocitiesAndForces_AccelerationBias;
+
+uint	SolveVelocitiesAndForces_ConeVelocityEnabled;
+float4	SolveVelocitiesAndForces_ConeQuat;
+float	SolveVelocitiesAndForces_ConeVelocityScale;
+float	SolveVelocitiesAndForces_ConeVelocityBias;
+float	SolveVelocitiesAndForces_ConeAngleScale;
+float	SolveVelocitiesAndForces_ConeAngleBias;
+float	SolveVelocitiesAndForces_ConeVelocityFalloff;
+
+uint	SolveVelocitiesAndForces_PontVelocityEnabled;
+float	SolveVelocitiesAndForces_PointVelocityScale;
+float	SolveVelocitiesAndForces_PointVelocityBias;
+float3	SolveVelocitiesAndForces_PointOrigin;
+
+uint	SolveVelocitiesAndForces_NoiseEnabled;
+float	SolveVelocitiesAndForces_NoiseAmplitude;
+float3	SolveVelocitiesAndForces_NoiseFrequency;
+Texture3D		SolveVelocitiesAndForces_NoiseTexture;
+SamplerState	SolveVelocitiesAndForces_NoiseSampler;
+
+uint	SolveVelocitiesAndForces_NoiseMode;
+uint	SolveVelocitiesAndForces_NoiseLUTOffset;
+uint	SolveVelocitiesAndForces_NoiseLUTNumChannel;
+uint	SolveVelocitiesAndForces_NoiseLUTChannelWidth;
+
+struct FStatelessModule_SolveVelocitiesAndForces
+{
+	float  Mass;
+	float  Drag;
+	float3 Velocity;
+	float3 Wind;
+	float3 Acceleration;
+};
+
+void SolveVelocitiesAndForces_IntegratePosition(in FStatelessModule_SolveVelocitiesAndForces ModuleData, float Age, inout float3 Position)
+{
+	const float3 IntVelocity = (ModuleData.Velocity - ModuleData.Wind) + (ModuleData.Wind * Age * Age);
+	const float LambdaDragMass = max(ModuleData.Drag * rcp(ModuleData.Mass), 0.0001f);
+	const float LambdaAge = (1.0f - exp(-(LambdaDragMass * Age))) / LambdaDragMass;
+	Position += IntVelocity * LambdaAge;
+	Position += (Age - LambdaAge) * (ModuleData.Acceleration / LambdaDragMass);
+}
+	
+void SolveVelocitiesAndForces_Initialize(inout FStatelessModule_SolveVelocitiesAndForces ModuleData, inout FStatelessParticle Particle)
+{
+	ModuleData.Mass = RandomScaleBiasFloat(Particle, SolveVelocitiesAndForces_MassScale, SolveVelocitiesAndForces_MassBias);
+	ModuleData.Drag = RandomScaleBiasFloat(Particle, SolveVelocitiesAndForces_DragScale, SolveVelocitiesAndForces_DragBias);
+	ModuleData.Velocity = RandomScaleBiasFloat(Particle, SolveVelocitiesAndForces_VelocityScale, SolveVelocitiesAndForces_VelocityBias);
+	ModuleData.Wind = RandomScaleBiasFloat(Particle, SolveVelocitiesAndForces_WindScale, SolveVelocitiesAndForces_WindBias);
+	ModuleData.Acceleration = RandomScaleBiasFloat(Particle, SolveVelocitiesAndForces_AccelerationScale, SolveVelocitiesAndForces_AccelerationBias);
+
+	if ( SolveVelocitiesAndForces_ConeVelocityEnabled != 0 )
+	{
+		const float ConeAngle = RandomScaleBiasFloat(Particle, SolveVelocitiesAndForces_ConeAngleScale, SolveVelocitiesAndForces_ConeAngleBias);
+		const float ConeRotation = RandomFloat(Particle) * UE_TWO_PI;
+			
+		float2 scAng = SinCos(ConeAngle);
+		float2 scRot = SinCos(ConeRotation);
+		const float3 Direction = float3(scRot.x * scAng.x, scRot.y * scAng.x, scAng.y);
+		
+		float VelocityScale = RandomScaleBiasFloat(Particle, SolveVelocitiesAndForces_ConeVelocityScale, SolveVelocitiesAndForces_ConeVelocityBias);
+		if ( SolveVelocitiesAndForces_ConeVelocityFalloff > 0.0f )
+		{
+			const float pf = pow(saturate(scAng.y), SolveVelocitiesAndForces_ConeVelocityFalloff * 10.0f);
+			VelocityScale *= lerp(1.0f, pf, SolveVelocitiesAndForces_ConeVelocityFalloff);
+		}
+			
+		ModuleData.Velocity += RotateVectorByQuat(Direction, SolveVelocitiesAndForces_ConeQuat) * VelocityScale;
+	}
+
+	if ( SolveVelocitiesAndForces_PontVelocityEnabled != 0 )
+	{
+		const float3 FallbackDir = RandomUnitFloat3(Particle);
+		const float3 Delta = Particle.Position - SolveVelocitiesAndForces_PointOrigin;
+		const float  DeltaL2 = length2(Delta);
+		const float3 Dir = DeltaL2 > UE_KINDA_SMALL_NUMBER ? Delta * rsqrt(DeltaL2) : FallbackDir;
+			
+		const float VelocityScale = RandomScaleBiasFloat(Particle, SolveVelocitiesAndForces_PointVelocityScale, SolveVelocitiesAndForces_PointVelocityBias);
+			
+		ModuleData.Velocity += Dir * VelocityScale;
+	}
+
+	if ( SolveVelocitiesAndForces_NoiseEnabled  != 0 )
+	{
+		const float3 NoiseOffset = RandomFloat3(Particle);
+		float3 NoiseValue = 0;
+		if ( SolveVelocitiesAndForces_NoiseMode == 0 )
+		{
+			const float3 NoiseTravel = float3(Particle.Age, Particle.NormalizedAge, Particle.UniqueIndex);
+			const float3 NoisePosition = NoiseOffset + (NoiseTravel * SolveVelocitiesAndForces_NoiseFrequency);
+			NoiseValue = Texture3DSample(SolveVelocitiesAndForces_NoiseTexture, SolveVelocitiesAndForces_NoiseSampler, NoisePosition).xyz;
+		}
+		else if ( SolveVelocitiesAndForces_NoiseMode == 1 )
+		{
+			const float3 NoiseTravel = float3(Particle.Age, Particle.NormalizedAge, Particle.UniqueIndex) * 5.0f;
+			const float3 NoisePosition = NoiseOffset + (NoiseTravel * SolveVelocitiesAndForces_NoiseFrequency);
+			const float3x4 J = JacobianSimplex_ALU(NoisePosition, false, 1.0f);
+			NoiseValue = float3(J[1][2] - J[2][1], J[2][0] - J[0][2], J[0][1] - J[1][0]); // See comments to JacobianSimplex_ALU in Random.ush
+		}
+		else if ( SolveVelocitiesAndForces_NoiseMode == 2 )
+		{
+			const uint3 ChannelIndex = RandomUInt3(Particle) % SolveVelocitiesAndForces_NoiseLUTNumChannel;
+			const float3 ChannelLerp = (Particle.Age * SolveVelocitiesAndForces_NoiseFrequency) + RandomFloat(Particle);
+			const uint3 ChannelLUTOffsetA = (ChannelIndex * SolveVelocitiesAndForces_NoiseLUTChannelWidth) + min(floor(ChannelLerp) + 0, SolveVelocitiesAndForces_NoiseLUTChannelWidth - 1);
+			const uint3 ChannelLUTOffsetB = (ChannelIndex * SolveVelocitiesAndForces_NoiseLUTChannelWidth) + min(floor(ChannelLerp) + 1, SolveVelocitiesAndForces_NoiseLUTChannelWidth - 1);
+				
+			const float3 ValuesA = float3(
+					GetStaticFloat(SolveVelocitiesAndForces_NoiseLUTOffset, ChannelLUTOffsetA.x),
+					GetStaticFloat(SolveVelocitiesAndForces_NoiseLUTOffset, ChannelLUTOffsetA.y),
+					GetStaticFloat(SolveVelocitiesAndForces_NoiseLUTOffset, ChannelLUTOffsetA.z)
+			);
+			const float3 ValuesB = float3(
+					GetStaticFloat(SolveVelocitiesAndForces_NoiseLUTOffset, ChannelLUTOffsetB.x),
+					GetStaticFloat(SolveVelocitiesAndForces_NoiseLUTOffset, ChannelLUTOffsetB.y),
+					GetStaticFloat(SolveVelocitiesAndForces_NoiseLUTOffset, ChannelLUTOffsetB.z)
+			);
+			NoiseValue = lerp(ValuesA, ValuesB, frac(ChannelLerp));
+		}
+		else if ( SolveVelocitiesAndForces_NoiseMode == 3 )
+		{
+			const uint3 ChannelIndex = RandomUInt3(Particle) % SolveVelocitiesAndForces_NoiseLUTNumChannel;
+			const float ChannelLerp = (Particle.Age * SolveVelocitiesAndForces_NoiseFrequency) + RandomFloat(Particle);
+			const uint3 ChannelLUTOffsetA = (ChannelIndex * SolveVelocitiesAndForces_NoiseLUTChannelWidth) + min(floor(ChannelLerp) + 0, SolveVelocitiesAndForces_NoiseLUTChannelWidth - 1);
+			const uint3 ChannelLUTOffsetB = (ChannelIndex * SolveVelocitiesAndForces_NoiseLUTChannelWidth) + min(floor(ChannelLerp) + 1, SolveVelocitiesAndForces_NoiseLUTChannelWidth - 1);
+			const uint3 ChannelLUTOffsetC = (ChannelIndex * SolveVelocitiesAndForces_NoiseLUTChannelWidth) + min(floor(ChannelLerp) + 2, SolveVelocitiesAndForces_NoiseLUTChannelWidth - 1);
+			const uint3 ChannelLUTOffsetD = (ChannelIndex * SolveVelocitiesAndForces_NoiseLUTChannelWidth) + min(floor(ChannelLerp) + 3, SolveVelocitiesAndForces_NoiseLUTChannelWidth - 1);
+				
+			const float3 Values0 = float3(GetStaticFloat(SolveVelocitiesAndForces_NoiseLUTOffset, ChannelLUTOffsetA.x), GetStaticFloat(SolveVelocitiesAndForces_NoiseLUTOffset, ChannelLUTOffsetA.y), GetStaticFloat(SolveVelocitiesAndForces_NoiseLUTOffset, ChannelLUTOffsetA.z));
+			const float3 Values1 = float3(GetStaticFloat(SolveVelocitiesAndForces_NoiseLUTOffset, ChannelLUTOffsetB.x), GetStaticFloat(SolveVelocitiesAndForces_NoiseLUTOffset, ChannelLUTOffsetB.y), GetStaticFloat(SolveVelocitiesAndForces_NoiseLUTOffset, ChannelLUTOffsetB.z));
+			const float3 Values2 = float3(GetStaticFloat(SolveVelocitiesAndForces_NoiseLUTOffset, ChannelLUTOffsetC.x), GetStaticFloat(SolveVelocitiesAndForces_NoiseLUTOffset, ChannelLUTOffsetC.y), GetStaticFloat(SolveVelocitiesAndForces_NoiseLUTOffset, ChannelLUTOffsetC.z));
+			const float3 Values3 = float3(GetStaticFloat(SolveVelocitiesAndForces_NoiseLUTOffset, ChannelLUTOffsetD.x), GetStaticFloat(SolveVelocitiesAndForces_NoiseLUTOffset, ChannelLUTOffsetD.y), GetStaticFloat(SolveVelocitiesAndForces_NoiseLUTOffset, ChannelLUTOffsetD.z));
+			{
+				const float u = frac(ChannelLerp);
+				const float u2 = u * u;
+				const float3 a0 = Values3 - Values2 - Values0 + Values1;
+				const float3 a1 = Values0 - Values1 - a0;
+				const float3 a2 = Values2 - Values0;
+				const float3 a3 = Values1;
+				NoiseValue = a0*u*u2+a1*u2+a2*u+a3;
+			}
+		}
+
+		ModuleData.Velocity += NoiseValue.xyz * SolveVelocitiesAndForces_NoiseAmplitude;
+	}
+}
+	
+void SolveVelocitiesAndForces_Simulate(in FStatelessModule_SolveVelocitiesAndForces ModuleData, inout FStatelessParticle Particle)
+{
+	SolveVelocitiesAndForces_IntegratePosition(ModuleData, Particle.Age, Particle.Position);
+	SolveVelocitiesAndForces_IntegratePosition(ModuleData, Particle.PreviousAge, Particle.PreviousPosition);
+		
+	Particle.Velocity = Particle.Position - Particle.PreviousPosition;
+	Particle.PreviousVelocity = Particle.Velocity;
+}
diff --git a/Engine/Plugins/FX/Niagara/Shaders/Private/Stateless/Modules/NiagaraStatelessModule_SpriteFacingAndAlignment.ush b/Engine/Plugins/FX/Niagara/Shaders/Private/Stateless/Modules/NiagaraStatelessModule_SpriteFacingAndAlignment.ush
new file mode 100644
--- /dev/null
+++ b/Engine/Plugins/FX/Niagara/Shaders/Private/Stateless/Modules/NiagaraStatelessModule_SpriteFacingAndAlignment.ush
@@ -0,0 +1,23 @@
+// Copyright Epic Games, Inc. All Rights Reserved.
+
+#include "../NiagaraStatelessCommon.ush"
+
+float3	SpriteFacingAndAlignment_SpriteFacing;
+float3	SpriteFacingAndAlignment_SpriteAlignment;
+
+struct FStatelessModule_SpriteFacingAndAlignment
+{
+};
+
+void SpriteFacingAndAlignment_Initialize(inout FStatelessModule_SpriteFacingAndAlignment ModuleData, inout FStatelessParticle Particle)
+{
+	Particle.SpriteFacing = SpriteFacingAndAlignment_SpriteFacing;
+	Particle.SpriteAlignment = SpriteFacingAndAlignment_SpriteAlignment;
+
+	Particle.PreviousSpriteFacing = SpriteFacingAndAlignment_SpriteFacing;
+	Particle.PreviousSpriteAlignment = SpriteFacingAndAlignment_SpriteAlignment;
+}
+	
+void SpriteFacingAndAlignment_Simulate(in FStatelessModule_SpriteFacingAndAlignment ModuleData, inout FStatelessParticle Particle)
+{
+}
diff --git a/Engine/Plugins/FX/Niagara/Shaders/Private/Stateless/Modules/NiagaraStatelessModule_SpriteRotationRate.ush b/Engine/Plugins/FX/Niagara/Shaders/Private/Stateless/Modules/NiagaraStatelessModule_SpriteRotationRate.ush
new file mode 100644
--- /dev/null
+++ b/Engine/Plugins/FX/Niagara/Shaders/Private/Stateless/Modules/NiagaraStatelessModule_SpriteRotationRate.ush
@@ -0,0 +1,22 @@
+// Copyright Epic Games, Inc. All Rights Reserved.
+
+#include "../NiagaraStatelessCommon.ush"
+
+float		SpriteRotationRate_Scale;
+float		SpriteRotationRate_Bias;
+
+struct FStatelessModule_SpriteRotationRate
+{
+	float RotationRate;
+};
+
+void SpriteRotationRate_Initialize(inout FStatelessModule_SpriteRotationRate ModuleData, inout FStatelessParticle Particle)
+{
+	ModuleData.RotationRate = RandomScaleBiasFloat(Particle, SpriteRotationRate_Scale, SpriteRotationRate_Bias);
+}
+	
+void SpriteRotationRate_Simulate(in FStatelessModule_SpriteRotationRate ModuleData, inout FStatelessParticle Particle)
+{
+	Particle.SpriteRotation += Particle.Age * ModuleData.RotationRate;
+	Particle.PreviousSpriteRotation += Particle.PreviousAge * ModuleData.RotationRate;
+}
diff --git a/Engine/Plugins/FX/Niagara/Shaders/Private/Stateless/Modules/NiagaraStatelessModule_SubUVAnimation.ush b/Engine/Plugins/FX/Niagara/Shaders/Private/Stateless/Modules/NiagaraStatelessModule_SubUVAnimation.ush
new file mode 100644
--- /dev/null
+++ b/Engine/Plugins/FX/Niagara/Shaders/Private/Stateless/Modules/NiagaraStatelessModule_SubUVAnimation.ush
@@ -0,0 +1,40 @@
+// Copyright Epic Games, Inc. All Rights Reserved.
+
+#include "../NiagaraStatelessCommon.ush"
+
+int		SubUVAnimation_Mode;
+float	SubUVAnimation_NumFrames;
+float	SubUVAnimation_InitialFrameScale;
+float	SubUVAnimation_InitialFrameBias;
+float	SubUVAnimation_InitialFrameRateChange;
+float	SubUVAnimation_AnimFrameStart;
+float	SubUVAnimation_AnimFrameRange;
+float	SubUVAnimation_RateScale;
+
+struct FStatelessModule_SubUVAnimation
+{
+	float InitialFrame;
+};
+
+void SubUVAnimation_Initialize(inout FStatelessModule_SubUVAnimation ModuleData, inout FStatelessParticle Particle)
+{
+	const int Seed = int(floor(Particle.Age * SubUVAnimation_InitialFrameRateChange) * 0xffff) ^ Particle.RandomSeed.x ^ Particle.UniqueIndex;
+	ModuleData.InitialFrame = RandBBSfloat(float(Seed) / 0xffff) * SubUVAnimation_InitialFrameScale + SubUVAnimation_InitialFrameBias;
+	//ModuleData.InitialFrame = RandomScaleBiasFloat(Particle, SubUVAnimation_InitialFrameScale, SubUVAnimation_InitialFrameBias);
+}
+	
+void SubUVAnimation_Simulate(in FStatelessModule_SubUVAnimation ModuleData, inout FStatelessParticle Particle)
+{
+	float Frame = ModuleData.InitialFrame;
+	if ( SubUVAnimation_Mode == 0 )
+	{
+		const float Interp = Particle.Age * SubUVAnimation_RateScale;
+		Frame = frac(Frame + SubUVAnimation_AnimFrameStart + (Interp * SubUVAnimation_AnimFrameRange));
+	}
+	else if ( SubUVAnimation_Mode == 1 )
+	{
+		const float Interp = Particle.NormalizedAge * SubUVAnimation_RateScale;
+		Frame = saturate(Frame + SubUVAnimation_AnimFrameStart + (Interp * SubUVAnimation_AnimFrameRange));
+	}
+	Particle.SubImageIndex = Frame * SubUVAnimation_NumFrames;
+}
diff --git a/Engine/Plugins/FX/Niagara/Shaders/Private/Stateless/NiagaraStatelessCommon.ush b/Engine/Plugins/FX/Niagara/Shaders/Private/Stateless/NiagaraStatelessCommon.ush
new file mode 100644
--- /dev/null
+++ b/Engine/Plugins/FX/Niagara/Shaders/Private/Stateless/NiagaraStatelessCommon.ush
@@ -0,0 +1,174 @@
+// Copyright Epic Games, Inc. All Rights Reserved.
+#pragma once
+
+#include "/Engine/Private/Common.ush"
+#include "/Engine/Private/Random.ush"
+
+#include "/Plugin/FX/Niagara/Private/NiagaraQuaternionUtils.ush"
+
+////////////////////////////////////////////////////////////////////////////////////
+
+#define INVALID_COMPONENT		-1
+#define UE_SMALL_NUMBER			(1.e-8f)
+#define UE_KINDA_SMALL_NUMBER	(1.e-4f)
+#define UE_PI 					(3.1415926535897932f)
+#define UE_INV_PI				(0.31830988618f)
+#define UE_HALF_PI				(1.57079632679f)
+#define UE_TWO_PI				(6.28318530717f)
+#define UE_PI_SQUARED			(9.86960440108f)
+#define UE_DEG_TO_RAD			(UE_PI / 180.0f)
+
+bool IsValidComponent(int ComponentOffset)
+{
+	return ComponentOffset != INVALID_COMPONENT;
+}
+
+uint				Common_OutputBufferStride;
+RWBuffer<float>		Common_FloatOutputBuffer;
+RWBuffer<int>		Common_IntOutputBuffer;
+Buffer<float>		Common_StaticFloatBuffer;
+
+////////////////////////////////////////////////////////////////////////////////////
+
+float2 SinCos(float v)
+{
+	float2 sc;
+	sincos(v, sc.x, sc.y);
+	return sc;
+}
+
+float2 SafeNormalize(float2 v, float2 Fallback)
+{
+	const float l2 = length2(v);
+	return l2 < UE_KINDA_SMALL_NUMBER ? Fallback : v * rsqrt(l2);
+}
+
+float2 SafeNormalize(float2 v)
+{
+	return SafeNormalize(v, float2(1.0f, 0.0f));
+}
+
+float3 SafeNormalize(float3 v, float3 Fallback)
+{
+	const float l2 = length2(v);
+	return l2 < UE_KINDA_SMALL_NUMBER ? Fallback : v * rsqrt(l2);
+}
+
+float3 SafeNormalize(float3 v)
+{
+	return SafeNormalize(v, float3(1.0f, 0.0f, 0.0f));
+}
+
+////////////////////////////////////////////////////////////////////////////////////
+
+struct FStatelessParticle
+{
+	bool bAlive;					// If the particle should remain alive or not (only evaluated as specific points)
+	uint UniqueIndex;				// Unique particle index, sequential based on previous spawn info accumulation
+	uint4 RandomSeed;				// Random seed for this particle
+	float Lifetime;					// Overall lifetime for the particle
+	float Age;						// Current age for the particle
+	float NormalizedAge;			// Current normalized age for the particle
+	float PreviousAge;				// Previous age for the particle
+	float PreviousNormalizedAge;	// Previous normalized age for the particle
+
+	float DeltaTime;				// Simulation DT
+	float InvDeltaTime;				// Simulation Inverse DT
+	
+	/////////////////////////////////////////////////////////////////////////////////////////////////////////////////////////
+	// TBD this might produce sub optimal code, needs investigate and is temporary
+	#define PARTICLE_ATTRIBUTE(TYPE, NAME) TYPE NAME##;
+		PARTICLE_ATTRIBUTES_OUTPUT
+		PARTICLE_ATTRIBUTES_TRANSIENT
+	#undef PARTICLE_ATTRIBUTE
+	/////////////////////////////////////////////////////////////////////////////////////////////////////////////////////////
+};
+	
+uint4 RandomUInt4(inout FStatelessParticle Particle)
+{
+	++Particle.RandomSeed.x;
+	Particle.RandomSeed.w = Particle.RandomSeed.x ^ Particle.RandomSeed.y;
+	return Rand4DPCG32(Particle.RandomSeed);
+}
+uint   RandomUInt(inout FStatelessParticle Particle) { return RandomUInt4(Particle).x; }
+uint2  RandomUInt2(inout FStatelessParticle Particle) { return RandomUInt4(Particle).xy; }
+uint3  RandomUInt3(inout FStatelessParticle Particle) { return RandomUInt4(Particle).xyz; }
+
+float4 RandomFloat4(inout FStatelessParticle Particle)
+{
+	uint4 v = RandomUInt4(Particle);
+	return float4((v >> 8) & 0x00ffffff) / 16777216.0; // 0x01000000 == 16777216
+}
+float  RandomFloat(inout FStatelessParticle Particle)  { return RandomFloat4(Particle).x; }
+float2 RandomFloat2(inout FStatelessParticle Particle) { return RandomFloat4(Particle).xy; }
+float3 RandomFloat3(inout FStatelessParticle Particle) { return RandomFloat4(Particle).xyz; }
+	
+float  RandomRangeFloat(inout FStatelessParticle Particle, float  Min, float  Max) { return (RandomFloat(Particle)  * (Max - Min)) + Min; }
+float2 RandomRangeFloat(inout FStatelessParticle Particle, float2 Min, float2 Max) { return (RandomFloat2(Particle) * (Max - Min)) + Min; }
+float3 RandomRangeFloat(inout FStatelessParticle Particle, float3 Min, float3 Max) { return (RandomFloat3(Particle) * (Max - Min)) + Min; }
+float4 RandomRangeFloat(inout FStatelessParticle Particle, float4 Min, float4 Max) { return (RandomFloat4(Particle) * (Max - Min)) + Min; }
+
+float  RandomScaleBiasFloat(inout FStatelessParticle Particle, float  Scale, float  Bias) { return (RandomFloat(Particle)  * Scale) + Bias; }
+float2 RandomScaleBiasFloat(inout FStatelessParticle Particle, float2 Scale, float2 Bias) { return (RandomFloat2(Particle) * Scale) + Bias; }
+float3 RandomScaleBiasFloat(inout FStatelessParticle Particle, float3 Scale, float3 Bias) { return (RandomFloat3(Particle) * Scale) + Bias; }
+float4 RandomScaleBiasFloat(inout FStatelessParticle Particle, float4 Scale, float4 Bias) { return (RandomFloat4(Particle) * Scale) + Bias; }
+
+float2 RandomUnitFloat2(inout FStatelessParticle Particle) { return SafeNormalize(RandomFloat2(Particle) - 0.5f); }
+float3 RandomUnitFloat3(inout FStatelessParticle Particle) { return SafeNormalize(RandomFloat3(Particle) - 0.5f); }
+
+float  GetStaticFloat (int BaseOffset, int Element)	{ return Common_StaticFloatBuffer[BaseOffset + Element]; }
+float2 GetStaticFloat2(int BaseOffset, int Element)	{ return float2(Common_StaticFloatBuffer[BaseOffset + Element*2], Common_StaticFloatBuffer[BaseOffset + 1 + Element*2]); }
+float3 GetStaticFloat3(int BaseOffset, int Element)	{ return float3(Common_StaticFloatBuffer[BaseOffset + Element*3], Common_StaticFloatBuffer[BaseOffset + 1 + Element*3], Common_StaticFloatBuffer[BaseOffset + 2 + Element*3]); }
+float4 GetStaticFloat4(int BaseOffset, int Element)	{ return float4(Common_StaticFloatBuffer[BaseOffset + Element*4], Common_StaticFloatBuffer[BaseOffset + 1 + Element*4], Common_StaticFloatBuffer[BaseOffset + 2 + Element*4], Common_StaticFloatBuffer[BaseOffset + 3 + Element*4]); }
+
+uint GetComponentOffset(uint OutputIndex, int ComponentOffset)
+{
+	return ComponentOffset * Common_OutputBufferStride + OutputIndex;
+}
+
+void OutputComponentData(uint OutputIndex, int ComponentOffset, float Value)
+{
+	if (IsValidComponent(ComponentOffset))
+	{
+		Common_FloatOutputBuffer[GetComponentOffset(OutputIndex, ComponentOffset)] = Value;
+	}
+}
+
+void OutputComponentData(uint OutputIndex, int ComponentOffset, float2 Value)
+{
+	if (IsValidComponent(ComponentOffset))
+	{
+		Common_FloatOutputBuffer[GetComponentOffset(OutputIndex, ComponentOffset + 0)] = Value.x;
+		Common_FloatOutputBuffer[GetComponentOffset(OutputIndex, ComponentOffset + 1)] = Value.y;
+	}
+}
+
+void OutputComponentData(uint OutputIndex, int ComponentOffset, float3 Value)
+{
+	if (IsValidComponent(ComponentOffset))
+	{
+		Common_FloatOutputBuffer[GetComponentOffset(OutputIndex, ComponentOffset + 0)] = Value.x;
+		Common_FloatOutputBuffer[GetComponentOffset(OutputIndex, ComponentOffset + 1)] = Value.y;
+		Common_FloatOutputBuffer[GetComponentOffset(OutputIndex, ComponentOffset + 2)] = Value.z;
+	}
+}
+
+void OutputComponentData(uint OutputIndex, int ComponentOffset, float4 Value)
+{
+	if (IsValidComponent(ComponentOffset))
+	{
+		Common_FloatOutputBuffer[GetComponentOffset(OutputIndex, ComponentOffset + 0)] = Value.x;
+		Common_FloatOutputBuffer[GetComponentOffset(OutputIndex, ComponentOffset + 1)] = Value.y;
+		Common_FloatOutputBuffer[GetComponentOffset(OutputIndex, ComponentOffset + 2)] = Value.z;
+		Common_FloatOutputBuffer[GetComponentOffset(OutputIndex, ComponentOffset + 3)] = Value.w;
+	}
+}
+
+void OutputComponentData(uint OutputIndex, int ComponentOffset, int Value)
+{
+	if (IsValidComponent(ComponentOffset))
+	{
+		Common_IntOutputBuffer[GetComponentOffset(OutputIndex, ComponentOffset)] = Value;
+	}
+}
+
diff --git a/Engine/Plugins/FX/Niagara/Shaders/Private/Stateless/NiagaraStatelessSimulationDefault.usf b/Engine/Plugins/FX/Niagara/Shaders/Private/Stateless/NiagaraStatelessSimulationDefault.usf
new file mode 100644
--- /dev/null
+++ b/Engine/Plugins/FX/Niagara/Shaders/Private/Stateless/NiagaraStatelessSimulationDefault.usf
@@ -0,0 +1,42 @@
+// Copyright Epic Games, Inc. All Rights Reserved.
+
+#define PARTICLE_ATTRIBUTES_OUTPUT \
+	PARTICLE_ATTRIBUTE(int,		UniqueID) \
+	PARTICLE_ATTRIBUTE(float3,	Position) \
+	PARTICLE_ATTRIBUTE(float4,	Color) \
+	PARTICLE_ATTRIBUTE(float4,	MeshOrientation) \
+	PARTICLE_ATTRIBUTE(float,	RibbonWidth) \
+	PARTICLE_ATTRIBUTE(float3,	Scale) \
+	PARTICLE_ATTRIBUTE(float2,	SpriteSize) \
+	PARTICLE_ATTRIBUTE(float3,	SpriteFacing) \
+	PARTICLE_ATTRIBUTE(float3,	SpriteAlignment) \
+	PARTICLE_ATTRIBUTE(float,	SpriteRotation) \
+	PARTICLE_ATTRIBUTE(float,	SubImageIndex) \
+	PARTICLE_ATTRIBUTE(float3,	Velocity) \
+	PARTICLE_ATTRIBUTE(float3,	PreviousPosition) \
+	PARTICLE_ATTRIBUTE(float4,	PreviousMeshOrientation) \
+	PARTICLE_ATTRIBUTE(float,	PreviousRibbonWidth) \
+	PARTICLE_ATTRIBUTE(float3,	PreviousScale) \
+	PARTICLE_ATTRIBUTE(float2,	PreviousSpriteSize) \
+	PARTICLE_ATTRIBUTE(float3,	PreviousSpriteFacing) \
+	PARTICLE_ATTRIBUTE(float3,	PreviousSpriteAlignment) \
+	PARTICLE_ATTRIBUTE(float,	PreviousSpriteRotation) \
+	PARTICLE_ATTRIBUTE(float3,	PreviousVelocity) \
+
+#define PARTICLE_ATTRIBUTES_TRANSIENT \
+
+#define PARTICLE_MODULES \
+	PARTICLE_MODULE(InitializeParticle) \
+	PARTICLE_MODULE(InitialMeshOrientation) \
+	PARTICLE_MODULE(ShapeLocation) \
+	PARTICLE_MODULE(ScaleColor) \
+	PARTICLE_MODULE(ScaleSpriteSize) \
+	PARTICLE_MODULE(SpriteFacingAndAlignment) \
+	PARTICLE_MODULE(SpriteRotationRate) \
+	PARTICLE_MODULE(ScaleMeshSize) \
+	PARTICLE_MODULE(SolveVelocitiesAndForces) \
+	PARTICLE_MODULE(ScaleMeshSizeBySpeed) \
+	PARTICLE_MODULE(ScaleSpriteSizeBySpeed) \
+	PARTICLE_MODULE(SubUVAnimation) \
+
+#include "NiagaraStatelessSimulationTemplate.ush"
diff --git a/Engine/Plugins/FX/Niagara/Shaders/Private/Stateless/NiagaraStatelessSimulationExample1.usf b/Engine/Plugins/FX/Niagara/Shaders/Private/Stateless/NiagaraStatelessSimulationExample1.usf
new file mode 100644
--- /dev/null
+++ b/Engine/Plugins/FX/Niagara/Shaders/Private/Stateless/NiagaraStatelessSimulationExample1.usf
@@ -0,0 +1,36 @@
+// Copyright Epic Games, Inc. All Rights Reserved.
+
+#define PARTICLE_ATTRIBUTES_OUTPUT \
+	PARTICLE_ATTRIBUTE(int,		UniqueID) \
+	PARTICLE_ATTRIBUTE(float3,	Position) \
+	PARTICLE_ATTRIBUTE(float4,	Color) \
+	PARTICLE_ATTRIBUTE(float,	RibbonWidth) \
+	PARTICLE_ATTRIBUTE(float3,	Scale) \
+	PARTICLE_ATTRIBUTE(float2,	SpriteSize) \
+	PARTICLE_ATTRIBUTE(float,	SpriteRotation) \
+	PARTICLE_ATTRIBUTE(float3,	Velocity) \
+	PARTICLE_ATTRIBUTE(float3,	PreviousPosition) \
+	PARTICLE_ATTRIBUTE(float,	PreviousRibbonWidth) \
+	PARTICLE_ATTRIBUTE(float3,	PreviousScale) \
+	PARTICLE_ATTRIBUTE(float2,	PreviousSpriteSize) \
+	PARTICLE_ATTRIBUTE(float,	PreviousSpriteRotation) \
+	PARTICLE_ATTRIBUTE(float3,	PreviousVelocity) \
+
+#define PARTICLE_ATTRIBUTES_TRANSIENT \
+	PARTICLE_ATTRIBUTE(float4,	MeshOrientation) \
+	PARTICLE_ATTRIBUTE(float3,	SpriteFacing) \
+	PARTICLE_ATTRIBUTE(float3,	SpriteAlignment) \
+	PARTICLE_ATTRIBUTE(float,	SubImageIndex) \
+	PARTICLE_ATTRIBUTE(float4,	PreviousMeshOrientation) \
+	PARTICLE_ATTRIBUTE(float3,	PreviousSpriteFacing) \
+	PARTICLE_ATTRIBUTE(float3,	PreviousSpriteAlignment) \
+
+#define PARTICLE_MODULES \
+	PARTICLE_MODULE(InitializeParticle) \
+	PARTICLE_MODULE(ShapeLocation) \
+	PARTICLE_MODULE(ScaleColor) \
+	PARTICLE_MODULE(ScaleSpriteSize) \
+	PARTICLE_MODULE(RotateAroundPoint) \
+	PARTICLE_MODULE(CalculateAccurateVelocity) \
+
+#include "NiagaraStatelessSimulationTemplate.ush"
diff --git a/Engine/Plugins/FX/Niagara/Shaders/Private/Stateless/NiagaraStatelessSimulationTemplate.ush b/Engine/Plugins/FX/Niagara/Shaders/Private/Stateless/NiagaraStatelessSimulationTemplate.ush
new file mode 100644
--- /dev/null
+++ b/Engine/Plugins/FX/Niagara/Shaders/Private/Stateless/NiagaraStatelessSimulationTemplate.ush
@@ -0,0 +1,144 @@
+// Copyright Epic Games, Inc. All Rights Reserved.
+
+#include "NiagaraStatelessCommon.ush"
+//#include "WaveOpUtil.ush"
+
+/////////////////////////////////////////////////////////////////////////////////////////////////////////////////////////
+
+#define PARTICLE_ATTRIBUTE(TYPE, NAME) int Permutation_##NAME##Component;
+	PARTICLE_ATTRIBUTES_OUTPUT
+#undef PARTICLE_ATTRIBUTE
+
+/////////////////////////////////////////////////////////////////////////////////////////////////////////////////////////
+// TBD this might produce sub optimal code, needs investigate and is temporary
+#include "Modules/NiagaraStatelessModule_CalculateAccurateVelocity.ush"
+#include "Modules/NiagaraStatelessModule_InitializeParticle.ush"
+#include "Modules/NiagaraStatelessModule_InitialMeshOrientation.ush"
+#include "Modules/NiagaraStatelessModule_RotateAroundPoint.ush"
+#include "Modules/NiagaraStatelessModule_ShapeLocation.ush"
+#include "Modules/NiagaraStatelessModule_ScaleColor.ush"
+#include "Modules/NiagaraStatelessModule_ScaleSpriteSize.ush"
+#include "Modules/NiagaraStatelessModule_ScaleSpriteSizeBySpeed.ush"
+#include "Modules/NiagaraStatelessModule_ScaleMeshSize.ush"
+#include "Modules/NiagaraStatelessModule_ScaleMeshSizebySpeed.ush"
+#include "Modules/NiagaraStatelessModule_SolveVelocitiesAndForces.ush"
+#include "Modules/NiagaraStatelessModule_SpriteFacingAndAlignment.ush"
+#include "Modules/NiagaraStatelessModule_SpriteRotationRate.ush"
+#include "Modules/NiagaraStatelessModule_SubUVAnimation.ush"
+
+struct FStatelessParticleSimuation
+{
+	#define PARTICLE_MODULE(NAME) FStatelessModule_##NAME	NAME##Module;
+		PARTICLE_MODULES
+	#undef PARTICLE_MODULE
+};
+
+/////////////////////////////////////////////////////////////////////////////////////////////////////////////////////////
+
+#define NIAGARA_MAX_GPU_SPAWN_INFOS 4
+
+/////////////////////////////////////////////////////////////////////////////////////////////////////////////////////////
+
+uint	Common_RandomSeed;
+float	Common_SimulationTime;
+float	Common_SimulationDeltaTime;
+float	Common_SimulationInvDeltaTime;
+float2	Common_LifetimeScaleBias;
+
+uint				Common_GPUCountBufferOffset;
+RWBuffer<int>		Common_GPUCountBuffer;
+
+DECLARE_SCALAR_ARRAY(uint,	SpawnInfo_NumActive, NIAGARA_MAX_GPU_SPAWN_INFOS);
+DECLARE_SCALAR_ARRAY(uint,	SpawnInfo_ParticleOffset, NIAGARA_MAX_GPU_SPAWN_INFOS);
+DECLARE_SCALAR_ARRAY(uint,	SpawnInfo_UniqueOffset, NIAGARA_MAX_GPU_SPAWN_INFOS);
+DECLARE_SCALAR_ARRAY(float,	SpawnInfo_Time, NIAGARA_MAX_GPU_SPAWN_INFOS);
+DECLARE_SCALAR_ARRAY(float,	SpawnInfo_Rate, NIAGARA_MAX_GPU_SPAWN_INFOS);
+
+/////////////////////////////////////////////////////////////////////////////////////////////////////////////////////////
+
+[numthreads(THREADGROUP_SIZE, 1, 1)]
+void StatelessMain(uint DispatchThreadId : SV_DispatchThreadID)
+{
+	// Initialize our particle data
+	FStatelessParticle Particle;
+	Particle.Age = -1.0f;
+
+	// Find which spawn info we belong to and fill out details
+	{
+		uint SpawnInfoIndex = DispatchThreadId.x;
+		for ( int i=0; i < NIAGARA_MAX_GPU_SPAWN_INFOS; ++i )
+		{
+			const uint SpawnInfoNumActive		= GET_SCALAR_ARRAY_ELEMENT(SpawnInfo_NumActive, i);
+			const uint SpawnInfoParticleOffset	= GET_SCALAR_ARRAY_ELEMENT(SpawnInfo_ParticleOffset, i);
+			const uint SpawnInfoUniqueOffset	= GET_SCALAR_ARRAY_ELEMENT(SpawnInfo_UniqueOffset, i);
+			const float SpawnInfoTime			= GET_SCALAR_ARRAY_ELEMENT(SpawnInfo_Time, i);
+			const float SpawnInfoRate			= GET_SCALAR_ARRAY_ELEMENT(SpawnInfo_Rate, i);
+			if ( SpawnInfoIndex < SpawnInfoNumActive )
+			{
+				const uint SpawnParticleIndex = SpawnInfoIndex + SpawnInfoParticleOffset;
+				Particle.UniqueIndex	= SpawnInfoIndex + SpawnInfoUniqueOffset + SpawnInfoParticleOffset;
+				Particle.Age			= Common_SimulationTime - (SpawnInfoTime + float(SpawnParticleIndex) * SpawnInfoRate);
+				break;
+			}
+
+			SpawnInfoIndex -= SpawnInfoNumActive;
+		}
+		if ( Particle.Age < 0.0f )
+		{
+			return;
+		}
+	}
+	
+	// Initialize variables / determine lifetime / etc
+	Particle.RandomSeed.x = 7123u;
+	Particle.RandomSeed.y = Common_RandomSeed;
+	Particle.RandomSeed.z = Particle.UniqueIndex * 3581u;
+	Particle.RandomSeed.w = 0;
+
+	Particle.Lifetime = RandomScaleBiasFloat(Particle, Common_LifetimeScaleBias.x, Common_LifetimeScaleBias.y);
+	if ( Particle.Lifetime <= 0.0f || Particle.Age >= Particle.Lifetime )
+	{
+		return;
+	}
+	Particle.NormalizedAge = Particle.Age / Particle.Lifetime;
+
+	Particle.PreviousAge = max(Particle.Age - Common_SimulationDeltaTime, 0.0f);
+	Particle.PreviousNormalizedAge = Particle.PreviousAge / Particle.Lifetime;
+
+	Particle.DeltaTime		= Common_SimulationDeltaTime;
+	Particle.InvDeltaTime	= Common_SimulationInvDeltaTime;
+	
+	Particle.bAlive = true;
+	
+	// Run Particle Initialization
+	FStatelessParticleSimuation ParticleSimulation;
+
+	#define PARTICLE_MODULE(NAME) ##NAME##_Initialize(ParticleSimulation.##NAME##Module, Particle);
+		PARTICLE_MODULES
+	#undef PARTICLE_MODULE
+
+	if ( !Particle.bAlive )
+	{
+		return;
+	}
+
+	// Run Particle Simulate
+	#define PARTICLE_MODULE(NAME) ##NAME##_Simulate(ParticleSimulation.##NAME##Module, Particle);
+		PARTICLE_MODULES
+	#undef PARTICLE_MODULE
+
+	if ( !Particle.bAlive )
+	{
+		return;
+	}
+
+	// Output Particles that are still alive
+	uint OutputIndex;
+	
+	//WaveInterlockedAddScalar_(Common_GPUCountBuffer[Common_GPUCountBufferOffset], 1, Particle.OutputIndex);
+	InterlockedAdd(Common_GPUCountBuffer[Common_GPUCountBufferOffset], 1, OutputIndex);
+
+	#define PARTICLE_ATTRIBUTE(TYPE, NAME) if ( IsValidComponent(Permutation_##NAME##Component) ) { OutputComponentData(OutputIndex, Permutation_##NAME##Component, Particle.##NAME); }
+		PARTICLE_ATTRIBUTES_OUTPUT
+	#undef PARTICLE_ATTRIBUTE
+}
diff --git a/Engine/Plugins/FX/Niagara/Source/Niagara/Classes/NiagaraComputeExecutionContext.h b/Engine/Plugins/FX/Niagara/Source/Niagara/Classes/NiagaraComputeExecutionContext.h
--- a/Engine/Plugins/FX/Niagara/Source/Niagara/Classes/NiagaraComputeExecutionContext.h
+++ b/Engine/Plugins/FX/Niagara/Source/Niagara/Classes/NiagaraComputeExecutionContext.h
@@ -55,10 +55,33 @@ struct FNiagaraGpuSpawnInfo
 	}
 };
 
-struct FNiagaraComputeExecutionContext
+struct INiagaraComputeDataBufferInterface
+{
+	virtual bool HasTranslucentDataToRender() const = 0;
+	virtual FNiagaraDataBuffer* GetDataToRender(bool bIsLowLatencyTranslucent) const = 0;
+};
+
+struct FNiagaraComputeExecutionContext : public INiagaraComputeDataBufferInterface
 {
 	FNiagaraComputeExecutionContext();
-	~FNiagaraComputeExecutionContext();
+	virtual ~FNiagaraComputeExecutionContext();
+
+	// Begin: INiagaraComputeDataBufferInterface
+	virtual bool HasTranslucentDataToRender() const override { return TranslucentDataToRender != nullptr; }
+	virtual FNiagaraDataBuffer* GetDataToRender(bool bIsLowLatencyTranslucent) const override
+	{
+		if (bIsLowLatencyTranslucent)
+		{
+			// Translucent rendering uses current frame (newest) data, which is TranslucentDataToRender if present
+			return TranslucentDataToRender ? TranslucentDataToRender : DataToRender;
+		}
+		else
+		{
+			// Non-translucent rendering uses previous frame (older) data, which is MultiViewPreviousDataToRender if present
+			return MultiViewPreviousDataToRender ? MultiViewPreviousDataToRender : DataToRender;
+		}
+	}
+	// End: INiagaraComputeDataBufferInterface
 
 	void Reset(FNiagaraGpuComputeDispatchInterface* ComputeDispatchInterface);
 
@@ -73,20 +96,6 @@ struct FNiagaraComputeExecutionContext
 	void SetDataToRender(FNiagaraDataBuffer* InDataToRender);
 	void SetTranslucentDataToRender(FNiagaraDataBuffer* InTranslucentDataToRender);
 	void SetMultiViewPreviousDataToRender(FNiagaraDataBuffer* InMultiViewPreviousDataToRender);
-	bool HasTranslucentDataToRender() const { return TranslucentDataToRender != nullptr; }	
-	FNiagaraDataBuffer* GetDataToRender(bool bIsLowLatencyTranslucent) const
-	{
-		if (bIsLowLatencyTranslucent)
-		{
-			// Translucent rendering uses current frame (newest) data, which is TranslucentDataToRender if present
-			return TranslucentDataToRender ? TranslucentDataToRender : DataToRender;
-		}
-		else
-		{
-			// Non-translucent rendering uses previous frame (older) data, which is MultiViewPreviousDataToRender if present
-			return MultiViewPreviousDataToRender ? MultiViewPreviousDataToRender : DataToRender;
-		}
-	}
 
 	int32 GetConstantBufferSize() const;
 	uint8* WriteConstantBufferInstanceData(uint8* InTargetBuffer, FNiagaraComputeInstanceData& InstanceData) const;
diff --git a/Engine/Plugins/FX/Niagara/Source/Niagara/Classes/NiagaraEmitterHandle.h b/Engine/Plugins/FX/Niagara/Source/Niagara/Classes/NiagaraEmitterHandle.h
--- a/Engine/Plugins/FX/Niagara/Source/Niagara/Classes/NiagaraEmitterHandle.h
+++ b/Engine/Plugins/FX/Niagara/Source/Niagara/Classes/NiagaraEmitterHandle.h
@@ -14,6 +14,14 @@
 #include "NiagaraEmitterHandle.generated.h"
 
 class UNiagaraSystem;
+class UNiagaraStatelessEmitter;
+
+UENUM()
+enum class ENiagaraEmitterMode : uint8
+{
+	Standard,
+	Stateless
+};
 
 /** 
  * Stores emitter information within the context of a System.
@@ -32,6 +40,10 @@ public:
 	
 	/** Creates a new emitter handle from an emitter. */
 	NIAGARA_API FNiagaraEmitterHandle(const FVersionedNiagaraEmitter& InEmitter);
+
+	//-TODO:Stateless
+	NIAGARA_API FNiagaraEmitterHandle(UNiagaraStatelessEmitter& InEmitter);
+	//-TODO:Stateless
 #endif
 
 	/** Whether or not this is a valid emitter handle. */
@@ -67,6 +79,17 @@ public:
 
 	NIAGARA_API FVersionedNiagaraEmitterData* GetEmitterData() const;
 
+	//-TODO:Stateless: Should we return a bass class here / have a factory method to generate the runtime instance?
+	UNiagaraStatelessEmitter* GetStatelessEmitter() const { return StatelessEmitter; }
+	void SetStatelessEmitter(UNiagaraStatelessEmitter* InEmitter) { StatelessEmitter = InEmitter; }
+	//-TODO:Stateless: Should we return a bass class here / have a factory method to generate the runtime instance?
+
+	ENiagaraEmitterMode GetEmitterMode() const { return EmitterMode; }
+#if WITH_EDITORONLY_DATA
+	NIAGARA_API void SetEmitterMode(UNiagaraSystem& InOwningSystem, ENiagaraEmitterMode InEmitterMode);
+	FSimpleMulticastDelegate& OnEmitterModeChanged() { return OnEmitterModeChangedDelegate; }
+#endif
+
 	/** Gets a unique name for this emitter instance for use in scripts and parameter stores etc.*/
 	NIAGARA_API FString GetUniqueInstanceName()const;
 
@@ -130,4 +153,17 @@ private:
 	/** The copied instance of the emitter this handle references. */
 	UPROPERTY()
 	FVersionedNiagaraEmitter VersionedInstance;
+
+	//-TODO:Stateless: Should we return a bass class here / have a factory method to generate the runtime instance?
+	UPROPERTY()
+	TObjectPtr<UNiagaraStatelessEmitter> StatelessEmitter = nullptr;
+	//-TODO:Stateless: Should we return a bass class here / have a factory method to generate the runtime instance?
+
+	UPROPERTY()
+	ENiagaraEmitterMode EmitterMode = ENiagaraEmitterMode::Standard;
+
+#if WITH_EDITORONLY_DATA
+	FSimpleMulticastDelegate OnEmitterModeChangedDelegate;
+#endif
+	
 };
diff --git a/Engine/Plugins/FX/Niagara/Source/Niagara/Classes/NiagaraEmitterInstance.h b/Engine/Plugins/FX/Niagara/Source/Niagara/Classes/NiagaraEmitterInstance.h
--- a/Engine/Plugins/FX/Niagara/Source/Niagara/Classes/NiagaraEmitterInstance.h
+++ b/Engine/Plugins/FX/Niagara/Source/Niagara/Classes/NiagaraEmitterInstance.h
@@ -9,12 +9,14 @@
 #include "NiagaraEmitter.h"
 #include "NiagaraScriptExecutionContext.h"
 
+struct INiagaraComputeDataBufferInterface;
 struct FNiagaraComputeExecutionContext;
 struct FNiagaraScriptExecutionContext;
 class UNiagaraDataInterface;
 class FNiagaraSystemInstance;
 
 class FNiagaraEmitterInstanceImpl;
+class FNiagaraStatelessEmitterInstance;
 
 /**
 * Base class for different emitter instances
@@ -54,6 +56,7 @@ public:
 	const FNiagaraDataSet& GetParticleData() const { check(ParticleDataSet); return *ParticleDataSet; }
 
 	FNiagaraComputeExecutionContext* GetGPUContext() const { return GPUExecContext; }
+	INiagaraComputeDataBufferInterface* GetComputeDataBufferInterface() const { return GPUDataBufferInterfaces; }
 
 	ENiagaraExecutionState GetExecutionState() const { return ExecutionState; }
 	bool IsActive() const { return ExecutionState == ENiagaraExecutionState::Active; }
@@ -94,11 +97,13 @@ public:
 	//////////////////////////////////////////////////////////////////////////////////////////////////////////////////
 
 	////////////////////////////////////////////////////////////////////////////////////////////////////
-	//-TODO:Stateless:
+	//-TODO:Stateless: Consider removing these virtual functions
 	// Can we remove these virtual functions
 	virtual FNiagaraEmitterInstanceImpl* AsStateful() { return nullptr; }
+	virtual FNiagaraStatelessEmitterInstance* AsStateless() { return nullptr; }
 
 	const FNiagaraEmitterInstanceImpl* AsStateful() const { return const_cast<FNiagaraEmitterInstance*>(this)->AsStateful(); }
+	const FNiagaraStatelessEmitterInstance* AsStateless() const { return const_cast<FNiagaraEmitterInstance*>(this)->AsStateless(); }
 
 	UNiagaraEmitter* GetEmitter() const { return VersionedEmitter.Emitter; }
 	const FVersionedNiagaraEmitter& GetVersionedEmitter() const { return VersionedEmitter; }
@@ -204,6 +209,7 @@ protected:
 	FNiagaraSystemInstance*				ParentSystemInstance = nullptr;
 	FNiagaraDataSet*					ParticleDataSet = nullptr;
 	FNiagaraComputeExecutionContext*	GPUExecContext = nullptr;
+	INiagaraComputeDataBufferInterface* GPUDataBufferInterfaces = nullptr;
 	FVersionedNiagaraEmitter			VersionedEmitter = {};
 
 	FNiagaraParameterStore				RendererBindings;
diff --git a/Engine/Plugins/FX/Niagara/Source/Niagara/Classes/NiagaraSystem.h b/Engine/Plugins/FX/Niagara/Source/Niagara/Classes/NiagaraSystem.h
--- a/Engine/Plugins/FX/Niagara/Source/Niagara/Classes/NiagaraSystem.h
+++ b/Engine/Plugins/FX/Niagara/Source/Niagara/Classes/NiagaraSystem.h
@@ -11,6 +11,7 @@
 #include "NiagaraMessageStore.h"
 #include "NiagaraParameterCollection.h"
 #include "NiagaraParameterDefinitionsSubscriber.h"
+#include "NiagaraSystemEmitterState.h"
 #include "NiagaraUserRedirectionParameterStore.h"
 #include "Particles/ParticleSystem.h"
 #include "UObject/Object.h"
@@ -785,6 +786,12 @@ public:
 	NIAGARA_API void PrepareRapidIterationParametersForCompilation();
 #endif
 
+	/** Can we run the code only system state path, i.e. we don't need to invoke the VVM / store per instance data set? */
+	bool SystemStateFastPathEnabled() const { return bSystemStateFastPathEnabled && bAllowSystemStateFastPath; }
+
+	/** Access the code system state data. */
+	const FNiagaraSystemStateData& GetSystemStateData() const { return SystemStateData; }
+
 private:
 #if WITH_EDITORONLY_DATA
 
@@ -815,6 +822,8 @@ private:
 	NIAGARA_API void ResolveParameterStoreBindings();
 #endif
 
+
+	void ResolveRequiresScripts();
 	NIAGARA_API void ResolveScalabilitySettings();
 	NIAGARA_API void UpdatePostCompileDIInfo();
 	NIAGARA_API void UpdateDITickFlags();
@@ -1018,6 +1027,13 @@ protected:
 	uint32 bHasAnyGPUEmitters : 1;
 	uint32 bNeedsSortedSignificanceCull : 1;
 
+	/** When enabled if all emitters don't require script execution and the system script is empty / constant we can invoke a faster CPU path. */
+	UPROPERTY(EditAnywhere, Category = "Performance", AdvancedDisplay)
+	uint32 bAllowSystemStateFastPath : 1 = true;
+
+	UPROPERTY()
+	uint32 bSystemStateFastPathEnabled : 1 = true;
+
 #if WITH_EDITORONLY_DATA
 	/** Messages associated with the System asset. */
 	UPROPERTY()
@@ -1035,6 +1051,9 @@ protected:
 	struct FStaticBuffersDeletor { void operator()(FNiagaraSystemStaticBuffers* Ptr) const; };
 	TUniquePtr<FNiagaraSystemStaticBuffers, FStaticBuffersDeletor> StaticBuffers;
 
+	UPROPERTY()
+	FNiagaraSystemStateData SystemStateData;
+
 	FRenderCommandFence WaitRenderCommandsFence;
 };
 
diff --git a/Engine/Plugins/FX/Niagara/Source/Niagara/Internal/NiagaraEmitterInstanceImpl.h b/Engine/Plugins/FX/Niagara/Source/Niagara/Internal/NiagaraEmitterInstanceImpl.h
--- a/Engine/Plugins/FX/Niagara/Source/Niagara/Internal/NiagaraEmitterInstanceImpl.h
+++ b/Engine/Plugins/FX/Niagara/Source/Niagara/Internal/NiagaraEmitterInstanceImpl.h
@@ -81,7 +81,6 @@ public:
 	NIAGARA_API void CalculateFixedBounds(const FTransform& ToWorldSpace);
 #endif
 
-	//-TODO:Stateless: Cache miss hell never call GetEmitterData()!
 	FORCEINLINE const FNiagaraEmitterScalabilitySettings& GetScalabilitySettings() const { check(VersionedEmitter.GetEmitterData() != nullptr); return VersionedEmitter.GetEmitterData()->GetScalabilitySettings(); }
 
 	//ENiagaraExecutionState GetExecutionState() { return ExecutionState; }
diff --git a/Engine/Plugins/FX/Niagara/Source/Niagara/Internal/NiagaraSystemEmitterState.h b/Engine/Plugins/FX/Niagara/Source/Niagara/Internal/NiagaraSystemEmitterState.h
new file mode 100644
--- /dev/null
+++ b/Engine/Plugins/FX/Niagara/Source/Niagara/Internal/NiagaraSystemEmitterState.h
@@ -0,0 +1,119 @@
+// Copyright Epic Games, Inc. All Rights Reserved.
+
+#pragma once
+
+#include "NiagaraCommon.h"
+
+#include "NiagaraSystemEmitterState.generated.h"
+
+UENUM()
+enum class ENiagaraSystemInactiveResponse
+{
+	/** Let Emitters Finish Then Kill Emitter */
+	Complete,
+	/** Emitter & Particles Die Immediatly */
+	Kill,
+};
+
+UENUM()
+enum class ENiagaraEmitterInactiveResponse
+{
+	/** Let Particles Finish Then Kill Emitter */
+	Complete,
+	/** Emitter & Particles Die Immediatly */
+	Kill,
+	/** Emitter deactivates but doesn't die until the system does */
+	//Continue,
+};
+
+UENUM()
+enum class ENiagaraLoopBehavior
+{
+	Infinite,
+	Multiple,
+	Once,
+};
+
+USTRUCT()
+struct FNiagaraSystemStateData
+{
+	GENERATED_BODY()
+
+	UPROPERTY(EditAnywhere, Category = "System State")
+	bool bIgnoreSystemState = true;
+
+	UPROPERTY(EditAnywhere, Category = "System State")
+	ENiagaraSystemInactiveResponse InactiveResponse = ENiagaraSystemInactiveResponse::Complete;
+
+	UPROPERTY(EditAnywhere, Category = "System State")
+	ENiagaraLoopBehavior LoopBehavior = ENiagaraLoopBehavior::Once;
+	
+	UPROPERTY(EditAnywhere, Category = "Emitter State", meta = (ClampMin = "0.0"))
+	float LoopDurationMin = 0.0f;
+
+	UPROPERTY(EditAnywhere, Category = "Emitter State", meta = (ClampMin = "0.0"))
+	float LoopDurationMax = 0.0f;
+
+	UPROPERTY(EditAnywhere, Category = "System State", meta = (ClampMin = "1", EditCondition = "LoopBehavior == ENiagaraLoopBehavior::Multiple", EditConditionHides))
+	int LoopCount = 1;
+
+	UPROPERTY(EditAnywhere, Category = "System State", meta = (EditCondition = "LoopBehavior != ENiagaraLoopBehavior::Once", EditConditionHides))
+	bool bRecalculateDurationEachLoop = false;
+
+	UPROPERTY(EditAnywhere, Category = "Emitter State", meta = (ClampMin = "0.0"))
+	float LoopDelayMin = 0.0f;
+
+	UPROPERTY(EditAnywhere, Category = "Emitter State", meta = (ClampMin = "0.0"))
+	float LoopDelayMax = 0.0f;
+
+	UPROPERTY(EditAnywhere, Category = "System State", meta = (EditCondition = "LoopBehavior != ENiagaraLoopBehavior::Once", EditConditionHides))
+	bool bDelayFirstLoopOnly = false;
+
+	UPROPERTY(EditAnywhere, Category = "System State", meta = (EditCondition = "LoopBehavior != ENiagaraLoopBehavior::Once && !bDelayFirstLoopOnly", EditConditionHides))
+	bool bRecalculateDelayEachLoop = false;
+};
+
+USTRUCT()
+struct FNiagaraEmitterStateData
+{
+	GENERATED_BODY()
+		
+	//UPROPERTY(EditAnywhere, Category="Emitter State")
+	//ENiagaraStatelessEmitterState_SelfSystem LifeCycleMode = ENiagaraStatelessEmitterState_SelfSystem::Self;
+
+	UPROPERTY(EditAnywhere, Category="Emitter State")
+	ENiagaraEmitterInactiveResponse InactiveResponse = ENiagaraEmitterInactiveResponse::Complete;
+
+	UPROPERTY(EditAnywhere, Category="Emitter State")
+	ENiagaraLoopBehavior LoopBehavior = ENiagaraLoopBehavior::Infinite;
+
+	UPROPERTY(EditAnywhere, Category = "Emitter State", meta = (ClampMin = "1", EditCondition = "LoopBehavior == ENiagaraLoopBehavior::Multiple", EditConditionHides))
+	int32 LoopCount = 1;
+
+	UPROPERTY(EditAnywhere, Category = "Emitter State", meta = (ClampMin = "0.0"))
+	float LoopDurationMin = 1.0f;
+
+	UPROPERTY(EditAnywhere, Category = "Emitter State", meta = (ClampMin = "0.0"))
+	float LoopDurationMax = 1.0f;
+
+	UPROPERTY(EditAnywhere, Category = "Emitter State", meta = (EditCondition = "LoopBehavior != ENiagaraLoopBehavior::Once", EditConditionHides))
+	bool bRecalculateDurationEachLoop = false;
+
+	UPROPERTY(EditAnywhere, Category = "Emitter State", meta = (ClampMin = "0.0"))
+	float LoopDelayMin = 0.0f;
+
+	UPROPERTY(EditAnywhere, Category = "Emitter State", meta = (ClampMin = "0.0"))
+	float LoopDelayMax = 0.0f;
+
+	UPROPERTY(EditAnywhere, Category = "Emitter State", meta = (EditCondition = "LoopBehavior != ENiagaraLoopBehavior::Once", EditConditionHides))
+	bool bDelayFirstLoopOnly = false;
+
+	UPROPERTY(EditAnywhere, Category = "Emitter State", meta = (EditCondition = "LoopBehavior != ENiagaraLoopBehavior::Once && !bDelayFirstLoopOnly", EditConditionHides))
+	bool bRecalculateDelayEachLoop = false;
+
+	//UPROPERTY(EditAnywhere, Category="Emitter State")
+	//ENiagaraStatelessEmitterState_SelfSystem Scalability = ENiagaraStatelessEmitterState_SelfSystem::Self;
+	//Enable Distance Culling
+	//Enable Visibility Culling
+	//Reset Age On Awaken
+};
diff --git a/Engine/Plugins/FX/Niagara/Source/Niagara/Internal/Stateless/Modules/NiagaraStatelessModuleCommon.h b/Engine/Plugins/FX/Niagara/Source/Niagara/Internal/Stateless/Modules/NiagaraStatelessModuleCommon.h
new file mode 100644
--- /dev/null
+++ b/Engine/Plugins/FX/Niagara/Source/Niagara/Internal/Stateless/Modules/NiagaraStatelessModuleCommon.h
@@ -0,0 +1,47 @@
+// Copyright Epic Games, Inc. All Rights Reserved.
+
+#pragma once
+
+#include "Stateless/NiagaraStatelessModule.h"
+
+namespace NiagaraStateless
+{
+	//-TODO: Examine structure padding, FQuat4f / UObject*
+	struct FPhysicsBuildData
+	{
+		static FName GetName() { return FName("FPhysicsBuildData"); }
+
+		float		MassMin = 1.0f;
+		float		MassMax = 1.0f;
+		float		DragMin = 0.0f;
+		float		DragMax = 0.0f;
+		FVector3f	VelocityMin = FVector3f::ZeroVector;
+		FVector3f	VelocityMax = FVector3f::ZeroVector;
+		FVector3f	WindMin = FVector3f::ZeroVector;
+		FVector3f	WindMax = FVector3f::ZeroVector;
+		FVector3f	AccelerationMin = FVector3f::ZeroVector;
+		FVector3f	AccelerationMax = FVector3f::ZeroVector;
+
+		bool		bConeVelocity = false;
+		FQuat4f		ConeQuat = FQuat4f::Identity;
+		float		ConeVelocityMin = 0.0f;
+		float		ConeVelocityMax = 0.0f;
+		float		ConeOuterAngle = 0.0f;
+		float		ConeInnerAngle = 0.0f;
+		float		ConeVelocityFalloff = 0.0f;
+
+		bool		bPointVelocity = false;
+		float		PointVelocityMin = 0.0f;
+		float		PointVelocityMax = 0.0f;
+		FVector3f	PointOrigin = FVector3f::ZeroVector;
+
+		bool		bNoiseEnabled = false;
+		float		NoiseAmplitude = 0.0f;
+		float		NoiseFrequency = 0.0f;
+		UObject*	NoiseTexture = nullptr;
+		int32		NoiseMode = 0;
+		int32		NoiseLUTOffset = 0;
+		int32		NoiseLUTNumChannel = 0;
+		int32		NoiseLUTChannelWidth = 0;
+	};
+}
diff --git a/Engine/Plugins/FX/Niagara/Source/Niagara/Internal/Stateless/Modules/NiagaraStatelessModule_AccelerationForce.h b/Engine/Plugins/FX/Niagara/Source/Niagara/Internal/Stateless/Modules/NiagaraStatelessModule_AccelerationForce.h
new file mode 100644
--- /dev/null
+++ b/Engine/Plugins/FX/Niagara/Source/Niagara/Internal/Stateless/Modules/NiagaraStatelessModule_AccelerationForce.h
@@ -0,0 +1,37 @@
+// Copyright Epic Games, Inc. All Rights Reserved.
+
+#pragma once
+
+#include "Stateless/NiagaraStatelessModule.h"
+#include "Stateless/NiagaraStatelessEmitterDataBuildContext.h"
+#include "Stateless/NiagaraStatelessModuleShaderParameters.h"
+#include "Stateless/Modules/NiagaraStatelessModuleCommon.h"
+
+#include "NiagaraStatelessModule_AccelerationForce.generated.h"
+
+UCLASS(MinimalAPI, EditInlineNew, meta = (DisplayName = "Acceleration Force"))
+class UNiagaraStatelessModule_AccelerationForce : public UNiagaraStatelessModule
+{
+	GENERATED_BODY()
+
+public:
+	UPROPERTY(EditAnywhere, Category = "Parameters")
+	FVector3f AccelerationMin = FVector3f(0.0f, 0.0f, 100.0f);
+	UPROPERTY(EditAnywhere, Category = "Parameters")
+	FVector3f AccelerationMax = FVector3f(0.0f, 0.0f, 100.0f);
+
+	virtual void BuildEmitterData(FNiagaraStatelessEmitterDataBuildContext& BuildContext) const override
+	{
+		if (!IsModuleEnabled())
+		{
+			return;
+		}
+		NiagaraStateless::FPhysicsBuildData& PhysicsBuildData = BuildContext.GetTransientBuildData<NiagaraStateless::FPhysicsBuildData>();
+		PhysicsBuildData.AccelerationMin += AccelerationMin;
+		PhysicsBuildData.AccelerationMax += AccelerationMax;
+	}
+
+#if WITH_EDITOR
+	virtual bool CanDisableModule() const override { return true; }
+#endif
+};
diff --git a/Engine/Plugins/FX/Niagara/Source/Niagara/Internal/Stateless/Modules/NiagaraStatelessModule_AddVelocity.h b/Engine/Plugins/FX/Niagara/Source/Niagara/Internal/Stateless/Modules/NiagaraStatelessModule_AddVelocity.h
new file mode 100644
--- /dev/null
+++ b/Engine/Plugins/FX/Niagara/Source/Niagara/Internal/Stateless/Modules/NiagaraStatelessModule_AddVelocity.h
@@ -0,0 +1,67 @@
+// Copyright Epic Games, Inc. All Rights Reserved.
+
+#pragma once
+
+#include "Stateless/NiagaraStatelessModule.h"
+#include "Stateless/NiagaraStatelessEmitterDataBuildContext.h"
+#include "Stateless/NiagaraStatelessModuleShaderParameters.h"
+#include "Stateless/Modules/NiagaraStatelessModuleCommon.h"
+
+#include "NiagaraStatelessModule_AddVelocity.generated.h"
+
+UENUM()
+enum class ENSM_VelocityType
+{
+	Linear,			// Velocity Min / Max -> Velocity Speed Scale
+	FromPoint,		// Min / Max / Origin
+	InCone,			// Max, Min, Cone| Cone Axis, Cone Angle, Inner Cone
+};
+
+UCLASS(MinimalAPI, EditInlineNew, meta = (DisplayName = "Add Velocity"))
+class UNiagaraStatelessModule_AddVelocity : public UNiagaraStatelessModule
+{
+	GENERATED_BODY()
+
+public:
+	UPROPERTY(EditAnywhere, Category = "Parameters")
+	ENSM_VelocityType VelocityType = ENSM_VelocityType::Linear;
+
+	UPROPERTY(EditAnywhere, Category = "Parameters", meta = (EditConditionHides, EditCondition = "VelocityType == ENSM_VelocityType::Linear"))
+	FVector3f VelocityMin = FVector3f(0.0f, 0.0f, 100.0f);
+	UPROPERTY(EditAnywhere, Category = "Parameters", meta = (EditConditionHides, EditCondition = "VelocityType == ENSM_VelocityType::Linear"))
+	FVector3f VelocityMax = FVector3f(0.0f, 0.0f, 100.0f);
+	UPROPERTY(EditAnywhere, Category = "Parameters", meta = (EditConditionHides, EditCondition = "VelocityType == ENSM_VelocityType::Linear"))
+	float VelocityScale = 1.0f;
+
+	UPROPERTY(EditAnywhere, Category = "Parameters", meta = (EditConditionHides, EditCondition = "VelocityType == ENSM_VelocityType::InCone"))
+	float ConeVelocityMin = 500.0f;
+	UPROPERTY(EditAnywhere, Category = "Parameters", meta = (EditConditionHides, EditCondition = "VelocityType == ENSM_VelocityType::InCone"))
+	float ConeVelocityMax = 500.0f;
+	UPROPERTY(EditAnywhere, Category = "Parameters", meta = (EditConditionHides, EditCondition = "VelocityType == ENSM_VelocityType::InCone"))
+	FRotator ConeRotation = FRotator::ZeroRotator;
+	UPROPERTY(EditAnywhere, Category = "Parameters", meta = (EditConditionHides, EditCondition = "VelocityType == ENSM_VelocityType::InCone", ClampMin = "0.0", ClampMax = "360.0", UIMin = "0.0", UIMax = "360.0"))
+	float ConeAngle = 45.0f;
+	UPROPERTY(EditAnywhere, Category = "Parameters", meta = (EditConditionHides, EditCondition = "VelocityType == ENSM_VelocityType::InCone", ClampMin = "0.0", ClampMax = "360.0", UIMin = "0.0", UIMax = "360.0"))
+	float InnerCone = 0.0f;
+
+	UPROPERTY(EditAnywhere, Category = "Parameters", meta = (EditConditionHides, EditCondition = "VelocityType == ENSM_VelocityType::FromPoint"))
+	float PointVelocityMin = 100.0f;
+	UPROPERTY(EditAnywhere, Category = "Parameters", meta = (EditConditionHides, EditCondition = "VelocityType == ENSM_VelocityType::FromPoint"))
+	float PointVelocityMax = 100.0f;
+	UPROPERTY(EditAnywhere, Category = "Parameters", meta = (EditConditionHides, EditCondition = "VelocityType == ENSM_VelocityType::FromPoint"))
+	FVector3f PointOrigin = FVector3f::ZeroVector;
+
+	UPROPERTY(EditAnywhere, Category = "Parameters", meta = (EditCondition = "VelocityType == ENSM_VelocityType::InCone"))
+	bool bSpeedFalloffFromConeAxisEnabled = false;
+	UPROPERTY(EditAnywhere, Category = "Parameters", meta = (EditCondition = "bSpeedFalloffFromConeAxisEnabled && VelocityType == ENSM_VelocityType::InCone", ClampMin = "0.0", ClampMax = "1.0", UIMin = "0.0", UIMax = "1.0"))
+	float SpeedFalloffFromConeAxis = 0.0f;
+
+	virtual void BuildEmitterData(FNiagaraStatelessEmitterDataBuildContext& BuildContext) const override;
+
+#if WITH_EDITOR
+	virtual bool CanDisableModule() const override { return true; }
+
+	virtual bool CanDebugDraw() const { return true; }
+	virtual void DrawDebug(const FNiagaraStatelessDrawDebugContext& DrawDebugContext) const override;
+#endif
+};
diff --git a/Engine/Plugins/FX/Niagara/Source/Niagara/Internal/Stateless/Modules/NiagaraStatelessModule_CalculateAccurateVelocity.h b/Engine/Plugins/FX/Niagara/Source/Niagara/Internal/Stateless/Modules/NiagaraStatelessModule_CalculateAccurateVelocity.h
new file mode 100644
--- /dev/null
+++ b/Engine/Plugins/FX/Niagara/Source/Niagara/Internal/Stateless/Modules/NiagaraStatelessModule_CalculateAccurateVelocity.h
@@ -0,0 +1,24 @@
+// Copyright Epic Games, Inc. All Rights Reserved.
+
+#pragma once
+
+#include "Stateless/NiagaraStatelessModule.h"
+#include "Stateless/NiagaraStatelessModuleShaderParameters.h"
+
+#include "NiagaraStatelessModule_CalculateAccurateVelocity.generated.h"
+
+UCLASS(MinimalAPI, EditInlineNew, meta = (DisplayName = "Rotate Around Point"))
+class UNiagaraStatelessModule_CalculateAccurateVelocity : public UNiagaraStatelessModule
+{
+	GENERATED_BODY()
+
+public:
+#if WITH_EDITORONLY_DATA
+	virtual void GetOutputVaruables(TArray<FNiagaraVariableBase>& OutVariables) const override
+	{
+		const FNiagaraStatelessGlobals& StatelessGlobals = FNiagaraStatelessGlobals::Get();
+		OutVariables.AddUnique(StatelessGlobals.VelocityVariable);
+		OutVariables.AddUnique(StatelessGlobals.PreviousVelocityVariable);
+	}
+#endif
+};
diff --git a/Engine/Plugins/FX/Niagara/Source/Niagara/Internal/Stateless/Modules/NiagaraStatelessModule_CurlNoiseForce.h b/Engine/Plugins/FX/Niagara/Source/Niagara/Internal/Stateless/Modules/NiagaraStatelessModule_CurlNoiseForce.h
new file mode 100644
--- /dev/null
+++ b/Engine/Plugins/FX/Niagara/Source/Niagara/Internal/Stateless/Modules/NiagaraStatelessModule_CurlNoiseForce.h
@@ -0,0 +1,45 @@
+// Copyright Epic Games, Inc. All Rights Reserved.
+
+#pragma once
+
+#include "Stateless/NiagaraStatelessModule.h"
+#include "Stateless/NiagaraStatelessEmitterDataBuildContext.h"
+#include "Stateless/NiagaraStatelessModuleShaderParameters.h"
+#include "Stateless/Modules/NiagaraStatelessModuleCommon.h"
+
+#include "NiagaraStatelessModule_CurlNoiseForce.generated.h"
+
+class UVectorField;
+
+UENUM()
+enum class ENSM_NoiseMode
+{
+	VectorField,
+	JacobNoise,
+	LUTJacob,
+	LUTJacobBicubic,
+};
+
+UCLASS(MinimalAPI, EditInlineNew, meta = (DisplayName = "Curl Noise Force"))
+class UNiagaraStatelessModule_CurlNoiseForce : public UNiagaraStatelessModule
+{
+	GENERATED_BODY()
+
+public:
+	UPROPERTY(EditAnywhere, Category = "Parameters")
+	float NoiseAmplitude = 200.0f;
+	UPROPERTY(EditAnywhere, Category = "Parameters")
+	float NoiseFrequency = 1.0f;
+
+	UPROPERTY(EditAnywhere, Category = "Parameters")
+	ENSM_NoiseMode NoiseMode = ENSM_NoiseMode::VectorField;
+
+	UPROPERTY(EditAnywhere, Category = "Parameters", meta = (AllowedClasses = "/Script/Engine.VectorField,/Script/Engine.VolumeTexture,/Script/Engine.TextureRenderTargetVolume"))
+	TObjectPtr<UObject> NoiseTexture;
+
+	virtual void BuildEmitterData(FNiagaraStatelessEmitterDataBuildContext& BuildContext) const override;
+
+#if WITH_EDITOR
+	virtual bool CanDisableModule() const override { return true; }
+#endif
+};
diff --git a/Engine/Plugins/FX/Niagara/Source/Niagara/Internal/Stateless/Modules/NiagaraStatelessModule_Drag.h b/Engine/Plugins/FX/Niagara/Source/Niagara/Internal/Stateless/Modules/NiagaraStatelessModule_Drag.h
new file mode 100644
--- /dev/null
+++ b/Engine/Plugins/FX/Niagara/Source/Niagara/Internal/Stateless/Modules/NiagaraStatelessModule_Drag.h
@@ -0,0 +1,38 @@
+// Copyright Epic Games, Inc. All Rights Reserved.
+
+#pragma once
+
+#include "Stateless/NiagaraStatelessModule.h"
+#include "Stateless/NiagaraStatelessEmitterDataBuildContext.h"
+#include "Stateless/NiagaraStatelessModuleShaderParameters.h"
+#include "Stateless/Modules/NiagaraStatelessModuleCommon.h"
+
+#include "NiagaraStatelessModule_Drag.generated.h"
+
+UCLASS(MinimalAPI, EditInlineNew, meta = (DisplayName = "Drag"))
+class UNiagaraStatelessModule_Drag : public UNiagaraStatelessModule
+{
+	GENERATED_BODY()
+
+public:
+	UPROPERTY(EditAnywhere, Category = "Parameters")
+	float DragMin = 0.1f;
+	UPROPERTY(EditAnywhere, Category = "Parameters")
+	float DragMax = 1.0f;
+
+	virtual void BuildEmitterData(FNiagaraStatelessEmitterDataBuildContext& BuildContext) const override
+	{
+		if (!IsModuleEnabled())
+		{
+			return;
+		}
+
+		NiagaraStateless::FPhysicsBuildData& PhysicsBuildData = BuildContext.GetTransientBuildData<NiagaraStateless::FPhysicsBuildData>();
+		PhysicsBuildData.DragMin += DragMin;
+		PhysicsBuildData.DragMax += DragMax;
+	}
+
+#if WITH_EDITOR
+	virtual bool CanDisableModule() const override { return true; }
+#endif
+};
diff --git a/Engine/Plugins/FX/Niagara/Source/Niagara/Internal/Stateless/Modules/NiagaraStatelessModule_GravityForce.h b/Engine/Plugins/FX/Niagara/Source/Niagara/Internal/Stateless/Modules/NiagaraStatelessModule_GravityForce.h
new file mode 100644
--- /dev/null
+++ b/Engine/Plugins/FX/Niagara/Source/Niagara/Internal/Stateless/Modules/NiagaraStatelessModule_GravityForce.h
@@ -0,0 +1,37 @@
+// Copyright Epic Games, Inc. All Rights Reserved.
+
+#pragma once
+
+#include "Stateless/NiagaraStatelessModule.h"
+#include "Stateless/NiagaraStatelessEmitterDataBuildContext.h"
+#include "Stateless/NiagaraStatelessModuleShaderParameters.h"
+#include "Stateless/Modules/NiagaraStatelessModuleCommon.h"
+
+#include "NiagaraStatelessModule_GravityForce.generated.h"
+
+UCLASS(MinimalAPI, EditInlineNew, meta = (DisplayName = "Gravity Force"))
+class UNiagaraStatelessModule_GravityForce : public UNiagaraStatelessModule
+{
+	GENERATED_BODY()
+
+public:
+	UPROPERTY(EditAnywhere, Category = "Parameters")
+	FVector3f GravityForceMin = FVector3f(0.0f, 0.0f, -980.0f);
+	UPROPERTY(EditAnywhere, Category = "Parameters")
+	FVector3f GravityForceMax = FVector3f(0.0f, 0.0f, -980.0f);
+
+	virtual void BuildEmitterData(FNiagaraStatelessEmitterDataBuildContext& BuildContext) const override
+	{
+		if (!IsModuleEnabled())
+		{
+			return;
+		}
+		NiagaraStateless::FPhysicsBuildData& PhysicsBuildData = BuildContext.GetTransientBuildData<NiagaraStateless::FPhysicsBuildData>();
+		PhysicsBuildData.AccelerationMin += GravityForceMin;
+		PhysicsBuildData.AccelerationMax += GravityForceMax;
+	}
+
+#if WITH_EDITOR
+	virtual bool CanDisableModule() const override { return true; }
+#endif
+};
diff --git a/Engine/Plugins/FX/Niagara/Source/Niagara/Internal/Stateless/Modules/NiagaraStatelessModule_InitialMeshOrientation.h b/Engine/Plugins/FX/Niagara/Source/Niagara/Internal/Stateless/Modules/NiagaraStatelessModule_InitialMeshOrientation.h
new file mode 100644
--- /dev/null
+++ b/Engine/Plugins/FX/Niagara/Source/Niagara/Internal/Stateless/Modules/NiagaraStatelessModule_InitialMeshOrientation.h
@@ -0,0 +1,51 @@
+// Copyright Epic Games, Inc. All Rights Reserved.
+
+#pragma once
+
+#include "Stateless/NiagaraStatelessModule.h"
+#include "Stateless/NiagaraStatelessEmitterDataBuildContext.h"
+#include "Stateless/NiagaraStatelessModuleShaderParameters.h"
+
+#include "NiagaraStatelessModule_InitialMeshOrientation.generated.h"
+
+UCLASS(MinimalAPI, EditInlineNew, meta = (DisplayName = "Initial Mesh Orientation"))
+class UNiagaraStatelessModule_InitialMeshOrientation : public UNiagaraStatelessModule
+{
+	GENERATED_BODY()
+
+public:
+	using FParameters = NiagaraStateless::FInitialMeshOrientationModule_ShaderParameters;
+
+	UPROPERTY(EditAnywhere, Category = "Parameters")
+	FVector3f	Rotation = FVector3f::ZeroVector;
+
+	UPROPERTY(EditAnywhere, Category = "Parameters")
+	FVector3f	RandomRotationRange = FVector3f(360.0f, 360.0f, 360.0f);
+
+	virtual void SetShaderParameters(const FNiagaraStatelessSetShaderParameterContext& SetShaderParameterContext) const override
+	{
+		FParameters* Parameters = SetShaderParameterContext.GetParameterNestedStruct<FParameters>();
+		if (IsModuleEnabled())
+		{
+			Parameters->InitialMeshOrientation_Rotation			= Rotation;
+			Parameters->InitialMeshOrientation_RandomRangeScale	= RandomRotationRange / 360.0f;
+		}
+		else
+		{
+			Parameters->InitialMeshOrientation_Rotation			= FVector3f::ZeroVector;
+			Parameters->InitialMeshOrientation_RandomRangeScale	= FVector3f::ZeroVector;
+		}
+	}
+
+#if WITH_EDITOR
+	virtual bool CanDisableModule() const override { return true; }
+#endif
+#if WITH_EDITORONLY_DATA
+	virtual void GetOutputVaruables(TArray<FNiagaraVariableBase>& OutVariables) const override
+	{
+		const FNiagaraStatelessGlobals& StatelessGlobals = FNiagaraStatelessGlobals::Get();
+		OutVariables.AddUnique(StatelessGlobals.MeshOrientationVariable);
+		OutVariables.AddUnique(StatelessGlobals.PreviousMeshOrientationVariable);
+	}
+#endif
+};
diff --git a/Engine/Plugins/FX/Niagara/Source/Niagara/Internal/Stateless/Modules/NiagaraStatelessModule_InitializeParticle.h b/Engine/Plugins/FX/Niagara/Source/Niagara/Internal/Stateless/Modules/NiagaraStatelessModule_InitializeParticle.h
new file mode 100644
--- /dev/null
+++ b/Engine/Plugins/FX/Niagara/Source/Niagara/Internal/Stateless/Modules/NiagaraStatelessModule_InitializeParticle.h
@@ -0,0 +1,203 @@
+// Copyright Epic Games, Inc. All Rights Reserved.
+
+#pragma once
+
+#include "Stateless/NiagaraStatelessModule.h"
+#include "Stateless/NiagaraStatelessEmitterDataBuildContext.h"
+#include "Stateless/NiagaraStatelessModuleShaderParameters.h"
+#include "Stateless/Modules/NiagaraStatelessModuleCommon.h"
+
+#include "NiagaraParameterBinding.h"
+#include "NiagaraParameterStore.h"
+
+#include "NiagaraStatelessModule_InitializeParticle.generated.h"
+
+UENUM()
+enum class ENSPositionDistributionType
+{
+	UniqueIndexClamp,
+	UniqueIndexWrap,
+	Random,
+	RandomInterpolateToNext,
+	RandomInterpolateToRandom,
+};
+
+USTRUCT()
+struct FPositionDistributionTest
+{
+	GENERATED_BODY()
+
+	UPROPERTY(EditAnywhere, Category = "Parameters")
+	ENSPositionDistributionType Type = ENSPositionDistributionType::UniqueIndexClamp;
+
+	UPROPERTY(EditAnywhere, Category = "Parameters")
+	TArray<FVector3f> Values;
+};
+
+//UENUM()
+//enum class NiagaraStatelessModuleInitializeParticle_ColorMode
+//{
+//	Unset,
+//	DirectSet,
+//	RandomRange,
+//	//RandomHSV UMETA(DisplayName="Random Hue/Saturation/Value"),
+//};
+
+UCLASS(MinimalAPI, EditInlineNew, meta = (DisplayName = "Initialize Particle"))
+class UNiagaraStatelessModule_InitializeParticle : public UNiagaraStatelessModule
+{
+	GENERATED_BODY()
+
+	struct FModuleBuiltData
+	{
+		int32	PositionParameterBinding = INDEX_NONE;
+		int32	PositionTableOffset = 0;
+		int32	PositionTableLength = 0;
+	};
+
+public:
+	using FParameters = NiagaraStateless::FInitializeParticleModule_ShaderParameters;
+
+	UPROPERTY(EditAnywhere, Category = "Parameters")
+	float LifetimeMin = 1.0f;
+	UPROPERTY(EditAnywhere, Category = "Parameters")
+	float LifetimeMax = 1.0f;
+
+	UPROPERTY(EditAnywhere, Category = "Parameters")
+	FVector3f	InitialPosition = FVector3f(0.0f);
+
+	UPROPERTY(EditAnywhere, Category = "Parameters")
+	FLinearColor	ColorMin = FLinearColor::White;
+	UPROPERTY(EditAnywhere, Category = "Parameters")
+	FLinearColor	ColorMax = FLinearColor::White;
+
+	UPROPERTY(EditAnywhere, Category = "Parameters")
+	float MassMin = 1.0f;
+	UPROPERTY(EditAnywhere, Category = "Parameters")
+	float MassMax = 1.0f;
+
+	UPROPERTY(EditAnywhere, Category = "Parameters")
+	FVector2f	SpriteSizeMin = FVector2f(10.0f);
+	UPROPERTY(EditAnywhere, Category = "Parameters")
+	FVector2f	SpriteSizeMax = FVector2f(10.0f);
+
+	UPROPERTY(EditAnywhere, Category = "Parameters")
+	float	SpriteRotationMin = 0.0f;
+	UPROPERTY(EditAnywhere, Category = "Parameters")
+	float	SpriteRotationMax = 0.0f;
+
+	UPROPERTY(EditAnywhere, Category = "Parameters")
+	FVector3f	MeshScaleMin = FVector3f(1.0f);
+	UPROPERTY(EditAnywhere, Category = "Parameters")
+	FVector3f	MeshScaleMax = FVector3f(1.0f);
+
+	UPROPERTY(EditAnywhere, Category = "Parameters", meta = (InlineEditConditionToggle))
+	bool bWriteRibbonWidth = false;
+
+	UPROPERTY(EditAnywhere, Category = "Parameters", meta = (EditCondition="bWriteRibbonWidth"))
+	FVector2f	RibbonWidthMinMax = FVector2f(10.0f, 10.0f);
+
+	UPROPERTY(EditAnywhere, Category = "Parameters")
+	FNiagaraParameterBindingWithValue	InitialPositionBinding;
+
+	UPROPERTY(EditAnywhere, Category = "Parameters")
+	FPositionDistributionTest PositionDistribution;
+
+	//~Begin: UObject Interface
+#if WITH_EDITORONLY_DATA
+	virtual void PostInitProperties()
+	{
+		Super::PostInitProperties();
+		if (HasAnyFlags(RF_ClassDefaultObject) == false)
+		{
+			InitialPositionBinding.SetUsage(ENiagaraParameterBindingUsage::NotParticle);
+			InitialPositionBinding.SetAllowedTypeDefinitions({ FNiagaraTypeDefinition::GetVec3Def() });
+			InitialPositionBinding.SetDefaultParameter(FNiagaraTypeDefinition::GetVec3Def(), FVector3f::ZeroVector);
+		}
+	}
+#endif
+	//~End: UObject Interface
+
+	virtual void BuildEmitterData(FNiagaraStatelessEmitterDataBuildContext& BuildContext) const override
+	{
+		const FVector3f DefaultPosition(0.0f);
+		TConstArrayView<FVector3f> PositionValues = PositionDistribution.Values.Num() > 0 ? MakeArrayView(PositionDistribution.Values) : MakeArrayView(&DefaultPosition, 1);
+
+		FModuleBuiltData* BuiltData = BuildContext.AllocateBuiltData<FModuleBuiltData>();
+		BuiltData->PositionTableOffset = BuildContext.AddStaticData(PositionValues);
+		BuiltData->PositionTableLength = PositionValues.Num();
+
+		BuiltData->PositionParameterBinding = BuildContext.AddRendererBinding(InitialPositionBinding);
+	
+		NiagaraStateless::FPhysicsBuildData& PhysicsBuildData = BuildContext.GetTransientBuildData<NiagaraStateless::FPhysicsBuildData>();
+		PhysicsBuildData.MassMin = MassMin;
+		PhysicsBuildData.MassMax = MassMax;
+	}
+
+	virtual void SetShaderParameters(const FNiagaraStatelessSetShaderParameterContext& SetShaderParameterContext) const override
+	{
+		FParameters* Parameters = SetShaderParameterContext.GetParameterNestedStruct<FParameters>();
+		const FModuleBuiltData* ModuleBuiltData = SetShaderParameterContext.ReadBuiltData<FModuleBuiltData>();
+
+		//Parameters->InitializeParticle_Position			= InitialPosition;
+		SetShaderParameterContext.GetRendererParameterValue(Parameters->InitializeParticle_Position, ModuleBuiltData->PositionParameterBinding, InitialPositionBinding.GetDefaultValue<FVector3f>());
+		Parameters->InitializeParticle_ColorScale			= ColorMax - ColorMin;
+		Parameters->InitializeParticle_ColorBias			= ColorMin;
+		Parameters->InitializeParticle_SpriteSizeScale		= SpriteSizeMax - SpriteSizeMin;
+		Parameters->InitializeParticle_SpriteSizeBias		= SpriteSizeMin;
+		Parameters->InitializeParticle_SpriteRotationScale	= SpriteRotationMax - SpriteRotationMin;
+		Parameters->InitializeParticle_SpriteRotationBias	= SpriteRotationMin;
+		Parameters->InitializeParticle_MeshScaleScale		= MeshScaleMax - MeshScaleMin;
+		Parameters->InitializeParticle_MeshScaleBias		= MeshScaleMin;
+		Parameters->InitializeParticle_RibbonWidthScale		= RibbonWidthMinMax.Y - RibbonWidthMinMax.X;
+		Parameters->InitializeParticle_RibbonWidthBias		= RibbonWidthMinMax.X;
+
+		Parameters->InitializeParticle_PositionMode			= 0;
+		Parameters->InitializeParticle_PositionTableMod		= ModuleBuiltData->PositionTableLength;
+		Parameters->InitializeParticle_PositionTableOffset	= ModuleBuiltData->PositionTableOffset;
+		Parameters->InitializeParticle_PositionTableLength	= ModuleBuiltData->PositionTableLength;
+		switch (PositionDistribution.Type)
+		{
+			case ENSPositionDistributionType::UniqueIndexClamp:
+				Parameters->InitializeParticle_PositionMode			= 0;
+				Parameters->InitializeParticle_PositionTableMod		= 0xffffffff;
+				break;
+			case ENSPositionDistributionType::UniqueIndexWrap:
+				Parameters->InitializeParticle_PositionMode			= 0;
+				break;
+			case ENSPositionDistributionType::Random:
+				Parameters->InitializeParticle_PositionMode			= 2;
+				break;
+			case ENSPositionDistributionType::RandomInterpolateToNext:
+				Parameters->InitializeParticle_PositionMode			= 1;
+				break;
+			case ENSPositionDistributionType::RandomInterpolateToRandom:
+				Parameters->InitializeParticle_PositionMode			= 3;
+				break;
+		}
+	}
+
+#if WITH_EDITORONLY_DATA
+	virtual void GetOutputVaruables(TArray<FNiagaraVariableBase>& OutVariables) const override
+	{
+		const FNiagaraStatelessGlobals& StatelessGlobals = FNiagaraStatelessGlobals::Get();
+		OutVariables.AddUnique(StatelessGlobals.UniqueIDVariable);
+		OutVariables.AddUnique(StatelessGlobals.PositionVariable);
+		OutVariables.AddUnique(StatelessGlobals.ColorVariable);
+		OutVariables.AddUnique(StatelessGlobals.SpriteSizeVariable);
+		OutVariables.AddUnique(StatelessGlobals.SpriteRotationVariable);
+		OutVariables.AddUnique(StatelessGlobals.ScaleVariable);
+
+		OutVariables.AddUnique(StatelessGlobals.PreviousPositionVariable);
+		OutVariables.AddUnique(StatelessGlobals.PreviousSpriteSizeVariable);
+		OutVariables.AddUnique(StatelessGlobals.PreviousSpriteRotationVariable);
+		OutVariables.AddUnique(StatelessGlobals.PreviousScaleVariable);
+
+		if (bWriteRibbonWidth)
+		{
+			OutVariables.AddUnique(StatelessGlobals.RibbonWidthVariable);
+			OutVariables.AddUnique(StatelessGlobals.PreviousRibbonWidthVariable);
+		}
+	}
+#endif
+};
diff --git a/Engine/Plugins/FX/Niagara/Source/Niagara/Internal/Stateless/Modules/NiagaraStatelessModule_RotateAroundPoint.h b/Engine/Plugins/FX/Niagara/Source/Niagara/Internal/Stateless/Modules/NiagaraStatelessModule_RotateAroundPoint.h
new file mode 100644
--- /dev/null
+++ b/Engine/Plugins/FX/Niagara/Source/Niagara/Internal/Stateless/Modules/NiagaraStatelessModule_RotateAroundPoint.h
@@ -0,0 +1,67 @@
+// Copyright Epic Games, Inc. All Rights Reserved.
+
+#pragma once
+
+#include "Stateless/NiagaraStatelessModule.h"
+#include "Stateless/NiagaraStatelessModuleShaderParameters.h"
+
+#include "NiagaraStatelessModule_RotateAroundPoint.generated.h"
+
+UCLASS(MinimalAPI, EditInlineNew, meta = (DisplayName = "Rotate Around Point"))
+class UNiagaraStatelessModule_RotateAroundPoint : public UNiagaraStatelessModule
+{
+	GENERATED_BODY()
+
+public:
+	using FParameters = NiagaraStateless::FRotateAroundPointModule_ShaderParameters;
+
+	UPROPERTY(EditAnywhere, Category = "Parameters")
+	float RateMin = 360.f;
+	UPROPERTY(EditAnywhere, Category = "Parameters")
+	float RateMax = 360.f;
+
+	UPROPERTY(EditAnywhere, Category = "Parameters")
+	float RadiusMin = 50.f;
+	UPROPERTY(EditAnywhere, Category = "Parameters")
+	float RadiusMax = 50.f;
+
+	UPROPERTY(EditAnywhere, Category = "Parameters")
+	float InitialPhaseMin = 0.f;
+	UPROPERTY(EditAnywhere, Category = "Parameters")
+	float InitialPhaseMax = 0.f;
+
+	virtual void SetShaderParameters(const FNiagaraStatelessSetShaderParameterContext& SetShaderParameterContext) const override
+	{
+		FParameters* Parameters = SetShaderParameterContext.GetParameterNestedStruct<FParameters>();
+		if (IsModuleEnabled())
+		{
+			Parameters->RotateAroundPoint_RateScale			= FMath::DegreesToRadians(RateMax - RateMin);
+			Parameters->RotateAroundPoint_RateBias			= FMath::DegreesToRadians(RateMin);
+			Parameters->RotateAroundPoint_RadiusScale		= RadiusMax - RadiusMin;
+			Parameters->RotateAroundPoint_RadiusBias		= RadiusMin;
+			Parameters->RotateAroundPoint_InitialPhaseScale	= InitialPhaseMax - InitialPhaseMin;
+			Parameters->RotateAroundPoint_InitialPhaseBias	= InitialPhaseMin;
+		}
+		else
+		{
+			Parameters->RotateAroundPoint_RateScale			= 0.0f;
+			Parameters->RotateAroundPoint_RateBias			= 0.0f;
+			Parameters->RotateAroundPoint_RadiusScale		= 0.0f;
+			Parameters->RotateAroundPoint_RadiusBias		= 0.0f;
+			Parameters->RotateAroundPoint_InitialPhaseScale	= 0.0f;
+			Parameters->RotateAroundPoint_InitialPhaseBias	= 0.0f;
+		}
+	}
+
+#if WITH_EDITOR
+	virtual bool CanDisableModule() const override { return true; }
+#endif
+#if WITH_EDITORONLY_DATA
+	virtual void GetOutputVaruables(TArray<FNiagaraVariableBase>& OutVariables) const override
+	{
+		const FNiagaraStatelessGlobals& StatelessGlobals = FNiagaraStatelessGlobals::Get();
+		OutVariables.AddUnique(StatelessGlobals.PositionVariable);
+		OutVariables.AddUnique(StatelessGlobals.PreviousPositionVariable);
+	}
+#endif
+};
diff --git a/Engine/Plugins/FX/Niagara/Source/Niagara/Internal/Stateless/Modules/NiagaraStatelessModule_ScaleColor.h b/Engine/Plugins/FX/Niagara/Source/Niagara/Internal/Stateless/Modules/NiagaraStatelessModule_ScaleColor.h
new file mode 100644
--- /dev/null
+++ b/Engine/Plugins/FX/Niagara/Source/Niagara/Internal/Stateless/Modules/NiagaraStatelessModule_ScaleColor.h
@@ -0,0 +1,57 @@
+// Copyright Epic Games, Inc. All Rights Reserved.
+
+#pragma once
+
+#include "Stateless/NiagaraStatelessModule.h"
+#include "Stateless/NiagaraStatelessEmitterDataBuildContext.h"
+#include "Stateless/NiagaraStatelessModuleShaderParameters.h"
+
+#include "NiagaraStatelessModule_ScaleColor.generated.h"
+
+UCLASS(MinimalAPI, EditInlineNew, meta = (DisplayName = "Scale Color"))
+class UNiagaraStatelessModule_ScaleColor : public UNiagaraStatelessModule
+{
+	GENERATED_BODY()
+
+	struct FModuleBuiltData
+	{
+		int32	TableOffset = 0;
+		int32	TableLength = 0;
+	};
+
+public:
+	using FParameters = NiagaraStateless::FScaleColorModule_ShaderParameters;
+
+	UPROPERTY(EditAnywhere, Category = "Parameters")
+	TArray<FLinearColor>	ScaleValues = {FLinearColor::White, FLinearColor::White, FLinearColor::White, FLinearColor::Transparent};
+
+	virtual void BuildEmitterData(FNiagaraStatelessEmitterDataBuildContext& BuildContext) const override
+	{
+		FModuleBuiltData* BuiltData = BuildContext.AllocateBuiltData<FModuleBuiltData>();
+		if (IsModuleEnabled())
+		{
+			BuiltData->TableOffset = BuildContext.AddStaticData(ScaleValues);
+			BuiltData->TableLength = ScaleValues.Num();
+		}
+	}
+
+	virtual void SetShaderParameters(const FNiagaraStatelessSetShaderParameterContext& SetShaderParameterContext) const override
+	{
+		const FModuleBuiltData* ModuleBuiltData = SetShaderParameterContext.ReadBuiltData<FModuleBuiltData>();
+
+		FParameters* Parameters = SetShaderParameterContext.GetParameterNestedStruct<FParameters>();
+		Parameters->ScaleColor_Offset = ModuleBuiltData->TableOffset;
+		Parameters->ScaleColor_Length = ModuleBuiltData->TableLength;
+	}
+
+#if WITH_EDITOR
+	virtual bool CanDisableModule() const override { return true; }
+#endif
+#if WITH_EDITORONLY_DATA
+	virtual void GetOutputVaruables(TArray<FNiagaraVariableBase>& OutVariables) const override
+	{
+		const FNiagaraStatelessGlobals& StatelessGlobals = FNiagaraStatelessGlobals::Get();
+		OutVariables.AddUnique(StatelessGlobals.ColorVariable);
+	}
+#endif
+};
diff --git a/Engine/Plugins/FX/Niagara/Source/Niagara/Internal/Stateless/Modules/NiagaraStatelessModule_ScaleMeshSize.h b/Engine/Plugins/FX/Niagara/Source/Niagara/Internal/Stateless/Modules/NiagaraStatelessModule_ScaleMeshSize.h
new file mode 100644
--- /dev/null
+++ b/Engine/Plugins/FX/Niagara/Source/Niagara/Internal/Stateless/Modules/NiagaraStatelessModule_ScaleMeshSize.h
@@ -0,0 +1,58 @@
+// Copyright Epic Games, Inc. All Rights Reserved.
+
+#pragma once
+
+#include "Stateless/NiagaraStatelessModule.h"
+#include "Stateless/NiagaraStatelessEmitterDataBuildContext.h"
+#include "Stateless/NiagaraStatelessModuleShaderParameters.h"
+
+#include "NiagaraStatelessModule_ScaleMeshSize.generated.h"
+
+UCLASS(MinimalAPI, EditInlineNew, meta = (DisplayName = "Scale Mesh Size"))
+class UNiagaraStatelessModule_ScaleMeshSize : public UNiagaraStatelessModule
+{
+	GENERATED_BODY()
+
+	struct FModuleBuiltData
+	{
+		int32	TableOffset = 0;
+		int32	TableLength = 0;
+	};
+
+public:
+	using FParameters = NiagaraStateless::FScaleMeshSizeModule_ShaderParameters;
+
+	UPROPERTY(EditAnywhere, Category = "Parameters")
+	TArray<FVector3f>	ScaleValues = { FVector3f::OneVector, FVector3f::OneVector, FVector3f::OneVector, FVector3f::ZeroVector };
+
+	virtual void BuildEmitterData(FNiagaraStatelessEmitterDataBuildContext& BuildContext) const override
+	{
+		FModuleBuiltData* BuiltData = BuildContext.AllocateBuiltData<FModuleBuiltData>();
+		if (IsModuleEnabled())
+		{
+			BuiltData->TableOffset = BuildContext.AddStaticData(ScaleValues);
+			BuiltData->TableLength = ScaleValues.Num();
+		}
+	}
+
+	virtual void SetShaderParameters(const FNiagaraStatelessSetShaderParameterContext& SetShaderParameterContext) const override
+	{
+		const FModuleBuiltData* ModuleBuiltData = SetShaderParameterContext.ReadBuiltData<FModuleBuiltData>();
+
+		FParameters* Parameters = SetShaderParameterContext.GetParameterNestedStruct<FParameters>();
+		Parameters->ScaleMeshSize_Offset = ModuleBuiltData->TableOffset;
+		Parameters->ScaleMeshSize_Length = ModuleBuiltData->TableLength;
+	}
+
+#if WITH_EDITOR
+	virtual bool CanDisableModule() const override { return true; }
+#endif
+#if WITH_EDITORONLY_DATA
+	virtual void GetOutputVaruables(TArray<FNiagaraVariableBase>& OutVariables) const override
+	{
+		const FNiagaraStatelessGlobals& StatelessGlobals = FNiagaraStatelessGlobals::Get();
+		OutVariables.AddUnique(StatelessGlobals.ScaleVariable);
+		OutVariables.AddUnique(StatelessGlobals.PreviousScaleVariable);
+	}
+#endif
+};
