- Experimental stateless emitters

diff --git a/Engine/Plugins/FX/Niagara/Source/Niagara/Internal/Stateless/Modules/NiagaraStatelessModule_ScaleMeshSizeBySpeed.h b/Engine/Plugins/FX/Niagara/Source/Niagara/Internal/Stateless/Modules/NiagaraStatelessModule_ScaleMeshSizeBySpeed.h
new file mode 100644
--- /dev/null
+++ b/Engine/Plugins/FX/Niagara/Source/Niagara/Internal/Stateless/Modules/NiagaraStatelessModule_ScaleMeshSizeBySpeed.h
@@ -0,0 +1,64 @@
+// Copyright Epic Games, Inc. All Rights Reserved.
+
+#pragma once
+
+#include "Stateless/NiagaraStatelessModule.h"
+#include "Stateless/NiagaraStatelessEmitterDataBuildContext.h"
+#include "Stateless/NiagaraStatelessModuleShaderParameters.h"
+
+#include "NiagaraStatelessModule_ScaleMeshSizeBySpeed.generated.h"
+
+UCLASS(MinimalAPI, EditInlineNew, meta = (DisplayName = "Scale Mesh Size By Speed"))
+class UNiagaraStatelessModule_ScaleMeshSizeBySpeed : public UNiagaraStatelessModule
+{
+	GENERATED_BODY()
+
+	struct FModuleBuiltData
+	{
+		float	VelocityNorm = 0.0f;
+		int32	TableOffset = 0;
+		int32	TableLength = 0;
+	};
+
+public:
+	using FParameters = NiagaraStateless::FScaleMeshSizeBySpeedModule_ShaderParameters;
+
+	UPROPERTY(EditAnywhere, Category = "Parameters", meta = (ClampMin = "0.01", UIMin = "0.01"))
+	float VelocityThreshold = 1000.0f;
+
+	UPROPERTY(EditAnywhere, Category = "Parameters")
+	TArray<FVector3f>	ScaleValues = { FVector3f(1.0, 1.0f, 1.0f), FVector3f(2.0, 2.0f, 2.0f) };
+
+	virtual void BuildEmitterData(FNiagaraStatelessEmitterDataBuildContext& BuildContext) const override
+	{
+		FModuleBuiltData* BuiltData = BuildContext.AllocateBuiltData<FModuleBuiltData>();
+		if (IsModuleEnabled())
+		{
+			BuiltData->VelocityNorm	= VelocityThreshold > 0.0f ? 1.0f / (VelocityThreshold * VelocityThreshold) : 0.0f;
+			BuiltData->TableOffset	= BuildContext.AddStaticData(ScaleValues);
+			BuiltData->TableLength	= ScaleValues.Num();
+		}
+	}
+
+	virtual void SetShaderParameters(const FNiagaraStatelessSetShaderParameterContext& SetShaderParameterContext) const override
+	{
+		const FModuleBuiltData* ModuleBuiltData = SetShaderParameterContext.ReadBuiltData<FModuleBuiltData>();
+
+		FParameters* Parameters = SetShaderParameterContext.GetParameterNestedStruct<FParameters>();
+		Parameters->ScaleMeshSizeBySpeed_VelocityNorm	= ModuleBuiltData->VelocityNorm;
+		Parameters->ScaleMeshSizeBySpeed_Offset			= ModuleBuiltData->TableOffset;
+		Parameters->ScaleMeshSizeBySpeed_Length			= ModuleBuiltData->TableLength;
+	}
+
+#if WITH_EDITOR
+	virtual bool CanDisableModule() const override { return true; }
+#endif
+#if WITH_EDITORONLY_DATA
+	virtual void GetOutputVaruables(TArray<FNiagaraVariableBase>& OutVariables) const override
+	{
+		const FNiagaraStatelessGlobals& StatelessGlobals = FNiagaraStatelessGlobals::Get();
+		OutVariables.AddUnique(StatelessGlobals.ScaleVariable);
+		OutVariables.AddUnique(StatelessGlobals.PreviousScaleVariable);
+	}
+#endif
+};
diff --git a/Engine/Plugins/FX/Niagara/Source/Niagara/Internal/Stateless/Modules/NiagaraStatelessModule_ScaleSpriteSize.h b/Engine/Plugins/FX/Niagara/Source/Niagara/Internal/Stateless/Modules/NiagaraStatelessModule_ScaleSpriteSize.h
new file mode 100644
--- /dev/null
+++ b/Engine/Plugins/FX/Niagara/Source/Niagara/Internal/Stateless/Modules/NiagaraStatelessModule_ScaleSpriteSize.h
@@ -0,0 +1,58 @@
+// Copyright Epic Games, Inc. All Rights Reserved.
+
+#pragma once
+
+#include "Stateless/NiagaraStatelessModule.h"
+#include "Stateless/NiagaraStatelessEmitterDataBuildContext.h"
+#include "Stateless/NiagaraStatelessModuleShaderParameters.h"
+
+#include "NiagaraStatelessModule_ScaleSpriteSize.generated.h"
+
+UCLASS(MinimalAPI, EditInlineNew, meta = (DisplayName = "Scale Sprite Size"))
+class UNiagaraStatelessModule_ScaleSpriteSize : public UNiagaraStatelessModule
+{
+	GENERATED_BODY()
+
+	struct FModuleBuiltData
+	{
+		int32	TableOffset = 0;
+		int32	TableLength = 0;
+	};
+
+public:
+	using FParameters = NiagaraStateless::FScaleSpriteSizeModule_ShaderParameters;
+
+	UPROPERTY(EditAnywhere, Category = "Parameters")
+	TArray<FVector2f>	ScaleValues = {FVector2f::One(), FVector2f::One(), FVector2f::One(), FVector2f::ZeroVector};
+
+	virtual void BuildEmitterData(FNiagaraStatelessEmitterDataBuildContext& BuildContext) const override
+	{
+		FModuleBuiltData* BuiltData = BuildContext.AllocateBuiltData<FModuleBuiltData>();
+		if (IsModuleEnabled())
+		{
+			BuiltData->TableOffset = BuildContext.AddStaticData(ScaleValues);
+			BuiltData->TableLength = ScaleValues.Num();
+		}
+	}
+
+	virtual void SetShaderParameters(const FNiagaraStatelessSetShaderParameterContext& SetShaderParameterContext) const override
+	{
+		const FModuleBuiltData* ModuleBuiltData = SetShaderParameterContext.ReadBuiltData<FModuleBuiltData>();
+
+		FParameters* Parameters = SetShaderParameterContext.GetParameterNestedStruct<FParameters>();
+		Parameters->ScaleSpriteSize_Offset = ModuleBuiltData->TableOffset;
+		Parameters->ScaleSpriteSize_Length = ModuleBuiltData->TableLength;
+	}
+
+#if WITH_EDITOR
+	virtual bool CanDisableModule() const override { return true; }
+#endif
+#if WITH_EDITORONLY_DATA
+	virtual void GetOutputVaruables(TArray<FNiagaraVariableBase>& OutVariables) const override
+	{
+		const FNiagaraStatelessGlobals& StatelessGlobals = FNiagaraStatelessGlobals::Get();
+		OutVariables.AddUnique(StatelessGlobals.SpriteSizeVariable);
+		OutVariables.AddUnique(StatelessGlobals.PreviousSpriteSizeVariable);
+	}
+#endif
+};
diff --git a/Engine/Plugins/FX/Niagara/Source/Niagara/Internal/Stateless/Modules/NiagaraStatelessModule_ScaleSpriteSizeBySpeed.h b/Engine/Plugins/FX/Niagara/Source/Niagara/Internal/Stateless/Modules/NiagaraStatelessModule_ScaleSpriteSizeBySpeed.h
new file mode 100644
--- /dev/null
+++ b/Engine/Plugins/FX/Niagara/Source/Niagara/Internal/Stateless/Modules/NiagaraStatelessModule_ScaleSpriteSizeBySpeed.h
@@ -0,0 +1,64 @@
+// Copyright Epic Games, Inc. All Rights Reserved.
+
+#pragma once
+
+#include "Stateless/NiagaraStatelessModule.h"
+#include "Stateless/NiagaraStatelessEmitterDataBuildContext.h"
+#include "Stateless/NiagaraStatelessModuleShaderParameters.h"
+
+#include "NiagaraStatelessModule_ScaleSpriteSizeBySpeed.generated.h"
+
+UCLASS(MinimalAPI, EditInlineNew, meta = (DisplayName = "Scale Sprite Size By Speed"))
+class UNiagaraStatelessModule_ScaleSpriteSizeBySpeed : public UNiagaraStatelessModule
+{
+	GENERATED_BODY()
+
+	struct FModuleBuiltData
+	{
+		float	VelocityNorm = 0.0f;
+		int32	TableOffset = 0;
+		int32	TableLength = 0;
+	};
+
+public:
+	using FParameters = NiagaraStateless::FScaleSpriteSizeBySpeedModule_ShaderParameters;
+
+	UPROPERTY(EditAnywhere, Category = "Parameters", meta = (ClampMin = "0.01", UIMin = "0.01"))
+	float VelocityThreshold = 1000.0f;
+
+	UPROPERTY(EditAnywhere, Category = "Parameters")
+	TArray<FVector2f>	ScaleValues = { FVector2f(1.0f, 1.0f), FVector2f(2.0f, 2.0f) };
+
+	virtual void BuildEmitterData(FNiagaraStatelessEmitterDataBuildContext& BuildContext) const override
+	{
+		FModuleBuiltData* BuiltData = BuildContext.AllocateBuiltData<FModuleBuiltData>();
+		if (IsModuleEnabled())
+		{
+			BuiltData->VelocityNorm	= VelocityThreshold > 0.0f ? 1.0f / (VelocityThreshold * VelocityThreshold) : 0.0f;
+			BuiltData->TableOffset	= BuildContext.AddStaticData(ScaleValues);
+			BuiltData->TableLength	= ScaleValues.Num();
+		}
+	}
+
+	virtual void SetShaderParameters(const FNiagaraStatelessSetShaderParameterContext& SetShaderParameterContext) const override
+	{
+		const FModuleBuiltData* ModuleBuiltData = SetShaderParameterContext.ReadBuiltData<FModuleBuiltData>();
+
+		FParameters* Parameters = SetShaderParameterContext.GetParameterNestedStruct<FParameters>();
+		Parameters->ScaleSpriteSizeBySpeed_VelocityNorm		= ModuleBuiltData->VelocityNorm;
+		Parameters->ScaleSpriteSizeBySpeed_Offset			= ModuleBuiltData->TableOffset;
+		Parameters->ScaleSpriteSizeBySpeed_Length			= ModuleBuiltData->TableLength;
+	}
+
+#if WITH_EDITOR
+	virtual bool CanDisableModule() const override { return true; }
+#endif
+#if WITH_EDITORONLY_DATA
+	virtual void GetOutputVaruables(TArray<FNiagaraVariableBase>& OutVariables) const override
+	{
+		const FNiagaraStatelessGlobals& StatelessGlobals = FNiagaraStatelessGlobals::Get();
+		OutVariables.AddUnique(StatelessGlobals.SpriteSizeVariable);
+		OutVariables.AddUnique(StatelessGlobals.PreviousSpriteSizeVariable);
+	}
+#endif
+};
diff --git a/Engine/Plugins/FX/Niagara/Source/Niagara/Internal/Stateless/Modules/NiagaraStatelessModule_ShapeLocation.h b/Engine/Plugins/FX/Niagara/Source/Niagara/Internal/Stateless/Modules/NiagaraStatelessModule_ShapeLocation.h
new file mode 100644
--- /dev/null
+++ b/Engine/Plugins/FX/Niagara/Source/Niagara/Internal/Stateless/Modules/NiagaraStatelessModule_ShapeLocation.h
@@ -0,0 +1,78 @@
+// Copyright Epic Games, Inc. All Rights Reserved.
+
+#pragma once
+
+#include "Stateless/NiagaraStatelessModule.h"
+#include "Stateless/NiagaraStatelessModuleShaderParameters.h"
+
+#include "NiagaraStatelessModule_ShapeLocation.generated.h"
+
+UENUM()
+enum class ENSM_ShapePrimitive
+{
+	Box,
+	Cylinder,
+	Plane,
+	Ring,
+	Sphere,
+};
+
+UCLASS(MinimalAPI, EditInlineNew, meta = (DisplayName = "Shape Location"))
+class UNiagaraStatelessModule_ShapeLocation : public UNiagaraStatelessModule
+{
+	GENERATED_BODY()
+
+public:
+	using FParameters = NiagaraStateless::FShapeLocationModule_ShaderParameters;
+
+	UPROPERTY(EditAnywhere, Category = "Parameters", meta = (ShowInStackItemHeader, StackItemHeaderAlignment = "Left"))
+	ENSM_ShapePrimitive ShapePrimitive = ENSM_ShapePrimitive::Sphere;
+
+	UPROPERTY(EditAnywhere, Category = "Parameters", meta = (EditConditionHides, EditCondition = "ShapePrimitive == ENSM_ShapePrimitive::Box"))
+	FVector3f BoxSize = FVector3f(100.0f);
+	UPROPERTY(EditAnywhere, Category = "Parameters", meta = (EditConditionHides, EditCondition = "ShapePrimitive == ENSM_ShapePrimitive::Box"))
+	bool bBoxSurfaceOnly = false;
+	UPROPERTY(EditAnywhere, Category = "Parameters", meta = (EditConditionHides, EditCondition = "ShapePrimitive == ENSM_ShapePrimitive::Box"))
+	float BoxSurfaceThicknessMin = 0.0f;
+	UPROPERTY(EditAnywhere, Category = "Parameters", meta = (EditConditionHides, EditCondition = "ShapePrimitive == ENSM_ShapePrimitive::Box"))
+	float BoxSurfaceThicknessMax = 0.0f;
+
+	UPROPERTY(EditAnywhere, Category = "Parameters", meta = (EditConditionHides, EditCondition = "ShapePrimitive == ENSM_ShapePrimitive::Plane"))
+	FVector2f PlaneSize = FVector2f(100.0f);
+
+	UPROPERTY(EditAnywhere, Category = "Parameters", meta = (EditConditionHides, EditCondition = "ShapePrimitive == ENSM_ShapePrimitive::Cylinder"))
+	float CylinderHeight = 100.0f;
+	UPROPERTY(EditAnywhere, Category = "Parameters", meta = (EditConditionHides, EditCondition = "ShapePrimitive == ENSM_ShapePrimitive::Cylinder"))
+	float CylinderRadius = 100.0f;
+	UPROPERTY(EditAnywhere, Category = "Parameters", meta = (EditConditionHides, EditCondition = "ShapePrimitive == ENSM_ShapePrimitive::Cylinder"))
+	float CylinderHeightMidpoint = 0.5f;
+
+	UPROPERTY(EditAnywhere, Category = "Parameters", meta = (EditConditionHides, EditCondition = "ShapePrimitive == ENSM_ShapePrimitive::Ring"))
+	float RingRadius = 100.0f;
+	UPROPERTY(EditAnywhere, Category = "Parameters", meta = (EditConditionHides, EditCondition = "ShapePrimitive == ENSM_ShapePrimitive::Ring", ClampMin = "0.0", ClampMax = "1.0", UIMin = "0.0", UIMax = "1.0"))
+	float DiscCoverage = 0.0f;
+	UPROPERTY(EditAnywhere, Category = "Parameters", meta = (EditConditionHides, EditCondition = "ShapePrimitive == ENSM_ShapePrimitive::Ring", ClampMin = "0.0", ClampMax = "1.0", UIMin = "0.0", UIMax = "1.0"))
+	float RingUDistribution = 0.0f;
+
+	UPROPERTY(EditAnywhere, Category = "Parameters", meta=(EditConditionHides, EditCondition="ShapePrimitive == ENSM_ShapePrimitive::Sphere"))
+	float SphereMin = 0.0f;
+	UPROPERTY(EditAnywhere, Category = "Parameters", meta=(EditConditionHides, EditCondition="ShapePrimitive == ENSM_ShapePrimitive::Sphere"))
+	float SphereMax = 100.0f;
+
+	virtual void SetShaderParameters(const FNiagaraStatelessSetShaderParameterContext& SetShaderParameterContext) const override;
+
+#if WITH_EDITOR
+	virtual bool CanDisableModule() const override { return true; }
+
+	virtual bool CanDebugDraw() const { return true; }
+	virtual void DrawDebug(const FNiagaraStatelessDrawDebugContext& DrawDebugContext) const override;
+#endif
+#if WITH_EDITORONLY_DATA
+	virtual void GetOutputVaruables(TArray<FNiagaraVariableBase>& OutVariables) const override
+	{
+		const FNiagaraStatelessGlobals& StatelessGlobals = FNiagaraStatelessGlobals::Get();
+		OutVariables.AddUnique(StatelessGlobals.PositionVariable);
+		OutVariables.AddUnique(StatelessGlobals.PreviousPositionVariable);
+	}
+#endif
+};
diff --git a/Engine/Plugins/FX/Niagara/Source/Niagara/Internal/Stateless/Modules/NiagaraStatelessModule_SolveVelocitiesAndForces.h b/Engine/Plugins/FX/Niagara/Source/Niagara/Internal/Stateless/Modules/NiagaraStatelessModule_SolveVelocitiesAndForces.h
new file mode 100644
--- /dev/null
+++ b/Engine/Plugins/FX/Niagara/Source/Niagara/Internal/Stateless/Modules/NiagaraStatelessModule_SolveVelocitiesAndForces.h
@@ -0,0 +1,148 @@
+// Copyright Epic Games, Inc. All Rights Reserved.
+
+#pragma once
+
+#include "Stateless/NiagaraStatelessModule.h"
+#include "Stateless/NiagaraStatelessEmitterDataBuildContext.h"
+#include "Stateless/NiagaraStatelessModuleShaderParameters.h"
+#include "Stateless/Modules/NiagaraStatelessModuleCommon.h"
+
+#include "VectorField/VectorField.h"
+#include "VectorField/VectorFieldStatic.h"
+#include "RHIStaticStates.h"
+
+#include "NiagaraStatelessModule_SolveVelocitiesAndForces.generated.h"
+
+UCLASS(MinimalAPI, EditInlineNew, meta = (DisplayName = "Solve Forces And Velocity"))
+class UNiagaraStatelessModule_SolveVelocitiesAndForces : public UNiagaraStatelessModule
+{
+	GENERATED_BODY()
+
+public:
+	using FModuleBuiltData = NiagaraStateless::FPhysicsBuildData;
+	using FParameters = NiagaraStateless::FSolveVelocitiesAndForcesModule_ShaderParameters;
+
+	virtual void BuildEmitterData(FNiagaraStatelessEmitterDataBuildContext& BuildContext) const override
+	{
+		NiagaraStateless::FPhysicsBuildData& PhysicsBuildData = BuildContext.GetTransientBuildData<NiagaraStateless::FPhysicsBuildData>();
+
+		FModuleBuiltData* BuiltData		= BuildContext.AllocateBuiltData<FModuleBuiltData>();
+		BuiltData->MassMin				= PhysicsBuildData.MassMin;
+		BuiltData->MassMax				= PhysicsBuildData.MassMax;
+		BuiltData->DragMin				= FMath::Clamp(PhysicsBuildData.DragMin, 0.01f, 1.0f);
+		BuiltData->DragMax				= FMath::Clamp(PhysicsBuildData.DragMax, 0.01f, 1.0f);
+		BuiltData->VelocityMin			= PhysicsBuildData.VelocityMin;
+		BuiltData->VelocityMax			= PhysicsBuildData.VelocityMax;
+		BuiltData->WindMin				= PhysicsBuildData.WindMin;
+		BuiltData->WindMax				= PhysicsBuildData.WindMax;
+		BuiltData->AccelerationMin		= PhysicsBuildData.AccelerationMin;
+		BuiltData->AccelerationMax		= PhysicsBuildData.AccelerationMax;
+		BuiltData->bConeVelocity		= PhysicsBuildData.bConeVelocity;
+		BuiltData->ConeQuat				= PhysicsBuildData.ConeQuat;
+		BuiltData->ConeVelocityMin		= PhysicsBuildData.ConeVelocityMin;
+		BuiltData->ConeVelocityMax		= PhysicsBuildData.ConeVelocityMax;
+		BuiltData->ConeOuterAngle		= PhysicsBuildData.ConeOuterAngle;
+		BuiltData->ConeInnerAngle		= PhysicsBuildData.ConeInnerAngle;
+		BuiltData->ConeVelocityFalloff	= PhysicsBuildData.ConeVelocityFalloff;
+		BuiltData->bPointVelocity		= PhysicsBuildData.bPointVelocity;
+		BuiltData->PointVelocityMin		= PhysicsBuildData.PointVelocityMin;
+		BuiltData->PointVelocityMax		= PhysicsBuildData.PointVelocityMax;
+		BuiltData->PointOrigin			= PhysicsBuildData.PointOrigin;
+		BuiltData->bNoiseEnabled		= PhysicsBuildData.bNoiseEnabled;
+		BuiltData->NoiseAmplitude		= PhysicsBuildData.NoiseAmplitude;
+		BuiltData->NoiseFrequency		= PhysicsBuildData.NoiseFrequency;
+		BuiltData->NoiseTexture			= PhysicsBuildData.NoiseTexture;
+		BuiltData->NoiseMode			= PhysicsBuildData.NoiseMode;
+		BuiltData->NoiseLUTOffset		= PhysicsBuildData.NoiseLUTOffset;
+		BuiltData->NoiseLUTNumChannel	= PhysicsBuildData.NoiseLUTNumChannel;
+		BuiltData->NoiseLUTChannelWidth	= PhysicsBuildData.NoiseLUTChannelWidth;
+	}
+
+	virtual void SetShaderParameters(const FNiagaraStatelessSetShaderParameterContext& SetShaderParameterContext) const override
+	{
+		const FModuleBuiltData* ModuleBuiltData = SetShaderParameterContext.ReadBuiltData<FModuleBuiltData>();
+
+		FParameters* Parameters = SetShaderParameterContext.GetParameterNestedStruct<FParameters>();
+		Parameters->SolveVelocitiesAndForces_MassScale				= ModuleBuiltData->MassMax - ModuleBuiltData->MassMin;
+		Parameters->SolveVelocitiesAndForces_MassBias				= ModuleBuiltData->MassMin;
+		Parameters->SolveVelocitiesAndForces_DragScale				= ModuleBuiltData->DragMax - ModuleBuiltData->DragMin;
+		Parameters->SolveVelocitiesAndForces_DragBias				= ModuleBuiltData->DragMin;
+		Parameters->SolveVelocitiesAndForces_VelocityScale			= ModuleBuiltData->VelocityMax - ModuleBuiltData->VelocityMin;
+		Parameters->SolveVelocitiesAndForces_VelocityBias			= ModuleBuiltData->VelocityMin;
+		Parameters->SolveVelocitiesAndForces_WindScale				= ModuleBuiltData->WindMax - ModuleBuiltData->WindMin;
+		Parameters->SolveVelocitiesAndForces_WindBias				= ModuleBuiltData->WindMin;
+		Parameters->SolveVelocitiesAndForces_AccelerationScale		= ModuleBuiltData->AccelerationMax - ModuleBuiltData->AccelerationMin;
+		Parameters->SolveVelocitiesAndForces_AccelerationBias		= ModuleBuiltData->AccelerationMin;
+
+		Parameters->SolveVelocitiesAndForces_ConeVelocityEnabled	= ModuleBuiltData->bConeVelocity ? 1 : 0;
+		Parameters->SolveVelocitiesAndForces_ConeQuat				= ModuleBuiltData->ConeQuat;
+		Parameters->SolveVelocitiesAndForces_ConeVelocityScale		= ModuleBuiltData->ConeVelocityMax - ModuleBuiltData->ConeVelocityMin;
+		Parameters->SolveVelocitiesAndForces_ConeVelocityBias		= ModuleBuiltData->ConeVelocityMin;
+		Parameters->SolveVelocitiesAndForces_ConeAngleScale			= (ModuleBuiltData->ConeOuterAngle - ModuleBuiltData->ConeInnerAngle) * (UE_PI / 360.0f);
+		Parameters->SolveVelocitiesAndForces_ConeAngleBias			= ModuleBuiltData->ConeInnerAngle * (UE_PI / 360.0f);
+		Parameters->SolveVelocitiesAndForces_ConeVelocityFalloff	= ModuleBuiltData->ConeVelocityFalloff;
+
+		Parameters->SolveVelocitiesAndForces_PontVelocityEnabled	= ModuleBuiltData->bPointVelocity ? 1 : 0;
+		Parameters->SolveVelocitiesAndForces_PointVelocityScale		= ModuleBuiltData->PointVelocityMax - ModuleBuiltData->PointVelocityMin;
+		Parameters->SolveVelocitiesAndForces_PointVelocityBias		= ModuleBuiltData->PointVelocityMax;
+		Parameters->SolveVelocitiesAndForces_PointOrigin			= ModuleBuiltData->PointOrigin;
+
+		Parameters->SolveVelocitiesAndForces_NoiseEnabled			= ModuleBuiltData->bNoiseEnabled ? 1 : 0;
+		Parameters->SolveVelocitiesAndForces_NoiseAmplitude			= ModuleBuiltData->NoiseAmplitude;
+		Parameters->SolveVelocitiesAndForces_NoiseFrequency			= FVector3f(ModuleBuiltData->NoiseFrequency, ModuleBuiltData->NoiseFrequency, ModuleBuiltData->NoiseFrequency);
+		//SetShaderParameterContext.SetTextureResource(&Parameters->SolveVelocitiesAndForces_NoiseTexture, ModuleBuildData->NoiseTexture);
+		Parameters->SolveVelocitiesAndForces_NoiseMode				= ModuleBuiltData->NoiseMode;
+		Parameters->SolveVelocitiesAndForces_NoiseLUTOffset			= ModuleBuiltData->NoiseLUTOffset;
+		Parameters->SolveVelocitiesAndForces_NoiseLUTNumChannel		= ModuleBuiltData->NoiseLUTNumChannel;
+		Parameters->SolveVelocitiesAndForces_NoiseLUTChannelWidth	= ModuleBuiltData->NoiseLUTChannelWidth;
+
+		FVectorFieldTextureAccessor TextureAccessor(Cast<UVectorField>(ModuleBuiltData->NoiseTexture));
+
+		ENQUEUE_RENDER_COMMAND(FNaughtyTest)(
+			[Parameters, TextureAccessor](FRHICommandListImmediate& RHICmdList)
+			{
+				FRHITexture* NoiseTextureRHI = TextureAccessor.GetTexture();
+				Parameters->SolveVelocitiesAndForces_NoiseTexture = NoiseTextureRHI;
+				Parameters->SolveVelocitiesAndForces_NoiseSampler = TStaticSamplerState<SF_Bilinear, AM_Wrap, AM_Wrap, AM_Wrap>::GetRHI();
+
+				if (Parameters->SolveVelocitiesAndForces_NoiseMode < 2)
+				{
+					const FIntVector TextureSize = NoiseTextureRHI ? NoiseTextureRHI->GetSizeXYZ() : FIntVector(1, 1, 1);
+					Parameters->SolveVelocitiesAndForces_NoiseFrequency.X *= 1.0f / float(TextureSize.X);
+					Parameters->SolveVelocitiesAndForces_NoiseFrequency.Y *= 1.0f / float(TextureSize.Y);
+					Parameters->SolveVelocitiesAndForces_NoiseFrequency.Z *= 1.0f / float(TextureSize.Z);
+				}
+			}
+		);
+
+		//ModuleBuildData->NoiseTexture
+		//SetShaderParameterContext.AddSetResource(
+		//	[]()
+		//	{
+		//	}
+		//);
+
+		//FRHITexture**		TextureRef = &Parameters->SolveVelocitiesAndForces_NoiseTexture;
+		//FRHISamplerState**	SamplerRef = &Parameters->SolveVelocitiesAndForces_NoiseSampler;
+		//FTextureRHIRef
+		//if (ModuleBuiltData->NoiseTexture)
+		//{
+		//}
+		//SetShaderParameterContext.AddRenderThreadSet(
+		//	[]()
+		//	{
+		//	}
+		//);
+	}
+
+#if WITH_EDITORONLY_DATA
+	virtual void GetOutputVaruables(TArray<FNiagaraVariableBase>& OutVariables) const override
+	{
+		const FNiagaraStatelessGlobals& StatelessGlobals = FNiagaraStatelessGlobals::Get();
+		OutVariables.AddUnique(StatelessGlobals.PositionVariable);
+		OutVariables.AddUnique(StatelessGlobals.VelocityVariable);
+		OutVariables.AddUnique(StatelessGlobals.PreviousPositionVariable);
+		OutVariables.AddUnique(StatelessGlobals.PreviousVelocityVariable);
+	}
+#endif
+};
diff --git a/Engine/Plugins/FX/Niagara/Source/Niagara/Internal/Stateless/Modules/NiagaraStatelessModule_SpriteFacingAndAlignment.h b/Engine/Plugins/FX/Niagara/Source/Niagara/Internal/Stateless/Modules/NiagaraStatelessModule_SpriteFacingAndAlignment.h
new file mode 100644
--- /dev/null
+++ b/Engine/Plugins/FX/Niagara/Source/Niagara/Internal/Stateless/Modules/NiagaraStatelessModule_SpriteFacingAndAlignment.h
@@ -0,0 +1,60 @@
+// Copyright Epic Games, Inc. All Rights Reserved.
+
+#pragma once
+
+#include "Stateless/NiagaraStatelessModule.h"
+#include "Stateless/NiagaraStatelessEmitterDataBuildContext.h"
+#include "Stateless/NiagaraStatelessModuleShaderParameters.h"
+
+#include "NiagaraStatelessModule_SpriteFacingAndAlignment.generated.h"
+
+UCLASS(MinimalAPI, EditInlineNew, meta = (DisplayName = "Sprite Facing Alignment"))
+class UNiagaraStatelessModule_SpriteFacingAndAlignment : public UNiagaraStatelessModule
+{
+	GENERATED_BODY()
+
+public:
+	using FParameters = NiagaraStateless::FSpriteFacingAndAlignmentModule_ShaderParameters;
+
+	UPROPERTY(EditAnywhere, Category = "Parameters", meta = (InlineEditConditionToggle))
+	bool bSpriteFacingEnabled = true;
+
+	UPROPERTY(EditAnywhere, Category = "Parameters", meta = (InlineEditConditionToggle))
+	bool bSpriteAlignmentEnabled = false;
+
+	UPROPERTY(EditAnywhere, Category = "Parameters", meta = (EditCondition = "bSpriteFacingEnabled"))
+	FVector3f	SpriteFacing = FVector3f::XAxisVector;
+
+	UPROPERTY(EditAnywhere, Category = "Parameters", meta = (EditCondition = "bSpriteAlignmentEnabled"))
+	FVector3f	SpriteAlignment = FVector3f::YAxisVector;
+
+	virtual void SetShaderParameters(const FNiagaraStatelessSetShaderParameterContext& SetShaderParameterContext) const override
+	{
+		FParameters* Parameters = SetShaderParameterContext.GetParameterNestedStruct<FParameters>();
+		Parameters->SpriteFacingAndAlignment_SpriteFacing		= SpriteFacing;
+		Parameters->SpriteFacingAndAlignment_SpriteAlignment	= SpriteAlignment;
+	}
+
+#if WITH_EDITOR
+	virtual bool CanDisableModule() const override { return true; }
+#endif
+#if WITH_EDITORONLY_DATA
+	virtual void GetOutputVaruables(TArray<FNiagaraVariableBase>& OutVariables) const override
+	{
+		const FNiagaraStatelessGlobals& StatelessGlobals = FNiagaraStatelessGlobals::Get();
+		if (IsModuleEnabled())
+		{
+			if (bSpriteFacingEnabled)
+			{
+				OutVariables.AddUnique(StatelessGlobals.SpriteFacingVariable);
+				OutVariables.AddUnique(StatelessGlobals.PreviousSpriteFacingVariable);
+			}
+			if (bSpriteAlignmentEnabled)
+			{
+				OutVariables.AddUnique(StatelessGlobals.SpriteAlignmentVariable);
+				OutVariables.AddUnique(StatelessGlobals.PreviousSpriteAlignmentVariable);
+			}
+		}
+	}
+#endif
+};
diff --git a/Engine/Plugins/FX/Niagara/Source/Niagara/Internal/Stateless/Modules/NiagaraStatelessModule_SpriteRotationRate.h b/Engine/Plugins/FX/Niagara/Source/Niagara/Internal/Stateless/Modules/NiagaraStatelessModule_SpriteRotationRate.h
new file mode 100644
--- /dev/null
+++ b/Engine/Plugins/FX/Niagara/Source/Niagara/Internal/Stateless/Modules/NiagaraStatelessModule_SpriteRotationRate.h
@@ -0,0 +1,51 @@
+// Copyright Epic Games, Inc. All Rights Reserved.
+
+#pragma once
+
+#include "Stateless/NiagaraStatelessModule.h"
+#include "Stateless/NiagaraStatelessEmitterDataBuildContext.h"
+#include "Stateless/NiagaraStatelessModuleShaderParameters.h"
+
+#include "NiagaraStatelessModule_SpriteRotationRate.generated.h"
+
+UCLASS(MinimalAPI, EditInlineNew, meta = (DisplayName = "Sprite Rotation Rate"))
+class UNiagaraStatelessModule_SpriteRotationRate : public UNiagaraStatelessModule
+{
+	GENERATED_BODY()
+
+public:
+	using FParameters = NiagaraStateless::FSpriteRotationRateModule_ShaderParameters;
+
+	UPROPERTY(EditAnywhere, Category = "Parameters")
+	float	RotationRateMin = 360.0f;
+
+	UPROPERTY(EditAnywhere, Category = "Parameters")
+	float	RotationRateMax = 360.0f;
+
+	virtual void SetShaderParameters(const FNiagaraStatelessSetShaderParameterContext& SetShaderParameterContext) const override
+	{
+		FParameters* Parameters = SetShaderParameterContext.GetParameterNestedStruct<FParameters>();
+		if (IsModuleEnabled())
+		{
+			Parameters->SpriteRotationRate_Scale	= RotationRateMax - RotationRateMin;
+			Parameters->SpriteRotationRate_Bias		= RotationRateMin;
+		}
+		else
+		{
+			Parameters->SpriteRotationRate_Scale	= 0.0f;
+			Parameters->SpriteRotationRate_Bias		= 0.0f;
+		}
+	}
+
+#if WITH_EDITOR
+	virtual bool CanDisableModule() const override { return true; }
+#endif
+#if WITH_EDITORONLY_DATA
+	virtual void GetOutputVaruables(TArray<FNiagaraVariableBase>& OutVariables) const override
+	{
+		const FNiagaraStatelessGlobals& StatelessGlobals = FNiagaraStatelessGlobals::Get();
+		OutVariables.AddUnique(StatelessGlobals.SpriteRotationVariable);
+		OutVariables.AddUnique(StatelessGlobals.PreviousSpriteRotationVariable);
+	}
+#endif
+};
diff --git a/Engine/Plugins/FX/Niagara/Source/Niagara/Internal/Stateless/Modules/NiagaraStatelessModule_SubUVAnimation.h b/Engine/Plugins/FX/Niagara/Source/Niagara/Internal/Stateless/Modules/NiagaraStatelessModule_SubUVAnimation.h
new file mode 100644
--- /dev/null
+++ b/Engine/Plugins/FX/Niagara/Source/Niagara/Internal/Stateless/Modules/NiagaraStatelessModule_SubUVAnimation.h
@@ -0,0 +1,120 @@
+// Copyright Epic Games, Inc. All Rights Reserved.
+
+#pragma once
+
+#include "Stateless/NiagaraStatelessModule.h"
+#include "Stateless/NiagaraStatelessEmitterDataBuildContext.h"
+#include "Stateless/NiagaraStatelessModuleShaderParameters.h"
+
+#include "NiagaraStatelessModule_SubUVAnimation.generated.h"
+
+UENUM()
+enum class ENSMSubUVAnimation_Mode
+{
+	InfiniteLoop,
+	Linear,
+	Random,
+};
+
+UCLASS(MinimalAPI, EditInlineNew, meta = (DisplayName = "Sub UV Animation"))
+class UNiagaraStatelessModule_SubUVAnimation : public UNiagaraStatelessModule
+{
+	GENERATED_BODY()
+
+public:
+	using FParameters = NiagaraStateless::FSubUVAnimationModule_ShaderParameters;
+
+	UPROPERTY(EditAnywhere, Category = "Parameters")
+	int32	NumFrames = 16;
+
+	UPROPERTY(EditAnywhere, Category = "Parameters", meta = (InlineEditConditionToggle))
+	bool	bStartFrameRangeOverride_Enabled = false;
+
+	UPROPERTY(EditAnywhere, Category = "Parameters", meta = (InlineEditConditionToggle))
+	bool	bEndFrameRangeOverride_Enabled = false;
+
+	UPROPERTY(EditAnywhere, Category = "Parameters", meta = (EditCondition = "bStartFrameRangeOverride_Enabled"))
+	int32	StartFrameRangeOverride = 0;
+
+	UPROPERTY(EditAnywhere, Category = "Parameters", meta = (EditCondition = "bEndFrameRangeOverride_Enabled"))
+	int32	EndFrameRangeOverride = 0;
+	
+	UPROPERTY(EditAnywhere, Category = "Parameters")
+	ENSMSubUVAnimation_Mode AnimationMode = ENSMSubUVAnimation_Mode::Linear;
+
+	//-Note: Main module has PlaybackMode (Loops / FPS) to choose between loops or frames per second
+	UPROPERTY(EditAnywhere, Category = "Parameters", meta = (EditConditionHides, EditCondition = "AnimationMode == ENSMSubUVAnimation_Mode::InfiniteLoop"))
+	float LoopsPerSecond = 1.0f;
+
+	//-Note: Main module has a few more options
+	//UPROPERTY(EditAnywhere, Category = "Parameters", meta = (EditConditionHides, EditCondition = "AnimationMode == ENSMSubUVAnimation_Mode::Linear"))
+	//bool bRandomStartFrame = false;
+	//int32 StartFrameOffset = 0;
+	//float LoopupIndexScale = 0.0f;
+
+	UPROPERTY(EditAnywhere, Category = "Parameters", meta = (EditConditionHides, EditCondition = "AnimationMode == ENSMSubUVAnimation_Mode::Random"))
+	float RandomChangeInterval = 0.1f;
+
+	virtual void SetShaderParameters(const FNiagaraStatelessSetShaderParameterContext& SetShaderParameterContext) const override
+	{
+		FParameters* Parameters = SetShaderParameterContext.GetParameterNestedStruct<FParameters>();
+		if (IsModuleEnabled())
+		{
+			const float FrameRangeStart	= bStartFrameRangeOverride_Enabled ? FMath::Clamp(float(StartFrameRangeOverride) / float(NumFrames - 1), 0.0f, 1.0f) : 0.0f;
+			const float FrameRangeEnd	= bEndFrameRangeOverride_Enabled   ? FMath::Clamp(float(EndFrameRangeOverride)   / float(NumFrames - 1), 0.0f, 1.0f) : 1.0f;
+
+			Parameters->SubUVAnimation_Mode			= (int)AnimationMode;
+			Parameters->SubUVAnimation_NumFrames	= float(NumFrames);
+			switch (AnimationMode)
+			{
+				case ENSMSubUVAnimation_Mode::InfiniteLoop:
+					Parameters->SubUVAnimation_InitialFrameScale		= 0.0f;
+					Parameters->SubUVAnimation_InitialFrameBias			= 0.0f;
+					Parameters->SubUVAnimation_InitialFrameRateChange	= 0.0f;
+					Parameters->SubUVAnimation_AnimFrameStart			= FrameRangeStart;
+					Parameters->SubUVAnimation_AnimFrameRange			= FrameRangeEnd - FrameRangeStart;
+					Parameters->SubUVAnimation_RateScale				= LoopsPerSecond;
+					break;
+
+				case ENSMSubUVAnimation_Mode::Linear:
+					Parameters->SubUVAnimation_InitialFrameScale		= 0.0f;
+					Parameters->SubUVAnimation_InitialFrameBias			= 0.0f;
+					Parameters->SubUVAnimation_InitialFrameRateChange	= 0.0f;
+					Parameters->SubUVAnimation_AnimFrameStart			= FrameRangeStart;
+					Parameters->SubUVAnimation_AnimFrameRange			= FrameRangeEnd - FrameRangeStart;
+					Parameters->SubUVAnimation_RateScale				= 1.0f;
+					break;
+
+				case ENSMSubUVAnimation_Mode::Random:
+					Parameters->SubUVAnimation_InitialFrameScale		= FrameRangeEnd - FrameRangeStart;
+					Parameters->SubUVAnimation_InitialFrameBias			= FrameRangeStart;
+					Parameters->SubUVAnimation_InitialFrameRateChange	= RandomChangeInterval > 0.0f ? 1.0f / RandomChangeInterval : 0.0f;
+					Parameters->SubUVAnimation_AnimFrameStart			= 0.0f;
+					Parameters->SubUVAnimation_AnimFrameRange			= 0.0f;
+					Parameters->SubUVAnimation_RateScale				= 0.0f;
+					break;
+			}
+		}
+		else
+		{
+			Parameters->SubUVAnimation_Mode						= 0;
+			Parameters->SubUVAnimation_InitialFrameScale		= 0.0f;
+			Parameters->SubUVAnimation_InitialFrameBias			= 0.0f;
+			Parameters->SubUVAnimation_InitialFrameRateChange	= 0.0f;
+			Parameters->SubUVAnimation_AnimFrameStart			= 0.0f;
+			Parameters->SubUVAnimation_AnimFrameRange			= 0.0f;
+			Parameters->SubUVAnimation_RateScale				= 0.0f;
+		}
+	}
+
+#if WITH_EDITOR
+	virtual bool CanDisableModule() const override { return true; }
+#endif
+#if WITH_EDITORONLY_DATA
+	virtual void GetOutputVaruables(TArray<FNiagaraVariableBase>& OutVariables) const override
+	{
+		const FNiagaraStatelessGlobals& StatelessGlobals = FNiagaraStatelessGlobals::Get();
+		OutVariables.AddUnique(StatelessGlobals.SubImageIndexVariable);
+	}
+#endif
+};
diff --git a/Engine/Plugins/FX/Niagara/Source/Niagara/Internal/Stateless/NiagaraStatelessCommon.h b/Engine/Plugins/FX/Niagara/Source/Niagara/Internal/Stateless/NiagaraStatelessCommon.h
new file mode 100644
--- /dev/null
+++ b/Engine/Plugins/FX/Niagara/Source/Niagara/Internal/Stateless/NiagaraStatelessCommon.h
@@ -0,0 +1,95 @@
+// Copyright Epic Games, Inc. All Rights Reserved.
+
+#pragma once
+
+#include "NiagaraCommon.h"
+
+#include "NiagaraStatelessCommon.generated.h"
+
+UENUM()
+enum class ENiagaraStatelessSpawnInfoType
+{
+	Burst,
+	Rate
+};
+
+USTRUCT()
+struct FNiagaraStatelessSpawnInfo
+{
+	GENERATED_BODY()
+
+#if WITH_EDITORONLY_DATA
+	UPROPERTY()
+	FGuid SourceId;
+#endif
+
+	UPROPERTY(EditAnywhere, Category = "Spawn", meta = (ShowInStackItemHeader, StackItemHeaderAlignment = "Left"))
+	ENiagaraStatelessSpawnInfoType	Type = ENiagaraStatelessSpawnInfoType::Burst;
+
+	UPROPERTY(EditAnywhere, Category = "Spawn", meta = (EditConditionHides, EditCondition = "Type == ENiagaraStatelessSpawnInfoType::Burst"))
+	float SpawnTime = 0.0f;
+
+	UPROPERTY(EditAnywhere, Category = "Spawn", meta = (EditConditionHides, EditCondition = "Type == ENiagaraStatelessSpawnInfoType::Burst", ClampMin = "0"))
+	int32 AmountMin = 1;
+
+	UPROPERTY(EditAnywhere, Category = "Spawn", meta = (EditConditionHides, EditCondition = "Type == ENiagaraStatelessSpawnInfoType::Burst", ClampMin = "0"))
+	int32 AmountMax = 1;
+
+	UPROPERTY(EditAnywhere, Category = "Spawn", meta = (EditConditionHides, EditCondition = "Type == ENiagaraStatelessSpawnInfoType::Rate", ClampMin = "0.0"))
+	float RateMin = 60.0f;
+
+	UPROPERTY(EditAnywhere, Category = "Spawn", meta = (EditConditionHides, EditCondition = "Type == ENiagaraStatelessSpawnInfoType::Rate", ClampMin = "0.0"))
+	float RateMax = 60.0f;
+
+	UPROPERTY(EditAnywhere, Category = "Spawn", meta = (InlineEditConditionToggle))
+	bool bSpawnProbabilityEnabled = false;
+
+	UPROPERTY(EditAnywhere, Category = "Spawn", meta = (EditCondition = "bSpawnProbabilityEnabled", ClampMin = "0.0", ClampMax = "1.0", UIMin = "0.0", UIMax = "1.0"))
+	float SpawnProbability = 1.0f;
+
+	bool IsValid(TOptional<float> LoopDuration) const;
+};
+
+struct FNiagaraStatelessRuntimeSpawnInfo
+{
+	ENiagaraStatelessSpawnInfoType Type = ENiagaraStatelessSpawnInfoType::Burst;
+	uint32	UniqueOffset	= 0;
+	float	SpawnTimeStart	= 0.0f;
+	float	SpawnTimeEnd	= 0.0f;
+	float	Rate			= 0.0f;
+	int32	Amount			= 0;
+};
+
+struct FNiagaraStatelessGlobals
+{
+	FNiagaraVariableBase	ColorVariable;
+	FNiagaraVariableBase	MeshOrientationVariable;
+	FNiagaraVariableBase	PositionVariable;
+	FNiagaraVariableBase	RibbonWidthVariable;
+	FNiagaraVariableBase	ScaleVariable;
+	FNiagaraVariableBase	SpriteAlignmentVariable;
+	FNiagaraVariableBase	SpriteFacingVariable;
+	FNiagaraVariableBase	SpriteSizeVariable;
+	FNiagaraVariableBase	SpriteRotationVariable;
+	FNiagaraVariableBase	SubImageIndexVariable;
+	FNiagaraVariableBase	UniqueIDVariable;
+	FNiagaraVariableBase	VelocityVariable;
+
+	//FNiagaraVariableBase	PreviousColorVariable;
+	FNiagaraVariableBase	PreviousMeshOrientationVariable;
+	FNiagaraVariableBase	PreviousPositionVariable;
+	FNiagaraVariableBase	PreviousRibbonWidthVariable;
+	FNiagaraVariableBase	PreviousScaleVariable;
+	FNiagaraVariableBase	PreviousSpriteAlignmentVariable;
+	FNiagaraVariableBase	PreviousSpriteFacingVariable;
+	FNiagaraVariableBase	PreviousSpriteSizeVariable;
+	FNiagaraVariableBase	PreviousSpriteRotationVariable;
+	FNiagaraVariableBase	PreviousVelocityVariable;
+
+	static const FNiagaraStatelessGlobals& Get();
+};
+
+namespace NiagaraStatelessCommon
+{
+	extern void Initialize();
+}
diff --git a/Engine/Plugins/FX/Niagara/Source/Niagara/Internal/Stateless/NiagaraStatelessDrawDebugContext.h b/Engine/Plugins/FX/Niagara/Source/Niagara/Internal/Stateless/NiagaraStatelessDrawDebugContext.h
new file mode 100644
--- /dev/null
+++ b/Engine/Plugins/FX/Niagara/Source/Niagara/Internal/Stateless/NiagaraStatelessDrawDebugContext.h
@@ -0,0 +1,21 @@
+// Copyright Epic Games, Inc. All Rights Reserved.
+
+#pragma once
+
+#include "NiagaraCommon.h"
+
+#if WITH_EDITOR
+struct FNiagaraStatelessDrawDebugContext
+{
+	UWorld*		World = nullptr;
+	FTransform	LocalToWorldTransform = FTransform::Identity;
+	FTransform	WorldTransform = FTransform::Identity;
+
+	void DrawArrow(const FVector& Origin, const FVector& DirectionWithLength, const FColor& Color = FColor::White) const;
+	void DrawBox(const FVector& Center, const FVector& Extent, const FColor& Color = FColor::White, const FQuat& Rotation = FQuat::Identity) const;
+	void DrawCone(const FVector& Origin, const FQuat& Rotation, float Angle, float Length, const FColor& Color = FColor::White) const;
+	void DrawCylinder(float CylinderHeight, float CylinderRadius, float CylinderHeightMidpoint, const FColor& Color = FColor::White) const;
+	void DrawCircle(const FVector& Center, const float Radius, const FColor& Color = FColor::White) const;
+	void DrawSphere(const FVector& Center, const float Radius, const FColor& Color = FColor::White) const;
+};
+#endif
diff --git a/Engine/Plugins/FX/Niagara/Source/Niagara/Internal/Stateless/NiagaraStatelessEmitterDataBuildContext.h b/Engine/Plugins/FX/Niagara/Source/Niagara/Internal/Stateless/NiagaraStatelessEmitterDataBuildContext.h
new file mode 100644
--- /dev/null
+++ b/Engine/Plugins/FX/Niagara/Source/Niagara/Internal/Stateless/NiagaraStatelessEmitterDataBuildContext.h
@@ -0,0 +1,74 @@
+// Copyright Epic Games, Inc. All Rights Reserved.
+
+#pragma once
+
+#include "NiagaraCommon.h"
+
+struct FNiagaraParameterBinding;
+struct FNiagaraParameterBindingWithValue;
+struct FNiagaraParameterStore;
+
+class FNiagaraStatelessEmitterDataBuildContext
+{
+public:
+	UE_NONCOPYABLE(FNiagaraStatelessEmitterDataBuildContext);
+
+	FNiagaraStatelessEmitterDataBuildContext(FNiagaraParameterStore& InRendererBindings, TArray<uint8>& InBuiltData, TArray<float>& InStaticFloatData)
+		: RendererBindings(InRendererBindings)
+		, BuiltData(InBuiltData)
+		, StaticFloatData(InStaticFloatData)
+	{
+	}
+
+	uint32 AddStaticData(TConstArrayView<float> FloatData);
+	uint32 AddStaticData(TConstArrayView<FVector2f> FloatData);
+	uint32 AddStaticData(TConstArrayView<FVector3f> FloatData);
+	uint32 AddStaticData(TConstArrayView<FVector4f> FloatData);
+	uint32 AddStaticData(TConstArrayView<FLinearColor> FloatData);
+
+	template<typename T>
+	T* AllocateBuiltData()
+	{
+		const int32 Align = BuiltData.Num() % alignof(T);
+		const int32 Offset = BuiltData.AddZeroed(sizeof(T) + Align);
+		return reinterpret_cast<T*>(BuiltData.GetData() + Offset + Align);
+	}
+
+	template<typename T>
+	T& GetTransientBuildData()
+	{
+		TUniquePtr<FTransientObject>& TransientObj = TransientBuildData.FindOrAdd(T::GetName());
+		if (TransientObj.IsValid() == false)
+		{
+			TransientObj.Reset(new TTransientObject<T>);
+		}
+		return *reinterpret_cast<T*>(TransientObj->GetObject());
+	}
+	//NiagaraStateless::FPhysicsBuildData& PhysicsBuildData = BuildContext.GetTransientBuildData<NiagaraStateless::FPhysicsBuildData>();
+
+	int32 AddRendererBinding(const FNiagaraParameterBinding& Binding);
+	int32 AddRendererBinding(const FNiagaraParameterBindingWithValue& Binding);	
+
+private:
+	FNiagaraParameterStore& RendererBindings;
+	TArray<uint8>&			BuiltData;
+	TArray<float>&			StaticFloatData;
+
+	struct FTransientObject
+	{
+		virtual ~FTransientObject() = default;
+		virtual void* GetObject() = 0;
+	};
+
+	template <typename T>
+	struct TTransientObject final : FTransientObject
+	{
+		template <typename... TArgs>
+		FORCEINLINE TTransientObject(TArgs&&... Args) : TheObject(Forward<TArgs&&>(Args)...) {}
+		virtual void* GetObject() { return &TheObject; }
+
+		T TheObject;
+	};
+
+	TMap<FName, TUniquePtr<FTransientObject>>	TransientBuildData;
+};
diff --git a/Engine/Plugins/FX/Niagara/Source/Niagara/Internal/Stateless/NiagaraStatelessEmitterTemplate.h b/Engine/Plugins/FX/Niagara/Source/Niagara/Internal/Stateless/NiagaraStatelessEmitterTemplate.h
new file mode 100644
--- /dev/null
+++ b/Engine/Plugins/FX/Niagara/Source/Niagara/Internal/Stateless/NiagaraStatelessEmitterTemplate.h
@@ -0,0 +1,76 @@
+// Copyright Epic Games, Inc. All Rights Reserved.
+
+#pragma once
+
+#include "NiagaraCommon.h"
+#include "Stateless/NiagaraStatelessSimulationShader.h"
+
+#include "RenderGraphFwd.h"
+#include "ShaderParameterStruct.h"
+
+#include "NiagaraStatelessEmitterTemplate.generated.h"
+
+UCLASS(MinimalAPI, abstract)
+class UNiagaraStatelessEmitterTemplate : public UObject
+{
+	GENERATED_BODY()
+
+public:
+	//~Begin: UObject Interface
+	virtual void PostInitProperties() override;
+	//~End: UObject Interface
+
+	static void InitCDOPropertiesAfterModuleStartup();
+	virtual void InitModulesAndAttributes() {}
+
+#if WITH_EDITORONLY_DATA
+	TConstArrayView<TObjectPtr<UClass>> GetModules() const { return Modules; }
+	TConstArrayView<FNiagaraVariableBase> GetOututputComponents() const { return OutputComponents; }
+#endif
+
+	virtual const FShaderParametersMetadata* GetShaderParametersMetadata() const { checkNoEntry(); return nullptr; }
+	virtual TShaderRef<NiagaraStateless::FSimulationShader> GetSimulationShader() const { checkNoEntry(); return TShaderRef<NiagaraStateless::FSimulationShader>(); }
+
+	//-TODO: Add a way to set them directly, we should know that the final struct is a series of ints in the order of the provided variables
+	virtual void SetShaderParameters(uint8* ShaderParametersBase, TConstArrayView<int32> ComponentOffsets) const { checkNoEntry(); }
+	//-TODO: Add a way to set them directly, we should know that the final struct is a series of ints in the order of the provided variables
+
+	static FNiagaraVariableBase MakeProvidedVariable(FNiagaraVariableBase Variable);
+	static FNiagaraVariableBase MakeProvidedPreviousVariable(FNiagaraVariableBase Variable);
+
+protected:
+#if WITH_EDITORONLY_DATA
+	TArray<TObjectPtr<UClass>>		Modules;
+	TArray<FNiagaraVariableBase>	OutputComponents;
+#endif
+};
+
+UCLASS(Category = "Statless", DisplayName = "Emitter Example 0")
+class UNiagaraStatelessEmitterDefault : public UNiagaraStatelessEmitterTemplate
+{
+	GENERATED_BODY()
+
+	virtual void InitModulesAndAttributes() override;
+
+	virtual const FShaderParametersMetadata* GetShaderParametersMetadata() const override;
+	virtual TShaderRef<NiagaraStateless::FSimulationShader> GetSimulationShader() const override;
+
+	//-TODO: Add a way to set them directly, we should know that the final struct is a series of ints in the order of the provided variables
+	virtual void SetShaderParameters(uint8* ShaderParametersBase, TConstArrayView<int32> ComponentOffsets) const override;
+	//-TODO: Add a way to set them directly, we should know that the final struct is a series of ints in the order of the provided variables
+};
+
+UCLASS(Category = "Statless", DisplayName = "Emitter Example 1")
+class UNiagaraStatelessEmitterExample1 : public UNiagaraStatelessEmitterTemplate
+{
+	GENERATED_BODY()
+
+	virtual void InitModulesAndAttributes() override;
+
+	virtual const FShaderParametersMetadata* GetShaderParametersMetadata() const override;
+	virtual TShaderRef<NiagaraStateless::FSimulationShader> GetSimulationShader() const override;
+
+	//-TODO: Add a way to set them directly, we should know that the final struct is a series of ints in the order of the provided variables
+	virtual void SetShaderParameters(uint8* ShaderParametersBase, TConstArrayView<int32> ComponentOffsets) const override;
+	//-TODO: Add a way to set them directly, we should know that the final struct is a series of ints in the order of the provided variables
+};
diff --git a/Engine/Plugins/FX/Niagara/Source/Niagara/Internal/Stateless/NiagaraStatelessModule.h b/Engine/Plugins/FX/Niagara/Source/Niagara/Internal/Stateless/NiagaraStatelessModule.h
new file mode 100644
--- /dev/null
+++ b/Engine/Plugins/FX/Niagara/Source/Niagara/Internal/Stateless/NiagaraStatelessModule.h
@@ -0,0 +1,123 @@
+// Copyright Epic Games, Inc. All Rights Reserved.
+
+#pragma once
+
+#include "NiagaraStatelessCommon.h"
+
+#include "ShaderParameterStruct.h"
+
+#include "NiagaraStatelessModule.generated.h"
+
+class FNiagaraStatelessEmitterDataBuildContext;
+#if WITH_EDITOR
+struct FNiagaraStatelessDrawDebugContext;
+#endif
+
+////////////////////////////////////////////////////////////////////////////////////////////////////////////////////////////////////////////////////////////////
+
+class FNiagaraStatelessSetShaderParameterContext
+{
+public:
+	UE_NONCOPYABLE(FNiagaraStatelessSetShaderParameterContext);
+
+	explicit FNiagaraStatelessSetShaderParameterContext(TConstArrayView<uint8> InRendererParameterData, TConstArrayView<uint8> InBuiltData, const FShaderParametersMetadata* InShaderParametersMetadata, uint8* InShaderParameters)
+		: RendererParameterData(InRendererParameterData)
+		, BuiltData(InBuiltData)
+		, ShaderParametersBase(InShaderParameters)
+		, ParameterOffset(0)
+		, ShaderParametersMetadata(InShaderParametersMetadata)
+	{
+	}
+
+	template<typename T>
+	T* GetParameterNestedStruct() const
+	{
+		const uint32 StructOffset = Align(ParameterOffset, TShaderParameterStructTypeInfo<T>::Alignment);
+	#if DO_CHECK
+		ValidateIncludeStructType(StructOffset, TShaderParameterStructTypeInfo<T>::GetStructMetadata());
+	#endif
+		ParameterOffset = StructOffset + TShaderParameterStructTypeInfo<T>::GetStructMetadata()->GetSize();
+		return reinterpret_cast<T*>(ShaderParametersBase + StructOffset);
+	}
+
+	template<typename T>
+	const T* ReadBuiltData() const
+	{
+		const int32 Align = BuiltDataOffset % alignof(T);
+		const int32 Offset = BuiltDataOffset + Align;
+		BuiltDataOffset = Offset + sizeof(T);
+		check(BuiltDataOffset <= BuiltData.Num());
+		return reinterpret_cast<const T*>(BuiltData.GetData() + Offset);
+	}
+
+	template<typename T>
+	void GetRendererParameterValue(T& OutValue, int32 Offset, const T& DefaultValue) const
+	{
+		if (Offset != INDEX_NONE)
+		{
+			check(Offset >= 0 && Offset + sizeof(T) <= RendererParameterData.Num());
+			FMemory::Memcpy(&OutValue, RendererParameterData.GetData() + Offset, sizeof(T));
+		}
+		else
+		{
+			OutValue = DefaultValue;
+		}
+	}
+
+protected:
+#if DO_CHECK
+	void ValidateIncludeStructType(uint32 StructOffset, const FShaderParametersMetadata* StructMetaData) const;
+#endif
+
+private:
+	TConstArrayView<uint8>				RendererParameterData;
+	TConstArrayView<uint8>				BuiltData;
+	mutable int32						BuiltDataOffset = 0;
+	uint8*								ShaderParametersBase = nullptr;
+	mutable uint32						ParameterOffset = 0;
+	const FShaderParametersMetadata*	ShaderParametersMetadata = nullptr;
+};
+
+////////////////////////////////////////////////////////////////////////////////////////////////////////////////////////////////////////////////////////////////
+
+UCLASS(MinimalAPI, abstract, EditInlineNew)
+class UNiagaraStatelessModule : public UObject
+{
+	GENERATED_BODY()
+
+protected:
+	UPROPERTY(EditAnywhere, Category = "Parameters", meta = (DisplayPriority = 0, HideInStack))
+	uint32 bModuleEnabled : 1 = true;
+
+#if WITH_EDITORONLY_DATA
+	UPROPERTY(EditAnywhere, Transient, Category = "Parameters", meta = (DisplayPriority = 0, ShowInStackItemHeader, StackItemHeaderIcon = "Icons.Visibility"))
+	uint32 bDebugDrawEnabled : 1 = false;
+#endif
+
+public:
+	bool IsModuleEnabled() const { return bModuleEnabled; }
+#if WITH_EDITOR
+	bool IsDebugDrawEnabled() const { return bDebugDrawEnabled; }
+#endif
+
+	virtual void BuildEmitterData(FNiagaraStatelessEmitterDataBuildContext& BuildContext) const {}
+	virtual void SetShaderParameters(const FNiagaraStatelessSetShaderParameterContext& SetShaderParameterContext) const { }
+
+#if WITH_EDITOR
+	virtual bool CanDisableModule() const { return false; }
+	void SetIsModuleEnabled(bool bInIsEnabled) { bModuleEnabled = bInIsEnabled; }
+
+	virtual bool CanDebugDraw() const { return false; }
+	virtual void DrawDebug(const FNiagaraStatelessDrawDebugContext& DrawDebugContext) const {}
+#endif
+
+#if WITH_EDITORONLY_DATA
+	virtual void GetOutputVaruables(TArray<FNiagaraVariableBase>& OutVariables) const {}
+#endif
+
+	//~UObject interface Begin
+#if WITH_EDITOR
+	virtual bool CanEditChange(const FProperty* InProperty) const override;
+#endif //WITH_EDITOR
+	//~UObject interface End
+};
diff --git a/Engine/Plugins/FX/Niagara/Source/Niagara/Private/DataInterface/NiagaraDataInterfaceDataChannelRead.cpp b/Engine/Plugins/FX/Niagara/Source/Niagara/Private/DataInterface/NiagaraDataInterfaceDataChannelRead.cpp
--- a/Engine/Plugins/FX/Niagara/Source/Niagara/Private/DataInterface/NiagaraDataInterfaceDataChannelRead.cpp
+++ b/Engine/Plugins/FX/Niagara/Source/Niagara/Private/DataInterface/NiagaraDataInterfaceDataChannelRead.cpp
@@ -579,7 +579,7 @@ void UNiagaraDataInterfaceDataChannelRead::PostStageTick(FNDICpuPostStageContext
 					{
 						SpawnInfo.SpawnGroup = i;
 					}
-					//-TODO:Stateless:
+					//-TODO:Stateless: We need to handle the stateless path here
 					if (FNiagaraEmitterInstanceImpl* StatefulEmitter = TargetEmitter->AsStateful())
 					{
 						StatefulEmitter->GetSpawnInfo().Emplace(SpawnInfo);
diff --git a/Engine/Plugins/FX/Niagara/Source/Niagara/Private/DataInterface/NiagaraDataInterfaceUtilities.cpp b/Engine/Plugins/FX/Niagara/Source/Niagara/Private/DataInterface/NiagaraDataInterfaceUtilities.cpp
--- a/Engine/Plugins/FX/Niagara/Source/Niagara/Private/DataInterface/NiagaraDataInterfaceUtilities.cpp
+++ b/Engine/Plugins/FX/Niagara/Source/Niagara/Private/DataInterface/NiagaraDataInterfaceUtilities.cpp
@@ -246,7 +246,6 @@ void FNiagaraDataInterfaceUtilities::ForEachDataInterface(FNiagaraSystemInstance
 
 	for (const FNiagaraEmitterInstanceRef& Emitter : SystemInstance->GetEmitters())
 	{
-		//-TODO:Stateless:
 		FNiagaraEmitterInstanceImpl* StatefulEmitter = Emitter->AsStateful();
 		if (!StatefulEmitter || Emitter->IsDisabled())
 		{
@@ -289,7 +288,6 @@ void FNiagaraDataInterfaceUtilities::ForEachDataInterface(FNiagaraSystemInstance
 
 	for (const FNiagaraEmitterInstanceRef& Emitter : SystemInstance->GetEmitters())
 	{
-		//-TODO:Stateless:
 		FNiagaraEmitterInstanceImpl* StatefulEmitter = Emitter->AsStateful();
 		if (!StatefulEmitter || Emitter->IsDisabled())
 		{
diff --git a/Engine/Plugins/FX/Niagara/Source/Niagara/Private/NiagaraDebugHud.cpp b/Engine/Plugins/FX/Niagara/Source/Niagara/Private/NiagaraDebugHud.cpp
--- a/Engine/Plugins/FX/Niagara/Source/Niagara/Private/NiagaraDebugHud.cpp
+++ b/Engine/Plugins/FX/Niagara/Source/Niagara/Private/NiagaraDebugHud.cpp
@@ -1039,6 +1039,7 @@ void FNiagaraDebugHud::GatherSystemInfo()
 	#if WITH_EDITORONLY_DATA
 		SystemDebugInfo.bCompileForEdit = NiagaraComponent ? NiagaraComponent->GetAsset()->GetCompileForEdit() : false;
 	#endif
+		SystemDebugInfo.bSystemStateFastPath = NiagaraComponent ? NiagaraComponent->GetAsset()->SystemStateFastPathEnabled() : false;
 		SystemDebugInfo.bShowInWorld = Settings.bSystemFilterEnabled && SystemDebugInfo.SystemName.MatchesWildcard(Settings.SystemFilter);
 		SystemDebugInfo.bPassesSystemFilter = !Settings.bSystemFilterEnabled || SystemDebugInfo.SystemName.MatchesWildcard(Settings.SystemFilter);
 
@@ -1131,8 +1132,7 @@ void FNiagaraDebugHud::GatherSystemInfo()
 				check(SystemInstance);
 				for (const FNiagaraEmitterInstanceRef& EmitterInstance : SystemInstance->GetEmitters())
 				{
-					const UNiagaraEmitter* NiagaraEmitter = EmitterInstance->GetEmitter();
-					if (NiagaraEmitter == nullptr)
+					if (EmitterInstance->IsDisabled())
 					{
 						continue;
 					}
@@ -1627,13 +1627,19 @@ void FNiagaraDebugHud::DrawOverview(class FNiagaraWorldManager* WorldManager, FC
 				Canvas->DrawTile(X, Y, Col.MaxWidth, fAdvanceHeight, 0,0,0,0, RowBGColor);
 				const FLinearColor RowColor = SystemInfo.bShowInWorld ? DetailHighlightColor : DetailColor;
 
-				const bool bShowEditMode = 
-#if WITH_EDITORONLY_DATA
-					SystemInfo.bCompileForEdit;
-#else
-					false;
-#endif
-					Canvas->DrawShadowedString(X, Y, bShowEditMode ? *FString::Printf(TEXT("%s (Edit Mode)"), *SystemInfo.SystemName) : *SystemInfo.SystemName, Font, RowColor);
+				FNameBuilder SystemNameString;
+				SystemNameString.Append(*SystemInfo.SystemName);
+			#if WITH_EDITORONLY_DATA
+				if (SystemInfo.bCompileForEdit)
+				{
+					SystemNameString.Append(TEXT(" (Edit Mode)"));
+				}
+			#endif
+				if (SystemInfo.bSystemStateFastPath)
+				{
+					SystemNameString.Append(TEXT(" (Fast Path)"));
+				}
+				Canvas->DrawShadowedString(X, Y, SystemNameString.ToString(), Font, RowColor);
 		});
 
 		if (Settings.OverviewMode == ENiagaraDebugHUDOverviewMode::Overview)
@@ -2934,8 +2940,7 @@ void FNiagaraDebugHud::DrawComponents(FNiagaraWorldManager* WorldManager, UCanva
 					int32 ActiveParticles = 0;
 					for (const FNiagaraEmitterInstanceRef& EmitterInstance : SystemInstance->GetEmitters())
 					{
-						const UNiagaraEmitter* NiagaraEmitter = EmitterInstance->GetEmitter();
-						if (NiagaraEmitter == nullptr)
+						if (EmitterInstance->IsDisabled())
 						{
 							continue;
 						}
@@ -2949,13 +2954,14 @@ void FNiagaraDebugHud::DrawComponents(FNiagaraWorldManager* WorldManager, UCanva
 
 						if (Settings.SystemEmitterVerbosity == ENiagaraDebugHudVerbosity::Verbose)
 						{
+							const FString EmitterName = EmitterInstance->GetEmitterHandle().GetUniqueInstanceName();
 							if ( EmitterInstance->GetGPUContext() )
 							{
-								StringBuilder.Appendf(TEXT("Emitter(GPU) %s - State %s - Particles %d\n"), *NiagaraEmitter->GetUniqueEmitterName(), *ExecutionStateEnum->GetNameStringByIndex((int32)EmitterInstance->GetExecutionState()), EmitterInstance->GetNumParticles());
+								StringBuilder.Appendf(TEXT("Emitter(GPU) %s - State %s - Particles %d\n"), *EmitterName, *ExecutionStateEnum->GetNameStringByIndex((int32)EmitterInstance->GetExecutionState()), EmitterInstance->GetNumParticles());
 							}
 							else
 							{
-								StringBuilder.Appendf(TEXT("Emitter %s - State %s - Particles %d\n"), *NiagaraEmitter->GetUniqueEmitterName(), *ExecutionStateEnum->GetNameStringByIndex((int32)EmitterInstance->GetExecutionState()), EmitterInstance->GetNumParticles());
+								StringBuilder.Appendf(TEXT("Emitter %s - State %s - Particles %d\n"), *EmitterName, *ExecutionStateEnum->GetNameStringByIndex((int32)EmitterInstance->GetExecutionState()), EmitterInstance->GetNumParticles());
 							}
 						}
 					}
@@ -2998,7 +3004,6 @@ void FNiagaraDebugHud::DrawComponents(FNiagaraWorldManager* WorldManager, UCanva
 							for (const FParameterStoreVariable& ParameterStoreVariable : CachedVariables.ParameterStoreVariables)
 							{
 								FNiagaraParameterStore* ParameterStore = nullptr;
-								//-TODO:Stateless:
 								FNiagaraEmitterInstanceImpl* StatefulEmitter = SystemInstance->Emitters.IsValidIndex(ParameterStoreVariable.EmitterIndex) ? SystemInstance->Emitters[ParameterStoreVariable.EmitterIndex]->AsStateful() : nullptr;
 								switch (ParameterStoreVariable.StoreLocation)
 								{
diff --git a/Engine/Plugins/FX/Niagara/Source/Niagara/Private/NiagaraEmitterHandle.cpp b/Engine/Plugins/FX/Niagara/Source/Niagara/Private/NiagaraEmitterHandle.cpp
--- a/Engine/Plugins/FX/Niagara/Source/Niagara/Private/NiagaraEmitterHandle.cpp
+++ b/Engine/Plugins/FX/Niagara/Source/Niagara/Private/NiagaraEmitterHandle.cpp
@@ -3,6 +3,7 @@
 #include "NiagaraEmitterHandle.h"
 #include "NiagaraSystem.h"
 #include "NiagaraEmitter.h"
+#include "Stateless/NiagaraStatelessEmitter.h"
 #include "NiagaraScriptSourceBase.h"
 #include "NiagaraCommon.h"
 #include "NiagaraCustomVersion.h"
@@ -33,6 +34,7 @@ FNiagaraEmitterHandle::FNiagaraEmitterHandle(UNiagaraEmitter& InEmitter, const F
 	, bIsolated(false)
 	, Instance_DEPRECATED(nullptr)
 	, VersionedInstance(FVersionedNiagaraEmitter(&InEmitter, Version))
+	, EmitterMode(ENiagaraEmitterMode::Standard)
 {
 }
 
@@ -46,8 +48,25 @@ FNiagaraEmitterHandle::FNiagaraEmitterHandle(const FVersionedNiagaraEmitter& InE
 	, bIsolated(false)
 	, Instance_DEPRECATED(nullptr)
 	, VersionedInstance(InEmitter)
+	, EmitterMode(ENiagaraEmitterMode::Standard)
 {
 }
+
+//-TODO:Stateless
+FNiagaraEmitterHandle::FNiagaraEmitterHandle(UNiagaraStatelessEmitter& InEmitter)
+	: Name(*InEmitter.GetUniqueEmitterName())
+	, Id(FGuid::NewGuid())
+	, IdName(*Id.ToString())
+	, bIsEnabled(true)
+	, Source_DEPRECATED(nullptr)
+	, LastMergedSource_DEPRECATED(nullptr)
+	, bIsolated(false)
+	, Instance_DEPRECATED(nullptr)
+	, StatelessEmitter(&InEmitter)
+	, EmitterMode(ENiagaraEmitterMode::Stateless)
+{
+}
+//-TODO:Stateless
 #endif
 
 bool FNiagaraEmitterHandle::IsValid() const
@@ -89,7 +108,13 @@ void FNiagaraEmitterHandle::SetName(FName InName, UNiagaraSystem& InOwnerSystem)
 	FName UniqueName = FNiagaraUtilities::GetUniqueName(SanitizedName, OtherEmitterNames);
 
 	Name = UniqueName;
-	if (VersionedInstance.Emitter->SetUniqueEmitterName(Name.ToString()))
+
+	if (StatelessEmitter != nullptr)
+	{
+		StatelessEmitter->SetUniqueEmitterName(Name.ToString());
+	}
+
+	if (VersionedInstance.Emitter && VersionedInstance.Emitter->SetUniqueEmitterName(Name.ToString()))
 	{
  #if WITH_EDITOR
 		if (InOwnerSystem.GetSystemSpawnScript() && InOwnerSystem.GetSystemSpawnScript()->GetLatestSource())
@@ -143,6 +168,24 @@ bool FNiagaraEmitterHandle::SetIsEnabled(bool bInIsEnabled, UNiagaraSystem& InOw
 	return false;
 }
 
+#if WITH_EDITORONLY_DATA
+void FNiagaraEmitterHandle::SetEmitterMode(UNiagaraSystem& InOwningSystem, ENiagaraEmitterMode InEmitterMode)
+{
+	if (EmitterMode != InEmitterMode)
+	{
+		EmitterMode = InEmitterMode;
+		if (EmitterMode == ENiagaraEmitterMode::Stateless && StatelessEmitter == nullptr)
+		{
+			StatelessEmitter = NewObject<UNiagaraStatelessEmitter>(&InOwningSystem, "StatelessEmitter", RF_Transactional);
+			StatelessEmitter->SetUniqueEmitterName(*Name.ToString());
+			StatelessEmitter->CacheFromCompiledData();
+		}
+		InOwningSystem.RequestCompile(false);
+		OnEmitterModeChangedDelegate.Broadcast();
+	}
+}
+#endif
+
 FVersionedNiagaraEmitter FNiagaraEmitterHandle::GetInstance() const
 {
 	return VersionedInstance;
@@ -165,6 +208,11 @@ FString FNiagaraEmitterHandle::GetUniqueInstanceName() const
 	{
 		return VersionedInstance.Emitter->GetUniqueEmitterName();
 	}
+	//-TODO:Stateless: We need an emitter data adapter / base class
+	else if (StatelessEmitter)
+	{
+		return StatelessEmitter->GetUniqueEmitterName();
+	}
 	return FString();
 }
 
diff --git a/Engine/Plugins/FX/Niagara/Source/Niagara/Private/NiagaraEmitterInstanceImpl.cpp b/Engine/Plugins/FX/Niagara/Source/Niagara/Private/NiagaraEmitterInstanceImpl.cpp
--- a/Engine/Plugins/FX/Niagara/Source/Niagara/Private/NiagaraEmitterInstanceImpl.cpp
+++ b/Engine/Plugins/FX/Niagara/Source/Niagara/Private/NiagaraEmitterInstanceImpl.cpp
@@ -372,6 +372,7 @@ void FNiagaraEmitterInstanceImpl::Init(int32 InEmitterIndex)
 			GPUExecContext->SystemStatID = GetSystemStatID(false, false);
 			GPUExecContext->EmitterStatID = GetEmitterStatID(false, false);
 		#endif
+			GPUDataBufferInterfaces = GPUExecContext;
 		}
 	}
 
diff --git a/Engine/Plugins/FX/Niagara/Source/Niagara/Private/NiagaraGPUSystemTick.cpp b/Engine/Plugins/FX/Niagara/Source/Niagara/Private/NiagaraGPUSystemTick.cpp
--- a/Engine/Plugins/FX/Niagara/Source/Niagara/Private/NiagaraGPUSystemTick.cpp
+++ b/Engine/Plugins/FX/Niagara/Source/Niagara/Private/NiagaraGPUSystemTick.cpp
@@ -129,7 +129,6 @@ void FNiagaraGPUSystemTick::Init(FNiagaraSystemInstance* InSystemInstance)
 				continue;
 			}
 
-			//-TODO:Stateless:
 			FNiagaraEmitterInstanceImpl* EmitterInstanceImpl = EmitterInstance->AsStateful();
 			if (!EmitterInstanceImpl)
 			{
diff --git a/Engine/Plugins/FX/Niagara/Source/Niagara/Private/NiagaraModule.cpp b/Engine/Plugins/FX/Niagara/Source/Niagara/Private/NiagaraModule.cpp
--- a/Engine/Plugins/FX/Niagara/Source/Niagara/Private/NiagaraModule.cpp
+++ b/Engine/Plugins/FX/Niagara/Source/Niagara/Private/NiagaraModule.cpp
@@ -45,6 +45,7 @@
 #include "DataInterface/NiagaraDataInterfaceDataChannelCommon.h"
 #include "DataInterface/NiagaraDataInterfaceDataChannelWrite.h"
 #include "DataInterface/NiagaraDataInterfaceDataChannelRead.h"
+#include "Stateless/NiagaraStatelessCommon.h"
 
 #if UE_USE_OPENVDB
 #include "NiagaraOpenVDB.h"
@@ -450,6 +451,7 @@ void INiagaraModule::StartupModule()
 	UNiagaraMeshRendererProperties::InitCDOPropertiesAfterModuleStartup();
 	UNiagaraComponentRendererProperties::InitCDOPropertiesAfterModuleStartup();
 	UNiagaraVolumeRendererProperties::InitCDOPropertiesAfterModuleStartup();
+	NiagaraStatelessCommon::Initialize();
 
 	// Register the data interface CDO finder with the shader module..
 	INiagaraShaderModule& NiagaraShaderModule = FModuleManager::LoadModuleChecked<INiagaraShaderModule>("NiagaraShader");
diff --git a/Engine/Plugins/FX/Niagara/Source/Niagara/Private/NiagaraRenderer.cpp b/Engine/Plugins/FX/Niagara/Source/Niagara/Private/NiagaraRenderer.cpp
--- a/Engine/Plugins/FX/Niagara/Source/Niagara/Private/NiagaraRenderer.cpp
+++ b/Engine/Plugins/FX/Niagara/Source/Niagara/Private/NiagaraRenderer.cpp
@@ -272,19 +272,15 @@ FNiagaraDynamicDataBase::FNiagaraDynamicDataBase(const FNiagaraEmitterInstance*
 
 	SystemInstanceID = InEmitter->GetParentSystemInstance()->GetId();
 
-	const FNiagaraDataSet& DataSet = InEmitter->GetParticleData();
-	const ENiagaraSimTarget SimTarget = DataSet.GetSimTarget();
-	if (SimTarget == ENiagaraSimTarget::CPUSim)
+	// CPU simulations we take a reference to the most recent data
+	if (InEmitter->GetSimTarget() == ENiagaraSimTarget::CPUSim)
 	{
-		//On CPU we pass through direct ptr to the most recent data buffer.
-		CPUParticleData = &DataSet.GetCurrentDataChecked();
+		CPUParticleData = &InEmitter->GetParticleData().GetCurrentDataChecked();
 	}
+	// GPU simulations we get a callback which will give us the correct data
 	else
 	{
-		//On GPU we must access the correct buffer via the GPUExecContext. Probably a way to route this data better outside the dynamic data in future.
-		//During simulation, the correct data buffer for rendering will be placed in the GPUContext and AddReadRef called.
-		check(SimTarget == ENiagaraSimTarget::GPUComputeSim);
-		GPUExecContext = InEmitter->GetGPUContext();
+		ComputeDataBufferInterface = InEmitter->GetComputeDataBufferInterface();
 	}
 }
 
@@ -294,12 +290,12 @@ FNiagaraDynamicDataBase::~FNiagaraDynamicDataBase()
 
 bool FNiagaraDynamicDataBase::IsGpuLowLatencyTranslucencyEnabled() const
 {
-	return GPUExecContext ? GPUExecContext->HasTranslucentDataToRender() : false;
+	return ComputeDataBufferInterface ? ComputeDataBufferInterface->HasTranslucentDataToRender() : false;
 }
 
-FNiagaraDataBuffer* FNiagaraDynamicDataBase::GetParticleDataToRender(bool bIsLowLatencyTranslucent)const
+FNiagaraDataBuffer* FNiagaraDynamicDataBase::GetParticleDataToRender(bool bIsLowLatencyTranslucent) const
 {
-	FNiagaraDataBuffer* Ret = GPUExecContext ? GPUExecContext->GetDataToRender(bIsLowLatencyTranslucent) : CPUParticleData.GetReference();
+	FNiagaraDataBuffer* Ret = ComputeDataBufferInterface ? ComputeDataBufferInterface->GetDataToRender(bIsLowLatencyTranslucent) : CPUParticleData.GetReference();
 	checkSlow(Ret == nullptr || Ret->IsBeingRead());
 	return Ret;
 }
diff --git a/Engine/Plugins/FX/Niagara/Source/Niagara/Private/NiagaraRendererComponents.cpp b/Engine/Plugins/FX/Niagara/Source/Niagara/Private/NiagaraRendererComponents.cpp
--- a/Engine/Plugins/FX/Niagara/Source/Niagara/Private/NiagaraRendererComponents.cpp
+++ b/Engine/Plugins/FX/Niagara/Source/Niagara/Private/NiagaraRendererComponents.cpp
@@ -403,7 +403,6 @@ void FNiagaraRendererComponents::PostSystemTick_GameThread(const UNiagaraRendere
 	}
 
 
-	//-TODO:Stateless:
 	const FNiagaraEmitterInstanceImpl* StatefulEmitter = Emitter->AsStateful();
 	if (StatefulEmitter == nullptr)
 	{
diff --git a/Engine/Plugins/FX/Niagara/Source/Niagara/Private/NiagaraRendererProperties.cpp b/Engine/Plugins/FX/Niagara/Source/Niagara/Private/NiagaraRendererProperties.cpp
--- a/Engine/Plugins/FX/Niagara/Source/Niagara/Private/NiagaraRendererProperties.cpp
+++ b/Engine/Plugins/FX/Niagara/Source/Niagara/Private/NiagaraRendererProperties.cpp
@@ -11,6 +11,8 @@
 #include "NiagaraSystem.h"
 #include "NiagaraSystemImpl.h"
 
+#include "Stateless/NiagaraStatelessEmitter.h"
+
 #include "Interfaces/ITargetPlatform.h"
 #include "Materials/MaterialInterface.h"
 #include "Materials/MaterialInstanceConstant.h"
@@ -858,6 +860,24 @@ bool UNiagaraRendererProperties::NeedsLoadForTargetPlatform(const ITargetPlatfor
 			}
 		}
 	}
+	//-TODO:Stateless: We need a base emitter type
+	else if (const UNiagaraStatelessEmitter* OwnerStatelessEmitter = GetTypedOuter<UNiagaraStatelessEmitter>())
+	{
+		if (OwnerStatelessEmitter->NeedsLoadForTargetPlatform(TargetPlatform))
+		{
+			if (bIsEnabled && Platforms.IsEnabledForPlatform(TargetPlatform->IniPlatformName()))
+			{
+//#if WITH_EDITORONLY_DATA
+//				if (GNiagaraRendererCookOutStaticEnabledBinding && RendererEnabledBinding.IsValid())
+//				{
+//					TOptional<bool> ResolvedStaticValue = TryResolveStaticVariableBool(OwnerEmitter, RendererEnabledBinding.GetParamMapBindableVariable());
+//					return ResolvedStaticValue.Get(true);
+//				}
+//#endif
+				return true;
+			}
+		}
+	}
 
 	return false;
 }
diff --git a/Engine/Plugins/FX/Niagara/Source/Niagara/Private/NiagaraSystemInstance.cpp b/Engine/Plugins/FX/Niagara/Source/Niagara/Private/NiagaraSystemInstance.cpp
--- a/Engine/Plugins/FX/Niagara/Source/Niagara/Private/NiagaraSystemInstance.cpp
+++ b/Engine/Plugins/FX/Niagara/Source/Niagara/Private/NiagaraSystemInstance.cpp
@@ -2,6 +2,7 @@
 
 #include "NiagaraSystemInstance.h"
 #include "Misc/LargeWorldRenderPosition.h"
+#include "Stateless/NiagaraStatelessEmitterInstance.h"
 #include "NiagaraSystemGpuComputeProxy.h"
 #include "NiagaraCommon.h"
 #include "NiagaraComputeExecutionContext.h"
@@ -901,7 +902,8 @@ void FNiagaraSystemInstance::ResetInternal(bool bResetSimulations)
 	Age = 0;
 	TickCount = 0;
 	CachedDeltaSeconds = 0.0f;
-	
+	InitSystemState();
+
 	if(!bLODDistanceIsOverridden)
 	{
 		bLODDistanceIsValid = false;
@@ -1027,6 +1029,7 @@ void FNiagaraSystemInstance::ReInitInternal()
 	CachedDeltaSeconds = 0.0f;
 	bAlreadyBound = false;
 	bSolo = bForceSolo;
+	InitSystemState();
 
 	if (!::IsValid(System))
 	{
@@ -1074,7 +1077,6 @@ void FNiagaraSystemInstance::ReInitInternal()
 	const FNiagaraSystemCompiledData& SystemCompiledData = System->GetSystemCompiledData();
 	InstanceParameters = SystemCompiledData.InstanceParamStore;
 
-
 	//When re initializing, throw away old emitters and init new ones.
 	Emitters.Reset();
 	InitEmitters();
@@ -2157,8 +2159,15 @@ void FNiagaraSystemInstance::InitEmitters()
 			//-TODO: We should not create emitter instances for disabled emitters
 			const bool EmitterEnabled = EmitterHandle.GetIsEnabled();
 
-			Emitters.Emplace(MakeShared<FNiagaraEmitterInstanceImpl, ESPMode::ThreadSafe>(this));
-
+			//-TODO:Stateless: Should this be a factory?
+			if ( EmitterHandle.GetEmitterMode() == ENiagaraEmitterMode::Stateless )
+			{
+				Emitters.Emplace(MakeShared<FNiagaraStatelessEmitterInstance, ESPMode::ThreadSafe>(this));
+			}
+			else
+			{
+				Emitters.Emplace(MakeShared<FNiagaraEmitterInstanceImpl, ESPMode::ThreadSafe>(this));
+			}
 			FNiagaraEmitterInstanceRef EmitterInstance = Emitters.Last();
 
 			if (System->bFixedBounds)
@@ -2978,3 +2987,77 @@ const FString& FNiagaraSystemInstance::GetCrashReporterTag()const
 	}
 	return CrashReporterTag;
 }
+
+void FNiagaraSystemInstance::InitSystemState()
+{
+	const FNiagaraSystemStateData& SystemStateData = System->GetSystemStateData();
+	SystemState_RandomStream.Initialize(RandomSeed + RandomSeedOffset);
+
+	SystemState_LoopCount			= 0;
+	SystemState_CurrentLoopDuration = SystemState_RandomStream.FRandRange(SystemStateData.LoopDurationMin, SystemStateData.LoopDurationMax);
+	SystemState_CurrentLoopDelay	= SystemState_RandomStream.FRandRange(SystemStateData.LoopDelayMin, SystemStateData.LoopDelayMax);
+	SystemState_CurrentLoopAgeStart	= 0.0f;
+	SystemState_CurrentLoopAgeEnd	= SystemState_CurrentLoopAgeStart + SystemState_CurrentLoopDelay + SystemState_CurrentLoopDuration;
+}
+
+void FNiagaraSystemInstance::TickSystemState()
+{
+	const uint32 ParameterIndex = GetParameterIndex();
+	const FNiagaraSystemParameters& CurrentSystemParameters = SystemParameters[ParameterIndex];
+	ENiagaraExecutionState ExecutionState = static_cast<ENiagaraExecutionState>(CurrentSystemParameters.EngineExecutionState);
+
+	if (ExecutionState == ENiagaraExecutionState::Active)
+	{
+		if (Age < SystemState_CurrentLoopAgeEnd)
+		{
+			return;
+		}
+
+		const FNiagaraSystemStateData& SystemStateData = System->GetSystemStateData();
+		if (SystemStateData.bIgnoreSystemState == false)
+		{
+			const ENiagaraExecutionState InactiveExecutionState = SystemStateData.InactiveResponse == ENiagaraSystemInactiveResponse::Kill ? ENiagaraExecutionState::Complete : ENiagaraExecutionState::Inactive;
+
+			if (SystemStateData.LoopBehavior == ENiagaraLoopBehavior::Once)
+			{
+				SetActualExecutionState(InactiveExecutionState);
+				return;
+			}
+
+			// Keep looping until we find out which loop we are in as a small loop age + large DT could result in crossing multiple loops
+			do
+			{
+				++SystemState_LoopCount;
+				if (SystemStateData.LoopBehavior == ENiagaraLoopBehavior::Multiple && SystemState_LoopCount >= SystemStateData.LoopCount)
+				{
+					SetActualExecutionState(InactiveExecutionState);
+					return;
+				}
+
+				if (SystemStateData.bRecalculateDurationEachLoop)
+				{
+					SystemState_CurrentLoopDuration = SystemState_RandomStream.FRandRange(SystemStateData.LoopDurationMin, SystemStateData.LoopDurationMax);
+				}
+
+				if (SystemStateData.bDelayFirstLoopOnly)
+				{
+					SystemState_CurrentLoopDelay = 0.0f;
+				}
+				else if (SystemStateData.bRecalculateDelayEachLoop)
+				{
+					SystemState_CurrentLoopDelay = SystemState_RandomStream.FRandRange(SystemStateData.LoopDelayMin, SystemStateData.LoopDelayMax);
+				}
+
+				SystemState_CurrentLoopAgeStart	= SystemState_CurrentLoopAgeEnd;
+				SystemState_CurrentLoopAgeEnd	= SystemState_CurrentLoopAgeStart + SystemState_CurrentLoopDelay + SystemState_CurrentLoopDuration;
+			} while (Age >= SystemState_CurrentLoopAgeEnd);
+		}
+	}
+	// Waiting on emitters to complete
+	//else if (ExecutionState == ENiagaraExecutionState::Inactive)
+	//{
+	//}
+
+	// Update the actual execution state
+	SetActualExecutionState(ExecutionState);
+}
diff --git a/Engine/Plugins/FX/Niagara/Source/Niagara/Private/NiagaraSystemSimulation.cpp b/Engine/Plugins/FX/Niagara/Source/Niagara/Private/NiagaraSystemSimulation.cpp
--- a/Engine/Plugins/FX/Niagara/Source/Niagara/Private/NiagaraSystemSimulation.cpp
+++ b/Engine/Plugins/FX/Niagara/Source/Niagara/Private/NiagaraSystemSimulation.cpp
@@ -575,6 +575,8 @@ bool FNiagaraSystemSimulation::Init(UNiagaraSystem* InSystem, UWorld* InWorld, b
 
 	bCanExecute = System->GetSystemSpawnScript()->GetVMExecutableData().IsValid() && System->GetSystemUpdateScript()->GetVMExecutableData().IsValid();
 
+	bSystemStateFastPathEnabled = System->SystemStateFastPathEnabled();
+
 	MaxDeltaTime = System->GetMaxDeltaTime();
 
 	if (bCanExecute)
@@ -994,21 +996,31 @@ void FNiagaraSystemSimulation::FlushTickBatch(FNiagaraSystemSimulationTickContex
 	// If we are running async create tasks to execute
 	if ( Context.IsRunningAsync() )
 	{
-		// Queue instance concurrent task and track information in the instance
-		FGraphEventRef InstanceAsyncGraphEvent = TGraphTask<FNiagaraSystemInstanceTickConcurrentTask>::CreateTask(&Context.BeforeInstancesTickGraphEvents).ConstructAndDispatchWhenReady(this, Context.TickBatch, Context.World);
-
-		for (FNiagaraSystemInstance* Inst : Context.TickBatch)
+		FGraphEventArray FinalizePrereqArray;
+		if (bSystemStateFastPathEnabled)
 		{
-			Inst->ConcurrentTickBatchGraphEvent = InstanceAsyncGraphEvent;
+			for (FNiagaraSystemInstance* Inst : Context.TickBatch)
+			{
+				Inst->Tick_Concurrent();
+			}
 		}
+		else
+		{
+			// Queue instance concurrent task and track information in the instance
+			FGraphEventRef InstanceAsyncGraphEvent = TGraphTask<FNiagaraSystemInstanceTickConcurrentTask>::CreateTask(&Context.BeforeInstancesTickGraphEvents).ConstructAndDispatchWhenReady(this, Context.TickBatch, Context.World);
 
-		// Ensure ConcurrentTickBatchGraphEvent is visible before we clear ConcurrentTickGraphEvent
-		FPlatformMisc::MemoryBarrier();
+			for (FNiagaraSystemInstance* Inst : Context.TickBatch)
+			{
+				Inst->ConcurrentTickBatchGraphEvent = InstanceAsyncGraphEvent;
+			}
+			FinalizePrereqArray.Add(InstanceAsyncGraphEvent);
+
+			// Ensure ConcurrentTickBatchGraphEvent is visible before we clear ConcurrentTickGraphEvent
+			FPlatformMisc::MemoryBarrier();
+		}
 
 		// Queue finalize task which will run after the instances are complete, track with our all completion event
-		FGraphEventArray FinalizePrereqArray;
 		FinalizePrereqArray.Add(Context.TickBatch[0]->ConcurrentTickGraphEvent);
-		FinalizePrereqArray.Add(InstanceAsyncGraphEvent);
 		FGraphEventRef FinalizeTask = TGraphTask<FNiagaraSystemInstanceFinalizeTask>::CreateTask(&FinalizePrereqArray).ConstructAndDispatchWhenReady(this, Context.TickBatch);
 
 		check(Context.CompletionEvents != nullptr);
@@ -1559,7 +1571,7 @@ void FNiagaraSystemSimulation::SimCachePostTick_Concurrent(float DeltaSeconds, c
 			//DataSetToEmitterUpdateParameters[iEmitter].DataSetToParameterStore(EmitterInstance.GetUpdateExecutionContext().Parameters, MainDataSet, iSystemInstance);
 			DataSetToEmitterRendererParameters[iEmitter].DataSetToParameterStore(EmitterInstance.GetRendererBoundVariables(), MainDataSet, iSystemInstance);
 
-			//-TODO:Stateless:
+			//-TODO:Stateless: Support for SimCache?
 			if ( FNiagaraEmitterInstanceImpl* StatefulEmitter = EmitterInstance.AsStateful() )
 			{
 				const ENiagaraExecutionState State = EmitterExecutionStateAccessors[iEmitter].GetReader(MainDataSet).GetSafe(iSystemInstance, ENiagaraExecutionState::Disabled);
@@ -1659,24 +1671,48 @@ void FNiagaraSystemSimulation::Tick_Concurrent(FNiagaraSystemSimulationTickConte
 	LLM_SCOPE(ELLMTag::Niagara);
 
 	FScopeCycleCounterUObject AdditionalScope(Context.System, GET_STATID(STAT_NiagaraOverview_GT_CNC));
-
-	FNiagaraSystemInstance* SoloSystemInstance = bIsSolo && Context.Instances.Num() == 1 ? Context.Instances[0] : nullptr;
+	if (!bCanExecute || !Context.Instances.Num())
+	{
+		return;
+	}
 
 	FNiagaraCrashReporterScope CRScope(this);
 
-	if (bCanExecute && Context.Instances.Num())
+	if (GbDumpSystemData || Context.System->bDumpDebugSystemInfo)
 	{
-		if (GbDumpSystemData || Context.System->bDumpDebugSystemInfo)
-		{
-			UE_LOG(LogNiagara, Log, TEXT("=========================================================="));
-			UE_LOG(LogNiagara, Log, TEXT("Niagara System Sim Tick_Concurrent(): %s"), *Context.System->GetName());
-			UE_LOG(LogNiagara, Log, TEXT("=========================================================="));
-		}
+		UE_LOG(LogNiagara, Log, TEXT("=========================================================="));
+		UE_LOG(LogNiagara, Log, TEXT("Niagara System Sim Tick_Concurrent(): %s"), *Context.System->GetName());
+		UE_LOG(LogNiagara, Log, TEXT("=========================================================="));
+	}
 
 #if STATS
-		FScopeCycleCounter SystemStatCounter(Context.System->GetStatID(true, true));
+	FScopeCycleCounter SystemStatCounter(Context.System->GetStatID(true, true));
 #endif
 
+	if (bSystemStateFastPathEnabled)
+	{
+		//-OPT: We should be able to avoid this but will require a lot of changes to the system simulation
+		//      We might want to consider having different system simulation types
+		{
+			Context.DataSet.BeginSimulate();
+			Context.DataSet.Allocate(Context.Instances.Num());
+			Context.DataSet.GetDestinationDataChecked().SetNumInstances(Context.Instances.Num());
+			Context.DataSet.EndSimulate();
+		}
+
+		for (FNiagaraSystemInstance* SystemInstance : Context.Instances)
+		{
+			//-TODO: Stateless doesn't require a lot of this data
+			SystemInstance->TickInstanceParameters_Concurrent();
+
+			SystemInstance->TickSystemState();
+
+			AddSystemToTickBatch(SystemInstance, Context);
+		}
+		FlushTickBatch(Context);
+	}
+	else
+	{
 		for (FNiagaraSystemInstance* SystemInstance : Context.Instances)
 		{
 			SystemInstance->TickInstanceParameters_Concurrent();
@@ -1698,36 +1734,39 @@ void FNiagaraSystemSimulation::Tick_Concurrent(FNiagaraSystemSimulationTickConte
 			AddSystemToTickBatch(Instance, Context);
 		}
 		FlushTickBatch(Context);
+	}
 
-		// When not running async we can finalize straight away
-		if ( !Context.IsRunningAsync() )
+	// When not running async we can finalize straight away
+	if ( !Context.IsRunningAsync() )
+	{
+		check(IsInGameThread());
+		int32 InstanceIndex = 0;
+		while (InstanceIndex < Context.Instances.Num())
 		{
-			check(IsInGameThread());
-			int32 InstanceIndex = 0;
-			while (InstanceIndex < Context.Instances.Num())
-			{
-				FNiagaraSystemInstance* Instance = Context.Instances[InstanceIndex];
-				Instance->FinalizeTick_GameThread();
+			FNiagaraSystemInstance* Instance = Context.Instances[InstanceIndex];
+			Instance->FinalizeTick_GameThread();
 
-				// Finalize can complete the instance and potentially reactivate
-				if (Context.Instances.IsValidIndex(InstanceIndex) && (Context.Instances[InstanceIndex] == Instance))
-				{
-					++InstanceIndex;
-				}
-
-				check(Context.DataSet.GetCurrentDataChecked().GetNumInstances() == Context.Instances.Num());
+			// Finalize can complete the instance and potentially reactivate
+			if (Context.Instances.IsValidIndex(InstanceIndex) && (Context.Instances[InstanceIndex] == Instance))
+			{
+				++InstanceIndex;
 			}
+
+			check(Context.DataSet.GetCurrentDataChecked().GetNumInstances() == Context.Instances.Num());
 		}
+	}
 
 	#if NIAGARA_SYSTEM_CAPTURE
+	{
+		FNiagaraSystemInstance* SoloSystemInstance = bIsSolo && Context.Instances.Num() == 1 ? Context.Instances[0] : nullptr;
 		if (SoloSystemInstance)
 		{
 			SoloSystemInstance->FinishCapture();
 		}
+	}
 	#endif
 
-		INC_DWORD_STAT_BY(STAT_NiagaraNumSystems, Context.Instances.Num());
-	}
+	INC_DWORD_STAT_BY(STAT_NiagaraNumSystems, Context.Instances.Num());
 }
 
 void FNiagaraSystemSimulation::SetupParameters_GameThread(float DeltaSeconds)
@@ -2015,7 +2054,6 @@ void FNiagaraSystemSimulation::TransferSystemSimResults(FNiagaraSystemSimulation
 
 				DataSetToEmitterRendererParameters[EmitterIdx].DataSetToParameterStore(EmitterInstance.GetRendererBoundVariables(), Context.DataSet, SystemIndex);
 
-				//-TODO:Stateless:
 				if ( FNiagaraEmitterInstanceImpl* StatefulEmitter = EmitterInstance.AsStateful() )
 				{
 					ENiagaraExecutionState State = EmitterExecutionStateAccessors[EmitterIdx].GetReader(Context.DataSet).GetSafe(SystemIndex, ENiagaraExecutionState::Disabled);
@@ -2217,7 +2255,6 @@ void FNiagaraSystemSimulation::InitParameterDataSetBindings(FNiagaraSystemInstan
 
 			DataSetToEmitterRendererParameters[EmitterIdx].Init(MainDataSet, EmitterInstance.GetRendererBoundVariables());
 
-			//-TODO:Stateless:
 			if (FNiagaraEmitterInstanceImpl* StatefulEmitter = EmitterInstance.AsStateful())
 			{
 				FNiagaraScriptExecutionContext& SpawnContext = StatefulEmitter->GetSpawnExecutionContext();
@@ -2283,7 +2320,6 @@ void FNiagaraConstantBufferToDataSetBinding::CopyToDataSets(
 		ApplyOffsets(CompiledData.UpdateInstanceOwnerBinding, OwnerParameters, UpdateDataSet, DataSestInstanceIndex);
 	}
 
-	//-TODO:Stateless: Remove when we have them inside EmitterHandles
 	TConstArrayView<FNiagaraEmitterInstanceRef> Emitters = SystemInstance.GetEmitters();
 	const int32 EmitterCount = FMath::Min(Emitters.Num(), CompiledData.SpawnInstanceEmitterBindings.Num());
 
diff --git a/Engine/Plugins/FX/Niagara/Source/Niagara/Private/Stateless/Modules/NiagaraStatelessModule_AddVelocity.cpp b/Engine/Plugins/FX/Niagara/Source/Niagara/Private/Stateless/Modules/NiagaraStatelessModule_AddVelocity.cpp
new file mode 100644
--- /dev/null
+++ b/Engine/Plugins/FX/Niagara/Source/Niagara/Private/Stateless/Modules/NiagaraStatelessModule_AddVelocity.cpp
@@ -0,0 +1,102 @@
+// Copyright Epic Games, Inc. All Rights Reserved.
+
+#include "Stateless/Modules/NiagaraStatelessModule_AddVelocity.h"
+#include "Stateless/NiagaraStatelessDrawDebugContext.h"
+
+void UNiagaraStatelessModule_AddVelocity::BuildEmitterData(FNiagaraStatelessEmitterDataBuildContext& BuildContext) const
+{
+	if (!IsModuleEnabled())
+	{
+		return;
+	}
+
+	NiagaraStateless::FPhysicsBuildData& PhysicsBuildData = BuildContext.GetTransientBuildData<NiagaraStateless::FPhysicsBuildData>();
+	if (VelocityType == ENSM_VelocityType::Linear)
+	{
+		PhysicsBuildData.VelocityMin += VelocityMin * VelocityScale;
+		PhysicsBuildData.VelocityMax += VelocityMax * VelocityScale;
+	}
+	else if (VelocityType == ENSM_VelocityType::FromPoint)
+	{
+		ensureMsgf(PhysicsBuildData.bPointVelocity == false, TEXT("Only a single point force is supported at the moment."));
+
+		PhysicsBuildData.bPointVelocity = true;
+		PhysicsBuildData.PointVelocityMin = PointVelocityMin;
+		PhysicsBuildData.PointVelocityMax = PointVelocityMax;
+		PhysicsBuildData.PointOrigin = PointOrigin;
+	}
+	else if (VelocityType == ENSM_VelocityType::InCone)
+	{
+		ensureMsgf(PhysicsBuildData.bConeVelocity == false, TEXT("Only a single cone force is supported at the moment."));
+
+		PhysicsBuildData.bConeVelocity = true;
+		PhysicsBuildData.ConeQuat = FQuat4f(ConeRotation.Quaternion());
+		PhysicsBuildData.ConeVelocityMin = ConeVelocityMin;
+		PhysicsBuildData.ConeVelocityMax = ConeVelocityMax;
+		PhysicsBuildData.ConeOuterAngle = ConeAngle;
+		PhysicsBuildData.ConeInnerAngle = InnerCone;
+		PhysicsBuildData.ConeVelocityFalloff = bSpeedFalloffFromConeAxisEnabled ? FMath::Clamp(SpeedFalloffFromConeAxis, 0.0f, 1.0f) : 0.0f;
+	}
+}
+
+#if WITH_EDITOR
+void UNiagaraStatelessModule_AddVelocity::DrawDebug(const FNiagaraStatelessDrawDebugContext& DrawDebugContext) const
+{
+	switch (VelocityType)
+	{
+		case ENSM_VelocityType::Linear:
+		{
+			const FVector MinDir = FVector(VelocityMin * VelocityScale);
+			const FVector MaxDir = FVector(VelocityMax * VelocityScale);
+			DrawDebugContext.DrawArrow(FVector::ZeroVector, MinDir);
+
+			if (!FMath::IsNearlyEqual(MinDir.Length(), MaxDir.Length()))
+			{
+				DrawDebugContext.DrawArrow(FVector::ZeroVector, MaxDir);
+			}
+			break;
+		}
+
+		case ENSM_VelocityType::InCone:
+		{
+			const FQuat Quat = ConeRotation.Quaternion();
+			const float ConeHAngle = ConeAngle / 2.0f;
+
+			TOptional<float> InnerConeHAngle;
+			if (InnerCone > 0.0f && !FMath::IsNearlyEqual(ConeAngle, InnerCone))
+			{
+				InnerConeHAngle = InnerCone / 2.0f;
+			}
+
+			DrawDebugContext.DrawCone(FVector::ZeroVector, Quat, ConeHAngle, ConeVelocityMin);
+			if ( InnerConeHAngle.IsSet() )
+			{
+				DrawDebugContext.DrawCone(FVector::ZeroVector, Quat, InnerConeHAngle.GetValue(), ConeVelocityMin);
+			}
+
+			if (!FMath::IsNearlyEqual(ConeVelocityMin, ConeVelocityMax))
+			{
+				DrawDebugContext.DrawCone(FVector::ZeroVector, Quat, ConeHAngle, ConeVelocityMax);
+				if (InnerConeHAngle.IsSet())
+				{
+					DrawDebugContext.DrawCone(FVector::ZeroVector, Quat, InnerConeHAngle.GetValue(), ConeVelocityMax);
+				}
+			}
+			break;
+		}
+
+		case ENSM_VelocityType::FromPoint:
+		{
+			if (!FMath::IsNearlyEqual(PointVelocityMin, 0.0f))
+			{
+				DrawDebugContext.DrawSphere(FVector(PointOrigin), PointVelocityMin);
+			}
+			if (!FMath::IsNearlyEqual(PointVelocityMin, PointVelocityMax))
+			{
+				DrawDebugContext.DrawSphere(FVector(PointOrigin), PointVelocityMax);
+			}
+			break;
+		}
+	}
+}
+#endif
diff --git a/Engine/Plugins/FX/Niagara/Source/Niagara/Private/Stateless/Modules/NiagaraStatelessModule_CurlNoiseForce.cpp b/Engine/Plugins/FX/Niagara/Source/Niagara/Private/Stateless/Modules/NiagaraStatelessModule_CurlNoiseForce.cpp
new file mode 100644
--- /dev/null
+++ b/Engine/Plugins/FX/Niagara/Source/Niagara/Private/Stateless/Modules/NiagaraStatelessModule_CurlNoiseForce.cpp
@@ -0,0 +1,347 @@
+// Copyright Epic Games, Inc. All Rights Reserved.
+
+#include "Stateless/Modules/NiagaraStatelessModule_CurlNoiseForce.h"
+
+namespace NoiseHelper
+{
+	struct FNiagaraUIntVector
+	{
+		uint32 X;
+		uint32 Y;
+		uint32 Z;
+
+		FNiagaraUIntVector() : X(0), Y(0), Z(0) {}
+		FNiagaraUIntVector(uint32 X, uint32 Y, uint32 Z) : X(X), Y(Y), Z(Z) {}
+		FNiagaraUIntVector(uint32 Value) : X(Value), Y(Value), Z(Value) {}
+		FNiagaraUIntVector(FIntVector Values) : X(Values.X), Y(Values.Y), Z(Values.Z) {}
+
+		FNiagaraUIntVector operator+(FNiagaraUIntVector Rhs)
+		{
+			return FNiagaraUIntVector{ X + Rhs.X, Y + Rhs.Y, Z + Rhs.Z };
+		}
+
+		FNiagaraUIntVector operator*(FNiagaraUIntVector Rhs)
+		{
+			return FNiagaraUIntVector{ X * Rhs.X, Y * Rhs.Y, Z * Rhs.Z };
+		}
+
+		FNiagaraUIntVector operator>>(uint32 Shift)
+		{
+			return FNiagaraUIntVector(X >> Shift, Y >> Shift, Z >> Shift);
+		}
+
+		FNiagaraUIntVector operator&(FNiagaraUIntVector Rhs)
+		{
+			return FNiagaraUIntVector{ X & Rhs.X, Y & Rhs.Y, Z & Rhs.Z };
+		}
+
+		uint32& operator[](int Index)
+		{
+			if (Index == 0) return X;
+			if (Index == 1) return Y;
+			if (Index == 2) return Z;
+			check(false);
+			return X;
+		}
+
+		const uint32& operator[](int Index) const
+		{
+			if (Index == 0) return X;
+			if (Index == 1) return Y;
+			if (Index == 2) return Z;
+			check(false);
+			return X;
+		}
+	};
+
+	FNiagaraUIntVector Rand3DPCG16(FIntVector p)
+	{
+		FNiagaraUIntVector v = FNiagaraUIntVector(p);
+
+		v = v * 1664525u + 1013904223u;
+
+		v.X += v.Y * v.Z;
+		v.Y += v.Z * v.X;
+		v.Z += v.X * v.Y;
+		v.X += v.Y * v.Z;
+		v.Y += v.Z * v.X;
+		v.Z += v.X * v.Y;
+
+		return v >> 16u;
+	}
+
+	FVector3f NiagaraVectorFrac(FVector3f v)
+	{
+		return FVector3f(FMath::Frac(v.X), FMath::Frac(v.Y), FMath::Frac(v.Z));
+	}
+
+	FVector3f NoiseTileWrap(FVector3f v, bool bTiling, float RepeatSize)
+	{
+		FVector3f vv = bTiling ? (NiagaraVectorFrac(v / RepeatSize) * RepeatSize) : v;
+		return vv;
+	}
+
+	struct FNiagaraMatrix4x3
+	{
+		FVector3f Row0;
+		FVector3f Row1;
+		FVector3f Row2;
+		FVector3f Row3;
+
+		FNiagaraMatrix4x3() : Row0(FVector3f::Zero()), Row1(FVector3f::Zero()), Row2(FVector3f::Zero()), Row3(FVector3f::Zero()) {}
+		FNiagaraMatrix4x3(FVector3f Row0, FVector3f Row1, FVector3f Row2, FVector3f Row3) : Row0(Row0), Row1(Row1), Row2(Row2), Row3(Row3) {}
+
+		FVector3f& operator[](int Row)
+		{
+			if (Row == 0) return Row0;
+			if (Row == 1) return Row1;
+			if (Row == 2) return Row2;
+			if (Row == 3) return Row3;
+			check(false);
+			return Row0;
+		}
+		const FVector3f& operator[](int Row) const
+		{
+			if (Row == 0) return Row0;
+			if (Row == 1) return Row1;
+			if (Row == 2) return Row2;
+			if (Row == 3) return Row3;
+			check(false);
+			return Row0;
+		}
+	};
+
+	FVector3f NiagaraVectorFloor(FVector3f v) {
+		return FVector3f(FGenericPlatformMath::FloorToFloat(v.X),
+			FGenericPlatformMath::FloorToFloat(v.Y),
+			FGenericPlatformMath::FloorToFloat(v.Z));
+	};
+
+	FIntVector NiagaraVectorFloorToInt(FVector3f v) {
+		return FIntVector(
+			FGenericPlatformMath::FloorToInt(v.X),
+			FGenericPlatformMath::FloorToInt(v.Y),
+			FGenericPlatformMath::FloorToInt(v.Z)
+		);
+	};
+
+	FVector3f NiagaraVectorStep(FVector3f v, FVector3f u)
+	{
+		return FVector3f(u.X >= v.X ? 1.0f : 0.0f,
+			u.Y >= v.Y ? 1.0f : 0.0f,
+			u.Z >= v.Z ? 1.0f : 0.0f);
+	}
+
+	FVector3f NiagaraVectorSwizzle(FVector3f v, uint32 x, uint32 y, uint32 z)
+	{
+		return FVector3f(v[x], v[y], v[z]);
+	}
+
+	FVector3f NiagaraVectorMin(FVector3f u, FVector3f v)
+	{
+		return FVector3f(u.X < v.X ? u.X : v.X,
+			u.Y < v.Y ? u.Y : v.Y,
+			u.Z < v.Z ? u.Z : v.Z);
+	}
+
+	FVector3f NiagaraVectorMax(FVector3f u, FVector3f v)
+	{
+		return FVector3f(u.X > v.X ? u.X : v.X,
+			u.Y > v.Y ? u.Y : v.Y,
+			u.Z > v.Z ? u.Z : v.Z);
+	}
+
+	FNiagaraMatrix4x3 SimplexCorners(FVector3f v)
+	{
+		FVector3f tet = NiagaraVectorFloor(v + v.X / 3 + v.Y / 3 + v.Z / 3);
+		FVector3f base = tet - tet.X / 6 - tet.Y / 6 - tet.Z / 6;
+		FVector3f f = v - base;
+
+		FVector3f g = NiagaraVectorStep(NiagaraVectorSwizzle(f, 1, 2, 0), f), h = FVector3f(1.0f) - NiagaraVectorSwizzle(g, 2, 0, 1);
+		FVector3f a1 = NiagaraVectorMin(g, h) - 1. / 6., a2 = NiagaraVectorMax(g, h) - 1. / 3.;
+
+		return FNiagaraMatrix4x3(base, base + a1, base + a2, base + 0.5);
+	}
+
+	FVector4f NiagaraVector4Saturate(const FVector4f& v) {
+		return FVector4f(FMath::Clamp(v.X, 0.0f, 1.0f), FMath::Clamp(v.Y, 0.0f, 1.0f), FMath::Clamp(v.Z, 0.0f, 1.0f), FMath::Clamp(v.W, 0.0f, 1.0f));
+	}
+
+	FVector4f SimplexSmooth(FNiagaraMatrix4x3 f)
+	{
+		const float scale = 1024.f / 375.f;
+		FVector4f d = FVector4f(FVector3f::DotProduct(f[0], f[0]), FVector3f::DotProduct(f[1], f[1]), FVector3f::DotProduct(f[2], f[2]), FVector3f::DotProduct(f[3], f[3]));
+		FVector4f s = NiagaraVector4Saturate(2.0f * d);
+		return scale * (FVector4f(1.0f, 1.0f, 1.0f, 1.0f) + s * (FVector4f(-3.0f, -3.0f, -3.0f, -3.0f) + s * (FVector4f(3.0f, 3.0f, 3.0f, 3.0f) - s)));
+	}
+
+	struct FNiagaraMatrix3x4
+	{
+		FVector4f row0;
+		FVector4f row1;
+		FVector4f row2;
+
+		FNiagaraMatrix3x4() : row0(FVector4f::Zero()), row1(FVector4f::Zero()), row2(FVector4f::Zero()) {}
+		FNiagaraMatrix3x4(const FVector4f& row0, const FVector4f& row1, const FVector4f& row2) : row0(row0), row1(row1), row2(row2) {}
+
+		FVector4f& operator[](int row)
+		{
+			if (row == 0) return row0;
+			if (row == 1) return row1;
+			if (row == 2) return row2;
+			check(false);
+			return row0;
+		}
+		const FVector4f& operator[](int row) const
+		{
+			if (row == 0) return row0;
+			if (row == 1) return row1;
+			if (row == 2) return row2;
+			check(false);
+			return row0;
+		}
+	};
+
+	FNiagaraMatrix3x4 SimplexDSmooth(FNiagaraMatrix4x3 f)
+	{
+		const float scale = 1024. / 375.;
+		FVector4f d = FVector4f(FVector3f::DotProduct(f[0], f[0]), FVector3f::DotProduct(f[1], f[1]), FVector3f::DotProduct(f[2], f[2]), FVector3f::DotProduct(f[3], f[3]));
+		FVector4f s = NiagaraVector4Saturate(2 * d);
+		s = -12 * FVector4f(scale, scale, scale, scale) + s * (24 * FVector4f(scale, scale, scale, scale) - s * 12 * scale);
+
+		return FNiagaraMatrix3x4(
+			s * FVector4f(f[0][0], f[1][0], f[2][0], f[3][0]),
+			s * FVector4f(f[0][1], f[1][1], f[2][1], f[3][1]),
+			s * FVector4f(f[0][2], f[1][2], f[2][2], f[3][2]));
+	}
+
+	FNiagaraUIntVector FNiagaraUIntVectorSwizzle(FNiagaraUIntVector v, int x, int y, int z)
+	{
+		return FNiagaraUIntVector(v[x], v[y], v[z]);
+	}
+
+	FVector3f FNiagaraUIntVectorToFVector(FNiagaraUIntVector v)
+	{
+		return FVector3f(float(v.X), float(v.Y), float(v.Z));
+	}
+
+	FVector3f MulFVector4AndFNiagaraMatrix4x3(const FVector4f& lhs, FNiagaraMatrix4x3 rhs)
+	{
+		return FVector3f(lhs[0] * rhs[0][0] + lhs[1] * rhs[1][0] + lhs[2] * rhs[2][0] + lhs[3] * rhs[3][0],
+			lhs[0] * rhs[0][1] + lhs[1] * rhs[1][1] + lhs[2] * rhs[2][1] + lhs[3] * rhs[3][1],
+			lhs[0] * rhs[0][2] + lhs[1] * rhs[1][2] + lhs[2] * rhs[2][2] + lhs[3] * rhs[3][2]);
+	}
+
+	FVector3f MulFNiagaraMatrix3x4FAndVector4(const FNiagaraMatrix3x4& lhs, const FVector4f& rhs)
+	{
+		return FVector3f(lhs[0][0] * rhs[0] + lhs[0][1] * rhs[1] + lhs[0][2] * rhs[2] + lhs[0][3] * rhs[3],
+			lhs[1][0] * rhs[0] + lhs[1][1] * rhs[1] + lhs[1][2] * rhs[2] + lhs[1][3] * rhs[3],
+			lhs[2][0] * rhs[0] + lhs[2][1] * rhs[1] + lhs[2][2] * rhs[2] + lhs[2][3] * rhs[3]);
+	}
+
+#define MGradientMask FNiagaraUIntVector(0x8000, 0x4000, 0x2000)
+#define MGradientScale FVector3f(1.f / 0x4000, 1.f / 0x2000, 1.f / 0x1000)
+
+	FNiagaraMatrix3x4 JacobianSimplex_ALU(FVector3f v)
+	{
+		FNiagaraMatrix4x3 T = SimplexCorners(v);
+		FNiagaraUIntVector rand;
+		FNiagaraMatrix4x3 gvec[3], fv;
+		FNiagaraMatrix3x4 grad;
+
+		fv[0] = v - T[0];
+		rand = Rand3DPCG16(NiagaraVectorFloorToInt(6 * T[0] + 0.5));
+		gvec[0][0] = FVector3f(FNiagaraUIntVectorToFVector(FNiagaraUIntVectorSwizzle(rand, 0, 0, 0) & MGradientMask)) * MGradientScale - 1;
+		gvec[1][0] = FVector3f(FNiagaraUIntVectorToFVector(FNiagaraUIntVectorSwizzle(rand, 1, 1, 1) & MGradientMask)) * MGradientScale - 1;
+		gvec[2][0] = FVector3f(FNiagaraUIntVectorToFVector(FNiagaraUIntVectorSwizzle(rand, 2, 2, 2) & MGradientMask)) * MGradientScale - 1;
+		grad[0][0] = FVector3f::DotProduct(gvec[0][0], fv[0]);
+		grad[1][0] = FVector3f::DotProduct(gvec[1][0], fv[0]);
+		grad[2][0] = FVector3f::DotProduct(gvec[2][0], fv[0]);
+
+		fv[1] = v - T[1];
+		rand = Rand3DPCG16(NiagaraVectorFloorToInt(6 * T[1] + 0.5));
+		gvec[0][1] = FVector3f(FNiagaraUIntVectorToFVector(FNiagaraUIntVectorSwizzle(rand, 0, 0, 0) & MGradientMask)) * MGradientScale - 1;
+		gvec[1][1] = FVector3f(FNiagaraUIntVectorToFVector(FNiagaraUIntVectorSwizzle(rand, 1, 1, 1) & MGradientMask)) * MGradientScale - 1;
+		gvec[2][1] = FVector3f(FNiagaraUIntVectorToFVector(FNiagaraUIntVectorSwizzle(rand, 2, 2, 2) & MGradientMask)) * MGradientScale - 1;
+		grad[0][1] = FVector3f::DotProduct(gvec[0][1], fv[1]);
+		grad[1][1] = FVector3f::DotProduct(gvec[1][1], fv[1]);
+		grad[2][1] = FVector3f::DotProduct(gvec[2][1], fv[1]);
+
+		fv[2] = v - T[2];
+		rand = Rand3DPCG16(NiagaraVectorFloorToInt(6 * T[2] + 0.5));
+		gvec[0][2] = FVector3f(FNiagaraUIntVectorToFVector(FNiagaraUIntVectorSwizzle(rand, 0, 0, 0) & MGradientMask)) * MGradientScale - 1;
+		gvec[1][2] = FVector3f(FNiagaraUIntVectorToFVector(FNiagaraUIntVectorSwizzle(rand, 1, 1, 1) & MGradientMask)) * MGradientScale - 1;
+		gvec[2][2] = FVector3f(FNiagaraUIntVectorToFVector(FNiagaraUIntVectorSwizzle(rand, 2, 2, 2) & MGradientMask)) * MGradientScale - 1;
+		grad[0][2] = FVector3f::DotProduct(gvec[0][2], fv[2]);
+		grad[1][2] = FVector3f::DotProduct(gvec[1][2], fv[2]);
+		grad[2][2] = FVector3f::DotProduct(gvec[2][2], fv[2]);
+
+		fv[3] = v - T[3];
+		rand = Rand3DPCG16(NiagaraVectorFloorToInt(6 * T[3] + 0.5));
+		gvec[0][3] = FVector3f(FNiagaraUIntVectorToFVector(FNiagaraUIntVectorSwizzle(rand, 0, 0, 0) & MGradientMask)) * MGradientScale - 1;
+		gvec[1][3] = FVector3f(FNiagaraUIntVectorToFVector(FNiagaraUIntVectorSwizzle(rand, 1, 1, 1) & MGradientMask)) * MGradientScale - 1;
+		gvec[2][3] = FVector3f(FNiagaraUIntVectorToFVector(FNiagaraUIntVectorSwizzle(rand, 2, 2, 2) & MGradientMask)) * MGradientScale - 1;
+		grad[0][3] = FVector3f::DotProduct(gvec[0][3], fv[3]);
+		grad[1][3] = FVector3f::DotProduct(gvec[1][3], fv[3]);
+		grad[2][3] = FVector3f::DotProduct(gvec[2][3], fv[3]);
+
+		FVector4f sv = SimplexSmooth(fv);
+		FNiagaraMatrix3x4 ds = SimplexDSmooth(fv);
+
+		FNiagaraMatrix3x4 jacobian;
+		jacobian[0] = FVector4f(MulFVector4AndFNiagaraMatrix4x3(sv, gvec[0]) + MulFNiagaraMatrix3x4FAndVector4(ds, grad[0]), Dot4(sv, grad[0]));
+		jacobian[1] = FVector4f(MulFVector4AndFNiagaraMatrix4x3(sv, gvec[1]) + MulFNiagaraMatrix3x4FAndVector4(ds, grad[1]), Dot4(sv, grad[1]));
+		jacobian[2] = FVector4f(MulFVector4AndFNiagaraMatrix4x3(sv, gvec[2]) + MulFNiagaraMatrix3x4FAndVector4(ds, grad[2]), Dot4(sv, grad[2]));
+
+		return jacobian;
+	}
+
+#undef MGradientMask
+#undef MGradientScale 
+}
+
+void UNiagaraStatelessModule_CurlNoiseForce::BuildEmitterData(FNiagaraStatelessEmitterDataBuildContext& BuildContext) const
+{
+	if (!IsModuleEnabled())
+	{
+		return;
+	}
+
+	NiagaraStateless::FPhysicsBuildData& PhysicsBuildData = BuildContext.GetTransientBuildData<NiagaraStateless::FPhysicsBuildData>();
+	PhysicsBuildData.bNoiseEnabled = true;
+	PhysicsBuildData.NoiseAmplitude = NoiseAmplitude;
+	PhysicsBuildData.NoiseFrequency = NoiseFrequency;
+	PhysicsBuildData.NoiseTexture = NoiseTexture;
+	PhysicsBuildData.NoiseMode = int32(NoiseMode);
+
+	// Build LUT
+	{
+		static int32 NumChannelsToBuild = 64;
+		static int32 ChannelWidth = 64;
+
+		TArray<float> ChannelData;
+		ChannelData.Reserve(NumChannelsToBuild * ChannelWidth);
+
+		for (int32 Channel=0; Channel < NumChannelsToBuild; ++Channel)
+		{
+			FVector3f NoisePosition = FVector3f(FMath::RandRange(-64.0f, 64.0f), FMath::RandRange(-64.0f, 64.0f), FMath::RandRange(-64.0f, 64.0f));
+			FVector3f NoiseTravel = FVector3f(FMath::RandRange(-1.0f, 1.0f), FMath::RandRange(-1.0f, 1.0f), FMath::RandRange(-1.0f, 1.0f));
+
+			float NoiseValue = 0.0f;
+			for (int32 Width=0; Width < ChannelWidth; ++Width)
+			{
+				//NoiseValue += FMath::RandRange(-1.0f, 1.0f);
+
+				NoiseHelper::FNiagaraMatrix3x4 J = NoiseHelper::JacobianSimplex_ALU(NoisePosition);
+				NoiseValue += J[1][2] - J[2][1];
+				NoisePosition += NoiseTravel;
+
+				ChannelData.Add(NoiseValue);
+			}
+		}
+
+		PhysicsBuildData.NoiseLUTOffset			= BuildContext.AddStaticData(ChannelData);
+		PhysicsBuildData.NoiseLUTNumChannel		= NumChannelsToBuild;
+		PhysicsBuildData.NoiseLUTChannelWidth	= ChannelWidth;
+	}
+}
