- Experimental stateless emitters

diff --git a/Engine/Plugins/FX/Niagara/Source/Niagara/Private/Stateless/Modules/NiagaraStatelessModule_ShapeLocation.cpp b/Engine/Plugins/FX/Niagara/Source/Niagara/Private/Stateless/Modules/NiagaraStatelessModule_ShapeLocation.cpp
new file mode 100644
--- /dev/null
+++ b/Engine/Plugins/FX/Niagara/Source/Niagara/Private/Stateless/Modules/NiagaraStatelessModule_ShapeLocation.cpp
@@ -0,0 +1,107 @@
+// Copyright Epic Games, Inc. All Rights Reserved.
+
+#include "Stateless/Modules/NiagaraStatelessModule_ShapeLocation.h"
+#include "Stateless/NiagaraStatelessDrawDebugContext.h"
+
+void UNiagaraStatelessModule_ShapeLocation::SetShaderParameters(const FNiagaraStatelessSetShaderParameterContext& SetShaderParameterContext) const
+{
+	FParameters* Parameters = SetShaderParameterContext.GetParameterNestedStruct<FParameters>();
+	if (!IsModuleEnabled())
+	{
+		Parameters->ShapeLocation_Mode.X		= 0;
+		Parameters->ShapeLocation_Parameters0	= FVector4f::Zero();
+		Parameters->ShapeLocation_Parameters1	= FVector4f::Zero();
+		return;
+	}
+
+	switch (ShapePrimitive)
+	{
+		case ENSM_ShapePrimitive::Box:
+		{
+			Parameters->ShapeLocation_Mode.X		= 0;
+			Parameters->ShapeLocation_Mode.Y		= bBoxSurfaceOnly ? 1 : 0;
+			Parameters->ShapeLocation_Parameters0	= FVector4f(BoxSize, BoxSurfaceThicknessMax - BoxSurfaceThicknessMin);
+			Parameters->ShapeLocation_Parameters1	= FVector4f(BoxSize * -0.5f, BoxSurfaceThicknessMin);
+			break;
+		}
+		case ENSM_ShapePrimitive::Plane:
+		{
+			Parameters->ShapeLocation_Mode.X		= 0;
+			Parameters->ShapeLocation_Parameters0	= FVector4f(PlaneSize.X, PlaneSize.Y, 0.0f, 0.0f);
+			Parameters->ShapeLocation_Parameters1	= FVector4f(-PlaneSize.X * 0.5f, -PlaneSize.Y * 0.5f, 0.0f, 0.0f);
+			break;
+		}
+		case ENSM_ShapePrimitive::Cylinder:
+		{
+			Parameters->ShapeLocation_Mode.X		= 1;
+			Parameters->ShapeLocation_Parameters0.X = CylinderHeight;
+			Parameters->ShapeLocation_Parameters0.Y = CylinderHeight * -CylinderHeightMidpoint;
+			Parameters->ShapeLocation_Parameters0.Z = CylinderRadius;
+			break;
+		}
+		case ENSM_ShapePrimitive::Ring:
+		{
+			const float DC = FMath::Clamp(1.0f - DiscCoverage, 0.0f, 1.0f);
+			const float SDC = DC > 0.0f ? FMath::Sqrt(DC) : 0.0f;
+
+			Parameters->ShapeLocation_Mode.X		= 2;
+			Parameters->ShapeLocation_Parameters0.X = RingRadius * (1.0f - SDC);
+			Parameters->ShapeLocation_Parameters0.Y = RingRadius * SDC;
+			Parameters->ShapeLocation_Parameters0.Z = -UE_TWO_PI * (1.0f - RingUDistribution);
+			Parameters->ShapeLocation_Parameters0.W = 0.0f;
+			break;
+		}
+		case ENSM_ShapePrimitive::Sphere:
+		{
+			Parameters->ShapeLocation_Mode.X		= 3;
+			Parameters->ShapeLocation_Parameters0.X = SphereMax - SphereMin;
+			Parameters->ShapeLocation_Parameters0.Y = SphereMin;
+			break;
+		}
+		default:
+		{
+			Parameters->ShapeLocation_Mode.X = 0;
+			Parameters->ShapeLocation_Parameters0 = FVector4f::Zero();
+			Parameters->ShapeLocation_Parameters1 = FVector4f::Zero();
+			checkNoEntry();
+		}
+	}
+}
+
+#if WITH_EDITOR
+void UNiagaraStatelessModule_ShapeLocation::DrawDebug(const FNiagaraStatelessDrawDebugContext& DrawDebugContext) const
+{
+	switch (ShapePrimitive)
+	{
+		case ENSM_ShapePrimitive::Box:
+		{
+			DrawDebugContext.DrawBox(FVector::ZeroVector, FVector(BoxSize * 0.5f));
+			break;
+		}
+		case ENSM_ShapePrimitive::Plane:
+		{
+			DrawDebugContext.DrawBox(FVector::ZeroVector, FVector(PlaneSize.X * 0.5f, PlaneSize.Y * 0.5f, 0.0f));
+			break;
+		}
+		case ENSM_ShapePrimitive::Cylinder:
+		{
+			DrawDebugContext.DrawCylinder(CylinderHeight, CylinderRadius, CylinderHeightMidpoint);
+			break;
+		}
+		case ENSM_ShapePrimitive::Ring:
+		{
+			const float DC = FMath::Clamp(1.0f - DiscCoverage, 0.0f, 1.0f);
+			const float SDC = DC > 0.0f ? FMath::Sqrt(DC) : 0.0f;
+			DrawDebugContext.DrawCircle(FVector::ZeroVector, RingRadius);
+			DrawDebugContext.DrawCircle(FVector::ZeroVector, RingRadius * SDC);
+			break;
+		}
+		case ENSM_ShapePrimitive::Sphere:
+		{
+			DrawDebugContext.DrawSphere(FVector::ZeroVector, SphereMin);
+			DrawDebugContext.DrawSphere(FVector::ZeroVector, SphereMax);
+			break;
+		}
+	}
+}
+#endif
diff --git a/Engine/Plugins/FX/Niagara/Source/Niagara/Private/Stateless/NiagaraStatelessCommon.cpp b/Engine/Plugins/FX/Niagara/Source/Niagara/Private/Stateless/NiagaraStatelessCommon.cpp
new file mode 100644
--- /dev/null
+++ b/Engine/Plugins/FX/Niagara/Source/Niagara/Private/Stateless/NiagaraStatelessCommon.cpp
@@ -0,0 +1,61 @@
+// Copyright Epic Games, Inc. All Rights Reserved.
+
+#include "Stateless/NiagaraStatelessCommon.h"
+
+#include "Stateless/NiagaraStatelessEmitterTemplate.h"
+
+namespace NiagaraStatelessCommon
+{
+	static FNiagaraStatelessGlobals GGlobals;
+
+	void Initialize()
+	{
+		GGlobals.ColorVariable				= FNiagaraVariableBase(FNiagaraTypeDefinition::GetColorDef(), TEXT("Color"));
+		GGlobals.MeshOrientationVariable	= FNiagaraVariableBase(FNiagaraTypeDefinition::GetQuatDef(), TEXT("MeshOrientation"));
+		GGlobals.PositionVariable			= FNiagaraVariableBase(FNiagaraTypeDefinition::GetPositionDef(), TEXT("Position"));
+		GGlobals.RibbonWidthVariable		= FNiagaraVariableBase(FNiagaraTypeDefinition::GetFloatDef(), TEXT("RibbonWidth"));
+		GGlobals.ScaleVariable				= FNiagaraVariableBase(FNiagaraTypeDefinition::GetVec3Def(), TEXT("Scale"));
+		GGlobals.SpriteAlignmentVariable	= FNiagaraVariableBase(FNiagaraTypeDefinition::GetVec3Def(), TEXT("SpriteAlignment"));
+		GGlobals.SpriteFacingVariable		= FNiagaraVariableBase(FNiagaraTypeDefinition::GetVec3Def(), TEXT("SpriteFacing"));
+		GGlobals.SpriteSizeVariable			= FNiagaraVariableBase(FNiagaraTypeDefinition::GetVec2Def(), TEXT("SpriteSize"));
+		GGlobals.SpriteRotationVariable		= FNiagaraVariableBase(FNiagaraTypeDefinition::GetFloatDef(), TEXT("SpriteRotation"));
+		GGlobals.SubImageIndexVariable		= FNiagaraVariableBase(FNiagaraTypeDefinition::GetFloatDef(), TEXT("SubImageIndex"));
+		GGlobals.UniqueIDVariable			= FNiagaraVariableBase(FNiagaraTypeDefinition::GetIntDef(), TEXT("UniqueID"));
+		GGlobals.VelocityVariable			= FNiagaraVariableBase(FNiagaraTypeDefinition::GetVec3Def(), TEXT("Velocity"));
+
+		//GGlobals.PreviousColorVariable			= FNiagaraVariableBase(FNiagaraTypeDefinition::GetColorDef(), TEXT("Previous.Color"));
+		GGlobals.PreviousMeshOrientationVariable	= FNiagaraVariableBase(FNiagaraTypeDefinition::GetQuatDef(), TEXT("Previous.MeshOrientation"));
+		GGlobals.PreviousPositionVariable			= FNiagaraVariableBase(FNiagaraTypeDefinition::GetPositionDef(), TEXT("Previous.Position"));
+		GGlobals.PreviousRibbonWidthVariable		= FNiagaraVariableBase(FNiagaraTypeDefinition::GetFloatDef(), TEXT("Previous.RibbonWidth"));
+		GGlobals.PreviousScaleVariable				= FNiagaraVariableBase(FNiagaraTypeDefinition::GetVec3Def(), TEXT("Previous.Scale"));
+		GGlobals.PreviousSpriteAlignmentVariable	= FNiagaraVariableBase(FNiagaraTypeDefinition::GetVec3Def(), TEXT("Previous.SpriteAlignment"));
+		GGlobals.PreviousSpriteFacingVariable		= FNiagaraVariableBase(FNiagaraTypeDefinition::GetVec3Def(), TEXT("Previous.SpriteFacing"));
+		GGlobals.PreviousSpriteSizeVariable			= FNiagaraVariableBase(FNiagaraTypeDefinition::GetVec2Def(), TEXT("Previous.SpriteSize"));
+		GGlobals.PreviousSpriteRotationVariable		= FNiagaraVariableBase(FNiagaraTypeDefinition::GetFloatDef(), TEXT("Previous.SpriteRotation"));
+		GGlobals.PreviousVelocityVariable			= FNiagaraVariableBase(FNiagaraTypeDefinition::GetVec3Def(), TEXT("Previous.Velocity"));
+
+		UNiagaraStatelessEmitterTemplate::InitCDOPropertiesAfterModuleStartup();
+	}
+} //NiagaraStatelessCommon
+
+bool FNiagaraStatelessSpawnInfo::IsValid(TOptional<float> LoopDuration) const
+{
+	switch (Type)
+	{
+		case ENiagaraStatelessSpawnInfoType::Burst:
+			return (AmountMin + AmountMax) > 0 && SpawnTime >= 0.0f && SpawnTime < LoopDuration.Get(SpawnTime + UE_SMALL_NUMBER);
+
+		case ENiagaraStatelessSpawnInfoType::Rate:
+			return (RateMin + RateMax) > 0.0f;
+
+		default:
+			checkNoEntry();
+			return false;
+	}
+}
+
+const FNiagaraStatelessGlobals& FNiagaraStatelessGlobals::Get()
+{
+	return NiagaraStatelessCommon::GGlobals;
+}
+
diff --git a/Engine/Plugins/FX/Niagara/Source/Niagara/Private/Stateless/NiagaraStatelessComputeManager.cpp b/Engine/Plugins/FX/Niagara/Source/Niagara/Private/Stateless/NiagaraStatelessComputeManager.cpp
new file mode 100644
--- /dev/null
+++ b/Engine/Plugins/FX/Niagara/Source/Niagara/Private/Stateless/NiagaraStatelessComputeManager.cpp
@@ -0,0 +1,235 @@
+// Copyright Epic Games, Inc. All Rights Reserved.
+
+#include "Stateless/NiagaraStatelessComputeManager.h"
+#include "Stateless/NiagaraStatelessEmitterData.h"
+#include "Stateless/NiagaraStatelessEmitterInstance.h"
+#include "Stateless/NiagaraStatelessSimulationShader.h"
+
+#include "NiagaraGpuComputeDispatchInterface.h"
+#include "NiagaraGPUInstanceCountManager.h"
+
+#include "GPUSortManager.h"
+#include "PrimitiveSceneProxy.h"
+#include "RenderGraphBuilder.h"
+#include "RenderGraphUtils.h"
+#include "SceneView.h"
+
+bool GNiagaraStatelessComputeManager_UseCache = 1;
+FAutoConsoleVariableRef CVarNiagaraStatelessComputeManager_UseCache(
+	TEXT("fx.NiagaraStateless.ComputeManager.UseCache"),
+	GNiagaraStatelessComputeManager_UseCache,
+	TEXT("When enabled we will attempt to reuse allocated buffers between frames."),
+	ECVF_Default
+);
+
+FNiagaraStatelessComputeManager::FNiagaraStatelessComputeManager(FNiagaraGpuComputeDispatchInterface* InOwnerInterface)
+	: FNiagaraGpuComputeDataManager(InOwnerInterface)
+{
+	InOwnerInterface->GetOnPreRenderEvent().AddRaw(this, &FNiagaraStatelessComputeManager::OnPostPreRender);
+	InOwnerInterface->GetOnPostRenderEvent().AddRaw(this, &FNiagaraStatelessComputeManager::OnPostPostRender);
+}
+
+FNiagaraStatelessComputeManager::~FNiagaraStatelessComputeManager()
+{
+}
+
+FNiagaraDataBuffer* FNiagaraStatelessComputeManager::GetDataBuffer(uintptr_t EmitterKey, const NiagaraStateless::FEmitterInstance_RT* EmitterInstance)
+{
+	if (TUniquePtr<FStatelessDataCache>* ExistingData = UsedData.Find(EmitterKey))
+	{
+		return (*ExistingData)->DataBuffer;
+	}
+
+	if (EmitterInstance->ExecutionState == ENiagaraExecutionState::Complete || EmitterInstance->ExecutionState == ENiagaraExecutionState::Disabled)
+	{
+		return nullptr;
+	}
+
+	const FNiagaraStatelessEmitterData* EmitterData = EmitterInstance->EmitterData.Get();
+	TShaderRef<NiagaraStateless::FSimulationShader> ComputeShader = EmitterData->GetShader();
+	if (!ComputeShader.IsValid() || !ComputeShader.GetComputeShader())
+	{
+		return nullptr;
+	}
+
+	uint32 ActiveParticles = 0;
+	{
+		NiagaraStateless::FCommonShaderParameters* ShaderParameters = EmitterInstance->ShaderParameters.Get();
+		ActiveParticles = EmitterData->CalculateActiveParticles(
+			EmitterInstance->RandomSeed,
+			EmitterInstance->SpawnInfos,
+			EmitterInstance->Age,
+			&ShaderParameters->SpawnParameters
+		);
+	}
+	if (ActiveParticles == 0)
+	{
+		return nullptr;
+	}
+
+	FNiagaraGpuComputeDispatchInterface* ComputeInterface = GetOwnerInterface();
+
+	const uint32 DataSetLayoutHash = EmitterInstance->EmitterData->ParticleDataSetCompiledData.GetLayoutHash();
+
+	FStatelessDataCache* CacheData = nullptr;
+	if (GNiagaraStatelessComputeManager_UseCache)
+	{
+		for (int32 i=0; i < FreeData.Num(); ++i)
+		{
+			if (FreeData[i]->DataSetLayoutHash == DataSetLayoutHash)
+			{
+				CacheData = FreeData[i].Release();
+				FreeData.RemoveAtSwap(i, 1, false);
+				break;
+			}
+		}
+	}
+
+	if ( CacheData == nullptr )
+	{
+		CacheData = new FStatelessDataCache();
+		CacheData->DataSetLayoutHash = DataSetLayoutHash;
+		CacheData->DataSet.Init(&EmitterInstance->EmitterData->ParticleDataSetCompiledData);
+		CacheData->DataBuffer = new FNiagaraDataBuffer(&CacheData->DataSet);
+	}
+
+	CacheData->EmitterInstance = EmitterInstance;
+	CacheData->ActiveParticles = ActiveParticles;
+
+	CacheData->DataBuffer->AllocateGPU(FRHICommandListExecutor::GetImmediateCommandList(), CacheData->ActiveParticles, ComputeInterface->GetFeatureLevel(), TEXT("StatelessSimBuffer"));
+	CacheData->DataBuffer->SetNumInstances(CacheData->ActiveParticles);
+
+	FNiagaraGPUInstanceCountManager& CountManager = ComputeInterface->GetGPUInstanceCounterManager();
+	const uint32 CountOffset = CountManager.AcquireOrAllocateEntry(FRHICommandListExecutor::GetImmediateCommandList());
+	CacheData->DataBuffer->SetGPUInstanceCountBufferOffset(CountOffset);
+	CountsToRelease.Add(CountOffset);
+
+	UsedData.Emplace(EmitterKey, CacheData);
+	return CacheData->DataBuffer;
+}
+
+void FNiagaraStatelessComputeManager::OnPostPreRender(FRDGBuilder& GraphBuilder)
+{
+	// Anything to process?
+	if (UsedData.Num() == 0)
+	{
+		return;
+	}
+
+	RDG_CSV_STAT_EXCLUSIVE_SCOPE(GraphBuilder, NiagaraStateless);
+	RDG_GPU_MASK_SCOPE(GraphBuilder, FRHIGPUMask::All());
+
+	AddPass(
+		GraphBuilder,
+		RDG_EVENT_NAME("FNiagaraStatelessComputeManager::OnPostPreRender"),
+		[this](FRHICommandListImmediate& RHICmdList)
+		{
+			const int32 NumJobs = UsedData.Num();
+
+			// Get Count Buffer
+			FNiagaraGpuComputeDispatchInterface* ComputeInterface = GetOwnerInterface();
+			FNiagaraGPUInstanceCountManager& CountManager = ComputeInterface->GetGPUInstanceCounterManager();
+			FRHIUnorderedAccessView* CountBufferUAV = CountManager.GetInstanceCountBuffer().UAV;
+
+			// Build Transitions
+			TArray<FRHITransitionInfo> TransitionsBefore;
+			TArray<FRHITransitionInfo> TransitionsAfter;
+			{
+				TransitionsBefore.Reserve(1 + (NumJobs * 2));
+				TransitionsAfter.Reserve(1 + (NumJobs * 2));
+
+				TransitionsBefore.Emplace(CountManager.GetInstanceCountBuffer().Buffer, FNiagaraGPUInstanceCountManager::kCountBufferDefaultState, ERHIAccess::UAVCompute);
+				TransitionsAfter.Emplace(CountManager.GetInstanceCountBuffer().Buffer, ERHIAccess::UAVCompute, FNiagaraGPUInstanceCountManager::kCountBufferDefaultState);
+
+				for (auto it = UsedData.CreateIterator(); it; ++it)
+				{
+					FStatelessDataCache* CacheData = it.Value().Get();
+
+					const FRWBuffer& FloatBuffer = CacheData->DataBuffer->GetGPUBufferFloat();
+					if (FloatBuffer.NumBytes > 0)
+					{
+						TransitionsBefore.Emplace(FloatBuffer.Buffer, ERHIAccess::SRVMask, ERHIAccess::UAVCompute);
+						TransitionsAfter.Emplace(FloatBuffer.Buffer, ERHIAccess::UAVCompute, ERHIAccess::SRVMask);
+					}
+					const FRWBuffer& IntBuffer = CacheData->DataBuffer->GetGPUBufferInt();
+					if (IntBuffer.NumBytes > 0)
+					{
+						TransitionsBefore.Emplace(IntBuffer.Buffer, ERHIAccess::SRVMask, ERHIAccess::UAVCompute);
+						TransitionsAfter.Emplace(IntBuffer.Buffer, ERHIAccess::UAVCompute, ERHIAccess::SRVMask);
+					}
+				}
+			}
+
+			FNiagaraGpuComputeDispatchInterface* ComputeDispatchInterface = GetOwnerInterface();
+			FNiagaraEmptyUAVPoolScopedAccess UAVPoolAccessScope(ComputeDispatchInterface->GetEmptyUAVPool());
+			FRHIUnorderedAccessView* EmptyIntBufferUAV = ComputeDispatchInterface->GetEmptyUAVFromPool(RHICmdList, PF_R32_SINT, ENiagaraEmptyUAVType::Buffer);
+
+			// Execute Simulations
+			RHICmdList.Transition(TransitionsBefore);
+
+			RHICmdList.BeginUAVOverlap(CountBufferUAV);
+			for (auto it = UsedData.CreateIterator(); it; ++it)
+			{
+				FStatelessDataCache* CacheData = it.Value().Get();
+				const FNiagaraStatelessEmitterData* EmitterData = CacheData->EmitterInstance->EmitterData.Get();
+
+				// Update parameters for this compute invocation
+				NiagaraStateless::FCommonShaderParameters* ShaderParameters = CacheData->EmitterInstance->ShaderParameters.Get();
+				ShaderParameters->Common_SimulationTime			= CacheData->EmitterInstance->Age;
+				ShaderParameters->Common_SimulationDeltaTime	= 1.0f / 60.0f;		//-TODO: Pull from view information, needs reworking of how we link to the dispatch
+				ShaderParameters->Common_SimulationInvDeltaTime	= 60.0f;
+				ShaderParameters->Common_OutputBufferStride		= CacheData->DataBuffer->GetFloatStride() / sizeof(float);
+				ShaderParameters->Common_GPUCountBufferOffset	= CacheData->DataBuffer->GetGPUInstanceCountBufferOffset();
+				ShaderParameters->Common_FloatOutputBuffer		= CacheData->DataBuffer->GetGPUBufferFloat().UAV;
+				//ShaderParameters->Common_HalfOutputBuffer		= CacheData->DataBuffer->GetGPUBufferHalf().UAV;
+				ShaderParameters->Common_IntOutputBuffer		= CacheData->DataBuffer->GetGPUBufferInt().UAV.IsValid() ? CacheData->DataBuffer->GetGPUBufferInt().UAV.GetReference() : EmptyIntBufferUAV;
+				ShaderParameters->Common_GPUCountBuffer			= CountBufferUAV;
+				ShaderParameters->Common_StaticFloatBuffer		= EmitterData->StaticFloatBuffer.SRV;
+
+				// Execute the simulation
+				TShaderRef<NiagaraStateless::FSimulationShader> ComputeShader = EmitterData->GetShader();
+				FRHIComputeShader* ComputeShaderRHI = ComputeShader.GetComputeShader();
+				const uint32 NumThreadGroups = FMath::DivideAndRoundUp<uint32>(CacheData->ActiveParticles, NiagaraStateless::FSimulationShader::ThreadGroupSize);
+
+				SetComputePipelineState(RHICmdList, ComputeShaderRHI);
+				SetShaderParameters(RHICmdList, ComputeShader, ComputeShaderRHI, EmitterData->GetShaderParametersMetadata(), *ShaderParameters);
+				RHICmdList.DispatchComputeShader(NumThreadGroups, 1, 1);
+				UnsetShaderUAVs(RHICmdList, ComputeShader, ComputeShaderRHI);
+			}
+			RHICmdList.EndUAVOverlap(CountBufferUAV);
+
+			RHICmdList.Transition(TransitionsAfter);
+		}
+	);
+}
+
+void FNiagaraStatelessComputeManager::OnPostPostRender(FRDGBuilder& GraphBuilder)
+{
+	// Anything to process?
+	if (UsedData.Num() + FreeData.Num() + CountsToRelease.Num() == 0)
+	{
+		return;
+	}
+
+	RDG_CSV_STAT_EXCLUSIVE_SCOPE(GraphBuilder, NiagaraStateless);
+	RDG_GPU_MASK_SCOPE(GraphBuilder, FRHIGPUMask::All());
+
+	AddPass(
+		GraphBuilder,
+		RDG_EVENT_NAME("FNiagaraStatelessComputeManager::OnPostPostRender"),
+		[this](FRHICommandListImmediate& RHICmdList)
+		{
+			FreeData.Empty();
+			for (auto it=UsedData.CreateIterator(); it; ++it)
+			{
+				FreeData.Emplace(it.Value().Release());
+			}
+			UsedData.Empty();
+
+			FNiagaraGpuComputeDispatchInterface* ComputeInterface = GetOwnerInterface();
+			FNiagaraGPUInstanceCountManager& CountManager = ComputeInterface->GetGPUInstanceCounterManager();
+			CountManager.FreeEntryArray(CountsToRelease);
+			CountsToRelease.Reset();
+		}
+	);
+}
diff --git a/Engine/Plugins/FX/Niagara/Source/Niagara/Private/Stateless/NiagaraStatelessDrawDebugContext.cpp b/Engine/Plugins/FX/Niagara/Source/Niagara/Private/Stateless/NiagaraStatelessDrawDebugContext.cpp
new file mode 100644
--- /dev/null
+++ b/Engine/Plugins/FX/Niagara/Source/Niagara/Private/Stateless/NiagaraStatelessDrawDebugContext.cpp
@@ -0,0 +1,102 @@
+// Copyright Epic Games, Inc. All Rights Reserved.
+
+#include "Stateless/NiagaraStatelessDrawDebugContext.h"
+#include "DrawDebugHelpers.h"
+
+#include "Components/LineBatchComponent.h"
+#include "Engine/World.h"
+
+#if WITH_EDITOR
+void FNiagaraStatelessDrawDebugContext::DrawArrow(const FVector& Origin, const FVector& DirectionWithLength, const FColor& Color) const
+{
+	const float Len = DirectionWithLength.Length();
+	if (Len > UE_SMALL_NUMBER)
+	{
+		DrawDebugDirectionalArrow(
+			World,
+			Origin,
+			Origin + DirectionWithLength,
+			Len,
+			Color
+		);
+	}
+}
+
+void FNiagaraStatelessDrawDebugContext::DrawBox(const FVector& Center, const FVector& Extent, const FColor& Color, const FQuat& Rotation) const
+{
+	DrawDebugBox(
+		World,
+		LocalToWorldTransform.TransformPosition(Center),
+		Extent,
+		LocalToWorldTransform.TransformRotation(Rotation),
+		Color
+	);
+}
+
+void FNiagaraStatelessDrawDebugContext::DrawCone(const FVector& Origin, const FQuat& Rotation, float Angle, float Length, const FColor& Color) const
+{
+	DrawDebugCone(
+		World,
+		Origin,
+		Rotation.GetAxisZ(),
+		Length,
+		FMath::DegreesToRadians(Angle),
+		FMath::DegreesToRadians(Angle),
+		16,
+		Color
+	);
+}
+
+void FNiagaraStatelessDrawDebugContext::DrawCylinder(float CylinderHeight, float CylinderRadius, float CylinderHeightMidpoint, const FColor& Color) const
+{
+	const FVector Axis = LocalToWorldTransform.TransformVector(FVector::ZAxisVector);
+	const FVector Center = LocalToWorldTransform.TransformPosition(FVector::ZeroVector);
+	const float Offset = -CylinderHeightMidpoint * CylinderHeight;
+
+	DrawDebugCylinder(
+		World,
+		Center + (Axis * Offset),
+		Center + (Axis * (Offset + CylinderHeight)),
+		CylinderRadius,
+		12,
+		Color
+	);
+}
+
+void FNiagaraStatelessDrawDebugContext::DrawCircle(const FVector& InCenter, const float Radius, const FColor& Color) const
+{
+	if (ULineBatchComponent* LineBatcher = World->LineBatcher)
+	{
+		int32 Segments = 16;
+		const float AngleStep = 2.f * UE_PI / float(Segments);
+
+		const FVector Center = LocalToWorldTransform.TransformVector(InCenter);
+		const FVector AxisX = LocalToWorldTransform.TransformVector(FVector::XAxisVector);
+		const FVector AxisY = LocalToWorldTransform.TransformVector(FVector::YAxisVector);
+
+		TArray<FBatchedLine> Lines;
+		Lines.Empty(Segments);
+
+		float Angle = 0.f;
+		while (Segments--)
+		{
+			const FVector Vertex1 = Center + Radius * (AxisX * FMath::Cos(Angle) + AxisY * FMath::Sin(Angle));
+			Angle += AngleStep;
+			const FVector Vertex2 = Center + Radius * (AxisX * FMath::Cos(Angle) + AxisY * FMath::Sin(Angle));
+			Lines.Emplace(Vertex1, Vertex2, Color, -1.0f, 0.0f, 0);
+		}
+		LineBatcher->DrawLines(Lines);
+	}
+}
+
+void FNiagaraStatelessDrawDebugContext::DrawSphere(const FVector& Center, const float Radius, const FColor& Color) const
+{
+	DrawDebugSphere(
+		World,
+		LocalToWorldTransform.TransformPosition(Center),
+		Radius,
+		16,
+		Color
+	);
+}
+#endif
diff --git a/Engine/Plugins/FX/Niagara/Source/Niagara/Private/Stateless/NiagaraStatelessEmitterData.cpp b/Engine/Plugins/FX/Niagara/Source/Niagara/Private/Stateless/NiagaraStatelessEmitterData.cpp
new file mode 100644
--- /dev/null
+++ b/Engine/Plugins/FX/Niagara/Source/Niagara/Private/Stateless/NiagaraStatelessEmitterData.cpp
@@ -0,0 +1,146 @@
+// Copyright Epic Games, Inc. All Rights Reserved.
+
+#include "Stateless/NiagaraStatelessEmitterData.h"
+#include "Stateless/NiagaraStatelessEmitterTemplate.h"
+#include "Stateless/NiagaraStatelessSimulationShader.h"
+
+void FNiagaraStatelessEmitterData::FDeleter::operator()(FNiagaraStatelessEmitterData* EmitterData) const
+{
+	if (IsInRenderingThread())
+	{
+		delete EmitterData;
+	}
+	else
+	{
+		ENQUEUE_RENDER_COMMAND(DeleteNiagaraStatelessEmitterData)
+		(
+			[EmitterData](FRHICommandListImmediate&)
+			{
+				delete EmitterData;
+			}
+		);
+	}
+}
+
+FNiagaraStatelessEmitterData::~FNiagaraStatelessEmitterData()
+{
+	check(IsInRenderingThread());
+	StaticFloatBuffer.Release();
+}
+
+void FNiagaraStatelessEmitterData::InitRenderResources()
+{
+	ENQUEUE_RENDER_COMMAND(InitNiagaraStatelessEmitterData)
+	(
+		[this](FRHICommandListImmediate& RHICmdList)
+		{
+			StaticFloatBuffer.Initialize(RHICmdList, TEXT("NiagaraStatelessEmitterData_StaticFloatBuffer"), sizeof(float), StaticFloatData.Num(), PF_R32_FLOAT, EBufferUsageFlags::Static);
+			void* BufferData = RHICmdList.LockBuffer(StaticFloatBuffer.Buffer, 0, StaticFloatBuffer.NumBytes, RLM_WriteOnly);
+			FMemory::Memcpy(BufferData, StaticFloatData.GetData(), StaticFloatBuffer.NumBytes);
+			RHICmdList.UnlockBuffer(StaticFloatBuffer.Buffer);
+		}
+	);
+}
+
+TShaderRef<NiagaraStateless::FSimulationShader>	FNiagaraStatelessEmitterData::GetShader() const
+{
+	if (FApp::CanEverRender())
+	{
+		check(EmitterTemplate);
+		return EmitterTemplate->GetSimulationShader();
+	}
+	return TShaderRef<NiagaraStateless::FSimulationShader>();
+}
+
+const FShaderParametersMetadata* FNiagaraStatelessEmitterData::GetShaderParametersMetadata() const
+{
+	check(EmitterTemplate);
+	return EmitterTemplate->GetShaderParametersMetadata();
+}
+
+//float FNiagaraStatelessEmitterData::CalculateCompletionAge(int32 InRandomSeed, TConstArrayView<FNiagaraStatelessRuntimeSpawnInfo> RuntimeSpawnInfos) const
+//{
+//	float CompletionAge = -1.0f;
+//	//for (const FNiagaraStatelessSpawnInfo& SpawnInfo : SpawnInfos)
+//	//{
+//	//	CompletionAge = FMath::Max(CompletionAge, SpawnInfo.AgeEnd);
+//	//}
+//
+//	return CompletionAge >= 0.0f ? CompletionAge + LifetimeRange.Y : 0.0f;
+//}
+
+uint32 FNiagaraStatelessEmitterData::CalculateActiveParticles(int32 InRandomSeed, TConstArrayView<FNiagaraStatelessRuntimeSpawnInfo> RuntimeSpawnInfos, TOptional<float> Age, NiagaraStateless::FSpawnInfoShaderParameters* SpawnParameters) const
+{
+	int32	GpuSpawnIndex = 0;
+	uint32	TotalActiveParticles = 0;
+
+	for (const FNiagaraStatelessRuntimeSpawnInfo& SpawnInfo : RuntimeSpawnInfos)
+	{
+		uint32 SpawnInfoTotalParticles = 0;
+
+		const bool bIsValidForAge = !Age.IsSet() || (Age.GetValue() >= SpawnInfo.SpawnTimeStart && Age.GetValue() < SpawnInfo.SpawnTimeEnd + LifetimeRange.Y);
+		if (bIsValidForAge && GpuSpawnIndex < NiagaraStateless::MaxGpuSpawnInfos)
+		{
+			uint32	NumActive = 0;
+			uint32	ParticleOffset = 0;
+			float	SpawnRate = 0.0f;
+			float	SpawnTimeStart = SpawnInfo.SpawnTimeStart;
+			switch (SpawnInfo.Type)
+			{
+				case ENiagaraStatelessSpawnInfoType::Burst:
+					NumActive				= SpawnInfo.Amount;
+					break;
+
+				case ENiagaraStatelessSpawnInfoType::Rate:
+					if (!Age.IsSet())
+					{
+						NumActive = FMath::FloorToInt((SpawnInfo.SpawnTimeEnd - SpawnInfo.SpawnTimeStart) * SpawnInfo.Rate);
+					}
+					else
+					{
+						const uint32 MaxActive = FMath::FloorToInt((SpawnInfo.SpawnTimeEnd - SpawnInfo.SpawnTimeStart) * SpawnInfo.Rate);
+						ParticleOffset = FMath::FloorToInt(FMath::Max(Age.GetValue() - SpawnInfo.SpawnTimeStart - LifetimeRange.Y, 0.0f) * SpawnInfo.Rate);
+						ParticleOffset = FMath::Min(ParticleOffset, MaxActive);
+						NumActive = FMath::FloorToInt(FMath::Max(Age.GetValue() - SpawnInfo.SpawnTimeStart, 0.0f) * SpawnInfo.Rate);
+						NumActive = FMath::Min(NumActive, MaxActive);
+						NumActive -= ParticleOffset;
+					}
+					SpawnRate = 1.0f / SpawnInfo.Rate;
+					SpawnTimeStart += SpawnRate;
+					break;
+
+				default:
+					checkNoEntry();
+					break;
+			}
+			if (NumActive > 0)
+			{
+				if (SpawnParameters)
+				{
+					GET_SCALAR_ARRAY_ELEMENT(SpawnParameters->SpawnInfo_NumActive, GpuSpawnIndex)		= NumActive;
+					GET_SCALAR_ARRAY_ELEMENT(SpawnParameters->SpawnInfo_ParticleOffset, GpuSpawnIndex)	= ParticleOffset;
+					GET_SCALAR_ARRAY_ELEMENT(SpawnParameters->SpawnInfo_UniqueOffset, GpuSpawnIndex)	= SpawnInfo.UniqueOffset;
+					GET_SCALAR_ARRAY_ELEMENT(SpawnParameters->SpawnInfo_Time, GpuSpawnIndex)			= SpawnTimeStart;
+					GET_SCALAR_ARRAY_ELEMENT(SpawnParameters->SpawnInfo_Rate, GpuSpawnIndex)			= SpawnRate;
+					++GpuSpawnIndex;
+				}
+				TotalActiveParticles += NumActive;
+			}
+		}
+	}
+
+	if (SpawnParameters)
+	{
+		while (GpuSpawnIndex < NiagaraStateless::MaxGpuSpawnInfos)
+		{
+			GET_SCALAR_ARRAY_ELEMENT(SpawnParameters->SpawnInfo_NumActive, GpuSpawnIndex)		= 0;
+			GET_SCALAR_ARRAY_ELEMENT(SpawnParameters->SpawnInfo_ParticleOffset, GpuSpawnIndex)	= 0;
+			GET_SCALAR_ARRAY_ELEMENT(SpawnParameters->SpawnInfo_UniqueOffset, GpuSpawnIndex)	= 0;
+			GET_SCALAR_ARRAY_ELEMENT(SpawnParameters->SpawnInfo_Time, GpuSpawnIndex)			= 0.0f;
+			GET_SCALAR_ARRAY_ELEMENT(SpawnParameters->SpawnInfo_Rate, GpuSpawnIndex)			= 0.0f;
+			++GpuSpawnIndex;
+		}
+	}
+	return TotalActiveParticles;
+}
+
diff --git a/Engine/Plugins/FX/Niagara/Source/Niagara/Private/Stateless/NiagaraStatelessEmitterDataBuildContext.cpp b/Engine/Plugins/FX/Niagara/Source/Niagara/Private/Stateless/NiagaraStatelessEmitterDataBuildContext.cpp
new file mode 100644
--- /dev/null
+++ b/Engine/Plugins/FX/Niagara/Source/Niagara/Private/Stateless/NiagaraStatelessEmitterDataBuildContext.cpp
@@ -0,0 +1,74 @@
+// Copyright Epic Games, Inc. All Rights Reserved.
+
+#include "Stateless/NiagaraStatelessEmitterDataBuildContext.h"
+
+#include "NiagaraParameterBinding.h"
+#include "NiagaraParameterStore.h"
+
+uint32 FNiagaraStatelessEmitterDataBuildContext::AddStaticData(TConstArrayView<float> FloatData)
+{
+	for (int32 i=0; i <= StaticFloatData.Num() - FloatData.Num(); ++i)
+	{
+		if (FMemory::Memcmp(StaticFloatData.GetData() + i, FloatData.GetData(), FloatData.Num() * sizeof(float)) == 0)
+		{
+			return i;
+		}
+	}
+
+	// Add new
+	const uint32 OutIndex = StaticFloatData.AddUninitialized(FloatData.Num());
+	FMemory::Memcpy(StaticFloatData.GetData() + OutIndex, FloatData.GetData(), FloatData.Num() * sizeof(float));
+	return OutIndex;
+}
+
+uint32 FNiagaraStatelessEmitterDataBuildContext::AddStaticData(TConstArrayView<FVector2f> FloatData)
+{
+	return AddStaticData(MakeArrayView(reinterpret_cast<const float*>(FloatData.GetData()), FloatData.Num() * 2));
+}
+
+uint32 FNiagaraStatelessEmitterDataBuildContext::AddStaticData(TConstArrayView<FVector3f> FloatData)
+{
+	return AddStaticData(MakeArrayView(reinterpret_cast<const float*>(FloatData.GetData()), FloatData.Num() * 3));
+}
+
+uint32 FNiagaraStatelessEmitterDataBuildContext::AddStaticData(TConstArrayView<FVector4f> FloatData)
+{
+	return AddStaticData(MakeArrayView(reinterpret_cast<const float*>(FloatData.GetData()), FloatData.Num() * 4));
+}
+
+uint32 FNiagaraStatelessEmitterDataBuildContext::AddStaticData(TConstArrayView<FLinearColor> FloatData)
+{
+	return AddStaticData(MakeArrayView(reinterpret_cast<const float*>(FloatData.GetData()), FloatData.Num() * 4));
+}
+
+int32 FNiagaraStatelessEmitterDataBuildContext::AddRendererBinding(const FNiagaraParameterBinding& Binding)
+{
+	int32 DataOffset = INDEX_NONE;
+	if (Binding.ResolvedParameter.IsValid())
+	{
+		FNiagaraVariable Var(Binding.ResolvedParameter);
+		RendererBindings.AddParameter(Var, false, false, &DataOffset);
+	}
+
+	return DataOffset;
+}
+
+int32 FNiagaraStatelessEmitterDataBuildContext::AddRendererBinding(const FNiagaraParameterBindingWithValue& Binding)
+{
+	int32 DataOffset = INDEX_NONE;
+	if (Binding.ResolvedParameter.IsValid())
+	{
+		FNiagaraVariable Var(Binding.ResolvedParameter);
+		RendererBindings.AddParameter(Var, false, false, &DataOffset);
+
+		TConstArrayView<uint8> DefaultValue = Binding.GetDefaultValueArray();
+		if (DefaultValue.Num() > 0)
+		{
+			check(DataOffset != INDEX_NONE);
+			check(DefaultValue.Num() == Var.GetSizeInBytes());
+			RendererBindings.SetParameterData(DefaultValue.GetData(), DataOffset, DefaultValue.Num());
+		}
+	}
+
+	return DataOffset;
+}
diff --git a/Engine/Plugins/FX/Niagara/Source/Niagara/Private/Stateless/NiagaraStatelessEmitterInstance.cpp b/Engine/Plugins/FX/Niagara/Source/Niagara/Private/Stateless/NiagaraStatelessEmitterInstance.cpp
new file mode 100644
--- /dev/null
+++ b/Engine/Plugins/FX/Niagara/Source/Niagara/Private/Stateless/NiagaraStatelessEmitterInstance.cpp
@@ -0,0 +1,501 @@
+// Copyright Epic Games, Inc. All Rights Reserved.
+
+#include "Stateless/NiagaraStatelessEmitterInstance.h"
+#include "Stateless/NiagaraStatelessEmitter.h"
+#include "Stateless/NiagaraStatelessEmitterData.h"
+#include "Stateless/NiagaraStatelessComputeManager.h"
+#include "NiagaraGpuComputeDispatchInterface.h"
+#include "NiagaraSystem.h"
+#include "NiagaraSystemInstance.h"
+
+#include "Shader.h"
+
+FNiagaraDataBuffer* NiagaraStateless::FEmitterInstance_RT::GetDataToRender(bool bIsLowLatencyTranslucent) const
+{
+	return ComputeManager ? ComputeManager->GetDataBuffer(uintptr_t(this), this) : nullptr;
+}
+
+FNiagaraStatelessEmitterInstance::FNiagaraStatelessEmitterInstance(FNiagaraSystemInstance* InParentSystemInstance)
+	: Super(InParentSystemInstance)
+{
+	// Setup base properties
+	bLocalSpace = true;
+	SimTarget = ENiagaraSimTarget::GPUComputeSim;
+	bNeedsPartialDepthTexture = false;
+	ParticleDataSet = new FNiagaraDataSet();
+}
+
+FNiagaraStatelessEmitterInstance::~FNiagaraStatelessEmitterInstance()
+{
+	//-TODO: Should we move this into the base class?
+	UnbindParameters(false);
+
+	NiagaraStateless::FEmitterInstance_RT* RenderThreadData = RenderThreadDataPtr.Release();
+	if (RenderThreadData || ParticleDataSet)
+	{
+		ENQUEUE_RENDER_COMMAND(FReleaseStatelessEmitter)(
+			[RenderThreadData, ParticleDataSet=ParticleDataSet](FRHICommandListImmediate& RHICmdList)
+			{
+				if (RenderThreadData != nullptr)
+				{
+					delete RenderThreadData;
+				}
+				if (ParticleDataSet != nullptr )
+				{
+					delete ParticleDataSet;
+				}
+			}
+		);
+		ParticleDataSet = nullptr;
+	}
+}
+
+void FNiagaraStatelessEmitterInstance::Init(int32 InEmitterIndex)
+{
+	Super::Init(InEmitterIndex);
+
+	// Initialize the EmitterData ptr if this is invalid the emitter is not allowed to run
+	InitEmitterData();
+	if (!bCanEverExecute)
+	{
+		ExecutionState = ENiagaraExecutionState::Disabled;
+		return;
+	}
+
+	// Pull out information
+	RandomSeed = EmitterData->RandomSeed + ParentSystemInstance->GetRandomSeedOffset();
+	if (!EmitterData->bDeterministic)
+	{
+		RandomSeed ^= FPlatformTime::Cycles();
+	}
+	RandomStream.Initialize(RandomSeed);
+	//CompletionAge = EmitterData->CalculateCompletionAge(RandomSeed);
+
+	// Initialize data set
+	ParticleDataSet->Init(&EmitterData->ParticleDataSetCompiledData);
+
+	// Prepare our parameters
+	RendererBindings = EmitterData->RendererBindings;
+
+	// Allocate and fill shader parameters
+	RenderThreadDataPtr.Reset(new NiagaraStateless::FEmitterInstance_RT());
+	RenderThreadDataPtr->EmitterData = EmitterData;
+	RenderThreadDataPtr->RandomSeed = RandomSeed;
+	RenderThreadDataPtr->Age = 0.0f;
+	RenderThreadDataPtr->ExecutionState = ENiagaraExecutionState::Active;
+	RenderThreadDataPtr->ShaderParameters.Reset(WeakStatelessEmitter->AllocateShaderParameters(RendererBindings));
+	RenderThreadDataPtr->ShaderParameters->Common_RandomSeed = RandomSeed;
+
+	ENQUEUE_RENDER_COMMAND(FInitStatelessEmitter)(
+		[RenderThreadData=RenderThreadDataPtr.Get(), ComputeInterface = ParentSystemInstance->GetComputeDispatchInterface()](FRHICommandListImmediate& RHICmdList)
+		{
+			RenderThreadData->ComputeManager = &ComputeInterface->GetOrCreateDataManager<FNiagaraStatelessComputeManager>();
+		}
+	);
+
+	GPUDataBufferInterfaces = RenderThreadDataPtr.Get();
+
+	InitEmitterState();
+	InitSpawnInfos();
+}
+
+void FNiagaraStatelessEmitterInstance::ResetSimulation(bool bKillExisting)
+{
+	if (!bCanEverExecute)
+	{
+		return;
+	}
+
+	if (bKillExisting)
+	{
+		SpawnInfos.Empty();
+	}
+	else
+	{
+		for (FNiagaraStatelessRuntimeSpawnInfo& SpawnInfo : SpawnInfos)
+		{
+			SpawnInfo.SpawnTimeStart -= Age;
+			SpawnInfo.SpawnTimeEnd -= Age;
+		}
+	}
+	ActiveSpawnRates.Empty();
+
+	RandomStream.Initialize(RandomSeed);
+
+	Age = 0.0f;
+	UniqueIndexOffset = 0;
+
+	InitEmitterState();
+	InitSpawnInfos();
+
+	ExecutionState = ENiagaraExecutionState::Active;
+	if (NiagaraStateless::FEmitterInstance_RT* RenderThreadData = RenderThreadDataPtr.Get())
+	{
+		ENQUEUE_RENDER_COMMAND(UpdateStatelessAge)(
+			[RenderThreadData](FRHICommandListImmediate& RHICmdList)
+			{
+				RenderThreadData->Age = 0.0f;
+				RenderThreadData->ExecutionState = ENiagaraExecutionState::Active;
+			}
+		);
+	}
+}
+
+bool FNiagaraStatelessEmitterInstance::HandleCompletion(bool bForce)
+{
+	bool bIsComplete = IsComplete();
+	if (!bIsComplete && bForce)
+	{
+		ExecutionState = ENiagaraExecutionState::Complete;
+		bIsComplete = true;
+
+		if (NiagaraStateless::FEmitterInstance_RT* RenderThreadData = RenderThreadDataPtr.Get())
+		{
+			ENQUEUE_RENDER_COMMAND(CompleteStatelessEmitter)(
+				[RenderThreadData](FRHICommandListImmediate& RHICmdList)
+				{
+					RenderThreadData->ExecutionState = ENiagaraExecutionState::Complete;
+				}
+			);
+		}
+	}
+
+	return bIsComplete;
+}
+
+int32 FNiagaraStatelessEmitterInstance::GetNumParticles() const
+{
+	return bCanEverExecute ? EmitterData->CalculateActiveParticles(RandomSeed, SpawnInfos, Age) : 0;
+}
+
+TConstArrayView<UNiagaraRendererProperties*> FNiagaraStatelessEmitterInstance::GetRenderers() const
+{
+	return EmitterData ? EmitterData->RendererProperties : TConstArrayView<UNiagaraRendererProperties*>();
+}
+
+void FNiagaraStatelessEmitterInstance::BindParameters(bool bExternalOnly)
+{
+	if (!RendererBindings.IsEmpty())
+	{
+		if (ParentSystemInstance)
+		{
+			ParentSystemInstance->BindToParameterStore(RendererBindings);
+		}
+	}
+}
+
+void FNiagaraStatelessEmitterInstance::UnbindParameters(bool bExternalOnly)
+{
+	if (!RendererBindings.IsEmpty())
+	{
+		if (ParentSystemInstance)
+		{
+			ParentSystemInstance->UnbindFromParameterStore(RendererBindings);
+		}
+	}
+}
+
+void FNiagaraStatelessEmitterInstance::Tick(float DeltaSeconds)
+{
+	Age += DeltaSeconds;
+
+	TickSpawnInfos();
+	TickEmitterState();
+	CalculateBounds();
+
+	if (NiagaraStateless::FEmitterInstance_RT* RenderThreadData = RenderThreadDataPtr.Get())
+	{
+		NiagaraStateless::FCommonShaderParameters* NewShaderParameters = nullptr;
+		if ( EmitterData->bModulesHaveRendererBindings && RendererBindings.GetParametersDirty() )
+		{
+			RendererBindings.Tick();
+
+			NewShaderParameters = WeakStatelessEmitter->AllocateShaderParameters(RendererBindings);
+			NewShaderParameters->Common_RandomSeed = RandomSeed;
+		}
+
+		ENQUEUE_RENDER_COMMAND(UpdateStatelessAge)(
+			[RenderThreadData, AgeForRT=Age, ExecutionStateForRT=ExecutionState, NewShaderParameters](FRHICommandListImmediate& RHICmdList)
+			{
+				RenderThreadData->Age = AgeForRT;
+				RenderThreadData->ExecutionState = ExecutionStateForRT;
+				if (NewShaderParameters)
+				{
+					RenderThreadData->ShaderParameters.Reset(NewShaderParameters);
+				}
+			}
+		);
+	}
+}
+
+void FNiagaraStatelessEmitterInstance::InitEmitterState()
+{
+	const FNiagaraEmitterStateData& EmitterState = EmitterData->EmitterState;
+	LoopCount			= 0;
+	CurrentLoopDuration = RandomStream.FRandRange(EmitterState.LoopDurationMin, EmitterState.LoopDurationMax);
+	CurrentLoopDelay	= RandomStream.FRandRange(EmitterState.LoopDelayMin, EmitterState.LoopDelayMax);
+	CurrentLoopAgeStart	= 0.0f;
+	CurrentLoopAgeEnd	= CurrentLoopAgeStart + CurrentLoopDelay + CurrentLoopDuration;
+}
+
+void FNiagaraStatelessEmitterInstance::TickEmitterState()
+{
+	if (ParentSystemInstance)
+	{
+		const ENiagaraExecutionState ParentExecutionState = ParentSystemInstance->GetActualExecutionState();
+		if (ParentExecutionState > ExecutionState)
+		{
+			SetExecutionStateInternal(ParentExecutionState);
+		}
+	}
+
+	if (ExecutionState == ENiagaraExecutionState::Active)
+	{
+		if (Age < CurrentLoopAgeEnd)
+		{
+			return;
+		}
+
+		const FNiagaraEmitterStateData& EmitterState = EmitterData->EmitterState;
+		if (EmitterState.LoopBehavior == ENiagaraLoopBehavior::Once)
+		{
+			SetExecutionStateInternal(ENiagaraExecutionState::Inactive);
+			return;
+		}
+
+		// Keep looping until we find out which loop we are in as a small loop age + large DT could result in crossing multiple loops
+		do
+		{
+			++LoopCount;
+			if (EmitterState.LoopBehavior == ENiagaraLoopBehavior::Multiple && LoopCount >= EmitterState.LoopCount)
+			{
+				SetExecutionStateInternal(ENiagaraExecutionState::Inactive);
+				return;
+			}
+
+			if (EmitterState.bRecalculateDurationEachLoop)
+			{
+				CurrentLoopDuration = RandomStream.FRandRange(EmitterState.LoopDurationMin, EmitterState.LoopDurationMax);
+			}
+
+			if (EmitterState.bDelayFirstLoopOnly)
+			{
+				CurrentLoopDelay	= 0.0f;
+			}
+			else if ( EmitterState.bRecalculateDelayEachLoop )
+			{
+				CurrentLoopDelay	= RandomStream.FRandRange(EmitterState.LoopDelayMin, EmitterState.LoopDelayMax);
+			}
+
+			CurrentLoopAgeStart	= CurrentLoopAgeEnd;
+			CurrentLoopAgeEnd	= CurrentLoopAgeStart + CurrentLoopDelay + CurrentLoopDuration;
+			InitSpawnInfosForLoop();
+		} while ( Age >= CurrentLoopAgeEnd );
+	}
+	else if (ExecutionState == ENiagaraExecutionState::Inactive)
+	{
+		if (SpawnInfos.Num() == 0 || GetNumParticles() == 0)
+		{
+			SetExecutionStateInternal(ENiagaraExecutionState::Complete);
+		}
+		return;
+	}
+}
+
+void FNiagaraStatelessEmitterInstance::CalculateBounds()
+{
+	CachedBounds.Init();
+	FRWScopeLock ScopeLock(FixedBoundsGuard, SLT_ReadOnly);
+	if (FixedBounds.IsValid)
+	{
+		CachedBounds = FixedBounds;
+	}
+	else if (CachedSystemFixedBounds.IsValid)
+	{
+		CachedBounds = CachedSystemFixedBounds;
+	}
+	else
+	{
+		CachedBounds = EmitterData->FixedBounds;
+	}
+}
+
+void FNiagaraStatelessEmitterInstance::InitSpawnInfos()
+{
+	UniqueIndexOffset = 0;
+	for (const FNiagaraStatelessSpawnInfo& SpawnInfo : EmitterData->SpawnInfos)
+	{
+		if (SpawnInfo.Type == ENiagaraStatelessSpawnInfoType::Rate)
+		{
+			const float SpawnRate = RandomStream.FRandRange(SpawnInfo.RateMin, SpawnInfo.RateMax);
+			if (SpawnRate > 0.0f)
+			{
+				FActiveSpawnRate& ActiveSpawnRate = ActiveSpawnRates.AddDefaulted_GetRef();
+				ActiveSpawnRate.Rate		= SpawnRate;
+				ActiveSpawnRate.SpawnTime	= CurrentLoopDelay;
+			}
+		}
+	}
+
+	InitSpawnInfosForLoop();
+}
+
+void FNiagaraStatelessEmitterInstance::InitSpawnInfosForLoop()
+{
+	// Add the next chunk for any active spawn rates
+	for (const FActiveSpawnRate& SpawnInfo : ActiveSpawnRates)
+	{
+		float SpawnTime = FMath::Max(CurrentLoopAgeStart - SpawnInfo.SpawnTime, 0.0f);
+		SpawnTime = FMath::CeilToFloat(SpawnTime * SpawnInfo.Rate) / SpawnInfo.Rate;
+		SpawnTime += SpawnInfo.SpawnTime;
+		if (SpawnTime >= CurrentLoopAgeEnd)
+		{
+			continue;
+		}
+
+		FNiagaraStatelessRuntimeSpawnInfo& NewSpawnInfo = SpawnInfos.AddDefaulted_GetRef();
+		NewSpawnInfo.Type			= ENiagaraStatelessSpawnInfoType::Rate;
+		NewSpawnInfo.UniqueOffset	= UniqueIndexOffset;
+		NewSpawnInfo.SpawnTimeStart	= SpawnTime;
+		NewSpawnInfo.SpawnTimeEnd	= CurrentLoopAgeEnd;
+		NewSpawnInfo.Rate			= SpawnInfo.Rate;
+
+		const float ActiveDuration	= CurrentLoopAgeEnd - SpawnTime;
+		const int32 NumSpawned		= FMath::FloorToInt(ActiveDuration * SpawnInfo.Rate);
+
+		UniqueIndexOffset += NumSpawned;
+	}
+
+	// Add bursts that fit within the loop duration (due to loop random they might not)
+	for (const FNiagaraStatelessSpawnInfo& SpawnInfo : EmitterData->SpawnInfos)
+	{
+		if (SpawnInfo.Type == ENiagaraStatelessSpawnInfoType::Rate || !SpawnInfo.IsValid(CurrentLoopDuration))
+		{
+			continue;
+		}
+
+		if (SpawnInfo.bSpawnProbabilityEnabled && SpawnInfo.SpawnProbability < RandomStream.FRand())
+		{
+			continue;
+		}
+
+		const int32 SpawnAmount = RandomStream.RandRange(SpawnInfo.AmountMin, SpawnInfo.AmountMax);
+		if (SpawnAmount <= 0)
+		{
+			continue;
+		}
+
+		FNiagaraStatelessRuntimeSpawnInfo& NewSpawnInfo = SpawnInfos.AddDefaulted_GetRef();
+		NewSpawnInfo.Type			= ENiagaraStatelessSpawnInfoType::Burst;
+		NewSpawnInfo.UniqueOffset	= UniqueIndexOffset;
+		NewSpawnInfo.SpawnTimeStart	= CurrentLoopAgeStart + CurrentLoopDelay + SpawnInfo.SpawnTime;
+		NewSpawnInfo.SpawnTimeEnd	= NewSpawnInfo.SpawnTimeStart;
+		NewSpawnInfo.Amount			= SpawnAmount;
+
+		UniqueIndexOffset += SpawnAmount;
+	}
+
+	//-TODO: Better way to send data to renderer
+	if (NiagaraStateless::FEmitterInstance_RT* RenderThreadData = RenderThreadDataPtr.Get())
+	{
+		ENQUEUE_RENDER_COMMAND(FInitStatelessEmitter)(
+			[RenderThreadData=RenderThreadDataPtr.Get(), SpawnInfos_RT=SpawnInfos](FRHICommandListImmediate& RHICmdList) mutable
+			{
+				RenderThreadData->SpawnInfos = MoveTemp(SpawnInfos_RT);
+			}
+		);
+	}
+}
+
+void FNiagaraStatelessEmitterInstance::TickSpawnInfos()
+{
+	const float MaxLifetime = EmitterData->LifetimeRange.Y;
+	SpawnInfos.RemoveAll(
+		[this, &MaxLifetime](const FNiagaraStatelessRuntimeSpawnInfo& SpawnInfo)
+		{
+			return Age >= SpawnInfo.SpawnTimeEnd + MaxLifetime;
+		}
+	);
+}
+
+void FNiagaraStatelessEmitterInstance::SetExecutionStateInternal(ENiagaraExecutionState RequestedExecutionState)
+{
+	if (ExecutionState == RequestedExecutionState)
+	{
+		return;
+	}
+	
+	switch (RequestedExecutionState)
+	{
+		case ENiagaraExecutionState::Active:
+			UE_LOG(LogNiagara, Error, TEXT("Stateless Emitter: Was requested to go Active and we do not supoprt that."));
+			break;
+
+		case ENiagaraExecutionState::Inactive:
+			if (EmitterData->EmitterState.InactiveResponse != ENiagaraEmitterInactiveResponse::Kill)
+			{
+				ActiveSpawnRates.Empty();
+
+				// Crop & Remove Spawn Infos
+				const float MaxLifetime = EmitterData->LifetimeRange.Y;
+				for (auto it=SpawnInfos.CreateIterator(); it; ++it)
+				{
+					FNiagaraStatelessRuntimeSpawnInfo& SpawnInfo = *it;
+					if (SpawnInfo.Type == ENiagaraStatelessSpawnInfoType::Rate)
+					{
+						SpawnInfo.SpawnTimeEnd = FMath::Min(SpawnInfo.SpawnTimeEnd, Age);
+					}
+					if (Age < SpawnInfo.SpawnTimeStart || Age >= SpawnInfo.SpawnTimeEnd + MaxLifetime)
+					{
+						it.RemoveCurrent();
+					}
+				}
+
+				if (SpawnInfos.Num() > 0)
+				{
+					//-TODO: Better way to send data to renderer
+					if (NiagaraStateless::FEmitterInstance_RT* RenderThreadData = RenderThreadDataPtr.Get())
+					{
+						ENQUEUE_RENDER_COMMAND(FInitStatelessEmitter)(
+							[RenderThreadData=RenderThreadDataPtr.Get(), SpawnInfos_RT=SpawnInfos](FRHICommandListImmediate& RHICmdList) mutable
+							{
+								RenderThreadData->SpawnInfos = MoveTemp(SpawnInfos_RT);
+							}
+						);
+					}
+
+					ExecutionState = ENiagaraExecutionState::Inactive;
+					break;
+				}
+			}
+			// Intentional fall through as we are complete
+
+		case ENiagaraExecutionState::InactiveClear:
+		case ENiagaraExecutionState::Complete:
+			SpawnInfos.Empty();
+			ActiveSpawnRates.Empty();
+			ExecutionState = ENiagaraExecutionState::Complete;
+			break;
+	}
+}
+
+void FNiagaraStatelessEmitterInstance::InitEmitterData()
+{
+	bCanEverExecute = false;
+	EmitterData = nullptr;
+	WeakStatelessEmitter = nullptr;
+
+	const FNiagaraEmitterHandle& EmitterHandle = GetEmitterHandle();
+	UNiagaraStatelessEmitter* StatelessEmitter = GetEmitterHandle().GetStatelessEmitter();
+	WeakStatelessEmitter = StatelessEmitter;
+	if (!StatelessEmitter)
+	{
+		return;
+	}
+	EmitterData = StatelessEmitter->GetEmitterData();
+
+	bCanEverExecute =
+		EmitterData != nullptr &&
+		EmitterData->bCanEverExecute &&
+		EmitterHandle.GetIsEnabled();
+}
diff --git a/Engine/Plugins/FX/Niagara/Source/Niagara/Private/Stateless/NiagaraStatelessEmitterTemplate.cpp b/Engine/Plugins/FX/Niagara/Source/Niagara/Private/Stateless/NiagaraStatelessEmitterTemplate.cpp
new file mode 100644
--- /dev/null
+++ b/Engine/Plugins/FX/Niagara/Source/Niagara/Private/Stateless/NiagaraStatelessEmitterTemplate.cpp
@@ -0,0 +1,229 @@
+// Copyright Epic Games, Inc. All Rights Reserved.
+
+#include "Stateless/NiagaraStatelessEmitterTemplate.h"
+#include "Stateless/NiagaraStatelessCommon.h"
+#include "Stateless/NiagaraStatelessEmitterShaders.h"
+#include "NiagaraConstants.h"
+#include "NiagaraModule.h"
+
+#include "Modules/ModuleManager.h"
+
+#include "Stateless/Modules/NiagaraStatelessModule_AddVelocity.h"
+#include "Stateless/Modules/NiagaraStatelessModule_AccelerationForce.h"
+#include "Stateless/Modules/NiagaraStatelessModule_CalculateAccurateVelocity.h"
+#include "Stateless/Modules/NiagaraStatelessModule_CurlNoiseForce.h"
+#include "Stateless/Modules/NiagaraStatelessModule_Drag.h"
+#include "Stateless/Modules/NiagaraStatelessModule_GravityForce.h"
+#include "Stateless/Modules/NiagaraStatelessModule_InitializeParticle.h"
+#include "Stateless/Modules/NiagaraStatelessModule_InitialMeshOrientation.h"
+#include "Stateless/Modules/NiagaraStatelessModule_RotateAroundPoint.h"
+#include "Stateless/Modules/NiagaraStatelessModule_ScaleColor.h"
+#include "Stateless/Modules/NiagaraStatelessModule_ScaleMeshSize.h"
+#include "Stateless/Modules/NiagaraStatelessModule_ScaleMeshSizeBySpeed.h"
+#include "Stateless/Modules/NiagaraStatelessModule_ScaleSpriteSize.h"
+#include "Stateless/Modules/NiagaraStatelessModule_ScaleSpriteSizeBySpeed.h"
+#include "Stateless/Modules/NiagaraStatelessModule_ShapeLocation.h"
+#include "Stateless/Modules/NiagaraStatelessModule_SolveVelocitiesAndForces.h"
+#include "Stateless/Modules/NiagaraStatelessModule_SpriteFacingAndAlignment.h"
+#include "Stateless/Modules/NiagaraStatelessModule_SpriteRotationRate.h"
+#include "Stateless/Modules/NiagaraStatelessModule_SubUVAnimation.h"
+
+namespace NiagaraStatelessEmitterTemplateImpl
+{
+	TArray<TWeakObjectPtr<UNiagaraStatelessEmitterTemplate>> ObjectsToDeferredInit;
+}
+
+//////////////////////////////////////////////////////////////////////////////////////////////////////////////////////////////////////////////////////////
+
+void UNiagaraStatelessEmitterTemplate::PostInitProperties()
+{
+	Super::PostInitProperties();
+
+	// We can end up hitting PostInitProperties before the Niagara Module has initialized bindings this needs, mark this object for deferred init and early out.
+	if (FModuleManager::Get().IsModuleLoaded("Niagara") == false)
+	{
+		NiagaraStatelessEmitterTemplateImpl::ObjectsToDeferredInit.Add(this);
+	}
+	else
+	{
+		InitModulesAndAttributes();
+	}
+}
+
+void UNiagaraStatelessEmitterTemplate::InitCDOPropertiesAfterModuleStartup()
+{
+	for (TWeakObjectPtr<UNiagaraStatelessEmitterTemplate>& WeakObject : NiagaraStatelessEmitterTemplateImpl::ObjectsToDeferredInit)
+	{
+		if (UNiagaraStatelessEmitterTemplate* EmitterObject = WeakObject.Get())
+		{
+			EmitterObject->InitModulesAndAttributes();
+		}
+	}
+}
+
+//////////////////////////////////////////////////////////////////////////////////////////////////////////////////////////////////////////////////////////
+
+void UNiagaraStatelessEmitterDefault::InitModulesAndAttributes()
+{
+#if WITH_EDITORONLY_DATA
+	Modules =
+	{
+		UNiagaraStatelessModule_InitializeParticle::StaticClass(),
+		UNiagaraStatelessModule_InitialMeshOrientation::StaticClass(),
+		UNiagaraStatelessModule_ShapeLocation::StaticClass(),
+		UNiagaraStatelessModule_ScaleColor::StaticClass(),
+		UNiagaraStatelessModule_ScaleSpriteSize::StaticClass(),
+		UNiagaraStatelessModule_ScaleSpriteSizeBySpeed::StaticClass(),
+		UNiagaraStatelessModule_ScaleMeshSize::StaticClass(),
+		UNiagaraStatelessModule_ScaleMeshSizeBySpeed::StaticClass(),
+		UNiagaraStatelessModule_AddVelocity::StaticClass(),
+		UNiagaraStatelessModule_AccelerationForce::StaticClass(),
+		UNiagaraStatelessModule_CurlNoiseForce::StaticClass(),
+		UNiagaraStatelessModule_Drag::StaticClass(),
+		UNiagaraStatelessModule_GravityForce::StaticClass(),
+		UNiagaraStatelessModule_SolveVelocitiesAndForces::StaticClass(),
+		UNiagaraStatelessModule_SpriteFacingAndAlignment::StaticClass(),
+		UNiagaraStatelessModule_SpriteRotationRate::StaticClass(),
+		UNiagaraStatelessModule_SubUVAnimation::StaticClass(),
+	};
+
+	const FNiagaraStatelessGlobals& StatelessGlobals = FNiagaraStatelessGlobals::Get();
+	OutputComponents =
+	{
+		StatelessGlobals.UniqueIDVariable,
+		StatelessGlobals.PositionVariable,
+		StatelessGlobals.ColorVariable,
+		StatelessGlobals.MeshOrientationVariable,
+		StatelessGlobals.RibbonWidthVariable,
+		StatelessGlobals.ScaleVariable,
+		StatelessGlobals.SpriteSizeVariable,
+		StatelessGlobals.SpriteFacingVariable,
+		StatelessGlobals.SpriteAlignmentVariable,
+		StatelessGlobals.SpriteRotationVariable,
+		StatelessGlobals.SubImageIndexVariable,
+		StatelessGlobals.VelocityVariable,
+		StatelessGlobals.PreviousPositionVariable,
+		StatelessGlobals.PreviousMeshOrientationVariable,
+		StatelessGlobals.PreviousRibbonWidthVariable,
+		StatelessGlobals.PreviousScaleVariable,
+		StatelessGlobals.PreviousSpriteSizeVariable,
+		StatelessGlobals.PreviousSpriteFacingVariable,
+		StatelessGlobals.PreviousSpriteAlignmentVariable,
+		StatelessGlobals.PreviousSpriteRotationVariable,
+		StatelessGlobals.PreviousVelocityVariable,
+	};
+#endif //WITH_EDITORONLY_DATA
+}
+
+const FShaderParametersMetadata* UNiagaraStatelessEmitterDefault::GetShaderParametersMetadata() const
+{
+	using namespace NiagaraStateless;
+	return FSimulationShaderDefaultCS::FParameters::FTypeInfo::GetStructMetadata();
+}
+
+TShaderRef<NiagaraStateless::FSimulationShader> UNiagaraStatelessEmitterDefault::GetSimulationShader() const
+{
+	using namespace NiagaraStateless;
+	return TShaderMapRef<FSimulationShaderDefaultCS>(GetGlobalShaderMap(GMaxRHIFeatureLevel));
+}
+
+//-TODO: Add a way to set them directly, we should know that the final struct is a series of ints in the order of the provided variables
+void UNiagaraStatelessEmitterDefault::SetShaderParameters(uint8* ShaderParametersBase, TConstArrayView<int32> ComponentOffsets) const
+{
+	using namespace NiagaraStateless;
+	FSimulationShaderDefaultCS::FParameters* ShaderParameters	= reinterpret_cast<FSimulationShaderDefaultCS::FParameters*>(ShaderParametersBase);
+	int iComponent = 0;
+	ShaderParameters->Permutation_UniqueIDComponent					= ComponentOffsets[iComponent++];
+	ShaderParameters->Permutation_PositionComponent					= ComponentOffsets[iComponent++];
+	ShaderParameters->Permutation_ColorComponent					= ComponentOffsets[iComponent++];
+	ShaderParameters->Permutation_MeshOrientationComponent			= ComponentOffsets[iComponent++];
+	ShaderParameters->Permutation_RibbonWidthComponent				= ComponentOffsets[iComponent++];
+	ShaderParameters->Permutation_ScaleComponent					= ComponentOffsets[iComponent++];
+	ShaderParameters->Permutation_SpriteSizeComponent				= ComponentOffsets[iComponent++];
+	ShaderParameters->Permutation_SpriteFacingComponent				= ComponentOffsets[iComponent++];
+	ShaderParameters->Permutation_SpriteAlignmentComponent			= ComponentOffsets[iComponent++];
+	ShaderParameters->Permutation_SpriteRotationComponent			= ComponentOffsets[iComponent++];
+	ShaderParameters->Permutation_SubImageIndexComponent			= ComponentOffsets[iComponent++];
+	ShaderParameters->Permutation_VelocityComponent					= ComponentOffsets[iComponent++];
+	ShaderParameters->Permutation_PreviousPositionComponent			= ComponentOffsets[iComponent++];
+	ShaderParameters->Permutation_PreviousMeshOrientationComponent	= ComponentOffsets[iComponent++];
+	ShaderParameters->Permutation_PreviousRibbonWidthComponent		= ComponentOffsets[iComponent++];
+	ShaderParameters->Permutation_PreviousScaleComponent			= ComponentOffsets[iComponent++];
+	ShaderParameters->Permutation_PreviousSpriteSizeComponent		= ComponentOffsets[iComponent++];
+	ShaderParameters->Permutation_PreviousSpriteFacingComponent		= ComponentOffsets[iComponent++];
+	ShaderParameters->Permutation_PreviousSpriteAlignmentComponent	= ComponentOffsets[iComponent++];
+	ShaderParameters->Permutation_PreviousSpriteRotationComponent	= ComponentOffsets[iComponent++];
+	ShaderParameters->Permutation_PreviousVelocityComponent			= ComponentOffsets[iComponent++];
+}
+//-TODO: Add a way to set them directly, we should know that the final struct is a series of ints in the order of the provided variables
+
+//////////////////////////////////////////////////////////////////////////////////////////////////////////////////////////////////////////////////////////
+
+void UNiagaraStatelessEmitterExample1::InitModulesAndAttributes()
+{
+#if WITH_EDITORONLY_DATA
+	Modules =
+	{
+		UNiagaraStatelessModule_InitializeParticle::StaticClass(),
+		UNiagaraStatelessModule_ShapeLocation::StaticClass(),
+		UNiagaraStatelessModule_ScaleColor::StaticClass(),
+		UNiagaraStatelessModule_ScaleSpriteSize::StaticClass(),
+		UNiagaraStatelessModule_RotateAroundPoint::StaticClass(),
+		UNiagaraStatelessModule_CalculateAccurateVelocity::StaticClass(),
+	};
+
+	const FNiagaraStatelessGlobals& StatelessGlobals = FNiagaraStatelessGlobals::Get();
+	OutputComponents =
+	{
+		StatelessGlobals.UniqueIDVariable,
+		StatelessGlobals.PositionVariable,
+		StatelessGlobals.ColorVariable,
+		StatelessGlobals.RibbonWidthVariable,
+		StatelessGlobals.ScaleVariable,
+		StatelessGlobals.SpriteSizeVariable,
+		StatelessGlobals.SpriteRotationVariable,
+		StatelessGlobals.VelocityVariable,
+		StatelessGlobals.PreviousPositionVariable,
+		StatelessGlobals.PreviousRibbonWidthVariable,
+		StatelessGlobals.PreviousScaleVariable,
+		StatelessGlobals.PreviousSpriteSizeVariable,
+		StatelessGlobals.PreviousSpriteRotationVariable,
+		StatelessGlobals.PreviousVelocityVariable,
+	};
+#endif //WITH_EDITORONLY_DATA
+}
+
+const FShaderParametersMetadata* UNiagaraStatelessEmitterExample1::GetShaderParametersMetadata() const
+{
+	using namespace NiagaraStateless;
+	return FSimulationShaderExample1CS::FParameters::FTypeInfo::GetStructMetadata();
+}
+
+TShaderRef<NiagaraStateless::FSimulationShader> UNiagaraStatelessEmitterExample1::GetSimulationShader() const
+{
+	using namespace NiagaraStateless;
+	return TShaderMapRef<FSimulationShaderExample1CS>(GetGlobalShaderMap(GMaxRHIFeatureLevel));
+}
+
+//-TODO: Add a way to set them directly, we should know that the final struct is a series of ints in the order of the provided variables
+void UNiagaraStatelessEmitterExample1::SetShaderParameters(uint8* ShaderParametersBase, TConstArrayView<int32> ComponentOffsets) const
+{
+	using namespace NiagaraStateless;
+	FSimulationShaderExample1CS::FParameters* ShaderParameters		= reinterpret_cast<FSimulationShaderExample1CS::FParameters*>(ShaderParametersBase);
+	int iComponent = 0;
+	ShaderParameters->Permutation_UniqueIDComponent					= ComponentOffsets[iComponent++];
+	ShaderParameters->Permutation_PositionComponent					= ComponentOffsets[iComponent++];
+	ShaderParameters->Permutation_ColorComponent					= ComponentOffsets[iComponent++];
+	ShaderParameters->Permutation_RibbonWidthComponent				= ComponentOffsets[iComponent++];
+	ShaderParameters->Permutation_ScaleComponent					= ComponentOffsets[iComponent++];
+	ShaderParameters->Permutation_SpriteSizeComponent				= ComponentOffsets[iComponent++];
+	ShaderParameters->Permutation_SpriteRotationComponent			= ComponentOffsets[iComponent++];
+	ShaderParameters->Permutation_VelocityComponent					= ComponentOffsets[iComponent++];
+	ShaderParameters->Permutation_PreviousPositionComponent			= ComponentOffsets[iComponent++];
+	ShaderParameters->Permutation_PreviousRibbonWidthComponent		= ComponentOffsets[iComponent++];
+	ShaderParameters->Permutation_PreviousScaleComponent			= ComponentOffsets[iComponent++];
+	ShaderParameters->Permutation_PreviousSpriteSizeComponent		= ComponentOffsets[iComponent++];
+	ShaderParameters->Permutation_PreviousSpriteRotationComponent	= ComponentOffsets[iComponent++];
+	ShaderParameters->Permutation_PreviousVelocityComponent			= ComponentOffsets[iComponent++];
+}
+//-TODO: Add a way to set them directly, we should know that the final struct is a series of ints in the order of the provided variables
diff --git a/Engine/Plugins/FX/Niagara/Source/Niagara/Private/Stateless/NiagaraStatelessModule.cpp b/Engine/Plugins/FX/Niagara/Source/Niagara/Private/Stateless/NiagaraStatelessModule.cpp
new file mode 100644
--- /dev/null
+++ b/Engine/Plugins/FX/Niagara/Source/Niagara/Private/Stateless/NiagaraStatelessModule.cpp
@@ -0,0 +1,55 @@
+// Copyright Epic Games, Inc. All Rights Reserved.
+
+#include "Stateless/NiagaraStatelessModule.h"
+
+#if DO_CHECK
+void FNiagaraStatelessSetShaderParameterContext::ValidateIncludeStructType(uint32 StructOffset, const FShaderParametersMetadata* StructMetaData) const
+{
+	for (const FShaderParametersMetadata::FMember& Member : ShaderParametersMetadata->GetMembers())
+	{
+		if (Member.GetOffset() != StructOffset)
+		{
+			continue;
+		}
+
+		if (Member.GetBaseType() == UBMT_INCLUDED_STRUCT && Member.GetStructMetadata()->GetLayout() == StructMetaData->GetLayout())
+		{
+			return;
+		}
+
+		UE_LOG(LogNiagara, Fatal, TEXT("Shader parameter struct member (%s) at offset (%u) is not of type (%s)"), Member.GetName(), StructOffset, StructMetaData->GetStructTypeName());
+		return;
+	}
+
+	UE_LOG(LogNiagara, Fatal, TEXT("Failed to find shader parameter struct member type (%s) at offset (%u)"), StructMetaData->GetStructTypeName(), StructOffset);
+}
+#endif //DO_CHECK
+
+#if WITH_EDITOR
+bool UNiagaraStatelessModule::CanEditChange(const FProperty* InProperty) const
+{
+	if (!Super::CanEditChange(InProperty))
+	{
+		return false;
+	}
+
+	if (InProperty != nullptr)
+	{
+		if (InProperty->GetFName() == GET_MEMBER_NAME_CHECKED(UNiagaraStatelessModule, bModuleEnabled))
+		{
+			return CanDisableModule();
+		}
+		else if (CanDisableModule() && !IsModuleEnabled())
+		{
+			return false;
+		}
+
+		if (InProperty->GetFName() == GET_MEMBER_NAME_CHECKED(UNiagaraStatelessModule, bDebugDrawEnabled))
+		{
+			return CanDebugDraw();
+		}
+	}
+
+	return true;
+}
+#endif //WITH_EDITOR
diff --git a/Engine/Plugins/FX/Niagara/Source/Niagara/Public/INiagaraEditorOnlyDataUtlities.h b/Engine/Plugins/FX/Niagara/Source/Niagara/Public/INiagaraEditorOnlyDataUtlities.h
--- a/Engine/Plugins/FX/Niagara/Source/Niagara/Public/INiagaraEditorOnlyDataUtlities.h
+++ b/Engine/Plugins/FX/Niagara/Source/Niagara/Public/INiagaraEditorOnlyDataUtlities.h
@@ -10,6 +10,7 @@ class UNiagaraEditorParametersAdapterBase;
 class UNiagaraMessageDataBase;
 class UNiagaraScriptSourceBase;
 class UNiagaraSystem;
+struct FNiagaraSystemStateData;
 
 /** Defines utility methods for creating editor only data which is stored on runtime objects. */
 class INiagaraEditorOnlyDataUtilities
@@ -30,4 +31,6 @@ public:
 	virtual bool IsEditorDataInterfaceInstance(const UNiagaraDataInterface* DataInterface) const = 0;
 
 	virtual UNiagaraDataInterface* GetResolvedRuntimeInstanceForEditorDataInterfaceInstance(const UNiagaraSystem& OwningSystem, UNiagaraDataInterface& EditorDataInterfaceInstance) const = 0;
-};
\ No newline at end of file
+
+	virtual TOptional<FNiagaraSystemStateData> TryGetSystemStateData(const UNiagaraSystem& System) const = 0;
+};
diff --git a/Engine/Plugins/FX/Niagara/Source/Niagara/Public/NiagaraRenderer.h b/Engine/Plugins/FX/Niagara/Source/Niagara/Public/NiagaraRenderer.h
--- a/Engine/Plugins/FX/Niagara/Source/Niagara/Public/NiagaraRenderer.h
+++ b/Engine/Plugins/FX/Niagara/Source/Niagara/Public/NiagaraRenderer.h
@@ -19,6 +19,7 @@ NiagaraRenderer.h: Base class for Niagara render modules
 #include "NiagaraBoundsCalculator.h"
 #include "RayTracingGeometry.h"
 
+struct INiagaraComputeDataBufferInterface;
 class FNiagaraDataSet;
 class FNiagaraSceneProxy;
 class FNiagaraGPURendererCount;
@@ -47,7 +48,7 @@ protected:
 	FNiagaraSystemInstanceID SystemInstanceID;
 
 	FNiagaraDataBufferRef CPUParticleData;
-	FNiagaraComputeExecutionContext* GPUExecContext = nullptr;
+	INiagaraComputeDataBufferInterface* ComputeDataBufferInterface = nullptr;
 };
 
 //////////////////////////////////////////////////////////////////////////
diff --git a/Engine/Plugins/FX/Niagara/Source/Niagara/Public/NiagaraSettings.h b/Engine/Plugins/FX/Niagara/Source/Niagara/Public/NiagaraSettings.h
--- a/Engine/Plugins/FX/Niagara/Source/Niagara/Public/NiagaraSettings.h
+++ b/Engine/Plugins/FX/Niagara/Source/Niagara/Public/NiagaraSettings.h
@@ -158,6 +158,16 @@ class UNiagaraSettings : public UDeveloperSettings
 	UPROPERTY(config, EditAnywhere, Category = Niagara, meta = (DisplayName = "Enable building data for Experimental VM"))
 	bool bExperimentalVMEnabled = false;
 
+	/**
+	Enables Stateless Emitters experimental feature.
+	Statless emitters are lightweight fixed function emitters, they are not fully programmable like regular emitters and do not run scripts on the CPU.
+	Particle data is extrapolated per frame for the current particle age.  This means we never store particle data, we only generate it on demand.
+	Systems that contain only stateless emitters and no system script modules can take advantage of a much faster path to execute.
+	** There is no guarantee on backwards compatability for this feature currently.  Do not ship stateless content. **
+	*/
+	UPROPERTY(config, EditAnywhere, Category = Niagara, meta = (DisplayName = "Enable Stateless Emitters (Experimental)", ConfigRestartRequired = true))
+	bool bStatelessEmittersEnabled = false;
+
 	/** Whether to limit the max tick delta time or not. */
 	UPROPERTY(config, EditAnywhere, BlueprintReadWrite, AdvancedDisplay, Category = "Niagara", meta = (InlineEditConditionToggle))
 	bool bLimitDeltaTime = true;
diff --git a/Engine/Plugins/FX/Niagara/Source/Niagara/Public/NiagaraSystemInstance.h b/Engine/Plugins/FX/Niagara/Source/Niagara/Public/NiagaraSystemInstance.h
--- a/Engine/Plugins/FX/Niagara/Source/Niagara/Public/NiagaraSystemInstance.h
+++ b/Engine/Plugins/FX/Niagara/Source/Niagara/Public/NiagaraSystemInstance.h
@@ -681,6 +681,16 @@ public:
 	};
 
 	FInstanceParameters GatheredInstanceParameters;
+
+	void InitSystemState();
+	void TickSystemState();
+
+	FRandomStream	SystemState_RandomStream;
+	int32			SystemState_LoopCount = 0;
+	float			SystemState_CurrentLoopDuration = 0.0f;
+	float			SystemState_CurrentLoopDelay = 0.0f;
+	float			SystemState_CurrentLoopAgeStart = 0.0f;
+	float			SystemState_CurrentLoopAgeEnd = 0.0f;
 };
 
 FORCEINLINE void FNiagaraSystemInstance::SetLODDistance(float InLODDistance, float InMaxLODDistance, bool bOverride)
diff --git a/Engine/Plugins/FX/Niagara/Source/NiagaraEditor/Private/Customizations/NiagaraTypeCustomizations.cpp b/Engine/Plugins/FX/Niagara/Source/NiagaraEditor/Private/Customizations/NiagaraTypeCustomizations.cpp
--- a/Engine/Plugins/FX/Niagara/Source/NiagaraEditor/Private/Customizations/NiagaraTypeCustomizations.cpp
+++ b/Engine/Plugins/FX/Niagara/Source/NiagaraEditor/Private/Customizations/NiagaraTypeCustomizations.cpp
@@ -45,6 +45,10 @@
 #include "NiagaraVariableMetaData.h"
 #include "ViewModels/NiagaraEmitterHandleViewModel.h"
 
+//-TODO:stateless:Remove and unify emitter
+#include "Stateless/NiagaraStatelessEmitter.h"
+//-TODO:stateless:Remove and unify emitter
+
 #include "Widgets/SNiagaraParameterMenu.h"
 
 #include UE_INLINE_GENERATED_CPP_BY_NAME(NiagaraTypeCustomizations)
@@ -131,12 +135,11 @@ void FNiagaraMatrixCustomization::CustomizeChildren(TSharedRef<IPropertyHandle>
 	}
 }
 
-TArray<FNiagaraVariableBase> FNiagaraStackAssetAction_VarBind::FindVariables(const FVersionedNiagaraEmitter& InEmitter, bool bSystem, bool bEmitter, bool bParticles, bool bUser, bool bAllowStatic)
+TArray<FNiagaraVariableBase> FNiagaraStackAssetAction_VarBind::FindVariables(UNiagaraSystem* NiagaraSystem, const FVersionedNiagaraEmitter& InEmitter, bool bSystem, bool bEmitter, bool bParticles, bool bUser, bool bAllowStatic)
 {
 	TArray<FNiagaraVariableBase> Bindings;
 	TArray<FNiagaraParameterMapHistory> Histories;
 
-	UNiagaraEmitter* Emitter = InEmitter.Emitter;
 	if (FVersionedNiagaraEmitterData* EmitterData = InEmitter.GetEmitterData())
 	{
 		if ( UNiagaraScriptSource* Source = Cast<UNiagaraScriptSource>(EmitterData->GraphSource) )
@@ -145,17 +148,16 @@ TArray<FNiagaraVariableBase> FNiagaraStackAssetAction_VarBind::FindVariables(con
 		}
 	}
 
-	if (bSystem || bEmitter)
+	if (NiagaraSystem && (bSystem || bEmitter))
 	{
-		if (UNiagaraSystem* Sys = Emitter->GetTypedOuter<UNiagaraSystem>())
+		if ( UNiagaraScriptSource* Source = Cast<UNiagaraScriptSource>(NiagaraSystem->GetSystemUpdateScript()->GetLatestSource()) )
 		{
-			if ( UNiagaraScriptSource* Source = Cast<UNiagaraScriptSource>(Sys->GetSystemUpdateScript()->GetLatestSource()) )
-			{
-				Histories.Append(UNiagaraNodeParameterMapBase::GetParameterMaps(Source->NodeGraph));
-			}
+			Histories.Append(UNiagaraNodeParameterMapBase::GetParameterMaps(Source->NodeGraph));
 		}
 	}
-	
+
+	UNiagaraEmitter* NiagaraEmitter = InEmitter.Emitter;
+	bEmitter &= NiagaraEmitter != nullptr;
 	for (const FNiagaraParameterMapHistory& History : Histories)
 	{
 		for (const FNiagaraVariable& Var : History.Variables)
@@ -163,42 +165,40 @@ TArray<FNiagaraVariableBase> FNiagaraStackAssetAction_VarBind::FindVariables(con
 			if (Var.GetType().IsStatic() && !bAllowStatic)
 				continue;
 
-			if (FNiagaraParameterUtilities::IsAttribute(Var) && bParticles)
+			if (bParticles && FNiagaraParameterUtilities::IsAttribute(Var))
 			{
 				Bindings.AddUnique(Var);
 			}
-			else if (FNiagaraParameterUtilities::IsSystemParameter(Var) && bSystem)
+			else if (bSystem && FNiagaraParameterUtilities::IsSystemParameter(Var))
 			{
 				Bindings.AddUnique(Var);
 			}
-			else if (Var.IsInNameSpace(Emitter->GetUniqueEmitterName()) && bEmitter)
+			else if (bEmitter && NiagaraEmitter && Var.IsInNameSpace(NiagaraEmitter->GetUniqueEmitterName()))
 			{
-				Bindings.AddUnique(FNiagaraUtilities::ResolveAliases(Var, FNiagaraAliasContext()
-					.ChangeEmitterNameToEmitter(Emitter->GetUniqueEmitterName())));
+				Bindings.AddUnique(
+					FNiagaraUtilities::ResolveAliases(Var, FNiagaraAliasContext().ChangeEmitterNameToEmitter(NiagaraEmitter->GetUniqueEmitterName()))
+				);
 			}
-			else if (FNiagaraParameterUtilities::IsAliasedEmitterParameter(Var) && bEmitter)
+			else if (bEmitter && FNiagaraParameterUtilities::IsAliasedEmitterParameter(Var))
 			{
 				Bindings.AddUnique(Var);
 			}
-			else if (Var.IsInNameSpace(FNiagaraConstants::EmitterNamespaceString) && bEmitter)
+			else if (bEmitter && Var.IsInNameSpace(FNiagaraConstants::EmitterNamespaceString))
 			{
 				Bindings.AddUnique(Var);
 			}
-			else if (FNiagaraParameterUtilities::IsUserParameter(Var) && bUser)
+			else if (bUser && FNiagaraParameterUtilities::IsUserParameter(Var))
 			{
 				Bindings.AddUnique(Var);
 			}
 		}
 	}
 
-	if (bUser)
+	if (NiagaraSystem && bUser)
 	{
-		if (UNiagaraSystem* Sys = Emitter->GetTypedOuter<UNiagaraSystem>())
+		for (const FNiagaraVariable Var : NiagaraSystem->GetExposedParameters().ReadParameterVariables())
 		{
-			for (const FNiagaraVariable Var : Sys->GetExposedParameters().ReadParameterVariables())
-			{
-				Bindings.AddUnique(Var);
-			}
+			Bindings.AddUnique(Var);
 		}
 	}
 	return Bindings;
@@ -782,7 +782,6 @@ void FNiagaraUserParameterBindingCustomization::CustomizeHeader(TSharedRef<IProp
 					.ToolTipText(this, &FNiagaraUserParameterBindingCustomization::GetTooltipText)
 					.ButtonContent()
 					[
-
 						SNew(SNiagaraParameterName)
 						.ParameterName(this, &FNiagaraUserParameterBindingCustomization::GetVariableName)
 						.IsReadOnly(true)
@@ -982,14 +981,18 @@ TArray<TPair<FNiagaraVariableBase, FNiagaraVariableBase> > FNiagaraMaterialAttri
 	TArray<TPair<FNiagaraVariableBase, FNiagaraVariableBase>> Names;
 	TArray<FNiagaraVariableBase> BaseVars;
 
-	if (BaseSystem && BaseEmitter.Emitter && TargetParameterBinding)
+	//-TODO:stateless:Remove and unify emitter
+	if (BaseSystem && (BaseEmitter.Emitter || StatelessEmitter) && TargetParameterBinding)
+	//-TODO:stateless:Remove and unify emitter
 	{
-		bool bSystem = true;
-		bool bEmitter = true;
-		bool bParticles = false;
-		bool bUser = true;
-		bool bStatic = false;
-		BaseVars = FNiagaraStackAssetAction_VarBind::FindVariables(BaseEmitter, bSystem, bEmitter, bParticles, bUser, bStatic);
+		const bool bSystem = true;
+		//-TODO:stateless:Remove and unify emitter
+		const bool bEmitter = true;
+		//-TODO:stateless:Remove and unify emitter
+		const bool bParticles = false;
+		const bool bUser = true;
+		const bool bStatic = false;
+		BaseVars = FNiagaraStackAssetAction_VarBind::FindVariables(BaseSystem, BaseEmitter, bSystem, bEmitter, bParticles, bUser, bStatic);
 
 		TArray<UNiagaraScript*> Scripts;
 		Scripts.Add(BaseSystem->GetSystemUpdateScript());
@@ -1000,7 +1003,12 @@ TArray<TPair<FNiagaraVariableBase, FNiagaraVariableBase> > FNiagaraMaterialAttri
 		}
 
 		TMap<FString, FString> EmitterAlias;
-		EmitterAlias.Emplace(FNiagaraConstants::EmitterNamespace.ToString(), BaseEmitter.Emitter->GetUniqueEmitterName());
+		//-TODO:stateless:Remove and unify emitter
+		if (BaseEmitter.Emitter)
+		//-TODO:stateless:Remove and unify emitter
+		{
+			EmitterAlias.Emplace(FNiagaraConstants::EmitterNamespace.ToString(), BaseEmitter.Emitter->GetUniqueEmitterName());
+		}
 
 		auto FindCachedDI = 
 			[&](const FNiagaraVariableBase& BaseVariable) -> UNiagaraDataInterface*
@@ -1008,6 +1016,9 @@ TArray<TPair<FNiagaraVariableBase, FNiagaraVariableBase> > FNiagaraMaterialAttri
 				FName VariableName = BaseVariable.GetName();
 				if (BaseVariable.IsInNameSpace(FNiagaraConstants::EmitterNamespaceString))
 				{
+					//-TODO:stateless:Remove and unify emitter
+					check(BaseEmitter.Emitter);
+					//-TODO:stateless:Remove and unify emitter
 					VariableName = FNiagaraUtilities::ResolveAliases(BaseVariable, FNiagaraAliasContext()
 						.ChangeEmitterToEmitterName(BaseEmitter.Emitter->GetUniqueEmitterName())).GetName();
 				}
@@ -1385,6 +1396,9 @@ void FNiagaraMaterialAttributeBindingCustomization::CustomizeChildren(TSharedRef
 		if (RenderProps)
 		{
 			BaseEmitter = RenderProps->GetOuterEmitter();
+			//-TODO:stateless:Remove and unify emitter
+			StatelessEmitter = RenderProps->GetTypedOuter<UNiagaraStatelessEmitter>();
+			//-TODO:stateless:Remove and unify emitter
 		}
 		if (BaseSystem)
 		{
@@ -2263,6 +2277,16 @@ FText FNiagaraRendererMaterialParameterCustomization::GetBindingNameText(TShared
 	return FText::FromName(BindingName);
 }
 
+FName FNiagaraRendererMaterialParameterCustomization::GetBindingName(TSharedPtr<IPropertyHandle> PropertyHandle) const
+{
+	FName BindingName;
+	if (PropertyHandle.IsValid())
+	{
+		PropertyHandle->GetValue(BindingName);
+	}
+	return BindingName;
+}
+
 FText FNiagaraRendererMaterialParameterCustomization::GetMaterialBindingTooltip(FName ParameterName, const FString& ParameterDesc)
 {
 	if (ParameterDesc.Len() > 0)
@@ -2368,9 +2392,9 @@ bool FNiagaraRendererMaterialStaticBoolParameterCustomization::CustomizeChildPro
 			.ForegroundColor(FAppStyle::GetColor("PropertyEditor.AssetName.ColorAndOpacity"))
 			.ButtonContent()
 			[
-				SNew(STextBlock)
-				.Text(this, &FNiagaraRendererMaterialParameterCustomization::GetBindingNameText, PropertyHandle)
-				.Font(IDetailLayoutBuilder::GetDetailFont())
+				SNew(SNiagaraParameterName)
+				.ParameterName(this, &FNiagaraRendererMaterialParameterCustomization::GetBindingName, PropertyHandle)
+				.IsReadOnly(true)
 			]
 		];
 
@@ -2398,54 +2422,54 @@ TSharedRef<SWidget> FNiagaraRendererMaterialStaticBoolParameterCustomization::On
 	{
 		const FVersionedNiagaraEmitter NiagaraEmitter = RenderProperties->GetOuterEmitter();
 		const FNiagaraTypeDefinition StaticBoolDef = FNiagaraTypeDefinition::GetBoolDef().ToStaticDef();
+		FNiagaraAliasContext AliasContext;
 
-		if (FVersionedNiagaraEmitterData* EmitterData = NiagaraEmitter.GetEmitterData())
-		{
-			const FNiagaraAliasContext AliasContext = FNiagaraAliasContext().ChangeEmitterNameToEmitter(NiagaraEmitter.Emitter->GetUniqueEmitterName());
-
-			auto AddStaticVariables =
-				[&](UNiagaraScript* NiagaraScript)
+		auto AddStaticVariables =
+			[&](UNiagaraScript* NiagaraScript)
+			{
+				if (NiagaraScript == nullptr)
+				{
+					return;
+				}
+				for ( const FNiagaraVariable& StaticVariable : NiagaraScript->GetVMExecutableData().StaticVariablesWritten )
 				{
-					if (NiagaraScript == nullptr)
+					if ( StaticVariable.GetType() != StaticBoolDef)
 					{
-						return;
+						continue;
 					}
-					for ( const FNiagaraVariable& StaticVariable : NiagaraScript->GetVMExecutableData().StaticVariablesWritten )
+					const FNiagaraVariableBase ResolvedVariable = FNiagaraUtilities::ResolveAliases(StaticVariable, AliasContext);
+					const bool bSystemAttribute = ResolvedVariable.IsInNameSpace(FNiagaraConstants::SystemNamespaceString);
+					const bool bEmitterAttribute = ResolvedVariable.IsInNameSpace(FNiagaraConstants::EmitterNamespaceString);
+					const bool bParticleAttribute = ResolvedVariable.IsInNameSpace(FNiagaraConstants::ParticleAttributeNamespaceString);
+					if (!bSystemAttribute && !bEmitterAttribute && !bParticleAttribute)
 					{
-						if ( StaticVariable.GetType() != StaticBoolDef)
-						{
-							continue;
-						}
-						const FNiagaraVariableBase ResolvedVariable = FNiagaraUtilities::ResolveAliases(StaticVariable, AliasContext);
-						const bool bSystemAttribute = ResolvedVariable.IsInNameSpace(FNiagaraConstants::SystemNamespaceString);
-						const bool bEmitterAttribute = ResolvedVariable.IsInNameSpace(FNiagaraConstants::EmitterNamespaceString);
-						const bool bParticleAttribute = ResolvedVariable.IsInNameSpace(FNiagaraConstants::ParticleAttributeNamespaceString);
-						if (!bSystemAttribute && !bEmitterAttribute && !bParticleAttribute)
-						{
-							continue;
-						}
+						continue;
+					}
 
-						MenuBuilder.AddMenuEntry(
-							FText::FromName(ResolvedVariable.GetName()),
-							TAttribute<FText>(),
-							FSlateIcon(),
-							FUIAction(
-								FExecuteAction::CreateLambda(
-									[ValidBinding= ResolvedVariable.GetName(), PropertyHandle]()
-									{
-										PropertyHandle->SetValue(ValidBinding);
-									}
-								)
+					MenuBuilder.AddMenuEntry(
+						FText::FromName(ResolvedVariable.GetName()),
+						TAttribute<FText>(),
+						FSlateIcon(),
+						FUIAction(
+							FExecuteAction::CreateLambda(
+								[ValidBinding= ResolvedVariable.GetName(), PropertyHandle]()
+								{
+									PropertyHandle->SetValue(ValidBinding);
+								}
 							)
-						);
-					}
-				};
+						)
+					);
+				}
+			};
 
-			if (UNiagaraSystem* NiagaraSystem = RenderProperties->GetTypedOuter<UNiagaraSystem>())
-			{
-				AddStaticVariables(NiagaraSystem->GetSystemSpawnScript());
-				AddStaticVariables(NiagaraSystem->GetSystemUpdateScript());
-			}
+		if (UNiagaraSystem* NiagaraSystem = RenderProperties->GetTypedOuter<UNiagaraSystem>())
+		{
+			AddStaticVariables(NiagaraSystem->GetSystemSpawnScript());
+			AddStaticVariables(NiagaraSystem->GetSystemUpdateScript());
+		}
+		if (FVersionedNiagaraEmitterData* EmitterData = NiagaraEmitter.GetEmitterData())
+		{
+			AliasContext.ChangeEmitterNameToEmitter(NiagaraEmitter.Emitter->GetUniqueEmitterName());
 			EmitterData->ForEachScript(AddStaticVariables);
 		}
 	}
diff --git a/Engine/Plugins/FX/Niagara/Source/NiagaraEditor/Private/Customizations/NiagaraTypeCustomizations.h b/Engine/Plugins/FX/Niagara/Source/NiagaraEditor/Private/Customizations/NiagaraTypeCustomizations.h
--- a/Engine/Plugins/FX/Niagara/Source/NiagaraEditor/Private/Customizations/NiagaraTypeCustomizations.h
+++ b/Engine/Plugins/FX/Niagara/Source/NiagaraEditor/Private/Customizations/NiagaraTypeCustomizations.h
@@ -153,7 +153,7 @@ struct FNiagaraStackAssetAction_VarBind : public FEdGraphSchemaAction
 	}
 	//~ End FEdGraphSchemaAction Interface
 
-	static TArray<FNiagaraVariableBase> FindVariables(const FVersionedNiagaraEmitter& InEmitter, bool bSystem, bool bEmitter, bool bParticles, bool bUser, bool bAllowStatic);
+	static TArray<FNiagaraVariableBase> FindVariables(UNiagaraSystem* NiagaraSystem, const FVersionedNiagaraEmitter& InEmitter, bool bSystem, bool bEmitter, bool bParticles, bool bUser, bool bAllowStatic);
 };
 
 class FNiagaraVariableAttributeBindingCustomization : public IPropertyTypeCustomization
@@ -274,6 +274,9 @@ private:
 	class UNiagaraSystem* BaseSystem;
 	FVersionedNiagaraEmitter BaseEmitter;
 	UNiagaraRendererProperties* RenderProps;
+	//-TODO:stateless:Remove and unify emitter
+	class UNiagaraStatelessEmitter* StatelessEmitter = nullptr;
+	//-TODO:stateless:Remove and unify emitter
 	struct FNiagaraMaterialAttributeBinding* TargetParameterBinding;
 
 };
@@ -386,6 +389,7 @@ public:
 	// IPropertyTypeCustomization interface end
 
 	FText GetBindingNameText(TSharedPtr<IPropertyHandle> PropertyHandle) const;
+	FName GetBindingName(TSharedPtr<IPropertyHandle> PropertyHandle) const;
 	static FText GetMaterialBindingTooltip(FName ParameterName, const FString& ParameterDesc);
 	TSharedRef<SWidget> OnGetMaterialBindingNameMenuContent(TSharedPtr<IPropertyHandle> PropertyHandle) const;
 
diff --git a/Engine/Plugins/FX/Niagara/Source/NiagaraEditor/Private/NiagaraCompiler.cpp b/Engine/Plugins/FX/Niagara/Source/NiagaraEditor/Private/NiagaraCompiler.cpp
--- a/Engine/Plugins/FX/Niagara/Source/NiagaraEditor/Private/NiagaraCompiler.cpp
+++ b/Engine/Plugins/FX/Niagara/Source/NiagaraEditor/Private/NiagaraCompiler.cpp
@@ -1029,9 +1029,8 @@ TSharedPtr<FNiagaraCompileRequestDataBase, ESPMode::ThreadSafe> FNiagaraEditorMo
 			EmitterPtr->EmitterUniqueName = Handle.GetUniqueInstanceName();
 			EmitterPtr->EmitterID = FNiagaraEmitterID(i);
 			EmitterPtr->SourceName = BasePtr->SourceName;
-			//-TODO:Stateless:
+			//-TODO:Stateless: We need to handle the stateless path here
 			EmitterPtr->Source = Handle.GetEmitterData() ? Cast<UNiagaraScriptSource>(Handle.GetEmitterData()->GraphSource) : nullptr;
-			//-TODO:Stateless:
 			EmitterPtr->bUseRapidIterationParams = BasePtr->bUseRapidIterationParams;
 			EmitterPtr->bDisableDebugSwitches = BasePtr->bDisableDebugSwitches;
 			EmitterPtr->SharedCompileDataInterfaceData = BasePtr->SharedCompileDataInterfaceData;
@@ -1065,9 +1064,7 @@ TSharedPtr<FNiagaraCompileRequestDataBase, ESPMode::ThreadSafe> FNiagaraEditorMo
 		{
 			const FNiagaraEmitterHandle& Handle = System->GetEmitterHandle(i);
 			FCompileConstantResolver ConstantResolver(Handle.GetInstance(), ENiagaraScriptUsage::EmitterSpawnScript);
-			//-TODO:Stateless:
 			if (Handle.GetIsEnabled() && Handle.GetEmitterData()) // Don't pull in the emitter if it isn't going to be used.
-			//-TODO:Stateless:
 			{
 				TSharedPtr<FNiagaraGraphCachedDataBase, ESPMode::ThreadSafe> CachedTraversalEmitterData = Handle.GetInstance().Emitter->GetCachedTraversalData(Handle.GetInstance().Version);
 				TArray<FNiagaraVariable> StaticVariablesFromEmitter = StaticVariablesFromSystem;
@@ -1200,9 +1197,7 @@ TSharedPtr<FNiagaraCompileRequestDataBase, ESPMode::ThreadSafe> FNiagaraEditorMo
 			{
 				const FNiagaraEmitterHandle& Handle = System->GetEmitterHandle(i);
 				FCompileConstantResolver ConstantResolver(Handle.GetInstance(), ENiagaraScriptUsage::EmitterSpawnScript);
-				//-TODO:Stateless:
 				if (Handle.GetIsEnabled() && Handle.GetEmitterData()) // Don't pull in the emitter if it isn't going to be used.
-					//-TODO:Stateless:
 				{
 					TArray<UNiagaraScript*> EmitterScripts;
 					Handle.GetEmitterData()->GetScripts(EmitterScripts, false);
@@ -1364,9 +1359,7 @@ TSharedPtr<FNiagaraCompileRequestDuplicateDataBase, ESPMode::ThreadSafe> FNiagar
 			EmitterPtr->SharedNameToDuplicatedDataInterfaceMap = BasePtr->SharedNameToDuplicatedDataInterfaceMap;
 			EmitterPtr->SharedDataInterfaceClassToDuplicatedCDOMap = BasePtr->SharedDataInterfaceClassToDuplicatedCDOMap;
 			//EmitterPtr->bSimulationStagesEnabled = Handle.GetInstance()->bSimulationStagesEnabled;
-			//-TODO:Stateless:
 			if (Handle.GetIsEnabled() && Handle.GetInstance().Emitter && (OwningEmitter == nullptr || OwningEmitter == Handle.GetInstance().Emitter)) // Don't need to copy the graph if we aren't going to use it.
-			//-TODO:Stateless:
 			{
 				EmitterPtr->DeepCopyGraphs(Handle.GetInstance());
 			}
@@ -1398,13 +1391,10 @@ TSharedPtr<FNiagaraCompileRequestDuplicateDataBase, ESPMode::ThreadSafe> FNiagar
 		{
 			const FNiagaraEmitterHandle& Handle = OwningSystem->GetEmitterHandle(i);
 
-
 			TArray<FNiagaraVariable> EncounterableEmitterVariables;
 			OwningSystemRequestData->GetDependentRequest(i)->GatherPreCompiledVariables(FString(), EncounterableEmitterVariables);
 
-			//-TODO:Stateless:
 			if (Handle.GetIsEnabled() && Handle.GetInstance().Emitter && (OwningEmitter == nullptr || OwningEmitter == Handle.GetInstance().Emitter))
-			//-TODO:Stateless:
 			{
 				TSharedPtr<FNiagaraGraphCachedDataBase, ESPMode::ThreadSafe> CachedTraversalEmitterData = Handle.GetInstance().Emitter->GetCachedTraversalData(Handle.GetInstance().Version);
 				TArray<FNiagaraVariable> StaticVariablesFromEmitter = StaticVariablesFromSystem;
diff --git a/Engine/Plugins/FX/Niagara/Source/NiagaraEditor/Private/NiagaraEditorModule.cpp b/Engine/Plugins/FX/Niagara/Source/NiagaraEditor/Private/NiagaraEditorModule.cpp
--- a/Engine/Plugins/FX/Niagara/Source/NiagaraEditor/Private/NiagaraEditorModule.cpp
+++ b/Engine/Plugins/FX/Niagara/Source/NiagaraEditor/Private/NiagaraEditorModule.cpp
@@ -289,6 +289,50 @@ class FNiagaraEditorOnlyDataUtilities : public INiagaraEditorOnlyDataUtilities
 	{
 		return FNiagaraEditorUtilities::GetResolvedRuntimeInstanceForEditorDataInterfaceInstance(OwningSystem, EditorDataInterfaceInstance);
 	}
+
+	virtual TOptional<FNiagaraSystemStateData> TryGetSystemStateData(const UNiagaraSystem& System) const override
+	{
+		TOptional<FNiagaraSystemStateData> SystemStateData;
+
+		// Never allow if stateless is not enabled
+		if (!GetDefault<UNiagaraSettings>()->bStatelessEmittersEnabled)
+		{
+			return SystemStateData;
+		}
+
+		// All emitters must be stateless currently
+		// We can perhaps look at this again, but we always write Emitter.RandomSeed currently even with an empty script
+		for (const FNiagaraEmitterHandle& EmitterHandle : System.GetEmitterHandles())
+		{
+			if ( !EmitterHandle.GetIsEnabled() && EmitterHandle.GetEmitterMode() != ENiagaraEmitterMode::Stateless )
+			{
+				return SystemStateData;
+			}
+		}
+
+		// Try to resolve system state from the system scripts
+		const UNiagaraScript* Script = System.GetSystemSpawnScript();
+		const UNiagaraScriptSource* ScriptSource = Script ? Cast<UNiagaraScriptSource>(Script->GetLatestSource()) : nullptr;
+		if (ScriptSource)
+		{
+			const TCHAR* SystemStateName = TEXT("/Niagara/Modules/System/SystemState.SystemState");
+			TArray<UNiagaraNodeFunctionCall*> Nodes;
+			ScriptSource->NodeGraph->GetNodesOfClass<UNiagaraNodeFunctionCall>(Nodes);
+			Nodes.RemoveAll([](UNiagaraNodeFunctionCall* Node) { return !Node || !Node->IsNodeEnabled(); });
+
+			// No function calls, we can enable fast path with the empty system state
+			if (Nodes.Num() == 0)
+			{
+				SystemStateData.Emplace(FNiagaraSystemStateData());
+			}
+			//-TODO:Stateless: Single function call which is system state, attempt to extract the data
+			//else if (Nodes.Num() == 1 && Nodes[0]->FunctionScript->GetPathName() == SystemStateName)
+			//{
+			//}
+		}
+
+		return SystemStateData;
+	}
 };
 
 class FNiagaraScriptGraphPanelPinFactory : public FGraphPanelPinFactory
diff --git a/Engine/Plugins/FX/Niagara/Source/NiagaraEditor/Private/NiagaraEditorStyle.cpp b/Engine/Plugins/FX/Niagara/Source/NiagaraEditor/Private/NiagaraEditorStyle.cpp
--- a/Engine/Plugins/FX/Niagara/Source/NiagaraEditor/Private/NiagaraEditorStyle.cpp
+++ b/Engine/Plugins/FX/Niagara/Source/NiagaraEditor/Private/NiagaraEditorStyle.cpp
@@ -353,7 +353,13 @@ void FNiagaraEditorStyle::InitCodeView()
 	Set("SyntaxHighlight.HLSL.PreProcessorKeyword", FTextBlockStyle(NormalHlslText).SetColorAndOpacity(FLinearColor(FColor(188, 98, 171))));
 
 	Set("SyntaxHighlight.HLSL.Error", HlslErrorText); 
-		
+	
+	Set("SyntaxHighlight.Python.Normal", FTextBlockStyle(NormalHlslText).SetColorAndOpacity(FLinearColor(FColor(189, 183, 107))));
+	Set("SyntaxHighlight.Python.Operator", FTextBlockStyle(NormalHlslText).SetColorAndOpacity(FLinearColor(FColor(180, 180, 180))));
+	Set("SyntaxHighlight.Python.Keyword", FTextBlockStyle(NormalHlslText).SetColorAndOpacity(FLinearColor(FColor(86, 156, 214))));
+	Set("SyntaxHighlight.Python.String", FTextBlockStyle(NormalHlslText).SetColorAndOpacity(FLinearColor(FColor(214, 157, 133))));
+	Set("SyntaxHighlight.Python.Number", FTextBlockStyle(NormalHlslText).SetColorAndOpacity(FLinearColor(FColor(181, 206, 168))));
+	Set("SyntaxHighlight.Python.Comment", FTextBlockStyle(NormalHlslText).SetColorAndOpacity(FLinearColor(FColor(87, 166, 74))));
 }
 
 void FNiagaraEditorStyle::InitSelectedEmitter()
@@ -396,6 +402,9 @@ void FNiagaraEditorStyle::InitIcons()
 	Set("NiagaraEditor.EventIcon", new IMAGE_BRUSH("Icons/Event", Icon12x12, ParticleIconColor));
 	Set("NiagaraEditor.SimulationStageIcon", new IMAGE_BRUSH("Icons/SimulationStage", Icon12x12, ParticleIconColor));
 	Set("NiagaraEditor.RenderIcon", new IMAGE_BRUSH("Icons/Render", Icon12x12, RendererIconColor));
+	Set("NiagaraEditor.Stateless.SpawnIcon", new IMAGE_BRUSH("Icons/Spawn", Icon12x12, FLinearColor::White));
+	Set("NiagaraEditor.Stateless.UpdateIcon", new IMAGE_BRUSH("Icons/Update", Icon12x12, FLinearColor::White));
+	Set("NiagaraEditor.Stateless.RenderIcon", new IMAGE_BRUSH("Icons/Render", Icon12x12, FLinearColor::White));
 	
 	Set("NiagaraEditor.HierarchyEditor.RootDropIcon", new IMAGE_BRUSH_SVG("Icons/caret-down", Icon20x20));
 	
@@ -454,6 +463,7 @@ void FNiagaraEditorStyle::InitAssetColors()
 	Set("NiagaraEditor.AssetColors.SimCache", FLinearColor(0.9f, 0.3f, 0.25f));
 	Set("NiagaraEditor.AssetColors.ValidationRuleSet", FLinearColor(0.25f, 0.25f, 0.25f));
 	Set("NiagaraEditor.AssetColors.DataChannelDefinitions", FLinearColor(0.5f, 0.0f, 1.0f));
+	Set("NiagaraEditor.AssetColors.StatelessConversionData", FLinearColor(1.0f, 1.0f, 1.0f));
 }
 
 void FNiagaraEditorStyle::InitThumbnails()
diff --git a/Engine/Plugins/FX/Niagara/Source/NiagaraEditor/Private/NiagaraEditorUtilities.cpp b/Engine/Plugins/FX/Niagara/Source/NiagaraEditor/Private/NiagaraEditorUtilities.cpp
--- a/Engine/Plugins/FX/Niagara/Source/NiagaraEditor/Private/NiagaraEditorUtilities.cpp
+++ b/Engine/Plugins/FX/Niagara/Source/NiagaraEditor/Private/NiagaraEditorUtilities.cpp
@@ -2569,6 +2569,49 @@ void FNiagaraEditorUtilities::ToggleSelectedEmittersIsolated(TSharedRef<FNiagara
 	SystemViewModel->IsolateEmitters(EmittersToIsolate);
 }
 
+ECheckBoxState FNiagaraEditorUtilities::GetSelectedEmittersEmitterModeCheckState(TSharedRef<FNiagaraSystemViewModel> SystemViewModel, ENiagaraEmitterMode EmitterMode)
+{
+	TOptional<ENiagaraEmitterMode> CurrentEmitterMode;
+	const TArray<FGuid>& SelectedEmitterHandleIds = SystemViewModel->GetSelectionViewModel()->GetSelectedEmitterHandleIds();
+	for (const TSharedRef<FNiagaraEmitterHandleViewModel>& EmitterHandleViewModel : SystemViewModel->GetEmitterHandleViewModels())
+	{
+		if (SelectedEmitterHandleIds.Contains(EmitterHandleViewModel->GetId()) == false)
+		{
+			continue;
+		}
+		if (CurrentEmitterMode.IsSet() == false)
+		{
+			CurrentEmitterMode = EmitterHandleViewModel->GetEmitterHandle()->GetEmitterMode();
+		}
+		else
+		{
+			if (CurrentEmitterMode.GetValue() != EmitterHandleViewModel->GetEmitterHandle()->GetEmitterMode())
+			{
+				return ECheckBoxState::Undetermined;
+			}
+		}
+	}
+	if (CurrentEmitterMode.IsSet())
+	{
+		return CurrentEmitterMode.GetValue() == EmitterMode ? ECheckBoxState::Checked : ECheckBoxState::Unchecked;
+	}
+	return ECheckBoxState::Unchecked;
+}
+
+void FNiagaraEditorUtilities::SetSelectedEmittersEmitterMode(TSharedRef<FNiagaraSystemViewModel> SystemViewModel, ENiagaraEmitterMode EmitterMode)
+{
+	FScopedTransaction ScopedTransaction(LOCTEXT("ChangeEmitterMode", "Change emitter mode"));
+	SystemViewModel->GetSystem().Modify();
+
+	const TArray<FGuid>& SelectedEmitterHandleIds = SystemViewModel->GetSelectionViewModel()->GetSelectedEmitterHandleIds();
+	for (const TSharedRef<FNiagaraEmitterHandleViewModel>& EmitterHandleViewModel : SystemViewModel->GetEmitterHandleViewModels())
+	{
+		if (SelectedEmitterHandleIds.Contains(EmitterHandleViewModel->GetId()))
+		{
+			EmitterHandleViewModel->GetEmitterHandle()->SetEmitterMode(SystemViewModel->GetSystem(), EmitterMode);
+		}
+	}
+}
 
 void FNiagaraEditorUtilities::CreateAssetFromEmitter(TSharedRef<FNiagaraEmitterHandleViewModel> EmitterHandleViewModel)
 {
diff --git a/Engine/Plugins/FX/Niagara/Source/NiagaraEditor/Private/Sequencer/NiagaraSequence/MovieSceneNiagaraEmitterTrack.cpp b/Engine/Plugins/FX/Niagara/Source/NiagaraEditor/Private/Sequencer/NiagaraSequence/MovieSceneNiagaraEmitterTrack.cpp
--- a/Engine/Plugins/FX/Niagara/Source/NiagaraEditor/Private/Sequencer/NiagaraSequence/MovieSceneNiagaraEmitterTrack.cpp
+++ b/Engine/Plugins/FX/Niagara/Source/NiagaraEditor/Private/Sequencer/NiagaraSequence/MovieSceneNiagaraEmitterTrack.cpp
@@ -271,6 +271,12 @@ void UMovieSceneNiagaraEmitterTrack::CreateSections(const FFrameRate& InFrameRes
 {
 	SectionInitializationErrors.Empty();
 
+	//-TODO:Stateless: We need to handle the stateless path here
+	if (GetEmitterHandleViewModel()->GetEmitterViewModel()->GetEmitter().GetEmitterData() == nullptr)
+	{
+		return;
+	}
+
 	UNiagaraScript* EmitterUpdateScript = GetEmitterHandleViewModel()->GetEmitterViewModel()->GetEmitter().GetEmitterData()->GetScript(ENiagaraScriptUsage::EmitterUpdateScript, FGuid());
 	UNiagaraScriptSource* ScriptSource = CastChecked<UNiagaraScriptSource>(EmitterUpdateScript->GetLatestSource());
 	UNiagaraNodeOutput* OutputNode = ScriptSource->NodeGraph->FindOutputNode(ENiagaraScriptUsage::EmitterUpdateScript);
diff --git a/Engine/Plugins/FX/Niagara/Source/NiagaraEditor/Private/Sequencer/NiagaraSequence/NiagaraEmitterTrackEditor.cpp b/Engine/Plugins/FX/Niagara/Source/NiagaraEditor/Private/Sequencer/NiagaraSequence/NiagaraEmitterTrackEditor.cpp
--- a/Engine/Plugins/FX/Niagara/Source/NiagaraEditor/Private/Sequencer/NiagaraSequence/NiagaraEmitterTrackEditor.cpp
+++ b/Engine/Plugins/FX/Niagara/Source/NiagaraEditor/Private/Sequencer/NiagaraSequence/NiagaraEmitterTrackEditor.cpp
@@ -92,8 +92,11 @@ public:
 		// Due to delegate bind order, the item creation happens after the refresh. To solve this, we just wait a frame. 
 		UNiagaraEmitter* NiagaraEmitter = EmitterTrack->GetEmitterHandleViewModel()->GetEmitterViewModel()->GetEmitter().Emitter;
 		WeakNiagaraEmitter = NiagaraEmitter;
-		NiagaraEmitter->OnRenderersChanged().AddSP(this, &SEmitterTrackWidget::RefreshRenderers);
-
+		//-TODO:Stateless: Add support for stateless
+		if (NiagaraEmitter)
+		{
+			NiagaraEmitter->OnRenderersChanged().AddSP(this, &SEmitterTrackWidget::RefreshRenderers);
+		}
 
 		// Enabled checkbox.
 		TrackBox->AddSlot()
diff --git a/Engine/Plugins/FX/Niagara/Source/NiagaraEditor/Private/ViewModels/HierarchyEditor/NiagaraSummaryViewViewModel.cpp b/Engine/Plugins/FX/Niagara/Source/NiagaraEditor/Private/ViewModels/HierarchyEditor/NiagaraSummaryViewViewModel.cpp
--- a/Engine/Plugins/FX/Niagara/Source/NiagaraEditor/Private/ViewModels/HierarchyEditor/NiagaraSummaryViewViewModel.cpp
+++ b/Engine/Plugins/FX/Niagara/Source/NiagaraEditor/Private/ViewModels/HierarchyEditor/NiagaraSummaryViewViewModel.cpp
@@ -117,11 +117,15 @@ void UNiagaraSummaryViewViewModel::Initialize(TSharedRef<FNiagaraEmitterViewMode
 {
 	EmitterViewModelWeak = EmitterViewModel;
 	EmitterViewModel->OnScriptGraphChanged().AddUObject(this, &UNiagaraSummaryViewViewModel::OnScriptGraphChanged);
-	EmitterViewModel->GetEmitter().Emitter->OnRenderersChanged().AddUObject(this, &UNiagaraSummaryViewViewModel::OnRenderersChanged);
-	EmitterViewModel->GetEmitter().Emitter->OnSimStagesChanged().AddUObject(this, &UNiagaraSummaryViewViewModel::OnSimStagesChanged);
-	EmitterViewModel->GetEmitter().Emitter->OnEventHandlersChanged().AddUObject(this, &UNiagaraSummaryViewViewModel::OnEventHandlersChanged);
+	//-TODO:Stateless: Do we need stateless support here?
+	if (EmitterViewModel->GetEmitter().Emitter)
+	{
+		EmitterViewModel->GetEmitter().Emitter->OnRenderersChanged().AddUObject(this, &UNiagaraSummaryViewViewModel::OnRenderersChanged);
+		EmitterViewModel->GetEmitter().Emitter->OnSimStagesChanged().AddUObject(this, &UNiagaraSummaryViewViewModel::OnSimStagesChanged);
+		EmitterViewModel->GetEmitter().Emitter->OnEventHandlersChanged().AddUObject(this, &UNiagaraSummaryViewViewModel::OnEventHandlersChanged);
 
-	UNiagaraHierarchyViewModelBase::Initialize();
+		UNiagaraHierarchyViewModelBase::Initialize();
+	}
 }
 
 void UNiagaraSummaryViewViewModel::FinalizeInternal()
diff --git a/Engine/Plugins/FX/Niagara/Source/NiagaraEditor/Private/ViewModels/NiagaraCurveSelectionViewModel.cpp b/Engine/Plugins/FX/Niagara/Source/NiagaraEditor/Private/ViewModels/NiagaraCurveSelectionViewModel.cpp
--- a/Engine/Plugins/FX/Niagara/Source/NiagaraEditor/Private/ViewModels/NiagaraCurveSelectionViewModel.cpp
+++ b/Engine/Plugins/FX/Niagara/Source/NiagaraEditor/Private/ViewModels/NiagaraCurveSelectionViewModel.cpp
@@ -886,9 +886,10 @@ void UNiagaraCurveSelectionViewModel::Refresh()
 		for (TSharedRef<FNiagaraEmitterHandleViewModel> EmitterHandleViewModel : SystemViewModel->GetEmitterHandleViewModels())
 		{
 			// During transactions Emitters may not have their EditorData set if they were created during the transaction; skip if this is true.
-			if (EmitterHandleViewModel->GetEmitterHandle() && EmitterHandleViewModel->GetEmitterHandle()->GetEmitterData()->GetEditorData())
+			FNiagaraEmitterHandle* EmitterHandle = EmitterHandleViewModel->GetEmitterHandle();
+			if (EmitterHandle && EmitterHandle->GetEmitterData() && EmitterHandle->GetEmitterData()->GetEditorData())
 			{
-				TSharedPtr<FNiagaraCurveSelectionTreeNode> EmitterNode = CreateNodeForEmitter(OldRootChildNodes, *EmitterHandleViewModel->GetEmitterHandle());
+				TSharedPtr<FNiagaraCurveSelectionTreeNode> EmitterNode = CreateNodeForEmitter(OldRootChildNodes, *EmitterHandle);
 				if (EmitterNode.IsValid())
 				{
 					NewRootChildNodes.Add(EmitterNode.ToSharedRef());
diff --git a/Engine/Plugins/FX/Niagara/Source/NiagaraEditor/Private/ViewModels/NiagaraEmitterHandleViewModel.cpp b/Engine/Plugins/FX/Niagara/Source/NiagaraEditor/Private/ViewModels/NiagaraEmitterHandleViewModel.cpp
--- a/Engine/Plugins/FX/Niagara/Source/NiagaraEditor/Private/ViewModels/NiagaraEmitterHandleViewModel.cpp
+++ b/Engine/Plugins/FX/Niagara/Source/NiagaraEditor/Private/ViewModels/NiagaraEmitterHandleViewModel.cpp
@@ -61,12 +61,16 @@ void FNiagaraEmitterHandleViewModel::GetRendererEntries(TArray<UNiagaraStackEntr
 	if (StackRoot)
 	{
 		TArray<UNiagaraStackEntry*> Children;
-		StackRoot->GetRenderGroup()->GetUnfilteredChildren(Children);
-		for (UNiagaraStackEntry* Child : Children)
+		UNiagaraStackRenderItemGroup* RendererGroup = StackRoot->GetRenderGroup();
+		if (RendererGroup != nullptr)
 		{
-			if (UNiagaraStackRendererItem* RendererItem = Cast<UNiagaraStackRendererItem>(Child))
+			RendererGroup->GetUnfilteredChildren(Children);
+			for (UNiagaraStackEntry* Child : Children)
 			{
-				InRenderingEntries.Add(Child);
+				if (UNiagaraStackRendererItem* RendererItem = Cast<UNiagaraStackRendererItem>(Child))
+				{
+					InRenderingEntries.Add(Child);
+				}
 			}
 		}
 	}
diff --git a/Engine/Plugins/FX/Niagara/Source/NiagaraEditor/Private/ViewModels/NiagaraEmitterViewModel.cpp b/Engine/Plugins/FX/Niagara/Source/NiagaraEditor/Private/ViewModels/NiagaraEmitterViewModel.cpp
--- a/Engine/Plugins/FX/Niagara/Source/NiagaraEditor/Private/ViewModels/NiagaraEmitterViewModel.cpp
+++ b/Engine/Plugins/FX/Niagara/Source/NiagaraEditor/Private/ViewModels/NiagaraEmitterViewModel.cpp
@@ -335,9 +335,8 @@ const UNiagaraEmitterEditorData& FNiagaraEmitterViewModel::GetEditorData() const
 	{
 		return *CastChecked<UNiagaraEmitterEditorData>(EmitterWeakPtr.GetEmitterData()->GetEditorData());
 	}
-	//-TODO:Stateless:Figure this out
+	//-TODO:Stateless: We don't have UNiagaraEmitterEditorData
 	return *GetDefault<UNiagaraEmitterEditorData>();
-	//-TODO:Stateless:Figure this out
 }
 
 UNiagaraEmitterEditorData& FNiagaraEmitterViewModel::GetEditorData()
@@ -346,9 +345,8 @@ UNiagaraEmitterEditorData& FNiagaraEmitterViewModel::GetEditorData()
 	{
 		return *CastChecked<UNiagaraEmitterEditorData>(EmitterWeakPtr.GetEmitterData()->GetEditorData());
 	}
-	//-TODO:Stateless:Figure this out
+	//-TODO:Stateless: We don't have UNiagaraEmitterEditorData
 	return *GetMutableDefault<UNiagaraEmitterEditorData>();
-	//-TODO:Stateless:Figure this out
 }
 
 void FNiagaraEmitterViewModel::GetEmitterMessageStores(TArray<FNiagaraMessageSourceAndStore>& OutMessageStores)
