- Experimental stateless emitters

diff --git a/Engine/Plugins/FX/Niagara/Source/NiagaraEditor/Private/ViewModels/NiagaraParameterPanelViewModel.cpp b/Engine/Plugins/FX/Niagara/Source/NiagaraEditor/Private/ViewModels/NiagaraParameterPanelViewModel.cpp
--- a/Engine/Plugins/FX/Niagara/Source/NiagaraEditor/Private/ViewModels/NiagaraParameterPanelViewModel.cpp
+++ b/Engine/Plugins/FX/Niagara/Source/NiagaraEditor/Private/ViewModels/NiagaraParameterPanelViewModel.cpp
@@ -823,11 +823,14 @@ void FNiagaraSystemToolkitParameterPanelViewModel::Cleanup()
 	GraphsToRemoveCallbacks.Add(GetGraphFromScript(System.GetSystemUpdateScript()));
 	for (TSharedRef<FNiagaraEmitterHandleViewModel> EmitterHandleViewModel : SystemViewModel->GetEmitterHandleViewModels())
 	{
-		if (FNiagaraEmitterHandle* EmitterHandle = EmitterHandleViewModel->GetEmitterHandle())
+		FNiagaraEmitterHandle* EmitterHandle = EmitterHandleViewModel->GetEmitterHandle();
+		//-TODO:Stateless: Add Stateless support
+		FVersionedNiagaraEmitterData* EmitterData = EmitterHandle ? EmitterHandle->GetEmitterData() : nullptr;
+		if (EmitterData)
 		{
 			TArray<UNiagaraScript*> EmitterScripts;
 			const bool bCompilableOnly = false;
-			EmitterHandle->GetEmitterData()->GetScripts(EmitterScripts, bCompilableOnly);
+			EmitterData->GetScripts(EmitterScripts, bCompilableOnly);
 			for (UNiagaraScript* EmitterScript : EmitterScripts)
 			{
 				GraphsToRemoveCallbacks.Add(GetGraphFromScript(EmitterScript));
@@ -917,6 +920,12 @@ void FNiagaraSystemToolkitParameterPanelViewModel::Init(const FSystemToolkitUICo
 	{
 		if (FNiagaraEmitterHandle* EmitterHandle = EmitterHandleViewModel->GetEmitterHandle())
 		{
+			//-TODO:Stateless: Add Stateless support?
+			if (EmitterHandle->GetEmitterData() == nullptr)
+			{
+				continue;
+			}
+
 			TArray<UNiagaraScript*> EmitterScripts;
 			const bool bCompilableOnly = false;
 			EmitterHandle->GetEmitterData()->GetScripts(EmitterScripts, bCompilableOnly);
@@ -2424,10 +2433,14 @@ void FNiagaraSystemToolkitParameterPanelViewModel::ReconcileOnGraphChangedBindin
 		{
 			TArray<UNiagaraScript*> EmitterScripts;
 			const bool bCompilableOnly = false;
-			EmitterHandle->GetEmitterData()->GetScripts(EmitterScripts, bCompilableOnly);
-			for (UNiagaraScript* EmitterScript : EmitterScripts)
+			//-TODO:Stateless: Add Stateless support?
+			if (FVersionedNiagaraEmitterData* EmitterData = EmitterHandle->GetEmitterData())
 			{
-				GraphsToAddCallbacks.Add(GetGraphFromScript(EmitterScript));
+				EmitterData->GetScripts(EmitterScripts, bCompilableOnly);
+				for (UNiagaraScript* EmitterScript : EmitterScripts)
+				{
+					GraphsToAddCallbacks.Add(GetGraphFromScript(EmitterScript));
+				}
 			}
 		}
 	}
diff --git a/Engine/Plugins/FX/Niagara/Source/NiagaraEditor/Private/ViewModels/NiagaraScratchPadUtilities.cpp b/Engine/Plugins/FX/Niagara/Source/NiagaraEditor/Private/ViewModels/NiagaraScratchPadUtilities.cpp
--- a/Engine/Plugins/FX/Niagara/Source/NiagaraEditor/Private/ViewModels/NiagaraScratchPadUtilities.cpp
+++ b/Engine/Plugins/FX/Niagara/Source/NiagaraEditor/Private/ViewModels/NiagaraScratchPadUtilities.cpp
@@ -27,6 +27,11 @@ void FNiagaraScratchPadUtilities::FixFunctionInputsFromFunctionScriptRename(UNia
 
 void FNiagaraScratchPadUtilities::FixExternalScratchPadScriptsForEmitter(UNiagaraSystem& SourceSystem, const FVersionedNiagaraEmitter& TargetEmitter)
 {
+	if ( TargetEmitter.Emitter == nullptr )
+	{
+		return;
+	}
+
 	UNiagaraSystem* TargetSystem = TargetEmitter.Emitter->GetTypedOuter<UNiagaraSystem>();
 	FVersionedNiagaraEmitterData* EmitterData = TargetEmitter.GetEmitterData();
 	if (EmitterData && TargetSystem != &SourceSystem)
diff --git a/Engine/Plugins/FX/Niagara/Source/NiagaraEditor/Private/ViewModels/NiagaraSystemScriptViewModel.cpp b/Engine/Plugins/FX/Niagara/Source/NiagaraEditor/Private/ViewModels/NiagaraSystemScriptViewModel.cpp
--- a/Engine/Plugins/FX/Niagara/Source/NiagaraEditor/Private/ViewModels/NiagaraSystemScriptViewModel.cpp
+++ b/Engine/Plugins/FX/Niagara/Source/NiagaraEditor/Private/ViewModels/NiagaraSystemScriptViewModel.cpp
@@ -52,12 +52,17 @@ void FNiagaraSystemScriptViewModel::OnSystemVMCompiled(UNiagaraSystem* InSystem)
 	SystemScripts.Add(InSystem->GetSystemUpdateScript());
 	ScriptsEnabled.Add(true);
 	ScriptsEnabled.Add(true);
-
 	
 	for (const FNiagaraEmitterHandle& Handle : InSystem->GetEmitterHandles())
 	{
+		//-TODO:Stateless: Do we need stateless support here?
+		FVersionedNiagaraEmitterData* EmitterData = Handle.GetEmitterData();
+		if (EmitterData == nullptr)
+		{
+			continue;
+		}
 		int32 NumScripts = SystemScripts.Num();
-		Handle.GetEmitterData()->GetScripts(SystemScripts, true);
+		EmitterData->GetScripts(SystemScripts, true);
 		for (; NumScripts < SystemScripts.Num(); NumScripts++)
 		{
 			if (Handle.GetIsEnabled())
diff --git a/Engine/Plugins/FX/Niagara/Source/NiagaraEditor/Private/ViewModels/NiagaraSystemSelectionViewModel.cpp b/Engine/Plugins/FX/Niagara/Source/NiagaraEditor/Private/ViewModels/NiagaraSystemSelectionViewModel.cpp
--- a/Engine/Plugins/FX/Niagara/Source/NiagaraEditor/Private/ViewModels/NiagaraSystemSelectionViewModel.cpp
+++ b/Engine/Plugins/FX/Niagara/Source/NiagaraEditor/Private/ViewModels/NiagaraSystemSelectionViewModel.cpp
@@ -476,6 +476,11 @@ void UNiagaraSystemSelectionViewModel::UpdateExternalSelectionState()
 	int32 NumRemoved = 0;
 	for (const FVersionedNiagaraEmitter& SelectedEmitter : SelectedEmitters)
 	{
+		//-TODO:Stateless: Do we need stateless support here?
+		if (SelectedEmitter.GetEmitterData() == nullptr)
+		{
+			continue;
+		}
 		FGuid EmitterId = GetSystemViewModel()->GetEmitterHandleViewModelForEmitter(SelectedEmitter)->GetId();
 		NumRemoved += OldSelectedEmitterHandleIds.Remove(EmitterId);
 		SelectedEmitterHandleIds.Add(EmitterId);
diff --git a/Engine/Plugins/FX/Niagara/Source/NiagaraEditor/Private/ViewModels/NiagaraSystemViewModel.cpp b/Engine/Plugins/FX/Niagara/Source/NiagaraEditor/Private/ViewModels/NiagaraSystemViewModel.cpp
--- a/Engine/Plugins/FX/Niagara/Source/NiagaraEditor/Private/ViewModels/NiagaraSystemViewModel.cpp
+++ b/Engine/Plugins/FX/Niagara/Source/NiagaraEditor/Private/ViewModels/NiagaraSystemViewModel.cpp
@@ -69,6 +69,11 @@
 #include "ViewModels/NiagaraParameterPanelViewModel.h"
 #include "Widgets/Notifications/SNotificationList.h"
 
+//-TODO:Stateless:
+#include "Stateless/NiagaraStatelessEmitter.h"
+#include "NiagaraSpriteRendererProperties.h"
+//-TODO:Stateless:
+
 DECLARE_CYCLE_STAT(TEXT("Niagara - SystemViewModel - CompileSystem"), STAT_NiagaraEditor_SystemViewModel_CompileSystem, STATGROUP_NiagaraEditor);
 
 #define LOCTEXT_NAMESPACE "NiagaraSystemViewModel"
@@ -595,6 +600,84 @@ TSharedPtr<FNiagaraEmitterHandleViewModel> FNiagaraSystemViewModel::AddEmptyEmit
 	EmptyEmitter->SetFlags(RF_Transactional);
 	return AddEmitter(FVersionedNiagaraEmitter(EmptyEmitter, FGuid()));
 }
+
+TSharedPtr<FNiagaraEmitterHandleViewModel> FNiagaraSystemViewModel::AddStatelessEmitter()
+{
+	UNiagaraSystem& NiagaraSystem = GetSystem();
+
+	auto SetupStatelessEmitter =
+		[](UNiagaraStatelessEmitter* StatelessEmitter, FName EmitterName)
+		{
+			StatelessEmitter->SetUniqueEmitterName(*EmitterName.ToString());
+
+			StatelessEmitter->AddSpawnInfo();
+			if (UClass* TemplateClass = FindObject<UClass>(nullptr, TEXT("/Script/Niagara.NiagaraStatelessEmitterDefault")))
+			{
+				StatelessEmitter->SetEmitterTemplateClass(TemplateClass);
+			}
+			{
+				FSoftObjectPath DefaultMaterial(TEXT("Material'/Niagara/DefaultAssets/DefaultSpriteMaterial.DefaultSpriteMaterial'"));
+
+				UNiagaraSpriteRendererProperties* NewRenderer = NewObject<UNiagaraSpriteRendererProperties>(StatelessEmitter, "Renderer", RF_Transactional);
+				NewRenderer->Material = Cast<UMaterialInterface>(DefaultMaterial.TryLoad());
+				StatelessEmitter->AddRenderer(NewRenderer, FGuid());// EmitterData->Version.VersionGuid);
+			}
+		};
+
+	// This method uses the combined standard / stateless setup
+	static bool bPureStatelessAdd = false;
+	if (bPureStatelessAdd == false)
+	{
+		TSharedPtr<FNiagaraEmitterHandleViewModel> NewViewModel = AddEmptyEmitter();
+		if (FNiagaraEmitterHandle* EmitterHandle = NewViewModel->GetEmitterHandle())
+		{
+			UNiagaraStatelessEmitter* NewEmitter = NewObject<UNiagaraStatelessEmitter>(&NiagaraSystem, "StatelessEmitter", RF_Transactional);
+			SetupStatelessEmitter(NewEmitter, EmitterHandle->GetName());
+
+			EmitterHandle->SetStatelessEmitter(NewEmitter);
+			EmitterHandle->SetEmitterMode(NiagaraSystem, ENiagaraEmitterMode::Stateless);
+		}
+		return NewViewModel;
+	}
+
+	// This method is pure stateless only which might result in compilation / graph issue that need ironing out
+	UNiagaraStatelessEmitter* NewEmitter = nullptr;
+	{
+		TSet<FName> EmitterHandleNames;
+		for (const FNiagaraEmitterHandle& EmitterHandle : NiagaraSystem.GetEmitterHandles())
+		{
+			EmitterHandleNames.Add(EmitterHandle.GetName());
+		}
+		FName NewEmitterName = FName("Stateless");
+		NewEmitterName = FNiagaraUtilities::GetUniqueName(NewEmitterName, EmitterHandleNames);
+		NewEmitter = NewObject<UNiagaraStatelessEmitter>(&NiagaraSystem, *NewEmitterName.ToString(), RF_Transactional);
+		SetupStatelessEmitter(NewEmitter, NewEmitterName);
+	}
+
+	ResetEmitterHandleViewModelsAndTracks();
+
+	GEditor->BeginTransaction(LOCTEXT("AddEmitter", "Add emitter"));
+	FNiagaraEmitterHandle EmitterHandle = FNiagaraEmitterHandle(*NewEmitter);
+	{
+		FNiagaraEditorUtilities::KillSystemInstances(NiagaraSystem);
+		NiagaraSystem.Modify();
+		NiagaraSystem.AddEmitterHandleDirect(EmitterHandle);
+		GetEditorData().SynchronizeOverviewGraphWithSystem(NiagaraSystem);
+	}
+	GEditor->EndTransaction();
+	const FGuid NewEmitterHandleId = EmitterHandle.GetId();
+
+	if (GetSystem().GetNumEmitters() == 1 && EditorSettings->GetAutoPlay() && Sequencer.IsValid())
+	{
+		// When adding a new emitter to an empty system start playing.
+		Sequencer->SetPlaybackStatus(EMovieScenePlayerStatus::Playing);
+	}
+
+	RefreshAll();
+	
+	TSharedPtr<FNiagaraEmitterHandleViewModel> NewEmitterHandleViewModel = GetEmitterHandleViewModelById(NewEmitterHandleId);
+	return NewEmitterHandleViewModel;
+}
 TSharedPtr<FNiagaraEmitterHandleViewModel> FNiagaraSystemViewModel::AddEmitter(const FVersionedNiagaraEmitter& VersionedEmitter)
 {
 	return AddEmitter(*VersionedEmitter.Emitter, VersionedEmitter.Version);
@@ -1394,10 +1477,9 @@ TArray<UNiagaraGraph*> FNiagaraSystemViewModel::GetAllGraphs()
 		for (const TSharedRef<FNiagaraEmitterHandleViewModel>& EmitterHandleViewModel : EmitterHandleViewModels)
 		{
 			FNiagaraEmitterHandle* EmitterHandle = EmitterHandleViewModel->GetEmitterHandle();
-			//-TODO:Stateless:
+			//-TODO:Stateless: Do we need stateless support here?
 			FVersionedNiagaraEmitterData* EmitterData = EmitterHandle ? EmitterHandle->GetEmitterData() : nullptr;
 			if (EmitterData == nullptr)
-			//-TODO:Stateless:
 			{
 				continue;
 			}
@@ -1411,10 +1493,9 @@ TArray<UNiagaraGraph*> FNiagaraSystemViewModel::GetAllGraphs()
 	else
 	{
 		FNiagaraEmitterHandle* EmitterHandle = GetEmitterHandleViewModels()[0]->GetEmitterHandle();
-		//-TODO:Stateless:
+		//-TODO:Stateless: Do we need stateless support here?
 		FVersionedNiagaraEmitterData* EmitterData = EmitterHandle ? EmitterHandle->GetEmitterData() : nullptr;
 		if (EmitterData != nullptr)
-		//-TODO:Stateless:
 		{
 			OutGraphs.Add(Cast<UNiagaraScriptSource>(EmitterData->GraphSource)->NodeGraph);
 		}
@@ -1609,9 +1690,8 @@ void FNiagaraSystemViewModel::SendLastCompileMessageJobs() const
 	for (const FNiagaraEmitterHandle& Handle : EmitterHandles)
 	{
 		FVersionedNiagaraEmitter EmitterInSystem = Handle.GetInstance();
-		//-TODO:Stateless:
+		//-TODO:Stateless: Do we need stateless support here?
 		if (FVersionedNiagaraEmitterData* EmitterData = EmitterInSystem.GetEmitterData())
-		//-TODO:Stateless:
 		{
 			TArray<UNiagaraScript*> EmitterScripts;
 			EmitterData->GetScripts(EmitterScripts, false);
@@ -1766,9 +1846,8 @@ void FNiagaraSystemViewModel::TickCompileStatus()
 					ScriptsToCheckForStatus.Add(System->GetSystemUpdateScript());
 					for (const FNiagaraEmitterHandle& EmitterHandle : System->GetEmitterHandles())
 					{
-						//-TODO:Stateless:
+						//-TODO:Stateless: Do we need stateless support here?
 						if (EmitterHandle.GetIsEnabled() && EmitterHandle.GetEmitterData())
-						//-TODO:Stateless:
 						{
 							EmitterHandle.GetEmitterData()->GetScripts(ScriptsToCheckForStatus, true);
 						}
@@ -3279,12 +3358,11 @@ void FNiagaraSystemViewModel::RefreshAssetMessages()
 		{
 			FVersionedNiagaraEmitter VersionedEmitter = EmitterHandle.GetInstance();
 			UNiagaraEmitter* Emitter = VersionedEmitter.Emitter.Get();
-			//-TODO:Stateless:
+			//-TODO:Stateless: Do we need stateless support here?
 			if (Emitter == nullptr)
 			{
 				continue;
 			}	
-		//-TODO:Stateless:
 			PublishMessages(FNiagaraMessageSourceAndStore(*Emitter, Emitter->GetMessageStore()));
 			MessageSourceObjectKeys.Add(FObjectKey(Emitter));
 			VersionedEmitter.GetEmitterData()->GetScripts(Scripts, false);
@@ -3298,9 +3376,8 @@ void FNiagaraSystemViewModel::RefreshAssetMessages()
 		if (ensureMsgf(EmitterHandles.Num() == 1, TEXT("There was not exactly 1 Emitter Handle for the SystemViewModel in Emitter edit mode!")))
 		{
 			FVersionedNiagaraEmitter VersionedEmitter = EmitterHandles[0].GetInstance();
-			//-TODO:Stateless:
+			//-TODO:Stateless: Do we need stateless support here?
 			if (UNiagaraEmitter* Emitter = VersionedEmitter.Emitter.Get())
-			//-TODO:Stateless:
 			{
 				PublishMessages(FNiagaraMessageSourceAndStore(*Emitter, Emitter->GetMessageStore()));
 				MessageSourceObjectKeys.Add(FObjectKey(Emitter));
diff --git a/Engine/Plugins/FX/Niagara/Source/NiagaraEditor/Private/ViewModels/Stack/NiagaraStackCommentCollection.cpp b/Engine/Plugins/FX/Niagara/Source/NiagaraEditor/Private/ViewModels/Stack/NiagaraStackCommentCollection.cpp
--- a/Engine/Plugins/FX/Niagara/Source/NiagaraEditor/Private/ViewModels/Stack/NiagaraStackCommentCollection.cpp
+++ b/Engine/Plugins/FX/Niagara/Source/NiagaraEditor/Private/ViewModels/Stack/NiagaraStackCommentCollection.cpp
@@ -48,8 +48,10 @@ void UNiagaraStackCommentCollection::RefreshChildrenInternal(const TArray<UNiaga
 			
 			if(StackObject == nullptr)
 			{
-			    StackObject = NewObject<UNiagaraStackObject>(this);
-			    StackObject->Initialize(CreateDefaultChildRequiredData(), CommentNode, true, GetStackEditorDataKey(), nullptr);
+				bool bIsInTopLevelObject = true;
+				bool bHideTopLevelCategories = false;
+				StackObject = NewObject<UNiagaraStackObject>(this);
+				StackObject->Initialize(CreateDefaultChildRequiredData(), CommentNode, bIsInTopLevelObject, bHideTopLevelCategories, GetStackEditorDataKey(), nullptr);
 			}
 			
 			NewChildren.Add(StackObject);
diff --git a/Engine/Plugins/FX/Niagara/Source/NiagaraEditor/Private/ViewModels/Stack/NiagaraStackEmitterPropertiesGroup.cpp b/Engine/Plugins/FX/Niagara/Source/NiagaraEditor/Private/ViewModels/Stack/NiagaraStackEmitterPropertiesGroup.cpp
--- a/Engine/Plugins/FX/Niagara/Source/NiagaraEditor/Private/ViewModels/Stack/NiagaraStackEmitterPropertiesGroup.cpp
+++ b/Engine/Plugins/FX/Niagara/Source/NiagaraEditor/Private/ViewModels/Stack/NiagaraStackEmitterPropertiesGroup.cpp
@@ -216,6 +216,11 @@ const FSlateBrush* UNiagaraStackEmitterPropertiesGroup::GetSecondaryIconBrush()
 	if (IsFinalized() == false && GetEmitterViewModel().IsValid())
 	{
 		FVersionedNiagaraEmitterData* EmitterData = GetEmitterViewModel()->GetEmitter().GetEmitterData();
+		//-TODO:Stateless: Do we need a stateless icon?
+		if (EmitterData == nullptr)
+		{
+			return FNiagaraEditorStyle::Get().GetBrush("NiagaraEditor.Stack.GPUIcon");
+		}
 		if (EmitterData->SimTarget == ENiagaraSimTarget::CPUSim)
 		{
 			return FNiagaraEditorStyle::Get().GetBrush("NiagaraEditor.Stack.CPUIcon");
diff --git a/Engine/Plugins/FX/Niagara/Source/NiagaraEditor/Private/ViewModels/Stack/NiagaraStackEmitterSettingsGroup.cpp b/Engine/Plugins/FX/Niagara/Source/NiagaraEditor/Private/ViewModels/Stack/NiagaraStackEmitterSettingsGroup.cpp
--- a/Engine/Plugins/FX/Niagara/Source/NiagaraEditor/Private/ViewModels/Stack/NiagaraStackEmitterSettingsGroup.cpp
+++ b/Engine/Plugins/FX/Niagara/Source/NiagaraEditor/Private/ViewModels/Stack/NiagaraStackEmitterSettingsGroup.cpp
@@ -32,7 +32,10 @@ void UNiagaraStackEmitterPropertiesItem::Initialize(FRequiredEntryData InRequire
 {
 	Super::Initialize(InRequiredEntryData, TEXT("EmitterProperties"));
 	EmitterWeakPtr = GetEmitterViewModel()->GetEmitter().ToWeakPtr();
-	EmitterWeakPtr.Emitter->OnPropertiesChanged().AddUObject(this, &UNiagaraStackEmitterPropertiesItem::EmitterPropertiesChanged);
+	if (EmitterWeakPtr.IsValid())
+	{
+		EmitterWeakPtr.Emitter->OnPropertiesChanged().AddUObject(this, &UNiagaraStackEmitterPropertiesItem::EmitterPropertiesChanged);
+	}
 }
 
 void UNiagaraStackEmitterPropertiesItem::FinalizeInternal()
@@ -118,9 +121,13 @@ void UNiagaraStackEmitterPropertiesItem::RefreshChildrenInternal(const TArray<UN
 	{
 		EmitterObject = NewObject<UNiagaraStackObject>(this);
 		FRequiredEntryData RequiredEntryData(GetSystemViewModel(), GetEmitterViewModel(), FExecutionCategoryNames::Emitter, NAME_None, GetStackEditorData());
-		bool bIsTopLevelObject = true;
-		EmitterObject->Initialize(RequiredEntryData, EmitterWeakPtr.Emitter.Get(), bIsTopLevelObject, GetStackEditorDataKey());
-		EmitterObject->RegisterInstancedCustomPropertyLayout(UNiagaraEmitter::StaticClass(), FOnGetDetailCustomizationInstance::CreateStatic(&FNiagaraEmitterDetails::MakeInstance));
+		if (EmitterWeakPtr.IsValid())
+		{
+			bool bIsTopLevelObject = true;
+			bool bHideTopLevelCategories = false;
+			EmitterObject->Initialize(RequiredEntryData, EmitterWeakPtr.Emitter.Get(), bIsTopLevelObject, bHideTopLevelCategories, GetStackEditorDataKey());
+			EmitterObject->RegisterInstancedCustomPropertyLayout(UNiagaraEmitter::StaticClass(), FOnGetDetailCustomizationInstance::CreateStatic(&FNiagaraEmitterDetails::MakeInstance));
+		}
 	}
 
 	NewChildren.Add(EmitterObject);
diff --git a/Engine/Plugins/FX/Niagara/Source/NiagaraEditor/Private/ViewModels/Stack/NiagaraStackEntry.cpp b/Engine/Plugins/FX/Niagara/Source/NiagaraEditor/Private/ViewModels/Stack/NiagaraStackEntry.cpp
--- a/Engine/Plugins/FX/Niagara/Source/NiagaraEditor/Private/ViewModels/Stack/NiagaraStackEntry.cpp
+++ b/Engine/Plugins/FX/Niagara/Source/NiagaraEditor/Private/ViewModels/Stack/NiagaraStackEntry.cpp
@@ -720,7 +720,12 @@ void UNiagaraStackEntry::FinalizeInternal()
 void UNiagaraStackEntry::RefreshChildren()
 {
 	checkf(bIsFinalized == false, TEXT("Can not refresh children on an entry after it has been finalized."));
-	checkf(SystemViewModel.IsValid(), TEXT("Base stack entry not initialized."));
+	//-TODO:Stateless: Do we need stateless support here?
+	//checkf(SystemViewModel.IsValid(), TEXT("Base stack entry not initialized."));
+	if (SystemViewModel.IsValid() == false)
+	{
+		return;
+	}
 
 	for (UNiagaraStackEntry* Child : Children)
 	{
@@ -829,20 +834,24 @@ void UNiagaraStackEntry::RefreshChildren()
 		ErrorChild->OnIssueModified().AddUObject(this, &UNiagaraStackEntry::IssueModified);
 	}
 
-	const FText* NewAlternateName = StackEditorData->GetStackEntryDisplayName(StackEditorDataKey);
-	if (NewAlternateName != nullptr && NewAlternateName->IsEmptyOrWhitespace() == false)
+	//-TODO:Stateless: Do we need stateless support here?
+	if (StackEditorData)
 	{
-		if (AlternateDisplayName.IsSet() == false || NewAlternateName->IdenticalTo(AlternateDisplayName.GetValue()) == false)
+		const FText* NewAlternateName = StackEditorData->GetStackEntryDisplayName(StackEditorDataKey);
+		if (NewAlternateName != nullptr && NewAlternateName->IsEmptyOrWhitespace() == false)
 		{
-			AlternateDisplayName = *NewAlternateName;
+			if (AlternateDisplayName.IsSet() == false || NewAlternateName->IdenticalTo(AlternateDisplayName.GetValue()) == false)
+			{
+				AlternateDisplayName = *NewAlternateName;
+				AlternateDisplayNameChangedDelegate.Broadcast();
+			}
+		}
+		else if (AlternateDisplayName.IsSet())
+		{
+			AlternateDisplayName.Reset();
 			AlternateDisplayNameChangedDelegate.Broadcast();
 		}
 	}
-	else if(AlternateDisplayName.IsSet())
-	{
-		AlternateDisplayName.Reset();
-		AlternateDisplayNameChangedDelegate.Broadcast();
-	}
 
 	PostRefreshChildrenInternal();
 
diff --git a/Engine/Plugins/FX/Niagara/Source/NiagaraEditor/Private/ViewModels/Stack/NiagaraStackEventScriptItemGroup.cpp b/Engine/Plugins/FX/Niagara/Source/NiagaraEditor/Private/ViewModels/Stack/NiagaraStackEventScriptItemGroup.cpp
--- a/Engine/Plugins/FX/Niagara/Source/NiagaraEditor/Private/ViewModels/Stack/NiagaraStackEventScriptItemGroup.cpp
+++ b/Engine/Plugins/FX/Niagara/Source/NiagaraEditor/Private/ViewModels/Stack/NiagaraStackEventScriptItemGroup.cpp
@@ -117,11 +117,12 @@ void UNiagaraStackEventHandlerPropertiesItem::RefreshChildrenInternal(const TArr
 
 		EmitterObject = NewObject<UNiagaraStackObject>(this);
 		bool bIsTopLevelObject = true;
-		EmitterObject->Initialize(CreateDefaultChildRequiredData(), EventWrapper, bIsTopLevelObject, GetStackEditorDataKey());
+		bool bHideTopLevelCategories = false;
+		EmitterObject->Initialize(CreateDefaultChildRequiredData(), EventWrapper, bIsTopLevelObject, bHideTopLevelCategories, GetStackEditorDataKey());
 		EmitterObject->RegisterInstancedCustomPropertyTypeLayout(FNiagaraEventScriptProperties::StaticStruct()->GetFName(), 
 			FOnGetPropertyTypeCustomizationInstance::CreateStatic(&FNiagaraEventScriptPropertiesCustomization::MakeInstance, 
 			TWeakObjectPtr<UNiagaraSystem>(&GetSystemViewModel()->GetSystem()), GetEmitterViewModel()->GetEmitter().ToWeakPtr()));
-		EmitterObject->SetOnGetCustomRootNodes(UNiagaraStackObject::FOnGetCustomRootNodes::CreateUObject(this, &UNiagaraStackEventHandlerPropertiesItem::SelectEmitterStackObjectRootTreeNodes));
+		EmitterObject->SetOnFilterDetailNodes(FNiagaraStackObjectShared::FOnFilterDetailNodes::CreateUObject(this, &UNiagaraStackEventHandlerPropertiesItem::FilterEmitterStackObjectRootTreeNodes));
 	}
 
 	FVersionedNiagaraEmitterData* EmitterData = GetEmitterViewModel()->GetEmitter().GetEmitterData();
@@ -201,9 +202,9 @@ TSharedPtr<IDetailTreeNode> GetEventHandlerArrayPropertyNode(const TArray<TShare
 	return GetEventHandlerArrayPropertyNode(ChildrenToCheck);
 }
 
-void UNiagaraStackEventHandlerPropertiesItem::SelectEmitterStackObjectRootTreeNodes(TArray<TSharedRef<IDetailTreeNode>> Source, TArray<TSharedRef<IDetailTreeNode>>* Selected)
+void UNiagaraStackEventHandlerPropertiesItem::FilterEmitterStackObjectRootTreeNodes(const TArray<TSharedRef<IDetailTreeNode>>& InSourceNodes, TArray<TSharedRef<IDetailTreeNode>>& OutFilteredNodes)
 {
-	TSharedPtr<IDetailTreeNode> EventHandlerArrayPropertyNode = GetEventHandlerArrayPropertyNode(Source);
+	TSharedPtr<IDetailTreeNode> EventHandlerArrayPropertyNode = GetEventHandlerArrayPropertyNode(InSourceNodes);
 	if (EventHandlerArrayPropertyNode.IsValid())
 	{
 		TArray<TSharedRef<IDetailTreeNode>> EventHandlerArrayItemNodes;
@@ -223,7 +224,7 @@ void UNiagaraStackEventHandlerPropertiesItem::SelectEmitterStackObjectRootTreeNo
 						FNiagaraEventScriptProperties* EventScriptProperties = static_cast<FNiagaraEventScriptProperties*>(RawData[0]);
 						if (EventScriptProperties->Script->GetUsageId() == EventScriptUsageId)
 						{
-							EventHandlerArrayItemNode->GetChildren(*Selected);
+							EventHandlerArrayItemNode->GetChildren(OutFilteredNodes);
 							return;
 						}
 					}
diff --git a/Engine/Plugins/FX/Niagara/Source/NiagaraEditor/Private/ViewModels/Stack/NiagaraStackFunctionInput.cpp b/Engine/Plugins/FX/Niagara/Source/NiagaraEditor/Private/ViewModels/Stack/NiagaraStackFunctionInput.cpp
--- a/Engine/Plugins/FX/Niagara/Source/NiagaraEditor/Private/ViewModels/Stack/NiagaraStackFunctionInput.cpp
+++ b/Engine/Plugins/FX/Niagara/Source/NiagaraEditor/Private/ViewModels/Stack/NiagaraStackFunctionInput.cpp
@@ -938,7 +938,8 @@ void UNiagaraStackFunctionInput::RefreshChildrenInternal(const TArray<UNiagaraSt
 		{
 			ValueObjectEntry = NewObject<UNiagaraStackObject>(this);
 			bool bIsTopLevelObject = false;
-			ValueObjectEntry->Initialize(CreateDefaultChildRequiredData(), InputValues.DataObject.Get(), bIsTopLevelObject, GetOwnerStackItemEditorDataKey(), OwningFunctionCallNode.Get());
+			bool bHideTopLevelCategories = false;
+			ValueObjectEntry->Initialize(CreateDefaultChildRequiredData(), InputValues.DataObject.Get(), bIsTopLevelObject, bHideTopLevelCategories, GetOwnerStackItemEditorDataKey(), OwningFunctionCallNode.Get());
 		}
 		NewChildren.Add(ValueObjectEntry);
 	}
diff --git a/Engine/Plugins/FX/Niagara/Source/NiagaraEditor/Private/ViewModels/Stack/NiagaraStackInputCategory.cpp b/Engine/Plugins/FX/Niagara/Source/NiagaraEditor/Private/ViewModels/Stack/NiagaraStackInputCategory.cpp
--- a/Engine/Plugins/FX/Niagara/Source/NiagaraEditor/Private/ViewModels/Stack/NiagaraStackInputCategory.cpp
+++ b/Engine/Plugins/FX/Niagara/Source/NiagaraEditor/Private/ViewModels/Stack/NiagaraStackInputCategory.cpp
@@ -18,6 +18,7 @@
 #include "ViewModels/Stack/NiagaraStackObject.h"
 #include "ViewModels/Stack/NiagaraStackPropertyRow.h"
 #include "ViewModels/Stack/NiagaraStackRendererItem.h"
+#include "ViewModels/Stack/NiagaraStackRenderersOwner.h"
 #include "ViewModels/Stack/NiagaraStackSimulationStageGroup.h"
 
 #include UE_INLINE_GENERATED_CPP_BY_NAME(NiagaraStackInputCategory)
@@ -530,7 +531,7 @@ void UNiagaraStackSummaryCategory::RefreshChildrenInternal(const TArray<UNiagara
 				if (StackRenderer == nullptr)
 				{
 					StackRenderer = NewObject<UNiagaraStackRendererItem>(this);
-					StackRenderer->Initialize(CreateDefaultChildRequiredData(), *MatchingRendererProperties);
+					StackRenderer->Initialize(CreateDefaultChildRequiredData(), FNiagaraStackRenderersOwnerStandard::CreateShared(GetEmitterViewModel().ToSharedRef()), *MatchingRendererProperties);
 				}
 			
 				NewChildren.Add(StackRenderer);	
@@ -648,35 +649,28 @@ void UNiagaraStackSummaryCategory::RefreshChildrenInternal(const TArray<UNiagara
 				UNiagaraStackObject* StackObjectWithProperty = FindCurrentChildOfTypeByPredicate<UNiagaraStackObject>(CurrentChildren,
 				[&](UNiagaraStackObject* StackObjectCandidate)
 				{
-					if(StackObjectCandidate->GetObject() != Object)
-					{
-						return false;
-					}
-
-					if(StackObjectCandidate->GetCustomRootNodeNames().Num() == 1 && StackObjectCandidate->GetCustomRootNodeNames()[0] == ObjectProperty->GetPersistentIdentity().Names[0])
-					{
-						return true;
-					}
-
-					return false;
+					return StackObjectCandidate->GetObject() == Object && StackObjectCandidate->GetCustomName() == ObjectProperty->GetPersistentIdentity().Names[0];
 				});
 
 				if(StackObjectWithProperty == nullptr)
 				{
 					StackObjectWithProperty = NewObject<UNiagaraStackObject>(this);
-					StackObjectWithProperty->Initialize(CreateDefaultChildRequiredData(), Object, false, GetStackEditorDataKey(), nullptr);
-					StackObjectWithProperty->SetOnGetCustomRootNodes(UNiagaraStackObject::FOnGetCustomRootNodes::CreateLambda([PropertyName = ObjectProperty->GetPersistentIdentity().Names[0]](TArray<TSharedRef<IDetailTreeNode>> DefaultRootNodes, TArray<TSharedRef<IDetailTreeNode>>* Result)
+					bool bIsInTopLevelObject = false;
+					bool bHideTopLevelCategories = false;
+					StackObjectWithProperty->Initialize(CreateDefaultChildRequiredData(), Object, bIsInTopLevelObject, bHideTopLevelCategories, GetStackEditorDataKey(), nullptr);
+					StackObjectWithProperty->SetCustomName(ObjectProperty->GetPersistentIdentity().Names[0]);
+					StackObjectWithProperty->SetOnFilterDetailNodes(FNiagaraStackObjectShared::FOnFilterDetailNodes::CreateLambda([PropertyName = ObjectProperty->GetPersistentIdentity().Names[0]]( const TArray<TSharedRef<IDetailTreeNode>>& InSourceNodes, TArray<TSharedRef<IDetailTreeNode>>& OutFilteredNodes)
 					{
-						for(TSharedRef<IDetailTreeNode>& RootNode : DefaultRootNodes)
+						for(const TSharedRef<IDetailTreeNode>& SourceNode : InSourceNodes)
 						{
 							TArray<TSharedRef<IDetailTreeNode>> ChildrenNodes;
-							RootNode->GetChildren(ChildrenNodes);
+							SourceNode->GetChildren(ChildrenNodes);
 							
 							for(TSharedRef<IDetailTreeNode> ChildNode : ChildrenNodes)
 							{
 								if(ChildNode->GetNodeName() == PropertyName)
 								{
-									Result->Add(ChildNode);
+									OutFilteredNodes.Add(ChildNode);
 								}
 							}
 						}
diff --git a/Engine/Plugins/FX/Niagara/Source/NiagaraEditor/Private/ViewModels/Stack/NiagaraStackObject.cpp b/Engine/Plugins/FX/Niagara/Source/NiagaraEditor/Private/ViewModels/Stack/NiagaraStackObject.cpp
--- a/Engine/Plugins/FX/Niagara/Source/NiagaraEditor/Private/ViewModels/Stack/NiagaraStackObject.cpp
+++ b/Engine/Plugins/FX/Niagara/Source/NiagaraEditor/Private/ViewModels/Stack/NiagaraStackObject.cpp
@@ -24,15 +24,17 @@ UNiagaraStackObject::UNiagaraStackObject()
 {
 }
 
-void UNiagaraStackObject::Initialize(FRequiredEntryData InRequiredEntryData, UObject* InObject, bool bInIsTopLevelObject, FString InOwnerStackItemEditorDataKey, UNiagaraNode* InOwningNiagaraNode)
+void UNiagaraStackObject::Initialize(FRequiredEntryData InRequiredEntryData, UObject* InObject, bool bInIsTopLevelObject, bool bInHideTopLevelCategories, FString InOwnerStackItemEditorDataKey, UNiagaraNode* InOwningNiagaraNode)
 {
-	checkf(WeakObject.IsValid() == false, TEXT("Can only initialize once."));
+	checkf(WeakObject.IsValid() == false || DisplayedStruct.IsValid() == false, TEXT("Can only initialize once."));
 	FString ObjectStackEditorDataKey = FString::Printf(TEXT("%s-%s"), *InOwnerStackItemEditorDataKey, *InObject->GetName());
 	Super::Initialize(InRequiredEntryData, InOwnerStackItemEditorDataKey, ObjectStackEditorDataKey);
 	WeakObject = InObject;
-	bIsTopLevelObject = bInIsTopLevelObject;
+	bIsTopLevel = bInIsTopLevelObject;
+	bHideTopLevelCategories = bInHideTopLevelCategories;
 	OwningNiagaraNode = InOwningNiagaraNode;
 	bIsRefreshingDataInterfaceErrors = false;
+	FilterMode = EDetailNodeFilterMode::FilterRootNodesOnly;
 
 	MessageLogGuid = GetSystemViewModel()->GetMessageLogGuid();
 
@@ -44,21 +46,31 @@ void UNiagaraStackObject::Initialize(FRequiredEntryData InRequiredEntryData, UOb
 	).BindUObject(this, &UNiagaraStackObject::OnMessageManagerRefresh);
 }
 
-void UNiagaraStackObject::SetOnGetCustomRootNodes(FOnGetCustomRootNodes OnSelectRootNodes)
+void UNiagaraStackObject::Initialize(
+		FRequiredEntryData InRequiredEntryData,
+		UObject* InOwningObject,
+		TSharedRef<FStructOnScope> InDisplayedStruct,
+		const FString& InStructName,
+		bool bInIsTopLevelStruct,
+		bool bInHideTopLevelCategories,
+		FString InOwnerStackItemEditorDataKey,
+		UNiagaraNode* InOwningNiagaraNode)
 {
-	OnSelectRootNodesDelegate = OnSelectRootNodes;
+	checkf(WeakObject.IsValid() == false && DisplayedStruct.IsValid() == false, TEXT("Can only initialize once."));
+	FString ObjectStackEditorDataKey = FString::Printf(TEXT("%s-%s"), *InOwnerStackItemEditorDataKey, *InStructName);
+	Super::Initialize(InRequiredEntryData, InOwnerStackItemEditorDataKey, ObjectStackEditorDataKey);
+	WeakObject = InOwningObject;
+	DisplayedStruct = InDisplayedStruct;
+	bIsTopLevel = bInIsTopLevelStruct;
+	bHideTopLevelCategories = bInHideTopLevelCategories;
+	OwningNiagaraNode = InOwningNiagaraNode;
+	bIsRefreshingDataInterfaceErrors = false;
 }
 
-TArray<FName> UNiagaraStackObject::GetCustomRootNodeNames() const
+void UNiagaraStackObject::SetOnFilterDetailNodes(FNiagaraStackObjectShared::FOnFilterDetailNodes InOnFilterDetailNodes, EDetailNodeFilterMode InFilterMode)
 {
-	TArray<FName> Result;
-
-	Algo::Transform(CustomRootNodes, Result, [](TSharedRef<IDetailTreeNode> RootNote)
-	{
-		return RootNote->GetNodeName();
-	});
-
-	return Result;
+	OnFilterDetailNodesDelegate = InOnFilterDetailNodes;
+	FilterMode = InFilterMode;
 }
 
 void UNiagaraStackObject::RegisterInstancedCustomPropertyLayout(UStruct* Class, FOnGetDetailCustomizationInstance DetailLayoutDelegate)
@@ -73,16 +85,23 @@ void UNiagaraStackObject::RegisterInstancedCustomPropertyTypeLayout(FName Proper
 	RegisteredPropertyCustomizations.Add({ PropertyTypeName, PropertyTypeLayoutDelegate, Identifier });
 }
 
-UObject* UNiagaraStackObject::GetObject()
+void UNiagaraStackObject::NotifyPreChange(FProperty* PropertyAboutToChange)
 {
-	return WeakObject.Get();
+	UObject* Object = GetObject();
+	if (Object != nullptr && DisplayedStruct.IsValid() && DisplayedStruct->OwnsStructMemory() == false)
+	{
+		Object->Modify();
+	}
 }
 
 void UNiagaraStackObject::NotifyPostChange(const FPropertyChangedEvent& PropertyChangedEvent, FProperty* PropertyThatChanged)
 {
-	TArray<UObject*> ChangedObjects;
-	ChangedObjects.Add(GetObject());
-	OnDataObjectModified().Broadcast(ChangedObjects, ENiagaraDataObjectChange::Changed);
+	if (GetObject() != nullptr)
+	{
+		TArray<UObject*> ChangedObjects;
+		ChangedObjects.Add(GetObject());
+		OnDataObjectModified().Broadcast(ChangedObjects, ENiagaraDataObjectChange::Changed);
+	}
 }
 
 bool UNiagaraStackObject::GetIsEnabled() const
@@ -109,12 +128,18 @@ FNiagaraHierarchyIdentity UNiagaraStackObject::DetermineSummaryIdentity() const
 
 void UNiagaraStackObject::FinalizeInternal()
 {
-	CustomRootNodes.Empty();
-	
 	if (PropertyRowGenerator.IsValid())
 	{
 		PropertyRowGenerator->OnRowsRefreshed().RemoveAll(this);
-		PropertyRowGenerator->SetObjects(TArray<UObject*>());
+
+		if (DisplayedStruct.IsValid())
+		{
+			PropertyRowGenerator->SetStructure(TSharedPtr<FStructOnScope>());
+		}
+		else
+		{
+			PropertyRowGenerator->SetObjects(TArray<UObject*>());
+		}
 
 		// Enqueue the row generator for destruction because stack entries might be finalized during the system view model tick
 		// and you can't destruct tickables while other tickables are being ticked.
@@ -186,13 +211,15 @@ void UNiagaraStackObject::RefreshChildrenInternal(const TArray<UNiagaraStackEntr
 	};
 
 	UObject* Object = WeakObject.Get();
-	if ( Object == nullptr )
+	if (Object == nullptr &&  DisplayedStruct.IsValid() == false)
 	{
 		return;
 	}
 
-	GatherIssueFromProperties((uint8*)Object, Object->GetClass(), true);
-
+	if (Object != nullptr && DisplayedStruct.IsValid() == false)
+	{
+		GatherIssueFromProperties((uint8*)Object, Object->GetClass(), true);
+	}
 
 	if (GetSystemViewModel()->GetIsForDataProcessingOnly() == false && PropertyRowGenerator.IsValid() == false)
 	{
@@ -212,9 +239,16 @@ void UNiagaraStackObject::RefreshChildrenInternal(const TArray<UNiagaraStackEntr
 				RegisteredPropertyCustomization.PropertyTypeLayoutDelegate, RegisteredPropertyCustomization.Identifier);
 		}
 
-		TArray<UObject*> Objects;
-		Objects.Add(Object);
-		PropertyRowGenerator->SetObjects(Objects);
+		if (DisplayedStruct.IsValid())
+		{
+			PropertyRowGenerator->SetStructure(DisplayedStruct);
+		}
+		else if (Object != nullptr)
+		{
+			TArray<UObject*> Objects;
+			Objects.Add(Object);
+			PropertyRowGenerator->SetObjects(Objects);
+		}
 
 		// Add the refresh delegate after setting the objects to prevent refreshing children immediately.
 		PropertyRowGenerator->OnRowsRefreshed().AddUObject(this, &UNiagaraStackObject::PropertyRowsRefreshed);
@@ -299,12 +333,12 @@ void UNiagaraStackObject::RefreshChildrenInternal(const TArray<UNiagaraStackEntr
 
 	if(PropertyRowGenerator.IsValid())
 	{
+		PropertyRowGenerator->InvalidateCachedState();
 		TArray<TSharedRef<IDetailTreeNode>> DefaultRootTreeNodes = PropertyRowGenerator->GetRootTreeNodes();
 		TArray<TSharedRef<IDetailTreeNode>> RootTreeNodes;
-		if (OnSelectRootNodesDelegate.IsBound())
+		if (OnFilterDetailNodesDelegate.IsBound())
 		{
-			OnSelectRootNodesDelegate.Execute(DefaultRootTreeNodes, &RootTreeNodes);
-			CustomRootNodes = RootTreeNodes;
+			OnFilterDetailNodesDelegate.Execute(DefaultRootTreeNodes, RootTreeNodes);
 		}
 		else
 		{
@@ -359,8 +393,12 @@ void UNiagaraStackObject::RefreshChildrenInternal(const TArray<UNiagaraStackEntr
 			if (ChildRow == nullptr)
 			{
 				ChildRow = NewObject<UNiagaraStackPropertyRow>(this);
-				ChildRow->Initialize(CreateDefaultChildRequiredData(), RootTreeNode, bIsTopLevelObject, GetOwnerStackItemEditorDataKey(), GetOwnerStackItemEditorDataKey(), OwningNiagaraNode);
+				ChildRow->Initialize(CreateDefaultChildRequiredData(), RootTreeNode, bIsTopLevel, bHideTopLevelCategories, GetOwnerStackItemEditorDataKey(), GetOwnerStackItemEditorDataKey(), OwningNiagaraNode);
 				ChildRow->SetOwnerGuid(ObjectGuid);
+				if (OnFilterDetailNodesDelegate.IsBound() && FilterMode == EDetailNodeFilterMode::FilterAllNodes)
+				{
+					ChildRow->SetOnFilterDetailNodes(OnFilterDetailNodesDelegate);
+				}
 			}
 
 			NewChildren.Add(ChildRow);
diff --git a/Engine/Plugins/FX/Niagara/Source/NiagaraEditor/Private/ViewModels/Stack/NiagaraStackPropertyRow.cpp b/Engine/Plugins/FX/Niagara/Source/NiagaraEditor/Private/ViewModels/Stack/NiagaraStackPropertyRow.cpp
--- a/Engine/Plugins/FX/Niagara/Source/NiagaraEditor/Private/ViewModels/Stack/NiagaraStackPropertyRow.cpp
+++ b/Engine/Plugins/FX/Niagara/Source/NiagaraEditor/Private/ViewModels/Stack/NiagaraStackPropertyRow.cpp
@@ -11,7 +11,14 @@
 
 #include UE_INLINE_GENERATED_CPP_BY_NAME(NiagaraStackPropertyRow)
 
-void UNiagaraStackPropertyRow::Initialize(FRequiredEntryData InRequiredEntryData, TSharedRef<IDetailTreeNode> InDetailTreeNode, bool bInIsTopLevelProperty, FString InOwnerStackItemEditorDataKey, FString InOwnerStackEditorDataKey, UNiagaraNode* InOwningNiagaraNode)
+void UNiagaraStackPropertyRow::Initialize(
+	FRequiredEntryData InRequiredEntryData,
+	TSharedRef<IDetailTreeNode> InDetailTreeNode,
+	bool bInIsTopLevelProperty,
+	bool bInHideTopLevelCategories,
+	FString InOwnerStackItemEditorDataKey,
+	FString InOwnerStackEditorDataKey,
+	UNiagaraNode* InOwningNiagaraNode)
 {
 	TSharedPtr<IPropertyHandle> PropertyHandle = InDetailTreeNode->CreatePropertyHandle();
 	FString RowStackEditorDataKey = FString::Printf(TEXT("%s-%s"), *InOwnerStackEditorDataKey, *InDetailTreeNode->GetNodeName().ToString());
@@ -20,6 +27,7 @@ void UNiagaraStackPropertyRow::Initialize(FRequiredEntryData InRequiredEntryData
 	SetIsAdvanced(bRowIsAdvanced);
 	DetailTreeNode = InDetailTreeNode;
 	bIsTopLevelProperty = bInIsTopLevelProperty;
+	bHideTopLevelCategories = bInHideTopLevelCategories;
 	OwningNiagaraNode = InOwningNiagaraNode;
 	CategorySpacer = nullptr;
 	if (DetailTreeNode->GetNodeType() == EDetailNodeType::Category)
@@ -40,6 +48,7 @@ void UNiagaraStackPropertyRow::Initialize(FRequiredEntryData InRequiredEntryData
 			bCannotEditInThisContext = true;
 		}
 	}
+	bIsHiddenCategory = false;
 }
 
 TSharedRef<IDetailTreeNode> UNiagaraStackPropertyRow::GetDetailTreeNode() const
@@ -59,6 +68,11 @@ UNiagaraStackEntry::EStackRowStyle UNiagaraStackPropertyRow::GetStackRowStyle()
 	return RowStyle;
 }
 
+bool UNiagaraStackPropertyRow::GetShouldShowInStack() const
+{
+	return bIsHiddenCategory == false;
+}
+
 bool UNiagaraStackPropertyRow::HasOverridenContent() const
 {
 	TSharedPtr<IPropertyHandle> PropertyHandle = DetailTreeNode->CreatePropertyHandle();
@@ -88,8 +102,20 @@ void UNiagaraStackPropertyRow::FinalizeInternal()
 
 void UNiagaraStackPropertyRow::RefreshChildrenInternal(const TArray<UNiagaraStackEntry*>& CurrentChildren, TArray<UNiagaraStackEntry*>& NewChildren, TArray<FStackIssue>& NewIssues)
 {
+	TArray<TSharedRef<IDetailTreeNode>> AllNodeChildren;
+	DetailTreeNode->GetChildren(AllNodeChildren);
+
 	TArray<TSharedRef<IDetailTreeNode>> NodeChildren;
-	DetailTreeNode->GetChildren(NodeChildren);
+	if (OnFilterDetailNodes.IsBound())
+	{
+		OnFilterDetailNodes.Execute(AllNodeChildren, NodeChildren);
+	}
+	else
+	{
+		NodeChildren = AllNodeChildren;
+	}
+
+	bIsHiddenCategory = DetailTreeNode->GetNodeType() == EDetailNodeType::Category && (NodeChildren.Num() == 0 || (bIsTopLevelProperty && bHideTopLevelCategories));
 	for (TSharedRef<IDetailTreeNode> NodeChild : NodeChildren)
 	{
 		if (NodeChild->GetNodeType() == EDetailNodeType::Advanced)
@@ -104,7 +130,7 @@ void UNiagaraStackPropertyRow::RefreshChildrenInternal(const TArray<UNiagaraStac
 		{
 			bool bChildIsTopLevelProperty = false;
 			ChildRow = NewObject<UNiagaraStackPropertyRow>(this);
-			ChildRow->Initialize(CreateDefaultChildRequiredData(), NodeChild, bChildIsTopLevelProperty, GetOwnerStackItemEditorDataKey(), GetStackEditorDataKey(), OwningNiagaraNode);
+			ChildRow->Initialize(CreateDefaultChildRequiredData(), NodeChild, bChildIsTopLevelProperty, bHideTopLevelCategories, GetOwnerStackItemEditorDataKey(), GetStackEditorDataKey(), OwningNiagaraNode);
 			ChildRow->SetOwnerGuid(OwnerGuid);
 		}
 
diff --git a/Engine/Plugins/FX/Niagara/Source/NiagaraEditor/Private/ViewModels/Stack/NiagaraStackRenderItemGroup.cpp b/Engine/Plugins/FX/Niagara/Source/NiagaraEditor/Private/ViewModels/Stack/NiagaraStackRenderItemGroup.cpp
--- a/Engine/Plugins/FX/Niagara/Source/NiagaraEditor/Private/ViewModels/Stack/NiagaraStackRenderItemGroup.cpp
+++ b/Engine/Plugins/FX/Niagara/Source/NiagaraEditor/Private/ViewModels/Stack/NiagaraStackRenderItemGroup.cpp
@@ -2,8 +2,11 @@
 
 #include "ViewModels/Stack/NiagaraStackRenderItemGroup.h"
 #include "ViewModels/Stack/NiagaraStackRendererItem.h"
+#include "ViewModels/Stack/NiagaraStackRenderersOwner.h"
+#include "ViewModels/NiagaraEmitterHandleViewModel.h"
 #include "ViewModels/NiagaraEmitterViewModel.h"
 #include "NiagaraEmitter.h"
+#include "Stateless/NiagaraStatelessEmitter.h"
 #include "NiagaraEmitterEditorData.h"
 #include "ViewModels/NiagaraSystemViewModel.h"
 #include "NiagaraRendererProperties.h"
@@ -59,9 +62,9 @@ private:
 class FRenderItemGroupAddUtilities : public TNiagaraStackItemGroupAddUtilities<UNiagaraRendererProperties*>
 {
 public:
-	FRenderItemGroupAddUtilities(TSharedRef<FNiagaraEmitterViewModel> InEmitterViewModel, FRenderItemGroupAddUtilities::FOnItemAdded OnItemAdded = FRenderItemGroupAddUtilities::FOnItemAdded())
+	FRenderItemGroupAddUtilities(TSharedPtr<INiagaraStackRenderersOwner> InRenderersOwner, FRenderItemGroupAddUtilities::FOnItemAdded OnItemAdded = FRenderItemGroupAddUtilities::FOnItemAdded())
 		: TNiagaraStackItemGroupAddUtilities(LOCTEXT("RenderGroupAddItemName", "Renderer"), EAddMode::AddFromAction, true, false, OnItemAdded)
-		, EmitterViewModel(InEmitterViewModel)
+		, RenderersOwner(InRenderersOwner)
 	{
 	}
 
@@ -83,8 +86,7 @@ public:
 
 	virtual void ExecuteAddAction(TSharedRef<INiagaraStackItemGroupAddAction> AddAction, int32 TargetIndex) override
 	{
-		TSharedPtr<FNiagaraEmitterViewModel> EmitterViewModelPinned = EmitterViewModel.Pin();
-		if (EmitterViewModelPinned.IsValid() == false)
+		if (RenderersOwner.IsValid() == false || RenderersOwner->IsValid() == false)
 		{
 			return;
 		}
@@ -93,48 +95,44 @@ public:
 
 		FScopedTransaction ScopedTransaction(LOCTEXT("AddNewRendererTransaction", "Add new renderer"));
 
-		FVersionedNiagaraEmitter VersionedEmitter = EmitterViewModelPinned->GetEmitter();
-		UNiagaraRendererProperties* RendererProperties = RenderAddAction->GetRendererCreationInfo().RendererFactory.Execute(VersionedEmitter.Emitter);
+		UNiagaraRendererProperties* RendererProperties = RenderAddAction->GetRendererCreationInfo().RendererFactory.Execute(RenderersOwner->GetOwnerObject());
 		check(RendererProperties != nullptr);
 		
-		VersionedEmitter.Emitter->AddRenderer(RendererProperties, VersionedEmitter.Version);
-
-		bool bVarsAdded = false;
-		TArray<FNiagaraVariable> MissingAttributes = UNiagaraStackRendererItem::GetMissingVariables(RendererProperties, VersionedEmitter.GetEmitterData());
-		for (int32 i = 0; i < MissingAttributes.Num(); i++)
-		{
-			if (UNiagaraStackRendererItem::AddMissingVariable(VersionedEmitter.GetEmitterData(), MissingAttributes[i]))
-			{
-				bVarsAdded = true;
-			}
-		}
-
-		FNiagaraSystemUpdateContext SystemUpdate(VersionedEmitter, true);
-
-		if (bVarsAdded)
-		{
-			FNotificationInfo Info(LOCTEXT("AddedVariables", "One or more variables have been added to the Spawn script to support the added renderer."));
-			Info.ExpireDuration = 5.0f;
-			Info.bFireAndForget = true;
-			Info.Image = FCoreStyle::Get().GetBrush(TEXT("MessageLog.Info"));
-			FSlateNotificationManager::Get().AddNotification(Info);
-		}
-
+		RenderersOwner->AddRenderer(RendererProperties);
 		OnItemAdded.ExecuteIfBound(RendererProperties);
 	}
 
 private:
-	TWeakPtr<FNiagaraEmitterViewModel> EmitterViewModel;
+	TSharedPtr<INiagaraStackRenderersOwner> RenderersOwner;
 };
 
-void UNiagaraStackRenderItemGroup::Initialize(FRequiredEntryData InRequiredEntryData)
+void UNiagaraStackRenderItemGroup::Initialize(FRequiredEntryData InRequiredEntryData, TSharedPtr<INiagaraStackRenderersOwner> InRenderersOwner)
 {
 	FText DisplayName = LOCTEXT("RenderGroupName", "Render");
 	FText ToolTip = LOCTEXT("RendererGroupTooltip", "Describes how we should display/present each particle. Note that this doesn't have to be visual. Multiple renderers are supported. Order in this stack is not necessarily relevant to draw order.");
-	AddUtilities = MakeShared<FRenderItemGroupAddUtilities>(InRequiredEntryData.EmitterViewModel.ToSharedRef(), FRenderItemGroupAddUtilities::FOnItemAdded::CreateUObject(this, &UNiagaraStackRenderItemGroup::OnRendererAdded));
+	RenderersOwner = InRenderersOwner;
+	AddUtilities = MakeShared<FRenderItemGroupAddUtilities>(RenderersOwner, FRenderItemGroupAddUtilities::FOnItemAdded::CreateUObject(this, &UNiagaraStackRenderItemGroup::OnRendererAdded));
 	Super::Initialize(InRequiredEntryData, DisplayName, ToolTip, AddUtilities.Get());
-	EmitterWeak = GetEmitterViewModel()->GetEmitter().ToWeakPtr();
-	EmitterWeak.Emitter->OnRenderersChanged().AddUObject(this, &UNiagaraStackRenderItemGroup::EmitterRenderersChanged);
+}
+
+UNiagaraStackEntry::EIconMode UNiagaraStackRenderItemGroup::GetSupportedIconMode() const
+{
+	return RenderersOwner.IsValid() ? RenderersOwner->GetSupportedIconMode() : EIconMode::None;
+}
+
+const FSlateBrush* UNiagaraStackRenderItemGroup::GetIconBrush() const
+{
+	return RenderersOwner.IsValid() ? RenderersOwner->GetIconBrush() : nullptr;
+}
+
+FText UNiagaraStackRenderItemGroup::GetIconText() const
+{
+	return RenderersOwner.IsValid() ? RenderersOwner->GetIconText() : FText();
+}
+
+bool UNiagaraStackRenderItemGroup::GetCanExpandInOverview() const
+{
+	return RenderersOwner.IsValid() && RenderersOwner->ShouldShowRendererItemsInOverview();
 }
 
 bool UNiagaraStackRenderItemGroup::TestCanPasteWithMessage(const UNiagaraClipboardContent* ClipboardContent, FText& OutMessage) const
@@ -155,21 +153,18 @@ FText UNiagaraStackRenderItemGroup::GetPasteTransactionText(const UNiagaraClipbo
 
 void UNiagaraStackRenderItemGroup::Paste(const UNiagaraClipboardContent* ClipboardContent, FText& OutPasteWarning)
 {
-	if (EmitterWeak.IsValid())
+	if (RenderersOwner.IsValid() && RenderersOwner->IsValid())
 	{
-		UNiagaraEmitter* Emitter = EmitterWeak.Emitter.Get();
 		for (const UNiagaraClipboardRenderer* ClipboardRenderer : ClipboardContent->Renderers)
 		{
 			if (ClipboardRenderer != nullptr && ClipboardRenderer->RendererProperties != nullptr)
 			{
-				UNiagaraRendererProperties* NewRenderer = ClipboardRenderer->RendererProperties->StaticDuplicateWithNewMergeId(Emitter);
-				
+				UNiagaraRendererProperties* NewRenderer = ClipboardRenderer->RendererProperties->StaticDuplicateWithNewMergeId(RenderersOwner->GetOwnerObject());
+				RenderersOwner->AddRenderer(NewRenderer);
 				if(ClipboardRenderer->StackNoteData.IsValid())
 				{
 					GetStackEditorData().AddOrReplaceStackNote(FNiagaraStackGraphUtilities::StackKeys::GenerateStackRendererEditorDataKey(*NewRenderer), ClipboardRenderer->StackNoteData);
 				}
-				
-				EmitterWeak.Emitter->AddRenderer(NewRenderer, EmitterWeak.Version);
 			}
 		}
 	}
@@ -178,22 +173,27 @@ void UNiagaraStackRenderItemGroup::Paste(const UNiagaraClipboardContent* Clipboa
 void UNiagaraStackRenderItemGroup::RefreshChildrenInternal(const TArray<UNiagaraStackEntry*>& CurrentChildren, TArray<UNiagaraStackEntry*>& NewChildren, TArray<FStackIssue>& NewIssues)
 {
 	int32 RendererIndex = 0;
-	for (UNiagaraRendererProperties* RendererProperties : GetEmitterViewModel()->GetEmitter().GetEmitterData()->GetRenderers())
+	if (RenderersOwner.IsValid())
 	{
-		UNiagaraStackRendererItem* RendererItem = FindCurrentChildOfTypeByPredicate<UNiagaraStackRendererItem>(CurrentChildren,
-			[=](UNiagaraStackRendererItem* CurrentRendererItem) { return CurrentRendererItem->GetRendererProperties() == RendererProperties; });
-
-		if (RendererItem == nullptr)
+		TArray<UNiagaraRendererProperties*> Renderers;
+		RenderersOwner->GetRenderers(Renderers);
+		for (UNiagaraRendererProperties* RendererProperties : Renderers)
 		{
-			RendererItem = NewObject<UNiagaraStackRendererItem>(this);
-			RendererItem->Initialize(CreateDefaultChildRequiredData(), RendererProperties);
-			RendererItem->SetOnRequestCanPaste(UNiagaraStackRendererItem::FOnRequestCanPaste::CreateUObject(this, &UNiagaraStackRenderItemGroup::ChildRequestCanPaste));
-			RendererItem->SetOnRequestPaste(UNiagaraStackRendererItem::FOnRequestPaste::CreateUObject(this, &UNiagaraStackRenderItemGroup::ChildRequestPaste));
-		}
+			UNiagaraStackRendererItem* RendererItem = FindCurrentChildOfTypeByPredicate<UNiagaraStackRendererItem>(CurrentChildren,
+				[=](UNiagaraStackRendererItem* CurrentRendererItem) { return CurrentRendererItem->GetRendererProperties() == RendererProperties; });
+
+			if (RendererItem == nullptr)
+			{
+				RendererItem = NewObject<UNiagaraStackRendererItem>(this);
+				RendererItem->Initialize(CreateDefaultChildRequiredData(), RenderersOwner, RendererProperties);
+				RendererItem->SetOnRequestCanPaste(UNiagaraStackRendererItem::FOnRequestCanPaste::CreateUObject(this, &UNiagaraStackRenderItemGroup::ChildRequestCanPaste));
+				RendererItem->SetOnRequestPaste(UNiagaraStackRendererItem::FOnRequestPaste::CreateUObject(this, &UNiagaraStackRenderItemGroup::ChildRequestPaste));
+			}
 
-		NewChildren.Add(RendererItem);
+			NewChildren.Add(RendererItem);
 
-		RendererIndex++;
+			RendererIndex++;
+		}
 	}
 
 	Super::RefreshChildrenInternal(CurrentChildren, NewChildren, NewIssues);
@@ -228,10 +228,7 @@ void UNiagaraStackRenderItemGroup::OnRendererAdded(UNiagaraRendererProperties* R
 
 void UNiagaraStackRenderItemGroup::FinalizeInternal()
 {
-	if (EmitterWeak.IsValid())
-	{
-		EmitterWeak.Emitter->OnRenderersChanged().RemoveAll(this);
-	}
+	RenderersOwner.Reset();
 	Super::FinalizeInternal();
 }
 
diff --git a/Engine/Plugins/FX/Niagara/Source/NiagaraEditor/Private/ViewModels/Stack/NiagaraStackRendererItem.cpp b/Engine/Plugins/FX/Niagara/Source/NiagaraEditor/Private/ViewModels/Stack/NiagaraStackRendererItem.cpp
--- a/Engine/Plugins/FX/Niagara/Source/NiagaraEditor/Private/ViewModels/Stack/NiagaraStackRendererItem.cpp
+++ b/Engine/Plugins/FX/Niagara/Source/NiagaraEditor/Private/ViewModels/Stack/NiagaraStackRendererItem.cpp
@@ -2,6 +2,7 @@
 
 #include "ViewModels/Stack/NiagaraStackRendererItem.h"
 #include "ViewModels/Stack/NiagaraStackObject.h"
+#include "ViewModels/Stack/NiagaraStackRenderersOwner.h"
 #include "NiagaraEmitter.h"
 #include "NiagaraStackEditorData.h"
 #include "NiagaraRendererProperties.h"
@@ -36,10 +37,11 @@ UNiagaraStackRendererItem::UNiagaraStackRendererItem()
 {
 }
 
-void UNiagaraStackRendererItem::Initialize(FRequiredEntryData InRequiredEntryData, UNiagaraRendererProperties* InRendererProperties)
+void UNiagaraStackRendererItem::Initialize(FRequiredEntryData InRequiredEntryData, TSharedPtr<INiagaraStackRenderersOwner> InRenderersOwner, UNiagaraRendererProperties* InRendererProperties)
 {
 	checkf(RendererProperties.IsValid() == false, TEXT("Can not initialize more than once."));
 	Super::Initialize(InRequiredEntryData, FNiagaraStackGraphUtilities::StackKeys::GenerateStackRendererEditorDataKey(*InRendererProperties));
+	RenderersOwner = InRenderersOwner;
 	RendererProperties = InRendererProperties;
 	RendererProperties->OnChanged().AddUObject(this, &UNiagaraStackRendererItem::RendererChanged);
 }
@@ -104,9 +106,11 @@ bool UNiagaraStackRendererItem::CanMoveRendererUp() const
 	{
 		return false;
 	}
-	if (FVersionedNiagaraEmitterData* EmitterData = GetEmitterViewModel()->GetEmitter().GetEmitterData())
+	if (RenderersOwner.IsValid() && RenderersOwner->IsValid())
 	{
-		return EmitterData->GetRenderers().IndexOfByKey(RendererProperties.Get()) > 0;
+		TArray<UNiagaraRendererProperties*> Renderers;
+		RenderersOwner->GetRenderers(Renderers);
+		return Renderers.IndexOfByKey(RendererProperties.Get()) > 0;
 	}
 	return false;
 }
@@ -118,12 +122,13 @@ void UNiagaraStackRendererItem::MoveRendererUp() const
 		return;
 	}
 
-	FVersionedNiagaraEmitter VersionedEmitter = GetEmitterViewModel()->GetEmitter();
-	if (FVersionedNiagaraEmitterData* EmitterData = VersionedEmitter.GetEmitterData())
+	if (RenderersOwner.IsValid() && RenderersOwner->IsValid())
 	{
 		FScopedTransaction ScopedTransaction(LOCTEXT("MoveRendererUpTransaction", "Move renderer up"));
-		int32 CurrentIndex = EmitterData->GetRenderers().IndexOfByKey(RendererProperties.Get());
-		VersionedEmitter.Emitter->MoveRenderer(RendererProperties.Get(), CurrentIndex - 1, VersionedEmitter.Version);
+		TArray<UNiagaraRendererProperties*> Renderers;
+		RenderersOwner->GetRenderers(Renderers);
+		int32 CurrentIndex = Renderers.IndexOfByKey(RendererProperties.Get());
+		RenderersOwner->MoveRenderer(RendererProperties.Get(), CurrentIndex - 1);
 	}
 }
 
@@ -133,9 +138,11 @@ bool UNiagaraStackRendererItem::CanMoveRendererDown() const
 	{
 		return false;
 	}
-	if (FVersionedNiagaraEmitterData* EmitterData = GetEmitterViewModel()->GetEmitter().GetEmitterData())
+	if (RenderersOwner.IsValid() && RenderersOwner->IsValid())
 	{
-		return EmitterData->GetRenderers().IndexOfByKey(RendererProperties.Get()) < EmitterData->GetRenderers().Num() - 1;
+		TArray<UNiagaraRendererProperties*> Renderers;
+		RenderersOwner->GetRenderers(Renderers);
+		return Renderers.IndexOfByKey(RendererProperties.Get()) < Renderers.Num() - 1;
 	}
 	return false;
 }
@@ -147,12 +154,13 @@ void UNiagaraStackRendererItem::MoveRendererDown() const
 		return;
 	}
 
-	FVersionedNiagaraEmitter VersionedEmitter = GetEmitterViewModel()->GetEmitter();
-	if (FVersionedNiagaraEmitterData* EmitterData = VersionedEmitter.GetEmitterData())
+	if (RenderersOwner.IsValid() && RenderersOwner->IsValid())
 	{
 		FScopedTransaction ScopedTransaction(LOCTEXT("MoveRendererDownTransaction", "Move renderer down"));
-		int32 CurrentIndex = EmitterData->GetRenderers().IndexOfByKey(RendererProperties.Get());
-		VersionedEmitter.Emitter->MoveRenderer(RendererProperties.Get(), CurrentIndex + 1, VersionedEmitter.Version);
+		TArray<UNiagaraRendererProperties*> Renderers;
+		RenderersOwner->GetRenderers(Renderers);
+		int32 CurrentIndex = Renderers.IndexOfByKey(RendererProperties.Get());
+		RenderersOwner->MoveRenderer(RendererProperties.Get(), CurrentIndex + 1);
 	}
 }
 
@@ -245,6 +253,7 @@ bool UNiagaraStackRendererItem::IsExcludedFromScalability() const
 
 bool UNiagaraStackRendererItem::IsOwningEmitterExcludedFromScalability() const
 {
+	// TODO - Stateless - RenderersOwner refactor
 	return GetEmitterViewModel().IsValid() ? !GetEmitterViewModel()->GetEmitter().GetEmitterData()->IsAllowedByScalability() : false;
 }
 
@@ -366,20 +375,18 @@ FText UNiagaraStackRendererItem::GetDeleteTransactionText() const
 
 void UNiagaraStackRendererItem::Delete()
 {
-	FVersionedNiagaraEmitter VersionedEmitter = GetEmitterViewModel()->GetEmitter();
-
-	UNiagaraRendererProperties* Renderer = RendererProperties.Get();
-	VersionedEmitter.Emitter->Modify();
-	VersionedEmitter.Emitter->RemoveRenderer(Renderer, VersionedEmitter.Version);
-	if (UNiagaraEmitterEditorData* Data = Cast<UNiagaraEmitterEditorData>(VersionedEmitter.GetEmitterData()->GetEditorData()))
+	if (RenderersOwner.IsValid() && RenderersOwner->IsValid())
 	{
-		Data->GetStackEditorData().Modify();
-		Data->GetStackEditorData().SetStackEntryDisplayName(GetStackEditorDataKey(), FText());
-	}
+		GetStackEditorData().Modify();
+		GetStackEditorData().SetStackEntryDisplayName(GetStackEditorDataKey(), FText());
 
-	TArray<UObject*> ChangedObjects;
-	ChangedObjects.Add(Renderer);
-	OnDataObjectModified().Broadcast(ChangedObjects, ENiagaraDataObjectChange::Removed);
+		UNiagaraRendererProperties* Renderer = RendererProperties.Get();
+		RenderersOwner->RemoveRenderer(Renderer);
+
+		TArray<UObject*> ChangedObjects;
+		ChangedObjects.Add(Renderer);
+		OnDataObjectModified().Broadcast(ChangedObjects, ENiagaraDataObjectChange::Removed);
+	}
 }
 
 bool UNiagaraStackRendererItem::GetIsInherited() const
@@ -405,9 +412,7 @@ bool UNiagaraStackRendererItem::HasBaseRenderer() const
 	{
 		if (bHasBaseRendererCache.IsSet() == false)
 		{
-			TSharedRef<FNiagaraScriptMergeManager> MergeManager = FNiagaraScriptMergeManager::Get();
-			FVersionedNiagaraEmitter BaseEmitter = GetEmitterViewModel()->GetEmitter().GetEmitterData()->GetParent();
-			bHasBaseRendererCache = BaseEmitter.Emitter != nullptr && MergeManager->HasBaseRenderer(BaseEmitter, RendererProperties->GetMergeId());
+			bHasBaseRendererCache = RenderersOwner.IsValid() && RenderersOwner->IsValid() && RenderersOwner->HasBaseRenderer(RendererProperties.Get());
 		}
 		return bHasBaseRendererCache.GetValue();
 	}
@@ -418,11 +423,9 @@ bool UNiagaraStackRendererItem::TestCanResetToBaseWithMessage(FText& OutCanReset
 {
 	if (bCanResetToBaseCache.IsSet() == false)
 	{
-		if (HasBaseRenderer())
+		if (HasBaseRenderer() && RenderersOwner.IsValid() && RenderersOwner->IsValid())
 		{
-			TSharedRef<FNiagaraScriptMergeManager> MergeManager = FNiagaraScriptMergeManager::Get();
-			FVersionedNiagaraEmitter BaseEmitter = GetEmitterViewModel()->GetEmitter().GetEmitterData()->GetParent();
-			bCanResetToBaseCache = BaseEmitter.Emitter != nullptr && MergeManager->IsRendererDifferentFromBase(GetEmitterViewModel()->GetEmitter(), BaseEmitter, RendererProperties->GetMergeId());
+			bCanResetToBaseCache = RenderersOwner->IsRendererDifferentFromBase(RendererProperties.Get());
 		}
 		else
 		{
@@ -446,13 +449,19 @@ void UNiagaraStackRendererItem::ResetToBase()
 	FText Unused;
 	if (TestCanResetToBaseWithMessage(Unused))
 	{
-		TSharedRef<FNiagaraScriptMergeManager> MergeManager = FNiagaraScriptMergeManager::Get();
-		FVersionedNiagaraEmitter BaseEmitter = GetEmitterViewModel()->GetEmitter().GetEmitterData()->GetParent();
-		MergeManager->ResetRendererToBase(GetEmitterViewModel()->GetEmitter(), BaseEmitter, RendererProperties->GetMergeId());
-		ModifiedGroupItemsDelegate.Broadcast();
+		if (RenderersOwner.IsValid() && RenderersOwner->IsValid())
+		{
+			RenderersOwner->ResetRendererToBase(RendererProperties.Get());
+			ModifiedGroupItemsDelegate.Broadcast();
+		}
 	}
 }
 
+bool UNiagaraStackRendererItem::GetShouldShowInOverview() const
+{
+	return RenderersOwner.IsValid() && RenderersOwner->ShouldShowRendererItemsInOverview();
+}
+
 bool UNiagaraStackRendererItem::GetIsEnabled() const
 {
 	if (IsFinalized() == false && RendererProperties.IsValid())
@@ -499,7 +508,8 @@ void UNiagaraStackRendererItem::RefreshChildrenInternal(const TArray<UNiagaraSta
 	{
 		RendererObject = NewObject<UNiagaraStackObject>(this);
 		bool bIsTopLevelObject = true;
-		RendererObject->Initialize(CreateDefaultChildRequiredData(), RendererProperties.Get(), bIsTopLevelObject, GetStackEditorDataKey());
+		bool bHideTopLevelCategories = false;
+		RendererObject->Initialize(CreateDefaultChildRequiredData(), RendererProperties.Get(), bIsTopLevelObject, bHideTopLevelCategories, GetStackEditorDataKey());
 		RendererObject->SetObjectGuid(GetRendererProperties()->GetMergeId());
 	}
 
diff --git a/Engine/Plugins/FX/Niagara/Source/NiagaraEditor/Private/ViewModels/Stack/NiagaraStackRoot.cpp b/Engine/Plugins/FX/Niagara/Source/NiagaraEditor/Private/ViewModels/Stack/NiagaraStackRoot.cpp
--- a/Engine/Plugins/FX/Niagara/Source/NiagaraEditor/Private/ViewModels/Stack/NiagaraStackRoot.cpp
+++ b/Engine/Plugins/FX/Niagara/Source/NiagaraEditor/Private/ViewModels/Stack/NiagaraStackRoot.cpp
@@ -4,9 +4,14 @@
 #include "ViewModels/Stack/NiagaraStackScriptItemGroup.h"
 #include "ViewModels/Stack/NiagaraStackEmitterPropertiesGroup.h"
 #include "ViewModels/Stack/NiagaraStackEmitterSettingsGroup.h"
+#include "ViewModels/Stack/NiagaraStackStatelessEmitterGroup.h"
+#include "ViewModels/Stack/NiagaraStackStatelessEmitterSimulateGroup.h"
+#include "ViewModels/Stack/NiagaraStackStatelessEmitterSpawnGroup.h"
+#include "ViewModels/Stack/NiagaraStackRenderersOwner.h"
 #include "ViewModels/Stack/NiagaraStackRenderItemGroup.h"
 #include "ViewModels/Stack/NiagaraStackEventScriptItemGroup.h"
 #include "ViewModels/Stack/NiagaraStackSimulationStageGroup.h"
+#include "ViewModels/Stack/NiagaraStackSystemSettingsGroup.h"
 #include "ViewModels/NiagaraSystemViewModel.h"
 #include "ViewModels/NiagaraEmitterViewModel.h"
 #include "ViewModels/NiagaraScriptViewModel.h"
@@ -14,24 +19,13 @@
 #include "NiagaraSystem.h"
 #include "NiagaraSystemEditorData.h"
 #include "NiagaraEmitterEditorData.h"
-#include "NiagaraObjectSelection.h"
-#include "ViewModels/NiagaraOverviewGraphViewModel.h"
-#include "ViewModels/NiagaraSystemSelectionViewModel.h"
-#include "ViewModels/Stack/NiagaraStackSystemSettingsGroup.h"
+#include "NiagaraSimulationStageBase.h"
 
 #include UE_INLINE_GENERATED_CPP_BY_NAME(NiagaraStackRoot)
 
 #define LOCTEXT_NAMESPACE "NiagaraStackViewModel"
 
 UNiagaraStackRoot::UNiagaraStackRoot()
-	: SystemPropertiesGroup(nullptr)
-	, EmitterPropertiesGroup(nullptr)
-	, EmitterSummaryGroup(nullptr)
-	, EmitterSpawnGroup(nullptr)
-	, EmitterUpdateGroup(nullptr)
-	, ParticleSpawnGroup(nullptr)
-	, ParticleUpdateGroup(nullptr)
-	, RenderGroup(nullptr)
 {
 }
 
@@ -40,27 +34,9 @@ void UNiagaraStackRoot::Initialize(FRequiredEntryData InRequiredEntryData, bool
 	Super::Initialize(InRequiredEntryData, TEXT("Root"));
 	bIncludeSystemInformation = bInIncludeSystemInformation;
 	bIncludeEmitterInformation = bInIncludeEmitterInformation;
-	SystemPropertiesGroup = nullptr;
-	EmitterPropertiesGroup = nullptr;
-	EmitterSummaryGroup = nullptr;
-	EmitterSpawnGroup = nullptr;
-	EmitterUpdateGroup = nullptr;
-	ParticleSpawnGroup = nullptr;
-	ParticleUpdateGroup = nullptr;
-	RenderGroup = nullptr;
-
 	if (bInIncludeEmitterInformation && GetEmitterViewModel())
 	{
 		GetEmitterViewModel()->GetEditorData().OnSummaryViewStateChanged().AddUObject(this, &UNiagaraStackRoot::OnSummaryViewStateChanged);
-
-		if(GetEmitterViewModel()->GetEditorData().ShouldShowSummaryView())
-		{
-			AddSummaryViewFilter();
-		}
-		else
-		{
-			AddDefaultViewFilter();
-		}
 	}
 }
 
@@ -83,272 +59,387 @@ bool UNiagaraStackRoot::GetShouldShowInStack() const
 	return false;
 }
 
+UNiagaraStackRenderItemGroup* UNiagaraStackRoot::GetRenderGroup() const
+{
+	TArray<UNiagaraStackRenderItemGroup*> RenderItemGroups;
+	GetUnfilteredChildrenOfType(RenderItemGroups);
+	return RenderItemGroups.Num() == 1 ? RenderItemGroups[0] : nullptr;
+}
+
+UNiagaraStackCommentCollection* UNiagaraStackRoot::GetCommentCollection() const
+{
+	TArray<UNiagaraStackCommentCollection*> CommentCollections;
+	GetUnfilteredChildrenOfType(CommentCollections);
+	return CommentCollections.Num() == 1 ? CommentCollections[0] : nullptr;
+}
+
 void UNiagaraStackRoot::RefreshChildrenInternal(const TArray<UNiagaraStackEntry*>& CurrentChildren, TArray<UNiagaraStackEntry*>& NewChildren, TArray<FStackIssue>& NewIssues)
 {
-	// Create static entries as needed.
-	if (bIncludeSystemInformation && SystemPropertiesGroup == nullptr)
+	if (bIncludeSystemInformation)
 	{
-		SystemPropertiesGroup = NewObject<UNiagaraStackSystemPropertiesGroup>(this);
-		FRequiredEntryData RequiredEntryData(GetSystemViewModel(), GetEmitterViewModel(),
-			FExecutionCategoryNames::System, FExecutionSubcategoryNames::Settings,
-			GetSystemViewModel()->GetEditorData().GetStackEditorData());
-		SystemPropertiesGroup->Initialize(RequiredEntryData);
+		RefreshSystemChildren(CurrentChildren, NewChildren);
 	}
 
-	if (bIncludeSystemInformation && SystemSpawnGroup == nullptr)
+	if (bIncludeEmitterInformation && GetEmitterViewModel().IsValid())
 	{
-		SystemSpawnGroup = NewObject<UNiagaraStackScriptItemGroup>(this);
-		FRequiredEntryData RequiredEntryData(GetSystemViewModel(), GetEmitterViewModel(),
-			FExecutionCategoryNames::System, FExecutionSubcategoryNames::Spawn,
-			GetSystemViewModel()->GetEditorData().GetStackEditorData());
-		FText DisplayName = LOCTEXT("SystemSpawnGroupName", "System Spawn");
-		FText ToolTip = LOCTEXT("SystemSpawnGroupToolTip", "Occurs once at System creation on the CPU. Modules in this stage should initialize defaults and/or do initial setup.\r\nModules are executed in order from top to bottom of the stack.");
-		SystemSpawnGroup->Initialize(RequiredEntryData, DisplayName, ToolTip, GetSystemViewModel()->GetSystemScriptViewModel().ToSharedRef(), ENiagaraScriptUsage::SystemSpawnScript);
+		TSharedPtr<FNiagaraEmitterHandleViewModel> EmitterHandleViewModel = GetSystemViewModel()->GetEmitterHandleViewModelForEmitter(GetEmitterViewModel()->GetEmitter());
+		bool bShouldShowEmitterStatelessView = EmitterHandleViewModel.IsValid() && EmitterHandleViewModel->GetEmitterHandle()->GetEmitterMode() == ENiagaraEmitterMode::Stateless;
+		
+		if (bShouldShowEmitterStatelessView)
+		{
+			RefreshEmitterStatelessChildren(CurrentChildren, NewChildren, EmitterHandleViewModel);
+		}
+		else
+		{
+			bool bShouldShowSummaryView = GetEmitterViewModel() ? GetEmitterViewModel()->GetEditorData().ShouldShowSummaryView() : false;
+			if (bShouldShowSummaryView)
+			{
+				RefreshEmitterSummaryChildren(CurrentChildren, NewChildren);
+			}
+			else
+			{
+				RefreshEmitterFullChildren(CurrentChildren, NewChildren);
+			}
+		}
 	}
+}
+
+void UNiagaraStackRoot::RefreshSystemChildren(const TArray<UNiagaraStackEntry*>& CurrentChildren, TArray<UNiagaraStackEntry*>& NewChildren)
+{
+	UNiagaraStackEditorData& SystemStackEditorData = GetSystemViewModel()->GetEditorData().GetStackEditorData();
+	TSharedRef<FNiagaraScriptViewModel> SystemScriptViewModel = GetSystemViewModel()->GetSystemScriptViewModel().ToSharedRef();
+	FName ExecutionCategory = FExecutionCategoryNames::System;
+
+	NewChildren.Add(GetOrCreateSystemPropertiesGroup(CurrentChildren));
+	NewChildren.Add(GetOrCreateSystemUserParametersGroup(CurrentChildren));
 
-	if (bIncludeSystemInformation && SystemUpdateGroup == nullptr)
+	UNiagaraStackEntry* SystemSpawnEntry = GetCurrentScriptGroup(CurrentChildren, ENiagaraScriptUsage::SystemSpawnScript, FGuid());
+	if (SystemSpawnEntry == nullptr)
 	{
-		SystemUpdateGroup = NewObject<UNiagaraStackScriptItemGroup>(this);
-		FRequiredEntryData RequiredEntryData(GetSystemViewModel(), GetEmitterViewModel(),
-			FExecutionCategoryNames::System, FExecutionSubcategoryNames::Update,
-			GetSystemViewModel()->GetEditorData().GetStackEditorData());
-		FText DisplayName = LOCTEXT("SystemUpdateGroupName", "System Update");
-		FText ToolTip = LOCTEXT("SystemUpdateGroupToolTip", "Occurs every Emitter tick on the CPU.Modules in this stage should compute values for parameters for emitter or particle update or spawning this frame.\r\nModules are executed in order from top to bottom of the stack.");
-		SystemUpdateGroup->Initialize(RequiredEntryData, DisplayName, ToolTip, GetSystemViewModel()->GetSystemScriptViewModel().ToSharedRef(), ENiagaraScriptUsage::SystemUpdateScript);
+		ENiagaraScriptUsage ScriptUsage = ENiagaraScriptUsage::SystemSpawnScript;
+		FGuid ScriptUsageId = FGuid();
+		FName ExecutionSubCategory = FExecutionSubcategoryNames::Spawn;
+		FText DisplayName = LOCTEXT("SystemSpawnGroupName", "System Spawn");
+		FText ToolTip =	LOCTEXT("SystemSpawnGroupToolTip", "Occurs once at System creation on the CPU. Modules in this stage should initialize defaults and/or do initial setup.\r\nModules are executed in order from top to bottom of the stack.");
+		SystemSpawnEntry = CreateScriptGroup(SystemScriptViewModel, ScriptUsage, ScriptUsageId, SystemStackEditorData, ExecutionCategory, ExecutionSubCategory, DisplayName, ToolTip);
 	}
+	NewChildren.Add(SystemSpawnEntry);
 
-	// we only allow comments in systems right now, not emitters
-	if(bIncludeSystemInformation && CommentCollection == nullptr)
+	UNiagaraStackEntry* SystemUpdateEntry = GetCurrentScriptGroup(CurrentChildren, ENiagaraScriptUsage::SystemUpdateScript, FGuid());
+	if (SystemUpdateEntry == nullptr)
 	{
-		CommentCollection = NewObject<UNiagaraStackCommentCollection>(this);
-		FRequiredEntryData RequiredEntryData(GetSystemViewModel(), GetEmitterViewModel(),
-			NAME_None, NAME_None,
-			GetSystemViewModel()->GetEditorData().GetStackEditorData());
-		CommentCollection->Initialize(RequiredEntryData);
+		ENiagaraScriptUsage ScriptUsage = ENiagaraScriptUsage::SystemUpdateScript;
+		FGuid ScriptUsageId = FGuid();
+		FName ExecutionSubCategory = FExecutionSubcategoryNames::Update;
+		FText DisplayName = LOCTEXT("SystemUpdateGroupName", "System Update");
+		FText ToolTip = LOCTEXT("SystemUpdateGroupToolTip", "Occurs every Emitter tick on the CPU.Modules in this stage should compute values for parameters for emitter or particle update or spawning this frame.\r\nModules are executed in order from top to bottom of the stack.");
+		SystemUpdateEntry = CreateScriptGroup(SystemScriptViewModel, ScriptUsage, ScriptUsageId, SystemStackEditorData, ExecutionCategory, ExecutionSubCategory, DisplayName, ToolTip);
 	}
-	
-	if (bIncludeEmitterInformation)
-	{
-		if(EmitterPropertiesGroup == nullptr)
-		{
-			EmitterPropertiesGroup = NewObject<UNiagaraStackEmitterPropertiesGroup>(this);
-			FRequiredEntryData RequiredEntryData(GetSystemViewModel(), GetEmitterViewModel(),
-				FExecutionCategoryNames::Emitter, FExecutionSubcategoryNames::Settings,
-				GetEmitterViewModel()->GetEditorData().GetStackEditorData());
-			EmitterPropertiesGroup->Initialize(RequiredEntryData);
-		}
+	NewChildren.Add(SystemUpdateEntry);
 
-		if(EmitterSummaryGroup == nullptr)
-		{
-			EmitterSummaryGroup = NewObject<UNiagaraStackEmitterSummaryGroup>(this);
-			FRequiredEntryData RequiredEntryData(GetSystemViewModel(), GetEmitterViewModel(),
-				FExecutionCategoryNames::Emitter, FExecutionSubcategoryNames::Settings,
-				GetEmitterViewModel()->GetEditorData().GetStackEditorData());
-			FText DisplayName = LOCTEXT("EmitterSummaryGroupName", "Emitter Summary");
-			FText Tooltip = LOCTEXT("EmitterSummaryTooltip", "Summary of parameters for this Emitter.");
-			EmitterSummaryGroup->Initialize(RequiredEntryData, DisplayName, Tooltip, nullptr);
-		}
+	NewChildren.Add(GetOrCreateSystemCommentCollection(CurrentChildren));
+}
 
-		if(RenderGroup == nullptr)
-		{
-			RenderGroup = NewObject<UNiagaraStackRenderItemGroup>(this);
-			FRequiredEntryData RequiredEntryData(GetSystemViewModel(), GetEmitterViewModel(),
-				FExecutionCategoryNames::Render, FExecutionSubcategoryNames::Render,
-				GetEmitterViewModel()->GetEditorData().GetStackEditorData());
-			RenderGroup->Initialize(RequiredEntryData);
-		}
-		
-		if(EmitterSpawnGroup == nullptr)
-		{
-			EmitterSpawnGroup = NewObject<UNiagaraStackScriptItemGroup>(this);
-			FRequiredEntryData RequiredEntryData(GetSystemViewModel(), GetEmitterViewModel(),
-				FExecutionCategoryNames::Emitter, FExecutionSubcategoryNames::Spawn,
-				GetEmitterViewModel()->GetEditorData().GetStackEditorData());
-			FText DisplayName = LOCTEXT("EmitterSpawnGroupName", "Emitter Spawn");
-			FText ToolTip = LOCTEXT("EmitterSpawnGroupTooltip", "Occurs once at Emitter creation on the CPU. Modules in this stage should initialize defaults and/or do initial setup.\r\nModules are executed in order from top to bottom of the stack.");
-			EmitterSpawnGroup->Initialize(RequiredEntryData, DisplayName, ToolTip, GetEmitterViewModel()->GetSharedScriptViewModel(), ENiagaraScriptUsage::EmitterSpawnScript);
-		}
+void UNiagaraStackRoot::RefreshEmitterFullChildren(const TArray<UNiagaraStackEntry*>& CurrentChildren, TArray<UNiagaraStackEntry*>& NewChildren)
+{
+	UNiagaraStackEditorData& EmitterStackEditorData = GetEmitterViewModel()->GetEditorData().GetStackEditorData();
+	TSharedRef<FNiagaraScriptViewModel> EmitterScriptViewModel = GetEmitterViewModel()->GetSharedScriptViewModel();
 
-		if(EmitterUpdateGroup == nullptr)
-		{
-			EmitterUpdateGroup = NewObject<UNiagaraStackScriptItemGroup>(this);
-			FRequiredEntryData RequiredEntryData(GetSystemViewModel(), GetEmitterViewModel(),
-				FExecutionCategoryNames::Emitter, FExecutionSubcategoryNames::Update,
-				GetEmitterViewModel()->GetEditorData().GetStackEditorData());
-			FText DisplayName = LOCTEXT("EmitterUpdateGroupName", "Emitter Update");
-			FText ToolTip = LOCTEXT("EmitterUpdateGroupTooltip", "Occurs every Emitter tick on the CPU. Modules in this stage should compute values for parameters for Particle Update or Spawning this frame.\r\nModules are executed in order from top to bottom of the stack.");
-			EmitterUpdateGroup->Initialize(RequiredEntryData, DisplayName, ToolTip, GetEmitterViewModel()->GetSharedScriptViewModel(), ENiagaraScriptUsage::EmitterUpdateScript);
-		}
+	NewChildren.Add(GetOrCreateEmitterPropertiesGroup(CurrentChildren));
 
-		if(ParticleSpawnGroup == nullptr)
-		{
-			ParticleSpawnGroup = NewObject<UNiagaraStackScriptItemGroup>(this);
-			FRequiredEntryData RequiredEntryData(GetSystemViewModel(), GetEmitterViewModel(),
-				FExecutionCategoryNames::Particle, FExecutionSubcategoryNames::Spawn,
-				GetEmitterViewModel()->GetEditorData().GetStackEditorData());
-			FText DisplayName = LOCTEXT("ParticleSpawnGroupName", "Particle Spawn");
-			FText ToolTip = LOCTEXT("ParticleSpawnGroupTooltip", "Called once per created particle. Modules in this stage should set up initial values for each particle.\r\nIf \"Use Interpolated Spawning\" is set, we will also run the Particle Update stage after the Particle Spawn stage.\r\nModules are executed in order from top to bottom of the stack.");
-			ParticleSpawnGroup->Initialize(RequiredEntryData, DisplayName, ToolTip, GetEmitterViewModel()->GetSharedScriptViewModel(), ENiagaraScriptUsage::ParticleSpawnScript);
-		}
+	UNiagaraStackEntry* EmitterSpawnEntry = GetCurrentScriptGroup(CurrentChildren, ENiagaraScriptUsage::EmitterSpawnScript, FGuid());
+	if (EmitterSpawnEntry == nullptr)
+	{
+		ENiagaraScriptUsage ScriptUsage = ENiagaraScriptUsage::EmitterSpawnScript;
+		FGuid ScriptUsageId = FGuid();
+		FName ExecutionCategory = FExecutionCategoryNames::Emitter;
+		FName ExecutionSubCategory = FExecutionSubcategoryNames::Spawn;
+		FText DisplayName = LOCTEXT("EmitterSpawnGroupName", "Emitter Spawn");
+		FText ToolTip = LOCTEXT("EmitterSpawnGroupTooltip", "Occurs once at Emitter creation on the CPU. Modules in this stage should initialize defaults and/or do initial setup.\r\nModules are executed in order from top to bottom of the stack.");
+		EmitterSpawnEntry = CreateScriptGroup(EmitterScriptViewModel, ScriptUsage, ScriptUsageId, EmitterStackEditorData, ExecutionCategory, ExecutionSubCategory, DisplayName, ToolTip);
+	}
+	NewChildren.Add(EmitterSpawnEntry);
 
-		if(ParticleUpdateGroup == nullptr)
-		{
-			ParticleUpdateGroup = NewObject<UNiagaraStackScriptItemGroup>(this);
-			FRequiredEntryData RequiredEntryData(GetSystemViewModel(), GetEmitterViewModel(),
-				FExecutionCategoryNames::Particle, FExecutionSubcategoryNames::Update,
-				GetEmitterViewModel()->GetEditorData().GetStackEditorData());
-			FText DisplayName = LOCTEXT("ParticleUpdateGroupName", "Particle Update");
-			FText ToolTip = LOCTEXT("ParticleUpdateGroupTooltip", "Called every frame per particle. Modules in this stage should update new values for this frame.\r\nModules are executed in order from top to bottom of the stack.");
-			ParticleUpdateGroup->Initialize(RequiredEntryData, DisplayName, ToolTip, GetEmitterViewModel()->GetSharedScriptViewModel(), ENiagaraScriptUsage::ParticleUpdateScript);
-		}
+	UNiagaraStackEntry* EmitterUpdateEntry = GetCurrentScriptGroup(CurrentChildren, ENiagaraScriptUsage::EmitterUpdateScript, FGuid());
+	if (EmitterUpdateEntry == nullptr)
+	{
+		ENiagaraScriptUsage ScriptUsage = ENiagaraScriptUsage::EmitterUpdateScript;
+		FGuid ScriptUsageId = FGuid();
+		FName ExecutionCategory = FExecutionCategoryNames::Emitter;
+		FName ExecutionSubCategory = FExecutionSubcategoryNames::Update;
+		FText DisplayName = LOCTEXT("EmitterUpdateGroupName", "Emitter Update");
+		FText ToolTip = LOCTEXT("EmitterUpdateGroupTooltip", "Occurs every Emitter tick on the CPU. Modules in this stage should compute values for parameters for Particle Update or Spawning this frame.\r\nModules are executed in order from top to bottom of the stack.");
+		EmitterUpdateEntry = CreateScriptGroup(EmitterScriptViewModel, ScriptUsage, ScriptUsageId, EmitterStackEditorData, ExecutionCategory, ExecutionSubCategory, DisplayName, ToolTip);
 	}
+	NewChildren.Add(EmitterUpdateEntry);
 
-	// Populate new children
-	if (bIncludeSystemInformation)
+	UNiagaraStackEntry* ParticleSpawnEntry = GetCurrentScriptGroup(CurrentChildren, ENiagaraScriptUsage::ParticleSpawnScript, FGuid());
+	if (ParticleSpawnEntry == nullptr)
 	{
-		NewChildren.Add(SystemPropertiesGroup);
-		NewChildren.Add(SystemSpawnGroup);
-		NewChildren.Add(SystemUpdateGroup);
-		NewChildren.Add(CommentCollection);
-	}	
+		ENiagaraScriptUsage ScriptUsage = ENiagaraScriptUsage::ParticleSpawnScript;
+		FGuid ScriptUsageId = FGuid();
+		FName ExecutionCategory = FExecutionCategoryNames::Particle;
+		FName ExecutionSubCategory = FExecutionSubcategoryNames::Spawn;
+		FText DisplayName = LOCTEXT("ParticleSpawnGroupName", "Particle Spawn");
+		FText ToolTip = LOCTEXT("ParticleSpawnGroupTooltip", "Called once per created particle. Modules in this stage should set up initial values for each particle.\r\nIf \"Use Interpolated Spawning\" is set, we will also run the Particle Update stage after the Particle Spawn stage.\r\nModules are executed in order from top to bottom of the stack.");
+		ParticleSpawnEntry = CreateScriptGroup(EmitterScriptViewModel, ScriptUsage, ScriptUsageId, EmitterStackEditorData, ExecutionCategory, ExecutionSubCategory, DisplayName, ToolTip);
+	}
+	NewChildren.Add(ParticleSpawnEntry);
 
-	if (bIncludeEmitterInformation)
+	UNiagaraStackEntry* ParticleUpdateEntry = GetCurrentScriptGroup(CurrentChildren, ENiagaraScriptUsage::ParticleUpdateScript, FGuid());
+	if (ParticleUpdateEntry == nullptr)
 	{
-		NewChildren.Add(EmitterSummaryGroup);
-		
-		NewChildren.Add(EmitterPropertiesGroup);
-		
-		NewChildren.Add(EmitterSpawnGroup);
-		NewChildren.Add(EmitterUpdateGroup);
+		ENiagaraScriptUsage ScriptUsage = ENiagaraScriptUsage::ParticleUpdateScript;
+		FGuid ScriptUsageId = FGuid();
+		FName ExecutionCategory = FExecutionCategoryNames::Particle;
+		FName ExecutionSubCategory = FExecutionSubcategoryNames::Update;
+		FText DisplayName = LOCTEXT("ParticleUpdateGroupName", "Particle Update");
+		FText ToolTip = LOCTEXT("ParticleUpdateGroupTooltip", "Called every frame per particle. Modules in this stage should update new values for this frame.\r\nModules are executed in order from top to bottom of the stack.");
+		ParticleUpdateEntry = CreateScriptGroup(EmitterScriptViewModel, ScriptUsage, ScriptUsageId, EmitterStackEditorData, ExecutionCategory, ExecutionSubCategory, DisplayName, ToolTip);
+	}
+	NewChildren.Add(ParticleUpdateEntry);
 
-		NewChildren.Add(ParticleSpawnGroup);
-		NewChildren.Add(ParticleUpdateGroup);
-		
-		for (const FNiagaraEventScriptProperties& EventScriptProperties : GetEmitterViewModel()->GetEmitter().GetEmitterData()->GetEventHandlers())
+	FVersionedNiagaraEmitterData* EmitterData = GetEmitterViewModel()->GetEmitter().GetEmitterData();
+	if (EmitterData != nullptr)
+	{
+		for (const FNiagaraEventScriptProperties& EventScriptProperties : EmitterData->GetEventHandlers())
 		{
-			UNiagaraStackEventScriptItemGroup* EventHandlerGroup = FindCurrentChildOfTypeByPredicate<UNiagaraStackEventScriptItemGroup>(CurrentChildren,
-				[&](UNiagaraStackEventScriptItemGroup* CurrentEventHandlerGroup) { return CurrentEventHandlerGroup->GetScriptUsageId() == EventScriptProperties.Script->GetUsageId(); });
-
-			if (EventHandlerGroup == nullptr)
+			UNiagaraStackEntry* EventEntry = GetCurrentScriptGroup(CurrentChildren, ENiagaraScriptUsage::ParticleEventScript, EventScriptProperties.Script->GetUsageId());
+			if (EventEntry == nullptr)
 			{
-				EventHandlerGroup = NewObject<UNiagaraStackEventScriptItemGroup>(this);
-				FRequiredEntryData RequiredEntryData(GetSystemViewModel(), GetEmitterViewModel(),
-					FExecutionCategoryNames::Particle, FExecutionSubcategoryNames::Event,
-					GetEmitterViewModel()->GetEditorData().GetStackEditorData());
-				EventHandlerGroup->Initialize(RequiredEntryData, GetEmitterViewModel()->GetSharedScriptViewModel(), ENiagaraScriptUsage::ParticleEventScript, EventScriptProperties.Script->GetUsageId(), EventScriptProperties.SourceEmitterID);
-				EventHandlerGroup->SetOnModifiedEventHandlers(UNiagaraStackEventScriptItemGroup::FOnModifiedEventHandlers::CreateUObject(this, &UNiagaraStackRoot::EmitterArraysChanged));
+				EventEntry = CreateEventScriptGroup(EventScriptProperties.SourceEmitterID, EmitterScriptViewModel, EventScriptProperties.Script->GetUsageId(), EmitterStackEditorData);
 			}
-
-			NewChildren.Add(EventHandlerGroup);
+			NewChildren.Add(EventEntry);
 		}
 
-		for (UNiagaraSimulationStageBase* SimulationStage : GetEmitterViewModel()->GetEmitter().GetEmitterData()->GetSimulationStages())
+		for (UNiagaraSimulationStageBase* SimulationStage : EmitterData->GetSimulationStages())
 		{
-			UNiagaraStackSimulationStageGroup* SimulationStageGroup = FindCurrentChildOfTypeByPredicate<UNiagaraStackSimulationStageGroup>(CurrentChildren,
-				[SimulationStage](UNiagaraStackSimulationStageGroup* CurrentSimulationStageGroup) { return CurrentSimulationStageGroup->GetSimulationStage() == SimulationStage; });
-
-			if (SimulationStageGroup == nullptr)
+			UNiagaraStackEntry* SimulationStageEntry = GetCurrentScriptGroup(CurrentChildren, ENiagaraScriptUsage::ParticleSimulationStageScript, SimulationStage->Script->GetUsageId());
+			if (SimulationStageEntry == nullptr)
 			{
-				SimulationStageGroup = NewObject<UNiagaraStackSimulationStageGroup>(this);
-				FRequiredEntryData RequiredEntryData(GetSystemViewModel(), GetEmitterViewModel(),
-					FExecutionCategoryNames::Particle, FExecutionSubcategoryNames::SimulationStage,
-					GetEmitterViewModel()->GetEditorData().GetStackEditorData());
-				SimulationStageGroup->Initialize(RequiredEntryData, GetEmitterViewModel()->GetSharedScriptViewModel(), SimulationStage);
-				SimulationStageGroup->SetOnModifiedSimulationStages(UNiagaraStackSimulationStageGroup::FOnModifiedSimulationStages::CreateUObject(this, &UNiagaraStackRoot::EmitterArraysChanged));
+				SimulationStageEntry = CreateSimulationStageScriptGroup(SimulationStage, EmitterScriptViewModel, EmitterStackEditorData);
 			}
-
-			NewChildren.Add(SimulationStageGroup);
+			NewChildren.Add(SimulationStageEntry);
 		}
-		
-		NewChildren.Add(RenderGroup);
 	}
+
+	NewChildren.Add(GetOrCreateEmitterRendererGroup(CurrentChildren, FNiagaraStackRenderersOwnerStandard::CreateShared(GetEmitterViewModel().ToSharedRef())));
 }
 
-void UNiagaraStackRoot::EmitterArraysChanged()
+void UNiagaraStackRoot::RefreshEmitterSummaryChildren(const TArray<UNiagaraStackEntry*>& CurrentChildren, TArray<UNiagaraStackEntry*>& NewChildren)
 {
-	RefreshChildren();
+	NewChildren.Add(GetOrCreateEmitterSummaryGroup(CurrentChildren));
 }
 
-void UNiagaraStackRoot::OnSummaryViewStateChanged()
+void UNiagaraStackRoot::RefreshEmitterStatelessChildren(const TArray<UNiagaraStackEntry*>& CurrentChildren, TArray<UNiagaraStackEntry*>& NewChildren, TSharedPtr<FNiagaraEmitterHandleViewModel> EmitterHandleViewModel)
 {
-	if (!IsFinalized())
+	NewChildren.Add(GetOrCreateStatelessEmitterGroup(CurrentChildren, EmitterHandleViewModel.ToSharedRef()));
+	NewChildren.Add(GetOrCreateStatelessEmitterSpawnGroup(CurrentChildren, EmitterHandleViewModel.ToSharedRef()));
+	NewChildren.Add(GetOrCreateStatelessEmitterSimulateGroup(CurrentChildren, EmitterHandleViewModel.ToSharedRef()));
+	NewChildren.Add(GetOrCreateEmitterRendererGroup(CurrentChildren, FNiagaraStackRenderersOwnerStateless::CreateShared(EmitterHandleViewModel.ToSharedRef())));
+}
+
+UNiagaraStackEntry* UNiagaraStackRoot::GetOrCreateSystemPropertiesGroup(const TArray<UNiagaraStackEntry*>& CurrentChildren)
+{
+	UNiagaraStackSystemPropertiesGroup* SystemPropertiesGroup = FindCurrentChildOfType<UNiagaraStackSystemPropertiesGroup>(CurrentChildren);
+	if (SystemPropertiesGroup == nullptr)
 	{
-		if(GetEmitterViewModel().IsValid())
-		{
-			if(GetEmitterViewModel()->GetEditorData().ShouldShowSummaryView())
-			{
-				AddSummaryViewFilter();
-				RemoveDefaultViewFilter();
-			}
-			else
-			{
-				AddDefaultViewFilter();
-				RemoveSummaryViewFilter();
-			}
-		}
+		SystemPropertiesGroup = NewObject<UNiagaraStackSystemPropertiesGroup>(this);
+		FRequiredEntryData RequiredEntryData(GetSystemViewModel(), GetEmitterViewModel(),
+			FExecutionCategoryNames::System, FExecutionSubcategoryNames::Settings,
+			GetSystemViewModel()->GetEditorData().GetStackEditorData());
+		SystemPropertiesGroup->Initialize(RequiredEntryData);
+	}
+	return SystemPropertiesGroup;
+}
 
-		RefreshChildren();
-		GetSystemViewModel()->GetSelectionViewModel()->Refresh();
+UNiagaraStackEntry* UNiagaraStackRoot::GetOrCreateSystemUserParametersGroup(const TArray<UNiagaraStackEntry*>& CurrentChildren)
+{
+	UNiagaraStackSystemUserParametersGroup* SystemUserParametersGroup = FindCurrentChildOfType<UNiagaraStackSystemUserParametersGroup>(CurrentChildren);
+	if (SystemUserParametersGroup == nullptr)
+	{
+		SystemUserParametersGroup = NewObject<UNiagaraStackSystemUserParametersGroup>(this);
+		FRequiredEntryData RequiredEntryData(GetSystemViewModel(), GetEmitterViewModel(),
+			FExecutionCategoryNames::System, FExecutionSubcategoryNames::Settings,
+			GetSystemViewModel()->GetEditorData().GetStackEditorData());
+		SystemUserParametersGroup->Initialize(RequiredEntryData, &GetSystemViewModel()->GetSystem(), &GetSystemViewModel()->GetSystem().GetExposedParameters());
 	}
+	return SystemUserParametersGroup;
 }
 
-void UNiagaraStackRoot::AddDefaultViewFilter()
+UNiagaraStackEntry* UNiagaraStackRoot::GetOrCreateSystemCommentCollection(const TArray<UNiagaraStackEntry*>& CurrentChildren)
 {
-	if(DefaultViewFilterHandle.IsValid())
+	UNiagaraStackCommentCollection* CommentCollection = FindCurrentChildOfType<UNiagaraStackCommentCollection>(CurrentChildren);
+	if (CommentCollection == nullptr)
 	{
-		RemoveDefaultViewFilter();
+		CommentCollection = NewObject<UNiagaraStackCommentCollection>(this);
+		FRequiredEntryData RequiredEntryData(GetSystemViewModel(), GetEmitterViewModel(),
+			NAME_None, NAME_None,
+			GetSystemViewModel()->GetEditorData().GetStackEditorData());
+		CommentCollection->Initialize(RequiredEntryData);
 	}
-	
-	DefaultViewFilterHandle = AddChildFilter(FOnFilterChild::CreateUObject(this, &UNiagaraStackRoot::FilterForDefaultView));
+	return CommentCollection;
 }
 
-void UNiagaraStackRoot::RemoveDefaultViewFilter()
+UNiagaraStackEntry* UNiagaraStackRoot::GetOrCreateEmitterPropertiesGroup(const TArray<UNiagaraStackEntry*>& CurrentChildren)
 {
-	if(DefaultViewFilterHandle.IsValid())
+	UNiagaraStackEmitterPropertiesGroup* EmitterPropertiesGroup = FindCurrentChildOfType<UNiagaraStackEmitterPropertiesGroup>(CurrentChildren);
+	if (EmitterPropertiesGroup == nullptr)
 	{
-		RemoveChildFilter(DefaultViewFilterHandle);
-		DefaultViewFilterHandle.Reset();
+		EmitterPropertiesGroup = NewObject<UNiagaraStackEmitterPropertiesGroup>(this);
+		FRequiredEntryData RequiredEntryData(GetSystemViewModel(), GetEmitterViewModel(),
+			FExecutionCategoryNames::Emitter, FExecutionSubcategoryNames::Settings,
+			GetEmitterViewModel()->GetEditorData().GetStackEditorData());
+		EmitterPropertiesGroup->Initialize(RequiredEntryData);
 	}
+	return EmitterPropertiesGroup;
 }
 
-void UNiagaraStackRoot::AddSummaryViewFilter()
+UNiagaraStackEntry* UNiagaraStackRoot::GetOrCreateEmitterSummaryGroup(const TArray<UNiagaraStackEntry*>& CurrentChildren)
 {
-	if(SummaryViewFilterHandle.IsValid())
+	UNiagaraStackEmitterSummaryGroup* EmitterSummaryGroup = FindCurrentChildOfType<UNiagaraStackEmitterSummaryGroup>(CurrentChildren);
+	if (EmitterSummaryGroup == nullptr)
 	{
-		RemoveSummaryViewFilter();
+		EmitterSummaryGroup = NewObject<UNiagaraStackEmitterSummaryGroup>(this);
+		FRequiredEntryData RequiredEntryData(GetSystemViewModel(), GetEmitterViewModel(),
+			FExecutionCategoryNames::Emitter, FExecutionSubcategoryNames::Settings,
+			GetEmitterViewModel()->GetEditorData().GetStackEditorData());
+		FText DisplayName = LOCTEXT("EmitterSummaryGroupName", "Emitter Summary");
+		FText Tooltip = LOCTEXT("EmitterSummaryTooltip", "Summary of parameters for this Emitter.");
+		EmitterSummaryGroup->Initialize(RequiredEntryData, DisplayName, Tooltip, nullptr);
 	}
-	
-	SummaryViewFilterHandle = AddChildFilter(FOnFilterChild::CreateUObject(this, &UNiagaraStackRoot::FilterForSummaryView));
+	return EmitterSummaryGroup;
 }
 
-void UNiagaraStackRoot::RemoveSummaryViewFilter()
+UNiagaraStackEntry* UNiagaraStackRoot::GetOrCreateEmitterRendererGroup(const TArray<UNiagaraStackEntry*>& CurrentChildren, TSharedPtr<INiagaraStackRenderersOwner> RenderersOwner)
 {
-	if(SummaryViewFilterHandle.IsValid())
+	UNiagaraStackRenderItemGroup* RenderGroup = FindCurrentChildOfTypeByPredicate<UNiagaraStackRenderItemGroup>(CurrentChildren, 
+		[RenderersOwner](const UNiagaraStackRenderItemGroup* CurrentChild) { return CurrentChild->GetRenderersOwner()->GetOwnerObject() == RenderersOwner->GetOwnerObject(); });
+	if (RenderGroup == nullptr)
 	{
-		RemoveChildFilter(SummaryViewFilterHandle);
-		SummaryViewFilterHandle.Reset();
+		RenderGroup = NewObject<UNiagaraStackRenderItemGroup>(this);
+		FRequiredEntryData RequiredEntryData(GetSystemViewModel(), GetEmitterViewModel(),
+			FExecutionCategoryNames::Render, FExecutionSubcategoryNames::Render,
+			GetEmitterViewModel()->GetEditorData().GetStackEditorData());
+		RenderGroup->Initialize(RequiredEntryData, RenderersOwner);
 	}
+	return RenderGroup;
 }
 
-bool UNiagaraStackRoot::FilterForDefaultView(const UNiagaraStackEntry& NiagaraStackEntry) const
+UNiagaraStackEntry* UNiagaraStackRoot::GetOrCreateStatelessEmitterGroup(const TArray<UNiagaraStackEntry*>& CurrentChildren, TSharedPtr<FNiagaraEmitterHandleViewModel> EmitterHandleViewModel)
 {
-	if(NiagaraStackEntry.IsA<UNiagaraStackEmitterSummaryGroup>())
+	UNiagaraStatelessEmitter* StatelessEmitter = EmitterHandleViewModel->GetEmitterHandle()->GetStatelessEmitter();
+	UNiagaraStackStatelessEmitterGroup* StatelessEmitterObjectGroup = FindCurrentChildOfTypeByPredicate<UNiagaraStackStatelessEmitterGroup>(CurrentChildren,
+		[StatelessEmitter](const UNiagaraStackStatelessEmitterGroup* Group) { return Group->GetStatelessEmitter() == StatelessEmitter; });
+	if (StatelessEmitterObjectGroup == nullptr)
 	{
-		return false;
+		StatelessEmitterObjectGroup = NewObject<UNiagaraStackStatelessEmitterGroup>(this);
+		FRequiredEntryData RequiredEntryData(GetSystemViewModel(), GetEmitterViewModel(),
+			FExecutionCategoryNames::Emitter, FExecutionSubcategoryNames::Settings,
+			GetEmitterViewModel()->GetEditorData().GetStackEditorData());
+		StatelessEmitterObjectGroup->Initialize(RequiredEntryData, StatelessEmitter);
 	}
+	return StatelessEmitterObjectGroup;
+}
 
-	return true;
+UNiagaraStackEntry* UNiagaraStackRoot::GetOrCreateStatelessEmitterSpawnGroup(const TArray<UNiagaraStackEntry*>& CurrentChildren, TSharedPtr<FNiagaraEmitterHandleViewModel> EmitterHandleViewModel)
+{
+	UNiagaraStatelessEmitter* StatelessEmitter = EmitterHandleViewModel->GetEmitterHandle()->GetStatelessEmitter();
+	UNiagaraStackStatelessEmitterSpawnGroup* StatelessEmitterSpawnGroup = FindCurrentChildOfTypeByPredicate<UNiagaraStackStatelessEmitterSpawnGroup>(CurrentChildren,
+		[StatelessEmitter](const UNiagaraStackStatelessEmitterSpawnGroup* Group) { return Group->GetStatelessEmitter() == StatelessEmitter; });
+	if (StatelessEmitterSpawnGroup == nullptr)
+	{
+		StatelessEmitterSpawnGroup = NewObject<UNiagaraStackStatelessEmitterSpawnGroup>(this);
+		FRequiredEntryData RequiredEntryData(GetSystemViewModel(), GetEmitterViewModel(),
+			FExecutionCategoryNames::Emitter, FExecutionSubcategoryNames::Spawn,
+			GetEmitterViewModel()->GetEditorData().GetStackEditorData());
+		StatelessEmitterSpawnGroup->Initialize(RequiredEntryData, EmitterHandleViewModel->GetEmitterHandle()->GetStatelessEmitter());
+	}
+	return StatelessEmitterSpawnGroup;
 }
 
-bool UNiagaraStackRoot::FilterForSummaryView(const UNiagaraStackEntry& NiagaraStackEntry) const
+UNiagaraStackEntry* UNiagaraStackRoot::GetOrCreateStatelessEmitterSimulateGroup(const TArray<UNiagaraStackEntry*>& CurrentChildren, TSharedPtr<FNiagaraEmitterHandleViewModel> EmitterHandleViewModel)
 {
-	if(NiagaraStackEntry.IsA<UNiagaraStackEmitterSummaryGroup>())
+	UNiagaraStatelessEmitter* StatelessEmitter = EmitterHandleViewModel->GetEmitterHandle()->GetStatelessEmitter();
+	UNiagaraStackStatelessEmitterSimulateGroup* StatelessEmitterSimulateGroup = FindCurrentChildOfTypeByPredicate<UNiagaraStackStatelessEmitterSimulateGroup>(CurrentChildren,
+		[StatelessEmitter](const UNiagaraStackStatelessEmitterSimulateGroup* Group) { return Group->GetStatelessEmitter() == StatelessEmitter; });
+	if (StatelessEmitterSimulateGroup == nullptr)
 	{
-		return true;
+		StatelessEmitterSimulateGroup = NewObject<UNiagaraStackStatelessEmitterSimulateGroup>(this);
+		FRequiredEntryData RequiredEntryData(GetSystemViewModel(), GetEmitterViewModel(),
+			FExecutionCategoryNames::Particle, FExecutionSubcategoryNames::Update,
+			GetEmitterViewModel()->GetEditorData().GetStackEditorData());
+		StatelessEmitterSimulateGroup->Initialize(RequiredEntryData, EmitterHandleViewModel->GetEmitterHandle()->GetStatelessEmitter());
+	}
+	return StatelessEmitterSimulateGroup;
+}
+
+UNiagaraStackEntry* UNiagaraStackRoot::GetCurrentScriptGroup(const TArray<UNiagaraStackEntry*>& CurrentChildren, ENiagaraScriptUsage InScriptUsage, FGuid InScriptUsageId) const
+{
+	return FindCurrentChildOfTypeByPredicate<UNiagaraStackScriptItemGroup>(
+		CurrentChildren, [InScriptUsage, InScriptUsageId](UNiagaraStackScriptItemGroup* ScriptItemGroup) {
+			return ScriptItemGroup->GetScriptUsage() == InScriptUsage && ScriptItemGroup->GetScriptUsageId() == InScriptUsageId; });
+}
+
+UNiagaraStackEntry* UNiagaraStackRoot::CreateScriptGroup(
+	TSharedRef<FNiagaraScriptViewModel> InScriptViewModel,
+	ENiagaraScriptUsage InScriptUsage,
+	FGuid InScriptUsageId,
+	UNiagaraStackEditorData& InStackEditorData,
+	FName InExecutionCategoryName,
+	FName InExecutionSubcategoryName,
+	FText InDisplayName,
+	FText InToolTip)
+{
+	UNiagaraStackScriptItemGroup* ScriptItemGroup = NewObject<UNiagaraStackScriptItemGroup>(this);
+	FRequiredEntryData RequiredEntryData(GetSystemViewModel(), GetEmitterViewModel(),
+		InExecutionCategoryName, InExecutionSubcategoryName, InStackEditorData);
+	ScriptItemGroup->Initialize(RequiredEntryData, InDisplayName, InToolTip, InScriptViewModel, InScriptUsage, InScriptUsageId);
+	return ScriptItemGroup;
+}
+
+UNiagaraStackEntry* UNiagaraStackRoot::CreateEventScriptGroup(
+	FGuid SourceEmitterId,
+	TSharedRef<FNiagaraScriptViewModel> InScriptViewModel,
+	FGuid InScriptUsageId,
+	UNiagaraStackEditorData& InStackEditorData)
+{
+	UNiagaraStackEventScriptItemGroup* EventHandlerGroup = NewObject<UNiagaraStackEventScriptItemGroup>(this);
+	FRequiredEntryData RequiredEntryData(GetSystemViewModel(), GetEmitterViewModel(),
+		FExecutionCategoryNames::Particle, FExecutionSubcategoryNames::Event,
+		InStackEditorData);
+	EventHandlerGroup->Initialize(RequiredEntryData, GetEmitterViewModel()->GetSharedScriptViewModel(), ENiagaraScriptUsage::ParticleEventScript, InScriptUsageId, SourceEmitterId);
+	EventHandlerGroup->SetOnModifiedEventHandlers(UNiagaraStackEventScriptItemGroup::FOnModifiedEventHandlers::CreateUObject(this, &UNiagaraStackRoot::EmitterArraysChanged));
+	return EventHandlerGroup;
+}
+
+UNiagaraStackEntry* UNiagaraStackRoot::CreateSimulationStageScriptGroup(
+	UNiagaraSimulationStageBase* InSimulationStage,
+	TSharedRef<FNiagaraScriptViewModel> InScriptViewModel,
+	UNiagaraStackEditorData& InStackEditorData)
+{
+	UNiagaraStackSimulationStageGroup* SimulationStageGroup = NewObject<UNiagaraStackSimulationStageGroup>(this);
+	FRequiredEntryData RequiredEntryData(GetSystemViewModel(), GetEmitterViewModel(),
+		FExecutionCategoryNames::Particle, FExecutionSubcategoryNames::SimulationStage,
+		InStackEditorData);
+	SimulationStageGroup->Initialize(RequiredEntryData, InScriptViewModel, InSimulationStage);
+	SimulationStageGroup->SetOnModifiedSimulationStages(UNiagaraStackSimulationStageGroup::FOnModifiedSimulationStages::CreateUObject(this, &UNiagaraStackRoot::EmitterArraysChanged));
+	return SimulationStageGroup;
+}
+
+void UNiagaraStackRoot::EmitterArraysChanged()
+{
+	RefreshChildren();
+}
+
+void UNiagaraStackRoot::OnSummaryViewStateChanged()
+{
+	if (!IsFinalized())
+	{
+		RefreshChildren();
 	}
-	
-	return false;
 }
 
 #undef LOCTEXT_NAMESPACE
 
+
diff --git a/Engine/Plugins/FX/Niagara/Source/NiagaraEditor/Private/ViewModels/Stack/NiagaraStackScriptItemGroup.cpp b/Engine/Plugins/FX/Niagara/Source/NiagaraEditor/Private/ViewModels/Stack/NiagaraStackScriptItemGroup.cpp
--- a/Engine/Plugins/FX/Niagara/Source/NiagaraEditor/Private/ViewModels/Stack/NiagaraStackScriptItemGroup.cpp
+++ b/Engine/Plugins/FX/Niagara/Source/NiagaraEditor/Private/ViewModels/Stack/NiagaraStackScriptItemGroup.cpp
@@ -393,8 +393,11 @@ void UNiagaraStackScriptItemGroup::Initialize(
 	ScriptUsage = InScriptUsage;
 	ScriptUsageId = InScriptUsageId;
 	ScriptGraph = InScriptViewModel->GetGraphViewModel()->GetGraph();
-	ScriptGraph->AddOnGraphChangedHandler(
-		FOnGraphChanged::FDelegate::CreateUObject(this, &UNiagaraStackScriptItemGroup::OnScriptGraphChanged));
+	//-TODO:Stateless: Do we need stateless support here?
+	if (ScriptGraph.IsValid())
+	{
+		ScriptGraph->AddOnGraphChangedHandler(FOnGraphChanged::FDelegate::CreateUObject(this, &UNiagaraStackScriptItemGroup::OnScriptGraphChanged));
+	}
 
 	if (ScriptUsage == ENiagaraScriptUsage::SystemSpawnScript || ScriptUsage == ENiagaraScriptUsage::EmitterSpawnScript)
 	{
@@ -408,11 +411,14 @@ void UNiagaraStackScriptItemGroup::Initialize(
 	}
 	else if (GetEmitterViewModel().IsValid())
 	{
-		OwningParticleScriptWeak = GetEmitterViewModel()->GetEmitter().GetEmitterData()->GetScript(ScriptUsage, ScriptUsageId);
-		if (OwningParticleScriptWeak.IsValid())
+		if (FVersionedNiagaraEmitterData* EmitterData = GetEmitterViewModel()->GetEmitter().GetEmitterData())
 		{
-			OwningParticleScriptWeak->OnVMScriptCompiled().AddUObject(this, &UNiagaraStackScriptItemGroup::OnParticleScriptCompiled);
-			OwningParticleScriptWeak->OnGPUScriptCompiled().AddUObject(this, &UNiagaraStackScriptItemGroup::OnParticleScriptCompiled);
+			OwningParticleScriptWeak = EmitterData->GetScript(ScriptUsage, ScriptUsageId);
+			if (OwningParticleScriptWeak.IsValid())
+			{
+				OwningParticleScriptWeak->OnVMScriptCompiled().AddUObject(this, &UNiagaraStackScriptItemGroup::OnParticleScriptCompiled);
+				OwningParticleScriptWeak->OnGPUScriptCompiled().AddUObject(this, &UNiagaraStackScriptItemGroup::OnParticleScriptCompiled);
+			}
 		}
 	}
 }
@@ -546,6 +552,12 @@ void UNiagaraStackScriptItemGroup::RefreshChildrenInternal(const TArray<UNiagara
 	checkf(ScriptViewModelPinned.IsValid(), TEXT("Can not refresh children when the script view model has been deleted."));
 
 	UNiagaraGraph* Graph = ScriptViewModelPinned->GetGraphViewModel()->GetGraph();
+	//-TODO:Stateless: Do we need stateless support here?
+	if (Graph == nullptr)
+	{
+		return;
+	}
+
 	FText ErrorMessage;
 	bIsValidForOutput = false;
 	if (FNiagaraStackGraphUtilities::ValidateGraphForOutput(*Graph, ScriptUsage, ScriptUsageId, ErrorMessage) == true)
