- Experimental stateless emitters

diff --git a/Engine/Plugins/FX/Niagara/Source/NiagaraEditor/Private/ViewModels/Stack/NiagaraStackSimulationStageGroup.cpp b/Engine/Plugins/FX/Niagara/Source/NiagaraEditor/Private/ViewModels/Stack/NiagaraStackSimulationStageGroup.cpp
--- a/Engine/Plugins/FX/Niagara/Source/NiagaraEditor/Private/ViewModels/Stack/NiagaraStackSimulationStageGroup.cpp
+++ b/Engine/Plugins/FX/Niagara/Source/NiagaraEditor/Private/ViewModels/Stack/NiagaraStackSimulationStageGroup.cpp
@@ -103,7 +103,8 @@ void UNiagaraStackSimulationStagePropertiesItem::RefreshChildrenInternal(const T
 	{
 		SimulationStageObject = NewObject<UNiagaraStackObject>(this);
 		bool bIsTopLevelObject = true;
-		SimulationStageObject->Initialize(CreateDefaultChildRequiredData(), SimulationStage.Get(), bIsTopLevelObject, GetStackEditorDataKey());
+		bool bHideTopLevelCategories = false;
+		SimulationStageObject->Initialize(CreateDefaultChildRequiredData(), SimulationStage.Get(), bIsTopLevelObject, bHideTopLevelCategories, GetStackEditorDataKey());
 		SimulationStageObject->SetObjectGuid(SimulationStage->GetMergeId());
 	}
 
diff --git a/Engine/Plugins/FX/Niagara/Source/NiagaraEditor/Private/ViewModels/Stack/NiagaraStackSummaryViewInputCollection.cpp b/Engine/Plugins/FX/Niagara/Source/NiagaraEditor/Private/ViewModels/Stack/NiagaraStackSummaryViewInputCollection.cpp
--- a/Engine/Plugins/FX/Niagara/Source/NiagaraEditor/Private/ViewModels/Stack/NiagaraStackSummaryViewInputCollection.cpp
+++ b/Engine/Plugins/FX/Niagara/Source/NiagaraEditor/Private/ViewModels/Stack/NiagaraStackSummaryViewInputCollection.cpp
@@ -23,6 +23,7 @@
 #include "ViewModels/Stack/NiagaraStackObject.h"
 #include "ViewModels/Stack/NiagaraStackPropertyRow.h"
 #include "ViewModels/Stack/NiagaraStackRendererItem.h"
+#include "ViewModels/Stack/NiagaraStackRenderersOwner.h"
 #include "ViewModels/Stack/NiagaraStackSimulationStageGroup.h"
 
 #include UE_INLINE_GENERATED_CPP_BY_NAME(NiagaraStackSummaryViewInputCollection)
@@ -101,6 +102,12 @@ void UNiagaraStackSummaryViewCollection::RefreshChildrenInternal(const TArray<UN
 
 	const UNiagaraHierarchyRoot* Root = ViewModel->GetEditorData().GetSummaryRoot();
 	TSharedPtr<FNiagaraHierarchyRootViewModel> RootViewModel = ViewModel->GetSummaryHierarchyViewModel()->GetHierarchyRootViewModel();
+	//-TODO:Stateless: Do we need stateless support here?
+	if (RootViewModel == nullptr)
+	{
+		return;
+	}
+
 	// we make sure to sync view models to data before refreshing in order to get rid of possibly removed entries
 	RootViewModel->SyncViewModelsToData();
 
@@ -340,7 +347,7 @@ void UNiagaraStackSummaryViewCollection::RefreshChildrenInternal(const TArray<UN
 				if (StackRenderer == nullptr)
 				{
 					StackRenderer = NewObject<UNiagaraStackRendererItem>(this);
-					StackRenderer->Initialize(CreateDefaultChildRequiredData(), *MatchingRendererProperties);
+					StackRenderer->Initialize(CreateDefaultChildRequiredData(), FNiagaraStackRenderersOwnerStandard::CreateShared(GetEmitterViewModel().ToSharedRef()), *MatchingRendererProperties);
 				}
 			
 				NewChildren.Add(StackRenderer);	
@@ -504,35 +511,28 @@ void UNiagaraStackSummaryViewCollection::RefreshChildrenInternal(const TArray<UN
 				UNiagaraStackObject* StackObjectWithProperty = FindCurrentChildOfTypeByPredicate<UNiagaraStackObject>(CurrentChildren,
 				[&](UNiagaraStackObject* StackObjectCandidate)
 				{
-					if(StackObjectCandidate->GetObject() != Object)
-					{
-						return false;
-					}
-					
-					if(StackObjectCandidate->GetCustomRootNodeNames().Num() == 1 && StackObjectCandidate->GetCustomRootNodeNames()[0] == ObjectProperty->GetPersistentIdentity().Names[0])
-					{
-						return true;
-					}
-
-					return false;
+					return StackObjectCandidate->GetObject() == Object && StackObjectCandidate->GetCustomName() == ObjectProperty->GetPersistentIdentity().Names[0];
 				});
 
 				if(StackObjectWithProperty == nullptr)
 				{
 					StackObjectWithProperty = NewObject<UNiagaraStackObject>(this);
-					StackObjectWithProperty->Initialize(CreateDefaultChildRequiredData(), Object, false, GetStackEditorDataKey(), nullptr);
-					StackObjectWithProperty->SetOnGetCustomRootNodes(UNiagaraStackObject::FOnGetCustomRootNodes::CreateLambda([PropertyName = ObjectProperty->GetPersistentIdentity().Names[0]](TArray<TSharedRef<IDetailTreeNode>> DefaultRootNodes, TArray<TSharedRef<IDetailTreeNode>>* Result)
+					bool bIsInTopLevelObject = false;
+					bool bHideTopLevelCategories = false;
+					StackObjectWithProperty->Initialize(CreateDefaultChildRequiredData(), Object, bIsInTopLevelObject, bHideTopLevelCategories, GetStackEditorDataKey(), nullptr);
+					StackObjectWithProperty->SetCustomName(ObjectProperty->GetPersistentIdentity().Names[0]);
+					StackObjectWithProperty->SetOnFilterDetailNodes(FNiagaraStackObjectShared::FOnFilterDetailNodes::CreateLambda([PropertyName = ObjectProperty->GetPersistentIdentity().Names[0]](const TArray<TSharedRef<IDetailTreeNode>>& InSourceNodes, TArray<TSharedRef<IDetailTreeNode>>& OutFilteredNodes)
 					{
-						for(TSharedRef<IDetailTreeNode>& RootNode : DefaultRootNodes)
+						for(const TSharedRef<IDetailTreeNode>& SourceNode : InSourceNodes)
 						{
 							TArray<TSharedRef<IDetailTreeNode>> ChildrenNodes;
-							RootNode->GetChildren(ChildrenNodes);
+							SourceNode->GetChildren(ChildrenNodes);
 							
 							for(TSharedRef<IDetailTreeNode> ChildNode : ChildrenNodes)
 							{
 								if(ChildNode->GetNodeName() == PropertyName)
 								{
-									Result->Add(ChildNode);
+									OutFilteredNodes.Add(ChildNode);
 								}
 							}
 						}
diff --git a/Engine/Plugins/FX/Niagara/Source/NiagaraEditor/Private/ViewModels/Stack/NiagaraStackSystemPropertiesItem.cpp b/Engine/Plugins/FX/Niagara/Source/NiagaraEditor/Private/ViewModels/Stack/NiagaraStackSystemPropertiesItem.cpp
--- a/Engine/Plugins/FX/Niagara/Source/NiagaraEditor/Private/ViewModels/Stack/NiagaraStackSystemPropertiesItem.cpp
+++ b/Engine/Plugins/FX/Niagara/Source/NiagaraEditor/Private/ViewModels/Stack/NiagaraStackSystemPropertiesItem.cpp
@@ -38,7 +38,8 @@ void UNiagaraStackSystemPropertiesItem::RefreshChildrenInternal(const TArray<UNi
 		SystemObject = NewObject<UNiagaraStackObject>(this);
 		FRequiredEntryData RequiredEntryData(GetSystemViewModel(), GetEmitterViewModel(), FExecutionCategoryNames::System, NAME_None, GetStackEditorData());
 		bool bIsTopLevelObject = true;
-		SystemObject->Initialize(RequiredEntryData, System.Get(), bIsTopLevelObject, GetStackEditorDataKey());
+		bool bHideTopLevelCategories = false;
+		SystemObject->Initialize(RequiredEntryData, System.Get(), bIsTopLevelObject, bHideTopLevelCategories, GetStackEditorDataKey());
 		SystemObject->RegisterInstancedCustomPropertyLayout(UNiagaraSystem::StaticClass(), FOnGetDetailCustomizationInstance::CreateStatic(&FNiagaraSystemDetails::MakeInstance));
 	}
 	NewChildren.Add(SystemObject);
diff --git a/Engine/Plugins/FX/Niagara/Source/NiagaraEditor/Private/ViewModels/Stack/NiagaraStackViewModel.cpp b/Engine/Plugins/FX/Niagara/Source/NiagaraEditor/Private/ViewModels/Stack/NiagaraStackViewModel.cpp
--- a/Engine/Plugins/FX/Niagara/Source/NiagaraEditor/Private/ViewModels/Stack/NiagaraStackViewModel.cpp
+++ b/Engine/Plugins/FX/Niagara/Source/NiagaraEditor/Private/ViewModels/Stack/NiagaraStackViewModel.cpp
@@ -117,31 +117,33 @@ void UNiagaraStackViewModel::UpdateStackWithValidationResults()
 void UNiagaraStackViewModel::InitializeWithViewModels(TSharedPtr<FNiagaraSystemViewModel> InSystemViewModel, TSharedPtr<FNiagaraEmitterHandleViewModel> InEmitterHandleViewModel, FNiagaraStackViewModelOptions InOptions)
 {
 	Reset();
-
 	Options = InOptions;
-	SystemViewModel = InSystemViewModel;
-	EmitterHandleViewModel = InEmitterHandleViewModel;
-
-	TSharedPtr<FNiagaraSystemViewModel> SystemViewModelPinned = SystemViewModel.Pin();
-	TSharedPtr<FNiagaraEmitterViewModel> EmitterViewModel = InEmitterHandleViewModel.IsValid() ? InEmitterHandleViewModel->GetEmitterViewModel() : TSharedPtr<FNiagaraEmitterViewModel>();
-
-	if (SystemViewModelPinned.IsValid())
+	if (InSystemViewModel.IsValid())
 	{
-		if (EmitterViewModel.IsValid())
+		TSharedPtr<FNiagaraEmitterViewModel> EmitterViewModel;
+		if (InEmitterHandleViewModel.IsValid())
 		{
-			EmitterViewModel->OnParentRemoved().AddUObject(this, &UNiagaraStackViewModel::EmitterParentRemoved);
-
-			if (UNiagaraEmitter* Emitter = EmitterViewModel->GetEmitter().Emitter)
+			EmitterViewModel = InEmitterHandleViewModel->GetEmitterViewModel();
+			if (EmitterViewModel.IsValid())
 			{
-				EmitterViewModel->GetEditorData().OnSummaryViewStateChanged().AddUObject(this, &UNiagaraStackViewModel::RequestRefreshDeferred);
-			}
+				EmitterViewModel->OnParentRemoved().AddUObject(this, &UNiagaraStackViewModel::EmitterParentRemoved);
 
+				if (UNiagaraEmitter* Emitter = EmitterViewModel->GetEmitter().Emitter)
+				{
+					EmitterViewModel->GetEditorData().OnSummaryViewStateChanged().AddUObject(this, &UNiagaraStackViewModel::RequestRefreshDeferred);
+				}
+			}
+			InEmitterHandleViewModel->GetEmitterHandle()->OnEmitterModeChanged().AddUObject(this, &UNiagaraStackViewModel::RequestRefreshDeferred);
+		}
+		else
+		{
+			EmitterViewModel = TSharedPtr<FNiagaraEmitterViewModel>();
 		}
 		
 		UNiagaraStackRoot* StackRoot = NewObject<UNiagaraStackRoot>(this);
-		UNiagaraStackEntry::FRequiredEntryData RequiredEntryData(SystemViewModelPinned.ToSharedRef(), EmitterViewModel,
+		UNiagaraStackEntry::FRequiredEntryData RequiredEntryData(InSystemViewModel.ToSharedRef(), EmitterViewModel,
 			UNiagaraStackEntry::FExecutionCategoryNames::System, UNiagaraStackEntry::FExecutionSubcategoryNames::Settings,
-			SystemViewModelPinned->GetEditorData().GetStackEditorData());
+			InSystemViewModel->GetEditorData().GetStackEditorData());
 		StackRoot->Initialize(RequiredEntryData, Options.GetIncludeSystemInformation(), Options.GetIncludeEmitterInformation());
 		StackRoot->RefreshChildren();
 		StackRoot->OnStructureChanged().AddUObject(this, &UNiagaraStackViewModel::EntryStructureChanged);
@@ -158,6 +160,9 @@ void UNiagaraStackViewModel::InitializeWithViewModels(TSharedPtr<FNiagaraSystemV
 		UpdateStackWithValidationResults();
 	}
 
+	SystemViewModel = InSystemViewModel;
+	EmitterHandleViewModel = InEmitterHandleViewModel;
+
 	StructureChangedDelegate.Broadcast(ENiagaraStructureChangedFlags::StructureChanged);
 }
 
@@ -204,13 +209,27 @@ void UNiagaraStackViewModel::Reset()
 	}
 	RootEntries.Empty();
 
-	TSharedPtr<FNiagaraEmitterViewModel> EmitterViewModel = EmitterHandleViewModel.IsValid() ? EmitterHandleViewModel.Pin()->GetEmitterViewModel() : TSharedPtr<FNiagaraEmitterViewModel>();
-	if (EmitterViewModel.IsValid())
+	TSharedPtr<FNiagaraEmitterHandleViewModel> EmitterHandleViewModelPinned = EmitterHandleViewModel.Pin();
+	if (EmitterHandleViewModelPinned.IsValid())
 	{
-		EmitterViewModel->OnParentRemoved().RemoveAll(this);
-		if (EmitterViewModel->GetEmitter().GetEmitterData())
+		TSharedPtr<FNiagaraEmitterViewModel> EmitterViewModel = EmitterHandleViewModelPinned->GetEmitterViewModel();
+		if (EmitterViewModel.IsValid())
 		{
-			EmitterViewModel->GetEditorData().OnSummaryViewStateChanged().RemoveAll(this);
+			EmitterViewModel->OnParentRemoved().RemoveAll(this);
+			if (EmitterViewModel->GetEmitter().GetEmitterData())
+			{
+				EmitterViewModel->GetEditorData().OnSummaryViewStateChanged().RemoveAll(this);
+			}
+		}
+
+		// Note we can not use EmitterHandleViewModelPinned->GetEmitterHandle() as it will be invalid on undo
+		if (SystemViewModel.IsValid())
+		{
+			UNiagaraSystem& NiagaraSystem = SystemViewModel.Pin()->GetSystem();
+			for (FNiagaraEmitterHandle& EmitterHandle : NiagaraSystem.GetEmitterHandles())
+			{
+				EmitterHandle.OnEmitterModeChanged().RemoveAll(this);
+			}
 		}
 	}
 
@@ -868,7 +887,12 @@ void UNiagaraStackViewModel::RefreshTopLevelViewModels()
 		if (RootChild->GetEmitterViewModel().IsValid())
 		{
 			TSharedPtr<FNiagaraEmitterHandleViewModel> RootChildEmitterHandleViewModel = RootChild->GetSystemViewModel()->GetEmitterHandleViewModelForEmitter(RootChild->GetEmitterViewModel()->GetEmitter());
-			TSharedRef<FTopLevelViewModel>* CurrentTopLevelViewModelPtr = CurrentTopLevelViewModels.FindByPredicate([&RootChildEmitterHandleViewModel](const TSharedRef<FTopLevelViewModel>& TopLevelViewModel) 
+			//-TODO:Stateless: Do we need stateless support here?
+			if (RootChildEmitterHandleViewModel == nullptr)
+			{
+				return;
+			}
+			TSharedRef<FTopLevelViewModel>* CurrentTopLevelViewModelPtr = CurrentTopLevelViewModels.FindByPredicate([&RootChildEmitterHandleViewModel](const TSharedRef<FTopLevelViewModel>& TopLevelViewModel)
 			{ 
 				return TopLevelViewModel->EmitterHandleViewModel == RootChildEmitterHandleViewModel &&
 					TopLevelViewModel->RootEntry == RootChildEmitterHandleViewModel->GetEmitterStackViewModel()->GetRootEntry(); 
diff --git a/Engine/Plugins/FX/Niagara/Source/NiagaraEditor/Private/Widgets/SNiagaraSystemViewport.cpp b/Engine/Plugins/FX/Niagara/Source/NiagaraEditor/Private/Widgets/SNiagaraSystemViewport.cpp
--- a/Engine/Plugins/FX/Niagara/Source/NiagaraEditor/Private/Widgets/SNiagaraSystemViewport.cpp
+++ b/Engine/Plugins/FX/Niagara/Source/NiagaraEditor/Private/Widgets/SNiagaraSystemViewport.cpp
@@ -7,6 +7,7 @@
 #include "EditorViewportCommands.h"
 #include "EngineUtils.h"
 #include "ImageUtils.h"
+#include "Stateless/NiagaraStatelessEmitter.h"
 #include "NiagaraComponent.h"
 #include "NiagaraEditorCommands.h"
 #include "NiagaraEditorModule.h"
@@ -189,6 +190,22 @@ void FNiagaraSystemViewportClient::Draw(FViewport* InViewport, FCanvas* Canvas)
 		}
 	}
 
+	//-TODO:Stateless: Temporary to draw some debug information
+	if (ParticleSystem && Component)
+	{
+		UWorld* World = Component->GetWorld();
+		for (const FNiagaraEmitterHandle& EmitterHandle : ParticleSystem->GetEmitterHandles())
+		{
+			UNiagaraStatelessEmitter* StatelessEmitter = EmitterHandle.GetStatelessEmitter();
+			if (!StatelessEmitter || EmitterHandle.GetEmitterMode() != ENiagaraEmitterMode::Stateless)
+			{
+				continue;
+			}
+
+			StatelessEmitter->DrawModuleDebug(World, Component->GetComponentTransform());
+		}
+	}
+
 	if (NiagaraViewport.IsValid() && NiagaraViewport->GetDrawElement(SNiagaraSystemViewport::EDrawElements::Bounds))
 	{
 		EngineShowFlags.SetBounds(true);
@@ -332,11 +349,13 @@ void FNiagaraSystemViewportClient::DrawParticleCounts(UNiagaraComponent* Compone
 		TextItem.Draw(Canvas);
 		CurrentY += FontHeight;
 
-		for (const TSharedRef<FNiagaraEmitterInstance, ESPMode::ThreadSafe>& EmitterInstance : SystemInstance->GetEmitters())
+		for (const FNiagaraEmitterInstanceRef& EmitterInstance : SystemInstance->GetEmitters())
 		{
+			FVersionedNiagaraEmitterData* EmitterData = EmitterInstance->GetEmitterHandle().GetEmitterData();
+
 			const FName EmitterName = EmitterInstance->GetEmitterHandle().GetName();
 			const int32 CurrentCount = EmitterInstance->GetNumParticles();
-			const int32 MaxCount = EmitterInstance->GetEmitterHandle().GetEmitterData()->GetMaxParticleCountEstimate();
+			const int32 MaxCount = EmitterData ? EmitterData->GetMaxParticleCountEstimate() : 0;
 			const bool IsIsolated = EmitterInstance->GetEmitterHandle().IsIsolated();
 			const bool IsEnabled = EmitterInstance->GetEmitterHandle().GetIsEnabled();
 			const ENiagaraExecutionState ExecutionState = EmitterInstance->GetExecutionState();
diff --git a/Engine/Plugins/FX/Niagara/Source/NiagaraEditor/Public/NiagaraEditorUtilities.h b/Engine/Plugins/FX/Niagara/Source/NiagaraEditor/Public/NiagaraEditorUtilities.h
--- a/Engine/Plugins/FX/Niagara/Source/NiagaraEditor/Public/NiagaraEditorUtilities.h
+++ b/Engine/Plugins/FX/Niagara/Source/NiagaraEditor/Public/NiagaraEditorUtilities.h
@@ -309,6 +309,9 @@ namespace FNiagaraEditorUtilities
 	ECheckBoxState GetSelectedEmittersIsolatedCheckState(TSharedRef<FNiagaraSystemViewModel> SystemViewModel);
 	void ToggleSelectedEmittersIsolated(TSharedRef<FNiagaraSystemViewModel> SystemViewModel);
 
+	ECheckBoxState GetSelectedEmittersEmitterModeCheckState(TSharedRef<FNiagaraSystemViewModel> SystemViewModel, ENiagaraEmitterMode EmitterMode);
+	void SetSelectedEmittersEmitterMode(TSharedRef<FNiagaraSystemViewModel> SystemViewModel, ENiagaraEmitterMode EmitterMode);
+
 	void CreateAssetFromEmitter(TSharedRef<FNiagaraEmitterHandleViewModel> EmitterHandleViewModel);
 
 	NIAGARAEDITOR_API void WarnWithToastAndLog(FText WarningMessage);
diff --git a/Engine/Plugins/FX/Niagara/Source/NiagaraEditor/Public/ViewModels/NiagaraSystemViewModel.h b/Engine/Plugins/FX/Niagara/Source/NiagaraEditor/Public/ViewModels/NiagaraSystemViewModel.h
--- a/Engine/Plugins/FX/Niagara/Source/NiagaraEditor/Public/ViewModels/NiagaraSystemViewModel.h
+++ b/Engine/Plugins/FX/Niagara/Source/NiagaraEditor/Public/ViewModels/NiagaraSystemViewModel.h
@@ -245,6 +245,9 @@ public:
 	/** Adds an empty emitter to the system. */
 	NIAGARAEDITOR_API TSharedPtr<FNiagaraEmitterHandleViewModel> AddEmptyEmitter();
 
+	/** Adds a stateless emitter to the system. */
+	NIAGARAEDITOR_API TSharedPtr<FNiagaraEmitterHandleViewModel> AddStatelessEmitter();
+
 	/** Deletes the emitters with the supplied ids from the system */
 	NIAGARAEDITOR_API void DeleteEmitters(TSet<FGuid> EmitterHandleIdsToDelete);
 
diff --git a/Engine/Plugins/FX/Niagara/Source/NiagaraEditor/Public/ViewModels/Stack/NiagaraStackEntry.h b/Engine/Plugins/FX/Niagara/Source/NiagaraEditor/Public/ViewModels/Stack/NiagaraStackEntry.h
--- a/Engine/Plugins/FX/Niagara/Source/NiagaraEditor/Public/ViewModels/Stack/NiagaraStackEntry.h
+++ b/Engine/Plugins/FX/Niagara/Source/NiagaraEditor/Public/ViewModels/Stack/NiagaraStackEntry.h
@@ -440,6 +440,20 @@ public:
 		return nullptr;
 	}
 
+	template<typename ChildType>
+	static ChildType* FindCurrentChildOfType(const TArray<UNiagaraStackEntry*>& CurrentChildren)
+	{
+		for (UNiagaraStackEntry* CurrentChild : CurrentChildren)
+		{
+			ChildType* TypedCurrentChild = Cast<ChildType>(CurrentChild);
+			if (TypedCurrentChild != nullptr)
+			{
+				return TypedCurrentChild;
+			}
+		}
+		return nullptr;
+	}
+
 	NIAGARAEDITOR_API virtual void GetSearchItems(TArray<FStackSearchItem>& SearchItems) const;
 
 	NIAGARAEDITOR_API virtual UObject* GetExternalAsset() const;
diff --git a/Engine/Plugins/FX/Niagara/Source/NiagaraEditor/Public/ViewModels/Stack/NiagaraStackEventScriptItemGroup.h b/Engine/Plugins/FX/Niagara/Source/NiagaraEditor/Public/ViewModels/Stack/NiagaraStackEventScriptItemGroup.h
--- a/Engine/Plugins/FX/Niagara/Source/NiagaraEditor/Public/ViewModels/Stack/NiagaraStackEventScriptItemGroup.h
+++ b/Engine/Plugins/FX/Niagara/Source/NiagaraEditor/Public/ViewModels/Stack/NiagaraStackEventScriptItemGroup.h
@@ -57,7 +57,7 @@ protected:
 private:
 	NIAGARAEDITOR_API void EventHandlerPropertiesChanged();
 
-	NIAGARAEDITOR_API void SelectEmitterStackObjectRootTreeNodes(TArray<TSharedRef<IDetailTreeNode>> Source, TArray<TSharedRef<IDetailTreeNode>>* Selected);
+	NIAGARAEDITOR_API void FilterEmitterStackObjectRootTreeNodes(const TArray<TSharedRef<IDetailTreeNode>>& InSourceNodes, TArray<TSharedRef<IDetailTreeNode>>& OutFilteredNodes);
 
 	NIAGARAEDITOR_API bool HasBaseEventHandler() const;
 
diff --git a/Engine/Plugins/FX/Niagara/Source/NiagaraEditor/Public/ViewModels/Stack/NiagaraStackItem.h b/Engine/Plugins/FX/Niagara/Source/NiagaraEditor/Public/ViewModels/Stack/NiagaraStackItem.h
--- a/Engine/Plugins/FX/Niagara/Source/NiagaraEditor/Public/ViewModels/Stack/NiagaraStackItem.h
+++ b/Engine/Plugins/FX/Niagara/Source/NiagaraEditor/Public/ViewModels/Stack/NiagaraStackItem.h
@@ -7,10 +7,35 @@
 #include "Layout/Visibility.h"
 #include "NiagaraStackItem.generated.h"
 
+struct FSlateBrush;
 class UNiagaraStackItemFooter;
 class UNiagaraNode;
 class UNiagaraClipboardContent;
 
+class INiagaraStackItemHeaderValueHandler
+{
+public: 
+	enum class EValueMode
+	{
+		BoolToggle,
+		EnumDropDown
+	};
+
+public:
+	virtual ~INiagaraStackItemHeaderValueHandler() { }
+	virtual EValueMode GetMode() const = 0;
+	virtual const UEnum* GetEnum() const = 0;
+	virtual const FText& GetLabelText() const = 0;
+	virtual const FSlateBrush* GetIconBrush() const = 0;
+	virtual const EHorizontalAlignment GetHAlign() const = 0;
+
+	virtual bool GetBoolValue() const = 0;
+	virtual void NotifyBoolValueChanged(bool bInValue) = 0;
+
+	virtual int32 GetEnumValue() const = 0;
+	virtual void NotifyEnumValueChanged(int32 bInValue) = 0;
+};
+
 UCLASS(MinimalAPI)
 class UNiagaraStackItem : public UNiagaraStackEntry
 {
@@ -46,6 +71,9 @@ public:
 	virtual TOptional<FText> GetEditModeButtonTooltip() const { return TOptional<FText>(); }
 	virtual EVisibility IsEditButtonVisible() const { return SupportsEditMode() ? EVisibility::Visible : EVisibility::Collapsed; }
 
+	virtual bool SupportsHeaderValues() const { return false; }
+	virtual void GetHeaderValueHandlers(TArray<TSharedRef<INiagaraStackItemHeaderValueHandler>>& OutHeaderValueHandlers) const { };
+
 	virtual bool GetIsInherited() const { return false; }
 
 protected:
diff --git a/Engine/Plugins/FX/Niagara/Source/NiagaraEditor/Public/ViewModels/Stack/NiagaraStackObject.h b/Engine/Plugins/FX/Niagara/Source/NiagaraEditor/Public/ViewModels/Stack/NiagaraStackObject.h
--- a/Engine/Plugins/FX/Niagara/Source/NiagaraEditor/Public/ViewModels/Stack/NiagaraStackObject.h
+++ b/Engine/Plugins/FX/Niagara/Source/NiagaraEditor/Public/ViewModels/Stack/NiagaraStackObject.h
@@ -3,6 +3,7 @@
 #pragma once
 
 #include "ViewModels/Stack/NiagaraStackItem.h"
+#include "ViewModels/Stack/NiagaraStackObjectShared.h"
 #include "Misc/NotifyHook.h"
 #include "PropertyEditorDelegates.h"
 #include "NiagaraStackObject.generated.h"
@@ -12,14 +13,13 @@ class UNiagaraNode;
 class IDetailTreeNode;
 class INiagaraMessage;
 
+class FStructOnScope;
+
 UCLASS(MinimalAPI)
 class UNiagaraStackObject : public UNiagaraStackItemContent, public FNotifyHook
 {
 	GENERATED_BODY()
 		
-public:
-	DECLARE_DELEGATE_TwoParams(FOnGetCustomRootNodes, TArray<TSharedRef<IDetailTreeNode>>, TArray<TSharedRef<IDetailTreeNode>>*);
-
 public:
 	NIAGARAEDITOR_API UNiagaraStackObject();
 
@@ -29,20 +29,51 @@ public:
 		@param bInIsTopLevelObject Whether or not the object being displayed should be treated as a top level object shown in the root of the stack.  This can be used to make layout decisions such as category styling.
 		@param InOwnerStackItemEditorDataKey The stack editor data key of the owning stack entry. 
 		@param InOwningNiagaraNode An optional niagara node which owns this object. */
-	NIAGARAEDITOR_API void Initialize(FRequiredEntryData InRequiredEntryData, UObject* InObject, bool bInIsTopLevelObject, FString InOwnerStackItemEditorDataKey, UNiagaraNode* InOwningNiagaraNode = nullptr);
-	
-	NIAGARAEDITOR_API void SetOnGetCustomRootNodes(FOnGetCustomRootNodes OnSelectRootNodes);
-	NIAGARAEDITOR_API TArray<FName> GetCustomRootNodeNames() const;
-	
+	NIAGARAEDITOR_API void Initialize(
+		FRequiredEntryData InRequiredEntryData,
+		UObject* InObject,
+		bool bInIsTopLevelObject,
+		bool bInHideTopLevelCategories,
+		FString InOwnerStackItemEditorDataKey,
+		UNiagaraNode* InOwningNiagaraNode = nullptr);
+
+	/** Initializes an entry for displaying a UStruct in the niagara stack.
+		@param InRequiredEntryData Struct with required data for all stack entries.
+		@param InOwningObject The object which owns the struct being displayed in the stack.
+		@param InStruct The struct which will be displayed in the stack.
+		@param InStructName A unique name for the struct which is used for identification and for generating the stack editor data key.
+		@param bInIsTopLevelStruct Whether or not the struct being displayed should be treated as a top level object shown in the root of the stack.  This can be used to make layout decisions such as category styling.
+		@param InOwnerStackItemEditorDataKey The stack editor data key of the owning stack entry.
+		@param InOwningNiagaraNode An optional niagara node which owns the displayed struct. */
+	void Initialize(
+		FRequiredEntryData InRequiredEntryData,
+		UObject* InOwningObject,
+		TSharedRef<FStructOnScope> InDisplayedStruct,
+		const FString& InStructName,
+		bool bInIsTopLevelStruct,
+		bool bInHideTopLevelCategories,
+		FString InOwnerStackItemEditorDataKey,
+		UNiagaraNode* InOwningNiagaraNode = nullptr);
+
+	enum class EDetailNodeFilterMode
+	{
+		FilterRootNodesOnly,
+		FilterAllNodes
+	};
+
+	NIAGARAEDITOR_API void SetOnFilterDetailNodes(FNiagaraStackObjectShared::FOnFilterDetailNodes InOnFilterDetailNodes, EDetailNodeFilterMode FilterMode = EDetailNodeFilterMode::FilterRootNodesOnly);
+
 	NIAGARAEDITOR_API void RegisterInstancedCustomPropertyLayout(UStruct* Class, FOnGetDetailCustomizationInstance DetailLayoutDelegate);
 	NIAGARAEDITOR_API void RegisterInstancedCustomPropertyTypeLayout(FName PropertyTypeName, FOnGetPropertyTypeCustomizationInstance PropertyTypeLayoutDelegate, TSharedPtr<IPropertyTypeIdentifier> Identifier = nullptr);
 
-	NIAGARAEDITOR_API UObject* GetObject();
+	NIAGARAEDITOR_API UObject* GetObject() const { return WeakObject.Get(); }
+	NIAGARAEDITOR_API TSharedPtr<FStructOnScope> GetDisplayedStruct() const { return DisplayedStruct; }
 
 	//~ FNotifyHook interface
+	NIAGARAEDITOR_API virtual void NotifyPreChange(FProperty* PropertyAboutToChange) override;
 	NIAGARAEDITOR_API virtual void NotifyPostChange(const FPropertyChangedEvent& PropertyChangedEvent, FProperty* PropertyThatChanged) override;
 
-	//~ UNiagaraStackEntry interface
+	//~ UNiagaraStackEntry interface+
 	NIAGARAEDITOR_API virtual bool GetIsEnabled() const override;
 	NIAGARAEDITOR_API virtual bool GetShouldShowInStack() const override;
 	NIAGARAEDITOR_API virtual UObject* GetDisplayedObject() const override;
@@ -50,6 +81,13 @@ public:
 	void SetObjectGuid(FGuid InGuid) { ObjectGuid = InGuid; }
 	virtual bool SupportsSummaryView() const override { return ObjectGuid.IsSet() && ObjectGuid->IsValid(); }
 	NIAGARAEDITOR_API virtual FNiagaraHierarchyIdentity DetermineSummaryIdentity() const override;
+	
+	/** Gets a custom name for indentifying this stack object. */
+	FName GetCustomName() const { return CustomName; }
+
+	/** Sets a custon name for identifying this stack object. */
+	void SetCustomName(FName InCustomName) { CustomName = InCustomName; }
+
 protected:
 	NIAGARAEDITOR_API virtual void FinalizeInternal() override;
 
@@ -75,12 +113,16 @@ private:
 	};
 
 	TWeakObjectPtr<UObject> WeakObject;
+	TSharedPtr<FStructOnScope> DisplayedStruct;
 
 	// Whether or not the object being displayed should be treated as a top level object shown in the root of the stack.  This can be used to make layout decisions such as category styling.
-	bool bIsTopLevelObject;
+	bool bIsTopLevel;
+
+	bool bHideTopLevelCategories;
 
 	UNiagaraNode* OwningNiagaraNode;
-	FOnGetCustomRootNodes OnSelectRootNodesDelegate;
+	FNiagaraStackObjectShared::FOnFilterDetailNodes OnFilterDetailNodesDelegate;
+	EDetailNodeFilterMode FilterMode;
 	TArray<FRegisteredClassCustomization> RegisteredClassCustomizations;
 	TArray<FRegisteredPropertyCustomization> RegisteredPropertyCustomizations;
 	TSharedPtr<IPropertyRowGenerator> PropertyRowGenerator;
@@ -89,8 +131,7 @@ private:
 	/** An optional object guid that can be provided to identify the object this stack entry represents. This can be used for summary view purposes and is also given to the child property rows. */
 	TOptional<FGuid> ObjectGuid;
 
-	/** We keep track of our custom root nodes, if available. This helps us determine identities of stack objects. */
-	TArray<TSharedRef<IDetailTreeNode>> CustomRootNodes;
+	FName CustomName = NAME_None;
 	
 	FGuid MessageManagerRegistrationKey;
 	TArray<FStackIssue> MessageManagerIssues;
diff --git a/Engine/Plugins/FX/Niagara/Source/NiagaraEditor/Public/ViewModels/Stack/NiagaraStackPropertyRow.h b/Engine/Plugins/FX/Niagara/Source/NiagaraEditor/Public/ViewModels/Stack/NiagaraStackPropertyRow.h
--- a/Engine/Plugins/FX/Niagara/Source/NiagaraEditor/Public/ViewModels/Stack/NiagaraStackPropertyRow.h
+++ b/Engine/Plugins/FX/Niagara/Source/NiagaraEditor/Public/ViewModels/Stack/NiagaraStackPropertyRow.h
@@ -3,6 +3,7 @@
 #pragma once
 
 #include "ViewModels/Stack/NiagaraStackItem.h"
+#include "ViewModels/Stack/NiagaraStackObjectShared.h"
 #include "NiagaraStackPropertyRow.generated.h"
 
 class IDetailTreeNode;
@@ -14,10 +15,21 @@ class UNiagaraStackPropertyRow : public UNiagaraStackItemContent
 	GENERATED_BODY()
 		
 public:
-	NIAGARAEDITOR_API void Initialize(FRequiredEntryData InRequiredEntryData, TSharedRef<IDetailTreeNode> InDetailTreeNode, bool bInIsTopLevelProperty, FString InOwnerStackItemEditorDataKey, FString InOwnerStackEditorDataKey, UNiagaraNode* InOwningNiagaraNode);
+	NIAGARAEDITOR_API void Initialize(
+		FRequiredEntryData InRequiredEntryData,
+		TSharedRef<IDetailTreeNode> InDetailTreeNode,
+		bool bInIsTopLevelProperty,
+		bool bInHideTopLevelCategories,
+		FString InOwnerStackItemEditorDataKey,
+		FString InOwnerStackEditorDataKey,
+		UNiagaraNode* InOwningNiagaraNode);
 	
+	NIAGARAEDITOR_API void SetOnFilterDetailNodes(FNiagaraStackObjectShared::FOnFilterDetailNodes InOnFilterDetailNodes) { OnFilterDetailNodes = InOnFilterDetailNodes; }
+
 	NIAGARAEDITOR_API virtual EStackRowStyle GetStackRowStyle() const override;
 
+	NIAGARAEDITOR_API virtual bool GetShouldShowInStack() const override;
+
 	NIAGARAEDITOR_API TSharedRef<IDetailTreeNode> GetDetailTreeNode() const;
 
 	NIAGARAEDITOR_API virtual bool GetIsEnabled() const override;
@@ -59,4 +71,8 @@ private:
 
 	bool bCannotEditInThisContext;
 	bool bIsTopLevelProperty;
+	bool bHideTopLevelCategories;
+	bool bIsHiddenCategory;
+
+	FNiagaraStackObjectShared::FOnFilterDetailNodes OnFilterDetailNodes;
 };
diff --git a/Engine/Plugins/FX/Niagara/Source/NiagaraEditor/Public/ViewModels/Stack/NiagaraStackRenderItemGroup.h b/Engine/Plugins/FX/Niagara/Source/NiagaraEditor/Public/ViewModels/Stack/NiagaraStackRenderItemGroup.h
--- a/Engine/Plugins/FX/Niagara/Source/NiagaraEditor/Public/ViewModels/Stack/NiagaraStackRenderItemGroup.h
+++ b/Engine/Plugins/FX/Niagara/Source/NiagaraEditor/Public/ViewModels/Stack/NiagaraStackRenderItemGroup.h
@@ -6,9 +6,9 @@
 #include "ViewModels/Stack/NiagaraStackItemGroup.h"
 #include "NiagaraStackRenderItemGroup.generated.h"
 
-class UNiagaraRendererProperties;
-class UNiagaraEmitter;
+class INiagaraStackRenderersOwner;
 class UNiagaraClipboardContent;
+class UNiagaraRendererProperties;
 
 UCLASS(MinimalAPI)
 class UNiagaraStackRenderItemGroup : public UNiagaraStackItemGroup
@@ -16,7 +16,15 @@ class UNiagaraStackRenderItemGroup : public UNiagaraStackItemGroup
 	GENERATED_BODY()
 
 public:
-	NIAGARAEDITOR_API void Initialize(FRequiredEntryData InRequiredEntryData);
+	NIAGARAEDITOR_API void Initialize(FRequiredEntryData InRequiredEntryData, TSharedPtr<INiagaraStackRenderersOwner> InRenderersOwner);
+
+	virtual UNiagaraStackEntry::EIconMode GetSupportedIconMode() const override;
+	virtual const FSlateBrush* GetIconBrush() const override;
+	virtual FText GetIconText() const override;
+
+	TSharedPtr<INiagaraStackRenderersOwner> GetRenderersOwner() const { return RenderersOwner; }
+
+	NIAGARAEDITOR_API virtual bool GetCanExpandInOverview() const;
 
 	virtual bool SupportsPaste() const override { return true; }
 	NIAGARAEDITOR_API virtual bool TestCanPasteWithMessage(const UNiagaraClipboardContent* ClipboardContent, FText& OutMessage) const override;
@@ -35,6 +43,5 @@ private:
 private:
 	NIAGARAEDITOR_API void OnRendererAdded(UNiagaraRendererProperties* RendererProperties) const;
 	TSharedPtr<INiagaraStackItemGroupAddUtilities> AddUtilities;
-
-	FVersionedNiagaraEmitterWeakPtr EmitterWeak;
+	TSharedPtr<INiagaraStackRenderersOwner> RenderersOwner;
 };
diff --git a/Engine/Plugins/FX/Niagara/Source/NiagaraEditor/Public/ViewModels/Stack/NiagaraStackRendererItem.h b/Engine/Plugins/FX/Niagara/Source/NiagaraEditor/Public/ViewModels/Stack/NiagaraStackRendererItem.h
--- a/Engine/Plugins/FX/Niagara/Source/NiagaraEditor/Public/ViewModels/Stack/NiagaraStackRendererItem.h
+++ b/Engine/Plugins/FX/Niagara/Source/NiagaraEditor/Public/ViewModels/Stack/NiagaraStackRendererItem.h
@@ -8,6 +8,7 @@
 #include "NiagaraStackRendererItem.generated.h"
 
 struct FVersionedNiagaraEmitterData;
+class INiagaraStackRenderersOwner;
 class UNiagaraEmitter;
 class UNiagaraRendererProperties;
 class UNiagaraStackObject;
@@ -20,7 +21,7 @@ class UNiagaraStackRendererItem : public UNiagaraStackItem
 public:
 	NIAGARAEDITOR_API UNiagaraStackRendererItem();
 
-	NIAGARAEDITOR_API void Initialize(FRequiredEntryData InRequiredEntryData, UNiagaraRendererProperties* InRendererProperties);
+	NIAGARAEDITOR_API void Initialize(FRequiredEntryData InRequiredEntryData, TSharedPtr<INiagaraStackRenderersOwner> InRenderersOwner, UNiagaraRendererProperties* InRendererProperties);
 
 	NIAGARAEDITOR_API UNiagaraRendererProperties* GetRendererProperties();
 	NIAGARAEDITOR_API const UNiagaraRendererProperties* GetRendererProperties() const;
@@ -70,6 +71,8 @@ public:
 	NIAGARAEDITOR_API virtual bool TestCanResetToBaseWithMessage(FText& OutCanResetToBaseMessage) const override;
 	NIAGARAEDITOR_API virtual void ResetToBase() override;
 
+	NIAGARAEDITOR_API virtual bool GetShouldShowInOverview() const;
+
 	static NIAGARAEDITOR_API TArray<FNiagaraVariable> GetMissingVariables(UNiagaraRendererProperties* RendererProperties, const FVersionedNiagaraEmitterData* EmitterData);
 	static NIAGARAEDITOR_API bool AddMissingVariable(const FVersionedNiagaraEmitterData* EmitterData, const FNiagaraVariable& Variable);
 
@@ -95,6 +98,8 @@ private:
 	NIAGARAEDITOR_API void ProcessRendererIssues(const TArray<FNiagaraRendererFeedback>& InIssues, EStackIssueSeverity Severity, TArray<FStackIssue>& OutIssues);
 
 private:
+	TSharedPtr<INiagaraStackRenderersOwner> RenderersOwner;
+
 	TWeakObjectPtr<UNiagaraRendererProperties> RendererProperties;
 
 	mutable TOptional<bool> bHasBaseRendererCache;
diff --git a/Engine/Plugins/FX/Niagara/Source/NiagaraEditor/Public/ViewModels/Stack/NiagaraStackRoot.h b/Engine/Plugins/FX/Niagara/Source/NiagaraEditor/Public/ViewModels/Stack/NiagaraStackRoot.h
--- a/Engine/Plugins/FX/Niagara/Source/NiagaraEditor/Public/ViewModels/Stack/NiagaraStackRoot.h
+++ b/Engine/Plugins/FX/Niagara/Source/NiagaraEditor/Public/ViewModels/Stack/NiagaraStackRoot.h
@@ -6,14 +6,17 @@
 #include "ViewModels/Stack/NiagaraStackEntry.h"
 #include "NiagaraStackRoot.generated.h"
 
+class FNiagaraEmitterHandleViewModel;
 class FNiagaraEmitterViewModel;
+class INiagaraStackRenderersOwner;
+class UNiagaraSimulationStageBase;
 class UNiagaraStackSystemPropertiesGroup;
 class UNiagaraStackSystemUserParametersGroup;
 class UNiagaraStackEmitterPropertiesGroup;
 class UNiagaraStackScriptItemGroup;
 class UNiagaraStackRenderItemGroup;
 class UNiagaraStackEmitterSummaryGroup;
-class UNiagaraStackSummaryViewCollapseButton;
+class UNiagaraStackEmitterStatelessGroup;
 
 UCLASS(MinimalAPI)
 class UNiagaraStackRoot : public UNiagaraStackEntry
@@ -28,20 +31,9 @@ public:
 
 	NIAGARAEDITOR_API virtual bool GetCanExpand() const override;
 	NIAGARAEDITOR_API virtual bool GetShouldShowInStack() const override;
-	UNiagaraStackRenderItemGroup* GetRenderGroup() const
-	{
-		return RenderGroup;
-	}
-
-	UNiagaraStackCommentCollection* GetCommentCollection() const
-	{
-		return CommentCollection;
-	}
-
-	UNiagaraStackEmitterSummaryGroup* GetEmitterSummaryGroup() const
-	{
-		return EmitterSummaryGroup;
-	}
+
+	UNiagaraStackRenderItemGroup* GetRenderGroup() const;
+	UNiagaraStackCommentCollection* GetCommentCollection() const;
 	
 protected:
 	NIAGARAEDITOR_API virtual void RefreshChildrenInternal(const TArray<UNiagaraStackEntry*>& CurrentChildren, TArray<UNiagaraStackEntry*>& NewChildren, TArray<FStackIssue>& NewIssues) override;
@@ -50,52 +42,47 @@ private:
 	NIAGARAEDITOR_API void EmitterArraysChanged();
 	NIAGARAEDITOR_API void OnSummaryViewStateChanged();
 
-	void AddDefaultViewFilter();
-	void RemoveDefaultViewFilter();
-	
-	void AddSummaryViewFilter();
-	void RemoveSummaryViewFilter();
-
-	NIAGARAEDITOR_API bool FilterForDefaultView(const UNiagaraStackEntry& NiagaraStackEntry) const;
-	NIAGARAEDITOR_API bool FilterForSummaryView(const UNiagaraStackEntry& NiagaraStackEntry) const;
+	void RefreshSystemChildren(const TArray<UNiagaraStackEntry*>& CurrentChildren, TArray<UNiagaraStackEntry*>& NewChildren);
+	void RefreshEmitterFullChildren(const TArray<UNiagaraStackEntry*>& CurrentChildren, TArray<UNiagaraStackEntry*>& NewChildren);
+	void RefreshEmitterSummaryChildren(const TArray<UNiagaraStackEntry*>& CurrentChildren, TArray<UNiagaraStackEntry*>& NewChildren);
+	void RefreshEmitterStatelessChildren(const TArray<UNiagaraStackEntry*>& CurrentChildren, TArray<UNiagaraStackEntry*>& NewChildren, TSharedPtr<FNiagaraEmitterHandleViewModel> EmitterHandleViewModel);
+
+	UNiagaraStackEntry* GetOrCreateSystemPropertiesGroup(const TArray<UNiagaraStackEntry*>& CurrentChildren);
+	UNiagaraStackEntry* GetOrCreateSystemUserParametersGroup(const TArray<UNiagaraStackEntry*>& CurrentChildren);
+	UNiagaraStackEntry* GetOrCreateSystemCommentCollection(const TArray<UNiagaraStackEntry*>& CurrentChildren);
+
+	UNiagaraStackEntry* GetOrCreateEmitterPropertiesGroup(const TArray<UNiagaraStackEntry*>& CurrentChildren);
+	UNiagaraStackEntry* GetOrCreateEmitterSummaryGroup(const TArray<UNiagaraStackEntry*>& CurrentChildren);
+	UNiagaraStackEntry* GetOrCreateEmitterRendererGroup(const TArray<UNiagaraStackEntry*>& CurrentChildren, TSharedPtr<INiagaraStackRenderersOwner> RenderersOwner);
+
+	UNiagaraStackEntry* GetOrCreateStatelessEmitterGroup(const TArray<UNiagaraStackEntry*>& CurrentChildren, TSharedPtr<FNiagaraEmitterHandleViewModel> EmitterHandleViewModel);
+	UNiagaraStackEntry* GetOrCreateStatelessEmitterSpawnGroup(const TArray<UNiagaraStackEntry*>& CurrentChildren, TSharedPtr<FNiagaraEmitterHandleViewModel> EmitterHandleViewModel);
+	UNiagaraStackEntry* GetOrCreateStatelessEmitterSimulateGroup(const TArray<UNiagaraStackEntry*>& CurrentChildren, TSharedPtr<FNiagaraEmitterHandleViewModel> EmitterHandleViewModel);
+
+	UNiagaraStackEntry* GetCurrentScriptGroup(const TArray<UNiagaraStackEntry*>& CurrentChildren, ENiagaraScriptUsage InScriptUsage, FGuid InScriptUsageId) const;
+
+	UNiagaraStackEntry* CreateScriptGroup(
+		TSharedRef<FNiagaraScriptViewModel> InScriptViewModel,
+		ENiagaraScriptUsage InScriptUsage,
+		FGuid InScriptUsageId,
+		UNiagaraStackEditorData& InStackEditorData,
+		FName InExecutionCategoryName,
+		FName InExecutionSubcategoryName,
+		FText InDisplayName,
+		FText InToolTip);
+
+	UNiagaraStackEntry* CreateEventScriptGroup(
+		FGuid SourceEmitterId,
+		TSharedRef<FNiagaraScriptViewModel> InScriptViewModel,
+		FGuid InScriptUsageId,
+		UNiagaraStackEditorData& InStackEditorData);
+
+	UNiagaraStackEntry* CreateSimulationStageScriptGroup(
+		UNiagaraSimulationStageBase* InSimulationStage,
+		TSharedRef<FNiagaraScriptViewModel> InScriptViewModel,
+		UNiagaraStackEditorData& InStackEditorData);
 
 private:
-	UPROPERTY()
-	TObjectPtr<UNiagaraStackSystemPropertiesGroup> SystemPropertiesGroup;
-
-	UPROPERTY()
-	TObjectPtr<UNiagaraStackScriptItemGroup> SystemSpawnGroup;
-
-	UPROPERTY()
-	TObjectPtr<UNiagaraStackScriptItemGroup> SystemUpdateGroup;
-
-	UPROPERTY()
-	TObjectPtr<UNiagaraStackEmitterPropertiesGroup> EmitterPropertiesGroup;
-
-	UPROPERTY()
-	TObjectPtr<UNiagaraStackEmitterSummaryGroup> EmitterSummaryGroup;
-
-	UPROPERTY()
-	TObjectPtr<UNiagaraStackScriptItemGroup> EmitterSpawnGroup;
-
-	UPROPERTY()
-	TObjectPtr<UNiagaraStackScriptItemGroup> EmitterUpdateGroup;
-
-	UPROPERTY()
-	TObjectPtr<UNiagaraStackScriptItemGroup> ParticleSpawnGroup;
-
-	UPROPERTY()
-	TObjectPtr<UNiagaraStackScriptItemGroup> ParticleUpdateGroup;
-
-	UPROPERTY()
-	TObjectPtr<UNiagaraStackRenderItemGroup> RenderGroup;
-
-	UPROPERTY()
-	TObjectPtr<UNiagaraStackCommentCollection> CommentCollection;
-
 	bool bIncludeSystemInformation;
 	bool bIncludeEmitterInformation;
-
-	FDelegateHandle DefaultViewFilterHandle;
-	FDelegateHandle SummaryViewFilterHandle;
 };
diff --git a/Engine/Plugins/FX/Niagara/Source/NiagaraEditorWidgets/Private/NiagaraEditorWidgetsStyle.cpp b/Engine/Plugins/FX/Niagara/Source/NiagaraEditorWidgets/Private/NiagaraEditorWidgetsStyle.cpp
--- a/Engine/Plugins/FX/Niagara/Source/NiagaraEditorWidgets/Private/NiagaraEditorWidgetsStyle.cpp
+++ b/Engine/Plugins/FX/Niagara/Source/NiagaraEditorWidgets/Private/NiagaraEditorWidgetsStyle.cpp
@@ -517,6 +517,17 @@ void FNiagaraEditorWidgetsStyle::InitStackTextStyles()
 		.SetFont(CategoryFont);
 
 	Set("NiagaraEditor.Stack.ItemText", StackItemText);
+
+	FSlateFontInfo StackItemHeaderValueLabelFont = FAppStyle::Get().GetFontStyle(TEXT("PropertyWindow.NormalFont"));
+	FTextBlockStyle StackItemHeaderValueLabelText = FTextBlockStyle(NormalText)
+		.SetFont(StackItemHeaderValueLabelFont);
+	Set("NiagaraEditor.Stack.ItemHeaderValueLabelText", StackItemHeaderValueLabelText);
+
+	FSlateFontInfo StackItemHeaderValueFont = FAppStyle::Get().GetFontStyle(TEXT("PropertyWindow.NormalFont"));
+	FTextBlockStyle StackItemHeaderValueText = FTextBlockStyle(NormalText)
+		.SetFont(StackItemHeaderValueFont)
+		.SetColorAndOpacity(FLinearColor(1.f, 1.f, 1.f, 1.f));
+	Set("NiagaraEditor.Stack.ItemHeaderValueText", StackItemHeaderValueText);
 	
 	FTextBlockStyle OverviewStackItemText = FTextBlockStyle(CategoryText)
 		.SetFont(NormalFont);
@@ -644,5 +655,7 @@ void FNiagaraEditorWidgetsStyle::InitMisc()
 	);
 
 	Set("NiagaraEditor.Stack.SearchResult", new BOX_BRUSH("Icons/SearchResultBorder", FMargin(1.f/8.f)));
+
+	Set("NiagaraEditor.Stack.ItemHeaderValue.BackgroundBrush", new FSlateRoundedBoxBrush(FStyleColors::Panel, 6));
 }
 
diff --git a/Engine/Plugins/FX/Niagara/Source/NiagaraEditorWidgets/Private/SNiagaraOverviewGraph.cpp b/Engine/Plugins/FX/Niagara/Source/NiagaraEditorWidgets/Private/SNiagaraOverviewGraph.cpp
--- a/Engine/Plugins/FX/Niagara/Source/NiagaraEditorWidgets/Private/SNiagaraOverviewGraph.cpp
+++ b/Engine/Plugins/FX/Niagara/Source/NiagaraEditorWidgets/Private/SNiagaraOverviewGraph.cpp
@@ -16,6 +16,7 @@
 #include "NiagaraObjectSelection.h"
 #include "NiagaraOverviewGraphNodeFactory.h"
 #include "NiagaraOverviewNode.h"
+#include "NiagaraSettings.h"
 #include "NiagaraSystem.h"
 #include "NiagaraSystemEditorData.h"
 #include "NiagaraStackEditorData.h"
@@ -274,6 +275,14 @@ FActionMenuContent SNiagaraOverviewGraph::OnCreateGraphActionMenu(UEdGraph* InGr
 				LOCTEXT("AddEmitterToolTip", "Adds an empty emitter without any modules or renderers."),
 				FSlateIcon(),
 				FExecuteAction::CreateSP(this, &SNiagaraOverviewGraph::OnCreateEmptyEmitter));
+			if (GetDefault<UNiagaraSettings>()->bStatelessEmittersEnabled)
+			{
+				MenuBuilder.AddMenuEntry(
+					LOCTEXT("AddStatelessEmitter", "Add stateless emitter"),
+					LOCTEXT("AddStatelessEmitterToolTip", "Adds a stateless emitter without any modules or renderers."),
+					FSlateIcon(),
+					FExecuteAction::CreateSP(this, &SNiagaraOverviewGraph::OnCreateStatelessEmitter));
+			}
 			MenuBuilder.AddMenuEntry(
 				LOCTEXT("CommentsLabel", "Add Comment"),
 				LOCTEXT("AddCommentBoxToolTip", "Add a comment box"),
@@ -313,6 +322,11 @@ void SNiagaraOverviewGraph::OnCreateEmptyEmitter()
 	ViewModel->GetSystemViewModel()->AddEmptyEmitter();
 }
 
+void SNiagaraOverviewGraph::OnCreateStatelessEmitter()
+{
+	ViewModel->GetSystemViewModel()->AddStatelessEmitter();
+}
+
 void SNiagaraOverviewGraph::OnCreateComment()
 {
 	// Emitter assets have a transient overview graph, so any created comments would also be transient. We skip creating these comments instead.
diff --git a/Engine/Plugins/FX/Niagara/Source/NiagaraEditorWidgets/Private/SNiagaraOverviewGraph.h b/Engine/Plugins/FX/Niagara/Source/NiagaraEditorWidgets/Private/SNiagaraOverviewGraph.h
--- a/Engine/Plugins/FX/Niagara/Source/NiagaraEditorWidgets/Private/SNiagaraOverviewGraph.h
+++ b/Engine/Plugins/FX/Niagara/Source/NiagaraEditorWidgets/Private/SNiagaraOverviewGraph.h
@@ -36,6 +36,7 @@ private:
 	FActionMenuContent OnCreateGraphActionMenu(UEdGraph* InGraph, const FVector2D& InNodePosition, const TArray<UEdGraphPin*>& InDraggedPins, bool bAutoExpand, SGraphEditor::FActionMenuClosed InOnMenuClosed);
 
 	void OnCreateEmptyEmitter();
+	void OnCreateStatelessEmitter();
 	void OnCreateComment();
 	void OnClearIsolated();
 
diff --git a/Engine/Plugins/FX/Niagara/Source/NiagaraEditorWidgets/Private/SNiagaraOverviewStackNode.cpp b/Engine/Plugins/FX/Niagara/Source/NiagaraEditorWidgets/Private/SNiagaraOverviewStackNode.cpp
--- a/Engine/Plugins/FX/Niagara/Source/NiagaraEditorWidgets/Private/SNiagaraOverviewStackNode.cpp
+++ b/Engine/Plugins/FX/Niagara/Source/NiagaraEditorWidgets/Private/SNiagaraOverviewStackNode.cpp
@@ -484,6 +484,27 @@ TSharedRef<SWidget> SNiagaraOverviewStackNode::CreateTitleWidget_Default(TShared
 				DefaultTitle
 			];
 	}
+
+	if (!GetDefault<UNiagaraSettings>()->bStatelessEmittersEnabled)
+	{
+		TSharedPtr<FNiagaraEmitterHandleViewModel> EmitterHandleViewModel = EmitterHandleViewModelWeak.Pin();
+		FNiagaraEmitterHandle* EmitterHandle = EmitterHandleViewModel.IsValid() ? EmitterHandleViewModel->GetEmitterHandle() : nullptr;
+		if (EmitterHandle && EmitterHandle->GetEmitterMode() != ENiagaraEmitterMode::Standard)
+		{
+			DefaultTitle =
+				SNew(SVerticalBox)
+				+ SVerticalBox::Slot()
+				.Padding(0, 0, 5, 0)
+				[
+					SNew(STextBlock)
+						.Text(LOCTEXT("StatelessNotEnabled", "Stateless not enabled\nin project settings."))
+				]
+				+ SVerticalBox::Slot()
+				[
+					DefaultTitle
+				];
+		}
+	}
 	
 	TSharedPtr<SWidget> TitleWidget = SNew(SHorizontalBox)
 	// Summary View Controls
diff --git a/Engine/Plugins/FX/Niagara/Source/NiagaraEditorWidgets/Private/Stack/SNiagaraStackFunctionInputValue.cpp b/Engine/Plugins/FX/Niagara/Source/NiagaraEditorWidgets/Private/Stack/SNiagaraStackFunctionInputValue.cpp
--- a/Engine/Plugins/FX/Niagara/Source/NiagaraEditorWidgets/Private/Stack/SNiagaraStackFunctionInputValue.cpp
+++ b/Engine/Plugins/FX/Niagara/Source/NiagaraEditorWidgets/Private/Stack/SNiagaraStackFunctionInputValue.cpp
@@ -71,6 +71,10 @@ void SNiagaraStackFunctionInputValue::Construct(const FArguments& InArgs, UNiaga
 	FunctionInput->OnValueChanged().AddSP(this, &SNiagaraStackFunctionInputValue::OnInputValueChanged);
 	SyntaxHighlighter = FNiagaraHLSLSyntaxHighlighter::Create();
 
+	TAttribute<bool> EntryIsEnabled;
+	EntryIsEnabled.Bind(this, &SNiagaraStackFunctionInputValue::GetEntryEnabled);
+	SetEnabled(EntryIsEnabled);
+
 	TSharedPtr<SHorizontalBox> OuterChildrenBox;
 	TSharedPtr<SHorizontalBox> ChildrenBox;
 	ChildSlot
@@ -548,7 +552,12 @@ void SNiagaraStackFunctionInputValue::SetToLocalValue()
 
 bool SNiagaraStackFunctionInputValue::GetInputEnabled() const
 {
-	return FunctionInput->GetHasEditCondition() == false || FunctionInput->GetEditConditionEnabled();
+	return FunctionInput->IsFinalized() == false && (FunctionInput->GetHasEditCondition() == false || FunctionInput->GetEditConditionEnabled());
+}
+
+bool SNiagaraStackFunctionInputValue::GetEntryEnabled() const
+{
+	return FunctionInput->IsFinalized() == false && FunctionInput->GetIsEnabledAndOwnerIsEnabled();
 }
 
 TSharedRef<SWidget> SNiagaraStackFunctionInputValue::ConstructLocalValueStructWidget()
diff --git a/Engine/Plugins/FX/Niagara/Source/NiagaraEditorWidgets/Private/Stack/SNiagaraStackFunctionInputValue.h b/Engine/Plugins/FX/Niagara/Source/NiagaraEditorWidgets/Private/Stack/SNiagaraStackFunctionInputValue.h
--- a/Engine/Plugins/FX/Niagara/Source/NiagaraEditorWidgets/Private/Stack/SNiagaraStackFunctionInputValue.h
+++ b/Engine/Plugins/FX/Niagara/Source/NiagaraEditorWidgets/Private/Stack/SNiagaraStackFunctionInputValue.h
@@ -53,6 +53,8 @@ private:
 
 	bool GetInputEnabled() const;
 
+	bool GetEntryEnabled() const;
+
 	TSharedRef<SWidget> ConstructLocalValueStructWidget();
 
 	void OnInputValueChanged();
diff --git a/Engine/Plugins/FX/Niagara/Source/NiagaraEditorWidgets/Private/Stack/SNiagaraStackItem.cpp b/Engine/Plugins/FX/Niagara/Source/NiagaraEditorWidgets/Private/Stack/SNiagaraStackItem.cpp
--- a/Engine/Plugins/FX/Niagara/Source/NiagaraEditorWidgets/Private/Stack/SNiagaraStackItem.cpp
+++ b/Engine/Plugins/FX/Niagara/Source/NiagaraEditorWidgets/Private/Stack/SNiagaraStackItem.cpp
@@ -3,24 +3,177 @@
 #include "Stack/SNiagaraStackItem.h"
 
 #include "EditorFontGlyphs.h"
-#include "NiagaraEditorWidgetsStyle.h"
 #include "NiagaraEditorStyle.h"
+#include "NiagaraEditorWidgetsStyle.h"
+#include "SEnumCombo.h"
+#include "Stack/SNiagaraStackInheritanceIcon.h"
 #include "Styling/AppStyle.h"
+#include "Styling/StyleColors.h"
+#include "ViewModels/Stack/NiagaraStackClipboardUtilities.h"
 #include "ViewModels/Stack/NiagaraStackItem.h"
 #include "ViewModels/Stack/NiagaraStackViewModel.h"
-#include "Widgets/SBoxPanel.h"
+#include "Widgets/Images/SImage.h"
 #include "Widgets/Input/SButton.h"
 #include "Widgets/Input/SCheckBox.h"
-#include "Widgets/Images/SImage.h"
-#include "Widgets/Text/STextBlock.h"
-#include "Widgets/Images/SImage.h"
 #include "Widgets/Layout/SBox.h"
-#include "ViewModels/Stack/NiagaraStackClipboardUtilities.h"
-#include "Styling/StyleColors.h"
-#include "Stack/SNiagaraStackInheritanceIcon.h"
+#include "Widgets/Layout/SSpacer.h"
+#include "Widgets/SBoxPanel.h"
+#include "Widgets/Text/STextBlock.h"
 
 #define LOCTEXT_NAMESPACE "NiagaraStackItem"
 
+class SNiagaraStackItemHeaderValue : public SCompoundWidget
+{
+public:
+	SLATE_BEGIN_ARGS(SNiagaraStackItemHeaderValue) { }
+	SLATE_END_ARGS();
+
+	void Construct(const FArguments& InArgs, TSharedRef<INiagaraStackItemHeaderValueHandler> InHeaderValueHandler)
+	{
+		HeaderValueHandler = InHeaderValueHandler;
+
+		TSharedPtr<SWidget> ChildWidget;
+		if (HeaderValueHandler->GetMode() == INiagaraStackItemHeaderValueHandler::EValueMode::BoolToggle)
+		{
+			ChildWidget = ConstructBoolToggle();
+		}
+		else if (HeaderValueHandler->GetMode() == INiagaraStackItemHeaderValueHandler::EValueMode::EnumDropDown)
+		{
+			ChildWidget = ConstructEnumDropDown();
+		}
+		else
+		{
+			ChildWidget = SNullWidget::NullWidget;
+		}
+
+		ChildSlot
+		[
+			ChildWidget.ToSharedRef()
+		];
+	}
+
+private:
+	TSharedRef<SWidget> ConstructEnumDropDown()
+	{
+		TSharedRef<SWidget> EnumWidget = SNew(SEnumComboBox, HeaderValueHandler->GetEnum())
+			.ComboButtonStyle(&FNiagaraEditorWidgetsStyle::Get().GetWidgetStyle<FComboButtonStyle>("NiagaraEditor.Stack.CompactComboButton"))
+			.CurrentValue(HeaderValueHandler.ToSharedRef(), &INiagaraStackItemHeaderValueHandler::GetEnumValue)
+			.ContentPadding(FMargin(2, 0))
+			.TextStyle(FNiagaraEditorWidgetsStyle::Get(), "NiagaraEditor.Stack.ItemHeaderValueText")
+			.OnEnumSelectionChanged(SEnumComboBox::FOnEnumSelectionChanged::CreateSP(this, &SNiagaraStackItemHeaderValue::EnumValueChanged));
+
+		TSharedPtr<SWidget> IconWidget;
+		if (HeaderValueHandler->GetIconBrush() != nullptr)
+		{
+			IconWidget = SNew(SImage)
+				.Image(HeaderValueHandler->GetIconBrush())
+				.DesiredSizeOverride(FVector2D(14.0f, 14.0f));
+		}
+
+		TSharedPtr<SWidget> LabelWidget;
+		if (HeaderValueHandler->GetLabelText().IsEmpty() == false)
+		{
+			LabelWidget = SNew(STextBlock)
+				.TextStyle(FNiagaraEditorWidgetsStyle::Get(), "NiagaraEditor.Stack.ItemHeaderValueLabelText")
+				.Text(HeaderValueHandler->GetLabelText());
+		}
+
+		if (IconWidget.IsValid() || LabelWidget.IsValid())
+		{
+			TSharedRef<SHorizontalBox> ContentBox = SNew(SHorizontalBox);
+			if (IconWidget.IsValid())
+			{
+				ContentBox->AddSlot()
+					.VAlign(VAlign_Center)
+					.AutoWidth()
+					.Padding(0, 0, 5, 0)
+					[
+						IconWidget.ToSharedRef()
+					];
+			}
+
+			if (LabelWidget.IsValid())
+			{
+				ContentBox->AddSlot()
+					.VAlign(VAlign_Center)
+					.AutoWidth()
+					.Padding(0, 0, 5, 0)
+					[
+						LabelWidget.ToSharedRef()
+					];
+			}
+
+			ContentBox->AddSlot()
+				.AutoWidth()
+				[
+					EnumWidget
+				];
+
+			return SNew(SBorder)
+				.BorderImage(FNiagaraEditorWidgetsStyle::Get().GetBrush(TEXT("NiagaraEditor.Stack.ItemHeaderValue.BackgroundBrush")))
+				.Padding(7, 2, 2, 2)
+				[
+					ContentBox
+				];
+		}
+		else
+		{
+			return EnumWidget;
+		}
+	}
+
+	TSharedRef<SWidget> ConstructBoolToggle()
+	{
+		if (HeaderValueHandler->GetIconBrush() != nullptr)
+		{
+			return SNew(SCheckBox)
+				.Style(FAppStyle::Get(), "DetailsView.SectionButton")
+				.ToolTipText(HeaderValueHandler->GetLabelText())
+				.OnCheckStateChanged(this, &SNiagaraStackItemHeaderValue::BoolCheckStateChanged)
+				.IsChecked(this, &SNiagaraStackItemHeaderValue::GetBoolCheckState)
+				.Padding(FMargin(4, 2))
+				[
+					SNew(SImage)
+					.Image(HeaderValueHandler->GetIconBrush())
+					.DesiredSizeOverride(FVector2D(14.0f, 14.0f))
+				];
+		}
+
+		if (HeaderValueHandler->GetLabelText().IsEmpty() == false)
+		{
+			return SNew(SCheckBox)
+				.Style(FAppStyle::Get(), "DetailsView.SectionButton")
+				.OnCheckStateChanged(this, &SNiagaraStackItemHeaderValue::BoolCheckStateChanged)
+				.IsChecked(this, &SNiagaraStackItemHeaderValue::GetBoolCheckState)
+				[
+					SNew(STextBlock)
+					.Text(HeaderValueHandler->GetLabelText())
+					.TextStyle(FNiagaraEditorStyle::Get(), "NiagaraEditor.ParameterText")
+				];
+		}
+
+		return SNullWidget::NullWidget;
+	}
+
+	void EnumValueChanged(int32 Value, ESelectInfo::Type SelectionType)
+	{
+		HeaderValueHandler->NotifyEnumValueChanged(Value);
+	}
+
+	ECheckBoxState GetBoolCheckState() const
+	{
+		return HeaderValueHandler->GetBoolValue() ? ECheckBoxState::Checked : ECheckBoxState::Unchecked;
+	}
+
+	void BoolCheckStateChanged(ECheckBoxState InCheckState)
+	{
+		HeaderValueHandler->NotifyBoolValueChanged(InCheckState == ECheckBoxState::Checked);
+	}
+
+private:
+	TSharedPtr<INiagaraStackItemHeaderValueHandler> HeaderValueHandler;
+};
+
 void SNiagaraStackItem::Construct(const FArguments& InArgs, UNiagaraStackItem& InItem, UNiagaraStackViewModel* InStackViewModel)
 {
 	Item = &InItem;
@@ -58,7 +211,8 @@ void SNiagaraStackItem::Construct(const FArguments& InArgs, UNiagaraStackItem& I
 
 	// Display name
 	RowBox->AddSlot()
-		.Padding(2, 0, 2, 0)
+		.Padding(2, 0, 10, 0)
+		.AutoWidth()
 		.VAlign(VAlign_Center)
 		[
 			SAssignNew(DisplayNameWidget, SNiagaraStackDisplayName, InItem, *InStackViewModel)
@@ -66,6 +220,9 @@ void SNiagaraStackItem::Construct(const FArguments& InArgs, UNiagaraStackItem& I
 			.EditableNameStyle(FNiagaraEditorWidgetsStyle::Get(), "NiagaraEditor.Stack.EditableItemText")
 		];
 
+	// Entry based header values.
+	ConstructHeaderValues(RowBox);
+
 	// Allow derived classes to add additional widgets.
 	AddCustomRowWidgets(RowBox);
 
@@ -204,6 +361,68 @@ TSharedRef<SWidget> SNiagaraStackItem::AddContainerForRowWidgets(TSharedRef<SWid
 	return RowWidgets;
 }
 
+void SNiagaraStackItem::ConstructHeaderValues(TSharedRef<SHorizontalBox> HorizontalBox)
+{
+	TArray<TSharedRef<INiagaraStackItemHeaderValueHandler>> HeaderValueHandlers;
+
+	if (Item->SupportsHeaderValues())
+	{
+		Item->GetHeaderValueHandlers(HeaderValueHandlers);
+	}
+
+	TArray<TSharedRef<SWidget>> AlignLeft;
+	TArray<TSharedRef<SWidget>> AlignCenterAlignFill;
+	TArray<TSharedRef<SWidget>> AlignRight;
+
+	for (TSharedRef<INiagaraStackItemHeaderValueHandler> HeaderValueHandler : HeaderValueHandlers)
+	{
+		TSharedRef<SWidget> HeaderValueWidget = SNew(SNiagaraStackItemHeaderValue, HeaderValueHandler)
+			.IsEnabled(this, &SNiagaraStackItem::GetEntryIsEnabled);
+		if (HeaderValueHandler->GetHAlign() == HAlign_Left)
+		{
+			AlignLeft.Add(HeaderValueWidget);
+		}
+		else if (HeaderValueHandler->GetHAlign() == HAlign_Right)
+		{
+			AlignRight.Add(HeaderValueWidget);
+		}
+		else
+		{
+			AlignCenterAlignFill.Add_GetRef(HeaderValueWidget);
+		}
+	}
+
+	if (AlignCenterAlignFill.Num() == 0)
+	{
+		AlignCenterAlignFill.Add(SNew(SSpacer));
+	}
+
+	auto AddWidgets = [](TSharedRef<SHorizontalBox> HorizontalBox, TArray<TSharedRef<SWidget>> Widgets, EHorizontalAlignment HorizontalAlignment, bool bAutoWidth)
+	{
+		for (TSharedRef<SWidget> Widget : Widgets)
+		{
+			SHorizontalBox::FScopedWidgetSlotArguments Slot = HorizontalBox->AddSlot();
+
+			Slot
+				.HAlign(HorizontalAlignment)
+				.VAlign(VAlign_Center)
+				.Padding(5, 0)
+				[
+					Widget
+				];
+			
+			if (bAutoWidth)
+			{
+				Slot.AutoWidth();
+			}
+		}
+	};
+
+	AddWidgets(HorizontalBox, AlignLeft, HAlign_Left, true);
+	AddWidgets(HorizontalBox, AlignCenterAlignFill, HAlign_Fill, false);
+	AddWidgets(HorizontalBox, AlignRight, HAlign_Right, true);
+}
+
 FText SNiagaraStackItem::GetEditModeButtonToolTipText() const
 {
 	return Item->GetEditModeButtonTooltip().IsSet() ? Item->GetEditModeButtonTooltip().GetValue() : LOCTEXT("EditButtonDefaultTooltip", "Activate Edit Mode");
@@ -280,7 +499,12 @@ ECheckBoxState SNiagaraStackItem::CheckEnabledStatus() const
 
 bool SNiagaraStackItem::GetEnabledCheckBoxEnabled() const
 {
-	return Item->GetOwnerIsEnabled();
+	return Item->IsFinalized() == false && Item->GetOwnerIsEnabled();
+}
+
+bool SNiagaraStackItem::GetEntryIsEnabled() const
+{
+	return Item->IsFinalized() == false && Item->GetIsEnabledAndOwnerIsEnabled();
 }
 
 #undef LOCTEXT_NAMESPACE
\ No newline at end of file
diff --git a/Engine/Plugins/FX/Niagara/Source/NiagaraEditorWidgets/Private/Stack/SNiagaraStackItem.h b/Engine/Plugins/FX/Niagara/Source/NiagaraEditorWidgets/Private/Stack/SNiagaraStackItem.h
--- a/Engine/Plugins/FX/Niagara/Source/NiagaraEditorWidgets/Private/Stack/SNiagaraStackItem.h
+++ b/Engine/Plugins/FX/Niagara/Source/NiagaraEditorWidgets/Private/Stack/SNiagaraStackItem.h
@@ -25,6 +25,8 @@ protected:
 	virtual TSharedRef<SWidget> AddContainerForRowWidgets(TSharedRef<SWidget> RowWidgets);
 
 private:
+	void ConstructHeaderValues(TSharedRef<SHorizontalBox> HorizontalBox);
+
 	FReply EditModeButtonClicked();
 	EVisibility IsEditButtonVisible() const;
 	FText GetEditModeButtonToolTipText() const;
@@ -47,6 +49,8 @@ private:
 
 	bool GetEnabledCheckBoxEnabled() const;
 
+	bool GetEntryIsEnabled() const;
+
 private:
 	UNiagaraStackItem* Item;
 	TSharedPtr<SNiagaraStackDisplayName> DisplayNameWidget;
diff --git a/Engine/Plugins/FX/Niagara/Source/NiagaraShader/Internal/Stateless/NiagaraStatelessEmitterShaders.h b/Engine/Plugins/FX/Niagara/Source/NiagaraShader/Internal/Stateless/NiagaraStatelessEmitterShaders.h
new file mode 100644
--- /dev/null
+++ b/Engine/Plugins/FX/Niagara/Source/NiagaraShader/Internal/Stateless/NiagaraStatelessEmitterShaders.h
@@ -0,0 +1,99 @@
+// Copyright Epic Games, Inc. All Rights Reserved.
+
+#pragma once
+
+#include "NiagaraCommon.h"
+#include "NiagaraStatelessSimulationShader.h"
+
+#include "NiagaraStatelessModuleShaderParameters.h"
+
+#include "RenderGraphFwd.h"
+#include "ShaderParameterStruct.h"
+
+namespace NiagaraStateless
+{
+	class FSimulationShaderDefaultCS : public FSimulationShader
+	{
+	public:
+		DECLARE_EXPORTED_GLOBAL_SHADER(FSimulationShaderDefaultCS, NIAGARASHADER_API);
+		SHADER_USE_PARAMETER_STRUCT(FSimulationShaderDefaultCS, FSimulationShader);
+
+		//BEGIN_SHADER_PARAMETER_STRUCT(FOutputParameters, NIAGARASHADER_API)
+		//	SHADER_PARAMETER(int,	Output_PositionComponent)
+		//	SHADER_PARAMETER(int,	Output_ColorComponent)
+		//	SHADER_PARAMETER(int,	Output_SpriteSizeComponent)
+		//	SHADER_PARAMETER(int,	Output_PreviousPositionComponent)
+		//	SHADER_PARAMETER(int,	Output_PreviousSpriteSizeComponent)
+		//END_SHADER_PARAMETER_STRUCT()
+
+		BEGIN_SHADER_PARAMETER_STRUCT(FParameters, NIAGARASHADER_API)
+			SHADER_PARAMETER_STRUCT_INCLUDE(FCommonShaderParameters,							CommonParameters)
+			SHADER_PARAMETER_STRUCT_INCLUDE(FInitializeParticleModule_ShaderParameters,			Module0Params)
+			SHADER_PARAMETER_STRUCT_INCLUDE(FInitialMeshOrientationModule_ShaderParameters,		Module1Params)
+			SHADER_PARAMETER_STRUCT_INCLUDE(FShapeLocationModule_ShaderParameters,				Module2Params)
+			SHADER_PARAMETER_STRUCT_INCLUDE(FScaleColorModule_ShaderParameters,					Module3Params)
+			SHADER_PARAMETER_STRUCT_INCLUDE(FScaleSpriteSizeModule_ShaderParameters,			Module4Params)
+			SHADER_PARAMETER_STRUCT_INCLUDE(FScaleSpriteSizeBySpeedModule_ShaderParameters,		Module5Params)
+			SHADER_PARAMETER_STRUCT_INCLUDE(FScaleMeshSizeModule_ShaderParameters,				Module6Params)
+			SHADER_PARAMETER_STRUCT_INCLUDE(FScaleMeshSizeBySpeedModule_ShaderParameters,		Module7Params)
+			SHADER_PARAMETER_STRUCT_INCLUDE(FSolveVelocitiesAndForcesModule_ShaderParameters,	Module8Params)
+			SHADER_PARAMETER_STRUCT_INCLUDE(FSpriteFacingAndAlignmentModule_ShaderParameters,	Module9Params)
+			SHADER_PARAMETER_STRUCT_INCLUDE(FSpriteRotationRateModule_ShaderParameters,			Module10Params)
+			SHADER_PARAMETER_STRUCT_INCLUDE(FSubUVAnimationModule_ShaderParameters,				Module11Params)
+
+			SHADER_PARAMETER(int, Permutation_UniqueIDComponent)
+			SHADER_PARAMETER(int, Permutation_PositionComponent)
+			SHADER_PARAMETER(int, Permutation_ColorComponent)
+			SHADER_PARAMETER(int, Permutation_MeshOrientationComponent)
+			SHADER_PARAMETER(int, Permutation_RibbonWidthComponent)
+			SHADER_PARAMETER(int, Permutation_ScaleComponent)
+			SHADER_PARAMETER(int, Permutation_SpriteSizeComponent)
+			SHADER_PARAMETER(int, Permutation_SpriteFacingComponent)
+			SHADER_PARAMETER(int, Permutation_SpriteAlignmentComponent)
+			SHADER_PARAMETER(int, Permutation_SpriteRotationComponent)
+			SHADER_PARAMETER(int, Permutation_SubImageIndexComponent)
+			SHADER_PARAMETER(int, Permutation_VelocityComponent)
+			SHADER_PARAMETER(int, Permutation_PreviousPositionComponent)
+			SHADER_PARAMETER(int, Permutation_PreviousMeshOrientationComponent)
+			SHADER_PARAMETER(int, Permutation_PreviousRibbonWidthComponent)
+			SHADER_PARAMETER(int, Permutation_PreviousScaleComponent)
+			SHADER_PARAMETER(int, Permutation_PreviousSpriteSizeComponent)
+			SHADER_PARAMETER(int, Permutation_PreviousSpriteFacingComponent)
+			SHADER_PARAMETER(int, Permutation_PreviousSpriteAlignmentComponent)
+			SHADER_PARAMETER(int, Permutation_PreviousSpriteRotationComponent)
+			SHADER_PARAMETER(int, Permutation_PreviousVelocityComponent)
+		END_SHADER_PARAMETER_STRUCT()
+	};
+
+	class FSimulationShaderExample1CS : public FSimulationShader
+	{
+	public:
+		DECLARE_EXPORTED_GLOBAL_SHADER(FSimulationShaderExample1CS, NIAGARASHADER_API);
+		SHADER_USE_PARAMETER_STRUCT(FSimulationShaderExample1CS, FSimulationShader);
+
+		BEGIN_SHADER_PARAMETER_STRUCT(FParameters, NIAGARASHADER_API)
+			SHADER_PARAMETER_STRUCT_INCLUDE(FCommonShaderParameters, CommonParameters)
+			SHADER_PARAMETER_STRUCT_INCLUDE(FInitializeParticleModule_ShaderParameters,			Module0Params)
+			SHADER_PARAMETER_STRUCT_INCLUDE(FShapeLocationModule_ShaderParameters,				Module1Params)
+			SHADER_PARAMETER_STRUCT_INCLUDE(FScaleColorModule_ShaderParameters,					Module2Params)
+			SHADER_PARAMETER_STRUCT_INCLUDE(FScaleSpriteSizeModule_ShaderParameters,			Module3Params)
+			SHADER_PARAMETER_STRUCT_INCLUDE(FRotateAroundPointModule_ShaderParameters,			Module4Params)
+			SHADER_PARAMETER_STRUCT_INCLUDE(FSolveVelocitiesAndForcesModule_ShaderParameters,	Module5Params)
+
+			SHADER_PARAMETER(int, Permutation_UniqueIDComponent)
+			SHADER_PARAMETER(int, Permutation_PositionComponent)
+			SHADER_PARAMETER(int, Permutation_ColorComponent)
+			SHADER_PARAMETER(int, Permutation_ScaleComponent)
+			SHADER_PARAMETER(int, Permutation_RibbonWidthComponent)
+			SHADER_PARAMETER(int, Permutation_SpriteSizeComponent)
+			SHADER_PARAMETER(int, Permutation_SpriteRotationComponent)
+			SHADER_PARAMETER(int, Permutation_VelocityComponent)
+			SHADER_PARAMETER(int, Permutation_PreviousPositionComponent)
+			SHADER_PARAMETER(int, Permutation_PreviousRibbonWidthComponent)
+			SHADER_PARAMETER(int, Permutation_PreviousScaleComponent)
+			SHADER_PARAMETER(int, Permutation_PreviousSpriteSizeComponent)
+			SHADER_PARAMETER(int, Permutation_PreviousSpriteRotationComponent)
+			SHADER_PARAMETER(int, Permutation_PreviousVelocityComponent)
+		END_SHADER_PARAMETER_STRUCT()
+	};
+}
diff --git a/Engine/Plugins/FX/Niagara/Source/NiagaraShader/Internal/Stateless/NiagaraStatelessModuleShaderParameters.h b/Engine/Plugins/FX/Niagara/Source/NiagaraShader/Internal/Stateless/NiagaraStatelessModuleShaderParameters.h
new file mode 100644
--- /dev/null
+++ b/Engine/Plugins/FX/Niagara/Source/NiagaraShader/Internal/Stateless/NiagaraStatelessModuleShaderParameters.h
@@ -0,0 +1,142 @@
+// Copyright Epic Games, Inc. All Rights Reserved.
+
+#pragma once
+
+#include "RenderGraphFwd.h"
+#include "ShaderParameterStruct.h"
+
+namespace NiagaraStateless
+{
+	BEGIN_SHADER_PARAMETER_STRUCT(FInitializeParticleModule_ShaderParameters, NIAGARASHADER_API)
+		SHADER_PARAMETER(FVector3f, InitializeParticle_Position)				// Unset / Direct Set / Simulation Position
+
+		SHADER_PARAMETER(FVector4f, InitializeParticle_ColorScale)				// Unset / Direct Set / Random Range(Link RGBA / Link RGB|Link A / Random Channels) / Random Hue|Saturation|Value
+		SHADER_PARAMETER(FVector4f, InitializeParticle_ColorBias)
+
+		SHADER_PARAMETER(FVector2f, InitializeParticle_SpriteSizeScale)			// Unset / Uniform / Random Uniform / Non-Uniform / Random Non-Uniform
+		SHADER_PARAMETER(FVector2f, InitializeParticle_SpriteSizeBias)
+		SHADER_PARAMETER(float, InitializeParticle_SpriteRotationScale)				// Unset / Random / Direct Set Deg / Direct Set Normalized
+		SHADER_PARAMETER(float, InitializeParticle_SpriteRotationBias)
+		//SHADER_PARAMETER(FVector3f, InitializeParticle_SpriteUVMode)			// Unset / Random / Random X / Random Y / Random XY / Direct Set
+
+		SHADER_PARAMETER(FVector3f, InitializeParticle_MeshScaleScale)			// Unset / Uniform / Random Uniform / Non-Uniform / Random Non-Uniform
+		SHADER_PARAMETER(FVector3f, InitializeParticle_MeshScaleBias)
+
+		SHADER_PARAMETER(float, InitializeParticle_RibbonWidthScale)		// Unset / Direct Set
+		SHADER_PARAMETER(float, InitializeParticle_RibbonWidthBias)
+		//SHADER_PARAMETER(FVector3f, InitializeParticle_RibbonFacingVector)		// Unset / Direct Set
+		//SHADER_PARAMETER(FVector3f, InitializeParticle_RibbonTwist)				// Unset / Direct Set
+
+		SHADER_PARAMETER(uint32,	InitializeParticle_PositionMode)
+		SHADER_PARAMETER(uint32,	InitializeParticle_PositionTableMod)
+		SHADER_PARAMETER(uint32,	InitializeParticle_PositionTableOffset)
+		SHADER_PARAMETER(uint32,	InitializeParticle_PositionTableLength)
+	END_SHADER_PARAMETER_STRUCT()
+
+	BEGIN_SHADER_PARAMETER_STRUCT(FInitialMeshOrientationModule_ShaderParameters, NIAGARASHADER_API)
+		SHADER_PARAMETER(FVector3f,	InitialMeshOrientation_Rotation)
+		SHADER_PARAMETER(FVector3f,	InitialMeshOrientation_RandomRangeScale)
+	END_SHADER_PARAMETER_STRUCT()
+		
+
+	BEGIN_SHADER_PARAMETER_STRUCT(FRotateAroundPointModule_ShaderParameters, NIAGARASHADER_API)
+		SHADER_PARAMETER(float,	RotateAroundPoint_RateScale)
+		SHADER_PARAMETER(float,	RotateAroundPoint_RateBias)
+		SHADER_PARAMETER(float,	RotateAroundPoint_RadiusScale)
+		SHADER_PARAMETER(float,	RotateAroundPoint_RadiusBias)
+		SHADER_PARAMETER(float,	RotateAroundPoint_InitialPhaseScale)
+		SHADER_PARAMETER(float,	RotateAroundPoint_InitialPhaseBias)
+	END_SHADER_PARAMETER_STRUCT()
+	
+	BEGIN_SHADER_PARAMETER_STRUCT(FScaleColorModule_ShaderParameters, NIAGARASHADER_API)
+		SHADER_PARAMETER(int32, ScaleColor_Offset)
+		SHADER_PARAMETER(int32, ScaleColor_Length)
+	END_SHADER_PARAMETER_STRUCT()
+
+	BEGIN_SHADER_PARAMETER_STRUCT(FScaleMeshSizeModule_ShaderParameters, NIAGARASHADER_API)
+		SHADER_PARAMETER(int32,		ScaleMeshSize_Offset)
+		SHADER_PARAMETER(int32,		ScaleMeshSize_Length)
+	END_SHADER_PARAMETER_STRUCT()
+
+	BEGIN_SHADER_PARAMETER_STRUCT(FScaleMeshSizeBySpeedModule_ShaderParameters, NIAGARASHADER_API)
+		SHADER_PARAMETER(int32,		ScaleMeshSizeBySpeed_Offset)
+		SHADER_PARAMETER(float,		ScaleMeshSizeBySpeed_VelocityNorm)
+		SHADER_PARAMETER(int32,		ScaleMeshSizeBySpeed_Length)
+	END_SHADER_PARAMETER_STRUCT()
+
+	BEGIN_SHADER_PARAMETER_STRUCT(FScaleSpriteSizeModule_ShaderParameters, NIAGARASHADER_API)
+		SHADER_PARAMETER(int32,		ScaleSpriteSize_Offset)
+		SHADER_PARAMETER(int32,		ScaleSpriteSize_Length)
+	END_SHADER_PARAMETER_STRUCT()
+
+	BEGIN_SHADER_PARAMETER_STRUCT(FScaleSpriteSizeBySpeedModule_ShaderParameters, NIAGARASHADER_API)
+		SHADER_PARAMETER(float,		ScaleSpriteSizeBySpeed_VelocityNorm)
+		SHADER_PARAMETER(int32,		ScaleSpriteSizeBySpeed_Offset)
+		SHADER_PARAMETER(int32,		ScaleSpriteSizeBySpeed_Length)
+	END_SHADER_PARAMETER_STRUCT()
+
+	BEGIN_SHADER_PARAMETER_STRUCT(FSpriteFacingAndAlignmentModule_ShaderParameters, NIAGARASHADER_API)
+		SHADER_PARAMETER(FVector3f,	SpriteFacingAndAlignment_SpriteFacing)
+		SHADER_PARAMETER(FVector3f,	SpriteFacingAndAlignment_SpriteAlignment)
+	END_SHADER_PARAMETER_STRUCT()
+
+	BEGIN_SHADER_PARAMETER_STRUCT(FSpriteRotationRateModule_ShaderParameters, NIAGARASHADER_API)
+		SHADER_PARAMETER(float,		SpriteRotationRate_Scale)
+		SHADER_PARAMETER(float,		SpriteRotationRate_Bias)
+	END_SHADER_PARAMETER_STRUCT()
+
+	BEGIN_SHADER_PARAMETER_STRUCT(FSubUVAnimationModule_ShaderParameters, NIAGARASHADER_API)
+		SHADER_PARAMETER(int,		SubUVAnimation_Mode)
+		SHADER_PARAMETER(float,		SubUVAnimation_NumFrames)
+		SHADER_PARAMETER(float,		SubUVAnimation_InitialFrameScale)
+		SHADER_PARAMETER(float,		SubUVAnimation_InitialFrameBias)
+		SHADER_PARAMETER(float,		SubUVAnimation_InitialFrameRateChange)
+		SHADER_PARAMETER(float,		SubUVAnimation_AnimFrameStart)
+		SHADER_PARAMETER(float,		SubUVAnimation_AnimFrameRange)
+		SHADER_PARAMETER(float,		SubUVAnimation_RateScale)
+	END_SHADER_PARAMETER_STRUCT()
+
+	BEGIN_SHADER_PARAMETER_STRUCT(FShapeLocationModule_ShaderParameters, NIAGARASHADER_API)
+		SHADER_PARAMETER(FUintVector4,	ShapeLocation_Mode)
+		SHADER_PARAMETER(FVector4f,		ShapeLocation_Parameters0)
+		SHADER_PARAMETER(FVector4f,		ShapeLocation_Parameters1)
+	END_SHADER_PARAMETER_STRUCT()
+
+	BEGIN_SHADER_PARAMETER_STRUCT(FSolveVelocitiesAndForcesModule_ShaderParameters, NIAGARASHADER_API)
+		SHADER_PARAMETER(float,			SolveVelocitiesAndForces_MassScale)
+		SHADER_PARAMETER(float,			SolveVelocitiesAndForces_MassBias)
+		SHADER_PARAMETER(float,			SolveVelocitiesAndForces_DragScale)
+		SHADER_PARAMETER(float,			SolveVelocitiesAndForces_DragBias)
+		SHADER_PARAMETER(FVector3f,		SolveVelocitiesAndForces_VelocityScale)
+		SHADER_PARAMETER(FVector3f,		SolveVelocitiesAndForces_VelocityBias)
+		SHADER_PARAMETER(FVector3f,		SolveVelocitiesAndForces_WindScale)
+		SHADER_PARAMETER(FVector3f,		SolveVelocitiesAndForces_WindBias)
+		SHADER_PARAMETER(FVector3f,		SolveVelocitiesAndForces_AccelerationScale)
+		SHADER_PARAMETER(FVector3f,		SolveVelocitiesAndForces_AccelerationBias)
+
+		SHADER_PARAMETER(uint32,		SolveVelocitiesAndForces_ConeVelocityEnabled)
+		SHADER_PARAMETER(FQuat4f,		SolveVelocitiesAndForces_ConeQuat)
+		SHADER_PARAMETER(float,			SolveVelocitiesAndForces_ConeVelocityScale)
+		SHADER_PARAMETER(float,			SolveVelocitiesAndForces_ConeVelocityBias)
+		SHADER_PARAMETER(float,			SolveVelocitiesAndForces_ConeAngleScale)
+		SHADER_PARAMETER(float,			SolveVelocitiesAndForces_ConeAngleBias)
+		SHADER_PARAMETER(float,			SolveVelocitiesAndForces_ConeVelocityFalloff)
+
+		SHADER_PARAMETER(uint32,		SolveVelocitiesAndForces_PontVelocityEnabled)
+		SHADER_PARAMETER(float,			SolveVelocitiesAndForces_PointVelocityScale)
+		SHADER_PARAMETER(float,			SolveVelocitiesAndForces_PointVelocityBias)
+		SHADER_PARAMETER(FVector3f,		SolveVelocitiesAndForces_PointOrigin)
+
+		SHADER_PARAMETER(uint32,		SolveVelocitiesAndForces_NoiseEnabled)
+		SHADER_PARAMETER(float,			SolveVelocitiesAndForces_NoiseAmplitude)
+		SHADER_PARAMETER(FVector3f,		SolveVelocitiesAndForces_NoiseFrequency)
+		SHADER_PARAMETER_TEXTURE(Texture3D,		SolveVelocitiesAndForces_NoiseTexture)
+		SHADER_PARAMETER_SAMPLER(SamplerState,	SolveVelocitiesAndForces_NoiseSampler)
+
+		SHADER_PARAMETER(uint32, SolveVelocitiesAndForces_NoiseMode)
+		SHADER_PARAMETER(uint32, SolveVelocitiesAndForces_NoiseLUTOffset)
+		SHADER_PARAMETER(uint32, SolveVelocitiesAndForces_NoiseLUTNumChannel)
+		SHADER_PARAMETER(uint32, SolveVelocitiesAndForces_NoiseLUTChannelWidth)
+
+	END_SHADER_PARAMETER_STRUCT()
+}
diff --git a/Engine/Plugins/FX/Niagara/Source/NiagaraShader/Private/Stateless/NiagaraStatelessEmitterShaders.cpp b/Engine/Plugins/FX/Niagara/Source/NiagaraShader/Private/Stateless/NiagaraStatelessEmitterShaders.cpp
new file mode 100644
--- /dev/null
+++ b/Engine/Plugins/FX/Niagara/Source/NiagaraShader/Private/Stateless/NiagaraStatelessEmitterShaders.cpp
@@ -0,0 +1,9 @@
+// Copyright Epic Games, Inc. All Rights Reserved.
+
+#include "Stateless/NiagaraStatelessEmitterShaders.h"
+
+namespace NiagaraStateless
+{
+	IMPLEMENT_GLOBAL_SHADER(FSimulationShaderDefaultCS, "/Plugin/FX/Niagara/Private/Stateless/NiagaraStatelessSimulationDefault.usf", "StatelessMain", SF_Compute);
+	IMPLEMENT_GLOBAL_SHADER(FSimulationShaderExample1CS, "/Plugin/FX/Niagara/Private/Stateless/NiagaraStatelessSimulationExample1.usf", "StatelessMain", SF_Compute);
+}
diff --git a/Engine/Source/Editor/EditorWidgets/Private/SEnumCombo.cpp b/Engine/Source/Editor/EditorWidgets/Private/SEnumCombo.cpp
--- a/Engine/Source/Editor/EditorWidgets/Private/SEnumCombo.cpp
+++ b/Engine/Source/Editor/EditorWidgets/Private/SEnumCombo.cpp
@@ -99,12 +99,14 @@ void SEnumComboBox::Construct(const FArguments& InArgs, const UEnum* InEnum)
 	}
 
 	SComboButton::Construct(SComboButton::FArguments()
+		.ComboButtonStyle(InArgs._ComboButtonStyle)
 		.ButtonStyle(InArgs._ButtonStyle)
 		.ContentPadding(InArgs._ContentPadding)
 		.OnGetMenuContent(this, &SEnumComboBox::OnGetMenuContent)
 		.ButtonContent()
 		[
 			SNew(STextBlock)
+			.TextStyle(InArgs._TextStyle)
 			.Font(Font)
 			.Text(this, &SEnumComboBox::GetCurrentValueText)
 			.ToolTipText(this, &SEnumComboBox::GetCurrentValueTooltip)
diff --git a/Engine/Source/Editor/EditorWidgets/Public/SEnumCombo.h b/Engine/Source/Editor/EditorWidgets/Public/SEnumCombo.h
--- a/Engine/Source/Editor/EditorWidgets/Public/SEnumCombo.h
+++ b/Engine/Source/Editor/EditorWidgets/Public/SEnumCombo.h
@@ -29,7 +29,9 @@ public:
 	SLATE_BEGIN_ARGS(SEnumComboBox)
 		: _CurrentValue()
 		, _ContentPadding(FMargin(4.0, 2.0))
+		, _ComboButtonStyle(&FAppStyle::Get().GetWidgetStyle<FComboButtonStyle>("ComboButton"))
 		, _ButtonStyle(nullptr)
+		, _TextStyle(&FCoreStyle::Get().GetWidgetStyle<FTextBlockStyle>("NormalText"))
 	{}
 
 		SLATE_ATTRIBUTE(int32, CurrentValue)
@@ -37,7 +39,9 @@ public:
 		SLATE_ATTRIBUTE(FSlateFontInfo, Font)
 		SLATE_EVENT(FOnEnumSelectionChanged, OnEnumSelectionChanged)
 		SLATE_EVENT(FOnGetToolTipForValue, OnGetToolTipForValue)
+		SLATE_STYLE_ARGUMENT(FComboButtonStyle, ComboButtonStyle)
 		SLATE_STYLE_ARGUMENT(FButtonStyle, ButtonStyle)
+		SLATE_STYLE_ARGUMENT(FTextBlockStyle, TextStyle)
 		SLATE_ARGUMENT(TArray<int32>, EnumValueSubset)
 	SLATE_END_ARGS()
 
