- Fix for potential nullptr access

diff --git a/Engine/Plugins/FX/Niagara/Source/NiagaraEditor/Private/Widgets/SNiagaraSystemViewport.cpp b/Engine/Plugins/FX/Niagara/Source/NiagaraEditor/Private/Widgets/SNiagaraSystemViewport.cpp
--- a/Engine/Plugins/FX/Niagara/Source/NiagaraEditor/Private/Widgets/SNiagaraSystemViewport.cpp
+++ b/Engine/Plugins/FX/Niagara/Source/NiagaraEditor/Private/Widgets/SNiagaraSystemViewport.cpp
@@ -172,7 +172,7 @@ void FNiagaraSystemViewportClient::Draw(FViewport* InViewport, FCanvas* Canvas)
 		}
 
 		UWorld* World = Component->GetWorld();
-		if (World && !ParticleSystem->bFixedBounds)
+		if (World && ParticleSystem && !ParticleSystem->bFixedBounds)
 		{
 			FNiagaraSystemInstance* SystemInstance = SystemInstanceController->GetSystemInstance_Unsafe();
 			const TArray<FNiagaraEmitterHandle>& EmitterHandles = ParticleSystem->GetEmitterHandles();
