fix culled physics assets

diff --git a/Engine/Plugins/FX/Niagara/Source/Niagara/Private/NiagaraDataInterfaceRigidMeshCollisionQuery.cpp b/Engine/Plugins/FX/Niagara/Source/Niagara/Private/NiagaraDataInterfaceRigidMeshCollisionQuery.cpp
--- a/Engine/Plugins/FX/Niagara/Source/Niagara/Private/NiagaraDataInterfaceRigidMeshCollisionQuery.cpp
+++ b/Engine/Plugins/FX/Niagara/Source/Niagara/Private/NiagaraDataInterfaceRigidMeshCollisionQuery.cpp
@@ -291,6 +291,12 @@ uint32 MeshComponentHash(const UStaticMeshComponent* StaticMeshComponent)
 
 /// End UStaticMeshComponent
 
+bool FilterShape(const FKShapeElem& Elem)
+{
+	// CollisionEnabledHasPhysics(Elem.GetCollisionEnabled());
+	return Elem.GetCollisionEnabled() != ECollisionEnabled::NoCollision;
+}
+
 template<typename TComponentType>
 void CountCollisionPrimitives(TConstArrayView<TComponentType*> Components, TArray<FNDIRigidMeshCollisionData::FComponentBodyCount>& PerComponentCounts, uint32& TotalBoxCount, uint32& TotalSphereCount, uint32& TotalCapsuleCount)
 {
@@ -305,28 +311,28 @@ void CountCollisionPrimitives(TConstArrayView<TComponentType*> Components, TArra
 		{
 			for (const FKConvexElem& ConvexElem : BodySetup->AggGeom.ConvexElems)
 			{
-				if (CollisionEnabledHasPhysics(ConvexElem.GetCollisionEnabled()))
+				if (FilterShape(ConvexElem))
 				{
 					++BodyCount.BoxCount;
 				}
 			}
 			for (const FKBoxElem& BoxElem : BodySetup->AggGeom.BoxElems)
 			{
-				if (CollisionEnabledHasPhysics(BoxElem.GetCollisionEnabled()))
+				if (FilterShape(BoxElem))
 				{
 					++BodyCount.BoxCount;
 				}
 			}
 			for (const FKSphereElem& SphereElem : BodySetup->AggGeom.SphereElems)
 			{
-				if (CollisionEnabledHasPhysics(SphereElem.GetCollisionEnabled()))
+				if (FilterShape(SphereElem))
 				{
 					++BodyCount.SphereCount;
 				}
 			}
 			for (const FKSphylElem& CapsuleElem : BodySetup->AggGeom.SphylElems)
 			{
-				if (CollisionEnabledHasPhysics(CapsuleElem.GetCollisionEnabled()))
+				if (FilterShape(CapsuleElem))
 				{
 					++BodyCount.CapsuleCount;
 				}
@@ -390,7 +396,7 @@ void UpdateAssetArrays(TConstArrayView<TComponentType*> Components, const FVecto
 
 		for (const FKConvexElem& ConvexElem : BodySetup->AggGeom.ConvexElems)
 		{
-			if (CollisionEnabledHasPhysics(ConvexElem.GetCollisionEnabled()))
+			if (FilterShape(ConvexElem))
 			{
 				FBox BBox = ConvexElem.ElemBox;
 
@@ -411,7 +417,7 @@ void UpdateAssetArrays(TConstArrayView<TComponentType*> Components, const FVecto
 		}
 		for (const FKBoxElem& BoxElem : BodySetup->AggGeom.BoxElems)
 		{
-			if (CollisionEnabledHasPhysics(BoxElem.GetCollisionEnabled()))
+			if (FilterShape(BoxElem))
 			{
 				if (InitializeStatics)
 				{
@@ -429,7 +435,7 @@ void UpdateAssetArrays(TConstArrayView<TComponentType*> Components, const FVecto
 
 		for (const FKSphereElem& SphereElem : BodySetup->AggGeom.SphereElems)
 		{
-			if (CollisionEnabledHasPhysics(SphereElem.GetCollisionEnabled()))
+			if (FilterShape(SphereElem))
 			{
 				if (InitializeStatics)
 				{
@@ -447,7 +453,7 @@ void UpdateAssetArrays(TConstArrayView<TComponentType*> Components, const FVecto
 
 		for (const FKSphylElem& CapsuleElem : BodySetup->AggGeom.SphylElems)
 		{
-			if (CollisionEnabledHasPhysics(CapsuleElem.GetCollisionEnabled()))
+			if (FilterShape(CapsuleElem))
 			{
 				if (InitializeStatics)
 				{
