Move AtributeTrimming CompileTag to editor only (fixup change that got lost during a previous merge)

diff --git a/Engine/Plugins/FX/Niagara/Source/NiagaraEditor/Private/NiagaraHlslTranslator.cpp b/Engine/Plugins/FX/Niagara/Source/NiagaraEditor/Private/NiagaraHlslTranslator.cpp
--- a/Engine/Plugins/FX/Niagara/Source/NiagaraEditor/Private/NiagaraHlslTranslator.cpp
+++ b/Engine/Plugins/FX/Niagara/Source/NiagaraEditor/Private/NiagaraHlslTranslator.cpp
@@ -10343,7 +10343,7 @@ void TNiagaraHlslTranslator<GraphBridge>::TrimAttributes(const FNiagaraCompileOp
 		{
 			if (!Attributes.Contains(Attribute))
 			{
-				TranslateResults.CompileTags.Emplace(Attribute, TEXT("Trimmed"));
+				TranslateResults.CompileTagsEditorOnly.Emplace(Attribute, TEXT("Trimmed"));
 			}
 		}
 	}
