Clean up unnecessary data out of the NiagaraShaderMap DDC payload.  Most significantly this removes the FriendlyName as the DDC data can be generated from multiple distinct NiagaraSystems, the name is non-determininistic.

diff --git a/Engine/Plugins/FX/Niagara/Source/NiagaraEditor/Private/NiagaraHlslTranslator.cpp b/Engine/Plugins/FX/Niagara/Source/NiagaraEditor/Private/NiagaraHlslTranslator.cpp
--- a/Engine/Plugins/FX/Niagara/Source/NiagaraEditor/Private/NiagaraHlslTranslator.cpp
+++ b/Engine/Plugins/FX/Niagara/Source/NiagaraEditor/Private/NiagaraHlslTranslator.cpp
@@ -125,8 +125,6 @@ void FNiagaraShaderQueueTickable::ProcessQueue()
 			continue;
 		}
 
-		FNiagaraComputeShaderCompilationOutput NewCompilationOutput;
-
 		ShaderScript->BuildScriptParametersMetadata(CompilableScript->GetVMExecutableData().ShaderScriptParametersMetadata);
 		ShaderScript->SetSourceName(TEXT("NiagaraComputeShader"));
 		UNiagaraEmitter* Emitter = Cast<UNiagaraEmitter>(CompilableScript->GetOuter());
@@ -144,7 +142,7 @@ void FNiagaraShaderQueueTickable::ProcessQueue()
 			const bool bSynchronousCompile = GIsAutomationTesting;
 
 			// Compile the shaders for the script.
-			NewShaderMap->Compile(ShaderScript, Item.ShaderMapId, CompilerEnvironment, NewCompilationOutput, Item.Platform, bSynchronousCompile, Item.bApply);
+			NewShaderMap->Compile(ShaderScript, Item.ShaderMapId, CompilerEnvironment, Item.Platform, bSynchronousCompile, Item.bApply);
 		}
 	}
 
diff --git a/Engine/Plugins/FX/Niagara/Source/NiagaraShader/Private/NiagaraShader.cpp b/Engine/Plugins/FX/Niagara/Source/NiagaraShader/Private/NiagaraShader.cpp
--- a/Engine/Plugins/FX/Niagara/Source/NiagaraShader/Private/NiagaraShader.cpp
+++ b/Engine/Plugins/FX/Niagara/Source/NiagaraShader/Private/NiagaraShader.cpp
@@ -471,8 +471,7 @@ void FNiagaraShaderType::AddUniformBufferIncludesToEnvironment(FShaderCompilerEn
  */
 FShader* FNiagaraShaderType::FinishCompileShader(
 	const FSHAHash& ShaderMapHash,
-	const FShaderCompileJob& CurrentJob,
-	const FString& InDebugDescription
+	const FShaderCompileJob& CurrentJob
 	) const
 {
 	check(CurrentJob.bSucceeded);
@@ -483,8 +482,7 @@ FShader* FNiagaraShaderType::FinishCompileShader(
 			static_cast<const FParameters*>(CurrentJob.ShaderParameters.Get()),
 			CurrentJob.Key.PermutationId,
 			CurrentJob.Output,
-			ShaderMapHash,
-			InDebugDescription
+			ShaderMapHash
 		)
 	);
 
@@ -578,14 +576,14 @@ void FNiagaraShaderMap::LoadFromDerivedDataCache(const FNiagaraShaderScript* Scr
 	}
 }
 
-void FNiagaraShaderMap::SaveToDerivedDataCache()
+void FNiagaraShaderMap::SaveToDerivedDataCache(const FNiagaraShaderScript* Script)
 {
 	COOK_STAT(auto Timer = NiagaraShaderCookStats::UsageStats.TimeSyncWork());
 	TArray<uint8> SaveData;
 	FMemoryWriter Ar(SaveData, true);
 	Serialize(Ar);
 
-	GetDerivedDataCacheRef().Put(*GetNiagaraShaderMapKeyString(GetContent()->ShaderMapId, GetShaderPlatform()), SaveData, FStringView(*GetFriendlyName()));
+	GetDerivedDataCacheRef().Put(*GetNiagaraShaderMapKeyString(GetContent()->ShaderMapId, GetShaderPlatform()), SaveData, Script ? Script->GetFriendlyName() : TEXT(""));
 	COOK_STAT(Timer.AddMiss(SaveData.Num()));
 }
 
@@ -620,7 +618,6 @@ void FNiagaraShaderMap::Compile(
 	FNiagaraShaderScript* Script,
 	const FNiagaraShaderMapId& InShaderMapId,
 	TRefCountPtr<FSharedShaderCompilerEnvironment> CompilationEnvironment,
-	const FNiagaraComputeShaderCompilationOutput& InNiagaraCompilationOutput,
 	EShaderPlatform InPlatform,
 	bool bSynchronousCompile,
 	bool bApplyCompletedShaderMapForRendering)
@@ -663,8 +660,6 @@ void FNiagaraShaderMap::Compile(
   
 			// Store the script name for debugging purposes.
 			FNiagaraShaderMapContent* NewContent = new FNiagaraShaderMapContent(InPlatform);
-			NewContent->FriendlyName = Script->GetFriendlyName();
-			NewContent->NiagaraCompilationOutput = InNiagaraCompilationOutput;
 			NewContent->ShaderMapId = InShaderMapId;
 			AssignContent(NewContent);
 
@@ -732,7 +727,7 @@ void FNiagaraShaderMap::Compile(
 			{
 				TArray<int32> CurrentShaderMapId;
 				CurrentShaderMapId.Add(CompilingId);
-				GNiagaraShaderCompilationManager.FinishCompilation(*NewContent->FriendlyName, CurrentShaderMapId);
+				GNiagaraShaderCompilationManager.FinishCompilation(CurrentShaderMapId);
 			}
 		}
 	}
@@ -771,7 +766,6 @@ void FNiagaraShaderMap::CreateCompileJobs(
 
 	// Store the script name for debugging purposes.
 	FNiagaraShaderMapContent* NewContent = new FNiagaraShaderMapContent(InPlatform);
-	NewContent->FriendlyName = FString(FriendlyName);
 	NewContent->ShaderMapId = InShaderMapId;
 	AssignContent(NewContent);
 
@@ -813,7 +807,7 @@ FShader* FNiagaraShaderMap::ProcessCompilationResultsForSingleJob(const TRefCoun
 
 	const FNiagaraShaderType* NiagaraShaderType = CurrentJob->Key.ShaderType->GetNiagaraShaderType();
 	check(NiagaraShaderType);
-	Shader = NiagaraShaderType->FinishCompileShader(ShaderMapHash, *CurrentJob, GetContent()->FriendlyName);
+	Shader = NiagaraShaderType->FinishCompileShader(ShaderMapHash, *CurrentJob);
 	bCompiledSuccessfully = CurrentJob->bSucceeded;
 
 	FNiagaraShader* NiagaraShader = static_cast<FNiagaraShader*>(Shader);
@@ -837,7 +831,7 @@ void FNiagaraShaderMap::ProcessAndFinalizeShaderCompileJob(const TRefCountPtr<FS
 	bCompilationFinalized = true;
 }
 
-bool FNiagaraShaderMap::ProcessCompilationResults(const TArray<FShaderCommonCompileJobPtr>& InCompilationResults, int32& InOutJobIndex, float& TimeBudget)
+bool FNiagaraShaderMap::ProcessCompilationResults(const TArray<FNiagaraShaderScript*>& InScripts, const TArray<FShaderCommonCompileJobPtr>& InCompilationResults, int32& InOutJobIndex, float& TimeBudget)
 {
 	check(InOutJobIndex < InCompilationResults.Num());
 
@@ -864,7 +858,7 @@ bool FNiagaraShaderMap::ProcessCompilationResults(const TArray<FShaderCommonComp
 	{
 		FinalizeContent();
 
-		SaveToDerivedDataCache();
+		SaveToDerivedDataCache(InScripts.IsEmpty() ? nullptr : InScripts.Top());
 		// The shader map can now be used on the rendering thread
 		bCompilationFinalized = true;
 		return true;
@@ -1178,10 +1172,7 @@ const FNiagaraShaderMap* FNiagaraShaderMap::GetShaderMapBeingCompiled(const FNia
 
 FNiagaraShader::FNiagaraShader(const FNiagaraShaderType::CompiledShaderInitializerType& Initializer)
 	: FShader(Initializer)
-	, DebugDescription(Initializer.DebugDescription)
 {
-	check(!DebugDescription.IsEmpty());
-
 	// Cache off requirements for later queries
 	bNeedsViewUniformBuffer = Initializer.ParameterMap.ContainsParameterAllocation(TEXT("View"));
 
diff --git a/Engine/Plugins/FX/Niagara/Source/NiagaraShader/Private/NiagaraShaderCompilationManager.cpp b/Engine/Plugins/FX/Niagara/Source/NiagaraShader/Private/NiagaraShaderCompilationManager.cpp
--- a/Engine/Plugins/FX/Niagara/Source/NiagaraShader/Private/NiagaraShaderCompilationManager.cpp
+++ b/Engine/Plugins/FX/Niagara/Source/NiagaraShader/Private/NiagaraShaderCompilationManager.cpp
@@ -177,7 +177,7 @@ void FNiagaraShaderCompilationManager::ProcessCompiledNiagaraShaderMaps(TMap<int
 			bool bShaderMapComplete = true;
 			if (bSuccess)
 			{
-				bShaderMapComplete = ShaderMap->ProcessCompilationResults(ResultArray, CompileResults.FinalizeJobIndex, TimeBudget);
+				bShaderMapComplete = ShaderMap->ProcessCompilationResults(ScriptArray, ResultArray, CompileResults.FinalizeJobIndex, TimeBudget);
 			}
 
 			if (bShaderMapComplete)
@@ -307,7 +307,7 @@ void FNiagaraShaderCompilationManager::ProcessCompiledNiagaraShaderMaps(TMap<int
 }
 
 
-void FNiagaraShaderCompilationManager::FinishCompilation(const TCHAR* ScriptName, const TArray<int32>& ShaderMapIdsToFinishCompiling)
+void FNiagaraShaderCompilationManager::FinishCompilation(const TArray<int32>& ShaderMapIdsToFinishCompiling)
 {
 	check(!FPlatformProperties::RequiresCookedData());
 	GShaderCompilingManager->FinishCompilation(NULL, ShaderMapIdsToFinishCompiling);
diff --git a/Engine/Plugins/FX/Niagara/Source/NiagaraShader/Private/NiagaraShared.cpp b/Engine/Plugins/FX/Niagara/Source/NiagaraShader/Private/NiagaraShared.cpp
--- a/Engine/Plugins/FX/Niagara/Source/NiagaraShader/Private/NiagaraShared.cpp
+++ b/Engine/Plugins/FX/Niagara/Source/NiagaraShader/Private/NiagaraShared.cpp
@@ -32,7 +32,6 @@
 IMPLEMENT_TYPE_LAYOUT(FNiagaraDataInterfaceParamRef);
 IMPLEMENT_TYPE_LAYOUT(FNiagaraShaderMapContent);
 IMPLEMENT_TYPE_LAYOUT(FNiagaraShaderMapId);
-IMPLEMENT_TYPE_LAYOUT(FNiagaraComputeShaderCompilationOutput);
 
 //* CVars */
 int32 GNiagaraTranslatorFailIfNotSetSeverity = 3;
@@ -608,7 +607,7 @@ bool FNiagaraShaderScript::CacheShaders(const FNiagaraShaderMapId& ShaderMapId,
 			FNiagaraShaderMap::LoadFromDerivedDataCache(this, ShaderMapId, ShaderPlatform, GameThreadShaderMap);
 			if (GameThreadShaderMap && GameThreadShaderMap->IsValid())
 			{
-				UE_LOG(LogShaders, Verbose, TEXT("Loaded shader %s for Niagara script %s from DDC"), *GameThreadShaderMap->GetFriendlyName(), *GetFriendlyName());
+				UE_LOG(LogShaders, Verbose, TEXT("Loaded shader for Niagara script %s from DDC"), *GetFriendlyName());
 			}
 			else
 			{
@@ -693,7 +692,7 @@ void FNiagaraShaderScript::FinishCompilation()
 		}
 		// Block until the shader maps that we will save have finished being compiled
 		// NIAGARATODO: implement when async compile works
-		GNiagaraShaderCompilationManager.FinishCompilation(*GetFriendlyName(), ShaderMapIdsToFinish);
+		GNiagaraShaderCompilationManager.FinishCompilation(ShaderMapIdsToFinish);
 
 		// Shouldn't have anything left to do...
 		TArray<int32> ShaderMapIdsToFinish2;
diff --git a/Engine/Plugins/FX/Niagara/Source/NiagaraShader/Public/NiagaraShader.h b/Engine/Plugins/FX/Niagara/Source/NiagaraShader/Public/NiagaraShader.h
--- a/Engine/Plugins/FX/Niagara/Source/NiagaraShader/Public/NiagaraShader.h
+++ b/Engine/Plugins/FX/Niagara/Source/NiagaraShader/Public/NiagaraShader.h
@@ -244,8 +244,6 @@ public:
 private:
 	// Data about parameters used for each Data Interface.
 	LAYOUT_FIELD(TMemoryImageArray<FNiagaraDataInterfaceParamRef>, DataInterfaceParameters);
-
-	LAYOUT_FIELD(FMemoryImageString, DebugDescription);
 };
 
 extern NIAGARASHADER_API int32 GNiagaraSkipVectorVMBackendOptimizations;
diff --git a/Engine/Plugins/FX/Niagara/Source/NiagaraShader/Public/NiagaraShaderCompilationManager.h b/Engine/Plugins/FX/Niagara/Source/NiagaraShader/Public/NiagaraShaderCompilationManager.h
--- a/Engine/Plugins/FX/Niagara/Source/NiagaraShader/Public/NiagaraShaderCompilationManager.h
+++ b/Engine/Plugins/FX/Niagara/Source/NiagaraShader/Public/NiagaraShaderCompilationManager.h
@@ -43,7 +43,7 @@ public:
 	NIAGARASHADER_API void AddJobs(TArray<FShaderCommonCompileJobPtr> InNewJobs);
 	NIAGARASHADER_API void ProcessAsyncResults();
 
-	void FinishCompilation(const TCHAR* ScriptName, const TArray<int32>& ShaderMapIdsToFinishCompiling);
+	void FinishCompilation(const TArray<int32>& ShaderMapIdsToFinishCompiling);
 private:
 	void ProcessCompiledNiagaraShaderMaps(TMap<int32, FNiagaraShaderMapFinalizeResults>& CompiledShaderMaps, float TimeBudget);
 
diff --git a/Engine/Plugins/FX/Niagara/Source/NiagaraShader/Public/NiagaraShaderType.h b/Engine/Plugins/FX/Niagara/Source/NiagaraShader/Public/NiagaraShaderType.h
--- a/Engine/Plugins/FX/Niagara/Source/NiagaraShader/Public/NiagaraShaderType.h
+++ b/Engine/Plugins/FX/Niagara/Source/NiagaraShader/Public/NiagaraShaderType.h
@@ -67,18 +67,14 @@ public:
 
 	struct CompiledShaderInitializerType : FShaderCompiledShaderInitializerType
 	{
-		const FString DebugDescription;
-
 		CompiledShaderInitializerType(
 			const FShaderType* InType,
 			const FParameters* InParameters,
 			int32 InPermutationId,
 			const FShaderCompilerOutput& CompilerOutput,
-			const FSHAHash& InNiagaraShaderMapHash,
-			const FString& InDebugDescription
+			const FSHAHash& InNiagaraShaderMapHash
 		)
 			: FShaderCompiledShaderInitializerType(InType, InParameters, InPermutationId, CompilerOutput, InNiagaraShaderMapHash, nullptr, nullptr)
-			, DebugDescription(InDebugDescription)
 		{
 		}
 	};
@@ -152,8 +148,7 @@ public:
 	 */
 	FShader* FinishCompileShader(
 		const FSHAHash& NiagaraShaderMapHash,
-		const FShaderCompileJob& CurrentJob,
-		const FString& InDebugDescription
+		const FShaderCompileJob& CurrentJob
 		) const;
 #endif // WITH_EDITOR
 
diff --git a/Engine/Plugins/FX/Niagara/Source/NiagaraShader/Public/NiagaraShared.h b/Engine/Plugins/FX/Niagara/Source/NiagaraShader/Public/NiagaraShared.h
--- a/Engine/Plugins/FX/Niagara/Source/NiagaraShader/Public/NiagaraShared.h
+++ b/Engine/Plugins/FX/Niagara/Source/NiagaraShader/Public/NiagaraShared.h
@@ -306,15 +306,6 @@ public:
 	LAYOUT_FIELD_WITH_WRITER(TMemoryImagePtr<FNiagaraDataInterfaceParametersCS>, Parameters, WriteFrozenParameters);
 };
 
-/** Stores outputs from the script compile that need to be saved. */
-class FNiagaraComputeShaderCompilationOutput
-{
-	DECLARE_TYPE_LAYOUT(FNiagaraComputeShaderCompilationOutput, NonVirtual);
-public:
-	FNiagaraComputeShaderCompilationOutput()
-	{}
-};
-
 /** Contains all the information needed to uniquely identify a FNiagaraShaderMapID. */
 class FNiagaraShaderMapId
 {
@@ -478,18 +469,9 @@ class FNiagaraShaderMapContent : public FShaderMapContent
 	DECLARE_TYPE_LAYOUT(FNiagaraShaderMapContent, NonVirtual);
 private:
 	explicit FNiagaraShaderMapContent(const EShaderPlatform InPlatform) : Super(InPlatform) {}
-	
-	/** The script's user friendly name, typically the object name. */
-	LAYOUT_FIELD(FMemoryImageString, FriendlyName);
-
-	/** Debug information about how the shader map was compiled. */
-	LAYOUT_FIELD(FMemoryImageString, DebugDescription);
 
 	/** The static parameter set that this shader map was compiled with */
 	LAYOUT_FIELD(FNiagaraShaderMapId, ShaderMapId);
-
-	/** Shader compilation output */
-	LAYOUT_FIELD(FNiagaraComputeShaderCompilationOutput, NiagaraCompilationOutput);
 };
 
 /**
@@ -542,7 +524,6 @@ public:
 		FNiagaraShaderScript* Script,
 		const FNiagaraShaderMapId& ShaderMapId,
 		TRefCountPtr<FSharedShaderCompilerEnvironment> CompilationEnvironment,
-		const FNiagaraComputeShaderCompilationOutput& InNiagaraCompilationOutput,
 		EShaderPlatform Platform,
 		bool bSynchronousCompile,
 		bool bApplyCompletedShaderMapForRendering
@@ -569,7 +550,7 @@ public:
 	NIAGARASHADER_API void ProcessAndFinalizeShaderCompileJob(const TRefCountPtr<FShaderCommonCompileJob>& SingleJob);
 
 	/** Sorts the incoming compiled jobs into the appropriate mesh shader maps, and finalizes this shader map so that it can be used for rendering. */
-	bool ProcessCompilationResults(const TArray<TRefCountPtr<class FShaderCommonCompileJob>>& InCompilationResults, int32& ResultIndex, float& TimeBudget);
+	bool ProcessCompilationResults(const TArray<FNiagaraShaderScript*>& InScripts, const TArray<TRefCountPtr<class FShaderCommonCompileJob>>& InCompilationResults, int32& ResultIndex, float& TimeBudget);
 
 	/**
 	* Checks whether the shader map is missing any shader types necessary for the given script.
@@ -619,7 +600,7 @@ public:
 
 #if WITH_EDITOR
 	/** Saves this shader map to the derived data cache. */
-	void SaveToDerivedDataCache();
+	void SaveToDerivedDataCache(const FNiagaraShaderScript* Script);
 
 	/** Backs up any FShaders in this shader map to memory through serialization and clears FShader references. */
 	TArray<uint8>* BackupShadersToMemory();
@@ -631,11 +612,9 @@ public:
 	// Accessors.
 	const FNiagaraShaderMapId& GetShaderMapId() const		{ return GetContent()->ShaderMapId; }
 	EShaderPlatform GetShaderPlatform() const				{ return GetContent()->GetShaderPlatform(); }
-	const FMemoryImageString& GetFriendlyName() const		{ return GetContent()->FriendlyName; }
 	uint32 GetCompilingId() const							{ return CompilingId; }
 	bool IsCompilationFinalized() const						{ return bCompilationFinalized; }
 	bool CompiledSuccessfully() const						{ return bCompiledSuccessfully; }
-	const FMemoryImageString& GetDebugDescription() const	{ return GetContent()->DebugDescription; }
 
 	bool IsValid() const
 	{
diff --git a/Engine/Source/Runtime/Core/Private/UObject/FortniteMainBranchObjectVersion.cpp b/Engine/Source/Runtime/Core/Private/UObject/FortniteMainBranchObjectVersion.cpp
--- a/Engine/Source/Runtime/Core/Private/UObject/FortniteMainBranchObjectVersion.cpp
+++ b/Engine/Source/Runtime/Core/Private/UObject/FortniteMainBranchObjectVersion.cpp
@@ -10,7 +10,7 @@ TMap<FGuid, FGuid> FFortniteMainBranchObjectVersion::GetSystemGuids()
 	SystemGuids.Add(DevGuids.LANDSCAPE_MOBILE_COOK_VERSION, FGuid("32D02EF867C74B71A0D4E0FA41392732"));
 	SystemGuids.Add(DevGuids.MATERIALSHADERMAP_DERIVEDDATA_VER, FGuid("CE7776E63E9E43FBB8C103EE1AB40B5D"));
 	SystemGuids.Add(DevGuids.NANITE_DERIVEDDATA_VER, FGuid("D28C003F825F4C178F698938FB0E47F7"));
-	SystemGuids.Add(DevGuids.NIAGARASHADERMAP_DERIVEDDATA_VER, FGuid("D0DF49D408F34E38A98FA097E0280EC7"));
+	SystemGuids.Add(DevGuids.NIAGARASHADERMAP_DERIVEDDATA_VER, FGuid("C417A50271B3427D98785388266E3B66"));
 	SystemGuids.Add(DevGuids.Niagara_LatestScriptCompileVersion, FGuid("CE33E29C6268469E8E18F1D96A55DFDA"));
 	SystemGuids.Add(DevGuids.SkeletalMeshDerivedDataVersion, FGuid("3BD78FFFB88E4DDF943FB2608A20C0D2"));
 	SystemGuids.Add(DevGuids.STATICMESH_DERIVEDDATA_VER, FGuid("543843B103794E0E9BA4BE80FB602F79"));
