Niagara Asset Browser 2.0
Resaved Niagara Template & Learning Content emitters

#virtualized

diff --git a/Engine/Plugins/FX/Niagara/Source/Niagara/Classes/NiagaraEmitter.h b/Engine/Plugins/FX/Niagara/Source/Niagara/Classes/NiagaraEmitter.h
--- a/Engine/Plugins/FX/Niagara/Source/Niagara/Classes/NiagaraEmitter.h
+++ b/Engine/Plugins/FX/Niagara/Source/Niagara/Classes/NiagaraEmitter.h
@@ -10,6 +10,7 @@
 #include "NiagaraMessageDataBase.h"
 #include "NiagaraMessageStore.h"
 #include "INiagaraMergeManager.h"
+#include "NiagaraAssetTagDefinitions.h"
 #include "NiagaraEffectType.h"
 #include "NiagaraDataSetAccessor.h"
 #include "NiagaraBoundsCalculator.h"
@@ -703,8 +704,12 @@ private:
 public:
 	
 #if WITH_EDITORONLY_DATA
-	UPROPERTY(EditAnywhere, AdvancedDisplay, Category = "Asset Options", AssetRegistrySearchable)
-	ENiagaraScriptTemplateSpecification TemplateSpecification;
+	UPROPERTY(EditAnywhere, Category = "Asset Options", meta=(ShowOnlyInnerProperties))
+	TArray<FNiagaraAssetTagDefinitionReference> AssetTags;
+
+	/** If an emitter is inheritable, new emitters based on an inheritable emitter, or Niagara Systems using an inheritable emitter, will automatically inherit changes made to the original emitter. */
+	UPROPERTY(EditAnywhere, Category = "Asset Options", AssetRegistrySearchable)
+	bool bIsInheritable = true;
 	
 	UPROPERTY(EditAnywhere, AdvancedDisplay, Category = "Asset Options", AssetRegistrySearchable)
 	FText TemplateAssetDescription;
@@ -761,7 +766,9 @@ public:
 	UPROPERTY(meta = (DeprecatedProperty))
 	FNiagaraEmitterScriptProperties SpawnScriptProps_DEPRECATED;
 
-
+	UPROPERTY(meta = (DeprecatedProperty))
+	ENiagaraScriptTemplateSpecification TemplateSpecification_DEPRECATED;
+	
 	/** Use property in struct returned from GetEmitterData() instead */ 
 	UPROPERTY(meta = (DeprecatedProperty))
 	FNiagaraEmitterScriptProperties EmitterSpawnScriptProps_DEPRECATED;
diff --git a/Engine/Plugins/FX/Niagara/Source/Niagara/Classes/NiagaraScript.h b/Engine/Plugins/FX/Niagara/Source/Niagara/Classes/NiagaraScript.h
--- a/Engine/Plugins/FX/Niagara/Source/Niagara/Classes/NiagaraScript.h
+++ b/Engine/Plugins/FX/Niagara/Source/Niagara/Classes/NiagaraScript.h
@@ -18,6 +18,7 @@
 #include "NiagaraVersionedObject.h"
 #include "HAL/CriticalSection.h"
 #include "VectorVM.h"
+#include "NiagaraAssetTagDefinitions.h"
 
 #include "NiagaraScript.generated.h"
 
@@ -694,6 +695,9 @@ public:
 	UPROPERTY(EditAnywhere, Category = Script)
 	FText Category;
 
+	UPROPERTY(EditAnywhere, Category = "Asset Options")
+	TArray<FNiagaraAssetTagDefinitionReference> AssetTagDefinitionReferences;
+
 	/** If true, this script will be added to a 'Suggested' category at the top of menus during searches */
 	UPROPERTY(AssetRegistrySearchable, EditAnywhere, Category = Script)
 	bool bSuggested = false;
diff --git a/Engine/Plugins/FX/Niagara/Source/Niagara/Classes/NiagaraSystem.h b/Engine/Plugins/FX/Niagara/Source/Niagara/Classes/NiagaraSystem.h
--- a/Engine/Plugins/FX/Niagara/Source/Niagara/Classes/NiagaraSystem.h
+++ b/Engine/Plugins/FX/Niagara/Source/Niagara/Classes/NiagaraSystem.h
@@ -3,6 +3,7 @@
 #pragma once
 
 #include "CoreMinimal.h"
+#include "..\Public\NiagaraAssetTagDefinitions.h"
 
 #include "NiagaraDataSetCompiledData.h"
 #include "NiagaraDataSetAccessor.h"
@@ -455,8 +456,8 @@ public:
 	bool bIsTemplateAsset_DEPRECATED;
 
 	/** If this system is a regular system, a template or a behavior example. */
-	UPROPERTY(EditAnywhere, AdvancedDisplay, Category = "Asset Options", AssetRegistrySearchable, meta = (SkipSystemResetOnChange = "true"))
-	ENiagaraScriptTemplateSpecification TemplateSpecification;
+	UPROPERTY()
+	ENiagaraScriptTemplateSpecification TemplateSpecification_DEPRECATED;
 
 	UPROPERTY(EditAnywhere, AdvancedDisplay, Category = "Asset Options", AssetRegistrySearchable, meta = (SkipSystemResetOnChange = "true"))
 	FText TemplateAssetDescription;
@@ -465,6 +466,9 @@ public:
 	UPROPERTY(EditAnywhere, AdvancedDisplay, Category = "Asset Options", AssetRegistrySearchable, meta = (SkipSystemResetOnChange = "true"))
 	FText Category;
 
+	UPROPERTY(EditAnywhere, AdvancedDisplay, Category = "Asset Options", meta = (SkipSystemResetOnChange = "true"))
+	TArray<FNiagaraAssetTagDefinitionReference> AssetTags;
+
 	UPROPERTY()
 	TArray<TObjectPtr<UNiagaraScript>> ScratchPadScripts;
 
diff --git a/Engine/Plugins/FX/Niagara/Source/Niagara/Private/NiagaraAssetTagDefinitions.cpp b/Engine/Plugins/FX/Niagara/Source/Niagara/Private/NiagaraAssetTagDefinitions.cpp
new file mode 100644
--- /dev/null
+++ b/Engine/Plugins/FX/Niagara/Source/Niagara/Private/NiagaraAssetTagDefinitions.cpp
@@ -0,0 +1,106 @@
+﻿// Copyright Epic Games, Inc. All Rights Reserved.
+
+#include "NiagaraAssetTagDefinitions.h"
+
+#include "GeneralProjectSettings.h"
+#include "NiagaraEmitter.h"
+#include "NiagaraSystem.h"
+#include "AssetRegistry/AssetRegistryModule.h"
+#include "Interfaces/IProjectManager.h"
+#include "AssetRegistry/AssetData.h"
+#include "UObject/AssetRegistryTagsContext.h"
+
+bool FNiagaraAssetTagDefinition::IsValid() const
+{
+	return TagGuid.IsValid() && AssetTag.IsEmpty() == false;
+}
+
+bool FNiagaraAssetTagDefinition::SupportsEmitters() const
+{
+	return (AssetFlags & (int32) ENiagaraAssetLibraryAssetTypes::Emitters) != 0;
+}
+
+bool FNiagaraAssetTagDefinition::SupportsSystems() const
+{
+	return (AssetFlags & (int32) ENiagaraAssetLibraryAssetTypes::Systems) != 0;
+}
+
+bool FNiagaraAssetTagDefinition::SupportsScripts() const
+{
+	return (AssetFlags & (int32) ENiagaraAssetLibraryAssetTypes::Scripts) != 0;
+}
+
+TArray<UClass*> FNiagaraAssetTagDefinition::GetSupportedClasses() const
+{
+	TArray<UClass*> Result;
+	
+	if(SupportsEmitters())
+	{
+		Result.Add(UNiagaraEmitter::StaticClass());
+	}
+
+	if(SupportsSystems())
+	{
+		Result.Add(UNiagaraSystem::StaticClass());
+	}
+
+	if(SupportsScripts())
+	{
+		Result.Add(UNiagaraScript::StaticClass());
+	}
+
+	return Result;
+}
+
+bool FNiagaraAssetTagDefinition::DoesAssetDataContainTag(const FAssetData& AssetData) const
+{
+	return AssetData.FindTag(FName(GetGuidAsString()));
+}
+
+FString FNiagaraAssetTagDefinition::GetGuidAsString() const
+{
+	return TagGuid.ToString(EGuidFormats::DigitsWithHyphens);
+}
+
+void FNiagaraAssetTagDefinition::AddTagToAssetRegistryTags(TArray<UObject::FAssetRegistryTag>& OutTags) const
+{
+	UObject::FAssetRegistryTag* FoundAssetTag = OutTags.FindByPredicate([this](const UObject::FAssetRegistryTag& Candidate)
+	{
+		return Candidate.Name == FName(GetGuidAsString());
+	});
+
+	if(FoundAssetTag)
+	{
+		return;
+	}
+	else
+	{
+		// Value has to be >1 of data to be considered non-empty to not get thrown away
+		OutTags.Add(UObject::FAssetRegistryTag(FName(GetGuidAsString()), "  ", UObject::FAssetRegistryTag::TT_Alphabetical));
+	}
+}
+
+FNiagaraAssetTagDefinitionReference::FNiagaraAssetTagDefinitionReference(const FNiagaraAssetTagDefinition& InTagDefinition)
+{
+	SetTagDefinitionReference(InTagDefinition);
+}
+
+FString FNiagaraAssetTagDefinitionReference::GetGuidAsString() const
+{
+	return AssetTagDefinitionGuid.ToString(EGuidFormats::DigitsWithHyphens);
+}
+
+void FNiagaraAssetTagDefinitionReference::AddTagToAssetRegistryTags(FAssetRegistryTagsContext& Context) const
+{
+	UObject::FAssetRegistryTag* FoundAssetTag = Context.FindTag(FName(GetGuidAsString()));
+
+	if(FoundAssetTag)
+	{
+		return;
+	}
+	else
+	{
+		// Value has to be >1 of data to be considered non-empty to not get thrown away
+		Context.AddTag(UObject::FAssetRegistryTag(FName(GetGuidAsString()), "  ", UObject::FAssetRegistryTag::TT_Alphabetical));
+	}
+}
diff --git a/Engine/Plugins/FX/Niagara/Source/Niagara/Private/NiagaraEmitter.cpp b/Engine/Plugins/FX/Niagara/Source/Niagara/Private/NiagaraEmitter.cpp
--- a/Engine/Plugins/FX/Niagara/Source/Niagara/Private/NiagaraEmitter.cpp
+++ b/Engine/Plugins/FX/Niagara/Source/Niagara/Private/NiagaraEmitter.cpp
@@ -631,11 +631,12 @@ void UNiagaraEmitter::PostLoad()
 			CineOverride.SpawnCountScale = GlobalSpawnCountScaleOverrides_DEPRECATED.Cine;
 		}
 	}
-
+	
 	// this can only ever be true for old assets that haven't been loaded yet, so this won't overwrite subsequent changes to the template specification
+	ENiagaraScriptTemplateSpecification CurrentTemplateSpecification = TemplateSpecification_DEPRECATED;
 	if(bIsTemplateAsset_DEPRECATED)
 	{
-		TemplateSpecification = ENiagaraScriptTemplateSpecification::Template;
+		CurrentTemplateSpecification = ENiagaraScriptTemplateSpecification::Template;
 	}
 
 	if(bExposeToLibrary_DEPRECATED)
@@ -643,6 +644,23 @@ void UNiagaraEmitter::PostLoad()
 		LibraryVisibility = ENiagaraScriptLibraryVisibility::Library;
 	}
 
+	if(NiagaraVer < FNiagaraCustomVersion::InheritanceUxRefactor)
+	{
+		if(CurrentTemplateSpecification == ENiagaraScriptTemplateSpecification::Template || CurrentTemplateSpecification == ENiagaraScriptTemplateSpecification::Behavior)
+		{
+			bIsInheritable = false;
+		}
+
+		if(CurrentTemplateSpecification == ENiagaraScriptTemplateSpecification::Template)
+		{			
+			AssetTags.AddUnique(INiagaraModule::TemplateTagDefinition);
+		}
+		else if(CurrentTemplateSpecification == ENiagaraScriptTemplateSpecification::Behavior)
+		{
+			AssetTags.AddUnique(INiagaraModule::LearningContentTagDefinition);
+		}
+	}
+	
 	if (MessageKeyToMessageMap_DEPRECATED.IsEmpty() == false)
 	{
 		MessageStore.SetMessages(MessageKeyToMessageMap_DEPRECATED);
@@ -1034,14 +1052,15 @@ void UNiagaraEmitter::GetAssetRegistryTags(FAssetRegistryTagsContext Context) co
 		Context.AddTag(FAssetRegistryTag(StringIter.Key(), LexToString(StringIter.Value()), FAssetRegistryTag::TT_Alphabetical));
 		++StringIter;
 	}
-
-	// TemplateSpecialization
-	FName TemplateSpecificationName = GET_MEMBER_NAME_CHECKED(UNiagaraEmitter, TemplateSpecification);
-	FText TemplateSpecializationValueString = StaticEnum<ENiagaraScriptTemplateSpecification>()->GetDisplayNameTextByValue((int64) TemplateSpecification);
-	Context.AddTag(FAssetRegistryTag(TemplateSpecificationName, TemplateSpecializationValueString.ToString(), FAssetRegistryTag::TT_Alphabetical));
 	
 	INiagaraModule& NiagaraModule = FModuleManager::GetModuleChecked<INiagaraModule>("Niagara");
 	Context.AddTag(NiagaraModule.GetEditorOnlyDataUtilities().CreateClassUsageAssetRegistryTag(this));
+
+	// Asset Library Tags
+	for(const FNiagaraAssetTagDefinitionReference& Tag : AssetTags)
+	{
+		Tag.AddTagToAssetRegistryTags(Context);
+	}
 #endif
 	Super::GetAssetRegistryTags(Context);
 }
diff --git a/Engine/Plugins/FX/Niagara/Source/Niagara/Private/NiagaraModule.cpp b/Engine/Plugins/FX/Niagara/Source/Niagara/Private/NiagaraModule.cpp
--- a/Engine/Plugins/FX/Niagara/Source/Niagara/Private/NiagaraModule.cpp
+++ b/Engine/Plugins/FX/Niagara/Source/Niagara/Private/NiagaraModule.cpp
@@ -67,6 +67,22 @@ static FAutoConsoleVariableRef CVarLogCompileIdGeneration(
 float INiagaraModule::EngineGlobalSpawnCountScale = 1.0f;
 float INiagaraModule::EngineGlobalSystemCountScale = 1.0f;
 
+const FNiagaraAssetTagDefinition INiagaraModule::TemplateTagDefinition = {
+	LOCTEXT("TemplateDisplayName", "Template")
+	, (int32) (ENiagaraAssetLibraryAssetTypes::Emitters) | (int32) (ENiagaraAssetLibraryAssetTypes::Systems)
+	, FText::GetEmpty()
+	, ENiagaraAssetTagDefinitionImportance::Primary
+	, FLinearColor::Yellow
+	, FGuid(0xFC157AFA, 0x50264263, 0xB9A9E658, 0xFFAC2D22)};
+
+const FNiagaraAssetTagDefinition INiagaraModule::LearningContentTagDefinition = {
+	LOCTEXT("LearningContentDisplayName", "Learning Content")
+	, (int32) (ENiagaraAssetLibraryAssetTypes::Emitters) | (int32) (ENiagaraAssetLibraryAssetTypes::Systems) | (int32) (ENiagaraAssetLibraryAssetTypes::Scripts)
+	, FText::GetEmpty()
+	, ENiagaraAssetTagDefinitionImportance::Primary
+	, FLinearColor::Green
+	, FGuid(0xFCF21AFA, 0x50764BFA, 0xB9B2E618, 0xF2A0CD6F)};
+
 std::atomic<bool> INiagaraModule::bDataChannelRefreshRequested = false;
 
 int32 GEnableVerboseNiagaraChangeIdLogging = 0;
@@ -477,6 +493,8 @@ void INiagaraModule::StartupModule()
 #endif
 #endif
 
+	RegisterInternalAssetTagDefinitions();
+	
 	FCoreDelegates::OnPostEngineInit.AddRaw(this, &INiagaraModule::OnPostEngineInit);
 	FCoreDelegates::OnPreExit.AddRaw(this, &INiagaraModule::OnPreExit);
 	FWorldDelegates::OnWorldBeginTearDown.AddRaw(this, &INiagaraModule::OnWorldBeginTearDown);
@@ -626,6 +644,11 @@ void INiagaraModule::ShutdownModule()
 #endif
 }
 
+INiagaraModule& INiagaraModule::Get()
+{
+	return FModuleManager::LoadModuleChecked<INiagaraModule>("Niagara");
+}
+
 #if WITH_EDITOR
 const INiagaraMergeManager& INiagaraModule::GetMergeManager() const
 {
@@ -1868,6 +1891,13 @@ const TArray<FNiagaraVariable>& FNiagaraEmitterParameters::GetVariables()
 
 #endif
 
+void INiagaraModule::RegisterInternalAssetTagDefinitions()
+{
+	// All internal tags have to be manually assigned a unique and stable guid so that names can be changed without affecting assigned assets.
+	InternalAssetTagDefinitions.Add(&TemplateTagDefinition);
+	InternalAssetTagDefinitions.Add(&LearningContentTagDefinition);
+}
+
 #if NIAGARA_PERF_BASELINES
 
 void INiagaraModule::GeneratePerfBaselines(TArray<UNiagaraEffectType*>& BaselinesToGenerate)
diff --git a/Engine/Plugins/FX/Niagara/Source/Niagara/Private/NiagaraScript.cpp b/Engine/Plugins/FX/Niagara/Source/Niagara/Private/NiagaraScript.cpp
--- a/Engine/Plugins/FX/Niagara/Source/Niagara/Private/NiagaraScript.cpp
+++ b/Engine/Plugins/FX/Niagara/Source/Niagara/Private/NiagaraScript.cpp
@@ -3543,6 +3543,13 @@ void UNiagaraScript::GetAssetRegistryTags(FAssetRegistryTagsContext Context) con
 	SuggestedTag = FString::FromInt(bSuggested);
 	Context.AddTag(FAssetRegistryTag(SuggestedName, SuggestedTag, FAssetRegistryTag::TT_Hidden));
 
+	// Asset Tags
+	TArray<FNiagaraAssetTagDefinitionReference> AssetTagsDefinitionReferences = ScriptData ? ScriptData->AssetTagDefinitionReferences : TArray<FNiagaraAssetTagDefinitionReference>();
+	for(const FNiagaraAssetTagDefinitionReference& TagDefinitionReference : AssetTagsDefinitionReferences)
+	{
+		TagDefinitionReference.AddTagToAssetRegistryTags(Context);
+	}
+	
 	// Add the current custom version to the tags so that tags can be fixed up in the future without having to load
 	// the whole asset.
 	const int32 NiagaraVer = GetLinkerCustomVersion(FNiagaraCustomVersion::GUID);
diff --git a/Engine/Plugins/FX/Niagara/Source/Niagara/Private/NiagaraSystem.cpp b/Engine/Plugins/FX/Niagara/Source/Niagara/Private/NiagaraSystem.cpp
--- a/Engine/Plugins/FX/Niagara/Source/Niagara/Private/NiagaraSystem.cpp
+++ b/Engine/Plugins/FX/Niagara/Source/Niagara/Private/NiagaraSystem.cpp
@@ -163,7 +163,6 @@ UNiagaraSystem::UNiagaraSystem(const FObjectInitializer& ObjectInitializer)
 : Super(ObjectInitializer)
 #if WITH_EDITORONLY_DATA
 , LibraryVisibility(ENiagaraScriptLibraryVisibility::Unexposed)
-, TemplateSpecification(ENiagaraScriptTemplateSpecification::None)
 , bBakeOutRapidIteration(false)
 , bBakeOutRapidIterationOnCook(true)
 , bTrimAttributes(false)
@@ -1082,12 +1081,6 @@ void UNiagaraSystem::PostLoad()
 		INiagaraModule& NiagaraModule = FModuleManager::GetModuleChecked<INiagaraModule>("Niagara");
 		EditorParameters = NiagaraModule.GetEditorOnlyDataUtilities().CreateDefaultEditorParameters(this);
 	}
-
-	// see the equivalent in NiagaraEmitter for details
-	if (bIsTemplateAsset_DEPRECATED)
-	{
-		TemplateSpecification = bIsTemplateAsset_DEPRECATED ? ENiagaraScriptTemplateSpecification::Template : ENiagaraScriptTemplateSpecification::None;
-	}
 	
 	if(bExposeToLibrary_DEPRECATED)
 	{
@@ -1117,9 +1110,22 @@ void UNiagaraSystem::PostLoad()
 
 #if WITH_EDITORONLY_DATA
 	// see the equivalent in NiagaraEmitter for details
+	ENiagaraScriptTemplateSpecification CurrentTemplateSpecification = TemplateSpecification_DEPRECATED;
 	if(bIsTemplateAsset_DEPRECATED)
 	{
-		TemplateSpecification = bIsTemplateAsset_DEPRECATED ? ENiagaraScriptTemplateSpecification::Template : ENiagaraScriptTemplateSpecification::None;
+		CurrentTemplateSpecification = ENiagaraScriptTemplateSpecification::Template;
+	}
+
+	if(NiagaraVer < FNiagaraCustomVersion::InheritanceUxRefactor)
+	{
+		if(CurrentTemplateSpecification == ENiagaraScriptTemplateSpecification::Template)
+		{
+			AssetTags.AddUnique(INiagaraModule::TemplateTagDefinition);
+		}
+		else if(CurrentTemplateSpecification == ENiagaraScriptTemplateSpecification::Behavior)
+		{
+			AssetTags.AddUnique(INiagaraModule::LearningContentTagDefinition);
+		}
 	}
 #endif // WITH_EDITORONLY_DATA
 
@@ -1582,11 +1588,6 @@ void UNiagaraSystem::GetAssetRegistryTags(FAssetRegistryTagsContext Context) con
 		++StringIter;
 	}
 
-	// TemplateSpecialization
-	FName TemplateSpecificationName = GET_MEMBER_NAME_CHECKED(UNiagaraSystem, TemplateSpecification);
-	FText TemplateSpecializationValueString = StaticEnum<ENiagaraScriptTemplateSpecification>()->GetDisplayNameTextByValue((int64) TemplateSpecification);
-	Context.AddTag(FAssetRegistryTag(TemplateSpecificationName, TemplateSpecializationValueString.ToString(), FAssetRegistryTag::TT_Alphabetical));
-
 	/*for (const UNiagaraDataInterface* DI : DataInterfaces)
 	{
 		FString ClassName;
@@ -1594,7 +1595,12 @@ void UNiagaraSystem::GetAssetRegistryTags(FAssetRegistryTagsContext Context) con
 		Context.AddTag(FAssetRegistryTag(*(TEXT("bHas")+ClassName), TEXT("True"), FAssetRegistryTag::TT_Alphabetical));
 	}*/
 
-
+	// Asset Library Tags
+	for(const FNiagaraAssetTagDefinitionReference& Tag : AssetTags)
+	{
+		Tag.AddTagToAssetRegistryTags(Context);
+	}
+	
 	//Context.AddTag(FAssetRegistryTag("CPUCollision", UsesCPUCollision() ? TEXT("True") : TEXT("False"), FAssetRegistryTag::TT_Alphabetical));
 	//Context.AddTag(FAssetRegistryTag("Looping", bAnyEmitterLoopsForever ? TEXT("True") : TEXT("False"), FAssetRegistryTag::TT_Alphabetical));
 	//Context.AddTag(FAssetRegistryTag("Immortal", IsImmortal() ? TEXT("True") : TEXT("False"), FAssetRegistryTag::TT_Alphabetical));
@@ -2442,10 +2448,9 @@ FNiagaraEmitterHandle UNiagaraSystem::AddEmitterHandle(UNiagaraEmitter& InEmitte
 {
 	UNiagaraEmitter* NewEmitter = UNiagaraEmitter::CreateWithParentAndOwner(FVersionedNiagaraEmitter(&InEmitter, EmitterVersion), this, EmitterName, ~(RF_Public | RF_Standalone));
 	FNiagaraEmitterHandle EmitterHandle(*NewEmitter, EmitterVersion);
-	if (InEmitter.TemplateSpecification == ENiagaraScriptTemplateSpecification::Template || InEmitter.TemplateSpecification == ENiagaraScriptTemplateSpecification::Behavior)
+	if (InEmitter.bIsInheritable == false)
 	{
 		NewEmitter->DisableVersioning(EmitterVersion);
-		NewEmitter->TemplateSpecification = ENiagaraScriptTemplateSpecification::None;
 		NewEmitter->TemplateAssetDescription = FText();
 		NewEmitter->GetLatestEmitterData()->RemoveParent();
 	}
diff --git a/Engine/Plugins/FX/Niagara/Source/Niagara/Public/NiagaraAssetTagDefinitions.h b/Engine/Plugins/FX/Niagara/Source/Niagara/Public/NiagaraAssetTagDefinitions.h
new file mode 100644
--- /dev/null
+++ b/Engine/Plugins/FX/Niagara/Source/Niagara/Public/NiagaraAssetTagDefinitions.h
@@ -0,0 +1,126 @@
+﻿// Copyright Epic Games, Inc. All Rights Reserved.
+
+#pragma once
+
+#include "CoreMinimal.h"
+#include "UObject/Object.h"
+#include "NiagaraAssetTagDefinitions.generated.h"
+
+UENUM(meta=(Bitflags, UseEnumValuesAsMaskValuesInEditor=true))
+enum class ENiagaraAssetLibraryAssetTypes : uint8
+{
+	Emitters = 1 << 0,
+	Systems = 1 << 1,
+	Scripts = 1 << 2
+};
+ENUM_CLASS_FLAGS(ENiagaraAssetLibraryAssetTypes)
+
+UENUM()
+enum class ENiagaraAssetTagDefinitionImportance : uint8
+{
+	Primary UMETA(ToolTip="This Asset Tag Definition is considered important and will get displayed in the primary filter UI in the Niagara Asset Browsers."),
+	Secondary UMETA(ToolTip="This Asset Tag Definition is considered less important and will only get displayed in the additional drop down filter UI in the Niagara Asset Browsers.")
+	//Both UMETA(ToolTip="This Asset Tag Definition will be displayed in both primary UI & the additional filter drop downs.")
+};
+
+/** An Asset Tag Definition defines a tag that can be added to various Niagara assets for sorting & filtering purposes. */
+USTRUCT()
+struct FNiagaraAssetTagDefinition
+{
+	GENERATED_BODY()
+
+public:
+	/** The Display Name used for this tag. */
+	UPROPERTY(EditAnywhere, Category="Properties")
+	FText AssetTag;
+
+	/** Select the asset types this tag can apply to. This is used to hide tags that can never apply for a given type. */
+	UPROPERTY(EditAnywhere, Category="Properties", meta=(Bitmask, BitmaskEnum="/Script/Niagara.ENiagaraAssetLibraryAssetTypes"))
+	int32 AssetFlags = 0;
+
+	/** Further explanation of what this tag is about. */
+	UPROPERTY(EditAnywhere, Category="Properties", meta=(MultiLine))
+	FText Description;
+
+	/** Whether this tag should be shown directly or in the drop down for additional filters. */
+	UPROPERTY(EditAnywhere, Category="Properties")
+	ENiagaraAssetTagDefinitionImportance DisplayType = ENiagaraAssetTagDefinitionImportance::Primary;
+
+	/** The color used in UI to represent this tag. */
+	UPROPERTY(EditAnywhere, Category="Properties")
+	FLinearColor Color;
+
+	/** The Tag Guid identifies this tag. This makes it possible to change the AssetTag name without it affecting functionality. */
+	UPROPERTY(VisibleAnywhere, Category="Properties")
+	FGuid TagGuid = FGuid::NewGuid();
+
+	bool operator<(const FNiagaraAssetTagDefinition& Other) const
+	{		
+		return AssetTag.ToString() < Other.AssetTag.ToString();
+	}
+
+	bool operator==(const FNiagaraAssetTagDefinition& Other) const
+	{
+		return TagGuid == Other.TagGuid;
+	}
+public:
+	NIAGARA_API bool IsValid() const;
+	
+	NIAGARA_API bool SupportsEmitters() const;
+	NIAGARA_API bool SupportsSystems() const;
+	NIAGARA_API bool SupportsScripts() const;
+	NIAGARA_API TArray<UClass*> GetSupportedClasses() const;
+	NIAGARA_API bool DoesAssetDataContainTag(const FAssetData& AssetData) const;
+
+	NIAGARA_API FString GetGuidAsString() const;
+	void AddTagToAssetRegistryTags(TArray<UObject::FAssetRegistryTag>& OutTags) const;
+};
+
+inline uint32 GetTypeHash(const FNiagaraAssetTagDefinition& AssetTagDefinition)
+{
+	return GetTypeHash(AssetTagDefinition.TagGuid);
+}
+
+/** A Tag Definition Reference stores the guid of a Tag Definition. This is what assets should be storing. */
+USTRUCT()
+struct NIAGARA_API FNiagaraAssetTagDefinitionReference
+{
+	GENERATED_BODY()
+
+public:
+	FNiagaraAssetTagDefinitionReference() = default;
+	FNiagaraAssetTagDefinitionReference(const FNiagaraAssetTagDefinition& InTagDefinition);
+	
+	FString GetGuidAsString() const;
+	void AddTagToAssetRegistryTags(FAssetRegistryTagsContext& Context) const;
+	
+	void SetTagDefinitionReference(const FNiagaraAssetTagDefinition& InTagDefinition) { SetTagDefinitionReferenceGuid(InTagDefinition.TagGuid); }
+	void SetTagDefinitionReferenceGuid(const FGuid& InTagDefinitionGuid) { AssetTagDefinitionGuid = InTagDefinitionGuid; }
+	FGuid GetTagDefinitionReferenceGuid() const { return AssetTagDefinitionGuid; }
+
+	bool operator==(const FNiagaraAssetTagDefinitionReference& Other) const
+	{
+		return AssetTagDefinitionGuid == Other.AssetTagDefinitionGuid;
+	}
+private:
+	UPROPERTY()
+	FGuid AssetTagDefinitionGuid;
+};
+
+/** An asset to define multiple tag definitions; used to sort and categorize Niagara assets. */
+UCLASS()
+class NIAGARA_API UNiagaraAssetTagDefinitions : public UObject
+{
+	GENERATED_BODY()
+public:	
+	const TArray<FNiagaraAssetTagDefinition>& GetAssetTagDefinitions() const { return TagDefinitions; }
+
+	int32 GetSortOrder() const { return SortOrder; }
+private:
+	UPROPERTY(EditAnywhere, Category="Properties", meta=(TitleProperty="AssetTag"))
+	TArray<FNiagaraAssetTagDefinition> TagDefinitions;
+
+	/** Tags are sorted by asset sort order first, then individually. That means tags of asset with sort order [0] come before tags of asset with sort order [1]. */
+	UPROPERTY(EditAnywhere, Category="Properties")
+	int32 SortOrder = INDEX_NONE;
+};
diff --git a/Engine/Plugins/FX/Niagara/Source/Niagara/Public/NiagaraModule.h b/Engine/Plugins/FX/Niagara/Source/Niagara/Public/NiagaraModule.h
--- a/Engine/Plugins/FX/Niagara/Source/Niagara/Public/NiagaraModule.h
+++ b/Engine/Plugins/FX/Niagara/Source/Niagara/Public/NiagaraModule.h
@@ -55,6 +55,10 @@ public:
 public:
 	NIAGARA_API virtual void StartupModule()override;
 	NIAGARA_API virtual void ShutdownModule()override;
+	
+	/** Get the instance of this module. */
+	NIAGARA_API static INiagaraModule& Get();
+	
 	NIAGARA_API void ShutdownRenderingResources();
 	
 	NIAGARA_API void OnPostEngineInit();
@@ -129,6 +133,9 @@ public:
 	static NIAGARA_API void RequestRefreshDataChannels() { bDataChannelRefreshRequested = true; }
 	static NIAGARA_API void RefreshDataChannels();
 
+	const TArray<const FNiagaraAssetTagDefinition*>& GetInternalAssetTagDefinitions() { return InternalAssetTagDefinitions; }
+	void RegisterInternalAssetTagDefinitions();
+	
 #if NIAGARA_PERF_BASELINES
 	NIAGARA_API void GeneratePerfBaselines(TArray<UNiagaraEffectType*>& BaselinesToGenerate);
 
@@ -260,9 +267,9 @@ public:
 	FORCEINLINE static const FNiagaraVariable&  GetVar_DataInstance_Alive() { return DataInstance_Alive; }
 	FORCEINLINE static const FNiagaraVariable&  GetVar_BeginDefaults() { return Translator_BeginDefaults; }
 	FORCEINLINE static const FNiagaraVariable&  GetVar_CallID() { return Translator_CallID; }
-
+	
 	FOnProcessQueue OnProcessQueue;
-
+	
 #if WITH_EDITORONLY_DATA
 	TSharedPtr<INiagaraMergeManager> MergeManager;
 	TSharedPtr<INiagaraEditorOnlyDataUtilities> EditorOnlyDataUtilities;
@@ -282,6 +289,9 @@ public:
 	static NIAGARA_API bool bUseGlobalFXBudget;
 	static NIAGARA_API bool bDataChannelsEnabled;
 
+	static const FNiagaraAssetTagDefinition TemplateTagDefinition;
+	static const FNiagaraAssetTagDefinition LearningContentTagDefinition;
+	TArray<const FNiagaraAssetTagDefinition*> InternalAssetTagDefinitions;
 private:
 	static NIAGARA_API FNiagaraVariable Engine_WorldDeltaTime;
 	static NIAGARA_API FNiagaraVariable Engine_DeltaTime;
diff --git a/Engine/Plugins/FX/Niagara/Source/NiagaraCore/Public/NiagaraCustomVersion.h b/Engine/Plugins/FX/Niagara/Source/NiagaraCore/Public/NiagaraCustomVersion.h
--- a/Engine/Plugins/FX/Niagara/Source/NiagaraCore/Public/NiagaraCustomVersion.h
+++ b/Engine/Plugins/FX/Niagara/Source/NiagaraCore/Public/NiagaraCustomVersion.h
@@ -199,6 +199,8 @@ struct FNiagaraCustomVersion
 
 		SimCache_BulkDataVersion1, // Sim Cache moved to bulk data by default
 
+		InheritanceUxRefactor, // Decoupling of 'Template' and 'Inheritance'
+		
 		// DO NOT ADD A NEW VERSION UNLESS YOU HAVE TALKED TO THE NIAGARA LEAD. Mismanagement of these versions can lead to data loss if it is adjusted in multiple streams simultaneously.
 		// -----<new versions can be added above this line>  -------------------------------------------------
 		VersionPlusOne,
diff --git a/Engine/Plugins/FX/Niagara/Source/NiagaraEditor/Internal/SNiagaraAssetPickerList.h b/Engine/Plugins/FX/Niagara/Source/NiagaraEditor/Internal/SNiagaraAssetPickerList.h
deleted file mode 100644
diff --git a/Engine/Plugins/FX/Niagara/Source/NiagaraEditor/NiagaraEditor.Build.cs b/Engine/Plugins/FX/Niagara/Source/NiagaraEditor/NiagaraEditor.Build.cs
--- a/Engine/Plugins/FX/Niagara/Source/NiagaraEditor/NiagaraEditor.Build.cs
+++ b/Engine/Plugins/FX/Niagara/Source/NiagaraEditor/NiagaraEditor.Build.cs
@@ -46,10 +46,13 @@ public class NiagaraEditor : ModuleRules
 				"ImageWrapper",
 				"AssetDefinition",
 				"ContentBrowser",
+				"ContentBrowserData",
 				"ToolWidgets",
 				"AssetTools",
 				"LevelSequence",
 				"SparseVolumeTexture",
+				"EngineSettings",
+				"EditorConfig"
 			}
 		);
 
diff --git a/Engine/Plugins/FX/Niagara/Source/NiagaraEditor/Private/AssetDefinitions/AssetDefinition_NiagaraAssetTagDefinitions.cpp b/Engine/Plugins/FX/Niagara/Source/NiagaraEditor/Private/AssetDefinitions/AssetDefinition_NiagaraAssetTagDefinitions.cpp
new file mode 100644
--- /dev/null
+++ b/Engine/Plugins/FX/Niagara/Source/NiagaraEditor/Private/AssetDefinitions/AssetDefinition_NiagaraAssetTagDefinitions.cpp
@@ -0,0 +1,17 @@
+﻿// Copyright Epic Games, Inc. All Rights Reserved.
+
+#include "AssetDefinition_NiagaraAssetTagDefinitions.h"
+
+#define LOCTEXT_NAMESPACE "UAssetDefinition_NiagaraAssetTagDefinitions"
+
+UAssetDefinition_NiagaraAssetTagDefinitions::UAssetDefinition_NiagaraAssetTagDefinitions()
+{
+}
+
+TConstArrayView<FAssetCategoryPath> UAssetDefinition_NiagaraAssetTagDefinitions::GetAssetCategories() const
+{
+	static const auto Categories = { EAssetCategoryPaths::FX / NSLOCTEXT("Niagara", "NiagaraAssetSubMenu_Advanced", "Advanced")};
+	return Categories;
+}
+
+#undef LOCTEXT_NAMESPACE
diff --git a/Engine/Plugins/FX/Niagara/Source/NiagaraEditor/Private/AssetDefinitions/AssetDefinition_NiagaraAssetTagDefinitions.h b/Engine/Plugins/FX/Niagara/Source/NiagaraEditor/Private/AssetDefinitions/AssetDefinition_NiagaraAssetTagDefinitions.h
new file mode 100644
--- /dev/null
+++ b/Engine/Plugins/FX/Niagara/Source/NiagaraEditor/Private/AssetDefinitions/AssetDefinition_NiagaraAssetTagDefinitions.h
@@ -0,0 +1,21 @@
+﻿// Copyright Epic Games, Inc. All Rights Reserved.
+
+#pragma once
+
+#include "CoreMinimal.h"
+#include "AssetDefinitionDefault.h"
+#include "NiagaraAssetTagDefinitions.h"
+#include "AssetDefinition_NiagaraAssetTagDefinitions.generated.h"
+
+UCLASS()
+class NIAGARAEDITOR_API UAssetDefinition_NiagaraAssetTagDefinitions : public UAssetDefinitionDefault
+{
+	GENERATED_BODY()
+
+	UAssetDefinition_NiagaraAssetTagDefinitions();
+	
+	virtual FText GetAssetDisplayName() const override { return NSLOCTEXT("AssetTypeActions", "AssetTypeActions_NiagaraAssetTagDefinitions", "Niagara Asset Tag Definitions"); }
+	virtual FLinearColor GetAssetColor() const override { return FLinearColor::White; }
+	virtual TSoftClassPtr<UObject> GetAssetClass() const override { return UNiagaraAssetTagDefinitions::StaticClass(); }
+	virtual TConstArrayView<FAssetCategoryPath> GetAssetCategories() const override;
+};
diff --git a/Engine/Plugins/FX/Niagara/Source/NiagaraEditor/Private/NiagaraAssetTagDefinitionsFactory.cpp b/Engine/Plugins/FX/Niagara/Source/NiagaraEditor/Private/NiagaraAssetTagDefinitionsFactory.cpp
new file mode 100644
--- /dev/null
+++ b/Engine/Plugins/FX/Niagara/Source/NiagaraEditor/Private/NiagaraAssetTagDefinitionsFactory.cpp
@@ -0,0 +1,18 @@
+﻿// Copyright Epic Games, Inc. All Rights Reserved.
+
+#include "NiagaraAssetTagDefinitionsFactory.h"
+
+#include "NiagaraAssetTagDefinitions.h"
+
+UNiagaraAssetTagDefinitionsFactory::UNiagaraAssetTagDefinitionsFactory()
+{
+	SupportedClass = UNiagaraAssetTagDefinitions::StaticClass();
+	bEditAfterNew = true;
+	bCreateNew = true;
+}
+
+UObject* UNiagaraAssetTagDefinitionsFactory::FactoryCreateNew(UClass* InClass, UObject* InParent, FName InName, EObjectFlags Flags, UObject* Context, FFeedbackContext* Warn)
+{
+	UNiagaraAssetTagDefinitions* NewAssetTagDefinitions = NewObject<UNiagaraAssetTagDefinitions>(InParent, InClass, InName, Flags | RF_Transactional);
+	return NewAssetTagDefinitions;
+}
diff --git a/Engine/Plugins/FX/Niagara/Source/NiagaraEditor/Private/NiagaraEditorMenuHelpers.cpp b/Engine/Plugins/FX/Niagara/Source/NiagaraEditor/Private/NiagaraEditorMenuHelpers.cpp
new file mode 100644
--- /dev/null
+++ b/Engine/Plugins/FX/Niagara/Source/NiagaraEditor/Private/NiagaraEditorMenuHelpers.cpp
@@ -0,0 +1,448 @@
+﻿// Copyright Epic Games, Inc. All Rights Reserved.
+
+#include "NiagaraEditorMenuHelpers.h"
+
+#include "ContentBrowserMenuContexts.h"
+#include "NiagaraEditorUtilities.h"
+#include "UObject/AssetRegistryTagsContext.h"
+#include "Widgets/AssetBrowser/NiagaraMenuFilters.h"
+
+#define LOCTEXT_NAMESPACE "NiagaraEditorMenuHelpers"
+
+void NiagaraEditorMenuHelpers::RegisterToolMenus()
+{
+	using namespace FNiagaraEditorUtilities::AssetBrowser;
+
+	UToolMenus::Get()->RegisterStartupCallback(FSimpleMulticastDelegate::FDelegate::CreateLambda([]()
+	{
+		FToolMenuOwnerScoped OwnerScoped(UE_MODULE_NAME);
+		
+		const FName ContentBrowserNiagaraTagMenuName = TEXT("NiagaraEditorModule.ContentBrowserNiagaraTags");
+		UToolMenu* AllTagsMenu = UToolMenus::Get()->RegisterMenu(ContentBrowserNiagaraTagMenuName);
+		
+		AllTagsMenu->AddDynamicSection(NAME_None, FNewToolMenuDelegate::CreateLambda([](UToolMenu* InMenu)
+		{
+			// We read from and write to the context
+			UNiagaraTagsContentBrowserFilterContext* Context = InMenu->FindContext<UNiagaraTagsContentBrowserFilterContext>();
+
+			if(Context == nullptr)
+			{
+				return;
+			}
+
+			TArray<FStructuredAssetTagDefinitionLookupData> AssetTagDefinitionData = GetStructuredSortedAssetTagDefinitions();
+
+			auto GetCheckState = [](const FNiagaraAssetTagDefinition& AssetTagDefinition, const FToolMenuContext& Context) -> ECheckBoxState
+			{
+				return Context.FindContext<UNiagaraTagsContentBrowserFilterContext>()->FilterData->ContainsActiveTagGuid(AssetTagDefinition.TagGuid) ? ECheckBoxState::Checked : ECheckBoxState::Unchecked;
+			};
+
+			auto AddTag = [](const FNiagaraAssetTagDefinition& AssetTagDefinition, const FToolMenuContext& Context)
+			{
+				Context.FindContext<UNiagaraTagsContentBrowserFilterContext>()->FilterData->AddTagGuid(AssetTagDefinition.TagGuid);
+			};
+
+			auto RemoveTag = [](const FNiagaraAssetTagDefinition& AssetTagDefinition, const FToolMenuContext& Context)
+			{
+				Context.FindContext<UNiagaraTagsContentBrowserFilterContext>()->FilterData->RemoveTagGuid(AssetTagDefinition.TagGuid);
+			};
+
+			// We create 3 menu sections: core Niagara tag filters, Project specified filters, and filters coming from other sources such as Engine Plugins (other than Niagara)
+			const FText NiagaraSectionName = GetAssetTagSectionNameFromSource(EAssetTagSectionSource::NiagaraInternal);
+			const FText ProjectSectionName = GetAssetTagSectionNameFromSource(EAssetTagSectionSource::Project);
+			const FText OtherSectionName = GetAssetTagSectionNameFromSource(EAssetTagSectionSource::Other);
+
+			{
+				FToolMenuSection& NiagaraSection = InMenu->AddSection(FName(NiagaraSectionName.ToString()), NiagaraSectionName);
+				NiagaraSection.InsertPosition = FToolMenuInsert(NAME_None, EToolMenuInsertType::First);
+			}
+
+			{
+				FToolMenuSection& ProjectSection = InMenu->AddSection(FName(ProjectSectionName.ToString()), ProjectSectionName);
+				ProjectSection.InsertPosition = FToolMenuInsert(FName(NiagaraSectionName.ToString()), EToolMenuInsertType::After);
+			}
+
+			{
+				FToolMenuSection& Section = InMenu->AddSection(FName(OtherSectionName.ToString()), OtherSectionName);
+				Section.InsertPosition = FToolMenuInsert(FName(ProjectSectionName.ToString()), EToolMenuInsertType::After);
+			}
+
+			for(const FStructuredAssetTagDefinitionLookupData& AssetTagDefinitionLookupData : AssetTagDefinitionData)
+			{
+				FName SectionName = FName(GetAssetTagSectionNameFromSource(AssetTagDefinitionLookupData.Source).ToString());
+				
+				for(const FNiagaraAssetTagDefinition& AssetTagDefinition : AssetTagDefinitionLookupData.AssetTagDefinitions)
+				{
+					FToolMenuSection* TagSection = InMenu->FindSection(SectionName);
+					
+					FText AssetTagLabel = AssetTagDefinition.AssetTag;
+					FText AssetTagDescription = AssetTagDefinition.Description;
+					
+					FToolUIAction Action;
+					Action.ExecuteAction = FToolMenuExecuteAction::CreateLambda([AssetTagDefinition, &GetCheckState, &AddTag, &RemoveTag](const FToolMenuContext& Context)
+					{
+						UNiagaraTagsContentBrowserFilterContext* NiagaraTagsContentBrowserFilterContext = Context.FindContext<UNiagaraTagsContentBrowserFilterContext>();
+
+						if(NiagaraTagsContentBrowserFilterContext == nullptr)
+						{
+							return;
+						}
+					
+						ECheckBoxState CurrentState = GetCheckState(AssetTagDefinition, Context);
+
+						if(CurrentState == ECheckBoxState::Checked)
+						{
+							RemoveTag(AssetTagDefinition, Context);
+						}
+						else if(CurrentState == ECheckBoxState::Unchecked)
+						{
+							AddTag(AssetTagDefinition, Context);
+						}
+					});
+					Action.GetActionCheckState = FToolMenuGetActionCheckState::CreateLambda([AssetTagDefinition, &GetCheckState](const FToolMenuContext& Context)
+					{
+						UNiagaraTagsContentBrowserFilterContext* NiagaraTagsContentBrowserFilterContext = Context.FindContext<UNiagaraTagsContentBrowserFilterContext>();
+
+						if(NiagaraTagsContentBrowserFilterContext == nullptr)
+						{
+							return ECheckBoxState::Unchecked;
+						}
+						
+						return GetCheckState(AssetTagDefinition, Context);
+					});
+					
+					FToolMenuEntry Entry = FToolMenuEntry::InitMenuEntry(
+						FName(AssetTagLabel.ToString()),
+						AssetTagLabel,
+						AssetTagDescription,
+						FSlateIcon(),
+						FToolUIActionChoice(Action), EUserInterfaceActionType::ToggleButton);
+					TagSection->AddEntry(Entry);
+				}
+			}
+		}));
+	}));
+}
+
+void NiagaraEditorMenuHelpers::RegisterMenuExtensions()
+{
+	using namespace FNiagaraEditorUtilities::AssetBrowser;
+	
+	UToolMenus::Get()->RegisterStartupCallback(FSimpleMulticastDelegate::FDelegate::CreateLambda([]()
+	{
+		FToolMenuOwnerScoped OwnerScoped(UE_MODULE_NAME);
+
+		// First we generate a standalone menu that will populate with all available tags dynamically
+		const FName StandaloneTagManagementMenuName = TEXT("NiagaraEditorModule.ManageAssetTags");
+		UToolMenu* TagMenu = UToolMenus::Get()->RegisterMenu(StandaloneTagManagementMenuName);
+		TagMenu->AddDynamicSection(NAME_None, FNewToolMenuDelegate::CreateLambda([](UToolMenu* InMenu)
+		{
+			TArray<FStructuredAssetTagDefinitionLookupData> AssetTagDefinitionData = GetStructuredSortedAssetTagDefinitions();
+
+			UContentBrowserAssetContextMenuContext* AssetContextMenuContext = InMenu->FindContext<UContentBrowserAssetContextMenuContext>();
+
+			if(AssetContextMenuContext == nullptr)
+			{
+				return;
+			}
+
+			auto GetCheckState = [](const FNiagaraAssetTagDefinition& AssetTagDefinition, TArray<FAssetData> SelectedAssets) -> ECheckBoxState
+			{
+				bool bAnyContainsTag = false;
+				bool bAnyDoesNotContainTag = false;
+				for(const FAssetData& SelectedAsset : SelectedAssets)
+				{
+					if(SelectedAsset.GetClass() == UNiagaraEmitter::StaticClass())
+					{
+						if(SelectedAsset.IsAssetLoaded())
+						{
+							UNiagaraEmitter* Emitter = Cast<UNiagaraEmitter>(SelectedAsset.GetAsset());
+							if(Emitter->AssetTags.Contains(AssetTagDefinition))
+							{
+								bAnyContainsTag = true;
+							}
+							else
+							{
+								bAnyDoesNotContainTag = true;
+							}
+						}
+						else
+						{
+							if(AssetTagDefinition.DoesAssetDataContainTag(SelectedAsset))
+							{
+								bAnyContainsTag = true;
+							}
+							else
+							{
+								bAnyDoesNotContainTag = true;
+							}
+						}
+					}
+					if(SelectedAsset.GetClass() == UNiagaraSystem::StaticClass())
+					{
+						if(SelectedAsset.IsAssetLoaded())
+						{
+							UNiagaraSystem* System = Cast<UNiagaraSystem>(SelectedAsset.GetAsset());
+							if(System->AssetTags.Contains(AssetTagDefinition))
+							{
+								bAnyContainsTag = true;
+							}
+							else
+							{
+								bAnyDoesNotContainTag = true;
+							}
+						}
+
+						else
+						{
+							if(AssetTagDefinition.DoesAssetDataContainTag(SelectedAsset))
+							{
+								bAnyContainsTag = true;
+							}
+							else
+							{
+								bAnyDoesNotContainTag = true;
+							}
+						}						
+					}
+					if(SelectedAsset.GetClass() == UNiagaraScript::StaticClass())
+					{
+						if(SelectedAsset.IsAssetLoaded())
+						{
+							UNiagaraScript* Script = Cast<UNiagaraScript>(SelectedAsset.GetAsset());
+							if(Script->GetLatestScriptData()->AssetTagDefinitionReferences.Contains(AssetTagDefinition))
+							{
+								bAnyContainsTag = true;
+							}
+							else
+							{
+								bAnyDoesNotContainTag = true;
+							}
+						}
+						else
+						{
+							if(AssetTagDefinition.DoesAssetDataContainTag(SelectedAsset))
+							{
+								bAnyContainsTag = true;
+							}
+							else
+							{
+								bAnyDoesNotContainTag = true;
+							}
+						}
+					}
+
+					if(bAnyContainsTag && bAnyDoesNotContainTag)
+					{
+						return ECheckBoxState::Undetermined;
+					}
+				}
+
+				if(bAnyContainsTag && bAnyDoesNotContainTag == false)
+				{
+					return ECheckBoxState::Checked;
+				}
+
+				return ECheckBoxState::Unchecked;
+			};
+
+			auto AddTag = [](const FNiagaraAssetTagDefinition& AssetTagDefinition, TArray<FAssetData> SelectedAssets)
+			{
+				FAssetRegistryModule& AssetRegistryModule = FModuleManager::LoadModuleChecked<FAssetRegistryModule>("AssetRegistry");
+
+				for(const FAssetData& SelectedAsset : SelectedAssets)
+				{
+					if(SelectedAsset.GetClass() == UNiagaraEmitter::StaticClass())
+					{
+						UNiagaraEmitter* Emitter = Cast<UNiagaraEmitter>(SelectedAsset.GetAsset());
+						if(Emitter->AssetTags.Contains(AssetTagDefinition))
+						{
+							continue;
+						}
+						else
+						{
+							Emitter->Modify();
+							Emitter->AssetTags.Add(AssetTagDefinition);
+						}
+					}
+					if(SelectedAsset.GetClass() == UNiagaraSystem::StaticClass())
+					{
+						UNiagaraSystem* System = Cast<UNiagaraSystem>(SelectedAsset.GetAsset());
+						if(System->AssetTags.Contains(AssetTagDefinition))
+						{
+							continue;
+						}
+						else
+						{
+							System->Modify();
+							System->AssetTags.Add(AssetTagDefinition);
+						}
+					}
+					if(SelectedAsset.GetClass() == UNiagaraScript::StaticClass())
+					{
+						UNiagaraScript* Script = Cast<UNiagaraScript>(SelectedAsset.GetAsset());
+						if(Script->GetLatestScriptData()->AssetTagDefinitionReferences.Contains(AssetTagDefinition))
+						{
+							continue;
+						}
+						else
+						{
+							Script->Modify();
+							Script->GetLatestScriptData()->AssetTagDefinitionReferences.Add(AssetTagDefinition);
+						}
+					}
+
+					AssetRegistryModule.Get().AssetUpdateTags(SelectedAsset.GetAsset(), EAssetRegistryTagsCaller::Fast);
+				}
+			};
+			
+			auto RemoveTag = [](const FNiagaraAssetTagDefinition& AssetTagDefinition, TArray<FAssetData> SelectedAssets)
+			{
+				FAssetRegistryModule& AssetRegistryModule = FModuleManager::LoadModuleChecked<FAssetRegistryModule>("AssetRegistry");
+
+				for(const FAssetData& SelectedAsset : SelectedAssets)
+				{
+					if(SelectedAsset.GetClass() == UNiagaraEmitter::StaticClass())
+					{
+						UNiagaraEmitter* Emitter = Cast<UNiagaraEmitter>(SelectedAsset.GetAsset());
+						if(Emitter->AssetTags.Contains(AssetTagDefinition) == false)
+						{
+							continue;
+						}
+						else
+						{
+							Emitter->Modify();
+							Emitter->AssetTags.Remove(AssetTagDefinition);
+						}
+					}
+					if(SelectedAsset.GetClass() == UNiagaraSystem::StaticClass())
+					{
+						UNiagaraSystem* System = Cast<UNiagaraSystem>(SelectedAsset.GetAsset());
+						if(System->AssetTags.Contains(AssetTagDefinition) == false)
+						{
+							continue;
+						}
+						else
+						{
+							System->Modify();
+							System->AssetTags.Remove(AssetTagDefinition);
+						}
+					}
+					if(SelectedAsset.GetClass() == UNiagaraScript::StaticClass())
+					{
+						UNiagaraScript* Script = Cast<UNiagaraScript>(SelectedAsset.GetAsset());
+						if(Script->GetLatestScriptData()->AssetTagDefinitionReferences.Contains(AssetTagDefinition) == false)
+						{
+							continue;
+						}
+						else
+						{
+							Script->Modify();
+							Script->GetLatestScriptData()->AssetTagDefinitionReferences.Remove(AssetTagDefinition);
+						}
+					}
+
+					AssetRegistryModule.Get().AssetUpdateTags(SelectedAsset.GetAsset(), EAssetRegistryTagsCaller::Fast);
+				}
+			};
+
+			// We create 3 menu sections: core Niagara tag filters, Project specified filters, and filters coming from other sources such as Engine Plugins (other than Niagara)
+			const FText NiagaraSectionName = GetAssetTagSectionNameFromSource(EAssetTagSectionSource::NiagaraInternal);
+			const FText ProjectSectionName = GetAssetTagSectionNameFromSource(EAssetTagSectionSource::Project);
+			const FText OtherSectionName = GetAssetTagSectionNameFromSource(EAssetTagSectionSource::Other);
+
+			{
+				FToolMenuSection& NiagaraSection = InMenu->AddSection(FName(NiagaraSectionName.ToString()), NiagaraSectionName);
+				NiagaraSection.InsertPosition = FToolMenuInsert(NAME_None, EToolMenuInsertType::First);
+			}
+
+			{
+				FToolMenuSection& ProjectSection = InMenu->AddSection(FName(ProjectSectionName.ToString()), ProjectSectionName);
+				ProjectSection.InsertPosition = FToolMenuInsert(FName(NiagaraSectionName.ToString()), EToolMenuInsertType::After);
+			}
+
+			{
+				FToolMenuSection& Section = InMenu->AddSection(FName(OtherSectionName.ToString()), OtherSectionName);
+				Section.InsertPosition = FToolMenuInsert(FName(ProjectSectionName.ToString()), EToolMenuInsertType::After);
+			}
+
+			for(const FStructuredAssetTagDefinitionLookupData& AssetTagDefinitionLookupData : AssetTagDefinitionData)
+			{
+				FName SectionName = FName(GetAssetTagSectionNameFromSource(AssetTagDefinitionLookupData.Source).ToString());
+				
+				for(const FNiagaraAssetTagDefinition& AssetTagDefinition : AssetTagDefinitionLookupData.AssetTagDefinitions)
+				{
+					if(AssetTagDefinition.GetSupportedClasses().Contains(AssetContextMenuContext->CommonClass))
+					{
+						FToolMenuSection* TagSection = InMenu->FindSection(SectionName);
+						
+						FText AssetTagLabel = AssetTagDefinition.AssetTag;
+						FText AssetTagDescription = AssetTagDefinition.Description;
+						
+						FToolUIAction Action;
+						Action.ExecuteAction = FToolMenuExecuteAction::CreateLambda([AssetTagDefinition, &GetCheckState, &AddTag, &RemoveTag](const FToolMenuContext& Context)
+						{
+							UContentBrowserAssetContextMenuContext* AssetContextMenuContext = Context.FindContext<UContentBrowserAssetContextMenuContext>();
+
+							if(AssetContextMenuContext == nullptr)
+							{
+								return;
+							}
+						
+							ECheckBoxState CurrentState = GetCheckState(AssetTagDefinition, AssetContextMenuContext->SelectedAssets);
+
+							if(CurrentState == ECheckBoxState::Checked)
+							{
+								RemoveTag(AssetTagDefinition, AssetContextMenuContext->SelectedAssets);
+							}
+							else if(CurrentState == ECheckBoxState::Unchecked)
+							{
+								AddTag(AssetTagDefinition, AssetContextMenuContext->SelectedAssets);
+							}
+							// if there is a mix of states, we are removing the tags
+							else if(CurrentState == ECheckBoxState::Undetermined)
+							{
+								RemoveTag(AssetTagDefinition, AssetContextMenuContext->SelectedAssets);
+							}
+						});
+						Action.GetActionCheckState = FToolMenuGetActionCheckState::CreateLambda([AssetTagDefinition, &GetCheckState](const FToolMenuContext& Context)
+						{
+							UContentBrowserAssetContextMenuContext* AssetContextMenuContext = Context.FindContext<UContentBrowserAssetContextMenuContext>();
+
+							if(AssetContextMenuContext == nullptr)
+							{
+								return ECheckBoxState::Unchecked;
+							}
+							
+							return GetCheckState(AssetTagDefinition, AssetContextMenuContext->SelectedAssets);
+						});
+						
+						FToolMenuEntry Entry = FToolMenuEntry::InitMenuEntry(
+							FName(AssetTagLabel.ToString()),
+							AssetTagLabel,
+							AssetTagDescription,
+							FSlateIcon(),
+							FToolUIActionChoice(Action), EUserInterfaceActionType::ToggleButton);
+						TagSection->AddEntry(Entry);
+					}
+				}
+			}
+		}));
+
+		for(UClass* ExtendedClass : {UNiagaraEmitter::StaticClass(), UNiagaraSystem::StaticClass(), UNiagaraScript::StaticClass()})
+		{
+			// Then we extend the asset context menu by adding an empty sub menu
+			const FName SubMenuName = TEXT("ManageTags");
+			UToolMenu* AssetContextMenu = UE::ContentBrowser::ExtendToolMenu_AssetContextMenu(ExtendedClass);
+			AssetContextMenu->AddSubMenu(OwnerScoped.GetOwner(), "GetAssetActions", SubMenuName, LOCTEXT("ManageTabsSubMenuLabel", "Manage Tags"));
+			
+			// Now that we have added a submenu with a specific name, we need the full asset context menu path including the same sub menu name
+			const FName AssetContextMenuNameWithSubMenu = UToolMenus::JoinMenuPaths(AssetContextMenu->GetMenuName(), SubMenuName);
+
+			// Then, we register the standalone menu as the parent of the submenu. This will 'link' the empty submenu with our standalone menu.
+			UToolMenus::Get()->RegisterMenu(AssetContextMenuNameWithSubMenu, StandaloneTagManagementMenuName);
+		}
+	}));
+}
+
+#undef LOCTEXT_NAMESPACE
diff --git a/Engine/Plugins/FX/Niagara/Source/NiagaraEditor/Private/NiagaraEditorMenuHelpers.h b/Engine/Plugins/FX/Niagara/Source/NiagaraEditor/Private/NiagaraEditorMenuHelpers.h
new file mode 100644
--- /dev/null
+++ b/Engine/Plugins/FX/Niagara/Source/NiagaraEditor/Private/NiagaraEditorMenuHelpers.h
@@ -0,0 +1,9 @@
+﻿// Copyright Epic Games, Inc. All Rights Reserved.
+
+#pragma once
+
+namespace NiagaraEditorMenuHelpers
+{
+	void RegisterToolMenus();
+	void RegisterMenuExtensions();
+}
\ No newline at end of file
diff --git a/Engine/Plugins/FX/Niagara/Source/NiagaraEditor/Private/NiagaraEditorModule.cpp b/Engine/Plugins/FX/Niagara/Source/NiagaraEditor/Private/NiagaraEditorModule.cpp
--- a/Engine/Plugins/FX/Niagara/Source/NiagaraEditor/Private/NiagaraEditorModule.cpp
+++ b/Engine/Plugins/FX/Niagara/Source/NiagaraEditor/Private/NiagaraEditorModule.cpp
@@ -5,6 +5,7 @@
 #include "NiagaraEditorTickables.h"
 #include "Modules/ModuleManager.h"
 #include "AssetToolsModule.h"
+#include "ContentBrowserMenuContexts.h"
 #include "ISequencerModule.h"
 #include "ISettingsModule.h"
 #include "SequencerChannelInterface.h"
@@ -148,6 +149,7 @@
 #include "ISourceControlModule.h"
 #include "DeviceProfiles/DeviceProfileManager.h"
 #include "Containers/Ticker.h"
+#include "UObject/AssetRegistryTagsContext.h"
 
 #include "ViewModels/Stack/NiagaraStackObjectIssueGenerator.h"
 #include "NiagaraPlatformSet.h"
@@ -159,11 +161,14 @@
 #include "Filters/CustomClassFilterData.h"
 
 #include "Widgets/SNiagaraDebugger.h"
+#include "Widgets/AssetBrowser/NiagaraAssetBrowserConfig.h"
+#include "Widgets/AssetBrowser/SNiagaraAssetBrowser.h"
 
 #include "NiagaraDebugVis.h"
 #include "NiagaraPerfBaseline.h"
 #include "NiagaraGraphDataCache.h"
 #include "NiagaraDecalRendererProperties.h"
+#include "NiagaraEditorMenuHelpers.h"
 #include "NiagaraLightRendererProperties.h"
 #include "NiagaraRibbonRendererProperties.h"
 #include "NiagaraSpriteRendererProperties.h"
@@ -171,6 +176,8 @@
 
 #include "Engine/AssetManager.h"
 #include "ViewModels/HierarchyEditor/NiagaraHierarchyCommands.h"
+#include "Widgets/AssetBrowser/SNiagaraSelectedAssetDetails.h"
+#include "NiagaraRecentAndFavoritesManager.h"
 
 #include UE_INLINE_GENERATED_CPP_BY_NAME(NiagaraEditorModule)
 
@@ -1036,13 +1043,17 @@ void FNiagaraEditorModule::StartupModule()
 	MenuExtensibilityManager = MakeShareable(new FExtensibilityManager);
 	ToolBarExtensibilityManager = MakeShareable(new FExtensibilityManager);
 
+	RecentAndFavoritesManager = MakeShared<FNiagaraRecentAndFavoritesManager>();
+	RecentAndFavoritesManager->Initialize();
+
+	UNiagaraAssetBrowserConfig::Initialize();
+
+	FNiagaraAssetDetailDatabase::Init();
+
 	FContentBrowserModule& ContentBrowserModule = FModuleManager::LoadModuleChecked<FContentBrowserModule>(TEXT("ContentBrowser"));
 	ContentBrowserModule.AddDynamicTagAssetClass(TEXT("NiagaraSystem"));
 	ContentBrowserModule.AddDynamicTagAssetClass(TEXT("NiagaraEmitter"));
 
-
-	IAssetTools& AssetTools = FModuleManager::LoadModuleChecked<FAssetToolsModule>("AssetTools").Get();
-
 	// Preload all parameter definition & collection assets so that they will be postloaded before postload calls to scripts/emitters/systems that rely on them.
 	{
 		FAssetRegistryModule& AssetRegistryModule = FModuleManager::LoadModuleChecked<FAssetRegistryModule>("AssetRegistry");
@@ -1230,6 +1241,8 @@ void FNiagaraEditorModule::StartupModule()
 	//Register Stack Object Issue Generators.
 	RegisterStackIssueGenerator(FNiagaraPlatformSet::StaticStruct()->GetFName(), new FNiagaraPlatformSetIssueGenerator());
 
+	NiagaraEditorMenuHelpers::RegisterToolMenus();
+	NiagaraEditorMenuHelpers::RegisterMenuExtensions();
 
 	FNiagaraEditorStyle::Register();
 	ReinitializeStyleCommand = IConsoleManager::Get().RegisterConsoleCommand(
@@ -1506,6 +1519,9 @@ void FNiagaraEditorModule::ShutdownModule()
 	MenuExtensibilityManager.Reset();
 	ToolBarExtensibilityManager.Reset();
 
+	RecentAndFavoritesManager->Shutdown();
+	RecentAndFavoritesManager.Reset();
+	
 	// Clean up asset registry callbacks
 	if (FModuleManager::Get().IsModuleLoaded("AssetRegistry"))
 	{
@@ -1728,6 +1744,11 @@ FNiagaraEditorModule& FNiagaraEditorModule::Get()
 	return FModuleManager::LoadModuleChecked<FNiagaraEditorModule>("NiagaraEditor");
 }
 
+FNiagaraRecentAndFavoritesManager* FNiagaraEditorModule::GetRecentsManager()
+{
+	return RecentAndFavoritesManager.Get();
+}
+
 void FNiagaraEditorModule::OnNiagaraSettingsChangedEvent(const FName& PropertyName, const UNiagaraSettings* Settings)
 {
 	if (PropertyName == GET_MEMBER_NAME_CHECKED(UNiagaraSettings, AdditionalParameterTypes)
diff --git a/Engine/Plugins/FX/Niagara/Source/NiagaraEditor/Private/NiagaraEditorStyle.cpp b/Engine/Plugins/FX/Niagara/Source/NiagaraEditor/Private/NiagaraEditorStyle.cpp
--- a/Engine/Plugins/FX/Niagara/Source/NiagaraEditor/Private/NiagaraEditorStyle.cpp
+++ b/Engine/Plugins/FX/Niagara/Source/NiagaraEditor/Private/NiagaraEditorStyle.cpp
@@ -106,6 +106,31 @@ void FNiagaraEditorStyle::InitAssetPicker()
 
 }
 
+void FNiagaraEditorStyle::InitAssetBrowser()
+{
+	const FTextBlockStyle NormalText = FAppStyle::GetWidgetStyle<FTextBlockStyle>("NormalText");
+
+	Set("NiagaraEditor.AssetBrowser.PropertySeparator", new FSlateColorBrush(FStyleColors::White25));
+
+	FTextBlockStyle AssetBrowserAssetTitleStyle = FAppStyle::GetWidgetStyle<FTextBlockStyle>("NormalText.Important");
+	AssetBrowserAssetTitleStyle.Font.Size = 12.f;
+	Set("NiagaraEditor.AssetBrowser.AssetTitle", AssetBrowserAssetTitleStyle);
+	
+	FTextBlockStyle AssetBrowserAssetTypeStyle = FAppStyle::GetWidgetStyle<FTextBlockStyle>("NormalText.Subdued");
+	AssetBrowserAssetTypeStyle.Font.Size = 10.f;
+	Set("NiagaraEditor.AssetBrowser.AssetType", AssetBrowserAssetTypeStyle);
+	
+	Set("NiagaraEditor.AssetBrowser.AssetTag.OuterBorder", new FSlateRoundedBoxBrush(FStyleColors::AccentGray, 4.f));
+	Set("NiagaraEditor.AssetBrowser.AssetTag.InnerBorder", new FSlateRoundedBoxBrush(FStyleColors::Black, 4.f));
+
+	FTextBlockStyle AssetTagText = NormalText;
+
+	FSlateFontInfo Font = FStyleFonts::Get().Normal;
+	Font.Size = 10.f;
+	AssetTagText.SetFont(Font);
+	Set("NiagaraEditor.AssetBrowser.AssetTag.Text", AssetTagText);
+}
+
 void FNiagaraEditorStyle::InitActionMenu()
 {
 	const FTextBlockStyle NormalText = FAppStyle::GetWidgetStyle<FTextBlockStyle>("NormalText");
@@ -878,6 +903,7 @@ FNiagaraEditorStyle::FNiagaraEditorStyle() : FSlateStyleSet("NiagaraEditorStyle"
 
 	InitStats();
 	InitAssetPicker();
+	InitAssetBrowser();
 	InitActionMenu();
 	InitEmitterHeader();
 	InitParameters();
diff --git a/Engine/Plugins/FX/Niagara/Source/NiagaraEditor/Private/NiagaraEditorUtilities.cpp b/Engine/Plugins/FX/Niagara/Source/NiagaraEditor/Private/NiagaraEditorUtilities.cpp
--- a/Engine/Plugins/FX/Niagara/Source/NiagaraEditor/Private/NiagaraEditorUtilities.cpp
+++ b/Engine/Plugins/FX/Niagara/Source/NiagaraEditor/Private/NiagaraEditorUtilities.cpp
@@ -6,6 +6,7 @@
 #include "AssetToolsModule.h"
 #include "ContentBrowserModule.h"
 #include "EdGraphSchema_Niagara.h"
+#include "GeneralProjectSettings.h"
 #include "Styling/AppStyle.h"
 #include "IContentBrowserSingleton.h"
 #include "INiagaraEditorTypeUtilities.h"
@@ -30,6 +31,7 @@
 #include "NiagaraOverviewNode.h"
 #include "NiagaraParameterMapHistory.h"
 #include "NiagaraParameterDefinitions.h"
+#include "NiagaraRecentAndFavoritesManager.h"
 #include "NiagaraScript.h"
 #include "ViewModels/NiagaraScriptGraphViewModel.h"
 #include "NiagaraScriptMergeManager.h"
@@ -1582,6 +1584,25 @@ ENiagaraScriptLibraryVisibility FNiagaraEditorUtilities::GetScriptAssetVisibilit
 	return ScriptVisibility;
 }
 
+bool FNiagaraEditorUtilities::GetIsInheritableFromAssetRegistryTags(const FAssetData& AssetData, bool& bUseInheritance)
+{
+	if(AssetData.GetTagValue(GET_MEMBER_NAME_CHECKED(UNiagaraEmitter, bIsInheritable), bUseInheritance))
+	{
+		return true;		
+	}
+	else
+	{
+		ENiagaraScriptTemplateSpecification DeprecatedTemplateTag;
+		if(GetTemplateSpecificationFromTag(AssetData, DeprecatedTemplateTag))
+		{
+			bUseInheritance = DeprecatedTemplateTag == ENiagaraScriptTemplateSpecification::None;
+			return true;
+		}
+	}
+
+	return false;
+}
+
 // this function is used as an overload inside FAssetData.GetTagValue for the ENiagaraScriptTemplateSpecification enum
 void LexFromString(ENiagaraScriptTemplateSpecification& OutValue, const TCHAR* Buffer)
 {
@@ -1591,7 +1612,7 @@ void LexFromString(ENiagaraScriptTemplateSpecification& OutValue, const TCHAR* B
 bool FNiagaraEditorUtilities::GetTemplateSpecificationFromTag(const FAssetData& Data, ENiagaraScriptTemplateSpecification& OutTemplateSpecification)
 {
 	ENiagaraScriptTemplateSpecification TemplateSpecification = ENiagaraScriptTemplateSpecification::None;
-	bool bTemplateEnumTagFound = Data.GetTagValue(GET_MEMBER_NAME_CHECKED(UNiagaraEmitter, TemplateSpecification), TemplateSpecification);
+	bool bTemplateEnumTagFound = Data.GetTagValue(FName("TemplateSpecification"), TemplateSpecification);
 
 	if(bTemplateEnumTagFound)
 	{
@@ -1835,6 +1856,7 @@ int32 FNiagaraEditorUtilities::GetWeightForItem(const TSharedPtr<FNiagaraMenuAct
 		}
 	}
 
+	// NPC Variables are rarely wanted so we halve their weight here
 	if(InCurrentAction->GetParameterVariable().IsSet())
 	{
 		if(InCurrentAction->GetParameterVariable().GetValue().IsInNameSpace(FNiagaraConstants::ParameterCollectionNamespaceString))
@@ -2060,6 +2082,8 @@ TArray<UNiagaraComponent*> FNiagaraEditorUtilities::GetComponentsThatReferenceSy
 
 const FGuid FNiagaraEditorUtilities::AddEmitterToSystem(UNiagaraSystem& InSystem, UNiagaraEmitter& InEmitterToAdd, FGuid EmitterVersion, bool bCreateCopy)
 {
+	FNiagaraEditorModule::Get().GetRecentsManager()->EmitterUsed(InEmitterToAdd);
+
 	// Kill all system instances before modifying the emitter handle list to prevent accessing deleted data.
 	KillSystemInstances(InSystem);
 
@@ -4765,6 +4789,207 @@ TMap<FNiagaraVariableBase, FGuid> FNiagaraEditorUtilities::Scripts::Validation::
 	return RemappedGuids;
 }
 
+TArray<FNiagaraEditorUtilities::AssetBrowser::FStructuredAssetTagDefinitionLookupData> FNiagaraEditorUtilities::AssetBrowser::GetStructuredSortedAssetTagDefinitions(bool bSortTagsByName)
+{
+	TArray<FStructuredAssetTagDefinitionLookupData> Result;
+	
+	// Internal Tags
+	{
+		FStructuredAssetTagDefinitionLookupData TagDefinitionData;
+		TagDefinitionData.DefinitionsAsset = nullptr;
+
+		for(const FNiagaraAssetTagDefinition* InternalAssetTagDefinition : INiagaraModule::Get().GetInternalAssetTagDefinitions())
+		{
+			TagDefinitionData.AssetTagDefinitions.Add(*InternalAssetTagDefinition);
+		}
+		
+		if(bSortTagsByName)
+		{
+			TagDefinitionData.AssetTagDefinitions.Sort();
+		}
+		TagDefinitionData.Source = EAssetTagSectionSource::NiagaraInternal;
+		Result.Add(TagDefinitionData);
+	}
+
+	FAssetRegistryModule& AssetRegistryModule = FModuleManager::Get().LoadModuleChecked<FAssetRegistryModule>("AssetRegistry");
+	FARFilter Filter;
+	Filter.ClassPaths.Add(UNiagaraAssetTagDefinitions::StaticClass()->GetClassPathName());
+
+	TArray<FAssetData> FilteredTagDefinitionAssets;
+	AssetRegistryModule.Get().GetAssets(Filter, FilteredTagDefinitionAssets);
+
+	for(const FAssetData& AssetData : FilteredTagDefinitionAssets)
+	{
+		UNiagaraAssetTagDefinitions* TagDefinitions = Cast<UNiagaraAssetTagDefinitions>(AssetData.GetAsset());
+		FStructuredAssetTagDefinitionLookupData TagDefinitionData;
+		TagDefinitionData.DefinitionsAsset = TagDefinitions;
+		TagDefinitionData.AssetTagDefinitions = TagDefinitions->GetAssetTagDefinitions();
+		if(bSortTagsByName)
+		{
+			TagDefinitionData.AssetTagDefinitions.Sort();
+		}
+		TagDefinitionData.Source = GetAssetTagDefinitionSource(AssetData);
+		Result.Add(TagDefinitionData);
+	}
+
+	Result.Sort([](const FStructuredAssetTagDefinitionLookupData& CandidateA, const FStructuredAssetTagDefinitionLookupData& CandidateB)
+	{
+		if((int32) CandidateA.Source != (int32) CandidateB.Source)
+		{
+			return (int32) CandidateA.Source < (int32) CandidateB.Source;
+		}
+
+		if(CandidateA.DefinitionsAsset == nullptr || CandidateB.DefinitionsAsset == nullptr)
+		{
+			if(CandidateB.DefinitionsAsset != nullptr)
+			{
+				return true;
+			}
+			else if(CandidateA.DefinitionsAsset != nullptr)
+			{
+				return false;
+			}
+			
+			return CandidateA.AssetTagDefinitions.Num() < CandidateB.AssetTagDefinitions.Num();
+		}
+		
+		return CandidateA.DefinitionsAsset->GetSortOrder() < CandidateB.DefinitionsAsset->GetSortOrder();
+	});
+	
+	return Result;
+}
+
+TArray<FNiagaraAssetTagDefinition> FNiagaraEditorUtilities::AssetBrowser::GetFlatSortedAssetTagDefinitions(bool bSortEndResultTagsByName)
+{
+	TArray<FNiagaraAssetTagDefinition> Result;
+	TArray<FStructuredAssetTagDefinitionLookupData> IntermediateResult;
+	FAssetRegistryModule& AssetRegistryModule = FModuleManager::Get().LoadModuleChecked<FAssetRegistryModule>("AssetRegistry");
+
+	for(const FNiagaraAssetTagDefinition* InternalAssetTagDefinition : INiagaraModule::Get().GetInternalAssetTagDefinitions())
+	{
+		Result.Add(*InternalAssetTagDefinition);
+	}
+
+	FARFilter Filter;
+	Filter.ClassPaths.Add(UNiagaraAssetTagDefinitions::StaticClass()->GetClassPathName());
+
+	TArray<FAssetData> AllTagDefinitionAssets;
+	AssetRegistryModule.Get().GetAssets(Filter, AllTagDefinitionAssets);
+
+	for(const FAssetData& AssetData : AllTagDefinitionAssets)
+	{
+		UNiagaraAssetTagDefinitions* TagDefinitions = Cast<UNiagaraAssetTagDefinitions>(AssetData.GetAsset());
+		FStructuredAssetTagDefinitionLookupData TagDefinitionData;
+		TagDefinitionData.DefinitionsAsset = TagDefinitions;
+		TagDefinitionData.AssetTagDefinitions = TagDefinitions->GetAssetTagDefinitions();
+		TagDefinitionData.Source = GetAssetTagDefinitionSource(AssetData);
+		IntermediateResult.Add(TagDefinitionData);
+	}
+	
+	IntermediateResult.Sort([](const FStructuredAssetTagDefinitionLookupData& CandidateA, const FStructuredAssetTagDefinitionLookupData& CandidateB)
+	{
+		if((int32) CandidateA.Source != (int32) CandidateB.Source)
+		{
+			return (int32) CandidateA.Source < (int32) CandidateB.Source;
+		}
+
+		if(CandidateA.DefinitionsAsset == nullptr || CandidateB.DefinitionsAsset == nullptr)
+		{
+			if(CandidateB.DefinitionsAsset != nullptr)
+			{
+				return true;
+			}
+			else if(CandidateA.DefinitionsAsset != nullptr)
+			{
+				return false;
+			}
+			
+			return CandidateA.AssetTagDefinitions.Num() < CandidateB.AssetTagDefinitions.Num();
+		}
+		
+		return CandidateA.DefinitionsAsset->GetSortOrder() < CandidateB.DefinitionsAsset->GetSortOrder();
+	});
+
+	for(const FStructuredAssetTagDefinitionLookupData& IntermediateData : IntermediateResult)
+	{
+		Result.Append(IntermediateData.AssetTagDefinitions);
+	}
+
+	if(bSortEndResultTagsByName)
+	{
+		Result.Sort();
+	}
+	
+	return Result;
+}
+
+const FNiagaraAssetTagDefinition& FNiagaraEditorUtilities::AssetBrowser::FindTagDefinitionForReference(const FNiagaraAssetTagDefinitionReference& Reference)
+{
+	TMap<FGuid, FNiagaraAssetTagDefinition> Result;
+	FAssetRegistryModule& AssetRegistryModule = FModuleManager::Get().LoadModuleChecked<FAssetRegistryModule>("AssetRegistry");
+
+	FARFilter Filter;
+	Filter.ClassPaths.Add(UNiagaraAssetTagDefinitions::StaticClass()->GetClassPathName());
+
+	const TArray<const FNiagaraAssetTagDefinition*>& InternalTagDefinitions = INiagaraModule::Get().GetInternalAssetTagDefinitions();
+	
+	for(const FNiagaraAssetTagDefinition* InternalTagDefinition : InternalTagDefinitions)
+	{
+		Result.Add(InternalTagDefinition->TagGuid, *InternalTagDefinition);
+	}
+
+	TArray<FAssetData> AssetTagDefinitionAssets;
+	AssetRegistryModule.Get().GetAssets(Filter, AssetTagDefinitionAssets);
+
+	for(const FAssetData& AssetTagDefinitionAsset : AssetTagDefinitionAssets)
+	{
+		UNiagaraAssetTagDefinitions* TagDefinitionsAsset = Cast<UNiagaraAssetTagDefinitions>(AssetTagDefinitionAsset.GetAsset());
+
+		for(const FNiagaraAssetTagDefinition& TagDefinition : TagDefinitionsAsset->GetAssetTagDefinitions())
+		{
+			Result.Add(TagDefinition.TagGuid, TagDefinition);
+		}
+	}
+
+	if(Result.Contains(Reference.GetTagDefinitionReferenceGuid()))
+	{
+		return Result[Reference.GetTagDefinitionReferenceGuid()];
+	}
+
+	static FNiagaraAssetTagDefinition DummyTagDefinition;
+	return DummyTagDefinition;
+}
+
+FNiagaraEditorUtilities::AssetBrowser::EAssetTagSectionSource FNiagaraEditorUtilities::AssetBrowser::GetAssetTagDefinitionSource(const FAssetData& AssetData)
+{	
+	ensure(AssetData.GetClass() == UNiagaraAssetTagDefinitions::StaticClass());
+
+	if(AssetData.PackagePath.ToString().StartsWith("/Niagara"))
+	{
+		return EAssetTagSectionSource::NiagaraInternal;
+	}
+	else if(AssetData.PackagePath.ToString().StartsWith("/Game"))
+	{
+		return EAssetTagSectionSource::Project;
+	}
+
+	return EAssetTagSectionSource::Other;
+}
+
+FText FNiagaraEditorUtilities::AssetBrowser::GetAssetTagSectionNameFromSource(EAssetTagSectionSource Source)
+{
+	if(Source == EAssetTagSectionSource::NiagaraInternal)
+	{
+		return FText::FromString("Niagara");
+	}
+	else if(Source == EAssetTagSectionSource::Project)
+	{
+		return FText::FromString(FApp::GetProjectName());
+	}
+
+	return LOCTEXT("NiagaraAssetTagSourceSectionName", "Misc");
+}
+
 void FNiagaraParameterUtilities::FilterToRelevantStaticVariables(TConstArrayView<FNiagaraVariable> InVars, TArray<FNiagaraVariable>& OutVars, FName InOldEmitterAlias, FName InNewEmitterAlias, bool bFilterByEmitterAliasAndConvertToUnaliased)
 {
 	FNiagaraAliasContext RenameContext(ENiagaraScriptUsage::ParticleSpawnScript);
diff --git a/Engine/Plugins/FX/Niagara/Source/NiagaraEditor/Private/NiagaraEmitterFactoryNew.cpp b/Engine/Plugins/FX/Niagara/Source/NiagaraEditor/Private/NiagaraEmitterFactoryNew.cpp
--- a/Engine/Plugins/FX/Niagara/Source/NiagaraEditor/Private/NiagaraEmitterFactoryNew.cpp
+++ b/Engine/Plugins/FX/Niagara/Source/NiagaraEditor/Private/NiagaraEmitterFactoryNew.cpp
@@ -14,11 +14,12 @@
 #include "NiagaraConstants.h"
 #include "NiagaraEmitterEditorData.h"
 #include "NiagaraNodeAssignment.h"
-#include "SNewEmitterDialog.h"
+#include "NiagaraRecentAndFavoritesManager.h"
 #include "Misc/MessageDialog.h"
 #include "Modules/ModuleManager.h"
 #include "Interfaces/IMainFrameModule.h"
 #include "Framework/Application/SlateApplication.h"
+#include "Widgets/AssetBrowser/SNiagaraAssetBrowser.h"
 
 #include UE_INLINE_GENERATED_CPP_BY_NAME(NiagaraEmitterFactoryNew)
 
@@ -40,45 +41,43 @@ bool UNiagaraEmitterFactoryNew::ConfigureProperties()
 	IMainFrameModule& MainFrame = FModuleManager::LoadModuleChecked<IMainFrameModule>("MainFrame");
 	TSharedPtr<SWindow>	ParentWindow = MainFrame.GetParentWindow();
 
-	TSharedRef<SNewEmitterDialog> NewEmitterDialog = SNew(SNewEmitterDialog);
-	FSlateApplication::Get().AddModalWindow(NewEmitterDialog, ParentWindow);
+	SNiagaraAssetBrowser::FArguments AssetBrowserArgs;
+	AssetBrowserArgs.AvailableClasses({UNiagaraEmitter::StaticClass()})
+	.RecentAndFavoritesList(FNiagaraEditorModule::Get().GetRecentsManager()->GetRecentEmitterAndSystemsList())
+	.EmptySelectionMessage(LOCTEXT("EmptyEmitterFactorySelectionMessage", "Select an emitter as a starting point for your new emitter.\n"));
 
-	if (NewEmitterDialog->GetUserConfirmedSelection() == false)
+	SNiagaraAssetBrowserWindow::FArguments AssetBrowserWindowArgs;
+	AssetBrowserWindowArgs.AssetBrowserArgs(AssetBrowserArgs)
+	.WindowTitle(LOCTEXT("EmitterAssetBrowserWindowTitle", "Create Niagara Emitter - Select an emitter as a base"));
+	
+	TSharedRef<SNiagaraCreateAssetWindow> CreateAssetBrowserWindow = SNew(SNiagaraCreateAssetWindow, *UNiagaraEmitter::StaticClass()).AssetBrowserWindowArgs(AssetBrowserWindowArgs);
+	
+	FSlateApplication::Get().AddModalWindow(CreateAssetBrowserWindow, ParentWindow);
+
+	if(CreateAssetBrowserWindow->ShouldProceedWithAction() == false)
 	{
-		// User cancelled or closed the dialog so abort asset creation.
-		return false;
+		return false;	
 	}
 
-	TOptional<FAssetData> SelectedEmitterAsset = NewEmitterDialog->GetSelectedEmitterAsset();
-	if (SelectedEmitterAsset.IsSet())
+	TArray<FAssetData> SelectedAssetData = CreateAssetBrowserWindow->GetSelectedAssets();
+	
+	if(SelectedAssetData.Num() == 1)
 	{
-		EmitterToCopy = Cast<UNiagaraEmitter>(SelectedEmitterAsset->GetAsset());
-		bUseInheritance = NewEmitterDialog->GetUseInheritance();
-		if (EmitterToCopy == nullptr)
-		{
-			EAppReturnType::Type DialogResult = FMessageDialog::Open(EAppMsgType::OkCancel, EAppReturnType::Cancel,
-				LOCTEXT("FailedToLoadMessage", "The selected emitter failed to load\nWould you like to create a default emitter?"),
-				LOCTEXT("FailedToLoadTitle", "Create Default?"));
-			if (DialogResult == EAppReturnType::Cancel)
-			{
-				return false;
-			}
-			else
-			{
-				// The selected emitter couldn't be loaded but the user still wants to create a default emitter so 
-				// well need to create a new empty emitter and add some default modules to it.
-				EmitterToCopy = nullptr;
-				bAddDefaultModulesAndRenderersToEmptyEmitter = true;
-			}
-		}
+		FAssetData SelectedAsset = SelectedAssetData[0];
+		ensure(SelectedAsset.GetClass() == UNiagaraEmitter::StaticClass());
+		
+		UNiagaraEmitter* EmitterAsset = Cast<UNiagaraEmitter>(SelectedAsset.GetAsset());
+			
+		EmitterToCopy = EmitterAsset;
+		bUseInheritance = EmitterAsset->bIsInheritable;
 	}
 	else
 	{
-		// User selected an empty emitter so set the emitter to copy to null, and disable creation of default modules.
 		EmitterToCopy = nullptr;
+		bUseInheritance = true;
 		bAddDefaultModulesAndRenderersToEmptyEmitter = false;
 	}
-
+	
 	return true;
 }
 
@@ -123,7 +122,7 @@ UObject* UNiagaraEmitterFactoryNew::FactoryCreateNew(UClass* Class, UObject* InP
 
 	if (EmitterToCopy != nullptr)
 	{
-		if (bUseInheritance && EmitterToCopy->TemplateSpecification == ENiagaraScriptTemplateSpecification::None)
+		if (bUseInheritance)
 		{
 			NewEmitter = UNiagaraEmitter::CreateWithParentAndOwner(FVersionedNiagaraEmitter(EmitterToCopy, EmitterToCopy->GetExposedVersion().VersionGuid), InParent, Name, Flags);
 		}
@@ -134,8 +133,9 @@ UObject* UNiagaraEmitterFactoryNew::FactoryCreateNew(UClass* Class, UObject* InP
 			NewEmitter->DisableVersioning(EmitterToCopy->GetExposedVersion().VersionGuid);
 		}
 
-		NewEmitter->TemplateSpecification = ENiagaraScriptTemplateSpecification::None;
+		NewEmitter->bIsInheritable = false;
 		NewEmitter->TemplateAssetDescription = FText();
+		NewEmitter->AssetTags.Empty();
 		NewEmitter->Category = FText();
 	}
 	else
@@ -150,7 +150,8 @@ UObject* UNiagaraEmitterFactoryNew::FactoryCreateNew(UClass* Class, UObject* InP
 		UNiagaraEmitterEditorData* EmitterEditorData = CastChecked<UNiagaraEmitterEditorData>(EmitterVersionData.GetEditorData());
 		EmitterEditorData->SetShowSummaryView(EmitterVersionData.AddEmitterDefaultViewState == ENiagaraEmitterDefaultSummaryState::Summary ? true : false);
 	});
-	
+
+	FNiagaraEditorModule::Get().GetRecentsManager()->EmitterUsed(*NewEmitter);
 	return NewEmitter;
 }
 
@@ -228,7 +229,6 @@ void UNiagaraEmitterFactoryNew::InitializeEmitter(UNiagaraEmitter* NewEmitter, b
 		FNiagaraStackGraphUtilities::RelayoutGraph(*Source->NodeGraph);
 		EmitterData->bInterpolatedSpawning = true;
 		EmitterData->bDeterminism = false; // NOTE: Default to non-determinism
-		NewEmitter->TemplateSpecification = ENiagaraScriptTemplateSpecification::None;
 		EmitterData->SpawnScriptProps.Script->SetUsage(ENiagaraScriptUsage::ParticleSpawnScriptInterpolated);
 	}
 }
diff --git a/Engine/Plugins/FX/Niagara/Source/NiagaraEditor/Private/NiagaraSystemFactoryNew.cpp b/Engine/Plugins/FX/Niagara/Source/NiagaraEditor/Private/NiagaraSystemFactoryNew.cpp
--- a/Engine/Plugins/FX/Niagara/Source/NiagaraEditor/Private/NiagaraSystemFactoryNew.cpp
+++ b/Engine/Plugins/FX/Niagara/Source/NiagaraEditor/Private/NiagaraSystemFactoryNew.cpp
@@ -10,15 +10,16 @@
 #include "ContentBrowserModule.h"
 #include "ImageUtils.h"
 #include "ViewModels/Stack/NiagaraStackGraphUtilities.h"
-#include "SNewSystemDialog.h"
 #include "Misc/MessageDialog.h"
 #include "ViewModels/NiagaraSystemViewModel.h"
 #include "Modules/ModuleManager.h"
 #include "Interfaces/IMainFrameModule.h"
 #include "Framework/Application/SlateApplication.h"
 #include "NiagaraEditorUtilities.h"
+#include "NiagaraRecentAndFavoritesManager.h"
 #include "NiagaraSettings.h"
 #include "ObjectTools.h"
+#include "Widgets/AssetBrowser/SNiagaraAssetBrowser.h"
 
 #include UE_INLINE_GENERATED_CPP_BY_NAME(NiagaraSystemFactoryNew)
 
@@ -37,74 +38,42 @@ bool UNiagaraSystemFactoryNew::ConfigureProperties()
 	IMainFrameModule& MainFrame = FModuleManager::LoadModuleChecked<IMainFrameModule>("MainFrame");
 	TSharedPtr<SWindow>	ParentWindow = MainFrame.GetParentWindow();
 
-	TSharedRef<SNewSystemDialog> NewSystemDialog = SNew(SNewSystemDialog);
-	FSlateApplication::Get().AddModalWindow(NewSystemDialog, ParentWindow);
+	SNiagaraAssetBrowser::FArguments AssetBrowserArgs;
+	AssetBrowserArgs.AvailableClasses({UNiagaraEmitter::StaticClass(), UNiagaraSystem::StaticClass()})
+	.RecentAndFavoritesList(FNiagaraEditorModule::Get().GetRecentsManager()->GetRecentEmitterAndSystemsList())
+	.EmptySelectionMessage(LOCTEXT("EmptySystemFactorySelectionMessage", "Select an emitter or system as a starting point for your new system.\nA system consists of one or more emitters."));
 
-	if (NewSystemDialog->GetUserConfirmedSelection() == false)
+	SNiagaraAssetBrowserWindow::FArguments AssetBrowserWindowArgs;
+	AssetBrowserWindowArgs.AssetBrowserArgs(AssetBrowserArgs)
+	.WindowTitle(LOCTEXT("SystemAssetBrowserWindowTitle", "Create Niagara System - Select an emitter or system as a base"));
+	
+	TSharedRef<SNiagaraCreateAssetWindow> CreateAssetBrowserWindow = SNew(SNiagaraCreateAssetWindow, *UNiagaraSystem::StaticClass()).AssetBrowserWindowArgs(AssetBrowserWindowArgs);
+
+	FSlateApplication::Get().AddModalWindow(CreateAssetBrowserWindow, ParentWindow);
+
+	if(CreateAssetBrowserWindow->ShouldProceedWithAction() == false)
 	{
-		// User cancelled or closed the dialog so abort asset creation.
 		return false;
 	}
 
-	TOptional<FAssetData> SelectedSystemAsset = NewSystemDialog->GetSelectedSystemAsset();
-	TArray<FAssetData> EmitterAssetsToAddToNewSystem = NewSystemDialog->GetSelectedEmitterAssets();
-	if (SelectedSystemAsset.IsSet())
+	TArray<FAssetData> SelectedAssetData = CreateAssetBrowserWindow->GetSelectedAssets();
+	
+	if(SelectedAssetData.Num() == 1)
 	{
-		SystemToCopy = Cast<UNiagaraSystem>(SelectedSystemAsset.GetValue().GetAsset());
-		if (SystemToCopy == nullptr)
+		FAssetData AssetData = SelectedAssetData[0];
+		if(AssetData.GetClass() == UNiagaraSystem::StaticClass())
 		{
-			EAppReturnType::Type DialogResult = FMessageDialog::Open(EAppMsgType::OkCancel, EAppReturnType::Cancel,
-				LOCTEXT("FailedToLoadEmitterMessage", "The selected system failed to load.\nWould you like to create an empty system?"),
-				LOCTEXT("FailedToLoadSystemTitle", "Create Default?"));
-			if (DialogResult == EAppReturnType::Cancel)
-			{
-				return false;
-			}
-			else
-			{
-				SystemToCopy = nullptr;
-				EmitterAssetsToAddToNewSystem.Empty();
-			}
+			SystemToCopy = Cast<UNiagaraSystem>(AssetData.GetAsset());
 		}
-	}
-	else if (EmitterAssetsToAddToNewSystem.Num() > 0)
-	{
-		bool bAllEmittersLoaded = true;
-		for (const FAssetData& EmitterAssetToAdd : EmitterAssetsToAddToNewSystem)
+		else if(AssetData.GetClass() == UNiagaraEmitter::StaticClass())
 		{
-			UNiagaraEmitter* EmitterToAdd = Cast<UNiagaraEmitter>(EmitterAssetToAdd.GetAsset());
-			if (EmitterToAdd != nullptr)
-			{
-				EmittersToAddToNewSystem.Add(FVersionedNiagaraEmitter(EmitterToAdd, EmitterToAdd->GetExposedVersion().VersionGuid));
-			}
-			else
-			{
-				bAllEmittersLoaded = false;
-				break;
-			}
+			UNiagaraEmitter* EmitterAsset = Cast<UNiagaraEmitter>(AssetData.GetAsset());
+			FVersionedNiagaraEmitter VersionedNiagaraEmitter(EmitterAsset, EmitterAsset->GetExposedVersion().VersionGuid);
+			
+			EmittersToAddToNewSystem.Add(VersionedNiagaraEmitter);
 		}
-
-		if(bAllEmittersLoaded == false)
-		{
-			EAppReturnType::Type DialogResult = FMessageDialog::Open(EAppMsgType::OkCancel, EAppReturnType::Cancel,
-				LOCTEXT("FailedToLoadMessage", "A selected emitter failed to load.\nWould you like to create an empty system system?"),
-				LOCTEXT("FailedToLoadEmitterTitle", "Create Default?"));
-			if (DialogResult == EAppReturnType::Cancel)
-			{
-				return false;
-			}
-			else
-			{
-				EmittersToAddToNewSystem.Empty();
-			}
-		}
-	}
-	else
-	{
-		SystemToCopy = nullptr;
-		EmittersToAddToNewSystem.Empty();
 	}
-
+	
 	return true;
 }
 
@@ -123,10 +92,13 @@ UObject* UNiagaraSystemFactoryNew::FactoryCreateNew(UClass* Class, UObject* InPa
 		{
 			SystemToCopy->WaitForCompilationComplete();
 		}
+		
+		FNiagaraEditorModule::Get().GetRecentsManager()->SystemUsed(*SystemToCopy);
+
 		NewSystem = Cast<UNiagaraSystem>(StaticDuplicateObject(SystemToCopy, InParent, Name, Flags, Class));
-		NewSystem->TemplateSpecification = ENiagaraScriptTemplateSpecification::None;
 		NewSystem->TemplateAssetDescription = FText();
 		NewSystem->Category = FText();
+		NewSystem->AssetTags.Empty();
 
 		// if the new system doesn't have a thumbnail image, check the thumbnail map of the original asset's upackage
 		if(NewSystem->ThumbnailImage == nullptr)
@@ -168,6 +140,7 @@ UObject* UNiagaraSystemFactoryNew::FactoryCreateNew(UClass* Class, UObject* InPa
 
 	NewSystem->RequestCompile(false);
 
+	FNiagaraEditorModule::Get().GetRecentsManager()->SystemUsed(*NewSystem);
 	return NewSystem;
 }
 
@@ -182,8 +155,7 @@ void UNiagaraSystemFactoryNew::InitializeSystem(UNiagaraSystem* System, bool bCr
 	{
 		SystemScriptSource->NodeGraph = NewObject<UNiagaraGraph>(SystemScriptSource, "SystemScriptGraph", RF_Transactional);
 	}
-
-	System->TemplateSpecification = ENiagaraScriptTemplateSpecification::None;
+	
 	SystemSpawnScript->SetLatestSource(SystemScriptSource);
 	SystemUpdateScript->SetLatestSource(SystemScriptSource);
 
diff --git a/Engine/Plugins/FX/Niagara/Source/NiagaraEditor/Private/SNewEmitterDialog.cpp b/Engine/Plugins/FX/Niagara/Source/NiagaraEditor/Private/SNewEmitterDialog.cpp
deleted file mode 100644
diff --git a/Engine/Plugins/FX/Niagara/Source/NiagaraEditor/Private/SNewEmitterDialog.h b/Engine/Plugins/FX/Niagara/Source/NiagaraEditor/Private/SNewEmitterDialog.h
deleted file mode 100644
diff --git a/Engine/Plugins/FX/Niagara/Source/NiagaraEditor/Private/SNewSystemDialog.cpp b/Engine/Plugins/FX/Niagara/Source/NiagaraEditor/Private/SNewSystemDialog.cpp
deleted file mode 100644
diff --git a/Engine/Plugins/FX/Niagara/Source/NiagaraEditor/Private/SNewSystemDialog.h b/Engine/Plugins/FX/Niagara/Source/NiagaraEditor/Private/SNewSystemDialog.h
deleted file mode 100644
diff --git a/Engine/Plugins/FX/Niagara/Source/NiagaraEditor/Private/SNiagaraAssetPickerList.cpp b/Engine/Plugins/FX/Niagara/Source/NiagaraEditor/Private/SNiagaraAssetPickerList.cpp
deleted file mode 100644
diff --git a/Engine/Plugins/FX/Niagara/Source/NiagaraEditor/Private/SNiagaraNewAssetDialog.cpp b/Engine/Plugins/FX/Niagara/Source/NiagaraEditor/Private/SNiagaraNewAssetDialog.cpp
deleted file mode 100644
diff --git a/Engine/Plugins/FX/Niagara/Source/NiagaraEditor/Private/SNiagaraNewAssetDialog.h b/Engine/Plugins/FX/Niagara/Source/NiagaraEditor/Private/SNiagaraNewAssetDialog.h
deleted file mode 100644
diff --git a/Engine/Plugins/FX/Niagara/Source/NiagaraEditor/Private/Toolkits/NiagaraSystemToolkit.cpp b/Engine/Plugins/FX/Niagara/Source/NiagaraEditor/Private/Toolkits/NiagaraSystemToolkit.cpp
--- a/Engine/Plugins/FX/Niagara/Source/NiagaraEditor/Private/Toolkits/NiagaraSystemToolkit.cpp
+++ b/Engine/Plugins/FX/Niagara/Source/NiagaraEditor/Private/Toolkits/NiagaraSystemToolkit.cpp
@@ -53,7 +53,7 @@
 #include "Modules/ModuleManager.h"
 #include "ViewModels/NiagaraMessageLogViewModel.h"
 #include "NiagaraVersionMetaData.h"
-#include "SNiagaraAssetPickerList.h"
+#include "Interfaces/IMainFrameModule.h"
 #include "Toolkits/SystemToolkitModes/NiagaraSystemToolkitModeBase.h"
 #include "SystemToolkitModes/NiagaraSystemToolkitMode_Default.h"
 #include "SystemToolkitModes/NiagaraSystemToolkitMode_Scalability.h"
@@ -63,6 +63,7 @@
 #include "Widgets/Layout/SBox.h"
 #include "Widgets/SBoxPanel.h"
 #include "Widgets/SNiagaraEmitterVersionWidget.h"
+#include "Widgets/AssetBrowser/SNiagaraAssetBrowser.h"
 
 #define LOCTEXT_NAMESPACE "NiagaraSystemEditor"
 
@@ -138,6 +139,9 @@ void FNiagaraSystemToolkit::InitializeWithSystem(const EToolkitMode::Type Mode,
 	Emitter = nullptr;
 	System->EnsureFullyLoaded();
 
+	// order of registering commands matters. SetupCommands before InitAssetEditor will make the toolkit prioritize niagara commands
+	SetupCommands();
+	
 	FNiagaraSystemViewModelOptions SystemOptions;
 	SystemOptions.bCanModifyEmittersFromTimeline = true;
 	SystemOptions.EditMode = ENiagaraSystemViewModelEditMode::SystemAsset;
@@ -297,8 +301,10 @@ void FNiagaraSystemToolkit::InitializeInternal(const EToolkitMode::Type Mode, co
 	constexpr bool bCreateDefaultStandaloneMenu = true;
 	constexpr bool bCreateDefaultToolbar = true;
 	UObject* ToolkitObject = SystemToolkitMode == ESystemToolkitMode::System ? (UObject*)System : (UObject*)Emitter;
-	// order of registering commands matters. SetupCommands before InitAssetEditor will make the toolkit prioritize niagara commands
-	SetupCommands();
+	// // order of registering commands matters. SetupCommands before InitAssetEditor will make the toolkit prioritize niagara commands
+	// SetupCommands();
+	// Now that everything should be initialized, we can link the command lists of various places
+	LinkCommandLists();
 
 	const TSharedRef<FTabManager::FLayout> DummyLayout = FTabManager::NewLayout("NullLayout")->AddArea(FTabManager::NewPrimaryArea());
 	FAssetEditorToolkit::InitAssetEditor(Mode, InitToolkitHost, FNiagaraEditorModule::NiagaraEditorAppIdentifier,
@@ -539,6 +545,14 @@ void FNiagaraSystemToolkit::SetupCommands()
 		}),
 		FCanExecuteAction());
 
+	GetToolkitCommands()->MapAction(
+		FNiagaraEditorCommands::Get().OpenAddEmitterMenu,
+		FExecuteAction::CreateSP(this, &FNiagaraSystemToolkit::OpenAddEmitterMenu),
+		FCanExecuteAction::CreateSP(this, &FNiagaraSystemToolkit::CanAddEmitters),
+		FIsActionChecked(),
+		FIsActionButtonVisible::CreateSP(this, &FNiagaraSystemToolkit::CanAddEmitters)
+	);
+	
 	GetToolkitCommands()->MapAction(
 		FNiagaraEditorCommands::Get().OpenDebugHUD,
 		FExecuteAction::CreateSP(this, &FNiagaraSystemToolkit::OpenDebugHUD));
@@ -552,7 +566,10 @@ void FNiagaraSystemToolkit::SetupCommands()
 	GetToolkitCommands()->MapAction(
 		FNiagaraEditorCommands::Get().EmitterVersioning,
 		FExecuteAction::CreateSP(this, &FNiagaraSystemToolkit::ManageVersions));
-	
+}
+
+void FNiagaraSystemToolkit::LinkCommandLists()
+{
 	// appending the sequencer commands will make the toolkit also check for sequencer commands (last)
 	GetToolkitCommands()->Append(SystemViewModel->GetSequencer()->GetCommandBindings(ESequencerCommandBindings::Sequencer).ToSharedRef());
 	SystemViewModel->GetSequencer()->GetCommandBindings(ESequencerCommandBindings::Sequencer)->Append(GetToolkitCommands());
@@ -568,6 +585,29 @@ TSharedPtr<FNiagaraEmitterViewModel> FNiagaraSystemToolkit::GetEditedEmitterView
 	return HasEmitter() ? SystemViewModel->GetEmitterHandleViewModels()[0]->GetEmitterViewModel() : TSharedPtr<FNiagaraEmitterViewModel>(); 
 }
 
+void FNiagaraSystemToolkit::OpenAddEmitterMenu()
+{
+	IMainFrameModule& MainFrame = FModuleManager::LoadModuleChecked<IMainFrameModule>("MainFrame");
+	TSharedPtr<SWindow>	ParentWindow = MainFrame.GetParentWindow();
+
+	TSharedRef<SWindow> AddEmitterWindow = SNew(SNiagaraAddEmitterToSystemWindow, GetSystemViewModel().ToSharedRef());
+	FSlateApplication::Get().AddModalWindow(AddEmitterWindow, ParentWindow);
+}
+
+bool FNiagaraSystemToolkit::CanAddEmitters() const
+{
+	return GetSystemViewModel()->GetEditMode() == ENiagaraSystemViewModelEditMode::SystemAsset ? true : false;
+}
+
+void FNiagaraSystemToolkit::GetSequencerAddMenuContent(FMenuBuilder& MenuBuilder, TSharedRef<ISequencer> Sequencer)
+{
+	MenuBuilder.PushCommandList(GetToolkitCommands());
+	
+	MenuBuilder.AddMenuEntry(FNiagaraEditorCommands::Get().OpenAddEmitterMenu);
+
+	MenuBuilder.PopCommandList();
+}
+
 void FNiagaraSystemToolkit::CaptureAssetThumbnail() const
 {
 	if (Viewport.IsValid())
@@ -694,49 +734,6 @@ const FName FNiagaraSystemToolkit::GetNiagaraSystemMessageLogName(UNiagaraSystem
 	return LogListingName;
 }
 
-void FNiagaraSystemToolkit::GetSequencerAddMenuContent(FMenuBuilder& MenuBuilder, TSharedRef<ISequencer> Sequencer)
-{
-	MenuBuilder.AddSubMenu(
-		LOCTEXT("EmittersLabel", "Emitters..."),
-		LOCTEXT("EmittersToolTip", "Add an existing emitter..."),
-		FNewMenuDelegate::CreateLambda([&](FMenuBuilder& InMenuBuilder)
-		{
-			InMenuBuilder.AddWidget(CreateAddEmitterMenuContent(), FText());
-		}));
-}
-
-TSharedRef<SWidget> FNiagaraSystemToolkit::CreateAddEmitterMenuContent()
-{
-	TArray<FRefreshItemSelectorDelegate*> RefreshItemSelectorDelegates;
-	RefreshItemSelectorDelegates.Add(&RefreshItemSelector);
-	FNiagaraAssetPickerListViewOptions ViewOptions;
-	ViewOptions.SetCategorizeUserDefinedCategory(true);
-	ViewOptions.SetCategorizeLibraryAssets(true);
-	ViewOptions.SetAddLibraryOnlyCheckbox(true);
-
-	SNiagaraTemplateTabBox::FNiagaraTemplateTabOptions TabOptions;
-	TabOptions.ChangeTabState(ENiagaraScriptTemplateSpecification::Template, true);
-	TabOptions.ChangeTabState(ENiagaraScriptTemplateSpecification::None, true);
-	TabOptions.ChangeTabState(ENiagaraScriptTemplateSpecification::Behavior, true);
-
-	return SNew(SVerticalBox)
-		+SVerticalBox::Slot()
-		.FillHeight(1.0f)
-		[
-			SNew(SBox)
-			.WidthOverride(450.f)
-			.HeightOverride(500.f)
-			[
-				SNew(SNiagaraAssetPickerList, UNiagaraEmitter::StaticClass())
-				.ClickActivateMode(EItemSelectorClickActivateMode::SingleClick)
-				.ViewOptions(ViewOptions)
-				.TabOptions(TabOptions)
-				.RefreshItemSelectorDelegates(RefreshItemSelectorDelegates)
-				.OnTemplateAssetActivated(this, &FNiagaraSystemToolkit::EmitterAssetSelected)
-			]
-		];
-}
-
 FText FNiagaraSystemToolkit::GetVersionButtonLabel() const
 {
 	FText BaseLabel = LOCTEXT("NiagaraShowEmitterVersions", "Versioning");
@@ -877,7 +874,7 @@ void FNiagaraSystemToolkit::CompileSystem(bool bFullRebuild)
 	SystemViewModel->CompileSystem(bFullRebuild);
 }
 
-TSharedPtr<FNiagaraSystemViewModel> FNiagaraSystemToolkit::GetSystemViewModel()
+TSharedPtr<FNiagaraSystemViewModel> FNiagaraSystemToolkit::GetSystemViewModel() const
 {
 	return SystemViewModel;
 }
@@ -1398,12 +1395,6 @@ bool FNiagaraSystemToolkit::OnRequestClose(EAssetEditorCloseReason InCloseReason
 	return FAssetEditorToolkit::OnRequestClose(InCloseReason);
 }
 
-void FNiagaraSystemToolkit::EmitterAssetSelected(const FAssetData& AssetData)
-{
-	FSlateApplication::Get().DismissAllMenus();
-	SystemViewModel->AddEmitterFromAssetData(AssetData);
-}
-
 void FNiagaraSystemToolkit::ToggleCompileEnabled()
 {
 	UNiagaraEditorSettings* Settings = GetMutableDefault<UNiagaraEditorSettings>();
diff --git a/Engine/Plugins/FX/Niagara/Source/NiagaraEditor/Private/Toolkits/NiagaraSystemToolkit.h b/Engine/Plugins/FX/Niagara/Source/NiagaraEditor/Private/Toolkits/NiagaraSystemToolkit.h
--- a/Engine/Plugins/FX/Niagara/Source/NiagaraEditor/Private/Toolkits/NiagaraSystemToolkit.h
+++ b/Engine/Plugins/FX/Niagara/Source/NiagaraEditor/Private/Toolkits/NiagaraSystemToolkit.h
@@ -96,7 +96,7 @@ public:
 	/** Compiles the system script. */
 	void CompileSystem(bool bFullRebuild);
 
-	TSharedPtr<FNiagaraSystemViewModel> GetSystemViewModel();
+	TSharedPtr<FNiagaraSystemViewModel> GetSystemViewModel() const;
 	TSharedPtr<FNiagaraSystemGraphSelectionViewModel> GetSystemGraphSelectionViewModel();
 
 	TSharedPtr<FNiagaraSimCacheViewModel> GetSimCacheViewModel();
@@ -178,16 +178,12 @@ private:
 
 	void UpdateOriginalEmitter();
 	void UpdateExistingEmitters();
-
+	
 	void SetupCommands();
+	void LinkCommandLists();
 
 	void ResetSimulation();
 
-	void GetSequencerAddMenuContent(FMenuBuilder& MenuBuilder, TSharedRef<ISequencer> Sequencer);
-	TSharedRef<SWidget> CreateAddEmitterMenuContent();
-
-	void EmitterAssetSelected(const FAssetData& AssetData);
-
 	static void ToggleCompileEnabled();
 	static bool IsAutoCompileEnabled();
 	
@@ -211,6 +207,9 @@ private:
 	void ManageVersions();
 	TSharedPtr<FNiagaraEmitterViewModel> GetEditedEmitterViewModel() const;
 
+	void OpenAddEmitterMenu();
+	bool CanAddEmitters() const;
+	void GetSequencerAddMenuContent(FMenuBuilder& MenuBuilder, TSharedRef<ISequencer> Sequencer);
 private:
 
 	/** The System being edited in system mode, or the placeholder system being edited in emitter mode. */
diff --git a/Engine/Plugins/FX/Niagara/Source/NiagaraEditor/Private/ViewModels/NiagaraSystemViewModel.cpp b/Engine/Plugins/FX/Niagara/Source/NiagaraEditor/Private/ViewModels/NiagaraSystemViewModel.cpp
--- a/Engine/Plugins/FX/Niagara/Source/NiagaraEditor/Private/ViewModels/NiagaraSystemViewModel.cpp
+++ b/Engine/Plugins/FX/Niagara/Source/NiagaraEditor/Private/ViewModels/NiagaraSystemViewModel.cpp
@@ -40,6 +40,7 @@
 #include "NiagaraNotificationWidgetProvider.h"
 #include "NiagaraOverviewNode.h"
 #include "NiagaraParameterDefinitionsSubscriber.h"
+#include "NiagaraRecentAndFavoritesManager.h"
 #include "NiagaraScriptGraphViewModel.h"
 #include "NiagaraScriptSource.h"
 #include "NiagaraScriptVariable.h"
@@ -526,6 +527,7 @@ TSharedPtr<FNiagaraEmitterHandleViewModel> FNiagaraSystemViewModel::AddEmitterFr
 	UNiagaraEmitter* Emitter = Cast<UNiagaraEmitter>(AssetData.GetAsset());
 	if (Emitter != nullptr)
 	{
+		FNiagaraEditorModule::Get().GetRecentsManager()->EmitterUsed(*Emitter);
 		return AddEmitter(*Emitter, Emitter->GetExposedVersion().VersionGuid);
 	}
 
@@ -594,7 +596,6 @@ TSharedPtr<FNiagaraEmitterHandleViewModel> FNiagaraSystemViewModel::AddEmptyEmit
 	UNiagaraEmitter* EmptyEmitter = NewObject<UNiagaraEmitter>(GetTransientPackage());
 	bool bAddDefaultModulesAndRenderers = false;
 	UNiagaraEmitterFactoryNew::InitializeEmitter(EmptyEmitter, bAddDefaultModulesAndRenderers);
-	EmptyEmitter->TemplateSpecification = ENiagaraScriptTemplateSpecification::Template;
 	FName EmptyEmitterName = FNiagaraEditorUtilities::GetUniqueObjectName<UNiagaraEmitter>(GetTransientPackage(), TEXT("Empty"));
 	EmptyEmitter->SetUniqueEmitterName(EmptyEmitterName.ToString());
 	EmptyEmitter->SetFlags(RF_Transactional);
diff --git a/Engine/Plugins/FX/Niagara/Source/NiagaraEditor/Private/Widgets/AssetBrowser/NiagaraAssetBrowserPreview.cpp b/Engine/Plugins/FX/Niagara/Source/NiagaraEditor/Private/Widgets/AssetBrowser/NiagaraAssetBrowserPreview.cpp
new file mode 100644
--- /dev/null
+++ b/Engine/Plugins/FX/Niagara/Source/NiagaraEditor/Private/Widgets/AssetBrowser/NiagaraAssetBrowserPreview.cpp
@@ -0,0 +1,292 @@
+﻿// Copyright Epic Games, Inc. All Rights Reserved.
+
+#include "Widgets/AssetBrowser/NiagaraAssetBrowserPreview.h"
+#include "AdvancedPreviewScene.h"
+#include "CanvasItem.h"
+#include "CanvasTypes.h"
+#include "NiagaraCommon.h"
+#include "NiagaraSystemEditorData.h"
+#include "NiagaraSystemFactoryNew.h"
+#include "NiagaraSystemInstanceController.h"
+#include "SlateOptMacros.h"
+
+BEGIN_SLATE_FUNCTION_BUILD_OPTIMIZATION
+
+/** Viewport Client for the preview viewport */
+class FNiagaraAssetPreviewViewportClient : public FEditorViewportClient, public TSharedFromThis<FNiagaraAssetPreviewViewportClient>
+{
+public:
+	DECLARE_DELEGATE_OneParam(FOnScreenShotCaptured, UTexture2D*);
+
+public:
+	FNiagaraAssetPreviewViewportClient(FAdvancedPreviewScene& InPreviewScene, const TSharedRef<SNiagaraAssetBrowserPreview>& InNiagaraEditorViewport);
+	virtual ~FNiagaraAssetPreviewViewportClient() override {}
+	
+	// FEditorViewportClient interface
+	virtual FLinearColor GetBackgroundColor() const override;
+	virtual void Tick(float DeltaSeconds) override;
+	virtual void Draw(FViewport* Viewport,FCanvas* Canvas) override;
+	virtual bool ShouldOrbitCamera() const override;
+	virtual bool CanSetWidgetMode(UE::Widget::EWidgetMode NewMode) const override { return false; }
+	virtual bool CanCycleWidgetMode() const override { return false; }
+	
+	void SetUpdateViewportFocus(bool bUpdate) { bUpdateViewportFocus = bUpdate; }
+
+	TWeakPtr<SNiagaraAssetBrowserPreview> AssetBrowserPreview;
+	
+	FAdvancedPreviewScene* AdvancedPreviewScene = nullptr;
+
+private:
+	bool bUpdateViewportFocus = false;
+};
+
+FNiagaraAssetPreviewViewportClient::FNiagaraAssetPreviewViewportClient(FAdvancedPreviewScene& InPreviewScene, const TSharedRef<SNiagaraAssetBrowserPreview>& InNiagaraEditorViewport)
+	: FEditorViewportClient(nullptr, &InPreviewScene, StaticCastSharedRef<SEditorViewport>(InNiagaraEditorViewport))
+	, AdvancedPreviewScene(&InPreviewScene)
+{
+	AssetBrowserPreview = InNiagaraEditorViewport;
+
+	// Setup defaults for the common draw helper.
+	DrawHelper.bDrawPivot = false;
+	DrawHelper.bDrawWorldBox = false;
+	DrawHelper.bDrawKillZ = false;
+	DrawHelper.bDrawGrid = true;
+	DrawHelper.GridColorAxis = FColor(80,80,80);
+	DrawHelper.GridColorMajor = FColor(72,72,72);
+	DrawHelper.GridColorMinor = FColor(64,64,64);
+	DrawHelper.PerspectiveGridSize = UE_OLD_HALF_WORLD_MAX1;
+	ShowWidget(false);
+
+	FEditorViewportClient::SetViewMode(VMI_Lit);
+
+	EngineShowFlags.SetSnap(0);
+
+	OverrideNearClipPlane(1.0f);
+
+	bDrawAxesGame = true;
+	bIsSimulateInEditorViewport = true;
+}
+
+FLinearColor FNiagaraAssetPreviewViewportClient::GetBackgroundColor() const
+{
+	return FEditorViewportClient::GetBackgroundColor();
+}
+
+void FNiagaraAssetPreviewViewportClient::Tick(float DeltaSeconds)
+{
+	FEditorViewportClient::Tick(DeltaSeconds);
+
+	// Tick the preview scene world as it's not being ticked automatically as this is used in a modal window
+	PreviewScene->GetWorld()->Tick(LEVELTICK_All, DeltaSeconds);
+	Viewport->Draw();
+}
+
+void FNiagaraAssetPreviewViewportClient::Draw(FViewport* InViewport, FCanvas* Canvas)
+{
+	EngineShowFlags.SetBounds(false);
+	EngineShowFlags.Game = 0;
+	
+	TSharedPtr<SNiagaraAssetBrowserPreview> NiagaraViewport = AssetBrowserPreview.Pin();
+	UNiagaraSystem* ParticleSystem = NiagaraViewport.IsValid() ? NiagaraViewport->GetPreviewComponent()->GetAsset() : nullptr;
+	UNiagaraComponent* Component = NiagaraViewport.IsValid() ? NiagaraViewport->GetPreviewComponent() : nullptr;
+
+	FNiagaraSystemInstanceControllerConstPtr SystemInstanceController = NiagaraViewport.IsValid() ? NiagaraViewport->GetPreviewComponent()->GetSystemInstanceController() : nullptr;
+	if (Component && SystemInstanceController.IsValid() && SystemInstanceController->GetAge() > 0.0)
+	{
+		FBox Bounds = Component->GetLocalBounds().GetBox();
+		if (bUpdateViewportFocus && Bounds.IsValid)
+		{
+			bool bCachedShouldOrbitCamera = ShouldOrbitCamera();
+			FocusViewportOnBox(Bounds, true);
+			ToggleOrbitCamera(bCachedShouldOrbitCamera);
+			SetUpdateViewportFocus(false);
+			
+		}
+
+		// UWorld* World = Component->GetWorld();
+		// if (World && !ParticleSystem->bFixedBounds)
+		// {
+		// 	FNiagaraSystemInstance* SystemInstance = SystemInstanceController->GetSystemInstance_Unsafe();
+		// 	const TArray<FNiagaraEmitterHandle>& EmitterHandles = ParticleSystem->GetEmitterHandles();
+		// 	for (int32 i=0; i < EmitterHandles.Num(); ++i)
+		// 	{
+		// 		if (!EmitterHandles[i].GetIsEnabled() || !EmitterHandles[i].GetDebugShowBounds())
+		// 		{
+		// 			continue;
+		// 		}
+		// 		const FNiagaraEmitterInstance& EmitterInstance = SystemInstance->GetEmitters()[i].Get();
+		// 		const FTransform& Transform = Component->GetComponentTransform();
+		// 		const FBox EmitterBounds = EmitterInstance.GetBounds();
+		// 		DrawDebugBox(World, EmitterBounds.GetCenter() + Transform.GetLocation(), EmitterBounds.GetExtent(), Transform.GetRotation(), FColor::Green);
+		// 	}
+		// }
+	}
+
+	FEditorViewportClient::Draw(InViewport, Canvas);
+
+	// Debug Text
+	FNiagaraSystemInstance* SystemInstance = SystemInstanceController.IsValid() ? SystemInstanceController->GetSoloSystemInstance() : nullptr;
+	float CurrentX = 10.0f;
+	float CurrentY = 50.0f;
+	UFont* Font = GEngine->GetSmallFont();
+	const float FontHeight = Font->GetMaxCharHeight() * 1.1f;
+
+	// Particle Count
+	// if (SystemInstance)
+	// {
+	// 	FCanvasTextItem TextItem(FVector2D(CurrentX, CurrentY), FText::FromString(TEXT("Particle Counts")), Font, FLinearColor::White);
+	// 	TextItem.EnableShadow(FLinearColor::Black);
+	// 	TextItem.Draw(Canvas);
+	// 	CurrentY += FontHeight;
+	// 	
+	// 	for (const TSharedRef<FNiagaraEmitterInstance, ESPMode::ThreadSafe>& EmitterInstance : SystemInstance->GetEmitters())
+	// 	{
+	// 		const FName EmitterName = EmitterInstance->GetEmitterHandle().GetName();
+	// 		const int32 CurrentCount = EmitterInstance->GetNumParticles();
+	// 		const int32 MaxCount = EmitterInstance->GetEmitterHandle().GetEmitterData()->GetMaxParticleCountEstimate();
+	// 		const bool IsIsolated = EmitterInstance->GetEmitterHandle().IsIsolated();
+	// 		const bool IsEnabled = EmitterInstance->GetEmitterHandle().GetIsEnabled();
+	// 		const ENiagaraExecutionState ExecutionState = EmitterInstance->GetExecutionState();
+	// 		const FString EmitterExecutionString = UEnum::GetValueAsString(ExecutionState);
+	// 		const int32 EmitterExecutionStringValueIndex = EmitterExecutionString.Find(TEXT("::"));
+	// 		const TCHAR* EmitterExecutionText = EmitterExecutionStringValueIndex == INDEX_NONE ? *EmitterExecutionString : *EmitterExecutionString + EmitterExecutionStringValueIndex + 2;
+	// 	
+	// 		TextItem.Text = FText::FromString(FString::Printf(TEXT("%i Current, %i Max (est.) - [%s] [%s]"), CurrentCount, MaxCount, *EmitterName.ToString(), EmitterExecutionText));
+	// 		TextItem.Position = FVector2D(CurrentX, CurrentY);
+	// 		TextItem.bOutlined = IsIsolated;
+	// 		TextItem.OutlineColor = FLinearColor(0.7f, 0.0f, 0.0f);
+	// 		TextItem.SetColor(IsEnabled ? FLinearColor::White : FLinearColor::Gray);
+	// 		TextItem.Draw(Canvas);
+	// 		CurrentY += FontHeight;
+	// 	}
+	// }
+
+	// Compiling Emitters
+	if(Component)
+	{
+		if(UNiagaraSystem* NiagaraSystem = Component->GetAsset())
+		{
+			const bool bSystemCompiling = NiagaraSystem->HasOutstandingCompilationRequests();
+		
+			for (const FNiagaraEmitterHandle& EmitterHandle : NiagaraSystem->GetEmitterHandles())
+			{
+				FVersionedNiagaraEmitterData* EmitterData = EmitterHandle.GetEmitterData();
+				if (EmitterData == nullptr)
+				{
+					continue;
+				}
+		
+				const bool bEmitterCompiling = bSystemCompiling || !EmitterData->IsReadyToRun();
+				if (!bEmitterCompiling)
+				{
+					continue;
+				}
+		
+				FString CompilingText = FString::Printf(TEXT("%s - %s"), *FText::FromString("Compiling Emitter").ToString(), *EmitterHandle.GetUniqueInstanceName());
+				Canvas->DrawShadowedString(CurrentX, CurrentY, *CompilingText, Font, FLinearColor::Yellow);
+				CurrentY += FontHeight;
+			}
+		}
+	}
+}
+
+bool FNiagaraAssetPreviewViewportClient::ShouldOrbitCamera() const
+{
+	return true;
+}
+
+void SNiagaraAssetBrowserPreview::Construct(const FArguments& InArgs)
+{
+	PreviewComponent = NewObject<UNiagaraComponent>();
+	PreviewComponent->CastShadow = 1;
+	PreviewComponent->bCastDynamicShadow = 1;
+	PreviewComponent->SetAllowScalability(false);
+	//PreviewComponent->SetAsset(System);
+	PreviewComponent->SetForceSolo(true);
+	PreviewComponent->SetAgeUpdateMode(ENiagaraAgeUpdateMode::TickDeltaTime);
+	PreviewComponent->SetCanRenderWhileSeeking(false);
+	PreviewComponent->Activate(true);
+
+	SystemForEmitterDisplay = NewObject<UNiagaraSystem>();
+	UNiagaraSystemFactoryNew::InitializeSystem(SystemForEmitterDisplay, true);
+	SystemForEmitterDisplay->EnsureFullyLoaded();
+	
+	AdvancedPreviewScene = MakeShareable(new FAdvancedPreviewScene(FPreviewScene::ConstructionValues()));
+	AdvancedPreviewScene->SetFloorVisibility(false);
+
+	AdvancedPreviewScene->AddComponent(PreviewComponent, PreviewComponent->GetRelativeTransform());
+
+	float Pitch = -40.0;
+	float Yaw = 128.0;
+	float Roll = 0.0;
+	AdvancedPreviewScene->SetLightDirection(FRotator(Pitch, Yaw, Roll));
+	
+	SEditorViewport::Construct(SEditorViewport::FArguments());
+}
+
+SNiagaraAssetBrowserPreview::~SNiagaraAssetBrowserPreview()
+{
+	PreviewComponent = nullptr;
+
+	if (AssetPreviewViewportClient.IsValid())
+	{
+		AssetPreviewViewportClient->Viewport = nullptr;
+	}
+}
+
+void SNiagaraAssetBrowserPreview::SetEmitter(UNiagaraEmitter& Emitter)
+{
+	FNiagaraEditorUtilities::KillSystemInstances(*SystemForEmitterDisplay);
+	
+	TSet<FGuid> EmitterHandleGuids;
+	for(const FNiagaraEmitterHandle& EmitterHandle : SystemForEmitterDisplay->GetEmitterHandles())
+	{
+		EmitterHandleGuids.Add(EmitterHandle.GetId());
+	}
+	SystemForEmitterDisplay->RemoveEmitterHandlesById(EmitterHandleGuids);
+	
+	FNiagaraEmitterHandle TemporaryEmitterHandle(Emitter, Emitter.GetExposedVersion().VersionGuid);
+	SystemForEmitterDisplay->DuplicateEmitterHandle(TemporaryEmitterHandle, *Emitter.GetUniqueEmitterName());
+	SystemForEmitterDisplay->RequestCompile(false);
+	PreviewComponent->SetAsset(SystemForEmitterDisplay);
+	PreviewComponent->ResetSystem();
+}
+
+void SNiagaraAssetBrowserPreview::SetSystem(UNiagaraSystem& System)
+{	
+	PreviewComponent->SetAsset(&System);
+}
+
+void SNiagaraAssetBrowserPreview::ResetAsset()
+{
+	PreviewComponent->SetAsset(nullptr);
+}
+
+void SNiagaraAssetBrowserPreview::Tick(const FGeometry& AllottedGeometry, const double InCurrentTime, const float InDeltaTime)
+{
+	SEditorViewport::Tick(AllottedGeometry, InCurrentTime, InDeltaTime);
+
+	PreviewComponent->TickComponent(InDeltaTime, LEVELTICK_All, nullptr);
+	AssetPreviewViewportClient->Tick(InDeltaTime);
+}
+
+void SNiagaraAssetBrowserPreview::AddReferencedObjects(FReferenceCollector& Collector)
+{
+	Collector.AddReferencedObject(PreviewComponent);
+	Collector.AddReferencedObject(SystemForEmitterDisplay);
+}
+
+TSharedRef<FEditorViewportClient> SNiagaraAssetBrowserPreview::MakeEditorViewportClient()
+{
+	AssetPreviewViewportClient = MakeShareable(new FNiagaraAssetPreviewViewportClient(*AdvancedPreviewScene.Get(), SharedThis(this)));
+	
+	AssetPreviewViewportClient->SetViewLocation( FVector::ZeroVector );
+	AssetPreviewViewportClient->SetViewRotation( FRotator::ZeroRotator );
+	AssetPreviewViewportClient->SetViewLocationForOrbiting( FVector::ZeroVector );
+	AssetPreviewViewportClient->VisibilityDelegate.BindSP( this, &SNiagaraAssetBrowserPreview::IsVisible );
+	AssetPreviewViewportClient->bSetListenerPosition = false;
+	
+	return AssetPreviewViewportClient.ToSharedRef();
+}
+
+END_SLATE_FUNCTION_BUILD_OPTIMIZATION
diff --git a/Engine/Plugins/FX/Niagara/Source/NiagaraEditor/Private/Widgets/AssetBrowser/NiagaraMenuFilters.cpp b/Engine/Plugins/FX/Niagara/Source/NiagaraEditor/Private/Widgets/AssetBrowser/NiagaraMenuFilters.cpp
new file mode 100644
--- /dev/null
+++ b/Engine/Plugins/FX/Niagara/Source/NiagaraEditor/Private/Widgets/AssetBrowser/NiagaraMenuFilters.cpp
@@ -0,0 +1,66 @@
+﻿// Copyright Epic Games, Inc. All Rights Reserved.
+
+#include "Widgets/AssetBrowser/NiagaraMenuFilters.h"
+
+#define LOCTEXT_NAMESPACE "NiagaraMenuFilters"
+
+void UNiagaraTagsContentBrowserFilterData::AddTagGuid(const FGuid& InGuid)
+{
+	ActiveTagGuids.AddUnique(InGuid);
+	OnActiveTagGuidsChangedDelegate.Broadcast();
+}
+
+void UNiagaraTagsContentBrowserFilterData::RemoveTagGuid(const FGuid& InGuid)
+{
+	ActiveTagGuids.Remove(InGuid);
+	OnActiveTagGuidsChangedDelegate.Broadcast();
+}
+
+bool UNiagaraTagsContentBrowserFilterData::ContainsActiveTagGuid(const FGuid& InGuid)
+{
+	return ActiveTagGuids.Contains(InGuid);
+}
+
+FName FNiagaraAssetBrowserMainFilter::GetIdentifier() const
+{
+	return FName(FText::AsCultureInvariant(GetDisplayName()).ToString());
+}
+
+FText FNiagaraAssetBrowserMainFilter::GetDisplayName() const
+{
+	if(FilterMode == EFilterMode::All)
+	{
+		return LOCTEXT("NiagaraMainFilterLabel_All", "All");
+	}
+	else if(FilterMode == EFilterMode::NiagaraAssetTag)
+	{
+		return FText::FromString(AssetTagDefinition.AssetTag.ToString());
+	}
+	else if(FilterMode == EFilterMode::Recent)
+	{
+		return LOCTEXT("NiagaraMainFilterLabel_Recent", "Recent");
+	}
+	else if(FilterMode == EFilterMode::Custom)
+	{
+		return CustomDisplayName;
+	}
+
+	return FText::FromString("Unknown Name");
+}
+
+bool FNiagaraAssetBrowserMainFilter::IsAssetRecent(const FAssetData& AssetCandidate) const
+{
+	return IsAssetRecentDelegate.Execute(AssetCandidate);
+}
+
+bool FNiagaraAssetBrowserMainFilter::ShouldCustomFilterAsset(const FAssetData& AssetCandidate) const
+{
+	return CustomShouldFilterAsset.Execute(AssetCandidate);
+}
+
+bool FNiagaraAssetBrowserMainFilter::DoesAssetHaveTag(const FAssetData& AssetCandidate) const
+{
+	return AssetTagDefinition.DoesAssetDataContainTag(AssetCandidate);
+}
+
+#undef LOCTEXT_NAMESPACE
diff --git a/Engine/Plugins/FX/Niagara/Source/NiagaraEditor/Private/Widgets/AssetBrowser/NiagaraSemanticTagsFrontEndFilterExtension.cpp b/Engine/Plugins/FX/Niagara/Source/NiagaraEditor/Private/Widgets/AssetBrowser/NiagaraSemanticTagsFrontEndFilterExtension.cpp
new file mode 100644
--- /dev/null
+++ b/Engine/Plugins/FX/Niagara/Source/NiagaraEditor/Private/Widgets/AssetBrowser/NiagaraSemanticTagsFrontEndFilterExtension.cpp
@@ -0,0 +1,182 @@
+﻿// Copyright Epic Games, Inc. All Rights Reserved.
+
+#include "Widgets/AssetBrowser/NiagaraSemanticTagsFrontEndFilterExtension.h"
+
+#include "Misc/ConfigCacheIni.h"
+#include "NiagaraEditorUtilities.h"
+
+#define LOCTEXT_NAMESPACE "NiagaraSemanticUNiagaraSemanticTagsFrontEndFilterExtension"
+
+bool FFrontendFilter_NiagaraTag::PassesFilter(const FContentBrowserItem& InItem) const
+{
+	return InItem.GetItemAttributes().Contains(FName(AssetTagDefinition.GetGuidAsString()));
+}
+
+FText FFrontendFilter_NiagaraTag_ContentBrowser::GetDisplayName() const
+{
+	FText BaseText = LOCTEXT("NiagaraAssetTagFilterDisplayName", "Niagara Asset Tags: {0}");
+	FAssetRegistryModule& AssetRegistryModule = FModuleManager::LoadModuleChecked<FAssetRegistryModule>("AssetRegistry");
+	if(AssetRegistryModule.Get().IsLoadingAssets())
+	{
+		return FText::FormatOrdered(BaseText, LOCTEXT("WaitingForAssetRegistryLoad", "Waiting for Asset Registry"));
+	}
+
+	if(FilterData->GetNumActiveTagGuids() == 0)
+	{
+		return FText::FormatOrdered(BaseText, LOCTEXT("Right-click", "Right-click"));
+	}
+	else if(FilterData->GetNumActiveTagGuids() == 1)
+	{
+		FNiagaraAssetTagDefinitionReference Reference;
+		Reference.SetTagDefinitionReferenceGuid(FilterData->GetActiveTagGuids()[0]);
+		const FNiagaraAssetTagDefinition& AssetTagDefinition = FNiagaraEditorUtilities::AssetBrowser::FindTagDefinitionForReference(Reference);
+
+		if(AssetTagDefinition.IsValid())
+		{
+			return FText::FormatOrdered(BaseText, AssetTagDefinition.AssetTag);
+		}
+		// In case an active tag was loaded from config that doesn't exist, we display 'Invalid'
+		else
+		{
+			return FText::FormatOrdered(BaseText, LOCTEXT("Invalid", "Invalid"));
+		}
+	}
+	else if(FilterData->GetNumActiveTagGuids() > 1)
+	{
+		return FText::FormatOrdered(BaseText, LOCTEXT("Complex", "Complex"));
+
+	}
+	
+	return FText::GetEmpty();
+}
+
+void FFrontendFilter_NiagaraTag_ContentBrowser::ModifyContextMenu(FMenuBuilder& MenuBuilder)
+{
+	UNiagaraTagsContentBrowserFilterContext* ContextObject = NewObject<UNiagaraTagsContentBrowserFilterContext>();
+	ContextObject->FilterData = FilterData;
+	FToolMenuContext Context(ContextObject);
+
+	MenuBuilder.AddWidget(UToolMenus::Get()->GenerateWidget("NiagaraEditorModule.ContentBrowserNiagaraTags", Context), FText::GetEmpty(), false);
+}
+
+bool FFrontendFilter_NiagaraTag_ContentBrowser::PassesFilter(const FContentBrowserItem& InItem) const
+{
+	for(const FGuid& Guid : FilterData->GetActiveTagGuids())
+	{
+		if(InItem.GetItemAttributes().Contains(FName(Guid.ToString(EGuidFormats::DigitsWithHyphens))) == false)
+		{
+			return false;
+		}
+	}
+
+	return true;
+}
+
+void FFrontendFilter_NiagaraTag_ContentBrowser::SaveSettings(const FString& IniFilename, const FString& IniSection, const FString& SettingsString) const
+{
+	FFrontendFilter::SaveSettings(IniFilename, IniSection, SettingsString);
+
+	TArray<FString> TagStrings;
+	TagStrings.Reserve(FilterData->GetNumActiveTagGuids());
+
+	for(const FGuid& ActiveTagGuid : FilterData->GetActiveTagGuids())
+	{
+		TagStrings.Add(ActiveTagGuid.ToString(EGuidFormats::DigitsWithHyphens));
+	}
+
+	GConfig->SetArray(*IniSection, *(SettingsString + TEXT(".NiagaraTags")), TagStrings, IniFilename);
+}
+
+void FFrontendFilter_NiagaraTag_ContentBrowser::LoadSettings(const FString& IniFilename, const FString& IniSection, const FString& SettingsString)
+{
+	FFrontendFilter::LoadSettings(IniFilename, IniSection, SettingsString);
+
+	TArray<FString> TagStrings;
+	GConfig->GetArray(*IniSection, *(SettingsString + TEXT(".NiagaraTags")), TagStrings, IniFilename);
+
+	FilterData->ResetActiveTagGuids();
+
+	for (const FString& TagString : TagStrings)
+	{
+		FGuid TagGuid = FGuid(TagString); 
+		if(TagGuid.IsValid())
+		{
+			FilterData->AddTagGuid(TagGuid);
+		}
+	}
+}
+
+FString FFrontendFilter_NiagaraTag_ContentBrowser::GetReferencerName() const
+{
+	return "FrontendFilter_NiagaraTag_ContentBrowser";
+}
+
+void FFrontendFilter_NiagaraTag_ContentBrowser::AddReferencedObjects(FReferenceCollector& Collector)
+{
+	Collector.AddReferencedObject(FilterData);
+}
+
+bool FFrontendFilter_NiagaraEmitterInheritance::PassesFilter(const FContentBrowserItem& InItem) const
+{
+	FAssetData AssetData;
+	InItem.Legacy_TryGetAssetData(AssetData);
+
+	if(AssetData.GetClass() != UNiagaraEmitter::StaticClass())
+	{
+		return false;
+	}
+	
+	bool bUseInheritance = false;
+	if(FNiagaraEditorUtilities::GetIsInheritableFromAssetRegistryTags(AssetData, bUseInheritance))
+	{
+		return bRequiredInheritanceState == bUseInheritance;
+	}
+
+	return false;
+}
+
+bool FFrontendFilter_NiagaraSystemEffectType::PassesFilter(const FContentBrowserItem& InItem) const
+{
+	if(InItem.GetItemAttribute("EffectType").IsValid())
+	{
+		FName EffectType = InItem.GetItemAttribute("EffectType").GetValue<FName>();
+		if(EffectType == NAME_None && bRequiresEffectType == false)
+		{
+			return true;
+		}
+
+		if(EffectType != NAME_None && bRequiresEffectType == true)
+		{
+			return true;
+		}
+	}
+	else
+	{
+		if(bRequiresEffectType == false)
+		{
+			return true;
+		}
+	}
+
+	return false;
+}
+
+void UNiagaraSemanticTagsFrontEndFilterExtension::AddFrontEndFilterExtensions(TSharedPtr<FFrontendFilterCategory> DefaultCategory, TArray<TSharedRef<FFrontendFilter>>& InOutFilterList) const
+{
+	using namespace FNiagaraEditorUtilities::AssetBrowser;
+	
+	TSharedRef<FFrontendFilterCategory> Category = MakeShared<FFrontendFilterCategory>(LOCTEXT("NiagaraSemanticFilterCategoryLabel", "Niagara Tags"),
+		LOCTEXT("NiagaraSemanticFilterCategoryTooltip", "Filter Niagara Assets by Tag"));
+
+	// by not sorting by name we ensure we keep the source & asset order
+	// for(const FNiagaraAssetTagDefinition& AssetTagDefinition : GetFlatSortedAssetTagDefinitions(false))
+	// {
+	// 	TSharedRef<FFrontendFilter_NiagaraTag> NiagaraTagFilter = MakeShared<FFrontendFilter_NiagaraTag>(AssetTagDefinition, Category);
+	// 	InOutFilterList.Add(NiagaraTagFilter);
+	// }
+
+	TSharedRef<FFrontendFilter_NiagaraTag_ContentBrowser> Filter = MakeShared<FFrontendFilter_NiagaraTag_ContentBrowser>(DefaultCategory);
+	InOutFilterList.Add(Filter);
+}
+
+#undef LOCTEXT_NAMESPACE
diff --git a/Engine/Plugins/FX/Niagara/Source/NiagaraEditor/Private/Widgets/AssetBrowser/SNiagaraAssetBrowser.cpp b/Engine/Plugins/FX/Niagara/Source/NiagaraEditor/Private/Widgets/AssetBrowser/SNiagaraAssetBrowser.cpp
new file mode 100644
--- /dev/null
+++ b/Engine/Plugins/FX/Niagara/Source/NiagaraEditor/Private/Widgets/AssetBrowser/SNiagaraAssetBrowser.cpp
@@ -0,0 +1,1046 @@
+﻿// Copyright Epic Games, Inc. All Rights Reserved.
+
+#include "Widgets\AssetBrowser/SNiagaraAssetBrowser.h"
+
+#include "ContentBrowserMenuContexts.h"
+#include "ContentBrowserModule.h"
+#include "NiagaraEditorSettings.h"
+#include "NiagaraEditorUtilities.h"
+#include "NiagaraEmitter.h"
+#include "NiagaraRecentAndFavoritesManager.h"
+#include "NiagaraSystem.h"
+#include "SlateOptMacros.h"
+#include "AssetRegistry/AssetRegistryModule.h"
+#include "Toolkits/GlobalEditorCommonCommands.h"
+#include "Toolkits/NiagaraSystemToolkit.h"
+#include "Widgets/SItemSelector.h"
+#include "Widgets/SNiagaraSystemViewport.h"
+#include "Widgets/Layout/SSpacer.h"
+#include "Widgets/Layout/SSeparator.h"
+#include "Widgets/AssetBrowser/NiagaraAssetBrowserConfig.h"
+#include "Widgets/AssetBrowser/NiagaraSemanticTagsFrontEndFilterExtension.h"
+#include "Widgets/AssetBrowser/SNiagaraSelectedAssetDetails.h"
+
+BEGIN_SLATE_FUNCTION_BUILD_OPTIMIZATION
+
+#define LOCTEXT_NAMESPACE "NiagaraAssetBrowser"
+
+struct FNiagaraAssetBrowserTabs
+{
+	// Tab identifiers
+	static const FName MainSelectionTabID;
+	static const FName ContentBrowserTabID;
+	static const FName DetailsTabID;
+	static const FName UserProvidedControlsTabID;
+};
+
+const FName FNiagaraAssetBrowserTabs::MainSelectionTabID(TEXT("MainSelection"));
+const FName FNiagaraAssetBrowserTabs::ContentBrowserTabID(TEXT("ContentBrowser"));
+const FName FNiagaraAssetBrowserTabs::DetailsTabID(TEXT("Details"));
+const FName FNiagaraAssetBrowserTabs::UserProvidedControlsTabID(TEXT("UserProvidedControls"));
+
+void SNiagaraAssetBrowser::Construct(const FArguments& InArgs)
+{
+	AvailableClasses = InArgs._AvailableClasses;
+	RecentAndFavoritesList = InArgs._RecentAndFavoritesList;
+	SaveSettingsName = InArgs._SaveSettingsName;
+	EmptySelectionMessage = InArgs._EmptySelectionMessage;
+	AssetSelectionMode = InArgs._AssetSelectionMode;
+	OnAssetSelectedDelegate = InArgs._OnAssetSelected;
+	OnAssetsActivatedDelegate = InArgs._OnAssetsActivated;
+
+	PreviewViewport = SNew(SNiagaraAssetBrowserPreview)
+	.Visibility(this, &SNiagaraAssetBrowser::OnGetViewportVisibility);
+	
+	ensureMsgf(AvailableClasses.Num() >= 1, TEXT("The Niagara asset browser has to be supplied at least one available class."));
+	
+	ChildSlot
+	[
+		SNew(SBorder)
+		.Padding(0.f)
+		.BorderImage(FAppStyle::Get().GetBrush("Brushes.Panel"))
+		[
+			SNew(SVerticalBox)
+			+ SVerticalBox::Slot()
+			.Padding(0.f)
+			[
+				SNew(SSplitter)
+				.Orientation(Orient_Horizontal)
+				.PhysicalSplitterHandleSize(2.f)
+				+ SSplitter::Slot().Expose(FiltersSlot)
+					.Value(0.15f)
+					.MinSize(50.f)
+				+ SSplitter::Slot().Expose(AssetBrowserContentSlot)
+					.Value(0.6f)
+				+ SSplitter::Slot().Expose(AssetBrowserDetailsAreaSlot)
+					.Value(0.25f)
+					.MinSize(100.f)
+			]
+			+ SVerticalBox::Slot()
+			.AutoHeight()
+			[
+				SNew(SSeparator)
+				.Orientation(Orient_Horizontal)
+				.Thickness(2.f)
+			]
+			+ SVerticalBox::Slot()
+			.AutoHeight()
+			.Expose(AdditionalWidgetSlot)
+		]
+	];
+
+	bSuppressSaveAndLoad = true;
+	PopulateAssetBrowserContentSlot();
+	PopulateFiltersSlot();
+	PopulateAssetBrowserDetailsSlot();
+
+	if(InArgs._AdditionalWidget.IsValid())
+	{
+		AdditionalWidgetSlot->AttachWidget(InArgs._AdditionalWidget.ToSharedRef());
+	}
+
+	bSuppressSaveAndLoad = false;
+	LoadSettings();
+
+	InitContextMenu();
+}
+
+SNiagaraAssetBrowser::~SNiagaraAssetBrowser()
+{
+	PreviewViewport.Reset();
+	SaveSettings();
+}
+
+TArray<UClass*> SNiagaraAssetBrowser::GetDisplayedAssetTypes() const
+{
+	return AvailableClasses;
+}
+
+void SNiagaraAssetBrowser::InitContextMenu()
+{
+	static FName MenuName("NiagaraAssetBrowser.ContextMenu");
+	if(UToolMenus::Get()->IsMenuRegistered(MenuName) == false)
+	{
+		UToolMenu* Menu = UToolMenus::Get()->RegisterMenu(MenuName);
+		FToolMenuSection& DefaultSection = Menu->AddSection("Default");
+
+		FToolMenuExecuteAction ToolMenuAction = FToolMenuExecuteAction::CreateLambda([](const FToolMenuContext& Context)
+		{
+			if(UContentBrowserAssetContextMenuContext* ContentBrowserContext = Context.FindContext<UContentBrowserAssetContextMenuContext>())
+			{
+				if(GEditor)
+				{
+					GEditor->SyncBrowserToObjects(ContentBrowserContext->SelectedAssets);
+				}
+			}
+		});
+		FToolUIActionChoice Action(ToolMenuAction);		
+		DefaultSection.AddEntry(FToolMenuEntry::InitMenuEntry("NavigateTo", LOCTEXT("NavigateTo", "Find in Content Browser"), FText::GetEmpty(), FSlateIcon(), Action));
+	}
+}
+
+FARFilter SNiagaraAssetBrowser::GetCurrentBackendFilter() const
+{
+	FARFilter Filter;
+
+	if(MainFilterSelector.IsValid())
+	{
+		FARFilter AssetTagSelectorFilter;
+
+		for(const TSharedRef<FNiagaraAssetBrowserMainFilter>& MainFilter : MainFilterSelector->GetSelectedItems())
+		{
+			if(MainFilter->FilterMode == FNiagaraAssetBrowserMainFilter::EFilterMode::NiagaraAssetTag)
+			{
+				AssetTagSelectorFilter.TagsAndValues.Add(FName(MainFilter->AssetTagDefinition.GetGuidAsString()), TOptional<FString>());
+			}
+		}
+
+		Filter.Append(AssetTagSelectorFilter);
+	}
+	
+	// if the filter list hasn't specified any classes of the available ones, we explicitly set the filter to the available classes.
+	// this is because if no classes are specified, we could potentially receive assets outside the available classes
+	if(Filter.ClassPaths.IsEmpty())
+	{
+		for(UClass* AvailableClass : AvailableClasses)
+		{
+			Filter.ClassPaths.Add(AvailableClass->GetClassPathName());
+		}
+	}
+
+	return Filter;
+}
+
+bool SNiagaraAssetBrowser::ShouldFilterAsset(const FAssetData& AssetData) const
+{
+	// TODO (ME) This currently implies only one main filter/folder can be active at a time. is this wanted?
+	for(const TSharedRef<FNiagaraAssetBrowserMainFilter>& MainFilter : MainFilterSelector->GetSelectedItems())
+	{
+		if(MainFilter->FilterMode == FNiagaraAssetBrowserMainFilter::EFilterMode::All)
+		{
+			return false;
+		}
+		else if(MainFilter->FilterMode == FNiagaraAssetBrowserMainFilter::EFilterMode::Recent)
+		{
+			return MainFilter->IsAssetRecent(AssetData) == false;
+		}
+		else if(MainFilter->FilterMode == FNiagaraAssetBrowserMainFilter::EFilterMode::NiagaraAssetTag)
+		{
+			return MainFilter->DoesAssetHaveTag(AssetData) == false;
+		}
+		else if(MainFilter->FilterMode == FNiagaraAssetBrowserMainFilter::EFilterMode::Custom)
+		{
+			return MainFilter->CustomShouldFilterAsset.Execute(AssetData);
+		}
+	}
+
+	return false;
+}
+
+void SNiagaraAssetBrowser::RefreshBackendFilter() const
+{
+	AssetBrowserContent->SetARFilter(GetCurrentBackendFilter());
+}
+
+void SNiagaraAssetBrowser::PopulateFiltersSlot()
+{
+	AssetBrowserMainFilters = GetMainFilters();
+
+	SAssignNew(MainFilterSelector, STreeView<TSharedRef<FNiagaraAssetBrowserMainFilter>>)
+		.TreeItemsSource(&AssetBrowserMainFilters)
+		.OnSelectionChanged(this, &SNiagaraAssetBrowser::OnMainFilterSelected)
+		.OnGenerateRow(this, &SNiagaraAssetBrowser::GenerateWidgetRowForMainFilter)
+		.OnGetChildren(this, &SNiagaraAssetBrowser::OnGetChildFiltersForFilter)
+		.SelectionMode(ESelectionMode::Single);
+	
+	TSharedRef<FNiagaraAssetBrowserMainFilter>* AllFilter = AssetBrowserMainFilters.FindByPredicate([](const TSharedRef<FNiagaraAssetBrowserMainFilter>& MainFilterCandidate)
+	{
+		return MainFilterCandidate->FilterMode == FNiagaraAssetBrowserMainFilter::EFilterMode::All;
+	});
+	
+	if(AllFilter)
+	{
+		MainFilterSelector->SetSelection(*AllFilter);
+		MainFilterSelector->SetItemExpansion(*AllFilter, true);
+	}
+	
+	FiltersSlot->AttachWidget(MainFilterSelector.ToSharedRef());
+}
+
+void SNiagaraAssetBrowser::PopulateAssetBrowserContentSlot()
+{
+	FAssetPickerConfig Config;
+	Config.Filter = GetCurrentBackendFilter();
+	Config.bCanShowClasses = false;
+	Config.bAddFilterUI = true;
+	Config.DefaultFilterMenuExpansion = EAssetTypeCategories::FX;
+	Config.ExtraFrontendFilters = OnGetExtraFrontendFilters();
+	Config.OnExtendAddFilterMenu = FOnExtendAddFilterMenu::CreateSP(this, &SNiagaraAssetBrowser::OnExtendAddFilterMenu);
+	Config.bUseSectionsForCustomFilterCategories = true;
+	Config.OnShouldFilterAsset = FOnShouldFilterAsset::CreateSP(this, &SNiagaraAssetBrowser::ShouldFilterAsset);
+	Config.OnAssetSelected = FOnAssetSelected::CreateSP(this, &SNiagaraAssetBrowser::OnAssetSelected);
+	Config.OnAssetsActivated = FOnAssetsActivated::CreateSP(this, &SNiagaraAssetBrowser::OnAssetsActivated);
+	Config.OnGetCustomAssetToolTip = FOnGetCustomAssetToolTip::CreateSP(this, &SNiagaraAssetBrowser::OnGetCustomAssetTooltip);
+	Config.OnIsAssetValidForCustomToolTip = FOnIsAssetValidForCustomToolTip::CreateLambda([](FAssetData& AssetData)
+	{
+		return true;
+	});
+	Config.bForceShowEngineContent = true;
+	Config.bForceShowPluginContent = true;
+	Config.SelectionMode = AssetSelectionMode;
+	Config.bAllowDragging = false;
+	Config.OnGetAssetContextMenu = FOnGetAssetContextMenu::CreateSP(this, &SNiagaraAssetBrowser::OnGetAssetContextMenu);
+
+	if(SaveSettingsName.IsSet())
+	{
+		Config.SaveSettingsName = SaveSettingsName.GetValue().ToString();
+	}
+	
+	SAssignNew(AssetBrowserContent, SNiagaraAssetBrowserContent)
+	.InitialConfig(Config);
+	
+	AssetBrowserContentSlot->AttachWidget(AssetBrowserContent.ToSharedRef());
+}
+
+void SNiagaraAssetBrowser::PopulateAssetBrowserDetailsSlot()
+{
+	DetailsContainer = SNew(SBox)
+		.Padding(5.f)
+		[
+			SAssignNew(DetailsSwitcher, SWidgetSwitcher)
+			+ SWidgetSwitcher::Slot()
+			.HAlign(HAlign_Fill)
+			.VAlign(VAlign_Top)
+			.Padding(30.f, 60.f)
+			[
+				SNew(STextBlock)
+				.Text(EmptySelectionMessage.Get(FText::GetEmpty()))
+				.AutoWrapText(true)
+				.TextStyle(&FAppStyle::GetWidgetStyle<FTextBlockStyle>("NormalText.Subdued"))
+			]
+			+ SWidgetSwitcher::Slot()
+			.Expose(AssetBrowserDetailsSlot)
+		];
+	
+	AssetBrowserDetailsAreaSlot->AttachWidget(DetailsContainer.ToSharedRef());
+}
+
+TArray<TSharedRef<FNiagaraAssetBrowserMainFilter>> SNiagaraAssetBrowser::GetMainFilters() const
+{
+	using namespace FNiagaraEditorUtilities::AssetBrowser;
+
+	TArray<TSharedRef<FNiagaraAssetBrowserMainFilter>> MainFilters;
+
+	// Recent
+	{
+		TSharedRef<FNiagaraAssetBrowserMainFilter> RecentFilter = MakeShared<FNiagaraAssetBrowserMainFilter>(FNiagaraAssetBrowserMainFilter::EFilterMode::Recent);
+		RecentFilter->IsAssetRecentDelegate = FNiagaraAssetBrowserMainFilter::FIsAssetRecent::CreateSP(this, &SNiagaraAssetBrowser::OnIsAssetRecent);
+		MainFilters.Add(RecentFilter);
+	}
+
+	// Tag Filters
+	TArray<TSharedRef<FNiagaraAssetBrowserMainFilter>> TagFilters;
+	{
+		TArray<UClass*> DisplayedAssetTypes = GetDisplayedAssetTypes();
+		TArray<FNiagaraAssetTagDefinition> AssetTagDefinitions;
+		
+		for(const FStructuredAssetTagDefinitionLookupData& AssetTagDefinitionData : GetStructuredSortedAssetTagDefinitions())
+		{
+			for(const FNiagaraAssetTagDefinition& AssetTagDefinition : AssetTagDefinitionData.AssetTagDefinitions)
+			{
+				if(AssetTagDefinition.DisplayType != ENiagaraAssetTagDefinitionImportance::Primary)
+				{
+					continue;
+				}
+				
+				TArray<UClass*> SupportedClasses = AssetTagDefinition.GetSupportedClasses();
+	
+				bool bCanAssetTagContainDisplayedAssetType = false;
+				for(UClass* SupportedClass : SupportedClasses)
+				{
+					if(DisplayedAssetTypes.Contains(SupportedClass))
+					{
+						bCanAssetTagContainDisplayedAssetType = true;
+						break;
+					}
+				}
+	
+				if(bCanAssetTagContainDisplayedAssetType == false)
+				{
+					continue;
+				}
+	
+				AssetTagDefinitions.Add(AssetTagDefinition);
+			}
+		}
+
+		for(const FNiagaraAssetTagDefinition& AssetTagDefinition : AssetTagDefinitions)
+		{
+			TSharedRef<FNiagaraAssetBrowserMainFilter> NiagaraAssetTagFilter = MakeShared<FNiagaraAssetBrowserMainFilter>(FNiagaraAssetBrowserMainFilter::EFilterMode::NiagaraAssetTag);
+			NiagaraAssetTagFilter->AssetTagDefinition = AssetTagDefinition;
+			TagFilters.Add(NiagaraAssetTagFilter);
+		}
+	}
+	
+	// All
+	{
+		TSharedRef<FNiagaraAssetBrowserMainFilter> AllFilter = MakeShared<FNiagaraAssetBrowserMainFilter>(FNiagaraAssetBrowserMainFilter::EFilterMode::All);
+		AllFilter->ChildFilters = TagFilters;
+		MainFilters.Add(AllFilter);
+	}
+
+	return MainFilters;
+}
+
+void SNiagaraAssetBrowser::OnFilterChanged() const
+{
+	RefreshBackendFilter();
+}
+
+void SNiagaraAssetBrowser::OnGetChildFiltersForFilter(TSharedRef<FNiagaraAssetBrowserMainFilter> NiagaraAssetBrowserMainFilter, TArray<TSharedRef<FNiagaraAssetBrowserMainFilter>>& OutChildren) const
+{
+	OutChildren.Append(NiagaraAssetBrowserMainFilter->ChildFilters);
+}
+
+bool SNiagaraAssetBrowser::OnCompareMainFiltersForEquality(const FNiagaraAssetBrowserMainFilter& MainFilterA, const FNiagaraAssetBrowserMainFilter& MainFilterB) const
+{
+	return MainFilterA == MainFilterB;
+}
+
+// TSharedRef<SWidget> SNiagaraAssetBrowser::GenerateWidgetForMainFilter(const FNiagaraAssetBrowserMainFilter& MainFilter) const
+// {
+// 	FLinearColor FolderColor = FAppStyle::Get().GetSlateColor("ContentBrowser.DefaultFolderColor").GetSpecifiedColor();
+// 	
+// 	TSharedRef<SWidget> Widget = SNew(SHorizontalBox)
+// 	+ SHorizontalBox::Slot()
+// 	.AutoWidth()
+// 	.Padding(2.f)
+// 	[
+// 		SNew(SImage)
+// 		.Image_Lambda([this, MainFilter]()
+// 		{
+// 			if(MainFilterSelector->GetSelectedItems().Contains(MainFilter))
+// 			{
+// 				return FAppStyle::GetBrush("ContentBrowser.AssetTreeFolderOpen");
+// 			}
+// 			else
+// 			{
+// 				return FAppStyle::GetBrush("ContentBrowser.AssetTreeFolderClosed");
+// 			}
+// 		})
+// 		.ColorAndOpacity(FolderColor)
+// 	]
+// 	+ SHorizontalBox::Slot()
+// 	[
+// 		SNew(STextBlock)
+// 		.Text(MainFilter.GetDisplayName())
+// 	];
+//
+// 	if(MainFilter.FilterMode == FNiagaraAssetBrowserMainFilter::EFilterMode::NiagaraAssetTag)
+// 	{
+// 		Widget->SetToolTipText(MainFilter.AssetTagDefinition.Description);	
+// 	}
+// 	
+// 	return Widget;
+// }
+
+TSharedRef<ITableRow> SNiagaraAssetBrowser::GenerateWidgetRowForMainFilter(TSharedRef<FNiagaraAssetBrowserMainFilter> MainFilter, const TSharedRef<STableViewBase>& OwningTable) const
+{
+	FLinearColor FolderColor = FAppStyle::Get().GetSlateColor("ContentBrowser.DefaultFolderColor").GetSpecifiedColor();
+	
+	TSharedRef<SWidget> Widget = SNew(SHorizontalBox)
+	+ SHorizontalBox::Slot()
+	.AutoWidth()
+	.Padding(2.f)
+	[
+		SNew(SImage)
+		.Image_Lambda([this, MainFilter]()
+		{
+			if(MainFilterSelector->GetSelectedItems().Contains(MainFilter))
+			{
+				return FAppStyle::GetBrush("ContentBrowser.AssetTreeFolderOpen");
+			}
+			else
+			{
+				return FAppStyle::GetBrush("ContentBrowser.AssetTreeFolderClosed");
+			}
+		})
+		.ColorAndOpacity(FolderColor)
+	]
+	+ SHorizontalBox::Slot()
+	[
+		SNew(STextBlock)
+		.Text(MainFilter->GetDisplayName())
+	];
+
+	if(MainFilter->FilterMode == FNiagaraAssetBrowserMainFilter::EFilterMode::NiagaraAssetTag)
+	{
+		Widget->SetToolTipText(MainFilter->AssetTagDefinition.Description);	
+	}
+	
+	return SNew(STableRow<TSharedRef<FNiagaraAssetBrowserMainFilter>>, OwningTable)
+	[
+		Widget
+	];
+}
+
+TArray<FAssetData> SNiagaraAssetBrowser::GetSelectedAssets() const
+{
+	return AssetBrowserContent->GetCurrentSelection();
+}
+
+void SNiagaraAssetBrowser::OnAssetSelected(const FAssetData& AssetData)
+{
+	if(AssetData.IsValid() == false)
+	{
+		AssetBrowserDetailsSlot->AttachWidget(SNullWidget::NullWidget);
+		DetailsSwitcher->SetActiveWidgetIndex(0);
+	}
+	
+	if(FNiagaraAssetDetailDatabase::NiagaraAssetDetailDatabase.Contains(AssetData.GetClass()))
+	{
+		if(GetShouldDisplayViewport())
+		{
+			if(AssetData.GetClass() == UNiagaraSystem::StaticClass())
+			{
+				PreviewViewport->SetSystem(*Cast<UNiagaraSystem>(AssetData.GetAsset()));
+			}
+			else
+			{
+				PreviewViewport->ResetAsset();
+			}
+		}
+		
+		TSharedRef<SWidget> SelectedAssetDetails = SNew(SNiagaraSelectedAssetDetails, AssetData)
+			.ShowThumbnail(TAttribute<EVisibility>::CreateSP(this, &SNiagaraAssetBrowser::OnGetThumbnailVisibility))
+			.OnAssetTagActivated(this, &SNiagaraAssetBrowser::OnAssetTagActivated)
+			.OnAssetTagActivatedTooltip(LOCTEXT("SecondaryAssetTagButtonTooltip", "\n\nClicking this tag will activate/deactivate its corresponding filter."));
+		
+		TSharedRef<SWidget> Details = SNew(SVerticalBox)
+		+ SVerticalBox::Slot()
+		.AutoHeight()
+		.HAlign(HAlign_Right)
+		[
+			SNew(SCheckBox)
+			.IsChecked(this, &SNiagaraAssetBrowser::OnShouldDisplayViewport)
+			.OnCheckStateChanged(this, &SNiagaraAssetBrowser::OnShouldDisplayViewportChanged)
+			.ToolTipText(this, &SNiagaraAssetBrowser::OnGetShouldDisplayViewportTooltip)
+			.Visibility(this, &SNiagaraAssetBrowser::OnGetShouldDisplayVisibilityCheckbox)
+			[
+				SNew(STextBlock).Text(LOCTEXT("DisplayViewport", "Display Viewport"))
+			]
+		]
+		+ SVerticalBox::Slot()
+		.AutoHeight()
+		[
+			PreviewViewport.ToSharedRef()
+		]
+		+ SVerticalBox::Slot()
+		.AutoHeight()
+		[
+			SelectedAssetDetails
+		];
+
+		// We don't attach the preview currently until some of the issues are resolved
+		AssetBrowserDetailsSlot->AttachWidget(SelectedAssetDetails);
+
+		DetailsSwitcher->SetActiveWidgetIndex(1);
+	}
+
+	OnAssetSelectedDelegate.ExecuteIfBound(AssetData);
+}
+
+void SNiagaraAssetBrowser::OnAssetsActivated(const TArray<FAssetData>& AssetData, EAssetTypeActivationMethod::Type ActivationType) const
+{
+	OnAssetsActivatedDelegate.ExecuteIfBound(AssetData, ActivationType);
+}
+
+TSharedPtr<SWidget> SNiagaraAssetBrowser::OnGetAssetContextMenu(const TArray<FAssetData>& AssetData) const
+{
+	if(AssetData.Num() == 1)
+	{
+		FName MenuName("NiagaraAssetBrowser.ContextMenu");
+		UContentBrowserAssetContextMenuContext* MenuContext = NewObject<UContentBrowserAssetContextMenuContext>();
+		MenuContext->SelectedAssets = AssetData;
+		FToolMenuContext Context(MenuContext);
+		return UToolMenus::Get()->GenerateWidget(MenuName, Context);
+	}
+
+	return nullptr;
+}
+
+TSharedRef<SToolTip> SNiagaraAssetBrowser::OnGetCustomAssetTooltip(FAssetData& AssetData)
+{
+	return SNew(SToolTip)
+	[
+		SNew(SNiagaraSelectedAssetDetails, AssetData).ShowThumbnail(EVisibility::Collapsed)
+	];
+}
+
+void SNiagaraAssetBrowser::OnMainFilterSelected(TSharedPtr<FNiagaraAssetBrowserMainFilter> MainFilter, ESelectInfo::Type Arg)
+{
+	if(MainFilter.IsValid())
+	{
+		OnFilterChanged();
+
+		LastSelectedMainFilterIdentifierFallback = MainFilter->GetIdentifier();
+		SaveSettings();
+	}
+}
+
+void SNiagaraAssetBrowser::OnAssetTagActivated(const FNiagaraAssetTagDefinition& NiagaraAssetTagDefinition)
+{
+	// First we attempt to select it as a main filter
+	// TArray<FNiagaraAssetBrowserMainFilter> MainFilters = GetMainFilters();
+	TArray<TSharedRef<FNiagaraAssetBrowserMainFilter>> MainFilters = GetMainFilters();
+
+	TArray<TSharedRef<FNiagaraAssetBrowserMainFilter>> AllFilters = MainFilters;
+	while(AllFilters.IsEmpty() == false)
+	{
+		TSharedRef<FNiagaraAssetBrowserMainFilter> CurrentFilter = AllFilters[0];
+		AllFilters.RemoveAt(0);
+		MainFilters.AddUnique(CurrentFilter);
+		AllFilters.Append(CurrentFilter->ChildFilters);
+	}
+	
+	TSharedRef<FNiagaraAssetBrowserMainFilter>* FoundMainFilter = MainFilters.FindByPredicate([NiagaraAssetTagDefinition](const TSharedRef<FNiagaraAssetBrowserMainFilter>& MainFilterCandidate)
+	{
+		return MainFilterCandidate->AssetTagDefinition == NiagaraAssetTagDefinition;
+	});
+
+	if(FoundMainFilter)
+	{
+		MainFilterSelector->SetItemSelection({*FoundMainFilter}, true);
+	}
+	else if(DropdownFilterCache.Contains(NiagaraAssetTagDefinition))
+	{
+		bool bNewState = !DropdownFilterCache[NiagaraAssetTagDefinition]->IsActive();
+		DropdownFilterCache[NiagaraAssetTagDefinition]->SetActive(bNewState);
+	}
+}
+
+TArray<TSharedRef<FFrontendFilter>> SNiagaraAssetBrowser::OnGetExtraFrontendFilters() const
+{
+	TArray<TSharedRef<FFrontendFilter>> Result;
+
+	{
+		using namespace FNiagaraEditorUtilities::AssetBrowser;
+		
+		TSharedRef<FFrontendFilterCategory> NiagaraTagFiltersCategory = MakeShared<FFrontendFilterCategory>(LOCTEXT("NiagaraTagFilterCategoryLabel", "Niagara Tags"), LOCTEXT("", ""));
+
+		// This requires the asset registry to be done loading
+		for(const FNiagaraAssetTagDefinition& AssetTagDefinition : GetFlatSortedAssetTagDefinitions(false))
+		{
+			if(AssetTagDefinition.DisplayType == ENiagaraAssetTagDefinitionImportance::Secondary)
+			{
+				TSharedRef<FFrontendFilter_NiagaraTag> NiagaraTagFilter = MakeShared<FFrontendFilter_NiagaraTag>(AssetTagDefinition, NiagaraTagFiltersCategory);
+				DropdownFilterCache.Add(AssetTagDefinition, NiagaraTagFilter);
+				Result.Add(NiagaraTagFilter);
+			}
+		}
+	}
+	
+	{
+		TSharedRef<FFrontendFilterCategory> NiagaraAdditionalFiltersCategory = MakeShared<FFrontendFilterCategory>(LOCTEXT("NiagaraPropertyFilterCategoryLabel", "Niagara Filters"), LOCTEXT("", ""));
+		if(AvailableClasses.Contains(UNiagaraEmitter::StaticClass()))
+		{
+			Result.Add(MakeShared<FFrontendFilter_NiagaraEmitterInheritance>(true, NiagaraAdditionalFiltersCategory));
+			Result.Add(MakeShared<FFrontendFilter_NiagaraEmitterInheritance>(false, NiagaraAdditionalFiltersCategory));
+		}
+		if(AvailableClasses.Contains(UNiagaraSystem::StaticClass()))
+		{
+			Result.Add(MakeShared<FFrontendFilter_NiagaraSystemEffectType>(true, NiagaraAdditionalFiltersCategory));
+			Result.Add(MakeShared<FFrontendFilter_NiagaraSystemEffectType>(false, NiagaraAdditionalFiltersCategory));
+		}
+	}
+	
+	return Result;
+}
+
+void SNiagaraAssetBrowser::OnExtendAddFilterMenu(UToolMenu* ToolMenu) const
+{
+	for (FToolMenuSection& Section : ToolMenu->Sections)
+	{
+		Section.Blocks.RemoveAll([](FToolMenuEntry& ToolMenuEntry)
+		{
+			return ToolMenuEntry.Name == FName("Common");
+		});
+	}
+	
+	TArray<FName> SectionsToKeep { FName("FilterBarResetFilters"), FName("FX Filters"), FName("AssetFilterBarFilterAdvancedAsset"), FName("Niagara Filters"), FName("Niagara Tags") };
+	ToolMenu->Sections.RemoveAll([&SectionsToKeep](FToolMenuSection& ToolMenuSection)
+	{
+		return SectionsToKeep.Contains(ToolMenuSection.Name) == false;
+	});
+
+	TArray<FNiagaraAssetTagDefinition> AssetTagDefinitions = FNiagaraEditorUtilities::AssetBrowser::GetFlatSortedAssetTagDefinitions();
+	AssetTagDefinitions.RemoveAll([](const FNiagaraAssetTagDefinition& AssetTagDefinition)
+	{
+		return AssetTagDefinition.DisplayType == ENiagaraAssetTagDefinitionImportance::Primary;
+	});
+	
+	for (FToolMenuSection& Section : ToolMenu->Sections)
+	{
+		if(Section.Name == FName("Niagara Tags"))
+		{
+			Section.Blocks.RemoveAll([AssetTagDefinitions](FToolMenuEntry& ToolMenuEntry)
+			{
+				return AssetTagDefinitions.ContainsByPredicate([ToolMenuEntry](const FNiagaraAssetTagDefinition& AssetTagDefinition)
+				{;
+					return AssetTagDefinition.AssetTag.EqualTo(ToolMenuEntry.Label.Get());
+				}) == false;
+			});
+		}
+	}
+}
+
+bool SNiagaraAssetBrowser::OnIsAssetRecent(const FAssetData& AssetCandidate) const
+{
+	return RecentAndFavoritesList.Get()->FindMRUItemIdx(AssetCandidate.PackageName.ToString()) != INDEX_NONE;
+}
+
+EVisibility SNiagaraAssetBrowser::OnGetViewportVisibility() const
+{
+	return (GetShouldDisplayViewport() == true && PreviewViewport->GetPreviewComponent()->GetAsset() != nullptr) ? EVisibility::Visible : EVisibility::Collapsed;
+}
+
+EVisibility SNiagaraAssetBrowser::OnGetThumbnailVisibility() const
+{
+	return OnGetViewportVisibility() == EVisibility::Visible ? EVisibility::Collapsed : EVisibility::Visible;
+}
+
+EVisibility SNiagaraAssetBrowser::OnGetShouldDisplayVisibilityCheckbox() const
+{
+	return AvailableClasses.Contains(UNiagaraSystem::StaticClass()) && GetSelectedAssets().Num() == 1 && GetSelectedAssets()[0].GetClass() == UNiagaraSystem::StaticClass()
+		? EVisibility::Visible
+		: EVisibility::Collapsed;
+}
+
+ECheckBoxState SNiagaraAssetBrowser::OnShouldDisplayViewport() const
+{
+	return bShouldDisplayViewport == true ? ECheckBoxState::Checked : ECheckBoxState::Unchecked;
+}
+
+void SNiagaraAssetBrowser::OnShouldDisplayViewportChanged(ECheckBoxState CheckBoxState)
+{
+	bShouldDisplayViewport = CheckBoxState == ECheckBoxState::Checked ? true : false;
+}
+
+FText SNiagaraAssetBrowser::OnGetShouldDisplayViewportTooltip() const
+{
+	return LOCTEXT("ShouldDisplayViewportTooltip", "If activated, displays Niagara Systems live in a viewport instead of a thumbnail.\nThis will compile the Niagara System if necessary and might slow down performance.");
+}
+
+void SNiagaraAssetBrowser::SaveSettings() const
+{
+	if(SaveSettingsName.IsSet() && bSuppressSaveAndLoad == false)
+	{
+		FNiagaraAssetBrowserConfiguration& Config = UNiagaraAssetBrowserConfig::Get()->MainFilterSelection.Add(SaveSettingsName.GetValue());
+		
+		for(const TSharedRef<FNiagaraAssetBrowserMainFilter>& SelectedMainFilter : MainFilterSelector->GetSelectedItems())
+		{
+			Config.MainFilterSelection.Add(SelectedMainFilter->GetIdentifier());
+		}
+
+		if(Config.MainFilterSelection.Num() == 0 && LastSelectedMainFilterIdentifierFallback != NAME_None)
+		{
+			Config.MainFilterSelection.Add(LastSelectedMainFilterIdentifierFallback);
+		}
+
+		Config.bShouldDisplayViewport = bShouldDisplayViewport;
+
+		UNiagaraAssetBrowserConfig::Get()->SaveEditorConfig();
+	}
+}
+
+void SNiagaraAssetBrowser::LoadSettings()
+{
+	if(SaveSettingsName.IsSet() && bSuppressSaveAndLoad == false)
+	{
+		UNiagaraAssetBrowserConfig::Get()->LoadEditorConfig();
+
+		if(UNiagaraAssetBrowserConfig::Get()->MainFilterSelection.Contains(FName(SaveSettingsName.GetValue())))
+		{
+			FNiagaraAssetBrowserConfiguration& Config = UNiagaraAssetBrowserConfig::Get()->MainFilterSelection[FName(SaveSettingsName.GetValue())];
+
+			bShouldDisplayViewport = Config.bShouldDisplayViewport;
+			
+			TArray<TSharedRef<FNiagaraAssetBrowserMainFilter>> MainFilters = AssetBrowserMainFilters;
+
+			TArray<TSharedRef<FNiagaraAssetBrowserMainFilter>> AllFilters = MainFilters;
+			while(AllFilters.IsEmpty() == false)
+			{
+				TSharedRef<FNiagaraAssetBrowserMainFilter> CurrentFilter = AllFilters[0];
+				AllFilters.RemoveAt(0);
+				MainFilters.AddUnique(CurrentFilter);
+				AllFilters.Append(CurrentFilter->ChildFilters);
+			}
+			
+			MainFilters.RemoveAll([&Config](const TSharedRef<FNiagaraAssetBrowserMainFilter>& MainFilterCandidate)
+			{
+				return Config.MainFilterSelection.Contains(MainFilterCandidate->GetIdentifier()) == false;
+			});
+
+			if(Config.MainFilterSelection.Num() == 1)
+			{
+				LastSelectedMainFilterIdentifierFallback = Config.MainFilterSelection[0];
+			}
+
+			if(Config.MainFilterSelection.Num() >= 1)
+			{
+				MainFilterSelector->SetSelection(MainFilters[0]);
+			}
+		}
+	}
+}
+
+void SNiagaraAssetBrowserWindow::Construct(const FArguments& InArgs)
+{
+	OnAssetsActivatedDelegate = InArgs._AssetBrowserArgs._OnAssetsActivated;
+	
+	SWindow::Construct(SWindow::FArguments()
+		.Title(InArgs._WindowTitle.Get(LOCTEXT("NiagaraAssetBrowserWindowTitle", "Niagara Asset Browser")))
+		.SizingRule(ESizingRule::UserSized)
+		.ClientSize(FVector2D(1400, 750))
+		.SupportsMaximize(false)
+		.SupportsMinimize(false)
+	[
+		SAssignNew(AssetBrowser, SNiagaraAssetBrowser)
+		.AvailableClasses(InArgs._AssetBrowserArgs._AvailableClasses)
+		.RecentAndFavoritesList(InArgs._AssetBrowserArgs._RecentAndFavoritesList)
+		.SaveSettingsName(InArgs._AssetBrowserArgs._SaveSettingsName)
+		.EmptySelectionMessage(InArgs._AssetBrowserArgs._EmptySelectionMessage)
+		.AssetSelectionMode(InArgs._AssetBrowserArgs._AssetSelectionMode)
+		.OnAssetSelected(InArgs._AssetBrowserArgs._OnAssetSelected)
+		.OnAssetsActivated(this, &SNiagaraAssetBrowserWindow::OnAssetsActivated)
+		.AdditionalWidget(InArgs._AssetBrowserArgs._AdditionalWidget)
+	]);
+}
+
+bool SNiagaraAssetBrowserWindow::HasSelectedAssets() const
+{
+	return GetSelectedAssets().Num() > 0;
+}
+
+TArray<FAssetData> SNiagaraAssetBrowserWindow::GetSelectedAssets() const
+{
+	return AssetBrowser->GetSelectedAssets();
+}
+
+void SNiagaraAssetBrowserWindow::OnAssetsActivated(const TArray<FAssetData>& AssetData, EAssetTypeActivationMethod::Type ActivationType)
+{
+	if(OnAssetsActivatedDelegate.ExecuteIfBound(AssetData, ActivationType))
+	{
+		RequestDestroyWindow();
+	}
+}
+
+void SNiagaraAssetBrowserWindow::OnFocusLost(const FFocusEvent& InFocusEvent)
+{
+	SWindow::OnFocusLost(InFocusEvent);
+}
+
+FReply SNiagaraAssetBrowserWindow::OnKeyDown(const FGeometry& MyGeometry, const FKeyEvent& InKeyEvent)
+{
+	FReply Reply = SWindow::OnKeyDown(MyGeometry, InKeyEvent);
+
+	if(InKeyEvent.GetKey() == EKeys::Escape)
+	{
+		RequestDestroyWindow();
+	}
+
+	return Reply;
+}
+
+void SNiagaraCreateAssetWindow::Construct(const FArguments& InArgs, UClass& InCreatedClass)
+{
+	CreatedClass = &InCreatedClass;
+
+	FArguments Args = InArgs;
+
+	TSharedRef<SWidget> CreateAssetControls = SNew(SBox)
+		.Padding(16.f, 16.f)
+		[
+			SNew(SHorizontalBox)
+			+ SHorizontalBox::Slot()
+			.AutoWidth()
+			.Padding(5.f, 1.f)
+			.HAlign(HAlign_Left)
+			[
+				SNew(SButton)
+				.OnClicked(this, &SNiagaraCreateAssetWindow::Proceed)
+				[
+					SNew(SHorizontalBox)
+					+ SHorizontalBox::Slot()
+					.AutoWidth()
+					.Padding(2.f)
+					[
+						SNew(SImage)
+						.Image(FSlateIconFinder::FindIconForClass(CreatedClass.Get()).GetIcon())
+					]
+					+ SHorizontalBox::Slot()
+					.AutoWidth()
+					.Padding(2.f)
+					[
+						SNew(STextBlock).Text(FText::FormatOrdered(LOCTEXT("CreateEmptyAssetButtonLabel", "Create Empty {0}"), CreatedClass->GetDisplayNameText()))
+					]
+				]
+			]
+			+ SHorizontalBox::Slot()
+			[
+				SNew(SSpacer)
+			]
+			+ SHorizontalBox::Slot()
+			.AutoWidth()
+			.Padding(5.f, 1.f)
+			[
+				SNew(SButton)
+				.ContentPadding(FMargin(50.f, 4.f, 50.f, 1.f))
+				.ButtonStyle(&FAppStyle::Get().GetWidgetStyle<FButtonStyle>("PrimaryButton"))
+				.TextStyle(&FAppStyle::Get().GetWidgetStyle<FTextBlockStyle>("PrimaryButtonText"))
+				.HAlign(HAlign_Center)
+				.Text(LOCTEXT("CreatePrimaryButtonLabel", "Create"))
+				// .Text(FText::FormatOrdered(LOCTEXT("CreatePrimaryButtonLabel", "Create {0}"), CreatedClass->GetDisplayNameText()))
+				.OnClicked(this, &SNiagaraCreateAssetWindow::Proceed)
+				.IsEnabled(this, &SNiagaraCreateAssetWindow::HasSelectedAssets)
+				.ToolTipText(this, &SNiagaraCreateAssetWindow::GetCreateButtonTooltip)
+			]
+			+ SHorizontalBox::Slot()
+			.AutoWidth()
+			.Padding(5.f, 1.f)
+			[
+				SNew(SButton)
+				.ContentPadding(FMargin(50.f, 4.f, 50.f, 1.f))
+				.TextStyle(&FAppStyle::Get().GetWidgetStyle<FTextBlockStyle>("PrimaryButtonText"))
+				.HAlign(HAlign_Center)
+				.Text(LOCTEXT("CancelButtonLabel", "Cancel"))
+				.OnClicked(this, &SNiagaraCreateAssetWindow::Cancel)
+			]
+		];
+	
+	Args._AssetBrowserWindowArgs._AssetBrowserArgs._AdditionalWidget = CreateAssetControls;
+	Args._AssetBrowserWindowArgs._AssetBrowserArgs._RecentAndFavoritesList = FNiagaraEditorModule::Get().GetRecentsManager()->GetRecentEmitterAndSystemsList();
+	Args._AssetBrowserWindowArgs._AssetBrowserArgs._OnAssetsActivated = FOnAssetsActivated::CreateSP(this, &SNiagaraCreateAssetWindow::OnAssetsActivatedInternal);
+	SNiagaraAssetBrowserWindow::Construct(Args._AssetBrowserWindowArgs);
+}
+
+void SNiagaraCreateAssetWindow::OnAssetsActivatedInternal(const TArray<FAssetData>& AssetData, EAssetTypeActivationMethod::Type)
+{
+	bProceedWithAction = true;
+}
+
+FReply SNiagaraCreateAssetWindow::Proceed()
+{
+	// The assets don't matter here as the factories making use of this window will retrieve the selected assets afterwards
+	OnAssetsActivated({}, EAssetTypeActivationMethod::Opened);
+	return FReply::Handled();
+}
+
+FReply SNiagaraCreateAssetWindow::Cancel()
+{
+	bProceedWithAction = false;
+	RequestDestroyWindow();
+	return FReply::Handled();
+}
+
+FText SNiagaraCreateAssetWindow::GetCreateButtonTooltip() const
+{
+	return HasSelectedAssets()
+	? FText::FormatOrdered(LOCTEXT("CreateAssetButtonTooltip_Enabled", "Create a new {0} with selected asset {1}"), CreatedClass->GetDisplayNameText(), FText::FromName(GetSelectedAssets()[0].AssetName))
+	: LOCTEXT("CreateAssetButtonTooltip_Disabled", "Please select an asset as a base for your new effect. Alternatively, create an empty asset.");
+}
+
+void SNiagaraAddEmitterToSystemWindow::Construct(const FArguments& InArgs, TSharedRef<FNiagaraSystemViewModel> SystemViewModel)
+{
+	WeakSystemViewModel = SystemViewModel;
+	
+	SNiagaraAssetBrowser::FArguments AssetBrowserArgs;
+	AssetBrowserArgs._AvailableClasses = {UNiagaraEmitter::StaticClass()};
+	AssetBrowserArgs._AssetSelectionMode = ESelectionMode::Single;
+	AssetBrowserArgs._SaveSettingsName = FName("NiagaraAssetBrowser.AddEmitter");
+	AssetBrowserArgs._EmptySelectionMessage = LOCTEXT("EmptyEmitterSelectionUserText", "Select an emitter to add it to your system, or add a new empty emitter");
+	AssetBrowserArgs._OnAssetsActivated = FOnAssetsActivated::CreateSP(this, &SNiagaraAddEmitterToSystemWindow::OnAssetsActivatedInternal);
+	AssetBrowserArgs._RecentAndFavoritesList = FNiagaraEditorModule::Get().GetRecentsManager()->GetRecentEmitterAndSystemsList();
+
+	AssetBrowserArgs._AdditionalWidget = SNew(SBox)
+	   .Padding(16.f, 16.f)
+	   [
+		   SNew(SHorizontalBox)
+		   + SHorizontalBox::Slot()
+		   .AutoWidth()
+		   .Padding(5.f, 1.f)
+		   .HAlign(HAlign_Left)
+		   [
+			   SNew(SButton)
+			   .OnClicked_Lambda([this]()
+			   {
+			   		if(WeakSystemViewModel.IsValid())
+			   		{
+						WeakSystemViewModel.Pin()->AddEmptyEmitter();
+			   			RequestDestroyWindow();
+			   		}
+			   	
+				   return FReply::Handled();
+			   })
+			   [
+				   SNew(SHorizontalBox)
+				   + SHorizontalBox::Slot()
+				   .AutoWidth()
+				   .Padding(2.f)
+				   [
+					   SNew(SImage)
+					   .Image(FSlateIconFinder::FindIconForClass(UNiagaraEmitter::StaticClass()).GetIcon())
+				   ]
+				   + SHorizontalBox::Slot()
+				   .AutoWidth()
+				   .Padding(2.f)
+				   [
+					   SNew(STextBlock).Text(FText::FormatOrdered(LOCTEXT("AddEmptyEmitterButtonLabel", "Add Empty {0}"), UNiagaraEmitter::StaticClass()->GetDisplayNameText()))
+				   ]
+			   ]
+		   ]
+		   + SHorizontalBox::Slot()
+		   [
+			   SNew(SSpacer)
+		   ]
+		   + SHorizontalBox::Slot()
+		   .AutoWidth()
+		   .Padding(5.f, 1.f)
+		   [
+			   SNew(SButton)
+			   .ContentPadding(FMargin(50.f, 4.f, 50.f, 1.f))
+			   .ButtonStyle(&FAppStyle::Get().GetWidgetStyle<FButtonStyle>("PrimaryButton"))
+			   .TextStyle(&FAppStyle::Get().GetWidgetStyle<FTextBlockStyle>("PrimaryButtonText"))
+			   .HAlign(HAlign_Center)
+			   .Text(LOCTEXT("AddEmitterPrimaryButtonLabel", "Add"))
+			   .OnClicked(this, &SNiagaraAddEmitterToSystemWindow::AddSelectedEmitters)
+			   .IsEnabled(this, &SNiagaraAddEmitterToSystemWindow::HasSelectedAssets)
+			   .ToolTipText(this, &SNiagaraAddEmitterToSystemWindow::GetAddButtonTooltip)
+		   ]
+		   + SHorizontalBox::Slot()
+		   .AutoWidth()
+		   .Padding(5.f, 1.f)
+		   [
+			   SNew(SButton)
+			   .ContentPadding(FMargin(50.f, 4.f, 50.f, 1.f))
+			   .TextStyle(&FAppStyle::Get().GetWidgetStyle<FTextBlockStyle>("PrimaryButtonText"))
+			   .HAlign(HAlign_Center)
+			   .Text(LOCTEXT("CancelButtonLabel", "Cancel"))
+			   .OnClicked(this, &SNiagaraAddEmitterToSystemWindow::Cancel)
+		   ]
+	   ];
+
+	SNiagaraAssetBrowserWindow::FArguments WindowArgs;
+	WindowArgs._WindowTitle = LOCTEXT("AddEmitterToSystemWindowTitle", "Add Emitter to your System");
+	WindowArgs._AssetBrowserArgs = AssetBrowserArgs;
+	
+	SNiagaraAssetBrowserWindow::Construct(WindowArgs);
+}
+
+void SNiagaraAddEmitterToSystemWindow::OnAssetsActivatedInternal(const TArray<FAssetData>& AssetData, EAssetTypeActivationMethod::Type) const
+{
+	if(WeakSystemViewModel.IsValid())
+	{
+		if(AssetData.Num() == 0)
+		{
+			WeakSystemViewModel.Pin()->AddEmptyEmitter();
+		}
+		else
+		{
+			for(const FAssetData& Asset : AssetData)
+			{
+				WeakSystemViewModel.Pin()->AddEmitterFromAssetData(Asset);
+			}
+		}
+	}
+}
+
+FReply SNiagaraAddEmitterToSystemWindow::AddEmptyEmitter()
+{
+	OnAssetsActivated({}, EAssetTypeActivationMethod::Opened);
+	return FReply::Handled();
+}
+
+FReply SNiagaraAddEmitterToSystemWindow::AddSelectedEmitters()
+{
+	OnAssetsActivated(AssetBrowser->GetSelectedAssets(), EAssetTypeActivationMethod::Opened);
+	return FReply::Handled();
+}
+
+FReply SNiagaraAddEmitterToSystemWindow::Cancel()
+{
+	RequestDestroyWindow();
+	return FReply::Handled();
+}
+
+FText SNiagaraAddEmitterToSystemWindow::GetAddButtonTooltip() const
+{
+	return GetSelectedAssets().Num() > 0
+	? LOCTEXT("AddEmitterButtonTooltip_ValidSelection", "Add the selected emitter to your System")
+	: LOCTEXT("AddEmitterButtonTooltip_InvalidSelection", "Select an emitter to add to your System");
+}
+
+
+#undef LOCTEXT_NAMESPACE
+
+END_SLATE_FUNCTION_BUILD_OPTIMIZATION
diff --git a/Engine/Plugins/FX/Niagara/Source/NiagaraEditor/Private/Widgets/AssetBrowser/SNiagaraAssetBrowserContent.cpp b/Engine/Plugins/FX/Niagara/Source/NiagaraEditor/Private/Widgets/AssetBrowser/SNiagaraAssetBrowserContent.cpp
new file mode 100644
--- /dev/null
+++ b/Engine/Plugins/FX/Niagara/Source/NiagaraEditor/Private/Widgets/AssetBrowser/SNiagaraAssetBrowserContent.cpp
@@ -0,0 +1,42 @@
+﻿// Copyright Epic Games, Inc. All Rights Reserved.
+
+#include "Widgets/AssetBrowser/SNiagaraAssetBrowserContent.h"
+
+#include "ContentBrowserModule.h"
+#include "IContentBrowserDataModule.h"
+#include "IContentBrowserSingleton.h"
+#include "SlateOptMacros.h"
+
+BEGIN_SLATE_FUNCTION_BUILD_OPTIMIZATION
+
+#define LOCTEXT_NAMESPACE "NiagaraAssetBrowserContent"
+
+void SNiagaraAssetBrowserContent::Construct(const FArguments& InArgs)
+{	
+	FAssetPickerConfig Config = InArgs._InitialConfig;
+	Config.GetCurrentSelectionDelegates.Add(&GetCurrentSelectionDelegate);
+	Config.RefreshAssetViewDelegates.Add(&RefreshAssetViewDelegate);
+	Config.SyncToAssetsDelegates.Add(&SyncToAssetsDelegate);
+	Config.SetFilterDelegates.Add(&SetNewFilterDelegate);
+	Config.bCanShowRealTimeThumbnails = true;
+	
+	FContentBrowserModule& ContentBrowserModule = FModuleManager::Get().LoadModuleChecked<FContentBrowserModule>("ContentBrowser");
+	ChildSlot
+	[
+		ContentBrowserModule.Get().CreateAssetPicker(Config)
+	];
+}
+
+void SNiagaraAssetBrowserContent::SetARFilter(FARFilter InFilter)
+{
+	SetNewFilterDelegate.Execute(InFilter);
+}
+
+TArray<FAssetData> SNiagaraAssetBrowserContent::GetCurrentSelection() const
+{
+	return GetCurrentSelectionDelegate.Execute();
+}
+
+#undef LOCTEXT_NAMESPACE
+
+END_SLATE_FUNCTION_BUILD_OPTIMIZATION
diff --git a/Engine/Plugins/FX/Niagara/Source/NiagaraEditor/Private/Widgets/AssetBrowser/SNiagaraSelectedAssetDetails.cpp b/Engine/Plugins/FX/Niagara/Source/NiagaraEditor/Private/Widgets/AssetBrowser/SNiagaraSelectedAssetDetails.cpp
new file mode 100644
--- /dev/null
+++ b/Engine/Plugins/FX/Niagara/Source/NiagaraEditor/Private/Widgets/AssetBrowser/SNiagaraSelectedAssetDetails.cpp
@@ -0,0 +1,511 @@
+﻿// Copyright Epic Games, Inc. All Rights Reserved.
+
+#include "Widgets/AssetBrowser/SNiagaraSelectedAssetDetails.h"
+#include "NiagaraAssetTagDefinitions.h"
+#include "NiagaraEditorStyle.h"
+#include "NiagaraEditorUtilities.h"
+#include "SlateOptMacros.h"
+#include "Styling/SlateIconFinder.h"
+#include "ThumbnailRendering/ThumbnailManager.h"
+#include "Widgets/Input/SButton.h"
+#include "Widgets/Layout/SScaleBox.h"
+#include "Widgets/Layout/SWrapBox.h"
+#include "Widgets/Layout/SSpacer.h"
+#include "Widgets/Layout/SSeparator.h"
+#include "Widgets/Layout/SScrollBox.h"
+
+#define LOCTEXT_NAMESPACE "SNiagaraSelectedAssetPreview"
+
+TMap<UClass*, FNiagaraAssetDetailClassInfo> FNiagaraAssetDetailDatabase::NiagaraAssetDetailDatabase;
+
+BEGIN_SLATE_FUNCTION_BUILD_OPTIMIZATION
+
+void FNiagaraAssetDetailDatabase::Init()
+{
+	FName SimulationTargetTag("HasGPUEmitter");
+	
+	// Emitters
+	{
+		FNiagaraAssetDetailClassInfo EmitterClassInfo;
+		EmitterClassInfo.DisplayedProperties =
+			{
+				// Inheritance
+				{
+					FDisplayedPropertyData::FShouldDisplayProperty(),
+					FDisplayedPropertyData::FGenerateWidget::CreateLambda([](const FAssetData& AssetData) -> TSharedRef<SWidget>
+					{
+						return SNew(STextBlock)
+						.Text(LOCTEXT("Inheritance", "Inheritance"));
+					}),
+					FDisplayedPropertyData::FGenerateWidget::CreateLambda([](const FAssetData& AssetData) -> TSharedRef<SWidget>
+					{
+						bool bUseInheritance = false;
+						if(FNiagaraEditorUtilities::GetIsInheritableFromAssetRegistryTags(AssetData, bUseInheritance))
+						{
+							return SNew(STextBlock).Text(bUseInheritance
+								? LOCTEXT("Emitter_UseInheritance_Yes", "Yes")
+								: LOCTEXT("Emitter_UseInheritance_No", "No")
+								);
+						}
+
+						return SNullWidget::NullWidget;
+					})
+				},
+				// Simulation Target
+				{
+					FDisplayedPropertyData::FShouldDisplayProperty(),
+					FDisplayedPropertyData::FGenerateWidget::CreateLambda([](const FAssetData& AssetData) -> TSharedRef<SWidget>
+					{
+						return SNew(STextBlock)
+						.Text(LOCTEXT("SimulationTargetLabel", "Runs on"));
+					}),
+					FDisplayedPropertyData::FGenerateWidget::CreateLambda([SimulationTargetTag](const FAssetData& Asset)
+					{
+						return SNew(SImage)
+						.Image_Lambda([Asset, SimulationTargetTag]() -> const FSlateBrush*
+						{
+							if(Asset.FindTag(SimulationTargetTag))
+							{
+								FString Value;
+								if(Asset.GetTagValue(SimulationTargetTag, Value))
+								{
+									if(Value == "True")
+									{
+										return FNiagaraEditorStyle::Get().GetBrush("NiagaraEditor.Stack.GPUIcon");
+									}
+									else
+									{
+										return FNiagaraEditorStyle::Get().GetBrush("NiagaraEditor.Stack.CPUIcon");
+									}
+								}
+							}
+
+							return FAppStyle::GetNoBrush();
+						});
+					})
+				}
+		};
+		EmitterClassInfo.GetDescriptionDelegate = FNiagaraAssetDetailClassInfo::FGetDescription::CreateLambda([](const FAssetData& AssetData)
+		{
+			FName DescriptionTagName("TemplateAssetDescription");
+			if(AssetData.FindTag(DescriptionTagName))
+			{
+				FString Description;
+				AssetData.GetTagValue(DescriptionTagName, Description);
+				return FText::FromString(Description);
+			}
+
+			return FText::GetEmpty();
+		});
+		
+		NiagaraAssetDetailDatabase.Add(UNiagaraEmitter::StaticClass(), EmitterClassInfo);
+	}
+
+	// Systems
+	{
+		FNiagaraAssetDetailClassInfo SystemClassInfo;
+		SystemClassInfo.DisplayedProperties =
+		{
+			// Number of emitters
+			{
+				FDisplayedPropertyData::FShouldDisplayProperty::CreateLambda([](const FAssetData& AssetData) -> bool
+				{
+					return AssetData.FindTag("NumEmitters");
+				}),
+				FDisplayedPropertyData::FGenerateWidget::CreateLambda([](const FAssetData& AssetData) -> TSharedRef<SWidget>
+				{
+					return SNew(STextBlock)
+					.Text(LOCTEXT("NumberOfEmitters", "Number of Emitters"));
+				}),
+				FDisplayedPropertyData::FGenerateWidget::CreateLambda([](const FAssetData& AssetData) -> TSharedRef<SWidget>
+				{
+					int32 NumEmitters = INDEX_NONE;
+					if(AssetData.GetTagValue("NumEmitters", NumEmitters))
+					{
+						return SNew(STextBlock).Text(FText::AsNumber(NumEmitters));
+					}
+
+					return SNullWidget::NullWidget;
+				})
+			}
+		};
+		SystemClassInfo.GetDescriptionDelegate = FNiagaraAssetDetailClassInfo::FGetDescription::CreateLambda([](const FAssetData& AssetData)
+		{
+			FName DescriptionTagName("TemplateAssetDescription");
+			if(AssetData.FindTag(DescriptionTagName))
+			{
+				FString Description;
+				AssetData.GetTagValue(DescriptionTagName, Description);
+				return FText::FromString(Description);
+			}
+		
+			return FText::GetEmpty();
+		});
+		
+		NiagaraAssetDetailDatabase.Add(UNiagaraSystem::StaticClass(), SystemClassInfo);
+	}
+}
+
+void SNiagaraAssetTag::Construct(const FArguments& InArgs, const FNiagaraAssetTagDefinition& InAssetTagDefinition)
+{
+	AssetTagDefinition = InAssetTagDefinition;
+	OnAssetTagActivated = InArgs._OnAssetTagActivated;
+	OnAssetTagActivatedTooltip = InArgs._OnAssetTagActivatedTooltip;
+
+	if(OnAssetTagActivated.IsBound() && OnAssetTagActivatedTooltip.IsSet())
+	{
+		FText TooltipText = FText::FormatOrdered(FText::AsCultureInvariant("{0}{1}"), AssetTagDefinition.Description,
+			OnAssetTagActivatedTooltip.GetValue());
+
+		SetToolTipText(TooltipText);
+	}
+	else
+	{
+		SetToolTipText(AssetTagDefinition.Description);
+	}
+
+	TSharedPtr<SWidget> ContentWidget = nullptr;
+
+	TSharedRef<STextBlock> DisplayNameWidget = SNew(STextBlock)
+		.Text(FText::FromString(AssetTagDefinition.AssetTag.ToString()))
+		.TextStyle(&FNiagaraEditorStyle::Get().GetWidgetStyle<FTextBlockStyle>("NiagaraEditor.AssetBrowser.AssetTag.Text"));
+
+	if(OnAssetTagActivated.IsBound())
+	{
+		ContentWidget = SNew(SButton)
+		.ButtonStyle(&FAppStyle::GetWidgetStyle<FButtonStyle>("HoverHintOnly"))
+		.OnClicked(this, &SNiagaraAssetTag::OnClicked)
+		[
+			DisplayNameWidget
+		];
+	}
+	else
+	{
+		ContentWidget = DisplayNameWidget;
+	}
+	
+	ChildSlot
+	[
+		SNew(SBorder)
+		.BorderImage(FNiagaraEditorStyle::Get().GetBrush("NiagaraEditor.AssetBrowser.AssetTag.OuterBorder"))
+		.BorderBackgroundColor(AssetTagDefinition.Color)
+		.Padding(1.f)
+		[
+			SNew(SBorder)
+			.BorderImage(FNiagaraEditorStyle::Get().GetBrush("NiagaraEditor.AssetBrowser.AssetTag.InnerBorder"))
+			.Padding(8.f, 2.f)
+			[
+				ContentWidget.ToSharedRef()
+			]
+		]
+	];
+}
+
+FReply SNiagaraAssetTag::OnClicked() const
+{
+	OnAssetTagActivated.ExecuteIfBound(AssetTagDefinition);
+	return FReply::Handled();
+}
+
+void SNiagaraAssetTagRow::Construct(const FArguments& InArgs, const FAssetData& Asset)
+{
+	using namespace FNiagaraEditorUtilities::AssetBrowser;
+
+	TArray<FNiagaraAssetTagDefinition> AllAssetTagDefinitions = GetFlatSortedAssetTagDefinitions();
+
+	TArray<FNiagaraAssetTagDefinition> MatchingAssetTagDefinitions;
+	for(const FNiagaraAssetTagDefinition& AssetTagDefinition : AllAssetTagDefinitions)
+	{
+		if(AssetTagDefinition.DoesAssetDataContainTag(Asset))
+		{
+			MatchingAssetTagDefinitions.Add(AssetTagDefinition);	
+		}		
+	}
+
+	TSharedRef<SWrapBox> AssetTagRow = SNew(SWrapBox).UseAllottedSize(true);
+	
+	for(const FNiagaraAssetTagDefinition& MatchingAssetTagDefinition : MatchingAssetTagDefinitions)
+	{
+		if(InArgs._DisplayType.IsSet() && InArgs._DisplayType.GetValue() != MatchingAssetTagDefinition.DisplayType)
+		{
+			continue;
+		}
+			
+		AssetTagRow->AddSlot()
+		.Padding(2.f)
+		[
+			SNew(SNiagaraAssetTag, MatchingAssetTagDefinition)
+			.OnAssetTagActivated(InArgs._OnAssetTagActivated)
+			.OnAssetTagActivatedTooltip(InArgs._OnAssetTagActivatedTooltip)
+		];
+	}
+
+	SetVisibility(AssetTagRow->GetChildren()->Num() > 0 ? EVisibility::Visible : EVisibility::Collapsed);
+
+	TSharedRef<SWidget> ContentWidget = AssetTagRow;
+
+	if(InArgs._DisplayType.IsSet())
+	{
+		FText Label = StaticEnum<ENiagaraAssetTagDefinitionImportance>()->GetDisplayNameTextByValue((int64) InArgs._DisplayType.GetValue());
+		
+		ContentWidget = SNew(SVerticalBox)
+		+ SVerticalBox::Slot()
+		.AutoHeight()
+		.Padding(0.f, 2.f)
+		[
+			SNew(STextBlock)
+			.Text(FText::FormatOrdered(LOCTEXT("AssetTagRowTypeLabel", "{0} Tags"), Label))
+			.TextStyle(&FAppStyle::GetWidgetStyle<FTextBlockStyle>("NormalText.Subdued"))
+			.ToolTipText(this, &SNiagaraAssetTagRow::GetDisplayTypeTooltipText, InArgs._DisplayType.GetValue())
+		]
+		+ SVerticalBox::Slot()
+		.AutoHeight()
+		.Padding(0.f, 2.f)
+		[
+			AssetTagRow
+		];
+	}
+	
+	ChildSlot
+	[
+		ContentWidget
+	];
+}
+
+FText SNiagaraAssetTagRow::GetDisplayTypeTooltipText(ENiagaraAssetTagDefinitionImportance DisplayType) const
+{
+	static FText TagDescriptionText = LOCTEXT("GeneralTagsLabelTooltip", "Assets can be assigned primary and secondary tags.\nThese tags are defined in an Niagara Asset Tag Definitions asset and can be assigned to an asset by right-clicking it in the Content Browser under 'Manage Tags'.");
+	if(DisplayType == ENiagaraAssetTagDefinitionImportance::Primary)
+	{
+		return TagDescriptionText;
+	}
+	else if(DisplayType == ENiagaraAssetTagDefinitionImportance::Secondary)
+	{
+		return FText::FormatOrdered(LOCTEXT("SecondaryTagsLabelTooltip", "{0}{1}"), TagDescriptionText,
+			LOCTEXT("SecondaryTagsLabelTooltipSpecific", "\n\nSecondary tags can be clicked on to manage additional filters.\nThese additional filters can also be accessed from the filters menu next to the search bar."));
+	}
+
+	return TagDescriptionText;
+}
+
+void SNiagaraSelectedAssetDetails::Construct(const FArguments& InArgs, const FAssetData& Asset)
+{
+	AssetData = Asset;
+	
+	ShowThumbnail = InArgs._ShowThumbnail;
+	OnAssetTagActivated = InArgs._OnAssetTagActivated;
+	OnAssetTagActivatedTooltip = InArgs._OnAssetTagActivatedTooltip;
+	
+	ChildSlot
+	[
+		SNew(SBox)
+		.MinDesiredWidth(200.f)
+		.MaxDesiredWidth(450.f)
+		.Padding(11.5f)
+		[
+			SNew(SVerticalBox)
+			+ SVerticalBox::Slot()
+			.AutoHeight()
+			.HAlign(HAlign_Fill)
+			.VAlign(VAlign_Fill)
+			.Padding(0.f, 0.f, 0.f, 22.f)
+			[
+				CreateAssetThumbnailWidget()
+			]
+			+ SVerticalBox::Slot()
+			.AutoHeight()
+			.Padding(0.f)
+			[
+				CreateTitleWidget()
+			]
+			+ SVerticalBox::Slot()
+			.AutoHeight()
+			.MaxHeight(150.f)
+			.Padding(0.f, 6.f)
+			[
+				SNew(SScrollBox)
+				+ SScrollBox::Slot()
+				[
+					SNew(SBox)
+					.MaxDesiredWidth(InArgs._MaxDesiredDescriptionWidth)
+					[
+						CreateDescriptionWidget()
+					]
+				]
+			]
+			+ SVerticalBox::Slot()
+			.AutoHeight()
+			.Padding(0.f, 6.f)
+			[
+				CreateOptionalPropertiesList()
+			]
+			+ SVerticalBox::Slot()
+			.AutoHeight()
+			.Padding(0.f, 6.f)
+			[
+				CreateAssetTagRow()
+			]
+		]
+	]; 
+}
+
+TSharedRef<SWidget> SNiagaraSelectedAssetDetails::CreateAssetThumbnailWidget()
+{
+	CurrentAssetThumbnail = MakeShared<FAssetThumbnail>(AssetData, 256.f, 256.f, UThumbnailManager::Get().GetSharedThumbnailPool());
+	FAssetThumbnailConfig Config;
+	Config.bAllowRealTimeOnHovered = false;
+	
+	return SNew(SBox)
+	.WidthOverride(256.f)
+	.HeightOverride(192.f)
+	.Visibility(ShowThumbnail)
+	[
+		SNew(SScaleBox)
+		.Stretch(EStretch::ScaleToFill)
+		[
+			CurrentAssetThumbnail->MakeThumbnailWidget(Config)
+		]
+	];
+}
+
+TSharedRef<SWidget> SNiagaraSelectedAssetDetails::CreateTitleWidget()
+{
+	return SNew(SVerticalBox)
+	+ SVerticalBox::Slot()
+	.AutoHeight()
+	.HAlign(HAlign_Left)
+	.Padding(2.f)
+	[
+		SNew(STextBlock)
+		.Text(FText::FromName(AssetData.AssetName))
+		.TextStyle(&FNiagaraEditorStyle::Get().GetWidgetStyle<FTextBlockStyle>("NiagaraEditor.AssetBrowser.AssetTitle"))
+	]
+	+ SVerticalBox::Slot()
+	.AutoHeight()
+	.HAlign(HAlign_Left)
+	.Padding(2.f)
+	[
+		CreateTypeWidget()
+	];
+}
+
+TSharedRef<SWidget> SNiagaraSelectedAssetDetails::CreateTypeWidget()
+{
+	return SNew(SHorizontalBox)
+	+ SHorizontalBox::Slot()
+	.AutoWidth()
+	.Padding(2.f)
+	[
+		SNew(SImage)
+		.Image(FSlateIconFinder::FindIconForClass(AssetData.GetClass()).GetIcon())
+	]
+	+ SHorizontalBox::Slot()
+	.AutoWidth()
+	[
+		SNew(STextBlock)
+		.Text(AssetData.GetClass()->GetDisplayNameText())
+		.TextStyle(&FNiagaraEditorStyle::Get().GetWidgetStyle<FTextBlockStyle>("NiagaraEditor.AssetBrowser.AssetType"))
+	];
+}
+
+TSharedRef<SWidget> SNiagaraSelectedAssetDetails::CreateDescriptionWidget()
+{
+	if(FNiagaraAssetDetailDatabase::NiagaraAssetDetailDatabase.Contains(AssetData.GetClass()))
+	{
+		FText Description = FNiagaraAssetDetailDatabase::NiagaraAssetDetailDatabase[AssetData.GetClass()].GetDescriptionDelegate.Execute(AssetData);
+		return SNew(STextBlock)
+			.AutoWrapText(true)
+			.Text(Description)
+			.Visibility(Description.IsEmpty() ? EVisibility::Collapsed : EVisibility::Visible);
+	}
+
+	return SNew(SBox).Visibility(EVisibility::Collapsed);
+}
+
+TSharedRef<SWidget> SNiagaraSelectedAssetDetails::CreateOptionalPropertiesList()
+{
+	if(FNiagaraAssetDetailDatabase::NiagaraAssetDetailDatabase.Contains(AssetData.GetClass()))
+	{
+		TSharedRef<SVerticalBox> Result = SNew(SVerticalBox);
+		
+		for(const FDisplayedPropertyData& DisplayedProperty : FNiagaraAssetDetailDatabase::NiagaraAssetDetailDatabase[AssetData.GetClass()].DisplayedProperties)
+		{			
+			if(DisplayedProperty.ShouldDisplayPropertyDelegate.IsBound())
+			{
+				if(DisplayedProperty.ShouldDisplayPropertyDelegate.Execute(AssetData) == false)
+				{
+					continue;
+				}
+			}
+			
+			Result->AddSlot()
+			.AutoHeight()
+			.Padding(10.f, 3.f)
+			[
+				SNew(SHorizontalBox)
+				+ SHorizontalBox::Slot()
+				.HAlign(HAlign_Left)
+				.AutoWidth()
+				[
+					DisplayedProperty.NameWidgetDelegate.Execute(AssetData)
+				]
+				+ SHorizontalBox::Slot()
+				[
+					SNew(SSpacer)
+				]
+				+ SHorizontalBox::Slot()
+				.HAlign(HAlign_Right)
+				.AutoWidth()
+				[
+					DisplayedProperty.ValueWidgetDelegate.Execute(AssetData)
+				]
+			];
+
+			Result->AddSlot()
+			.AutoHeight()
+			.Padding(10.f, 3.f)
+			[
+				SNew(SSeparator)
+				.Orientation(Orient_Horizontal)
+				.SeparatorImage(FNiagaraEditorStyle::Get().GetBrush("NiagaraEditor.AssetBrowser.PropertySeparator"))
+				.Thickness(1.f)
+			];
+		}
+
+		if(Result->NumSlots() == 0)
+		{
+			return SNew(SBox).Visibility(EVisibility::Collapsed);
+		}
+		
+		return Result;
+	}
+
+	return SNew(SBox).Visibility(EVisibility::Collapsed);
+}
+
+TSharedRef<SWidget> SNiagaraSelectedAssetDetails::CreateAssetTagRow()
+{
+	return SNew(SVerticalBox)
+		+ SVerticalBox::Slot()
+		.AutoHeight()
+		.Padding(0.f, 2.f)
+		[
+			SNew(SNiagaraAssetTagRow, AssetData)
+			//.OnAssetTagActivated(OnAssetTagActivated)
+			//.OnAssetTagActivatedTooltip(OnAssetTagActivatedTooltip)
+			.DisplayType(ENiagaraAssetTagDefinitionImportance::Primary)
+		]
+		+ SVerticalBox::Slot()
+		.AutoHeight()
+		.Padding(0.f, 2.f)
+		[
+			SNew(SNiagaraAssetTagRow, AssetData)
+			.OnAssetTagActivated(OnAssetTagActivated)
+			.DisplayType(ENiagaraAssetTagDefinitionImportance::Secondary)
+			.OnAssetTagActivatedTooltip(OnAssetTagActivatedTooltip)
+		];	
+}
+
+END_SLATE_FUNCTION_BUILD_OPTIMIZATION
+
+#undef LOCTEXT_NAMESPACE
diff --git a/Engine/Plugins/FX/Niagara/Source/NiagaraEditor/Public/NiagaraAssetTagDefinitionsFactory.h b/Engine/Plugins/FX/Niagara/Source/NiagaraEditor/Public/NiagaraAssetTagDefinitionsFactory.h
new file mode 100644
--- /dev/null
+++ b/Engine/Plugins/FX/Niagara/Source/NiagaraEditor/Public/NiagaraAssetTagDefinitionsFactory.h
@@ -0,0 +1,17 @@
+﻿// Copyright Epic Games, Inc. All Rights Reserved.
+
+#pragma once
+
+#include "CoreMinimal.h"
+#include "Factories/Factory.h"
+#include "NiagaraAssetTagDefinitionsFactory.generated.h"
+
+UCLASS()
+class NIAGARAEDITOR_API UNiagaraAssetTagDefinitionsFactory : public UFactory
+{
+	GENERATED_BODY()
+
+	UNiagaraAssetTagDefinitionsFactory();
+	
+	virtual UObject* FactoryCreateNew(UClass* InClass, UObject* InParent, FName InName, EObjectFlags Flags, UObject* Context, FFeedbackContext* Warn) override;
+};
diff --git a/Engine/Plugins/FX/Niagara/Source/NiagaraEditor/Public/NiagaraEditorModule.h b/Engine/Plugins/FX/Niagara/Source/NiagaraEditor/Public/NiagaraEditorModule.h
--- a/Engine/Plugins/FX/Niagara/Source/NiagaraEditor/Public/NiagaraEditorModule.h
+++ b/Engine/Plugins/FX/Niagara/Source/NiagaraEditor/Public/NiagaraEditorModule.h
@@ -15,6 +15,7 @@
 #include "Customizations/NiagaraDataInterfaceSimCacheVisualizer.h"
 #include "NiagaraEditorModule.generated.h"
 
+class FNiagaraRecentAndFavoritesManager;
 class IAssetTools;
 class IAssetTypeActions;
 class INiagaraEditorTypeUtilities;
@@ -149,6 +150,8 @@ public:
 	/** Get the instance of this module. */
 	NIAGARAEDITOR_API static FNiagaraEditorModule& Get();
 
+	NIAGARAEDITOR_API FNiagaraRecentAndFavoritesManager* GetRecentsManager();
+	
 	/** Start the compilation of the specified script. */
 	virtual int32 CompileScript(const FNiagaraCompileRequestDataBase* InCompileRequest, const FNiagaraCompileRequestDuplicateDataBase* InCompileRequestDuplicate, const FNiagaraCompileOptions& InCompileOptions);
 	virtual TSharedPtr<FNiagaraVMExecutableData> GetCompilationResult(int32 JobID, bool bWait, FNiagaraScriptCompileMetrics& ScriptMetrics);
@@ -176,7 +179,7 @@ public:
 	/** Register/unregister niagara editor settings. */
 	void RegisterSettings();
 	void UnregisterSettings();
-
+	
 	/** Gets Niagara editor type utilities for a specific type if there are any registered. */
 	TSharedPtr<INiagaraEditorTypeUtilities, ESPMode::ThreadSafe> NIAGARAEDITOR_API GetTypeUtilities(const FNiagaraTypeDefinition& Type);
 
@@ -432,7 +435,7 @@ private:
 	TSharedPtr<FNiagaraComponentBroker> NiagaraComponentBroker;
 
 	TMap<const UScriptStruct*, FOnCreateMovieSceneTrackForParameter> TypeToParameterTrackCreatorMap;
-
+	
 	IConsoleCommand* TestCompileScriptCommand;
 	IConsoleCommand* ValidateScriptVariableGuidsCommand;
 	IConsoleCommand* ValidateAndFixScriptVariableGuidsCommand;
@@ -494,4 +497,6 @@ private:
 	TAssetPreloadCache<UNiagaraParameterDefinitions> ParameterDefinitionsAssetCache;
 
 	TArray<UClass*> PluginAssetClassesPreloaded;
+
+	TSharedPtr<FNiagaraRecentAndFavoritesManager> RecentAndFavoritesManager;
 };
diff --git a/Engine/Plugins/FX/Niagara/Source/NiagaraEditor/Public/NiagaraEditorStyle.h b/Engine/Plugins/FX/Niagara/Source/NiagaraEditor/Public/NiagaraEditorStyle.h
--- a/Engine/Plugins/FX/Niagara/Source/NiagaraEditor/Public/NiagaraEditorStyle.h
+++ b/Engine/Plugins/FX/Niagara/Source/NiagaraEditor/Public/NiagaraEditorStyle.h
@@ -24,6 +24,7 @@ private:
 	FNiagaraEditorStyle();
 	void InitStats();
 	void InitAssetPicker();
+	void InitAssetBrowser();
 	void InitActionMenu();
 	void InitEmitterHeader();
 	void InitParameters();
diff --git a/Engine/Plugins/FX/Niagara/Source/NiagaraEditor/Public/NiagaraEditorUtilities.h b/Engine/Plugins/FX/Niagara/Source/NiagaraEditor/Public/NiagaraEditorUtilities.h
--- a/Engine/Plugins/FX/Niagara/Source/NiagaraEditor/Public/NiagaraEditorUtilities.h
+++ b/Engine/Plugins/FX/Niagara/Source/NiagaraEditor/Public/NiagaraEditorUtilities.h
@@ -64,6 +64,7 @@ struct NIAGARAEDITOR_API FNiagaraEditorSharedTexts
 {
 	static const FText DebugDrawUIActionBaseText;
 };
+
 namespace FNiagaraEditorUtilities
 {
 	/** Determines if the contents of two sets matches */
@@ -242,6 +243,8 @@ namespace FNiagaraEditorUtilities
 
 	NIAGARAEDITOR_API ENiagaraScriptLibraryVisibility GetScriptAssetVisibility(const FAssetData& ScriptAssetData);
 
+	NIAGARAEDITOR_API bool GetIsInheritableFromAssetRegistryTags(const FAssetData& AssetData, bool& bUseInheritance);
+	
 	/** Used instead of reading the template tag directly for backwards compatibility reasons when changing from a bool template specifier to an enum */
 	NIAGARAEDITOR_API bool GetTemplateSpecificationFromTag(const FAssetData& Data, ENiagaraScriptTemplateSpecification& OutTemplateSpecification);
 
@@ -424,7 +427,7 @@ namespace FNiagaraEditorUtilities
 	NIAGARAEDITOR_API bool IsEditorDataInterfaceInstance(const UNiagaraDataInterface* DataInterface);
 
 	NIAGARAEDITOR_API UNiagaraDataInterface* GetResolvedRuntimeInstanceForEditorDataInterfaceInstance(const UNiagaraSystem& OwningSystem, UNiagaraDataInterface& EditorDataInterfaceInstance);
-
+	
 	namespace Tooltips
 	{
 		NIAGARAEDITOR_API TSharedRef<SToolTip> CreateStackNoteTooltip(UNiagaraStackNote& StackNote);
@@ -438,6 +441,36 @@ namespace FNiagaraEditorUtilities
 			TMap<FNiagaraVariableBase, FGuid> FixupDuplicateScriptVariableGuids(UNiagaraScript* Script);
 		}
 	}
+
+	namespace AssetBrowser
+	{
+		enum EAssetTagSectionSource
+		{		
+			NiagaraInternal,
+			Project,
+			Other
+		};
+		
+		struct FStructuredAssetTagDefinitionLookupData
+		{
+			UNiagaraAssetTagDefinitions* DefinitionsAsset;
+			TArray<FNiagaraAssetTagDefinition> AssetTagDefinitions;
+			EAssetTagSectionSource Source;
+		};
+
+		
+		/** Returns a fully sorted array of tag definitions. Sorted by [Source.AssetSortOrder.(Optionally)TagNames].
+		 * If not sorted by tag names, the order defined in the asset is used. Sorting by tag names makes sense for flat lists. */
+		TArray<FStructuredAssetTagDefinitionLookupData> GetStructuredSortedAssetTagDefinitions(bool bSortTagsByName = false);
+		/** Same as above, but returns a flat list. If end result is not sorted, the order is [Source.AssetSorderOrder.OriginalTagOrder]. */
+		TArray<FNiagaraAssetTagDefinition> GetFlatSortedAssetTagDefinitions(bool bSortEndResultTagsByName = false);
+
+		const FNiagaraAssetTagDefinition& FindTagDefinitionForReference(const FNiagaraAssetTagDefinitionReference& Reference);
+
+		/* Expects a UNiagaraAssetTagDefinition asset */
+		EAssetTagSectionSource GetAssetTagDefinitionSource(const FAssetData& AssetData);
+		FText GetAssetTagSectionNameFromSource(EAssetTagSectionSource Source);
+	}
 }
 
 namespace FNiagaraParameterUtilities
diff --git a/Engine/Plugins/FX/Niagara/Source/NiagaraEditor/Public/NiagaraRecentAndFavoritesManager.cpp b/Engine/Plugins/FX/Niagara/Source/NiagaraEditor/Public/NiagaraRecentAndFavoritesManager.cpp
new file mode 100644
--- /dev/null
+++ b/Engine/Plugins/FX/Niagara/Source/NiagaraEditor/Public/NiagaraRecentAndFavoritesManager.cpp
@@ -0,0 +1,198 @@
+﻿// Copyright Epic Games, Inc. All Rights Reserved.
+
+#include "NiagaraRecentAndFavoritesManager.h"
+
+#include "ModuleDescriptor.h"
+#include "NiagaraEditorModule.h"
+#include "NiagaraEmitter.h"
+#include "NiagaraSystem.h"
+#include "AssetRegistry/AssetRegistryModule.h"
+#include "Interfaces/IPluginManager.h"
+#include "UObject/PackageResourceManager.h"
+
+void FNiagaraRecentAndFavoritesManager::Initialize()
+{
+	RecentAndFavoriteEmittersAndSystem = MakeUnique<FMainMRUFavoritesList>(TEXT("NiagaraRecent_EmittersAndSystems"), TEXT("NiagaraFavorites_EmittersAndSystems"),10);
+	RecentAndFavoriteModules = MakeUnique<FMainMRUFavoritesList>(TEXT("NiagaraRecent_Scripts"), TEXT("NiagaraFavorites_Scripts"),30);
+	RecentAndFavoriteParameters = MakeUnique<FMainMRUFavoritesList>(TEXT("NiagaraRecent_Parameters"), TEXT("NiagaraFavorites_Parameters"), 10);
+
+	RecentAndFavoriteEmittersAndSystem->ReadFromINI();
+	RecentAndFavoriteModules->ReadFromINI();
+	RecentAndFavoriteParameters->ReadFromINI();
+	
+	FAssetRegistryModule& AssetRegistryModule = FModuleManager::LoadModuleChecked<FAssetRegistryModule>("AssetRegistry");
+	AssetRegistryModule.Get().OnAssetRemoved().AddSP(this, &FNiagaraRecentAndFavoritesManager::OnAssetRemoved);
+	AssetRegistryModule.Get().OnAssetRenamed().AddSP(this, &FNiagaraRecentAndFavoritesManager::OnAssetRenamed);
+
+	IPluginManager& PluginManager = IPluginManager::Get();
+	PluginManager.OnLoadingPhaseComplete().AddRaw(this, &FNiagaraRecentAndFavoritesManager::OnPluginLoadingPhaseComplete);
+	
+	IPackageResourceManager::GetOnClearPackageResourceManagerDelegate().AddSP(this, &FNiagaraRecentAndFavoritesManager::OnPackageResourceManagerShutdown);
+}
+
+void FNiagaraRecentAndFavoritesManager::OnAssetAdded(const FAssetData& AssetData) const
+{
+	if(AssetData.GetClass() == UNiagaraEmitter::StaticClass() || AssetData.GetClass() == UNiagaraSystem::StaticClass())
+	{
+		RecentAndFavoriteEmittersAndSystem->AddMRUItem(AssetData.PackageName.ToString());
+	}
+	else if(AssetData.GetClass() == UNiagaraScript::StaticClass())
+	{
+		RecentAndFavoriteModules->AddMRUItem(AssetData.PackageName.ToString());
+	}
+}
+
+
+void FNiagaraRecentAndFavoritesManager::OnAssetRemoved(const FAssetData& AssetData) const
+{
+	if(FMainMRUFavoritesList* RecentAssetsList = GetRecentsList(AssetData))
+	{
+		FString Identifier = AssetData.PackageName.ToString();
+		if(RecentAssetsList->FindMRUItemIdx(Identifier) != INDEX_NONE)
+		{
+			RecentAssetsList->RemoveMRUItem(AssetData.PackageName.ToString());
+		}
+
+		if(RecentAssetsList->ContainsFavoritesItem(Identifier))
+		{
+			RecentAssetsList->RemoveFavoritesItem(AssetData.PackageName.ToString());
+		}
+	}
+}
+
+void FNiagaraRecentAndFavoritesManager::OnAssetOpenedInEditor(UObject* Object, IAssetEditorInstance* AssetEditorInstance) const
+{
+	FAssetData AssetData(Object);
+	OnAssetAdded(AssetData);
+}
+
+void FNiagaraRecentAndFavoritesManager::OnAssetRenamed(const FAssetData& AssetData, const FString& AssetOldName) const
+{
+	FString OldPathName = FPaths::GetBaseFilename(AssetOldName, false);
+	FString NewPathName = FPaths::GetBaseFilename(AssetData.GetObjectPathString(), false);
+
+	// We need this early exit because FindMRUItemIdx has a check() for non valid long package names
+	if (!FPackageName::IsValidLongPackageName(OldPathName) || !FPackageName::IsValidLongPackageName(NewPathName))
+	{
+		return;
+	}
+	
+	if(FMainMRUFavoritesList* RecentAssetsList = GetRecentsList(AssetData))
+	{
+		// If the asset did not previously exist in the recents list, we have nothing to do
+		if(RecentAssetsList->FindMRUItemIdx(OldPathName) == INDEX_NONE)
+		{
+			return;
+		}
+
+		// Otherwise remove the old name of the asset, and re-add it with the new name
+		// NOTE: This has an unintentional side effect of bringing it to the top of the MRU list that can't be avoided
+		RecentAssetsList->RemoveMRUItem(OldPathName);
+		RecentAssetsList->AddMRUItem(NewPathName);
+	}
+}
+
+void FNiagaraRecentAndFavoritesManager::OnPluginLoadingPhaseComplete(ELoadingPhase::Type LoadingPhase, bool bSuccess)
+{
+	if(LoadingPhase == ELoadingPhase::PostEngineInit)
+	{
+		GEditor->GetEditorSubsystem<UAssetEditorSubsystem>()->OnAssetOpenedInEditor().AddSP(this, &FNiagaraRecentAndFavoritesManager::OnAssetOpenedInEditor);
+		IPluginManager& PluginManager = IPluginManager::Get();
+		PluginManager.OnLoadingPhaseComplete().RemoveAll(this);
+
+		GEditor->OnEditorClose().AddSP(this, &FNiagaraRecentAndFavoritesManager::OnEditorClose);
+	}
+}
+
+void FNiagaraRecentAndFavoritesManager::Shutdown()
+{
+	Save();
+	
+	RecentAndFavoriteEmittersAndSystem.Reset();
+	RecentAndFavoriteModules.Reset();
+	RecentAndFavoriteParameters.Reset();
+}
+
+void FNiagaraRecentAndFavoritesManager::Save() const
+{
+	if(bCanCheckForPackages)
+	{
+		// Remove assets that have never been saved to disk
+		for(int CurRecentIndex = 0; CurRecentIndex < RecentAndFavoriteEmittersAndSystem->GetNumItems(); ++CurRecentIndex)
+		{
+			const FString& RecentAsset = RecentAndFavoriteEmittersAndSystem->GetMRUItem(CurRecentIndex);
+			// TODO If the engine is closing already the PackageManager might be nullptr
+			if(!FPackageName::DoesPackageExist(RecentAsset))
+			{
+				RecentAndFavoriteEmittersAndSystem->RemoveMRUItem(CurRecentIndex);
+				--CurRecentIndex;
+			}
+		}
+		
+		for(int CurRecentIndex = 0; CurRecentIndex < RecentAndFavoriteModules->GetNumItems(); ++CurRecentIndex)
+		{
+			const FString& RecentAsset = RecentAndFavoriteModules->GetMRUItem(CurRecentIndex);
+			// TODO If the engine is closing already the PackageManager might be nullptr
+			if(!FPackageName::DoesPackageExist(RecentAsset))
+			{
+				RecentAndFavoriteModules->RemoveMRUItem(CurRecentIndex);
+				--CurRecentIndex;
+			}
+		}
+	}
+
+	
+	RecentAndFavoriteEmittersAndSystem->WriteToINI();
+	RecentAndFavoriteModules->WriteToINI();
+	RecentAndFavoriteParameters->WriteToINI();
+}
+
+void FNiagaraRecentAndFavoritesManager::EmitterUsed(const UNiagaraEmitter& Emitter)
+{
+	const FAssetData EmitterAssetData(&Emitter);
+	OnAssetAdded(EmitterAssetData);
+}
+
+void FNiagaraRecentAndFavoritesManager::SystemUsed(const UNiagaraSystem& System)
+{
+	const FAssetData SystemAssetData(&System);
+	OnAssetAdded(SystemAssetData);
+}
+
+void FNiagaraRecentAndFavoritesManager::ModuleScriptUsed(const UNiagaraScript& Script)
+{
+	const FAssetData ScriptAssetData(&Script);
+	OnAssetAdded(ScriptAssetData);
+}
+
+void FNiagaraRecentAndFavoritesManager::ParameterUsed(const FNiagaraVariable& VariableBase)
+{
+	RecentAndFavoriteParameters->AddMRUItem(VariableBase.ToString());
+}
+
+FMainMRUFavoritesList* FNiagaraRecentAndFavoritesManager::GetRecentsList(const FAssetData& AssetData) const
+{
+	FMainMRUFavoritesList* RecentAssetsList = nullptr;
+	if(AssetData.GetClass() == UNiagaraEmitter::StaticClass() || AssetData.GetClass() == UNiagaraSystem::StaticClass())
+	{
+		RecentAssetsList = RecentAndFavoriteEmittersAndSystem.Get();
+	}
+	else if(AssetData.GetClass() == UNiagaraScript::StaticClass())
+	{
+		RecentAssetsList = RecentAndFavoriteModules.Get();
+	}
+
+	return RecentAssetsList;
+}
+
+void FNiagaraRecentAndFavoritesManager::OnEditorClose()
+{
+	Save();
+	GEditor->OnEditorClose().RemoveAll(this);
+}
+
+void FNiagaraRecentAndFavoritesManager::OnPackageResourceManagerShutdown()
+{
+	bCanCheckForPackages = false;
+	IPackageResourceManager::GetOnClearPackageResourceManagerDelegate().RemoveAll(this);
+}
diff --git a/Engine/Plugins/FX/Niagara/Source/NiagaraEditor/Public/NiagaraRecentAndFavoritesManager.h b/Engine/Plugins/FX/Niagara/Source/NiagaraEditor/Public/NiagaraRecentAndFavoritesManager.h
new file mode 100644
--- /dev/null
+++ b/Engine/Plugins/FX/Niagara/Source/NiagaraEditor/Public/NiagaraRecentAndFavoritesManager.h
@@ -0,0 +1,51 @@
+﻿// Copyright Epic Games, Inc. All Rights Reserved.
+
+#pragma once
+
+#include "ModuleDescriptor.h"
+#include "NiagaraSystem.h"
+#include "MRUFavoritesList.h"
+#include "Subsystems/AssetEditorSubsystem.h"
+
+class UNiagaraScript;
+class UNiagaraSystem;
+class UNiagaraEmitter;
+struct FNiagaraVariable;
+
+class FNiagaraRecentAndFavoritesManager : public TSharedFromThis<FNiagaraRecentAndFavoritesManager>
+{
+public:
+	void Initialize();
+	void Shutdown();
+	
+	void Save() const;
+
+	void EmitterUsed(const UNiagaraEmitter& Emitter);
+	void SystemUsed(const UNiagaraSystem& System);
+	void ModuleScriptUsed(const UNiagaraScript& Script);
+	void ParameterUsed(const FNiagaraVariable& VariableBase);
+
+	const FMainMRUFavoritesList* GetRecentEmitterAndSystemsList() const { return RecentAndFavoriteEmittersAndSystem.Get(); }
+	const FMainMRUFavoritesList* GetRecentScriptsList() const { return RecentAndFavoriteModules.Get(); }
+	const FMainMRUFavoritesList* GetRecentParametersList() const { return RecentAndFavoriteParameters.Get(); }
+
+private:
+	FMainMRUFavoritesList* GetRecentsList(const FAssetData& AssetData) const;
+	void OnEditorClose();
+	void OnPackageResourceManagerShutdown();
+
+	void OnAssetAdded(const FAssetData& AssetData) const;
+	void OnAssetRemoved(const FAssetData& AssetData) const;
+	void OnAssetOpenedInEditor(UObject* Object, IAssetEditorInstance* AssetEditorInstance) const;
+	void OnAssetRenamed(const FAssetData& AssetData, const FString& String) const;
+
+	void OnPluginLoadingPhaseComplete(ELoadingPhase::Type LoadingPhase, bool bSuccess);
+
+private:
+	TUniquePtr<FMainMRUFavoritesList> RecentAndFavoriteEmittersAndSystem;
+	TUniquePtr<FMainMRUFavoritesList> RecentAndFavoriteModules;
+	TUniquePtr<FMainMRUFavoritesList> RecentAndFavoriteParameters;
+	
+	bool bCanCheckForPackages = true;
+};
+
diff --git a/Engine/Plugins/FX/Niagara/Source/NiagaraEditor/Public/Widgets/AssetBrowser/NiagaraAssetBrowserConfig.cpp b/Engine/Plugins/FX/Niagara/Source/NiagaraEditor/Public/Widgets/AssetBrowser/NiagaraAssetBrowserConfig.cpp
new file mode 100644
--- /dev/null
+++ b/Engine/Plugins/FX/Niagara/Source/NiagaraEditor/Public/Widgets/AssetBrowser/NiagaraAssetBrowserConfig.cpp
@@ -0,0 +1,14 @@
+﻿// Copyright Epic Games, Inc. All Rights Reserved.
+
+#include "NiagaraAssetBrowserConfig.h"
+
+TObjectPtr<UNiagaraAssetBrowserConfig> UNiagaraAssetBrowserConfig::Instance = nullptr;
+
+void UNiagaraAssetBrowserConfig::Initialize()
+{
+	if(!Instance)
+	{
+		Instance = NewObject<UNiagaraAssetBrowserConfig>(); 
+		Instance->AddToRoot();
+	}
+}
diff --git a/Engine/Plugins/FX/Niagara/Source/NiagaraEditor/Public/Widgets/AssetBrowser/NiagaraAssetBrowserConfig.h b/Engine/Plugins/FX/Niagara/Source/NiagaraEditor/Public/Widgets/AssetBrowser/NiagaraAssetBrowserConfig.h
new file mode 100644
--- /dev/null
+++ b/Engine/Plugins/FX/Niagara/Source/NiagaraEditor/Public/Widgets/AssetBrowser/NiagaraAssetBrowserConfig.h
@@ -0,0 +1,35 @@
+﻿// Copyright Epic Games, Inc. All Rights Reserved.
+
+#pragma once
+
+#include "CoreMinimal.h"
+#include "EditorConfigBase.h"
+#include "NiagaraAssetBrowserConfig.generated.h"
+
+USTRUCT()
+struct FNiagaraAssetBrowserConfiguration
+{
+	GENERATED_BODY()
+	
+	UPROPERTY()
+	TArray<FName> MainFilterSelection;
+
+	UPROPERTY()
+	bool bShouldDisplayViewport;
+};
+
+UCLASS(EditorConfig="NiagaraAssetBrowser")
+class NIAGARAEDITOR_API UNiagaraAssetBrowserConfig : public UEditorConfigBase
+{
+	GENERATED_BODY()
+	
+public:		
+	static void Initialize();
+	static UNiagaraAssetBrowserConfig* Get() { return Instance; }
+	
+	UPROPERTY(meta=(EditorConfig))
+	TMap<FName, FNiagaraAssetBrowserConfiguration> MainFilterSelection;
+	
+private:
+	static TObjectPtr<UNiagaraAssetBrowserConfig> Instance;
+};
diff --git a/Engine/Plugins/FX/Niagara/Source/NiagaraEditor/Public/Widgets/AssetBrowser/NiagaraAssetBrowserPreview.h b/Engine/Plugins/FX/Niagara/Source/NiagaraEditor/Public/Widgets/AssetBrowser/NiagaraAssetBrowserPreview.h
new file mode 100644
--- /dev/null
+++ b/Engine/Plugins/FX/Niagara/Source/NiagaraEditor/Public/Widgets/AssetBrowser/NiagaraAssetBrowserPreview.h
@@ -0,0 +1,51 @@
+﻿// Copyright Epic Games, Inc. All Rights Reserved.
+
+#pragma once
+
+#include "CoreMinimal.h"
+#include "NiagaraComponent.h"
+#include "SEditorViewport.h"
+#include "Widgets/SCompoundWidget.h"
+
+/**
+ * 
+ */
+class NIAGARAEDITOR_API SNiagaraAssetBrowserPreview : public SEditorViewport, public FGCObject
+{
+public:
+	SLATE_BEGIN_ARGS(SNiagaraAssetBrowserPreview)
+		{
+		}
+	SLATE_END_ARGS()
+
+	/** Constructs this widget with InArgs */
+	void Construct(const FArguments& InArgs);
+	virtual ~SNiagaraAssetBrowserPreview() override;
+
+	void SetEmitter(UNiagaraEmitter& Emitter);
+	void SetSystem(UNiagaraSystem& System);
+	void ResetAsset();
+	
+	UNiagaraComponent* GetPreviewComponent() { return PreviewComponent; }
+private:
+	virtual void Tick(const FGeometry& AllottedGeometry, const double InCurrentTime, const float InDeltaTime) override;
+	
+	virtual void AddReferencedObjects( FReferenceCollector& Collector ) override;
+	virtual FString GetReferencerName() const override
+	{
+		return TEXT("SNiagaraAssetBrowserPreview");
+	}
+
+private:
+	TObjectPtr<UObject> AssetToDisplay;
+	TObjectPtr<UNiagaraComponent> PreviewComponent;
+	TObjectPtr<UNiagaraSystem> SystemForEmitterDisplay;
+
+	/** Preview Scene - uses advanced preview settings */
+	TSharedPtr<class FAdvancedPreviewScene> AdvancedPreviewScene;
+	TSharedPtr<class FNiagaraAssetPreviewViewportClient> AssetPreviewViewportClient;
+
+protected:
+	/** SEditorViewport interface */
+	virtual TSharedRef<FEditorViewportClient> MakeEditorViewportClient() override;
+};
diff --git a/Engine/Plugins/FX/Niagara/Source/NiagaraEditor/Public/Widgets/AssetBrowser/NiagaraMenuFilters.h b/Engine/Plugins/FX/Niagara/Source/NiagaraEditor/Public/Widgets/AssetBrowser/NiagaraMenuFilters.h
new file mode 100644
--- /dev/null
+++ b/Engine/Plugins/FX/Niagara/Source/NiagaraEditor/Public/Widgets/AssetBrowser/NiagaraMenuFilters.h
@@ -0,0 +1,96 @@
+// Copyright Epic Games, Inc. All Rights Reserved.
+
+#pragma once
+
+#include "CoreMinimal.h"
+#include "FrontendFilterBase.h"
+#include "NiagaraAssetTagDefinitions.h"
+#include "NiagaraMenuFilters.generated.h"
+
+UCLASS()
+class UNiagaraTagsContentBrowserFilterData : public UObject
+{
+	GENERATED_BODY()
+
+public:
+	void AddTagGuid(const FGuid& InGuid);
+	void RemoveTagGuid(const FGuid& InGuid);
+	bool ContainsActiveTagGuid(const FGuid& InGuid);
+	int32 GetNumActiveTagGuids() const { return ActiveTagGuids.Num(); }
+	const TArray<FGuid>& GetActiveTagGuids() const { return ActiveTagGuids; }
+	void ResetActiveTagGuids() { ActiveTagGuids.Reset(); }
+
+	FSimpleMulticastDelegate& OnActiveTagGuidsChanged() { return OnActiveTagGuidsChangedDelegate; }
+private:
+	UPROPERTY()
+	TArray<FGuid> ActiveTagGuids;
+
+	FSimpleMulticastDelegate OnActiveTagGuidsChangedDelegate;
+};
+
+UCLASS()
+class UNiagaraTagsContentBrowserFilterContext : public UObject
+{
+	GENERATED_BODY()
+
+public:
+	UPROPERTY()
+	TObjectPtr<UNiagaraTagsContentBrowserFilterData> FilterData;
+};
+
+
+struct FNiagaraAssetBrowserMainFilter
+{
+	DECLARE_DELEGATE_RetVal_OneParam(bool, FIsAssetRecent, const FAssetData&)
+	
+	enum class EFilterMode
+	{
+		All,
+		NiagaraAssetTag,
+		Recent,
+		Custom
+	};
+
+	FNiagaraAssetBrowserMainFilter(EFilterMode InFilterMode) : FilterMode(InFilterMode) {}
+	
+	EFilterMode FilterMode;
+	TArray<TSharedRef<FNiagaraAssetBrowserMainFilter>> ChildFilters;
+
+	FNiagaraAssetTagDefinition AssetTagDefinition;
+
+	/** The custom guid should always be set the same for a specific main filter.
+	 * Should be used if there are multiple items with the same filter mode that aren't niagara tags (i.e. set to Custom) */
+	FGuid CustomGuid;
+
+	FIsAssetRecent IsAssetRecentDelegate;
+	
+	FText CustomDisplayName;
+	FOnShouldFilterAsset CustomShouldFilterAsset;
+
+	FName GetIdentifier() const;
+	FText GetDisplayName() const;
+	
+	bool IsAssetRecent(const FAssetData& AssetCandidate) const;
+	bool ShouldCustomFilterAsset(const FAssetData& AssetCandidate) const;
+	bool DoesAssetHaveTag(const FAssetData& AssetCandidate) const;
+
+	bool operator==(const FNiagaraAssetBrowserMainFilter& Other) const
+	{
+		if(FilterMode != Other.FilterMode)
+		{
+			return false;
+		}
+
+		if(FilterMode == EFilterMode::NiagaraAssetTag)
+		{
+			return AssetTagDefinition == Other.AssetTagDefinition;
+		}
+
+		if(CustomGuid.IsValid())
+		{
+			return CustomGuid == Other.CustomGuid;
+		}
+		
+		return true;
+	}
+};
diff --git a/Engine/Plugins/FX/Niagara/Source/NiagaraEditor/Public/Widgets/AssetBrowser/NiagaraSemanticTagsFrontEndFilterExtension.h b/Engine/Plugins/FX/Niagara/Source/NiagaraEditor/Public/Widgets/AssetBrowser/NiagaraSemanticTagsFrontEndFilterExtension.h
new file mode 100644
--- /dev/null
+++ b/Engine/Plugins/FX/Niagara/Source/NiagaraEditor/Public/Widgets/AssetBrowser/NiagaraSemanticTagsFrontEndFilterExtension.h
@@ -0,0 +1,121 @@
+﻿// Copyright Epic Games, Inc. All Rights Reserved.
+
+#pragma once
+
+#include "CoreMinimal.h"
+#include "ContentBrowserFrontEndFilterExtension.h"
+#include "NiagaraAssetTagDefinitions.h"
+#include "NiagaraMenuFilters.h"
+#include "NiagaraSemanticTagsFrontEndFilterExtension.generated.h"
+
+#define LOCTEXT_NAMESPACE "NiagaraFrontendFilters"
+
+/** A frontend filter used for Niagara specific windows. One filter per tag. */
+class FFrontendFilter_NiagaraTag : public FFrontendFilter
+{
+public:
+	FFrontendFilter_NiagaraTag(FNiagaraAssetTagDefinition InAssetTagDefinition, TSharedPtr<FFrontendFilterCategory> Category) : FFrontendFilter(Category), AssetTagDefinition(InAssetTagDefinition)
+	{}
+
+	virtual FLinearColor GetColor() const override { return AssetTagDefinition.Color; }
+	virtual FString GetName() const override { return AssetTagDefinition.AssetTag.ToString(); }
+	virtual FText GetDisplayName() const override { return FText::FromString(AssetTagDefinition.AssetTag.ToString()); }
+	virtual FText GetToolTipText() const override { return AssetTagDefinition.Description; }
+	
+	virtual bool PassesFilter(const FContentBrowserItem& InItem) const override;
+private:
+	FNiagaraAssetTagDefinition AssetTagDefinition;
+};
+
+/** A special case filter for the content browser.
+ *  It has a right click context menu that shows available asset tags and serializes them using only their guids */
+class FFrontendFilter_NiagaraTag_ContentBrowser : public FFrontendFilter, public FGCObject
+{
+public:
+	FFrontendFilter_NiagaraTag_ContentBrowser(TSharedPtr<FFrontendFilterCategory> Category) : FFrontendFilter(Category)
+	{
+		FilterData = NewObject<UNiagaraTagsContentBrowserFilterData>();
+		FilterData->OnActiveTagGuidsChanged().AddLambda([this]()
+		{
+			BroadcastChangedEvent();
+		});
+	}
+	
+	virtual FString GetName() const override { return "Niagara Asset Tags"; }
+	virtual FText GetDisplayName() const override;
+	virtual FText GetToolTipText() const override { return LOCTEXT("NiagaraAssetTagFilterTooltip", "Right-click to specify the Niagara Tags to filter for.\nA combination of tags will filter for assets owning all tags."); }
+	virtual void ModifyContextMenu(FMenuBuilder& MenuBuilder) override;
+	
+	virtual bool PassesFilter(const FContentBrowserItem& InItem) const override;
+
+	virtual void SaveSettings(const FString& IniFilename, const FString& IniSection, const FString& SettingsString) const override;
+	virtual void LoadSettings(const FString& IniFilename, const FString& IniSection, const FString& SettingsString) override;
+
+	virtual FString GetReferencerName() const override;
+	virtual void AddReferencedObjects(FReferenceCollector& Collector) override;
+private:
+	TObjectPtr<UNiagaraTagsContentBrowserFilterData> FilterData;
+};
+
+class FFrontendFilter_NiagaraEmitterInheritance : public FFrontendFilter
+{
+public:
+	FFrontendFilter_NiagaraEmitterInheritance(bool bInRequiredInheritanceState, TSharedPtr<FFrontendFilterCategory> Category) : FFrontendFilter(Category), bRequiredInheritanceState(bInRequiredInheritanceState)
+	{}
+
+	virtual FLinearColor GetColor() const override { return FLinearColor::Red; }
+	virtual FString GetName() const override { return FString::Printf(TEXT("Inheritance: %s"), bRequiredInheritanceState ? TEXT("Yes") : TEXT("No")); }
+	virtual FText GetDisplayName() const override { return FText::FormatOrdered(LOCTEXT("NiagarEmitterInheritanceFilterDisplayName", "Inheritance: {0}"), bRequiredInheritanceState ?
+		LOCTEXT("NiagarEmitterInheritanceFilterDisplayName_Yes", "Yes") : LOCTEXT("NiagarEmitterInheritanceFilterDisplayName_No", "No")); }
+	virtual FText GetToolTipText() const override
+	{
+		if(bRequiredInheritanceState)
+		{
+			return LOCTEXT("NiagaraEmitterInheritanceFilterTooltip_Yes", "Only display emitters with Inheritance");
+		}
+
+		return LOCTEXT("NiagaraEmitterInheritanceFilterTooltip_No", "Only display emitters without Inheritance");
+	}
+	
+	virtual bool PassesFilter(const FContentBrowserItem& InItem) const override;
+private:
+	bool bRequiredInheritanceState;
+};
+
+class FFrontendFilter_NiagaraSystemEffectType : public FFrontendFilter
+{
+public:
+	FFrontendFilter_NiagaraSystemEffectType(bool bInRequiresEffectType, TSharedPtr<FFrontendFilterCategory> Category) : FFrontendFilter(Category), bRequiresEffectType(bInRequiresEffectType)
+	{}
+
+	virtual FLinearColor GetColor() const override { return FLinearColor::Red; }
+	virtual FString GetName() const override { return FString::Printf(TEXT("Effect Type: %s"), bRequiresEffectType ? TEXT("Yes") : TEXT("No")); }
+	virtual FText GetDisplayName() const override { return FText::FormatOrdered(LOCTEXT("NiagaraSystemEffectTypeFilterDisplayName", "Effect Type: {0}"), bRequiresEffectType ?
+		LOCTEXT("NiagaraSystemEffectTypeFilterDisplayName_Yes", "Yes") : LOCTEXT("NiagaraSystemEffectTypeFilterDisplayName_No", "No")); }
+	virtual FText GetToolTipText() const override
+	{
+		if(bRequiresEffectType)
+		{
+			return LOCTEXT("NiagaraSystemEffectTypeFilterTooltip_Yes", "Only display Niagara Systems with an effect type");
+		}
+
+		return LOCTEXT("NiagaraSystemEffectTypeFilterTooltip_No", "Only display Niagara Systems without Inheritance");
+	}
+	
+	virtual bool PassesFilter(const FContentBrowserItem& InItem) const override;
+private:
+	bool bRequiresEffectType;
+};
+
+/**
+ * 
+ */
+UCLASS()
+class NIAGARAEDITOR_API UNiagaraSemanticTagsFrontEndFilterExtension : public UContentBrowserFrontEndFilterExtension
+{
+	GENERATED_BODY()
+
+	virtual void AddFrontEndFilterExtensions(TSharedPtr<class FFrontendFilterCategory> DefaultCategory, TArray< TSharedRef<class FFrontendFilter> >& InOutFilterList) const override;
+};
+
+#undef LOCTEXT_NAMESPACE
diff --git a/Engine/Plugins/FX/Niagara/Source/NiagaraEditor/Public/Widgets/AssetBrowser/SNiagaraAssetBrowser.h b/Engine/Plugins/FX/Niagara/Source/NiagaraEditor/Public/Widgets/AssetBrowser/SNiagaraAssetBrowser.h
new file mode 100644
--- /dev/null
+++ b/Engine/Plugins/FX/Niagara/Source/NiagaraEditor/Public/Widgets/AssetBrowser/SNiagaraAssetBrowser.h
@@ -0,0 +1,198 @@
+﻿// Copyright Epic Games, Inc. All Rights Reserved.
+
+#pragma once
+
+#include "CoreMinimal.h"
+#include "NiagaraAssetBrowserPreview.h"
+#include "SNiagaraAssetBrowserContent.h"
+#include "NiagaraMenuFilters.h"
+#include "Widgets/SCompoundWidget.h"
+#include "Widgets/SItemSelector.h"
+#include "NiagaraAssetTagDefinitions.h"
+#include "SAssetEditorViewport.h"
+#include "Subsystems/AssetEditorSubsystem.h"
+#include "MRUFavoritesList.h"
+#include "Editor/EditorEngine.h"
+
+class FNiagaraSystemViewModel;
+
+class NIAGARAEDITOR_API SNiagaraAssetBrowser : public SCompoundWidget
+{
+public:
+	SLATE_BEGIN_ARGS(SNiagaraAssetBrowser)
+		: _RecentAndFavoritesList(GEditor->GetEditorSubsystem<UAssetEditorSubsystem>()->GetRecentlyOpenedAssets())
+		, _SaveSettingsName("NiagaraAssetBrowser")
+		, _AssetSelectionMode(ESelectionMode::Single)
+		{
+		}
+		SLATE_ARGUMENT(TArray<UClass*>, AvailableClasses)
+		SLATE_ATTRIBUTE(const FMainMRUFavoritesList*, RecentAndFavoritesList)
+		SLATE_ARGUMENT(TOptional<FName>, SaveSettingsName)
+		SLATE_ARGUMENT(TOptional<FText>, EmptySelectionMessage)
+		SLATE_ARGUMENT(ESelectionMode::Type, AssetSelectionMode)
+		SLATE_EVENT(FOnAssetSelected, OnAssetSelected)
+		SLATE_EVENT(FOnAssetsActivated, OnAssetsActivated)
+		SLATE_ARGUMENT(TSharedPtr<SWidget>, AdditionalWidget)
+	SLATE_END_ARGS()
+
+	void Construct(const FArguments& InArgs);
+	virtual ~SNiagaraAssetBrowser() override;
+
+	void RefreshBackendFilter() const;
+
+	TArray<FAssetData> GetSelectedAssets() const;
+	
+	TArray<UClass*> GetDisplayedAssetTypes() const;
+
+	void InitContextMenu();
+private:
+	/** This function should take into account all the different widgets states that could affect the FARFilter. */
+	FARFilter GetCurrentBackendFilter() const;
+	
+	bool ShouldFilterAsset(const FAssetData& AssetData) const;
+	
+	void PopulateFiltersSlot();
+	void PopulateAssetBrowserContentSlot();
+	void PopulateAssetBrowserDetailsSlot();
+	
+	TArray<TSharedRef<FNiagaraAssetBrowserMainFilter>> GetMainFilters() const;
+
+private:
+	/** This has to be called whenever any kind of filtering changes and refreshes the actual asset view */
+	void OnFilterChanged() const;
+	
+	void OnGetChildFiltersForFilter(TSharedRef<FNiagaraAssetBrowserMainFilter> NiagaraAssetBrowserMainFilter, TArray<TSharedRef<FNiagaraAssetBrowserMainFilter>>& OutChildren) const;
+	bool OnCompareMainFiltersForEquality(const FNiagaraAssetBrowserMainFilter& NiagaraAssetBrowserMainFilter, const FNiagaraAssetBrowserMainFilter& NiagaraAssetBrowserMainFilter1) const;
+	
+	TSharedRef<ITableRow> GenerateWidgetRowForMainFilter(TSharedRef<FNiagaraAssetBrowserMainFilter> MainFilter, const TSharedRef<STableViewBase>& OwningTable) const;
+
+
+	void OnAssetSelected(const FAssetData& AssetData);
+	void OnAssetsActivated(const TArray<FAssetData>& AssetData, EAssetTypeActivationMethod::Type) const;
+	TSharedPtr<SWidget> OnGetAssetContextMenu(const TArray<FAssetData>& AssetData) const;
+	TSharedRef<SToolTip> OnGetCustomAssetTooltip(FAssetData& AssetData);
+	void OnMainFilterSelected(TSharedPtr<FNiagaraAssetBrowserMainFilter> MainFilter, ESelectInfo::Type Arg);
+	
+	void OnAssetTagActivated(const FNiagaraAssetTagDefinition& NiagaraAssetTagDefinition);
+
+	TArray<TSharedRef<FFrontendFilter>> OnGetExtraFrontendFilters() const;
+	void OnExtendAddFilterMenu(UToolMenu* ToolMenu) const;
+
+	bool OnIsAssetRecent(const FAssetData& AssetCandidate) const;
+
+	bool GetShouldDisplayViewport() const { return bShouldDisplayViewport; }
+	
+	EVisibility OnGetViewportVisibility() const;
+	EVisibility OnGetThumbnailVisibility() const;
+	EVisibility OnGetShouldDisplayVisibilityCheckbox() const;
+
+	ECheckBoxState OnShouldDisplayViewport() const;
+	void OnShouldDisplayViewportChanged(ECheckBoxState CheckBoxState);
+	FText OnGetShouldDisplayViewportTooltip() const;
+
+	void SaveSettings() const;
+	void LoadSettings();
+private:
+	TArray<UClass*> AvailableClasses;
+	TAttribute<const FMainMRUFavoritesList*> RecentAndFavoritesList;
+	TOptional<FName> SaveSettingsName;
+	TOptional<FText> EmptySelectionMessage;
+	FOnAssetSelected OnAssetSelectedDelegate;
+	FOnAssetsActivated OnAssetsActivatedDelegate;
+	ESelectionMode::Type AssetSelectionMode = ESelectionMode::Single;
+	bool bShouldDisplayViewport = false;
+
+	mutable TMap<FNiagaraAssetTagDefinition, TSharedRef<FFrontendFilter>> DropdownFilterCache;
+	TArray<TSharedRef<FNiagaraAssetBrowserMainFilter>> AssetBrowserMainFilters;
+private:
+	TSharedPtr<STreeView<TSharedRef<FNiagaraAssetBrowserMainFilter>>> MainFilterSelector;
+	TSharedPtr<SNiagaraAssetBrowserContent> AssetBrowserContent;
+	TSharedPtr<SWidget> DetailsContainer;
+	TSharedPtr<SWidgetSwitcher> DetailsSwitcher;
+	TSharedPtr<SNiagaraAssetBrowserPreview> PreviewViewport;
+
+	SSplitter::FSlot* FiltersSlot = nullptr;
+	SSplitter::FSlot* AssetBrowserContentSlot = nullptr;
+	SSplitter::FSlot* AssetBrowserDetailsAreaSlot = nullptr;
+	SWidgetSwitcher::FSlot* AssetBrowserDetailsSlot = nullptr;
+	SVerticalBox::FSlot* AdditionalWidgetSlot = nullptr;
+
+	/** Used to temporarily stop saving & loading config state.
+	 * This is used for initialization where setting the default state should not be saved into config.*/
+	bool bSuppressSaveAndLoad = false;
+	
+	/** We use this fallback as a means to save our last selected filter, in case we have no valid active selections anymore. */
+	mutable FName LastSelectedMainFilterIdentifierFallback = NAME_None;
+};
+
+class NIAGARAEDITOR_API SNiagaraAssetBrowserWindow : public SWindow
+{
+public:
+	SLATE_BEGIN_ARGS(SNiagaraAssetBrowserWindow)
+		{
+		}
+		SLATE_ARGUMENT(SNiagaraAssetBrowser::FArguments, AssetBrowserArgs)
+		SLATE_ARGUMENT(TOptional<FText>, WindowTitle)
+	SLATE_END_ARGS()
+
+	void Construct(const FArguments& InArgs);
+
+	bool HasSelectedAssets() const;
+	TArray<FAssetData> GetSelectedAssets() const;
+	
+protected:	
+	void OnAssetsActivated(const TArray<FAssetData>& AssetData, EAssetTypeActivationMethod::Type ActivationType);
+
+	virtual void OnFocusLost(const FFocusEvent& InFocusEvent) override;
+	virtual FReply OnKeyDown(const FGeometry& MyGeometry, const FKeyEvent& InKeyEvent) override;
+protected:
+	TSharedPtr<SNiagaraAssetBrowser> AssetBrowser;
+private:
+	FOnAssetsActivated OnAssetsActivatedDelegate;
+};
+
+class NIAGARAEDITOR_API SNiagaraCreateAssetWindow : public SNiagaraAssetBrowserWindow
+{
+	SLATE_BEGIN_ARGS(SNiagaraCreateAssetWindow)
+		{
+		}
+		SLATE_ARGUMENT(SNiagaraAssetBrowserWindow::FArguments, AssetBrowserWindowArgs)
+	SLATE_END_ARGS()
+
+	void Construct(const FArguments& InArgs, UClass& CreatedClass);
+
+	bool ShouldProceedWithAction() const { return bProceedWithAction; }
+private:
+	/** The function that will be called by our buttons or by the asset picker itself if double-clicking, hitting enter etc. */
+	void OnAssetsActivatedInternal(const TArray<FAssetData>& AssetData, EAssetTypeActivationMethod::Type);
+	
+	FReply Proceed();
+	FReply Cancel();
+
+	FText GetCreateButtonTooltip() const;
+private:
+	TWeakObjectPtr<UClass> CreatedClass;
+	bool bProceedWithAction = false;
+};
+
+class NIAGARAEDITOR_API SNiagaraAddEmitterToSystemWindow : public SNiagaraAssetBrowserWindow
+{
+	SLATE_BEGIN_ARGS(SNiagaraAddEmitterToSystemWindow)
+	{
+	}
+		SLATE_ARGUMENT(SNiagaraAssetBrowserWindow::FArguments, AssetBrowserWindowArgs)
+	SLATE_END_ARGS()
+
+	void Construct(const FArguments& InArgs, TSharedRef<FNiagaraSystemViewModel> SystemViewModel);
+private:
+	/** The function that will be called by our buttons or by the asset picker itself if double-clicking, hitting enter etc. */
+	void OnAssetsActivatedInternal(const TArray<FAssetData>& AssetData, EAssetTypeActivationMethod::Type) const;
+
+	FReply AddEmptyEmitter();
+	FReply AddSelectedEmitters();
+	FReply Cancel();
+	
+	FText GetAddButtonTooltip() const;
+private:
+	TWeakPtr<FNiagaraSystemViewModel> WeakSystemViewModel;
+};
\ No newline at end of file
diff --git a/Engine/Plugins/FX/Niagara/Source/NiagaraEditor/Public/Widgets/AssetBrowser/SNiagaraAssetBrowserContent.h b/Engine/Plugins/FX/Niagara/Source/NiagaraEditor/Public/Widgets/AssetBrowser/SNiagaraAssetBrowserContent.h
new file mode 100644
--- /dev/null
+++ b/Engine/Plugins/FX/Niagara/Source/NiagaraEditor/Public/Widgets/AssetBrowser/SNiagaraAssetBrowserContent.h
@@ -0,0 +1,37 @@
+﻿// Copyright Epic Games, Inc. All Rights Reserved.
+
+#pragma once
+
+#include "CoreMinimal.h"
+#include "FrontendFilterBase.h"
+#include "..\..\..\..\Niagara\Public\NiagaraAssetTagDefinitions.h"
+#include "Filters/GenericFilter.h"
+#include "Filters/SAssetFilterBar.h"
+#include "Widgets/SCompoundWidget.h"
+
+class NIAGARAEDITOR_API SNiagaraAssetBrowserContent : public SCompoundWidget
+{
+public:
+	DECLARE_DELEGATE_RetVal_OneParam(bool, FOnShouldFilterAsset, const FAssetData&)
+	DECLARE_DELEGATE_RetVal(TArray<TSharedRef<FFrontendFilter>>, FOnGetExtraFrontendFilters)
+	
+	SLATE_BEGIN_ARGS(SNiagaraAssetBrowserContent)
+		{
+		}
+		SLATE_ARGUMENT(FAssetPickerConfig, InitialConfig)
+	SLATE_END_ARGS()
+
+	/** Constructs this widget with InArgs */
+	void Construct(const FArguments& InArgs);
+
+	void SetARFilter(FARFilter InFilter);
+
+	TArray<FAssetData> GetCurrentSelection() const;
+private:	
+	FRefreshAssetViewDelegate RefreshAssetViewDelegate;
+	FSyncToAssetsDelegate SyncToAssetsDelegate;
+	FSetARFilterDelegate SetNewFilterDelegate;
+	FGetCurrentSelectionDelegate GetCurrentSelectionDelegate;
+
+	FAssetPickerConfig InitialAssetPickerConfig;
+};
diff --git a/Engine/Plugins/FX/Niagara/Source/NiagaraEditor/Public/Widgets/AssetBrowser/SNiagaraSelectedAssetDetails.h b/Engine/Plugins/FX/Niagara/Source/NiagaraEditor/Public/Widgets/AssetBrowser/SNiagaraSelectedAssetDetails.h
new file mode 100644
--- /dev/null
+++ b/Engine/Plugins/FX/Niagara/Source/NiagaraEditor/Public/Widgets/AssetBrowser/SNiagaraSelectedAssetDetails.h
@@ -0,0 +1,111 @@
+﻿// Copyright Epic Games, Inc. All Rights Reserved.
+
+#pragma once
+
+#include "CoreMinimal.h"
+#include "NiagaraAssetTagDefinitions.h"
+#include "Widgets/SCompoundWidget.h"
+#include "AssetThumbnail.h"
+#include "AssetRegistry/AssetData.h"
+
+struct FDisplayedPropertyData
+{
+	DECLARE_DELEGATE_RetVal_OneParam(bool, FShouldDisplayProperty, const FAssetData&)
+	DECLARE_DELEGATE_RetVal_OneParam(TSharedRef<SWidget>, FGenerateWidget, const FAssetData&)
+
+	FShouldDisplayProperty ShouldDisplayPropertyDelegate;
+	FGenerateWidget NameWidgetDelegate;
+	FGenerateWidget ValueWidgetDelegate;
+};
+
+struct FNiagaraAssetDetailClassInfo
+{
+	DECLARE_DELEGATE_RetVal_OneParam(FText, FGetDescription, const FAssetData&);
+	
+	FGetDescription GetDescriptionDelegate;
+	TArray<FDisplayedPropertyData> DisplayedProperties;
+};
+
+struct FNiagaraAssetDetailDatabase
+{
+	static void Init();
+	static TMap<UClass*, FNiagaraAssetDetailClassInfo> NiagaraAssetDetailDatabase;
+};
+
+DECLARE_DELEGATE_OneParam(FOnAssetTagActivated, const FNiagaraAssetTagDefinition& AssetTagDefinition);
+
+class NIAGARAEDITOR_API SNiagaraAssetTag : public SCompoundWidget
+{
+	SLATE_BEGIN_ARGS(SNiagaraAssetTag)
+	{
+	}
+		SLATE_EVENT(FOnAssetTagActivated, OnAssetTagActivated)
+		SLATE_ARGUMENT(TOptional<FText>, OnAssetTagActivatedTooltip)
+	SLATE_END_ARGS()
+
+	void Construct(const FArguments& InArgs, const FNiagaraAssetTagDefinition& AssetTagDefinition);
+
+private:
+	FReply OnClicked() const;
+
+private:
+	FNiagaraAssetTagDefinition AssetTagDefinition;
+	FOnAssetTagActivated OnAssetTagActivated;
+	TOptional<FText> OnAssetTagActivatedTooltip;
+};
+
+class NIAGARAEDITOR_API SNiagaraAssetTagRow : public SCompoundWidget
+{
+	SLATE_BEGIN_ARGS(SNiagaraAssetTagRow)
+		{
+		}
+		SLATE_EVENT(FOnAssetTagActivated, OnAssetTagActivated)
+		SLATE_ARGUMENT(TOptional<FText>, OnAssetTagActivatedTooltip)
+		SLATE_ARGUMENT(TOptional<ENiagaraAssetTagDefinitionImportance>, DisplayType)
+	SLATE_END_ARGS()
+
+	void Construct(const FArguments& InArgs, const FAssetData& Asset);
+
+private:
+	FText GetDisplayTypeTooltipText(ENiagaraAssetTagDefinitionImportance DisplayType) const;
+};
+
+/**
+ * 
+ */
+class NIAGARAEDITOR_API SNiagaraSelectedAssetDetails : public SCompoundWidget
+{
+public:
+	
+	SLATE_BEGIN_ARGS(SNiagaraSelectedAssetDetails)
+		: _ShowThumbnail(EVisibility::Visible)
+		, _MaxDesiredDescriptionWidth(300.f)
+		{
+		}
+		SLATE_ATTRIBUTE(EVisibility, ShowThumbnail)
+		SLATE_EVENT(FOnAssetTagActivated, OnAssetTagActivated)
+		SLATE_ARGUMENT(TOptional<FText>, OnAssetTagActivatedTooltip)
+	
+		/* Text will automatically wrap according to actual width.
+		 * With this you can specify the maximum width you want the description to be; this is not the actual width.
+		 * Actual width might be increased due to big asset names.
+		 */
+		SLATE_ARGUMENT(float, MaxDesiredDescriptionWidth)
+	SLATE_END_ARGS()
+	
+	/** Constructs this widget with InArgs */
+	void Construct(const FArguments& InArgs, const FAssetData& Asset);
+private:
+	FAssetData AssetData;
+	TAttribute<EVisibility> ShowThumbnail;
+	TSharedPtr<FAssetThumbnail> CurrentAssetThumbnail;
+	FOnAssetTagActivated OnAssetTagActivated;
+	TOptional<FText> OnAssetTagActivatedTooltip;
+
+	TSharedRef<SWidget> CreateAssetThumbnailWidget();
+	TSharedRef<SWidget> CreateTitleWidget();
+	TSharedRef<SWidget> CreateTypeWidget();
+	TSharedRef<SWidget> CreateDescriptionWidget();
+	TSharedRef<SWidget> CreateOptionalPropertiesList();
+	TSharedRef<SWidget> CreateAssetTagRow();
+};
diff --git a/Engine/Plugins/FX/Niagara/Source/NiagaraEditor/Public/Widgets/SItemSelector.h b/Engine/Plugins/FX/Niagara/Source/NiagaraEditor/Public/Widgets/SItemSelector.h
--- a/Engine/Plugins/FX/Niagara/Source/NiagaraEditor/Public/Widgets/SItemSelector.h
+++ b/Engine/Plugins/FX/Niagara/Source/NiagaraEditor/Public/Widgets/SItemSelector.h
@@ -1656,7 +1656,6 @@ public:
 
 		// Search Box
 		SAssignNew(SearchBox, SSearchBox)
-		.Visibility(this, &SItemSelector::GetSearchBoxVisibility)
 		.OnTextChanged(this, &SItemSelector::OnSearchTextChanged)
 		.DelayChangeNotificationsWhileTyping(false);
 
@@ -1667,8 +1666,9 @@ public:
 			.AutoHeight()
 			.Padding(0, 0, 0, 5)
 			[
-				SNew(SHorizontalBox)
 				// Wrapped Search Box
+				SNew(SHorizontalBox)
+				.Visibility(this, &SItemSelector::GetSearchBoxVisibility)
 				+ SHorizontalBox::Slot()
 				.FillWidth(1.0f)
 				.VAlign(VAlign_Center)
@@ -1750,18 +1750,25 @@ public:
 	void SetSelectedItems(const TArray<ItemType>& NewSelectedItems, bool bExpandToShow = false)
 	{
 		checkf(OnCompareItemsForEquality.IsBound(), TEXT("OnCompareItemsForEquality event must be handled to use the SetSelectedItems function."));
-		TGuardValue<bool> SelectionGuard(bIsSettingSelection, true);
 		if (NewSelectedItems.Num() == 0)
 		{
 			ItemTree->ClearSelection();
 		}
 		else
 		{
+			TGuardValue<bool> SelectionGuard(bIsSettingSelection, true);
+
 			TArray<TSharedRef<FItemSelectorItemViewModel>> SelectedItemViewModels;
 			ViewModelUtilities->GetItemViewModelsForItems(NewSelectedItems, SelectedItemViewModels);
 			for (TSharedRef<FItemSelectorItemViewModel> SelectedItemViewModel : SelectedItemViewModels)
 			{
 				ItemTree->SetSelection(SelectedItemViewModel);
+				if(SelectedItemViewModel->GetType() == EItemSelectorItemViewModelType::Item)
+				{
+					TSharedRef<FItemSelectorItemContainerViewModel> ItemViewModel = StaticCastSharedRef<FItemSelectorItemContainerViewModel>(SelectedItemViewModel);
+					OnItemSelected.ExecuteIfBound(ItemViewModel->GetItem(), ESelectInfo::Direct);
+				}
+				
 				if (bExpandToShow)
 				{
 					TSharedPtr<FRootViewModel> RootViewModel = ViewModelUtilities->GetRootViewModel();
@@ -1773,8 +1780,8 @@ public:
 					}
 				}
 			}
+			OnSelectionChanged.ExecuteIfBound();
 		}
-		OnSelectionChanged.ExecuteIfBound();
 	}
 
 	void RequestScrollIntoView(const ItemType& Item)
diff --git a/Engine/Plugins/FX/Niagara/Source/NiagaraEditorWidgets/Private/SNiagaraOverviewGraph.cpp b/Engine/Plugins/FX/Niagara/Source/NiagaraEditorWidgets/Private/SNiagaraOverviewGraph.cpp
--- a/Engine/Plugins/FX/Niagara/Source/NiagaraEditorWidgets/Private/SNiagaraOverviewGraph.cpp
+++ b/Engine/Plugins/FX/Niagara/Source/NiagaraEditorWidgets/Private/SNiagaraOverviewGraph.cpp
@@ -20,15 +20,16 @@
 #include "NiagaraSystem.h"
 #include "NiagaraSystemEditorData.h"
 #include "NiagaraStackEditorData.h"
-#include "SNiagaraAssetPickerList.h"
 #include "SNiagaraOverviewGraphTitleBar.h"
 #include "SNiagaraStack.h"
 #include "ScopedTransaction.h"
+#include "Interfaces/IMainFrameModule.h"
 #include "Templates/SharedPointer.h"
 #include "ViewModels/NiagaraEmitterHandleViewModel.h"
 #include "ViewModels/NiagaraOverviewGraphViewModel.h"
 #include "ViewModels/NiagaraSystemViewModel.h"
 #include "Widgets/SItemSelector.h"
+#include "Widgets/AssetBrowser/SNiagaraAssetBrowser.h"
 
 #define LOCTEXT_NAMESPACE "NiagaraOverviewGraph"
 
@@ -214,46 +215,12 @@ void SNiagaraOverviewGraph::PreClose()
 }
 
 void SNiagaraOverviewGraph::OpenAddEmitterMenu()
-{	
-	FNiagaraAssetPickerListViewOptions ViewOptions;
-	ViewOptions.SetCategorizeUserDefinedCategory(true);
-	ViewOptions.SetCategorizeLibraryAssets(true);
-
-	SNiagaraTemplateTabBox::FNiagaraTemplateTabOptions TabOptions;
-	TabOptions.ChangeTabState(ENiagaraScriptTemplateSpecification::Template, true);
-	TabOptions.ChangeTabState(ENiagaraScriptTemplateSpecification::None, true);
-	TabOptions.ChangeTabState(ENiagaraScriptTemplateSpecification::Behavior, true);	
-
-	TSharedPtr<SNiagaraAssetPickerList> AssetPickerList = SNew(SNiagaraAssetPickerList, UNiagaraEmitter::StaticClass())
-	.ClickActivateMode(EItemSelectorClickActivateMode::SingleClick)
-	.ViewOptions(ViewOptions)
-	.TabOptions(TabOptions)
-	.OnTemplateAssetActivated_Lambda([this](const FAssetData& AssetData) {
-		TSharedPtr<FNiagaraEmitterHandleViewModel> EmitterHandleViewModel = ViewModel->GetSystemViewModel()->AddEmitterFromAssetData(AssetData);
-		FSlateApplication::Get().DismissAllMenus();
-	});
-    
-	TSharedRef<SWidget> EmitterAddSubMenu =
-		SNew(SBorder)
-		.BorderImage(FNiagaraEditorStyle::Get().GetBrush("GraphActionMenu.Background"))
-		.Padding(0.0f)
-		[
-			SNew(SBox)
-			.WidthOverride(450.0f)
-			.HeightOverride(500.0f)
-			[
-				AssetPickerList.ToSharedRef()
-			]
-		];
-	
-	FSlateApplication::Get().PushMenu(SharedThis(this), FWidgetPath(), EmitterAddSubMenu, FSlateApplication::Get().GetCursorPos(), FPopupTransitionEffect::None);
+{
+	IMainFrameModule& MainFrame = FModuleManager::LoadModuleChecked<IMainFrameModule>("MainFrame");
+	TSharedPtr<SWindow>	ParentWindow = MainFrame.GetParentWindow();
 
-	// we need to set the keyboard focus for next tick instead of immediately
-	// if we open the menu via shortcut, we'd type the shortcut key into the search box
-	GEditor->GetTimerManager()->SetTimerForNextTick(FTimerDelegate::CreateLambda([AssetPickerList]()
-	{
-		FSlateApplication::Get().SetKeyboardFocus(AssetPickerList->GetSearchBox());
-	}));
+	TSharedRef<SWindow> AddEmitterWindow = SNew(SNiagaraAddEmitterToSystemWindow, ViewModel->GetSystemViewModel());
+	FSlateApplication::Get().AddModalWindow(AddEmitterWindow, ParentWindow);
 }
 
 bool SNiagaraOverviewGraph::CanAddEmitters() const
