Fixed include path

diff --git a/Engine/Plugins/FX/Niagara/Source/Niagara/Classes/NiagaraSystem.h b/Engine/Plugins/FX/Niagara/Source/Niagara/Classes/NiagaraSystem.h
--- a/Engine/Plugins/FX/Niagara/Source/Niagara/Classes/NiagaraSystem.h
+++ b/Engine/Plugins/FX/Niagara/Source/Niagara/Classes/NiagaraSystem.h
@@ -3,7 +3,7 @@
 #pragma once
 
 #include "CoreMinimal.h"
-#include "..\Public\NiagaraAssetTagDefinitions.h"
+#include "NiagaraAssetTagDefinitions.h"
 
 #include "NiagaraDataSetCompiledData.h"
 #include "NiagaraDataSetAccessor.h"
