Fixed 2 more include paths

diff --git a/Engine/Plugins/FX/Niagara/Source/NiagaraEditor/Private/Widgets/AssetBrowser/SNiagaraAssetBrowser.cpp b/Engine/Plugins/FX/Niagara/Source/NiagaraEditor/Private/Widgets/AssetBrowser/SNiagaraAssetBrowser.cpp
--- a/Engine/Plugins/FX/Niagara/Source/NiagaraEditor/Private/Widgets/AssetBrowser/SNiagaraAssetBrowser.cpp
+++ b/Engine/Plugins/FX/Niagara/Source/NiagaraEditor/Private/Widgets/AssetBrowser/SNiagaraAssetBrowser.cpp
@@ -1,6 +1,6 @@
 ﻿// Copyright Epic Games, Inc. All Rights Reserved.
 
-#include "Widgets\AssetBrowser/SNiagaraAssetBrowser.h"
+#include "Widgets/AssetBrowser/SNiagaraAssetBrowser.h"
 
 #include "ContentBrowserMenuContexts.h"
 #include "ContentBrowserModule.h"
diff --git a/Engine/Plugins/FX/Niagara/Source/NiagaraEditor/Public/Widgets/AssetBrowser/SNiagaraAssetBrowserContent.h b/Engine/Plugins/FX/Niagara/Source/NiagaraEditor/Public/Widgets/AssetBrowser/SNiagaraAssetBrowserContent.h
--- a/Engine/Plugins/FX/Niagara/Source/NiagaraEditor/Public/Widgets/AssetBrowser/SNiagaraAssetBrowserContent.h
+++ b/Engine/Plugins/FX/Niagara/Source/NiagaraEditor/Public/Widgets/AssetBrowser/SNiagaraAssetBrowserContent.h
@@ -4,7 +4,7 @@
 
 #include "CoreMinimal.h"
 #include "FrontendFilterBase.h"
-#include "..\..\..\..\Niagara\Public\NiagaraAssetTagDefinitions.h"
+#include "NiagaraAssetTagDefinitions.h"
 #include "Filters/GenericFilter.h"
 #include "Filters/SAssetFilterBar.h"
 #include "Widgets/SCompoundWidget.h"
