Fixed faulty property initialization

diff --git a/Engine/Plugins/FX/Niagara/Source/Niagara/Private/NiagaraAssetTagDefinitions.cpp b/Engine/Plugins/FX/Niagara/Source/Niagara/Private/NiagaraAssetTagDefinitions.cpp
--- a/Engine/Plugins/FX/Niagara/Source/Niagara/Private/NiagaraAssetTagDefinitions.cpp
+++ b/Engine/Plugins/FX/Niagara/Source/Niagara/Private/NiagaraAssetTagDefinitions.cpp
@@ -10,6 +10,11 @@
 #include "AssetRegistry/AssetData.h"
 #include "UObject/AssetRegistryTagsContext.h"
 
+FNiagaraAssetTagDefinition::FNiagaraAssetTagDefinition(FText InAssetTag, int32 InAssetFlags, FText InDescription, ENiagaraAssetTagDefinitionImportance InDisplayType, FLinearColor InColor, FGuid InTagGuid)
+	: AssetTag(InAssetTag), AssetFlags(InAssetFlags), Description(InDescription), DisplayType(InDisplayType), Color(InColor), TagGuid(InTagGuid)
+{
+}
+
 bool FNiagaraAssetTagDefinition::IsValid() const
 {
 	return TagGuid.IsValid() && AssetTag.IsEmpty() == false;
diff --git a/Engine/Plugins/FX/Niagara/Source/Niagara/Public/NiagaraAssetTagDefinitions.h b/Engine/Plugins/FX/Niagara/Source/Niagara/Public/NiagaraAssetTagDefinitions.h
--- a/Engine/Plugins/FX/Niagara/Source/Niagara/Public/NiagaraAssetTagDefinitions.h
+++ b/Engine/Plugins/FX/Niagara/Source/Niagara/Public/NiagaraAssetTagDefinitions.h
@@ -30,6 +30,13 @@ struct FNiagaraAssetTagDefinition
 	GENERATED_BODY()
 
 public:
+	FNiagaraAssetTagDefinition()
+	{
+		TagGuid = FGuid::NewGuid();
+	}
+
+	FNiagaraAssetTagDefinition(FText InAssetTag, int32 InAssetFlags, FText InDescription, ENiagaraAssetTagDefinitionImportance InDisplayType, FLinearColor InColor, FGuid InTagGuid);
+	
 	/** The Display Name used for this tag. */
 	UPROPERTY(EditAnywhere, Category="Properties")
 	FText AssetTag;
@@ -48,11 +55,11 @@ public:
 
 	/** The color used in UI to represent this tag. */
 	UPROPERTY(EditAnywhere, Category="Properties")
-	FLinearColor Color;
+	FLinearColor Color = FLinearColor::Black;
 
 	/** The Tag Guid identifies this tag. This makes it possible to change the AssetTag name without it affecting functionality. */
 	UPROPERTY(VisibleAnywhere, Category="Properties")
-	FGuid TagGuid = FGuid::NewGuid();
+	FGuid TagGuid = FGuid();
 
 	bool operator<(const FNiagaraAssetTagDefinition& Other) const
 	{		
diff --git a/Engine/Plugins/FX/Niagara/Source/NiagaraEditor/Public/Widgets/AssetBrowser/NiagaraAssetBrowserConfig.h b/Engine/Plugins/FX/Niagara/Source/NiagaraEditor/Public/Widgets/AssetBrowser/NiagaraAssetBrowserConfig.h
--- a/Engine/Plugins/FX/Niagara/Source/NiagaraEditor/Public/Widgets/AssetBrowser/NiagaraAssetBrowserConfig.h
+++ b/Engine/Plugins/FX/Niagara/Source/NiagaraEditor/Public/Widgets/AssetBrowser/NiagaraAssetBrowserConfig.h
@@ -15,7 +15,7 @@ struct FNiagaraAssetBrowserConfiguration
 	TArray<FName> MainFilterSelection;
 
 	UPROPERTY()
-	bool bShouldDisplayViewport;
+	bool bShouldDisplayViewport = false;
 };
 
 UCLASS(EditorConfig="NiagaraAssetBrowser")
