added Array out of bounds guard

diff --git a/Engine/Plugins/FX/Niagara/Source/NiagaraEditor/Private/Widgets/AssetBrowser/SNiagaraAssetBrowser.cpp b/Engine/Plugins/FX/Niagara/Source/NiagaraEditor/Private/Widgets/AssetBrowser/SNiagaraAssetBrowser.cpp
--- a/Engine/Plugins/FX/Niagara/Source/NiagaraEditor/Private/Widgets/AssetBrowser/SNiagaraAssetBrowser.cpp
+++ b/Engine/Plugins/FX/Niagara/Source/NiagaraEditor/Private/Widgets/AssetBrowser/SNiagaraAssetBrowser.cpp
@@ -746,7 +746,7 @@ void SNiagaraAssetBrowser::LoadSettings()
 				LastSelectedMainFilterIdentifierFallback = Config.MainFilterSelection[0];
 			}
 
-			if(Config.MainFilterSelection.Num() >= 1)
+			if(Config.MainFilterSelection.Num() >= 1 && MainFilters.Num() >= 1)
 			{
 				MainFilterSelector->SetSelection(MainFilters[0]);
 			}
