Added the option for tag definition assets to get their own entry in the main filter list.
Clicking on it will show all assets having at least 1 contained tag.
This allows 1 level of hierarchy depth that makes organizing things a lot easier

diff --git a/Engine/Plugins/FX/Niagara/Source/Niagara/Private/NiagaraAssetTagDefinitions.cpp b/Engine/Plugins/FX/Niagara/Source/Niagara/Private/NiagaraAssetTagDefinitions.cpp
--- a/Engine/Plugins/FX/Niagara/Source/Niagara/Private/NiagaraAssetTagDefinitions.cpp
+++ b/Engine/Plugins/FX/Niagara/Source/Niagara/Private/NiagaraAssetTagDefinitions.cpp
@@ -109,3 +109,31 @@ void FNiagaraAssetTagDefinitionReference::AddTagToAssetRegistryTags(FAssetRegist
 		Context.AddTag(UObject::FAssetRegistryTag(FName(GetGuidAsString()), "  ", UObject::FAssetRegistryTag::TT_Alphabetical));
 	}
 }
+
+FText UNiagaraAssetTagDefinitions::GetDisplayName() const
+{
+	if(DisplayName.IsEmptyOrWhitespace())
+	{
+		return FText::FromString(GetName());
+	}
+
+	return DisplayName;
+}
+
+FText UNiagaraAssetTagDefinitions::GetDescription() const
+{
+	return Description;
+}
+
+bool UNiagaraAssetTagDefinitions::DoesAssetDataContainAnyTag(const FAssetData& AssetData) const
+{
+	for(const FNiagaraAssetTagDefinition& AssetTagDefinition : TagDefinitions)
+	{
+		if(AssetTagDefinition.DoesAssetDataContainTag(AssetData))
+		{
+			return true;
+		}
+	}
+
+	return false;
+}
diff --git a/Engine/Plugins/FX/Niagara/Source/Niagara/Public/NiagaraAssetTagDefinitions.h b/Engine/Plugins/FX/Niagara/Source/Niagara/Public/NiagaraAssetTagDefinitions.h
--- a/Engine/Plugins/FX/Niagara/Source/Niagara/Public/NiagaraAssetTagDefinitions.h
+++ b/Engine/Plugins/FX/Niagara/Source/Niagara/Public/NiagaraAssetTagDefinitions.h
@@ -122,11 +122,31 @@ class NIAGARA_API UNiagaraAssetTagDefinitions : public UObject
 public:	
 	const TArray<FNiagaraAssetTagDefinition>& GetAssetTagDefinitions() const { return TagDefinitions; }
 
+	FText GetDisplayName() const;
+	
+	FText GetDescription() const;
+
+	bool DisplayTagsAsFlatList() const { return bDisplayTagsAsFlatList; }
+	
+	bool DoesAssetDataContainAnyTag(const FAssetData& AssetData) const;
+
 	int32 GetSortOrder() const { return SortOrder; }
 private:
 	UPROPERTY(EditAnywhere, Category="Properties", meta=(TitleProperty="AssetTag"))
 	TArray<FNiagaraAssetTagDefinition> TagDefinitions;
 
+	/** The display name to use when listing this asset in the Niagara Asset Browser */
+	UPROPERTY(EditAnywhere, Category="Properties")
+	FText DisplayName;
+
+	/** A description for this group of tags. Used for tooltips. */
+	UPROPERTY(EditAnywhere, Category="Properties")
+	FText Description;
+	
+	/** If true, no 'parent' entry for this asset will be displayed in the Niagara Asset Browser. Instead a flat list of the contained tags will be added. */
+	UPROPERTY(EditAnywhere, Category="Properties")
+	bool bDisplayTagsAsFlatList = false;
+	
 	/** Tags are sorted by asset sort order first, then individually. That means tags of asset with sort order [0] come before tags of asset with sort order [1]. */
 	UPROPERTY(EditAnywhere, Category="Properties")
 	int32 SortOrder = INDEX_NONE;
diff --git a/Engine/Plugins/FX/Niagara/Source/NiagaraEditor/Private/Widgets/AssetBrowser/NiagaraMenuFilters.cpp b/Engine/Plugins/FX/Niagara/Source/NiagaraEditor/Private/Widgets/AssetBrowser/NiagaraMenuFilters.cpp
--- a/Engine/Plugins/FX/Niagara/Source/NiagaraEditor/Private/Widgets/AssetBrowser/NiagaraMenuFilters.cpp
+++ b/Engine/Plugins/FX/Niagara/Source/NiagaraEditor/Private/Widgets/AssetBrowser/NiagaraMenuFilters.cpp
@@ -36,6 +36,10 @@ FText FNiagaraAssetBrowserMainFilter::GetDisplayName() const
 	{
 		return FText::FromString(AssetTagDefinition.AssetTag.ToString());
 	}
+	else if(FilterMode == EFilterMode::NiagaraAssetTagDefinitionsAsset)
+	{
+		return AssetTagDefinitionsAsset->GetDisplayName();
+	}
 	else if(FilterMode == EFilterMode::Recent)
 	{
 		return LOCTEXT("NiagaraMainFilterLabel_Recent", "Recent");
@@ -63,4 +67,9 @@ bool FNiagaraAssetBrowserMainFilter::DoesAssetHaveTag(const FAssetData& AssetCan
 	return AssetTagDefinition.DoesAssetDataContainTag(AssetCandidate);
 }
 
+bool FNiagaraAssetBrowserMainFilter::DoesAssetHaveAnyTagFromTagDefinitionsAsset(const FAssetData& AssetCandidate) const
+{
+	return AssetTagDefinitionsAsset->DoesAssetDataContainAnyTag(AssetCandidate);
+}
+
 #undef LOCTEXT_NAMESPACE
diff --git a/Engine/Plugins/FX/Niagara/Source/NiagaraEditor/Private/Widgets/AssetBrowser/SNiagaraAssetBrowser.cpp b/Engine/Plugins/FX/Niagara/Source/NiagaraEditor/Private/Widgets/AssetBrowser/SNiagaraAssetBrowser.cpp
--- a/Engine/Plugins/FX/Niagara/Source/NiagaraEditor/Private/Widgets/AssetBrowser/SNiagaraAssetBrowser.cpp
+++ b/Engine/Plugins/FX/Niagara/Source/NiagaraEditor/Private/Widgets/AssetBrowser/SNiagaraAssetBrowser.cpp
@@ -188,6 +188,10 @@ bool SNiagaraAssetBrowser::ShouldFilterAsset(const FAssetData& AssetData) const
 		{
 			return MainFilter->DoesAssetHaveTag(AssetData) == false;
 		}
+		else if(MainFilter->FilterMode == FNiagaraAssetBrowserMainFilter::EFilterMode::NiagaraAssetTagDefinitionsAsset)
+		{
+			return MainFilter->DoesAssetHaveAnyTagFromTagDefinitionsAsset(AssetData) == false;
+		}
 		else if(MainFilter->FilterMode == FNiagaraAssetBrowserMainFilter::EFilterMode::Custom)
 		{
 			return MainFilter->CustomShouldFilterAsset.Execute(AssetData);
@@ -302,43 +306,98 @@ TArray<TSharedRef<FNiagaraAssetBrowserMainFilter>> SNiagaraAssetBrowser::GetMain
 	TArray<TSharedRef<FNiagaraAssetBrowserMainFilter>> TagFilters;
 	{
 		TArray<UClass*> DisplayedAssetTypes = GetDisplayedAssetTypes();
-		TArray<FNiagaraAssetTagDefinition> AssetTagDefinitions;
+		TArray<UNiagaraAssetTagDefinitions*> DisplayedAssetTagDefinitionAssets;
+		TArray<FNiagaraAssetTagDefinition> DisplayedFlatAssetTagDefinitionsList;
+
+		auto IsAssetTagDefinitionValid = [DisplayedAssetTypes](const FNiagaraAssetTagDefinition& AssetTagDefinition) -> bool
+		{
+			if(AssetTagDefinition.DisplayType != ENiagaraAssetTagDefinitionImportance::Primary)
+			{
+				return false;
+			}
+				
+			TArray<UClass*> SupportedClasses = AssetTagDefinition.GetSupportedClasses();
+	
+			bool bCanAssetTagContainDisplayedAssetType = false;
+			for(UClass* SupportedClass : SupportedClasses)
+			{
+				if(DisplayedAssetTypes.Contains(SupportedClass))
+				{
+					bCanAssetTagContainDisplayedAssetType = true;
+					break;
+				}
+			}
+	
+			if(bCanAssetTagContainDisplayedAssetType == false)
+			{
+				return false;
+			}
+
+			return true;
+		};
 		
 		for(const FStructuredAssetTagDefinitionLookupData& AssetTagDefinitionData : GetStructuredSortedAssetTagDefinitions())
 		{
-			for(const FNiagaraAssetTagDefinition& AssetTagDefinition : AssetTagDefinitionData.AssetTagDefinitions)
+			// If there is no definitions asset, the tags have been declared internally
+			if(AssetTagDefinitionData.DefinitionsAsset == nullptr)
 			{
-				if(AssetTagDefinition.DisplayType != ENiagaraAssetTagDefinitionImportance::Primary)
+				for(const FNiagaraAssetTagDefinition& AssetTagDefinition : AssetTagDefinitionData.AssetTagDefinitions)
 				{
-					continue;
+					if(IsAssetTagDefinitionValid(AssetTagDefinition))
+					{
+						DisplayedFlatAssetTagDefinitionsList.Add(AssetTagDefinition);
+					}	
 				}
-				
-				TArray<UClass*> SupportedClasses = AssetTagDefinition.GetSupportedClasses();
-	
-				bool bCanAssetTagContainDisplayedAssetType = false;
-				for(UClass* SupportedClass : SupportedClasses)
+			}
+			// If there is an asset, we check if we want to display a parent-entry per asset first
+			else
+			{
+				// If we only want to display the tags or we only have 1 tag defined, we add the tags directly to a flat list
+				if(AssetTagDefinitionData.DefinitionsAsset->DisplayTagsAsFlatList() || AssetTagDefinitionData.AssetTagDefinitions.Num() == 1)
 				{
-					if(DisplayedAssetTypes.Contains(SupportedClass))
+					for(const FNiagaraAssetTagDefinition& AssetTagDefinition : AssetTagDefinitionData.AssetTagDefinitions)
 					{
-						bCanAssetTagContainDisplayedAssetType = true;
-						break;
+						if(IsAssetTagDefinitionValid(AssetTagDefinition))
+						{
+							DisplayedFlatAssetTagDefinitionsList.Add(AssetTagDefinition);
+						}	
 					}
 				}
-	
-				if(bCanAssetTagContainDisplayedAssetType == false)
+				// If not, we keep track of the asset here so we can construct a hierarchy of filters per asset
+				else
 				{
-					continue;
+					DisplayedAssetTagDefinitionAssets.Add(AssetTagDefinitionData.DefinitionsAsset);
 				}
-	
-				AssetTagDefinitions.Add(AssetTagDefinition);
-			}
+			}			
 		}
 
-		for(const FNiagaraAssetTagDefinition& AssetTagDefinition : AssetTagDefinitions)
+		// First we add all 'flat list' entries at the top
+		for(const FNiagaraAssetTagDefinition& AssetTagDefinition : DisplayedFlatAssetTagDefinitionsList)
 		{
-			TSharedRef<FNiagaraAssetBrowserMainFilter> NiagaraAssetTagFilter = MakeShared<FNiagaraAssetBrowserMainFilter>(FNiagaraAssetBrowserMainFilter::EFilterMode::NiagaraAssetTag);
-			NiagaraAssetTagFilter->AssetTagDefinition = AssetTagDefinition;
-			TagFilters.Add(NiagaraAssetTagFilter);
+			TSharedRef<FNiagaraAssetBrowserMainFilter> AssetTagFilter = MakeShared<FNiagaraAssetBrowserMainFilter>(FNiagaraAssetBrowserMainFilter::EFilterMode::NiagaraAssetTag);
+			AssetTagFilter->AssetTagDefinition = AssetTagDefinition;
+			TagFilters.Add(AssetTagFilter);
+		}
+
+		for(const UNiagaraAssetTagDefinitions* AssetTagDefinitionsAsset : DisplayedAssetTagDefinitionAssets)
+		{
+			// This code should only execute for assets with > 1 tag. If there is only 1 tag, it should have been automatically added to the flat list instead
+			if(ensure(AssetTagDefinitionsAsset->GetAssetTagDefinitions().Num() > 1))
+			{
+				TSharedRef<FNiagaraAssetBrowserMainFilter> AssetTagDefinitionsAssetsFilter = MakeShared<FNiagaraAssetBrowserMainFilter>(FNiagaraAssetBrowserMainFilter::EFilterMode::NiagaraAssetTagDefinitionsAsset);
+				AssetTagDefinitionsAssetsFilter->AssetTagDefinitionsAsset = AssetTagDefinitionsAsset;
+				TagFilters.Add(AssetTagDefinitionsAssetsFilter);
+
+				TArray<TSharedRef<FNiagaraAssetBrowserMainFilter>> TagChildFilters;
+				for(const FNiagaraAssetTagDefinition& AssetTagDefinition : AssetTagDefinitionsAsset->GetAssetTagDefinitions())
+				{
+					TSharedRef<FNiagaraAssetBrowserMainFilter> AssetTagFilter = MakeShared<FNiagaraAssetBrowserMainFilter>(FNiagaraAssetBrowserMainFilter::EFilterMode::NiagaraAssetTag);
+					AssetTagFilter->AssetTagDefinition = AssetTagDefinition;
+					TagChildFilters.Add(AssetTagFilter);
+				}
+
+				AssetTagDefinitionsAssetsFilter->ChildFilters = TagChildFilters;
+			}			
 		}
 	}
 	
@@ -432,11 +491,16 @@ TSharedRef<ITableRow> SNiagaraAssetBrowser::GenerateWidgetRowForMainFilter(TShar
 		SNew(STextBlock)
 		.Text(MainFilter->GetDisplayName())
 	];
-
+	
 	if(MainFilter->FilterMode == FNiagaraAssetBrowserMainFilter::EFilterMode::NiagaraAssetTag)
 	{
 		Widget->SetToolTipText(MainFilter->AssetTagDefinition.Description);	
 	}
+
+	if(MainFilter->FilterMode == FNiagaraAssetBrowserMainFilter::EFilterMode::NiagaraAssetTagDefinitionsAsset)
+	{
+		Widget->SetToolTipText(MainFilter->AssetTagDefinitionsAsset->GetDescription());	
+	}
 	
 	return SNew(STableRow<TSharedRef<FNiagaraAssetBrowserMainFilter>>, OwningTable)
 	[
diff --git a/Engine/Plugins/FX/Niagara/Source/NiagaraEditor/Public/Widgets/AssetBrowser/NiagaraMenuFilters.h b/Engine/Plugins/FX/Niagara/Source/NiagaraEditor/Public/Widgets/AssetBrowser/NiagaraMenuFilters.h
--- a/Engine/Plugins/FX/Niagara/Source/NiagaraEditor/Public/Widgets/AssetBrowser/NiagaraMenuFilters.h
+++ b/Engine/Plugins/FX/Niagara/Source/NiagaraEditor/Public/Widgets/AssetBrowser/NiagaraMenuFilters.h
@@ -47,6 +47,7 @@ struct FNiagaraAssetBrowserMainFilter
 	{
 		All,
 		NiagaraAssetTag,
+		NiagaraAssetTagDefinitionsAsset,
 		Recent,
 		Custom
 	};
@@ -56,16 +57,20 @@ struct FNiagaraAssetBrowserMainFilter
 	EFilterMode FilterMode;
 	TArray<TSharedRef<FNiagaraAssetBrowserMainFilter>> ChildFilters;
 
+	// If in mode Asset Tag, this is the tag to filter against
 	FNiagaraAssetTagDefinition AssetTagDefinition;
+	// If in mode Asset Tag Definitions Asset, this is the asset to filter against (should include all assets that match any tag within the asset)
+	TObjectPtr<const UNiagaraAssetTagDefinitions> AssetTagDefinitionsAsset;
+	// If in mode Recent, this is set and used to determine if an asset is recent
+	FIsAssetRecent IsAssetRecentDelegate;
+	// If in mode Custom, this needs to be set
+	FOnShouldFilterAsset CustomShouldFilterAsset;
 
 	/** The custom guid should always be set the same for a specific main filter.
 	 * Should be used if there are multiple items with the same filter mode that aren't niagara tags (i.e. set to Custom) */
 	FGuid CustomGuid;
-
-	FIsAssetRecent IsAssetRecentDelegate;
 	
 	FText CustomDisplayName;
-	FOnShouldFilterAsset CustomShouldFilterAsset;
 
 	FName GetIdentifier() const;
 	FText GetDisplayName() const;
@@ -73,6 +78,7 @@ struct FNiagaraAssetBrowserMainFilter
 	bool IsAssetRecent(const FAssetData& AssetCandidate) const;
 	bool ShouldCustomFilterAsset(const FAssetData& AssetCandidate) const;
 	bool DoesAssetHaveTag(const FAssetData& AssetCandidate) const;
+	bool DoesAssetHaveAnyTagFromTagDefinitionsAsset(const FAssetData& AssetCandidate) const;
 
 	bool operator==(const FNiagaraAssetBrowserMainFilter& Other) const
 	{
@@ -86,6 +92,11 @@ struct FNiagaraAssetBrowserMainFilter
 			return AssetTagDefinition == Other.AssetTagDefinition;
 		}
 
+		if(FilterMode == EFilterMode::NiagaraAssetTagDefinitionsAsset)
+		{
+			return AssetTagDefinitionsAsset == Other.AssetTagDefinitionsAsset;
+		}
+
 		if(CustomGuid.IsValid())
 		{
 			return CustomGuid == Other.CustomGuid;
