Added analytics for system data, cascade conversion and sim caches

diff --git a/Engine/Plugins/FX/CascadeToNiagaraConverter/Source/CascadeToNiagaraConverter/Private/CascadeToNiagaraConverterModule.cpp b/Engine/Plugins/FX/CascadeToNiagaraConverter/Source/CascadeToNiagaraConverter/Private/CascadeToNiagaraConverterModule.cpp
--- a/Engine/Plugins/FX/CascadeToNiagaraConverter/Source/CascadeToNiagaraConverter/Private/CascadeToNiagaraConverterModule.cpp
+++ b/Engine/Plugins/FX/CascadeToNiagaraConverter/Source/CascadeToNiagaraConverter/Private/CascadeToNiagaraConverterModule.cpp
@@ -3,6 +3,7 @@
 #include "CascadeToNiagaraConverterModule.h"
 #include "ContentBrowserModule.h"
 #include "IPythonScriptPlugin.h"
+#include "NiagaraAnalytics.h"
 #include "Framework/MultiBox/MultiBoxBuilder.h"
 #include "NiagaraMessageManager.h"
 #include "NiagaraStackGraphUtilitiesAdapterLibrary.h"
@@ -74,6 +75,9 @@ void ICascadeToNiagaraConverterModule::ExecuteConvertCascadeSystemToNiagaraSyste
 {
 	UConvertCascadeToNiagaraResults* Results = NewObject<UConvertCascadeToNiagaraResults>();
 
+	bool bCancelled = false;
+	bool bSuccess = true;
+	double StartTime = FPlatformTime::Seconds();
 	for (UParticleSystem* CascadeSystem : CascadeSystems)
 	{
 		Results->Init();
@@ -88,16 +92,24 @@ void ICascadeToNiagaraConverterModule::ExecuteConvertCascadeSystemToNiagaraSyste
 
 		if (Results->bCancelledByUser)
 		{
-			// raise modal dialog of status
-			return;
+			bCancelled = true;
+			break;
 		}
-		else if (Results->bCancelledByPythonError)
+		if (Results->bCancelledByPythonError)
 		{
 			// raise modal dialog of status
 			//PythonCommand.CommandResult
-			return;
+			bSuccess = false;
+			break;
 		}
 	}
+
+	TArray<FAnalyticsEventAttribute> Attributes;
+	Attributes.Emplace(TEXT("DurationSeconds"), FPlatformTime::Seconds() - StartTime);
+	Attributes.Emplace(TEXT("Cancelled"), bCancelled);
+	Attributes.Emplace(TEXT("Error"), bSuccess == false);
+	Attributes.Emplace(TEXT("AssetCount"), CascadeSystems.Num());
+	NiagaraAnalytics::RecordEvent("Cascade.Conversion", Attributes);
 }
 
 #undef LOCTEXT_NAMESPACE //"CascadeToNiagaraConverterModule"
diff --git a/Engine/Plugins/FX/Niagara/Source/Niagara/Classes/NiagaraScriptSourceBase.h b/Engine/Plugins/FX/Niagara/Source/Niagara/Classes/NiagaraScriptSourceBase.h
--- a/Engine/Plugins/FX/Niagara/Source/Niagara/Classes/NiagaraScriptSourceBase.h
+++ b/Engine/Plugins/FX/Niagara/Source/Niagara/Classes/NiagaraScriptSourceBase.h
@@ -29,6 +29,13 @@ struct EditorExposedVectorCurveConstant
 	class UCurveVector *Value;
 };
 
+struct FNiagaraScriptSourceAnalytics
+{
+	int32 ActiveModules = 0;
+	int32 DisabledModules = 0;
+	TSet<FString> UsedNiagaraModules;
+};
+
 /** External reference to the compile request data generated.*/
 class FNiagaraCompileRequestDataBase
 {
@@ -119,12 +126,13 @@ class UNiagaraScriptSourceBase : public UObject
 
 	virtual FGuid GetChangeID() { return FGuid(); }
 
-	virtual void RegisterVMCompilationIdDependencies(FNiagaraScriptHashCollector& Collector, ENiagaraScriptUsage InUsage, const FGuid& InUsageId) const {};
+	virtual void RegisterVMCompilationIdDependencies(FNiagaraScriptHashCollector& Collector, ENiagaraScriptUsage InUsage, const FGuid& InUsageId) const {}
 	virtual void ComputeVMCompilationId(FNiagaraVMExecutableDataId& Id, FNiagaraScriptHashCollector& HashCollector, ENiagaraScriptUsage InUsage, const FGuid& InUsageId) const {};
 
 #if WITH_EDITORONLY_DATA
 	// Will conditionally refresh the graph's CompileId
 	virtual void RefreshGraphCompileId() {};
+	virtual void ReportAnalyticsData(FNiagaraScriptSourceAnalytics& InData) const {};
 #endif
 
 	virtual TMap<FName, UNiagaraDataInterface*> ComputeObjectNameMap(UNiagaraSystem& System, ENiagaraScriptUsage Usage, FGuid UsageId, FString EmitterUniqueName) const { return TMap<FName, UNiagaraDataInterface*>(); }
diff --git a/Engine/Plugins/FX/Niagara/Source/Niagara/Classes/NiagaraSimCache.h b/Engine/Plugins/FX/Niagara/Source/Niagara/Classes/NiagaraSimCache.h
--- a/Engine/Plugins/FX/Niagara/Source/Niagara/Classes/NiagaraSimCache.h
+++ b/Engine/Plugins/FX/Niagara/Source/Niagara/Classes/NiagaraSimCache.h
@@ -420,7 +420,7 @@ public:
 	NIAGARA_API bool WriteFrame(UNiagaraComponent* NiagaraComponent);
 	NIAGARA_API bool WriteFrame(UNiagaraComponent* NiagaraComponent, FNiagaraSimCacheFeedbackContext& FeedbackContext);
 
-	NIAGARA_API bool EndWrite();
+	NIAGARA_API bool EndWrite(bool bAllowAnalytics = false);
 
 	NIAGARA_API bool CanRead(UNiagaraSystem* NiagaraSystem);
 	NIAGARA_API bool Read(float TimeSeconds, FNiagaraSystemInstance* SystemInstance) const;
@@ -600,6 +600,7 @@ private:
 	TMap<FNiagaraVariableBase, TObjectPtr<UObject>> DataInterfaceStorage;
 
 	int32 CaptureTickCount = INDEX_NONE;
+	double CaptureStartTime = 0;
 
 	mutable std::atomic<int32> PendingCommandsInFlight;
 
diff --git a/Engine/Plugins/FX/Niagara/Source/Niagara/Classes/NiagaraSystem.h b/Engine/Plugins/FX/Niagara/Source/Niagara/Classes/NiagaraSystem.h
--- a/Engine/Plugins/FX/Niagara/Source/Niagara/Classes/NiagaraSystem.h
+++ b/Engine/Plugins/FX/Niagara/Source/Niagara/Classes/NiagaraSystem.h
@@ -501,6 +501,8 @@ public:
 
 	FORCEINLINE bool IsInitialOwnerVelocityFromActor() const { return bInitialOwnerVelocityFromActor; }
 
+	NIAGARA_API void ReportAnalyticsData(bool bIsCooking);
+
 #if WITH_EDITORONLY_DATA
 	NIAGARA_API bool UsesEmitter(UNiagaraEmitter* Emitter) const;
 	NIAGARA_API bool UsesEmitter(const FVersionedNiagaraEmitter& VersionedEmitter) const;
@@ -827,7 +829,6 @@ private:
 #endif
 
 
-	void ResolveRequiresScripts();
 	NIAGARA_API void ResolveScalabilitySettings();
 	NIAGARA_API void UpdatePostCompileDIInfo();
 	NIAGARA_API void UpdateDITickFlags();
diff --git a/Engine/Plugins/FX/Niagara/Source/Niagara/Internal/NiagaraAnalytics.cpp b/Engine/Plugins/FX/Niagara/Source/Niagara/Internal/NiagaraAnalytics.cpp
--- a/Engine/Plugins/FX/Niagara/Source/Niagara/Internal/NiagaraAnalytics.cpp
+++ b/Engine/Plugins/FX/Niagara/Source/Niagara/Internal/NiagaraAnalytics.cpp
@@ -4,6 +4,7 @@
 #include "EngineAnalytics.h"
 #include "UObject/Object.h"
 #include "UObject/Package.h"
+#include "UObject/Class.h"
 
 void NiagaraAnalytics::RecordEvent(FString&& EventName, const TArray<FAnalyticsEventAttribute>& Attributes)
 {
@@ -29,3 +30,13 @@ bool NiagaraAnalytics::IsPluginAsset(const UObject* Obj)
 	}
 	return false;
 }
+
+bool NiagaraAnalytics::IsPluginClass(const UClass* Class)
+{
+	if (Class)
+	{
+		FString Name = Class->GetPackage()->GetName();
+		return Name == TEXT("/Script/Niagara") || Name == TEXT("/Script/NiagaraFluids");
+	}
+	return false;
+}
diff --git a/Engine/Plugins/FX/Niagara/Source/Niagara/Internal/NiagaraAnalytics.h b/Engine/Plugins/FX/Niagara/Source/Niagara/Internal/NiagaraAnalytics.h
--- a/Engine/Plugins/FX/Niagara/Source/Niagara/Internal/NiagaraAnalytics.h
+++ b/Engine/Plugins/FX/Niagara/Source/Niagara/Internal/NiagaraAnalytics.h
@@ -13,4 +13,7 @@ namespace NiagaraAnalytics
 
 	// Checks if the provided object is part of a Niagara plugin and can safely be reported on (e.g. the asset name)
 	NIAGARA_API bool IsPluginAsset(const UObject* Obj);
+
+	// Checks if the provided class is part of a Niagara plugin and can safely be reported on (e.g. the asset name)
+	NIAGARA_API bool IsPluginClass(const UClass* Class);
 };
\ No newline at end of file
diff --git a/Engine/Plugins/FX/Niagara/Source/Niagara/Private/NiagaraSimCache.cpp b/Engine/Plugins/FX/Niagara/Source/Niagara/Private/NiagaraSimCache.cpp
--- a/Engine/Plugins/FX/Niagara/Source/Niagara/Private/NiagaraSimCache.cpp
+++ b/Engine/Plugins/FX/Niagara/Source/Niagara/Private/NiagaraSimCache.cpp
@@ -2,6 +2,8 @@
 
 #include "NiagaraSimCache.h"
 
+#include "EngineAnalytics.h"
+#include "NiagaraAnalytics.h"
 #include "Engine/World.h"
 #include "Misc/LargeWorldRenderPosition.h"
 #include "NiagaraConstants.h"
@@ -324,6 +326,7 @@ bool UNiagaraSimCache::BeginWrite(FNiagaraSimCacheCreateParameters InCreateParam
 	CacheLayout = FNiagaraSimCacheLayout();
 	CacheFrames.Empty();
 	CaptureTickCount = INDEX_NONE;
+	CaptureStartTime = FPlatformTime::Seconds();
 
 	for ( auto it=DataInterfaceStorage.CreateIterator(); it; ++it )
 	{
@@ -548,6 +551,7 @@ bool UNiagaraSimCache::BeginAppend(FNiagaraSimCacheCreateParameters InCreatePara
 	{
 		CaptureTickCount = SimulationTickCount - 1;
 	}
+	CaptureStartTime = FPlatformTime::Seconds();
 
 	return true;
 }
@@ -717,7 +721,7 @@ bool UNiagaraSimCache::WriteFrame(UNiagaraComponent* NiagaraComponent, FNiagaraS
 	return true;
 }
 
-bool UNiagaraSimCache::EndWrite()
+bool UNiagaraSimCache::EndWrite(bool bAllowAnalytics)
 {
 	check(PendingCommandsInFlight == 0);
 	if ( CacheFrames.Num() == 0 )
@@ -802,6 +806,29 @@ bool UNiagaraSimCache::EndWrite()
 		}
 	}
 
+#if WITH_EDITORONLY_DATA
+	if (FEngineAnalytics::IsAvailable() && bAllowAnalytics)
+	{		
+		TArray<FAnalyticsEventAttribute> Attributes;
+		Attributes.Emplace(TEXT("DurationSeconds"), FPlatformTime::Seconds() - CaptureStartTime);
+		Attributes.Emplace(TEXT("ValidResult"), IsCacheValid());
+		Attributes.Emplace(TEXT("NumFrames"), GetNumFrames());
+		Attributes.Emplace(TEXT("NumEmitters"), GetNumEmitters());
+		Attributes.Emplace(TEXT("NumDataInterfaces"), GetStoredDataInterfaces().Num());
+		int64 StoredBytes = 0;
+		for (const FNiagaraSimCacheFrame& Frame : CacheFrames)
+		{
+			StoredBytes += Frame.SystemData.SystemDataBuffers.DataBuffer.GetView().Num();
+			for (const FNiagaraSimCacheEmitterFrame& EmitterFrame : Frame.EmitterData)
+			{
+				StoredBytes += EmitterFrame.ParticleDataBuffers.DataBuffer.GetView().Num();
+			}
+		}
+		Attributes.Emplace(TEXT("ParticleStorageBytes"), StoredBytes);
+		NiagaraAnalytics::RecordEvent("SimCache.Capture", Attributes);
+	}
+#endif
+
 	OnCacheEndWrite.Broadcast(this);
 	return IsCacheValid();
 }
diff --git a/Engine/Plugins/FX/Niagara/Source/Niagara/Private/NiagaraSystem.cpp b/Engine/Plugins/FX/Niagara/Source/Niagara/Private/NiagaraSystem.cpp
--- a/Engine/Plugins/FX/Niagara/Source/Niagara/Private/NiagaraSystem.cpp
+++ b/Engine/Plugins/FX/Niagara/Source/Niagara/Private/NiagaraSystem.cpp
@@ -2,11 +2,14 @@
 
 
 #include "NiagaraSystem.h"
+
+#include "EngineAnalytics.h"
 #include "Components/PrimitiveComponent.h"
 #include "NiagaraSystemImpl.h"
 
 #include "AssetRegistry/IAssetRegistry.h"
 #include "INiagaraEditorOnlyDataUtlities.h"
+#include "NiagaraAnalytics.h"
 #include "Materials/MaterialInterface.h"
 #include "Stateless/NiagaraStatelessEmitter.h"
 #include "NiagaraAsyncCompile.h"
@@ -118,6 +121,14 @@ static FAutoConsoleVariableRef CVarNiagaraObjectNeedsLoadMode(
 	TEXT("2 - Validate objects are loaded and force preload\n"),
 	ECVF_Default
 );
+
+static int GNiagaraReportOnCook = 1;
+static FAutoConsoleVariableRef CVarNiagaraReportOnCook(
+	TEXT("fx.Niagara.Analytics.ReportOnCook"),
+	GNiagaraReportOnCook,
+	TEXT("If true then basic system info will be gathered and reported as part of the editor analytics for every cooked system."),
+	ECVF_Default
+);
 #endif
 
 //////////////////////////////////////////////////////////////////////////
@@ -302,6 +313,7 @@ void UNiagaraSystem::BeginCacheForCookedPlatformData(const ITargetPlatform *Targ
 	{
 		WaitForCompilationComplete();
 	}
+	ReportAnalyticsData(true);
 }
 
 bool UNiagaraSystem::IsCachedCookedPlatformDataLoaded(const ITargetPlatform* TargetPlatform)
@@ -627,9 +639,7 @@ void UNiagaraSystem::UpdateSystemAfterLoad()
 		bIsValidCached = IsValidInternal();
 		bIsReadyToRunCached = IsReadyToRunInternal();
 	}
-
-	ResolveRequiresScripts();
-
+	
 	ResolveScalabilitySettings();
 
 	ComputeEmittersExecutionOrder();
@@ -724,6 +734,128 @@ bool UNiagaraSystem::UsesScript(const UNiagaraScript* Script) const
 	return false;
 }
 
+void UNiagaraSystem::ReportAnalyticsData(bool bIsCooking)
+{
+	if (FEngineAnalytics::IsAvailable() && IsAsset() && (bIsCooking == false || GNiagaraReportOnCook))
+	{
+		TArray<FAnalyticsEventAttribute> Attributes;
+		Attributes.Emplace(TEXT("FromCook"), bIsCooking);
+		Attributes.Emplace(TEXT("EmitterCount"), EmitterHandles.Num());
+		Attributes.Emplace(TEXT("UserParameterCount"), GetExposedParameters().Num());
+		Attributes.Emplace(TEXT("IsNiagaraAsset"), NiagaraAnalytics::IsPluginAsset(this));
+		Attributes.Emplace(TEXT("AssetPathHash"), GetTypeHash(GetPathName()));
+		
+		// gather data interface data
+		TSet<UNiagaraDataInterface*> DataInterfaces;
+		TArray<FString> DataInterfaceClasses;
+		FNiagaraDataInterfaceUtilities::FDataInterfaceSearchOptions SearchOptions;
+		SearchOptions.bIncludeInternal = true;
+		ForEachDataInterface(this, [&](const FNiagaraDataInterfaceUtilities::FDataInterfaceUsageContext& UsageContext) -> bool
+		{
+			if (UsageContext.DataInterface && !DataInterfaces.Contains(UsageContext.DataInterface))
+			{
+				DataInterfaces.Add(UsageContext.DataInterface);
+
+				if (NiagaraAnalytics::IsPluginClass(UsageContext.DataInterface->GetClass()))
+				{
+					DataInterfaceClasses.AddUnique(UsageContext.DataInterface->GetClass()->GetName());
+				}
+			}
+			return true;
+		}, SearchOptions);
+		DataInterfaceClasses.Sort();
+		Attributes.Emplace(TEXT("DataInterfaceCount"), DataInterfaces.Num());
+		Attributes.Emplace(TEXT("ExposedDataInterfaceCount"), GetExposedParameters().GetDataInterfaces().Num());
+		Attributes.Add(FAnalyticsEventAttribute(TEXT("DataInterfaces"), DataInterfaceClasses));
+
+		// gather emitter data
+		int32 GpuCount = 0;
+		int32 DisabledEmitters = 0;
+		int32 StatelessEmitters = 0;
+		int32 RendererCount = 0;
+		int32 CustomRendererCount = 0;
+		int32 DisabledRendererCount = 0;
+		int32 ParentEmitterCount = 0;
+		int32 ScratchPadCount = ScratchPadScripts.Num();
+		TArray<FString> ParentEmitters;
+		FNiagaraScriptSourceAnalytics ScriptAnalytics;
+		for (const FNiagaraEmitterHandle& Handle : EmitterHandles)
+		{
+			if (!Handle.GetIsEnabled())
+			{
+				DisabledEmitters++;
+			}
+			if (Handle.GetEmitterMode() == ENiagaraEmitterMode::Stateless)
+			{
+				StatelessEmitters++;
+			}
+			if (FVersionedNiagaraEmitterData* EmitterData = Handle.GetInstance().GetEmitterData())
+			{
+				if (EmitterData->SimTarget == ENiagaraSimTarget::GPUComputeSim)
+				{
+					GpuCount++;
+				}
+				
+				if (EmitterData->GraphSource && Handle.GetIsEnabled())
+				{
+					EmitterData->GraphSource->ReportAnalyticsData(ScriptAnalytics);
+					EmitterData->ForEachRenderer([&](UNiagaraRendererProperties* Renderer)
+					{
+						RendererCount++;
+						if (Renderer && !Renderer->GetIsEnabled())
+						{
+							DisabledRendererCount++;
+						}
+						if (Renderer && !NiagaraAnalytics::IsPluginClass(Renderer->GetClass()))
+						{
+							CustomRendererCount++;
+						}
+					});
+					ScratchPadCount += EmitterData->ScratchPads->Scripts.Num();
+
+					if (EmitterData->GetParent().GetEmitterData())
+					{
+						ParentEmitterCount++;
+						if (NiagaraAnalytics::IsPluginAsset(EmitterData->GetParent().Emitter))
+						{
+							FNiagaraAssetVersion AssetVersion = EmitterData->GetParent().GetEmitterData()->Version;
+							ParentEmitters.AddUnique(FString::Format(TEXT("{0}:{1}.{2}"), {GetPathNameSafe(EmitterData->GetParent().Emitter->GetPackage()), AssetVersion.MajorVersion, AssetVersion.MinorVersion}));
+						}
+					}
+				}
+			}
+		}
+		ParentEmitters.Sort();
+		Attributes.Emplace(TEXT("TotalEmitters"), EmitterHandles.Num());
+		Attributes.Emplace(TEXT("ActiveGpuEmitters"), GpuCount);
+		Attributes.Emplace(TEXT("StatelessEmitters"), StatelessEmitters);
+		Attributes.Emplace(TEXT("DisabledEmitters"), DisabledEmitters);
+		Attributes.Emplace(TEXT("RendererCount"), RendererCount);
+		Attributes.Emplace(TEXT("CustomRendererCount"), CustomRendererCount);
+		Attributes.Emplace(TEXT("DisabledRendererCount"), DisabledRendererCount);
+		Attributes.Emplace(TEXT("ScratchPadCount"), ScratchPadCount);
+		Attributes.Emplace(TEXT("ParentEmitterCount"), ParentEmitterCount);
+		Attributes.Emplace(TEXT("ParentEmitters"), ParentEmitters);
+		Attributes.Emplace(TEXT("UsesEffectType"), GetEffectType() != nullptr);
+		Attributes.Emplace(TEXT("OverrideScalabilitySettings"), GetOverrideScalabilitySettings());
+		Attributes.Emplace(TEXT("UsesWarmup"), NeedsWarmup());
+
+		// gather system script module data
+		UNiagaraScriptSourceBase* SystemScriptSource = SystemSpawnScript->GetLatestSource();
+		if (SystemSpawnScript && SystemScriptSource)
+		{
+			SystemScriptSource->ReportAnalyticsData(ScriptAnalytics);
+		}
+		TArray<FString> UsedModules = ScriptAnalytics.UsedNiagaraModules.Array();
+		UsedModules.Sort();
+		Attributes.Emplace(TEXT("Script.ActiveModules"), ScriptAnalytics.ActiveModules);
+		Attributes.Emplace(TEXT("Script.DisabledModules"), ScriptAnalytics.DisabledModules);
+		Attributes.Emplace(TEXT("Script.UsedNiagaraModules"), UsedModules);
+		
+		NiagaraAnalytics::RecordEvent(TEXT("System.Content"), Attributes);
+	}
+}
+
 bool UNiagaraSystem::UsesEmitter(UNiagaraEmitter* Emitter) const
 {
 	if (!Emitter)
@@ -895,14 +1027,6 @@ void UNiagaraSystem::PreEditChange(FProperty* PropertyThatWillChange)
 		UpdateContext.SetDestroyOnAdd(true);
 		UpdateContext.Add(this, false);
 	}
-
-	////-TODO:Stateless: Merge into emitter handle
-	//if (PropertyThatWillChange && PropertyThatWillChange->GetFName() == GET_MEMBER_NAME_CHECKED(UNiagaraSystem, StatelessEmitters))
-	//{
-	//	UpdateContext.SetDestroyOnAdd(true);
-	//	UpdateContext.Add(this, false);
-	//}
-	////-TODO:Stateless: Merge into emitter handle
 }
 
 void UNiagaraSystem::PostEditChangeProperty(struct FPropertyChangedEvent& PropertyChangedEvent)
@@ -1839,7 +1963,7 @@ void UNiagaraSystem::ComputeEmittersExecutionOrder()
 			continue;
 		}
 
-		EmitterDependencies.SetNum(0, EAllowShrinking::No);
+		EmitterDependencies.SetNum(0, false);
 
 		if (EmitterData->SimTarget == ENiagaraSimTarget::GPUComputeSim && EmitterData->GetGPUComputeScript())
 		{
@@ -1987,34 +2111,15 @@ void UNiagaraSystem::ComputeRenderersDrawOrder()
 			continue;
 		}
 
-		if (EmitterHandle.GetEmitterMode() == ENiagaraEmitterMode::Standard)
-		{
-			if (FVersionedNiagaraEmitterData* EmitterData = EmitterHandle.GetEmitterData())
-			{
-				EmitterData->ForEachEnabledRenderer(
-					[&](UNiagaraRendererProperties* Properties)
-					{
-						RendererSortInfo.Emplace(Properties->SortOrderHint, RendererSortInfo.Num());
-					}
-				);
-			}
-		}
-		//-TODO:Stateless:Abstract this away
-		else
+		if (FVersionedNiagaraEmitterData* EmitterData = EmitterHandle.GetEmitterData())
 		{
-			ensure(EmitterHandle.GetEmitterMode() == ENiagaraEmitterMode::Stateless);
-			const UNiagaraStatelessEmitter* StatelessEmitter = EmitterHandle.GetStatelessEmitter();
-			if (StatelessEmitter)
-			{
-				StatelessEmitter->ForEachEnabledRenderer(
-					[&](UNiagaraRendererProperties* Properties)
-					{
-						RendererSortInfo.Emplace(Properties->SortOrderHint, RendererSortInfo.Num());
-					}
-				);
-			}
+			EmitterData->ForEachEnabledRenderer(
+				[&](UNiagaraRendererProperties* Properties)
+				{
+					RendererSortInfo.Emplace(Properties->SortOrderHint, RendererSortInfo.Num());
+				}
+			);
 		}
-		//-TODO:Stateless:Abstract this away
 	}
 
 	// We sort by the sort hint in order to guarantee that we submit according to the preferred sort order..
@@ -2064,84 +2169,71 @@ void UNiagaraSystem::CacheFromCompiledData()
 			continue;
 		}
 
-		if (EmitterHandle.GetEmitterMode() == ENiagaraEmitterMode::Standard)
+		if (FVersionedNiagaraEmitterData* EmitterData = EmitterHandle.GetEmitterData())
 		{
-			if (FVersionedNiagaraEmitterData* EmitterData = EmitterHandle.GetEmitterData())
-			{
-				// Cache system instance accessors
-				ExecutionStateNameBuilder.Reset();
-				ExecutionStateNameBuilder << EmitterHandle.GetInstance().Emitter->GetUniqueEmitterName();
-				ExecutionStateNameBuilder << TEXT(".ExecutionState");
-				const FName ExecutionStateName(ExecutionStateNameBuilder);
+			// Cache system instance accessors
+			ExecutionStateNameBuilder.Reset();
+			ExecutionStateNameBuilder << EmitterHandle.GetInstance().Emitter->GetUniqueEmitterName();
+			ExecutionStateNameBuilder << TEXT(".ExecutionState");
+			const FName ExecutionStateName(ExecutionStateNameBuilder);
 
-				EmitterExecutionState.Init(SystemDataSet, ExecutionStateName);
+			EmitterExecutionState.Init(SystemDataSet, ExecutionStateName);
 
-				// Cache emitter data set accessors, for things like bounds, etc
-				const FNiagaraDataSetCompiledData* DataSetCompiledData = nullptr;
-				if (EmitterCompiledData.IsValidIndex(i))
+			// Cache emitter data set accessors, for things like bounds, etc
+			const FNiagaraDataSetCompiledData* DataSetCompiledData = nullptr;
+			if (EmitterCompiledData.IsValidIndex(i))
+			{
+				for (const FName& SpawnName : EmitterCompiledData[i]->SpawnAttributes)
 				{
-					for (const FName& SpawnName : EmitterCompiledData[i]->SpawnAttributes)
-					{
-						EmitterSpawnInfoAccessors[i].Emplace(SystemDataSet, SpawnName);
-					}
+					EmitterSpawnInfoAccessors[i].Emplace(SystemDataSet, SpawnName);
+				}
 
-					DataSetCompiledData = &EmitterCompiledData[i]->DataSetCompiledData;
+				DataSetCompiledData = &EmitterCompiledData[i]->DataSetCompiledData;
 
-					if (NiagaraSettings->bLimitDeltaTime)
-					{
-						MaxDeltaTime = MaxDeltaTime.IsSet() ? FMath::Min(MaxDeltaTime.GetValue(), NiagaraSettings->MaxDeltaTimePerTick) : NiagaraSettings->MaxDeltaTimePerTick;
-					}
+				if (NiagaraSettings->bLimitDeltaTime)
+				{
+					MaxDeltaTime = MaxDeltaTime.IsSet() ? FMath::Min(MaxDeltaTime.GetValue(), NiagaraSettings->MaxDeltaTimePerTick) : NiagaraSettings->MaxDeltaTimePerTick;
 				}
-				EmitterHandle.GetInstance().Emitter->ConditionalPostLoad();
-				EmitterData->CacheFromCompiledData(DataSetCompiledData, *EmitterHandle.GetInstance().Emitter);
+			}
+			EmitterHandle.GetInstance().Emitter->ConditionalPostLoad();
+			EmitterData->CacheFromCompiledData(DataSetCompiledData, *EmitterHandle.GetInstance().Emitter);
 
-				PSOPrecacheEvents.Append(EmitterData->PrecacheComputePSOs(*EmitterHandle.GetInstance().Emitter));
+			PSOPrecacheEvents.Append(EmitterData->PrecacheComputePSOs(*EmitterHandle.GetInstance().Emitter));
 
-				// Allow data interfaces to cache static buffers
-				UNiagaraScript* NiagaraEmitterScripts[] =
+			// Allow data interfaces to cache static buffers
+			UNiagaraScript* NiagaraEmitterScripts[] =
+			{
+				EmitterData->SpawnScriptProps.Script,
+				EmitterData->UpdateScriptProps.Script,
+				EmitterData->GetGPUComputeScript(),
+			};
+
+			for (UNiagaraScript* NiagaraScript : NiagaraEmitterScripts)
+			{
+				if (NiagaraScript == nullptr)
 				{
-					EmitterData->SpawnScriptProps.Script,
-					EmitterData->UpdateScriptProps.Script,
-					EmitterData->GetGPUComputeScript(),
-				};
+					continue;
+				}
 
-				for (UNiagaraScript* NiagaraScript : NiagaraEmitterScripts)
+				const bool bUsedByCPU = EmitterData->SimTarget == ENiagaraSimTarget::CPUSim;
+				const bool bUsedByGPU = EmitterData->SimTarget == ENiagaraSimTarget::GPUComputeSim;
+				for (const FNiagaraScriptResolvedDataInterfaceInfo& DataInterfaceInfo : NiagaraScript->GetResolvedDataInterfaces())
 				{
-					if (NiagaraScript == nullptr)
+					if (UNiagaraDataInterface* DataInterface = DataInterfaceInfo.ResolvedDataInterface)
 					{
-						continue;
-					}
+						DataInterface->CacheStaticBuffers(*StaticBuffers.Get(), DataInterfaceInfo.ResolvedVariable, bUsedByCPU, bUsedByGPU);
 
-					const bool bUsedByCPU = EmitterData->SimTarget == ENiagaraSimTarget::CPUSim;
-					const bool bUsedByGPU = EmitterData->SimTarget == ENiagaraSimTarget::GPUComputeSim;
-					for (const FNiagaraScriptResolvedDataInterfaceInfo& DataInterfaceInfo : NiagaraScript->GetResolvedDataInterfaces())
-					{
-						if (UNiagaraDataInterface* DataInterface = DataInterfaceInfo.ResolvedDataInterface)
+						if ( bUsedByGPU )
 						{
-							DataInterface->CacheStaticBuffers(*StaticBuffers.Get(), DataInterfaceInfo.ResolvedVariable, bUsedByCPU, bUsedByGPU);
-
-							if ( bUsedByGPU )
+							if (DataInterfaceInfo.bIsInternal == false)
 							{
-								if (DataInterfaceInfo.bIsInternal == false)
-								{
-									DataInterfaceGpuUsage.Add(DataInterfaceInfo.ParameterStoreVariable.GetName()); 
-								}
+								DataInterfaceGpuUsage.Add(DataInterfaceInfo.ParameterStoreVariable.GetName()); 
 							}
 						}
 					}
 				}
 			}
 		}
-		//-TODO:Stateless:Abstract this away
-		else
-		{
-			ensure(EmitterHandle.GetEmitterMode() == ENiagaraEmitterMode::Stateless);
-			if (UNiagaraStatelessEmitter* StatelessEmitter = EmitterHandle.GetStatelessEmitter())
-			{
-				StatelessEmitter->CacheFromCompiledData();
-			}
-		}
-		//-TODO:Stateless:Abstract this away
 	}
 
 	// If we had any precache events to wait on add a task to ensure the system does not start until they are ready
@@ -2469,25 +2561,13 @@ void UNiagaraSystem::AddEmitterHandleDirect(FNiagaraEmitterHandle& EmitterHandle
 
 FNiagaraEmitterHandle UNiagaraSystem::DuplicateEmitterHandle(const FNiagaraEmitterHandle& EmitterHandleToDuplicate, FName EmitterName)
 {
-	UNiagaraEmitter* DuplicateEmitter = EmitterHandleToDuplicate.GetInstance().Emitter ? UNiagaraEmitter::CreateAsDuplicate(*EmitterHandleToDuplicate.GetInstance().Emitter, EmitterName, *this) : nullptr;
-	UNiagaraStatelessEmitter* DuplicateStatelessEmitter = EmitterHandleToDuplicate.GetStatelessEmitter() ? EmitterHandleToDuplicate.GetStatelessEmitter()->CreateAsDuplicate(EmitterName, *this) : nullptr;
-	if (DuplicateEmitter)
-	{
-		EmitterHandles.Emplace_GetRef(*DuplicateEmitter, EmitterHandleToDuplicate.GetInstance().Version);
-	}
-	else
-	{
-		check(DuplicateStatelessEmitter);
-		EmitterHandles.Emplace(*DuplicateStatelessEmitter);
-	}
-
-	FNiagaraEmitterHandle& NewHandle = EmitterHandles.Last();
-	NewHandle.SetIsEnabled(EmitterHandleToDuplicate.GetIsEnabled(), *this, false);
-	NewHandle.SetStatelessEmitter(DuplicateStatelessEmitter);
-	NewHandle.SetEmitterMode(*this, EmitterHandleToDuplicate.GetEmitterMode());
-	RefreshSystemParametersFromEmitter(NewHandle);
+	UNiagaraEmitter* DuplicateEmitter = UNiagaraEmitter::CreateAsDuplicate(*EmitterHandleToDuplicate.GetInstance().Emitter, EmitterName, *this);
+	FNiagaraEmitterHandle EmitterHandle(*DuplicateEmitter, EmitterHandleToDuplicate.GetInstance().Version);
+	EmitterHandle.SetIsEnabled(EmitterHandleToDuplicate.GetIsEnabled(), *this, false);
+	EmitterHandles.Add(EmitterHandle);
+	RefreshSystemParametersFromEmitter(EmitterHandle);
 	InvalidateCachedData();
-	return NewHandle;
+	return EmitterHandle;
 }
 
 void UNiagaraSystem::RemoveEmitterHandle(const FNiagaraEmitterHandle& EmitterHandleToDelete)
@@ -3013,7 +3093,6 @@ bool UNiagaraSystem::QueryCompileComplete(bool bWait)
 	UpdateHasGPUEmitters();
 	UpdateDITickFlags();
 	ResolveScalabilitySettings();
-	ResolveRequiresScripts();
 
 	CurrentCompilation->ReportResults(Options);
 
@@ -3315,9 +3394,7 @@ void UNiagaraSystem::InitEmitterCompiledData()
 			const FNiagaraEmitterHandle& EmitterHandle = EmitterHandles[EmitterIdx];
 			const UNiagaraEmitter* Emitter = EmitterHandle.GetInstance().Emitter;
 			FNiagaraDataSetCompiledData& EmitterDataSetCompiledData = NewEmitterCompiledData[EmitterIdx]->DataSetCompiledData;
-			//-TODO:Stateless: Do we need stateless support here?
-			//if (ensureMsgf(Emitter != nullptr, TEXT("Failed to get Emitter Instance from Emitter Handle in post compile, please investigate.")))
-			if (Emitter != nullptr)
+			if (ensureMsgf(Emitter != nullptr, TEXT("Failed to get Emitter Instance from Emitter Handle in post compile, please investigate.")))
 			{
 				InitEmitterVariableAliasNames(NewEmitterCompiledData[EmitterIdx].Get(), Emitter);
 				InitEmitterDataSetCompiledData(EmitterDataSetCompiledData, EmitterHandle);			
@@ -3386,9 +3463,7 @@ void UNiagaraSystem::InitSystemCompiledData()
 	{
 		const FNiagaraEmitterHandle& PerEmitterHandle = EmitterHandles[EmitterIdx];
 		const UNiagaraEmitter* Emitter = PerEmitterHandle.GetInstance().Emitter;
-		//-TODO:Stateless: Do we need stateless support here?
-		//if (ensureMsgf(Emitter != nullptr, TEXT("Failed to get Emitter Instance from Emitter Handle when post compiling Niagara System %s!"), *GetPathNameSafe(this)))
-		if (Emitter != nullptr)
+		if (ensureMsgf(Emitter != nullptr, TEXT("Failed to get Emitter Instance from Emitter Handle when post compiling Niagara System %s!"), *GetPathNameSafe(this)))
 		{
 			const FString EmitterName = Emitter->GetUniqueEmitterName();
 
@@ -3637,23 +3712,6 @@ void UNiagaraSystem::GatherStaticVariables(TArray<FNiagaraVariable>& OutVars, TA
 }
 #endif
 
-void UNiagaraSystem::ResolveRequiresScripts()
-{
-#if WITH_EDITORONLY_DATA
-	TOptional<FNiagaraSystemStateData> NewSystemStateData;
-
-	if (bAllowSystemStateFastPath)
-	{
-		INiagaraModule& NiagaraModule = FModuleManager::GetModuleChecked<INiagaraModule>("Niagara");
-		const INiagaraEditorOnlyDataUtilities& EditorOnlyDataUtilities = NiagaraModule.GetEditorOnlyDataUtilities();
-		NewSystemStateData = EditorOnlyDataUtilities.TryGetSystemStateData(*this);
-	}
-
-	bSystemStateFastPathEnabled = NewSystemStateData.IsSet();
-	SystemStateData = NewSystemStateData.Get(FNiagaraSystemStateData());
-#endif //WITH_EDITORONLY_DATA
-}
-
 void UNiagaraSystem::ResolveScalabilitySettings()
 {
 	CurrentScalabilitySettings.Clear();
diff --git a/Engine/Plugins/FX/Niagara/Source/NiagaraEditor/Private/NiagaraBakerRendererOutputSimCache.cpp b/Engine/Plugins/FX/Niagara/Source/NiagaraEditor/Private/NiagaraBakerRendererOutputSimCache.cpp
--- a/Engine/Plugins/FX/Niagara/Source/NiagaraEditor/Private/NiagaraBakerRendererOutputSimCache.cpp
+++ b/Engine/Plugins/FX/Niagara/Source/NiagaraEditor/Private/NiagaraBakerRendererOutputSimCache.cpp
@@ -91,7 +91,7 @@ void FNiagaraBakerRendererOutputSimCache::EndBake(FNiagaraBakerFeedbackContext&
 		return;
 	}
 
-	BakeSimCache->EndWrite();
+	BakeSimCache->EndWrite(true);
 
 	BakeSimCache->RemoveFromRoot();
 }
diff --git a/Engine/Plugins/FX/Niagara/Source/NiagaraEditor/Private/NiagaraScriptSource.cpp b/Engine/Plugins/FX/Niagara/Source/NiagaraEditor/Private/NiagaraScriptSource.cpp
--- a/Engine/Plugins/FX/Niagara/Source/NiagaraEditor/Private/NiagaraScriptSource.cpp
+++ b/Engine/Plugins/FX/Niagara/Source/NiagaraEditor/Private/NiagaraScriptSource.cpp
@@ -5,6 +5,7 @@
 #include "EdGraphUtilities.h"
 #include "GraphEditAction.h"
 #include "INiagaraEditorTypeUtilities.h"
+#include "NiagaraAnalytics.h"
 #include "NiagaraCommon.h"
 #include "NiagaraComponent.h"
 #include "NiagaraConstants.h"
@@ -99,6 +100,42 @@ void UNiagaraScriptSource::RefreshGraphCompileId()
 	}
 }
 
+void UNiagaraScriptSource::ReportAnalyticsData(FNiagaraScriptSourceAnalytics& InData) const
+{
+	static const UEnum* EnumType = StaticEnum<ENiagaraScriptUsage>();
+
+	TArray<UNiagaraNodeOutput*> OutputNodes;
+	NodeGraph->GetNodesOfClass<UNiagaraNodeOutput>(OutputNodes);	
+	for (UNiagaraNodeOutput* OutputNode : OutputNodes)
+	{
+		FString UsageString = TEXT("Script.") + EnumType->GetNameStringByValue(static_cast<uint64>(OutputNode->GetUsage()));
+
+		// gather module data
+		TArray<UNiagaraNodeFunctionCall*> ModuleNodes;
+		FNiagaraStackGraphUtilities::GetOrderedModuleNodes(*OutputNode, ModuleNodes);
+
+		for (UNiagaraNodeFunctionCall* Node : ModuleNodes)
+		{
+			if (!Node->FunctionScript || !Node->GetScriptData())
+			{
+				continue;
+			}
+			
+			if (!Node->IsNodeEnabled())
+			{
+				InData.DisabledModules++;
+				continue;
+			}
+			InData.ActiveModules++;
+			if (NiagaraAnalytics::IsPluginAsset(Node->FunctionScript))
+			{
+				FNiagaraAssetVersion AssetVersion = Node->GetScriptData()->Version;
+				InData.UsedNiagaraModules.Add(FString::Format(TEXT("{0}:{1}.{2}"), {GetPathNameSafe(Node->FunctionScript->GetPackage()), AssetVersion.MajorVersion, AssetVersion.MinorVersion}));
+			}
+		}
+	}
+}
+
 FGuid UNiagaraScriptSource::GetCompileBaseId(ENiagaraScriptUsage InUsage, const FGuid& InUsageId) const
 {
 	return NodeGraph->GetBaseId(InUsage, InUsageId);
diff --git a/Engine/Plugins/FX/Niagara/Source/NiagaraEditor/Private/NiagaraSystemFactoryNew.cpp b/Engine/Plugins/FX/Niagara/Source/NiagaraEditor/Private/NiagaraSystemFactoryNew.cpp
--- a/Engine/Plugins/FX/Niagara/Source/NiagaraEditor/Private/NiagaraSystemFactoryNew.cpp
+++ b/Engine/Plugins/FX/Niagara/Source/NiagaraEditor/Private/NiagaraSystemFactoryNew.cpp
@@ -9,6 +9,7 @@
 #include "AssetRegistry/AssetData.h"
 #include "ContentBrowserModule.h"
 #include "ImageUtils.h"
+#include "NiagaraAnalytics.h"
 #include "ViewModels/Stack/NiagaraStackGraphUtilities.h"
 #include "Misc/MessageDialog.h"
 #include "ViewModels/NiagaraSystemViewModel.h"
diff --git a/Engine/Plugins/FX/Niagara/Source/NiagaraEditor/Private/Toolkits/NiagaraSystemToolkit.cpp b/Engine/Plugins/FX/Niagara/Source/NiagaraEditor/Private/Toolkits/NiagaraSystemToolkit.cpp
--- a/Engine/Plugins/FX/Niagara/Source/NiagaraEditor/Private/Toolkits/NiagaraSystemToolkit.cpp
+++ b/Engine/Plugins/FX/Niagara/Source/NiagaraEditor/Private/Toolkits/NiagaraSystemToolkit.cpp
@@ -1328,6 +1328,7 @@ bool FNiagaraSystemToolkit::OnRequestClose(EAssetEditorCloseReason InCloseReason
 	}
 
 	SystemViewModel->NotifyPreClose();
+	SystemViewModel->GetSystem().ReportAnalyticsData(false);
 
 	bool bHasUnappliedScratchPadChanges = false;
 	for (TSharedRef<FNiagaraScratchPadScriptViewModel> ScratchPadViewModel : SystemViewModel->GetScriptScratchPadViewModel()->GetScriptViewModels())
diff --git a/Engine/Plugins/FX/Niagara/Source/NiagaraEditor/Public/NiagaraScriptSource.h b/Engine/Plugins/FX/Niagara/Source/NiagaraEditor/Public/NiagaraScriptSource.h
--- a/Engine/Plugins/FX/Niagara/Source/NiagaraEditor/Public/NiagaraScriptSource.h
+++ b/Engine/Plugins/FX/Niagara/Source/NiagaraEditor/Public/NiagaraScriptSource.h
@@ -15,7 +15,7 @@ struct FNiagaraScriptHashCollector;
 struct FNiagaraVMExecutableDataId;
 
 UCLASS(MinimalAPI)
-class UNiagaraScriptSource : public UNiagaraScriptSourceBase
+class UNiagaraScriptSource final : public UNiagaraScriptSourceBase
 {
 	GENERATED_UCLASS_BODY()
 
@@ -46,6 +46,8 @@ class UNiagaraScriptSource : public UNiagaraScriptSourceBase
 
 	// Will conditionally refresh the graph's CompileId
 	virtual void RefreshGraphCompileId() override;
+	
+	virtual void ReportAnalyticsData(FNiagaraScriptSourceAnalytics& InData) const override;
 
 	virtual FGuid GetCompileBaseId(ENiagaraScriptUsage InUsage, const FGuid& InUsageId) const override;
 
diff --git a/Engine/Plugins/FX/NiagaraSimCaching/Source/NiagaraSimCachingEditor/Private/Sequencer/MovieSceneNiagaraTrackRecorder.cpp b/Engine/Plugins/FX/NiagaraSimCaching/Source/NiagaraSimCachingEditor/Private/Sequencer/MovieSceneNiagaraTrackRecorder.cpp
--- a/Engine/Plugins/FX/NiagaraSimCaching/Source/NiagaraSimCachingEditor/Private/Sequencer/MovieSceneNiagaraTrackRecorder.cpp
+++ b/Engine/Plugins/FX/NiagaraSimCaching/Source/NiagaraSimCachingEditor/Private/Sequencer/MovieSceneNiagaraTrackRecorder.cpp
@@ -201,7 +201,7 @@ void UMovieSceneNiagaraTrackRecorder::FinalizeTrackImpl()
 	if (NiagaraCacheSection.IsValid())
 	{
 		// finalize the sim cache
-		NiagaraCacheSection->Params.SimCache->EndWrite();
+		NiagaraCacheSection->Params.SimCache->EndWrite(true);
 		
 		// Activate the section
 		NiagaraCacheSection->SetIsActive(true);
