Fixed missing loctext keys

diff --git a/Engine/Plugins/FX/Niagara/Source/NiagaraEditor/Private/Widgets/AssetBrowser/SNiagaraAssetBrowser.cpp b/Engine/Plugins/FX/Niagara/Source/NiagaraEditor/Private/Widgets/AssetBrowser/SNiagaraAssetBrowser.cpp
--- a/Engine/Plugins/FX/Niagara/Source/NiagaraEditor/Private/Widgets/AssetBrowser/SNiagaraAssetBrowser.cpp
+++ b/Engine/Plugins/FX/Niagara/Source/NiagaraEditor/Private/Widgets/AssetBrowser/SNiagaraAssetBrowser.cpp
@@ -650,7 +650,7 @@ TArray<TSharedRef<FFrontendFilter>> SNiagaraAssetBrowser::OnGetExtraFrontendFilt
 	{
 		using namespace FNiagaraEditorUtilities::AssetBrowser;
 		
-		TSharedRef<FFrontendFilterCategory> NiagaraTagFiltersCategory = MakeShared<FFrontendFilterCategory>(LOCTEXT("NiagaraTagFilterCategoryLabel", "Niagara Tags"), LOCTEXT("", ""));
+		TSharedRef<FFrontendFilterCategory> NiagaraTagFiltersCategory = MakeShared<FFrontendFilterCategory>(LOCTEXT("NiagaraTagFilterCategoryLabel", "Niagara Tags"), LOCTEXT("NiagaraTagFiltersTooltip", "Secondary Niagara Tags used for filtering"));
 
 		// This requires the asset registry to be done loading
 		for(const FNiagaraAssetTagDefinition& AssetTagDefinition : GetFlatSortedAssetTagDefinitions(false))
@@ -665,7 +665,7 @@ TArray<TSharedRef<FFrontendFilter>> SNiagaraAssetBrowser::OnGetExtraFrontendFilt
 	}
 	
 	{
-		TSharedRef<FFrontendFilterCategory> NiagaraAdditionalFiltersCategory = MakeShared<FFrontendFilterCategory>(LOCTEXT("NiagaraPropertyFilterCategoryLabel", "Niagara Filters"), LOCTEXT("", ""));
+		TSharedRef<FFrontendFilterCategory> NiagaraAdditionalFiltersCategory = MakeShared<FFrontendFilterCategory>(LOCTEXT("NiagaraPropertyFilterCategoryLabel", "Niagara Filters"), LOCTEXT("NiagaraAdditionalFiltersTooltip", "Additional filters for filtering Niagara assets"));
 		if(AvailableClasses.Contains(UNiagaraEmitter::StaticClass()))
 		{
 			Result.Add(MakeShared<FFrontendFilter_NiagaraEmitterInheritance>(true, NiagaraAdditionalFiltersCategory));
