Fixed lack of attaching the updated content bar to its parent slot..
Fixes missing renderer thumbnails in niagara node

diff --git a/Engine/Plugins/FX/Niagara/Source/NiagaraEditorWidgets/Private/SNiagaraOverviewStackNode.cpp b/Engine/Plugins/FX/Niagara/Source/NiagaraEditorWidgets/Private/SNiagaraOverviewStackNode.cpp
--- a/Engine/Plugins/FX/Niagara/Source/NiagaraEditorWidgets/Private/SNiagaraOverviewStackNode.cpp
+++ b/Engine/Plugins/FX/Niagara/Source/NiagaraEditorWidgets/Private/SNiagaraOverviewStackNode.cpp
@@ -419,9 +419,9 @@ void SNiagaraOverviewStackNode::Tick(const FGeometry& AllottedGeometry, const do
 			OverviewStackNode->RenameStarted();
 		}
 
-		if (bTopContentBarRefreshPending)
+		if (bTopContentBarRefreshPending && TopContentBarSlot != nullptr)
 		{
-			CreateTopContentBar();
+			TopContentBarSlot->AttachWidget(CreateTopContentBar());
 			bTopContentBarRefreshPending = false;
 		}
 
@@ -610,6 +610,7 @@ TSharedRef<SWidget> SNiagaraOverviewStackNode::CreateNodeContentArea_Default()
 			.HAlign(HAlign_Fill)
 			.VAlign(VAlign_Center)
 			.Padding(2.f)
+			.Expose(TopContentBarSlot)
 			[
 				CreateTopContentBar()
 			]
diff --git a/Engine/Plugins/FX/Niagara/Source/NiagaraEditorWidgets/Private/SNiagaraOverviewStackNode.h b/Engine/Plugins/FX/Niagara/Source/NiagaraEditorWidgets/Private/SNiagaraOverviewStackNode.h
--- a/Engine/Plugins/FX/Niagara/Source/NiagaraEditorWidgets/Private/SNiagaraOverviewStackNode.h
+++ b/Engine/Plugins/FX/Niagara/Source/NiagaraEditorWidgets/Private/SNiagaraOverviewStackNode.h
@@ -153,4 +153,6 @@ private:
 
 	TSharedPtr<SBox> ThumbnailContainer;
 	TSharedPtr<FAssetThumbnail> PreviewThumbnail;
+	
+	SVerticalBox::FSlot* TopContentBarSlot = nullptr;
 };
\ No newline at end of file
