Added safety nullptr check for GEditor

diff --git a/Engine/Plugins/FX/Niagara/Source/NiagaraEditor/Public/NiagaraRecentAndFavoritesManager.cpp b/Engine/Plugins/FX/Niagara/Source/NiagaraEditor/Public/NiagaraRecentAndFavoritesManager.cpp
--- a/Engine/Plugins/FX/Niagara/Source/NiagaraEditor/Public/NiagaraRecentAndFavoritesManager.cpp
+++ b/Engine/Plugins/FX/Niagara/Source/NiagaraEditor/Public/NiagaraRecentAndFavoritesManager.cpp
@@ -96,11 +96,15 @@ void FNiagaraRecentAndFavoritesManager::OnPluginLoadingPhaseComplete(ELoadingPha
 {
 	if(LoadingPhase == ELoadingPhase::PostEngineInit)
 	{
-		GEditor->GetEditorSubsystem<UAssetEditorSubsystem>()->OnAssetOpenedInEditor().AddSP(this, &FNiagaraRecentAndFavoritesManager::OnAssetOpenedInEditor);
-		IPluginManager& PluginManager = IPluginManager::Get();
-		PluginManager.OnLoadingPhaseComplete().RemoveAll(this);
+		// GEditor might not always exist
+		if(GEditor)
+		{
+			GEditor->GetEditorSubsystem<UAssetEditorSubsystem>()->OnAssetOpenedInEditor().AddSP(this, &FNiagaraRecentAndFavoritesManager::OnAssetOpenedInEditor);
+			IPluginManager& PluginManager = IPluginManager::Get();
+			PluginManager.OnLoadingPhaseComplete().RemoveAll(this);
 
-		GEditor->OnEditorClose().AddSP(this, &FNiagaraRecentAndFavoritesManager::OnEditorClose);
+			GEditor->OnEditorClose().AddSP(this, &FNiagaraRecentAndFavoritesManager::OnEditorClose);
+		}
 	}
 }
 
