Fixed static analysis warning regarding virtual methods

diff --git a/Engine/Plugins/FX/Niagara/Source/NiagaraEditor/Public/NiagaraScriptSource.h b/Engine/Plugins/FX/Niagara/Source/NiagaraEditor/Public/NiagaraScriptSource.h
--- a/Engine/Plugins/FX/Niagara/Source/NiagaraEditor/Public/NiagaraScriptSource.h
+++ b/Engine/Plugins/FX/Niagara/Source/NiagaraEditor/Public/NiagaraScriptSource.h
@@ -59,8 +59,9 @@ class UNiagaraScriptSource final : public UNiagaraScriptSourceBase
 
 	NIAGARAEDITOR_API virtual bool AddModuleIfMissing(FString ModulePath, ENiagaraScriptUsage Usage, bool& bOutFoundModule)override;
 
-	virtual void FixupRenamedParameters(UNiagaraNodeFunctionCall* FunctionCallNode, TConstArrayView<FNiagaraVariable> ModuleInputVariables, FNiagaraParameterStore& RapidIterationParameters, const TArray<FNiagaraVariable>& OldRapidIterationVariables, const FVersionedNiagaraEmitter& Emitter, ENiagaraScriptUsage ScriptUsage) const;
-	virtual void InitializeNewParameters(UNiagaraNodeFunctionCall* FunctionCallNode, TConstArrayView<FNiagaraVariable> ModuleInputVariables, FNiagaraParameterStore& RapidIterationParameters, const FVersionedNiagaraEmitter& VersionedEmitter, ENiagaraScriptUsage ScriptUsage, TSet<FNiagaraVariableBase>& ValidRapidIterationParameters) const;
+
+	void FixupRenamedParameters(UNiagaraNodeFunctionCall* FunctionCallNode, TConstArrayView<FNiagaraVariable> ModuleInputVariables, FNiagaraParameterStore& RapidIterationParameters, const TArray<FNiagaraVariable>& OldRapidIterationVariables, const FVersionedNiagaraEmitter& Emitter, ENiagaraScriptUsage ScriptUsage) const;
+	void InitializeNewParameters(UNiagaraNodeFunctionCall* FunctionCallNode, TConstArrayView<FNiagaraVariable> ModuleInputVariables, FNiagaraParameterStore& RapidIterationParameters, const FVersionedNiagaraEmitter& VersionedEmitter, ENiagaraScriptUsage ScriptUsage, TSet<FNiagaraVariableBase>& ValidRapidIterationParameters) const;
 
 	virtual void CleanUpOldAndInitializeNewRapidIterationParameters(const FVersionedNiagaraEmitter& Emitter, ENiagaraScriptUsage ScriptUsage, FGuid ScriptUsageId, FNiagaraParameterStore& RapidIterationParameters) const override;
 	virtual void ForceGraphToRecompileOnNextCheck() override;
