Niagara - Fix renderer entries not refreshing properly when added and deleted.

diff --git a/Engine/Plugins/FX/Niagara/Source/NiagaraEditor/Private/ViewModels/Stack/NiagaraStackRenderItemGroup.cpp b/Engine/Plugins/FX/Niagara/Source/NiagaraEditor/Private/ViewModels/Stack/NiagaraStackRenderItemGroup.cpp
--- a/Engine/Plugins/FX/Niagara/Source/NiagaraEditor/Private/ViewModels/Stack/NiagaraStackRenderItemGroup.cpp
+++ b/Engine/Plugins/FX/Niagara/Source/NiagaraEditor/Private/ViewModels/Stack/NiagaraStackRenderItemGroup.cpp
@@ -111,6 +111,7 @@ void UNiagaraStackRenderItemGroup::Initialize(FRequiredEntryData InRequiredEntry
 	FText DisplayName = LOCTEXT("RenderGroupName", "Render");
 	FText ToolTip = LOCTEXT("RendererGroupTooltip", "Describes how we should display/present each particle. Note that this doesn't have to be visual. Multiple renderers are supported. Order in this stack is not necessarily relevant to draw order.");
 	RenderersOwner = InRenderersOwner;
+	RenderersOwner->OnRenderersChanged().BindUObject(this, &UNiagaraStackRenderItemGroup::OwnerRenderersChanged);
 	AddUtilities = MakeShared<FRenderItemGroupAddUtilities>(RenderersOwner, FRenderItemGroupAddUtilities::FOnItemAdded::CreateUObject(this, &UNiagaraStackRenderItemGroup::OnRendererAdded));
 	Super::Initialize(InRequiredEntryData, DisplayName, ToolTip, AddUtilities.Get());
 }
@@ -199,7 +200,7 @@ void UNiagaraStackRenderItemGroup::RefreshChildrenInternal(const TArray<UNiagara
 	Super::RefreshChildrenInternal(CurrentChildren, NewChildren, NewIssues);
 }
 
-void UNiagaraStackRenderItemGroup::EmitterRenderersChanged()
+void UNiagaraStackRenderItemGroup::OwnerRenderersChanged()
 {
 	if (IsFinalized() == false)
 	{
@@ -228,7 +229,11 @@ void UNiagaraStackRenderItemGroup::OnRendererAdded(UNiagaraRendererProperties* R
 
 void UNiagaraStackRenderItemGroup::FinalizeInternal()
 {
-	RenderersOwner.Reset();
+	if (RenderersOwner.IsValid())
+	{
+		RenderersOwner->OnRenderersChanged().Unbind();
+		RenderersOwner.Reset();
+	}
 	Super::FinalizeInternal();
 }
 
diff --git a/Engine/Plugins/FX/Niagara/Source/NiagaraEditor/Public/ViewModels/Stack/NiagaraStackRenderItemGroup.h b/Engine/Plugins/FX/Niagara/Source/NiagaraEditor/Public/ViewModels/Stack/NiagaraStackRenderItemGroup.h
--- a/Engine/Plugins/FX/Niagara/Source/NiagaraEditor/Public/ViewModels/Stack/NiagaraStackRenderItemGroup.h
+++ b/Engine/Plugins/FX/Niagara/Source/NiagaraEditor/Public/ViewModels/Stack/NiagaraStackRenderItemGroup.h
@@ -35,7 +35,7 @@ protected:
 	NIAGARAEDITOR_API virtual void RefreshChildrenInternal(const TArray<UNiagaraStackEntry*>& CurrentChildren, TArray<UNiagaraStackEntry*>& NewChildren, TArray<FStackIssue>& NewIssues) override;
 	NIAGARAEDITOR_API virtual void FinalizeInternal() override;
 private:
-	NIAGARAEDITOR_API void EmitterRenderersChanged();
+	NIAGARAEDITOR_API void OwnerRenderersChanged();
 
 	NIAGARAEDITOR_API bool ChildRequestCanPaste(const UNiagaraClipboardContent* ClipboardContent,FText& OutCanPasteMessage);
 	NIAGARAEDITOR_API void ChildRequestPaste(const UNiagaraClipboardContent* ClipboardContent, int32 PasteIndex, FText& OutPasteWarning);
