New method for allowing particles spawned from NDC to access the NDC element that generated them.

Work using an indirection table that allocates out power of 2 chunks of instances to each NDC item according to how many particles each spawns.
This is a bit convoluted but provides a good balance of minimal size of the table and allowing both high particle counts per NDC and high NDC counts spawning few particles.
Method is described in greater detail in code comments.

New function GetNDCSpawnData(). Returns the NDC index for a particle as well as the total count for that NDC and the spawn index for the current particle inside that NDC spawn. i.e NDC local exec index.

Previous SpawnGroup based method still works as a backup where needed.
There may be cases where the new method cannot work. For example if a proper ExecIndex per NDC is needed.

Misc:
- Fixed bug in translator that meant emitter.ID would not be set for GPU scripts.

diff --git a/Engine/Plugins/FX/Niagara/Shaders/Private/DataChannel/NiagaraDataInterfaceDataChannelTemplateCommon.ush b/Engine/Plugins/FX/Niagara/Shaders/Private/DataChannel/NiagaraDataInterfaceDataChannelTemplateCommon.ush
--- a/Engine/Plugins/FX/Niagara/Shaders/Private/DataChannel/NiagaraDataInterfaceDataChannelTemplateCommon.ush
+++ b/Engine/Plugins/FX/Niagara/Shaders/Private/DataChannel/NiagaraDataInterfaceDataChannelTemplateCommon.ush
@@ -21,6 +21,131 @@ int {ParameterName}_FloatStride;
 int {ParameterName}_Int32Stride;
 int {ParameterName}_HalfStride;
 
+//Buffer containing back mapping from spawn particle exec index to the NDC index that spawned that particle.
+//Arraged into power of two buckets so that we can avoid an entry per particle.
+Buffer<int> {ParameterName}_NDCSpawnDataBuffer;
+
+void Num_{ParameterName}(out int OutNum)
+{
+	OutNum = {ParameterName}_Num;
+}
+
+void GetNDCSpawnData_{ParameterName}(NiagaraEmitterID InEmitterID, int InExecIndex, out int OutNDCIndex, out int OutNDCSpawnIndex, out int OutNDCSpawnCount)
+{
+	/*
+	Look up the various bits of spawn data from the NDC that spawned this particle.
+	The first section of the NDCSpawnDataBuffer contains a SpawnCount for every NDC Item from which we can calcualte NDCSpawnCount.
+	The next section is a series of buckets containing all NDC indices that contain enough particles to spawn 1 << Bucket particles.
+	ie. if an NDC spawns 10 particles it will have an entry in the 8 and 2 buckets and none in any other bucket.
+	From this we can map the main ExecIndex back into an NDCIndex that generated it. However ofc this mapping is by bucket and so particles from one NDC may not be contiuous in the spawn script execution/particle data.
+	From the per NDCSpawnCount we can also determine if an NDC has an entry in each bucket and so calculate the SpawnIndex of each particle within the NDC Spawn.
+
+	This is a little convoluted but means we can access the spawning NDC Index when spawning in a single large SpawnInfo/Execution.
+	This power of 2 bucket system also means we can keep the indirection table size reasonable for large and small particle counts in exchange for some ALU work.
+
+	|| NDCSpawnDataBuffer |
+	|| Bucket Sizes	                                           || Spawn Counts for each NDC         || 2^15 Bucket   || 2^14 Bucket || ... || 2^1 Bucket || 2^0 Bucket   ||
+	|| Bucket 0 Size | Bucket 1 Size | ... | Bucket 15 Size    || NDC 2 | ... | NDC N               || 0 | 0 | 5 | 7 || 0 |.........|| ... || 0 | 1 | 5  ||              || 
+	
+	Above shows the layout for NDCSpawnDataBuffer.
+	All Emitters are packed into the same buffer. The first part of the buffer is a lookup table for each emitter (by emitter ID/Index) into the rest of the data.
+
+	As we proceed through each ExecIndex in the spawn execution, we iterate through each bucket and can know for example, we can allocate the first 2^15 particles to NDC 0 as that has an entry in the 2^15 bucket.
+	The 2^15 bucket is special in that it can have multiple entries for the same NDC as it's the final bucket. All other buckets have only one as higher particle counts just move the NDC into other buckets. 
+	Less fills up lower buckets, more with fill up this and lower buckets. 2x will push it into the bucket above. Except in the final bucket if we hit 2x or more 2^15 then we just have to add another entry to that bucket.
+	So in the above example the first and second 2^15 instances of the execution are given to NDC 0.
+	The next 2^15 are given to 5, then 7. Now we've processed the 2^15 bucket, further instances take from the next bucket and so on until we're at the 2^0 bucket and we've handled all instances.
+	*/
+
+	OutNDCIndex = -1;
+	OutNDCSpawnIndex = -1;
+	OutNDCSpawnCount = -1;
+
+	uint NumNDCEntries = {ParameterName}_Num;
+
+	//First part of the buffer is the per emitter offsets into the data.
+	//Grab the right base offset for this emitter.
+	uint NDCSpawnDataOffset = { ParameterName }_NDCSpawnDataBuffer[InEmitterID.ID];
+	if(NDCSpawnDataOffset == 0 || NumNDCEntries == 0)
+	{
+		OutNDCIndex = NDCSpawnDataOffset;
+		OutNDCSpawnIndex = InEmitterID.ID;
+		OutNDCSpawnCount = NumNDCEntries;
+
+		return;
+	}
+
+	uint NDCSpawnData_BucketSizeStart = NDCSpawnDataOffset;
+	uint NDCSpawnData_NDCCountsStart = NDCSpawnData_BucketSizeStart + 16;
+	uint NDCSPawnData_NDCIndexBucketsStart = NDCSpawnData_NDCCountsStart + NumNDCEntries;
+	
+	uint MaxBucketExecIndex = 0;
+	uint BucketEntryStart = NDCSPawnData_NDCIndexBucketsStart;
+
+	[UNROLL]
+	for(int BucketIdx = 0; BucketIdx < 16; ++BucketIdx)
+	{
+		uint BucketSize = (1U << 15U) >> BucketIdx;
+		uint NumEntriesInBucket = {ParameterName}_NDCSpawnDataBuffer[NDCSpawnData_BucketSizeStart + BucketIdx];
+		uint MinBucketExecIndex = MaxBucketExecIndex; 
+		MaxBucketExecIndex += BucketSize * NumEntriesInBucket;
+		
+		if(InExecIndex < MaxBucketExecIndex)
+		{
+			//We found our bucket.
+			//Now we need to find our NDCIndex Entry Index.
+			uint NDCIndexEntry = (InExecIndex - MinBucketExecIndex) / BucketSize;			
+			OutNDCIndex = {ParameterName}_NDCSpawnDataBuffer[BucketEntryStart + NDCIndexEntry];
+			break;
+		}
+		
+		BucketEntryStart += NumEntriesInBucket;
+	}
+	
+	if(OutNDCIndex >= 0 && OutNDCIndex < (int)NumNDCEntries)
+	{
+		//Pull the NDC Spawn count from the first part of the buffer.
+		OutNDCSpawnCount = {ParameterName}_NDCSpawnDataBuffer[NDCSpawnData_NDCCountsStart + OutNDCIndex];
+
+		//Do another pass to calculate our SpawnIndex for this NDC within the total count for this NDC.
+		OutNDCSpawnIndex = 0;
+		MaxBucketExecIndex = 0;
+		uint Count = OutNDCSpawnCount;
+		[UNROLL]
+		for(int BucketIdx = 0; BucketIdx < 16; ++BucketIdx)
+		{
+			uint BucketSize = (1U << 15U) >> BucketIdx;
+			uint Mask = (0xFFFFU >> (BucketIdx + 1U));
+			uint CountMasked = Count & ~Mask;
+			Count &= Mask;
+			uint NumNDCEntriesInBucket = CountMasked >> (15U - BucketIdx);
+			uint NumEntriesInBucket = {ParameterName}_NDCSpawnDataBuffer[NDCSpawnData_BucketSizeStart + BucketIdx];
+			uint NumNDCInstancesInBucket = NumNDCEntriesInBucket * BucketSize;
+
+			uint MinBucketExecIndex = MaxBucketExecIndex;
+			MaxBucketExecIndex += BucketSize * NumEntriesInBucket;
+			if (InExecIndex < MaxBucketExecIndex && NumNDCInstancesInBucket > 0)
+			{
+				uint NDCIndexEntry = (InExecIndex - MinBucketExecIndex) >> (15U - BucketIdx);
+				uint MinNDCBucketExecIndex = MinBucketExecIndex + (BucketSize * NDCIndexEntry);
+
+				OutNDCSpawnIndex += (InExecIndex - MinNDCBucketExecIndex);
+				break;
+			}
+			else
+			{
+				OutNDCSpawnIndex += NumNDCInstancesInBucket;
+			}
+		}
+	}
+	else
+	{
+		OutNDCIndex = -1;
+		OutNDCSpawnIndex = -1;
+		OutNDCSpawnCount = -1;
+	}
+}
+
 // Outputs the starting offset of each component in the channel data for a partiuclar parameter.
 bool GetParameterLayoutInfo_{ParameterName}(uint ParameterIndex, out uint OutFloatRegister, out uint OutInt32Register, out uint OutHalfRegister)
 {
diff --git a/Engine/Plugins/FX/Niagara/Source/Niagara/Classes/NiagaraEmitterInstance.h b/Engine/Plugins/FX/Niagara/Source/Niagara/Classes/NiagaraEmitterInstance.h
--- a/Engine/Plugins/FX/Niagara/Source/Niagara/Classes/NiagaraEmitterInstance.h
+++ b/Engine/Plugins/FX/Niagara/Source/Niagara/Classes/NiagaraEmitterInstance.h
@@ -94,6 +94,9 @@ public:
 
 	[[nodiscard]] NIAGARA_API UObject* FindBinding(const FNiagaraVariable& InVariable) const;
 	[[nodiscard]] NIAGARA_API UNiagaraDataInterface* FindDataInterface(const FNiagaraVariable& InVariable) const;
+
+	[[nodiscard]] NIAGARA_API FNiagaraEmitterID GetEmitterID() const{ return FNiagaraEmitterID(EmitterIndex); }
+
 	//////////////////////////////////////////////////////////////////////////////////////////////////////////////////
 
 	////////////////////////////////////////////////////////////////////////////////////////////////////
diff --git a/Engine/Plugins/FX/Niagara/Source/Niagara/Internal/DataInterface/NiagaraDataInterfaceDataChannelCommon.h b/Engine/Plugins/FX/Niagara/Source/Niagara/Internal/DataInterface/NiagaraDataInterfaceDataChannelCommon.h
--- a/Engine/Plugins/FX/Niagara/Source/Niagara/Internal/DataInterface/NiagaraDataInterfaceDataChannelCommon.h
+++ b/Engine/Plugins/FX/Niagara/Source/Niagara/Internal/DataInterface/NiagaraDataInterfaceDataChannelCommon.h
@@ -166,6 +166,7 @@ struct FNDIDataChannelCompiledData
 
 	bool UsedByCPU()const{ return bUsedByCPU; }
 	bool UsedByGPU()const{ return bUsedByGPU; }
+	bool NeedSpawnDataTable()const { return bNeedsSpawnDataTable; }
 	int32 GetTotalParams()const{ return TotalParams; }
 
 protected:
@@ -193,6 +194,9 @@ protected:
 
 	UPROPERTY()
 	bool bUsedByGPU = false;
+	
+	UPROPERTY()
+	bool bNeedsSpawnDataTable = true;
 
 	/** Iterates over all scripts for the owning system and gathers all functions and parameters accessing this DI. Building the FunctionInfoTable and GPUScriptParameterInfos map.  */
 	void GatherAccessInfo(UNiagaraSystem* System, UNiagaraDataInterface* Owner);
@@ -201,7 +205,9 @@ protected:
 
 namespace NDIDataChannelUtilities
 {
-	void SortParameters(TArray<FNiagaraVariableBase>& Parameters);	
+	extern const FName GetNDCSpawnDataName;
+
+	void SortParameters(TArray<FNiagaraVariableBase>& Parameters);
 }
 
 
diff --git a/Engine/Plugins/FX/Niagara/Source/Niagara/Internal/DataInterface/NiagaraDataInterfaceDataChannelRead.h b/Engine/Plugins/FX/Niagara/Source/Niagara/Internal/DataInterface/NiagaraDataInterfaceDataChannelRead.h
--- a/Engine/Plugins/FX/Niagara/Source/Niagara/Internal/DataInterface/NiagaraDataInterfaceDataChannelRead.h
+++ b/Engine/Plugins/FX/Niagara/Source/Niagara/Internal/DataInterface/NiagaraDataInterfaceDataChannelRead.h
@@ -90,7 +90,15 @@ public:
 	bool bUpdateSourceDataEveryTick = true;
 
 
-	/** When true, Emitter.Spawn group for any spawned particles will be overridden to the index of the data channel element that generated that spawn. This allows particles to read further data from the data channel element that generated them. */
+	/**
+	When true, Emitter.Spawn group for any spawned particles will be overridden to the index of the data channel element that generated that spawn.
+	Doing this will submit all NDC spawns individually and will be less performant.
+	However it will allow particles to access the NDC data that generated then via the SpawnGroup value.
+	It will also mean that Exec Index will be correct on a per NDC Entry level. 
+	Without this settings ExecIndex will be 0...TotalSpawnCount-1. With this it will be 0...SpawnCount for each NDC item individually.
+	Unless absolutely needed this is discouraged as it comes at significant performance cost when spawning and GPU emitters can currently only handle 8 individual spawns per frame.
+	Calling GetNDCSpawnInfo() in the particle spawn script to get the spawning NDC Index is prefered.
+	*/
 	UPROPERTY(EditAnywhere, Category = "Spawning", AdvancedDisplay)
 	bool bOverrideSpawnGroupToDataChannelIndex = true;
 
@@ -146,6 +154,8 @@ public:
 	NIAGARA_API void Num(FVectorVMExternalFunctionContext& Context);
 	NIAGARA_API void Read(FVectorVMExternalFunctionContext& Context, int32 FuncIdx);
 	NIAGARA_API void Consume(FVectorVMExternalFunctionContext& Context, int32 FuncIdx);
+	
+	NIAGARA_API void GetNDCSpawnData(FVectorVMExternalFunctionContext& Context);
 
 	//Emitter only functions.
 	NIAGARA_API void SpawnConditional(FVectorVMExternalFunctionContext& Context, int32 FuncIndex);
@@ -163,6 +173,34 @@ protected:
  	FNDIDataChannelCompiledData CompiledData;
 };
 
+struct FNDIDataChannelRead_EmitterSpawnData
+{
+	TArray<int32> NDCSpawnData;
+	int32 NDCSpawnDataBuckets[16];
+
+	void Reset()
+	{
+		NDCSpawnData.Reset();
+		FMemory::Memzero(NDCSpawnDataBuckets, sizeof(int32) * 16);
+	}
+};
+
+struct FNDIDataChannelRead_EmitterInstanceData
+{
+	//Spawn data buffers needed for accessing the correct NDCIndex and spawn data from a particle during CPU or GPU execution.
+	//Rebuilt every frame and only valid inside the spawn execution generated by this DI.
+	FNDIDataChannelRead_EmitterSpawnData NDCSpawnData;
+
+	//Spawn Counts for each entry in the NDC.
+	TArray<int32> NDCSpawnCounts;
+
+	void Reset()
+	{
+		NDCSpawnCounts.Reset();
+		NDCSpawnData.Reset();
+	}
+};
+
 struct FNDIDataChannelReadInstanceData
 {
 	//TODO: Local reads.
@@ -181,8 +219,8 @@ struct FNDIDataChannelReadInstanceData
 	/** Cached hash to check if the layout of our source data has changed. */
 	uint64 ChachedDataSetLayoutHash = INDEX_NONE;
 
-	/** When true we should update our RT data next tick. */
-	mutable bool bUpdateRTData = false;
+	/** When true we should update our function binding info on the RT next tick. */
+	mutable bool bUpdateFunctionBindingRTData = false;
 
 	/** Binding info from the Data Channel to each of our function's parameters. */
 	//FNDIDataChannelBindingInfo BindingInfo;
@@ -200,10 +238,9 @@ struct FNDIDataChannelReadInstanceData
 	std::atomic<int32> ConsumeIndex = 0;
 
 	/** 
-	Conditional spawns to inject into the bound emitter based on our DI function calls.
-	Each data channel element we read in a particular execution will have an entry in this array and can spawn(or not) independently.
+	Instance data for each emitter using this DI.
 	*/
-	TMap<FNiagaraEmitterInstance*,TArray<FNiagaraSpawnInfo>> PerEmitterConditionalSpawns;
+	TMap<FNiagaraEmitterInstance*,FNDIDataChannelRead_EmitterInstanceData> EmitterInstanceData;
 
 	uint32 CachedLayoutHash = INDEX_NONE;
 	FNiagaraSystemInstance* Owner = nullptr;
@@ -211,7 +248,7 @@ struct FNDIDataChannelReadInstanceData
 	virtual ~FNDIDataChannelReadInstanceData();
 	FNiagaraDataBuffer* GetReadBufferCPU(bool bPrevFrame);
 	bool Init(UNiagaraDataInterfaceDataChannelRead* Interface, FNiagaraSystemInstance* Instance);
-	bool Tick(UNiagaraDataInterfaceDataChannelRead* Interface, FNiagaraSystemInstance* Instance, bool bIsInit=false);
+	bool Tick(UNiagaraDataInterfaceDataChannelRead* Interface, FNiagaraSystemInstance* Instance, bool bIsInit = false);	
 	bool PostTick(UNiagaraDataInterfaceDataChannelRead* Interface, FNiagaraSystemInstance* Instance);
 };
 
@@ -221,6 +258,8 @@ struct FNiagaraDataInterfaceProxy_DataChannelRead : public FNiagaraDataInterface
 	virtual int32 PerInstanceDataPassedToRenderThreadSize() const override;
 	virtual void GetDispatchArgs(const FNDIGpuComputeDispatchArgsGenContext& Context) override;
 
+	virtual void PreStage(const FNDIGpuComputePreStageContext& Context)override;
+
 	/** Persistent per instance data on the RT. Constructed when consuming data passed from GT->RT. */
 	struct FInstanceData
 	{
@@ -242,6 +281,11 @@ struct FNiagaraDataInterfaceProxy_DataChannelRead : public FNiagaraDataInterface
 		At hlsl gen time we can only know which parameter are accessed by each script individually so each script must have it's own parameter binding table.
 		*/
 		TMap<FNiagaraCompileHash, uint32> GPUScriptParameterTableOffsets;
+
+		/** Buffer containing packed data for all emitters NDC spawning data for use on the GPU. */
+		TArray<int32> NDCSpawnData;
+
+		FRDGBufferRef NDCSpawnDataBuffer;
 	};
 
 	TMap<FNiagaraSystemInstanceID, FInstanceData> SystemInstancesToProxyData_RT;
diff --git a/Engine/Plugins/FX/Niagara/Source/Niagara/Private/DataInterface/NiagaraDataInterfaceDataChannelCommon.cpp b/Engine/Plugins/FX/Niagara/Source/Niagara/Private/DataInterface/NiagaraDataInterfaceDataChannelCommon.cpp
--- a/Engine/Plugins/FX/Niagara/Source/Niagara/Private/DataInterface/NiagaraDataInterfaceDataChannelCommon.cpp
+++ b/Engine/Plugins/FX/Niagara/Source/Niagara/Private/DataInterface/NiagaraDataInterfaceDataChannelCommon.cpp
@@ -206,6 +206,12 @@ void FNDIDataChannelCompiledData::GatherAccessInfo(UNiagaraSystem* System, UNiag
 			}
 		}
 		bUsedByCPU = true;
+
+		if(BindingInfo.Name == NDIDataChannelUtilities::GetNDCSpawnDataName)
+		{
+			bNeedsSpawnDataTable = true;
+		}
+
 		return true;
 	};
 	FNiagaraDataInterfaceUtilities::ForEachVMFunction(Owner, System, HandleVMFunc);
@@ -232,6 +238,11 @@ void FNDIDataChannelCompiledData::GatherAccessInfo(UNiagaraSystem* System, UNiag
 			}
 		}
 		bUsedByGPU = true;
+
+		if (BindingInfo.DefinitionName == TEXT("GetNDCSpawnData"))
+		{
+			bNeedsSpawnDataTable = true;
+		}
 		return true;
 	};
 	FNiagaraDataInterfaceUtilities::ForEachGpuFunction(Owner, System, HandleGpuFunc);
@@ -276,6 +287,8 @@ int32 FNDIDataChannelCompiledData::FindFunctionInfoIndex(FName Name, const TArra
 
 namespace NDIDataChannelUtilities
 {
+	const FName GetNDCSpawnDataName(TEXT("GetNDCSpawnData"));
+
 	void SortParameters(TArray<FNiagaraVariableBase>& Parameters)
 	{
 		Parameters.Sort([](const FNiagaraVariableBase& Lhs, const FNiagaraVariableBase& Rhs)
diff --git a/Engine/Plugins/FX/Niagara/Source/Niagara/Private/DataInterface/NiagaraDataInterfaceDataChannelRead.cpp b/Engine/Plugins/FX/Niagara/Source/Niagara/Private/DataInterface/NiagaraDataInterfaceDataChannelRead.cpp
--- a/Engine/Plugins/FX/Niagara/Source/Niagara/Private/DataInterface/NiagaraDataInterfaceDataChannelRead.cpp
+++ b/Engine/Plugins/FX/Niagara/Source/Niagara/Private/DataInterface/NiagaraDataInterfaceDataChannelRead.cpp
@@ -4,6 +4,8 @@
 
 #include "NiagaraCompileHashVisitor.h"
 #include "ShaderCompilerCore.h"
+#include "RenderGraphUtils.h"
+#include "SystemTextures.h"
 #include "NiagaraModule.h"
 #include "NiagaraShaderParametersBuilder.h"
 
@@ -82,6 +84,29 @@ namespace NDIDataChannelReadLocal
 		return Sig;
 	}
 
+	const FNiagaraFunctionSignature& GetFunctionSig_GetNDCSpawnData()
+	{
+		static FNiagaraFunctionSignature Sig;
+		if (!Sig.IsValid())
+		{
+			Sig.Name = NDIDataChannelUtilities::GetNDCSpawnDataName;
+#if WITH_EDITORONLY_DATA
+			Sig.Description = LOCTEXT("GetNDCSpawnInfoFunctionDescription", "Returns useful data in relation the the NDC item that spawned this particle.");
+			NIAGARA_ADD_FUNCTION_SOURCE_INFO(Sig)
+#endif		
+			Sig.bMemberFunction = true;
+			Sig.bExperimental = true;
+			Sig.ModuleUsageBitmask = ENiagaraScriptUsageMask::Particle;
+			Sig.AddInput(FNiagaraVariable(FNiagaraTypeDefinition(UNiagaraDataInterfaceDataChannelRead::StaticClass()), TEXT("DataChannel interface")));
+			Sig.AddInput(FNiagaraVariable(FNiagaraTypeDefinition(FNiagaraEmitterID::StaticStruct()), TEXT("Emitter ID")), LOCTEXT("EmitterIDDesc", "ID of the emitter we'd like to spawn into. This can be obtained from Engine.Emitter.ID."));
+			Sig.AddInput(FNiagaraVariable(FNiagaraTypeDefinition::GetIntDef(), TEXT("Spawned Particle Exec Index")), LOCTEXT("GetNDCSpawnData_InExecIndexDesc","The execution index of the spawned particle."));
+			Sig.AddOutput(FNiagaraVariable(FNiagaraTypeDefinition::GetIntDef(), TEXT("NDC Index")), LOCTEXT("GetNDCSpawnData_OutNDCIndexDesc","Index of the NDC item that spawned this particle."));
+			Sig.AddOutput(FNiagaraVariable(FNiagaraTypeDefinition::GetIntDef(), TEXT("NDC Spawn Index")), LOCTEXT("GetNDCSpawnData_OutNDCSpawnIndexDesc","The index of this particle in relation to all the particle spawned by the same NDC item."));
+			Sig.AddOutput(FNiagaraVariable(FNiagaraTypeDefinition::GetIntDef(), TEXT("NDC Spawn Count")), LOCTEXT("GetNDCSpawnData_OutNDCSpawnCountDesc","The number of particles spawned by the same NDC item."));
+		}
+		return Sig;
+	}
+
 	const FNiagaraFunctionSignature& GetFunctionSig_Read()
 	{
 		static FNiagaraFunctionSignature Sig;
@@ -155,7 +180,7 @@ namespace NDIDataChannelReadLocal
 			Sig.ModuleUsageBitmask = ENiagaraScriptUsageMask::Emitter | ENiagaraScriptUsageMask::System;
 			Sig.AddInput(FNiagaraVariable(FNiagaraTypeDefinition(UNiagaraDataInterfaceDataChannelRead::StaticClass()), TEXT("DataChannel interface")));
 			Sig.AddInput(EnabledVar, LOCTEXT("SpawnEnableInputDesc", "Enable or disable this function call. If false, this call with have no effetcs."));
-			Sig.AddInput(FNiagaraVariable(FNiagaraTypeDefinition(FNiagaraEmitterID::StaticStruct()), TEXT("Emitter ID")), LOCTEXT("EmitterIDDesc", "ID of the emitter we'd like to spawn into. This can be obtained from Engine.Emitter.Index."));
+			Sig.AddInput(FNiagaraVariable(FNiagaraTypeDefinition(FNiagaraEmitterID::StaticStruct()), TEXT("Emitter ID")), LOCTEXT("EmitterIDDesc", "ID of the emitter we'd like to spawn into. This can be obtained from Engine.Emitter.ID."));
 			Sig.AddInput(FNiagaraVariable(StaticEnum<ENDIDataChannelSpawnMode>(), TEXT("Mode")), LOCTEXT("SpawnCondModeInputDesc", "A mode switch that controls how this funciton will behave."));
 			Sig.AddInput(FNiagaraVariable(StaticEnum<ENiagaraConditionalOperator>(), TEXT("Operator")), LOCTEXT("SpawnCondOpInputDesc", "The comparison operator to use when comparing values in the data channel to conditional parameters."));
 			Sig.AddInput(FNiagaraVariable(FNiagaraTypeDefinition::GetIntDef(), TEXT("Min Spawn Count")), LOCTEXT("MinSpawnCountInputDesc", "Minimum number of particles to spawn for each element in the data channel."));
@@ -185,6 +210,7 @@ namespace NDIDataChannelReadLocal
 		SHADER_PARAMETER(int32, FloatStride)
 		SHADER_PARAMETER(int32, Int32Stride)
 		SHADER_PARAMETER(int32, HalfStride)
+		SHADER_PARAMETER_RDG_BUFFER_SRV(Buffer<int32>, NDCSpawnDataBuffer)
 	END_SHADER_PARAMETER_STRUCT()
 }
 
@@ -211,8 +237,11 @@ struct FNDIDataChannelReadInstanceData_RT
 	*/
 	TMap<FNiagaraCompileHash, uint32> GPUScriptParameterTableOffsets;
 
-	/** Signal we have updated data from the GT this frame we should process. */
-	bool bHasUpdate = false;
+	/** Signal we have updated function binding data from the GT this frame we should process. */
+	bool bHasFunctionBindingUpdate = false;
+
+	/** Buffer containing packed data for all emitters NDC spawning data for use on the GPU. */
+	TArray<int32> NDCSpawnData;
 };
 
 //////////////////////////////////////////////////////////////////////////
@@ -251,9 +280,9 @@ bool FNDIDataChannelReadInstanceData::Init(UNiagaraDataInterfaceDataChannelRead*
 bool FNDIDataChannelReadInstanceData::Tick(UNiagaraDataInterfaceDataChannelRead* Interface, FNiagaraSystemInstance* Instance, bool bIsInit)
 {
 	ConsumeIndex = 0;
-	for(auto& ConditionalSpawnPair : PerEmitterConditionalSpawns)
+	for(auto& EmitterInstDataPair : EmitterInstanceData)
 	{
-		ConditionalSpawnPair.Value.Reset();
+		EmitterInstDataPair.Value.Reset();
 	}
 
 	const FNDIDataChannelCompiledData& CompiledData = Interface->GetCompiledData();
@@ -389,7 +418,7 @@ bool FNDIDataChannelReadInstanceData::Tick(UNiagaraDataInterfaceDataChannelRead*
 
 				//We can likely be more targeted here.
 				//Could probably only update the RT when the GPU data changes and only update the bindings if the function hashes change etc.
-				bUpdateRTData = CompiledData.UsedByGPU();
+				bUpdateFunctionBindingRTData = CompiledData.UsedByGPU();
 				int32 NumFuncs = CompiledData.GetFunctionInfo().Num();
 				FuncToDataSetBindingInfo.SetNum(NumFuncs);
 				//FuncToDataSetLayoutKeys.SetNumZeroed(NumFuncs);
@@ -566,28 +595,104 @@ void UNiagaraDataInterfaceDataChannelRead::PostStageTick(FNDICpuPostStageContext
 	FNDIDataChannelReadInstanceData* InstanceData = Context.GetPerInstanceData<FNDIDataChannelReadInstanceData>();
 
 	check(InstanceData);
-	for (auto& ConditionalSpawnPair : InstanceData->PerEmitterConditionalSpawns)
+	check(Context.Usage == ENiagaraScriptUsage::EmitterUpdateScript || Context.Usage == ENiagaraScriptUsage::SystemUpdateScript);
+
+	for (auto& EmitterInstanceDataPair : InstanceData->EmitterInstanceData)
 	{
-		if(FNiagaraEmitterInstance* TargetEmitter = ConditionalSpawnPair.Key)
+		FNDIDataChannelRead_EmitterInstanceData& EmitterInstData = EmitterInstanceDataPair.Value;
+
+		TArray<int32> PerNDCSpawnCounts = MoveTemp(EmitterInstData.NDCSpawnCounts);
+	
+		EmitterInstData.Reset();
+
+		if(FNiagaraEmitterInstance* TargetEmitter = EmitterInstanceDataPair.Key)
 		{
-			for (int32 i = 0; i < ConditionalSpawnPair.Value.Num(); ++i)
+			//-TODO:Stateless:
+			if (FNiagaraEmitterInstanceImpl* StatefulEmitter = TargetEmitter->AsStateful())
 			{
-				FNiagaraSpawnInfo& SpawnInfo = ConditionalSpawnPair.Value[i];
-				if (SpawnInfo.Count > 0)
+				if (bOverrideSpawnGroupToDataChannelIndex)
 				{
-					if (bOverrideSpawnGroupToDataChannelIndex)
+					//If we're overriding the spawn group then we must submit one SpawnInfo per NDC entry.
+					FNiagaraSpawnInfo NewSpawnInfo(0,0.0f,0.0f,0);
+					for (int32 i = 0; i < PerNDCSpawnCounts.Num(); ++i)
 					{
-						SpawnInfo.SpawnGroup = i;
+						int32 SpawnCount = PerNDCSpawnCounts[i];
+						if (SpawnCount > 0)
+						{
+							NewSpawnInfo.Count = SpawnCount;
+							NewSpawnInfo.SpawnGroup = i;
+							StatefulEmitter->GetSpawnInfo().Emplace(NewSpawnInfo);
+						}
 					}
-					//-TODO:Stateless: We need to handle the stateless path here
-					if (FNiagaraEmitterInstanceImpl* StatefulEmitter = TargetEmitter->AsStateful())
+				}
+				else 
+				{
+					//No need for indirection table but we're not overriding the spawn group either so still push a single combined spawn info.
+					FNiagaraSpawnInfo NewSpawnInfo(0, 0.0f, 0.0f, 0);
+					for (int32 i = 0; i < PerNDCSpawnCounts.Num(); ++i)
 					{
-						StatefulEmitter->GetSpawnInfo().Emplace(SpawnInfo);
+						NewSpawnInfo.Count += PerNDCSpawnCounts[i];
+					}
+					StatefulEmitter->GetSpawnInfo().Emplace(NewSpawnInfo);
+				}
+
+				if(CompiledData.NeedSpawnDataTable())
+				{
+					//Build an indirection table that allows us to map from ExecIndex back to the NDCIndex that generated it.
+					//The indirection table is arranged in power of two buckets.
+					//An NDC that spawns say 37 particles would add an entry to the 32, 4 and 1 buckets.
+					//This allows us to spawn any number of particles from each NDC and only have a max of 16 indirection table entries.
+					//Vs the naive per particle approach of 1 entry per particle.
+					//Buckets are processed in descending size order.
+
+					//TODO: It should be possible to write this from the GPU too as long as we allocate fixed size buckets.
+					int32* SpawnDataBuckets = EmitterInstData.NDCSpawnData.NDCSpawnDataBuckets;					
+
+					TArray<int32>& NDCSpawnData = EmitterInstData.NDCSpawnData.NDCSpawnData;
+
+					//Start of the buffer is the per NDC spawn counts.
+					uint32 TotalNDCSpawnDataSize = PerNDCSpawnCounts.Num();
+
+					for (int32 i = 0; i < PerNDCSpawnCounts.Num(); ++i)
+					{						
+						uint32 Count = PerNDCSpawnCounts[i];
+
+						//First section is the per NDC counts.
+						NDCSpawnData.Add(Count);
+
+						for (uint32 Bucket = 0; Bucket < 16; ++Bucket)
+						{
+							uint32 BucketSize = (1<<15) >> Bucket;
+							uint32 Mask = (0xFFFF >> (Bucket + 1));
+							uint32 CountMasked = Count & ~Mask;
+							Count &= Mask;
+							uint32 NumBucketEntries = CountMasked / BucketSize;
+							SpawnDataBuckets[Bucket] += NumBucketEntries;
+							TotalNDCSpawnDataSize += NumBucketEntries;
+						}
+					}
+
+					//Second part is the counts decomposed into power of two buckets that allows us to map ExecIndex at runtime to an NDCIndex entry in this table.
+					for (int32 Bucket = 0; Bucket < 16; ++Bucket)
+					{
+						uint32 BucketSize = (1 << 15) >> Bucket;
+						uint32 StartSize = NDCSpawnData.Num();
+						for (int32 i = 0; i < PerNDCSpawnCounts.Num(); ++i)
+						{
+							int32& Count = PerNDCSpawnCounts[i];
+							while ((uint32)Count >= BucketSize)
+							{
+								Count -= BucketSize;
+								NDCSpawnData.Add(i);
+							}
+						}
+
+						uint32 EndSize = NDCSpawnData.Num();
+						check(EndSize - StartSize == SpawnDataBuckets[Bucket]);
 					}
 				}
 			}
 		}
-		ConditionalSpawnPair.Value.Reset();
 	}
 }
 
@@ -602,11 +707,11 @@ void UNiagaraDataInterfaceDataChannelRead::ProvidePerInstanceDataForRenderThread
 	bool bReadPrevFrame = bReadCurrentFrame == false || GNDCReadForcePrevFrame;
 	TargetData->bReadPrevFrame = bReadPrevFrame;
 
-	if (SourceData.bUpdateRTData && INiagaraModule::DataChannelsEnabled())
+	if (SourceData.bUpdateFunctionBindingRTData && INiagaraModule::DataChannelsEnabled())
 	{
-		SourceData.bUpdateRTData = false;
+		SourceData.bUpdateFunctionBindingRTData = false;
 
-		TargetData->bHasUpdate = true;
+		TargetData->bHasFunctionBindingUpdate = true;
 		
 		const FNiagaraDataSetCompiledData& GPUCompiledData = SourceData.DataChannel->GetDataChannel()->GetCompiledData(ENiagaraSimTarget::GPUComputeSim);
 
@@ -639,6 +744,74 @@ void UNiagaraDataInterfaceDataChannelRead::ProvidePerInstanceDataForRenderThread
 			}
 		}
 	}
+
+	//Always need to fill in the NDCSpawnData array as it will change every frame and be pushed into an RDG buffer.
+
+	//New buffer is every emitter continuous NDCSpawnDataArray. We need to store an offset that we pass in as a uniform.
+	//The buckets come first, then the per NDC SpawnCounts, Then the bucket back ptrs.
+
+	//Do one pass to calculate size.
+	auto GetEmitterNDCSpawnDataSize = [](const FNDIDataChannelRead_EmitterInstanceData& EmitterInstData)
+	{
+		return 16 + EmitterInstData.NDCSpawnCounts.Num() + EmitterInstData.NDCSpawnData.NDCSpawnData.Num();
+	};
+
+	uint32 NumEmitters = SourceData.EmitterInstanceData.Num();
+	uint32 TotalPacckedNDCSpawnDataSize = 0;
+	int32 MaxEmitterIndex = 0;
+	for (const TPair<FNiagaraEmitterInstance*, FNDIDataChannelRead_EmitterInstanceData>& EmitterInstDataPair : SourceData.EmitterInstanceData)
+	{
+		if (const FNiagaraEmitterInstance* EmitterInst = EmitterInstDataPair.Key)
+		{
+			const FNDIDataChannelRead_EmitterInstanceData& EmitterInstData = EmitterInstDataPair.Value;
+
+			TotalPacckedNDCSpawnDataSize += GetEmitterNDCSpawnDataSize(EmitterInstData);
+			FNiagaraEmitterID ID = EmitterInst->GetEmitterID();
+			MaxEmitterIndex = FMath::Max(MaxEmitterIndex, ID.ID);
+		}
+	}
+
+	//First section of the NDCSpawnDataBuffer is an offset into the buffer for each emitter.
+	TotalPacckedNDCSpawnDataSize += MaxEmitterIndex;
+
+	TargetData->NDCSpawnData.Reset(TotalPacckedNDCSpawnDataSize);
+	TArray<int32>& TargetNDCSpawnData = TargetData->NDCSpawnData;
+
+	//First grab space for the per emitter offset table. We'll fill this in as we go.
+	TargetNDCSpawnData.AddZeroed(MaxEmitterIndex + 1);
+
+	uint32 CurrentSpawnDataOffset = TargetNDCSpawnData.Num();
+
+	for (const TPair<FNiagaraEmitterInstance*, FNDIDataChannelRead_EmitterInstanceData>& EmitterInstDataPair : SourceData.EmitterInstanceData)
+	{
+		if (const FNiagaraEmitterInstance* EmitterInst = EmitterInstDataPair.Key)
+		{
+			const FNDIDataChannelRead_EmitterInstanceData& EmitterInstData = EmitterInstDataPair.Value;
+			if(EmitterInst->GetGPUContext() == nullptr)
+			{
+				continue;
+			}
+
+
+			uint32 EmitterNDCSpawnDataSize = GetEmitterNDCSpawnDataSize(EmitterInstData);
+
+			//First fill in the current offset for this emitter.
+			FNiagaraEmitterID EmitterID = EmitterInst->GetEmitterID();
+			TargetNDCSpawnData[EmitterID.ID] = CurrentSpawnDataOffset;
+
+			CurrentSpawnDataOffset += EmitterNDCSpawnDataSize;
+
+			//Next fill in bucket counts
+			for (int32 i = 0; i < 16; ++i)
+			{
+				TargetNDCSpawnData.Add(EmitterInstData.NDCSpawnData.NDCSpawnDataBuckets[i]);
+			}
+			//Next the per NDC Spawn Counts
+			TargetNDCSpawnData.Append(EmitterInstData.NDCSpawnCounts);
+			//Finally the exec index to NDC index mapping table
+			TargetNDCSpawnData.Append(EmitterInstData.NDCSpawnData.NDCSpawnData);
+		}
+	}
 }
 
 void UNiagaraDataInterfaceDataChannelRead::GetEmitterDependencies(UNiagaraSystem* Asset, TArray<FVersionedNiagaraEmitter>& Dependencies)const
@@ -910,6 +1083,7 @@ bool UNiagaraDataInterfaceDataChannelRead::CopyToInternal(UNiagaraDataInterface*
 void UNiagaraDataInterfaceDataChannelRead::GetFunctionsInternal(TArray<FNiagaraFunctionSignature>& OutFunctions) const
 {
 	OutFunctions.Add(NDIDataChannelReadLocal::GetFunctionSig_Num());
+	OutFunctions.Add(NDIDataChannelReadLocal::GetFunctionSig_GetNDCSpawnData());
 	OutFunctions.Add(NDIDataChannelReadLocal::GetFunctionSig_Read());
 	OutFunctions.Add(NDIDataChannelReadLocal::GetFunctionSig_Consume());
 	OutFunctions.Add(NDIDataChannelReadLocal::GetFunctionSig_SpawnConditional());
@@ -922,6 +1096,10 @@ void UNiagaraDataInterfaceDataChannelRead::GetVMExternalFunction(const FVMExtern
 	{
 		OutFunc = FVMExternalFunction::CreateLambda([this](FVectorVMExternalFunctionContext& Context) { this->Num(Context); });
 	}
+	else if (BindingInfo.Name == NDIDataChannelReadLocal::GetFunctionSig_GetNDCSpawnData().Name)
+	{
+		OutFunc = FVMExternalFunction::CreateLambda([this](FVectorVMExternalFunctionContext& Context) { this->GetNDCSpawnData(Context); });
+	}
 	else
 	{
 		int32 FuncIndex = CompiledData.FindFunctionInfoIndex(BindingInfo.Name, BindingInfo.VariadicInputs, BindingInfo.VariadicOutputs);
@@ -965,6 +1143,188 @@ void UNiagaraDataInterfaceDataChannelRead::Num(FVectorVMExternalFunctionContext&
 	}
 }
 
+/** 
+GetNDCSpawnData - Retrieves spawn data about the NDC that spawned a particular particle.
+Uses an indirection table that decomposes each NDC spawn into power of 2 buckets of particles.
+We do this to strike a balance between allowing many particles per NDC entry and allowing many NDC entires to spawn particles.
+We have a max of 16 buckets with the highest being for spawns with 1<15 particles or more and the lowest being for individual particles.
+An example with two NDC entries. The first spawning 10 and the second 8.
+The first's spawn count decomposes into an entry in the 8 bucket and 2 bucket.
+The seconds just has an entry in the 8 bucket.
+We have 16 buckets so the bucket counts array looks like
+0,0,0,0,0,0,0,0,0,0,0,0,2,0,1,0
+This means the we have 2 buckets with data so the rest of our buffer is.
+0,1,0,
+
+We have two entires in the 8 bucket. 
+As we spawn particles we use our bucket sizes counts and the exec index to see which bucket entry each exec index should use.
+The first 8 particles processed, exec index 0-7 will lookup the first entry and so use NDC 0.
+THe next 8 particles, exec index 8-15 will use the next and so use NDC 1.
+Finally the last two particle spawned will use the next entry and so also use NDC 0.
+
+So in total we do have 10 particles from NDC 0 and 8 from NDC 1.
+However they will not be processed all together with their own spawning NDC.
+
+In the worst case an NDC entry could add to all 16 buckets and so we'd have 16 entries for that NDC entry.
+Which may seem like a lot but consider that is spawning 1<<15 particles so not all that bad really.
+
+It also means the lookup does not need to search an arbitrary sized list.
+It just has to loop over a size 16 array and do some math to get an index into the main buffer from which to retreive the NDC Index.
+
+Once we have the NDCIndex we the do another similar pass and use the total spawn counts for that NDC to work out a SpawnIndex within the NDC.
+*/
+void UNiagaraDataInterfaceDataChannelRead::GetNDCSpawnData(FVectorVMExternalFunctionContext& Context)
+{
+	VectorVM::FUserPtrHandler<FNDIDataChannelReadInstanceData> InstData(Context);
+
+	FNDIInputParam<FNiagaraEmitterID> InEmitterID(Context);
+	FNDIInputParam<int32> InExecIndex(Context);
+	FNDIOutputParam<int32> OutNDCIndex(Context);
+	FNDIOutputParam<int32> OutNDCSpawnIndex(Context);
+	FNDIOutputParam<int32> OutNDCSpawnCount(Context);
+
+	FNiagaraSystemInstance* SystemInstance = InstData->Owner;
+	check(SystemInstance);
+
+	bool bReadPrevFrame = bReadCurrentFrame == false || GNDCReadForcePrevFrame;
+	FNiagaraDataBuffer* Buffer = InstData->GetReadBufferCPU(bReadPrevFrame);
+
+	if(Buffer && INiagaraModule::DataChannelsEnabled())
+	{
+		uint32 NumNDCEntries = Buffer->GetNumInstances();
+
+		auto CalculateNDCSpawnInfo = [&](const FNDIDataChannelRead_EmitterInstanceData& EmitterInstanceData)
+		{
+			const int32* NDCSpawnBukets = EmitterInstanceData.NDCSpawnData.NDCSpawnDataBuckets;
+			TConstArrayView<int32> NDCSpawnData(EmitterInstanceData.NDCSpawnData.NDCSpawnData);
+
+			uint32 ExecIndex = InExecIndex.GetAndAdvance();
+			uint32 NDCIndex = INDEX_NONE;
+
+			//First we find which bucket this exec index is in.	
+			uint32 MaxBucketExecIndex = 0;
+			uint32 BucketEntryStart = NumNDCEntries;
+			for (uint32 BucketIdx = 0; BucketIdx < 16; ++BucketIdx)
+			{
+				uint32 BucketSize = (1 << 15) >> BucketIdx;
+				uint32 NumEntriesInBucket = NDCSpawnBukets[BucketIdx];
+				uint32 MinBucketExecIndex = MaxBucketExecIndex;
+				MaxBucketExecIndex += BucketSize * NumEntriesInBucket;
+				if (ExecIndex < MaxBucketExecIndex)
+				{
+					//We found our bucket.
+					//Now we need to find our NDCIndex Entry Index.
+					uint32 NDCIndexEntry = (ExecIndex - MinBucketExecIndex) >> (15 - BucketIdx);
+
+					NDCIndex = NDCSpawnData[BucketEntryStart + NDCIndexEntry];
+					break;
+				}
+
+				BucketEntryStart += NumEntriesInBucket;
+			}
+
+			if (NDCIndex >= 0 && NDCIndex < NumNDCEntries)
+			{
+				OutNDCIndex.SetAndAdvance(NDCIndex);
+
+				uint32 NDCSpawnCount = NDCSpawnData[NDCIndex];
+				OutNDCSpawnCount.SetAndAdvance(NDCSpawnCount);
+
+				//Do another pass to calculate our SpawnIndex for this NDC within the total count for this NDC.
+				if (OutNDCSpawnIndex.IsValid())
+				{
+					uint32 NDCSpawnIndex = 0;
+					uint32 Count = NDCSpawnCount;
+					MaxBucketExecIndex = 0;
+					for (int32 BucketIdx = 0; BucketIdx < 16; ++BucketIdx)
+					{
+						uint32 BucketSize = (1 << 15) >> BucketIdx;
+						uint32 Mask = (0xFFFF >> (BucketIdx + 1));
+						uint32 CountMasked = Count & ~Mask;
+						Count &= Mask;
+						uint32 NumNDCEntriesInBucket = CountMasked >> (15 - BucketIdx);
+						uint32 NumEntriesInBucket = NDCSpawnBukets[BucketIdx];
+						uint32 NumNDCInstancesInBucket = NumNDCEntriesInBucket * BucketSize;
+
+						int32 MinBucketExecIndex = MaxBucketExecIndex;
+						MaxBucketExecIndex += BucketSize * NumEntriesInBucket;
+						if (ExecIndex < MaxBucketExecIndex && NumNDCInstancesInBucket > 0)
+						{
+							uint32 NDCIndexEntry = (ExecIndex - MinBucketExecIndex) >> (15 - BucketIdx);
+
+							uint32 MinNDCBucketExecIndex = MinBucketExecIndex + (BucketSize * NDCIndexEntry);
+
+							NDCSpawnIndex += (ExecIndex - MinNDCBucketExecIndex);
+							break;
+						}
+						else
+						{
+							NDCSpawnIndex += NumNDCInstancesInBucket;
+						}
+					}
+					OutNDCSpawnIndex.SetAndAdvance(NDCSpawnIndex);
+				}
+			}
+			else
+			{
+				OutNDCIndex.SetAndAdvance(INDEX_NONE);
+				OutNDCSpawnCount.SetAndAdvance(INDEX_NONE);
+				OutNDCSpawnIndex.SetAndAdvance(INDEX_NONE);
+			}
+		};
+		
+		if(InEmitterID.IsConstant())
+		{		
+			//TODO: Can likely vectorize all this.
+			const FNiagaraEmitterID EmitterID = InEmitterID.GetAndAdvance();
+			FNiagaraEmitterInstance* EmitterInst = SystemInstance->GetEmitterByID(EmitterID);
+			if(const FNDIDataChannelRead_EmitterInstanceData* EmitterInstData = InstData->EmitterInstanceData.Find(EmitterInst))
+			{
+				for (int32 i = 0; i < Context.GetNumInstances(); ++i)
+				{
+					CalculateNDCSpawnInfo(*EmitterInstData);
+				}
+			}
+			else
+			{
+				for (int32 i = 0; i < Context.GetNumInstances(); ++i)
+				{
+					OutNDCIndex.SetAndAdvance(INDEX_NONE);
+					OutNDCSpawnCount.SetAndAdvance(INDEX_NONE);
+					OutNDCSpawnIndex.SetAndAdvance(INDEX_NONE);
+				}
+			}
+		}
+		else
+		{
+			for (int32 i = 0; i < Context.GetNumInstances(); ++i)
+			{
+				const FNiagaraEmitterID EmitterID = InEmitterID.GetAndAdvance();
+				FNiagaraEmitterInstance* EmitterInst = SystemInstance->GetEmitterByID(EmitterID);
+				if (const FNDIDataChannelRead_EmitterInstanceData* EmitterInstData = InstData->EmitterInstanceData.Find(EmitterInst))
+				{
+					CalculateNDCSpawnInfo(*EmitterInstData);
+				}
+				else
+				{
+					OutNDCIndex.SetAndAdvance(INDEX_NONE);
+					OutNDCSpawnCount.SetAndAdvance(INDEX_NONE);
+					OutNDCSpawnIndex.SetAndAdvance(INDEX_NONE);
+				}
+			}			
+		}
+	}	
+	else
+	{
+		for (int32 i = 0; i < Context.GetNumInstances(); ++i)
+		{
+			OutNDCIndex.SetAndAdvance(INDEX_NONE);
+			OutNDCSpawnCount.SetAndAdvance(INDEX_NONE);
+			OutNDCSpawnIndex.SetAndAdvance(INDEX_NONE);
+		}
+	}
+}
+
 void UNiagaraDataInterfaceDataChannelRead::Read(FVectorVMExternalFunctionContext& Context, int32 FuncIdx)
 {
 	SCOPE_CYCLE_COUNTER(STAT_NDIDataChannelRead_Read);
@@ -1166,7 +1526,8 @@ void UNiagaraDataInterfaceDataChannelRead::SpawnConditional(FVectorVMExternalFun
 		int32 SpawnMax = InSpawnMax.GetAndAdvance();
 
 		//Each data channel element has an additional spawn entry which accumulates across all spawning calls and can be nulled independently by a suppression call.
-		TArray<FNiagaraSpawnInfo>& EmitterConditionalSpawns = InstData->PerEmitterConditionalSpawns.FindOrAdd(EmitterInst);
+		FNDIDataChannelRead_EmitterInstanceData& EmitterInstData = InstData->EmitterInstanceData.FindOrAdd(EmitterInst);
+		TArray<int32>& EmitterConditionalSpawns = EmitterInstData.NDCSpawnCounts;
 		EmitterConditionalSpawns.SetNumZeroed(NumDataChannelInstances);
 
 		for (int32 DataChannelIdx = 0; DataChannelIdx < NumDataChannelInstances; ++DataChannelIdx)
@@ -1195,11 +1556,11 @@ void UNiagaraDataInterfaceDataChannelRead::SpawnConditional(FVectorVMExternalFun
 				int32 Count = RandHelper.RandRange(DataChannelIdx, SpawnMin, SpawnMax);
 				if (Mode == ENDIDataChannelSpawnMode::Accumulate)
 				{
-					EmitterConditionalSpawns[DataChannelIdx].Count += Count;
+					EmitterConditionalSpawns[DataChannelIdx] += Count;
 				}
 				else if (Mode == ENDIDataChannelSpawnMode::Override)
 				{
-					EmitterConditionalSpawns[DataChannelIdx].Count = Count;
+					EmitterConditionalSpawns[DataChannelIdx] = Count;
 				}
 			}
 		}
@@ -1228,6 +1589,7 @@ void UNiagaraDataInterfaceDataChannelRead::GetCommonHLSL(FString& OutHLSL)
 bool UNiagaraDataInterfaceDataChannelRead::GetFunctionHLSL(FNiagaraDataInterfaceHlslGenerationContext& HlslGenContext, FString& OutHLSL)
 {
 	return	HlslGenContext.GetFunctionInfo().DefinitionName == GET_FUNCTION_NAME_CHECKED(UNiagaraDataInterfaceDataChannelRead, Num) ||
+		HlslGenContext.GetFunctionInfo().DefinitionName == GET_FUNCTION_NAME_CHECKED(UNiagaraDataInterfaceDataChannelRead, GetNDCSpawnData) ||
 		HlslGenContext.GetFunctionInfo().DefinitionName == GET_FUNCTION_NAME_CHECKED(UNiagaraDataInterfaceDataChannelRead, Read) ||
 		HlslGenContext.GetFunctionInfo().DefinitionName == GET_FUNCTION_NAME_CHECKED(UNiagaraDataInterfaceDataChannelRead, Consume);
 }
@@ -1570,7 +1932,8 @@ void UNiagaraDataInterfaceDataChannelRead::SetShaderParameters(const FNiagaraDat
 			{
 				const FReadBuffer& ParameterLayoutBuffer = InstanceData->ParameterLayoutBuffer;
 
-				if (ParameterLayoutBuffer.SRV.IsValid() && ParameterLayoutBuffer.NumBytes > 0)
+				FRDGBufferSRVRef NDCSpawnDataBufferSRV = Context.GetGraphBuilder().CreateSRV(InstanceData->NDCSpawnDataBuffer, PF_R32_SINT);
+				if (NDCSpawnDataBufferSRV && ParameterLayoutBuffer.SRV.IsValid() && ParameterLayoutBuffer.NumBytes > 0)
 				{
 					InstParameters->ParamOffsetTable = ParameterLayoutBuffer.SRV.IsValid() ? ParameterLayoutBuffer.SRV.GetReference() : FNiagaraRenderer::GetDummyUIntBuffer();
 					InstParameters->DataFloat = Data->GetGPUBufferFloat().SRV.IsValid() ? Data->GetGPUBufferFloat().SRV.GetReference() : FNiagaraRenderer::GetDummyFloatBuffer();
@@ -1582,6 +1945,7 @@ void UNiagaraDataInterfaceDataChannelRead::SetShaderParameters(const FNiagaraDat
 					InstParameters->Int32Stride = Data->GetInt32Stride() / sizeof(int32);
 					InstParameters->HalfStride = Data->GetHalfStride() / sizeof(FFloat16);
 
+					InstParameters->NDCSpawnDataBuffer = NDCSpawnDataBufferSRV;
 					bSuccess = true;
 				}
 			}
@@ -1599,6 +1963,22 @@ void UNiagaraDataInterfaceDataChannelRead::SetShaderParameters(const FNiagaraDat
 		InstParameters->FloatStride = 0;
 		InstParameters->Int32Stride = 0;
 		InstParameters->HalfStride = 0;
+		
+		FRDGBufferRef DummyBuffer = GSystemTextures.GetDefaultBuffer(Context.GetGraphBuilder(), 4, 0u);
+		InstParameters->NDCSpawnDataBuffer = Context.GetGraphBuilder().CreateSRV(DummyBuffer);		
+	}
+}
+
+void FNiagaraDataInterfaceProxy_DataChannelRead::PreStage(const FNDIGpuComputePreStageContext& Context)
+{
+	FNiagaraDataInterfaceProxy_DataChannelRead::FInstanceData* InstanceData = SystemInstancesToProxyData_RT.Find(Context.GetSystemInstanceID());
+
+	if(InstanceData->NDCSpawnDataBuffer == nullptr)
+	{
+		InstanceData->NDCSpawnDataBuffer = CreateUploadBuffer<int32>(
+			Context.GetGraphBuilder(),
+			TEXT("Niagara_NDCReadDI_NDCSpawnData"),
+			InstanceData->NDCSpawnData);
 	}
 }
 
@@ -1612,7 +1992,10 @@ void FNiagaraDataInterfaceProxy_DataChannelRead::ConsumePerInstanceDataFromGameT
 	InstData.ChannelDataRTProxy = SourceData.ChannelDataRTProxy;
 	InstData.bReadPrevFrame = SourceData.bReadPrevFrame;
 
-	if (SourceData.bHasUpdate)
+	InstData.NDCSpawnData = SourceData.NDCSpawnData;
+	InstData.NDCSpawnDataBuffer = nullptr;//Clear the RDG buffer ready for re-up to the GPU.
+
+	if (SourceData.bHasFunctionBindingUpdate) 
 	{
 		//Take the offset map from the source data.
 		//This maps from GPU script to that scripts offset into the ParameterLayoutBuffer.
diff --git a/Engine/Plugins/FX/Niagara/Source/NiagaraEditor/Private/NiagaraHlslTranslator.cpp b/Engine/Plugins/FX/Niagara/Source/NiagaraEditor/Private/NiagaraHlslTranslator.cpp
--- a/Engine/Plugins/FX/Niagara/Source/NiagaraEditor/Private/NiagaraHlslTranslator.cpp
+++ b/Engine/Plugins/FX/Niagara/Source/NiagaraEditor/Private/NiagaraHlslTranslator.cpp
@@ -3086,6 +3086,8 @@ void TNiagaraHlslTranslator<GraphBridge>::DefineMainGPUFunctions(
 	DataSetReads.GetKeys(ReadDataSetIDs);
 	DataSetWrites.GetKeys(WriteDataSetIDs);
 
+	FNiagaraEmitterID EmitterID = GetEmitterID();
+
 	// Whether Alive is used and must be set at each run
 	bool bUsesAlive = false;
 	TArray<FName> DataSetNames;
@@ -3193,6 +3195,8 @@ void TNiagaraHlslTranslator<GraphBridge>::DefineMainGPUFunctions(
 			HlslOutput += ContextName + TEXT("Particles.ID.Index = IDIndex;\n");
 			HlslOutput += ContextName + TEXT("Particles.ID.AcquireTag = IDAcquireTag;\n");
 		}
+
+		HlslOutput += FString::Printf(TEXT("\t%sEngine.Emitter.ID.ID = %d;\n"), *ContextName, EmitterID.ID);
 	}
 	HlslOutput += TEXT("}\n\n");
 
@@ -3285,6 +3289,8 @@ void TNiagaraHlslTranslator<GraphBridge>::DefineMainGPUFunctions(
 			{
 				HlslOutput += ContextName + TEXT("DataInstance.Alive=true;\n");
 			}
+
+			HlslOutput += FString::Printf(TEXT("\t%sEngine.Emitter.ID.ID = %d;\n"), *ContextName, EmitterID.ID);
 		}
 	}
 	HlslOutput += TEXT("}\n\n");
