- Downgrade warning to log
- We have a stack message for this so the log is less important

diff --git a/Engine/Plugins/FX/Niagara/Source/Niagara/Private/NiagaraSystem.cpp b/Engine/Plugins/FX/Niagara/Source/Niagara/Private/NiagaraSystem.cpp
--- a/Engine/Plugins/FX/Niagara/Source/Niagara/Private/NiagaraSystem.cpp
+++ b/Engine/Plugins/FX/Niagara/Source/Niagara/Private/NiagaraSystem.cpp
@@ -2053,7 +2053,7 @@ void UNiagaraSystem::ComputeEmittersExecutionOrder()
 			{
 				if (!ComputeEmitterPriority(EmitterIdx, EmitterPriorities, EmitterDependencyGraph))
 				{
-#if WITH_EDITORONLY_DATA
+				#if WITH_EDITORONLY_DATA
 					FName EmitterName = EmitterHandles[EmitterIdx].GetName();
 					FText ErrorMessage = FText::Format(LOCTEXT("CircularDependencyError", "Found circular dependency involving emitter '{0}' in system '{1}'. The execution order will be undefined."), FText::FromName(EmitterName), FText::FromString(GetName()));
 					INiagaraModule& NiagaraModule = FModuleManager::GetModuleChecked<INiagaraModule>("Niagara");
@@ -2067,9 +2067,9 @@ void UNiagaraSystem::ComputeEmittersExecutionOrder()
 					MessageStore.AddMessage(ComputeEmitterExecutionOrderMessageId, ComputeEmitterExecutionOrderMessage);
 					if (MessageStore.IsMessageDismissed(ComputeEmitterExecutionOrderMessageId) == false)
 					{
-						UE_LOG(LogNiagara, Warning, TEXT("%s"), *ErrorMessage.ToString());
+						UE_LOG(LogNiagara, Log, TEXT("%s"), *ErrorMessage.ToString());
 					}
-#endif
+				#endif
 					break;
 				}
 			}
