Fixed issue preventing the Niagara Debug HUD perf graph from displaying GPU performance data.
Simple fix. At some point it was just left out of a check to enable the underlying GPUProfiler that generates the stats.

diff --git a/Engine/Plugins/FX/Niagara/Source/Niagara/Private/NiagaraDebugHud.cpp b/Engine/Plugins/FX/Niagara/Source/Niagara/Private/NiagaraDebugHud.cpp
--- a/Engine/Plugins/FX/Niagara/Source/Niagara/Private/NiagaraDebugHud.cpp
+++ b/Engine/Plugins/FX/Niagara/Source/Niagara/Private/NiagaraDebugHud.cpp
@@ -950,7 +950,10 @@ void FNiagaraDebugHud::GatherSystemInfo()
 
 #if WITH_NIAGARA_GPU_PROFILER
 	// Update Gpu capture state
-	GpuProfilerListener.SetEnabled(Settings.bOverviewEnabled && (Settings.OverviewMode == ENiagaraDebugHUDOverviewMode::GpuComputePerformance));
+	bool bEnableGPUStats = Settings.bOverviewEnabled && 
+		((Settings.OverviewMode == ENiagaraDebugHUDOverviewMode::GpuComputePerformance) || (Settings.OverviewMode == ENiagaraDebugHUDOverviewMode::PerformanceGraph && Settings.PerfGraphMode == ENiagaraDebugHUDPerfGraphMode::GPU));
+							
+	GpuProfilerListener.SetEnabled(bEnableGPUStats);
 #endif
 
 	// When not enabled do nothing
