Improved sequencer sim caching:
* allow scalability override per cache track
* allow locking of cache tracks to prevent unintentionally overriding a cache asset
* display system name inside cache assets
* disable sequencer looping when recording

diff --git a/Engine/Plugins/FX/Niagara/Source/NiagaraEditor/Private/Toolkits/NiagaraSimCacheToolkit.cpp b/Engine/Plugins/FX/Niagara/Source/NiagaraEditor/Private/Toolkits/NiagaraSimCacheToolkit.cpp
--- a/Engine/Plugins/FX/Niagara/Source/NiagaraEditor/Private/Toolkits/NiagaraSimCacheToolkit.cpp
+++ b/Engine/Plugins/FX/Niagara/Source/NiagaraEditor/Private/Toolkits/NiagaraSimCacheToolkit.cpp
@@ -5,6 +5,9 @@
 
 // Widgets & ViewModels
 #include "AdvancedPreviewSceneModule.h"
+#include "ContentBrowserModule.h"
+#include "IContentBrowserSingleton.h"
+#include "NiagaraSystem.h"
 #include "Widgets/Docking/SDockTab.h"
 #include "ViewModels/NiagaraSimCacheViewModel.h"
 #include "Widgets/SNiagaraSimCacheTreeView.h"
@@ -99,8 +102,8 @@ void FNiagaraSimCacheToolkit::Initialize(const EToolkitMode::Type Mode, const TS
 
 		constexpr bool bCreateDefaultStandaloneMenu = true;
 		constexpr bool bCreateDefaultToolbar = true;
-		FAssetEditorToolkit::InitAssetEditor(Mode, InitToolkitHost, FNiagaraEditorModule::NiagaraEditorAppIdentifier, StandaloneDefaultLayout, bCreateDefaultStandaloneMenu, bCreateDefaultToolbar, SimCache.Get());
-
+		InitAssetEditor(Mode, InitToolkitHost, FNiagaraEditorModule::NiagaraEditorAppIdentifier, StandaloneDefaultLayout, bCreateDefaultStandaloneMenu, bCreateDefaultToolbar, SimCache.Get());
+		ExtendToolbar();
 		RegenerateMenusAndToolbars();
 	}
 }
@@ -227,4 +230,53 @@ TSharedRef<SDockTab> FNiagaraSimCacheToolkit::SpawnTab_Overview(const FSpawnTabA
 	return SpawnedTab;
 }
 
-#undef LOCTEXT_NAMESPACE /** NiagaraSimCacheToolkit **/
\ No newline at end of file
+void FNiagaraSimCacheToolkit::ExtendToolbar()
+{
+	struct Local
+	{
+		static void FillToolbar(FToolBarBuilder& ToolbarBuilder, FNiagaraSimCacheToolkit* Toolkit)
+		{
+			if (Toolkit->SimCache.IsValid() && Toolkit->SimCache->GetSystem())
+			{
+				ToolbarBuilder.BeginSection("System");
+				{
+					TWeakObjectPtr<UNiagaraSystem> NiagaraSystem = Toolkit->SimCache->GetSystem();
+					FUIAction GotoAction = FUIAction(
+						FExecuteAction::CreateLambda([NiagaraSystem] ()
+						{
+							if (NiagaraSystem.IsValid())
+							{
+								const TArray<FAssetData>& Assets = { NiagaraSystem.Get() };
+								const FContentBrowserModule& ContentBrowserModule = FModuleManager::Get().LoadModuleChecked<FContentBrowserModule>("ContentBrowser");
+								ContentBrowserModule.Get().SyncBrowserToAssets(Assets);
+							}
+						}),
+						FCanExecuteAction::CreateLambda([] () { return true; })
+					);
+					ToolbarBuilder.AddToolBarButton(
+						GotoAction,
+						NAME_None,
+						FText::FromString(NiagaraSystem->GetName()),
+						LOCTEXT("GotoSystem", "Browses to the associated Niagara system asset"),
+						FSlateIcon(FAppStyle::GetAppStyleSetName(), "SystemWideCommands.FindInContentBrowser.Small")
+					);
+				}
+				ToolbarBuilder.EndSection();
+			}
+		}
+	};
+	
+	TSharedPtr<FExtender> ToolbarExtender = MakeShareable(new FExtender);
+	ToolbarExtender->AddToolBarExtension(
+		"Asset",
+		EExtensionHook::After,
+		GetToolkitCommands(),
+		FToolBarExtensionDelegate::CreateStatic(&Local::FillToolbar, this)
+	);
+
+	AddToolbarExtender(ToolbarExtender);
+	FNiagaraEditorModule& NiagaraEditorModule = FModuleManager::LoadModuleChecked<FNiagaraEditorModule>("NiagaraEditor");
+	AddToolbarExtender(NiagaraEditorModule.GetToolBarExtensibilityManager()->GetAllExtenders(GetToolkitCommands(), GetEditingObjects()));
+}
+
+#undef LOCTEXT_NAMESPACE /** NiagaraSimCacheToolkit **/
diff --git a/Engine/Plugins/FX/Niagara/Source/NiagaraEditor/Private/Toolkits/NiagaraSimCacheToolkit.h b/Engine/Plugins/FX/Niagara/Source/NiagaraEditor/Private/Toolkits/NiagaraSimCacheToolkit.h
--- a/Engine/Plugins/FX/Niagara/Source/NiagaraEditor/Private/Toolkits/NiagaraSimCacheToolkit.h
+++ b/Engine/Plugins/FX/Niagara/Source/NiagaraEditor/Private/Toolkits/NiagaraSimCacheToolkit.h
@@ -42,6 +42,8 @@ private:
 	TSharedRef<SDockTab> SpawnTab_PreviewSettings(const FSpawnTabArgs& Args);
 	TSharedRef<SDockTab> SpawnTab_Overview(const FSpawnTabArgs& Args);
 
+	void ExtendToolbar();
+
 	// The sim cache being viewed.
 	TWeakObjectPtr <UNiagaraSimCache> SimCache;
 
diff --git a/Engine/Plugins/FX/NiagaraSimCaching/Source/NiagaraSimCaching/Private/Niagara/Sequencer/MovieSceneNiagaraCacheTrack.cpp b/Engine/Plugins/FX/NiagaraSimCaching/Source/NiagaraSimCaching/Private/Niagara/Sequencer/MovieSceneNiagaraCacheTrack.cpp
--- a/Engine/Plugins/FX/NiagaraSimCaching/Source/NiagaraSimCaching/Private/Niagara/Sequencer/MovieSceneNiagaraCacheTrack.cpp
+++ b/Engine/Plugins/FX/NiagaraSimCaching/Source/NiagaraSimCaching/Private/Niagara/Sequencer/MovieSceneNiagaraCacheTrack.cpp
@@ -175,7 +175,34 @@ void UMovieSceneNiagaraCacheTrack::SetCacheRecordingAllowed(bool bShouldRecord)
 
 bool UMovieSceneNiagaraCacheTrack::IsCacheRecordingAllowed() const
 {
-	return bCacheRecordingEnabled;
+	bool bAllSectionCanRecord = true;
+	for (UMovieSceneSection* Section : AnimationSections)
+	{
+		if (UMovieSceneNiagaraCacheSection* NiagaraCacheSection = Cast<UMovieSceneNiagaraCacheSection>(Section))
+		{
+			if (NiagaraCacheSection->IsLocked() || NiagaraCacheSection->Params.bLockCacheToReadOnly)
+			{
+				bAllSectionCanRecord = false;
+			}
+		}
+	}
+	return bCacheRecordingEnabled && !IsEvalDisabled() && bAllSectionCanRecord;
+}
+
+int32 UMovieSceneNiagaraCacheTrack::GetMinimumEngineScalabilitySetting() const
+{
+	int32 MinSetting = -1;
+	for (UMovieSceneSection* Section : AnimationSections)
+	{
+		if (UMovieSceneNiagaraCacheSection* NiagaraCacheSection = Cast<UMovieSceneNiagaraCacheSection>(Section))
+		{
+			if (NiagaraCacheSection->Params.bOverrideQualityLevel)
+			{
+				MinSetting = FMath::Max(MinSetting, static_cast<int32>(NiagaraCacheSection->Params.RecordQualityLevel));
+			}
+		}
+	}
+	return MinSetting;
 }
 
 #undef LOCTEXT_NAMESPACE
diff --git a/Engine/Plugins/FX/NiagaraSimCaching/Source/NiagaraSimCaching/Public/Niagara/Sequencer/MovieSceneNiagaraCacheSection.h b/Engine/Plugins/FX/NiagaraSimCaching/Source/NiagaraSimCaching/Public/Niagara/Sequencer/MovieSceneNiagaraCacheSection.h
--- a/Engine/Plugins/FX/NiagaraSimCaching/Source/NiagaraSimCaching/Public/Niagara/Sequencer/MovieSceneNiagaraCacheSection.h
+++ b/Engine/Plugins/FX/NiagaraSimCaching/Source/NiagaraSimCaching/Public/Niagara/Sequencer/MovieSceneNiagaraCacheSection.h
@@ -4,6 +4,7 @@
 
 #include "NiagaraSimCache.h"
 #include "Sections/MovieSceneBaseCacheSection.h"
+#include "PerQualityLevelProperties.h"
 #include "MovieSceneNiagaraCacheSection.generated.h"
 
 UENUM(BlueprintType)
@@ -52,6 +53,17 @@ struct NIAGARASIMCACHING_API FMovieSceneNiagaraCacheParams : public FMovieSceneB
 	UPROPERTY(EditAnywhere, Category = "NiagaraCache")
 	TObjectPtr<UNiagaraSimCache> SimCache;
 
+	/** If true then the section properties might still be changed (so the section itself is not locked), but the cache cannot be rerecorded to prevent accidentally overriding the data within */
+	UPROPERTY(EditAnywhere, Category = "NiagaraCache")
+	bool bLockCacheToReadOnly = false;
+
+	UPROPERTY(EditAnywhere, Category = "NiagaraCache", meta=(InlineEditConditionToggle))
+	bool bOverrideQualityLevel = false;
+
+	/** If set, then the engine scalability setting will be overriden with this value when recording a new cache for this track */
+	UPROPERTY(EditAnywhere, Category = "NiagaraCache", meta=(EditCondition="bOverrideQualityLevel"))
+	EPerQualityLevels RecordQualityLevel;
+
 	/** What should the effect do when the track has no cache data to display */
 	UPROPERTY(EditAnywhere, Category="SimCache")
 	ENiagaraSimCacheSectionPlayMode CacheReplayPlayMode = ENiagaraSimCacheSectionPlayMode::DisplayCacheOnly;
diff --git a/Engine/Plugins/FX/NiagaraSimCaching/Source/NiagaraSimCaching/Public/Niagara/Sequencer/MovieSceneNiagaraCacheTrack.h b/Engine/Plugins/FX/NiagaraSimCaching/Source/NiagaraSimCaching/Public/Niagara/Sequencer/MovieSceneNiagaraCacheTrack.h
--- a/Engine/Plugins/FX/NiagaraSimCaching/Source/NiagaraSimCaching/Public/Niagara/Sequencer/MovieSceneNiagaraCacheTrack.h
+++ b/Engine/Plugins/FX/NiagaraSimCaching/Source/NiagaraSimCaching/Public/Niagara/Sequencer/MovieSceneNiagaraCacheTrack.h
@@ -50,6 +50,7 @@ public:
 	virtual void ResetCache() override;
 	virtual void SetCacheRecordingAllowed(bool bShouldRecord) override;
 	virtual bool IsCacheRecordingAllowed() const override;
+	virtual int32 GetMinimumEngineScalabilitySetting() const override;
 	
 	UPROPERTY(Transient)
 	bool bIsRecording = false;
diff --git a/Engine/Plugins/FX/NiagaraSimCaching/Source/NiagaraSimCachingEditor/Private/Sequencer/NiagaraCacheTrackEditor.cpp b/Engine/Plugins/FX/NiagaraSimCaching/Source/NiagaraSimCachingEditor/Private/Sequencer/NiagaraCacheTrackEditor.cpp
--- a/Engine/Plugins/FX/NiagaraSimCaching/Source/NiagaraSimCachingEditor/Private/Sequencer/NiagaraCacheTrackEditor.cpp
+++ b/Engine/Plugins/FX/NiagaraSimCaching/Source/NiagaraSimCachingEditor/Private/Sequencer/NiagaraCacheTrackEditor.cpp
@@ -357,7 +357,8 @@ void FNiagaraCacheTrackEditor::BuildTrackContextMenu(FMenuBuilder& MenuBuilder,
 					FExecuteAction::CreateLambda([this]()
 					{
 						UCacheTrackRecorder::RecordSelectedTracks(GetSequencer());
-					})
+					}),
+					FCanExecuteAction::CreateUObject(CacheTrack, &UMovieSceneNiagaraCacheTrack::IsCacheRecordingAllowed)
 				)
 			);
 		}
@@ -515,7 +516,8 @@ TSharedPtr<SWidget> FNiagaraCacheTrackEditor::BuildOutlinerEditWidget(const FGui
 			.ButtonStyle(FAppStyle::Get(), "HoverHintOnly")
 			.IsEnabled_Lambda([Track]()
 			{
-				return !Track->IsEvalDisabled();
+				IMovieSceneCachedTrack* CachedTrack = Cast<IMovieSceneCachedTrack>(Track);
+				return CachedTrack && CachedTrack->IsCacheRecordingAllowed();
 			})
 			.OnClicked(this, &FNiagaraCacheTrackEditor::RecordCacheTrack, Cast<IMovieSceneCachedTrack>(Track))
 			[
diff --git a/Engine/Plugins/VirtualProduction/Takes/Source/CacheTrackRecorder/Private/Recorder/CacheTrackRecorder.cpp b/Engine/Plugins/VirtualProduction/Takes/Source/CacheTrackRecorder/Private/Recorder/CacheTrackRecorder.cpp
--- a/Engine/Plugins/VirtualProduction/Takes/Source/CacheTrackRecorder/Private/Recorder/CacheTrackRecorder.cpp
+++ b/Engine/Plugins/VirtualProduction/Takes/Source/CacheTrackRecorder/Private/Recorder/CacheTrackRecorder.cpp
@@ -40,6 +40,7 @@
 // LevelEditor includes
 #include "IAssetViewport.h"
 #include "LevelEditor.h"
+#include "SequencerSettings.h"
 #include "Subsystems/AssetEditorSubsystem.h"
 
 #include UE_INLINE_GENERATED_CPP_BY_NAME(CacheTrackRecorder)
@@ -351,7 +352,10 @@ void UCacheTrackRecorder::RecordSelectedTracks(TSharedPtr<ISequencer> Sequencer,
 	{
 		if (IMovieSceneCachedTrack* CacheTrack = Cast<IMovieSceneCachedTrack>(SelectedTrack))
 		{
-			CacheTracks.Add(CacheTrack);
+			if (CacheTrack->IsCacheRecordingAllowed())
+			{
+				CacheTracks.Add(CacheTrack);
+			}
 		}
 	}
 
@@ -463,6 +467,12 @@ bool UCacheTrackRecorder::Initialize(ULevelSequence* RootLevelSequence, const TA
 		TRange<double> NewRange(ViewRangeStartSeconds - 0.5f, ViewRangeStartSeconds + (Range.GetUpperBoundValue() - Range.GetLowerBoundValue()) + 0.5f);
 		Sequencer->SetViewRange(NewRange, EViewRangeInterpolation::Immediate);
 		Sequencer->SetClampRange(TRange(Sequencer->GetViewRange()));
+		ESequencerLoopMode LoopMode = Sequencer->GetSequencerSettings()->GetLoopMode();
+		if (LoopMode != SLM_NoLoop)
+		{
+			Sequencer->GetSequencerSettings()->SetLoopMode(SLM_NoLoop);
+			OnStopCleanup.Add([LoopMode, this] { WeakSequencer.Pin()->GetSequencerSettings()->SetLoopMode(LoopMode); });
+		}
 	}
 
 	return true;
@@ -511,6 +521,20 @@ void UCacheTrackRecorder::InitializeFromParameters()
 			OnStopCleanup.Add(RestoreImmersiveMode);
 		}
 	}
+
+	int32 MinScalability = -1;
+	for (FCachedTrackSource& CacheTrack : CacheTracks)
+	{
+		MinScalability = FMath::Max(MinScalability, CacheTrack.Track->GetMinimumEngineScalabilitySetting());
+	}
+	if (MinScalability >= 0)
+	{
+		Scalability::FQualityLevels CurrentQualityLevels = Scalability::GetQualityLevels();
+		Scalability::FQualityLevels NewQualityLevels = CurrentQualityLevels;
+		NewQualityLevels.SetFromSingleQualityLevel(MinScalability);
+		SetQualityLevels(NewQualityLevels, true);
+		OnStopCleanup.Add([CurrentQualityLevels] { SetQualityLevels(CurrentQualityLevels, true); });
+	}
 }
 
 bool UCacheTrackRecorder::ShouldShowNotifications()
diff --git a/Engine/Source/Runtime/MovieScene/Public/Tracks/MovieSceneCachedTrack.h b/Engine/Source/Runtime/MovieScene/Public/Tracks/MovieSceneCachedTrack.h
--- a/Engine/Source/Runtime/MovieScene/Public/Tracks/MovieSceneCachedTrack.h
+++ b/Engine/Source/Runtime/MovieScene/Public/Tracks/MovieSceneCachedTrack.h
@@ -32,4 +32,6 @@ public:
 	virtual void SetCacheRecordingAllowed(bool bShouldRecord) {};
 
 	virtual bool IsCacheRecordingAllowed() const { return true; };
+
+	virtual int32 GetMinimumEngineScalabilitySetting() const { return -1; };
 };
