- Fix crash when data interfaces classes are no longer valid

diff --git a/Engine/Plugins/FX/Niagara/Source/Niagara/Private/NiagaraScript.cpp b/Engine/Plugins/FX/Niagara/Source/Niagara/Private/NiagaraScript.cpp
--- a/Engine/Plugins/FX/Niagara/Source/Niagara/Private/NiagaraScript.cpp
+++ b/Engine/Plugins/FX/Niagara/Source/Niagara/Private/NiagaraScript.cpp
@@ -1789,8 +1789,8 @@ void UNiagaraScript::GenerateDefaultFunctionBindings()
 		const auto& ScriptDataInterfaces = ScriptParameterStore->GetDataInterfaces();
 
 		const int32 DataInterfaceCount = FMath::Min(CachedScriptVM.DataInterfaceInfo.Num(), ScriptDataInterfaces.Num());
-		check(DataInterfaceCount == CachedScriptVM.DataInterfaceInfo.Num());
-		check(DataInterfaceCount == ScriptDataInterfaces.Num());
+		ensureMsgf(DataInterfaceCount == CachedScriptVM.DataInterfaceInfo.Num(), TEXT("DataInterface count does not match VM this is likely caused by missing data interface classes"));
+		ensureMsgf(DataInterfaceCount == ScriptDataInterfaces.Num(), TEXT("DataInterface count does not match script data interfaces this is likely caused by missing data interface classes"));
 
 		for (const FVMExternalFunctionBindingInfo& BindingInfo : CachedScriptVM.CalledVMExternalFunctions)
 		{
diff --git a/Engine/Plugins/FX/Niagara/Source/Niagara/Private/NiagaraScriptExecutionParameterStore.cpp b/Engine/Plugins/FX/Niagara/Source/Niagara/Private/NiagaraScriptExecutionParameterStore.cpp
--- a/Engine/Plugins/FX/Niagara/Source/Niagara/Private/NiagaraScriptExecutionParameterStore.cpp
+++ b/Engine/Plugins/FX/Niagara/Source/Niagara/Private/NiagaraScriptExecutionParameterStore.cpp
@@ -245,7 +245,10 @@ void FNiagaraScriptExecutionParameterStore::AddScriptParams(UNiagaraScript* Scri
 	{
 		int32 VarOffset = INDEX_NONE;
 		bAdded |= AddParameter(ResolvedDataInterface.ParameterStoreVariable, false, false, &VarOffset);
-		SetDataInterface(ResolvedDataInterface.ResolvedDataInterface, VarOffset);
+		if (GetDataInterfaces().IsValidIndex(VarOffset))
+		{
+			SetDataInterface(ResolvedDataInterface.ResolvedDataInterface, VarOffset);
+		}
 	}
 
 	check(Script->GetVMExecutableData().UObjectInfos.Num() == Script->GetCachedDefaultUObjects().Num());
@@ -253,7 +256,10 @@ void FNiagaraScriptExecutionParameterStore::AddScriptParams(UNiagaraScript* Scri
 	{
 		int32 VarOffset = INDEX_NONE;
 		bAdded |= AddParameter(Info.ResolvedVariable, false, false, &VarOffset);
-		SetUObject(Info.Object, VarOffset);
+		if (GetUObjects().IsValidIndex(VarOffset))
+		{
+			SetUObject(Info.Object, VarOffset);
+		}
 	}
 
 	if (bAdded && bTriggerRebind)
