Fix for RDG crash caused by NDC.

diff --git a/Engine/Plugins/FX/Niagara/Source/Niagara/Internal/DataInterface/NiagaraDataInterfaceDataChannelRead.h b/Engine/Plugins/FX/Niagara/Source/Niagara/Internal/DataInterface/NiagaraDataInterfaceDataChannelRead.h
--- a/Engine/Plugins/FX/Niagara/Source/Niagara/Internal/DataInterface/NiagaraDataInterfaceDataChannelRead.h
+++ b/Engine/Plugins/FX/Niagara/Source/Niagara/Internal/DataInterface/NiagaraDataInterfaceDataChannelRead.h
@@ -259,6 +259,7 @@ struct FNiagaraDataInterfaceProxy_DataChannelRead : public FNiagaraDataInterface
 	virtual void GetDispatchArgs(const FNDIGpuComputeDispatchArgsGenContext& Context) override;
 
 	virtual void PreStage(const FNDIGpuComputePreStageContext& Context)override;
+	virtual void PostSimulate(const FNDIGpuComputePostSimulateContext& Context)override;
 
 	/** Persistent per instance data on the RT. Constructed when consuming data passed from GT->RT. */
 	struct FInstanceData
diff --git a/Engine/Plugins/FX/Niagara/Source/Niagara/Private/DataInterface/NiagaraDataInterfaceDataChannelRead.cpp b/Engine/Plugins/FX/Niagara/Source/Niagara/Private/DataInterface/NiagaraDataInterfaceDataChannelRead.cpp
--- a/Engine/Plugins/FX/Niagara/Source/Niagara/Private/DataInterface/NiagaraDataInterfaceDataChannelRead.cpp
+++ b/Engine/Plugins/FX/Niagara/Source/Niagara/Private/DataInterface/NiagaraDataInterfaceDataChannelRead.cpp
@@ -1982,6 +1982,15 @@ void FNiagaraDataInterfaceProxy_DataChannelRead::PreStage(const FNDIGpuComputePr
 	}
 }
 
+void FNiagaraDataInterfaceProxy_DataChannelRead::PostSimulate(const FNDIGpuComputePostSimulateContext& Context)
+{
+	if (Context.IsFinalPostSimulate())
+	{
+		FNiagaraDataInterfaceProxy_DataChannelRead::FInstanceData* InstanceData = SystemInstancesToProxyData_RT.Find(Context.GetSystemInstanceID());
+		InstanceData->NDCSpawnDataBuffer = nullptr;
+	}
+}
+
 void FNiagaraDataInterfaceProxy_DataChannelRead::ConsumePerInstanceDataFromGameThread(void* PerInstanceData, const FNiagaraSystemInstanceID& Instance)
 {
 	FNDIDataChannelReadInstanceData_RT& SourceData = *reinterpret_cast<FNDIDataChannelReadInstanceData_RT*>(PerInstanceData);
