fix crash when trying to from NDC channel - set NDCSpawnDataBuffer to null in PostSimulate to avoid dangling ptr

diff --git a/Engine/Plugins/FX/Niagara/Source/Niagara/Internal/DataInterface/NiagaraDataInterfaceDataChannelRead.h b/Engine/Plugins/FX/Niagara/Source/Niagara/Internal/DataInterface/NiagaraDataInterfaceDataChannelRead.h
--- a/Engine/Plugins/FX/Niagara/Source/Niagara/Internal/DataInterface/NiagaraDataInterfaceDataChannelRead.h
+++ b/Engine/Plugins/FX/Niagara/Source/Niagara/Internal/DataInterface/NiagaraDataInterfaceDataChannelRead.h
@@ -261,6 +261,8 @@ struct FNiagaraDataInterfaceProxy_DataChannelRead : public FNiagaraDataInterface
 	virtual void PreStage(const FNDIGpuComputePreStageContext& Context)override;
 	virtual void PostSimulate(const FNDIGpuComputePostSimulateContext& Context)override;
 
+	virtual void PostSimulate(const FNDIGpuComputePostSimulateContext& Context) override;
+
 	/** Persistent per instance data on the RT. Constructed when consuming data passed from GT->RT. */
 	struct FInstanceData
 	{
