Removed unneeded fguid initialization that was causing warnings

diff --git a/Engine/Plugins/FX/Niagara/Source/Niagara/Public/NiagaraAssetTagDefinitions.h b/Engine/Plugins/FX/Niagara/Source/Niagara/Public/NiagaraAssetTagDefinitions.h
--- a/Engine/Plugins/FX/Niagara/Source/Niagara/Public/NiagaraAssetTagDefinitions.h
+++ b/Engine/Plugins/FX/Niagara/Source/Niagara/Public/NiagaraAssetTagDefinitions.h
@@ -59,7 +59,7 @@ public:
 
 	/** The Tag Guid identifies this tag. This makes it possible to change the AssetTag name without it affecting functionality. */
 	UPROPERTY(VisibleAnywhere, Category="Properties")
-	FGuid TagGuid = FGuid();
+	FGuid TagGuid;
 
 	bool operator<(const FNiagaraAssetTagDefinition& Other) const
 	{		
