- Modify static variable tooltip

diff --git a/Engine/Plugins/FX/Niagara/Source/NiagaraEditorWidgets/Private/Stack/SNiagaraStackFunctionInputValue.cpp b/Engine/Plugins/FX/Niagara/Source/NiagaraEditorWidgets/Private/Stack/SNiagaraStackFunctionInputValue.cpp
--- a/Engine/Plugins/FX/Niagara/Source/NiagaraEditorWidgets/Private/Stack/SNiagaraStackFunctionInputValue.cpp
+++ b/Engine/Plugins/FX/Niagara/Source/NiagaraEditorWidgets/Private/Stack/SNiagaraStackFunctionInputValue.cpp
@@ -1014,7 +1014,7 @@ const FSlateBrush* SNiagaraStackFunctionInputValue::GetTypeModifierIcon() const
 
 FText SNiagaraStackFunctionInputValue::GetTypeModifierIconToolTip() const
 {
-	return LOCTEXT("TypeModifierTooltip", "Static variables can only be set once in the graph and are meant to communicate inputs to static switches across the emitter.");
+	return LOCTEXT("TypeModifierTooltip", "This variable is Static and cannot be changed at runtime, or linked to any dynamic attributes. Static variables can be used to reduce script instructions & memory.");
 }
 
 FSlateColor SNiagaraStackFunctionInputValue::GetTypeModifierIconColor() const
