Reverted previous fix as it introduces flickering in the UI until a proper fix is found

diff --git a/Engine/Plugins/FX/Niagara/Source/NiagaraEditorWidgets/Private/SNiagaraOverviewStackNode.cpp b/Engine/Plugins/FX/Niagara/Source/NiagaraEditorWidgets/Private/SNiagaraOverviewStackNode.cpp
--- a/Engine/Plugins/FX/Niagara/Source/NiagaraEditorWidgets/Private/SNiagaraOverviewStackNode.cpp
+++ b/Engine/Plugins/FX/Niagara/Source/NiagaraEditorWidgets/Private/SNiagaraOverviewStackNode.cpp
@@ -421,7 +421,7 @@ void SNiagaraOverviewStackNode::Tick(const FGeometry& AllottedGeometry, const do
 
 		if (bTopContentBarRefreshPending && TopContentBarSlot != nullptr)
 		{
-			TopContentBarSlot->AttachWidget(CreateTopContentBar());
+			CreateTopContentBar();
 			bTopContentBarRefreshPending = false;
 		}
 
