Fixed member initialization test by adding tags

diff --git a/Engine/Plugins/FX/Niagara/Source/Niagara/Public/NiagaraAssetTagDefinitions.h b/Engine/Plugins/FX/Niagara/Source/Niagara/Public/NiagaraAssetTagDefinitions.h
--- a/Engine/Plugins/FX/Niagara/Source/Niagara/Public/NiagaraAssetTagDefinitions.h
+++ b/Engine/Plugins/FX/Niagara/Source/Niagara/Public/NiagaraAssetTagDefinitions.h
@@ -30,10 +30,7 @@ struct FNiagaraAssetTagDefinition
 	GENERATED_BODY()
 
 public:
-	FNiagaraAssetTagDefinition()
-	{
-		TagGuid = FGuid::NewGuid();
-	}
+	FNiagaraAssetTagDefinition() = default;
 
 	FNiagaraAssetTagDefinition(FText InAssetTag, int32 InAssetFlags, FText InDescription, ENiagaraAssetTagDefinitionImportance InDisplayType, FLinearColor InColor, FGuid InTagGuid);
 	
@@ -58,8 +55,8 @@ public:
 	FLinearColor Color = FLinearColor::Black;
 
 	/** The Tag Guid identifies this tag. This makes it possible to change the AssetTag name without it affecting functionality. */
-	UPROPERTY(VisibleAnywhere, Category="Properties")
-	FGuid TagGuid;
+	UPROPERTY(VisibleAnywhere, Category="Properties", meta=(IgnoreForMemberInitializationTest))
+	FGuid TagGuid = FGuid::NewGuid();
 
 	bool operator<(const FNiagaraAssetTagDefinition& Other) const
 	{		
