- Fix for lack of GPU stats in performance mode

diff --git a/Engine/Plugins/FX/Niagara/Source/Niagara/Private/NiagaraDebugHud.cpp b/Engine/Plugins/FX/Niagara/Source/Niagara/Private/NiagaraDebugHud.cpp
--- a/Engine/Plugins/FX/Niagara/Source/Niagara/Private/NiagaraDebugHud.cpp
+++ b/Engine/Plugins/FX/Niagara/Source/Niagara/Private/NiagaraDebugHud.cpp
@@ -788,6 +788,19 @@ namespace NiagaraDebugLocal
 			}
 		}
 	}
+
+#if WITH_NIAGARA_GPU_PROFILER
+	bool NeedsGpuStatCapture()
+	{
+		return
+			Settings.bOverviewEnabled &&
+			(
+				(Settings.OverviewMode == ENiagaraDebugHUDOverviewMode::Performance) ||
+				(Settings.OverviewMode == ENiagaraDebugHUDOverviewMode::GpuComputePerformance) ||
+				(Settings.OverviewMode == ENiagaraDebugHUDOverviewMode::PerformanceGraph && Settings.PerfGraphMode == ENiagaraDebugHUDPerfGraphMode::GPU)
+			);
+	}
+#endif
 }
 
 FNiagaraDebugHud::FNiagaraDebugHud(UWorld* World)
@@ -949,11 +962,7 @@ void FNiagaraDebugHud::GatherSystemInfo()
 	}
 
 #if WITH_NIAGARA_GPU_PROFILER
-	// Update Gpu capture state
-	bool bEnableGPUStats = Settings.bOverviewEnabled && 
-		((Settings.OverviewMode == ENiagaraDebugHUDOverviewMode::GpuComputePerformance) || (Settings.OverviewMode == ENiagaraDebugHUDOverviewMode::PerformanceGraph && Settings.PerfGraphMode == ENiagaraDebugHUDPerfGraphMode::GPU));
-							
-	GpuProfilerListener.SetEnabled(bEnableGPUStats);
+	GpuProfilerListener.SetEnabled(NeedsGpuStatCapture());
 #endif
 
 	// When not enabled do nothing
