Fixed a crash when using a Data Channel asset without a type set

diff --git a/Engine/Plugins/FX/Niagara/Source/NiagaraBlueprintNodes/Internal/K2Node_DataChannelBase.h b/Engine/Plugins/FX/Niagara/Source/NiagaraBlueprintNodes/Internal/K2Node_DataChannelBase.h
--- a/Engine/Plugins/FX/Niagara/Source/NiagaraBlueprintNodes/Internal/K2Node_DataChannelBase.h
+++ b/Engine/Plugins/FX/Niagara/Source/NiagaraBlueprintNodes/Internal/K2Node_DataChannelBase.h
@@ -15,6 +15,7 @@ class UK2Node_DataChannelBase : public UK2Node_CallFunction
 public:
 
 	NIAGARABLUEPRINTNODES_API UNiagaraDataChannel* GetDataChannel() const;
+	NIAGARABLUEPRINTNODES_API bool HasValidDataChannel() const;
 
 	virtual void PostLoad() override;
 
@@ -37,14 +38,16 @@ public:
 
 protected:
 
-	UPROPERTY()
-	TObjectPtr<UNiagaraDataChannelAsset> DataChannel;
-
 	UEdGraphPin* GetChannelSelectorPin() const;
 
 #if WITH_EDITORONLY_DATA
 	UPROPERTY()
 	FGuid DataChannelVersion;
 #endif
+
+private:
+
+	UPROPERTY()
+	TObjectPtr<UNiagaraDataChannelAsset> DataChannel;
 };
 
diff --git a/Engine/Plugins/FX/Niagara/Source/NiagaraBlueprintNodes/Private/K2Node_DataChannelBase.cpp b/Engine/Plugins/FX/Niagara/Source/NiagaraBlueprintNodes/Private/K2Node_DataChannelBase.cpp
--- a/Engine/Plugins/FX/Niagara/Source/NiagaraBlueprintNodes/Private/K2Node_DataChannelBase.cpp
+++ b/Engine/Plugins/FX/Niagara/Source/NiagaraBlueprintNodes/Private/K2Node_DataChannelBase.cpp
@@ -13,12 +13,17 @@ UNiagaraDataChannel* UK2Node_DataChannelBase::GetDataChannel() const
 	return DataChannel ? DataChannel->Get() : nullptr;
 }
 
+bool UK2Node_DataChannelBase::HasValidDataChannel() const
+{
+	return DataChannel && DataChannel->Get();
+}
+
 void UK2Node_DataChannelBase::PostLoad()
 {
 	Super::PostLoad();
 
 #if WITH_EDITORONLY_DATA
-	if (DataChannel && DataChannel->Get()->GetVersion() != DataChannelVersion && HasValidBlueprint())
+	if (HasValidDataChannel() && GetDataChannel()->GetVersion() != DataChannelVersion && HasValidBlueprint())
 	{
 		ReconstructNode();
 	}
@@ -30,9 +35,9 @@ void UK2Node_DataChannelBase::AllocateDefaultPins()
 	Super::AllocateDefaultPins();
 
 #if WITH_EDITORONLY_DATA
-	if (DataChannel)
+	if (HasValidDataChannel())
 	{
-		DataChannelVersion = DataChannel->Get()->GetVersion();
+		DataChannelVersion = GetDataChannel()->GetVersion();
 	}
 #endif
 }
@@ -92,7 +97,7 @@ void UK2Node_DataChannelBase::PreloadRequiredAssets()
 {
 	Super::PreloadRequiredAssets();
 
-	if (DataChannel)
+	if (HasValidDataChannel())
 	{
 		PreloadObject(DataChannel);
 		PreloadObject(DataChannel->Get());
diff --git a/Engine/Plugins/FX/Niagara/Source/NiagaraBlueprintNodes/Private/K2Node_ReadDataChannel.cpp b/Engine/Plugins/FX/Niagara/Source/NiagaraBlueprintNodes/Private/K2Node_ReadDataChannel.cpp
--- a/Engine/Plugins/FX/Niagara/Source/NiagaraBlueprintNodes/Private/K2Node_ReadDataChannel.cpp
+++ b/Engine/Plugins/FX/Niagara/Source/NiagaraBlueprintNodes/Private/K2Node_ReadDataChannel.cpp
@@ -26,9 +26,9 @@ void UK2Node_ReadDataChannel::AllocateDefaultPins()
 {
 	Super::AllocateDefaultPins();
 
-	if (DataChannel)
+	if (HasValidDataChannel())
 	{
-		for (const FNiagaraDataChannelVariable& InVar : DataChannel->Get()->GetVariables())
+		for (const FNiagaraDataChannelVariable& InVar : GetDataChannel()->GetVariables())
 		{
 			if (IgnoredVariables.Contains(InVar.Version))
 			{
@@ -58,7 +58,7 @@ void UK2Node_ReadDataChannel::ExpandNode(FKismetCompilerContext& CompilerContext
 	const UEdGraphSchema_K2* Schema = CompilerContext.GetSchema();
 	UEdGraphPin* SuccessPin = FindPinChecked(FName("Success"), EGPD_Output);
 	UEdGraphPin* FailurePin = FindPinChecked(FName("Failure"), EGPD_Output);
-	if (!DataChannel)
+	if (!HasValidDataChannel())
 	{
 		// replace function with a noop if we don't have a data channel
 		UK2Node_ExecutionSequence* NoopNode = CompilerContext.SpawnIntermediateNode<UK2Node_ExecutionSequence>(this, SourceGraph);
@@ -69,7 +69,7 @@ void UK2Node_ReadDataChannel::ExpandNode(FKismetCompilerContext& CompilerContext
 		Schema->BreakPinLinks(*SuccessPin, false);
 		return;
 	}
-	if (DataChannelVersion != DataChannel->Get()->GetVersion())
+	if (DataChannelVersion != GetDataChannel()->GetVersion())
 	{
 		CompilerContext.MessageLog.Error(*LOCTEXT("StaleNode", "Node is out of sync with the data channel asset, please refresh node to fix up the pins - @@").ToString(), this);
 		return;
@@ -114,7 +114,7 @@ void UK2Node_ReadDataChannel::ExpandNode(FKismetCompilerContext& CompilerContext
 	// create the read function nodes
 	UEdGraphPin* LastIsValidPin = nullptr;
 	UEdGraphPin* IndexPin = FindPinChecked(FName("Index"), EGPD_Input);
-	for (const FNiagaraDataChannelVariable& InVar : DataChannel->Get()->GetVariables())
+	for (const FNiagaraDataChannelVariable& InVar : GetDataChannel()->GetVariables())
 	{
 		if (IgnoredVariables.Contains(InVar.Version))
 		{
diff --git a/Engine/Plugins/FX/Niagara/Source/NiagaraBlueprintNodes/Private/K2Node_WriteDataChannel.cpp b/Engine/Plugins/FX/Niagara/Source/NiagaraBlueprintNodes/Private/K2Node_WriteDataChannel.cpp
--- a/Engine/Plugins/FX/Niagara/Source/NiagaraBlueprintNodes/Private/K2Node_WriteDataChannel.cpp
+++ b/Engine/Plugins/FX/Niagara/Source/NiagaraBlueprintNodes/Private/K2Node_WriteDataChannel.cpp
@@ -23,9 +23,9 @@ void UK2Node_WriteDataChannel::AllocateDefaultPins()
 {
 	Super::AllocateDefaultPins();
 
-	if (DataChannel)
+	if (HasValidDataChannel())
 	{
-		for (const FNiagaraDataChannelVariable& InVar : DataChannel->Get()->GetVariables())
+		for (const FNiagaraDataChannelVariable& InVar : GetDataChannel()->GetVariables())
 		{
 			if (IgnoredVariables.Contains(InVar.Version))
 			{
@@ -52,7 +52,7 @@ void UK2Node_WriteDataChannel::GetMenuActions(FBlueprintActionDatabaseRegistrar&
 
 void UK2Node_WriteDataChannel::ExpandNode(FKismetCompilerContext& CompilerContext, UEdGraph* SourceGraph)
 {
-	if (!DataChannel)
+	if (!HasValidDataChannel())
 	{
 		// replace function with a noop if we don't have a data channel
 		UK2Node_ExecutionSequence* NoopNode = CompilerContext.SpawnIntermediateNode<UK2Node_ExecutionSequence>(this, SourceGraph);
@@ -61,7 +61,7 @@ void UK2Node_WriteDataChannel::ExpandNode(FKismetCompilerContext& CompilerContex
 		CompilerContext.MovePinLinksToIntermediate(*GetThenPin(), *NoopNode->GetThenPinGivenIndex(0));
 		return;
 	}
-	if (DataChannelVersion != DataChannel->Get()->GetVersion())
+	if (DataChannelVersion != GetDataChannel()->GetVersion())
 	{
 		CompilerContext.MessageLog.Error(*LOCTEXT("StaleNode", "Node is out of sync with the data channel asset, please refresh node to fix up the pins - @@").ToString(), this);
 		return;
@@ -108,7 +108,7 @@ void UK2Node_WriteDataChannel::ExpandNode(FKismetCompilerContext& CompilerContex
 	// create the write function nodes
 	UEdGraphPin* LastExecPin = CreateWriterNode->GetThenPin();
 	UEdGraphPin* WriterResultPin = CreateWriterNode->GetReturnValuePin();
-	for (const FNiagaraDataChannelVariable& InVar : DataChannel->Get()->GetVariables())
+	for (const FNiagaraDataChannelVariable& InVar : GetDataChannel()->GetVariables())
 	{
 		if (IgnoredVariables.Contains(InVar.Version))
 		{
