Niagara: Missing file from CL 31062557

Somehow submit tool didn't submit this even though it was in the CL, then did some kind of revert -k on it :/

diff --git a/Engine/Plugins/FX/Niagara/Source/NiagaraEditor/Private/Sequencer/LevelSequence/NiagaraSystemTrackEditor.cpp b/Engine/Plugins/FX/Niagara/Source/NiagaraEditor/Private/Sequencer/LevelSequence/NiagaraSystemTrackEditor.cpp
--- a/Engine/Plugins/FX/Niagara/Source/NiagaraEditor/Private/Sequencer/LevelSequence/NiagaraSystemTrackEditor.cpp
+++ b/Engine/Plugins/FX/Niagara/Source/NiagaraEditor/Private/Sequencer/LevelSequence/NiagaraSystemTrackEditor.cpp
@@ -16,6 +16,9 @@
 #include "NiagaraSystemEditorData.h"
 #include "SEnumCombo.h"
 
+#include "MVVM/ViewModels/OutlinerColumns/OutlinerColumnTypes.h"
+#include "MVVM/Extensions/ITrackExtension.h"
+
 #define LOCTEXT_NAMESPACE "NiagaraSystemTrackEditor"
 
 TSharedRef<ISequencerTrackEditor> FNiagaraSystemTrackEditor::CreateTrackEditor(TSharedRef<ISequencer> InSequencer)
@@ -132,15 +135,20 @@ void FNiagaraSystemTrackEditor::BuildObjectBindingTrackMenu(FMenuBuilder& MenuBu
 	}
 }
 
-TSharedPtr<SWidget> FNiagaraSystemTrackEditor::BuildOutlinerEditWidget(const FGuid& ObjectBinding, UMovieSceneTrack* Track, const FBuildEditWidgetParams& Params)
+TSharedPtr<SWidget> FNiagaraSystemTrackEditor::BuildOutlinerColumnWidget(const FBuildColumnWidgetParams& Params, const FName& ColumnName)
 {
-	if (UMovieSceneNiagaraSystemTrack* NiagaraSystemTrack = Cast<UMovieSceneNiagaraSystemTrack>(Track))
+	using namespace UE::Sequencer;
+
+	UMovieSceneNiagaraSystemTrack* NiagaraSystemTrack = Cast<UMovieSceneNiagaraSystemTrack>(Params.TrackModel->GetTrack());
+	if (!NiagaraSystemTrack)
+	{
+		return nullptr;
+	}
+
+	if (ColumnName == FCommonOutlinerNames::Edit)
 	{
-		return SNew(SHorizontalBox)
-			+ SHorizontalBox::Slot()
-			.Padding(5, 0, 5, 0)
-			.VAlign(VAlign_Center)
-			.AutoWidth()
+		return SNew(SBox)
+			.HAlign(HAlign_Left)
 			[
 				SNew(SEnumComboBox, StaticEnum<ENiagaraAgeUpdateMode>())
 				.Font(FAppStyle::GetFontStyle(TEXT("PropertyWindow.NormalFont")))
@@ -168,7 +176,7 @@ TSharedPtr<SWidget> FNiagaraSystemTrackEditor::BuildOutlinerEditWidget(const FGu
 				})
 			];
 	}
-	return FMovieSceneTrackEditor::BuildOutlinerEditWidget(ObjectBinding, Track, Params);
+	return FMovieSceneTrackEditor::BuildOutlinerColumnWidget(Params, ColumnName);
 }
 
 void FNiagaraSystemTrackEditor::AddDefaultSystemTracks(const AActor& SourceActor, const FGuid& Binding,	TSharedPtr<ISequencer> Sequencer)
