PR #11396: [Niagara] Fixed the problem of repeated recycling of the Resource List

diff --git a/Engine/Plugins/FX/Niagara/Source/Niagara/Private/NiagaraRendererRibbons.cpp b/Engine/Plugins/FX/Niagara/Source/Niagara/Private/NiagaraRendererRibbons.cpp
--- a/Engine/Plugins/FX/Niagara/Source/Niagara/Private/NiagaraRendererRibbons.cpp
+++ b/Engine/Plugins/FX/Niagara/Source/Niagara/Private/NiagaraRendererRibbons.cpp
@@ -320,17 +320,16 @@ struct FNiagaraDynamicDataRibbon : public FNiagaraDynamicDataBase
 
 struct FNiagaraRibbonRenderingFrameViewResources
 {
-	FNiagaraRibbonVertexFactory		VertexFactory;
-	FNiagaraRibbonUniformBufferRef	UniformBuffer;
-	FNiagaraRibbonIndexBuffer		IndexBuffer;
-	FRWBuffer						IndirectDrawBuffer;
-	FNiagaraIndexGenerationInput	IndexGenerationSettings;
+	FNiagaraRibbonVertexFactory				VertexFactory;
+	FNiagaraRibbonUniformBufferRef			UniformBuffer;
+	TSharedPtr<FNiagaraRibbonIndexBuffer>	IndexBuffer;
+	FRWBuffer								IndirectDrawBuffer;
+	FNiagaraIndexGenerationInput			IndexGenerationSettings;
 
 	~FNiagaraRibbonRenderingFrameViewResources()
 	{
 		UniformBuffer.SafeRelease();
 		VertexFactory.ReleaseResource();
-		IndexBuffer.ReleaseResource();
 		IndirectDrawBuffer.Release();
 	}
 };
@@ -536,9 +535,9 @@ class FNiagaraGpuRibbonsDataManager final : public FNiagaraGpuComputeDataManager
 
 	struct FIndexBufferEntry
 	{
-		uint64						FrameUsed = 0;
-		int32						NumIndices = 0;
-		FNiagaraRibbonIndexBuffer	Buffer;
+		uint64									FrameUsed = 0;
+		int32									NumIndices = 0;
+		TSharedPtr<FNiagaraRibbonIndexBuffer>	Buffer;
 	};
 
 public:
@@ -584,7 +583,7 @@ public:
 		return BufferEntry->Buffer;
 	}
 
-	FNiagaraRibbonIndexBuffer GetOrAllocateIndexBuffer(FRHICommandListBase& RHICmdList, int32 NumIndices, int32 MaxIndicesEstimate)
+	TSharedPtr<FNiagaraRibbonIndexBuffer>& GetOrAllocateIndexBuffer(FRHICommandListBase& RHICmdList, int32 NumIndices, int32 MaxIndicesEstimate)
 	{
 		if (GNiagaraRibbonGpuBufferCachePurgeCounter >= 0)
 		{
@@ -601,7 +600,8 @@ public:
 		{
 			BufferEntry = &Index32BufferCache.AddDefaulted_GetRef();
 			BufferEntry->NumIndices = NumIndices;
-			BufferEntry->Buffer.Initialize(RHICmdList, NumIndices);
+			BufferEntry->Buffer = MakeShared<FNiagaraRibbonIndexBuffer>();
+			BufferEntry->Buffer->Initialize(RHICmdList, NumIndices);
 		}
 
 		BufferEntry->FrameUsed = FrameCounter;
@@ -1080,8 +1080,8 @@ void FNiagaraRendererRibbons::GetDynamicRayTracingInstances(FRayTracingMaterialG
 	RayTracingInstance.Geometry = &RayTracingGeometry;
 	RayTracingInstance.InstanceTransforms.Add(FMatrix::Identity);
 	
-	RayTracingGeometry.Initializer.IndexBuffer = RenderingViewResources->IndexBuffer.IndexBufferRHI;// PerViewGeneratedData.IndexAllocation.IndexBuffer->IndexBufferRHI;
-	RayTracingGeometry.Initializer.IndexBufferOffset = RenderingViewResources->IndexBuffer.FirstIndex;//PerViewGeneratedData.IndexAllocation.FirstIndex * PerViewGeneratedData.IndexAllocation.IndexStride;
+	RayTracingGeometry.Initializer.IndexBuffer = RenderingViewResources->IndexBuffer->IndexBufferRHI;// PerViewGeneratedData.IndexAllocation.IndexBuffer->IndexBufferRHI;
+	RayTracingGeometry.Initializer.IndexBufferOffset = RenderingViewResources->IndexBuffer->FirstIndex;//PerViewGeneratedData.IndexAllocation.FirstIndex * PerViewGeneratedData.IndexAllocation.IndexStride;
 	
 	FMeshBatch MeshBatch;
 	
@@ -1888,16 +1888,17 @@ void FNiagaraRendererRibbons::GenerateIndexBufferForView(
 		}
 		else
 		{
+			RenderingViewResources->IndexBuffer = MakeShared<FNiagaraRibbonIndexBuffer>();
 			if (GeneratedData.TotalBitCount <= 16)
 			{
 				FGlobalDynamicIndexBuffer::FAllocationEx IndexAllocation = Collector.GetDynamicIndexBuffer().Allocate<uint16>(GeneratedData.TotalNumIndices);
-				RenderingViewResources->IndexBuffer.Initialize(RHICmdList, IndexAllocation);
+				RenderingViewResources->IndexBuffer->Initialize(RHICmdList, IndexAllocation);
 				GenerateIndexBufferCPU<uint16>(GeneratedData, DynamicDataRibbon, ShapeState, reinterpret_cast<uint16*>(IndexAllocation.Buffer), View, ViewOriginForDistanceCulling, FeatureLevel, DrawDirection);
 			}
 			else
 			{
 				FGlobalDynamicIndexBuffer::FAllocationEx IndexAllocation = Collector.GetDynamicIndexBuffer().Allocate<uint32>(GeneratedData.TotalNumIndices);
-				RenderingViewResources->IndexBuffer.Initialize(RHICmdList, IndexAllocation);
+				RenderingViewResources->IndexBuffer->Initialize(RHICmdList, IndexAllocation);
 				GenerateIndexBufferCPU<uint32>(GeneratedData, DynamicDataRibbon, ShapeState, reinterpret_cast<uint32*>(IndexAllocation.Buffer), View, ViewOriginForDistanceCulling, FeatureLevel, DrawDirection);
 			}
 		}
@@ -2109,8 +2110,8 @@ inline void FNiagaraRendererRibbons::SetupMeshBatchAndCollectorResourceForView(F
 	OutMeshBatch.MaterialRenderProxy = bIsWireframe? UMaterial::GetDefaultMaterial(MD_Surface)->GetRenderProxy() : MaterialRenderProxy;
 	
 	FMeshBatchElement& MeshElement = OutMeshBatch.Elements[0];
-	MeshElement.IndexBuffer = &RenderingViewResources->IndexBuffer;
-	MeshElement.FirstIndex = RenderingViewResources->IndexBuffer.FirstIndex;
+	MeshElement.IndexBuffer = RenderingViewResources->IndexBuffer.Get();
+	MeshElement.FirstIndex = RenderingViewResources->IndexBuffer->FirstIndex;
 	MeshElement.NumInstances = 1;
 	MeshElement.MinVertexIndex = 0;
 	MeshElement.MaxVertexIndex = 0;
@@ -2215,7 +2216,7 @@ void FNiagaraRendererRibbons::InitializeViewIndexBuffersGPU(FRHICommandListImmed
 		FNiagaraRibbonGenerateIndices Params;
 		FMemory::Memzero(Params);
 		
-		Params.GeneratedIndicesBuffer = RenderingViewResources->IndexBuffer.UAV;
+		Params.GeneratedIndicesBuffer = RenderingViewResources->IndexBuffer->UAV;
 		Params.SortedIndices = VertexBuffers.SortedIndicesBuffer.SRV;
 		Params.MultiRibbonIndices = VertexBuffers.MultiRibbonIndicesBuffer.SRV;
 		Params.Segments = VertexBuffers.SegmentsBuffer.SRV;
@@ -2243,9 +2244,9 @@ void FNiagaraRendererRibbons::InitializeViewIndexBuffersGPU(FRHICommandListImmed
 		Params.SubSegmentBitShift = RenderingViewResources->IndexGenerationSettings.SubSegmentBitShift;
 		Params.SubSegmentBitMask = RenderingViewResources->IndexGenerationSettings.SubSegmentBitMask;
 		
-		RHICmdList.Transition(FRHITransitionInfo(RenderingViewResources->IndexBuffer.UAV, ERHIAccess::VertexOrIndexBuffer, ERHIAccess::UAVCompute));
+		RHICmdList.Transition(FRHITransitionInfo(RenderingViewResources->IndexBuffer->UAV, ERHIAccess::VertexOrIndexBuffer, ERHIAccess::UAVCompute));
 		FComputeShaderUtils::DispatchIndirect(RHICmdList, ComputeShader, Params, RenderingViewResources->IndirectDrawBuffer.Buffer, IndirectDispatchArgsOffset);
-		RHICmdList.Transition(FRHITransitionInfo(RenderingViewResources->IndexBuffer.UAV, ERHIAccess::UAVCompute, ERHIAccess::VertexOrIndexBuffer));
+		RHICmdList.Transition(FRHITransitionInfo(RenderingViewResources->IndexBuffer->UAV, ERHIAccess::UAVCompute, ERHIAccess::VertexOrIndexBuffer));
 	}
 }
 
