remove unused hlsl function causing warnings

diff --git a/Engine/Plugins/FX/Niagara/Shaders/Private/NiagaraDataInterfaceGrid3DCollection.ush b/Engine/Plugins/FX/Niagara/Shaders/Private/NiagaraDataInterfaceGrid3DCollection.ush
--- a/Engine/Plugins/FX/Niagara/Shaders/Private/NiagaraDataInterfaceGrid3DCollection.ush
+++ b/Engine/Plugins/FX/Niagara/Shaders/Private/NiagaraDataInterfaceGrid3DCollection.ush
@@ -1,5 +1,6 @@
 // Copyright Epic Games, Inc. All Rights Reserved.
 
+/*
 float SampleExplicitTrilinear_{ParameterName}(SamplerState Sampler, float3 UVW, int MipLevel)
 {	
 	int3 TotalDim = {NumCellsName} * {NumTiles};
@@ -29,6 +30,7 @@ float SampleExplicitTrilinear_{ParameterName}(SamplerState Sampler, float3 UVW,
 	// interp between bottom and top
 	return TV * delta.z + BV * (1.-delta.z);	
 }
+*/
 
 float{NumChannels} SampleTriCubicLagrange_{ParameterName}(SamplerState Sampler, float3 UVW, int MipLevel)
 {		    
