Provides a fallback path to the parameter map traversal to ensure that we get accurate aliased variables in the history.  Without the aliases in place we can miss resetting the chunk cache in the translator which leads to variables leaking between emitters and between the emitter and the system scripts.  This was typically seen using a dynamic input in both system and update scripts (Multiply Float).

diff --git a/Engine/Plugins/FX/Niagara/Source/NiagaraEditor/Private/NiagaraParameterMapHistory.cpp b/Engine/Plugins/FX/Niagara/Source/NiagaraEditor/Private/NiagaraParameterMapHistory.cpp
--- a/Engine/Plugins/FX/Niagara/Source/NiagaraEditor/Private/NiagaraParameterMapHistory.cpp
+++ b/Engine/Plugins/FX/Niagara/Source/NiagaraEditor/Private/NiagaraParameterMapHistory.cpp
@@ -384,6 +384,27 @@ int32 TNiagaraParameterMapHistory<GraphBridge>::FindVariable(const FName& Variab
 	return FoundIdx;
 }
 
+// because of how we do variable traversal we encounter situations where we are registering variables without having the context
+// of what modules might be relevant when it comes to aliasing.  As an example, when processing a dynamic input for Multiply Float
+// we could be doing a ParamMapSet on Mulitply_Float.A outside of the module where we're not able to provide AddVariable with
+// InVar=Multiply_Float.A and InAliasedVar=Module.A.  As a way around this problem we'll conditionally update the aliased version
+// of the variable as we encounter variables through the rest of the parameter map traversal
+template<typename GraphBridge>
+void TNiagaraParameterMapHistory<GraphBridge>::ConditionalUpdateAliasedVariable(int32 VariableIndex, const FNiagaraVariableBase& InAliasedVariable)
+{
+	if (VariablesWithOriginalAliasesIntact.IsValidIndex(VariableIndex))
+	{
+		FNiagaraVariable& AliasedVariable = VariablesWithOriginalAliasesIntact[VariableIndex];
+
+		if (AliasedVariable.GetName() != InAliasedVariable.GetName()
+			&& !FNiagaraParameterUtilities::IsAliasedModuleParameter(AliasedVariable)
+			&& FNiagaraParameterUtilities::IsAliasedModuleParameter(InAliasedVariable))
+		{
+			AliasedVariable.SetName(InAliasedVariable.GetName());
+		}
+	}
+}
+
 template<typename GraphBridge>
 int32 TNiagaraParameterMapHistory<GraphBridge>::AddVariable(
 	  const FNiagaraVariable& InVar
@@ -395,7 +416,7 @@ int32 TNiagaraParameterMapHistory<GraphBridge>::AddVariable(
 	FNiagaraVariable Var = InVar;
 
 	int32 FoundIdx = FindVariable(Var.GetName(), Var.GetType());
-	if (FoundIdx == -1)
+	if (FoundIdx == INDEX_NONE)
 	{
 		FoundIdx = Variables.Add(Var);
 		VariablesWithOriginalAliasesIntact.Add(InAliasedVar);
@@ -411,9 +432,14 @@ int32 TNiagaraParameterMapHistory<GraphBridge>::AddVariable(
 			check(Variables.Num() == VariableMetaData.Num());
 		}
 	}
-	else if (Variables[FoundIdx].GetType() != Var.GetType())
+	else
 	{
-		PerVariableWarnings[FoundIdx].Append(FString::Printf(TEXT("Type mismatch %s instead of %s in map!"), *Var.GetType().GetName(), *Variables[FoundIdx].GetType().GetName()));
+		if (Variables[FoundIdx].GetType() != Var.GetType())
+		{
+			PerVariableWarnings[FoundIdx].Append(FString::Printf(TEXT("Type mismatch %s instead of %s in map!"), *Var.GetType().GetName(), *Variables[FoundIdx].GetType().GetName()));
+		}
+
+		ConditionalUpdateAliasedVariable(FoundIdx, InAliasedVar);
 	}
 
 	if (InPin != nullptr)
@@ -1876,6 +1902,8 @@ int32 TNiagaraParameterMapHistoryBuilder<GraphBridge>::HandleVariableRead(int32
 		{
 			History.PerVariableWarnings[FoundIdx].Append(FString::Printf(TEXT("Type mismatch %s instead of %s in map!"), *Var.GetType().GetName(), *History.Variables[FoundIdx].GetType().GetName()));
 		}
+
+		History.ConditionalUpdateAliasedVariable(FoundIdx, AliasedVar);
 	}
 
 	typename FHistory::FReadHistory& ReadEntry = History.PerVariableReadHistory[FoundIdx].AddDefaulted_GetRef();
diff --git a/Engine/Plugins/FX/Niagara/Source/NiagaraEditor/Public/NiagaraParameterMapHistory.h b/Engine/Plugins/FX/Niagara/Source/NiagaraEditor/Public/NiagaraParameterMapHistory.h
--- a/Engine/Plugins/FX/Niagara/Source/NiagaraEditor/Public/NiagaraParameterMapHistory.h
+++ b/Engine/Plugins/FX/Niagara/Source/NiagaraEditor/Public/NiagaraParameterMapHistory.h
@@ -375,6 +375,8 @@ public:
 	FParameterCollection IsParameterCollectionParameter(FNiagaraVariable& InVar, bool& bMissingParameter);
 
 	bool ShouldIgnoreVariableDefault(const FNiagaraVariable& Var)const;
+
+	void ConditionalUpdateAliasedVariable(int32 VariableIndex, const FNiagaraVariableBase& InAliasedVar);
 };
 
 template<typename GraphBridge>
diff --git a/Engine/Source/Runtime/Core/Private/UObject/FortniteReleaseBranchCustomObjectVersion.cpp b/Engine/Source/Runtime/Core/Private/UObject/FortniteReleaseBranchCustomObjectVersion.cpp
--- a/Engine/Source/Runtime/Core/Private/UObject/FortniteReleaseBranchCustomObjectVersion.cpp
+++ b/Engine/Source/Runtime/Core/Private/UObject/FortniteReleaseBranchCustomObjectVersion.cpp
@@ -6,7 +6,7 @@ TMap<FGuid, FGuid> FFortniteReleaseBranchCustomObjectVersion::GetSystemGuids()
 	TMap<FGuid, FGuid> SystemGuids;
 	const FDevSystemGuids& DevGuids = FDevSystemGuids::Get();
 
-	SystemGuids.Add(DevGuids.Niagara_LatestScriptCompileVersion, FGuid("E65477F2755541CEB2C366179FB735CA"));
+	SystemGuids.Add(DevGuids.Niagara_LatestScriptCompileVersion, FGuid("7B09FC98626D42E1993952FD3BC3ECE1"));
 	SystemGuids.Add(DevGuids.SkeletalMeshDerivedDataVersion, FGuid("5705956EC7134274A5A1BD3FC74DB3A9"));
 	return SystemGuids;
 }
