Niagara multi-view bug fix:  count clearing related logic wasn't respecting the "bClearCounts" condition, causing counts to be incorrectly cleared in certain cases, leading to certain types of GPU sim particles not showing up in additional views.

#rnx

diff --git a/Engine/Plugins/FX/Niagara/Source/Niagara/Private/NiagaraGPUInstanceCountManager.cpp b/Engine/Plugins/FX/Niagara/Source/Niagara/Private/NiagaraGPUInstanceCountManager.cpp
--- a/Engine/Plugins/FX/Niagara/Source/Niagara/Private/NiagaraGPUInstanceCountManager.cpp
+++ b/Engine/Plugins/FX/Niagara/Source/Niagara/Private/NiagaraGPUInstanceCountManager.cpp
@@ -408,7 +408,7 @@ void FNiagaraGPUInstanceCountManager::UpdateDrawIndirectBuffers(FNiagaraGpuCompu
 		FReadBuffer TaskInfosBuffer;
 		{
 			const uint32 ArgGenSize = ArgTasks.Num() * sizeof(FNiagaraDrawIndirectArgGenTaskInfo);
-			const uint32 InstanceCountClearSize = ComputeDispatchInterface->IsFirstViewFamily() ? InstanceCountClearTasks.Num() * sizeof(uint32) : 0;
+			const uint32 InstanceCountClearSize = bClearCounts ? InstanceCountClearTasks.Num() * sizeof(uint32) : 0;
 			const uint32 TaskBufferSize = ArgGenSize + InstanceCountClearSize;
 			TaskInfosBuffer.Initialize(RHICmdList, TEXT("NiagaraTaskInfosBuffer"), sizeof(uint32), TaskBufferSize / sizeof(uint32), EPixelFormat::PF_R32_UINT, BUF_Volatile);
 
@@ -493,7 +493,7 @@ void FNiagaraGPUInstanceCountManager::UpdateDrawIndirectBuffers(FNiagaraGpuCompu
 			}
 
 			const bool bIsLastDispatch = DispatchIdx == (NumDispatches - 1);
-			const int32 NumInstanceCountClearTasks = bIsLastDispatch && ComputeDispatchInterface->IsFirstViewFamily() ? InstanceCountClearTasks.Num() : 0;
+			const int32 NumInstanceCountClearTasks = bIsLastDispatch && bClearCounts ? InstanceCountClearTasks.Num() : 0;
 
 			// Do we have anything to do for this pool?
 			if (NumArgGenTasks + NumInstanceCountClearTasks == 0)
