Remove debug only data from Niagara scripts during cook to reduce the size

diff --git a/Engine/Plugins/FX/Niagara/Source/Niagara/Classes/NiagaraScript.h b/Engine/Plugins/FX/Niagara/Source/Niagara/Classes/NiagaraScript.h
--- a/Engine/Plugins/FX/Niagara/Source/Niagara/Classes/NiagaraScript.h
+++ b/Engine/Plugins/FX/Niagara/Source/Niagara/Classes/NiagaraScript.h
@@ -1117,10 +1117,6 @@ public:
 #endif
 
 	//~ Begin UObject interface
-	PRAGMA_DISABLE_DEPRECATION_WARNINGS // Suppress compiler warning on override of deprecated function
-	UE_DEPRECATED(5.0, "Use version that takes FObjectPreSaveContext instead.")
-	virtual void PreSave(const class ITargetPlatform* TargetPlatform) override;
-	PRAGMA_ENABLE_DEPRECATION_WARNINGS
 	virtual void PreSave(FObjectPreSaveContext ObjectSaveContext) override;
 	virtual void Serialize(FArchive& Ar)override;
 	virtual void PostLoad() override;
diff --git a/Engine/Plugins/FX/Niagara/Source/Niagara/Classes/NiagaraSystem.h b/Engine/Plugins/FX/Niagara/Source/Niagara/Classes/NiagaraSystem.h
--- a/Engine/Plugins/FX/Niagara/Source/Niagara/Classes/NiagaraSystem.h
+++ b/Engine/Plugins/FX/Niagara/Source/Niagara/Classes/NiagaraSystem.h
@@ -238,10 +238,6 @@ public:
 #endif
 	NIAGARA_API virtual void BeginDestroy() override;
 	NIAGARA_API virtual bool IsReadyForFinishDestroy() override;
-	PRAGMA_DISABLE_DEPRECATION_WARNINGS // Suppress compiler warning on override of deprecated function
-	UE_DEPRECATED(5.0, "Use version that takes FObjectPreSaveContext instead.")
-	NIAGARA_API virtual void PreSave(const class ITargetPlatform* TargetPlatform) override;
-	NIAGARA_API PRAGMA_ENABLE_DEPRECATION_WARNINGS
 	virtual void PreSave(FObjectPreSaveContext ObjectSaveContext) override;
 #if WITH_EDITOR
 	NIAGARA_API virtual void PostRename(UObject* OldOuter, const FName OldName) override;
diff --git a/Engine/Plugins/FX/Niagara/Source/Niagara/Private/NiagaraScript.cpp b/Engine/Plugins/FX/Niagara/Source/Niagara/Private/NiagaraScript.cpp
--- a/Engine/Plugins/FX/Niagara/Source/Niagara/Private/NiagaraScript.cpp
+++ b/Engine/Plugins/FX/Niagara/Source/Niagara/Private/NiagaraScript.cpp
@@ -1809,13 +1809,6 @@ void UNiagaraScript::GenerateDefaultFunctionBindings()
 	}
 }
 
-void UNiagaraScript::PreSave(const class ITargetPlatform* TargetPlatform)
-{
-	PRAGMA_DISABLE_DEPRECATION_WARNINGS;
-	Super::PreSave(TargetPlatform);
-	PRAGMA_ENABLE_DEPRECATION_WARNINGS;
-}
-
 void UNiagaraScript::PreSave(FObjectPreSaveContext ObjectSaveContext)
 {
 	Super::PreSave(ObjectSaveContext);
@@ -1926,6 +1919,17 @@ void UNiagaraScript::Serialize(FArchive& Ar)
 			UE_LOG(LogNiagara, Verbose, TEXT("Pruned %d/%d parameters from script %s"), NumRemoved, ParameterVariables.Num(), *GetFullName());
 		}
 
+		// when we are cooking we want to strip out some debug only data.  While this data is already intended for editor only it will
+		// bloat the size of packages intended for the cooked editor
+		TOptional<FNiagaraVMExecutableData> LocalStrippedDebugData;
+		auto SwapDebugData = [](FNiagaraVMExecutableData& SourceData, FNiagaraVMExecutableData& TargetData) -> void
+		{
+			Swap(TargetData.LastHlslTranslation, SourceData.LastHlslTranslation);
+			Swap(TargetData.LastHlslTranslationGPU, SourceData.LastHlslTranslationGPU);
+			Swap(TargetData.LastAssemblyTranslation, SourceData.LastAssemblyTranslation);
+			Swap(TargetData.CompileTagsEditorOnly, SourceData.CompileTagsEditorOnly);
+		};
+
 		if (Ar.IsSaving())
 		{
 			auto& ExecutableData = GetVMExecutableData();
@@ -1943,10 +1947,18 @@ void UNiagaraScript::Serialize(FArchive& Ar)
 			{
 				ExecutableData.ByteCode.Compress();
 			}
+
+			LocalStrippedDebugData.Emplace();
+			SwapDebugData(*LocalStrippedDebugData, CachedScriptVM);
 		}
 
 		Super::Serialize(Ar);
 
+		if (LocalStrippedDebugData.IsSet())
+		{
+			SwapDebugData(*LocalStrippedDebugData, CachedScriptVM);
+		}
+
 		// If we cached the rapid iteration parameters before serialize, restore them here.
 		if (RapidIterationParametersCookedEditorCache.Num() != 0)
 		{
diff --git a/Engine/Plugins/FX/Niagara/Source/Niagara/Private/NiagaraSystem.cpp b/Engine/Plugins/FX/Niagara/Source/Niagara/Private/NiagaraSystem.cpp
--- a/Engine/Plugins/FX/Niagara/Source/Niagara/Private/NiagaraSystem.cpp
+++ b/Engine/Plugins/FX/Niagara/Source/Niagara/Private/NiagaraSystem.cpp
@@ -268,13 +268,6 @@ bool UNiagaraSystem::IsReadyForFinishDestroy()
 	return bReady && WaitRenderCommandsFence.IsFenceComplete();
 }
 
-void UNiagaraSystem::PreSave(const class ITargetPlatform* TargetPlatform)
-{
-	PRAGMA_DISABLE_DEPRECATION_WARNINGS;
-	Super::PreSave(TargetPlatform);
-	PRAGMA_ENABLE_DEPRECATION_WARNINGS;
-}
-
 void UNiagaraSystem::PreSave(FObjectPreSaveContext ObjectSaveContext)
 {
 	Super::PreSave(ObjectSaveContext);
