Fixed clipping & added 'force display as disabled' on the overview graph nodes on Initialization of the overview graph model

diff --git a/Engine/Plugins/FX/Niagara/Source/NiagaraEditor/Private/NiagaraOverviewNode.cpp b/Engine/Plugins/FX/Niagara/Source/NiagaraEditor/Private/NiagaraOverviewNode.cpp
--- a/Engine/Plugins/FX/Niagara/Source/NiagaraEditor/Private/NiagaraOverviewNode.cpp
+++ b/Engine/Plugins/FX/Niagara/Source/NiagaraEditor/Private/NiagaraOverviewNode.cpp
@@ -36,6 +36,22 @@ UNiagaraOverviewNode::UNiagaraOverviewNode()
 	UEdGraphNode::bCanRenameNode = true;
 };
 
+static FNiagaraEmitterHandle* FindEmitterHandleByID(UNiagaraSystem* System, const FGuid& Guid)
+{
+	check(System != nullptr);
+
+	for (int Idx = 0; Idx < System->GetNumEmitters(); ++Idx)
+	{
+		FNiagaraEmitterHandle& EmitterHandle = System->GetEmitterHandle(Idx);
+		if (EmitterHandle.GetId() == Guid)
+		{
+			return &EmitterHandle;
+		}
+	}
+
+	return nullptr;
+}
+
 void UNiagaraOverviewNode::Initialize(UNiagaraSystem* InOwningSystem)
 {
 	OwningSystem = InOwningSystem;
@@ -47,25 +63,24 @@ void UNiagaraOverviewNode::Initialize(UNiagaraSystem* InOwningSystem, FGuid InEm
 	EmitterHandleGuid = InEmitterHandleGuid;
 }
 
-const FGuid UNiagaraOverviewNode::GetEmitterHandleGuid() const
+void UNiagaraOverviewNode::UpdateStatus()
 {
-	return EmitterHandleGuid;
-}
-
-static FNiagaraEmitterHandle* FindEmitterHandleByID(UNiagaraSystem* System, const FGuid& Guid)
-{
-	check(System != nullptr);
-
-	for (int Idx = 0; Idx < System->GetNumEmitters(); ++Idx)
+	if(FNiagaraEmitterHandle* EmitterHandle = FindEmitterHandleByID(OwningSystem, EmitterHandleGuid))
 	{
-		FNiagaraEmitterHandle& EmitterHandle = System->GetEmitterHandle(Idx);
-		if (EmitterHandle.GetId() == Guid)
+		if(!GetDefault<UNiagaraSettings>()->bStatelessEmittersEnabled && EmitterHandle->GetEmitterMode() == ENiagaraEmitterMode::Stateless)
 		{
-			return &EmitterHandle;
+			SetForceDisplayAsDisabled(true);
+		}
+		else
+		{
+			SetForceDisplayAsDisabled(false);
 		}
 	}
+}
 
-	return nullptr;
+const FGuid UNiagaraOverviewNode::GetEmitterHandleGuid() const
+{
+	return EmitterHandleGuid;
 }
 
 FNiagaraEmitterHandle* UNiagaraOverviewNode::TryGetEmitterHandle() const
diff --git a/Engine/Plugins/FX/Niagara/Source/NiagaraEditor/Private/ViewModels/NiagaraOverviewGraphViewModel.cpp b/Engine/Plugins/FX/Niagara/Source/NiagaraEditor/Private/ViewModels/NiagaraOverviewGraphViewModel.cpp
--- a/Engine/Plugins/FX/Niagara/Source/NiagaraEditor/Private/ViewModels/NiagaraOverviewGraphViewModel.cpp
+++ b/Engine/Plugins/FX/Niagara/Source/NiagaraEditor/Private/ViewModels/NiagaraOverviewGraphViewModel.cpp
@@ -62,6 +62,8 @@ void FNiagaraOverviewGraphViewModel::Initialize(TSharedRef<FNiagaraSystemViewMod
 
 	NodeSelection->OnSelectedObjectsChanged().AddSP(this, &FNiagaraOverviewGraphViewModel::GraphSelectionChanged);
 	InSystemViewModel->GetSelectionViewModel()->OnEntrySelectionChanged().AddSP(this, &FNiagaraOverviewGraphViewModel::SystemSelectionChanged);
+
+	UpdateOverviewGraphNodes();
 }
 
 FText FNiagaraOverviewGraphViewModel::GetDisplayName() const
@@ -121,6 +123,17 @@ const TSharedRef<FNiagaraSystemViewModel> FNiagaraOverviewGraphViewModel::GetSys
 	return SystemViewModelPinned.ToSharedRef();
 }
 
+void FNiagaraOverviewGraphViewModel::UpdateOverviewGraphNodes()
+{
+	for(UEdGraphNode* Node : OverviewGraph->Nodes)
+	{
+		if(UNiagaraOverviewNode* OverviewNode = Cast<UNiagaraOverviewNode>(Node))
+		{
+			OverviewNode->UpdateStatus();
+		}
+	}
+}
+
 void FNiagaraOverviewGraphViewModel::SetupCommands()
 {
 	Commands->MapAction(
diff --git a/Engine/Plugins/FX/Niagara/Source/NiagaraEditor/Public/NiagaraOverviewNode.h b/Engine/Plugins/FX/Niagara/Source/NiagaraEditor/Public/NiagaraOverviewNode.h
--- a/Engine/Plugins/FX/Niagara/Source/NiagaraEditor/Public/NiagaraOverviewNode.h
+++ b/Engine/Plugins/FX/Niagara/Source/NiagaraEditor/Public/NiagaraOverviewNode.h
@@ -15,6 +15,7 @@ public:
 	NIAGARAEDITOR_API UNiagaraOverviewNode();
 	NIAGARAEDITOR_API void Initialize(UNiagaraSystem* InOwningSystem);
 	NIAGARAEDITOR_API void Initialize(UNiagaraSystem* InOwningSystem, FGuid InEmitterHandleGuid);
+	void UpdateStatus();
 	NIAGARAEDITOR_API const FGuid GetEmitterHandleGuid() const;
 	NIAGARAEDITOR_API struct FNiagaraEmitterHandle* TryGetEmitterHandle() const;
 	
diff --git a/Engine/Plugins/FX/Niagara/Source/NiagaraEditor/Public/ViewModels/NiagaraOverviewGraphViewModel.h b/Engine/Plugins/FX/Niagara/Source/NiagaraEditor/Public/ViewModels/NiagaraOverviewGraphViewModel.h
--- a/Engine/Plugins/FX/Niagara/Source/NiagaraEditor/Public/ViewModels/NiagaraOverviewGraphViewModel.h
+++ b/Engine/Plugins/FX/Niagara/Source/NiagaraEditor/Public/ViewModels/NiagaraOverviewGraphViewModel.h
@@ -53,7 +53,8 @@ public:
 	//~ FEditorUndoClient interface.
 	virtual void PostUndo(bool bSuccess) override;
 	virtual void PostRedo(bool bSuccess) override { PostUndo(bSuccess); }
-	
+
+	void UpdateOverviewGraphNodes();
 private:
 	void SetupCommands();
 
diff --git a/Engine/Plugins/FX/Niagara/Source/NiagaraEditorWidgets/Private/SNiagaraOverviewStackNode.cpp b/Engine/Plugins/FX/Niagara/Source/NiagaraEditorWidgets/Private/SNiagaraOverviewStackNode.cpp
--- a/Engine/Plugins/FX/Niagara/Source/NiagaraEditorWidgets/Private/SNiagaraOverviewStackNode.cpp
+++ b/Engine/Plugins/FX/Niagara/Source/NiagaraEditorWidgets/Private/SNiagaraOverviewStackNode.cpp
@@ -495,9 +495,13 @@ TSharedRef<SWidget> SNiagaraOverviewStackNode::CreateTitleWidget_Default(TShared
 				SNew(SVerticalBox)
 				+ SVerticalBox::Slot()
 				.Padding(0, 0, 5, 0)
+				.AutoHeight()
 				[
 					SNew(STextBlock)
-						.Text(LOCTEXT("StatelessNotEnabled", "Stateless not enabled\nin project settings."))
+						.Text(LOCTEXT("StatelessNotEnabled", "Stateless not enabled in project settings."))
+						.Font(FCoreStyle::GetDefaultFontStyle("Regular", 8.f))
+						.OverflowPolicy(ETextOverflowPolicy::Ellipsis)
+						.AutoWrapText(true)
 				]
 				+ SVerticalBox::Slot()
 				[
