Readded the removal of scalability from default view as there were details panel bugs of unknown origin. Meant as a temp fix.
Use scalability mode to access scalability settings instead

diff --git a/Engine/Plugins/FX/Niagara/Source/NiagaraEditor/Private/NiagaraEmitterDetailsCustomization.cpp b/Engine/Plugins/FX/Niagara/Source/NiagaraEditor/Private/NiagaraEmitterDetailsCustomization.cpp
--- a/Engine/Plugins/FX/Niagara/Source/NiagaraEditor/Private/NiagaraEmitterDetailsCustomization.cpp
+++ b/Engine/Plugins/FX/Niagara/Source/NiagaraEditor/Private/NiagaraEmitterDetailsCustomization.cpp
@@ -83,7 +83,12 @@ void FNiagaraEmitterDetails::CustomizeDetails(IDetailLayoutBuilder& InDetailLayo
 			if (ChildProperty->HasAllPropertyFlags(CPF_Edit))
 			{
 				FName Category = FName(ChildProperty->GetMetaData(TEXT("Category")));
-				CategoryPropertyMap.FindOrAdd(Category).Add(ChildProperty);
+				
+				// we display the scalability category within scalability mode, which is why we hide it here
+				if (Category != FName("Scalability"))
+				{
+					CategoryPropertyMap.FindOrAdd(Category).Add(ChildProperty);
+				}
 			}
 		}
 
