- Stateless merge [at] CL31115585

diff --git a/Engine/Plugins/FX/Niagara/Shaders/Private/Stateless/Modules/NiagaraStatelessModule_InitializeParticle.ush b/Engine/Plugins/FX/Niagara/Shaders/Private/Stateless/Modules/NiagaraStatelessModule_InitializeParticle.ush
--- a/Engine/Plugins/FX/Niagara/Shaders/Private/Stateless/Modules/NiagaraStatelessModule_InitializeParticle.ush
+++ b/Engine/Plugins/FX/Niagara/Shaders/Private/Stateless/Modules/NiagaraStatelessModule_InitializeParticle.ush
@@ -2,6 +2,10 @@
 
 #include "../NiagaraStatelessCommon.ush"
 
+static const uint EInitializeParticleModuleFlag_UniformSpriteSize	= 1 << 0;
+static const uint EInitializeParticleModuleFlag_UniformMeshScale	= 1 << 1;
+
+uint	InitializeParticle_ModuleFlags;
 float3	InitializeParticle_Position;
 float4	InitializeParticle_ColorScale;
 float4	InitializeParticle_ColorBias;
@@ -16,45 +20,22 @@ float3	InitializeParticle_MeshScaleBias;
 float	InitializeParticle_RibbonWidthScale;
 float	InitializeParticle_RibbonWidthBias;
 	
-uint InitializeParticle_PositionMode;
-uint InitializeParticle_PositionTableMod;
-uint InitializeParticle_PositionTableOffset;
-uint InitializeParticle_PositionTableLength;
-
 struct FStatelessModule_InitializeParticle
 {
 };
 
 void InitializeParticle_Initialize(inout FStatelessModule_InitializeParticle ModuleData, inout FStatelessParticle Particle)
 {  
+	const bool bUniformSpriteSize = (InitializeParticle_ModuleFlags & EInitializeParticleModuleFlag_UniformSpriteSize) != 0;
+	const bool bUniformMeshScale  = (InitializeParticle_ModuleFlags & EInitializeParticleModuleFlag_UniformMeshScale) != 0;
+
 	Particle.UniqueID		= Particle.UniqueIndex;
 	Particle.Position		= InitializeParticle_Position;
 	Particle.Color			= RandomScaleBiasFloat(Particle, InitializeParticle_ColorScale, InitializeParticle_ColorBias);
 	Particle.RibbonWidth	= RandomScaleBiasFloat(Particle, InitializeParticle_RibbonWidthScale, InitializeParticle_RibbonWidthBias);
-	Particle.SpriteSize		= RandomScaleBiasFloat(Particle, InitializeParticle_SpriteSizeScale, InitializeParticle_SpriteSizeBias);
+	Particle.SpriteSize		= RandomScaleBiasFloat(Particle, InitializeParticle_SpriteSizeScale, InitializeParticle_SpriteSizeBias, bUniformSpriteSize);
 	Particle.SpriteRotation	= RandomScaleBiasFloat(Particle, InitializeParticle_SpriteRotationScale, InitializeParticle_SpriteRotationBias);
-	Particle.Scale			= RandomScaleBiasFloat(Particle, InitializeParticle_MeshScaleScale, InitializeParticle_MeshScaleBias);
-
-	// For testing purposes only
-	{
-		if ( InitializeParticle_PositionMode == 0 )
-		{
-			const uint Len = InitializeParticle_PositionTableLength - 1;
-			const uint Element = clamp(Particle.UniqueIndex % InitializeParticle_PositionTableMod, 0, Len);
-			Particle.Position += GetStaticFloat3(InitializeParticle_PositionTableOffset, Element);
-		}
-		else
-		{
-			const uint Len = InitializeParticle_PositionTableLength - 1;
-			const float	R0 = RandomScaleBiasFloat(Particle, Len, 0);
-			const float U = InitializeParticle_PositionMode != 2 ? frac(R0) : 0.0f;
-			const uint E0 = floor(R0);
-			const uint E1 = InitializeParticle_PositionMode == 1 ? min(E0 + 1, Len) : (RandomUInt(Particle) % InitializeParticle_PositionTableMod);
-			const float3 P0 = GetStaticFloat3(InitializeParticle_PositionTableOffset, E0);
-			const float3 P1 = GetStaticFloat3(InitializeParticle_PositionTableOffset, E1);
-			Particle.Position += lerp(P0, P1, U);
-		}
-	}
+	Particle.Scale			= RandomScaleBiasFloat(Particle, InitializeParticle_MeshScaleScale, InitializeParticle_MeshScaleBias, bUniformMeshScale);
 
 	Particle.PreviousPosition		= Particle.Position;
 	Particle.PreviousRibbonWidth	= Particle.RibbonWidth;
diff --git a/Engine/Plugins/FX/Niagara/Shaders/Private/Stateless/Modules/NiagaraStatelessModule_ScaleColor.ush b/Engine/Plugins/FX/Niagara/Shaders/Private/Stateless/Modules/NiagaraStatelessModule_ScaleColor.ush
--- a/Engine/Plugins/FX/Niagara/Shaders/Private/Stateless/Modules/NiagaraStatelessModule_ScaleColor.ush
+++ b/Engine/Plugins/FX/Niagara/Shaders/Private/Stateless/Modules/NiagaraStatelessModule_ScaleColor.ush
@@ -2,35 +2,22 @@
 
 #include "../NiagaraStatelessCommon.ush"
 
-int		ScaleColor_Offset;
-int		ScaleColor_Length;
+uint3	ScaleColor_Distribution;
 
 struct FStatelessModule_ScaleColor
 {
+	FStatelessDistributionFloat4	ScaleDistribution;
 };
 
-float4 ScaleColor_GetScale(in FStatelessParticle Particle, float NormalizedAge)
-{
-	const float Offset = NormalizedAge * float(ScaleColor_Length - 1);
-	const int A = floor(Offset);
-	const int B = ceil(Offset);
-	const float U = frac(Offset);
-
-	return lerp(
-		GetStaticFloat4(ScaleColor_Offset, A),
-		GetStaticFloat4(ScaleColor_Offset, B), 
-		U
-	);
-}
-
 void ScaleColor_Initialize(inout FStatelessModule_ScaleColor ModuleData, inout FStatelessParticle Particle)
 {
+	ModuleData.ScaleDistribution.Init(Particle, ScaleColor_Distribution);
 }
 	
 void ScaleColor_Simulate(in FStatelessModule_ScaleColor ModuleData, inout FStatelessParticle Particle)
 {
-	if (ScaleColor_Length > 0)
+	if (ModuleData.ScaleDistribution.IsValid())
 	{
-		Particle.Color *= ScaleColor_GetScale(Particle, Particle.NormalizedAge);
+		Particle.Color *= ModuleData.ScaleDistribution.GetValue(Particle.NormalizedAge);
 	}
 }
diff --git a/Engine/Plugins/FX/Niagara/Shaders/Private/Stateless/Modules/NiagaraStatelessModule_ScaleMeshSize.ush b/Engine/Plugins/FX/Niagara/Shaders/Private/Stateless/Modules/NiagaraStatelessModule_ScaleMeshSize.ush
--- a/Engine/Plugins/FX/Niagara/Shaders/Private/Stateless/Modules/NiagaraStatelessModule_ScaleMeshSize.ush
+++ b/Engine/Plugins/FX/Niagara/Shaders/Private/Stateless/Modules/NiagaraStatelessModule_ScaleMeshSize.ush
@@ -2,37 +2,23 @@
 
 #include "../NiagaraStatelessCommon.ush"
 
-int		ScaleMeshSize_Offset;
-int		ScaleMeshSize_Length;
-
+uint3 ScaleMeshSize_Distribution;
+	
 struct FStatelessModule_ScaleMeshSize
 {
+	FStatelessDistributionFloat3	MeshSizeDistribution;
 };
 
-
-float3 ScaleMeshSize_GetScale(in FStatelessParticle Particle, float NormalizedAge)
-{
-	const float Offset = NormalizedAge * float(ScaleMeshSize_Length - 1);
-	const int A = floor(Offset);
-	const int B = ceil(Offset);
-	const float U = frac(Offset);
-
-	return lerp(
-		GetStaticFloat3(ScaleMeshSize_Offset, A),
-		GetStaticFloat3(ScaleMeshSize_Offset, B), 
-		U
-	);
-}
-	
 void ScaleMeshSize_Initialize(inout FStatelessModule_ScaleMeshSize ModuleData, inout FStatelessParticle Particle)
 {
+	ModuleData.MeshSizeDistribution.Init(Particle, ScaleMeshSize_Distribution);
 }
 	
 void ScaleMeshSize_Simulate(in FStatelessModule_ScaleMeshSize ModuleData, inout FStatelessParticle Particle)
 {
-	if (ScaleMeshSize_Length > 0)
+	if (ModuleData.MeshSizeDistribution.IsValid())
 	{
-		Particle.Scale *= ScaleMeshSize_GetScale(Particle, Particle.NormalizedAge);
-		Particle.PreviousScale *= ScaleMeshSize_GetScale(Particle, Particle.PreviousNormalizedAge);
+		Particle.Scale			*= ModuleData.MeshSizeDistribution.GetValue(Particle.NormalizedAge);
+		Particle.PreviousScale	*= ModuleData.MeshSizeDistribution.GetValue(Particle.PreviousNormalizedAge);
 	}
 }
diff --git a/Engine/Plugins/FX/Niagara/Shaders/Private/Stateless/Modules/NiagaraStatelessModule_ScaleSpriteSize.ush b/Engine/Plugins/FX/Niagara/Shaders/Private/Stateless/Modules/NiagaraStatelessModule_ScaleSpriteSize.ush
--- a/Engine/Plugins/FX/Niagara/Shaders/Private/Stateless/Modules/NiagaraStatelessModule_ScaleSpriteSize.ush
+++ b/Engine/Plugins/FX/Niagara/Shaders/Private/Stateless/Modules/NiagaraStatelessModule_ScaleSpriteSize.ush
@@ -2,36 +2,23 @@
 
 #include "../NiagaraStatelessCommon.ush"
 
-int		ScaleSpriteSize_Offset;
-int		ScaleSpriteSize_Length;
+uint3 ScaleSpriteSize_Distribution;
 
 struct FStatelessModule_ScaleSpriteSize
 {
+	FStatelessDistributionFloat2	ScaleDistribution;
 };
 	
-float2 ScaleSpriteSize_GetScale(in FStatelessParticle Particle, float NormalizedAge)
-{
-	const float Offset = NormalizedAge * float(ScaleSpriteSize_Length - 1);
-	const int A = floor(Offset);
-	const int B = ceil(Offset);
-	const float U = frac(Offset);
-
-	return lerp(
-		GetStaticFloat2(ScaleSpriteSize_Offset, A),
-		GetStaticFloat2(ScaleSpriteSize_Offset, B), 
-		U
-	);
-}
-	
 void ScaleSpriteSize_Initialize(inout FStatelessModule_ScaleSpriteSize ModuleData, inout FStatelessParticle Particle)
 {
+	ModuleData.ScaleDistribution.Init(Particle, ScaleSpriteSize_Distribution);
 }
 	
 void ScaleSpriteSize_Simulate(in FStatelessModule_ScaleSpriteSize ModuleData, inout FStatelessParticle Particle)
 {
-	if (ScaleSpriteSize_Length > 0)
+	if (ModuleData.ScaleDistribution.IsValid())
 	{
-		Particle.SpriteSize *= ScaleSpriteSize_GetScale(Particle, Particle.NormalizedAge);
-		Particle.PreviousSpriteSize *= ScaleSpriteSize_GetScale(Particle, Particle.PreviousNormalizedAge);
+		Particle.SpriteSize			*= ModuleData.ScaleDistribution.GetValue(Particle.NormalizedAge);
+		Particle.PreviousSpriteSize	*= ModuleData.ScaleDistribution.GetValue(Particle.PreviousNormalizedAge);
 	}
 }
diff --git a/Engine/Plugins/FX/Niagara/Shaders/Private/Stateless/NiagaraStatelessCommon.ush b/Engine/Plugins/FX/Niagara/Shaders/Private/Stateless/NiagaraStatelessCommon.ush
--- a/Engine/Plugins/FX/Niagara/Shaders/Private/Stateless/NiagaraStatelessCommon.ush
+++ b/Engine/Plugins/FX/Niagara/Shaders/Private/Stateless/NiagaraStatelessCommon.ush
@@ -113,6 +113,11 @@ float2 RandomScaleBiasFloat(inout FStatelessParticle Particle, float2 Scale, flo
 float3 RandomScaleBiasFloat(inout FStatelessParticle Particle, float3 Scale, float3 Bias) { return (RandomFloat3(Particle) * Scale) + Bias; }
 float4 RandomScaleBiasFloat(inout FStatelessParticle Particle, float4 Scale, float4 Bias) { return (RandomFloat4(Particle) * Scale) + Bias; }
 
+float  RandomScaleBiasFloat(inout FStatelessParticle Particle, float  Scale, float  Bias, bool bUniform) { return (RandomFloat(Particle)  * Scale) + Bias; }
+float2 RandomScaleBiasFloat(inout FStatelessParticle Particle, float2 Scale, float2 Bias, bool bUniform) { float2 Random = (RandomFloat2(Particle) * Scale) + Bias; return bUniform ? Random.xx : Random; }
+float3 RandomScaleBiasFloat(inout FStatelessParticle Particle, float3 Scale, float3 Bias, bool bUniform) { float3 Random = (RandomFloat3(Particle) * Scale) + Bias; return bUniform ? Random.xxx : Random; }
+float4 RandomScaleBiasFloat(inout FStatelessParticle Particle, float4 Scale, float4 Bias, bool bUniform) { float4 Random = (RandomFloat4(Particle) * Scale) + Bias; return bUniform ? Random.xxxx : Random; }
+
 float2 RandomUnitFloat2(inout FStatelessParticle Particle) { return SafeNormalize(RandomFloat2(Particle) - 0.5f); }
 float3 RandomUnitFloat3(inout FStatelessParticle Particle) { return SafeNormalize(RandomFloat3(Particle) - 0.5f); }
 
@@ -172,3 +177,4 @@ void OutputComponentData(uint OutputIndex, int ComponentOffset, int Value)
 	}
 }
 
+#include "/Plugin/FX/Niagara/Private/Stateless/NiagaraStatelessDistribution.ush"
diff --git a/Engine/Plugins/FX/Niagara/Shaders/Private/Stateless/NiagaraStatelessDistribution.ush b/Engine/Plugins/FX/Niagara/Shaders/Private/Stateless/NiagaraStatelessDistribution.ush
new file mode 100644
--- /dev/null
+++ b/Engine/Plugins/FX/Niagara/Shaders/Private/Stateless/NiagaraStatelessDistribution.ush
@@ -0,0 +1,184 @@
+// Copyright Epic Games, Inc. All Rights Reserved.
+#pragma once
+
+#define StatelessDistributionFlag_Random	0x00000001
+#define StatelessDistributionFlag_Uniform	0x00000002
+
+struct FStatelessDistributionFloat
+{
+	bool	bRandomDistribution;
+	bool	bUniformValue;
+	uint	TableOffset;
+	float	TableLength;
+	float	RandomOffset;
+
+	void Init(inout FStatelessParticle Particle, uint3 Parameters)
+	{
+		bRandomDistribution = (Parameters.x & StatelessDistributionFlag_Random) != 0;
+		bUniformValue		= (Parameters.x & StatelessDistributionFlag_Uniform) != 0;
+		TableOffset			= Parameters.y;
+		TableLength			= asfloat(Parameters.z);
+		RandomOffset		= RandomFloat(Particle);
+	}
+
+	float GetValue(float NormalizedAge)
+	{
+		if ( bRandomDistribution )
+		{
+			const float	Value0 = GetStaticFloat(TableOffset, 0);
+			const float	Value1 = GetStaticFloat(TableOffset, 1);
+			return lerp(Value0, Value1, RandomOffset);
+		}
+		else
+		{
+			const float Offset = NormalizedAge * TableLength;
+			const float	Value0 = GetStaticFloat(TableOffset, floor(Offset));
+			const float	Value1 = GetStaticFloat(TableOffset, ceil(Offset));
+			return lerp(Value0, Value1, frac(Offset));
+		}
+	}
+};
+
+struct FStatelessDistributionFloat2
+{
+	bool	bRandomDistribution;
+	bool	bUniformValue;
+	uint	TableOffset;
+	float	TableLength;
+	float2	RandomOffset;
+
+	void Init(inout FStatelessParticle Particle, uint3 Parameters)
+	{
+		bRandomDistribution = (Parameters.x & StatelessDistributionFlag_Random) != 0;
+		bUniformValue		= (Parameters.x & StatelessDistributionFlag_Uniform) != 0;
+		TableOffset			= Parameters.y;
+		TableLength			= asfloat(Parameters.z);
+		RandomOffset		= 0;
+		if ( bRandomDistribution )
+		{
+			RandomOffset	= bUniformValue ? RandomFloat(Particle) : RandomFloat2(Particle);
+		}
+	}
+
+	bool IsValid()
+	{
+		return TableLength > 0.0f;
+	}
+	
+	float2 GetValue(float NormalizedAge)
+	{
+		if ( bRandomDistribution )
+		{
+			const float2 Value0 = GetStaticFloat2(TableOffset, 0);
+			const float2 Value1 = GetStaticFloat2(TableOffset, 1);
+			return float2(
+				lerp(Value0.x, Value1.x, RandomOffset.x),
+				lerp(Value0.y, Value1.y, RandomOffset.y)
+			);
+		}
+		else
+		{
+			const float  Offset = NormalizedAge * TableLength;
+			const float2 Value0 = GetStaticFloat2(TableOffset, floor(Offset));
+			const float2 Value1 = GetStaticFloat2(TableOffset, ceil(Offset));
+			return lerp(Value0, Value1, frac(Offset));
+		}
+	}
+};
+
+struct FStatelessDistributionFloat3
+{
+	bool	bRandomDistribution;
+	bool	bUniformValue;
+	uint	TableOffset;
+	float	TableLength;
+	float3	RandomOffset;
+
+	void Init(inout FStatelessParticle Particle, uint3 Parameters)
+	{
+		bRandomDistribution = (Parameters.x & StatelessDistributionFlag_Random) != 0;
+		bUniformValue		= (Parameters.x & StatelessDistributionFlag_Uniform) != 0;
+		TableOffset			= Parameters.y;
+		TableLength			= asfloat(Parameters.z);
+		RandomOffset		= 0;
+		if ( bRandomDistribution )
+		{
+			RandomOffset	= bUniformValue ? RandomFloat(Particle) : RandomFloat3(Particle);
+		}
+	}
+
+	bool IsValid()
+	{
+		return TableLength > 0.0f;
+	}
+	
+	float3 GetValue(float NormalizedAge)
+	{
+		if ( bRandomDistribution )
+		{
+			const float3 Value0 = GetStaticFloat3(TableOffset, 0);
+			const float3 Value1 = GetStaticFloat3(TableOffset, 1);
+			return float3(
+				lerp(Value0.x, Value1.x, RandomOffset.x),
+				lerp(Value0.y, Value1.y, RandomOffset.y),
+				lerp(Value0.z, Value1.z, RandomOffset.z)
+			);
+		}
+		else
+		{
+			const float  Offset = NormalizedAge * TableLength;
+			const float3 Value0 = GetStaticFloat3(TableOffset, floor(Offset));
+			const float3 Value1 = GetStaticFloat3(TableOffset, ceil(Offset));
+			return lerp(Value0, Value1, frac(Offset));
+		}
+	}
+};
+
+struct FStatelessDistributionFloat4
+{
+	bool	bRandomDistribution;
+	bool	bUniformValue;
+	uint	TableOffset;
+	float	TableLength;
+	float4	RandomOffset;
+
+	void Init(inout FStatelessParticle Particle, uint3 Parameters)
+	{
+		bRandomDistribution = (Parameters.x & StatelessDistributionFlag_Random) != 0;
+		bUniformValue		= (Parameters.x & StatelessDistributionFlag_Uniform) != 0;
+		TableOffset			= Parameters.y;
+		TableLength			= asfloat(Parameters.z);
+		RandomOffset		= 0;
+		if ( bRandomDistribution )
+		{
+			RandomOffset	= bUniformValue ? RandomFloat(Particle) : RandomFloat4(Particle);
+		}
+	}
+
+	bool IsValid()
+	{
+		return TableLength > 0.0f;
+	}
+	
+	float4 GetValue(float NormalizedAge)
+	{
+		if ( bRandomDistribution )
+		{
+			const float4 Value0 = GetStaticFloat4(TableOffset, 0);
+			const float4 Value1 = GetStaticFloat4(TableOffset, 1);
+			return float4(
+				lerp(Value0.x, Value1.x, RandomOffset.x),
+				lerp(Value0.y, Value1.y, RandomOffset.y),
+				lerp(Value0.z, Value1.z, RandomOffset.z),
+				lerp(Value0.w, Value1.w, RandomOffset.w)
+			);
+		}
+		else
+		{
+			const float  Offset = NormalizedAge * TableLength;
+			const float4 Value0 = GetStaticFloat4(TableOffset, floor(Offset));
+			const float4 Value1 = GetStaticFloat4(TableOffset, ceil(Offset));
+			return lerp(Value0, Value1, frac(Offset));
+		}
+	}
+};
diff --git a/Engine/Plugins/FX/Niagara/Source/Niagara/Internal/Stateless/Modules/NiagaraStatelessModuleCommon.h b/Engine/Plugins/FX/Niagara/Source/Niagara/Internal/Stateless/Modules/NiagaraStatelessModuleCommon.h
--- a/Engine/Plugins/FX/Niagara/Source/Niagara/Internal/Stateless/Modules/NiagaraStatelessModuleCommon.h
+++ b/Engine/Plugins/FX/Niagara/Source/Niagara/Internal/Stateless/Modules/NiagaraStatelessModuleCommon.h
@@ -11,29 +11,23 @@ namespace NiagaraStateless
 	{
 		static FName GetName() { return FName("FPhysicsBuildData"); }
 
-		float		MassMin = 1.0f;
-		float		MassMax = 1.0f;
-		float		DragMin = 0.0f;
-		float		DragMax = 0.0f;
-		FVector3f	VelocityMin = FVector3f::ZeroVector;
-		FVector3f	VelocityMax = FVector3f::ZeroVector;
-		FVector3f	WindMin = FVector3f::ZeroVector;
-		FVector3f	WindMax = FVector3f::ZeroVector;
-		FVector3f	AccelerationMin = FVector3f::ZeroVector;
-		FVector3f	AccelerationMax = FVector3f::ZeroVector;
+		FNiagaraStatelessRangeFloat		MassRange = FNiagaraStatelessRangeFloat(1.0f);
+		FNiagaraStatelessRangeFloat		DragRange = FNiagaraStatelessRangeFloat(0.0f);
+		FNiagaraStatelessRangeVector3	VelocityRange = FNiagaraStatelessRangeVector3(FVector3f::ZeroVector);
+		FNiagaraStatelessRangeVector3	WindRange = FNiagaraStatelessRangeVector3(FVector3f::ZeroVector);
+		FNiagaraStatelessRangeVector3	AccelerationRange = FNiagaraStatelessRangeVector3(FVector3f::ZeroVector);
 
-		bool		bConeVelocity = false;
-		FQuat4f		ConeQuat = FQuat4f::Identity;
-		float		ConeVelocityMin = 0.0f;
-		float		ConeVelocityMax = 0.0f;
-		float		ConeOuterAngle = 0.0f;
-		float		ConeInnerAngle = 0.0f;
-		float		ConeVelocityFalloff = 0.0f;
+		bool							bConeVelocity = false;
+		FQuat4f							ConeQuat = FQuat4f::Identity;
+		FNiagaraStatelessRangeFloat		ConeVelocityRange = FNiagaraStatelessRangeFloat(0.0f);
+		float							ConeOuterAngle = 0.0f;
+		float							ConeInnerAngle = 0.0f;
+		float							ConeVelocityFalloff = 0.0f;
 
-		bool		bPointVelocity = false;
-		float		PointVelocityMin = 0.0f;
-		float		PointVelocityMax = 0.0f;
-		FVector3f	PointOrigin = FVector3f::ZeroVector;
+		bool							bPointVelocity = false;
+		FNiagaraStatelessRangeFloat		PointVelocityRange = FNiagaraStatelessRangeFloat(0.0f);
+		float							PointVelocityMax = 0.0f;
+		FVector3f						PointOrigin = FVector3f::ZeroVector;
 
 		bool		bNoiseEnabled = false;
 		float		NoiseAmplitude = 0.0f;
diff --git a/Engine/Plugins/FX/Niagara/Source/Niagara/Internal/Stateless/Modules/NiagaraStatelessModule_AccelerationForce.h b/Engine/Plugins/FX/Niagara/Source/Niagara/Internal/Stateless/Modules/NiagaraStatelessModule_AccelerationForce.h
--- a/Engine/Plugins/FX/Niagara/Source/Niagara/Internal/Stateless/Modules/NiagaraStatelessModule_AccelerationForce.h
+++ b/Engine/Plugins/FX/Niagara/Source/Niagara/Internal/Stateless/Modules/NiagaraStatelessModule_AccelerationForce.h
@@ -15,10 +15,8 @@ class UNiagaraStatelessModule_AccelerationForce : public UNiagaraStatelessModule
 	GENERATED_BODY()
 
 public:
-	UPROPERTY(EditAnywhere, Category = "Parameters")
-	FVector3f AccelerationMin = FVector3f(0.0f, 0.0f, 100.0f);
-	UPROPERTY(EditAnywhere, Category = "Parameters")
-	FVector3f AccelerationMax = FVector3f(0.0f, 0.0f, 100.0f);
+	UPROPERTY(EditAnywhere, Category = "Parameters", meta = (DisplayName = "Acceleration", DisableUniformDistribution))
+	FNiagaraDistributionRangeVector3 AccelerationDistribution = FNiagaraDistributionRangeVector3(FVector3f::ZeroVector);
 
 	virtual void BuildEmitterData(FNiagaraStatelessEmitterDataBuildContext& BuildContext) const override
 	{
@@ -26,9 +24,12 @@ public:
 		{
 			return;
 		}
+
+		const FNiagaraStatelessRangeVector3 AccelerationRange = AccelerationDistribution.CalculateRange(FVector3f::ZeroVector);
+
 		NiagaraStateless::FPhysicsBuildData& PhysicsBuildData = BuildContext.GetTransientBuildData<NiagaraStateless::FPhysicsBuildData>();
-		PhysicsBuildData.AccelerationMin += AccelerationMin;
-		PhysicsBuildData.AccelerationMax += AccelerationMax;
+		PhysicsBuildData.AccelerationRange.Min += AccelerationRange.Min;
+		PhysicsBuildData.AccelerationRange.Max += AccelerationRange.Max;
 	}
 
 #if WITH_EDITOR
diff --git a/Engine/Plugins/FX/Niagara/Source/Niagara/Internal/Stateless/Modules/NiagaraStatelessModule_AddVelocity.h b/Engine/Plugins/FX/Niagara/Source/Niagara/Internal/Stateless/Modules/NiagaraStatelessModule_AddVelocity.h
--- a/Engine/Plugins/FX/Niagara/Source/Niagara/Internal/Stateless/Modules/NiagaraStatelessModule_AddVelocity.h
+++ b/Engine/Plugins/FX/Niagara/Source/Niagara/Internal/Stateless/Modules/NiagaraStatelessModule_AddVelocity.h
@@ -26,17 +26,13 @@ public:
 	UPROPERTY(EditAnywhere, Category = "Parameters")
 	ENSM_VelocityType VelocityType = ENSM_VelocityType::Linear;
 
-	UPROPERTY(EditAnywhere, Category = "Parameters", meta = (EditConditionHides, EditCondition = "VelocityType == ENSM_VelocityType::Linear"))
-	FVector3f VelocityMin = FVector3f(0.0f, 0.0f, 100.0f);
-	UPROPERTY(EditAnywhere, Category = "Parameters", meta = (EditConditionHides, EditCondition = "VelocityType == ENSM_VelocityType::Linear"))
-	FVector3f VelocityMax = FVector3f(0.0f, 0.0f, 100.0f);
-	UPROPERTY(EditAnywhere, Category = "Parameters", meta = (EditConditionHides, EditCondition = "VelocityType == ENSM_VelocityType::Linear"))
-	float VelocityScale = 1.0f;
+	UPROPERTY(EditAnywhere, Category = "Parameters", meta = (DisplayName = "Velocity", DisableUniformDistribution, EditConditionHides, EditCondition = "VelocityType == ENSM_VelocityType::Linear"))
+	FNiagaraDistributionRangeVector3 LinearVelocityDistribution = FNiagaraDistributionRangeVector3(FVector3f(0.0f, 0.0f, 100.0f));
+	UPROPERTY(EditAnywhere, Category = "Parameters", meta = (DisplayName = "Velocity Scale", EditConditionHides, EditCondition = "VelocityType == ENSM_VelocityType::Linear"))
+	float LinearVelocityScale = 1.0f;
 
-	UPROPERTY(EditAnywhere, Category = "Parameters", meta = (EditConditionHides, EditCondition = "VelocityType == ENSM_VelocityType::InCone"))
-	float ConeVelocityMin = 500.0f;
-	UPROPERTY(EditAnywhere, Category = "Parameters", meta = (EditConditionHides, EditCondition = "VelocityType == ENSM_VelocityType::InCone"))
-	float ConeVelocityMax = 500.0f;
+	UPROPERTY(EditAnywhere, Category = "Parameters", meta = (DisplayName = "Cone Velocity", EditConditionHides, EditCondition = "VelocityType == ENSM_VelocityType::InCone"))
+	FNiagaraDistributionRangeFloat ConeVelocityDistribution = FNiagaraDistributionRangeFloat(500.0f);
 	UPROPERTY(EditAnywhere, Category = "Parameters", meta = (EditConditionHides, EditCondition = "VelocityType == ENSM_VelocityType::InCone"))
 	FRotator ConeRotation = FRotator::ZeroRotator;
 	UPROPERTY(EditAnywhere, Category = "Parameters", meta = (EditConditionHides, EditCondition = "VelocityType == ENSM_VelocityType::InCone", ClampMin = "0.0", ClampMax = "360.0", UIMin = "0.0", UIMax = "360.0"))
@@ -44,10 +40,8 @@ public:
 	UPROPERTY(EditAnywhere, Category = "Parameters", meta = (EditConditionHides, EditCondition = "VelocityType == ENSM_VelocityType::InCone", ClampMin = "0.0", ClampMax = "360.0", UIMin = "0.0", UIMax = "360.0"))
 	float InnerCone = 0.0f;
 
-	UPROPERTY(EditAnywhere, Category = "Parameters", meta = (EditConditionHides, EditCondition = "VelocityType == ENSM_VelocityType::FromPoint"))
-	float PointVelocityMin = 100.0f;
-	UPROPERTY(EditAnywhere, Category = "Parameters", meta = (EditConditionHides, EditCondition = "VelocityType == ENSM_VelocityType::FromPoint"))
-	float PointVelocityMax = 100.0f;
+	UPROPERTY(EditAnywhere, Category = "Parameters", meta = (DisplayName = "Point Velocity", EditConditionHides, EditCondition = "VelocityType == ENSM_VelocityType::FromPoint"))
+	FNiagaraDistributionRangeFloat PointVelocityDistribution = FNiagaraDistributionRangeFloat(100.0f);
 	UPROPERTY(EditAnywhere, Category = "Parameters", meta = (EditConditionHides, EditCondition = "VelocityType == ENSM_VelocityType::FromPoint"))
 	FVector3f PointOrigin = FVector3f::ZeroVector;
 
diff --git a/Engine/Plugins/FX/Niagara/Source/Niagara/Internal/Stateless/Modules/NiagaraStatelessModule_Drag.h b/Engine/Plugins/FX/Niagara/Source/Niagara/Internal/Stateless/Modules/NiagaraStatelessModule_Drag.h
--- a/Engine/Plugins/FX/Niagara/Source/Niagara/Internal/Stateless/Modules/NiagaraStatelessModule_Drag.h
+++ b/Engine/Plugins/FX/Niagara/Source/Niagara/Internal/Stateless/Modules/NiagaraStatelessModule_Drag.h
@@ -14,11 +14,11 @@ class UNiagaraStatelessModule_Drag : public UNiagaraStatelessModule
 {
 	GENERATED_BODY()
 
+	static constexpr float DefaultDrag = 1.0f;
+
 public:
-	UPROPERTY(EditAnywhere, Category = "Parameters")
-	float DragMin = 0.1f;
-	UPROPERTY(EditAnywhere, Category = "Parameters")
-	float DragMax = 1.0f;
+	UPROPERTY(EditAnywhere, Category = "Parameters", meta = (DisplayName = "Drag"))
+	FNiagaraDistributionRangeFloat DragDistribution = FNiagaraDistributionRangeFloat(DefaultDrag);
 
 	virtual void BuildEmitterData(FNiagaraStatelessEmitterDataBuildContext& BuildContext) const override
 	{
@@ -28,8 +28,7 @@ public:
 		}
 
 		NiagaraStateless::FPhysicsBuildData& PhysicsBuildData = BuildContext.GetTransientBuildData<NiagaraStateless::FPhysicsBuildData>();
-		PhysicsBuildData.DragMin += DragMin;
-		PhysicsBuildData.DragMax += DragMax;
+		PhysicsBuildData.DragRange = DragDistribution.CalculateRange(DefaultDrag);
 	}
 
 #if WITH_EDITOR
diff --git a/Engine/Plugins/FX/Niagara/Source/Niagara/Internal/Stateless/Modules/NiagaraStatelessModule_GravityForce.h b/Engine/Plugins/FX/Niagara/Source/Niagara/Internal/Stateless/Modules/NiagaraStatelessModule_GravityForce.h
--- a/Engine/Plugins/FX/Niagara/Source/Niagara/Internal/Stateless/Modules/NiagaraStatelessModule_GravityForce.h
+++ b/Engine/Plugins/FX/Niagara/Source/Niagara/Internal/Stateless/Modules/NiagaraStatelessModule_GravityForce.h
@@ -15,10 +15,8 @@ class UNiagaraStatelessModule_GravityForce : public UNiagaraStatelessModule
 	GENERATED_BODY()
 
 public:
-	UPROPERTY(EditAnywhere, Category = "Parameters")
-	FVector3f GravityForceMin = FVector3f(0.0f, 0.0f, -980.0f);
-	UPROPERTY(EditAnywhere, Category = "Parameters")
-	FVector3f GravityForceMax = FVector3f(0.0f, 0.0f, -980.0f);
+	UPROPERTY(EditAnywhere, Category = "Parameters", meta = (DisplayName = "Gravity", DisableUniformDistribution))
+	FNiagaraDistributionRangeVector3 GravityDistribution = FNiagaraDistributionRangeVector3(GetDefaultValue());
 
 	virtual void BuildEmitterData(FNiagaraStatelessEmitterDataBuildContext& BuildContext) const override
 	{
@@ -26,11 +24,15 @@ public:
 		{
 			return;
 		}
+		const FNiagaraStatelessRangeVector3 GravityRange = GravityDistribution.CalculateRange(GetDefaultValue());
+
 		NiagaraStateless::FPhysicsBuildData& PhysicsBuildData = BuildContext.GetTransientBuildData<NiagaraStateless::FPhysicsBuildData>();
-		PhysicsBuildData.AccelerationMin += GravityForceMin;
-		PhysicsBuildData.AccelerationMax += GravityForceMax;
+		PhysicsBuildData.AccelerationRange.Min += GravityRange.Min;
+		PhysicsBuildData.AccelerationRange.Max += GravityRange.Max;
 	}
 
+	static FVector3f GetDefaultValue() { return FVector3f(0.0f, 0.0f, -980.0f); }
+
 #if WITH_EDITOR
 	virtual bool CanDisableModule() const override { return true; }
 #endif
diff --git a/Engine/Plugins/FX/Niagara/Source/Niagara/Internal/Stateless/Modules/NiagaraStatelessModule_InitializeParticle.h b/Engine/Plugins/FX/Niagara/Source/Niagara/Internal/Stateless/Modules/NiagaraStatelessModule_InitializeParticle.h
--- a/Engine/Plugins/FX/Niagara/Source/Niagara/Internal/Stateless/Modules/NiagaraStatelessModule_InitializeParticle.h
+++ b/Engine/Plugins/FX/Niagara/Source/Niagara/Internal/Stateless/Modules/NiagaraStatelessModule_InitializeParticle.h
@@ -12,97 +12,57 @@
 
 #include "NiagaraStatelessModule_InitializeParticle.generated.h"
 
-UENUM()
-enum class ENSPositionDistributionType
-{
-	UniqueIndexClamp,
-	UniqueIndexWrap,
-	Random,
-	RandomInterpolateToNext,
-	RandomInterpolateToRandom,
-};
-
-USTRUCT()
-struct FPositionDistributionTest
-{
-	GENERATED_BODY()
-
-	UPROPERTY(EditAnywhere, Category = "Parameters")
-	ENSPositionDistributionType Type = ENSPositionDistributionType::UniqueIndexClamp;
-
-	UPROPERTY(EditAnywhere, Category = "Parameters")
-	TArray<FVector3f> Values;
-};
-
-//UENUM()
-//enum class NiagaraStatelessModuleInitializeParticle_ColorMode
-//{
-//	Unset,
-//	DirectSet,
-//	RandomRange,
-//	//RandomHSV UMETA(DisplayName="Random Hue/Saturation/Value"),
-//};
-
 UCLASS(MinimalAPI, EditInlineNew, meta = (DisplayName = "Initialize Particle"))
 class UNiagaraStatelessModule_InitializeParticle : public UNiagaraStatelessModule
 {
 	GENERATED_BODY()
 
+	static const uint32 EInitializeParticleModuleFlag_UniformSpriteSize	= 1 << 0;
+	static const uint32 EInitializeParticleModuleFlag_UniformMeshScale	= 1 << 1;
+
 	struct FModuleBuiltData
 	{
-		int32	PositionParameterBinding = INDEX_NONE;
-		int32	PositionTableOffset = 0;
-		int32	PositionTableLength = 0;
+		uint32							ModuleFlags = 0;
+		int32							PositionParameterBinding = INDEX_NONE;
+		FNiagaraStatelessRangeFloat		LifetimeRange;
+		FNiagaraStatelessRangeColor		ColorRange;
+		FNiagaraStatelessRangeFloat		MassRange;
+		FNiagaraStatelessRangeVector2	SpriteSizeRange;
+		FNiagaraStatelessRangeFloat		SpriteRotationRange;
+		FNiagaraStatelessRangeVector3	MeshScaleRange;
+		FNiagaraStatelessRangeFloat		RibbonWidthRange;
 	};
 
 public:
 	using FParameters = NiagaraStateless::FInitializeParticleModule_ShaderParameters;
 
-	UPROPERTY(EditAnywhere, Category = "Parameters")
-	float LifetimeMin = 1.0f;
-	UPROPERTY(EditAnywhere, Category = "Parameters")
-	float LifetimeMax = 1.0f;
-
-	UPROPERTY(EditAnywhere, Category = "Parameters")
-	FVector3f	InitialPosition = FVector3f(0.0f);
+	UPROPERTY(EditAnywhere, Category = "Parameters", meta = (DisplayName = "Lifetime"))
+	FNiagaraDistributionRangeFloat LifetimeDistribution = FNiagaraDistributionRangeFloat(FNiagaraStatelessGlobals::GetDefaultLifetimeValue());
 
-	UPROPERTY(EditAnywhere, Category = "Parameters")
-	FLinearColor	ColorMin = FLinearColor::White;
-	UPROPERTY(EditAnywhere, Category = "Parameters")
-	FLinearColor	ColorMax = FLinearColor::White;
+	UPROPERTY(EditAnywhere, Category = "Parameters", meta = (DisplayName = "Color", DisableCurveDistribution))
+	FNiagaraDistributionColor ColorDistribution = FNiagaraDistributionColor(FNiagaraStatelessGlobals::GetDefaultColorValue());
 
-	UPROPERTY(EditAnywhere, Category = "Parameters")
-	float MassMin = 1.0f;
-	UPROPERTY(EditAnywhere, Category = "Parameters")
-	float MassMax = 1.0f;
+	UPROPERTY(EditAnywhere, Category = "Parameters", meta = (DisplayName = "Mass"))
+	FNiagaraDistributionRangeFloat MassDistribution = FNiagaraDistributionRangeFloat(FNiagaraStatelessGlobals::GetDefaultMassValue());
 
-	UPROPERTY(EditAnywhere, Category = "Parameters")
-	FVector2f	SpriteSizeMin = FVector2f(10.0f);
-	UPROPERTY(EditAnywhere, Category = "Parameters")
-	FVector2f	SpriteSizeMax = FVector2f(10.0f);
+	UPROPERTY(EditAnywhere, Category = "Parameters", meta = (DisplayName = "Sprite Size"))
+	FNiagaraDistributionRangeVector2 SpriteSizeDistribution = FNiagaraDistributionRangeVector2(FNiagaraStatelessGlobals::GetDefaultSpriteSizeValue());
 
-	UPROPERTY(EditAnywhere, Category = "Parameters")
-	float	SpriteRotationMin = 0.0f;
-	UPROPERTY(EditAnywhere, Category = "Parameters")
-	float	SpriteRotationMax = 0.0f;
+	UPROPERTY(EditAnywhere, Category = "Parameters", meta = (DisplayName = "Sprite Rotation"))
+	FNiagaraDistributionRangeFloat SpriteRotationDistribution = FNiagaraDistributionRangeFloat(FNiagaraStatelessGlobals::GetDefaultSpriteRotationValue());
 
-	UPROPERTY(EditAnywhere, Category = "Parameters")
-	FVector3f	MeshScaleMin = FVector3f(1.0f);
-	UPROPERTY(EditAnywhere, Category = "Parameters")
-	FVector3f	MeshScaleMax = FVector3f(1.0f);
+	UPROPERTY(EditAnywhere, Category = "Parameters", meta = (DisplayName = "Mesh Scale"))
+	FNiagaraDistributionRangeVector3 MeshScaleDistribution = FNiagaraDistributionRangeVector3(FNiagaraStatelessGlobals::GetDefaultScaleValue());
 
 	UPROPERTY(EditAnywhere, Category = "Parameters", meta = (InlineEditConditionToggle))
 	bool bWriteRibbonWidth = false;
 
-	UPROPERTY(EditAnywhere, Category = "Parameters", meta = (EditCondition="bWriteRibbonWidth"))
-	FVector2f	RibbonWidthMinMax = FVector2f(10.0f, 10.0f);
+	UPROPERTY(EditAnywhere, Category = "Parameters", meta = (DisplayName = "Ribbon Width", EditCondition="bWriteRibbonWidth"))
+	FNiagaraDistributionRangeFloat RibbonWidthDistribution = FNiagaraDistributionRangeFloat(FNiagaraStatelessGlobals::GetDefaultRibbonWidthValue());
 
 	UPROPERTY(EditAnywhere, Category = "Parameters")
 	FNiagaraParameterBindingWithValue	InitialPositionBinding;
 
-	UPROPERTY(EditAnywhere, Category = "Parameters")
-	FPositionDistributionTest PositionDistribution;
-
 	//~Begin: UObject Interface
 #if WITH_EDITORONLY_DATA
 	virtual void PostInitProperties()
@@ -120,18 +80,21 @@ public:
 
 	virtual void BuildEmitterData(FNiagaraStatelessEmitterDataBuildContext& BuildContext) const override
 	{
-		const FVector3f DefaultPosition(0.0f);
-		TConstArrayView<FVector3f> PositionValues = PositionDistribution.Values.Num() > 0 ? MakeArrayView(PositionDistribution.Values) : MakeArrayView(&DefaultPosition, 1);
-
 		FModuleBuiltData* BuiltData = BuildContext.AllocateBuiltData<FModuleBuiltData>();
-		BuiltData->PositionTableOffset = BuildContext.AddStaticData(PositionValues);
-		BuiltData->PositionTableLength = PositionValues.Num();
+		BuiltData->ModuleFlags				 = SpriteSizeDistribution.IsUniform() ? EInitializeParticleModuleFlag_UniformSpriteSize : 0;
+		BuiltData->ModuleFlags				|= MeshScaleDistribution.IsUniform() ? EInitializeParticleModuleFlag_UniformMeshScale : 0;
+
+		BuiltData->PositionParameterBinding	= BuildContext.AddRendererBinding(InitialPositionBinding);
+		BuiltData->LifetimeRange			= LifetimeDistribution.CalculateRange(FNiagaraStatelessGlobals::GetDefaultLifetimeValue());
+		BuiltData->ColorRange				= ColorDistribution.CalculateRange(FNiagaraStatelessGlobals::GetDefaultColorValue());
+		BuiltData->MassRange				= MassDistribution.CalculateRange(FNiagaraStatelessGlobals::GetDefaultMassValue());
+		BuiltData->SpriteSizeRange			= SpriteSizeDistribution.CalculateRange(FNiagaraStatelessGlobals::GetDefaultSpriteSizeValue());
+		BuiltData->SpriteRotationRange		= SpriteRotationDistribution.CalculateRange(FNiagaraStatelessGlobals::GetDefaultSpriteRotationValue());
+		BuiltData->MeshScaleRange			= MeshScaleDistribution.CalculateRange(FNiagaraStatelessGlobals::GetDefaultScaleValue());
+		BuiltData->RibbonWidthRange			= RibbonWidthDistribution.CalculateRange(FNiagaraStatelessGlobals::GetDefaultRibbonWidthValue());
 
-		BuiltData->PositionParameterBinding = BuildContext.AddRendererBinding(InitialPositionBinding);
-	
 		NiagaraStateless::FPhysicsBuildData& PhysicsBuildData = BuildContext.GetTransientBuildData<NiagaraStateless::FPhysicsBuildData>();
-		PhysicsBuildData.MassMin = MassMin;
-		PhysicsBuildData.MassMax = MassMax;
+		PhysicsBuildData.MassRange			= MassDistribution.CalculateRange(FNiagaraStatelessGlobals::GetDefaultMassValue());
 	}
 
 	virtual void SetShaderParameters(const FNiagaraStatelessSetShaderParameterContext& SetShaderParameterContext) const override
@@ -139,42 +102,18 @@ public:
 		FParameters* Parameters = SetShaderParameterContext.GetParameterNestedStruct<FParameters>();
 		const FModuleBuiltData* ModuleBuiltData = SetShaderParameterContext.ReadBuiltData<FModuleBuiltData>();
 
-		//Parameters->InitializeParticle_Position			= InitialPosition;
 		SetShaderParameterContext.GetRendererParameterValue(Parameters->InitializeParticle_Position, ModuleBuiltData->PositionParameterBinding, InitialPositionBinding.GetDefaultValue<FVector3f>());
-		Parameters->InitializeParticle_ColorScale			= ColorMax - ColorMin;
-		Parameters->InitializeParticle_ColorBias			= ColorMin;
-		Parameters->InitializeParticle_SpriteSizeScale		= SpriteSizeMax - SpriteSizeMin;
-		Parameters->InitializeParticle_SpriteSizeBias		= SpriteSizeMin;
-		Parameters->InitializeParticle_SpriteRotationScale	= SpriteRotationMax - SpriteRotationMin;
-		Parameters->InitializeParticle_SpriteRotationBias	= SpriteRotationMin;
-		Parameters->InitializeParticle_MeshScaleScale		= MeshScaleMax - MeshScaleMin;
-		Parameters->InitializeParticle_MeshScaleBias		= MeshScaleMin;
-		Parameters->InitializeParticle_RibbonWidthScale		= RibbonWidthMinMax.Y - RibbonWidthMinMax.X;
-		Parameters->InitializeParticle_RibbonWidthBias		= RibbonWidthMinMax.X;
-
-		Parameters->InitializeParticle_PositionMode			= 0;
-		Parameters->InitializeParticle_PositionTableMod		= ModuleBuiltData->PositionTableLength;
-		Parameters->InitializeParticle_PositionTableOffset	= ModuleBuiltData->PositionTableOffset;
-		Parameters->InitializeParticle_PositionTableLength	= ModuleBuiltData->PositionTableLength;
-		switch (PositionDistribution.Type)
-		{
-			case ENSPositionDistributionType::UniqueIndexClamp:
-				Parameters->InitializeParticle_PositionMode			= 0;
-				Parameters->InitializeParticle_PositionTableMod		= 0xffffffff;
-				break;
-			case ENSPositionDistributionType::UniqueIndexWrap:
-				Parameters->InitializeParticle_PositionMode			= 0;
-				break;
-			case ENSPositionDistributionType::Random:
-				Parameters->InitializeParticle_PositionMode			= 2;
-				break;
-			case ENSPositionDistributionType::RandomInterpolateToNext:
-				Parameters->InitializeParticle_PositionMode			= 1;
-				break;
-			case ENSPositionDistributionType::RandomInterpolateToRandom:
-				Parameters->InitializeParticle_PositionMode			= 3;
-				break;
-		}
+		Parameters->InitializeParticle_ModuleFlags			= ModuleBuiltData->ModuleFlags;
+		Parameters->InitializeParticle_ColorScale			= ModuleBuiltData->ColorRange.GetScale();
+		Parameters->InitializeParticle_ColorBias			= ModuleBuiltData->ColorRange.Min;
+		Parameters->InitializeParticle_SpriteSizeScale		= ModuleBuiltData->SpriteSizeRange.GetScale();
+		Parameters->InitializeParticle_SpriteSizeBias		= ModuleBuiltData->SpriteSizeRange.Min;
+		Parameters->InitializeParticle_SpriteRotationScale	= ModuleBuiltData->SpriteRotationRange.GetScale();
+		Parameters->InitializeParticle_SpriteRotationBias	= ModuleBuiltData->SpriteRotationRange.Min;
+		Parameters->InitializeParticle_MeshScaleScale		= ModuleBuiltData->MeshScaleRange.GetScale();
+		Parameters->InitializeParticle_MeshScaleBias		= ModuleBuiltData->MeshScaleRange.Min;
+		Parameters->InitializeParticle_RibbonWidthScale		= ModuleBuiltData->RibbonWidthRange.GetScale();
+		Parameters->InitializeParticle_RibbonWidthBias		= ModuleBuiltData->RibbonWidthRange.Min;
 	}
 
 #if WITH_EDITORONLY_DATA
diff --git a/Engine/Plugins/FX/Niagara/Source/Niagara/Internal/Stateless/Modules/NiagaraStatelessModule_ScaleColor.h b/Engine/Plugins/FX/Niagara/Source/Niagara/Internal/Stateless/Modules/NiagaraStatelessModule_ScaleColor.h
--- a/Engine/Plugins/FX/Niagara/Source/Niagara/Internal/Stateless/Modules/NiagaraStatelessModule_ScaleColor.h
+++ b/Engine/Plugins/FX/Niagara/Source/Niagara/Internal/Stateless/Modules/NiagaraStatelessModule_ScaleColor.h
@@ -15,24 +15,19 @@ class UNiagaraStatelessModule_ScaleColor : public UNiagaraStatelessModule
 
 	struct FModuleBuiltData
 	{
-		int32	TableOffset = 0;
-		int32	TableLength = 0;
+		FUintVector3	DistributionParameters = FUintVector3::ZeroValue;
 	};
 
 public:
 	using FParameters = NiagaraStateless::FScaleColorModule_ShaderParameters;
 
-	UPROPERTY(EditAnywhere, Category = "Parameters")
-	TArray<FLinearColor>	ScaleValues = {FLinearColor::White, FLinearColor::White, FLinearColor::White, FLinearColor::Transparent};
+	UPROPERTY(EditAnywhere, Category = "Parameters", meta = (DisplayName = "Scale"))
+	FNiagaraDistributionColor ScaleDistribution = FNiagaraDistributionColor(FLinearColor::White);
 
 	virtual void BuildEmitterData(FNiagaraStatelessEmitterDataBuildContext& BuildContext) const override
 	{
 		FModuleBuiltData* BuiltData = BuildContext.AllocateBuiltData<FModuleBuiltData>();
-		if (IsModuleEnabled())
-		{
-			BuiltData->TableOffset = BuildContext.AddStaticData(ScaleValues);
-			BuiltData->TableLength = ScaleValues.Num();
-		}
+		BuiltData->DistributionParameters = BuildContext.AddDistribution(ScaleDistribution, IsModuleEnabled());
 	}
 
 	virtual void SetShaderParameters(const FNiagaraStatelessSetShaderParameterContext& SetShaderParameterContext) const override
@@ -40,8 +35,7 @@ public:
 		const FModuleBuiltData* ModuleBuiltData = SetShaderParameterContext.ReadBuiltData<FModuleBuiltData>();
 
 		FParameters* Parameters = SetShaderParameterContext.GetParameterNestedStruct<FParameters>();
-		Parameters->ScaleColor_Offset = ModuleBuiltData->TableOffset;
-		Parameters->ScaleColor_Length = ModuleBuiltData->TableLength;
+		Parameters->ScaleColor_Distribution = ModuleBuiltData->DistributionParameters;
 	}
 
 #if WITH_EDITOR
diff --git a/Engine/Plugins/FX/Niagara/Source/Niagara/Internal/Stateless/Modules/NiagaraStatelessModule_ScaleMeshSize.h b/Engine/Plugins/FX/Niagara/Source/Niagara/Internal/Stateless/Modules/NiagaraStatelessModule_ScaleMeshSize.h
--- a/Engine/Plugins/FX/Niagara/Source/Niagara/Internal/Stateless/Modules/NiagaraStatelessModule_ScaleMeshSize.h
+++ b/Engine/Plugins/FX/Niagara/Source/Niagara/Internal/Stateless/Modules/NiagaraStatelessModule_ScaleMeshSize.h
@@ -15,24 +15,19 @@ class UNiagaraStatelessModule_ScaleMeshSize : public UNiagaraStatelessModule
 
 	struct FModuleBuiltData
 	{
-		int32	TableOffset = 0;
-		int32	TableLength = 0;
+		FUintVector3	DistributionParameters = FUintVector3::ZeroValue;
 	};
 
 public:
 	using FParameters = NiagaraStateless::FScaleMeshSizeModule_ShaderParameters;
 
-	UPROPERTY(EditAnywhere, Category = "Parameters")
-	TArray<FVector3f>	ScaleValues = { FVector3f::OneVector, FVector3f::OneVector, FVector3f::OneVector, FVector3f::ZeroVector };
+	UPROPERTY(EditAnywhere, Category = "Parameters", meta = (DisplayName = "Scale"))
+	FNiagaraDistributionVector3 ScaleDistribution = FNiagaraDistributionVector3(FVector3f::OneVector);
 
 	virtual void BuildEmitterData(FNiagaraStatelessEmitterDataBuildContext& BuildContext) const override
 	{
 		FModuleBuiltData* BuiltData = BuildContext.AllocateBuiltData<FModuleBuiltData>();
-		if (IsModuleEnabled())
-		{
-			BuiltData->TableOffset = BuildContext.AddStaticData(ScaleValues);
-			BuiltData->TableLength = ScaleValues.Num();
-		}
+		BuiltData->DistributionParameters = BuildContext.AddDistribution(ScaleDistribution, IsModuleEnabled());
 	}
 
 	virtual void SetShaderParameters(const FNiagaraStatelessSetShaderParameterContext& SetShaderParameterContext) const override
@@ -40,8 +35,7 @@ public:
 		const FModuleBuiltData* ModuleBuiltData = SetShaderParameterContext.ReadBuiltData<FModuleBuiltData>();
 
 		FParameters* Parameters = SetShaderParameterContext.GetParameterNestedStruct<FParameters>();
-		Parameters->ScaleMeshSize_Offset = ModuleBuiltData->TableOffset;
-		Parameters->ScaleMeshSize_Length = ModuleBuiltData->TableLength;
+		Parameters->ScaleMeshSize_Distribution = ModuleBuiltData->DistributionParameters;
 	}
 
 #if WITH_EDITOR
diff --git a/Engine/Plugins/FX/Niagara/Source/Niagara/Internal/Stateless/Modules/NiagaraStatelessModule_ScaleSpriteSize.h b/Engine/Plugins/FX/Niagara/Source/Niagara/Internal/Stateless/Modules/NiagaraStatelessModule_ScaleSpriteSize.h
--- a/Engine/Plugins/FX/Niagara/Source/Niagara/Internal/Stateless/Modules/NiagaraStatelessModule_ScaleSpriteSize.h
+++ b/Engine/Plugins/FX/Niagara/Source/Niagara/Internal/Stateless/Modules/NiagaraStatelessModule_ScaleSpriteSize.h
@@ -2,6 +2,7 @@
 
 #pragma once
 
+#include "Stateless/NiagaraStatelessCommon.h"
 #include "Stateless/NiagaraStatelessModule.h"
 #include "Stateless/NiagaraStatelessEmitterDataBuildContext.h"
 #include "Stateless/NiagaraStatelessModuleShaderParameters.h"
@@ -15,24 +16,19 @@ class UNiagaraStatelessModule_ScaleSpriteSize : public UNiagaraStatelessModule
 
 	struct FModuleBuiltData
 	{
-		int32	TableOffset = 0;
-		int32	TableLength = 0;
+		FUintVector3	DistributionParameters = FUintVector3::ZeroValue;
 	};
 
 public:
 	using FParameters = NiagaraStateless::FScaleSpriteSizeModule_ShaderParameters;
 
-	UPROPERTY(EditAnywhere, Category = "Parameters")
-	TArray<FVector2f>	ScaleValues = {FVector2f::One(), FVector2f::One(), FVector2f::One(), FVector2f::ZeroVector};
-
+	UPROPERTY(EditAnywhere, Category = "Parameters", meta = (DisplayName = "Scale"))
+	FNiagaraDistributionVector2 ScaleDistribution = FNiagaraDistributionVector2(FVector2f::One());
+	
 	virtual void BuildEmitterData(FNiagaraStatelessEmitterDataBuildContext& BuildContext) const override
 	{
 		FModuleBuiltData* BuiltData = BuildContext.AllocateBuiltData<FModuleBuiltData>();
-		if (IsModuleEnabled())
-		{
-			BuiltData->TableOffset = BuildContext.AddStaticData(ScaleValues);
-			BuiltData->TableLength = ScaleValues.Num();
-		}
+		BuiltData->DistributionParameters = BuildContext.AddDistribution(ScaleDistribution, IsModuleEnabled());
 	}
 
 	virtual void SetShaderParameters(const FNiagaraStatelessSetShaderParameterContext& SetShaderParameterContext) const override
@@ -40,8 +36,7 @@ public:
 		const FModuleBuiltData* ModuleBuiltData = SetShaderParameterContext.ReadBuiltData<FModuleBuiltData>();
 
 		FParameters* Parameters = SetShaderParameterContext.GetParameterNestedStruct<FParameters>();
-		Parameters->ScaleSpriteSize_Offset = ModuleBuiltData->TableOffset;
-		Parameters->ScaleSpriteSize_Length = ModuleBuiltData->TableLength;
+		Parameters->ScaleSpriteSize_Distribution = ModuleBuiltData->DistributionParameters;
 	}
 
 #if WITH_EDITOR
diff --git a/Engine/Plugins/FX/Niagara/Source/Niagara/Internal/Stateless/Modules/NiagaraStatelessModule_SolveVelocitiesAndForces.h b/Engine/Plugins/FX/Niagara/Source/Niagara/Internal/Stateless/Modules/NiagaraStatelessModule_SolveVelocitiesAndForces.h
--- a/Engine/Plugins/FX/Niagara/Source/Niagara/Internal/Stateless/Modules/NiagaraStatelessModule_SolveVelocitiesAndForces.h
+++ b/Engine/Plugins/FX/Niagara/Source/Niagara/Internal/Stateless/Modules/NiagaraStatelessModule_SolveVelocitiesAndForces.h
@@ -27,26 +27,20 @@ public:
 		NiagaraStateless::FPhysicsBuildData& PhysicsBuildData = BuildContext.GetTransientBuildData<NiagaraStateless::FPhysicsBuildData>();
 
 		FModuleBuiltData* BuiltData		= BuildContext.AllocateBuiltData<FModuleBuiltData>();
-		BuiltData->MassMin				= PhysicsBuildData.MassMin;
-		BuiltData->MassMax				= PhysicsBuildData.MassMax;
-		BuiltData->DragMin				= FMath::Clamp(PhysicsBuildData.DragMin, 0.01f, 1.0f);
-		BuiltData->DragMax				= FMath::Clamp(PhysicsBuildData.DragMax, 0.01f, 1.0f);
-		BuiltData->VelocityMin			= PhysicsBuildData.VelocityMin;
-		BuiltData->VelocityMax			= PhysicsBuildData.VelocityMax;
-		BuiltData->WindMin				= PhysicsBuildData.WindMin;
-		BuiltData->WindMax				= PhysicsBuildData.WindMax;
-		BuiltData->AccelerationMin		= PhysicsBuildData.AccelerationMin;
-		BuiltData->AccelerationMax		= PhysicsBuildData.AccelerationMax;
+		BuiltData->MassRange			= PhysicsBuildData.MassRange;
+		BuiltData->DragRange.Min		= FMath::Clamp(PhysicsBuildData.DragRange.Min, 0.01f, 1.0f);
+		BuiltData->DragRange.Max		= FMath::Clamp(PhysicsBuildData.DragRange.Max, 0.01f, 1.0f);
+		BuiltData->VelocityRange		= PhysicsBuildData.VelocityRange;
+		BuiltData->WindRange			= PhysicsBuildData.WindRange;
+		BuiltData->AccelerationRange	= PhysicsBuildData.AccelerationRange;
 		BuiltData->bConeVelocity		= PhysicsBuildData.bConeVelocity;
 		BuiltData->ConeQuat				= PhysicsBuildData.ConeQuat;
-		BuiltData->ConeVelocityMin		= PhysicsBuildData.ConeVelocityMin;
-		BuiltData->ConeVelocityMax		= PhysicsBuildData.ConeVelocityMax;
+		BuiltData->ConeVelocityRange	= PhysicsBuildData.ConeVelocityRange;
 		BuiltData->ConeOuterAngle		= PhysicsBuildData.ConeOuterAngle;
 		BuiltData->ConeInnerAngle		= PhysicsBuildData.ConeInnerAngle;
 		BuiltData->ConeVelocityFalloff	= PhysicsBuildData.ConeVelocityFalloff;
 		BuiltData->bPointVelocity		= PhysicsBuildData.bPointVelocity;
-		BuiltData->PointVelocityMin		= PhysicsBuildData.PointVelocityMin;
-		BuiltData->PointVelocityMax		= PhysicsBuildData.PointVelocityMax;
+		BuiltData->PointVelocityRange	= PhysicsBuildData.PointVelocityRange;
 		BuiltData->PointOrigin			= PhysicsBuildData.PointOrigin;
 		BuiltData->bNoiseEnabled		= PhysicsBuildData.bNoiseEnabled;
 		BuiltData->NoiseAmplitude		= PhysicsBuildData.NoiseAmplitude;
@@ -63,28 +57,28 @@ public:
 		const FModuleBuiltData* ModuleBuiltData = SetShaderParameterContext.ReadBuiltData<FModuleBuiltData>();
 
 		FParameters* Parameters = SetShaderParameterContext.GetParameterNestedStruct<FParameters>();
-		Parameters->SolveVelocitiesAndForces_MassScale				= ModuleBuiltData->MassMax - ModuleBuiltData->MassMin;
-		Parameters->SolveVelocitiesAndForces_MassBias				= ModuleBuiltData->MassMin;
-		Parameters->SolveVelocitiesAndForces_DragScale				= ModuleBuiltData->DragMax - ModuleBuiltData->DragMin;
-		Parameters->SolveVelocitiesAndForces_DragBias				= ModuleBuiltData->DragMin;
-		Parameters->SolveVelocitiesAndForces_VelocityScale			= ModuleBuiltData->VelocityMax - ModuleBuiltData->VelocityMin;
-		Parameters->SolveVelocitiesAndForces_VelocityBias			= ModuleBuiltData->VelocityMin;
-		Parameters->SolveVelocitiesAndForces_WindScale				= ModuleBuiltData->WindMax - ModuleBuiltData->WindMin;
-		Parameters->SolveVelocitiesAndForces_WindBias				= ModuleBuiltData->WindMin;
-		Parameters->SolveVelocitiesAndForces_AccelerationScale		= ModuleBuiltData->AccelerationMax - ModuleBuiltData->AccelerationMin;
-		Parameters->SolveVelocitiesAndForces_AccelerationBias		= ModuleBuiltData->AccelerationMin;
+		Parameters->SolveVelocitiesAndForces_MassScale				= ModuleBuiltData->MassRange.GetScale();
+		Parameters->SolveVelocitiesAndForces_MassBias				= ModuleBuiltData->MassRange.Min;
+		Parameters->SolveVelocitiesAndForces_DragScale				= ModuleBuiltData->DragRange.GetScale();
+		Parameters->SolveVelocitiesAndForces_DragBias				= ModuleBuiltData->DragRange.Min;
+		Parameters->SolveVelocitiesAndForces_VelocityScale			= ModuleBuiltData->VelocityRange.GetScale();
+		Parameters->SolveVelocitiesAndForces_VelocityBias			= ModuleBuiltData->VelocityRange.Min;
+		Parameters->SolveVelocitiesAndForces_WindScale				= ModuleBuiltData->WindRange.GetScale();
+		Parameters->SolveVelocitiesAndForces_WindBias				= ModuleBuiltData->WindRange.Min;
+		Parameters->SolveVelocitiesAndForces_AccelerationScale		= ModuleBuiltData->AccelerationRange.GetScale();
+		Parameters->SolveVelocitiesAndForces_AccelerationBias		= ModuleBuiltData->AccelerationRange.Min;
 
 		Parameters->SolveVelocitiesAndForces_ConeVelocityEnabled	= ModuleBuiltData->bConeVelocity ? 1 : 0;
 		Parameters->SolveVelocitiesAndForces_ConeQuat				= ModuleBuiltData->ConeQuat;
-		Parameters->SolveVelocitiesAndForces_ConeVelocityScale		= ModuleBuiltData->ConeVelocityMax - ModuleBuiltData->ConeVelocityMin;
-		Parameters->SolveVelocitiesAndForces_ConeVelocityBias		= ModuleBuiltData->ConeVelocityMin;
+		Parameters->SolveVelocitiesAndForces_ConeVelocityScale		= ModuleBuiltData->ConeVelocityRange.GetScale();
+		Parameters->SolveVelocitiesAndForces_ConeVelocityBias		= ModuleBuiltData->ConeVelocityRange.Min;
 		Parameters->SolveVelocitiesAndForces_ConeAngleScale			= (ModuleBuiltData->ConeOuterAngle - ModuleBuiltData->ConeInnerAngle) * (UE_PI / 360.0f);
 		Parameters->SolveVelocitiesAndForces_ConeAngleBias			= ModuleBuiltData->ConeInnerAngle * (UE_PI / 360.0f);
 		Parameters->SolveVelocitiesAndForces_ConeVelocityFalloff	= ModuleBuiltData->ConeVelocityFalloff;
 
 		Parameters->SolveVelocitiesAndForces_PontVelocityEnabled	= ModuleBuiltData->bPointVelocity ? 1 : 0;
-		Parameters->SolveVelocitiesAndForces_PointVelocityScale		= ModuleBuiltData->PointVelocityMax - ModuleBuiltData->PointVelocityMin;
-		Parameters->SolveVelocitiesAndForces_PointVelocityBias		= ModuleBuiltData->PointVelocityMax;
+		Parameters->SolveVelocitiesAndForces_PointVelocityScale		= ModuleBuiltData->PointVelocityRange.GetScale();
+		Parameters->SolveVelocitiesAndForces_PointVelocityBias		= ModuleBuiltData->PointVelocityRange.Min;
 		Parameters->SolveVelocitiesAndForces_PointOrigin			= ModuleBuiltData->PointOrigin;
 
 		Parameters->SolveVelocitiesAndForces_NoiseEnabled			= ModuleBuiltData->bNoiseEnabled ? 1 : 0;
diff --git a/Engine/Plugins/FX/Niagara/Source/Niagara/Internal/Stateless/Modules/NiagaraStatelessModule_SpriteRotationRate.h b/Engine/Plugins/FX/Niagara/Source/Niagara/Internal/Stateless/Modules/NiagaraStatelessModule_SpriteRotationRate.h
--- a/Engine/Plugins/FX/Niagara/Source/Niagara/Internal/Stateless/Modules/NiagaraStatelessModule_SpriteRotationRate.h
+++ b/Engine/Plugins/FX/Niagara/Source/Niagara/Internal/Stateless/Modules/NiagaraStatelessModule_SpriteRotationRate.h
@@ -16,19 +16,17 @@ class UNiagaraStatelessModule_SpriteRotationRate : public UNiagaraStatelessModul
 public:
 	using FParameters = NiagaraStateless::FSpriteRotationRateModule_ShaderParameters;
 
-	UPROPERTY(EditAnywhere, Category = "Parameters")
-	float	RotationRateMin = 360.0f;
-
-	UPROPERTY(EditAnywhere, Category = "Parameters")
-	float	RotationRateMax = 360.0f;
+	UPROPERTY(EditAnywhere, Category = "Parameters", meta = (DisplayName = "Rotation Rate"))
+	FNiagaraDistributionRangeFloat RotationRateDistribution = FNiagaraDistributionRangeFloat(FNiagaraStatelessGlobals::GetDefaultSpriteRotationValue());
 
 	virtual void SetShaderParameters(const FNiagaraStatelessSetShaderParameterContext& SetShaderParameterContext) const override
 	{
 		FParameters* Parameters = SetShaderParameterContext.GetParameterNestedStruct<FParameters>();
 		if (IsModuleEnabled())
 		{
-			Parameters->SpriteRotationRate_Scale	= RotationRateMax - RotationRateMin;
-			Parameters->SpriteRotationRate_Bias		= RotationRateMin;
+			const FNiagaraStatelessRangeFloat RotationRateRange = RotationRateDistribution.CalculateRange(FNiagaraStatelessGlobals::GetDefaultSpriteRotationValue());
+			Parameters->SpriteRotationRate_Scale	= RotationRateRange.GetScale();
+			Parameters->SpriteRotationRate_Bias		= RotationRateRange.Min;
 		}
 		else
 		{
diff --git a/Engine/Plugins/FX/Niagara/Source/Niagara/Internal/Stateless/NiagaraStatelessCommon.h b/Engine/Plugins/FX/Niagara/Source/Niagara/Internal/Stateless/NiagaraStatelessCommon.h
--- a/Engine/Plugins/FX/Niagara/Source/Niagara/Internal/Stateless/NiagaraStatelessCommon.h
+++ b/Engine/Plugins/FX/Niagara/Source/Niagara/Internal/Stateless/NiagaraStatelessCommon.h
@@ -4,61 +4,26 @@
 
 #include "NiagaraCommon.h"
 
-#include "NiagaraStatelessCommon.generated.h"
-
-UENUM()
-enum class ENiagaraStatelessSpawnInfoType
-{
-	Burst,
-	Rate
-};
-
-USTRUCT()
-struct FNiagaraStatelessSpawnInfo
+template<typename TType>
+struct FNiagaraStatelessRange
 {
-	GENERATED_BODY()
-
-#if WITH_EDITORONLY_DATA
-	UPROPERTY()
-	FGuid SourceId;
-#endif
-
-	UPROPERTY(EditAnywhere, Category = "Spawn", meta = (ShowInStackItemHeader, StackItemHeaderAlignment = "Left"))
-	ENiagaraStatelessSpawnInfoType	Type = ENiagaraStatelessSpawnInfoType::Burst;
-
-	UPROPERTY(EditAnywhere, Category = "Spawn", meta = (EditConditionHides, EditCondition = "Type == ENiagaraStatelessSpawnInfoType::Burst"))
-	float SpawnTime = 0.0f;
+	using ValueType = TType;
 
-	UPROPERTY(EditAnywhere, Category = "Spawn", meta = (EditConditionHides, EditCondition = "Type == ENiagaraStatelessSpawnInfoType::Burst", ClampMin = "0"))
-	int32 AmountMin = 1;
+	FNiagaraStatelessRange() = default;
+	explicit FNiagaraStatelessRange(const ValueType& InMinMax) : Min(InMinMax), Max(InMinMax) {}
+	explicit FNiagaraStatelessRange(const ValueType& InMin, const ValueType& InMax) : Min(InMin), Max(InMax) {}
 
-	UPROPERTY(EditAnywhere, Category = "Spawn", meta = (EditConditionHides, EditCondition = "Type == ENiagaraStatelessSpawnInfoType::Burst", ClampMin = "0"))
-	int32 AmountMax = 1;
+	ValueType GetScale() const { return Max - Min; }
 
-	UPROPERTY(EditAnywhere, Category = "Spawn", meta = (EditConditionHides, EditCondition = "Type == ENiagaraStatelessSpawnInfoType::Rate", ClampMin = "0.0"))
-	float RateMin = 60.0f;
-
-	UPROPERTY(EditAnywhere, Category = "Spawn", meta = (EditConditionHides, EditCondition = "Type == ENiagaraStatelessSpawnInfoType::Rate", ClampMin = "0.0"))
-	float RateMax = 60.0f;
-
-	UPROPERTY(EditAnywhere, Category = "Spawn", meta = (InlineEditConditionToggle))
-	bool bSpawnProbabilityEnabled = false;
-
-	UPROPERTY(EditAnywhere, Category = "Spawn", meta = (EditCondition = "bSpawnProbabilityEnabled", ClampMin = "0.0", ClampMax = "1.0", UIMin = "0.0", UIMax = "1.0"))
-	float SpawnProbability = 1.0f;
-
-	bool IsValid(TOptional<float> LoopDuration) const;
+	ValueType Min = {};
+	ValueType Max = {};
 };
 
-struct FNiagaraStatelessRuntimeSpawnInfo
-{
-	ENiagaraStatelessSpawnInfoType Type = ENiagaraStatelessSpawnInfoType::Burst;
-	uint32	UniqueOffset	= 0;
-	float	SpawnTimeStart	= 0.0f;
-	float	SpawnTimeEnd	= 0.0f;
-	float	Rate			= 0.0f;
-	int32	Amount			= 0;
-};
+using FNiagaraStatelessRangeFloat	= FNiagaraStatelessRange<float>;
+using FNiagaraStatelessRangeVector2	= FNiagaraStatelessRange<FVector2f>;
+using FNiagaraStatelessRangeVector3	= FNiagaraStatelessRange<FVector3f>;
+using FNiagaraStatelessRangeVector4	= FNiagaraStatelessRange<FVector4f>;
+using FNiagaraStatelessRangeColor	= FNiagaraStatelessRange<FLinearColor>;
 
 struct FNiagaraStatelessGlobals
 {
@@ -86,6 +51,15 @@ struct FNiagaraStatelessGlobals
 	FNiagaraVariableBase	PreviousSpriteRotationVariable;
 	FNiagaraVariableBase	PreviousVelocityVariable;
 
+	inline static FLinearColor	GetDefaultColorValue() { return FLinearColor::White; }
+	inline static float			GetDefaultLifetimeValue() { return 1.0f; }
+	inline static float			GetDefaultMassValue() { return 1.0f; }
+	inline static FQuat4f		GetDefaultMeshOrientationValue() { return FQuat4f::Identity; }
+	inline static float			GetDefaultRibbonWidthValue() { return 10.0f; }
+	inline static FVector3f		GetDefaultScaleValue() { return FVector3f::OneVector; }
+	inline static FVector2f		GetDefaultSpriteSizeValue() { return FVector2f(10.0f); }
+	inline static float			GetDefaultSpriteRotationValue() { return 0.0f; }
+
 	static const FNiagaraStatelessGlobals& Get();
 };
 
diff --git a/Engine/Plugins/FX/Niagara/Source/Niagara/Internal/Stateless/NiagaraStatelessDistribution.h b/Engine/Plugins/FX/Niagara/Source/Niagara/Internal/Stateless/NiagaraStatelessDistribution.h
new file mode 100644
--- /dev/null
+++ b/Engine/Plugins/FX/Niagara/Source/Niagara/Internal/Stateless/NiagaraStatelessDistribution.h
@@ -0,0 +1,195 @@
+// Copyright Epic Games, Inc. All Rights Reserved.
+
+#pragma once
+
+#include "Curves/RichCurve.h"
+#include "NiagaraStatelessCommon.h"
+
+#include "NiagaraStatelessDistribution.generated.h"
+
+UENUM()
+enum class ENiagaraDistributionMode
+{
+	UniformConstant,
+	NonUniformConstant,
+	UniformRange,
+	NonUniformRange,
+	UniformCurve,
+	NonUniformCurve
+};
+
+USTRUCT()
+struct FNiagaraDistributionBase
+{
+	GENERATED_BODY()
+
+	virtual ~FNiagaraDistributionBase() 
+	{
+	}
+
+	UPROPERTY(EditAnywhere, Category = "Parameters")
+	ENiagaraDistributionMode Mode = ENiagaraDistributionMode::UniformConstant;
+
+	bool IsConstant() const { return Mode == ENiagaraDistributionMode::UniformConstant || Mode == ENiagaraDistributionMode::NonUniformConstant; }
+	bool IsUniform() const { return Mode == ENiagaraDistributionMode::UniformConstant || Mode == ENiagaraDistributionMode::UniformRange; }
+	bool IsCurve() const { return Mode == ENiagaraDistributionMode::UniformCurve || Mode == ENiagaraDistributionMode::NonUniformCurve; }
+
+#if WITH_EDITORONLY_DATA
+	UPROPERTY(EditAnywhere, Category = "Parameters")
+	TArray<float> ChannelConstantsAndRanges;
+
+	UPROPERTY(EditAnywhere, Category = "Parameters")
+	TArray<FRichCurve> ChannelCurves;
+
+	UPROPERTY(EditAnywhere, Category = "Parameters")
+	int32 MaxLutSampleCount = 64;
+
+	virtual bool AllowCurves() const { return true; }
+	virtual void UpdateValuesFromDistribution() { }
+#endif
+};
+
+USTRUCT()
+struct FNiagaraDistributionRangeFloat : public FNiagaraDistributionBase
+{
+	GENERATED_BODY()
+
+	FNiagaraDistributionRangeFloat() = default;
+	explicit FNiagaraDistributionRangeFloat(float ConstantValue) { InitConstant(ConstantValue); }
+
+	UPROPERTY(EditAnywhere, Category = "Parameters")
+	float Min = 0.0f;
+
+	UPROPERTY(EditAnywhere, Category = "Parameters")
+	float Max = 0.0f;
+
+	NIAGARA_API void InitConstant(float Value);
+	NIAGARA_API FNiagaraStatelessRangeFloat CalculateRange(const float Default = 0.0f) const;
+
+#if WITH_EDITORONLY_DATA
+	virtual bool AllowCurves() const override { return false; }
+	NIAGARA_API virtual void UpdateValuesFromDistribution() override;
+#endif
+};
+
+USTRUCT()
+struct FNiagaraDistributionRangeVector2 : public FNiagaraDistributionBase
+{
+	GENERATED_BODY()
+
+	FNiagaraDistributionRangeVector2() = default;
+	explicit FNiagaraDistributionRangeVector2(const FVector2f& ConstantValue) { InitConstant(ConstantValue); }
+
+	UPROPERTY(EditAnywhere, Category = "Parameters")
+	FVector2f Min = FVector2f::ZeroVector;
+
+	UPROPERTY(EditAnywhere, Category = "Parameters")
+	FVector2f Max = FVector2f::ZeroVector;
+
+	NIAGARA_API void InitConstant(const FVector2f& Value);
+	NIAGARA_API FNiagaraStatelessRangeVector2 CalculateRange(const FVector2f& Default = FVector2f::ZeroVector) const;
+
+#if WITH_EDITORONLY_DATA
+	virtual bool AllowCurves() const override { return false; }
+	virtual void UpdateValuesFromDistribution() override;
+#endif
+};
+
+USTRUCT()
+struct FNiagaraDistributionRangeVector3 : public FNiagaraDistributionBase
+{
+	GENERATED_BODY()
+
+	FNiagaraDistributionRangeVector3() = default;
+	explicit FNiagaraDistributionRangeVector3(const FVector3f& ConstantValue) { InitConstant(ConstantValue); }
+
+	UPROPERTY(EditAnywhere, Category = "Parameters")
+	FVector3f Min = FVector3f::ZeroVector;
+
+	UPROPERTY(EditAnywhere, Category = "Parameters")
+	FVector3f Max = FVector3f::ZeroVector;
+
+	NIAGARA_API void InitConstant(const FVector3f& Value);
+	NIAGARA_API FNiagaraStatelessRangeVector3 CalculateRange(const FVector3f& Default = FVector3f::ZeroVector) const;
+
+#if WITH_EDITORONLY_DATA
+	virtual bool AllowCurves() const override { return false; }
+	virtual void UpdateValuesFromDistribution() override;
+#endif
+};
+
+USTRUCT()
+struct FNiagaraDistributionFloat : public FNiagaraDistributionBase
+{
+	GENERATED_BODY()
+
+	FNiagaraDistributionFloat() = default;
+	explicit FNiagaraDistributionFloat(float ConstantValue) { InitConstant(ConstantValue); }
+
+	UPROPERTY(EditAnywhere, Category = "Parameters")
+	TArray<float> Values;
+
+	NIAGARA_API void InitConstant(float Value);
+	NIAGARA_API FNiagaraStatelessRangeFloat CalculateRange(const float Default = 0.0f) const;
+
+#if WITH_EDITORONLY_DATA
+	virtual void UpdateValuesFromDistribution() override;
+#endif
+};
+
+USTRUCT()
+struct FNiagaraDistributionVector2 : public FNiagaraDistributionBase
+{
+	GENERATED_BODY()
+
+	FNiagaraDistributionVector2() = default;
+	explicit FNiagaraDistributionVector2(const FVector2f& ConstantValue) { InitConstant(ConstantValue); }
+
+	UPROPERTY(EditAnywhere, Category = "Parameters")
+	TArray<FVector2f> Values;
+
+	NIAGARA_API void InitConstant(const FVector2f& Value);
+	NIAGARA_API FNiagaraStatelessRangeVector2 CalculateRange(const FVector2f& Default = FVector2f::ZeroVector) const;
+
+#if WITH_EDITORONLY_DATA
+	virtual void UpdateValuesFromDistribution() override;
+#endif
+};
+
+USTRUCT()
+struct FNiagaraDistributionVector3 : public FNiagaraDistributionBase
+{
+	GENERATED_BODY()
+
+	FNiagaraDistributionVector3() = default;
+	explicit FNiagaraDistributionVector3(const FVector3f& ConstantValue) { InitConstant(ConstantValue); }
+
+	UPROPERTY(EditAnywhere, Category = "Parameters")
+	TArray<FVector3f> Values;
+
+	NIAGARA_API void InitConstant(const FVector3f& Value);
+	NIAGARA_API FNiagaraStatelessRangeVector3 CalculateRange(const FVector3f& Default = FVector3f::ZeroVector) const;
+
+#if WITH_EDITORONLY_DATA
+	virtual void UpdateValuesFromDistribution() override;
+#endif
+};
+
+USTRUCT()
+struct FNiagaraDistributionColor : public FNiagaraDistributionBase
+{
+	GENERATED_BODY()
+
+	FNiagaraDistributionColor() = default;
+	explicit FNiagaraDistributionColor(const FLinearColor& ConstantValue) { InitConstant(ConstantValue); }
+
+	UPROPERTY(EditAnywhere, Category = "Parameters")
+	TArray<FLinearColor> Values;
+
+	NIAGARA_API void InitConstant(const FLinearColor& Value);
+	NIAGARA_API FNiagaraStatelessRangeColor CalculateRange(const FLinearColor& Default = FLinearColor::White) const;
+
+#if WITH_EDITORONLY_DATA
+	virtual void UpdateValuesFromDistribution() override;
+#endif
+};
diff --git a/Engine/Plugins/FX/Niagara/Source/Niagara/Internal/Stateless/NiagaraStatelessEmitter.h b/Engine/Plugins/FX/Niagara/Source/Niagara/Internal/Stateless/NiagaraStatelessEmitter.h
--- a/Engine/Plugins/FX/Niagara/Source/Niagara/Internal/Stateless/NiagaraStatelessEmitter.h
+++ b/Engine/Plugins/FX/Niagara/Source/Niagara/Internal/Stateless/NiagaraStatelessEmitter.h
@@ -3,6 +3,7 @@
 #pragma once
 
 #include "NiagaraStatelessCommon.h"
+#include "NiagaraStatelessSpawnInfo.h"
 #include "NiagaraDataSet.h"
 #include "NiagaraRendererProperties.h"
 #include "NiagaraSystemEmitterState.h"
diff --git a/Engine/Plugins/FX/Niagara/Source/Niagara/Internal/Stateless/NiagaraStatelessEmitterData.h b/Engine/Plugins/FX/Niagara/Source/Niagara/Internal/Stateless/NiagaraStatelessEmitterData.h
--- a/Engine/Plugins/FX/Niagara/Source/Niagara/Internal/Stateless/NiagaraStatelessEmitterData.h
+++ b/Engine/Plugins/FX/Niagara/Source/Niagara/Internal/Stateless/NiagaraStatelessEmitterData.h
@@ -3,6 +3,7 @@
 #pragma once
 
 #include "NiagaraStatelessCommon.h"
+#include "NiagaraStatelessSpawnInfo.h"
 #include "NiagaraDataSet.h"
 #include "NiagaraParameterStore.h"
 #include "NiagaraSystemEmitterState.h"
@@ -35,7 +36,7 @@ struct FNiagaraStatelessEmitterData
 	bool											bCanEverExecute = false;
 	bool											bDeterministic = false;
 	int32											RandomSeed = 0;
-	FVector2f										LifetimeRange = FVector2f::ZeroVector;
+	FNiagaraStatelessRangeFloat						LifetimeRange = FNiagaraStatelessRangeFloat(0.0f, 0.0f);
 	FBox											FixedBounds = FBox(ForceInit);
 
 	FNiagaraEmitterStateData						EmitterState;
diff --git a/Engine/Plugins/FX/Niagara/Source/Niagara/Internal/Stateless/NiagaraStatelessEmitterDataBuildContext.h b/Engine/Plugins/FX/Niagara/Source/Niagara/Internal/Stateless/NiagaraStatelessEmitterDataBuildContext.h
--- a/Engine/Plugins/FX/Niagara/Source/Niagara/Internal/Stateless/NiagaraStatelessEmitterDataBuildContext.h
+++ b/Engine/Plugins/FX/Niagara/Source/Niagara/Internal/Stateless/NiagaraStatelessEmitterDataBuildContext.h
@@ -2,7 +2,8 @@
 
 #pragma once
 
-#include "NiagaraCommon.h"
+#include "NiagaraStatelessCommon.h"
+#include "NiagaraStatelessDistribution.h"
 
 struct FNiagaraParameterBinding;
 struct FNiagaraParameterBindingWithValue;
@@ -44,11 +45,37 @@ public:
 		}
 		return *reinterpret_cast<T*>(TransientObj->GetObject());
 	}
-	//NiagaraStateless::FPhysicsBuildData& PhysicsBuildData = BuildContext.GetTransientBuildData<NiagaraStateless::FPhysicsBuildData>();
 
 	int32 AddRendererBinding(const FNiagaraParameterBinding& Binding);
 	int32 AddRendererBinding(const FNiagaraParameterBindingWithValue& Binding);	
 
+	// Adds a distribution into the LUT if enabled and returns the packed information to send to the shader
+	template<typename TDistribution>
+	FUintVector3 AddDistribution(const TDistribution& Distribution, bool bEnabled)
+	{
+		FUintVector3 Parameters = FUintVector3::ZeroValue;
+		if (bEnabled && Distribution.Values.Num() > 0)
+		{
+			constexpr uint32 StatelessDistributionFlag_Random = 0x00000001;
+			constexpr uint32 StatelessDistributionFlag_Uniform = 0x00000002;
+
+			switch (Distribution.Mode)
+			{
+				case ENiagaraDistributionMode::UniformConstant:		Parameters.X = StatelessDistributionFlag_Random | StatelessDistributionFlag_Uniform; break;
+				case ENiagaraDistributionMode::NonUniformConstant:	Parameters.X = StatelessDistributionFlag_Random; break;
+				case ENiagaraDistributionMode::UniformRange:		Parameters.X = StatelessDistributionFlag_Random | StatelessDistributionFlag_Uniform; break;
+				case ENiagaraDistributionMode::NonUniformRange:		Parameters.X = StatelessDistributionFlag_Random; break;
+				case ENiagaraDistributionMode::UniformCurve:		Parameters.X = StatelessDistributionFlag_Uniform; break;
+				case ENiagaraDistributionMode::NonUniformCurve:		Parameters.X = 0; break;
+				default:											checkNoEntry(); break;
+			}
+
+			Parameters.Y = AddStaticData(Distribution.Values);
+			reinterpret_cast<float&>(Parameters.Z) = Distribution.Values.Num() - 1;
+		}
+		return Parameters;
+	}
+
 private:
 	FNiagaraParameterStore& RendererBindings;
 	TArray<uint8>&			BuiltData;
diff --git a/Engine/Plugins/FX/Niagara/Source/Niagara/Internal/Stateless/NiagaraStatelessModule.h b/Engine/Plugins/FX/Niagara/Source/Niagara/Internal/Stateless/NiagaraStatelessModule.h
--- a/Engine/Plugins/FX/Niagara/Source/Niagara/Internal/Stateless/NiagaraStatelessModule.h
+++ b/Engine/Plugins/FX/Niagara/Source/Niagara/Internal/Stateless/NiagaraStatelessModule.h
@@ -3,6 +3,7 @@
 #pragma once
 
 #include "NiagaraStatelessCommon.h"
+#include "NiagaraStatelessDistribution.h"
 
 #include "ShaderParameterStruct.h"
 
@@ -90,7 +91,7 @@ protected:
 	uint32 bModuleEnabled : 1 = true;
 
 #if WITH_EDITORONLY_DATA
-	UPROPERTY(EditAnywhere, Transient, Category = "Parameters", meta = (DisplayPriority = 0, ShowInStackItemHeader, StackItemHeaderIcon = "Icons.Visibility"))
+	UPROPERTY(EditAnywhere, Transient, Category = "Parameters", meta = (DisplayPriority = 0, StackItemHeaderIcon = "Icons.Visibility"))
 	uint32 bDebugDrawEnabled : 1 = false;
 #endif
 
@@ -98,6 +99,13 @@ public:
 	bool IsModuleEnabled() const { return bModuleEnabled; }
 #if WITH_EDITOR
 	bool IsDebugDrawEnabled() const { return bDebugDrawEnabled; }
+
+	struct PrivateMemberNames
+	{
+		static NIAGARA_API const FName bModuleEnabled;
+		static NIAGARA_API const FName bDebugDrawEnabled;
+	};
+
 #endif
 
 	virtual void BuildEmitterData(FNiagaraStatelessEmitterDataBuildContext& BuildContext) const {}
@@ -118,6 +126,7 @@ public:
 	//~UObject interface Begin
 #if WITH_EDITOR
 	virtual bool CanEditChange(const FProperty* InProperty) const override;
+	virtual void PostEditChangeProperty(struct FPropertyChangedEvent& PropertyChangedEvent) override;
 #endif //WITH_EDITOR
 	//~UObject interface End
 };
diff --git a/Engine/Plugins/FX/Niagara/Source/Niagara/Internal/Stateless/NiagaraStatelessSpawnInfo.h b/Engine/Plugins/FX/Niagara/Source/Niagara/Internal/Stateless/NiagaraStatelessSpawnInfo.h
new file mode 100644
--- /dev/null
+++ b/Engine/Plugins/FX/Niagara/Source/Niagara/Internal/Stateless/NiagaraStatelessSpawnInfo.h
@@ -0,0 +1,62 @@
+// Copyright Epic Games, Inc. All Rights Reserved.
+
+#pragma once
+
+#include "NiagaraStatelessCommon.h"
+#include "NiagaraStatelessDistribution.h"
+
+#include "NiagaraStatelessSpawnInfo.generated.h"
+
+UENUM()
+enum class ENiagaraStatelessSpawnInfoType
+{
+	Burst,
+	Rate
+};
+
+USTRUCT()
+struct FNiagaraStatelessSpawnInfo
+{
+	GENERATED_BODY()
+
+#if WITH_EDITORONLY_DATA
+	UPROPERTY()
+	FGuid SourceId;
+#endif
+
+	UPROPERTY(EditAnywhere, Category = "Spawn", meta = (ShowInStackItemHeader, StackItemHeaderAlignment = "Left"))
+	ENiagaraStatelessSpawnInfoType	Type = ENiagaraStatelessSpawnInfoType::Burst;
+
+	UPROPERTY(EditAnywhere, Category = "Spawn", meta = (EditConditionHides, EditCondition = "Type == ENiagaraStatelessSpawnInfoType::Burst"))
+	float SpawnTime = 0.0f;
+
+	UPROPERTY(EditAnywhere, Category = "Spawn", meta = (EditConditionHides, EditCondition = "Type == ENiagaraStatelessSpawnInfoType::Burst", ClampMin = "0"))
+	int32 AmountMin = 1;
+
+	UPROPERTY(EditAnywhere, Category = "Spawn", meta = (EditConditionHides, EditCondition = "Type == ENiagaraStatelessSpawnInfoType::Burst", ClampMin = "0"))
+	int32 AmountMax = 1;
+
+	UPROPERTY(EditAnywhere, Category = "Spawn", meta = (EditConditionHides, EditCondition = "Type == ENiagaraStatelessSpawnInfoType::Rate", ClampMin = "0.0"))
+	float RateMin = 60.0f;
+
+	UPROPERTY(EditAnywhere, Category = "Spawn", meta = (EditConditionHides, EditCondition = "Type == ENiagaraStatelessSpawnInfoType::Rate", ClampMin = "0.0"))
+	float RateMax = 60.0f;
+
+	UPROPERTY(EditAnywhere, Category = "Spawn", meta = (InlineEditConditionToggle))
+	bool bSpawnProbabilityEnabled = false;
+
+	UPROPERTY(EditAnywhere, Category = "Spawn", meta = (EditCondition = "bSpawnProbabilityEnabled", ClampMin = "0.0", ClampMax = "1.0", UIMin = "0.0", UIMax = "1.0"))
+	float SpawnProbability = 1.0f;
+
+	bool IsValid(TOptional<float> LoopDuration) const;
+};
+
+struct FNiagaraStatelessRuntimeSpawnInfo
+{
+	ENiagaraStatelessSpawnInfoType Type = ENiagaraStatelessSpawnInfoType::Burst;
+	uint32	UniqueOffset	= 0;
+	float	SpawnTimeStart	= 0.0f;
+	float	SpawnTimeEnd	= 0.0f;
+	float	Rate			= 0.0f;
+	int32	Amount			= 0;
+};
diff --git a/Engine/Plugins/FX/Niagara/Source/Niagara/Private/NiagaraRenderer.cpp b/Engine/Plugins/FX/Niagara/Source/Niagara/Private/NiagaraRenderer.cpp
--- a/Engine/Plugins/FX/Niagara/Source/Niagara/Private/NiagaraRenderer.cpp
+++ b/Engine/Plugins/FX/Niagara/Source/Niagara/Private/NiagaraRenderer.cpp
@@ -15,6 +15,9 @@
 #include "NiagaraSystemInstanceController.h"
 #include "Materials/MaterialInstanceDynamic.h"
 #include "PrimitiveViewRelevance.h"
+//-TODO:Stateless: Consider a TFunction or something else to get the data
+#include "Stateless/NiagaraStatelessEmitterInstance.h"
+//-TODO:Stateless: Consider a TFunction or something else to get the data
 
 DECLARE_CYCLE_STAT(TEXT("Sort Particles"), STAT_NiagaraSortParticles, STATGROUP_Niagara);
 DECLARE_CYCLE_STAT(TEXT("Global Float Alloc - All"), STAT_NiagaraAllocateGlobalFloatAll, STATGROUP_Niagara);
diff --git a/Engine/Plugins/FX/Niagara/Source/Niagara/Private/NiagaraRendererProperties.cpp b/Engine/Plugins/FX/Niagara/Source/Niagara/Private/NiagaraRendererProperties.cpp
--- a/Engine/Plugins/FX/Niagara/Source/Niagara/Private/NiagaraRendererProperties.cpp
+++ b/Engine/Plugins/FX/Niagara/Source/Niagara/Private/NiagaraRendererProperties.cpp
@@ -878,6 +878,7 @@ bool UNiagaraRendererProperties::NeedsLoadForTargetPlatform(const ITargetPlatfor
 			}
 		}
 	}
+	//-TODO:Stateless: We need a base emitter type
 
 	return false;
 }
diff --git a/Engine/Plugins/FX/Niagara/Source/Niagara/Private/NiagaraSystem.cpp b/Engine/Plugins/FX/Niagara/Source/Niagara/Private/NiagaraSystem.cpp
--- a/Engine/Plugins/FX/Niagara/Source/Niagara/Private/NiagaraSystem.cpp
+++ b/Engine/Plugins/FX/Niagara/Source/Niagara/Private/NiagaraSystem.cpp
@@ -3513,9 +3513,9 @@ void UNiagaraSystem::InitSystemCompiledData()
 	{
 		const FNiagaraEmitterHandle& PerEmitterHandle = EmitterHandles[EmitterIdx];
 		const UNiagaraEmitter* Emitter = PerEmitterHandle.GetInstance().Emitter;
-		//-TODO:Stateless: Do we need stateless support here?
 		//if (ensureMsgf(Emitter != nullptr, TEXT("Failed to get Emitter Instance from Emitter Handle when post compiling Niagara System %s!"), *GetPathNameSafe(this)))
 		if (Emitter != nullptr)
+		//-TODO:Stateless:
 		{
 			const FString EmitterName = Emitter->GetUniqueEmitterName();
 
diff --git a/Engine/Plugins/FX/Niagara/Source/Niagara/Private/Stateless/Modules/NiagaraStatelessModule_AddVelocity.cpp b/Engine/Plugins/FX/Niagara/Source/Niagara/Private/Stateless/Modules/NiagaraStatelessModule_AddVelocity.cpp
--- a/Engine/Plugins/FX/Niagara/Source/Niagara/Private/Stateless/Modules/NiagaraStatelessModule_AddVelocity.cpp
+++ b/Engine/Plugins/FX/Niagara/Source/Niagara/Private/Stateless/Modules/NiagaraStatelessModule_AddVelocity.cpp
@@ -13,16 +13,16 @@ void UNiagaraStatelessModule_AddVelocity::BuildEmitterData(FNiagaraStatelessEmit
 	NiagaraStateless::FPhysicsBuildData& PhysicsBuildData = BuildContext.GetTransientBuildData<NiagaraStateless::FPhysicsBuildData>();
 	if (VelocityType == ENSM_VelocityType::Linear)
 	{
-		PhysicsBuildData.VelocityMin += VelocityMin * VelocityScale;
-		PhysicsBuildData.VelocityMax += VelocityMax * VelocityScale;
+		const FNiagaraStatelessRangeVector3 VelocityRange = LinearVelocityDistribution.CalculateRange(FVector3f::ZeroVector);
+		PhysicsBuildData.VelocityRange.Min += VelocityRange.Min * LinearVelocityScale;
+		PhysicsBuildData.VelocityRange.Max += VelocityRange.Max * LinearVelocityScale;
 	}
 	else if (VelocityType == ENSM_VelocityType::FromPoint)
 	{
 		ensureMsgf(PhysicsBuildData.bPointVelocity == false, TEXT("Only a single point force is supported at the moment."));
 
 		PhysicsBuildData.bPointVelocity = true;
-		PhysicsBuildData.PointVelocityMin = PointVelocityMin;
-		PhysicsBuildData.PointVelocityMax = PointVelocityMax;
+		PhysicsBuildData.PointVelocityRange = PointVelocityDistribution.CalculateRange(0.0f);
 		PhysicsBuildData.PointOrigin = PointOrigin;
 	}
 	else if (VelocityType == ENSM_VelocityType::InCone)
@@ -31,8 +31,7 @@ void UNiagaraStatelessModule_AddVelocity::BuildEmitterData(FNiagaraStatelessEmit
 
 		PhysicsBuildData.bConeVelocity = true;
 		PhysicsBuildData.ConeQuat = FQuat4f(ConeRotation.Quaternion());
-		PhysicsBuildData.ConeVelocityMin = ConeVelocityMin;
-		PhysicsBuildData.ConeVelocityMax = ConeVelocityMax;
+		PhysicsBuildData.ConeVelocityRange = ConeVelocityDistribution.CalculateRange(0.0f);
 		PhysicsBuildData.ConeOuterAngle = ConeAngle;
 		PhysicsBuildData.ConeInnerAngle = InnerCone;
 		PhysicsBuildData.ConeVelocityFalloff = bSpeedFalloffFromConeAxisEnabled ? FMath::Clamp(SpeedFalloffFromConeAxis, 0.0f, 1.0f) : 0.0f;
@@ -46,8 +45,9 @@ void UNiagaraStatelessModule_AddVelocity::DrawDebug(const FNiagaraStatelessDrawD
 	{
 		case ENSM_VelocityType::Linear:
 		{
-			const FVector MinDir = FVector(VelocityMin * VelocityScale);
-			const FVector MaxDir = FVector(VelocityMax * VelocityScale);
+			const FNiagaraStatelessRangeVector3 VelocityRange = LinearVelocityDistribution.CalculateRange(FVector3f::ZeroVector);
+			const FVector MinDir = FVector(VelocityRange.Min * LinearVelocityScale);
+			const FVector MaxDir = FVector(VelocityRange.Max * LinearVelocityScale);
 			DrawDebugContext.DrawArrow(FVector::ZeroVector, MinDir);
 
 			if (!FMath::IsNearlyEqual(MinDir.Length(), MaxDir.Length()))
@@ -68,18 +68,19 @@ void UNiagaraStatelessModule_AddVelocity::DrawDebug(const FNiagaraStatelessDrawD
 				InnerConeHAngle = InnerCone / 2.0f;
 			}
 
-			DrawDebugContext.DrawCone(FVector::ZeroVector, Quat, ConeHAngle, ConeVelocityMin);
+			const FNiagaraStatelessRangeFloat ConeVelocityRange = ConeVelocityDistribution.CalculateRange(0.0f);
+			DrawDebugContext.DrawCone(FVector::ZeroVector, Quat, ConeHAngle, ConeVelocityRange.Min);
 			if ( InnerConeHAngle.IsSet() )
 			{
-				DrawDebugContext.DrawCone(FVector::ZeroVector, Quat, InnerConeHAngle.GetValue(), ConeVelocityMin);
+				DrawDebugContext.DrawCone(FVector::ZeroVector, Quat, InnerConeHAngle.GetValue(), ConeVelocityRange.Min);
 			}
 
-			if (!FMath::IsNearlyEqual(ConeVelocityMin, ConeVelocityMax))
+			if (!FMath::IsNearlyEqual(ConeVelocityRange.Min, ConeVelocityRange.Max))
 			{
-				DrawDebugContext.DrawCone(FVector::ZeroVector, Quat, ConeHAngle, ConeVelocityMax);
+				DrawDebugContext.DrawCone(FVector::ZeroVector, Quat, ConeHAngle, ConeVelocityRange.Max);
 				if (InnerConeHAngle.IsSet())
 				{
-					DrawDebugContext.DrawCone(FVector::ZeroVector, Quat, InnerConeHAngle.GetValue(), ConeVelocityMax);
+					DrawDebugContext.DrawCone(FVector::ZeroVector, Quat, InnerConeHAngle.GetValue(), ConeVelocityRange.Max);
 				}
 			}
 			break;
@@ -87,13 +88,14 @@ void UNiagaraStatelessModule_AddVelocity::DrawDebug(const FNiagaraStatelessDrawD
 
 		case ENSM_VelocityType::FromPoint:
 		{
-			if (!FMath::IsNearlyEqual(PointVelocityMin, 0.0f))
+			const FNiagaraStatelessRangeFloat PointVelocityRange = PointVelocityDistribution.CalculateRange(0.0f);
+			if (!FMath::IsNearlyEqual(PointVelocityRange.Min, 0.0f))
 			{
-				DrawDebugContext.DrawSphere(FVector(PointOrigin), PointVelocityMin);
+				DrawDebugContext.DrawSphere(FVector(PointOrigin), PointVelocityRange.Min);
 			}
-			if (!FMath::IsNearlyEqual(PointVelocityMin, PointVelocityMax))
+			if (!FMath::IsNearlyEqual(PointVelocityRange.Min, PointVelocityRange.Max))
 			{
-				DrawDebugContext.DrawSphere(FVector(PointOrigin), PointVelocityMax);
+				DrawDebugContext.DrawSphere(FVector(PointOrigin), PointVelocityRange.Max);
 			}
 			break;
 		}
diff --git a/Engine/Plugins/FX/Niagara/Source/Niagara/Private/Stateless/NiagaraStatelessCommon.cpp b/Engine/Plugins/FX/Niagara/Source/Niagara/Private/Stateless/NiagaraStatelessCommon.cpp
--- a/Engine/Plugins/FX/Niagara/Source/Niagara/Private/Stateless/NiagaraStatelessCommon.cpp
+++ b/Engine/Plugins/FX/Niagara/Source/Niagara/Private/Stateless/NiagaraStatelessCommon.cpp
@@ -38,24 +38,7 @@ namespace NiagaraStatelessCommon
 	}
 } //NiagaraStatelessCommon
 
-bool FNiagaraStatelessSpawnInfo::IsValid(TOptional<float> LoopDuration) const
-{
-	switch (Type)
-	{
-		case ENiagaraStatelessSpawnInfoType::Burst:
-			return (AmountMin + AmountMax) > 0 && SpawnTime >= 0.0f && SpawnTime < LoopDuration.Get(SpawnTime + UE_SMALL_NUMBER);
-
-		case ENiagaraStatelessSpawnInfoType::Rate:
-			return (RateMin + RateMax) > 0.0f;
-
-		default:
-			checkNoEntry();
-			return false;
-	}
-}
-
 const FNiagaraStatelessGlobals& FNiagaraStatelessGlobals::Get()
 {
 	return NiagaraStatelessCommon::GGlobals;
 }
-
diff --git a/Engine/Plugins/FX/Niagara/Source/Niagara/Private/Stateless/NiagaraStatelessDistribution.cpp b/Engine/Plugins/FX/Niagara/Source/Niagara/Private/Stateless/NiagaraStatelessDistribution.cpp
new file mode 100644
--- /dev/null
+++ b/Engine/Plugins/FX/Niagara/Source/Niagara/Private/Stateless/NiagaraStatelessDistribution.cpp
@@ -0,0 +1,408 @@
+// Copyright Epic Games, Inc. All Rights Reserved.
+
+#include "Stateless/NiagaraStatelessDistribution.h"
+#include "HAL/IConsoleManager.h"
+
+namespace NiagaraStatelessDistributionPrivate
+{
+#if WITH_EDITORONLY_DATA
+	bool GOptimizeLUTs = true;
+	FAutoConsoleVariableRef CVarNiagaraStatelessDistributionOptimizeLUTs(
+		TEXT("fx.NiagaraStateless.Distribution.OptimizeLUTs"),
+		GOptimizeLUTs,
+		TEXT("When enabled we optimize the LUT generation."),
+		ECVF_Default
+	);
+
+	TArray<float> CurvesToLUT(TArrayView<const FRichCurve> Curves, int32 NumSamples)
+	{
+		const int32 NumChannels = Curves.Num();
+
+		TArray<float> LUT;
+		LUT.SetNumUninitialized(NumChannels * NumSamples);
+		for (int32 iSample=0; iSample < NumSamples; ++iSample)
+		{
+			const float Time = float(iSample) / float(NumSamples - 1);
+			for (int32 iChannel=0; iChannel < NumChannels; ++iChannel)
+			{
+				LUT[(iSample * NumChannels) + iChannel] = Curves[iChannel].Eval(Time);
+			}
+		}
+		return LUT;
+	}
+
+	bool AreLUTsAlmostEqual(TArrayView<float> Lhs, TArrayView<float> Rhs, int32 NumChannels, float ErrorThreshold = 0.01f)
+	{
+		const int32 LhsNumSamples = Lhs.Num() / NumChannels;
+		const int32 RhsNumSamples = Rhs.Num() / NumChannels;
+		const int32 MaxSamples = FMath::Max(LhsNumSamples, RhsNumSamples);
+		for ( int32 i=0; i < MaxSamples; ++i )
+		{
+			const float U = float(i) / float(MaxSamples - 1);
+
+			const float LhsT = U * float(LhsNumSamples - 1);
+			const float LhsU = FMath::Frac(LhsT);
+			const int32 LhsA = FMath::Min(FMath::FloorToInt(LhsT), LhsNumSamples - 1);
+			const int32 LhsB = FMath::Min(LhsA + 1, LhsNumSamples - 1);
+
+			const float RhsT = U * float(RhsNumSamples - 1);
+			const float RhsU = FMath::Frac(RhsT);
+			const int32 RhsA = FMath::Min(FMath::FloorToInt(RhsT), RhsNumSamples - 1);
+			const int32 RhsB = FMath::Min(RhsA + 1, RhsNumSamples - 1);
+
+			for (int iChannel=0; iChannel < NumChannels; ++iChannel)
+			{
+				const float LhsValue = FMath::Lerp(Lhs[LhsA * NumChannels + iChannel], Lhs[LhsB * NumChannels + iChannel], LhsU);
+				const float RhsValue = FMath::Lerp(Rhs[RhsA * NumChannels + iChannel], Rhs[RhsB * NumChannels + iChannel], RhsU);
+				const float Error = FMath::Abs(LhsValue - RhsValue);
+				if (Error > ErrorThreshold)
+				{
+					return false;
+				}
+			}
+		}
+		return true;
+	}
+
+	TArray<float> CurvesToOptimizedLUT(TArrayView<const FRichCurve> Curves, int32 MaxLutSampleCount)
+	{
+		TArray<float> LUT = CurvesToLUT(Curves, MaxLutSampleCount);
+		if (GOptimizeLUTs)
+		{
+			for (int32 iSamples = 2; iSamples < MaxLutSampleCount; ++iSamples)
+			{
+				TArray<float> NewLUT = CurvesToLUT(Curves, iSamples);
+				if (AreLUTsAlmostEqual(NewLUT, LUT, Curves.Num()))
+				{
+					return NewLUT;
+				}
+			}
+		}
+		return LUT;
+	}
+
+	template<typename FValueContainerSetNum, typename FValueAndChannelAccessor>
+	void UpdateDistributionValues(ENiagaraDistributionMode InMode, const TArray<float>& InChannelConstantsAndRanges, const TArray<FRichCurve>& InChannelCurves, int32 InChannelCount, FValueContainerSetNum InContainerNum, FValueAndChannelAccessor InValueAndChannelAccessor, int32 MaxLutSampleCount)
+	{
+		switch(InMode)
+		{
+			case ENiagaraDistributionMode::UniformConstant:
+				if (InChannelConstantsAndRanges.Num() >= 1)
+				{
+					// Note we set two values to simplify the GPU code
+					InContainerNum(2);
+					for (int32 ChannelIndex = 0; ChannelIndex < InChannelCount; ChannelIndex++)
+					{
+						InValueAndChannelAccessor(0, ChannelIndex) = InChannelConstantsAndRanges[0];
+						InValueAndChannelAccessor(1, ChannelIndex) = InChannelConstantsAndRanges[0];
+					}
+				}
+				break;
+			case ENiagaraDistributionMode::NonUniformConstant:
+				if (InChannelConstantsAndRanges.Num() >= InChannelCount)
+				{
+					// Note we set two values to simplify the GPU code
+					InContainerNum(2);
+					for (int32 ChannelIndex = 0; ChannelIndex < InChannelCount; ChannelIndex++)
+					{
+						InValueAndChannelAccessor(0, ChannelIndex) = InChannelConstantsAndRanges[ChannelIndex];
+						InValueAndChannelAccessor(1, ChannelIndex) = InChannelConstantsAndRanges[ChannelIndex];
+					}
+				}
+				break;
+			case ENiagaraDistributionMode::UniformRange:
+				if (InChannelConstantsAndRanges.Num() >= 2)
+				{
+					InContainerNum(2);
+					for (int32 ChannelIndex = 0; ChannelIndex < InChannelCount; ChannelIndex++)
+					{
+						InValueAndChannelAccessor(0, ChannelIndex) = InChannelConstantsAndRanges[0];
+						InValueAndChannelAccessor(1, ChannelIndex) = InChannelConstantsAndRanges[1];
+					}
+				}
+				break;
+			case ENiagaraDistributionMode::NonUniformRange:
+				if (InChannelConstantsAndRanges.Num() >= 2 * InChannelCount)
+				{
+					InContainerNum(2);
+					for (int32 ChannelIndex = 0; ChannelIndex < InChannelCount; ChannelIndex++)
+					{
+						InValueAndChannelAccessor(0, ChannelIndex) = InChannelConstantsAndRanges[ChannelIndex];
+						InValueAndChannelAccessor(1, ChannelIndex) = InChannelConstantsAndRanges[InChannelCount + ChannelIndex];
+					}
+				}
+				break;
+			case ENiagaraDistributionMode::UniformCurve:
+			case ENiagaraDistributionMode::NonUniformCurve:
+			{
+				const int32 ExpectedChannels = InMode == ENiagaraDistributionMode::NonUniformCurve ? InChannelCount : 1;
+				if (InChannelCurves.Num() >= ExpectedChannels)
+				{
+					MaxLutSampleCount = FMath::Max(MaxLutSampleCount, 2);
+					//const TArray<float> LUT = CurvesToLUT(MakeArrayView(InChannelCurves.GetData(), ExpectedChannels), MaxLutSampleCount);
+					const TArray<float> LUT = CurvesToOptimizedLUT(MakeArrayView(InChannelCurves.GetData(), ExpectedChannels), MaxLutSampleCount);
+					const int32 NumSamples = LUT.Num() / ExpectedChannels;
+
+					InContainerNum(NumSamples);
+					for (int32 SampleIndex = 0; SampleIndex < NumSamples; ++SampleIndex)
+					{
+						for (int32 ChannelIndex=0; ChannelIndex < InChannelCount; ++ChannelIndex)
+						{
+							const int32 LUTChannelIndex = InMode == ENiagaraDistributionMode::NonUniformCurve ? ChannelIndex : 0;
+							InValueAndChannelAccessor(SampleIndex, ChannelIndex) = LUT[(SampleIndex * ExpectedChannels) + LUTChannelIndex];
+						}
+					}
+				}
+				break;
+			}
+		}
+	}
+#endif
+} //NiagaraStatelessCommon
+
+void FNiagaraDistributionRangeFloat::InitConstant(float Value)
+{
+	Mode = ENiagaraDistributionMode::UniformConstant;
+	Min = Value;
+	Max = Value;
+#if WITH_EDITORONLY_DATA
+	ChannelConstantsAndRanges = TArray<float>({ Value });
+#endif
+}
+
+FNiagaraStatelessRangeFloat FNiagaraDistributionRangeFloat::CalculateRange(const float Default) const
+{
+	return Mode == ENiagaraDistributionMode::UniformConstant ? FNiagaraStatelessRangeFloat(Min, Min) : FNiagaraStatelessRangeFloat(Min, Max);
+}
+
+void FNiagaraDistributionRangeVector2::InitConstant(const FVector2f& Value)
+{
+	Mode = ENiagaraDistributionMode::NonUniformConstant;
+	Min = Value;
+	Max = Value;
+#if WITH_EDITORONLY_DATA
+	ChannelConstantsAndRanges = TArray<float>({ Value.X, Value.Y });
+#endif
+}
+
+FNiagaraStatelessRangeVector2 FNiagaraDistributionRangeVector2::CalculateRange(const FVector2f& Default) const
+{
+	return Mode == ENiagaraDistributionMode::UniformConstant ? FNiagaraStatelessRangeVector2(Min, Min) : FNiagaraStatelessRangeVector2(Min, Max);
+}
+
+void FNiagaraDistributionRangeVector3::InitConstant(const FVector3f& Value)
+{
+	Mode = ENiagaraDistributionMode::NonUniformConstant;
+	Min = Value;
+	Max = Value;
+#if WITH_EDITORONLY_DATA
+	ChannelConstantsAndRanges = TArray<float>({ Value.X, Value.Y, Value.Z });
+#endif
+}
+
+FNiagaraStatelessRangeVector3 FNiagaraDistributionRangeVector3::CalculateRange(const FVector3f& Default) const
+{
+	return Mode == ENiagaraDistributionMode::UniformConstant ? FNiagaraStatelessRangeVector3(Min, Min) : FNiagaraStatelessRangeVector3(Min, Max);
+}
+
+void FNiagaraDistributionFloat::InitConstant(float Value)
+{
+	Mode = ENiagaraDistributionMode::UniformConstant;
+	Values = TArray<float>({ Value, Value });
+#if WITH_EDITORONLY_DATA
+	ChannelConstantsAndRanges = TArray<float>({ Value });
+#endif
+}
+
+FNiagaraStatelessRangeFloat FNiagaraDistributionFloat::CalculateRange(const float Default) const
+{
+	FNiagaraStatelessRangeFloat Range(Default, Default);
+	if (Values.Num() > 0)
+	{
+		Range.Min = Values[0];
+		Range.Max = Values[0];
+		for (int i=1; i < Values.Num(); ++i)
+		{
+			Range.Min = FMath::Min(Range.Min, Values[i]);
+			Range.Max = FMath::Max(Range.Max, Values[i]);
+		}
+	}
+	return Range;
+}
+
+void FNiagaraDistributionVector2::InitConstant(const FVector2f& Value)
+{
+	Mode = ENiagaraDistributionMode::NonUniformConstant;
+	Values = TArray<FVector2f>({ Value, Value });
+#if WITH_EDITORONLY_DATA
+	ChannelConstantsAndRanges = TArray<float>({ Value.X, Value.Y });
+#endif
+}
+
+FNiagaraStatelessRangeVector2 FNiagaraDistributionVector2::CalculateRange(const FVector2f& Default) const
+{
+	FNiagaraStatelessRangeVector2 Range(Default, Default);
+	if (Values.Num() > 0)
+	{
+		Range.Min = Values[0];
+		Range.Max = Values[0];
+		for (int i = 1; i < Values.Num(); ++i)
+		{
+			Range.Min = FVector2f::Min(Range.Min, Values[i]);
+			Range.Max = FVector2f::Max(Range.Max, Values[i]);
+		}
+	}
+	return Range;
+}
+
+void FNiagaraDistributionVector3::InitConstant(const FVector3f& Value)
+{
+	Mode = ENiagaraDistributionMode::NonUniformConstant;
+	Values = TArray<FVector3f>({ Value, Value });
+#if WITH_EDITORONLY_DATA
+	ChannelConstantsAndRanges = TArray<float>({ Value.X, Value.Y, Value.Z });
+#endif
+}
+
+FNiagaraStatelessRangeVector3 FNiagaraDistributionVector3::CalculateRange(const FVector3f& Default) const
+{
+	FNiagaraStatelessRangeVector3 Range(Default, Default);
+	if (Values.Num() > 0)
+	{
+		Range.Min = Values[0];
+		Range.Max = Values[0];
+		for (int i = 1; i < Values.Num(); ++i)
+		{
+			Range.Min = FVector3f::Min(Range.Min, Values[i]);
+			Range.Max = FVector3f::Max(Range.Max, Values[i]);
+		}
+	}
+	return Range;
+}
+
+void FNiagaraDistributionColor::InitConstant(const FLinearColor& Value)
+{
+	Mode = ENiagaraDistributionMode::NonUniformConstant;
+	Values = TArray<FLinearColor>({ Value, Value });
+#if WITH_EDITORONLY_DATA
+	ChannelConstantsAndRanges = TArray<float>({ Value.R, Value.G, Value.B, Value.A });
+#endif
+}
+
+FNiagaraStatelessRangeColor FNiagaraDistributionColor::CalculateRange(const FLinearColor& Default) const
+{
+	FNiagaraStatelessRangeColor Range(Default, Default);
+	if (Values.Num() > 0)
+	{
+		Range.Min = Values[0];
+		Range.Max = Values[0];
+		for (int i = 1; i < Values.Num(); ++i)
+		{
+			Range.Min.R = FMath::Min(Range.Min.R, Values[i].R);
+			Range.Min.G = FMath::Min(Range.Min.G, Values[i].G);
+			Range.Min.B = FMath::Min(Range.Min.B, Values[i].B);
+			Range.Min.A = FMath::Min(Range.Min.A, Values[i].A);
+
+			Range.Max.R = FMath::Max(Range.Max.R, Values[i].R);
+			Range.Max.G = FMath::Max(Range.Max.G, Values[i].G);
+			Range.Max.B = FMath::Max(Range.Max.B, Values[i].B);
+			Range.Max.A = FMath::Max(Range.Max.A, Values[i].A);
+		}
+	}
+	return Range;
+}
+
+#if WITH_EDITORONLY_DATA
+void FNiagaraDistributionRangeFloat::UpdateValuesFromDistribution()
+{
+	NiagaraStatelessDistributionPrivate::UpdateDistributionValues(
+		Mode,
+		ChannelConstantsAndRanges,
+		ChannelCurves,
+		1,
+		[this](int32 Num) {},
+		[this](int32 ValueIndex, int32 ChannelIndex) -> float& { return ValueIndex == 0 ? Min : Max; },
+		MaxLutSampleCount
+	);
+}
+
+void FNiagaraDistributionRangeVector2::UpdateValuesFromDistribution()
+{
+	NiagaraStatelessDistributionPrivate::UpdateDistributionValues(
+		Mode,
+		ChannelConstantsAndRanges,
+		ChannelCurves,
+		2,
+		[this](int32 Num) {},
+		[this](int32 ValueIndex, int32 ChannelIndex) -> float& { return ValueIndex == 0 ? Min[ChannelIndex] : Max[ChannelIndex]; },
+		MaxLutSampleCount
+	);
+}
+
+void FNiagaraDistributionRangeVector3::UpdateValuesFromDistribution()
+{
+	NiagaraStatelessDistributionPrivate::UpdateDistributionValues(
+		Mode,
+		ChannelConstantsAndRanges,
+		ChannelCurves,
+		3,
+		[this](int32 Num) {},
+		[this](int32 ValueIndex, int32 ChannelIndex) -> float& { return ValueIndex == 0 ? Min[ChannelIndex] : Max[ChannelIndex]; },
+		MaxLutSampleCount
+	);
+}
+
+void FNiagaraDistributionFloat::UpdateValuesFromDistribution()
+{
+	NiagaraStatelessDistributionPrivate::UpdateDistributionValues(
+		Mode, 
+		ChannelConstantsAndRanges, 
+		ChannelCurves,
+		1,
+		[this](int32 NumValues) { Values.SetNum(NumValues); },
+		[this](int32 ValueIndex, int32 ChannelIndex) -> float& { return Values[ValueIndex]; },
+		MaxLutSampleCount
+	);
+}
+
+void FNiagaraDistributionVector2::UpdateValuesFromDistribution()
+{
+	NiagaraStatelessDistributionPrivate::UpdateDistributionValues(
+		Mode,
+		ChannelConstantsAndRanges,
+		ChannelCurves,
+		2,
+		[this](int32 NumValues) { Values.SetNum(NumValues); },
+		[this](int32 ValueIndex, int32 ChannelIndex) -> float& { return Values[ValueIndex][ChannelIndex]; },
+		MaxLutSampleCount
+	);
+}
+
+void FNiagaraDistributionVector3::UpdateValuesFromDistribution()
+{
+	NiagaraStatelessDistributionPrivate::UpdateDistributionValues(
+		Mode,
+		ChannelConstantsAndRanges,
+		ChannelCurves,
+		3,
+		[this](int32 NumValues) { Values.SetNum(NumValues); },
+		[this](int32 ValueIndex, int32 ChannelIndex) -> float& { return Values[ValueIndex][ChannelIndex]; },
+		MaxLutSampleCount
+	);
+}
+
+void FNiagaraDistributionColor::UpdateValuesFromDistribution()
+{
+	NiagaraStatelessDistributionPrivate::UpdateDistributionValues(
+		Mode,
+		ChannelConstantsAndRanges,
+		ChannelCurves,
+		4,
+		[this](int32 NumValues) { Values.SetNum(NumValues); },
+		[this](int32 ValueIndex, int32 ChannelIndex) -> float& { return Values[ValueIndex].Component(ChannelIndex); },
+		MaxLutSampleCount
+	);
+}
+
+#endif
+
diff --git a/Engine/Plugins/FX/Niagara/Source/Niagara/Private/Stateless/NiagaraStatelessEmitter.cpp b/Engine/Plugins/FX/Niagara/Source/Niagara/Private/Stateless/NiagaraStatelessEmitter.cpp
--- a/Engine/Plugins/FX/Niagara/Source/Niagara/Private/Stateless/NiagaraStatelessEmitter.cpp
+++ b/Engine/Plugins/FX/Niagara/Source/Niagara/Private/Stateless/NiagaraStatelessEmitter.cpp
@@ -57,7 +57,7 @@ namespace NiagaraStatelessInternal
 		}
 
 		// Lifetime is not valid so we have nothing to every render with
-		if ((EmitterData.LifetimeRange.X <= 0.0f) && (EmitterData.LifetimeRange.Y <= 0.0f))
+		if ((EmitterData.LifetimeRange.Min <= 0.0f) && (EmitterData.LifetimeRange.Max <= 0.0f))
 		{
 			return false;
 		}
@@ -222,7 +222,7 @@ void UNiagaraStatelessEmitter::CacheFromCompiledData()
 		Modules.FindItemByClass(&InitializeParticleModule);
 		if (ensure(InitializeParticleModule))
 		{
-			StatelessEmitterData->LifetimeRange = FVector2f(InitializeParticleModule->LifetimeMin, InitializeParticleModule->LifetimeMax);
+			StatelessEmitterData->LifetimeRange = InitializeParticleModule->LifetimeDistribution.CalculateRange();
 		}
 	}
 
@@ -409,7 +409,7 @@ NiagaraStateless::FCommonShaderParameters* UNiagaraStatelessEmitter::AllocateSha
 	);
 
 	NiagaraStateless::FCommonShaderParameters* CommonParameters = SetShaderParametersContext.GetParameterNestedStruct<NiagaraStateless::FCommonShaderParameters>();
-	CommonParameters->Common_LifetimeScaleBias	= FVector2f(StatelessEmitterData->LifetimeRange.Y - StatelessEmitterData->LifetimeRange.X, StatelessEmitterData->LifetimeRange.X);
+	CommonParameters->Common_LifetimeScaleBias	= FVector2f(StatelessEmitterData->LifetimeRange.Max - StatelessEmitterData->LifetimeRange.Min, StatelessEmitterData->LifetimeRange.Min);
 
 	for (UNiagaraStatelessModule* Module : Modules)
 	{
diff --git a/Engine/Plugins/FX/Niagara/Source/Niagara/Private/Stateless/NiagaraStatelessEmitterData.cpp b/Engine/Plugins/FX/Niagara/Source/Niagara/Private/Stateless/NiagaraStatelessEmitterData.cpp
--- a/Engine/Plugins/FX/Niagara/Source/Niagara/Private/Stateless/NiagaraStatelessEmitterData.cpp
+++ b/Engine/Plugins/FX/Niagara/Source/Niagara/Private/Stateless/NiagaraStatelessEmitterData.cpp
@@ -78,7 +78,7 @@ uint32 FNiagaraStatelessEmitterData::CalculateActiveParticles(int32 InRandomSeed
 	{
 		uint32 SpawnInfoTotalParticles = 0;
 
-		const bool bIsValidForAge = !Age.IsSet() || (Age.GetValue() >= SpawnInfo.SpawnTimeStart && Age.GetValue() < SpawnInfo.SpawnTimeEnd + LifetimeRange.Y);
+		const bool bIsValidForAge = !Age.IsSet() || (Age.GetValue() >= SpawnInfo.SpawnTimeStart && Age.GetValue() < SpawnInfo.SpawnTimeEnd + LifetimeRange.Max);
 		if (bIsValidForAge && GpuSpawnIndex < NiagaraStateless::MaxGpuSpawnInfos)
 		{
 			uint32	NumActive = 0;
@@ -99,7 +99,7 @@ uint32 FNiagaraStatelessEmitterData::CalculateActiveParticles(int32 InRandomSeed
 					else
 					{
 						const uint32 MaxActive = FMath::FloorToInt((SpawnInfo.SpawnTimeEnd - SpawnInfo.SpawnTimeStart) * SpawnInfo.Rate);
-						ParticleOffset = FMath::FloorToInt(FMath::Max(Age.GetValue() - SpawnInfo.SpawnTimeStart - LifetimeRange.Y, 0.0f) * SpawnInfo.Rate);
+						ParticleOffset = FMath::FloorToInt(FMath::Max(Age.GetValue() - SpawnInfo.SpawnTimeStart - LifetimeRange.Max, 0.0f) * SpawnInfo.Rate);
 						ParticleOffset = FMath::Min(ParticleOffset, MaxActive);
 						NumActive = FMath::FloorToInt(FMath::Max(Age.GetValue() - SpawnInfo.SpawnTimeStart, 0.0f) * SpawnInfo.Rate);
 						NumActive = FMath::Min(NumActive, MaxActive);
diff --git a/Engine/Plugins/FX/Niagara/Source/Niagara/Private/Stateless/NiagaraStatelessEmitterInstance.cpp b/Engine/Plugins/FX/Niagara/Source/Niagara/Private/Stateless/NiagaraStatelessEmitterInstance.cpp
--- a/Engine/Plugins/FX/Niagara/Source/Niagara/Private/Stateless/NiagaraStatelessEmitterInstance.cpp
+++ b/Engine/Plugins/FX/Niagara/Source/Niagara/Private/Stateless/NiagaraStatelessEmitterInstance.cpp
@@ -409,7 +409,7 @@ void FNiagaraStatelessEmitterInstance::InitSpawnInfosForLoop()
 
 void FNiagaraStatelessEmitterInstance::TickSpawnInfos()
 {
-	const float MaxLifetime = EmitterData->LifetimeRange.Y;
+	const float MaxLifetime = EmitterData->LifetimeRange.Max;
 	SpawnInfos.RemoveAll(
 		[this, &MaxLifetime](const FNiagaraStatelessRuntimeSpawnInfo& SpawnInfo)
 		{
@@ -437,7 +437,7 @@ void FNiagaraStatelessEmitterInstance::SetExecutionStateInternal(ENiagaraExecuti
 				ActiveSpawnRates.Empty();
 
 				// Crop & Remove Spawn Infos
-				const float MaxLifetime = EmitterData->LifetimeRange.Y;
+				const float MaxLifetime = EmitterData->LifetimeRange.Max;
 				for (auto it=SpawnInfos.CreateIterator(); it; ++it)
 				{
 					FNiagaraStatelessRuntimeSpawnInfo& SpawnInfo = *it;
diff --git a/Engine/Plugins/FX/Niagara/Source/Niagara/Private/Stateless/NiagaraStatelessModule.cpp b/Engine/Plugins/FX/Niagara/Source/Niagara/Private/Stateless/NiagaraStatelessModule.cpp
--- a/Engine/Plugins/FX/Niagara/Source/Niagara/Private/Stateless/NiagaraStatelessModule.cpp
+++ b/Engine/Plugins/FX/Niagara/Source/Niagara/Private/Stateless/NiagaraStatelessModule.cpp
@@ -2,6 +2,11 @@
 
 #include "Stateless/NiagaraStatelessModule.h"
 
+#if WITH_EDITORONLY_DATA
+const FName UNiagaraStatelessModule::PrivateMemberNames::bModuleEnabled = GET_MEMBER_NAME_CHECKED(UNiagaraStatelessModule, bModuleEnabled);
+const FName UNiagaraStatelessModule::PrivateMemberNames::bDebugDrawEnabled = GET_MEMBER_NAME_CHECKED(UNiagaraStatelessModule, bDebugDrawEnabled);
+#endif
+
 #if DO_CHECK
 void FNiagaraStatelessSetShaderParameterContext::ValidateIncludeStructType(uint32 StructOffset, const FShaderParametersMetadata* StructMetaData) const
 {
@@ -52,4 +57,18 @@ bool UNiagaraStatelessModule::CanEditChange(const FProperty* InProperty) const
 
 	return true;
 }
+
+void UNiagaraStatelessModule::PostEditChangeProperty(struct FPropertyChangedEvent& PropertyChangedEvent)
+{
+	FStructProperty* StructProperty = CastField<FStructProperty>(PropertyChangedEvent.MemberProperty);
+	if (StructProperty != nullptr && StructProperty->Struct->IsChildOf(FNiagaraDistributionBase::StaticStruct()))
+	{
+		FNiagaraDistributionBase* ValuePtr = StructProperty->ContainerPtrToValuePtr<FNiagaraDistributionBase>(this);
+		if (ValuePtr != nullptr)
+		{
+			ValuePtr->UpdateValuesFromDistribution();
+		}
+	}
+	Super::PostEditChangeProperty(PropertyChangedEvent);
+}
 #endif //WITH_EDITOR
diff --git a/Engine/Plugins/FX/Niagara/Source/Niagara/Private/Stateless/NiagaraStatelessSpawnInfo.cpp b/Engine/Plugins/FX/Niagara/Source/Niagara/Private/Stateless/NiagaraStatelessSpawnInfo.cpp
new file mode 100644
--- /dev/null
+++ b/Engine/Plugins/FX/Niagara/Source/Niagara/Private/Stateless/NiagaraStatelessSpawnInfo.cpp
@@ -0,0 +1,19 @@
+// Copyright Epic Games, Inc. All Rights Reserved.
+
+#include "Stateless/NiagaraStatelessSpawnInfo.h"
+
+bool FNiagaraStatelessSpawnInfo::IsValid(TOptional<float> LoopDuration) const
+{
+	switch (Type)
+	{
+		case ENiagaraStatelessSpawnInfoType::Burst:
+			return (AmountMin + AmountMax) > 0 && SpawnTime >= 0.0f && SpawnTime < LoopDuration.Get(SpawnTime + UE_SMALL_NUMBER);
+
+		case ENiagaraStatelessSpawnInfoType::Rate:
+			return (RateMin + RateMax) > 0.0f;
+
+		default:
+			checkNoEntry();
+			return false;
+	}
+}
diff --git a/Engine/Plugins/FX/Niagara/Source/Niagara/Public/NiagaraSystemInstance.h b/Engine/Plugins/FX/Niagara/Source/Niagara/Public/NiagaraSystemInstance.h
--- a/Engine/Plugins/FX/Niagara/Source/Niagara/Public/NiagaraSystemInstance.h
+++ b/Engine/Plugins/FX/Niagara/Source/Niagara/Public/NiagaraSystemInstance.h
@@ -513,6 +513,8 @@ private:
 	int32 WarmupTickCount = -1;
 	float WarmupTickDelta = 0;
 	
+	//-TODO:Stateless:
+	//TArray<TSharedRef<FNiagaraEmitterInstanceImpl, ESPMode::ThreadSafe>> Emitters;
 	TArray<FNiagaraEmitterInstanceRef> Emitters;
 
 	FOnPostTick OnPostTickDelegate;
diff --git a/Engine/Plugins/FX/Niagara/Source/NiagaraEditor/Private/NiagaraEditorStyle.cpp b/Engine/Plugins/FX/Niagara/Source/NiagaraEditor/Private/NiagaraEditorStyle.cpp
--- a/Engine/Plugins/FX/Niagara/Source/NiagaraEditor/Private/NiagaraEditorStyle.cpp
+++ b/Engine/Plugins/FX/Niagara/Source/NiagaraEditor/Private/NiagaraEditorStyle.cpp
@@ -437,6 +437,14 @@ void FNiagaraEditorStyle::InitIcons()
 	*SlateBrush = FAppStyle::Get().GetWidgetStyle<FEditableTextBoxStyle>("NormalEditableTextBox").BackgroundImageNormal;
 	SlateBrush->TintColor = FLinearColor(0.01f, 0.01f, 0.01f);
 	Set("NiagaraEditor.Module.InputTypeBorder", SlateBrush);
+
+	Set("NiagaraEditor.DistributionEditor.UniformConstant", new IMAGE_BRUSH_SVG("Icons/DistributionEditor/UniformConstant", Icon20x20));
+	Set("NiagaraEditor.DistributionEditor.NonUniformConstant", new IMAGE_BRUSH_SVG("Icons/DistributionEditor/NonUniformConstant", Icon20x20));
+	Set("NiagaraEditor.DistributionEditor.UniformRange", new IMAGE_BRUSH_SVG("Icons/DistributionEditor/UniformRange", Icon20x20));
+	Set("NiagaraEditor.DistributionEditor.NonUniformRange", new IMAGE_BRUSH_SVG("Icons/DistributionEditor/NonUniformRange", Icon20x20));
+	Set("NiagaraEditor.DistributionEditor.UniformCurve", new IMAGE_BRUSH_SVG("Icons/DistributionEditor/UniformCurve", Icon20x20));
+	Set("NiagaraEditor.DistributionEditor.NonUniformCurve", new IMAGE_BRUSH_SVG("Icons/DistributionEditor/NonUniformCurve", Icon20x20));
+	Set("NiagaraEditor.DistributionEditor.CurveKeyHandle", new IMAGE_BRUSH_SVG("Icons/DistributionEditor/CurveKeyHandle", Icon20x20));
 }
 
 void FNiagaraEditorStyle::InitTextStyles()
diff --git a/Engine/Plugins/FX/Niagara/Source/NiagaraEditor/Private/Stateless/NiagaraDistributionPropertyCustomization.cpp b/Engine/Plugins/FX/Niagara/Source/NiagaraEditor/Private/Stateless/NiagaraDistributionPropertyCustomization.cpp
new file mode 100644
--- /dev/null
+++ b/Engine/Plugins/FX/Niagara/Source/NiagaraEditor/Private/Stateless/NiagaraDistributionPropertyCustomization.cpp
@@ -0,0 +1,638 @@
+// Copyright Epic Games, Inc. All Rights Reserved.
+
+#include "Stateless/NiagaraDistributionPropertyCustomization.h"
+
+#include "DetailWidgetRow.h"
+#include "Editor.h"
+#include "EditorUndoClient.h"
+#include "PropertyHandle.h"
+#include "ScopedTransaction.h"
+#include "Stateless/NiagaraStatelessDistribution.h"
+#include "Styling/StyleColors.h"
+#include "Widgets/DeclarativeSyntaxSupport.h"
+#include "Widgets/NiagaraDistributionEditorUtilities.h"
+#include "Widgets/SBoxPanel.h"
+#include "Widgets/SCompoundWidget.h"
+#include "Widgets/SNiagaraDistributionEditor.h"
+
+#define LOCTEXT_NAMESPACE "NiagaraDistributionPropertyCustomization"
+
+class FNiagaraDistributionAdapter : public INiagaraDistributionAdapter
+{
+public:
+	FNiagaraDistributionAdapter(TSharedPtr<IPropertyHandle> InPropertyHandle, UObject* InOwningObject, FNiagaraDistributionBase* InDistribution, int32 InNumChannels)
+		: PropertyHandle(InPropertyHandle)
+		, OwnerObjectWeak(TWeakObjectPtr<UObject>(InOwningObject))
+		, SourceDistribution(InDistribution)
+		, SourceNumChannels(InNumChannels)
+	{
+		if (IsValid())
+		{
+			EditorMode = GetDistributionEditorModeFromSourceMode(SourceNumChannels, SourceDistribution->Mode);
+			EditorNumChannels = FNiagaraDistributionEditorUtilities::IsUniform(EditorMode) ? 1 : SourceNumChannels;
+
+			const FName DisableCurveDistributionName("DisableCurveDistribution");
+			const FName DisableUniformDistributionName("DisableUniformDistribution");
+			bAllowUniform = PropertyHandle ? PropertyHandle->HasMetaData(DisableUniformDistributionName) == false : false;
+			bAllowCurves = InDistribution->AllowCurves() && (PropertyHandle ? PropertyHandle->HasMetaData(DisableCurveDistributionName) == false : false);
+		}
+	}
+
+	virtual bool IsValid() const override { return SourceDistribution != nullptr; }
+
+	virtual int32 GetNumChannels() const override { return EditorNumChannels; }
+
+	virtual FText GetChannelDisplayName(int32 ChannelIndex) const override 
+	{
+		switch(ChannelIndex)
+		{
+		case 0:
+			return LOCTEXT("XChannelName", "X");
+		case 1:
+			return LOCTEXT("YChannelName", "Y");
+		case 2:
+			return LOCTEXT("ZChannelName", "Z");
+		case 3:
+			return LOCTEXT("WChannelName", "W");
+		}
+		return FText();
+	}
+
+	virtual FSlateColor GetChannelColor(int32 ChannelIndex) const override
+	{
+		if (EditorNumChannels != 1)
+		{
+			switch (ChannelIndex)
+			{
+			case 0:
+				return FSlateColor(EStyleColor::AccentRed);
+			case 1:
+				return FSlateColor(EStyleColor::AccentGreen);
+			case 2:
+				return FSlateColor(EStyleColor::AccentBlue);
+			case 3:
+				return FSlateColor(EStyleColor::AccentWhite);
+			}
+		}
+		return FSlateColor(EStyleColor::AccentWhite);
+	}
+
+	virtual void GetSupportedDistributionModes(TArray<ENiagaraDistributionEditorMode>& OutSupportedModes) const override
+	{
+		if (SourceNumChannels == 1)
+		{
+			OutSupportedModes.Add(ENiagaraDistributionEditorMode::Constant);
+			OutSupportedModes.Add(ENiagaraDistributionEditorMode::Range);
+			if (bAllowCurves)
+			{
+				OutSupportedModes.Add(ENiagaraDistributionEditorMode::Curve);
+			}
+		}
+		else if (SourceNumChannels > 1)
+		{
+			if (bAllowUniform)
+			{
+				OutSupportedModes.Add(ENiagaraDistributionEditorMode::UniformConstant);
+				OutSupportedModes.Add(ENiagaraDistributionEditorMode::UniformRange);
+			}
+			OutSupportedModes.Add(ENiagaraDistributionEditorMode::NonUniformConstant);
+			OutSupportedModes.Add(ENiagaraDistributionEditorMode::NonUniformRange);
+			if (bAllowCurves)
+			{
+				OutSupportedModes.Add(ENiagaraDistributionEditorMode::UniformCurve);
+				OutSupportedModes.Add(ENiagaraDistributionEditorMode::NonUniformCurve);
+			}
+		}
+	}
+
+	static ENiagaraDistributionEditorMode GetDistributionEditorModeFromSourceMode(int32 InSourceNumChannels, ENiagaraDistributionMode InSourceMode)
+	{
+		if (InSourceNumChannels == 1)
+		{
+			switch (InSourceMode)
+			{
+			case ENiagaraDistributionMode::UniformConstant:
+				return ENiagaraDistributionEditorMode::Constant;
+			case ENiagaraDistributionMode::UniformRange:
+				return ENiagaraDistributionEditorMode::Range;
+			case ENiagaraDistributionMode::UniformCurve:
+				return ENiagaraDistributionEditorMode::Curve;
+			default:
+				return ENiagaraDistributionEditorMode::Constant;
+			}
+		}
+		else
+		{
+			switch (InSourceMode)
+			{
+			case ENiagaraDistributionMode::UniformConstant:
+				return ENiagaraDistributionEditorMode::UniformConstant;
+			case ENiagaraDistributionMode::NonUniformConstant:
+				return ENiagaraDistributionEditorMode::NonUniformConstant;
+			case ENiagaraDistributionMode::UniformRange:
+				return ENiagaraDistributionEditorMode::UniformRange;
+			case ENiagaraDistributionMode::NonUniformRange:
+				return ENiagaraDistributionEditorMode::NonUniformRange;
+			case ENiagaraDistributionMode::UniformCurve:
+				return ENiagaraDistributionEditorMode::UniformCurve;
+			case ENiagaraDistributionMode::NonUniformCurve:
+				return ENiagaraDistributionEditorMode::NonUniformCurve;
+			default:
+				return ENiagaraDistributionEditorMode::UniformConstant;
+			}
+		}
+	}
+
+	static ENiagaraDistributionMode GetDistributionSourceModeFromEditorMode(ENiagaraDistributionEditorMode InEditorMode)
+	{
+		switch (InEditorMode)
+		{
+		case ENiagaraDistributionEditorMode::Constant:
+		case ENiagaraDistributionEditorMode::UniformConstant:
+			return ENiagaraDistributionMode::UniformConstant;
+		case ENiagaraDistributionEditorMode::NonUniformConstant:
+			return ENiagaraDistributionMode::NonUniformConstant;
+		case ENiagaraDistributionEditorMode::Range:
+		case ENiagaraDistributionEditorMode::UniformRange:
+			return ENiagaraDistributionMode::UniformRange;
+		case ENiagaraDistributionEditorMode::NonUniformRange:
+			return ENiagaraDistributionMode::NonUniformRange;
+		case ENiagaraDistributionEditorMode::Curve:
+		case ENiagaraDistributionEditorMode::UniformCurve:
+			return ENiagaraDistributionMode::UniformCurve;
+		case ENiagaraDistributionEditorMode::NonUniformCurve:
+			return ENiagaraDistributionMode::NonUniformCurve;
+		default:
+			return ENiagaraDistributionMode::UniformConstant;
+		}
+	}
+
+	virtual ENiagaraDistributionEditorMode GetDistributionMode() const override { return EditorMode; }
+
+	virtual void SetDistributionMode(ENiagaraDistributionEditorMode InEditorMode) override
+	{
+		UObject* OwnerObject = OwnerObjectWeak.Get();
+		if (OwnerObject == nullptr)
+		{
+			return;
+		}
+
+		ENiagaraDistributionMode NewMode = GetDistributionSourceModeFromEditorMode(InEditorMode);
+		if (SourceDistribution->Mode != NewMode)
+		{
+			int32 NewEditorNumChannels = FNiagaraDistributionEditorUtilities::IsUniform(InEditorMode) ? 1 : SourceNumChannels;
+
+			const FScopedTransaction Transaction(LOCTEXT("SetDistributionModeTransaction", "Set distribution mode"));
+			OwnerObject->Modify();
+			PropertyHandle->NotifyPreChange();
+
+			MigrateDataFromModeChange(NewMode, NewEditorNumChannels);
+			SourceDistribution->Mode = NewMode;
+			EditorNumChannels = NewEditorNumChannels;
+			EditorMode = InEditorMode;
+
+			PropertyHandle->NotifyPostChange(EPropertyChangeType::ValueSet);
+		}
+	}
+
+	float GetConstantOrRangeValue(int32 ChannelIndex, int32 ValueIndex) const override
+	{
+		UObject* OwnerObject = OwnerObjectWeak.Get();
+		int32 DataIndex = ValueIndex * EditorNumChannels + ChannelIndex;
+		if (OwnerObject != nullptr &&
+			SourceDistribution != nullptr &&
+			ChannelIndex < EditorNumChannels &&
+			DataIndex < SourceDistribution->ChannelConstantsAndRanges.Num())
+		{
+			return SourceDistribution->ChannelConstantsAndRanges[DataIndex];
+		}
+		return 0.0f;
+	}
+
+	void SetConstantOrRangeValue(int32 ChannelIndex, int32 ValueIndex, float InValue) override
+	{
+		UObject* OwnerObject = OwnerObjectWeak.Get();
+		if (OwnerObject != nullptr &&
+			GetConstantOrRangeValue(ChannelIndex, ValueIndex) != InValue &&
+			ChannelIndex < EditorNumChannels)
+		{
+			FText TransactionText;
+			int32 RequiredValueCount = 0;
+			if (FNiagaraDistributionEditorUtilities::IsConstant(GetDistributionMode()))
+			{
+				TransactionText = LOCTEXT("SetConstantValueTransaction", "Set constant value");
+				RequiredValueCount = EditorNumChannels;
+			}
+			else if (FNiagaraDistributionEditorUtilities::IsRange(GetDistributionMode()))
+			{
+				TransactionText = ValueIndex == 0
+					? LOCTEXT("SetRangeMinValueTransaction", "Set range min value")
+					: LOCTEXT("SetRangeMaxValueTransaction", "Set range max value");
+				RequiredValueCount = EditorNumChannels * 2;
+			}
+
+			if (bContinuousTransactionPending)
+			{
+				bContinuousTransactionPending = false;
+				ContinuousTransactionIndex = GEditor->BeginTransaction(TransactionText);
+			}
+
+			const FScopedTransaction Transaction(TransactionText, ContinuousTransactionIndex.IsSet() == false);
+			OwnerObject->Modify();
+			PropertyHandle->NotifyPreChange();
+
+			if (SourceDistribution->ChannelConstantsAndRanges.Num() < RequiredValueCount)
+			{
+				SourceDistribution->ChannelConstantsAndRanges.AddZeroed(RequiredValueCount - SourceDistribution->ChannelConstantsAndRanges.Num());
+			}
+			else if (SourceDistribution->ChannelConstantsAndRanges.Num() > RequiredValueCount)
+			{
+				SourceDistribution->ChannelConstantsAndRanges.SetNum(RequiredValueCount);
+			}
+
+			SourceDistribution->ChannelConstantsAndRanges[ValueIndex * EditorNumChannels + ChannelIndex] = InValue;
+			PropertyHandle->NotifyPostChange(bContinuousChangeActive ? EPropertyChangeType::Interactive : EPropertyChangeType::ValueSet);
+		}
+	}
+
+	virtual FRichCurve* GetCurveValue(int32 ChannelIndex) const override
+	{
+		if (ChannelIndex < EditorNumChannels &&
+			ChannelIndex < SourceDistribution->ChannelCurves.Num())
+		{
+			return &SourceDistribution->ChannelCurves[ChannelIndex];
+		}
+		return nullptr;
+	}
+
+	virtual void SetCurveValue(int32 ChannelIndex, const FRichCurve& InValue) override
+	{
+		UObject* OwnerObject = OwnerObjectWeak.Get();
+		FRichCurve* CurrentValue = GetCurveValue(ChannelIndex);
+		if (OwnerObject != nullptr &&
+			(CurrentValue == nullptr || *CurrentValue != InValue) &&
+			ChannelIndex < EditorNumChannels)
+		{
+			FText TransactionText = LOCTEXT("SetCurveValueTransaction", "Set curve value");
+
+			if (bContinuousTransactionPending)
+			{
+				bContinuousTransactionPending = false;
+				ContinuousTransactionIndex = GEditor->BeginTransaction(TransactionText);
+			}
+
+			const FScopedTransaction Transaction(TransactionText, ContinuousTransactionIndex.IsSet() == false);
+			OwnerObject->Modify();
+			PropertyHandle->NotifyPreChange();
+
+			const int32 RequiredCurveCount = EditorNumChannels;
+			if (SourceDistribution->ChannelCurves.Num() < RequiredCurveCount)
+			{
+				SourceDistribution->ChannelCurves.AddDefaulted(RequiredCurveCount - SourceDistribution->ChannelCurves.Num());
+			}
+			else if (SourceDistribution->ChannelCurves.Num() > RequiredCurveCount)
+			{
+				SourceDistribution->ChannelCurves.SetNum(RequiredCurveCount);
+			}
+
+			SourceDistribution->ChannelCurves[ChannelIndex] = InValue;
+			PropertyHandle->NotifyPostChange(bContinuousChangeActive ? EPropertyChangeType::Interactive : EPropertyChangeType::ValueSet);
+		}
+	}
+
+	virtual void BeginContinuousChange() override
+	{
+		bContinuousTransactionPending = true;
+		bContinuousChangeActive = true;
+	}
+
+	virtual void EndContinuousChange() override
+	{
+		if (ContinuousTransactionIndex.IsSet())
+		{
+			GEditor->EndTransaction();
+		}
+		bContinuousTransactionPending = false;
+		ContinuousTransactionIndex.Reset();
+		bContinuousChangeActive = false;
+		PropertyHandle->NotifyFinishedChangingProperties();
+	}
+
+private:
+	void MigrateDataFromModeChange(ENiagaraDistributionMode NewMode, int32 NewEditorNumChannels)
+	{
+		if (NewMode == ENiagaraDistributionMode::UniformConstant || NewMode == ENiagaraDistributionMode::NonUniformConstant)
+		{
+			MigrateDataToConstants(NewEditorNumChannels);
+		}
+		else if(NewMode == ENiagaraDistributionMode::UniformRange || NewMode == ENiagaraDistributionMode::NonUniformRange)
+		{
+			MigrateDataToRanges(NewEditorNumChannels);
+		}
+		else if (NewMode == ENiagaraDistributionMode::UniformCurve || NewMode == ENiagaraDistributionMode::NonUniformCurve)
+		{
+			MigrateDataToCurves(NewEditorNumChannels);
+		}
+	}
+
+	void MigrateDataToConstants(int32 NewEditorNumChannels)
+	{
+		TArray<float> ConstantValues;
+		switch (SourceDistribution->Mode)
+		{
+		case ENiagaraDistributionMode::UniformConstant:
+		case ENiagaraDistributionMode::UniformRange:
+		case ENiagaraDistributionMode::NonUniformConstant:
+		case ENiagaraDistributionMode::NonUniformRange:
+			for (int32 ChannelIndex = 0; ChannelIndex < EditorNumChannels; ChannelIndex++)
+			{
+				ConstantValues.Add(GetConstantOrRangeValue(ChannelIndex, 0));
+			}
+			break;
+		case ENiagaraDistributionMode::UniformCurve:
+		case ENiagaraDistributionMode::NonUniformCurve:
+			for (int32 ChannelIndex = 0; ChannelIndex < EditorNumChannels; ChannelIndex++)
+			{
+				FRichCurve* Curve = GetCurveValue(ChannelIndex);
+				float ConstantValue = Curve != nullptr && Curve->GetNumKeys() > 0 ? Curve->GetFirstKey().Value : 0;
+				ConstantValues.Add(ConstantValue);
+			}
+			break;
+		default:
+			ConstantValues.Add(0);
+		}
+
+		SourceDistribution->ChannelConstantsAndRanges.SetNum(NewEditorNumChannels);
+		SourceDistribution->ChannelCurves.Empty();
+		for (int32 ChannelIndex = 0; ChannelIndex < NewEditorNumChannels; ChannelIndex++)
+		{
+			int32 ValueIndex = ChannelIndex < ConstantValues.Num() ? ChannelIndex : 0;
+			SourceDistribution->ChannelConstantsAndRanges[ChannelIndex] = ConstantValues[ValueIndex];
+		}
+	}
+
+	void MigrateDataToRanges(int32 NewEditorNumChannels)
+	{
+		TArray<float> MinValues;
+		TArray<float> MaxValues;
+		switch (SourceDistribution->Mode)
+		{
+		case ENiagaraDistributionMode::UniformConstant:
+		case ENiagaraDistributionMode::NonUniformConstant:
+			for (int32 ChannelIndex = 0; ChannelIndex < EditorNumChannels; ChannelIndex++)
+			{
+				float ConstantValue = GetConstantOrRangeValue(ChannelIndex, 0);
+				MinValues.Add(ConstantValue);
+				MaxValues.Add(ConstantValue);
+			}
+			break;
+		case ENiagaraDistributionMode::UniformRange:
+		case ENiagaraDistributionMode::NonUniformRange:
+			for (int32 ChannelIndex = 0; ChannelIndex < EditorNumChannels; ChannelIndex++)
+			{
+				MinValues.Add(GetConstantOrRangeValue(ChannelIndex, 0));
+				MaxValues.Add(GetConstantOrRangeValue(ChannelIndex, 1));
+			}
+			break;
+		case ENiagaraDistributionMode::UniformCurve:
+		case ENiagaraDistributionMode::NonUniformCurve:
+			for (int32 ChannelIndex = 0; ChannelIndex < EditorNumChannels; ChannelIndex++)
+			{
+				FRichCurve* Curve = GetCurveValue(ChannelIndex);
+				if (Curve != nullptr && Curve->GetNumKeys() > 0)
+				{
+					MinValues.Add(Curve->GetFirstKey().Value);
+					MaxValues.Add(Curve->GetLastKey().Value);
+				}
+				else
+				{
+					MinValues.Add(0);
+					MaxValues.Add(0);
+				}
+			}
+			break;
+		default:
+			MinValues.Add(0);
+			MaxValues.Add(0);
+			break;
+		}
+
+		SourceDistribution->ChannelConstantsAndRanges.SetNum(NewEditorNumChannels * 2);
+		SourceDistribution->ChannelCurves.Empty();
+		for (int32 ChannelIndex = 0; ChannelIndex < NewEditorNumChannels; ChannelIndex++)
+		{
+			int32 MinValueIndex = ChannelIndex < MinValues.Num() ? ChannelIndex : 0;
+			SourceDistribution->ChannelConstantsAndRanges[ChannelIndex] = MinValues[MinValueIndex];
+
+			int32 MaxValueIndex = ChannelIndex < MaxValues.Num() ? ChannelIndex : 0;
+			SourceDistribution->ChannelConstantsAndRanges[NewEditorNumChannels + ChannelIndex] = MaxValues[MaxValueIndex];
+		}
+	}
+
+	void MigrateDataToCurves(int32 NewEditorNumChannels)
+	{
+		TArray<FRichCurve> CurveValues;
+		switch (SourceDistribution->Mode)
+		{
+		case ENiagaraDistributionMode::UniformConstant:
+		case ENiagaraDistributionMode::NonUniformConstant:
+			for (int32 ChannelIndex = 0; ChannelIndex < EditorNumChannels; ChannelIndex++)
+			{
+				FRichCurve& CurveValue = CurveValues.AddDefaulted_GetRef();
+				CurveValue.AddKey(0, GetConstantOrRangeValue(ChannelIndex, 0));
+			}
+			break;
+		case ENiagaraDistributionMode::UniformRange:
+		case ENiagaraDistributionMode::NonUniformRange:
+			for (int32 ChannelIndex = 0; ChannelIndex < EditorNumChannels; ChannelIndex++)
+			{
+				FRichCurve& CurveValue = CurveValues.AddDefaulted_GetRef();
+				CurveValue.AddKey(0, GetConstantOrRangeValue(ChannelIndex, 0));
+				CurveValue.AddKey(1, GetConstantOrRangeValue(ChannelIndex, 1));
+			}
+			break;
+		case ENiagaraDistributionMode::UniformCurve:
+		case ENiagaraDistributionMode::NonUniformCurve:
+			for (int32 ChannelIndex = 0; ChannelIndex < EditorNumChannels; ChannelIndex++)
+			{
+				FRichCurve* Curve = GetCurveValue(ChannelIndex);
+				if (Curve != nullptr)
+				{
+					CurveValues.Add(*Curve);
+				}
+				else
+				{
+					FRichCurve& CurveValue = CurveValues.AddDefaulted_GetRef();
+					CurveValue.AddKey(0, 0);
+				}
+			}
+			break;
+		default:
+			FRichCurve& CurveValue = CurveValues.AddDefaulted_GetRef();
+			CurveValue.AddKey(0, 0);
+			break;
+		}
+
+		SourceDistribution->ChannelCurves.SetNum(NewEditorNumChannels);
+		SourceDistribution->ChannelConstantsAndRanges.Empty();
+		for (int32 ChannelIndex = 0; ChannelIndex < NewEditorNumChannels; ChannelIndex++)
+		{
+			int32 ValueIndex = ChannelIndex < CurveValues.Num() ? ChannelIndex : 0;
+			SourceDistribution->ChannelCurves[ChannelIndex] = CurveValues[ValueIndex];
+		}
+	}
+
+private:
+	TSharedPtr<IPropertyHandle> PropertyHandle;
+	TWeakObjectPtr<UObject> OwnerObjectWeak;
+	FNiagaraDistributionBase* SourceDistribution = nullptr;
+	int32 SourceNumChannels = 0;
+
+	ENiagaraDistributionEditorMode EditorMode = ENiagaraDistributionEditorMode::Constant;
+	int32 EditorNumChannels = 0;
+
+	bool bContinuousTransactionPending = false;
+	TOptional<int32> ContinuousTransactionIndex;
+	bool bContinuousChangeActive = false;
+
+	bool bAllowUniform = true;
+	bool bAllowCurves = true;
+};
+
+class SNiagaraDistributionPropertyWidget : public SCompoundWidget, public FEditorUndoClient
+{
+	SLATE_BEGIN_ARGS(SNiagaraDistributionPropertyWidget) { }
+	SLATE_END_ARGS();
+
+	void Construct(const FArguments& InArgs, TSharedRef<IPropertyHandle> InDistributionPropertyHandle, TSharedRef<INiagaraDistributionAdapter> InDistributionAdapter)
+	{
+		GEditor->RegisterForUndo(this);
+		DistributionPropertyHandleWeak = InDistributionPropertyHandle;
+		DistributionAdapter = InDistributionAdapter;
+		InDistributionPropertyHandle->SetOnPropertyValueChanged(FSimpleDelegate::CreateSP(this, &SNiagaraDistributionPropertyWidget::HandleValueChanged));
+
+		ChildSlot
+		[
+			SNew(SNiagaraDistributionEditor, DistributionAdapter.ToSharedRef())
+		];
+	}
+
+	virtual ~SNiagaraDistributionPropertyWidget()
+	{
+		if (GEditor != nullptr)
+		{
+			GEditor->UnregisterForUndo(this);
+		}
+	}
+
+	// Begin FEditorUndoClient
+	virtual void PostUndo(bool bSuccess) override
+	{
+	}
+
+	virtual void PostRedo(bool bSuccess) override { PostUndo(bSuccess); }
+	// End of FEditorUndoClient
+
+private:
+	void HandleValueChanged()
+	{
+		if (bUpdatingHandle == false)
+		{
+		}
+	}
+
+private:
+	TWeakPtr<IPropertyHandle> DistributionPropertyHandleWeak;
+	TSharedPtr<INiagaraDistributionAdapter> DistributionAdapter;
+	bool bUpdatingHandle = false;
+};
+
+TSharedRef<IPropertyTypeCustomization> FNiagaraDistributionPropertyCustomization::MakeFloatInstance()
+{
+	FPropertyHandleToDistributionAdapter FloatDistributionPropertyHandleToDistributionAdapter = FPropertyHandleToDistributionAdapter::CreateLambda([](TSharedRef<IPropertyHandle> FloatDistributionPropertyHandle)
+	{
+		void* ValueData = nullptr;
+		TArray<UObject*> OuterObjects;
+		FloatDistributionPropertyHandle->GetOuterObjects(OuterObjects);
+		if (OuterObjects.Num() == 1 && FloatDistributionPropertyHandle->GetValueData(ValueData) == FPropertyAccess::Success)
+		{
+			FNiagaraDistributionBase* FloatDistribution = static_cast<FNiagaraDistributionBase*>(ValueData);
+			return MakeShared<FNiagaraDistributionAdapter>(FloatDistributionPropertyHandle, OuterObjects[0], FloatDistribution, 1);
+		}
+		return MakeShared<FNiagaraDistributionAdapter>(nullptr, nullptr, nullptr, INDEX_NONE);
+	});
+
+	return MakeShareable<FNiagaraDistributionPropertyCustomization>(new FNiagaraDistributionPropertyCustomization(FloatDistributionPropertyHandleToDistributionAdapter));
+}
+
+TSharedRef<IPropertyTypeCustomization> FNiagaraDistributionPropertyCustomization::MakeVector2Instance()
+{
+	FPropertyHandleToDistributionAdapter Vector2DistributionPropertyHandleToDistributionAdapter = FPropertyHandleToDistributionAdapter::CreateLambda([](TSharedRef<IPropertyHandle> Vector2DistributionPropertyHandle)
+	{
+		void* ValueData = nullptr;
+		TArray<UObject*> OuterObjects;
+		Vector2DistributionPropertyHandle->GetOuterObjects(OuterObjects);
+		if (OuterObjects.Num() == 1 && Vector2DistributionPropertyHandle->GetValueData(ValueData) == FPropertyAccess::Success)
+		{
+			FNiagaraDistributionBase* Vector2Distribution = static_cast<FNiagaraDistributionBase*>(ValueData);
+			return MakeShared<FNiagaraDistributionAdapter>(Vector2DistributionPropertyHandle, OuterObjects[0], Vector2Distribution, 2);
+		}
+		return MakeShared<FNiagaraDistributionAdapter>(nullptr, nullptr, nullptr, INDEX_NONE);
+	});
+
+	return MakeShareable<FNiagaraDistributionPropertyCustomization>(new FNiagaraDistributionPropertyCustomization(Vector2DistributionPropertyHandleToDistributionAdapter));
+}
+
+TSharedRef<IPropertyTypeCustomization> FNiagaraDistributionPropertyCustomization::MakeVector3Instance()
+{
+	FPropertyHandleToDistributionAdapter Vector3DistributionPropertyHandleToDistributionAdapter = FPropertyHandleToDistributionAdapter::CreateLambda([](TSharedRef<IPropertyHandle> Vector3DistributionPropertyHandle)
+	{
+		void* ValueData = nullptr;
+		TArray<UObject*> OuterObjects;
+		Vector3DistributionPropertyHandle->GetOuterObjects(OuterObjects);
+		if (OuterObjects.Num() == 1 && Vector3DistributionPropertyHandle->GetValueData(ValueData) == FPropertyAccess::Success)
+		{
+			FNiagaraDistributionBase* Vector3Distribution = static_cast<FNiagaraDistributionBase*>(ValueData);
+			return MakeShared<FNiagaraDistributionAdapter>(Vector3DistributionPropertyHandle, OuterObjects[0], Vector3Distribution, 3);
+		}
+		return MakeShared<FNiagaraDistributionAdapter>(nullptr, nullptr, nullptr, INDEX_NONE);
+	});
+
+	return MakeShareable<FNiagaraDistributionPropertyCustomization>(new FNiagaraDistributionPropertyCustomization(Vector3DistributionPropertyHandleToDistributionAdapter));
+}
+
+TSharedRef<IPropertyTypeCustomization> FNiagaraDistributionPropertyCustomization::MakeColorInstance()
+{
+	FPropertyHandleToDistributionAdapter ColorDistributionPropertyHandleToDistributionAdapter = FPropertyHandleToDistributionAdapter::CreateLambda([](TSharedRef<IPropertyHandle> ColorDistributionPropertyHandle)
+	{
+		void* ValueData = nullptr;
+		TArray<UObject*> OuterObjects;
+		ColorDistributionPropertyHandle->GetOuterObjects(OuterObjects);
+		if (OuterObjects.Num() == 1 && ColorDistributionPropertyHandle->GetValueData(ValueData) == FPropertyAccess::Success)
+		{
+			FNiagaraDistributionBase* ColorDistribution = static_cast<FNiagaraDistributionBase*>(ValueData);
+			return MakeShared<FNiagaraDistributionAdapter>(ColorDistributionPropertyHandle, OuterObjects[0], ColorDistribution, 4);
+		}
+		return MakeShared<FNiagaraDistributionAdapter>(nullptr, nullptr, nullptr, INDEX_NONE);
+	});
+
+	return MakeShareable<FNiagaraDistributionPropertyCustomization>(new FNiagaraDistributionPropertyCustomization(ColorDistributionPropertyHandleToDistributionAdapter));
+}
+
+void FNiagaraDistributionPropertyCustomization::CustomizeHeader(TSharedRef<IPropertyHandle> PropertyHandle, FDetailWidgetRow& HeaderRow, IPropertyTypeCustomizationUtils& CustomizationUtils)
+{
+	HeaderRow.NameContent()
+	[
+		PropertyHandle->CreatePropertyNameWidget()
+	];
+	HeaderRow.ValueContent()
+	[
+		SNew(SNiagaraDistributionPropertyWidget, PropertyHandle, PropertyHandleToDistributionAdapter.Execute(PropertyHandle))
+	];
+}
+
+void FNiagaraDistributionPropertyCustomization::CustomizeChildren(TSharedRef<IPropertyHandle> PropertyHandle, IDetailChildrenBuilder& ChildBuilder, IPropertyTypeCustomizationUtils& CustomizationUtils)
+{
+}
+
+#undef LOCTEXT_NAMESPACE
\ No newline at end of file
diff --git a/Engine/Plugins/FX/Niagara/Source/NiagaraEditor/Private/Stateless/NiagaraDistributionPropertyCustomization.h b/Engine/Plugins/FX/Niagara/Source/NiagaraEditor/Private/Stateless/NiagaraDistributionPropertyCustomization.h
new file mode 100644
--- /dev/null
+++ b/Engine/Plugins/FX/Niagara/Source/NiagaraEditor/Private/Stateless/NiagaraDistributionPropertyCustomization.h
@@ -0,0 +1,33 @@
+// Copyright Epic Games, Inc. All Rights Reserved.
+
+#pragma once
+
+#include "IPropertyTypeCustomization.h"
+
+class INiagaraDistributionAdapter;
+
+class FNiagaraDistributionPropertyCustomization : public IPropertyTypeCustomization
+{
+public:
+	static TSharedRef<IPropertyTypeCustomization> MakeFloatInstance();
+	static TSharedRef<IPropertyTypeCustomization> MakeVector2Instance();
+	static TSharedRef<IPropertyTypeCustomization> MakeVector3Instance();
+	static TSharedRef<IPropertyTypeCustomization> MakeColorInstance();
+
+	virtual void CustomizeHeader(TSharedRef<IPropertyHandle> PropertyHandle, FDetailWidgetRow& HeaderRow, IPropertyTypeCustomizationUtils& CustomizationUtils) override;
+
+	virtual void CustomizeChildren(TSharedRef<IPropertyHandle> PropertyHandle, IDetailChildrenBuilder& ChildBuilder, IPropertyTypeCustomizationUtils& CustomizationUtils) override;
+
+private:
+	DECLARE_DELEGATE_RetVal_OneParam(TSharedRef<INiagaraDistributionAdapter>, FPropertyHandleToDistributionAdapter, TSharedRef<IPropertyHandle> /* InDistributionPropertyHandle */)
+
+private:
+	FNiagaraDistributionPropertyCustomization(FPropertyHandleToDistributionAdapter InPropertyHandleToDistributionAdapter)
+		: PropertyHandleToDistributionAdapter(InPropertyHandleToDistributionAdapter)
+	{
+	}
+
+private:
+	
+	FPropertyHandleToDistributionAdapter PropertyHandleToDistributionAdapter;
+};
\ No newline at end of file
diff --git a/Engine/Plugins/FX/Niagara/Source/NiagaraEditor/Private/ViewModels/NiagaraCurveSelectionViewModel.cpp b/Engine/Plugins/FX/Niagara/Source/NiagaraEditor/Private/ViewModels/NiagaraCurveSelectionViewModel.cpp
--- a/Engine/Plugins/FX/Niagara/Source/NiagaraEditor/Private/ViewModels/NiagaraCurveSelectionViewModel.cpp
+++ b/Engine/Plugins/FX/Niagara/Source/NiagaraEditor/Private/ViewModels/NiagaraCurveSelectionViewModel.cpp
@@ -886,8 +886,10 @@ void UNiagaraCurveSelectionViewModel::Refresh()
 		for (TSharedRef<FNiagaraEmitterHandleViewModel> EmitterHandleViewModel : SystemViewModel->GetEmitterHandleViewModels())
 		{
 			// During transactions Emitters may not have their EditorData set if they were created during the transaction; skip if this is true.
+			//-TODO:Stateless:
 			FNiagaraEmitterHandle* EmitterHandle = EmitterHandleViewModel->GetEmitterHandle();
 			if (EmitterHandle && EmitterHandle->GetEmitterData() && EmitterHandle->GetEmitterData()->GetEditorData())
+			//-TODO:Stateless:
 			{
 				TSharedPtr<FNiagaraCurveSelectionTreeNode> EmitterNode = CreateNodeForEmitter(OldRootChildNodes, *EmitterHandle);
 				if (EmitterNode.IsValid())
diff --git a/Engine/Plugins/FX/Niagara/Source/NiagaraEditor/Private/ViewModels/NiagaraSystemViewModel.cpp b/Engine/Plugins/FX/Niagara/Source/NiagaraEditor/Private/ViewModels/NiagaraSystemViewModel.cpp
--- a/Engine/Plugins/FX/Niagara/Source/NiagaraEditor/Private/ViewModels/NiagaraSystemViewModel.cpp
+++ b/Engine/Plugins/FX/Niagara/Source/NiagaraEditor/Private/ViewModels/NiagaraSystemViewModel.cpp
@@ -2970,11 +2970,18 @@ void FNiagaraSystemViewModel::UpdateEmitterHandleSelectionFromSequencer()
 			return;
 		}
 	}
-	
-	TArray<FGuid> NewSelectedEmitterHandleIds;
 
 	TArray<UMovieSceneTrack*> SelectedTracks;
+	TArray<UMovieSceneSection*> SelectedSections;
 	Sequencer->GetSelectedTracks(SelectedTracks);
+	Sequencer->GetSelectedSections(SelectedSections);
+	if (SelectedTracks.Num() == 0 && SelectedSections.Num() == 0)
+	{
+		return;
+	}
+
+	TArray<FGuid> NewSelectedEmitterHandleIds;
+
 	for (UMovieSceneTrack* SelectedTrack : SelectedTracks)
 	{
 		UMovieSceneNiagaraEmitterTrack* SelectedEmitterTrack = Cast<UMovieSceneNiagaraEmitterTrack>(SelectedTrack);
@@ -2984,8 +2991,6 @@ void FNiagaraSystemViewModel::UpdateEmitterHandleSelectionFromSequencer()
 		}
 	}
 
-	TArray<UMovieSceneSection*> SelectedSections;
-	Sequencer->GetSelectedSections(SelectedSections);
 	for (UMovieSceneSection* SelectedSection : SelectedSections)
 	{
 		UMovieSceneNiagaraEmitterSectionBase* SelectedEmitterSection = Cast<UMovieSceneNiagaraEmitterSectionBase>(SelectedSection);
diff --git a/Engine/Plugins/FX/Niagara/Source/NiagaraEditor/Private/ViewModels/Stack/NiagaraStackEmitterSettingsGroup.cpp b/Engine/Plugins/FX/Niagara/Source/NiagaraEditor/Private/ViewModels/Stack/NiagaraStackEmitterSettingsGroup.cpp
--- a/Engine/Plugins/FX/Niagara/Source/NiagaraEditor/Private/ViewModels/Stack/NiagaraStackEmitterSettingsGroup.cpp
+++ b/Engine/Plugins/FX/Niagara/Source/NiagaraEditor/Private/ViewModels/Stack/NiagaraStackEmitterSettingsGroup.cpp
@@ -32,7 +32,9 @@ void UNiagaraStackEmitterPropertiesItem::Initialize(FRequiredEntryData InRequire
 {
 	Super::Initialize(InRequiredEntryData, TEXT("EmitterProperties"));
 	EmitterWeakPtr = GetEmitterViewModel()->GetEmitter().ToWeakPtr();
+	//-TODO:Stateless:
 	if (EmitterWeakPtr.IsValid())
+	//-TODO:Stateless:
 	{
 		EmitterWeakPtr.Emitter->OnPropertiesChanged().AddUObject(this, &UNiagaraStackEmitterPropertiesItem::EmitterPropertiesChanged);
 	}
diff --git a/Engine/Plugins/FX/Niagara/Source/NiagaraEditor/Private/ViewModels/Stack/NiagaraStackItemPropertyHeaderValueShared.cpp b/Engine/Plugins/FX/Niagara/Source/NiagaraEditor/Private/ViewModels/Stack/NiagaraStackItemPropertyHeaderValueShared.cpp
--- a/Engine/Plugins/FX/Niagara/Source/NiagaraEditor/Private/ViewModels/Stack/NiagaraStackItemPropertyHeaderValueShared.cpp
+++ b/Engine/Plugins/FX/Niagara/Source/NiagaraEditor/Private/ViewModels/Stack/NiagaraStackItemPropertyHeaderValueShared.cpp
@@ -180,7 +180,7 @@ void FNiagaraStackItemPropertyHeaderValueShared::GenerateHeaderValueHandlers(UOb
 	for (TFieldIterator<FProperty> PropertyIt(&InTargetStruct, EFieldIteratorFlags::SuperClassFlags::IncludeSuper, EFieldIteratorFlags::DeprecatedPropertyFlags::ExcludeDeprecated); PropertyIt; ++PropertyIt)
 	{
 		FProperty* Property = *PropertyIt;
-		if (Property->HasMetaData("ShowInStackItemHeader") && InOwnerObject.CanEditChange(Property))
+		if (Property->HasMetaData("ShowInStackItemHeader"))
 		{
 			FEnumProperty* EnumProperty = CastField<FEnumProperty>(Property);
 			if (EnumProperty != nullptr)
diff --git a/Engine/Plugins/FX/Niagara/Source/NiagaraEditor/Private/ViewModels/Stack/NiagaraStackStatelessEmitterGroup.cpp b/Engine/Plugins/FX/Niagara/Source/NiagaraEditor/Private/ViewModels/Stack/NiagaraStackStatelessEmitterGroup.cpp
--- a/Engine/Plugins/FX/Niagara/Source/NiagaraEditor/Private/ViewModels/Stack/NiagaraStackStatelessEmitterGroup.cpp
+++ b/Engine/Plugins/FX/Niagara/Source/NiagaraEditor/Private/ViewModels/Stack/NiagaraStackStatelessEmitterGroup.cpp
@@ -50,7 +50,7 @@ void UNiagaraStackStatelessEmitterGroup::RefreshChildrenInternal(const TArray<UN
 		UNiagaraStackStatelessEmitterObjectItem* FilteredObjectItem = FilteredObjectItemWeak.Get();
 		if (FilteredObjectItem == nullptr || FilteredObjectItem->GetStatelessEmitter() != StatelessEmitter)
 		{
-			bool bExpandedByDefault = false;
+			bool bExpandedByDefault = true;
 			FilteredObjectItem = NewObject<UNiagaraStackStatelessEmitterObjectItem>(this);
 			FilteredObjectItem->Initialize(
 				CreateDefaultChildRequiredData(),
diff --git a/Engine/Plugins/FX/Niagara/Source/NiagaraEditor/Private/ViewModels/Stack/NiagaraStackStatelessEmitterSimulateGroup.cpp b/Engine/Plugins/FX/Niagara/Source/NiagaraEditor/Private/ViewModels/Stack/NiagaraStackStatelessEmitterSimulateGroup.cpp
--- a/Engine/Plugins/FX/Niagara/Source/NiagaraEditor/Private/ViewModels/Stack/NiagaraStackStatelessEmitterSimulateGroup.cpp
+++ b/Engine/Plugins/FX/Niagara/Source/NiagaraEditor/Private/ViewModels/Stack/NiagaraStackStatelessEmitterSimulateGroup.cpp
@@ -6,6 +6,8 @@
 #include "NiagaraEditorStyle.h"
 #include "PropertyHandle.h"
 #include "ScopedTransaction.h"
+#include "Stateless/NiagaraDistributionPropertyCustomization.h"
+#include "Stateless/NiagaraStatelessCommon.h"
 #include "Stateless/NiagaraStatelessEmitter.h"
 #include "Stateless/NiagaraStatelessModule.h"
 #include "ViewModels/Stack/NiagaraStackItemPropertyHeaderValueShared.h"
@@ -83,6 +85,13 @@ void UNiagaraStackStatelessModuleItem::RefreshChildrenInternal(const TArray<UNia
 			ModuleObject = NewObject<UNiagaraStackObject>(this);
 			ModuleObject->Initialize(CreateDefaultChildRequiredData(), StatelessModule, bIsInTopLevelObject, bHideTopLevelCategories, GetStackEditorDataKey());
 			ModuleObject->SetOnFilterDetailNodes(FNiagaraStackObjectShared::FOnFilterDetailNodes::CreateStatic(&UNiagaraStackStatelessModuleItem::FilterDetailNodes), UNiagaraStackObject::EDetailNodeFilterMode::FilterAllNodes);
+			ModuleObject->RegisterInstancedCustomPropertyTypeLayout(FNiagaraDistributionFloat::StaticStruct()->GetFName(), FOnGetPropertyTypeCustomizationInstance::CreateStatic(&FNiagaraDistributionPropertyCustomization::MakeFloatInstance));
+			ModuleObject->RegisterInstancedCustomPropertyTypeLayout(FNiagaraDistributionVector2::StaticStruct()->GetFName(), FOnGetPropertyTypeCustomizationInstance::CreateStatic(&FNiagaraDistributionPropertyCustomization::MakeVector2Instance));
+			ModuleObject->RegisterInstancedCustomPropertyTypeLayout(FNiagaraDistributionVector3::StaticStruct()->GetFName(), FOnGetPropertyTypeCustomizationInstance::CreateStatic(&FNiagaraDistributionPropertyCustomization::MakeVector3Instance));
+			ModuleObject->RegisterInstancedCustomPropertyTypeLayout(FNiagaraDistributionColor::StaticStruct()->GetFName(), FOnGetPropertyTypeCustomizationInstance::CreateStatic(&FNiagaraDistributionPropertyCustomization::MakeColorInstance));
+			ModuleObject->RegisterInstancedCustomPropertyTypeLayout(FNiagaraDistributionRangeFloat::StaticStruct()->GetFName(), FOnGetPropertyTypeCustomizationInstance::CreateStatic(&FNiagaraDistributionPropertyCustomization::MakeFloatInstance));
+			ModuleObject->RegisterInstancedCustomPropertyTypeLayout(FNiagaraDistributionRangeVector2::StaticStruct()->GetFName(), FOnGetPropertyTypeCustomizationInstance::CreateStatic(&FNiagaraDistributionPropertyCustomization::MakeVector2Instance));
+			ModuleObject->RegisterInstancedCustomPropertyTypeLayout(FNiagaraDistributionRangeVector3::StaticStruct()->GetFName(), FOnGetPropertyTypeCustomizationInstance::CreateStatic(&FNiagaraDistributionPropertyCustomization::MakeVector3Instance));
 			ModuleObjectWeak = ModuleObject;
 		}
 		NewChildren.Add(ModuleObject);
@@ -91,6 +100,25 @@ void UNiagaraStackStatelessModuleItem::RefreshChildrenInternal(const TArray<UNia
 		{
 			bGeneratedHeaderValueHandlers = true;
 			FNiagaraStackItemPropertyHeaderValueShared::GenerateHeaderValueHandlers(*StatelessModule, nullptr, *StatelessModule->GetClass(), FSimpleDelegate::CreateUObject(this, &UNiagaraStackStatelessModuleItem::OnHeaderValueChanged), HeaderValueHandlers);
+			if (StatelessModule->CanDebugDraw())
+			{
+				// Debug draw is handled separately here since it's visibility needs to be determined by the results of the CanDebugDraw function, rather than by property metadata.
+				FBoolProperty* DebugDrawProperty = nullptr;
+				for (TFieldIterator<FProperty> PropertyIt(StatelessModule->GetClass(), EFieldIteratorFlags::SuperClassFlags::IncludeSuper, EFieldIteratorFlags::DeprecatedPropertyFlags::ExcludeDeprecated); PropertyIt; ++PropertyIt)
+				{
+					if (PropertyIt->GetFName() == UNiagaraStatelessModule::PrivateMemberNames::bDebugDrawEnabled)
+					{
+						DebugDrawProperty = CastField<FBoolProperty>(*PropertyIt);
+						break;
+					}
+				}
+				if (DebugDrawProperty != nullptr)
+				{
+					HeaderValueHandlers.Add(MakeShared<FNiagaraStackItemPropertyHeaderValue>(
+						*ModuleObject, nullptr, *DebugDrawProperty, 
+						FSimpleDelegate::CreateUObject(this, &UNiagaraStackStatelessModuleItem::OnHeaderValueChanged)));
+				}
+			}
 		}
 		else
 		{
diff --git a/Engine/Plugins/FX/Niagara/Source/NiagaraEditor/Private/Widgets/INiagaraDistributionAdapter.h b/Engine/Plugins/FX/Niagara/Source/NiagaraEditor/Private/Widgets/INiagaraDistributionAdapter.h
new file mode 100644
--- /dev/null
+++ b/Engine/Plugins/FX/Niagara/Source/NiagaraEditor/Private/Widgets/INiagaraDistributionAdapter.h
@@ -0,0 +1,49 @@
+// Copyright Epic Games, Inc. All Rights Reserved.
+
+#pragma once
+
+#include "Containers/Array.h"
+#include "HAL/Platform.h"
+#include "Internationalization/Text.h"
+#include "Styling/SlateColor.h"
+
+struct FRichCurve;
+
+enum class ENiagaraDistributionEditorMode
+{
+	Constant,
+	UniformConstant,
+	NonUniformConstant,
+	Range,
+	UniformRange,
+	NonUniformRange,
+	Curve,
+	UniformCurve,
+	NonUniformCurve
+};
+
+class INiagaraDistributionAdapter
+{
+public:
+	virtual ~INiagaraDistributionAdapter() { }
+
+	virtual bool IsValid() const = 0;
+
+	virtual int32 GetNumChannels() const = 0;
+	virtual FText GetChannelDisplayName(int32 ChannelIndex) const = 0;
+	virtual FSlateColor GetChannelColor(int32 ChannelIndex) const = 0;
+
+	virtual void GetSupportedDistributionModes(TArray<ENiagaraDistributionEditorMode>& OutSupportedModes) const = 0;
+
+	virtual ENiagaraDistributionEditorMode GetDistributionMode() const = 0;
+	virtual void SetDistributionMode(ENiagaraDistributionEditorMode InMode) = 0;
+
+	virtual float GetConstantOrRangeValue(int32 ChannelIndex, int32 ValueIndex) const = 0;
+	virtual void SetConstantOrRangeValue(int32 ChannelIndex, int32 ValueIndex, float InValue) = 0;
+
+	virtual FRichCurve* GetCurveValue(int32 ChannelIndex) const = 0;
+	virtual void SetCurveValue(int32 ChannelIndex, const FRichCurve& InValue) = 0;
+
+	virtual void BeginContinuousChange() = 0;
+	virtual void EndContinuousChange() = 0;
+};
\ No newline at end of file
diff --git a/Engine/Plugins/FX/Niagara/Source/NiagaraEditor/Private/Widgets/NiagaraDistributionEditorUtilities.cpp b/Engine/Plugins/FX/Niagara/Source/NiagaraEditor/Private/Widgets/NiagaraDistributionEditorUtilities.cpp
new file mode 100644
--- /dev/null
+++ b/Engine/Plugins/FX/Niagara/Source/NiagaraEditor/Private/Widgets/NiagaraDistributionEditorUtilities.cpp
@@ -0,0 +1,134 @@
+// Copyright Epic Games, Inc. All Rights Reserved.
+
+#include "Widgets/NiagaraDistributionEditorUtilities.h"
+
+#include "NiagaraEditorStyle.h"
+#include "Textures/SlateIcon.h"
+#include "Widgets/INiagaraDistributionAdapter.h"
+
+#define LOCTEXT_NAMESPACE "NiagaraDistributionEditorUtilities"
+
+FText FNiagaraDistributionEditorUtilities::DistributionModeToDisplayName(ENiagaraDistributionEditorMode InMode)
+{
+	switch (InMode)
+	{
+	case ENiagaraDistributionEditorMode::Constant:
+		return LOCTEXT("ConstantDisplayName", "Constant");
+	case ENiagaraDistributionEditorMode::UniformConstant:
+		return LOCTEXT("UniformContstantDisplayName", "Uniform Constant");
+	case ENiagaraDistributionEditorMode::NonUniformConstant:
+		return LOCTEXT("NonUniformConstantDisplayName", "Non-uniform Constant");
+	case ENiagaraDistributionEditorMode::Range:
+		return LOCTEXT("RangeDisplayName", "Range");
+	case ENiagaraDistributionEditorMode::UniformRange:
+		return LOCTEXT("UniformRangeDisplayName", "Uniform Range");
+	case ENiagaraDistributionEditorMode::NonUniformRange:
+		return LOCTEXT("NonUniformRangeDisplayName", "Non-uniform Range");
+	case ENiagaraDistributionEditorMode::Curve:
+		return LOCTEXT("CurveDisplayName", "Curve");
+	case ENiagaraDistributionEditorMode::UniformCurve:
+		return LOCTEXT("UniformCurveDisplayName", "Uniform Curve");
+	case ENiagaraDistributionEditorMode::NonUniformCurve:
+		return LOCTEXT("NonUniformCurveDisplayName", "Non-uniform Curve");
+	default:
+		return LOCTEXT("UnknownDisplayName", "Unknown");
+	}
+}
+
+FText FNiagaraDistributionEditorUtilities::DistributionModeToToolTipText(ENiagaraDistributionEditorMode InMode)
+{
+	switch (InMode)
+	{
+	case ENiagaraDistributionEditorMode::Constant:
+		return LOCTEXT("ConstantToolTip", "A constant single value.");
+	case ENiagaraDistributionEditorMode::UniformConstant:
+		return LOCTEXT("UniformContstantToolTip", "A constant value applied to all value components.");
+	case ENiagaraDistributionEditorMode::NonUniformConstant:
+		return LOCTEXT("NonUniformConstantToolTip", "Constant values which can be different for each value component.");
+	case ENiagaraDistributionEditorMode::Range:
+		return LOCTEXT("RangeToolTip", "A single min/max range.");
+	case ENiagaraDistributionEditorMode::UniformRange:
+		return LOCTEXT("UniformRangeToolTip", "A min/max range applied to all value components.");
+	case ENiagaraDistributionEditorMode::NonUniformRange:
+		return LOCTEXT("NonUniformRangeToolTip", "Min/max ranges which can be different for each value component.");
+	case ENiagaraDistributionEditorMode::Curve:
+		return LOCTEXT("CurveToolTip", "This value is driven by a curve.");
+	case ENiagaraDistributionEditorMode::UniformCurve:
+		return LOCTEXT("UniformCurveToolTip", "All components of this value are driven by the same curve.");
+	case ENiagaraDistributionEditorMode::NonUniformCurve:
+		return LOCTEXT("NonUniformCurveToolTip", "Each component of this value is driven by its own curve.");
+	default:
+		return LOCTEXT("UnknownToolTip", "Unknown");
+	}
+}
+
+FName FNiagaraDistributionEditorUtilities::DistributionModeToIconBrushName(ENiagaraDistributionEditorMode InMode)
+{
+	switch (InMode)
+	{
+	case ENiagaraDistributionEditorMode::Constant:
+	case ENiagaraDistributionEditorMode::UniformConstant:
+		return "NiagaraEditor.DistributionEditor.UniformConstant";
+	case ENiagaraDistributionEditorMode::NonUniformConstant:
+		return "NiagaraEditor.DistributionEditor.NonUniformConstant";
+	case ENiagaraDistributionEditorMode::Range:
+	case ENiagaraDistributionEditorMode::UniformRange:
+		return "NiagaraEditor.DistributionEditor.UniformRange";
+	case ENiagaraDistributionEditorMode::NonUniformRange:
+		return "NiagaraEditor.DistributionEditor.NonUniformRange";
+	case ENiagaraDistributionEditorMode::Curve:
+	case ENiagaraDistributionEditorMode::UniformCurve:
+		return "NiagaraEditor.DistributionEditor.UniformCurve";
+	case ENiagaraDistributionEditorMode::NonUniformCurve:
+		return "NiagaraEditor.DistributionEditor.NonUniformCurve";
+	default:
+		return NAME_None;
+	}
+}
+
+const FSlateBrush* FNiagaraDistributionEditorUtilities::DistributionModeToIconBrush(ENiagaraDistributionEditorMode InMode)
+{
+	return FNiagaraEditorStyle::Get().GetBrush(DistributionModeToIconBrushName(InMode));
+}
+
+FSlateIcon FNiagaraDistributionEditorUtilities::DistributionModeToIcon(ENiagaraDistributionEditorMode InMode)
+{
+	return FSlateIcon(FNiagaraEditorStyle::Get().GetStyleSetName(), DistributionModeToIconBrushName(InMode));
+}
+
+bool FNiagaraDistributionEditorUtilities::IsUniform(ENiagaraDistributionEditorMode InMode)
+{
+	return
+		InMode == ENiagaraDistributionEditorMode::Constant ||
+		InMode == ENiagaraDistributionEditorMode::UniformConstant ||
+		InMode == ENiagaraDistributionEditorMode::Range ||
+		InMode == ENiagaraDistributionEditorMode::UniformRange ||
+		InMode == ENiagaraDistributionEditorMode::Curve ||
+		InMode == ENiagaraDistributionEditorMode::UniformCurve;
+}
+
+bool FNiagaraDistributionEditorUtilities::IsConstant(ENiagaraDistributionEditorMode InMode)
+{
+	return
+		InMode == ENiagaraDistributionEditorMode::Constant ||
+		InMode == ENiagaraDistributionEditorMode::UniformConstant ||
+		InMode == ENiagaraDistributionEditorMode::NonUniformConstant;
+}
+
+bool FNiagaraDistributionEditorUtilities::IsRange(ENiagaraDistributionEditorMode InMode)
+{
+	return
+		InMode == ENiagaraDistributionEditorMode::Range ||
+		InMode == ENiagaraDistributionEditorMode::UniformRange ||
+		InMode == ENiagaraDistributionEditorMode::NonUniformRange;
+}
+
+bool FNiagaraDistributionEditorUtilities::IsCurve(ENiagaraDistributionEditorMode InMode)
+{
+	return
+		InMode == ENiagaraDistributionEditorMode::Curve ||
+		InMode == ENiagaraDistributionEditorMode::UniformCurve ||
+		InMode == ENiagaraDistributionEditorMode::NonUniformCurve;
+}
+
+#undef LOCTEXT_NAMESPACE
\ No newline at end of file
diff --git a/Engine/Plugins/FX/Niagara/Source/NiagaraEditor/Private/Widgets/NiagaraDistributionEditorUtilities.h b/Engine/Plugins/FX/Niagara/Source/NiagaraEditor/Private/Widgets/NiagaraDistributionEditorUtilities.h
new file mode 100644
--- /dev/null
+++ b/Engine/Plugins/FX/Niagara/Source/NiagaraEditor/Private/Widgets/NiagaraDistributionEditorUtilities.h
@@ -0,0 +1,26 @@
+// Copyright Epic Games, Inc. All Rights Reserved.
+
+#pragma once
+
+#include "Widgets/INiagaraDistributionAdapter.h"
+
+#include "Internationalization/Text.h"
+#include "UObject/NameTypes.h"
+
+struct FSlateBrush;
+struct FSlateIcon;
+
+namespace FNiagaraDistributionEditorUtilities
+{
+	FText DistributionModeToDisplayName(ENiagaraDistributionEditorMode InMode);
+	FText DistributionModeToToolTipText(ENiagaraDistributionEditorMode InMode);
+
+	FName DistributionModeToIconBrushName(ENiagaraDistributionEditorMode InMode);
+	const FSlateBrush* DistributionModeToIconBrush(ENiagaraDistributionEditorMode InMode);
+	FSlateIcon DistributionModeToIcon(ENiagaraDistributionEditorMode InMode);
+
+	bool IsUniform(ENiagaraDistributionEditorMode InMode);
+	bool IsConstant(ENiagaraDistributionEditorMode InMode);
+	bool IsRange(ENiagaraDistributionEditorMode InMode);
+	bool IsCurve(ENiagaraDistributionEditorMode InMode);
+}
\ No newline at end of file
diff --git a/Engine/Plugins/FX/Niagara/Source/NiagaraEditor/Private/Widgets/SNiagaraDistributionCurveEditor.cpp b/Engine/Plugins/FX/Niagara/Source/NiagaraEditor/Private/Widgets/SNiagaraDistributionCurveEditor.cpp
new file mode 100644
--- /dev/null
+++ b/Engine/Plugins/FX/Niagara/Source/NiagaraEditor/Private/Widgets/SNiagaraDistributionCurveEditor.cpp
@@ -0,0 +1,880 @@
+// Copyright Epic Games, Inc. All Rights Reserved.
+
+#include "Widgets/SNiagaraDistributionCurveEditor.h"
+
+#include "NiagaraEditorStyle.h"
+#include "SEnumCombo.h"
+#include "Widgets/INiagaraDistributionAdapter.h"
+#include "Widgets/Input/SButton.h"
+#include "Widgets/Input/SNumericEntryBox.h"
+#include "Widgets/Layout/SBox.h"
+#include "Widgets/Layout/SGridPanel.h"
+#include "Widgets/SLeafWidget.h"
+
+#define LOCTEXT_NAMESPACE "NiagaraDistributionCurveEditor"
+
+class SNiagaraCompactCurveView : public SLeafWidget
+{
+	SLATE_BEGIN_ARGS(SNiagaraCompactCurveView)
+		: _Width(20)
+		, _Height(20)
+		, _CurveColor(FSlateColor(EStyleColor::White))
+		, _CurveStateSerialNumber(0)
+		{ }
+		SLATE_ARGUMENT(float, Width)
+		SLATE_ARGUMENT(float, Height)
+		SLATE_ARGUMENT(FSlateColor, CurveColor)
+		SLATE_ATTRIBUTE(FKeyHandle, SelectedKeyHandle)
+		SLATE_ATTRIBUTE(int32, CurveStateSerialNumber)
+	SLATE_END_ARGS()
+
+	void Construct(const FArguments& InArgs, const FRichCurve* InCurveToDisplay)
+	{
+		Width = InArgs._Width;
+		Height = InArgs._Height;
+		CurveColor = InArgs._CurveColor;
+		SelectedKeyHandle = InArgs._SelectedKeyHandle;
+		CurveStateSerialNumber = InArgs._CurveStateSerialNumber;
+		CurveToDisplay = InCurveToDisplay;
+
+		KeyColor = FSlateColor(EStyleColor::AccentGray);
+		SelectedKeyColor = FSlateColor(EStyleColor::Select);
+	}
+
+	virtual int32 OnPaint(const FPaintArgs& Args, const FGeometry& AllottedGeometry, const FSlateRect& MyCullingRect, FSlateWindowElementList& OutDrawElements, int32 LayerId, const FWidgetStyle& InWidgetStyle, bool bParentEnabled) const override
+	{
+		FKeyHandle NewSelectedKeyHandle = SelectedKeyHandle.Get();
+		int32 NewCurveStateSerialNumber = CurveStateSerialNumber.Get(0);
+
+		if (CurveStateSerialNumberCache != NewCurveStateSerialNumber || AllottedGeometryCache != AllottedGeometry)
+		{
+			SelectedKeyHandleCache = NewSelectedKeyHandle;
+			CurveStateSerialNumberCache = NewCurveStateSerialNumber;
+			AllottedGeometryCache = AllottedGeometry;
+			CacheLines(AllottedGeometry.Size.X, AllottedGeometry.Size.Y, true, true, true);
+		}
+		else if (SelectedKeyHandleCache != NewSelectedKeyHandle)
+		{
+			SelectedKeyHandleCache = NewSelectedKeyHandle;
+			CacheLines(AllottedGeometry.Size.X, AllottedGeometry.Size.Y, false, false, true);
+		}
+
+		FSlateDrawElement::MakeLines(OutDrawElements, LayerId, AllottedGeometry.ToPaintGeometry(), CurvePointsCache, ESlateDrawEffect::None, CurveColor.GetColor(InWidgetStyle), true, 2.0f);
+		for (const TArray<FVector2D>& KeyLinePoints : KeyLinesCache)
+		{
+			FSlateDrawElement::MakeLines(OutDrawElements, LayerId, AllottedGeometry.ToPaintGeometry(), KeyLinePoints, ESlateDrawEffect::None, KeyColor.GetColor(InWidgetStyle), true, 2.0f);
+		}
+		FSlateDrawElement::MakeLines(OutDrawElements, LayerId, AllottedGeometry.ToPaintGeometry(), SelectedKeyLineCache, ESlateDrawEffect::None, SelectedKeyColor.GetColor(InWidgetStyle), true, 3.0f);
+		return LayerId;
+	}
+
+	virtual FVector2D ComputeDesiredSize(float) const override
+	{
+		return FVector2D(Width, Height);
+	}
+
+private:
+	void CacheLines(float InWidth, float InHeight, bool bCacheCurve, bool bCacheKeys, bool bCacheSelectedKey) const
+	{
+		if (CurveToDisplay == nullptr)
+		{
+			return;
+		}
+
+		float TimeMin;
+		float TimeMax;
+		float ValueMin;
+		float ValueMax;
+		CurveToDisplay->GetTimeRange(TimeMin, TimeMax);
+		CurveToDisplay->GetValueRange(ValueMin, ValueMax);
+
+		float TimeRange = TimeMax != TimeMin ? TimeMax - TimeMin : 1;
+		float ValueRange = ValueMax != ValueMin ? ValueMax - ValueMin : 1;
+
+		if (bCacheCurve)
+		{
+			CurvePointsCache.Empty();
+			int32 Points = CurveToDisplay->GetNumKeys() * 20;
+			float TimeIncrement = TimeRange / (Points - 1);
+			for (int32 i = 0; i < Points; i++)
+			{
+				float Time = TimeMin + i * TimeIncrement;
+				float Value = CurveToDisplay->Eval(Time);
+
+				float NormalizedX = (Time - TimeMin) / TimeRange;
+				float NormalizedY = (Value - ValueMin) / ValueRange;
+				CurvePointsCache.Add(FVector2D(NormalizedX * InWidth, (1 - NormalizedY) * InHeight));
+			}
+		}
+
+		if (bCacheKeys)
+		{
+			KeyLinesCache.Empty();
+			for (auto KeyIterator(CurveToDisplay->GetKeyIterator()); KeyIterator; ++KeyIterator)
+			{
+				const FRichCurveKey& Key = *KeyIterator;
+				TArray<FVector2D>& KeyLinePoints = KeyLinesCache.AddDefaulted_GetRef();
+
+				float NormalizedX = (Key.Time - TimeMin) / TimeRange;
+				KeyLinePoints.Add(FVector2D(NormalizedX * InWidth, 0));
+				KeyLinePoints.Add(FVector2D(NormalizedX * InWidth, InHeight));
+			}
+		}
+
+		if (bCacheSelectedKey)
+		{
+			SelectedKeyLineCache.Empty();
+			for (auto KeyHandleIterator(CurveToDisplay->GetKeyHandleIterator()); KeyHandleIterator; ++KeyHandleIterator)
+			{
+				if (*KeyHandleIterator == SelectedKeyHandleCache)
+				{
+					const FRichCurveKey& SelectedKey = CurveToDisplay->GetKey(SelectedKeyHandleCache);
+					float NormalizedX = (SelectedKey.Time - TimeMin) / TimeRange;
+					SelectedKeyLineCache.Add(FVector2D(NormalizedX * InWidth, 0));
+					SelectedKeyLineCache.Add(FVector2D(NormalizedX * InWidth, InHeight));
+					break;
+				}
+			}
+		}
+	}
+
+	const FRichCurve* CurveToDisplay = nullptr;
+	mutable TArray<FVector2D> CurvePointsCache;
+	mutable TArray<TArray<FVector2D>> KeyLinesCache; 
+	mutable TArray<FVector2D> SelectedKeyLineCache;
+	mutable FGeometry AllottedGeometryCache;
+	mutable FKeyHandle SelectedKeyHandleCache;
+	mutable int32 CurveStateSerialNumberCache = 0;
+	float Width;
+	float Height;
+	FSlateColor CurveColor;
+	FSlateColor KeyColor;
+	FSlateColor SelectedKeyColor;
+	TAttribute<FKeyHandle> SelectedKeyHandle;
+	TAttribute<int32> CurveStateSerialNumber;
+};
+
+class SNiagaraCompactCurveKeyHandle : public SBorder
+{
+	DECLARE_DELEGATE_OneParam(FOnMoved, float /* LocalDelta */);
+
+	SLATE_BEGIN_ARGS(SNiagaraCompactCurveKeyHandle)
+		: _Size(10, 10)
+	{ }
+		SLATE_ARGUMENT(FVector2f, Size)
+		SLATE_ATTRIBUTE(bool, IsSelected)
+		SLATE_EVENT(FSimpleDelegate, OnSelected)
+		SLATE_EVENT(FOnMoved, OnMoved)
+	SLATE_END_ARGS()
+
+	void Construct(const FArguments& InArgs)
+	{
+		Size = InArgs._Size;
+		IsSelected = InArgs._IsSelected;
+		OnSelectedDelegate = InArgs._OnSelected;
+		OnMovedDelegate = InArgs._OnMoved;
+		SBorder::Construct(SBorder::FArguments()
+			.BorderImage(FNiagaraEditorStyle::Get().GetBrush("NiagaraEditor.DistributionEditor.CurveKeyHandle"))
+			.BorderBackgroundColor(this, &SNiagaraCompactCurveKeyHandle::GetHandleColor)
+			.Padding(0)
+			[
+				SNew(SBox)
+				.WidthOverride(Size.X)
+				.HeightOverride(Size.Y)
+			]);
+	}
+
+	virtual FReply OnMouseButtonDown(const FGeometry& MyGeometry, const FPointerEvent& MouseEvent) override
+	{
+		if (MouseEvent.GetEffectingButton() == EKeys::LeftMouseButton)
+		{
+			FVector2D MouseLocation = MouseEvent.GetScreenSpacePosition();
+			LastMouseLocation = MouseLocation.X;
+			return FReply::Handled().CaptureMouse(SharedThis(this));
+		}
+		return FReply::Unhandled();
+	}
+
+	virtual FReply OnMouseButtonUp(const FGeometry& MyGeometry, const FPointerEvent& MouseEvent) override
+	{
+		if (this->HasMouseCapture())
+		{
+			if (IsSelected.Get() == false)
+			{
+				OnSelectedDelegate.ExecuteIfBound();
+			}
+			return FReply::Handled().ReleaseMouseCapture();
+		}
+		return FReply::Unhandled();
+	}
+
+	virtual FReply OnMouseMove(const FGeometry& MyGeometry, const FPointerEvent& MouseEvent) override
+	{
+		if (this->HasMouseCapture())
+		{
+			FVector2D MouseLocation = MouseEvent.GetScreenSpacePosition();
+			OnMovedDelegate.ExecuteIfBound((MouseLocation.X - LastMouseLocation) / MyGeometry.Scale);
+			LastMouseLocation = MouseLocation.X;
+			return FReply::Handled();
+		}
+		return FReply::Unhandled();
+	}
+
+private:
+	FSlateColor GetHandleColor() const
+	{
+		return IsSelected.Get()
+			? FSlateColor(EStyleColor::Select)
+			: FSlateColor(EStyleColor::White);
+	}
+
+private:
+	FVector2f Size;
+	TAttribute<bool> IsSelected;
+	FSimpleDelegate OnSelectedDelegate;
+	FOnMoved OnMovedDelegate;
+	float LastMouseLocation;
+};
+
+class SNiagaraCompactCurveKeySelector : public SCompoundWidget
+{
+	DECLARE_DELEGATE_OneParam(FOnKeySelected, FKeyHandle /* SelectedKeyHandle */);
+	DECLARE_DELEGATE_TwoParams(FOnKeyMoved, FKeyHandle /* SelectedKeyHandle */, float /* TimeDelta */)
+
+	SLATE_BEGIN_ARGS(SNiagaraCompactCurveKeySelector)
+		: _HandleSize(10, 10)
+		, _CurveStateSerialNumber(0)
+	{ }
+		SLATE_ARGUMENT(FVector2f, HandleSize)
+		SLATE_ATTRIBUTE(FKeyHandle, SelectedKeyHandle)
+		SLATE_ATTRIBUTE(int32, CurveStateSerialNumber)
+		SLATE_EVENT(FOnKeySelected, OnKeySelected)
+		SLATE_EVENT(FOnKeyMoved, OnKeyMoved)
+	SLATE_END_ARGS()
+
+private:
+	struct FKeyWidgetData
+	{
+		FKeyHandle Handle;
+		TSharedPtr<SNiagaraCompactCurveKeyHandle> HandleWidget;
+		FMargin BorderPadding;
+	};
+
+public:
+	void Construct(const FArguments& InArgs, FRichCurve* InEditCurve)
+	{
+		EditCurve = InEditCurve;
+		HandleSize = InArgs._HandleSize;
+		SelectedKeyHandle = InArgs._SelectedKeyHandle;
+		CurveStateSerialNumber = InArgs._CurveStateSerialNumber;
+		OnKeySelectedDelegate = InArgs._OnKeySelected;
+		OnKeyMovedDelegate = InArgs._OnKeyMoved;
+		LocalWidthCache = 100;
+
+		ChildSlot
+		[
+			SAssignNew(BorderOverlay, SOverlay)
+		];
+
+		UpdateKeyWidgets();
+	}
+
+	virtual void Tick(const FGeometry& AllottedGeometry, const double InCurrentTime, const float InDeltaTime) override
+	{
+		SCompoundWidget::Tick(AllottedGeometry, InCurrentTime, InDeltaTime);
+		int32 CurrentCurveStateSerialNumber = CurveStateSerialNumber.Get();
+		LocalWidthCache = AllottedGeometry.GetLocalSize().X - HandleSize.X;
+		if (CachedCurveStateSerialNumber != CurrentCurveStateSerialNumber || AllottedGeometry != AllottedGeometryCache)
+		{
+			CachedCurveStateSerialNumber = CurrentCurveStateSerialNumber;
+			UpdateKeyWidgets();
+		}
+		AllottedGeometryCache = AllottedGeometry;
+	}
+
+private:
+	void UpdateKeyWidgets()
+	{
+		if (KeyWidgetData.Num() != EditCurve->GetNumKeys())
+		{
+			BorderOverlay->ClearChildren();
+			KeyWidgetData.Empty();
+			for (auto KeyHandleIterator = EditCurve->GetKeyHandleIterator(); KeyHandleIterator; KeyHandleIterator++)
+			{
+				int32 DataIndex = KeyWidgetData.Num();
+				FKeyWidgetData& KeyWidgetDataItem = KeyWidgetData.AddDefaulted_GetRef();
+				KeyWidgetDataItem.Handle = *KeyHandleIterator;
+				TAttribute<FMargin> BorderPadding = TAttribute<FMargin>::CreateSP(this, &SNiagaraCompactCurveKeySelector::GetKeyBorderPadding, DataIndex);
+				BorderOverlay->AddSlot()
+				.Padding(BorderPadding)
+				.HAlign(HAlign_Left)
+				[
+					SAssignNew(KeyWidgetDataItem.HandleWidget, SNiagaraCompactCurveKeyHandle)
+					.Size(HandleSize)
+					.IsSelected(this, &SNiagaraCompactCurveKeySelector::GetKeyIsSelected, KeyWidgetDataItem.Handle)
+					.OnSelected(this, &SNiagaraCompactCurveKeySelector::KeyHandleSelected, KeyWidgetDataItem.Handle)
+					.OnMoved(this, &SNiagaraCompactCurveKeySelector::KeyHandleMoved, KeyWidgetDataItem.Handle)
+				];
+			}
+		}
+
+		float TimeMin;
+		float TimeMax;
+		float ValueMin;
+		float ValueMax;
+		EditCurve->GetTimeRange(TimeMin, TimeMax);
+		EditCurve->GetValueRange(ValueMin, ValueMax);
+		float TimeRange = TimeMax != TimeMin ? TimeMax - TimeMin : 1;
+		float ValueRange = ValueMax != ValueMin ? ValueMax - ValueMin : 1;
+		for (FKeyWidgetData& KeyWidgetDataItem : KeyWidgetData)
+		{
+			const FRichCurveKey& Key = EditCurve->GetKey(KeyWidgetDataItem.Handle);
+			float NormalizedX = (Key.Time - TimeMin) / TimeRange;
+			KeyWidgetDataItem.BorderPadding = FMargin(LocalWidthCache * NormalizedX, 0, 0, 0);
+		}
+	}
+
+	FMargin GetKeyBorderPadding(int32 KeyWidgetDataIndex) const
+	{
+		return KeyWidgetDataIndex < KeyWidgetData.Num() ? KeyWidgetData[KeyWidgetDataIndex].BorderPadding : FMargin(0);
+	}
+
+	bool GetKeyIsSelected(FKeyHandle KeyHandle) const
+	{
+		return KeyHandle == SelectedKeyHandle.Get();
+	}
+
+	void KeyHandleSelected(FKeyHandle KeyHandle)
+	{
+		OnKeySelectedDelegate.ExecuteIfBound(KeyHandle);
+	}
+
+	void KeyHandleMoved(float LocalDelta, FKeyHandle KeyHandle)
+	{
+		float TimeMin;
+		float TimeMax;
+		EditCurve->GetTimeRange(TimeMin, TimeMax);
+		float TimeRange = TimeMax - TimeMin;
+		float NormalizedDelta = LocalDelta / LocalWidthCache;
+		if (SelectedKeyHandle.Get() != KeyHandle)
+		{
+			OnKeySelectedDelegate.ExecuteIfBound(KeyHandle);
+		}
+		OnKeyMovedDelegate.ExecuteIfBound(KeyHandle, NormalizedDelta * TimeRange);
+	}
+
+private:
+	FRichCurve* EditCurve = nullptr;
+	FVector2f HandleSize;
+	TAttribute<FKeyHandle> SelectedKeyHandle;
+	TAttribute<int32> CurveStateSerialNumber;
+	FOnKeySelected OnKeySelectedDelegate;
+	FOnKeyMoved OnKeyMovedDelegate;
+
+	FGeometry AllottedGeometryCache;
+	float LocalWidthCache;
+
+	TSharedPtr<SOverlay> BorderOverlay;
+	int32 CachedCurveStateSerialNumber = INDEX_NONE;
+	TArray<FKeyWidgetData> KeyWidgetData;
+};
+
+void SNiagaraDistributionCurveEditor::Construct(const FArguments& InArgs, TSharedRef<INiagaraDistributionAdapter> InDistributionAdapter, int32 InChannelIndex)
+{
+	DistributionAdapter = InDistributionAdapter;
+	ChannelIndex = InChannelIndex;
+	const FRichCurve* EditCurvePtr = DistributionAdapter->GetCurveValue(ChannelIndex);
+	if (EditCurvePtr != nullptr)
+	{
+		EditCurve = *EditCurvePtr;
+	}
+	else
+	{
+		EditCurve.AddKey(0, 0);
+	}
+	UpdateSelectedKeyIndex(0);
+
+	float KeyHandleWidth = 16;
+	float KeyHandleHeight = 16;
+	ChildSlot
+	[
+		SNew(SVerticalBox)
+		+ SVerticalBox::Slot()
+		.AutoHeight()
+		.Padding(0, 0, 0, 3)
+		[
+			SNew(SHorizontalBox)
+			+ SHorizontalBox::Slot()
+			.AutoWidth()
+			.Padding(0, 0, 3, 0)
+			[
+				SNew(SVerticalBox)
+				+ SVerticalBox::Slot()
+				.AutoHeight()
+				.Padding(0, 0, 0, 2)
+				[
+					SNew(SNumericEntryBox<int32>)
+					.Font(FAppStyle::Get().GetFontStyle("PropertyWindow.NormalFont"))
+					.Value(this, &SNiagaraDistributionCurveEditor::GetSelectedKeyIndex)
+					.OnValueChanged(this, &SNiagaraDistributionCurveEditor::SelectedKeyIndexChanged)
+					.OnValueCommitted(this, &SNiagaraDistributionCurveEditor::SelectedKeyIndexComitted)
+					.AllowSpin(true)
+					.MinValue(0)
+					.MaxValue(this, &SNiagaraDistributionCurveEditor::GetMaxKeyIndex)
+					.MinSliderValue(0)
+					.MaxSliderValue(this, &SNiagaraDistributionCurveEditor::GetMaxKeyIndex)
+					.Delta(1)
+					.MinDesiredValueWidth(8)
+					.LabelVAlign(VAlign_Center)
+					.Label()
+					[
+						SNew(STextBlock)
+						.Font(FAppStyle::Get().GetFontStyle("PropertyWindow.NormalFont"))
+						.Text(LOCTEXT("KeyLabel", "Key"))
+					]
+				]
+				+ SVerticalBox::Slot()
+				.AutoHeight()
+				[
+					SNew(SHorizontalBox)
+					+ SHorizontalBox::Slot()
+					.AutoWidth()
+					.Padding(0, 0, 2, 0)
+					[
+						SNew(SButton)
+						.ButtonStyle( FAppStyle::Get(), "SimpleButton" )
+						.OnClicked(this, &SNiagaraDistributionCurveEditor::AddKeyButtonClicked)
+						.ToolTipText(LOCTEXT("AddKeyToolTip", "Add a key"))
+						.ContentPadding(0.0f)
+						[ 
+							SNew( SImage )
+							.Image( FAppStyle::GetBrush("Icons.PlusCircle") )
+							.ColorAndOpacity( FSlateColor::UseForeground() )
+						]
+					]
+					+ SHorizontalBox::Slot()
+					.AutoWidth()
+					.Padding(0, 0, 2, 0)
+					[
+						SNew(SButton)
+						.ButtonStyle(FAppStyle::Get(), "SimpleButton")
+						.OnClicked(this, &SNiagaraDistributionCurveEditor::DeleteKeyButtonClicked)
+						.ToolTipText(LOCTEXT("DeleteKeyToolTip", "Delete this key"))
+						.ContentPadding(0.0f)
+						[
+							SNew(SImage)
+							.Image(FAppStyle::GetBrush("Icons.Delete"))
+							.ColorAndOpacity(FSlateColor::UseForeground())
+						]
+					]
+				]
+			]
+			// Curve view and key selector.
+			+ SHorizontalBox::Slot()
+			.Padding(2)
+			[
+				SNew(SGridPanel)
+				.FillRow(0, 1)
+				.FillColumn(0, 1)
+				// Value Labels
+				+ SGridPanel::Slot(1, 0)
+				[
+					SNew(SVerticalBox)
+					+ SVerticalBox::Slot()
+					.VAlign(VAlign_Top)
+					[
+						SNew(STextBlock)
+						.Font(FAppStyle::Get().GetFontStyle("SmallFont"))
+						.Text(this, &SNiagaraDistributionCurveEditor::GetCurveValueMaxText)
+					]
+					+ SVerticalBox::Slot()
+					.VAlign(VAlign_Bottom)
+					[
+						SNew(STextBlock)
+						.Font(FAppStyle::Get().GetFontStyle("SmallFont"))
+						.Text(this, &SNiagaraDistributionCurveEditor::GetCurveValueMinText)
+					]
+				]
+				// Curve view
+				+ SGridPanel::Slot(0, 0)
+				.Padding(FMargin(KeyHandleWidth / 2, 3, KeyHandleWidth / 2, 3))
+				[
+					SNew(SNiagaraCompactCurveView, &EditCurve)
+					.Height(50)
+					.CurveColor(InArgs._CurveColor)
+					.SelectedKeyHandle(this, &SNiagaraDistributionCurveEditor::GetSelectedKeyHandle)
+					.CurveStateSerialNumber(this, &SNiagaraDistributionCurveEditor::GetCurveStateSerialNumber)
+				]
+				// Time Labels
+				+ SGridPanel::Slot(0, 2)
+				[
+					SNew(SHorizontalBox)
+					+ SHorizontalBox::Slot()
+					.HAlign(HAlign_Left)
+					[
+						SNew(STextBlock)
+						.Font(FAppStyle::Get().GetFontStyle("SmallFont"))
+						.Text(this, &SNiagaraDistributionCurveEditor::GetCurveTimeMinText)
+					]
+					+ SHorizontalBox::Slot()
+					.HAlign(HAlign_Right)
+					[
+						SNew(STextBlock)
+						.Font(FAppStyle::Get().GetFontStyle("SmallFont"))
+						.Text(this, &SNiagaraDistributionCurveEditor::GetCurveTimeMaxText)
+					]
+				]
+				// key selector
+				+ SGridPanel::Slot(0, 1)
+				.Padding(0, 0, 0, 3)
+				[
+					SNew(SNiagaraCompactCurveKeySelector, &EditCurve)
+					.HandleSize(FVector2f(KeyHandleWidth, KeyHandleHeight))
+					.SelectedKeyHandle(this, &SNiagaraDistributionCurveEditor::GetSelectedKeyHandle)
+					.CurveStateSerialNumber(this, &SNiagaraDistributionCurveEditor::GetCurveStateSerialNumber)
+					.OnKeySelected(this, &SNiagaraDistributionCurveEditor::OnKeySelectorKeySelected)
+					.OnKeyMoved(this, &SNiagaraDistributionCurveEditor::OnKeySelectorKeyMoved)
+				]
+			]
+		]
+		+ SVerticalBox::Slot()
+		.AutoHeight()
+		[
+			SNew(SHorizontalBox)
+			+ SHorizontalBox::Slot()
+			.Padding(0, 0, 5, 0)
+			[
+				SNew(SNumericEntryBox<float>)
+				.Font(FAppStyle::Get().GetFontStyle("PropertyWindow.NormalFont"))
+				.Value(this, &SNiagaraDistributionCurveEditor::GetSelectedKeyTime)
+				.OnValueChanged(this, &SNiagaraDistributionCurveEditor::SelectedKeyTimeChanged)
+				.OnValueCommitted(this, &SNiagaraDistributionCurveEditor::SelectedKeyTimeComitted)
+				.OnBeginSliderMovement(this, &SNiagaraDistributionCurveEditor::BeginSliderMovement)
+				.OnEndSliderMovement(this, &SNiagaraDistributionCurveEditor::EndSliderMovement)
+				.AllowSpin(true)
+				.MinValue(TOptional<float>())
+				.MaxValue(TOptional<float>())
+				.MinSliderValue(TOptional<float>())
+				.MaxSliderValue(TOptional<float>())
+				.LabelVAlign(VAlign_Center)
+				.Label()
+				[
+					SNew(STextBlock)
+					.Font(FAppStyle::Get().GetFontStyle("PropertyWindow.NormalFont"))
+					.Text(LOCTEXT("TimeLabel", "Time"))
+				]
+			]
+			+ SHorizontalBox::Slot()
+			.Padding(0, 0, 5, 0)
+			[
+				SNew(SNumericEntryBox<float>)
+				.Font(FAppStyle::Get().GetFontStyle("PropertyWindow.NormalFont"))
+				.Value(this, &SNiagaraDistributionCurveEditor::GetSelectedKeyValue)
+				.OnValueChanged(this, &SNiagaraDistributionCurveEditor::SelectedKeyValueChanged)
+				.OnValueCommitted(this, &SNiagaraDistributionCurveEditor::SelectedKeyValueComitted)
+				.OnBeginSliderMovement(this, &SNiagaraDistributionCurveEditor::BeginSliderMovement)
+				.OnEndSliderMovement(this, &SNiagaraDistributionCurveEditor::EndSliderMovement)
+				.AllowSpin(true)
+				.MinValue(TOptional<float>())
+				.MaxValue(TOptional<float>())
+				.MinSliderValue(TOptional<float>())
+				.MaxSliderValue(TOptional<float>())
+				.LabelVAlign(VAlign_Center)
+				.Label()
+				[
+					SNew(STextBlock)
+					.Font(FAppStyle::Get().GetFontStyle("PropertyWindow.NormalFont"))
+					.Text(LOCTEXT("ValueLabel", "Value"))
+				]
+			]
+			+ SHorizontalBox::Slot()
+			[
+				SNew(SHorizontalBox)
+				+ SHorizontalBox::Slot()
+				.VAlign(VAlign_Center)
+				.AutoWidth()
+				.Padding(0, 0, 3, 0)
+				[
+					SNew(STextBlock)
+					.Font(FAppStyle::Get().GetFontStyle("PropertyWindow.NormalFont"))
+					.Text(LOCTEXT("InterpLabel", "Interp"))
+				]
+				+ SHorizontalBox::Slot()
+				[
+					SNew(SEnumComboBox, StaticEnum<ERichCurveInterpMode>())
+					.Font(FAppStyle::Get().GetFontStyle("PropertyWindow.NormalFont"))
+					.ContentPadding(FMargin(2, 0))
+					.CurrentValue(this, &SNiagaraDistributionCurveEditor::GetSelectedKeyInterpMode)
+					.OnEnumSelectionChanged(this, &SNiagaraDistributionCurveEditor::SelectedKeyInterpModeChanged)
+				]
+			]
+		]
+	];
+}
+
+TOptional<int32> SNiagaraDistributionCurveEditor::GetSelectedKeyIndex() const
+{
+	return SelectedKeyIndex;
+}
+
+TOptional<int32> SNiagaraDistributionCurveEditor::GetMaxKeyIndex() const
+{
+	return EditCurve.GetNumKeys() - 1;
+}
+
+void SNiagaraDistributionCurveEditor::SelectedKeyIndexChanged(int32 InValue)
+{
+	UpdateSelectedKeyIndex(InValue);
+}
+
+void SNiagaraDistributionCurveEditor::SelectedKeyIndexComitted(int32 InValue, ETextCommit::Type CommitInfo)
+{
+	if (CommitInfo == ETextCommit::OnEnter || CommitInfo == ETextCommit::OnUserMovedFocus)
+	{
+		UpdateSelectedKeyIndex(InValue);
+	}
+}
+
+void SNiagaraDistributionCurveEditor::UpdateSelectedKeyIndex(int32 NewKeyIndex)
+{
+	SelectedKeyIndex = NewKeyIndex;
+	int32 KeyIndex = 0;
+	for (auto KeyHandleIterator = EditCurve.GetKeyHandleIterator(); KeyHandleIterator; KeyHandleIterator++)
+	{
+		if (KeyIndex == NewKeyIndex)
+		{
+			SelectedKeyHandle = *KeyHandleIterator;
+			break;
+		}
+		KeyIndex++;
+	}
+	CurveStateSerialNumber++;
+}
+
+TOptional<float> SNiagaraDistributionCurveEditor::GetSelectedKeyTime() const
+{
+	UpdateCachedValues();
+	return KeyTimeCache;
+}
+
+void SNiagaraDistributionCurveEditor::SelectedKeyTimeChanged(float InValue)
+{
+	UpdateSelectedKeyTime(InValue);
+}
+
+void SNiagaraDistributionCurveEditor::SelectedKeyTimeComitted(float InValue, ETextCommit::Type CommitInfo)
+{
+	if (CommitInfo == ETextCommit::OnEnter || CommitInfo == ETextCommit::OnUserMovedFocus)
+	{
+		UpdateSelectedKeyTime(InValue);
+	}
+}
+
+void SNiagaraDistributionCurveEditor::UpdateSelectedKeyTime(float NewKeyTime)
+{
+	if (GetSelectedKeyTime() != NewKeyTime)
+	{
+		EditCurve.SetKeyTime(SelectedKeyHandle, NewKeyTime);
+		DistributionAdapter->SetCurveValue(ChannelIndex, EditCurve);
+
+		int32 NewKeyIndex = INDEX_NONE;
+		int32 KeyIndex = 0;
+		for (auto KeyHandleIterator = EditCurve.GetKeyHandleIterator(); KeyHandleIterator; KeyHandleIterator++)
+		{
+			if (SelectedKeyHandle == *KeyHandleIterator)
+			{
+				NewKeyIndex = KeyIndex;
+				break;
+			}
+			KeyIndex++;
+		}
+
+		SelectedKeyIndex = NewKeyIndex;
+		KeyTimeCache = NewKeyTime;
+		CurveStateSerialNumber++;
+	}
+}
+
+TOptional<float> SNiagaraDistributionCurveEditor::GetSelectedKeyValue() const
+{
+	UpdateCachedValues();
+	return KeyValueCache;
+}
+
+void SNiagaraDistributionCurveEditor::SelectedKeyValueChanged(float InValue)
+{
+	UpdateSelectedKeyValue(InValue);
+}
+
+void SNiagaraDistributionCurveEditor::SelectedKeyValueComitted(float InValue, ETextCommit::Type CommitInfo)
+{
+	if (CommitInfo == ETextCommit::OnEnter || CommitInfo == ETextCommit::OnUserMovedFocus)
+	{
+		UpdateSelectedKeyValue(InValue);
+	}
+}
+
+void SNiagaraDistributionCurveEditor::UpdateSelectedKeyValue(float NewKeyValue)
+{
+	if (GetSelectedKeyValue() != NewKeyValue)
+	{
+		EditCurve.SetKeyValue(SelectedKeyHandle, NewKeyValue);
+		DistributionAdapter->SetCurveValue(ChannelIndex, EditCurve);
+		KeyValueCache = NewKeyValue;
+		CurveStateSerialNumber++;
+	}
+}
+
+int32 SNiagaraDistributionCurveEditor::GetSelectedKeyInterpMode() const
+{
+	UpdateCachedValues();
+	return KeyInterpModeCache;
+}
+
+void SNiagaraDistributionCurveEditor::SelectedKeyInterpModeChanged(int32 InNewValue, ESelectInfo::Type InInfo)
+{
+	if (GetSelectedKeyInterpMode() != InNewValue)
+	{
+		EditCurve.SetKeyInterpMode(SelectedKeyHandle, (ERichCurveInterpMode)InNewValue);
+		DistributionAdapter->SetCurveValue(ChannelIndex, EditCurve);
+		KeyInterpModeCache = InNewValue;
+		CurveStateSerialNumber++;
+	}
+}
+
+void SNiagaraDistributionCurveEditor::BeginSliderMovement()
+{
+	DistributionAdapter->BeginContinuousChange();
+}
+
+void SNiagaraDistributionCurveEditor::EndSliderMovement(float Value)
+{
+	DistributionAdapter->EndContinuousChange();
+}
+
+FReply SNiagaraDistributionCurveEditor::AddKeyButtonClicked()
+{
+	int32 NumKeys = EditCurve.GetNumKeys();
+	int32 NewSelectedKeyIndex;
+	if (NumKeys == 0)
+	{
+		EditCurve.AddKey(0, 0);
+		NewSelectedKeyIndex = 0;
+	}
+	else if (NumKeys == 1)
+	{
+		FRichCurveKey FirstKey = EditCurve.GetFirstKey();
+		EditCurve.AddKey(FirstKey.Time + 1, FirstKey.Value);
+		NewSelectedKeyIndex = 1;
+	}
+	else
+	{
+		if (SelectedKeyIndex < NumKeys - 1)
+		{
+			FKeyHandle NextKeyHandle = EditCurve.GetNextKey(SelectedKeyHandle);
+			FRichCurveKey SelectedKey = EditCurve.GetKey(SelectedKeyHandle);
+			FRichCurveKey NextKey = EditCurve.GetKey(NextKeyHandle);
+			EditCurve.AddKey((SelectedKey.Time + NextKey.Time) / 2, (SelectedKey.Value + NextKey.Value) / 2);
+			NewSelectedKeyIndex = SelectedKeyIndex + 1;
+		}
+		else
+		{
+			FKeyHandle PreviousKeyHandle = EditCurve.GetPreviousKey(SelectedKeyHandle);
+			FRichCurveKey SelectedKey = EditCurve.GetKey(SelectedKeyHandle);
+			FRichCurveKey PreviousKey = EditCurve.GetKey(PreviousKeyHandle);
+			EditCurve.AddKey(SelectedKey.Time + (SelectedKey.Time - PreviousKey.Time), SelectedKey.Value + (SelectedKey.Value - PreviousKey.Value));
+			NewSelectedKeyIndex = SelectedKeyIndex;
+		}
+	}
+
+	DistributionAdapter->SetCurveValue(ChannelIndex, EditCurve);
+	UpdateSelectedKeyIndex(NewSelectedKeyIndex);
+	CurveStateSerialNumber++;
+	return FReply::Handled();
+}
+
+FReply SNiagaraDistributionCurveEditor::DeleteKeyButtonClicked()
+{
+	if (EditCurve.GetNumKeys() > 1)
+	{
+		EditCurve.DeleteKey(GetSelectedKeyHandle());
+		DistributionAdapter->SetCurveValue(ChannelIndex, EditCurve);
+		int32 NewKeyIndex = FMath::Clamp(SelectedKeyIndex - 1, 0, EditCurve.GetNumKeys() - 1);
+		UpdateSelectedKeyIndex(NewKeyIndex);
+		CurveStateSerialNumber++;
+		return FReply::Handled();
+	}
+	return FReply::Unhandled();
+}
+
+void SNiagaraDistributionCurveEditor::OnKeySelectorKeySelected(FKeyHandle InSelectedKeyHandle)
+{
+	int32 KeyIndex = 0;
+	for (auto KeyHandleIterator = EditCurve.GetKeyHandleIterator(); KeyHandleIterator; KeyHandleIterator++)
+	{
+		if (*KeyHandleIterator == InSelectedKeyHandle)
+		{
+			SelectedKeyHandle = *KeyHandleIterator;
+			SelectedKeyIndex = KeyIndex;
+			CurveStateSerialNumber++;
+			break;
+		}
+		KeyIndex++;
+	}
+}
+
+void SNiagaraDistributionCurveEditor::OnKeySelectorKeyMoved(FKeyHandle InSelectedKeyHandle, float TimeDelta)
+{
+	if (SelectedKeyHandle == InSelectedKeyHandle)
+	{
+		UpdateSelectedKeyTime(GetSelectedKeyTime().GetValue() + TimeDelta);
+	}
+}
+
+FText SNiagaraDistributionCurveEditor::GetCurveTimeMinText() const
+{
+	UpdateCachedValues();
+	return CurveTimeMinTextCache;
+}
+
+FText SNiagaraDistributionCurveEditor::GetCurveTimeMaxText() const
+{
+	UpdateCachedValues();
+	return CurveTimeMaxTextCache;
+}
+
+FText SNiagaraDistributionCurveEditor::GetCurveValueMinText() const
+{
+	UpdateCachedValues();
+	return CurveValueMinTextCache;
+}
+
+FText SNiagaraDistributionCurveEditor::GetCurveValueMaxText() const
+{
+	UpdateCachedValues();
+	return CurveValueMaxTextCache;
+}
+
+void SNiagaraDistributionCurveEditor::UpdateCachedValues() const
+{
+	if (CacheCurveStateSerialNumber != CurveStateSerialNumber)
+	{
+		CacheCurveStateSerialNumber = CurveStateSerialNumber;
+		KeyTimeCache = EditCurve.GetKeyTime(SelectedKeyHandle);
+		KeyValueCache = EditCurve.GetKeyValue(SelectedKeyHandle);
+		KeyInterpModeCache = (int32)EditCurve.GetKeyInterpMode(SelectedKeyHandle);
+
+		float MinTime;
+		float MaxTime;
+		float MinValue;
+		float MaxValue;
+		EditCurve.GetTimeRange(MinTime, MaxTime);
+		EditCurve.GetValueRange(MinValue, MaxValue);
+		CurveTimeMinTextCache = FText::AsNumber(MinTime);
+		CurveTimeMaxTextCache = FText::AsNumber(MaxTime);
+		CurveValueMinTextCache = FText::AsNumber(MinValue);
+		CurveValueMaxTextCache = FText::AsNumber(MaxValue);
+	}
+}
+
+#undef LOCTEXT_NAMESPACE
\ No newline at end of file
diff --git a/Engine/Plugins/FX/Niagara/Source/NiagaraEditor/Private/Widgets/SNiagaraDistributionCurveEditor.h b/Engine/Plugins/FX/Niagara/Source/NiagaraEditor/Private/Widgets/SNiagaraDistributionCurveEditor.h
new file mode 100644
--- /dev/null
+++ b/Engine/Plugins/FX/Niagara/Source/NiagaraEditor/Private/Widgets/SNiagaraDistributionCurveEditor.h
@@ -0,0 +1,80 @@
+// Copyright Epic Games, Inc. All Rights Reserved.
+
+#pragma once
+
+#include "Curves/RichCurve.h"
+#include "Styling/StyleColors.h"
+#include "Widgets/SCompoundWidget.h"
+
+class INiagaraDistributionAdapter;
+
+class SNiagaraDistributionCurveEditor : public SCompoundWidget
+{
+public:
+	SLATE_BEGIN_ARGS(SNiagaraDistributionCurveEditor)
+		: _CurveColor(FSlateColor(EStyleColor::White))
+		{ }
+		SLATE_EVENT(FSimpleDelegate, OnCurveChanged)
+		SLATE_ARGUMENT(FSlateColor, CurveColor)
+	SLATE_END_ARGS();
+
+	void Construct(const FArguments& InArgs, TSharedRef<INiagaraDistributionAdapter> InDistributionAdapter, int32 InChannelIndex);
+
+private:
+	TOptional<int32> GetSelectedKeyIndex() const;
+	TOptional<int32> GetMaxKeyIndex() const;
+	void SelectedKeyIndexChanged(int32 InValue);
+	void SelectedKeyIndexComitted(int32 InValue, ETextCommit::Type CommitInfo);
+	void UpdateSelectedKeyIndex(int32 NewKeyIndex);
+
+	TOptional<float> GetSelectedKeyTime() const;
+	void SelectedKeyTimeChanged(float InValue);
+	void SelectedKeyTimeComitted(float InValue, ETextCommit::Type CommitInfo);
+	void UpdateSelectedKeyTime(float NewKeyTime);
+
+	TOptional<float> GetSelectedKeyValue() const;
+	void SelectedKeyValueChanged(float InValue);
+	void SelectedKeyValueComitted(float InValue, ETextCommit::Type CommitInfo);
+	void UpdateSelectedKeyValue(float NewKeyValue);
+
+	int32 GetSelectedKeyInterpMode() const;
+	void SelectedKeyInterpModeChanged(int32 InNewValue, ESelectInfo::Type InInfo);
+
+	void BeginSliderMovement();
+	void EndSliderMovement(float Value);
+
+	FReply AddKeyButtonClicked();
+	FReply DeleteKeyButtonClicked();
+
+	FKeyHandle GetSelectedKeyHandle() const { return SelectedKeyHandle; }
+	int32 GetCurveStateSerialNumber() const { return CurveStateSerialNumber; }
+
+	void OnKeySelectorKeySelected(FKeyHandle InSelectedKeyHandle);
+	void OnKeySelectorKeyMoved(FKeyHandle InSelectedKeyHandle, float TimeDelta);
+
+	FText GetCurveTimeMinText() const;
+	FText GetCurveTimeMaxText() const;
+	FText GetCurveValueMinText() const;
+	FText GetCurveValueMaxText() const;
+
+	void UpdateCachedValues() const;
+
+private:
+	TSharedPtr<INiagaraDistributionAdapter> DistributionAdapter;
+	int32 ChannelIndex;
+
+	FRichCurve EditCurve;
+	FKeyHandle SelectedKeyHandle;
+
+	int32 SelectedKeyIndex = INDEX_NONE;
+	int32 CurveStateSerialNumber = 0;
+
+	mutable int32 CacheCurveStateSerialNumber = INDEX_NONE;
+	mutable TOptional<float> KeyTimeCache;
+	mutable TOptional<float> KeyValueCache;
+	mutable int32 KeyInterpModeCache;
+	mutable FText CurveTimeMinTextCache;
+	mutable FText CurveTimeMaxTextCache;
+	mutable FText CurveValueMinTextCache;
+	mutable FText CurveValueMaxTextCache;
+};
\ No newline at end of file
diff --git a/Engine/Plugins/FX/Niagara/Source/NiagaraEditor/Private/Widgets/SNiagaraDistributionEditor.cpp b/Engine/Plugins/FX/Niagara/Source/NiagaraEditor/Private/Widgets/SNiagaraDistributionEditor.cpp
new file mode 100644
--- /dev/null
+++ b/Engine/Plugins/FX/Niagara/Source/NiagaraEditor/Private/Widgets/SNiagaraDistributionEditor.cpp
@@ -0,0 +1,377 @@
+// Copyright Epic Games, Inc. All Rights Reserved.
+
+#include "Widgets/SNiagaraDistributionEditor.h"
+
+#include "Framework/Commands/UIAction.h"
+#include "Framework/MultiBox/MultiBoxBuilder.h"
+#include "NiagaraEditorStyle.h"
+#include "Textures/SlateIcon.h"
+#include "Widgets/Input/SComboButton.h"
+#include "Widgets/Input/SNumericEntryBox.h"
+#include "Widgets/Input/SVectorInputBox.h"
+#include "Widgets/Layout/SBox.h"
+#include "Widgets/Layout/SGridPanel.h"
+#include "Widgets/NiagaraDistributionEditorUtilities.h"
+#include "Widgets/SNiagaraDistributionCurveEditor.h"
+
+#define LOCTEXT_NAMESPACE "NiagaraDistributionEditor"
+
+class SNiagaraDistributionModeSelector : public SComboButton
+{
+public:
+	SLATE_BEGIN_ARGS(SNiagaraDistributionModeSelector) { }
+		SLATE_EVENT(FSimpleDelegate, OnDistributionModeChanged)
+	SLATE_END_ARGS();
+
+	void Construct(const FArguments& InArgs, TSharedRef<INiagaraDistributionAdapter> InDistributionAdapter)
+	{
+		
+		DistributionAdapter = InDistributionAdapter;
+		OnDistributionChangedDelegate = InArgs._OnDistributionModeChanged;
+
+		SComboButton::Construct(SComboButton::FArguments()
+			.ButtonStyle(FAppStyle::Get(), "HoverHintOnly")
+			.ContentPadding(FMargin(0))
+			.OnGetMenuContent(this, &SNiagaraDistributionModeSelector::OnGetMenuContent)
+			.ButtonContent()
+			[
+				SNew(SImage)
+				.Image(this, &SNiagaraDistributionModeSelector::GetModeIcon)
+				.ToolTipText(this, &SNiagaraDistributionModeSelector::GetModeToolTip)
+			]);
+	}
+
+private:
+	const FSlateBrush* GetModeIcon() const
+	{
+		UpdateCachedValues();
+		return ModeIconCache;
+	}
+
+	FText GetModeToolTip() const
+	{
+		UpdateCachedValues();
+		return ModeToolTipCache;
+	}
+
+	TSharedRef<SWidget> OnGetMenuContent()
+	{
+		FMenuBuilder MenuBuilder(true, nullptr);
+		TArray<ENiagaraDistributionEditorMode> SupportedModes;
+		DistributionAdapter->GetSupportedDistributionModes(SupportedModes);
+		for (ENiagaraDistributionEditorMode SupportedMode : SupportedModes)
+		{
+			MenuBuilder.AddMenuEntry(
+				FNiagaraDistributionEditorUtilities::DistributionModeToDisplayName(SupportedMode),
+				FNiagaraDistributionEditorUtilities::DistributionModeToToolTipText(SupportedMode),
+				FNiagaraDistributionEditorUtilities::DistributionModeToIcon(SupportedMode),
+				FUIAction(
+					FExecuteAction::CreateSP(this, &SNiagaraDistributionModeSelector::DistributionModeSelected, SupportedMode),
+					FCanExecuteAction(),
+					FIsActionChecked::CreateSP(this, &SNiagaraDistributionModeSelector::IsDistributionModeSelected, SupportedMode)
+				),
+				NAME_None,
+				EUserInterfaceActionType::RadioButton
+			);
+		}
+
+		return MenuBuilder.MakeWidget();
+	}
+
+	void UpdateCachedValues() const
+	{
+		ENiagaraDistributionEditorMode CurrentMode = DistributionAdapter->GetDistributionMode();
+		if (ModeCache.IsSet() == false || ModeCache.GetValue() != CurrentMode)
+		{
+			ModeCache = CurrentMode;
+			ModeIconCache = FNiagaraDistributionEditorUtilities::DistributionModeToIconBrush(CurrentMode);
+			ModeToolTipCache = FText::Format(LOCTEXT("ModeSelectorToolTip", "{0} - {1}"),
+				FNiagaraDistributionEditorUtilities::DistributionModeToDisplayName(CurrentMode),
+				FNiagaraDistributionEditorUtilities::DistributionModeToToolTipText(CurrentMode));
+		}
+	}
+
+	bool IsDistributionModeSelected(ENiagaraDistributionEditorMode InSelectedMode) const
+	{
+		UpdateCachedValues();
+		return ModeCache == InSelectedMode;
+	}
+
+	void DistributionModeSelected(ENiagaraDistributionEditorMode InSelectedMode)
+	{
+		DistributionAdapter->SetDistributionMode(InSelectedMode);
+		UpdateCachedValues();
+		OnDistributionChangedDelegate.ExecuteIfBound();
+	}
+
+private:
+	TSharedPtr<INiagaraDistributionAdapter> DistributionAdapter;
+	FSimpleDelegate OnDistributionChangedDelegate;
+
+	mutable TOptional<ENiagaraDistributionEditorMode> ModeCache;
+	mutable const FSlateBrush* ModeIconCache = nullptr;
+	mutable FText ModeToolTipCache;
+};
+
+
+class SNiagaraDistributionValueEditor : public SCompoundWidget
+{
+public:
+	SLATE_BEGIN_ARGS(SNiagaraDistributionValueEditor) { }
+	SLATE_END_ARGS();
+
+	void Construct(const FArguments& InArgs, TSharedRef<INiagaraDistributionAdapter> InDistributionAdapter)
+	{
+		DistributionAdapter = InDistributionAdapter;
+		bool bIsUniform = FNiagaraDistributionEditorUtilities::IsUniform(DistributionAdapter->GetDistributionMode());
+		bool bIsConstant = FNiagaraDistributionEditorUtilities::IsConstant(DistributionAdapter->GetDistributionMode());
+
+		FText MinLabelText = LOCTEXT("MinLabel", "Min");
+		FText MaxLabelText = LOCTEXT("MaxLabel", "Max");
+
+		TSharedPtr<SWidget> ChildWidget;
+		if (bIsUniform)
+		{
+			if (bIsConstant)
+			{
+				ChildWidget = ConstructFloatWidget(0, 0, FText());
+			}
+			else
+			{
+				ChildWidget = SNew(SHorizontalBox)
+					+ SHorizontalBox::Slot()
+					[
+						ConstructFloatWidget(0, 0, MinLabelText)
+					]
+					+ SHorizontalBox::Slot()
+					[
+						ConstructFloatWidget(0, 1, MaxLabelText)
+					];
+			}
+		}
+		else
+		{
+			if (bIsConstant)
+			{
+				ChildWidget = ConstructVectorWidget(DistributionAdapter->GetNumChannels(), 0);
+			}
+			else
+			{
+				ChildWidget = SNew(SGridPanel)
+					.FillColumn(1, 1)
+					+ SGridPanel::Slot(0, 0)
+					.VAlign(VAlign_Center)
+					.Padding(0, 0, 5, 3)
+					[
+						SNew(STextBlock)
+						.Text(MinLabelText)
+					]
+					+ SGridPanel::Slot(1, 0)
+					.Padding(0, 0, 0, 3)
+					[
+						ConstructVectorWidget(DistributionAdapter->GetNumChannels(), 0)
+					]
+					+ SGridPanel::Slot(0, 1)
+					.VAlign(VAlign_Center)
+					.Padding(0, 0, 5, 0)
+					[
+						SNew(STextBlock)
+						.Text(MaxLabelText)
+					]
+					+ SGridPanel::Slot(1, 1)
+					[
+						ConstructVectorWidget(DistributionAdapter->GetNumChannels(), 1)
+					];
+			}
+		}
+
+		ChildSlot
+		[
+			ChildWidget.ToSharedRef()
+		];
+	}
+
+private:
+	TSharedRef<SWidget> ConstructFloatWidget(int32 ChannelIndex, int32 ValueIndex, FText LabelText)
+	{
+		TSharedRef<SWidget> LabelWidget = LabelText.IsEmpty()
+			? SNullWidget::NullWidget
+			: SNew(STextBlock)
+				.TextStyle(FNiagaraEditorStyle::Get(), "NiagaraEditor.ParameterText")
+				.Text(LabelText);
+
+		return SNew(SNumericEntryBox<float>)
+			.Font(FAppStyle::Get().GetFontStyle("PropertyWindow.NormalFont"))
+			.Value(this, &SNiagaraDistributionValueEditor::GetValue, ChannelIndex, ValueIndex)
+			.OnValueChanged(this, &SNiagaraDistributionValueEditor::ValueChanged, ChannelIndex, ValueIndex)
+			.OnValueCommitted(this, &SNiagaraDistributionValueEditor::ValueCommitted, ChannelIndex, ValueIndex)
+			.OnBeginSliderMovement(this, &SNiagaraDistributionValueEditor::BeginValueSliderMovement)
+			.OnEndSliderMovement(this, &SNiagaraDistributionValueEditor::EndValueSliderMovement)
+			.AllowSpin(true)
+			.MinValue(TOptional<float>())
+			.MaxValue(TOptional<float>())
+			.MinSliderValue(TOptional<float>())
+			.MaxSliderValue(TOptional<float>())
+			.BroadcastValueChangesPerKey(false)
+			.LabelVAlign(EVerticalAlignment::VAlign_Center)
+			.MinDesiredValueWidth(30)
+			.Label()
+			[
+				LabelWidget
+			];
+	}
+
+	typedef SNumericVectorInputBox<float, UE::Math::TVector2<float>, 2> SNumericVectorInputBox2;
+	typedef SNumericVectorInputBox<float, UE::Math::TVector<float>, 3> SNumericVectorInputBox3;
+	typedef SNumericVectorInputBox<float, UE::Math::TVector4<float>, 4> SNumericVectorInputBox4;
+
+	TSharedRef<SWidget> ConstructVectorWidget(int32 ChannelCount, int32 ValueIndex)
+	{
+		if (ChannelCount == 2)
+		{
+			return SNew(SNumericVectorInputBox2)
+				.Font(FAppStyle::Get().GetFontStyle("PropertyWindow.NormalFont"))
+				.AllowSpin(true)
+				.bColorAxisLabels(true)
+				.X(this, &SNiagaraDistributionValueEditor::GetValue, 0, ValueIndex)
+				.OnXChanged(this, &SNiagaraDistributionValueEditor::ValueChanged, 0, ValueIndex)
+				.OnXCommitted(this, &SNiagaraDistributionValueEditor::ValueCommitted, 0, ValueIndex)
+				.Y(this, &SNiagaraDistributionValueEditor::GetValue, 1, ValueIndex)
+				.OnYChanged(this, &SNiagaraDistributionValueEditor::ValueChanged, 1, ValueIndex)
+				.OnYCommitted(this, &SNiagaraDistributionValueEditor::ValueCommitted, 1, ValueIndex)
+				.OnBeginSliderMovement(this, &SNiagaraDistributionValueEditor::BeginValueSliderMovement)
+				.OnEndSliderMovement(this, &SNiagaraDistributionValueEditor::EndValueSliderMovement);
+		}
+		if (ChannelCount == 3)
+		{
+			return SNew(SNumericVectorInputBox3)
+				.Font(FAppStyle::Get().GetFontStyle("PropertyWindow.NormalFont"))
+				.AllowSpin(true)
+				.bColorAxisLabels(true)
+				.X(this, &SNiagaraDistributionValueEditor::GetValue, 0, ValueIndex)
+				.OnXChanged(this, &SNiagaraDistributionValueEditor::ValueChanged, 0, ValueIndex)
+				.OnXCommitted(this, &SNiagaraDistributionValueEditor::ValueCommitted, 0, ValueIndex)
+				.Y(this, &SNiagaraDistributionValueEditor::GetValue, 1, ValueIndex)
+				.OnYChanged(this, &SNiagaraDistributionValueEditor::ValueChanged, 1, ValueIndex)
+				.OnYCommitted(this, &SNiagaraDistributionValueEditor::ValueCommitted, 1, ValueIndex)
+				.Z(this, &SNiagaraDistributionValueEditor::GetValue, 2, ValueIndex)
+				.OnZChanged(this, &SNiagaraDistributionValueEditor::ValueChanged, 2, ValueIndex)
+				.OnZCommitted(this, &SNiagaraDistributionValueEditor::ValueCommitted, 2, ValueIndex)
+				.OnBeginSliderMovement(this, &SNiagaraDistributionValueEditor::BeginValueSliderMovement)
+				.OnEndSliderMovement(this, &SNiagaraDistributionValueEditor::EndValueSliderMovement);
+		}
+		if (ChannelCount == 4)
+		{
+			return SNew(SNumericVectorInputBox4)
+				.Font(FAppStyle::Get().GetFontStyle("PropertyWindow.NormalFont"))
+				.AllowSpin(true)
+				.bColorAxisLabels(true)
+				.X(this, &SNiagaraDistributionValueEditor::GetValue, 0, ValueIndex)
+				.OnXChanged(this, &SNiagaraDistributionValueEditor::ValueChanged, 0, ValueIndex)
+				.OnXCommitted(this, &SNiagaraDistributionValueEditor::ValueCommitted, 0, ValueIndex)
+				.Y(this, &SNiagaraDistributionValueEditor::GetValue, 1, ValueIndex)
+				.OnYChanged(this, &SNiagaraDistributionValueEditor::ValueChanged, 1, ValueIndex)
+				.OnYCommitted(this, &SNiagaraDistributionValueEditor::ValueCommitted, 1, ValueIndex)
+				.Z(this, &SNiagaraDistributionValueEditor::GetValue, 2, ValueIndex)
+				.OnZChanged(this, &SNiagaraDistributionValueEditor::ValueChanged, 2, ValueIndex)
+				.OnZCommitted(this, &SNiagaraDistributionValueEditor::ValueCommitted, 2, ValueIndex)
+				.W(this, &SNiagaraDistributionValueEditor::GetValue, 3, ValueIndex)
+				.OnWChanged(this, &SNiagaraDistributionValueEditor::ValueChanged, 3, ValueIndex)
+				.OnWCommitted(this, &SNiagaraDistributionValueEditor::ValueCommitted, 3, ValueIndex)
+				.OnBeginSliderMovement(this, &SNiagaraDistributionValueEditor::BeginValueSliderMovement)
+				.OnEndSliderMovement(this, &SNiagaraDistributionValueEditor::EndValueSliderMovement);
+		}
+		return SNullWidget::NullWidget;
+	}
+
+	TOptional<float> GetValue(int32 ChannelIndex, int32 ValueIndex) const
+	{
+		return DistributionAdapter->GetConstantOrRangeValue(ChannelIndex, ValueIndex);
+	}
+
+	void ValueChanged(float Value, int32 ChannelIndex, int32 ValueIndex)
+	{
+		DistributionAdapter->SetConstantOrRangeValue(ChannelIndex, ValueIndex, Value);
+	}
+	
+	void ValueCommitted(float Value, ETextCommit::Type CommitInfo, int32 ChannelIndex, int32 ValueIndex)
+	{
+		if (CommitInfo == ETextCommit::OnEnter || CommitInfo == ETextCommit::OnUserMovedFocus)
+		{
+			DistributionAdapter->SetConstantOrRangeValue(ChannelIndex, ValueIndex, Value);
+		}
+	}
+
+	void BeginValueSliderMovement()
+	{
+		DistributionAdapter->BeginContinuousChange();
+	}
+
+	void EndValueSliderMovement(float Value)
+	{
+		DistributionAdapter->EndContinuousChange();
+	}
+
+private:
+	TSharedPtr<INiagaraDistributionAdapter> DistributionAdapter;
+};
+
+void SNiagaraDistributionEditor::Construct(const FArguments& InArgs, TSharedRef<INiagaraDistributionAdapter> InDistributionAdapter)
+{
+	DistributionAdapter = InDistributionAdapter;
+
+	ChildSlot
+	[
+		SNew(SHorizontalBox)
+		+ SHorizontalBox::Slot()
+		.AutoWidth()
+		.Padding(0, 0, 10, 0)
+		[
+			SNew(SNiagaraDistributionModeSelector, InDistributionAdapter)
+			.OnDistributionModeChanged(this, &SNiagaraDistributionEditor::OnDistributionModeChanged)
+		]
+		+ SHorizontalBox::Slot()
+		[
+			SAssignNew(ContentBox, SBox)
+			[
+				ConstructContentForMode()
+			]
+		]
+	];
+}
+
+void  SNiagaraDistributionEditor::OnDistributionModeChanged()
+{
+	ContentBox->SetContent(ConstructContentForMode());
+}
+
+TSharedRef<SWidget> SNiagaraDistributionEditor::ConstructContentForMode()
+{
+	ENiagaraDistributionEditorMode Mode = DistributionAdapter->GetDistributionMode();
+	if (FNiagaraDistributionEditorUtilities::IsConstant(Mode) || FNiagaraDistributionEditorUtilities::IsRange(Mode))
+	{
+		return SNew(SNiagaraDistributionValueEditor, DistributionAdapter.ToSharedRef());
+	}
+	else if (FNiagaraDistributionEditorUtilities::IsCurve(Mode))
+	{
+		TSharedRef<SVerticalBox> CurveBox = SNew(SVerticalBox);
+		int32 NumberOfChannels = FNiagaraDistributionEditorUtilities::IsUniform(Mode) ? 1 : DistributionAdapter->GetNumChannels();
+		for (int32 ChannelIndex = 0; ChannelIndex < NumberOfChannels; ChannelIndex++)
+		{
+			CurveBox->AddSlot()
+			.AutoHeight()
+			.Padding(0, 0, 0, 3)
+			[
+				SNew(SNiagaraDistributionCurveEditor, DistributionAdapter.ToSharedRef(), ChannelIndex)
+				.CurveColor(DistributionAdapter->GetChannelColor(ChannelIndex))
+			];
+		}
+		return CurveBox;
+	}
+	else
+	{
+		return SNullWidget::NullWidget;
+	}
+}
+
+#undef LOCTEXT_NAMESPACE
\ No newline at end of file
diff --git a/Engine/Plugins/FX/Niagara/Source/NiagaraEditor/Private/Widgets/SNiagaraDistributionEditor.h b/Engine/Plugins/FX/Niagara/Source/NiagaraEditor/Private/Widgets/SNiagaraDistributionEditor.h
new file mode 100644
--- /dev/null
+++ b/Engine/Plugins/FX/Niagara/Source/NiagaraEditor/Private/Widgets/SNiagaraDistributionEditor.h
@@ -0,0 +1,29 @@
+// Copyright Epic Games, Inc. All Rights Reserved.
+
+#pragma once
+
+#include "Widgets/INiagaraDistributionAdapter.h"
+#include "Widgets/SCompoundWidget.h"
+
+struct FSlateBrush;
+struct FSlateIcon;
+class INiagaraDistributionAdapter;
+class SBox;
+
+class SNiagaraDistributionEditor : public SCompoundWidget
+{
+public:
+	SLATE_BEGIN_ARGS(SNiagaraDistributionEditor) { }
+	SLATE_END_ARGS();
+
+	void Construct(const FArguments& InArgs, TSharedRef<INiagaraDistributionAdapter> InDistributionAdapter);
+
+private:
+	void OnDistributionModeChanged();
+
+	TSharedRef<SWidget> ConstructContentForMode();
+
+private:
+	TSharedPtr<INiagaraDistributionAdapter> DistributionAdapter;
+	TSharedPtr<SBox> ContentBox;
+};
\ No newline at end of file
diff --git a/Engine/Plugins/FX/Niagara/Source/NiagaraShader/Internal/Stateless/NiagaraStatelessModuleShaderParameters.h b/Engine/Plugins/FX/Niagara/Source/NiagaraShader/Internal/Stateless/NiagaraStatelessModuleShaderParameters.h
--- a/Engine/Plugins/FX/Niagara/Source/NiagaraShader/Internal/Stateless/NiagaraStatelessModuleShaderParameters.h
+++ b/Engine/Plugins/FX/Niagara/Source/NiagaraShader/Internal/Stateless/NiagaraStatelessModuleShaderParameters.h
@@ -8,29 +8,25 @@
 namespace NiagaraStateless
 {
 	BEGIN_SHADER_PARAMETER_STRUCT(FInitializeParticleModule_ShaderParameters, NIAGARASHADER_API)
-		SHADER_PARAMETER(FVector3f, InitializeParticle_Position)				// Unset / Direct Set / Simulation Position
+		SHADER_PARAMETER(uint32,	InitializeParticle_ModuleFlags)
+		SHADER_PARAMETER(FVector3f,	InitializeParticle_Position)				// Unset / Direct Set / Simulation Position
 
-		SHADER_PARAMETER(FVector4f, InitializeParticle_ColorScale)				// Unset / Direct Set / Random Range(Link RGBA / Link RGB|Link A / Random Channels) / Random Hue|Saturation|Value
-		SHADER_PARAMETER(FVector4f, InitializeParticle_ColorBias)
+		SHADER_PARAMETER(FVector4f,	InitializeParticle_ColorScale)				// Unset / Direct Set / Random Range(Link RGBA / Link RGB|Link A / Random Channels) / Random Hue|Saturation|Value
+		SHADER_PARAMETER(FVector4f,	InitializeParticle_ColorBias)
 
 		SHADER_PARAMETER(FVector2f, InitializeParticle_SpriteSizeScale)			// Unset / Uniform / Random Uniform / Non-Uniform / Random Non-Uniform
 		SHADER_PARAMETER(FVector2f, InitializeParticle_SpriteSizeBias)
-		SHADER_PARAMETER(float, InitializeParticle_SpriteRotationScale)				// Unset / Random / Direct Set Deg / Direct Set Normalized
-		SHADER_PARAMETER(float, InitializeParticle_SpriteRotationBias)
+		SHADER_PARAMETER(float,		InitializeParticle_SpriteRotationScale)				// Unset / Random / Direct Set Deg / Direct Set Normalized
+		SHADER_PARAMETER(float,		InitializeParticle_SpriteRotationBias)
 		//SHADER_PARAMETER(FVector3f, InitializeParticle_SpriteUVMode)			// Unset / Random / Random X / Random Y / Random XY / Direct Set
 
-		SHADER_PARAMETER(FVector3f, InitializeParticle_MeshScaleScale)			// Unset / Uniform / Random Uniform / Non-Uniform / Random Non-Uniform
-		SHADER_PARAMETER(FVector3f, InitializeParticle_MeshScaleBias)
+		SHADER_PARAMETER(FVector3f,	InitializeParticle_MeshScaleScale)			// Unset / Uniform / Random Uniform / Non-Uniform / Random Non-Uniform
+		SHADER_PARAMETER(FVector3f,	InitializeParticle_MeshScaleBias)
 
-		SHADER_PARAMETER(float, InitializeParticle_RibbonWidthScale)		// Unset / Direct Set
-		SHADER_PARAMETER(float, InitializeParticle_RibbonWidthBias)
+		SHADER_PARAMETER(float,		InitializeParticle_RibbonWidthScale)		// Unset / Direct Set
+		SHADER_PARAMETER(float,		InitializeParticle_RibbonWidthBias)
 		//SHADER_PARAMETER(FVector3f, InitializeParticle_RibbonFacingVector)		// Unset / Direct Set
 		//SHADER_PARAMETER(FVector3f, InitializeParticle_RibbonTwist)				// Unset / Direct Set
-
-		SHADER_PARAMETER(uint32,	InitializeParticle_PositionMode)
-		SHADER_PARAMETER(uint32,	InitializeParticle_PositionTableMod)
-		SHADER_PARAMETER(uint32,	InitializeParticle_PositionTableOffset)
-		SHADER_PARAMETER(uint32,	InitializeParticle_PositionTableLength)
 	END_SHADER_PARAMETER_STRUCT()
 
 	BEGIN_SHADER_PARAMETER_STRUCT(FInitialMeshOrientationModule_ShaderParameters, NIAGARASHADER_API)
@@ -49,13 +45,11 @@ namespace NiagaraStateless
 	END_SHADER_PARAMETER_STRUCT()
 	
 	BEGIN_SHADER_PARAMETER_STRUCT(FScaleColorModule_ShaderParameters, NIAGARASHADER_API)
-		SHADER_PARAMETER(int32, ScaleColor_Offset)
-		SHADER_PARAMETER(int32, ScaleColor_Length)
+		SHADER_PARAMETER(FUintVector3,	ScaleColor_Distribution)
 	END_SHADER_PARAMETER_STRUCT()
 
 	BEGIN_SHADER_PARAMETER_STRUCT(FScaleMeshSizeModule_ShaderParameters, NIAGARASHADER_API)
-		SHADER_PARAMETER(int32,		ScaleMeshSize_Offset)
-		SHADER_PARAMETER(int32,		ScaleMeshSize_Length)
+		SHADER_PARAMETER(FUintVector3,	ScaleMeshSize_Distribution)
 	END_SHADER_PARAMETER_STRUCT()
 
 	BEGIN_SHADER_PARAMETER_STRUCT(FScaleMeshSizeBySpeedModule_ShaderParameters, NIAGARASHADER_API)
@@ -65,8 +59,7 @@ namespace NiagaraStateless
 	END_SHADER_PARAMETER_STRUCT()
 
 	BEGIN_SHADER_PARAMETER_STRUCT(FScaleSpriteSizeModule_ShaderParameters, NIAGARASHADER_API)
-		SHADER_PARAMETER(int32,		ScaleSpriteSize_Offset)
-		SHADER_PARAMETER(int32,		ScaleSpriteSize_Length)
+		SHADER_PARAMETER(FUintVector3, ScaleSpriteSize_Distribution)
 	END_SHADER_PARAMETER_STRUCT()
 
 	BEGIN_SHADER_PARAMETER_STRUCT(FScaleSpriteSizeBySpeedModule_ShaderParameters, NIAGARASHADER_API)
