* Added a fix for inaccurate quat interpolation in spawn scripts
* some translator code cleanup

diff --git a/Engine/Plugins/FX/Niagara/Shaders/Private/NiagaraCommon.ush b/Engine/Plugins/FX/Niagara/Shaders/Private/NiagaraCommon.ush
--- a/Engine/Plugins/FX/Niagara/Shaders/Private/NiagaraCommon.ush
+++ b/Engine/Plugins/FX/Niagara/Shaders/Private/NiagaraCommon.ush
@@ -93,30 +93,6 @@ float3x3 NiagaraFindRotationBetweenNormals(float3 A, float3 B)
 	return NiagaraQuatTo3x3(Quat);
 }
 
-float4 NiagaraQuatSLerp(float4 Quat1, float4 Quat2, float Slerp)
-{
-	const float RawCosom = dot(Quat1, Quat2);
-	const float Cosom = abs(RawCosom);
-	
-	float Scale0, Scale1;
-	if (Cosom < 0.9999f)
-	{
-		const float Omega = acos(Cosom);
-		const float InvSin = 1.f / sin(Omega);
-		Scale0 = sin((1.f - Slerp) * Omega) * InvSin;
-		Scale1 = sin(Slerp * Omega) * InvSin;
-	}
-	else
-	{
-		Scale0 = 1.0f - Slerp;
-		Scale1 = Slerp;
-	}
-
-	Scale1 = RawCosom >= 0.0f ? Scale1 : -Scale1;
-
-	return (Scale0 * Quat1) + (Scale1 * Quat2);
-}
-
 float4 NiagaraQuatMul(float4 Quat1, float4 Quat2)
 {
 	return float4(
diff --git a/Engine/Plugins/FX/Niagara/Shaders/Private/NiagaraEmitterInstanceShader.usf b/Engine/Plugins/FX/Niagara/Shaders/Private/NiagaraEmitterInstanceShader.usf
--- a/Engine/Plugins/FX/Niagara/Shaders/Private/NiagaraEmitterInstanceShader.usf
+++ b/Engine/Plugins/FX/Niagara/Shaders/Private/NiagaraEmitterInstanceShader.usf
@@ -460,6 +460,31 @@ int rand_int(int x, int Seed1, int Seed2, int Seed3)
 	return rand(x.x-1, Seed1, Seed2, Seed3);
 }
 
+// used to interpolate rotations in interpolated spawn scripts
+float4 NiagaraQuatSLerp(float4 Quat1, float4 Quat2, float Slerp)
+{
+	const float RawCosom = dot(Quat1, Quat2);
+	const float Cosom = abs(RawCosom);
+	
+	float Scale0, Scale1;
+	if (Cosom < 0.9999f)
+	{
+		const float Omega = acos(Cosom);
+		const float InvSin = 1.f / sin(Omega);
+		Scale0 = sin((1.f - Slerp) * Omega) * InvSin;
+		Scale1 = sin(Slerp * Omega) * InvSin;
+	}
+	else
+	{
+		Scale0 = 1.0f - Slerp;
+		Scale1 = Slerp;
+	}
+
+	Scale1 = RawCosom >= 0.0f ? Scale1 : -Scale1;
+
+	return (Scale0 * Quat1) + (Scale1 * Quat2);
+}
+
 /* -----------------------------------------------------------------
  * VM simulation function declarations
  * -----------------------------------------------------------------
diff --git a/Engine/Plugins/FX/Niagara/Source/Niagara/Classes/NiagaraScriptSourceBase.h b/Engine/Plugins/FX/Niagara/Source/Niagara/Classes/NiagaraScriptSourceBase.h
--- a/Engine/Plugins/FX/Niagara/Source/Niagara/Classes/NiagaraScriptSourceBase.h
+++ b/Engine/Plugins/FX/Niagara/Source/Niagara/Classes/NiagaraScriptSourceBase.h
@@ -72,6 +72,7 @@ public:
 	NIAGARA_API static const FString EventSpawnDefine;
 	NIAGARA_API static const FString EventSpawnInitialAttribWritesDefine;
 	NIAGARA_API static const FString ExperimentalVMDisabled;
+	NIAGARA_API static const FString AccurateQuatInterpolation;
 
 	FNiagaraCompileOptions() : TargetUsage(ENiagaraScriptUsage::Function), TargetUsageBitmask(0)
 	{
diff --git a/Engine/Plugins/FX/Niagara/Source/Niagara/Private/NiagaraCommon.cpp b/Engine/Plugins/FX/Niagara/Source/Niagara/Private/NiagaraCommon.cpp
--- a/Engine/Plugins/FX/Niagara/Source/Niagara/Private/NiagaraCommon.cpp
+++ b/Engine/Plugins/FX/Niagara/Source/Niagara/Private/NiagaraCommon.cpp
@@ -1651,6 +1651,7 @@ const FString FNiagaraCompileOptions::GpuScriptDefine = TEXT("GPUComputeSim");
 const FString FNiagaraCompileOptions::EventSpawnDefine = TEXT("EventSpawn");
 const FString FNiagaraCompileOptions::EventSpawnInitialAttribWritesDefine = TEXT("EventSpawnInitialAttribWrites");
 const FString FNiagaraCompileOptions::ExperimentalVMDisabled = TEXT("ExperimentalVMDisabled");
+const FString FNiagaraCompileOptions::AccurateQuatInterpolation = TEXT("AccurateQuatInterpolation");
 
 FSynchronizeWithParameterDefinitionsArgs::FSynchronizeWithParameterDefinitionsArgs()
 	: SpecificDefinitionsUniqueIds(TArray<FGuid>())
diff --git a/Engine/Plugins/FX/Niagara/Source/Niagara/Private/NiagaraScript.cpp b/Engine/Plugins/FX/Niagara/Source/Niagara/Private/NiagaraScript.cpp
--- a/Engine/Plugins/FX/Niagara/Source/Niagara/Private/NiagaraScript.cpp
+++ b/Engine/Plugins/FX/Niagara/Source/Niagara/Private/NiagaraScript.cpp
@@ -1097,6 +1097,10 @@ void UNiagaraScript::ComputeVMCompilationId(FNiagaraVMExecutableDataId& Id, cons
 	Id.bInterpolatedSpawn = false;
 	Id.bRequiresPersistentIDs = false;
 	Id.ScriptVersionID = IsVersioningEnabled() ? (VersionGuid.IsValid() ? VersionGuid : ExposedVersion) : FGuid();
+	if (GetDefault<UNiagaraSettings>()->bAccurateQuatInterpolation)
+	{
+		Id.AdditionalDefines.Add(FNiagaraCompileOptions::AccurateQuatInterpolation);
+	}
 	
 	TArray<FNiagaraVariableBase> ReferencedStaticVars;
 	TArray<const uint8*> ReferencedStaticVarValues;
diff --git a/Engine/Plugins/FX/Niagara/Source/Niagara/Public/NiagaraSettings.h b/Engine/Plugins/FX/Niagara/Source/Niagara/Public/NiagaraSettings.h
--- a/Engine/Plugins/FX/Niagara/Source/Niagara/Public/NiagaraSettings.h
+++ b/Engine/Plugins/FX/Niagara/Source/Niagara/Public/NiagaraSettings.h
@@ -168,6 +168,10 @@ class UNiagaraSettings : public UDeveloperSettings
 	UPROPERTY(config, EditAnywhere, Category = Niagara, meta = (DisplayName = "Enable Stateless Emitters (Experimental)", ConfigRestartRequired = true))
 	bool bStatelessEmittersEnabled = false;
 
+	/** If set to true, quaternion attributes will be interpolated via slerp instead of lerp in interpolated spawn scripts. */
+	UPROPERTY(config, EditAnywhere, AdvancedDisplay, Category = Niagara)
+	bool bAccurateQuatInterpolation = true;
+
 	/** Whether to limit the max tick delta time or not. */
 	UPROPERTY(config, EditAnywhere, BlueprintReadWrite, AdvancedDisplay, Category = "Niagara", meta = (InlineEditConditionToggle))
 	bool bLimitDeltaTime = true;
diff --git a/Engine/Plugins/FX/Niagara/Source/NiagaraEditor/Private/NiagaraHlslTranslator.cpp b/Engine/Plugins/FX/Niagara/Source/NiagaraEditor/Private/NiagaraHlslTranslator.cpp
--- a/Engine/Plugins/FX/Niagara/Source/NiagaraEditor/Private/NiagaraHlslTranslator.cpp
+++ b/Engine/Plugins/FX/Niagara/Source/NiagaraEditor/Private/NiagaraHlslTranslator.cpp
@@ -10,7 +10,6 @@
 #include "NiagaraCompilationPrivate.h"
 #include "NiagaraConstants.h"
 #include "NiagaraDataInterface.h"
-#include "NiagaraDataInterfaceVector2DCurve.h"
 #include "NiagaraEditorModule.h"
 #include "NiagaraEditorSettings.h"
 #include "NiagaraEditorTickables.h"
@@ -32,8 +31,6 @@
 #include "NiagaraNodeSelect.h"
 #include "NiagaraNodeStaticSwitch.h"
 #include "NiagaraNodeOp.h"
-#include "NiagaraParameterCollection.h"
-#include "NiagaraScriptSource.h"
 #include "NiagaraSettings.h"
 #include "NiagaraShared.h"
 #include "NiagaraSimulationStageBase.h"
@@ -613,7 +610,10 @@ void TNiagaraHlslTranslator<GraphBridge>::GenerateFunctionSignature(ENiagaraScri
 //////////////////////////////////////////////////////////////////////////
 
 FNiagaraHlslTranslator::FNiagaraHlslTranslator()
-	: CurrentBodyChunkMode(ENiagaraCodeChunkMode::Body)
+	: ReadIdx(0)
+	, WriteIdx(0)
+	, CompilationTarget(ENiagaraSimTarget::CPUSim)
+	, CurrentBodyChunkMode(ENiagaraCodeChunkMode::Body)
 	, ActiveStageIdx(-1)
 	, bInitializedDefaults(false)
 {
@@ -1187,6 +1187,7 @@ FNiagaraTranslateResults TNiagaraHlslTranslator<GraphBridge>::Translate(const FN
 	CompilationTarget = TranslationOptions.SimTarget;
 	TranslateResults.bHLSLGenSucceeded = false;
 	TranslateResults.OutputHLSL = "";
+	HlslOutput.Reserve(1024 * 1024); // reserve some space to prevent reallocations during translation
 
 	const FGraph* SourceGraph = GraphBridge::GetGraph(CompileDuplicateData);
 
@@ -2792,11 +2793,11 @@ void TNiagaraHlslTranslator<GraphBridge>::DefineInterpolatedParametersFunction(F
 
 						if (FoundNamespacedVar != nullptr)
 						{
-							FString FoundName = GetSanitizedSymbolName(FoundNamespacedVar->GetName().ToString());
 							FNiagaraCodeChunk& Chunk = CodeChunks[ChunkIdx];
 							if (ShouldInterpolateParameter(*FoundNamespacedVar))
 							{
-								HlslOutputString += TEXT("\tContext.") + PrevMap + TEXT(".") + FoundName + TEXT(" = lerp(") + INTERPOLATED_PARAMETER_PREFIX + Chunk.SymbolName + Chunk.ComponentMask + TEXT(", ") + Chunk.SymbolName + Chunk.ComponentMask + TEXT(", ") + TEXT("SpawnInterp);\n");
+								
+								HlslOutputString += GetInterpolateHlsl(*FoundNamespacedVar, PrevMap, Chunk);
 							}
 							else
 							{
@@ -4798,6 +4799,16 @@ bool FNiagaraHlslTranslator::ShouldInterpolateParameter(const FNiagaraVariable&
 	return true;
 }
 
+FString FNiagaraHlslTranslator::GetInterpolateHlsl(const FNiagaraVariable& Parameter, const FString& PrevMapName, const FNiagaraCodeChunk& Chunk) const
+{
+	FString FoundName = GetSanitizedSymbolName(Parameter.GetName().ToString());
+	if (Parameter.GetType() == FNiagaraTypeDefinition::GetQuatDef() && CompileOptions.AdditionalDefines.Contains(FNiagaraCompileOptions::AccurateQuatInterpolation))
+	{
+		return FString::Format(TEXT("\tContext.{0}.{1} = NiagaraQuatSLerp({2}{3}{4}, {3}{4}, SpawnInterp);\n"), {PrevMapName, FoundName, INTERPOLATED_PARAMETER_PREFIX, Chunk.SymbolName, Chunk.ComponentMask});
+	}
+	return FString::Format(TEXT("\tContext.{0}.{1} = lerp({2}{3}{4}, {3}{4}, SpawnInterp);\n"), {PrevMapName, FoundName, INTERPOLATED_PARAMETER_PREFIX, Chunk.SymbolName, Chunk.ComponentMask});
+}
+
 // specialization handling the case where the translation process treats the graph as mutable
 template<>
 void TNiagaraHlslTranslator<FNiagaraCompilationDigestBridge>::UpdateStaticSwitchConstants(const FPin* Pin)
@@ -5977,7 +5988,7 @@ bool TNiagaraHlslTranslator<GraphBridge>::GetLiteralConstantVariable(FNiagaraVar
 		FNiagaraVariable ResolvedLocalSpaceCompileOptionVar = ActiveHistoryForFunctionCalls.ResolveAliases(SYS_PARAM_EMITTER_LOCALSPACE);
 		if (CompileOptions.AdditionalDefines.Contains(ResolvedLocalSpaceCompileOptionVar.GetName().ToString()))
 		{
-			OutVar.SetValue(FVector3f(EForceInit::ForceInitToZero));
+			OutVar.SetValue(FVector3f(ForceInitToZero));
 			return true;
 		}
 	}
@@ -6727,7 +6738,7 @@ void TNiagaraHlslTranslator<GraphBridge>::ValidateFailIfPreviouslyNotSet(const F
 		{
 			for (const typename GraphBridge::FModuleScopedPin& ScopedPin : OtherOutputParamMapHistories[OtherParamIdx].PerVariableWriteHistory[FoundInParamIdx])
 			{
-				if (ScopedPin.Pin->Direction == EEdGraphPinDirection::EGPD_Input && ScopedPin.Pin->bHidden == false)
+				if (ScopedPin.Pin->Direction == EGPD_Input && ScopedPin.Pin->bHidden == false)
 				{
 					bSetPreviously = true;
 					break;
@@ -9189,7 +9200,7 @@ FString FNiagaraHlslTranslator::GenerateFunctionHlslPrototype(FStringView InVari
 	{
 		StringBuilder.Append(InVariableName);
 		StringBuilder.Append(TEXT("."));
-		StringBuilder.Append(FNiagaraHlslTranslator::GetSanitizedSymbolName(FunctionSignature.Name.ToString()));
+		StringBuilder.Append(GetSanitizedSymbolName(FunctionSignature.Name.ToString()));
 
 		// Build specifiers
 		if (FunctionSignature.FunctionSpecifiers.Num())
@@ -9229,9 +9240,9 @@ FString FNiagaraHlslTranslator::GenerateFunctionHlslPrototype(FStringView InVari
 				bNeedsComma = true;
 
 				StringBuilder.Append(TEXT("in "));
-				StringBuilder.Append(FNiagaraHlslTranslator::GetStructHlslTypeName(InputVar.GetType()));
+				StringBuilder.Append(GetStructHlslTypeName(InputVar.GetType()));
 				StringBuilder.Append(TEXT(" In_"));
-				StringBuilder.Append(FNiagaraHlslTranslator::GetSanitizedSymbolName(InputVar.GetName().ToString()));
+				StringBuilder.Append(GetSanitizedSymbolName(InputVar.GetName().ToString()));
 			}
 
 			// Outputs
@@ -9245,9 +9256,9 @@ FString FNiagaraHlslTranslator::GenerateFunctionHlslPrototype(FStringView InVari
 				bNeedsComma = true;
 
 				StringBuilder.Append(TEXT("out "));
-				StringBuilder.Append(FNiagaraHlslTranslator::GetStructHlslTypeName(OutputVar.GetType()));
+				StringBuilder.Append(GetStructHlslTypeName(OutputVar.GetType()));
 				StringBuilder.Append(TEXT(" Out_"));
-				StringBuilder.Append(FNiagaraHlslTranslator::GetSanitizedSymbolName(OutputVar.GetName().ToString()));
+				StringBuilder.Append(GetSanitizedSymbolName(OutputVar.GetName().ToString()));
 			}
 			StringBuilder.Append(TEXT(");"));
 		}
@@ -9320,7 +9331,7 @@ FString FNiagaraHlslTranslator::GetFunctionSignature(const FNiagaraFunctionSigna
 				}
 
 				FNiagaraVariable SimInput = ConvertToSimulationVariable(Input);
-				SigStr += FNiagaraHlslTranslator::GetStructHlslTypeName(SimInput.GetType()) + TEXT(" In_") + FNiagaraHlslTranslator::GetSanitizedSymbolName(Input.GetName().ToString(), true);
+				SigStr += GetStructHlslTypeName(SimInput.GetType()) + TEXT(" In_") + GetSanitizedSymbolName(Input.GetName().ToString(), true);
 				++ParamIdx;
 			}
 		}
@@ -9344,7 +9355,7 @@ FString FNiagaraHlslTranslator::GetFunctionSignature(const FNiagaraFunctionSigna
 				}
 
 				FNiagaraVariable SimOutput = ConvertToSimulationVariable(Output);
-				SigStr += TEXT("out ") + FNiagaraHlslTranslator::GetStructHlslTypeName(SimOutput.GetType()) + TEXT(" ") + FNiagaraHlslTranslator::GetSanitizedSymbolName(TEXT("Out_") + Output.GetName().ToString());
+				SigStr += TEXT("out ") + GetStructHlslTypeName(SimOutput.GetType()) + TEXT(" ") + GetSanitizedSymbolName(TEXT("Out_") + Output.GetName().ToString());
 				++ParamIdx;
 			}
 		}
diff --git a/Engine/Plugins/FX/Niagara/Source/NiagaraEditor/Private/NiagaraHlslTranslator.h b/Engine/Plugins/FX/Niagara/Source/NiagaraEditor/Private/NiagaraHlslTranslator.h
--- a/Engine/Plugins/FX/Niagara/Source/NiagaraEditor/Private/NiagaraHlslTranslator.h
+++ b/Engine/Plugins/FX/Niagara/Source/NiagaraEditor/Private/NiagaraHlslTranslator.h
@@ -466,6 +466,7 @@ protected:
 	virtual bool ParameterMapRegisterNamespaceAttributeVariable(const FNiagaraVariable& InVariable, int32 InParamMapHistoryIdx, int32& Output) = 0;
 
 	bool ShouldInterpolateParameter(const FNiagaraVariable& Parameter);
+	FString GetInterpolateHlsl(const FNiagaraVariable& Parameter, const FString& PrevMapName, const FNiagaraCodeChunk& Chunk) const;
 
 	bool IsBulkSystemScript() const;
 	bool IsSpawnScript() const;
