Make newly created emitter assets inheritable by default

diff --git a/Engine/Plugins/FX/Niagara/Source/NiagaraEditor/Private/AssetDefinitions/AssetDefinition_NiagaraEmitter.cpp b/Engine/Plugins/FX/Niagara/Source/NiagaraEditor/Private/AssetDefinitions/AssetDefinition_NiagaraEmitter.cpp
--- a/Engine/Plugins/FX/Niagara/Source/NiagaraEditor/Private/AssetDefinitions/AssetDefinition_NiagaraEmitter.cpp
+++ b/Engine/Plugins/FX/Niagara/Source/NiagaraEditor/Private/AssetDefinitions/AssetDefinition_NiagaraEmitter.cpp
@@ -97,28 +97,26 @@ namespace MenuExtension_NiagaraEmitter
 		const UContentBrowserAssetContextMenuContext* CBContext = UContentBrowserAssetContextMenuContext::FindContextWithAssets(InContext);
 		for (UNiagaraEmitter* Emitter : CBContext->LoadSelectedObjects<UNiagaraEmitter>())
 		{
-			if (Emitter->IsVersioningEnabled())
-			{
-				FString Name;
-				FString PackageName;
-				AssetToolsModule.Get().CreateUniqueAssetName(Emitter->GetOutermost()->GetName(), DefaultSuffix, PackageName, Name);
+			FString Name;
+			FString PackageName;
+			AssetToolsModule.Get().CreateUniqueAssetName(Emitter->GetOutermost()->GetName(), DefaultSuffix, PackageName, Name);
 
-				// Create the factory used to generate the asset
-				UNiagaraEmitter* NewEmitter = Cast<UNiagaraEmitter>(AssetToolsModule.Get().DuplicateAssetWithDialog(Name, FPackageName::GetLongPackagePath(PackageName), Emitter));
+			// Create the factory used to generate the asset
+			UNiagaraEmitter* NewEmitter = Cast<UNiagaraEmitter>(AssetToolsModule.Get().DuplicateAssetWithDialog(Name, FPackageName::GetLongPackagePath(PackageName), Emitter));
 
-				if (NewEmitter != nullptr)
-				{
-					GEditor->GetEditorSubsystem<UAssetEditorSubsystem>()->CloseAllEditorsForAsset(Emitter);
+			if (NewEmitter != nullptr)
+			{
+				GEditor->GetEditorSubsystem<UAssetEditorSubsystem>()->CloseAllEditorsForAsset(Emitter);
 
-					NewEmitter->Modify();
-					NewEmitter->GetLatestEmitterData()->GraphSource->MarkNotSynchronized(TEXT("Emitter created"));
+				NewEmitter->Modify();
+				NewEmitter->bIsInheritable = true;
+				NewEmitter->GetLatestEmitterData()->GraphSource->MarkNotSynchronized(TEXT("Emitter created"));
 
-					Emitter->Modify();
-					Emitter->SetParent(FVersionedNiagaraEmitter(NewEmitter, NewEmitter->GetExposedVersion().VersionGuid));
-					Emitter->GetLatestEmitterData()->GraphSource->MarkNotSynchronized(TEXT("Emitter parent changed"));
+				Emitter->Modify();
+				Emitter->SetParent(FVersionedNiagaraEmitter(NewEmitter, NewEmitter->GetExposedVersion().VersionGuid));
+				Emitter->GetLatestEmitterData()->GraphSource->MarkNotSynchronized(TEXT("Emitter parent changed"));
 
-					ObjectsToSync.Add(NewEmitter);
-				}
+				ObjectsToSync.Add(NewEmitter);
 			}
 		}
 
@@ -130,23 +128,6 @@ namespace MenuExtension_NiagaraEmitter
 			ContentBrowserModule.Get().SyncBrowserToAssets(ObjectsToSync);
 		}
 	}
-	
-	bool CanExecuteCreateDuplicateParent(const FToolMenuContext& InContext)
-	{
-		const static FName NAME_VersioningEnabled("VersioningEnabled");
-	
-		const UContentBrowserAssetContextMenuContext* CBContext = UContentBrowserAssetContextMenuContext::FindContextWithAssets(InContext);
-        for (const FAssetData& SelectedAsset : CBContext->SelectedAssets)
-        {
-			bool VersioningEnabled = false;
-        	if (SelectedAsset.GetTagValue<bool>(NAME_VersioningEnabled, VersioningEnabled) && !VersioningEnabled)
-        	{
-        		return false;
-        	}
-        }
- 
-		return true;
-	}
 
 	static void ExecuteMarkDependentCompilableAssetsDirty(const FToolMenuContext& InContext)
 	{
@@ -179,7 +160,6 @@ namespace MenuExtension_NiagaraEmitter
 
 					FToolUIAction UIAction;
 					UIAction.ExecuteAction = FToolMenuExecuteAction::CreateStatic(&ExecuteCreateDuplicateParent);
-					UIAction.CanExecuteAction = FToolMenuCanExecuteAction::CreateStatic(&CanExecuteCreateDuplicateParent);
 					InSection.AddMenuEntry("Emitter_CreateDuplicateParent", Label, ToolTip, Icon, UIAction);
 				}
 				{
diff --git a/Engine/Plugins/FX/Niagara/Source/NiagaraEditor/Private/NiagaraEditorUtilities.cpp b/Engine/Plugins/FX/Niagara/Source/NiagaraEditor/Private/NiagaraEditorUtilities.cpp
--- a/Engine/Plugins/FX/Niagara/Source/NiagaraEditor/Private/NiagaraEditorUtilities.cpp
+++ b/Engine/Plugins/FX/Niagara/Source/NiagaraEditor/Private/NiagaraEditorUtilities.cpp
@@ -2665,6 +2665,7 @@ void FNiagaraEditorUtilities::CreateAssetFromEmitter(TSharedRef<FNiagaraEmitterH
 	if (CreatedAsset != nullptr)
 	{
 		CreatedAsset->SetUniqueEmitterName(CreatedAsset->GetName());
+		CreatedAsset->bIsInheritable = true;
 
 		GEditor->GetEditorSubsystem<UAssetEditorSubsystem>()->OpenEditorForAsset(CreatedAsset);
 
