Added alternate search name for Niagara graph actions, only used by math ops currently. These will get evaluated separately for search scoring, meaning the altername name can result in an 'exact match', boosting the score.
Useful for cases such the action with name 'Divide' and alternate search name ' / '

diff --git a/Engine/Plugins/FX/Niagara/Source/NiagaraEditor/Private/EdGraphSchema_Niagara.cpp b/Engine/Plugins/FX/Niagara/Source/NiagaraEditor/Private/EdGraphSchema_Niagara.cpp
--- a/Engine/Plugins/FX/Niagara/Source/NiagaraEditor/Private/EdGraphSchema_Niagara.cpp
+++ b/Engine/Plugins/FX/Niagara/Source/NiagaraEditor/Private/EdGraphSchema_Niagara.cpp
@@ -488,7 +488,8 @@ TArray<TSharedPtr<FNiagaraAction_NewNode>> UEdGraphSchema_Niagara::GetGraphActio
 				// todo suggestion info per op?
 				UNiagaraNodeOp* OpNode = NewObject<UNiagaraNodeOp>(OwnerOfTemporaries);
 				OpNode->OpName = OpInfo.Name;
-				AddNewNodeMenuAction(NewActions, OpNode, OpInfo.FriendlyName, ENiagaraMenuSections::General, {OpInfo.Category.ToString()}, OpInfo.Description,  OpInfo.Keywords);
+				TSharedPtr<FNiagaraAction_NewNode> MathOpAction = AddNewNodeMenuAction(NewActions, OpNode, OpInfo.FriendlyName, ENiagaraMenuSections::General, {OpInfo.Category.ToString()}, OpInfo.Description,  OpInfo.Keywords);
+				MathOpAction->AlternateSearchName = OpInfo.AlternateSearchName.IsSet() ? FText::FromName(OpInfo.AlternateSearchName.GetValue()) : TOptional<FText>();
 			}
 		}
 	}
diff --git a/Engine/Plugins/FX/Niagara/Source/NiagaraEditor/Private/NiagaraEditorCommon.cpp b/Engine/Plugins/FX/Niagara/Source/NiagaraEditor/Private/NiagaraEditorCommon.cpp
--- a/Engine/Plugins/FX/Niagara/Source/NiagaraEditor/Private/NiagaraEditorCommon.cpp
+++ b/Engine/Plugins/FX/Niagara/Source/NiagaraEditor/Private/NiagaraEditorCommon.cpp
@@ -145,6 +145,7 @@ void FNiagaraOpInfo::Init()
 		Op->CompactName = FText::FromString(TEXT("+"));
 		Op->CompactNameFontSizeOverride = 30.f;
 		Op->Description = NSLOCTEXT("NiagaraOpInfo", "Add Desc", "Result = A + B");
+		Op->AlternateSearchName = FName("+");
 		Op->Keywords = FText::FromString(TEXT("+"));
 		Op->Inputs.Add(FNiagaraOpInOutInfo(A, NumericType, AText, AText, DefaultStr_Zero));
 		Op->Inputs.Add(FNiagaraOpInOutInfo(B, NumericType, BText, BText, DefaultStr_Zero));
@@ -198,6 +199,7 @@ void FNiagaraOpInfo::Init()
 		Op->CompactName = FText::FromString(TEXT("-"));
 		Op->CompactNameFontSizeOverride = 28.f;
 		Op->Description = NSLOCTEXT("NiagaraOpInfo", "Subtract Desc", "Result = A - B");
+		Op->AlternateSearchName = FName("-");
 		Op->Keywords = FText::FromString(TEXT("-"));
 		Op->Inputs.Add(FNiagaraOpInOutInfo(A, NumericType, AText, AText, DefaultStr_Zero));
 		Op->Inputs.Add(FNiagaraOpInOutInfo(B, NumericType, BText, BText, DefaultStr_Zero));
@@ -255,6 +257,7 @@ void FNiagaraOpInfo::Init()
 		Op->CompactName = FText::FromString(TEXT("\xD7"));
 		Op->CompactNameFontSizeOverride = 30.f;
 		Op->Description = NSLOCTEXT("NiagaraOpInfo", "Multiply Desc", "Result = A * B");
+		Op->AlternateSearchName = FName("*");
 		Op->Keywords = FText::FromString(TEXT("*"));
 		Op->Inputs.Add(FNiagaraOpInOutInfo(A, NumericType, AText, AText, DefaultStr_One));
 		Op->Inputs.Add(FNiagaraOpInOutInfo(B, NumericType, BText, BText, DefaultStr_One));
@@ -289,6 +292,7 @@ void FNiagaraOpInfo::Init()
 		Op->CompactName = FText::FromString(TEXT("\xF7"));
 		Op->CompactNameFontSizeOverride = 25.f;
 		Op->Description = NSLOCTEXT("NiagaraOpInfo", "Divide Desc", "Result = A / B");
+		Op->AlternateSearchName = FName("/");
 		Op->Keywords = FText::FromString(TEXT("/"));
 		Op->Inputs.Add(FNiagaraOpInOutInfo(A, NumericType, AText, AText, DefaultStr_One));
 		Op->Inputs.Add(FNiagaraOpInOutInfo(B, NumericType, BText, BText, DefaultStr_One));
@@ -326,6 +330,7 @@ void FNiagaraOpInfo::Init()
 		Op->CompactName = FText::FromString(CompactNameString);
 		Op->bShowPinNamesInCompactMode = true;
 		Op->Description = NSLOCTEXT("NiagaraOpInfo", "MultiplyAdd Desc", "Result = (A * B) + C");
+		Op->AlternateSearchName = FName("*+");
 		Op->Inputs.Add(FNiagaraOpInOutInfo(A, NumericType, AText, AText, DefaultStr_One));
 		Op->Inputs.Add(FNiagaraOpInOutInfo(B, NumericType, BText, BText, DefaultStr_One));
 		Op->Inputs.Add(FNiagaraOpInOutInfo(C, NumericType, CText, CText, DefaultStr_Zero));
@@ -456,6 +461,7 @@ void FNiagaraOpInfo::Init()
 		Op->CompactName = FText::FromString(TEXT("1-A"));
 		Op->bShowPinNamesInCompactMode = true;
 		Op->Description = NSLOCTEXT("NiagaraOpInfo", "One Minus Desc", "Result = 1 - A");
+		Op->AlternateSearchName = FName("1-");
 		Op->Keywords = FText::FromString(TEXT("1-x"));
 		Op->Inputs.Add(FNiagaraOpInOutInfo(A, NumericType, AText, AText, DefaultStr_One));
 		Op->Outputs.Add(FNiagaraOpInOutInfo(Result, NumericType, ResultText, ResultText, DefaultStr_One, TEXT("1 - {0}")));
diff --git a/Engine/Plugins/FX/Niagara/Source/NiagaraEditor/Private/NiagaraEditorUtilities.cpp b/Engine/Plugins/FX/Niagara/Source/NiagaraEditor/Private/NiagaraEditorUtilities.cpp
--- a/Engine/Plugins/FX/Niagara/Source/NiagaraEditor/Private/NiagaraEditorUtilities.cpp
+++ b/Engine/Plugins/FX/Niagara/Source/NiagaraEditor/Private/NiagaraEditorUtilities.cpp
@@ -1705,6 +1705,14 @@ int32 FNiagaraEditorUtilities::GetWeightForItem(const TSharedPtr<FNiagaraMenuAct
 	InCurrentAction->DisplayName.ToString().ParseIntoArray(DisplayNameArray, TEXT(" "), true);
 	WeightedArrayList.Add(FArrayWithWeight(&DisplayNameArray, DisplayNameWeight, 10, StartsWith));
 
+	// Alternate Name
+	TArray<FString> AlternameNameArray;
+	if(InCurrentAction->AlternateSearchName.IsSet())
+	{
+		InCurrentAction->AlternateSearchName.GetValue().ToString().ParseIntoArray(AlternameNameArray, TEXT(" "), true);
+		WeightedArrayList.Add(FArrayWithWeight(&AlternameNameArray, DisplayNameWeight, 10, StartsWith));
+	}
+
 	// Keywords
 	TArray<FString> KeywordsArray;
 	InCurrentAction->Keywords.ToString().ParseIntoArray(KeywordsArray, TEXT(" "), true);
diff --git a/Engine/Plugins/FX/Niagara/Source/NiagaraEditor/Public/NiagaraActions.h b/Engine/Plugins/FX/Niagara/Source/NiagaraEditor/Public/NiagaraActions.h
--- a/Engine/Plugins/FX/Niagara/Source/NiagaraEditor/Public/NiagaraActions.h
+++ b/Engine/Plugins/FX/Niagara/Source/NiagaraEditor/Public/NiagaraActions.h
@@ -135,6 +135,9 @@ struct FNiagaraMenuAction_Base
 	/** The DisplayName used in lists */
 	FText DisplayName;
 
+	/** An alternate name used explicitly for search purposes. Example: Divide and '/'.*/
+	TOptional<FText> AlternateSearchName;
+
 	/** The Tooltip text for this action */
 	FText ToolTip;
 
diff --git a/Engine/Plugins/FX/Niagara/Source/NiagaraEditor/Public/NiagaraEditorCommon.h b/Engine/Plugins/FX/Niagara/Source/NiagaraEditor/Public/NiagaraEditorCommon.h
--- a/Engine/Plugins/FX/Niagara/Source/NiagaraEditor/Public/NiagaraEditorCommon.h
+++ b/Engine/Plugins/FX/Niagara/Source/NiagaraEditor/Public/NiagaraEditorCommon.h
@@ -32,6 +32,8 @@ public:
 	{}
 
 	FName Name;
+	// An alternate name for search purposes, such as Divide and '/'
+	TOptional<FName> AlternateSearchName;
 	FText Category;
 	FText FriendlyName;
 	FText CompactName;
