Make new emitter assets inheritable by default

diff --git a/Engine/Plugins/FX/Niagara/Source/NiagaraEditor/Private/NiagaraEmitterFactoryNew.cpp b/Engine/Plugins/FX/Niagara/Source/NiagaraEditor/Private/NiagaraEmitterFactoryNew.cpp
--- a/Engine/Plugins/FX/Niagara/Source/NiagaraEditor/Private/NiagaraEmitterFactoryNew.cpp
+++ b/Engine/Plugins/FX/Niagara/Source/NiagaraEditor/Private/NiagaraEmitterFactoryNew.cpp
@@ -133,7 +133,7 @@ UObject* UNiagaraEmitterFactoryNew::FactoryCreateNew(UClass* Class, UObject* InP
 			NewEmitter->DisableVersioning(EmitterToCopy->GetExposedVersion().VersionGuid);
 		}
 
-		NewEmitter->bIsInheritable = false;
+		NewEmitter->bIsInheritable = true;
 		NewEmitter->TemplateAssetDescription = FText();
 		NewEmitter->AssetTags.Empty();
 		NewEmitter->Category = FText();
diff --git a/Engine/Plugins/FX/Niagara/Source/NiagaraEditor/Private/ViewModels/NiagaraSystemViewModel.cpp b/Engine/Plugins/FX/Niagara/Source/NiagaraEditor/Private/ViewModels/NiagaraSystemViewModel.cpp
--- a/Engine/Plugins/FX/Niagara/Source/NiagaraEditor/Private/ViewModels/NiagaraSystemViewModel.cpp
+++ b/Engine/Plugins/FX/Niagara/Source/NiagaraEditor/Private/ViewModels/NiagaraSystemViewModel.cpp
@@ -599,6 +599,7 @@ TSharedPtr<FNiagaraEmitterHandleViewModel> FNiagaraSystemViewModel::AddEmptyEmit
 	FName EmptyEmitterName = FNiagaraEditorUtilities::GetUniqueObjectName<UNiagaraEmitter>(GetTransientPackage(), TEXT("Empty"));
 	EmptyEmitter->SetUniqueEmitterName(EmptyEmitterName.ToString());
 	EmptyEmitter->SetFlags(RF_Transactional);
+	EmptyEmitter->bIsInheritable = false;
 	return AddEmitter(FVersionedNiagaraEmitter(EmptyEmitter, FGuid()));
 }
 
