- Debug HUD add performance unit selection (millisecond / microsecond)
- Add an auto mode to the perf graph for the unit scale

diff --git a/Engine/Plugins/FX/Niagara/Source/Niagara/Private/NiagaraDebugHud.cpp b/Engine/Plugins/FX/Niagara/Source/Niagara/Private/NiagaraDebugHud.cpp
--- a/Engine/Plugins/FX/Niagara/Source/Niagara/Private/NiagaraDebugHud.cpp
+++ b/Engine/Plugins/FX/Niagara/Source/Niagara/Private/NiagaraDebugHud.cpp
@@ -445,10 +445,22 @@ namespace NiagaraDebugLocal
 	#endif
 	}
 
-	FString FormatPerfValue(double dValue, int32 Length = 7)
+	double GetMillisecondsToPerfUnits()
 	{
-		const int32 Value = int32(dValue);
-		FString TempString = FString::FormatAsNumber(Value);
+		return (Settings.PerfUnits == ENiagaraDebugHUDPerfUnits::Milliseconds) ? 1.0 / 1000.0 : 1.0;
+	}
+
+	FString FormatPerfValue(double Microseconds, int32 Length = 7)
+	{
+		FString TempString;
+		if (Settings.PerfUnits == ENiagaraDebugHUDPerfUnits::Milliseconds)
+		{
+			TempString = FString::Printf(TEXT("%5.3f"), float(Microseconds / 1000.0));
+		}
+		else //if (Settings.PerfUnits == ENiagaraDebugHUDPerfUnits::Milliseconds)
+		{
+			TempString = FString::FormatAsNumber(int32(Microseconds));
+		}
 		while (TempString.Len() < Length)
 		{
 			TempString.AppendChar(' ');
@@ -456,6 +468,15 @@ namespace NiagaraDebugLocal
 		return TempString;
 	}
 
+	FString FormatPerfString(const TCHAR* InText)
+	{
+		if (Settings.PerfUnits == ENiagaraDebugHUDPerfUnits::Milliseconds)
+		{
+			return FString::Printf(TEXT("%s (ms)"), InText);
+		}
+		return FString::Printf(TEXT("%s (us)"), InText);
+	}
+
 	const FCachedVariables& GetCachedVariables(UNiagaraSystem* NiagaraSystem)
 	{
 		FCachedVariables* CachedVariables = GCachedSystemVariables.Find(NiagaraSystem);
@@ -1809,7 +1830,7 @@ void FNiagaraDebugHud::DrawOverview(class FNiagaraWorldManager* WorldManager, FC
 			const int32 GlobalDataStringSize = GetStringSize(Font, TEXT("000,000")).X;
 			const int32 SystemStringSize = GetStringSize(Font, TEXT("000,000")).X;
 
-			OverviewColumns.Emplace(TEXT("Game Thread Avg:"), FormatPerfValue(GlobalPerfStats.Avg.Time_GT), GlobalDataStringSize, TEXT("GT Avg (us)"), ColumnOffset, Font, SystemStringSize,
+			OverviewColumns.Emplace(TEXT("Game Thread Avg:"), FormatPerfValue(GlobalPerfStats.Avg.Time_GT), GlobalDataStringSize, FormatPerfString(TEXT("GT Avg")), ColumnOffset, Font, SystemStringSize,
 				[&DetailColor, &DetailHighlightColor, &fAdvanceHeight](FCanvas* Canvas, UFont* Font, float X, float Y, FOverviewColumn& Col, const FSystemDebugInfo& SystemInfo)
 				{
 					FLinearColor RowBGColor = SystemInfo.UniqueColor;
@@ -1819,7 +1840,7 @@ void FNiagaraDebugHud::DrawOverview(class FNiagaraWorldManager* WorldManager, FC
 					Canvas->DrawShadowedString(X, Y, *FormatPerfValue(SystemInfo.PerfStats ? SystemInfo.PerfStats->Avg.Time_GT : 0.0), Font, RowColor);
 				});
 
-			OverviewColumns.Emplace(TEXT("Game Thread Max:"), FormatPerfValue(GlobalPerfStats.Max.Time_GT), GlobalDataStringSize, TEXT("GT Max (us)"), ColumnOffset, Font, SystemStringSize,
+			OverviewColumns.Emplace(TEXT("Game Thread Max:"), FormatPerfValue(GlobalPerfStats.Max.Time_GT), GlobalDataStringSize, FormatPerfString(TEXT("GT Max")), ColumnOffset, Font, SystemStringSize,
 				[&DetailColor, &DetailHighlightColor, &fAdvanceHeight](FCanvas* Canvas, UFont* Font, float X, float Y, FOverviewColumn& Col, const FSystemDebugInfo& SystemInfo)
 				{
 					FLinearColor RowBGColor = SystemInfo.UniqueColor;
@@ -1829,7 +1850,7 @@ void FNiagaraDebugHud::DrawOverview(class FNiagaraWorldManager* WorldManager, FC
 					Canvas->DrawShadowedString(X, Y, *FormatPerfValue(SystemInfo.PerfStats ? SystemInfo.PerfStats->Max.Time_GT : 0.0), Font, RowColor);
 				});
 
-			OverviewColumns.Emplace_GetRef(TEXT("Render Thread Avg:"), FormatPerfValue(GlobalPerfStats.Avg.Time_RT), GlobalDataStringSize, TEXT("RT Avg (us)"), ColumnOffset, Font, SystemStringSize,
+			OverviewColumns.Emplace_GetRef(TEXT("Render Thread Avg:"), FormatPerfValue(GlobalPerfStats.Avg.Time_RT), GlobalDataStringSize, FormatPerfString(TEXT("RT Avg")), ColumnOffset, Font, SystemStringSize,
 				[&DetailColor, &DetailHighlightColor, &fAdvanceHeight](FCanvas* Canvas, UFont* Font, float X, float Y, FOverviewColumn& Col, const FSystemDebugInfo& SystemInfo)
 				{
 					FLinearColor RowBGColor = SystemInfo.UniqueColor;
@@ -1839,7 +1860,7 @@ void FNiagaraDebugHud::DrawOverview(class FNiagaraWorldManager* WorldManager, FC
 					Canvas->DrawShadowedString(X, Y, *FormatPerfValue(SystemInfo.PerfStats ? SystemInfo.PerfStats->Avg.Time_RT : 0.0), Font, RowColor);
 				});
 
-			OverviewColumns.Emplace(TEXT("Render Thread Max:"), FormatPerfValue(GlobalPerfStats.Max.Time_RT), GlobalDataStringSize, TEXT("RT Max (us)"), ColumnOffset, Font, SystemStringSize,
+			OverviewColumns.Emplace(TEXT("Render Thread Max:"), FormatPerfValue(GlobalPerfStats.Max.Time_RT), GlobalDataStringSize, FormatPerfString(TEXT("RT Max")), ColumnOffset, Font, SystemStringSize,
 				[&DetailColor, &DetailHighlightColor, &fAdvanceHeight](FCanvas* Canvas, UFont* Font, float X, float Y, FOverviewColumn& Col, const FSystemDebugInfo& SystemInfo)
 				{
 					FLinearColor RowBGColor = SystemInfo.UniqueColor;
@@ -1849,7 +1870,7 @@ void FNiagaraDebugHud::DrawOverview(class FNiagaraWorldManager* WorldManager, FC
 					Canvas->DrawShadowedString(X, Y, *FormatPerfValue(SystemInfo.PerfStats ? SystemInfo.PerfStats->Max.Time_RT : 0.0), Font, RowColor);
 				});
 
-			OverviewColumns.Emplace_GetRef(TEXT("Gpu Avg:"), FormatPerfValue(GlobalPerfStats.Avg.Time_GPU), GlobalDataStringSize, TEXT("Gpu Avg (us)"), ColumnOffset, Font, SystemStringSize,
+			OverviewColumns.Emplace_GetRef(TEXT("Gpu Avg:"), FormatPerfValue(GlobalPerfStats.Avg.Time_GPU), GlobalDataStringSize, FormatPerfString(TEXT("Gpu Avg")), ColumnOffset, Font, SystemStringSize,
 				[&DetailColor, &DetailHighlightColor, &fAdvanceHeight](FCanvas* Canvas, UFont* Font, float X, float Y, FOverviewColumn& Col, const FSystemDebugInfo& SystemInfo)
 				{
 					FLinearColor RowBGColor = SystemInfo.UniqueColor;
@@ -1859,7 +1880,7 @@ void FNiagaraDebugHud::DrawOverview(class FNiagaraWorldManager* WorldManager, FC
 					Canvas->DrawShadowedString(X, Y, *FormatPerfValue(SystemInfo.PerfStats ? SystemInfo.PerfStats->Avg.Time_GPU : 0.0), Font, RowColor);
 				});
 
-			OverviewColumns.Emplace(TEXT("Gpu Max:"), FormatPerfValue(GlobalPerfStats.Max.Time_GPU), GlobalDataStringSize, TEXT("Gpu Max (us)"), ColumnOffset, Font, SystemStringSize,
+			OverviewColumns.Emplace(TEXT("Gpu Max:"), FormatPerfValue(GlobalPerfStats.Max.Time_GPU), GlobalDataStringSize, FormatPerfString(TEXT("Gpu Max")), ColumnOffset, Font, SystemStringSize,
 				[&DetailColor, &DetailHighlightColor, &fAdvanceHeight](FCanvas* Canvas, UFont* Font, float X, float Y, FOverviewColumn& Col, const FSystemDebugInfo& SystemInfo)
 				{
 					FLinearColor RowBGColor = SystemInfo.UniqueColor;
@@ -1871,7 +1892,7 @@ void FNiagaraDebugHud::DrawOverview(class FNiagaraWorldManager* WorldManager, FC
 		}
 		else if (Settings.OverviewMode == ENiagaraDebugHUDOverviewMode::PerformanceGraph && StatsListener)
 		{
-			OverviewColumns.Emplace(TEXT(""), TEXT(""), TEXT(""), ColumnOffset, Font, TEXT("------"),
+			OverviewColumns.Emplace(TEXT("Total Active"), TEXT(""), TEXT(""), ColumnOffset, Font, TEXT("------"),
 				[&DetailColor, &DetailHighlightColor, &fAdvanceHeight](FCanvas* Canvas, UFont* Font, float X, float Y, FOverviewColumn& Col, const FSystemDebugInfo& SystemInfo)
 				{
 					Canvas->DrawTile(X, Y, Col.MaxWidth - 6.0f, fAdvanceHeight - 2.0f, 0.0f, 0.0f, 0.0f, 0.0f, SystemInfo.UniqueColor);
@@ -2084,11 +2105,15 @@ void FNiagaraDebugHud::DrawOverview(class FNiagaraWorldManager* WorldManager, FC
 				GraphLocation.X += OverviewColumns.Last().Offset + OverviewColumns.Last().MaxWidth + 50.0f;
 			}
 
-			FGraph<double> Graph(DrawCanvas, Font, GraphLocation, FVector2f(Settings.PerfGraphSize), FVector2f(float(Settings.PerfHistoryFrames), Settings.PerfGraphTimeRange), TEXT("Frame"), TEXT("Time(us)"));
+			const double MillisecondsToPerfUnits = GetMillisecondsToPerfUnits();
+			const float GraphTimeRange = (Settings.bUsePerfGraphTimeRange ? Settings.PerfGraphTimeRange : AutoGraphTimeRange) * MillisecondsToPerfUnits;
+
+			FGraph<double> Graph(DrawCanvas, Font, GraphLocation, FVector2f(Settings.PerfGraphSize), FVector2f(float(Settings.PerfHistoryFrames), GraphTimeRange), TEXT("Frame"), FormatPerfString(TEXT("Time")));
 
 			Graph.Draw(Settings.PerfGraphAxisColor, BackgroundColor);
 
 			//Add each line to the graph.
+			double NewGraphTimeRange = 0.0;
 			for (const auto& Pair : PerSystemDebugInfo)
 			{
 				const FSystemDebugInfo& SysInfo = Pair.Value;
@@ -2131,9 +2156,28 @@ void FNiagaraDebugHud::DrawOverview(class FNiagaraWorldManager* WorldManager, FC
 						Frames = MoveTemp(Smoothed);
 					}
 
+					for (double& Value : Frames)
+					{
+						NewGraphTimeRange = FMath::Max(Value, NewGraphTimeRange);
+						Value *= MillisecondsToPerfUnits;
+					}
+
 					Graph.DrawLine(SysInfo.SystemName, SysInfo.UniqueColor, Frames);
 				}
 			}
+
+			constexpr double AutoGraphTimeRangeChangeRate = 5.0;
+			constexpr double AutoGraphMinUnits = 100.0;
+
+			NewGraphTimeRange = FMath::CeilToDouble(NewGraphTimeRange / AutoGraphMinUnits) * AutoGraphMinUnits;
+			NewGraphTimeRange = FMath::Max(NewGraphTimeRange, AutoGraphMinUnits);
+			AutoGraphTimeRangeChangeMax = FMath::Max(AutoGraphTimeRangeChangeMax, NewGraphTimeRange);
+			if (AutoGraphTimeRangeChangeMax > AutoGraphTimeRange || (FPlatformTime::Seconds() - AutoGraphTimeRangeChangedTime > AutoGraphTimeRangeChangeRate))
+			{
+				AutoGraphTimeRange				= AutoGraphTimeRangeChangeMax;
+				AutoGraphTimeRangeChangeMax		= AutoGraphMinUnits;
+				AutoGraphTimeRangeChangedTime	= FPlatformTime::Seconds();
+			}
 		}
 		#endif//WITH_PARTICLE_PERF_STATS
 	}
@@ -2240,8 +2284,11 @@ void FNiagaraDebugHud::DrawGpuComputeOverriew(class FNiagaraWorldManager* WorldM
 		SimpleTable.GetColumnText(0).Append(TEXT("Gpu Compute Overview"));
 		SimpleTable.RowComplete();
 
-		SimpleTable.GetColumnText(0).Appendf(TEXT("TotalMicroseconds : %llu"), GpuTotalMicroseconds.GetAverage());
-		SimpleTable.GetColumnText(1).Appendf(TEXT("TotalDispatches : %d"), GpuTotalDispatches.GetAverage());
+		const UEnum* PerfUnitsEnum = StaticEnum<ENiagaraDebugHUDPerfUnits>();
+		const FString PerfUnits = PerfUnitsEnum ? PerfUnitsEnum->GetNameStringByValue((int)Settings.PerfUnits) : FString();
+
+		SimpleTable.GetColumnText(0).Appendf(TEXT("Total %s : %s"), *PerfUnits, *FormatPerfValue(GpuTotalMicroseconds.GetAverage()));
+		SimpleTable.GetColumnText(1).Appendf(TEXT("Total Dispatches : %d"), GpuTotalDispatches.GetAverage());
 
 		SimpleTable.Draw(Font, DrawCanvas, TextLocation, DetailColor, BackgroundColor);
 		TextLocation.Y += fAdvanceHeight;
@@ -2308,8 +2355,8 @@ void FNiagaraDebugHud::DrawGpuComputeOverriew(class FNiagaraWorldManager* WorldM
 			SimpleTable.GetColumnText(1).Append(TEXT("Emitter Name"));
 			SimpleTable.GetColumnText(2).Append(TEXT("Stage Name"));
 			SimpleTable.GetColumnText(3).Append(TEXT("Avg Instances"));
-			SimpleTable.GetColumnText(4).Append(TEXT("Avg us"));
-			SimpleTable.GetColumnText(5).Append(TEXT("Max us"));
+			SimpleTable.GetColumnText(4).Append(FormatPerfString(TEXT("Avg")));
+			SimpleTable.GetColumnText(5).Append(FormatPerfString(TEXT("Max")));
 			SimpleTable.RowComplete();
 
 			for (auto EmitterIt = SystemIt.Value().Emitters.CreateIterator(); EmitterIt; ++EmitterIt)
@@ -2329,16 +2376,16 @@ void FNiagaraDebugHud::DrawGpuComputeOverriew(class FNiagaraWorldManager* WorldM
 					OwnerEmitter->GetFName().AppendString(SimpleTable.GetColumnText(1));
 					StageIt.Key().AppendString(SimpleTable.GetColumnText(2));
 					SimpleTable.GetColumnText(3).Appendf(TEXT("%4.1f"), StageUsage.InstanceCount.GetAverage<float>());
-					SimpleTable.GetColumnText(4).Appendf(TEXT("%llu"), StageUsage.Microseconds.GetAverage());
-					SimpleTable.GetColumnText(5).Appendf(TEXT("%llu"), StageUsage.Microseconds.GetMax());
+					SimpleTable.GetColumnText(4).Append(FormatPerfValue(StageUsage.Microseconds.GetAverage()));
+					SimpleTable.GetColumnText(5).Append(FormatPerfValue(StageUsage.Microseconds.GetMax()));
 					SimpleTable.RowComplete();
 				}
 			}
 
 			SimpleTable.GetColumnText(2).Append(TEXT("Total"));
 			SimpleTable.GetColumnText(3).Appendf(TEXT("%4.1f"), SystemUsage.InstanceCount.GetAverage<float>());
-			SimpleTable.GetColumnText(4).Appendf(TEXT("%llu"), SystemUsage.Microseconds.GetAverage());
-			SimpleTable.GetColumnText(5).Appendf(TEXT("%llu"), SystemUsage.Microseconds.GetMax());
+			SimpleTable.GetColumnText(4).Append(FormatPerfValue(SystemUsage.Microseconds.GetAverage()));
+			SimpleTable.GetColumnText(5).Append(FormatPerfValue(SystemUsage.Microseconds.GetMax()));
 			SimpleTable.RowComplete();
 		}
 		SimpleTable.Draw(Font, DrawCanvas, TextLocation, DetailColor, BackgroundColor);
@@ -2353,8 +2400,8 @@ void FNiagaraDebugHud::DrawGpuComputeOverriew(class FNiagaraWorldManager* WorldM
 		SimpleTable.AddColumns(3, 100.0f);
 		SimpleTable.GetColumnText(0).Append(TEXT("SystemName"));
 		SimpleTable.GetColumnText(1).Append(TEXT("Avg Instances"));
-		SimpleTable.GetColumnText(2).Append(TEXT("Avg us"));
-		SimpleTable.GetColumnText(3).Append(TEXT("Max us"));
+		SimpleTable.GetColumnText(2).Append(FormatPerfString(TEXT("Avg")));
+		SimpleTable.GetColumnText(3).Append(FormatPerfString(TEXT("Max")));
 		SimpleTable.RowComplete();
 
 		for (auto SystemIt=GpuUsagePerSystem.CreateIterator(); SystemIt; ++SystemIt)
@@ -2377,8 +2424,8 @@ void FNiagaraDebugHud::DrawGpuComputeOverriew(class FNiagaraWorldManager* WorldM
 			}
 #endif
 			SimpleTable.GetColumnText(1).Appendf(TEXT("%4.1f"), SystemUsage.InstanceCount.GetAverage<float>());
-			SimpleTable.GetColumnText(2).Appendf(TEXT("%llu"), SystemUsage.Microseconds.GetAverage());
-			SimpleTable.GetColumnText(3).Appendf(TEXT("%llu"), SystemUsage.Microseconds.GetMax());
+			SimpleTable.GetColumnText(2).Append(FormatPerfValue(SystemUsage.Microseconds.GetAverage()));
+			SimpleTable.GetColumnText(3).Append(FormatPerfValue(SystemUsage.Microseconds.GetMax()));
 			SimpleTable.RowComplete();
 		}
 		SimpleTable.Draw(Font, DrawCanvas, TextLocation, DetailColor, BackgroundColor);
@@ -2393,8 +2440,8 @@ void FNiagaraDebugHud::DrawGpuComputeOverriew(class FNiagaraWorldManager* WorldM
 		SimpleTable.AddColumns(3, 100.0f);
 		SimpleTable.GetColumnText(0).Append(TEXT("EventName"));
 		SimpleTable.GetColumnText(1).Append(TEXT("Avg Instances"));
-		SimpleTable.GetColumnText(2).Append(TEXT("Avg us"));
-		SimpleTable.GetColumnText(3).Append(TEXT("Max us"));
+		SimpleTable.GetColumnText(2).Append(FormatPerfString(TEXT("Avg")));
+		SimpleTable.GetColumnText(3).Append(FormatPerfString(TEXT("Max")));
 		SimpleTable.RowComplete();
 
 		for (auto EventIt=GpuUsagePerEvent.CreateIterator(); EventIt; ++EventIt)
@@ -2402,8 +2449,8 @@ void FNiagaraDebugHud::DrawGpuComputeOverriew(class FNiagaraWorldManager* WorldM
 			const FGpuUsagePerEvent& EventUsage = EventIt.Value();
 			EventIt.Key().AppendString(SimpleTable.GetColumnText(0));
 			SimpleTable.GetColumnText(1).Appendf(TEXT("%4.1f"), EventUsage.InstanceCount.GetAverage<float>());
-			SimpleTable.GetColumnText(2).Appendf(TEXT("%llu"), EventUsage.Microseconds.GetAverage());
-			SimpleTable.GetColumnText(3).Appendf(TEXT("%llu"), EventUsage.Microseconds.GetMax());
+			SimpleTable.GetColumnText(2).Append(FormatPerfValue(EventUsage.Microseconds.GetAverage()));
+			SimpleTable.GetColumnText(3).Append(FormatPerfValue(EventUsage.Microseconds.GetMax()));
 			SimpleTable.RowComplete();
 		}
 		SimpleTable.Draw(Font, DrawCanvas, TextLocation, DetailColor, BackgroundColor);
diff --git a/Engine/Plugins/FX/Niagara/Source/Niagara/Private/NiagaraDebugHud.h b/Engine/Plugins/FX/Niagara/Source/Niagara/Private/NiagaraDebugHud.h
--- a/Engine/Plugins/FX/Niagara/Source/Niagara/Private/NiagaraDebugHud.h
+++ b/Engine/Plugins/FX/Niagara/Source/Niagara/Private/NiagaraDebugHud.h
@@ -356,6 +356,10 @@ private:
 	};
 	TArray<FDebugBox2D> Boxes2D;
 
+	double AutoGraphTimeRange = 100.0;
+	double AutoGraphTimeRangeChangeMax = 100.0;
+	double AutoGraphTimeRangeChangedTime = 0.0;
+
 public:
 	/** Add a 2D line to the debug renering. Positions are in normalized screen space. (0,0) in top left, (1,1) bottom right.*/
 	void AddLine2D(FVector2f Start, FVector2f End, FLinearColor Color, float Thickness, float Lifetime);
diff --git a/Engine/Plugins/FX/Niagara/Source/Niagara/Public/NiagaraDebuggerCommon.h b/Engine/Plugins/FX/Niagara/Source/Niagara/Public/NiagaraDebuggerCommon.h
--- a/Engine/Plugins/FX/Niagara/Source/Niagara/Public/NiagaraDebuggerCommon.h
+++ b/Engine/Plugins/FX/Niagara/Source/Niagara/Public/NiagaraDebuggerCommon.h
@@ -405,6 +405,13 @@ enum class ENiagaraDebugHUDPerfSampleMode
 	PerInstanceAverage,	
 };
 
+UENUM()
+enum class ENiagaraDebugHUDPerfUnits
+{
+	Microseconds,
+	Milliseconds,
+};
+
 UENUM()
 enum class ENiagaraDebugHUDDOverviewSort
 {
@@ -665,6 +672,9 @@ struct FNiagaraDebugHUDSettingsData
 	UPROPERTY(Config, EditAnywhere, Category = "Perf Overview")
 	ENiagaraDebugHUDPerfSampleMode PerfSampleMode = ENiagaraDebugHUDPerfSampleMode::FrameTotal;
 
+	UPROPERTY(Config, EditAnywhere, Category = "Perf Overview")
+	ENiagaraDebugHUDPerfUnits PerfUnits = ENiagaraDebugHUDPerfUnits::Microseconds;
+
 	/** Time range of the Y Axis of the perf graph */
 	UPROPERTY(Config, EditAnywhere, Category = "Perf Overview")
 	ENiagaraDebugHUDPerfGraphMode PerfGraphMode = ENiagaraDebugHUDPerfGraphMode::GameThread;
@@ -673,8 +683,12 @@ struct FNiagaraDebugHUDSettingsData
 	UPROPERTY(Config, EditAnywhere, Category = "Perf Overview")
 	int32 PerfHistoryFrames = 600;
 
+	/** Use the specified user range when enabled, otherwise we will auto detect the range to use. */
+	UPROPERTY(Config, EditAnywhere, Category = "Perf Overview", meta = (InlineEditConditionToggle))
+	bool bUsePerfGraphTimeRange = false;
+
 	/** Time range of the Y Axis of the perf graph */
-	UPROPERTY(Config, EditAnywhere, Category = "Perf Overview")
+	UPROPERTY(Config, EditAnywhere, Category = "Perf Overview", meta = (EditCondition = "bUsePerfGraphTimeRange"))
 	float PerfGraphTimeRange = 500.0f;
 
 	/** Pixel size of the perf graph. */
