- Add metadata for asset tags so they can be localized

diff --git a/Engine/Plugins/FX/Niagara/Source/Niagara/Classes/NiagaraSystem.h b/Engine/Plugins/FX/Niagara/Source/Niagara/Classes/NiagaraSystem.h
--- a/Engine/Plugins/FX/Niagara/Source/Niagara/Classes/NiagaraSystem.h
+++ b/Engine/Plugins/FX/Niagara/Source/Niagara/Classes/NiagaraSystem.h
@@ -390,6 +390,8 @@ public:
 #endif
 
 #if WITH_EDITORONLY_DATA
+	NIAGARA_API virtual void GetAssetRegistryTagMetadata(TMap<FName, FAssetRegistryTagMetadata>& OutMetadata) const override;
+
 	/** Are there any pending compile requests?*/
 	NIAGARA_API bool HasOutstandingCompilationRequests(bool bIncludingGPUShaders = false) const;
 
diff --git a/Engine/Plugins/FX/Niagara/Source/Niagara/Private/NiagaraSystem.cpp b/Engine/Plugins/FX/Niagara/Source/Niagara/Private/NiagaraSystem.cpp
--- a/Engine/Plugins/FX/Niagara/Source/Niagara/Private/NiagaraSystem.cpp
+++ b/Engine/Plugins/FX/Niagara/Source/Niagara/Private/NiagaraSystem.cpp
@@ -74,6 +74,22 @@ namespace NiagaraScriptCookStats
 }
 #endif
 
+namespace NiagaraSystemPrivate
+{
+	static const FName NAME_ActiveEmitters("ActiveEmitters");
+	static const FName NAME_ActiveRenderers("ActiveRenderers");
+	static const FName NAME_GPUSimsMissingFixedBounds("GPUSimsMissingFixedBounds");
+	static const FName NAME_EffectType("EffectType");
+	static const FName NAME_WarmupTime("WarmupTime");
+	static const FName NAME_HasOverrideScalabilityForSystem("HasOverrideScalabilityForSystem");
+	static const FName NAME_HasDIsWithPostSimulateTick("HasDIsWithPostSimulateTick");
+	static const FName NAME_NeedsSortedSignificanceCull("NeedsSortedSignificanceCull");
+	static const FName NAME_ActiveDIs("ActiveDIs");
+	static const FName NAME_HasGPUEmitter("HasGPUEmitter");
+	static const FName NAME_FixedBoundsSize("FixedBoundsSize");
+	static const FName NAME_NumEmitters("NumEmitters");
+}
+
 //Disable for now until we can spend more time on a good method of applying the data gathered.
 int32 GEnableNiagaraRuntimeCycleCounts = 0;
 static FAutoConsoleVariableRef CVarEnableNiagaraRuntimeCycleCounts(TEXT("fx.EnableNiagaraRuntimeCycleCounts"), GEnableNiagaraRuntimeCycleCounts, TEXT("Toggle for runtime cylce counts tracking Niagara's frame time. \n"), ECVF_ReadOnly);
@@ -1549,12 +1565,12 @@ void UNiagaraSystem::GetAssetRegistryTags(FAssetRegistryTagsContext Context) con
 	}
 
 #if WITH_EDITOR
-	Context.AddTag(FAssetRegistryTag("HasGPUEmitter", HasAnyGPUEmitters() ? TEXT("True") : TEXT("False"), FAssetRegistryTag::TT_Alphabetical));
+	Context.AddTag(FAssetRegistryTag(NiagaraSystemPrivate::NAME_HasGPUEmitter, HasAnyGPUEmitters() ? TEXT("True") : TEXT("False"), FAssetRegistryTag::TT_Alphabetical));
 
 	const float BoundsSize = float(FixedBounds.GetSize().GetMax());
-	Context.AddTag(FAssetRegistryTag("FixedBoundsSize", bFixedBounds ? FString::Printf(TEXT("%.2f"), BoundsSize) : FString(TEXT("None")), FAssetRegistryTag::TT_Numerical));
+	Context.AddTag(FAssetRegistryTag(NiagaraSystemPrivate::NAME_FixedBoundsSize, bFixedBounds ? FString::Printf(TEXT("%.2f"), BoundsSize) : FString(TEXT("None")), FAssetRegistryTag::TT_Numerical));
 
-	Context.AddTag(FAssetRegistryTag("NumEmitters", LexToString(EmitterHandles.Num()), FAssetRegistryTag::TT_Numerical));
+	Context.AddTag(FAssetRegistryTag(NiagaraSystemPrivate::NAME_NumEmitters, LexToString(EmitterHandles.Num()), FAssetRegistryTag::TT_Numerical));
 
 	uint32 GPUSimsMissingFixedBounds = 0;
 
@@ -1588,14 +1604,14 @@ void UNiagaraSystem::GetAssetRegistryTags(FAssetRegistryTagsContext Context) con
 		}
 	}
 
-	Context.AddTag(FAssetRegistryTag("ActiveEmitters", LexToString(NumActiveEmitters), FAssetRegistryTag::TT_Numerical));
-	Context.AddTag(FAssetRegistryTag("ActiveRenderers", LexToString(NumActiveRenderers), FAssetRegistryTag::TT_Numerical));
-	Context.AddTag(FAssetRegistryTag("GPUSimsMissingFixedBounds", LexToString(GPUSimsMissingFixedBounds), FAssetRegistryTag::TT_Numerical));
-	Context.AddTag(FAssetRegistryTag("EffectType", EffectType != nullptr ? EffectType->GetName() : FString(TEXT("None")), FAssetRegistryTag::TT_Alphabetical));
-	Context.AddTag(FAssetRegistryTag("WarmupTime", LexToString(WarmupTime), FAssetRegistryTag::TT_Numerical));
-	Context.AddTag(FAssetRegistryTag("HasOverrideScalabilityForSystem", bOverrideScalabilitySettings ? TEXT("True") : TEXT("False"), FAssetRegistryTag::TT_Alphabetical));
-	Context.AddTag(FAssetRegistryTag("HasDIsWithPostSimulateTick", bHasDIsWithPostSimulateTick ? TEXT("True") : TEXT("False"), FAssetRegistryTag::TT_Alphabetical));
-	Context.AddTag(FAssetRegistryTag("NeedsSortedSignificanceCull", bNeedsSortedSignificanceCull ? TEXT("True") : TEXT("False"), FAssetRegistryTag::TT_Alphabetical));
+	Context.AddTag(FAssetRegistryTag(NiagaraSystemPrivate::NAME_ActiveEmitters, LexToString(NumActiveEmitters), FAssetRegistryTag::TT_Numerical));
+	Context.AddTag(FAssetRegistryTag(NiagaraSystemPrivate::NAME_ActiveRenderers, LexToString(NumActiveRenderers), FAssetRegistryTag::TT_Numerical));
+	Context.AddTag(FAssetRegistryTag(NiagaraSystemPrivate::NAME_GPUSimsMissingFixedBounds, LexToString(GPUSimsMissingFixedBounds), FAssetRegistryTag::TT_Numerical));
+	Context.AddTag(FAssetRegistryTag(NiagaraSystemPrivate::NAME_EffectType, EffectType != nullptr ? EffectType->GetName() : FString(TEXT("None")), FAssetRegistryTag::TT_Alphabetical));
+	Context.AddTag(FAssetRegistryTag(NiagaraSystemPrivate::NAME_WarmupTime, LexToString(WarmupTime), FAssetRegistryTag::TT_Numerical));
+	Context.AddTag(FAssetRegistryTag(NiagaraSystemPrivate::NAME_HasOverrideScalabilityForSystem, bOverrideScalabilitySettings ? TEXT("True") : TEXT("False"), FAssetRegistryTag::TT_Alphabetical));
+	Context.AddTag(FAssetRegistryTag(NiagaraSystemPrivate::NAME_HasDIsWithPostSimulateTick, bHasDIsWithPostSimulateTick ? TEXT("True") : TEXT("False"), FAssetRegistryTag::TT_Alphabetical));
+	Context.AddTag(FAssetRegistryTag(NiagaraSystemPrivate::NAME_NeedsSortedSignificanceCull, bNeedsSortedSignificanceCull ? TEXT("True") : TEXT("False"), FAssetRegistryTag::TT_Alphabetical));
 
 	// Gather up NumActive emitters based off of quality level.
 	const UNiagaraSettings* Settings = GetDefault<UNiagaraSettings>();
@@ -1694,7 +1710,7 @@ void UNiagaraSystem::GetAssetRegistryTags(FAssetRegistryTagsContext Context) con
 				PropDefault->GetAssetTagsForContext(this, FGuid(), DataInterfaces, NumericKeys, StringKeys);
 			}
 		}
-		Context.AddTag(FAssetRegistryTag("ActiveDIs", LexToString(DataInterfaces.Num()), FAssetRegistryTag::TT_Numerical));
+		Context.AddTag(FAssetRegistryTag(NiagaraSystemPrivate::NAME_ActiveDIs, LexToString(DataInterfaces.Num()), FAssetRegistryTag::TT_Numerical));
 	}
 
 
@@ -1741,6 +1757,73 @@ void UNiagaraSystem::GetAssetRegistryTags(FAssetRegistryTagsContext Context) con
 }
 
 #if WITH_EDITORONLY_DATA
+void UNiagaraSystem::GetAssetRegistryTagMetadata(TMap<FName, FAssetRegistryTagMetadata>& OutMetadata) const
+{
+	Super::GetAssetRegistryTagMetadata(OutMetadata);
+
+	OutMetadata.Add(
+		NiagaraSystemPrivate::NAME_ActiveEmitters,
+		FAssetRegistryTagMetadata()
+		.SetDisplayName(LOCTEXT("ActiveEmitters", "Active Emitters"))
+		.SetTooltip(LOCTEXT("ActiveEmittersTooltip", "The nunmber of active emitters in the system"))
+	);
+	OutMetadata.Add(
+		NiagaraSystemPrivate::NAME_ActiveRenderers,
+		FAssetRegistryTagMetadata()
+		.SetDisplayName(LOCTEXT("ActiveRenderers", "Active Renderers"))
+	);
+	OutMetadata.Add(
+		NiagaraSystemPrivate::NAME_GPUSimsMissingFixedBounds,
+		FAssetRegistryTagMetadata()
+		.SetDisplayName(LOCTEXT("GPUSimsMissingFixedBounds", "GPU Sims Missing Fixed Bounds"))
+	);
+	OutMetadata.Add(
+		NiagaraSystemPrivate::NAME_EffectType,
+		FAssetRegistryTagMetadata()
+		.SetDisplayName(LOCTEXT("EffectType", "Effect Type"))
+	);
+	OutMetadata.Add(
+		NiagaraSystemPrivate::NAME_WarmupTime,
+		FAssetRegistryTagMetadata()
+		.SetDisplayName(LOCTEXT("WarmupTime", "Warmup Time"))
+	);
+	OutMetadata.Add(
+		NiagaraSystemPrivate::NAME_HasOverrideScalabilityForSystem,
+		FAssetRegistryTagMetadata()
+		.SetDisplayName(LOCTEXT("HasOverrideScalabilityForSystem", "Has Override Scalability For System"))
+	);
+	OutMetadata.Add(
+		NiagaraSystemPrivate::NAME_HasDIsWithPostSimulateTick,
+		FAssetRegistryTagMetadata()
+		.SetDisplayName(LOCTEXT("HasDIsWithPostSimulateTick", "Has DIs With Post Simulate Tick"))
+	);
+	OutMetadata.Add(
+		NiagaraSystemPrivate::NAME_NeedsSortedSignificanceCull,
+		FAssetRegistryTagMetadata()
+		.SetDisplayName(LOCTEXT("NeedsSortedSignificanceCull", "Needs Sorted Significance Cull"))
+	);
+	OutMetadata.Add(
+		NiagaraSystemPrivate::NAME_ActiveDIs,
+		FAssetRegistryTagMetadata()
+		.SetDisplayName(LOCTEXT("ActiveDIs", "Active DIs"))
+	);
+	OutMetadata.Add(
+		NiagaraSystemPrivate::NAME_HasGPUEmitter,
+		FAssetRegistryTagMetadata()
+		.SetDisplayName(LOCTEXT("HasGPUEmitter", "Has GPU Emitter"))
+	);
+	OutMetadata.Add(
+		NiagaraSystemPrivate::NAME_FixedBoundsSize,
+		FAssetRegistryTagMetadata()
+		.SetDisplayName(LOCTEXT("FixedBoundsSize", "Fixed Bounds Size"))
+	);
+	OutMetadata.Add(
+		NiagaraSystemPrivate::NAME_NumEmitters,
+		FAssetRegistryTagMetadata()
+		.SetDisplayName(LOCTEXT("NumEmitters", "Num Emitters"))
+	);
+}
+
 bool UNiagaraSystem::HasOutstandingCompilationRequests(bool bIncludingGPUShaders) const
 {
 	if (bNeedsRequestCompile)
