Fix crash when unloading plugin
#rnx

diff --git a/Engine/Plugins/FX/Niagara/Source/NiagaraEditor/Public/NiagaraRecentAndFavoritesManager.cpp b/Engine/Plugins/FX/Niagara/Source/NiagaraEditor/Public/NiagaraRecentAndFavoritesManager.cpp
--- a/Engine/Plugins/FX/Niagara/Source/NiagaraEditor/Public/NiagaraRecentAndFavoritesManager.cpp
+++ b/Engine/Plugins/FX/Niagara/Source/NiagaraEditor/Public/NiagaraRecentAndFavoritesManager.cpp
@@ -48,14 +48,19 @@ void FNiagaraRecentAndFavoritesManager::OnAssetRemoved(const FAssetData& AssetDa
 	if(FMainMRUFavoritesList* RecentAssetsList = GetRecentsList(AssetData))
 	{
 		FString Identifier = AssetData.PackageName.ToString();
-		if(RecentAssetsList->FindMRUItemIdx(Identifier) != INDEX_NONE)
-		{
-			RecentAssetsList->RemoveMRUItem(AssetData.PackageName.ToString());
-		}
 
-		if(RecentAssetsList->ContainsFavoritesItem(Identifier))
+		// We need this because FindMRUItemIdx has a check() for non valid long package names
+		if (FPackageName::IsValidLongPackageName(Identifier))
 		{
-			RecentAssetsList->RemoveFavoritesItem(AssetData.PackageName.ToString());
+			if (RecentAssetsList->FindMRUItemIdx(Identifier) != INDEX_NONE)
+			{
+				RecentAssetsList->RemoveMRUItem(Identifier);
+			}
+
+			if (RecentAssetsList->ContainsFavoritesItem(Identifier))
+			{
+				RecentAssetsList->RemoveFavoritesItem(Identifier);
+			}
 		}
 	}
 }
