Fixed primary asset tags showing up in the primary tag list despite being set to secondary

#changelist validated
#virtualized

diff --git a/Engine/Plugins/FX/Niagara/Source/NiagaraEditor/Private/Widgets/AssetBrowser/SNiagaraAssetBrowser.cpp b/Engine/Plugins/FX/Niagara/Source/NiagaraEditor/Private/Widgets/AssetBrowser/SNiagaraAssetBrowser.cpp
--- a/Engine/Plugins/FX/Niagara/Source/NiagaraEditor/Private/Widgets/AssetBrowser/SNiagaraAssetBrowser.cpp
+++ b/Engine/Plugins/FX/Niagara/Source/NiagaraEditor/Private/Widgets/AssetBrowser/SNiagaraAssetBrowser.cpp
@@ -382,9 +382,12 @@ TArray<TSharedRef<FNiagaraAssetBrowserMainFilter>> SNiagaraAssetBrowser::GetMain
 				TArray<TSharedRef<FNiagaraAssetBrowserMainFilter>> TagChildFilters;
 				for(const FNiagaraAssetTagDefinition& AssetTagDefinition : AssetTagDefinitionsAsset->GetAssetTagDefinitions())
 				{
-					TSharedRef<FNiagaraAssetBrowserMainFilter> AssetTagFilter = MakeShared<FNiagaraAssetBrowserMainFilter>(FNiagaraAssetBrowserMainFilter::EFilterMode::NiagaraAssetTag);
-					AssetTagFilter->AssetTagDefinition = AssetTagDefinition;
-					TagChildFilters.Add(AssetTagFilter);
+					if(IsAssetTagDefinitionValid(AssetTagDefinition))
+					{
+						TSharedRef<FNiagaraAssetBrowserMainFilter> AssetTagFilter = MakeShared<FNiagaraAssetBrowserMainFilter>(FNiagaraAssetBrowserMainFilter::EFilterMode::NiagaraAssetTag);
+						AssetTagFilter->AssetTagDefinition = AssetTagDefinition;
+						TagChildFilters.Add(AssetTagFilter);
+					}
 				}
 
 				AssetTagDefinitionsAssetsFilter->ChildFilters = TagChildFilters;
