Exposed emitter summary in default view again

diff --git a/Engine/Plugins/FX/Niagara/Source/NiagaraEditor/Private/ViewModels/Stack/NiagaraStackRoot.cpp b/Engine/Plugins/FX/Niagara/Source/NiagaraEditor/Private/ViewModels/Stack/NiagaraStackRoot.cpp
--- a/Engine/Plugins/FX/Niagara/Source/NiagaraEditor/Private/ViewModels/Stack/NiagaraStackRoot.cpp
+++ b/Engine/Plugins/FX/Niagara/Source/NiagaraEditor/Private/ViewModels/Stack/NiagaraStackRoot.cpp
@@ -147,6 +147,8 @@ void UNiagaraStackRoot::RefreshEmitterFullChildren(const TArray<UNiagaraStackEnt
 
 	NewChildren.Add(GetOrCreateEmitterPropertiesGroup(CurrentChildren));
 
+	NewChildren.Add(GetOrCreateEmitterSummaryGroup(CurrentChildren));
+
 	UNiagaraStackEntry* EmitterSpawnEntry = GetCurrentScriptGroup(CurrentChildren, ENiagaraScriptUsage::EmitterSpawnScript, FGuid());
 	if (EmitterSpawnEntry == nullptr)
 	{
diff --git a/Engine/Plugins/FX/Niagara/Source/NiagaraEditor/Private/ViewModels/Stack/NiagaraStackSummaryViewInputCollection.cpp b/Engine/Plugins/FX/Niagara/Source/NiagaraEditor/Private/ViewModels/Stack/NiagaraStackSummaryViewInputCollection.cpp
--- a/Engine/Plugins/FX/Niagara/Source/NiagaraEditor/Private/ViewModels/Stack/NiagaraStackSummaryViewInputCollection.cpp
+++ b/Engine/Plugins/FX/Niagara/Source/NiagaraEditor/Private/ViewModels/Stack/NiagaraStackSummaryViewInputCollection.cpp
@@ -86,12 +86,7 @@ void UNiagaraStackSummaryViewCollection::RefreshForAdvancedToggle()
 
 bool UNiagaraStackSummaryViewCollection::GetShouldShowInStack() const
 {
-	if(GetEmitterViewModel().IsValid())
-	{
-		return GetEmitterViewModel()->GetEditorData().ShouldShowSummaryView();
-	}
-
-	return false;
+	return true;
 }
 
 void UNiagaraStackSummaryViewCollection::RefreshChildrenInternal(const TArray<UNiagaraStackEntry*>& CurrentChildren, TArray<UNiagaraStackEntry*>& NewChildren, TArray<FStackIssue>& NewIssues)
diff --git a/Engine/Plugins/FX/Niagara/Source/NiagaraEditor/Public/ViewModels/Stack/NiagaraStackEmitterSettingsGroup.h b/Engine/Plugins/FX/Niagara/Source/NiagaraEditor/Public/ViewModels/Stack/NiagaraStackEmitterSettingsGroup.h
--- a/Engine/Plugins/FX/Niagara/Source/NiagaraEditor/Public/ViewModels/Stack/NiagaraStackEmitterSettingsGroup.h
+++ b/Engine/Plugins/FX/Niagara/Source/NiagaraEditor/Public/ViewModels/Stack/NiagaraStackEmitterSettingsGroup.h
@@ -100,7 +100,7 @@ public:
 
 protected:
 	virtual bool GetShouldShowInStack() const override { return false; }
-	virtual bool GetShouldShowInOverview() const override { return false; }
+	virtual bool GetShouldShowInOverview() const override { return true; }
 	
 	NIAGARAEDITOR_API virtual void RefreshChildrenInternal(const TArray<UNiagaraStackEntry*>& CurrentChildren, TArray<UNiagaraStackEntry*>& NewChildren, TArray<FStackIssue>& NewIssues) override;
 
