- Fix renderer post tick order issue when emitters are disabled

diff --git a/Engine/Plugins/FX/Niagara/Source/Niagara/Private/NiagaraSystem.cpp b/Engine/Plugins/FX/Niagara/Source/Niagara/Private/NiagaraSystem.cpp
--- a/Engine/Plugins/FX/Niagara/Source/Niagara/Private/NiagaraSystem.cpp
+++ b/Engine/Plugins/FX/Niagara/Source/Niagara/Private/NiagaraSystem.cpp
@@ -2045,7 +2045,6 @@ void UNiagaraSystem::ComputeEmittersExecutionOrder()
 
 		if (!EmitterHandle.GetIsEnabled())
 		{
-			EmitterData->ForEachEnabledRenderer([&] (const UNiagaraRendererProperties*) { ++SystemRendererIndex; });
 			continue;
 		}
 
