[Niagara] Fail to compile when writing to attributes from the wrong namespace

diff --git a/Engine/Plugins/FX/Niagara/Source/NiagaraEditor/Private/NiagaraCompilationTasks.cpp b/Engine/Plugins/FX/Niagara/Source/NiagaraEditor/Private/NiagaraCompilationTasks.cpp
--- a/Engine/Plugins/FX/Niagara/Source/NiagaraEditor/Private/NiagaraCompilationTasks.cpp
+++ b/Engine/Plugins/FX/Niagara/Source/NiagaraEditor/Private/NiagaraCompilationTasks.cpp
@@ -959,7 +959,10 @@ void FNiagaraSystemCompilationTask::FCompileTaskInfo::RetrieveTranslateResult()
 		}
 	}
 
-	ExeData->LastCompileEvents.Append(TranslateResults.CompileEvents);
+	for (const FNiagaraCompileEvent& Event : TranslateResults.CompileEvents)
+	{
+		ExeData->LastCompileEvents.AddUnique(Event);
+	}
 	ExeData->ExternalDependencies = TranslateResults.CompileDependencies;
 	ExeData->CompileTags = TranslateResults.CompileTags;
 	ExeData->CompileTagsEditorOnly = TranslateResults.CompileTagsEditorOnly;
diff --git a/Engine/Plugins/FX/Niagara/Source/NiagaraEditor/Private/NiagaraHlslTranslator.cpp b/Engine/Plugins/FX/Niagara/Source/NiagaraEditor/Private/NiagaraHlslTranslator.cpp
--- a/Engine/Plugins/FX/Niagara/Source/NiagaraEditor/Private/NiagaraHlslTranslator.cpp
+++ b/Engine/Plugins/FX/Niagara/Source/NiagaraEditor/Private/NiagaraHlslTranslator.cpp
@@ -5288,6 +5288,30 @@ bool FNiagaraHlslTranslationStage::IsExternalConstantNamespace(const FNiagaraVar
 	return false;
 }
 
+bool FNiagaraHlslTranslator::IsWriteAllowedForNamespace(const FNiagaraVariable& Var, ENiagaraScriptUsage TargetUsage, FText& ErrorMsg)
+{
+	if (UNiagaraScript::IsStandaloneScript(TargetUsage))
+	{
+		return true;
+	}
+	if (UNiagaraScript::IsSystemScript(TargetUsage) && (Var.IsInNameSpace(FNiagaraConstants::EmitterNamespace) || Var.IsInNameSpace(FNiagaraConstants::ParticleAttributeNamespace)))
+	{
+		ErrorMsg = FText::Format(LOCTEXT("WriteAllowedForNamespaceFail_System", "Cannot set variable {0} in system scripts."), FText::FromName(Var.GetName()));
+		return false;
+	}
+	if (UNiagaraScript::IsEmitterScript(TargetUsage) && (Var.IsInNameSpace(FNiagaraConstants::SystemNamespace) || Var.IsInNameSpace(FNiagaraConstants::ParticleAttributeNamespace)))
+	{
+		ErrorMsg = FText::Format(LOCTEXT("WriteAllowedForNamespaceFail_Emitter", "Cannot set variable {0} in emitter scripts."), FText::FromName(Var.GetName()));
+		return false;
+	}
+	if (UNiagaraScript::IsParticleScript(TargetUsage) && (Var.IsInNameSpace(FNiagaraConstants::SystemNamespace) || Var.IsInNameSpace(FNiagaraConstants::EmitterNamespace)))
+	{
+		ErrorMsg = FText::Format(LOCTEXT("WriteAllowedForNamespaceFail_Particles", "Cannot set variable {0} in particle scripts."), FText::FromName(Var.GetName()));
+		return false;
+	}
+	return true;
+}
+
 template<typename GraphBridge>
 bool FNiagaraHlslTranslationStage::IsRelevantToSpawnForStage(const typename GraphBridge::FParamMapHistory& InHistory, const FNiagaraVariable& InAliasedVar, const FNiagaraVariable& InVar) const
 {
@@ -5649,12 +5673,19 @@ void TNiagaraHlslTranslator<GraphBridge>::ParameterMapSet(const FParamMapSetNode
 
 			if (!AddStructToDefinitionSet(Var.GetType()))
 			{
-				Error(FText::Format(LOCTEXT("ParameterMapSetTypeError", "Cannot handle type {0}! Variable: {1}"), Var.GetType().GetNameText(), FText::FromName(Var.GetName())));
+				Error(FText::Format(LOCTEXT("ParameterMapSetTypeError", "Cannot handle type {0}! Variable: {1}"), Var.GetType().GetNameText(), FText::FromName(Var.GetName())), SetNode, Inputs[i].Pin);
 			}
 
 			if (TranslationStages[ActiveStageIdx].IsExternalConstantNamespace(Var, CompileOptions.TargetUsage, CompileOptions.GetTargetUsageBitmask()))
 			{
-				Error(FText::Format(LOCTEXT("SetSystemConstantFail", "Cannot Set external constant, Type: {0} Variable: {1}"), Var.GetType().GetNameText(), FText::FromName(Var.GetName())), SetNode, nullptr);
+				Error(FText::Format(LOCTEXT("SetSystemConstantFail", "Cannot Set external constant, Type: {0} Variable: {1}"), Var.GetType().GetNameText(), FText::FromName(Var.GetName())), SetNode, Inputs[i].Pin);
+				continue;
+			}
+
+			FText WriteErrorText;
+			if (!IsWriteAllowedForNamespace(Var, ActiveHistoryForFunctionCalls.GetCurrentUsageContext(), WriteErrorText))
+			{
+				Error(WriteErrorText, SetNode, Inputs[i].Pin);
 				continue;
 			}
 
@@ -5664,12 +5695,12 @@ void TNiagaraHlslTranslator<GraphBridge>::ParameterMapSet(const FParamMapSetNode
 			if (ConstantInfo.ConstantVar != nullptr && ConstantInfo.ConstantVar->GetType() != Var.GetType() && ConstantInfo.ConstantType != ENiagaraKnownConstantType::Attribute)
 			{
 				Error(FText::Format(LOCTEXT("MismatchedConstantTypes", "Variable {0} is a system constant, but its type is different! {1} != {2}"), FText::FromName(Var.GetName()),
-					ConstantInfo.ConstantVar->GetType().GetNameText(), Var.GetType().GetNameText()));
+					ConstantInfo.ConstantVar->GetType().GetNameText(), Var.GetType().GetNameText()), SetNode, Inputs[i].Pin);
 			}
 
 			if (FNiagaraConstants::IsEngineManagedAttribute(Var))
 			{
-				Error(FText::Format(LOCTEXT("SettingSystemAttr", "Variable {0} is an engine managed particle attribute and cannot be set directly."), FText::FromName(Var.GetName())));
+				Error(FText::Format(LOCTEXT("SettingSystemAttr", "Variable {0} is an engine managed particle attribute and cannot be set directly."), FText::FromName(Var.GetName())), SetNode, Inputs[i].Pin);
 				continue;
 			}
 
diff --git a/Engine/Plugins/FX/Niagara/Source/NiagaraEditor/Private/NiagaraHlslTranslator.h b/Engine/Plugins/FX/Niagara/Source/NiagaraEditor/Private/NiagaraHlslTranslator.h
--- a/Engine/Plugins/FX/Niagara/Source/NiagaraEditor/Private/NiagaraHlslTranslator.h
+++ b/Engine/Plugins/FX/Niagara/Source/NiagaraEditor/Private/NiagaraHlslTranslator.h
@@ -471,6 +471,7 @@ protected:
 	bool IsBulkSystemScript() const;
 	bool IsSpawnScript() const;
 	bool RequiresInterpolation() const;
+	static bool IsWriteAllowedForNamespace(const FNiagaraVariable& Var, ENiagaraScriptUsage TargetUsage, FText& ErrorMsg);
 
 	template<typename T>
 	void BuildConstantBuffer(ENiagaraCodeChunkMode ChunkMode);
diff --git a/Engine/Plugins/FX/Niagara/Source/NiagaraShader/Public/NiagaraShared.h b/Engine/Plugins/FX/Niagara/Source/NiagaraShader/Public/NiagaraShared.h
--- a/Engine/Plugins/FX/Niagara/Source/NiagaraShader/Public/NiagaraShared.h
+++ b/Engine/Plugins/FX/Niagara/Source/NiagaraShader/Public/NiagaraShared.h
@@ -110,6 +110,10 @@ public:
 	UPROPERTY()
 	FNiagaraCompileEventSource Source;
 
+	bool operator==(const FNiagaraCompileEvent& Other) const
+	{
+		return Other.Severity == Severity && Other.Message == Message && Other.ShortDescription == ShortDescription && Other.NodeGuid == NodeGuid && Other.PinGuid == PinGuid && Other.Source == Source;
+	}
 };
 
 //
