Fixed a bug where new parameter collections were only visible in Niagara systems after an editor restart

diff --git a/Engine/Plugins/FX/Niagara/Source/NiagaraEditor/Private/NiagaraEditorModule.cpp b/Engine/Plugins/FX/Niagara/Source/NiagaraEditor/Private/NiagaraEditorModule.cpp
--- a/Engine/Plugins/FX/Niagara/Source/NiagaraEditor/Private/NiagaraEditorModule.cpp
+++ b/Engine/Plugins/FX/Niagara/Source/NiagaraEditor/Private/NiagaraEditorModule.cpp
@@ -1062,6 +1062,17 @@ void FNiagaraEditorModule::StartupModule()
 			ParameterCollectionAssetCache.RefreshCache(true /*bAllowLoading*/);
 			ParameterDefinitionsAssetCache.RefreshCache(true /*bAllowLoading*/);
 		});
+		AssetRegistryModule.Get().OnAssetAdded().AddLambda([this](const FAssetData& InAssetData)
+		{
+			if (InAssetData.IsInstanceOf(UNiagaraParameterCollection::StaticClass()))
+			{
+				ParameterCollectionAssetCache.RefreshCache(false);
+			}
+			else if (InAssetData.IsInstanceOf(UNiagaraParameterDefinitions::StaticClass()))
+			{
+				ParameterDefinitionsAssetCache.RefreshCache(false);
+			}
+		});
 	}
 
 	UNiagaraSettings::OnSettingsChanged().AddRaw(this, &FNiagaraEditorModule::OnNiagaraSettingsChangedEvent);
@@ -1527,6 +1538,7 @@ void FNiagaraEditorModule::ShutdownModule()
 	{
 		FAssetRegistryModule& AssetRegistryModule = FModuleManager::GetModuleChecked<FAssetRegistryModule>("AssetRegistry");
 		AssetRegistryModule.Get().OnFilesLoaded().RemoveAll(this);
+		AssetRegistryModule.Get().OnAssetAdded().RemoveAll(this);
 	}
 
 	UNiagaraSettings::OnSettingsChanged().RemoveAll(this);
