Adjust MIC creation for niagara renderers to avoid unnecessary allocations of MIC.  This restructure fixes a problem with the AsyncTask compilation mode where a previously created MIC would interfere with the cook because the unused version would be cached, and wouldn't get reprocessed after it was recreated.

diff --git a/Engine/Plugins/FX/Niagara/Source/Niagara/Private/NiagaraRendererProperties.cpp b/Engine/Plugins/FX/Niagara/Source/Niagara/Private/NiagaraRendererProperties.cpp
--- a/Engine/Plugins/FX/Niagara/Source/Niagara/Private/NiagaraRendererProperties.cpp
+++ b/Engine/Plugins/FX/Niagara/Source/Niagara/Private/NiagaraRendererProperties.cpp
@@ -416,8 +416,10 @@ void UNiagaraRendererProperties::ChangeToPositionBinding(FNiagaraVariableAttribu
 	}
 }
 
-bool UNiagaraRendererProperties::UpdateMaterialStaticParameters(const FNiagaraRendererMaterialParameters& MaterialParameters, UMaterialInstanceConstant* MIC)
+bool UNiagaraRendererProperties::BuildMaterialStaticParameterSet(const FNiagaraRendererMaterialParameters& MaterialParameters, const UMaterialInterface* Material, FStaticParameterSet& StaticParameterSet) const
 {
+	StaticParameterSet.Empty();
+
 	UNiagaraSystem* NiagaraSystem = GetTypedOuter<UNiagaraSystem>();
 	if (NiagaraSystem == nullptr)
 	{
@@ -428,9 +430,8 @@ bool UNiagaraRendererProperties::UpdateMaterialStaticParameters(const FNiagaraRe
 	TArray<FMaterialParameterInfo> AllStaticSwitchParameterInfos;
 	{
 		TArray<FGuid> ParameterGuids;
-		MIC->GetAllStaticSwitchParameterInfo(AllStaticSwitchParameterInfos, ParameterGuids);
+		Material->GetAllStaticSwitchParameterInfo(AllStaticSwitchParameterInfos, ParameterGuids);
 	}
-	FStaticParameterSet StaticParameterSet;
 
 	bool bModified = false;
 	for (const FNiagaraRendererMaterialStaticBoolParameter& ParameterBinding : MaterialParameters.StaticBoolParameters)
@@ -438,7 +439,7 @@ bool UNiagaraRendererProperties::UpdateMaterialStaticParameters(const FNiagaraRe
 		NiagaraSystem->ForEachScript(
 			[&](UNiagaraScript* NiagaraScript)
 			{
-				for ( const FNiagaraVariable& StaticVariable : NiagaraScript->GetVMExecutableData().StaticVariablesWritten )
+				for (const FNiagaraVariable& StaticVariable : NiagaraScript->GetVMExecutableData().StaticVariablesWritten)
 				{
 					if (StaticVariable.GetType() != FNiagaraTypeDefinition::GetBoolDef().ToStaticDef())
 					{
@@ -474,7 +475,7 @@ bool UNiagaraRendererProperties::UpdateMaterialStaticParameters(const FNiagaraRe
 						{
 							FGuid ParameterGuid;
 							bool bDefaultValue = false;
-							if (MIC->GetStaticSwitchParameterDefaultValue(ParameterInfo, bDefaultValue, ParameterGuid))
+							if (Material->GetStaticSwitchParameterDefaultValue(ParameterInfo, bDefaultValue, ParameterGuid))
 							{
 								if (bDefaultValue != bNewValue)
 								{
@@ -497,12 +498,19 @@ bool UNiagaraRendererProperties::UpdateMaterialStaticParameters(const FNiagaraRe
 		);
 	}
 
-	if (bModified)
+	return bModified;
+}
+
+bool UNiagaraRendererProperties::UpdateMaterialStaticParameters(const FNiagaraRendererMaterialParameters& MaterialParameters, UMaterialInstanceConstant* MIC)
+{
+	FStaticParameterSet StaticParameterSet;
+	if (BuildMaterialStaticParameterSet(MaterialParameters, MIC, StaticParameterSet))
 	{
 		MIC->UpdateStaticPermutation(StaticParameterSet);
+		return true;
 	}
 
-	return bModified;
+	return false;
 }
 
 void UNiagaraRendererProperties::UpdateMaterialParametersMIC(const FNiagaraRendererMaterialParameters& MaterialParameters, TObjectPtr<UMaterialInterface>& InOutMaterial, TObjectPtr<UMaterialInstanceConstant>& InOutMIC)
@@ -513,22 +521,28 @@ void UNiagaraRendererProperties::UpdateMaterialParametersMIC(const FNiagaraRende
 		return;
 	}
 
-	FNameBuilder NameBuilder;
-	InOutMaterial->GetFName().ToString(NameBuilder);
-	NameBuilder.Append(TEXT("_MIC"));
-	if (InOutMIC == nullptr)
+	FStaticParameterSet MaterialParameterSet;
+	if (BuildMaterialStaticParameterSet(MaterialParameters, InOutMaterial, MaterialParameterSet))
 	{
-		InOutMIC = NewObject<UMaterialInstanceConstant>(this, FName(NameBuilder));
-		InOutMIC->SetParentEditorOnly(InOutMaterial);
-	}
-	else if (InOutMIC->Parent != InOutMaterial)
-	{
-		InOutMIC->Rename(NameBuilder.ToString());
-		InOutMIC->SetParentEditorOnly(InOutMaterial);
-	}
+		FNameBuilder NameBuilder;
+		InOutMaterial->GetFName().ToString(NameBuilder);
+		NameBuilder.Append(TEXT("_MIC"));
+		if (InOutMIC == nullptr)
+		{
+			InOutMIC = NewObject<UMaterialInstanceConstant>(this, FName(NameBuilder));
+			InOutMIC->SetParentEditorOnly(InOutMaterial);
+		}
+		else if (InOutMIC->Parent != InOutMaterial)
+		{
+			InOutMIC->Rename(NameBuilder.ToString());
+			InOutMIC->SetParentEditorOnly(InOutMaterial);
+		}
 
-	if (UpdateMaterialStaticParameters(MaterialParameters, InOutMIC) == false)
+		InOutMIC->UpdateStaticPermutation(MaterialParameterSet);
+	}
+	else if (InOutMIC)
 	{
+		InOutMIC->MarkAsGarbage();
 		InOutMIC = nullptr;
 	}
 }
@@ -562,44 +576,43 @@ void UNiagaraRendererProperties::UpdateMaterialParametersMIC(const FNiagaraRende
 			continue;
 		}
 
-		//-OPT: We should be able to reuse rather than create
-		FNameBuilder NameBuilder;
-		Material->GetFName().ToString(NameBuilder);
-		NameBuilder.Append(TEXT("_MIC"));
-
-		UMaterialInstanceConstant* MIC = nullptr;
-		bool bNeedsRename = false;
-		if (MICPool.Num() > 0)
+		FStaticParameterSet MaterialParameterSet;
+		if (BuildMaterialStaticParameterSet(MaterialParameters, Material, MaterialParameterSet))
 		{
-			FName MICName(NameBuilder);
-			const int32 ExistingIndex = MICPool.IndexOfByPredicate([&MICName](UMaterialInstanceConstant* MIC) { return MIC->GetFName() == MICName; });
-			if (ExistingIndex != INDEX_NONE)
+			//-OPT: We should be able to reuse rather than create
+			FNameBuilder NameBuilder;
+			Material->GetFName().ToString(NameBuilder);
+			NameBuilder.Append(TEXT("_MIC"));
+
+			UMaterialInstanceConstant* MIC = nullptr;
+			bool bNeedsRename = false;
+			if (MICPool.Num() > 0)
 			{
-				MIC = MICPool[ExistingIndex];
-				MICPool.RemoveAtSwap(ExistingIndex, 1, EAllowShrinking::No);
+				FName MICName(NameBuilder);
+				const int32 ExistingIndex = MICPool.IndexOfByPredicate([&MICName](UMaterialInstanceConstant* MIC) { return MIC->GetFName() == MICName; });
+				if (ExistingIndex != INDEX_NONE)
+				{
+					MIC = MICPool[ExistingIndex];
+					MICPool.RemoveAtSwap(ExistingIndex, 1, EAllowShrinking::No);
+				}
+				else
+				{
+					bNeedsRename = true;
+					MIC = MICPool.Pop();
+				}
+				if (MIC->Parent != Material)
+				{
+					MIC->SetParentEditorOnly(Material);
+				}
 			}
 			else
 			{
-				bNeedsRename = true;
-				MIC = MICPool.Pop();
-			}
-			if (MIC->Parent != Material)
-			{
+				MIC = NewObject<UMaterialInstanceConstant>(this, FName(NameBuilder));
 				MIC->SetParentEditorOnly(Material);
 			}
-		}
-		else
-		{
-			MIC = NewObject<UMaterialInstanceConstant>(this, FName(NameBuilder));
-			MIC->SetParentEditorOnly(Material);
-		}
 
-		if (UpdateMaterialStaticParameters(MaterialParameters, MIC) == false)
-		{
-			MICPool.Add(MIC);
-		}
-		else
-		{
+			MIC->UpdateStaticPermutation(MaterialParameterSet);
+
 			InOutMICs.SetNum(i + 1);
 			InOutMICs[i] = MIC;
 			if (bNeedsRename)
diff --git a/Engine/Plugins/FX/Niagara/Source/Niagara/Public/NiagaraRendererProperties.h b/Engine/Plugins/FX/Niagara/Source/Niagara/Public/NiagaraRendererProperties.h
--- a/Engine/Plugins/FX/Niagara/Source/Niagara/Public/NiagaraRendererProperties.h
+++ b/Engine/Plugins/FX/Niagara/Source/Niagara/Public/NiagaraRendererProperties.h
@@ -25,6 +25,7 @@ class SWidget;
 class FAssetThumbnailPool;
 struct FNiagaraDataSetCompiledData;
 struct FSlateBrush;
+struct FStaticParameterSet;
 struct FStreamingRenderAssetPrimitiveInfo;
 
 #if WITH_EDITOR
@@ -454,7 +455,11 @@ protected:
 	/** utility function that can be used to fix up old vec3 bindings into position bindings. */
 	static NIAGARA_API void ChangeToPositionBinding(FNiagaraVariableAttributeBinding& Binding);
 
+	/** Generates the static parameter set for the parent emitter */
+	bool BuildMaterialStaticParameterSet(const FNiagaraRendererMaterialParameters& MaterialParameters, const UMaterialInterface* Material, FStaticParameterSet& StaticParameterSet) const;
+
 	/** Update MIC Static Parameters. */
+	UE_DEPRECATED(5.5, "This helper function has been deprecated.  Use BuildMaterialStaticParameterSet instead.")
 	NIAGARA_API bool UpdateMaterialStaticParameters(const FNiagaraRendererMaterialParameters& MaterialParameters, UMaterialInstanceConstant* MIC);
 
 	/** Utility function to updates MICs. */
