Allow new async task compilation mode to not block GC

diff --git a/Engine/Plugins/FX/Niagara/Source/Niagara/Classes/NiagaraSystem.h b/Engine/Plugins/FX/Niagara/Source/Niagara/Classes/NiagaraSystem.h
--- a/Engine/Plugins/FX/Niagara/Source/Niagara/Classes/NiagaraSystem.h
+++ b/Engine/Plugins/FX/Niagara/Source/Niagara/Classes/NiagaraSystem.h
@@ -395,6 +395,8 @@ public:
 	/** Are there any pending compile requests?*/
 	NIAGARA_API bool HasOutstandingCompilationRequests(bool bIncludingGPUShaders = false) const;
 
+	NIAGARA_API bool CompileRequestsShouldBlockGC() const;
+
 	/** Determines if this system has the supplied emitter as an editable and simulating emitter instance. */
 	NIAGARA_API bool ReferencesInstanceEmitter(const FVersionedNiagaraEmitter& Emitter) const;
 
diff --git a/Engine/Plugins/FX/Niagara/Source/Niagara/Private/NiagaraCompilationTypes.cpp b/Engine/Plugins/FX/Niagara/Source/Niagara/Private/NiagaraCompilationTypes.cpp
--- a/Engine/Plugins/FX/Niagara/Source/Niagara/Private/NiagaraCompilationTypes.cpp
+++ b/Engine/Plugins/FX/Niagara/Source/Niagara/Private/NiagaraCompilationTypes.cpp
@@ -407,6 +407,11 @@ public:
 		}
 	}
 
+	virtual bool BlocksGarbageCollection() const override
+	{
+		return false;
+	}
+
 	FNiagaraCompilationTaskHandle CompileRequestHandle;
 	FNiagaraSystemAsyncCompileResults CompileResults;
 };
diff --git a/Engine/Plugins/FX/Niagara/Source/Niagara/Private/NiagaraSystem.cpp b/Engine/Plugins/FX/Niagara/Source/Niagara/Private/NiagaraSystem.cpp
--- a/Engine/Plugins/FX/Niagara/Source/Niagara/Private/NiagaraSystem.cpp
+++ b/Engine/Plugins/FX/Niagara/Source/Niagara/Private/NiagaraSystem.cpp
@@ -1860,6 +1860,20 @@ bool UNiagaraSystem::HasOutstandingCompilationRequests(bool bIncludingGPUShaders
 
 	return false;
 }
+
+bool UNiagaraSystem::CompileRequestsShouldBlockGC() const
+{
+	for (const TUniquePtr<FNiagaraActiveCompilation>& ActiveCompilation : ActiveCompilations)
+	{
+		if (ActiveCompilation && ActiveCompilation->BlocksGarbageCollection())
+		{
+			return true;
+		}
+	}
+
+	return false;
+}
+
 #endif
 
 bool UNiagaraSystem::ComputeEmitterPriority(int32 EmitterIdx, TArray<int32, TInlineAllocator<32>>& EmitterPriorities, const TBitArray<TInlineAllocator<32>>& EmitterDependencyGraph)
diff --git a/Engine/Plugins/FX/Niagara/Source/Niagara/Public/NiagaraCompilationTypes.h b/Engine/Plugins/FX/Niagara/Source/Niagara/Public/NiagaraCompilationTypes.h
--- a/Engine/Plugins/FX/Niagara/Source/Niagara/Public/NiagaraCompilationTypes.h
+++ b/Engine/Plugins/FX/Niagara/Source/Niagara/Public/NiagaraCompilationTypes.h
@@ -114,6 +114,7 @@ public:
 	virtual void Apply(const FNiagaraQueryCompilationOptions& Options) = 0;
 	virtual void ReportResults(const FNiagaraQueryCompilationOptions& Options) const = 0;
 	virtual bool BlocksBeginCacheForCooked() const { return false; }
+	virtual bool BlocksGarbageCollection() const { return true; }
 
 	void Invalidate()
 	{
diff --git a/Engine/Plugins/FX/Niagara/Source/NiagaraEditor/Private/NiagaraEditorModule.cpp b/Engine/Plugins/FX/Niagara/Source/NiagaraEditor/Private/NiagaraEditorModule.cpp
--- a/Engine/Plugins/FX/Niagara/Source/NiagaraEditor/Private/NiagaraEditorModule.cpp
+++ b/Engine/Plugins/FX/Niagara/Source/NiagaraEditor/Private/NiagaraEditorModule.cpp
@@ -2206,10 +2206,12 @@ void FNiagaraEditorModule::OnPreGarbageCollection()
 		// release is due to go out).
 		for (TObjectIterator<UNiagaraSystem> It; It; ++It)
 		{
-			UNiagaraSystem* System = *It;
-			if (System && System->HasOutstandingCompilationRequests())
+			if (UNiagaraSystem* System = *It)
 			{
-				System->WaitForCompilationComplete();
+				if (System->CompileRequestsShouldBlockGC())
+				{
+					System->WaitForCompilationComplete();
+				}
 			}
 		}
 	}
