Add support for specular scale with the simple lights.  Includes support for Niagara to set the attribute per light through the light renderer.

#virtualized

diff --git a/Engine/Plugins/FX/Niagara/Source/Niagara/Classes/NiagaraConstants.h b/Engine/Plugins/FX/Niagara/Source/Niagara/Classes/NiagaraConstants.h
--- a/Engine/Plugins/FX/Niagara/Source/Niagara/Classes/NiagaraConstants.h
+++ b/Engine/Plugins/FX/Niagara/Source/Niagara/Classes/NiagaraConstants.h
@@ -127,6 +127,7 @@
 #define SYS_PARAM_PARTICLES_LIGHT_EXPONENT               INiagaraModule::GetVar_Particles_LightExponent()
 #define SYS_PARAM_PARTICLES_LIGHT_ENABLED                INiagaraModule::GetVar_Particles_LightEnabled()
 #define SYS_PARAM_PARTICLES_LIGHT_VOLUMETRIC_SCATTERING  INiagaraModule::GetVar_Particles_LightVolumetricScattering()
+#define SYS_PARAM_PARTICLES_LIGHT_SPECULAR_SCALE         INiagaraModule::GetVar_Particles_LightSpecularScale()
 #define SYS_PARAM_PARTICLES_VISIBILITY_TAG				 INiagaraModule::GetVar_Particles_VisibilityTag()
 #define SYS_PARAM_PARTICLES_MESH_INDEX					 INiagaraModule::GetVar_Particles_MeshIndex()
 #define SYS_PARAM_PARTICLES_COMPONENTS_ENABLED           INiagaraModule::GetVar_Particles_ComponentsEnabled()
diff --git a/Engine/Plugins/FX/Niagara/Source/Niagara/Private/NiagaraConstants.cpp b/Engine/Plugins/FX/Niagara/Source/Niagara/Private/NiagaraConstants.cpp
--- a/Engine/Plugins/FX/Niagara/Source/Niagara/Private/NiagaraConstants.cpp
+++ b/Engine/Plugins/FX/Niagara/Source/Niagara/Private/NiagaraConstants.cpp
@@ -349,6 +349,7 @@ void FNiagaraConstants::Init()
 		AttrDataSetKeyMap.Add(SYS_PARAM_PARTICLES_LIGHT_RADIUS, GetAttributeAsParticleDataSetKey(SYS_PARAM_PARTICLES_LIGHT_RADIUS));
 		AttrDataSetKeyMap.Add(SYS_PARAM_PARTICLES_LIGHT_EXPONENT, GetAttributeAsParticleDataSetKey(SYS_PARAM_PARTICLES_LIGHT_EXPONENT));
 		AttrDataSetKeyMap.Add(SYS_PARAM_PARTICLES_LIGHT_VOLUMETRIC_SCATTERING, GetAttributeAsParticleDataSetKey(SYS_PARAM_PARTICLES_LIGHT_VOLUMETRIC_SCATTERING));
+		AttrDataSetKeyMap.Add(SYS_PARAM_PARTICLES_LIGHT_SPECULAR_SCALE, GetAttributeAsParticleDataSetKey(SYS_PARAM_PARTICLES_LIGHT_SPECULAR_SCALE));
 		AttrDataSetKeyMap.Add(SYS_PARAM_PARTICLES_VISIBILITY_TAG, GetAttributeAsParticleDataSetKey(SYS_PARAM_PARTICLES_VISIBILITY_TAG));
 		AttrDataSetKeyMap.Add(SYS_PARAM_PARTICLES_COMPONENTS_ENABLED, GetAttributeAsParticleDataSetKey(SYS_PARAM_PARTICLES_COMPONENTS_ENABLED));
 		AttrDataSetKeyMap.Add(SYS_PARAM_PARTICLES_RIBBONID, GetAttributeAsParticleDataSetKey(SYS_PARAM_PARTICLES_RIBBONID));
@@ -489,6 +490,11 @@ void FNiagaraConstants::Init()
 		Var.SetValue<float>(0.0f);
 		AttrDefaultsValueMap.Add(SYS_PARAM_PARTICLES_LIGHT_VOLUMETRIC_SCATTERING, Var);
 
+		AttrDefaultsStrMap.Add(SYS_PARAM_PARTICLES_LIGHT_SPECULAR_SCALE, TEXT("1.0"));
+		Var = SYS_PARAM_PARTICLES_LIGHT_SPECULAR_SCALE;
+		Var.SetValue<float>(1.0f);
+		AttrDefaultsValueMap.Add(SYS_PARAM_PARTICLES_LIGHT_SPECULAR_SCALE, Var);
+
 		AttrDefaultsStrMap.Add(SYS_PARAM_PARTICLES_VISIBILITY_TAG, TEXT("0"));
 		Var = SYS_PARAM_PARTICLES_VISIBILITY_TAG;
 		Var.SetValue<float>(0.0f);
@@ -578,6 +584,7 @@ void FNiagaraConstants::Init()
 		AttrDescStrMap.Add(SYS_PARAM_PARTICLES_LIGHT_ENABLED, LOCTEXT("LightEnabledParamDesc", "This can be used with the Light renderer to dynamically control if a light exists on a per-particle basis."));
 		AttrDescStrMap.Add(SYS_PARAM_PARTICLES_LIGHT_EXPONENT, LOCTEXT("LightExponentParamDesc", "Used to drive the attenuation of the light when using a Light renderer without inverse squared falloff enabled."));
 		AttrDescStrMap.Add(SYS_PARAM_PARTICLES_LIGHT_VOLUMETRIC_SCATTERING, LOCTEXT("LightVolumetricScatteringParamDesc", "Used to drive the volumetric scattering intensity of the light when using a Light renderer."));
+		AttrDescStrMap.Add(SYS_PARAM_PARTICLES_LIGHT_SPECULAR_SCALE, LOCTEXT("LightSpecularScaleParamDesc", "Used to drive the specular scale of the light when using a Light renderer."));
 		AttrDescStrMap.Add(SYS_PARAM_INSTANCE_ALIVE, LOCTEXT("AliveParamDesc", "Used to determine whether or not this particle instance is still valid or if it can be deleted."));
 		AttrDescStrMap.Add(SYS_PARAM_PARTICLES_VISIBILITY_TAG, LOCTEXT("VisibilityTag", "Used for selecting renderers to use when rendering this particle. Without this, the particle will render in all renderers"));
 		AttrDescStrMap.Add(SYS_PARAM_PARTICLES_COMPONENTS_ENABLED, LOCTEXT("ComponentRenderEnabledParamDesc", "Used to check if component rendering should be enabled on a per-particle basis. Without this, the each particle will spawn a component."));
diff --git a/Engine/Plugins/FX/Niagara/Source/Niagara/Private/NiagaraLightRendererProperties.cpp b/Engine/Plugins/FX/Niagara/Source/Niagara/Private/NiagaraLightRendererProperties.cpp
--- a/Engine/Plugins/FX/Niagara/Source/Niagara/Private/NiagaraLightRendererProperties.cpp
+++ b/Engine/Plugins/FX/Niagara/Source/Niagara/Private/NiagaraLightRendererProperties.cpp
@@ -40,6 +40,7 @@ UNiagaraLightRendererProperties::UNiagaraLightRendererProperties()
 		&RadiusBinding,
 		&VolumetricScatteringBinding,
 		&RendererVisibilityTagBinding,
+		&SpecularScaleBinding,
 	};
 }
 
@@ -74,6 +75,7 @@ void UNiagaraLightRendererProperties::PostInitProperties()
 			LightRenderingEnabledBinding = FNiagaraConstants::GetAttributeDefaultBinding(SYS_PARAM_PARTICLES_LIGHT_ENABLED);
 			VolumetricScatteringBinding = FNiagaraConstants::GetAttributeDefaultBinding(SYS_PARAM_PARTICLES_LIGHT_VOLUMETRIC_SCATTERING);
 			RendererVisibilityTagBinding = FNiagaraConstants::GetAttributeDefaultBinding(SYS_PARAM_PARTICLES_VISIBILITY_TAG);
+			SpecularScaleBinding = FNiagaraConstants::GetAttributeDefaultBinding(SYS_PARAM_PARTICLES_LIGHT_SPECULAR_SCALE);
 		}
 	}
 }
@@ -103,6 +105,7 @@ void UNiagaraLightRendererProperties::InitCDOPropertiesAfterModuleStartup()
 	CDO->LightRenderingEnabledBinding = FNiagaraConstants::GetAttributeDefaultBinding(SYS_PARAM_PARTICLES_LIGHT_ENABLED);
 	CDO->VolumetricScatteringBinding = FNiagaraConstants::GetAttributeDefaultBinding(SYS_PARAM_PARTICLES_LIGHT_VOLUMETRIC_SCATTERING);
 	CDO->RendererVisibilityTagBinding = FNiagaraConstants::GetAttributeDefaultBinding(SYS_PARAM_PARTICLES_VISIBILITY_TAG);
+	CDO->SpecularScaleBinding = FNiagaraConstants::GetAttributeDefaultBinding(SYS_PARAM_PARTICLES_LIGHT_SPECULAR_SCALE);
 
 	for (TWeakObjectPtr<UNiagaraLightRendererProperties>& WeakLightRendererProperties : LightRendererPropertiesToDeferredInit)
 	{
@@ -117,6 +120,7 @@ void UNiagaraLightRendererProperties::InitCDOPropertiesAfterModuleStartup()
 				WeakLightRendererProperties->LightRenderingEnabledBinding = FNiagaraConstants::GetAttributeDefaultBinding(SYS_PARAM_PARTICLES_LIGHT_ENABLED);
 				WeakLightRendererProperties->VolumetricScatteringBinding = FNiagaraConstants::GetAttributeDefaultBinding(SYS_PARAM_PARTICLES_LIGHT_VOLUMETRIC_SCATTERING);
 				WeakLightRendererProperties->RendererVisibilityTagBinding = FNiagaraConstants::GetAttributeDefaultBinding(SYS_PARAM_PARTICLES_VISIBILITY_TAG);
+				WeakLightRendererProperties->SpecularScaleBinding = FNiagaraConstants::GetAttributeDefaultBinding(SYS_PARAM_PARTICLES_LIGHT_SPECULAR_SCALE);
 			}
 		}
 	}
@@ -195,6 +199,7 @@ void UNiagaraLightRendererProperties::CacheFromCompiledData(const FNiagaraDataSe
 	InitParticleDataSetAccessor(ScatteringDataSetAccessor, CompiledData, VolumetricScatteringBinding);
 	InitParticleDataSetAccessor(EnabledDataSetAccessor, CompiledData, LightRenderingEnabledBinding);
 	InitParticleDataSetAccessor(RendererVisibilityTagAccessor, CompiledData, RendererVisibilityTagBinding);
+	InitParticleDataSetAccessor(SpecularScaleAccessor, CompiledData, SpecularScaleBinding);
 }
 
 void UNiagaraLightRendererProperties::UpdateSourceModeDerivates(ENiagaraRendererSourceDataMode InSourceMode, bool bFromPropertyEdit)
@@ -231,6 +236,7 @@ const TArray<FNiagaraVariable>& UNiagaraLightRendererProperties::GetOptionalAttr
 		Attrs.Add(SYS_PARAM_PARTICLES_LIGHT_ENABLED);
 		Attrs.Add(SYS_PARAM_PARTICLES_LIGHT_VOLUMETRIC_SCATTERING);
 		Attrs.Add(SYS_PARAM_PARTICLES_VISIBILITY_TAG);
+		Attrs.Add(SYS_PARAM_PARTICLES_LIGHT_SPECULAR_SCALE);
 	}
 	return Attrs;
 }
diff --git a/Engine/Plugins/FX/Niagara/Source/Niagara/Private/NiagaraModule.cpp b/Engine/Plugins/FX/Niagara/Source/Niagara/Private/NiagaraModule.cpp
--- a/Engine/Plugins/FX/Niagara/Source/Niagara/Private/NiagaraModule.cpp
+++ b/Engine/Plugins/FX/Niagara/Source/Niagara/Private/NiagaraModule.cpp
@@ -274,6 +274,7 @@ FNiagaraVariable INiagaraModule::Particles_LightRadius;
 FNiagaraVariable INiagaraModule::Particles_LightExponent;
 FNiagaraVariable INiagaraModule::Particles_LightEnabled;
 FNiagaraVariable INiagaraModule::Particles_LightVolumetricScattering;
+FNiagaraVariable INiagaraModule::Particles_LightSpecularScale;
 FNiagaraVariable INiagaraModule::Particles_RibbonID;
 FNiagaraVariable INiagaraModule::Particles_RibbonWidth;
 FNiagaraVariable INiagaraModule::Particles_RibbonTwist;
@@ -436,6 +437,7 @@ void INiagaraModule::StartupModule()
 	Particles_LightExponent = FNiagaraVariable(FNiagaraTypeDefinition::GetFloatDef(), TEXT("Particles.LightExponent"));
 	Particles_LightEnabled = FNiagaraVariable(FNiagaraTypeDefinition::GetBoolDef(), TEXT("Particles.LightEnabled"));
 	Particles_LightVolumetricScattering = FNiagaraVariable(FNiagaraTypeDefinition::GetFloatDef(), TEXT("Particles.LightVolumetricScattering"));
+	Particles_LightSpecularScale = FNiagaraVariable(FNiagaraTypeDefinition::GetFloatDef(), TEXT("Particles.LightSpecularScale"));
 	Particles_RibbonID = FNiagaraVariable(FNiagaraTypeDefinition::GetIDDef(), TEXT("Particles.RibbonID"));
 	Particles_RibbonWidth = FNiagaraVariable(FNiagaraTypeDefinition::GetFloatDef(), TEXT("Particles.RibbonWidth"));
 	Particles_RibbonTwist = FNiagaraVariable(FNiagaraTypeDefinition::GetFloatDef(), TEXT("Particles.RibbonTwist"));
diff --git a/Engine/Plugins/FX/Niagara/Source/Niagara/Private/NiagaraRendererLights.cpp b/Engine/Plugins/FX/Niagara/Source/Niagara/Private/NiagaraRendererLights.cpp
--- a/Engine/Plugins/FX/Niagara/Source/Niagara/Private/NiagaraRendererLights.cpp
+++ b/Engine/Plugins/FX/Niagara/Source/Niagara/Private/NiagaraRendererLights.cpp
@@ -117,6 +117,7 @@ FNiagaraDynamicDataBase* FNiagaraRendererLights::GenerateDynamicData(const FNiag
 	const FNiagaraBool DefaultEnabled = ParameterStore.GetParameterValueOrDefault(Properties->LightRenderingEnabledBinding.GetParamMapBindableVariable(), FNiagaraBool(true));
 	const int32 DefaultVisibilityTag = ParameterStore.GetParameterValueOrDefault(Properties->RendererVisibilityTagBinding.GetParamMapBindableVariable(), Properties->RendererVisibility);
 	const float DefaultExponent = ParameterStore.GetParameterValueOrDefault(Properties->LightExponentBinding.GetParamMapBindableVariable(), Properties->DefaultExponent);
+	const float DefaultSpecularScale = ParameterStore.GetParameterValueOrDefault(Properties->SpecularScaleBinding.GetParamMapBindableVariable(), Properties->SpecularScale);
 
 	const float InverseExposureBlend = Properties->bOverrideInverseExposureBlend ? Properties->InverseExposureBlend : GetDefault<UNiagaraSettings>()->DefaultLightInverseExposureBlend;
 
@@ -132,6 +133,7 @@ FNiagaraDynamicDataBase* FNiagaraRendererLights::GenerateDynamicData(const FNiag
 		const auto ScatteringReader = Properties->ScatteringDataSetAccessor.GetReader(Data);
 		const auto EnabledReader = Properties->EnabledDataSetAccessor.GetReader(Data);
 		const auto VisTagReader = Properties->RendererVisibilityTagAccessor.GetReader(Data);
+		const auto SpecularScaleReader = Properties->SpecularScaleAccessor.GetReader(Data);
 
 		for (uint32 ParticleIndex = 0; ParticleIndex < DataToRender->GetNumInstances(); ParticleIndex++)
 		{
@@ -152,6 +154,7 @@ FNiagaraDynamicDataBase* FNiagaraRendererLights::GenerateDynamicData(const FNiag
 				LightData.LightEntry.InverseExposureBlend = InverseExposureBlend;
 				LightData.LightEntry.bAffectTranslucency = Properties->bAffectsTranslucency;
 				LightData.LightEntry.VolumetricScatteringIntensity = ScatteringReader.GetSafe(ParticleIndex, DefaultScattering);
+				LightData.LightEntry.SpecularScale = SpecularScaleReader.GetSafe(ParticleIndex, DefaultSpecularScale);
 				LightData.PerViewEntry.Position = SimToWorld.TransformPosition(FVector(SimPos));
 			}
 		}
@@ -176,6 +179,8 @@ FNiagaraDynamicDataBase* FNiagaraRendererLights::GenerateDynamicData(const FNiag
 			LightData.LightEntry.InverseExposureBlend = InverseExposureBlend;
 			LightData.LightEntry.bAffectTranslucency = Properties->bAffectsTranslucency;
 			LightData.LightEntry.VolumetricScatteringIntensity = LightScattering;
+			LightData.LightEntry.SpecularScale = Properties->SpecularScale;
+
 			LightData.PerViewEntry.Position = SimToWorld.TransformPosition(FVector(SimPos));
 		}
 	}
diff --git a/Engine/Plugins/FX/Niagara/Source/Niagara/Public/NiagaraLightRendererProperties.h b/Engine/Plugins/FX/Niagara/Source/Niagara/Public/NiagaraLightRendererProperties.h
--- a/Engine/Plugins/FX/Niagara/Source/Niagara/Public/NiagaraLightRendererProperties.h
+++ b/Engine/Plugins/FX/Niagara/Source/Niagara/Public/NiagaraLightRendererProperties.h
@@ -78,6 +78,10 @@ public:
 	UPROPERTY(EditAnywhere, Category = "Light Rendering", meta = (ClampMin = "0", EditCondition = "!bUseInverseSquaredFalloff", EditConditionHides))
 	float DefaultExponent;
 
+	/** The specular scale to use for all lights if no binding was found */
+	UPROPERTY(EditAnywhere, Category = "Light Rendering", meta = (ClampMin = "0"))
+	float SpecularScale = 1.0f;
+
 	/** A static color shift applied to each rendered light */
 	UPROPERTY(EditAnywhere, Category = "Light Rendering")
 	FVector3f ColorAdd;
@@ -122,6 +126,10 @@ public:
 	UPROPERTY(EditAnywhere, Category = "Bindings")
 	FNiagaraVariableAttributeBinding RendererVisibilityTagBinding;
 
+	/** Which attribute should we use for the specular scale when generating lights?*/
+	UPROPERTY(EditAnywhere, Category = "Bindings")
+	FNiagaraVariableAttributeBinding SpecularScaleBinding;
+
 	FNiagaraDataSetAccessor<FNiagaraPosition> PositionDataSetAccessor;
 	FNiagaraDataSetAccessor<FLinearColor> ColorDataSetAccessor;
 	FNiagaraDataSetAccessor<float> RadiusDataSetAccessor;
@@ -129,6 +137,7 @@ public:
 	FNiagaraDataSetAccessor<float> ScatteringDataSetAccessor;
 	FNiagaraDataSetAccessor<FNiagaraBool> EnabledDataSetAccessor;
 	FNiagaraDataSetAccessor<int32> RendererVisibilityTagAccessor;
+	FNiagaraDataSetAccessor<float> SpecularScaleAccessor;
 
 private:
 	static TArray<TWeakObjectPtr<UNiagaraLightRendererProperties>> LightRendererPropertiesToDeferredInit;
diff --git a/Engine/Plugins/FX/Niagara/Source/Niagara/Public/NiagaraModule.h b/Engine/Plugins/FX/Niagara/Source/Niagara/Public/NiagaraModule.h
--- a/Engine/Plugins/FX/Niagara/Source/Niagara/Public/NiagaraModule.h
+++ b/Engine/Plugins/FX/Niagara/Source/Niagara/Public/NiagaraModule.h
@@ -250,6 +250,7 @@ public:
 	FORCEINLINE static const FNiagaraVariable&  GetVar_Particles_LightExponent() { return Particles_LightExponent; }
 	FORCEINLINE static const FNiagaraVariable&  GetVar_Particles_LightEnabled() { return Particles_LightEnabled; }
 	FORCEINLINE static const FNiagaraVariable&  GetVar_Particles_LightVolumetricScattering() { return Particles_LightVolumetricScattering; }
+	FORCEINLINE static const FNiagaraVariable&  GetVar_Particles_LightSpecularScale() { return Particles_LightSpecularScale; }
 	FORCEINLINE static const FNiagaraVariable&  GetVar_Particles_RibbonID() { return Particles_RibbonID; }
 	FORCEINLINE static const FNiagaraVariable&  GetVar_Particles_RibbonWidth() { return Particles_RibbonWidth; }
 	FORCEINLINE static const FNiagaraVariable&  GetVar_Particles_RibbonTwist() { return Particles_RibbonTwist; }
@@ -387,6 +388,7 @@ private:
 	static NIAGARA_API FNiagaraVariable Particles_LightExponent;
 	static NIAGARA_API FNiagaraVariable Particles_LightEnabled;
 	static NIAGARA_API FNiagaraVariable Particles_LightVolumetricScattering;
+	static NIAGARA_API FNiagaraVariable Particles_LightSpecularScale;
 	static NIAGARA_API FNiagaraVariable Particles_RibbonID;
 	static NIAGARA_API FNiagaraVariable Particles_RibbonWidth;
 	static NIAGARA_API FNiagaraVariable Particles_RibbonTwist;
diff --git a/Engine/Source/Runtime/Engine/Public/PrimitiveSceneProxy.h b/Engine/Source/Runtime/Engine/Public/PrimitiveSceneProxy.h
--- a/Engine/Source/Runtime/Engine/Public/PrimitiveSceneProxy.h
+++ b/Engine/Source/Runtime/Engine/Public/PrimitiveSceneProxy.h
@@ -66,6 +66,7 @@ public:
 	float Exponent;
 	float InverseExposureBlend = 0.0f;
 	float VolumetricScatteringIntensity;
+	float SpecularScale = 1.0f;
 	bool bAffectTranslucency;
 };
 
diff --git a/Engine/Source/Runtime/Renderer/Private/LightRendering.cpp b/Engine/Source/Runtime/Renderer/Private/LightRendering.cpp
--- a/Engine/Source/Runtime/Renderer/Private/LightRendering.cpp
+++ b/Engine/Source/Runtime/Renderer/Private/LightRendering.cpp
@@ -660,7 +660,7 @@ FDeferredLightUniformStruct GetSimpleDeferredLightParameters(
 	Out.LightParameters.Direction = FVector3f(1, 0, 0);
 	Out.LightParameters.Tangent = FVector3f(1, 0, 0);
 	Out.LightParameters.SpotAngles = FVector2f(-2, 1);
-	Out.LightParameters.SpecularScale = 1.0f;
+	Out.LightParameters.SpecularScale = SimpleLight.SpecularScale;
 	Out.LightParameters.SourceRadius = 0.0f;
 	Out.LightParameters.SoftSourceRadius = 0.0f;
 	Out.LightParameters.SourceLength = 0.0f;
