Fix for not propagating static properties over to non-static variables during default propagation with AsyncTask compilation mode

diff --git a/Engine/Plugins/FX/Niagara/Source/Niagara/Private/NiagaraModule.cpp b/Engine/Plugins/FX/Niagara/Source/Niagara/Private/NiagaraModule.cpp
--- a/Engine/Plugins/FX/Niagara/Source/Niagara/Private/NiagaraModule.cpp
+++ b/Engine/Plugins/FX/Niagara/Source/Niagara/Private/NiagaraModule.cpp
@@ -1359,63 +1359,63 @@ bool FNiagaraTypeDefinition::IsScalarDefinition(const FNiagaraTypeDefinition& Ty
 	return ScalarStructs.Contains(Type.GetScriptStruct()) || (Type.GetScriptStruct() == IntStruct && Type.GetEnum() != nullptr);
 }
 
-bool FNiagaraTypeDefinition::TypesAreAssignable(const FNiagaraTypeDefinition& TypeInput, const FNiagaraTypeDefinition& TypeOutput, bool bAllowLossyLWCConversions)
+bool FNiagaraTypeDefinition::TypesAreAssignable(const FNiagaraTypeDefinition& InputPinType, const FNiagaraTypeDefinition& OutputPinType, bool bAllowLossyLWCConversions)
 {
-	if (const UClass* AClass = TypeInput.GetClass())
+	if (const UClass* AClass = InputPinType.GetClass())
 	{
-		if (const UClass* BClass = TypeOutput.GetClass())
+		if (const UClass* BClass = OutputPinType.GetClass())
 		{
 			return AClass == BClass;
 		}
 	}
 	
-	if (const UClass* BClass = TypeOutput.GetClass())
+	if (const UClass* BClass = OutputPinType.GetClass())
 	{
 		return false;
 	}
 
-	if (const UClass* AClass = TypeInput.GetClass())
+	if (const UClass* AClass = InputPinType.GetClass())
 	{
 		return false;
 	}
 
 	// Enums can only be assigned from enums of the same type.
-	if(TypeInput.IsEnum() && TypeInput.GetEnum() != TypeOutput.GetEnum())
+	if(InputPinType.IsEnum() && InputPinType.GetEnum() != OutputPinType.GetEnum())
 	{
 		return false;
 	}
 
 	//Enums can be assigned to enums of the same type or plain integers.
-	if(TypeOutput.IsEnum() && TypeInput.GetEnum() != TypeOutput.GetEnum() && TypeInput != FNiagaraTypeDefinition::GetIntDef())
+	if(OutputPinType.IsEnum() && InputPinType.GetEnum() != OutputPinType.GetEnum() && InputPinType != FNiagaraTypeDefinition::GetIntDef())
 	{
 		return false;
 	}
 
 	// Static can go to static or static can go to nonstatic, but nonstatic can't go to static
-	if (((!TypeInput.IsStatic() && TypeOutput.IsStatic()) || (TypeInput.IsStatic() && TypeOutput.IsStatic())) && TypeInput.GetStruct() == TypeOutput.GetStruct() && TypeInput.GetEnum() == TypeOutput.GetEnum())
+	if (((!InputPinType.IsStatic() && OutputPinType.IsStatic()) || (InputPinType.IsStatic() && OutputPinType.IsStatic())) && InputPinType.GetStruct() == OutputPinType.GetStruct() && InputPinType.GetEnum() == OutputPinType.GetEnum())
 	{
 		return true;
 	}
-	else if ((TypeInput.IsStatic() || TypeOutput.IsStatic()) && TypeInput.GetStruct() == TypeOutput.GetStruct() && TypeInput.GetEnum() == TypeOutput.GetEnum())
+	else if ((InputPinType.IsStatic() || OutputPinType.IsStatic()) && InputPinType.GetStruct() == OutputPinType.GetStruct() && InputPinType.GetEnum() == OutputPinType.GetEnum())
 	{
 		return false;
 	}
 
-	if (TypeInput.GetStruct() == TypeOutput.GetStruct())
+	if (InputPinType.GetStruct() == OutputPinType.GetStruct())
 	{
 		return true;
 	}
 
 	bool bIsSupportedConversion = false;
-	if (IsScalarDefinition(TypeInput) && IsScalarDefinition(TypeOutput))
+	if (IsScalarDefinition(InputPinType) && IsScalarDefinition(OutputPinType))
 	{
-		bIsSupportedConversion = (TypeInput == IntDef && TypeOutput == FloatDef) || (TypeOutput == IntDef && TypeInput == FloatDef) || (TypeInput == IntDef && TypeOutput == BoolDef) || (TypeOutput == IntDef && TypeInput == BoolDef);
+		bIsSupportedConversion = (InputPinType == IntDef && OutputPinType == FloatDef) || (OutputPinType == IntDef && InputPinType == FloatDef) || (InputPinType == IntDef && OutputPinType == BoolDef) || (OutputPinType == IntDef && InputPinType == BoolDef);
 	}
 	else
 	{
-		bIsSupportedConversion = (TypeInput == ColorDef && TypeOutput == Vec4Def) || (TypeOutput == ColorDef && TypeInput == Vec4Def)
+		bIsSupportedConversion = (InputPinType == ColorDef && OutputPinType == Vec4Def) || (OutputPinType == ColorDef && InputPinType == Vec4Def)
 			// we only allow direct assignments between vec <-> position if we are allowing lossy conversions
-			|| (bAllowLossyLWCConversions && ((TypeOutput == PositionDef && TypeInput == Vec3Def) || (TypeOutput == Vec3Def && TypeInput == PositionDef)));
+			|| (bAllowLossyLWCConversions && ((OutputPinType == PositionDef && InputPinType == Vec3Def) || (OutputPinType == Vec3Def && InputPinType == PositionDef)));
 	}
 
 	if (bIsSupportedConversion)
@@ -1424,10 +1424,10 @@ bool FNiagaraTypeDefinition::TypesAreAssignable(const FNiagaraTypeDefinition& Ty
 	}
 
 
-	return	(TypeInput == NumericDef && NumericStructs.Contains(TypeOutput.GetScriptStruct())) ||
-		(TypeOutput == NumericDef && NumericStructs.Contains(TypeInput.GetScriptStruct())) ||
-		(TypeInput == NumericDef && (TypeOutput.GetStruct() == GetIntStruct()) && TypeOutput.GetEnum() != nullptr) ||
-		(TypeOutput == NumericDef && (TypeInput.GetStruct() == GetIntStruct()) && TypeInput.GetEnum() != nullptr);
+	return	(InputPinType == NumericDef && NumericStructs.Contains(OutputPinType.GetScriptStruct())) ||
+		(OutputPinType == NumericDef && NumericStructs.Contains(InputPinType.GetScriptStruct())) ||
+		(InputPinType == NumericDef && (OutputPinType.GetStruct() == GetIntStruct()) && OutputPinType.GetEnum() != nullptr) ||
+		(OutputPinType == NumericDef && (InputPinType.GetStruct() == GetIntStruct()) && InputPinType.GetEnum() != nullptr);
 }
 
 bool FNiagaraTypeDefinition::IsLossyConversion(const FNiagaraTypeDefinition& FromType, const FNiagaraTypeDefinition& ToType)
diff --git a/Engine/Plugins/FX/Niagara/Source/Niagara/Public/NiagaraTypes.h b/Engine/Plugins/FX/Niagara/Source/Niagara/Public/NiagaraTypes.h
--- a/Engine/Plugins/FX/Niagara/Source/Niagara/Public/NiagaraTypes.h
+++ b/Engine/Plugins/FX/Niagara/Source/Niagara/Public/NiagaraTypes.h
@@ -1017,7 +1017,8 @@ public:
 		return FNiagaraTypeHelper::ToString(ValueData, ClassStructOrEnum);
 	}
 	
-	static NIAGARA_API bool TypesAreAssignable(const FNiagaraTypeDefinition& TypeA, const FNiagaraTypeDefinition& TypeB, bool bAllowLossyLWCConversions = false);
+	// Evaluates if two pin types are compatible for assignment.  Note that the assignment being evaluated is InputPinType = OutputPinType
+	static NIAGARA_API bool TypesAreAssignable(const FNiagaraTypeDefinition& InputPinType, const FNiagaraTypeDefinition& OutputPinType, bool bAllowLossyLWCConversions = false);
 	static NIAGARA_API bool IsLossyConversion(const FNiagaraTypeDefinition& FromType, const FNiagaraTypeDefinition& ToType);
 	static NIAGARA_API FNiagaraTypeDefinition GetNumericOutputType(TConstArrayView<FNiagaraTypeDefinition> TypeDefinintions, ENiagaraNumericOutputTypeSelectionMode SelectionMode);
 
diff --git a/Engine/Plugins/FX/Niagara/Source/NiagaraEditor/Private/NiagaraGraphDigest.cpp b/Engine/Plugins/FX/Niagara/Source/NiagaraEditor/Private/NiagaraGraphDigest.cpp
--- a/Engine/Plugins/FX/Niagara/Source/NiagaraEditor/Private/NiagaraGraphDigest.cpp
+++ b/Engine/Plugins/FX/Niagara/Source/NiagaraEditor/Private/NiagaraGraphDigest.cpp
@@ -1568,7 +1568,7 @@ FNiagaraCompilationInputPin::FNiagaraCompilationInputPin(const FNiagaraCompilati
 			}
 		}
 
-		if (ensure(Variable.GetType() == TracedInputPin->Variable.GetType()))
+		if (ensure(FNiagaraTypeDefinition::TypesAreAssignable(Variable.GetType(), TracedInputPin->Variable.GetType())))
 		{
 			if (TracedInputPin->Variable.IsDataAllocated())
 			{
