Cleanup some of the UI around the EVVM
-renames public facing things from Experimental to Optimized
-disables the Performance mode button if it's not supported (when EVVM is being used)

diff --git a/Engine/Plugins/FX/Niagara/Source/Niagara/Classes/NiagaraSystem.h b/Engine/Plugins/FX/Niagara/Source/Niagara/Classes/NiagaraSystem.h
--- a/Engine/Plugins/FX/Niagara/Source/Niagara/Classes/NiagaraSystem.h
+++ b/Engine/Plugins/FX/Niagara/Source/Niagara/Classes/NiagaraSystem.h
@@ -524,6 +524,9 @@ public:
 	FORCEINLINE void SetBakeOutRapidIterationOnCook(bool bBakeOut) { bBakeOutRapidIteration = bBakeOut; bBakeOutRapidIterationOnCook = bBakeOut; }
 	FORCEINLINE void SetTrimAttributesOnCook(bool bTrim) { bTrimAttributes = bTrim; bTrimAttributesOnCook = bTrim; }
 
+	/** returns true if the system, and it's underlying scripts support stat scopes and the collection of the timings */
+	NIAGARA_API bool SupportsStatScopedPerformanceMode() const;
+
 protected:
 	/** When enable constant values are baked into the scripts while editing the system, this will increase iteration times but improve performance. */
 	UPROPERTY(EditAnywhere, AdvancedDisplay, Category = "Performance", meta=(DisplayName="Bake Rapid Iteration Parameters During Edit"))
@@ -616,8 +619,8 @@ public:
 	UPROPERTY(EditAnywhere, BlueprintReadOnly, Category="Rendering", meta=(DisplayName="Default Render CustomDepth Pass", EditCondition="bOverrideRenderCustomDepth"))
 	uint8 bRenderCustomDepth : 1;
 
-	/** If true, disables experimental VM, if available */
-	UPROPERTY(EditAnywhere, AdvancedDisplay, Category = "Performance", meta = (DisplayName = "Disable Experimental VM"))
+	/** If true, disables optimized VM, if available */
+	UPROPERTY(EditAnywhere, AdvancedDisplay, Category = "Performance", meta = (DisplayName = "Disable Optimized VM"))
 	uint8 bDisableExperimentalVM : 1;
 
 private:
diff --git a/Engine/Plugins/FX/Niagara/Source/Niagara/Private/NiagaraSystem.cpp b/Engine/Plugins/FX/Niagara/Source/Niagara/Private/NiagaraSystem.cpp
--- a/Engine/Plugins/FX/Niagara/Source/Niagara/Private/NiagaraSystem.cpp
+++ b/Engine/Plugins/FX/Niagara/Source/Niagara/Private/NiagaraSystem.cpp
@@ -942,6 +942,12 @@ void UNiagaraSystem::RecomputeExecutionOrderForDataInterface(class UNiagaraDataI
 	}
 }
 
+bool UNiagaraSystem::SupportsStatScopedPerformanceMode() const
+{
+	const UNiagaraSettings* Settings = GetDefault<UNiagaraSettings>();
+	return bDisableExperimentalVM || !Settings || !Settings->bExperimentalVMEnabled;
+}
+
 #endif
 
 void UNiagaraSystem::Serialize(FArchive& Ar)
diff --git a/Engine/Plugins/FX/Niagara/Source/Niagara/Public/NiagaraSettings.h b/Engine/Plugins/FX/Niagara/Source/Niagara/Public/NiagaraSettings.h
--- a/Engine/Plugins/FX/Niagara/Source/Niagara/Public/NiagaraSettings.h
+++ b/Engine/Plugins/FX/Niagara/Source/Niagara/Public/NiagaraSettings.h
@@ -98,7 +98,7 @@ enum class ENiagaraStripScriptByteCodeOption : uint8
 {
 	Default = 0 UMETA(DisplayName = "No Stripping"),
 	Strip_Original = 1 UMETA(DisplayName = "Strip Original ByteCode"),
-	Strip_Experimental = 2 UMETA(DisplayName = "Strip Experimental ByteCode"),
+	Strip_Experimental = 2 UMETA(DisplayName = "Strip Optimized ByteCode"),
 };
 
 #if WITH_EDITORONLY_DATA
@@ -107,7 +107,7 @@ enum class ENiagaraCompilationMode : int32
 {
 	Original = 0	UMETA(DisplayName = "Standard Compilation"),
 	AsyncTasks = 1	UMETA(DisplayName = "Experimental - Async Compilation"),
-	Verify = 2		UMETA(DisplayName = "Experimental - Validate Async")
+	Verify = 2      UMETA(Hidden)
 };
 #endif
 
@@ -151,11 +151,11 @@ class UNiagaraSettings : public UDeveloperSettings
 	bool bEnforceStrictStackTypes = true;
 
 	/**
-	 True indicates that we will generate byte code for the new experimental VM.  Control over whether the new VM will
+	 True indicates that we will generate byte code for the new optimized VM.  Control over whether the new VM will
 	 be used when executing NiagaraScripts will also take into account the overrides on the system (bDisableExperimentalVM) and
 	 the cvars fx.NiagaraScript.StripByteCodeOnLoad and fx.ForceExecVMPath.
 	*/
-	UPROPERTY(config, EditAnywhere, Category = Niagara, meta = (DisplayName = "Enable building data for Experimental VM"))
+	UPROPERTY(config, EditAnywhere, Category = Niagara, meta = (DisplayName = "Enable building data for Optimized VM"))
 	bool bExperimentalVMEnabled = false;
 
 	/**
@@ -193,7 +193,7 @@ class UNiagaraSettings : public UDeveloperSettings
 	FLinearColor PositionPinTypeColor;
 
 	/**
-	 Controls how byte code will be stripped when loading assets that have multiple sets of bytecode (i.e. experimental).
+	 Controls how byte code will be stripped when loading assets that have multiple sets of bytecode (i.e. optimized).
 	 */
 	UPROPERTY(config, EditAnywhere, Category = Niagara, meta = (DisplayName = "Option for how to strip bytecode"))
 	ENiagaraStripScriptByteCodeOption ByteCodeStripOption = ENiagaraStripScriptByteCodeOption::Default;
diff --git a/Engine/Plugins/FX/Niagara/Source/NiagaraEditor/Private/NiagaraEditorCommands.cpp b/Engine/Plugins/FX/Niagara/Source/NiagaraEditor/Private/NiagaraEditorCommands.cpp
--- a/Engine/Plugins/FX/Niagara/Source/NiagaraEditor/Private/NiagaraEditorCommands.cpp
+++ b/Engine/Plugins/FX/Niagara/Source/NiagaraEditor/Private/NiagaraEditorCommands.cpp
@@ -29,7 +29,7 @@ void FNiagaraEditorCommands::RegisterCommands()
 	UI_COMMAND(ToggleMemoryInfo, "Memory", "Show approximate memory usage", EUserInterfaceActionType::ToggleButton, FInputChord());
 	UI_COMMAND(TogglePreviewBackground, "Background", "Toggles the preview pane's background.", EUserInterfaceActionType::ToggleButton, FInputChord());
 	UI_COMMAND(ToggleUnlockToChanges, "Lock/Unlock To Changes", "Toggles whether or not changes in the source asset get pulled into this asset automatically.", EUserInterfaceActionType::Button, FInputChord());
-	UI_COMMAND(ToggleStatPerformance, "Performance", "Show runtime performance for particle scripts.", EUserInterfaceActionType::ToggleButton, FInputChord());
+	UI_COMMAND(ToggleStatPerformance, "Performance", "Show runtime performance for particle scripts. Not supported for optimized VM byte code.", EUserInterfaceActionType::ToggleButton, FInputChord());
 	UI_COMMAND(ToggleStatPerformanceGPU, "GPU Profiling", "Measure gpu runtime cost of scripts. Can be expensive for simulation stages with lots of iterations.", EUserInterfaceActionType::ToggleButton, FInputChord());
 	UI_COMMAND(ClearStatPerformance, "Clear Stats", "Clear all existing stat captures.", EUserInterfaceActionType::Button, FInputChord());
 	UI_COMMAND(ToggleStatPerformanceTypeAvg, "Display Average", "Displays the average of captured stats.", EUserInterfaceActionType::ToggleButton, FInputChord());
diff --git a/Engine/Plugins/FX/Niagara/Source/NiagaraEditor/Private/Toolkits/NiagaraSystemToolkit.cpp b/Engine/Plugins/FX/Niagara/Source/NiagaraEditor/Private/Toolkits/NiagaraSystemToolkit.cpp
--- a/Engine/Plugins/FX/Niagara/Source/NiagaraEditor/Private/Toolkits/NiagaraSystemToolkit.cpp
+++ b/Engine/Plugins/FX/Niagara/Source/NiagaraEditor/Private/Toolkits/NiagaraSystemToolkit.cpp
@@ -422,10 +422,13 @@ void FNiagaraSystemToolkit::SetupCommands()
 		FExecuteAction::CreateRaw(this, &FNiagaraSystemToolkit::ResetSimulation)); 
 
 	GetToolkitCommands()->MapAction(
-        FNiagaraEditorCommands::Get().ToggleStatPerformance,
-        FExecuteAction::CreateSP(this, &FNiagaraSystemToolkit::ToggleStatPerformance),
-        FCanExecuteAction(),
-        FIsActionChecked::CreateSP(this, &FNiagaraSystemToolkit::IsStatPerformanceChecked));
+		FNiagaraEditorCommands::Get().ToggleStatPerformance,
+		FExecuteAction::CreateSP(this, &FNiagaraSystemToolkit::ToggleStatPerformance),
+		FCanExecuteAction::CreateLambda([this]()
+		{
+			return System && System->SupportsStatScopedPerformanceMode();
+		}),
+		FIsActionChecked::CreateSP(this, &FNiagaraSystemToolkit::IsStatPerformanceChecked));
 	GetToolkitCommands()->MapAction(
         FNiagaraEditorCommands::Get().ClearStatPerformance,
         FExecuteAction::CreateSP(this, &FNiagaraSystemToolkit::ClearStatPerformance));
