- Stateless merge [at] CL 31242318

diff --git a/Engine/Plugins/FX/Niagara/Shaders/Private/Stateless/Modules/NiagaraStatelessModule_ScaleMeshSizeBySpeed.ush b/Engine/Plugins/FX/Niagara/Shaders/Private/Stateless/Modules/NiagaraStatelessModule_ScaleMeshSizeBySpeed.ush
--- a/Engine/Plugins/FX/Niagara/Shaders/Private/Stateless/Modules/NiagaraStatelessModule_ScaleMeshSizeBySpeed.ush
+++ b/Engine/Plugins/FX/Niagara/Shaders/Private/Stateless/Modules/NiagaraStatelessModule_ScaleMeshSizeBySpeed.ush
@@ -2,40 +2,25 @@
 
 #include "../NiagaraStatelessCommon.ush"
 
+uint2	ScaleMeshSizeBySpeed_ScaleDistribution;
 float	ScaleMeshSizeBySpeed_VelocityNorm;
-int		ScaleMeshSizeBySpeed_Offset;
-int		ScaleMeshSizeBySpeed_Length;
 
 struct FStatelessModule_ScaleMeshSizeBySpeed
 {
 };
 
-float3 ScaleMeshSizeBySpeed_GetScale(in FStatelessParticle Particle, float NormSpeed)
-{
-	const float Offset = NormSpeed * float(ScaleMeshSizeBySpeed_Length - 1);
-	const int A = floor(Offset);
-	const int B = ceil(Offset);
-	const float U = frac(Offset);
-
-	return lerp(
-		GetStaticFloat3(ScaleMeshSizeBySpeed_Offset, A),
-		GetStaticFloat3(ScaleMeshSizeBySpeed_Offset, B),
-		U
-	);
-}
-
-void ScaleMeshSizeBySpeed_Initialize(inout FStatelessModule_ScaleMeshSizeBySpeed MouldeData, inout FStatelessParticle Particle)
+void ScaleMeshSizeBySpeed_Initialize(inout FStatelessModule_ScaleMeshSizeBySpeed ModuleData, inout FStatelessParticle Particle)
 {
 }
 	
-void ScaleMeshSizeBySpeed_Simulate(in FStatelessModule_ScaleMeshSizeBySpeed MouldeData, inout FStatelessParticle Particle)
+void ScaleMeshSizeBySpeed_Simulate(in FStatelessModule_ScaleMeshSizeBySpeed ModuleData, inout FStatelessParticle Particle)
 {
-	if (ScaleMeshSizeBySpeed_Length > 0)
+	if (ScaleMeshSizeBySpeed_ScaleDistribution.y > 0)
 	{
 		const float Speed = length2(Particle.Position - Particle.PreviousPosition) * Particle.InvDeltaTime;
 		const float NormSpeed = saturate(Speed * ScaleMeshSizeBySpeed_VelocityNorm);
-			
-		Particle.Scale			*= ScaleMeshSizeBySpeed_GetScale(Particle, NormSpeed);
-		Particle.PreviousScale	*= ScaleMeshSizeBySpeed_GetScale(Particle, NormSpeed);
+
+		Particle.Scale			*= LerpStaticFloat3(ScaleMeshSizeBySpeed_ScaleDistribution, NormSpeed);
+		Particle.PreviousScale	*= LerpStaticFloat3(ScaleMeshSizeBySpeed_ScaleDistribution, NormSpeed);
 	}
 }
diff --git a/Engine/Plugins/FX/Niagara/Shaders/Private/Stateless/Modules/NiagaraStatelessModule_ScaleSpriteSizeBySpeed.ush b/Engine/Plugins/FX/Niagara/Shaders/Private/Stateless/Modules/NiagaraStatelessModule_ScaleSpriteSizeBySpeed.ush
--- a/Engine/Plugins/FX/Niagara/Shaders/Private/Stateless/Modules/NiagaraStatelessModule_ScaleSpriteSizeBySpeed.ush
+++ b/Engine/Plugins/FX/Niagara/Shaders/Private/Stateless/Modules/NiagaraStatelessModule_ScaleSpriteSizeBySpeed.ush
@@ -2,40 +2,25 @@
 
 #include "../NiagaraStatelessCommon.ush"
 
+uint2	ScaleSpriteSizeBySpeed_ScaleDistribution;
 float	ScaleSpriteSizeBySpeed_VelocityNorm;
-int		ScaleSpriteSizeBySpeed_Offset;
-int		ScaleSpriteSizeBySpeed_Length;
 
 struct FStatelessModule_ScaleSpriteSizeBySpeed
 {
 };
 
-float2 ScaleSpriteSizeBySpeed_GetScale(in FStatelessParticle Particle, float NormSpeed)
-{
-	const float Offset = NormSpeed * float(ScaleSpriteSizeBySpeed_Length - 1);
-	const int A = floor(Offset);
-	const int B = ceil(Offset);
-	const float U = frac(Offset);
-
-	return lerp(
-		GetStaticFloat2(ScaleSpriteSizeBySpeed_Offset, A),
-		GetStaticFloat2(ScaleSpriteSizeBySpeed_Offset, B),
-		U
-	);
-}
-	
 void ScaleSpriteSizeBySpeed_Initialize(inout FStatelessModule_ScaleSpriteSizeBySpeed ModuleData, inout FStatelessParticle Particle)
 {
 }
 	
 void ScaleSpriteSizeBySpeed_Simulate(in FStatelessModule_ScaleSpriteSizeBySpeed ModuleData, inout FStatelessParticle Particle)
 {
-	if (ScaleSpriteSizeBySpeed_Length > 0)
+	if (ScaleSpriteSizeBySpeed_ScaleDistribution.y > 0)
 	{
 		const float Speed = length2(Particle.Position - Particle.PreviousPosition) * Particle.InvDeltaTime;
 		const float NormSpeed = saturate(Speed * ScaleSpriteSizeBySpeed_VelocityNorm);
-			
-		Particle.SpriteSize			*= ScaleSpriteSizeBySpeed_GetScale(Particle, NormSpeed);
-		Particle.PreviousSpriteSize	*= ScaleSpriteSizeBySpeed_GetScale(Particle, NormSpeed);
+
+		Particle.SpriteSize			*= LerpStaticFloat2(ScaleSpriteSizeBySpeed_ScaleDistribution, NormSpeed);
+		Particle.PreviousSpriteSize	*= LerpStaticFloat2(ScaleSpriteSizeBySpeed_ScaleDistribution, NormSpeed);
 	}
 }
diff --git a/Engine/Plugins/FX/Niagara/Shaders/Private/Stateless/NiagaraStatelessCommon.ush b/Engine/Plugins/FX/Niagara/Shaders/Private/Stateless/NiagaraStatelessCommon.ush
--- a/Engine/Plugins/FX/Niagara/Shaders/Private/Stateless/NiagaraStatelessCommon.ush
+++ b/Engine/Plugins/FX/Niagara/Shaders/Private/Stateless/NiagaraStatelessCommon.ush
@@ -126,6 +126,30 @@ float2 GetStaticFloat2(int BaseOffset, int Element)	{ return float2(Common_Stati
 float3 GetStaticFloat3(int BaseOffset, int Element)	{ return float3(Common_StaticFloatBuffer[BaseOffset + Element*3], Common_StaticFloatBuffer[BaseOffset + 1 + Element*3], Common_StaticFloatBuffer[BaseOffset + 2 + Element*3]); }
 float4 GetStaticFloat4(int BaseOffset, int Element)	{ return float4(Common_StaticFloatBuffer[BaseOffset + Element*4], Common_StaticFloatBuffer[BaseOffset + 1 + Element*4], Common_StaticFloatBuffer[BaseOffset + 2 + Element*4], Common_StaticFloatBuffer[BaseOffset + 3 + Element*4]); }
 
+float LerpStaticFloat (uint2 Parameters, float U)
+{
+	const float Offset = U * float(Parameters.y);
+	const float Value0 = GetStaticFloat(Parameters.x, floor(Offset));
+	const float Value1 = GetStaticFloat(Parameters.x, ceil(Offset));
+	return lerp(Value0, Value1, frac(Offset));
+}
+
+float2 LerpStaticFloat2(uint2 Parameters, float U)
+{
+	const float  Offset = U * float(Parameters.y);
+	const float2 Value0 = GetStaticFloat2(Parameters.x, floor(Offset));
+	const float2 Value1 = GetStaticFloat2(Parameters.x, ceil(Offset));
+	return lerp(Value0, Value1, frac(Offset));
+}
+
+float3 LerpStaticFloat3(uint2 Parameters, float U)
+{
+	const float  Offset = U * float(Parameters.y);
+	const float3 Value0 = GetStaticFloat3(Parameters.x, floor(Offset));
+	const float3 Value1 = GetStaticFloat3(Parameters.x, ceil(Offset));
+	return lerp(Value0, Value1, frac(Offset));
+}
+	
 uint GetComponentOffset(uint OutputIndex, int ComponentOffset)
 {
 	return ComponentOffset * Common_OutputBufferStride + OutputIndex;
diff --git a/Engine/Plugins/FX/Niagara/Source/Niagara/Internal/NiagaraSystemEmitterState.h b/Engine/Plugins/FX/Niagara/Source/Niagara/Internal/NiagaraSystemEmitterState.h
--- a/Engine/Plugins/FX/Niagara/Source/Niagara/Internal/NiagaraSystemEmitterState.h
+++ b/Engine/Plugins/FX/Niagara/Source/Niagara/Internal/NiagaraSystemEmitterState.h
@@ -3,6 +3,7 @@
 #pragma once
 
 #include "NiagaraCommon.h"
+#include "Stateless/NiagaraStatelessDistribution.h"
 
 #include "NiagaraSystemEmitterState.generated.h"
 
@@ -47,12 +48,9 @@ struct FNiagaraSystemStateData
 
 	UPROPERTY(EditAnywhere, Category = "System State")
 	ENiagaraLoopBehavior LoopBehavior = ENiagaraLoopBehavior::Once;
-	
-	UPROPERTY(EditAnywhere, Category = "Emitter State", meta = (ClampMin = "0.0"))
-	float LoopDurationMin = 0.0f;
 
-	UPROPERTY(EditAnywhere, Category = "Emitter State", meta = (ClampMin = "0.0"))
-	float LoopDurationMax = 0.0f;
+	UPROPERTY(EditAnywhere, Category = "System State", meta = (ClampMin = "0.0"))
+	FNiagaraDistributionRangeFloat LoopDuration = FNiagaraDistributionRangeFloat(0.0f);
 
 	UPROPERTY(EditAnywhere, Category = "System State", meta = (ClampMin = "1", EditCondition = "LoopBehavior == ENiagaraLoopBehavior::Multiple", EditConditionHides))
 	int LoopCount = 1;
@@ -60,11 +58,8 @@ struct FNiagaraSystemStateData
 	UPROPERTY(EditAnywhere, Category = "System State", meta = (EditCondition = "LoopBehavior != ENiagaraLoopBehavior::Once", EditConditionHides))
 	bool bRecalculateDurationEachLoop = false;
 
-	UPROPERTY(EditAnywhere, Category = "Emitter State", meta = (ClampMin = "0.0"))
-	float LoopDelayMin = 0.0f;
-
-	UPROPERTY(EditAnywhere, Category = "Emitter State", meta = (ClampMin = "0.0"))
-	float LoopDelayMax = 0.0f;
+	UPROPERTY(EditAnywhere, Category = "System State", meta = (ClampMin = "0.0"))
+	FNiagaraDistributionRangeFloat LoopDelay = FNiagaraDistributionRangeFloat(0.0f);
 
 	UPROPERTY(EditAnywhere, Category = "System State", meta = (EditCondition = "LoopBehavior != ENiagaraLoopBehavior::Once", EditConditionHides))
 	bool bDelayFirstLoopOnly = false;
@@ -91,19 +86,13 @@ struct FNiagaraEmitterStateData
 	int32 LoopCount = 1;
 
 	UPROPERTY(EditAnywhere, Category = "Emitter State", meta = (ClampMin = "0.0"))
-	float LoopDurationMin = 1.0f;
-
-	UPROPERTY(EditAnywhere, Category = "Emitter State", meta = (ClampMin = "0.0"))
-	float LoopDurationMax = 1.0f;
+	FNiagaraDistributionRangeFloat LoopDuration = FNiagaraDistributionRangeFloat(1.0f);
 
 	UPROPERTY(EditAnywhere, Category = "Emitter State", meta = (EditCondition = "LoopBehavior != ENiagaraLoopBehavior::Once", EditConditionHides))
 	bool bRecalculateDurationEachLoop = false;
 
 	UPROPERTY(EditAnywhere, Category = "Emitter State", meta = (ClampMin = "0.0"))
-	float LoopDelayMin = 0.0f;
-
-	UPROPERTY(EditAnywhere, Category = "Emitter State", meta = (ClampMin = "0.0"))
-	float LoopDelayMax = 0.0f;
+	FNiagaraDistributionRangeFloat LoopDelay = FNiagaraDistributionRangeFloat(0.0f);
 
 	UPROPERTY(EditAnywhere, Category = "Emitter State", meta = (EditCondition = "LoopBehavior != ENiagaraLoopBehavior::Once", EditConditionHides))
 	bool bDelayFirstLoopOnly = false;
diff --git a/Engine/Plugins/FX/Niagara/Source/Niagara/Internal/Stateless/Modules/NiagaraStatelessModule_ScaleMeshSizeBySpeed.h b/Engine/Plugins/FX/Niagara/Source/Niagara/Internal/Stateless/Modules/NiagaraStatelessModule_ScaleMeshSizeBySpeed.h
--- a/Engine/Plugins/FX/Niagara/Source/Niagara/Internal/Stateless/Modules/NiagaraStatelessModule_ScaleMeshSizeBySpeed.h
+++ b/Engine/Plugins/FX/Niagara/Source/Niagara/Internal/Stateless/Modules/NiagaraStatelessModule_ScaleMeshSizeBySpeed.h
@@ -15,9 +15,8 @@ class UNiagaraStatelessModule_ScaleMeshSizeBySpeed : public UNiagaraStatelessMod
 
 	struct FModuleBuiltData
 	{
-		float	VelocityNorm = 0.0f;
-		int32	TableOffset = 0;
-		int32	TableLength = 0;
+		float			VelocityNorm = 0.0f;
+		FUintVector2	ScaleDistribution = FUintVector2::ZeroValue;
 	};
 
 public:
@@ -26,17 +25,26 @@ public:
 	UPROPERTY(EditAnywhere, Category = "Parameters", meta = (ClampMin = "0.01", UIMin = "0.01"))
 	float VelocityThreshold = 1000.0f;
 
-	UPROPERTY(EditAnywhere, Category = "Parameters")
-	TArray<FVector3f>	ScaleValues = { FVector3f(1.0, 1.0f, 1.0f), FVector3f(2.0, 2.0f, 2.0f) };
+	UPROPERTY(EditAnywhere, Category = "Parameters", meta = (DisplayName = "Scale", DisableRangeDistribution))
+	FNiagaraDistributionVector3 ScaleDistribution = FNiagaraDistributionVector3(1.0f);
 
 	virtual void BuildEmitterData(FNiagaraStatelessEmitterDataBuildContext& BuildContext) const override
 	{
 		FModuleBuiltData* BuiltData = BuildContext.AllocateBuiltData<FModuleBuiltData>();
 		if (IsModuleEnabled())
 		{
-			BuiltData->VelocityNorm	= VelocityThreshold > 0.0f ? 1.0f / (VelocityThreshold * VelocityThreshold) : 0.0f;
-			BuiltData->TableOffset	= BuildContext.AddStaticData(ScaleValues);
-			BuiltData->TableLength	= ScaleValues.Num();
+			BuiltData->VelocityNorm = VelocityThreshold > 0.0f ? 1.0f / (VelocityThreshold * VelocityThreshold) : 0.0f;
+			if (ScaleDistribution.IsCurve() && ScaleDistribution.Values.Num() > 1)
+			{
+				BuiltData->ScaleDistribution.X = BuildContext.AddStaticData(ScaleDistribution.Values);
+				BuiltData->ScaleDistribution.Y = ScaleDistribution.Values.Num() - 1;
+			}
+			else
+			{
+				const FVector3f Values[] = { FVector3f::One(), ScaleDistribution.Values.Num() > 0 ? ScaleDistribution.Values[0] : FVector3f::One() };
+				BuiltData->ScaleDistribution.X = BuildContext.AddStaticData(Values);
+				BuiltData->ScaleDistribution.Y = 1;
+			}
 		}
 	}
 
@@ -45,9 +53,8 @@ public:
 		const FModuleBuiltData* ModuleBuiltData = SetShaderParameterContext.ReadBuiltData<FModuleBuiltData>();
 
 		FParameters* Parameters = SetShaderParameterContext.GetParameterNestedStruct<FParameters>();
-		Parameters->ScaleMeshSizeBySpeed_VelocityNorm	= ModuleBuiltData->VelocityNorm;
-		Parameters->ScaleMeshSizeBySpeed_Offset			= ModuleBuiltData->TableOffset;
-		Parameters->ScaleMeshSizeBySpeed_Length			= ModuleBuiltData->TableLength;
+		Parameters->ScaleMeshSizeBySpeed_VelocityNorm		= ModuleBuiltData->VelocityNorm;
+		Parameters->ScaleMeshSizeBySpeed_ScaleDistribution	= ModuleBuiltData->ScaleDistribution;
 	}
 
 #if WITH_EDITOR
diff --git a/Engine/Plugins/FX/Niagara/Source/Niagara/Internal/Stateless/Modules/NiagaraStatelessModule_ScaleSpriteSizeBySpeed.h b/Engine/Plugins/FX/Niagara/Source/Niagara/Internal/Stateless/Modules/NiagaraStatelessModule_ScaleSpriteSizeBySpeed.h
--- a/Engine/Plugins/FX/Niagara/Source/Niagara/Internal/Stateless/Modules/NiagaraStatelessModule_ScaleSpriteSizeBySpeed.h
+++ b/Engine/Plugins/FX/Niagara/Source/Niagara/Internal/Stateless/Modules/NiagaraStatelessModule_ScaleSpriteSizeBySpeed.h
@@ -15,9 +15,8 @@ class UNiagaraStatelessModule_ScaleSpriteSizeBySpeed : public UNiagaraStatelessM
 
 	struct FModuleBuiltData
 	{
-		float	VelocityNorm = 0.0f;
-		int32	TableOffset = 0;
-		int32	TableLength = 0;
+		float			VelocityNorm = 0.0f;
+		FUintVector2	ScaleDistribution = FUintVector2::ZeroValue;
 	};
 
 public:
@@ -26,17 +25,26 @@ public:
 	UPROPERTY(EditAnywhere, Category = "Parameters", meta = (ClampMin = "0.01", UIMin = "0.01"))
 	float VelocityThreshold = 1000.0f;
 
-	UPROPERTY(EditAnywhere, Category = "Parameters")
-	TArray<FVector2f>	ScaleValues = { FVector2f(1.0f, 1.0f), FVector2f(2.0f, 2.0f) };
+	UPROPERTY(EditAnywhere, Category = "Parameters", meta = (DisplayName = "Scale", DisableRangeDistribution))
+	FNiagaraDistributionVector2 ScaleDistribution = FNiagaraDistributionVector2(1.0f);
 
 	virtual void BuildEmitterData(FNiagaraStatelessEmitterDataBuildContext& BuildContext) const override
 	{
 		FModuleBuiltData* BuiltData = BuildContext.AllocateBuiltData<FModuleBuiltData>();
 		if (IsModuleEnabled())
 		{
-			BuiltData->VelocityNorm	= VelocityThreshold > 0.0f ? 1.0f / (VelocityThreshold * VelocityThreshold) : 0.0f;
-			BuiltData->TableOffset	= BuildContext.AddStaticData(ScaleValues);
-			BuiltData->TableLength	= ScaleValues.Num();
+			BuiltData->VelocityNorm			= VelocityThreshold > 0.0f ? 1.0f / (VelocityThreshold * VelocityThreshold) : 0.0f;
+			if (ScaleDistribution.IsCurve() && ScaleDistribution.Values.Num() > 1)
+			{
+				BuiltData->ScaleDistribution.X = BuildContext.AddStaticData(ScaleDistribution.Values);
+				BuiltData->ScaleDistribution.Y = ScaleDistribution.Values.Num() - 1;
+			}
+			else
+			{
+				const FVector2f Values[] = { FVector2f::One(), ScaleDistribution.Values.Num() > 0 ? ScaleDistribution.Values[0] : FVector2f::One() };
+				BuiltData->ScaleDistribution.X = BuildContext.AddStaticData(Values);
+				BuiltData->ScaleDistribution.Y = 1;
+			}
 		}
 	}
 
@@ -45,9 +53,8 @@ public:
 		const FModuleBuiltData* ModuleBuiltData = SetShaderParameterContext.ReadBuiltData<FModuleBuiltData>();
 
 		FParameters* Parameters = SetShaderParameterContext.GetParameterNestedStruct<FParameters>();
-		Parameters->ScaleSpriteSizeBySpeed_VelocityNorm		= ModuleBuiltData->VelocityNorm;
-		Parameters->ScaleSpriteSizeBySpeed_Offset			= ModuleBuiltData->TableOffset;
-		Parameters->ScaleSpriteSizeBySpeed_Length			= ModuleBuiltData->TableLength;
+		Parameters->ScaleSpriteSizeBySpeed_VelocityNorm			= ModuleBuiltData->VelocityNorm;
+		Parameters->ScaleSpriteSizeBySpeed_ScaleDistribution	= ModuleBuiltData->ScaleDistribution;
 	}
 
 #if WITH_EDITOR
diff --git a/Engine/Plugins/FX/Niagara/Source/Niagara/Internal/Stateless/NiagaraStatelessCommon.h b/Engine/Plugins/FX/Niagara/Source/Niagara/Internal/Stateless/NiagaraStatelessCommon.h
--- a/Engine/Plugins/FX/Niagara/Source/Niagara/Internal/Stateless/NiagaraStatelessCommon.h
+++ b/Engine/Plugins/FX/Niagara/Source/Niagara/Internal/Stateless/NiagaraStatelessCommon.h
@@ -19,6 +19,7 @@ struct FNiagaraStatelessRange
 	ValueType Max = {};
 };
 
+using FNiagaraStatelessRangeInt		= FNiagaraStatelessRange<int32>;
 using FNiagaraStatelessRangeFloat	= FNiagaraStatelessRange<float>;
 using FNiagaraStatelessRangeVector2	= FNiagaraStatelessRange<FVector2f>;
 using FNiagaraStatelessRangeVector3	= FNiagaraStatelessRange<FVector3f>;
diff --git a/Engine/Plugins/FX/Niagara/Source/Niagara/Internal/Stateless/NiagaraStatelessDistribution.h b/Engine/Plugins/FX/Niagara/Source/Niagara/Internal/Stateless/NiagaraStatelessDistribution.h
--- a/Engine/Plugins/FX/Niagara/Source/Niagara/Internal/Stateless/NiagaraStatelessDistribution.h
+++ b/Engine/Plugins/FX/Niagara/Source/Niagara/Internal/Stateless/NiagaraStatelessDistribution.h
@@ -23,9 +23,7 @@ struct FNiagaraDistributionBase
 {
 	GENERATED_BODY()
 
-	virtual ~FNiagaraDistributionBase() 
-	{
-	}
+	virtual ~FNiagaraDistributionBase() = default;
 
 	UPROPERTY(EditAnywhere, Category = "Parameters")
 	ENiagaraDistributionMode Mode = ENiagaraDistributionMode::UniformConstant;
@@ -46,9 +44,32 @@ struct FNiagaraDistributionBase
 
 	virtual bool AllowCurves() const { return true; }
 	virtual void UpdateValuesFromDistribution() { }
+
+	static void PostEditChangeProperty(UObject* OwnerObject, FPropertyChangedEvent& PropertyChangedEvent);
 #endif
 };
 
+USTRUCT()
+struct FNiagaraDistributionRangeInt
+{
+	GENERATED_BODY()
+
+	FNiagaraDistributionRangeInt() = default;
+	explicit FNiagaraDistributionRangeInt(int32 ConstantValue) { InitConstant(ConstantValue); }
+
+	UPROPERTY(EditAnywhere, Category = "Parameters")
+	ENiagaraDistributionMode Mode = ENiagaraDistributionMode::UniformConstant;
+
+	UPROPERTY(EditAnywhere, Category = "Parameters")
+	int32 Min = 0;
+
+	UPROPERTY(EditAnywhere, Category = "Parameters")
+	int32 Max = 0;
+
+	NIAGARA_API void InitConstant(int32 Value);
+	NIAGARA_API FNiagaraStatelessRangeInt CalculateRange(const int32 Default = 0) const;
+};
+
 USTRUCT()
 struct FNiagaraDistributionRangeFloat : public FNiagaraDistributionBase
 {
@@ -143,11 +164,13 @@ struct FNiagaraDistributionVector2 : public FNiagaraDistributionBase
 	GENERATED_BODY()
 
 	FNiagaraDistributionVector2() = default;
+	explicit FNiagaraDistributionVector2(const float ConstantValue) { InitConstant(ConstantValue); }
 	explicit FNiagaraDistributionVector2(const FVector2f& ConstantValue) { InitConstant(ConstantValue); }
 
 	UPROPERTY(EditAnywhere, Category = "Parameters")
 	TArray<FVector2f> Values;
 
+	NIAGARA_API void InitConstant(const float Value);
 	NIAGARA_API void InitConstant(const FVector2f& Value);
 	NIAGARA_API FNiagaraStatelessRangeVector2 CalculateRange(const FVector2f& Default = FVector2f::ZeroVector) const;
 
@@ -162,11 +185,13 @@ struct FNiagaraDistributionVector3 : public FNiagaraDistributionBase
 	GENERATED_BODY()
 
 	FNiagaraDistributionVector3() = default;
+	explicit FNiagaraDistributionVector3(const float ConstantValue) { InitConstant(ConstantValue); }
 	explicit FNiagaraDistributionVector3(const FVector3f& ConstantValue) { InitConstant(ConstantValue); }
 
 	UPROPERTY(EditAnywhere, Category = "Parameters")
 	TArray<FVector3f> Values;
 
+	NIAGARA_API void InitConstant(const float Value);
 	NIAGARA_API void InitConstant(const FVector3f& Value);
 	NIAGARA_API FNiagaraStatelessRangeVector3 CalculateRange(const FVector3f& Default = FVector3f::ZeroVector) const;
 
diff --git a/Engine/Plugins/FX/Niagara/Source/Niagara/Internal/Stateless/NiagaraStatelessEmitterDataBuildContext.h b/Engine/Plugins/FX/Niagara/Source/Niagara/Internal/Stateless/NiagaraStatelessEmitterDataBuildContext.h
--- a/Engine/Plugins/FX/Niagara/Source/Niagara/Internal/Stateless/NiagaraStatelessEmitterDataBuildContext.h
+++ b/Engine/Plugins/FX/Niagara/Source/Niagara/Internal/Stateless/NiagaraStatelessEmitterDataBuildContext.h
@@ -49,17 +49,17 @@ public:
 	int32 AddRendererBinding(const FNiagaraParameterBinding& Binding);
 	int32 AddRendererBinding(const FNiagaraParameterBindingWithValue& Binding);	
 
-	// Adds a distribution into the LUT if enabled and returns the packed information to send to the shader
-	template<typename TDistribution>
-	FUintVector3 AddDistribution(const TDistribution& Distribution, bool bEnabled)
+	// Adds an distribution into the LUT if enabled
+	template<typename TType>
+	FUintVector3 AddDistribution(ENiagaraDistributionMode Mode, TConstArrayView<TType> Values, bool bEnabled)
 	{
 		FUintVector3 Parameters = FUintVector3::ZeroValue;
-		if (bEnabled && Distribution.Values.Num() > 0)
+		if (bEnabled && Values.Num() > 0)
 		{
 			constexpr uint32 StatelessDistributionFlag_Random = 0x00000001;
 			constexpr uint32 StatelessDistributionFlag_Uniform = 0x00000002;
 
-			switch (Distribution.Mode)
+			switch (Mode)
 			{
 				case ENiagaraDistributionMode::UniformConstant:		Parameters.X = StatelessDistributionFlag_Random | StatelessDistributionFlag_Uniform; break;
 				case ENiagaraDistributionMode::NonUniformConstant:	Parameters.X = StatelessDistributionFlag_Random; break;
@@ -70,12 +70,19 @@ public:
 				default:											checkNoEntry(); break;
 			}
 
-			Parameters.Y = AddStaticData(Distribution.Values);
-			reinterpret_cast<float&>(Parameters.Z) = Distribution.Values.Num() - 1;
+			Parameters.Y = AddStaticData(Values);
+			reinterpret_cast<float&>(Parameters.Z) = Values.Num() - 1;
 		}
 		return Parameters;
 	}
 
+	// Adds a distribution into the LUT if enabled and returns the packed information to send to the shader
+	template<typename TDistribution>
+	FUintVector3 AddDistribution(const TDistribution& Distribution, bool bEnabled)
+	{
+		return AddDistribution(Distribution.Mode, MakeArrayView(Distribution.Values), bEnabled);
+	}
+
 private:
 	FNiagaraParameterStore& RendererBindings;
 	TArray<uint8>&			BuiltData;
diff --git a/Engine/Plugins/FX/Niagara/Source/Niagara/Internal/Stateless/NiagaraStatelessSpawnInfo.h b/Engine/Plugins/FX/Niagara/Source/Niagara/Internal/Stateless/NiagaraStatelessSpawnInfo.h
--- a/Engine/Plugins/FX/Niagara/Source/Niagara/Internal/Stateless/NiagaraStatelessSpawnInfo.h
+++ b/Engine/Plugins/FX/Niagara/Source/Niagara/Internal/Stateless/NiagaraStatelessSpawnInfo.h
@@ -31,16 +31,10 @@ struct FNiagaraStatelessSpawnInfo
 	float SpawnTime = 0.0f;
 
 	UPROPERTY(EditAnywhere, Category = "Spawn", meta = (EditConditionHides, EditCondition = "Type == ENiagaraStatelessSpawnInfoType::Burst", ClampMin = "0"))
-	int32 AmountMin = 1;
-
-	UPROPERTY(EditAnywhere, Category = "Spawn", meta = (EditConditionHides, EditCondition = "Type == ENiagaraStatelessSpawnInfoType::Burst", ClampMin = "0"))
-	int32 AmountMax = 1;
-
-	UPROPERTY(EditAnywhere, Category = "Spawn", meta = (EditConditionHides, EditCondition = "Type == ENiagaraStatelessSpawnInfoType::Rate", ClampMin = "0.0"))
-	float RateMin = 60.0f;
+	FNiagaraDistributionRangeInt Amount = FNiagaraDistributionRangeInt(1);
 
 	UPROPERTY(EditAnywhere, Category = "Spawn", meta = (EditConditionHides, EditCondition = "Type == ENiagaraStatelessSpawnInfoType::Rate", ClampMin = "0.0"))
-	float RateMax = 60.0f;
+	FNiagaraDistributionRangeFloat Rate = FNiagaraDistributionRangeFloat(60.0f);
 
 	UPROPERTY(EditAnywhere, Category = "Spawn", meta = (InlineEditConditionToggle))
 	bool bSpawnProbabilityEnabled = false;
diff --git a/Engine/Plugins/FX/Niagara/Source/Niagara/Private/NiagaraSystemInstance.cpp b/Engine/Plugins/FX/Niagara/Source/Niagara/Private/NiagaraSystemInstance.cpp
--- a/Engine/Plugins/FX/Niagara/Source/Niagara/Private/NiagaraSystemInstance.cpp
+++ b/Engine/Plugins/FX/Niagara/Source/Niagara/Private/NiagaraSystemInstance.cpp
@@ -2994,8 +2994,8 @@ void FNiagaraSystemInstance::InitSystemState()
 	SystemState_RandomStream.Initialize(RandomSeed + RandomSeedOffset);
 
 	SystemState_LoopCount			= 0;
-	SystemState_CurrentLoopDuration = SystemState_RandomStream.FRandRange(SystemStateData.LoopDurationMin, SystemStateData.LoopDurationMax);
-	SystemState_CurrentLoopDelay	= SystemState_RandomStream.FRandRange(SystemStateData.LoopDelayMin, SystemStateData.LoopDelayMax);
+	SystemState_CurrentLoopDuration = SystemState_RandomStream.FRandRange(SystemStateData.LoopDuration.Min, SystemStateData.LoopDuration.Max);
+	SystemState_CurrentLoopDelay	= SystemState_RandomStream.FRandRange(SystemStateData.LoopDelay.Min, SystemStateData.LoopDelay.Max);
 	SystemState_CurrentLoopAgeStart	= 0.0f;
 	SystemState_CurrentLoopAgeEnd	= SystemState_CurrentLoopAgeStart + SystemState_CurrentLoopDelay + SystemState_CurrentLoopDuration;
 }
@@ -3036,7 +3036,7 @@ void FNiagaraSystemInstance::TickSystemState()
 
 				if (SystemStateData.bRecalculateDurationEachLoop)
 				{
-					SystemState_CurrentLoopDuration = SystemState_RandomStream.FRandRange(SystemStateData.LoopDurationMin, SystemStateData.LoopDurationMax);
+					SystemState_CurrentLoopDuration = SystemState_RandomStream.FRandRange(SystemStateData.LoopDuration.Min, SystemStateData.LoopDuration.Max);
 				}
 
 				if (SystemStateData.bDelayFirstLoopOnly)
@@ -3045,7 +3045,7 @@ void FNiagaraSystemInstance::TickSystemState()
 				}
 				else if (SystemStateData.bRecalculateDelayEachLoop)
 				{
-					SystemState_CurrentLoopDelay = SystemState_RandomStream.FRandRange(SystemStateData.LoopDelayMin, SystemStateData.LoopDelayMax);
+					SystemState_CurrentLoopDelay = SystemState_RandomStream.FRandRange(SystemStateData.LoopDelay.Min, SystemStateData.LoopDelay.Max);
 				}
 
 				SystemState_CurrentLoopAgeStart	= SystemState_CurrentLoopAgeEnd;
diff --git a/Engine/Plugins/FX/Niagara/Source/Niagara/Private/Stateless/NiagaraStatelessDistribution.cpp b/Engine/Plugins/FX/Niagara/Source/Niagara/Private/Stateless/NiagaraStatelessDistribution.cpp
--- a/Engine/Plugins/FX/Niagara/Source/Niagara/Private/Stateless/NiagaraStatelessDistribution.cpp
+++ b/Engine/Plugins/FX/Niagara/Source/Niagara/Private/Stateless/NiagaraStatelessDistribution.cpp
@@ -160,6 +160,18 @@ namespace NiagaraStatelessDistributionPrivate
 #endif
 } //NiagaraStatelessCommon
 
+void FNiagaraDistributionRangeInt::InitConstant(int32 Value)
+{
+	Mode = ENiagaraDistributionMode::UniformConstant;
+	Min = Value;
+	Max = Value;
+}
+
+FNiagaraStatelessRangeInt FNiagaraDistributionRangeInt::CalculateRange(const int32 Default) const
+{
+	return Mode == ENiagaraDistributionMode::UniformConstant ? FNiagaraStatelessRangeInt(Min, Min) : FNiagaraStatelessRangeInt(Min, Max);
+}
+
 void FNiagaraDistributionRangeFloat::InitConstant(float Value)
 {
 	Mode = ENiagaraDistributionMode::UniformConstant;
@@ -230,6 +242,15 @@ FNiagaraStatelessRangeFloat FNiagaraDistributionFloat::CalculateRange(const floa
 	return Range;
 }
 
+void FNiagaraDistributionVector2::InitConstant(const float Value)
+{
+	Mode = ENiagaraDistributionMode::UniformConstant;
+	Values = TArray<FVector2f>({ FVector2f(Value, Value), FVector2f(Value, Value) });
+#if WITH_EDITORONLY_DATA
+	ChannelConstantsAndRanges = TArray<float>({ Value, Value });
+#endif
+}
+
 void FNiagaraDistributionVector2::InitConstant(const FVector2f& Value)
 {
 	Mode = ENiagaraDistributionMode::NonUniformConstant;
@@ -255,6 +276,15 @@ FNiagaraStatelessRangeVector2 FNiagaraDistributionVector2::CalculateRange(const
 	return Range;
 }
 
+void FNiagaraDistributionVector3::InitConstant(const float Value)
+{
+	Mode = ENiagaraDistributionMode::UniformConstant;
+	Values = TArray<FVector3f>({ FVector3f(Value, Value, Value), FVector3f(Value, Value, Value) });
+#if WITH_EDITORONLY_DATA
+	ChannelConstantsAndRanges = TArray<float>({ Value, Value });
+#endif
+}
+
 void FNiagaraDistributionVector3::InitConstant(const FVector3f& Value)
 {
 	Mode = ENiagaraDistributionMode::NonUniformConstant;
@@ -313,6 +343,31 @@ FNiagaraStatelessRangeColor FNiagaraDistributionColor::CalculateRange(const FLin
 }
 
 #if WITH_EDITORONLY_DATA
+void FNiagaraDistributionBase::PostEditChangeProperty(UObject* OwnerObject, FPropertyChangedEvent& PropertyChangedEvent)
+{
+	FStructProperty* StructProperty = CastField<FStructProperty>(PropertyChangedEvent.Property);
+	if (StructProperty && StructProperty->Struct && StructProperty->Struct->IsChildOf(FNiagaraDistributionBase::StaticStruct()))
+	{
+		FNiagaraDistributionBase* ValuePtr = nullptr;
+		if (PropertyChangedEvent.Property != PropertyChangedEvent.MemberProperty)
+		{
+			// Properties stored in a UStruct inside a UObject need to first offset from UObject -> UStruct then UStruct -> Property
+			FStructProperty* MemberStructProperty = CastField<FStructProperty>(PropertyChangedEvent.MemberProperty);
+			void* StructPtr = MemberStructProperty->ContainerPtrToValuePtr<void>(OwnerObject);
+			ValuePtr = StructProperty->ContainerPtrToValuePtr<FNiagaraDistributionBase>(StructPtr);
+		}
+		else
+		{
+			ValuePtr = StructProperty->ContainerPtrToValuePtr<FNiagaraDistributionBase>(OwnerObject);
+		}
+
+		if (ValuePtr != nullptr)
+		{
+			ValuePtr->UpdateValuesFromDistribution();
+		}
+	}
+}
+
 void FNiagaraDistributionRangeFloat::UpdateValuesFromDistribution()
 {
 	NiagaraStatelessDistributionPrivate::UpdateDistributionValues(
diff --git a/Engine/Plugins/FX/Niagara/Source/Niagara/Private/Stateless/NiagaraStatelessEmitter.cpp b/Engine/Plugins/FX/Niagara/Source/Niagara/Private/Stateless/NiagaraStatelessEmitter.cpp
--- a/Engine/Plugins/FX/Niagara/Source/Niagara/Private/Stateless/NiagaraStatelessEmitter.cpp
+++ b/Engine/Plugins/FX/Niagara/Source/Niagara/Private/Stateless/NiagaraStatelessEmitter.cpp
@@ -144,6 +144,8 @@ bool UNiagaraStatelessEmitter::NeedsLoadForTargetPlatform(const ITargetPlatform*
 #if WITH_EDITOR
 void UNiagaraStatelessEmitter::PostEditChangeProperty(FPropertyChangedEvent& PropertyChangedEvent)
 {
+	FNiagaraDistributionBase::PostEditChangeProperty(this, PropertyChangedEvent);
+
 	Super::PostEditChangeProperty(PropertyChangedEvent);
 
 	//-TODO: This should be done differently
@@ -237,7 +239,7 @@ void UNiagaraStatelessEmitter::CacheFromCompiledData()
 		StatelessEmitterData->SpawnInfos,
 		[this](const FNiagaraStatelessSpawnInfo& SpawnInfo)
 		{
-			return SpawnInfo.IsValid(EmitterState.LoopDurationMax);
+			return SpawnInfo.IsValid(EmitterState.LoopDuration.Max);
 		}
 	);
 
diff --git a/Engine/Plugins/FX/Niagara/Source/Niagara/Private/Stateless/NiagaraStatelessEmitterInstance.cpp b/Engine/Plugins/FX/Niagara/Source/Niagara/Private/Stateless/NiagaraStatelessEmitterInstance.cpp
--- a/Engine/Plugins/FX/Niagara/Source/Niagara/Private/Stateless/NiagaraStatelessEmitterInstance.cpp
+++ b/Engine/Plugins/FX/Niagara/Source/Niagara/Private/Stateless/NiagaraStatelessEmitterInstance.cpp
@@ -232,8 +232,8 @@ void FNiagaraStatelessEmitterInstance::InitEmitterState()
 {
 	const FNiagaraEmitterStateData& EmitterState = EmitterData->EmitterState;
 	LoopCount			= 0;
-	CurrentLoopDuration = RandomStream.FRandRange(EmitterState.LoopDurationMin, EmitterState.LoopDurationMax);
-	CurrentLoopDelay	= RandomStream.FRandRange(EmitterState.LoopDelayMin, EmitterState.LoopDelayMax);
+	CurrentLoopDuration = RandomStream.FRandRange(EmitterState.LoopDuration.Min, EmitterState.LoopDuration.Max);
+	CurrentLoopDelay	= RandomStream.FRandRange(EmitterState.LoopDelay.Min, EmitterState.LoopDelay.Max);
 	CurrentLoopAgeStart	= 0.0f;
 	CurrentLoopAgeEnd	= CurrentLoopAgeStart + CurrentLoopDelay + CurrentLoopDuration;
 }
@@ -275,7 +275,7 @@ void FNiagaraStatelessEmitterInstance::TickEmitterState()
 
 			if (EmitterState.bRecalculateDurationEachLoop)
 			{
-				CurrentLoopDuration = RandomStream.FRandRange(EmitterState.LoopDurationMin, EmitterState.LoopDurationMax);
+				CurrentLoopDuration = RandomStream.FRandRange(EmitterState.LoopDuration.Min, EmitterState.LoopDuration.Max);
 			}
 
 			if (EmitterState.bDelayFirstLoopOnly)
@@ -284,7 +284,7 @@ void FNiagaraStatelessEmitterInstance::TickEmitterState()
 			}
 			else if ( EmitterState.bRecalculateDelayEachLoop )
 			{
-				CurrentLoopDelay	= RandomStream.FRandRange(EmitterState.LoopDelayMin, EmitterState.LoopDelayMax);
+				CurrentLoopDelay	= RandomStream.FRandRange(EmitterState.LoopDelay.Min, EmitterState.LoopDelay.Max);
 			}
 
 			CurrentLoopAgeStart	= CurrentLoopAgeEnd;
@@ -327,7 +327,7 @@ void FNiagaraStatelessEmitterInstance::InitSpawnInfos()
 	{
 		if (SpawnInfo.Type == ENiagaraStatelessSpawnInfoType::Rate)
 		{
-			const float SpawnRate = RandomStream.FRandRange(SpawnInfo.RateMin, SpawnInfo.RateMax);
+			const float SpawnRate = RandomStream.FRandRange(SpawnInfo.Rate.Min, SpawnInfo.Rate.Max);
 			if (SpawnRate > 0.0f)
 			{
 				FActiveSpawnRate& ActiveSpawnRate = ActiveSpawnRates.AddDefaulted_GetRef();
@@ -379,7 +379,7 @@ void FNiagaraStatelessEmitterInstance::InitSpawnInfosForLoop()
 			continue;
 		}
 
-		const int32 SpawnAmount = RandomStream.RandRange(SpawnInfo.AmountMin, SpawnInfo.AmountMax);
+		const int32 SpawnAmount = RandomStream.RandRange(SpawnInfo.Amount.Min, SpawnInfo.Amount.Max);
 		if (SpawnAmount <= 0)
 		{
 			continue;
diff --git a/Engine/Plugins/FX/Niagara/Source/Niagara/Private/Stateless/NiagaraStatelessModule.cpp b/Engine/Plugins/FX/Niagara/Source/Niagara/Private/Stateless/NiagaraStatelessModule.cpp
--- a/Engine/Plugins/FX/Niagara/Source/Niagara/Private/Stateless/NiagaraStatelessModule.cpp
+++ b/Engine/Plugins/FX/Niagara/Source/Niagara/Private/Stateless/NiagaraStatelessModule.cpp
@@ -60,15 +60,8 @@ bool UNiagaraStatelessModule::CanEditChange(const FProperty* InProperty) const
 
 void UNiagaraStatelessModule::PostEditChangeProperty(struct FPropertyChangedEvent& PropertyChangedEvent)
 {
-	FStructProperty* StructProperty = CastField<FStructProperty>(PropertyChangedEvent.MemberProperty);
-	if (StructProperty != nullptr && StructProperty->Struct->IsChildOf(FNiagaraDistributionBase::StaticStruct()))
-	{
-		FNiagaraDistributionBase* ValuePtr = StructProperty->ContainerPtrToValuePtr<FNiagaraDistributionBase>(this);
-		if (ValuePtr != nullptr)
-		{
-			ValuePtr->UpdateValuesFromDistribution();
-		}
-	}
+	FNiagaraDistributionBase::PostEditChangeProperty(this, PropertyChangedEvent);
+
 	Super::PostEditChangeProperty(PropertyChangedEvent);
 }
 #endif //WITH_EDITOR
diff --git a/Engine/Plugins/FX/Niagara/Source/Niagara/Private/Stateless/NiagaraStatelessSpawnInfo.cpp b/Engine/Plugins/FX/Niagara/Source/Niagara/Private/Stateless/NiagaraStatelessSpawnInfo.cpp
--- a/Engine/Plugins/FX/Niagara/Source/Niagara/Private/Stateless/NiagaraStatelessSpawnInfo.cpp
+++ b/Engine/Plugins/FX/Niagara/Source/Niagara/Private/Stateless/NiagaraStatelessSpawnInfo.cpp
@@ -7,10 +7,10 @@ bool FNiagaraStatelessSpawnInfo::IsValid(TOptional<float> LoopDuration) const
 	switch (Type)
 	{
 		case ENiagaraStatelessSpawnInfoType::Burst:
-			return (AmountMin + AmountMax) > 0 && SpawnTime >= 0.0f && SpawnTime < LoopDuration.Get(SpawnTime + UE_SMALL_NUMBER);
+			return (Amount.Min + Amount.Max) > 0 && SpawnTime >= 0.0f && SpawnTime < LoopDuration.Get(SpawnTime + UE_SMALL_NUMBER);
 
 		case ENiagaraStatelessSpawnInfoType::Rate:
-			return (RateMin + RateMax) > 0.0f;
+			return (Rate.Min + Rate.Max) > 0.0f;
 
 		default:
 			checkNoEntry();
diff --git a/Engine/Plugins/FX/Niagara/Source/NiagaraEditor/Private/NiagaraOverviewNode.cpp b/Engine/Plugins/FX/Niagara/Source/NiagaraEditor/Private/NiagaraOverviewNode.cpp
--- a/Engine/Plugins/FX/Niagara/Source/NiagaraEditor/Private/NiagaraOverviewNode.cpp
+++ b/Engine/Plugins/FX/Niagara/Source/NiagaraEditor/Private/NiagaraOverviewNode.cpp
@@ -300,7 +300,7 @@ void UNiagaraOverviewNode::GetNodeContextMenuActions(class UToolMenu* Menu, clas
 						)
 					)
 				);
-
+			#if 0
 				if (GetDefault<UNiagaraSettings>()->bStatelessEmittersEnabled)
 				{
 					Section.AddSubMenu(
@@ -340,6 +340,7 @@ void UNiagaraOverviewNode::GetNodeContextMenuActions(class UToolMenu* Menu, clas
 							}
 						}));
 				}
+			#endif
 			}
 		}
 	}
diff --git a/Engine/Plugins/FX/Niagara/Source/NiagaraEditor/Private/Stateless/NiagaraDistributionIntPropertyCustomization.cpp b/Engine/Plugins/FX/Niagara/Source/NiagaraEditor/Private/Stateless/NiagaraDistributionIntPropertyCustomization.cpp
new file mode 100644
--- /dev/null
+++ b/Engine/Plugins/FX/Niagara/Source/NiagaraEditor/Private/Stateless/NiagaraDistributionIntPropertyCustomization.cpp
@@ -0,0 +1,304 @@
+// Copyright Epic Games, Inc. All Rights Reserved.
+
+#include "Stateless/NiagaraDistributionIntPropertyCustomization.h"
+
+#include "DetailWidgetRow.h"
+#include "Editor.h"
+#include "EditorUndoClient.h"
+#include "NiagaraEditorStyle.h"
+#include "PropertyHandle.h"
+#include "ScopedTransaction.h"
+
+#include "Framework/MultiBox/MultiBoxBuilder.h"
+#include "Stateless/NiagaraStatelessDistribution.h"
+#include "Styling/StyleColors.h"
+#include "Widgets/DeclarativeSyntaxSupport.h"
+#include "Widgets/NiagaraDistributionEditorUtilities.h"
+#include "Widgets/SBoxPanel.h"
+#include "Widgets/SCompoundWidget.h"
+#include "Widgets/INiagaraDistributionAdapter.h"
+#include "Widgets/Input/SComboButton.h"
+#include "Widgets/Input/SNumericEntryBox.h"
+#include "Widgets/Layout/SBox.h"
+
+#define LOCTEXT_NAMESPACE "NiagaraDistributionIntPropertyCustomization"
+
+class SNiagaraDistributionIntPropertyWidget : public SCompoundWidget
+{
+	SLATE_BEGIN_ARGS(SNiagaraDistributionIntPropertyWidget) { }
+	SLATE_END_ARGS();
+
+	void Construct(const FArguments& InArgs, TSharedRef<IPropertyHandle> InPropertyHandle, UObject* InOwnerObject, FNiagaraDistributionRangeInt* InDistribution)
+	{
+		WeakPropertyHandle	= InPropertyHandle;
+		WeakOwnerObject = InOwnerObject;
+		Distribution = InDistribution;
+
+		ChildSlot
+		[
+			SNew(SHorizontalBox)
+			+ SHorizontalBox::Slot()
+			.AutoWidth()
+			.Padding(0, 0, 10, 0)
+			[
+				SNew(SComboButton)
+				.ButtonStyle(FAppStyle::Get(), "HoverHintOnly")
+				.ContentPadding(FMargin(0))
+				.OnGetMenuContent(this, &SNiagaraDistributionIntPropertyWidget::OnGetDistributionModeMenuContent)
+				.ButtonContent()
+				[
+					SNew(SImage)
+					.Image(this, &SNiagaraDistributionIntPropertyWidget::GetDistributionModeBrush)
+					.ToolTipText(this, &SNiagaraDistributionIntPropertyWidget::GetDistributionModeToolTip)
+				]
+			]
+			+ SHorizontalBox::Slot()
+			[
+				SAssignNew(ContentBox, SBox)
+				[
+					ConstructContentForMode()
+				]
+			]
+		];
+	}
+
+	virtual ~SNiagaraDistributionIntPropertyWidget()
+	{
+	}
+
+	ENiagaraDistributionEditorMode GetDistributionMode() const
+	{
+		TSharedPtr<IPropertyHandle> PropertyHandle = WeakPropertyHandle.Pin();
+		UObject* OwnerObject = WeakOwnerObject.Get();
+		if (PropertyHandle.IsValid() && OwnerObject && Distribution)
+		{
+			return Distribution->Mode == ENiagaraDistributionMode::UniformConstant ? ENiagaraDistributionEditorMode::UniformConstant : ENiagaraDistributionEditorMode::UniformRange;
+		}
+		return ENiagaraDistributionEditorMode::UniformConstant;
+	}
+
+	void SetDistributionMode(ENiagaraDistributionEditorMode NewMode)
+	{
+		TSharedPtr<IPropertyHandle> PropertyHandle = WeakPropertyHandle.Pin();
+		UObject* OwnerObject = WeakOwnerObject.Get();
+		if (PropertyHandle.IsValid() && OwnerObject && Distribution)
+		{
+			const FScopedTransaction Transaction(LOCTEXT("SetDistributionMode", "Set distribution mode"));
+			OwnerObject->Modify();
+			PropertyHandle->NotifyPreChange();
+
+			Distribution->Mode = NewMode == ENiagaraDistributionEditorMode::UniformConstant ? ENiagaraDistributionMode::UniformConstant : ENiagaraDistributionMode::UniformRange;
+			ContentBox->SetContent(ConstructContentForMode());
+
+			PropertyHandle->NotifyPostChange(bContinuousChangeActive ? EPropertyChangeType::Interactive : EPropertyChangeType::ValueSet);
+		}
+	}
+
+	bool IsDistributionModeSelected(ENiagaraDistributionEditorMode Mode) const
+	{
+		return Mode == GetDistributionMode();
+	}
+
+	TSharedRef<SWidget> OnGetDistributionModeMenuContent()
+	{
+		FMenuBuilder MenuBuilder(true, nullptr);
+		TArray<ENiagaraDistributionEditorMode> SupportedModes { ENiagaraDistributionEditorMode::UniformConstant, ENiagaraDistributionEditorMode::UniformRange };
+		for (ENiagaraDistributionEditorMode SupportedMode : SupportedModes)
+		{
+			MenuBuilder.AddMenuEntry(
+				FNiagaraDistributionEditorUtilities::DistributionModeToDisplayName(SupportedMode),
+				FNiagaraDistributionEditorUtilities::DistributionModeToToolTipText(SupportedMode),
+				FNiagaraDistributionEditorUtilities::DistributionModeToIcon(SupportedMode),
+				FUIAction(
+					FExecuteAction::CreateSP(this, &SNiagaraDistributionIntPropertyWidget::SetDistributionMode, SupportedMode),
+					FCanExecuteAction(),
+					FIsActionChecked::CreateSP(this, &SNiagaraDistributionIntPropertyWidget::IsDistributionModeSelected, SupportedMode)
+				),
+				NAME_None,
+				EUserInterfaceActionType::RadioButton
+			);
+		}
+
+		return MenuBuilder.MakeWidget();
+	}
+
+	const FSlateBrush* GetDistributionModeBrush() const
+	{
+		return FNiagaraDistributionEditorUtilities::DistributionModeToIconBrush(GetDistributionMode());
+	}
+
+	FText GetDistributionModeToolTip() const
+	{
+		return FNiagaraDistributionEditorUtilities::DistributionModeToToolTipText(GetDistributionMode());
+	}
+
+	TSharedRef<SWidget> ConstructContentForMode()
+	{
+		const ENiagaraDistributionEditorMode Mode = GetDistributionMode();
+		if (Mode == ENiagaraDistributionEditorMode::UniformConstant)
+		{
+			return SNew(SHorizontalBox)
+				+ SHorizontalBox::Slot()
+				[
+					ConstructIntWidget(0, FText())
+				];
+		}
+		else
+		{
+			return SNew(SHorizontalBox)
+				+ SHorizontalBox::Slot()
+				[
+					ConstructIntWidget(0, LOCTEXT("MinLabel", "Min"))
+				]
+				+ SHorizontalBox::Slot()
+				[
+					ConstructIntWidget(1, LOCTEXT("MaxLabel", "Max"))
+				];
+		}
+	}
+
+	TSharedRef<SWidget> ConstructIntWidget(int32 ValueIndex, FText LabelText)
+	{
+		TSharedRef<SWidget> LabelWidget = SNullWidget::NullWidget;
+		if ( !LabelText.IsEmpty() )
+		{
+			LabelWidget =
+				SNew(STextBlock)
+				.TextStyle(FNiagaraEditorStyle::Get(), "NiagaraEditor.ParameterText")
+				.Text(LabelText);
+		}
+
+		return SNew(SNumericEntryBox<int32>)
+			.Font(FAppStyle::Get().GetFontStyle("PropertyWindow.NormalFont"))
+			.Value(this, &SNiagaraDistributionIntPropertyWidget::GetValue, ValueIndex)
+			.OnValueChanged(this, &SNiagaraDistributionIntPropertyWidget::ValueChanged, ValueIndex)
+			.OnValueCommitted(this, &SNiagaraDistributionIntPropertyWidget::ValueCommitted, ValueIndex)
+			.OnBeginSliderMovement(this, &SNiagaraDistributionIntPropertyWidget::BeginValueSliderMovement)
+			.OnEndSliderMovement(this, &SNiagaraDistributionIntPropertyWidget::EndValueSliderMovement)
+			.AllowSpin(true)
+			.MinValue(TOptional<int32>())
+			.MaxValue(TOptional<int32>())
+			.MinSliderValue(TOptional<int32>())
+			.MaxSliderValue(TOptional<int32>())
+			.BroadcastValueChangesPerKey(false)
+			.LabelVAlign(EVerticalAlignment::VAlign_Center)
+			.MinDesiredValueWidth(30)
+			.Label()
+			[
+				LabelWidget
+			];
+	}
+
+	TOptional<int32> GetValue(int32 ValueIndex) const
+	{
+		TOptional<int32> Value;
+		TSharedPtr<IPropertyHandle> PropertyHandle = WeakPropertyHandle.Pin();
+		UObject* OwnerObject = WeakOwnerObject.Get();
+		if (PropertyHandle.IsValid() && OwnerObject && Distribution)
+		{
+			Value = ValueIndex == 0  ? Distribution->Min : Distribution->Max;
+		}
+		return Value;
+	}
+
+	void ValueChanged(int32 Value, int32 ValueIndex)
+	{
+		TSharedPtr<IPropertyHandle> PropertyHandle = WeakPropertyHandle.Pin();
+		UObject* OwnerObject = WeakOwnerObject.Get();
+		if (PropertyHandle.IsValid() && OwnerObject && Distribution)
+		{
+			const FScopedTransaction Transaction(LOCTEXT("SetDistributionValue", "Set distribution value"));
+			OwnerObject->Modify();
+			PropertyHandle->NotifyPreChange();
+
+			Distribution->Min = (ValueIndex == 0) || Distribution->Mode == ENiagaraDistributionMode::UniformConstant ? Value : Distribution->Min;
+			Distribution->Max = (ValueIndex != 0) || Distribution->Mode == ENiagaraDistributionMode::UniformConstant ? Value : Distribution->Max;
+
+			PropertyHandle->NotifyPostChange(bContinuousChangeActive ? EPropertyChangeType::Interactive : EPropertyChangeType::ValueSet);
+		}
+	}
+
+	void ValueCommitted(int32 Value, ETextCommit::Type CommitInfo, int32 ValueIndex)
+	{
+		if (CommitInfo == ETextCommit::OnEnter || CommitInfo == ETextCommit::OnUserMovedFocus)
+		{
+			ValueChanged(Value, ValueIndex);
+		}
+	}
+
+	void BeginValueSliderMovement()
+	{
+		bContinuousChangeActive = true;
+	}
+
+	void EndValueSliderMovement(int32 Value)
+	{
+		bContinuousChangeActive = false;
+	}
+
+private:
+	TWeakPtr<IPropertyHandle>		WeakPropertyHandle;
+	TWeakObjectPtr<UObject>			WeakOwnerObject;
+	FNiagaraDistributionRangeInt*	Distribution = nullptr;
+	bool							bContinuousChangeActive = false;
+	TSharedPtr<SBox>				ContentBox;
+};
+
+TSharedRef<IPropertyTypeCustomization> FNiagaraDistributionIntPropertyCustomization::MakeIntInstance(UObject* OptionalOuter)
+{
+	return MakeShared<FNiagaraDistributionIntPropertyCustomization>(OptionalOuter);
+}
+
+void FNiagaraDistributionIntPropertyCustomization::CustomizeHeader(TSharedRef<IPropertyHandle> PropertyHandle, FDetailWidgetRow& HeaderRow, IPropertyTypeCustomizationUtils& CustomizationUtils)
+{
+	UObject* OwnerObject = nullptr;
+	FNiagaraDistributionRangeInt* Distribution = nullptr;
+	{
+		TArray<UObject*> OuterObjects;
+		if (UObject* OuterObject = WeakOuterObject.Get())
+		{
+			TArray<TSharedPtr<FStructOnScope>> OutStructOnScopes;
+			PropertyHandle->GetOuterStructs(OutStructOnScopes);
+			if (OutStructOnScopes.Num() == 1)
+			{
+				OuterObjects.Add(OuterObject);
+			}
+		}
+		else
+		{
+			PropertyHandle->GetOuterObjects(OuterObjects);
+		}
+
+		if (OuterObjects.Num() == 1)
+		{
+			if (FProperty* Property = PropertyHandle->GetProperty())
+			{
+				FStructProperty* StructProperty = CastField<FStructProperty>(Property);
+				if (StructProperty && StructProperty->Struct && StructProperty->Struct->IsChildOf(FNiagaraDistributionRangeInt::StaticStruct()))
+				{
+					void* ValuePtr = nullptr;
+					if (PropertyHandle->GetValueData(ValuePtr) == FPropertyAccess::Success)
+					{
+						OwnerObject = OuterObjects[0];
+						Distribution = static_cast<FNiagaraDistributionRangeInt*>(ValuePtr);
+					}
+				}
+			}
+		}
+	}
+
+	HeaderRow.NameContent()
+	[
+		PropertyHandle->CreatePropertyNameWidget()
+	];
+	HeaderRow.ValueContent()
+	[
+		SNew(SNiagaraDistributionIntPropertyWidget, PropertyHandle, OwnerObject, Distribution)
+	];
+}
+
+void FNiagaraDistributionIntPropertyCustomization::CustomizeChildren(TSharedRef<IPropertyHandle> PropertyHandle, IDetailChildrenBuilder& ChildBuilder, IPropertyTypeCustomizationUtils& CustomizationUtils)
+{
+}
+
+#undef LOCTEXT_NAMESPACE
diff --git a/Engine/Plugins/FX/Niagara/Source/NiagaraEditor/Private/Stateless/NiagaraDistributionIntPropertyCustomization.h b/Engine/Plugins/FX/Niagara/Source/NiagaraEditor/Private/Stateless/NiagaraDistributionIntPropertyCustomization.h
new file mode 100644
--- /dev/null
+++ b/Engine/Plugins/FX/Niagara/Source/NiagaraEditor/Private/Stateless/NiagaraDistributionIntPropertyCustomization.h
@@ -0,0 +1,21 @@
+// Copyright Epic Games, Inc. All Rights Reserved.
+
+#pragma once
+
+#include "IPropertyTypeCustomization.h"
+
+class INiagaraDistributionAdapter;
+
+class FNiagaraDistributionIntPropertyCustomization : public IPropertyTypeCustomization
+{
+public:
+	static TSharedRef<IPropertyTypeCustomization> MakeIntInstance(UObject* OptionalOuter);
+
+	virtual void CustomizeHeader(TSharedRef<IPropertyHandle> PropertyHandle, FDetailWidgetRow& HeaderRow, IPropertyTypeCustomizationUtils& CustomizationUtils) override;
+	virtual void CustomizeChildren(TSharedRef<IPropertyHandle> PropertyHandle, IDetailChildrenBuilder& ChildBuilder, IPropertyTypeCustomizationUtils& CustomizationUtils) override;
+
+	FNiagaraDistributionIntPropertyCustomization(UObject* InOuterObject) : WeakOuterObject(InOuterObject) { }
+
+private:
+	TWeakObjectPtr<UObject> WeakOuterObject;
+};
diff --git a/Engine/Plugins/FX/Niagara/Source/NiagaraEditor/Private/Stateless/NiagaraDistributionPropertyCustomization.cpp b/Engine/Plugins/FX/Niagara/Source/NiagaraEditor/Private/Stateless/NiagaraDistributionPropertyCustomization.cpp
--- a/Engine/Plugins/FX/Niagara/Source/NiagaraEditor/Private/Stateless/NiagaraDistributionPropertyCustomization.cpp
+++ b/Engine/Plugins/FX/Niagara/Source/NiagaraEditor/Private/Stateless/NiagaraDistributionPropertyCustomization.cpp
@@ -33,7 +33,11 @@ public:
 
 			const FName DisableCurveDistributionName("DisableCurveDistribution");
 			const FName DisableUniformDistributionName("DisableUniformDistribution");
+			const FName DisableNonUniformDistributionName("DisableNonUniformDistribution");
+			const FName DisableRangeDistributionName("DisableRangeDistribution");
 			bAllowUniform = PropertyHandle ? PropertyHandle->HasMetaData(DisableUniformDistributionName) == false : false;
+			bAllowNonUniform = PropertyHandle ? PropertyHandle->HasMetaData(DisableNonUniformDistributionName) == false : false;
+			bAllowRange = PropertyHandle ? PropertyHandle->HasMetaData(DisableRangeDistributionName) == false : false;
 			bAllowCurves = InDistribution->AllowCurves() && (PropertyHandle ? PropertyHandle->HasMetaData(DisableCurveDistributionName) == false : false);
 		}
 	}
@@ -81,8 +85,14 @@ public:
 	{
 		if (SourceNumChannels == 1)
 		{
-			OutSupportedModes.Add(ENiagaraDistributionEditorMode::Constant);
-			OutSupportedModes.Add(ENiagaraDistributionEditorMode::Range);
+			if (bAllowUniform)
+			{
+				OutSupportedModes.Add(ENiagaraDistributionEditorMode::Constant);
+				if (bAllowRange)
+				{
+					OutSupportedModes.Add(ENiagaraDistributionEditorMode::Range);
+				}
+			}
 			if (bAllowCurves)
 			{
 				OutSupportedModes.Add(ENiagaraDistributionEditorMode::Curve);
@@ -93,14 +103,26 @@ public:
 			if (bAllowUniform)
 			{
 				OutSupportedModes.Add(ENiagaraDistributionEditorMode::UniformConstant);
-				OutSupportedModes.Add(ENiagaraDistributionEditorMode::UniformRange);
+				if (bAllowRange)
+				{
+					OutSupportedModes.Add(ENiagaraDistributionEditorMode::UniformRange);
+				}
+			}
+			if (bAllowNonUniform)
+			{
+				OutSupportedModes.Add(ENiagaraDistributionEditorMode::NonUniformConstant);
+				if (bAllowRange)
+				{
+					OutSupportedModes.Add(ENiagaraDistributionEditorMode::NonUniformRange);
+				}
 			}
-			OutSupportedModes.Add(ENiagaraDistributionEditorMode::NonUniformConstant);
-			OutSupportedModes.Add(ENiagaraDistributionEditorMode::NonUniformRange);
 			if (bAllowCurves)
 			{
 				OutSupportedModes.Add(ENiagaraDistributionEditorMode::UniformCurve);
-				OutSupportedModes.Add(ENiagaraDistributionEditorMode::NonUniformCurve);
+				if (bAllowNonUniform)
+				{
+					OutSupportedModes.Add(ENiagaraDistributionEditorMode::NonUniformCurve);
+				}
 			}
 		}
 	}
@@ -496,6 +518,8 @@ private:
 	bool bContinuousChangeActive = false;
 
 	bool bAllowUniform = true;
+	bool bAllowNonUniform = true;
+	bool bAllowRange = true;
 	bool bAllowCurves = true;
 };
 
@@ -547,13 +571,25 @@ private:
 	bool bUpdatingHandle = false;
 };
 
-TSharedRef<IPropertyTypeCustomization> FNiagaraDistributionPropertyCustomization::MakeFloatInstance()
+TSharedRef<IPropertyTypeCustomization> FNiagaraDistributionPropertyCustomization::MakeFloatInstance(UObject* OptionalOuter)
 {
-	FPropertyHandleToDistributionAdapter FloatDistributionPropertyHandleToDistributionAdapter = FPropertyHandleToDistributionAdapter::CreateLambda([](TSharedRef<IPropertyHandle> FloatDistributionPropertyHandle)
+	FPropertyHandleToDistributionAdapter FloatDistributionPropertyHandleToDistributionAdapter = FPropertyHandleToDistributionAdapter::CreateLambda([OptionalOuter](TSharedRef<IPropertyHandle> FloatDistributionPropertyHandle)
 	{
 		void* ValueData = nullptr;
 		TArray<UObject*> OuterObjects;
-		FloatDistributionPropertyHandle->GetOuterObjects(OuterObjects);
+		if (OptionalOuter)
+		{
+			TArray<TSharedPtr<FStructOnScope>> OutStructOnScopes;
+			FloatDistributionPropertyHandle->GetOuterStructs(OutStructOnScopes);
+			if (OutStructOnScopes.Num() == 1)
+			{
+				OuterObjects.Add(OptionalOuter);
+			}
+		}
+		else
+		{
+			FloatDistributionPropertyHandle->GetOuterObjects(OuterObjects);
+		}
 		if (OuterObjects.Num() == 1 && FloatDistributionPropertyHandle->GetValueData(ValueData) == FPropertyAccess::Success)
 		{
 			FNiagaraDistributionBase* FloatDistribution = static_cast<FNiagaraDistributionBase*>(ValueData);
@@ -565,6 +601,11 @@ TSharedRef<IPropertyTypeCustomization> FNiagaraDistributionPropertyCustomization
 	return MakeShareable<FNiagaraDistributionPropertyCustomization>(new FNiagaraDistributionPropertyCustomization(FloatDistributionPropertyHandleToDistributionAdapter));
 }
 
+TSharedRef<IPropertyTypeCustomization> FNiagaraDistributionPropertyCustomization::MakeFloatInstance()
+{
+	return MakeFloatInstance(nullptr);
+}
+
 TSharedRef<IPropertyTypeCustomization> FNiagaraDistributionPropertyCustomization::MakeVector2Instance()
 {
 	FPropertyHandleToDistributionAdapter Vector2DistributionPropertyHandleToDistributionAdapter = FPropertyHandleToDistributionAdapter::CreateLambda([](TSharedRef<IPropertyHandle> Vector2DistributionPropertyHandle)
diff --git a/Engine/Plugins/FX/Niagara/Source/NiagaraEditor/Private/Stateless/NiagaraDistributionPropertyCustomization.h b/Engine/Plugins/FX/Niagara/Source/NiagaraEditor/Private/Stateless/NiagaraDistributionPropertyCustomization.h
--- a/Engine/Plugins/FX/Niagara/Source/NiagaraEditor/Private/Stateless/NiagaraDistributionPropertyCustomization.h
+++ b/Engine/Plugins/FX/Niagara/Source/NiagaraEditor/Private/Stateless/NiagaraDistributionPropertyCustomization.h
@@ -9,6 +9,7 @@ class INiagaraDistributionAdapter;
 class FNiagaraDistributionPropertyCustomization : public IPropertyTypeCustomization
 {
 public:
+	static TSharedRef<IPropertyTypeCustomization> MakeFloatInstance(UObject* OptionalOuter);
 	static TSharedRef<IPropertyTypeCustomization> MakeFloatInstance();
 	static TSharedRef<IPropertyTypeCustomization> MakeVector2Instance();
 	static TSharedRef<IPropertyTypeCustomization> MakeVector3Instance();
diff --git a/Engine/Plugins/FX/Niagara/Source/NiagaraEditor/Private/ViewModels/Stack/NiagaraStackStatelessEmitterGroup.cpp b/Engine/Plugins/FX/Niagara/Source/NiagaraEditor/Private/ViewModels/Stack/NiagaraStackStatelessEmitterGroup.cpp
--- a/Engine/Plugins/FX/Niagara/Source/NiagaraEditor/Private/ViewModels/Stack/NiagaraStackStatelessEmitterGroup.cpp
+++ b/Engine/Plugins/FX/Niagara/Source/NiagaraEditor/Private/ViewModels/Stack/NiagaraStackStatelessEmitterGroup.cpp
@@ -4,6 +4,7 @@
 
 #include "IDetailTreeNode.h"
 #include "PropertyHandle.h"
+#include "Stateless/NiagaraDistributionPropertyCustomization.h"
 #include "Stateless/NiagaraStatelessEmitter.h"
 #include "Styling/AppStyle.h"
 #include "ViewModels/Stack/NiagaraStackObject.h"
@@ -113,6 +114,13 @@ void UNiagaraStackStatelessEmitterObjectItem::RefreshChildrenInternal(const TArr
 			bool bHideTopLevelCategories = false;
 			StatelessEmitterStackObject->Initialize(CreateDefaultChildRequiredData(), StatelessEmitter, bIsTopLevelObject, bHideTopLevelCategories, GetStackEditorDataKey());
 			StatelessEmitterStackObject->SetOnFilterDetailNodes(OnFilterDetailNodes, UNiagaraStackObject::EDetailNodeFilterMode::FilterAllNodes);
+			StatelessEmitterStackObject->RegisterInstancedCustomPropertyTypeLayout(FNiagaraDistributionFloat::StaticStruct()->GetFName(), FOnGetPropertyTypeCustomizationInstance::CreateStatic(&FNiagaraDistributionPropertyCustomization::MakeFloatInstance));
+			StatelessEmitterStackObject->RegisterInstancedCustomPropertyTypeLayout(FNiagaraDistributionVector2::StaticStruct()->GetFName(), FOnGetPropertyTypeCustomizationInstance::CreateStatic(&FNiagaraDistributionPropertyCustomization::MakeVector2Instance));
+			StatelessEmitterStackObject->RegisterInstancedCustomPropertyTypeLayout(FNiagaraDistributionVector3::StaticStruct()->GetFName(), FOnGetPropertyTypeCustomizationInstance::CreateStatic(&FNiagaraDistributionPropertyCustomization::MakeVector3Instance));
+			StatelessEmitterStackObject->RegisterInstancedCustomPropertyTypeLayout(FNiagaraDistributionColor::StaticStruct()->GetFName(), FOnGetPropertyTypeCustomizationInstance::CreateStatic(&FNiagaraDistributionPropertyCustomization::MakeColorInstance));
+			StatelessEmitterStackObject->RegisterInstancedCustomPropertyTypeLayout(FNiagaraDistributionRangeFloat::StaticStruct()->GetFName(), FOnGetPropertyTypeCustomizationInstance::CreateStatic(&FNiagaraDistributionPropertyCustomization::MakeFloatInstance));
+			StatelessEmitterStackObject->RegisterInstancedCustomPropertyTypeLayout(FNiagaraDistributionRangeVector2::StaticStruct()->GetFName(), FOnGetPropertyTypeCustomizationInstance::CreateStatic(&FNiagaraDistributionPropertyCustomization::MakeVector2Instance));
+			StatelessEmitterStackObject->RegisterInstancedCustomPropertyTypeLayout(FNiagaraDistributionRangeVector3::StaticStruct()->GetFName(), FOnGetPropertyTypeCustomizationInstance::CreateStatic(&FNiagaraDistributionPropertyCustomization::MakeVector3Instance));
 			StatelessEmitterStackObjectWeak = StatelessEmitterStackObject;
 		}
 		NewChildren.Add(StatelessEmitterStackObject);
diff --git a/Engine/Plugins/FX/Niagara/Source/NiagaraEditor/Private/ViewModels/Stack/NiagaraStackStatelessEmitterSpawnGroup.cpp b/Engine/Plugins/FX/Niagara/Source/NiagaraEditor/Private/ViewModels/Stack/NiagaraStackStatelessEmitterSpawnGroup.cpp
--- a/Engine/Plugins/FX/Niagara/Source/NiagaraEditor/Private/ViewModels/Stack/NiagaraStackStatelessEmitterSpawnGroup.cpp
+++ b/Engine/Plugins/FX/Niagara/Source/NiagaraEditor/Private/ViewModels/Stack/NiagaraStackStatelessEmitterSpawnGroup.cpp
@@ -6,6 +6,8 @@
 #include "NiagaraEditorStyle.h"
 #include "PropertyHandle.h"
 #include "ScopedTransaction.h"
+#include "Stateless/NiagaraDistributionPropertyCustomization.h"
+#include "Stateless/NiagaraDistributionIntPropertyCustomization.h"
 #include "Stateless/NiagaraStatelessEmitter.h"
 #include "Styling/AppStyle.h"
 #include "ViewModels/Stack/NiagaraStackItemPropertyHeaderValueShared.h"
@@ -120,6 +122,7 @@ void UNiagaraStackStatelessEmitterSpawnItem::Initialize(FRequiredEntryData InReq
 	StatelessEmitterWeak = InStatelessEmitter;
 	Index = InIndex;
 	SourceId = InStatelessEmitter->GetSpawnInfoByIndex(InIndex)->SourceId;
+	OnDataObjectModified().AddUObject(this, &UNiagaraStackStatelessEmitterSpawnItem::OnSpawnInfoModified);
 }
 
 FText UNiagaraStackStatelessEmitterSpawnItem::GetDisplayName() const
@@ -164,16 +167,26 @@ void UNiagaraStackStatelessEmitterSpawnItem::RefreshChildrenInternal(const TArra
 		}
 
 		UNiagaraStackObject* SpawnInfoObject = SpawnInfoObjectWeak.Get();
-		if (SpawnInfoObject == nullptr || SpawnInfoObject->GetObject() != StatelessEmitterWeak.Get() ||
+		UObject* StatelessEmitterObject = StatelessEmitterWeak.Get();
+		if (SpawnInfoObject == nullptr || SpawnInfoObject->GetObject() != StatelessEmitterObject ||
 			SpawnInfoObject->GetDisplayedStruct().IsValid() == false || SpawnInfoObject->GetDisplayedStruct()->GetStructMemory() != (uint8*)SpawnInfo)
 		{
 			bool bIsInTopLevelStruct = true;
 			bool bHideTopLevelCategories = true;
 			SpawnInfoObject = NewObject<UNiagaraStackObject>(this);
-			SpawnInfoObject->Initialize(CreateDefaultChildRequiredData(), StatelessEmitterWeak.Get(), SpawnInfoStructOnScope.ToSharedRef(), TEXT("SpawnInfo"), bIsInTopLevelStruct, bHideTopLevelCategories, GetStackEditorDataKey());
+			SpawnInfoObject->Initialize(CreateDefaultChildRequiredData(), StatelessEmitterObject, SpawnInfoStructOnScope.ToSharedRef(), TEXT("SpawnInfo"), bIsInTopLevelStruct, bHideTopLevelCategories, GetStackEditorDataKey());
 			SpawnInfoObject->SetOnFilterDetailNodes(
 				FNiagaraStackObjectShared::FOnFilterDetailNodes::CreateStatic(&UNiagaraStackStatelessEmitterSpawnItem::FilterDetailNodes),
 				UNiagaraStackObject::EDetailNodeFilterMode::FilterAllNodes);
+			//SpawnInfoObject->RegisterInstancedCustomPropertyTypeLayout(FNiagaraDistributionFloat::StaticStruct()->GetFName(), FOnGetPropertyTypeCustomizationInstance::CreateStatic(&FNiagaraDistributionPropertyCustomization::MakeFloatInstance, StatelessEmitterObject));
+			//SpawnInfoObject->RegisterInstancedCustomPropertyTypeLayout(FNiagaraDistributionVector2::StaticStruct()->GetFName(), FOnGetPropertyTypeCustomizationInstance::CreateStatic(&FNiagaraDistributionPropertyCustomization::MakeVector2Instance));
+			//SpawnInfoObject->RegisterInstancedCustomPropertyTypeLayout(FNiagaraDistributionVector3::StaticStruct()->GetFName(), FOnGetPropertyTypeCustomizationInstance::CreateStatic(&FNiagaraDistributionPropertyCustomization::MakeVector3Instance));
+			//SpawnInfoObject->RegisterInstancedCustomPropertyTypeLayout(FNiagaraDistributionColor::StaticStruct()->GetFName(), FOnGetPropertyTypeCustomizationInstance::CreateStatic(&FNiagaraDistributionPropertyCustomization::MakeColorInstance));
+			SpawnInfoObject->RegisterInstancedCustomPropertyTypeLayout(FNiagaraDistributionRangeFloat::StaticStruct()->GetFName(), FOnGetPropertyTypeCustomizationInstance::CreateStatic(&FNiagaraDistributionPropertyCustomization::MakeFloatInstance, StatelessEmitterObject));
+			//SpawnInfoObject->RegisterInstancedCustomPropertyTypeLayout(FNiagaraDistributionRangeVector2::StaticStruct()->GetFName(), FOnGetPropertyTypeCustomizationInstance::CreateStatic(&FNiagaraDistributionPropertyCustomization::MakeVector2Instance));
+			//SpawnInfoObject->RegisterInstancedCustomPropertyTypeLayout(FNiagaraDistributionRangeVector3::StaticStruct()->GetFName(), FOnGetPropertyTypeCustomizationInstance::CreateStatic(&FNiagaraDistributionPropertyCustomization::MakeVector3Instance));
+			SpawnInfoObject->RegisterInstancedCustomPropertyTypeLayout(FNiagaraDistributionRangeInt::StaticStruct()->GetFName(), FOnGetPropertyTypeCustomizationInstance::CreateStatic(&FNiagaraDistributionIntPropertyCustomization::MakeIntInstance, StatelessEmitterObject));
+
 			SpawnInfoObjectWeak = SpawnInfoObject;
 		}
 		NewChildren.Add(SpawnInfoObject);
@@ -181,7 +194,7 @@ void UNiagaraStackStatelessEmitterSpawnItem::RefreshChildrenInternal(const TArra
 		if (bGeneratedHeaderValueHandlers == false)
 		{
 			bGeneratedHeaderValueHandlers = true;
-			FNiagaraStackItemPropertyHeaderValueShared::GenerateHeaderValueHandlers(*StatelessEmitterWeak.Get(), (uint8*)SpawnInfo, *SpawnInfo->StaticStruct(), FSimpleDelegate::CreateUObject(this, &UNiagaraStackStatelessEmitterSpawnItem::OnHeaderValueChanged), HeaderValueHandlers);
+			FNiagaraStackItemPropertyHeaderValueShared::GenerateHeaderValueHandlers(*StatelessEmitterObject, (uint8*)SpawnInfo, *SpawnInfo->StaticStruct(), FSimpleDelegate::CreateUObject(this, &UNiagaraStackStatelessEmitterSpawnItem::OnHeaderValueChanged), HeaderValueHandlers);
 		}
 		else
 		{
@@ -240,4 +253,15 @@ void UNiagaraStackStatelessEmitterSpawnItem::OnHeaderValueChanged()
 	}
 }
 
+//-TODO:Stateless: There should be cleaner way of doing this
+void UNiagaraStackStatelessEmitterSpawnItem::OnSpawnInfoModified(TArray<UObject*> Objects, ENiagaraDataObjectChange ChangeType)
+{
+	UNiagaraStatelessEmitter* StatelessEmitter = StatelessEmitterWeak.Get();
+	if (Objects.Num() == 1 && Objects[0] == StatelessEmitter && StatelessEmitter)
+	{
+		FNiagaraStatelessSpawnInfo* SpawnInfo = GetSpawnInfo();
+		SpawnInfo->Rate.UpdateValuesFromDistribution();
+	}
+}
+
 #undef LOCTEXT_NAMESPACE
diff --git a/Engine/Plugins/FX/Niagara/Source/NiagaraEditor/Public/ViewModels/Stack/NiagaraStackStatelessEmitterSpawnGroup.h b/Engine/Plugins/FX/Niagara/Source/NiagaraEditor/Public/ViewModels/Stack/NiagaraStackStatelessEmitterSpawnGroup.h
--- a/Engine/Plugins/FX/Niagara/Source/NiagaraEditor/Public/ViewModels/Stack/NiagaraStackStatelessEmitterSpawnGroup.h
+++ b/Engine/Plugins/FX/Niagara/Source/NiagaraEditor/Public/ViewModels/Stack/NiagaraStackStatelessEmitterSpawnGroup.h
@@ -80,6 +80,7 @@ private:
 	static void FilterDetailNodes(const TArray<TSharedRef<IDetailTreeNode>>& InSourceNodes, TArray<TSharedRef<IDetailTreeNode>>& OutFilteredNodes);
 
 	void OnHeaderValueChanged();
+	void OnSpawnInfoModified(TArray<UObject*> Objects, ENiagaraDataObjectChange ChangeType);
 
 private:
 	TWeakObjectPtr<UNiagaraStatelessEmitter> StatelessEmitterWeak;
diff --git a/Engine/Plugins/FX/Niagara/Source/NiagaraShader/Internal/Stateless/NiagaraStatelessModuleShaderParameters.h b/Engine/Plugins/FX/Niagara/Source/NiagaraShader/Internal/Stateless/NiagaraStatelessModuleShaderParameters.h
--- a/Engine/Plugins/FX/Niagara/Source/NiagaraShader/Internal/Stateless/NiagaraStatelessModuleShaderParameters.h
+++ b/Engine/Plugins/FX/Niagara/Source/NiagaraShader/Internal/Stateless/NiagaraStatelessModuleShaderParameters.h
@@ -53,9 +53,8 @@ namespace NiagaraStateless
 	END_SHADER_PARAMETER_STRUCT()
 
 	BEGIN_SHADER_PARAMETER_STRUCT(FScaleMeshSizeBySpeedModule_ShaderParameters, NIAGARASHADER_API)
-		SHADER_PARAMETER(int32,		ScaleMeshSizeBySpeed_Offset)
-		SHADER_PARAMETER(float,		ScaleMeshSizeBySpeed_VelocityNorm)
-		SHADER_PARAMETER(int32,		ScaleMeshSizeBySpeed_Length)
+		SHADER_PARAMETER(FUintVector2,	ScaleMeshSizeBySpeed_ScaleDistribution)
+		SHADER_PARAMETER(float,			ScaleMeshSizeBySpeed_VelocityNorm)
 	END_SHADER_PARAMETER_STRUCT()
 
 	BEGIN_SHADER_PARAMETER_STRUCT(FScaleSpriteSizeModule_ShaderParameters, NIAGARASHADER_API)
@@ -63,9 +62,8 @@ namespace NiagaraStateless
 	END_SHADER_PARAMETER_STRUCT()
 
 	BEGIN_SHADER_PARAMETER_STRUCT(FScaleSpriteSizeBySpeedModule_ShaderParameters, NIAGARASHADER_API)
-		SHADER_PARAMETER(float,		ScaleSpriteSizeBySpeed_VelocityNorm)
-		SHADER_PARAMETER(int32,		ScaleSpriteSizeBySpeed_Offset)
-		SHADER_PARAMETER(int32,		ScaleSpriteSizeBySpeed_Length)
+		SHADER_PARAMETER(FUintVector2,	ScaleSpriteSizeBySpeed_ScaleDistribution)
+		SHADER_PARAMETER(float,			ScaleSpriteSizeBySpeed_VelocityNorm)
 	END_SHADER_PARAMETER_STRUCT()
 
 	BEGIN_SHADER_PARAMETER_STRUCT(FSpriteFacingAndAlignmentModule_ShaderParameters, NIAGARASHADER_API)
