* Fixed missing tooltip for set attributes node when linked to a parameter definition asset
* Fixed ui bug in param definition editor

diff --git a/Engine/Plugins/FX/Niagara/Source/NiagaraEditor/Private/Customizations/NiagaraScriptVariableCustomization.cpp b/Engine/Plugins/FX/Niagara/Source/NiagaraEditor/Private/Customizations/NiagaraScriptVariableCustomization.cpp
--- a/Engine/Plugins/FX/Niagara/Source/NiagaraEditor/Private/Customizations/NiagaraScriptVariableCustomization.cpp
+++ b/Engine/Plugins/FX/Niagara/Source/NiagaraEditor/Private/Customizations/NiagaraScriptVariableCustomization.cpp
@@ -244,7 +244,6 @@ void FNiagaraScriptVariableDetails::CustomizeDetailsParameterDefinitionsSynchron
 		[
 			SNew(SEnumComboBox, LibrarySourceDefaultModeEnum)
 			.CurrentValue(this, &FNiagaraScriptVariableDetails::GetLibraryDefaultModeValue)
-			.ButtonStyle(FAppStyle::Get(), "FlatButton.Light")
 			.ContentPadding(FMargin(2, 0))
 			.Font(FAppStyle::GetFontStyle("Sequencer.AnimationOutliner.RegularFont"))
 			.OnEnumSelectionChanged(SEnumComboBox::FOnEnumSelectionChanged::CreateSP(this, &FNiagaraScriptVariableDetails::OnLibrarySourceDefaultModeChanged))
@@ -264,7 +263,6 @@ void FNiagaraScriptVariableDetails::CustomizeDetailsParameterDefinitionsSynchron
 		[
 			SNew(SEnumComboBox, LibrarySynchronizedDefaultModeEnum)
 			.CurrentValue(this, &FNiagaraScriptVariableDetails::GetLibraryDefaultModeValue)
-			.ButtonStyle(FAppStyle::Get(), "FlatButton.Light")
 			.ContentPadding(FMargin(2, 0))
 			.Font(FAppStyle::GetFontStyle("Sequencer.AnimationOutliner.RegularFont"))
 			.OnEnumSelectionChanged(SEnumComboBox::FOnEnumSelectionChanged::CreateSP(this, &FNiagaraScriptVariableDetails::OnLibrarySynchronizedDefaultModeChanged))
diff --git a/Engine/Plugins/FX/Niagara/Source/NiagaraEditor/Private/ViewModels/Stack/NiagaraStackFunctionInput.cpp b/Engine/Plugins/FX/Niagara/Source/NiagaraEditor/Private/ViewModels/Stack/NiagaraStackFunctionInput.cpp
--- a/Engine/Plugins/FX/Niagara/Source/NiagaraEditor/Private/ViewModels/Stack/NiagaraStackFunctionInput.cpp
+++ b/Engine/Plugins/FX/Niagara/Source/NiagaraEditor/Private/ViewModels/Stack/NiagaraStackFunctionInput.cpp
@@ -24,6 +24,7 @@
 #include "NiagaraNodeOutput.h"
 #include "NiagaraNodeParameterMapGet.h"
 #include "NiagaraNodeParameterMapSet.h"
+#include "NiagaraParameterDefinitions.h"
 #include "NiagaraRendererProperties.h"
 #include "NiagaraScript.h"
 #include "ViewModels/NiagaraScriptGraphViewModel.h"
@@ -1301,6 +1302,10 @@ void UNiagaraStackFunctionInput::RefreshFromMetaData(TArray<FStackIssue>& NewIss
 				InputMetaData = *FoundMetaData;
 			}
 		}
+		else if (UNiagaraScriptVariable* ScriptVariable = GetSystemViewModel()->FindSubscribedParameterDefinitionsScriptVarByName(InputParameterHandle.GetName()))
+		{
+			InputMetaData = ScriptVariable->Metadata;
+		}
 	}
 	else if (OwningFunctionCallNode->FunctionScript != nullptr)
 	{
