Fix for verbose debug rendering of the rigid body DI

diff --git a/Engine/Plugins/FX/Niagara/Source/Niagara/Private/NiagaraDataInterfaceRigidMeshCollisionQuery.cpp b/Engine/Plugins/FX/Niagara/Source/Niagara/Private/NiagaraDataInterfaceRigidMeshCollisionQuery.cpp
--- a/Engine/Plugins/FX/Niagara/Source/Niagara/Private/NiagaraDataInterfaceRigidMeshCollisionQuery.cpp
+++ b/Engine/Plugins/FX/Niagara/Source/Niagara/Private/NiagaraDataInterfaceRigidMeshCollisionQuery.cpp
@@ -1055,26 +1055,14 @@ void UNiagaraDataInterfaceRigidMeshCollisionQuery::DrawDebugHud(FNDIDrawDebugHud
 
 	if (DebugHudContext.IsVerbose())
 	{
-		UCanvas* Canvas = DebugHudContext.GetCanvas();
-
-		// the DrawDebugCanvas* functions don't reasoanbly handle the near clip plane (both in terms of clipping and in terms of
-		// objects being behind the camera); so we introduce this culling behavior to work around it
-		auto ShouldClip = [&](UCanvas* Canvas, const FMatrix& Transform, const FBoxSphereBounds& Bounds)
-		{
-			const FVector Origin = Transform.TransformPosition(Bounds.Origin);
-			return (Canvas->Project(Origin).GetMin() < UE_KINDA_SMALL_NUMBER);
-		};
+		const UWorld* World = DebugHudContext.GetWorld();
 
 		// Boxes
 		for (uint32 BoxIt = 0; BoxIt < BoxCount; ++BoxIt)
 		{
 			const FVector3f HalfBoxExtent = 0.5f * InstanceData_GT->AssetArrays->ElementExtent[ElementOffsets.BoxOffset + BoxIt];
-			const FBox Box(-HalfBoxExtent, HalfBoxExtent);
 			const FMatrix CurrentTransform = GetCurrentTransform(ElementOffsets.BoxOffset + BoxIt);
-			if (!ShouldClip(Canvas, CurrentTransform, FSphere(FVector::ZeroVector, HalfBoxExtent.Size())))
-			{
-				DrawDebugCanvasWireBox(Canvas, CurrentTransform, Box, FColor::Blue);
-			}
+			DrawDebugBox(World, CurrentTransform.TransformPosition(FVector::ZeroVector), FVector(HalfBoxExtent), CurrentTransform.ToQuat(), FColor::Blue);
 		}
 
 		// Spheres
@@ -1082,10 +1070,7 @@ void UNiagaraDataInterfaceRigidMeshCollisionQuery::DrawDebugHud(FNDIDrawDebugHud
 		{
 			const float Radius = InstanceData_GT->AssetArrays->ElementExtent[ElementOffsets.SphereOffset + SphereIt].X;
 			const FMatrix CurrentTransform = GetCurrentTransform(ElementOffsets.SphereOffset + SphereIt);
-			if (!ShouldClip(Canvas, CurrentTransform, FSphere(FVector::ZeroVector, Radius)))
-			{
-				DrawDebugCanvasWireSphere(Canvas, CurrentTransform.TransformPosition(FVector::ZeroVector), FColor::Blue, Radius, 20);
-			}
+			DrawDebugSphere(World, CurrentTransform.TransformPosition(FVector::ZeroVector), Radius, 20, FColor::Blue);
 		}
 
 		// Capsules
@@ -1093,15 +1078,21 @@ void UNiagaraDataInterfaceRigidMeshCollisionQuery::DrawDebugHud(FNDIDrawDebugHud
 		{
 			const FVector2f RadiusLength(InstanceData_GT->AssetArrays->ElementExtent[ElementOffsets.CapsuleOffset + CapsuleIt]);
 			const FMatrix CurrentTransform = GetCurrentTransform(ElementOffsets.CapsuleOffset + CapsuleIt);
-			const float HalfTotalLength = RadiusLength.X + 0.5f * RadiusLength.Y;
-			if (!ShouldClip(Canvas, CurrentTransform, FSphere(FVector::ZeroVector, HalfTotalLength)))
-			{
-				DrawDebugCanvasCapsule(Canvas, CurrentTransform, HalfTotalLength, RadiusLength.X, FColor::Blue);
-			}
+			DrawDebugCapsule(World, CurrentTransform.TransformPosition(FVector::ZeroVector), RadiusLength.Y * 0.5f, RadiusLength.X, CurrentTransform.ToQuat(), FColor::Blue);
 		}
 
 		if (!InstanceData_GT->ExplicitActors.IsEmpty() || !InstanceData_GT->FoundActors.IsEmpty())
 		{
+			UCanvas* Canvas = DebugHudContext.GetCanvas();
+
+			// the DrawDebugCanvas* functions don't reasoanbly handle the near clip plane (both in terms of clipping and in terms of
+			// objects being behind the camera); so we introduce this culling behavior to work around it
+			auto ShouldClip = [Canvas](const FMatrix& Transform, const FBoxSphereBounds& Bounds)
+			{
+				const FVector Origin = Transform.TransformPosition(Bounds.Origin);
+				return (Canvas->Project(Origin).GetMin() < UE_KINDA_SMALL_NUMBER);
+			};
+
 			const UFont* Font = GEngine->GetMediumFont();
 			Canvas->SetDrawColor(FColor::White);
 
@@ -1114,10 +1105,9 @@ void UNiagaraDataInterfaceRigidMeshCollisionQuery::DrawDebugHud(FNDIDrawDebugHud
 					Actor->GetActorBounds(true, ActorOrigin, ActorBoundsExtent);
 
 					const FMatrix CurrentTransform = FTranslationMatrix(ActorOrigin);
-					if (!ShouldClip(Canvas, CurrentTransform, FSphere(FVector::ZeroVector, ActorBoundsExtent.Size())))
+					DrawDebugBox(World, CurrentTransform.TransformPosition(FVector::ZeroVector), ActorBoundsExtent, CurrentTransform.ToQuat(), FColor::Yellow);
+					if (!ShouldClip(CurrentTransform, FSphere(FVector::ZeroVector, ActorBoundsExtent.Size())))
 					{
-						DrawDebugCanvasWireBox(Canvas, CurrentTransform, FBox(-ActorBoundsExtent, ActorBoundsExtent), FColor::Yellow);
-
 						FString ActorLabel;
 #if WITH_EDITOR
 						ActorLabel = Actor->GetActorLabel();
