Disable IsWriteAllowedForNamespace till we can collect the list of assets causing trouble
#rnx

diff --git a/Engine/Plugins/FX/Niagara/Source/NiagaraEditor/Private/NiagaraHlslTranslator.cpp b/Engine/Plugins/FX/Niagara/Source/NiagaraEditor/Private/NiagaraHlslTranslator.cpp
--- a/Engine/Plugins/FX/Niagara/Source/NiagaraEditor/Private/NiagaraHlslTranslator.cpp
+++ b/Engine/Plugins/FX/Niagara/Source/NiagaraEditor/Private/NiagaraHlslTranslator.cpp
@@ -5682,12 +5682,12 @@ void TNiagaraHlslTranslator<GraphBridge>::ParameterMapSet(const FParamMapSetNode
 				continue;
 			}
 
-			FText WriteErrorText;
-			if (!IsWriteAllowedForNamespace(Var, ActiveHistoryForFunctionCalls.GetCurrentUsageContext(), WriteErrorText))
-			{
-				Error(WriteErrorText, SetNode, Inputs[i].Pin);
-				continue;
-			}
+			//FText WriteErrorText;
+			//if (!IsWriteAllowedForNamespace(Var, ActiveHistoryForFunctionCalls.GetCurrentUsageContext(), WriteErrorText))
+			//{
+			//	Error(WriteErrorText, SetNode, Inputs[i].Pin);
+			//	continue;
+			//}
 
 			Var = ActiveHistoryForFunctionCalls.ResolveAliases(Var);
 			
