Validation fix for iteration source DI writes

diff --git a/Engine/Plugins/FX/Niagara/Source/NiagaraEditor/Private/NiagaraHlslTranslator.cpp b/Engine/Plugins/FX/Niagara/Source/NiagaraEditor/Private/NiagaraHlslTranslator.cpp
--- a/Engine/Plugins/FX/Niagara/Source/NiagaraEditor/Private/NiagaraHlslTranslator.cpp
+++ b/Engine/Plugins/FX/Niagara/Source/NiagaraEditor/Private/NiagaraHlslTranslator.cpp
@@ -5306,6 +5306,12 @@ bool FNiagaraHlslTranslator::IsWriteAllowedForNamespace(const FNiagaraVariable&
 	}
 	if (UNiagaraScript::IsParticleScript(TargetUsage) && (Var.IsInNameSpace(FNiagaraConstants::SystemNamespace) || Var.IsInNameSpace(FNiagaraConstants::EmitterNamespace)))
 	{
+		const FNiagaraHlslTranslationStage& ActiveStage = TranslationStages[ActiveStageIdx];
+		if (ActiveStage.IterationSourceType == ENiagaraIterationSource::DataInterface && Var.IsInNameSpace(ActiveStage.IterationDataInterface))
+		{
+			// weird special case where particle script can write to emitter data interface sub properties
+			return true;
+		}
 		ErrorMsg = FText::Format(LOCTEXT("WriteAllowedForNamespaceFail_Particles", "Cannot set variable {0} in particle scripts."), FText::FromName(Var.GetName()));
 		return false;
 	}
diff --git a/Engine/Plugins/FX/Niagara/Source/NiagaraEditor/Private/NiagaraHlslTranslator.h b/Engine/Plugins/FX/Niagara/Source/NiagaraEditor/Private/NiagaraHlslTranslator.h
--- a/Engine/Plugins/FX/Niagara/Source/NiagaraEditor/Private/NiagaraHlslTranslator.h
+++ b/Engine/Plugins/FX/Niagara/Source/NiagaraEditor/Private/NiagaraHlslTranslator.h
@@ -471,7 +471,7 @@ protected:
 	bool IsBulkSystemScript() const;
 	bool IsSpawnScript() const;
 	bool RequiresInterpolation() const;
-	static bool IsWriteAllowedForNamespace(const FNiagaraVariable& Var, ENiagaraScriptUsage TargetUsage, FText& ErrorMsg);
+	bool IsWriteAllowedForNamespace(const FNiagaraVariable& Var, ENiagaraScriptUsage TargetUsage, FText& ErrorMsg);
 
 	template<typename T>
 	void BuildConstantBuffer(ENiagaraCodeChunkMode ChunkMode);
