Another attempt at trying to make the default compilation mode resiliant to changes being made (deleting emitters) during compilation

diff --git a/Engine/Plugins/FX/Niagara/Source/Niagara/Private/NiagaraAsyncCompile.cpp b/Engine/Plugins/FX/Niagara/Source/Niagara/Private/NiagaraAsyncCompile.cpp
--- a/Engine/Plugins/FX/Niagara/Source/Niagara/Private/NiagaraAsyncCompile.cpp
+++ b/Engine/Plugins/FX/Niagara/Source/Niagara/Private/NiagaraAsyncCompile.cpp
@@ -112,6 +112,39 @@ TSharedPtr<FNiagaraCompileRequestDuplicateDataBase, ESPMode::ThreadSafe> FNiagar
 	return PrecompileDuplicateData;
 }
 
+bool FNiagaraLazyPrecompileReference::IsValidForPrecompile() const
+{
+	if (System)
+	{
+		const TArray<FNiagaraEmitterHandle>& EmitterHandles = System->GetEmitterHandles();
+		for (TMap<UNiagaraScript*, int32>::TConstIterator It = EmitterScriptIndex.CreateConstIterator(); It; ++It)
+		{
+			if (EmitterHandles.IsValidIndex(It.Value()))
+			{
+				const FNiagaraEmitterHandle& EmitterHandle = EmitterHandles[It.Value()];
+				const FVersionedNiagaraEmitterData* EmitterData = EmitterHandle.GetEmitterData();
+				if (!EmitterHandle.GetIsEnabled() || !EmitterData)
+				{
+					return false;
+				}
+				TArray<UNiagaraScript*> EmitterScripts;
+				EmitterData->GetScripts(EmitterScripts, false, true);
+
+				if (!EmitterScripts.Contains(It.Key()))
+				{
+					return false;
+				}
+			}
+			else
+			{
+				return false;
+			}
+		}
+	}
+
+	return true;
+}
+
 FNiagaraAsyncCompileTask::FNiagaraAsyncCompileTask(UNiagaraSystem* InOwningSystem, FString InAssetPath, const FEmitterCompiledScriptPair& InScriptPair)
 {
 	OwningSystem = InOwningSystem;
@@ -278,6 +311,14 @@ void FNiagaraAsyncCompileTask::PrecompileData()
 		return;
 	}
 
+	// destructive changes (like emitters being deleted) can disrupt the precompile so we need to validate that we can continue.
+	if (!PrecompileReference->IsValidForPrecompile())
+	{
+		UE_LOG(LogNiagara, Warning, TEXT("Unable to proceed with compilation due to changes to asset %s.  Aborting."), *AssetPath);
+		AbortTask();
+		return;
+	}
+
 	ComputedPrecompileData = PrecompileReference->GetPrecompileData(ScriptPair.CompiledScript);
 	TSharedPtr<FNiagaraCompileRequestDuplicateDataBase, ESPMode::ThreadSafe> SystemPrecompileDuplicateData = PrecompileReference->GetPrecompileDuplicateData(ScriptPair.VersionedEmitter.Emitter, ScriptPair.CompiledScript);
 
@@ -778,14 +819,33 @@ bool FNiagaraActiveCompilationDefault::QueryCompileComplete(const FNiagaraQueryC
 
 void FNiagaraActiveCompilationDefault::Apply(const FNiagaraQueryCompilationOptions& Options)
 {
+	auto ShouldApplyTask = [](const FAsyncTaskPtr& AsyncTask) -> bool
+	{
+		if (AsyncTask->CurrentState == ENiagaraCompilationState::Aborted || !AsyncTask->ScriptPair.bResultsReady)
+		{
+			return false;
+		}
+
+		// for now also check some common failure points (the ComputedPrecompileDuplicateData being cleaned up/incomplete)
+		if (AsyncTask->ComputedPrecompileDuplicateData.IsValid())
+		{
+			if (AsyncTask->ComputedPrecompileDuplicateData->GetScriptSource() == nullptr)
+			{
+				return false;
+			}
+		}
+
+		return true;
+	};
+
 	if (bEvaluateParametersPending)
 	{
 		// run a first pass to apply the rapid iteration parameters across all of the tasks with results
 		for (FAsyncTaskPtr& AsyncTask : Tasks)
 		{
-			FEmitterCompiledScriptPair& EmitterCompiledScriptPair = AsyncTask->ScriptPair;
-			if (EmitterCompiledScriptPair.bResultsReady)
+			if (ShouldApplyTask(AsyncTask))
 			{
+				FEmitterCompiledScriptPair& EmitterCompiledScriptPair = AsyncTask->ScriptPair;
 				if (TSharedPtr<FNiagaraVMExecutableData> ExeData = EmitterCompiledScriptPair.CompileResults)
 				{
 					if (UNiagaraScript* CompiledScript = EmitterCompiledScriptPair.CompiledScript)
@@ -799,9 +859,9 @@ void FNiagaraActiveCompilationDefault::Apply(const FNiagaraQueryCompilationOptio
 
 	for (FAsyncTaskPtr& AsyncTask : Tasks)
 	{
-		FEmitterCompiledScriptPair& EmitterCompiledScriptPair = AsyncTask->ScriptPair;
-		if (EmitterCompiledScriptPair.bResultsReady)
+		if (ShouldApplyTask(AsyncTask))
 		{
+			FEmitterCompiledScriptPair& EmitterCompiledScriptPair = AsyncTask->ScriptPair;
 			TSharedPtr<FNiagaraVMExecutableData> ExeData = EmitterCompiledScriptPair.CompileResults;
 			UNiagaraScript* CompiledScript = EmitterCompiledScriptPair.CompiledScript;
 
diff --git a/Engine/Plugins/FX/Niagara/Source/Niagara/Private/NiagaraAsyncCompile.h b/Engine/Plugins/FX/Niagara/Source/Niagara/Private/NiagaraAsyncCompile.h
--- a/Engine/Plugins/FX/Niagara/Source/Niagara/Private/NiagaraAsyncCompile.h
+++ b/Engine/Plugins/FX/Niagara/Source/Niagara/Private/NiagaraAsyncCompile.h
@@ -42,6 +42,8 @@ struct FNiagaraLazyPrecompileReference
 {
 	TSharedPtr<FNiagaraCompileRequestDataBase, ESPMode::ThreadSafe> GetPrecompileData(UNiagaraScript* ForScript);
 	TSharedPtr<FNiagaraCompileRequestDuplicateDataBase, ESPMode::ThreadSafe> GetPrecompileDuplicateData(UNiagaraEmitter* OwningEmitter, UNiagaraScript* TargetScript);
+
+	bool IsValidForPrecompile() const;
 	
 	UNiagaraSystem* System = nullptr;
 	TArray<UNiagaraScript*> Scripts;
