- Add platform set to sorting rules

diff --git a/Engine/Plugins/FX/Niagara/Source/NiagaraEditor/Private/NiagaraValidationRules.cpp b/Engine/Plugins/FX/Niagara/Source/NiagaraEditor/Private/NiagaraValidationRules.cpp
--- a/Engine/Plugins/FX/Niagara/Source/NiagaraEditor/Private/NiagaraValidationRules.cpp
+++ b/Engine/Plugins/FX/Niagara/Source/NiagaraEditor/Private/NiagaraValidationRules.cpp
@@ -674,6 +674,13 @@ void UNiagaraValidationRule_RendererSortingEnabled::CheckValidity(const FNiagara
 	for (TSharedRef<FNiagaraEmitterHandleViewModel> EmitterHandleModel : EmitterHandleViewModels)
 	{
 		FVersionedNiagaraEmitterData* EmitterData = EmitterHandleModel.Get().GetEmitterHandle()->GetEmitterData();
+
+		const FString PlatformConflictsString = NiagaraValidation::GetPlatformConflictsString(Platforms, EmitterData->Platforms);
+		if (PlatformConflictsString.IsEmpty())
+		{
+			continue;
+		}
+
 		EmitterData->ForEachRenderer(
 			[&Results, EmitterHandleModel, this, &System](UNiagaraRendererProperties* RendererProperties)
 			{
diff --git a/Engine/Plugins/FX/Niagara/Source/NiagaraEditor/Public/NiagaraValidationRules.h b/Engine/Plugins/FX/Niagara/Source/NiagaraEditor/Public/NiagaraValidationRules.h
--- a/Engine/Plugins/FX/Niagara/Source/NiagaraEditor/Public/NiagaraValidationRules.h
+++ b/Engine/Plugins/FX/Niagara/Source/NiagaraEditor/Public/NiagaraValidationRules.h
@@ -106,6 +106,9 @@ public:
 	UPROPERTY(EditAnywhere, Category = Validation)
 	ENiagaraValidationSeverity Severity = ENiagaraValidationSeverity::Warning;
 
+	UPROPERTY(EditAnywhere, Category = Validation)
+	FNiagaraPlatformSet Platforms;
+
 	virtual void CheckValidity(const FNiagaraValidationContext& Context, TArray<FNiagaraValidationResult>& OutResults) const override;
 };
 
