Make severity of invalid namespace writes in niagara scripts a plugin setting

diff --git a/Engine/Plugins/FX/Niagara/Source/Niagara/Private/NiagaraScript.cpp b/Engine/Plugins/FX/Niagara/Source/Niagara/Private/NiagaraScript.cpp
--- a/Engine/Plugins/FX/Niagara/Source/Niagara/Private/NiagaraScript.cpp
+++ b/Engine/Plugins/FX/Niagara/Source/Niagara/Private/NiagaraScript.cpp
@@ -1097,7 +1097,10 @@ void UNiagaraScript::ComputeVMCompilationId(FNiagaraVMExecutableDataId& Id, cons
 	Id.bInterpolatedSpawn = false;
 	Id.bRequiresPersistentIDs = false;
 	Id.ScriptVersionID = IsVersioningEnabled() ? (VersionGuid.IsValid() ? VersionGuid : ExposedVersion) : FGuid();
-	if (GetDefault<UNiagaraSettings>()->bAccurateQuatInterpolation)
+	
+	const UNiagaraSettings* NiagaraSettings = GetDefault<UNiagaraSettings>();
+	Id.AdditionalDefines.Add(TEXT("InvalidNamespaceWrite=") + FString::FormatAsNumber(static_cast<int32>(NiagaraSettings->InvalidNamespaceWriteSeverity)));
+	if (NiagaraSettings->bAccurateQuatInterpolation)
 	{
 		Id.AdditionalDefines.Add(FNiagaraCompileOptions::AccurateQuatInterpolation);
 	}
@@ -1132,7 +1135,7 @@ void UNiagaraScript::ComputeVMCompilationId(FNiagaraVMExecutableDataId& Id, cons
 				Id.AdditionalDefines.Add(TEXT("IgnoreParticleReadsForAttributeTrim"));
 			}
 
-			if (EmitterOwner->ShouldDisableExperimentalVM() || !GetDefault<UNiagaraSettings>()->bExperimentalVMEnabled)
+			if (EmitterOwner->ShouldDisableExperimentalVM() || !NiagaraSettings->bExperimentalVMEnabled)
 			{
 				Id.AdditionalDefines.Add(FNiagaraCompileOptions::ExperimentalVMDisabled);
 			}
@@ -1344,7 +1347,7 @@ void UNiagaraScript::ComputeVMCompilationId(FNiagaraVMExecutableDataId& Id, cons
 			Id.AdditionalDefines.Add(TEXT("CompressAttributes"));
 		}
 
-		if (System->ShouldDisableExperimentalVM() || !GetDefault<UNiagaraSettings>()->bExperimentalVMEnabled)
+		if (System->ShouldDisableExperimentalVM() || !NiagaraSettings->bExperimentalVMEnabled)
 		{
 			Id.AdditionalDefines.Add(FNiagaraCompileOptions::ExperimentalVMDisabled);
 		}
diff --git a/Engine/Plugins/FX/Niagara/Source/Niagara/Public/NiagaraSettings.h b/Engine/Plugins/FX/Niagara/Source/Niagara/Public/NiagaraSettings.h
--- a/Engine/Plugins/FX/Niagara/Source/Niagara/Public/NiagaraSettings.h
+++ b/Engine/Plugins/FX/Niagara/Source/Niagara/Public/NiagaraSettings.h
@@ -81,6 +81,12 @@ enum class ENiagaraDefaultGpuTranslucentLatency : uint8
 	Latent,
 };
 
+UENUM()
+enum class ENiagaraCompileErrorSeverity : uint8
+{
+	Ignore, LogOnly, Warning, Error
+};
+
 UENUM()
 namespace ENDICollisionQuery_AsyncGpuTraceProvider
 {
@@ -172,6 +178,10 @@ class UNiagaraSettings : public UDeveloperSettings
 	UPROPERTY(config, EditAnywhere, AdvancedDisplay, Category = Niagara)
 	bool bAccurateQuatInterpolation = true;
 
+	/** If the Niagara compiler sees that a script writes to a namespace that is read only (e.g. a particle script writing to a system attribute), what should it do. */
+	UPROPERTY(config, EditAnywhere, AdvancedDisplay, Category = Niagara)
+	ENiagaraCompileErrorSeverity InvalidNamespaceWriteSeverity = ENiagaraCompileErrorSeverity::Warning;
+
 	/** Whether to limit the max tick delta time or not. */
 	UPROPERTY(config, EditAnywhere, BlueprintReadWrite, AdvancedDisplay, Category = "Niagara", meta = (InlineEditConditionToggle))
 	bool bLimitDeltaTime = true;
diff --git a/Engine/Plugins/FX/Niagara/Source/NiagaraEditor/Private/NiagaraHlslTranslator.cpp b/Engine/Plugins/FX/Niagara/Source/NiagaraEditor/Private/NiagaraHlslTranslator.cpp
--- a/Engine/Plugins/FX/Niagara/Source/NiagaraEditor/Private/NiagaraHlslTranslator.cpp
+++ b/Engine/Plugins/FX/Niagara/Source/NiagaraEditor/Private/NiagaraHlslTranslator.cpp
@@ -5688,12 +5688,24 @@ void TNiagaraHlslTranslator<GraphBridge>::ParameterMapSet(const FParamMapSetNode
 				continue;
 			}
 
-			//FText WriteErrorText;
-			//if (!IsWriteAllowedForNamespace(Var, ActiveHistoryForFunctionCalls.GetCurrentUsageContext(), WriteErrorText))
-			//{
-			//	Error(WriteErrorText, SetNode, Inputs[i].Pin);
-			//	continue;
-			//}
+			FText WriteErrorText;
+			if (!IsWriteAllowedForNamespace(Var, ActiveHistoryForFunctionCalls.GetCurrentUsageContext(), WriteErrorText))
+			{
+				const UNiagaraSettings* Settings = GetDefault<UNiagaraSettings>();
+				if (Settings->InvalidNamespaceWriteSeverity == ENiagaraCompileErrorSeverity::Error)
+				{
+					Error(WriteErrorText, SetNode, Inputs[i].Pin);
+				}
+				else if (Settings->InvalidNamespaceWriteSeverity == ENiagaraCompileErrorSeverity::Warning)
+				{
+					Warning(WriteErrorText, SetNode, Inputs[i].Pin);
+				}
+				else if (Settings->InvalidNamespaceWriteSeverity == ENiagaraCompileErrorSeverity::LogOnly)
+				{
+					Message(FNiagaraCompileEventSeverity::Log, WriteErrorText, SetNode, Inputs[i].Pin);
+				}
+				continue;
+			}
 
 			Var = ActiveHistoryForFunctionCalls.ResolveAliases(Var);
 			
