- Stateless Dynamic Parameter Module (no channel masking support yet)
- Fix spelling for GetOutputVariables

diff --git a/Engine/Plugins/FX/Niagara/Shaders/Private/Stateless/Modules/NiagaraStatelessModule_DynamicMaterialParameters.ush b/Engine/Plugins/FX/Niagara/Shaders/Private/Stateless/Modules/NiagaraStatelessModule_DynamicMaterialParameters.ush
new file mode 100644
--- /dev/null
+++ b/Engine/Plugins/FX/Niagara/Shaders/Private/Stateless/Modules/NiagaraStatelessModule_DynamicMaterialParameters.ush
@@ -0,0 +1,40 @@
+// Copyright Epic Games, Inc. All Rights Reserved.
+
+#include "../NiagaraStatelessCommon.ush"
+
+uint	DynamicMaterialParameters_ChannelMask;
+uint3	DynamicMaterialParameters_Parameter0X;
+uint3	DynamicMaterialParameters_Parameter0Y;
+uint3	DynamicMaterialParameters_Parameter0Z;
+uint3	DynamicMaterialParameters_Parameter0W;
+
+struct FStatelessModule_DynamicMaterialParameters
+{
+	FStatelessDistributionFloat	Parameter0XDistribution;
+	FStatelessDistributionFloat	Parameter0YDistribution;
+	FStatelessDistributionFloat	Parameter0ZDistribution;
+	FStatelessDistributionFloat	Parameter0WDistribution;
+};
+
+void DynamicMaterialParameters_Initialize(inout FStatelessModule_DynamicMaterialParameters ModuleData, inout FStatelessParticle Particle)
+{
+	if ((DynamicMaterialParameters_ChannelMask & 0xf) != 0)
+	{
+		ModuleData.Parameter0XDistribution.Init(Particle, DynamicMaterialParameters_Parameter0X);
+		ModuleData.Parameter0YDistribution.Init(Particle, DynamicMaterialParameters_Parameter0Y);
+		ModuleData.Parameter0ZDistribution.Init(Particle, DynamicMaterialParameters_Parameter0Z);
+		ModuleData.Parameter0WDistribution.Init(Particle, DynamicMaterialParameters_Parameter0W);
+	}
+}
+	
+void DynamicMaterialParameters_Simulate(in FStatelessModule_DynamicMaterialParameters ModuleData, inout FStatelessParticle Particle)
+{
+	Particle.DynamicMaterialParameter0 = 0.0f;
+	if ((DynamicMaterialParameters_ChannelMask & 0xf) != 0)
+	{
+		Particle.DynamicMaterialParameter0.x = (DynamicMaterialParameters_ChannelMask & 0x1) != 0 ? ModuleData.Parameter0XDistribution.GetValue(Particle.NormalizedAge) : 0.0f;
+		Particle.DynamicMaterialParameter0.y = (DynamicMaterialParameters_ChannelMask & 0x2) != 0 ? ModuleData.Parameter0YDistribution.GetValue(Particle.NormalizedAge) : 0.0f;
+		Particle.DynamicMaterialParameter0.z = (DynamicMaterialParameters_ChannelMask & 0x4) != 0 ? ModuleData.Parameter0ZDistribution.GetValue(Particle.NormalizedAge) : 0.0f;
+		Particle.DynamicMaterialParameter0.w = (DynamicMaterialParameters_ChannelMask & 0x8) != 0 ? ModuleData.Parameter0WDistribution.GetValue(Particle.NormalizedAge) : 0.0f;
+	}
+}
diff --git a/Engine/Plugins/FX/Niagara/Shaders/Private/Stateless/Modules/NiagaraStatelessModule_SolveVelocitiesAndForces.ush b/Engine/Plugins/FX/Niagara/Shaders/Private/Stateless/Modules/NiagaraStatelessModule_SolveVelocitiesAndForces.ush
--- a/Engine/Plugins/FX/Niagara/Shaders/Private/Stateless/Modules/NiagaraStatelessModule_SolveVelocitiesAndForces.ush
+++ b/Engine/Plugins/FX/Niagara/Shaders/Private/Stateless/Modules/NiagaraStatelessModule_SolveVelocitiesAndForces.ush
@@ -133,7 +133,7 @@ void SolveVelocitiesAndForces_Initialize(inout FStatelessModule_SolveVelocitiesA
 		else if ( SolveVelocitiesAndForces_NoiseMode == 3 )
 		{
 			const uint3 ChannelIndex = RandomUInt3(Particle) % SolveVelocitiesAndForces_NoiseLUTNumChannel;
-			const float ChannelLerp = (Particle.Age * SolveVelocitiesAndForces_NoiseFrequency) + RandomFloat(Particle);
+			const float ChannelLerp = (Particle.Age * SolveVelocitiesAndForces_NoiseFrequency.x) + RandomFloat(Particle);
 			const uint3 ChannelLUTOffsetA = (ChannelIndex * SolveVelocitiesAndForces_NoiseLUTChannelWidth) + min(floor(ChannelLerp) + 0, SolveVelocitiesAndForces_NoiseLUTChannelWidth - 1);
 			const uint3 ChannelLUTOffsetB = (ChannelIndex * SolveVelocitiesAndForces_NoiseLUTChannelWidth) + min(floor(ChannelLerp) + 1, SolveVelocitiesAndForces_NoiseLUTChannelWidth - 1);
 			const uint3 ChannelLUTOffsetC = (ChannelIndex * SolveVelocitiesAndForces_NoiseLUTChannelWidth) + min(floor(ChannelLerp) + 2, SolveVelocitiesAndForces_NoiseLUTChannelWidth - 1);
diff --git a/Engine/Plugins/FX/Niagara/Shaders/Private/Stateless/NiagaraStatelessDistribution.ush b/Engine/Plugins/FX/Niagara/Shaders/Private/Stateless/NiagaraStatelessDistribution.ush
--- a/Engine/Plugins/FX/Niagara/Shaders/Private/Stateless/NiagaraStatelessDistribution.ush
+++ b/Engine/Plugins/FX/Niagara/Shaders/Private/Stateless/NiagaraStatelessDistribution.ush
@@ -21,6 +21,11 @@ struct FStatelessDistributionFloat
 		RandomOffset		= RandomFloat(Particle);
 	}
 
+	bool IsValid()
+	{
+		return TableLength > 0.0f;
+	}
+	
 	float GetValue(float NormalizedAge)
 	{
 		if ( bRandomDistribution )
diff --git a/Engine/Plugins/FX/Niagara/Shaders/Private/Stateless/NiagaraStatelessSimulationDefault.usf b/Engine/Plugins/FX/Niagara/Shaders/Private/Stateless/NiagaraStatelessSimulationDefault.usf
--- a/Engine/Plugins/FX/Niagara/Shaders/Private/Stateless/NiagaraStatelessSimulationDefault.usf
+++ b/Engine/Plugins/FX/Niagara/Shaders/Private/Stateless/NiagaraStatelessSimulationDefault.usf
@@ -4,6 +4,7 @@
 	PARTICLE_ATTRIBUTE(int,		UniqueID) \
 	PARTICLE_ATTRIBUTE(float3,	Position) \
 	PARTICLE_ATTRIBUTE(float4,	Color) \
+	PARTICLE_ATTRIBUTE(float4,	DynamicMaterialParameter0) \
 	PARTICLE_ATTRIBUTE(float4,	MeshOrientation) \
 	PARTICLE_ATTRIBUTE(float,	RibbonWidth) \
 	PARTICLE_ATTRIBUTE(float3,	Scale) \
@@ -38,5 +39,6 @@
 	PARTICLE_MODULE(ScaleMeshSizeBySpeed) \
 	PARTICLE_MODULE(ScaleSpriteSizeBySpeed) \
 	PARTICLE_MODULE(SubUVAnimation) \
+	PARTICLE_MODULE(DynamicMaterialParameters) \
 
 #include "NiagaraStatelessSimulationTemplate.ush"
diff --git a/Engine/Plugins/FX/Niagara/Shaders/Private/Stateless/NiagaraStatelessSimulationExample1.usf b/Engine/Plugins/FX/Niagara/Shaders/Private/Stateless/NiagaraStatelessSimulationExample1.usf
--- a/Engine/Plugins/FX/Niagara/Shaders/Private/Stateless/NiagaraStatelessSimulationExample1.usf
+++ b/Engine/Plugins/FX/Niagara/Shaders/Private/Stateless/NiagaraStatelessSimulationExample1.usf
@@ -17,6 +17,7 @@
 	PARTICLE_ATTRIBUTE(float3,	PreviousVelocity) \
 
 #define PARTICLE_ATTRIBUTES_TRANSIENT \
+	PARTICLE_ATTRIBUTE(float4,	DynamicMaterialParameter0) \
 	PARTICLE_ATTRIBUTE(float4,	MeshOrientation) \
 	PARTICLE_ATTRIBUTE(float3,	SpriteFacing) \
 	PARTICLE_ATTRIBUTE(float3,	SpriteAlignment) \
diff --git a/Engine/Plugins/FX/Niagara/Shaders/Private/Stateless/NiagaraStatelessSimulationTemplate.ush b/Engine/Plugins/FX/Niagara/Shaders/Private/Stateless/NiagaraStatelessSimulationTemplate.ush
--- a/Engine/Plugins/FX/Niagara/Shaders/Private/Stateless/NiagaraStatelessSimulationTemplate.ush
+++ b/Engine/Plugins/FX/Niagara/Shaders/Private/Stateless/NiagaraStatelessSimulationTemplate.ush
@@ -12,6 +12,7 @@
 /////////////////////////////////////////////////////////////////////////////////////////////////////////////////////////
 // TBD this might produce sub optimal code, needs investigate and is temporary
 #include "Modules/NiagaraStatelessModule_CalculateAccurateVelocity.ush"
+#include "Modules/NiagaraStatelessModule_DynamicMaterialParameters.ush"
 #include "Modules/NiagaraStatelessModule_InitializeParticle.ush"
 #include "Modules/NiagaraStatelessModule_InitialMeshOrientation.ush"
 #include "Modules/NiagaraStatelessModule_RotateAroundPoint.ush"
diff --git a/Engine/Plugins/FX/Niagara/Source/Niagara/Internal/Stateless/Modules/NiagaraStatelessModule_CalculateAccurateVelocity.h b/Engine/Plugins/FX/Niagara/Source/Niagara/Internal/Stateless/Modules/NiagaraStatelessModule_CalculateAccurateVelocity.h
--- a/Engine/Plugins/FX/Niagara/Source/Niagara/Internal/Stateless/Modules/NiagaraStatelessModule_CalculateAccurateVelocity.h
+++ b/Engine/Plugins/FX/Niagara/Source/Niagara/Internal/Stateless/Modules/NiagaraStatelessModule_CalculateAccurateVelocity.h
@@ -14,7 +14,7 @@ class UNiagaraStatelessModule_CalculateAccurateVelocity : public UNiagaraStatele
 
 public:
 #if WITH_EDITORONLY_DATA
-	virtual void GetOutputVaruables(TArray<FNiagaraVariableBase>& OutVariables) const override
+	virtual void GetOutputVariables(TArray<FNiagaraVariableBase>& OutVariables) const override
 	{
 		const FNiagaraStatelessGlobals& StatelessGlobals = FNiagaraStatelessGlobals::Get();
 		OutVariables.AddUnique(StatelessGlobals.VelocityVariable);
diff --git a/Engine/Plugins/FX/Niagara/Source/Niagara/Internal/Stateless/Modules/NiagaraStatelessModule_DynamicMaterialParameters.h b/Engine/Plugins/FX/Niagara/Source/Niagara/Internal/Stateless/Modules/NiagaraStatelessModule_DynamicMaterialParameters.h
new file mode 100644
--- /dev/null
+++ b/Engine/Plugins/FX/Niagara/Source/Niagara/Internal/Stateless/Modules/NiagaraStatelessModule_DynamicMaterialParameters.h
@@ -0,0 +1,147 @@
+// Copyright Epic Games, Inc. All Rights Reserved.
+
+#pragma once
+
+#include "Stateless/NiagaraStatelessModule.h"
+#include "Stateless/NiagaraStatelessEmitterDataBuildContext.h"
+#include "Stateless/NiagaraStatelessModuleShaderParameters.h"
+
+#include "NiagaraStatelessModule_DynamicMaterialParameters.generated.h"
+
+UCLASS(MinimalAPI, EditInlineNew, meta = (DisplayName = "Dynamic Material Parameters"))
+class UNiagaraStatelessModule_DynamicMaterialParameters : public UNiagaraStatelessModule
+{
+	GENERATED_BODY()
+
+	static constexpr int32 NumParameters			= 1;
+	static constexpr int32 NumChannelPerParameter	= 4;
+
+	struct FModuleBuiltData
+	{
+		FModuleBuiltData()
+		{
+			for (FUintVector3& v : ParameterDistributions)
+			{
+				v = FUintVector3::ZeroValue;
+			}
+		}
+
+		uint32			ChannelMask = 0;
+		FUintVector3	ParameterDistributions[NumParameters * NumChannelPerParameter];
+	};
+
+public:
+	using FParameters = NiagaraStateless::FDynamicMaterialParametersModule_ShaderParameters;
+
+	UPROPERTY(EditAnywhere, Category = "Parameters", meta = (InlineEditConditionToggle))
+	uint32 bParameter0XEnabled : 1 = false;
+
+	UPROPERTY(EditAnywhere, Category = "Parameters", meta = (InlineEditConditionToggle))
+	uint32 bParameter0YEnabled : 1 = false;
+
+	UPROPERTY(EditAnywhere, Category = "Parameters", meta = (InlineEditConditionToggle))
+	uint32 bParameter0ZEnabled : 1 = false;
+
+	UPROPERTY(EditAnywhere, Category = "Parameters", meta = (InlineEditConditionToggle))
+	uint32 bParameter0WEnabled : 1 = false;
+
+	UPROPERTY(EditAnywhere, Category = "Parameters", meta = (DisplayName = "Parameter0.X", EditCondition = "bParameter0XEnabled"))
+	FNiagaraDistributionFloat Parameter0XDistribution = FNiagaraDistributionFloat(FNiagaraStatelessGlobals::GetDefaultDynamicMaterialParameters0Value().X);
+
+	UPROPERTY(EditAnywhere, Category = "Parameters", meta = (DisplayName = "Parameter0.Y", EditCondition = "bParameter0YEnabled"))
+	FNiagaraDistributionFloat Parameter0YDistribution = FNiagaraDistributionFloat(FNiagaraStatelessGlobals::GetDefaultDynamicMaterialParameters0Value().Y);
+
+	UPROPERTY(EditAnywhere, Category = "Parameters", meta = (DisplayName = "Parameter0.Z", EditCondition = "bParameter0ZEnabled"))
+	FNiagaraDistributionFloat Parameter0ZDistribution = FNiagaraDistributionFloat(FNiagaraStatelessGlobals::GetDefaultDynamicMaterialParameters0Value().Z);
+
+	UPROPERTY(EditAnywhere, Category = "Parameters", meta = (DisplayName = "Parameter0.W", EditCondition = "bParameter0WEnabled"))
+	FNiagaraDistributionFloat Parameter0WDistribution = FNiagaraDistributionFloat(FNiagaraStatelessGlobals::GetDefaultDynamicMaterialParameters0Value().W);
+
+	uint32 GetParameterChannelMask(int32 ParameterIndex) const
+	{
+		uint32 Mask = 0;
+		switch (ParameterIndex)
+		{
+			case 0:
+				Mask |= bParameter0XEnabled ? 1 << 0 : 0;
+				Mask |= bParameter0YEnabled ? 1 << 1 : 0;
+				Mask |= bParameter0ZEnabled ? 1 << 2 : 0;
+				Mask |= bParameter0WEnabled ? 1 << 3 : 0;
+				break;
+		}
+		return Mask;
+	}
+
+	const FNiagaraDistributionFloat& GetParameterDistributions(int32 ParameterIndex, int32 ChannelIndex) const
+	{
+		switch (ParameterIndex)
+		{
+			case 0:
+				switch (ChannelIndex)
+				{
+					case 0:		return Parameter0XDistribution;
+					case 1:		return Parameter0YDistribution;
+					case 2:		return Parameter0ZDistribution;
+					case 3:		return Parameter0WDistribution;
+					default:	checkNoEntry();	return Parameter0XDistribution;
+				}
+				break;
+
+			default:
+				checkNoEntry();
+				return Parameter0XDistribution;
+		}
+	}
+
+	virtual void BuildEmitterData(FNiagaraStatelessEmitterDataBuildContext& BuildContext) const override
+	{
+		FModuleBuiltData* BuiltData = BuildContext.AllocateBuiltData<FModuleBuiltData>();
+
+		if (IsModuleEnabled())
+		{
+			for (int32 iParameter=0; iParameter < NumParameters; ++iParameter)
+			{
+				const uint32 ParameterChannelMask = GetParameterChannelMask(iParameter);
+				if (ParameterChannelMask == 0)
+				{
+					continue;
+				}
+				BuiltData->ChannelMask |= ParameterChannelMask << (iParameter * 4);
+				for (int32 iChannel = 0; iChannel < NumChannelPerParameter; ++iChannel)
+				{
+					BuiltData->ParameterDistributions[iParameter * NumParameters + iChannel] = BuildContext.AddDistribution(GetParameterDistributions(iParameter, iChannel), true);
+				}
+			}
+		}
+	}
+
+	virtual void SetShaderParameters(const FNiagaraStatelessSetShaderParameterContext& SetShaderParameterContext) const override
+	{
+		const FModuleBuiltData* ModuleBuiltData = SetShaderParameterContext.ReadBuiltData<FModuleBuiltData>();
+
+		FParameters* Parameters = SetShaderParameterContext.GetParameterNestedStruct<FParameters>();
+		Parameters->DynamicMaterialParameters_ChannelMask = ModuleBuiltData->ChannelMask;
+		Parameters->DynamicMaterialParameters_Parameter0X = ModuleBuiltData->ParameterDistributions[0];
+		Parameters->DynamicMaterialParameters_Parameter0Y = ModuleBuiltData->ParameterDistributions[1];
+		Parameters->DynamicMaterialParameters_Parameter0Z = ModuleBuiltData->ParameterDistributions[2];
+		Parameters->DynamicMaterialParameters_Parameter0W = ModuleBuiltData->ParameterDistributions[3];
+	}
+
+#if WITH_EDITOR
+	virtual bool CanDisableModule() const override { return true; }
+#endif
+#if WITH_EDITORONLY_DATA
+	virtual void GetOutputVariables(TArray<FNiagaraVariableBase>& OutVariables) const override
+	{
+		const FNiagaraStatelessGlobals& StatelessGlobals = FNiagaraStatelessGlobals::Get();
+		if (IsModuleEnabled())
+		{
+			//-TODO: Channel masking, etc
+			if (GetParameterChannelMask(0) != 0)
+			{
+				OutVariables.AddUnique(StatelessGlobals.DynamicMaterialParameters0Variable);
+			}
+		}
+	}
+#endif
+};
diff --git a/Engine/Plugins/FX/Niagara/Source/Niagara/Internal/Stateless/Modules/NiagaraStatelessModule_InitialMeshOrientation.h b/Engine/Plugins/FX/Niagara/Source/Niagara/Internal/Stateless/Modules/NiagaraStatelessModule_InitialMeshOrientation.h
--- a/Engine/Plugins/FX/Niagara/Source/Niagara/Internal/Stateless/Modules/NiagaraStatelessModule_InitialMeshOrientation.h
+++ b/Engine/Plugins/FX/Niagara/Source/Niagara/Internal/Stateless/Modules/NiagaraStatelessModule_InitialMeshOrientation.h
@@ -41,7 +41,7 @@ public:
 	virtual bool CanDisableModule() const override { return true; }
 #endif
 #if WITH_EDITORONLY_DATA
-	virtual void GetOutputVaruables(TArray<FNiagaraVariableBase>& OutVariables) const override
+	virtual void GetOutputVariables(TArray<FNiagaraVariableBase>& OutVariables) const override
 	{
 		const FNiagaraStatelessGlobals& StatelessGlobals = FNiagaraStatelessGlobals::Get();
 		OutVariables.AddUnique(StatelessGlobals.MeshOrientationVariable);
diff --git a/Engine/Plugins/FX/Niagara/Source/Niagara/Internal/Stateless/Modules/NiagaraStatelessModule_InitializeParticle.h b/Engine/Plugins/FX/Niagara/Source/Niagara/Internal/Stateless/Modules/NiagaraStatelessModule_InitializeParticle.h
--- a/Engine/Plugins/FX/Niagara/Source/Niagara/Internal/Stateless/Modules/NiagaraStatelessModule_InitializeParticle.h
+++ b/Engine/Plugins/FX/Niagara/Source/Niagara/Internal/Stateless/Modules/NiagaraStatelessModule_InitializeParticle.h
@@ -80,7 +80,7 @@ public:
 
 	virtual void BuildEmitterData(FNiagaraStatelessEmitterDataBuildContext& BuildContext) const override
 	{
-		FModuleBuiltData* BuiltData = BuildContext.AllocateBuiltData<FModuleBuiltData>();
+		FModuleBuiltData* BuiltData			= BuildContext.AllocateBuiltData<FModuleBuiltData>();
 		BuiltData->ModuleFlags				 = SpriteSizeDistribution.IsUniform() ? EInitializeParticleModuleFlag_UniformSpriteSize : 0;
 		BuiltData->ModuleFlags				|= MeshScaleDistribution.IsUniform() ? EInitializeParticleModuleFlag_UniformMeshScale : 0;
 
@@ -117,7 +117,7 @@ public:
 	}
 
 #if WITH_EDITORONLY_DATA
-	virtual void GetOutputVaruables(TArray<FNiagaraVariableBase>& OutVariables) const override
+	virtual void GetOutputVariables(TArray<FNiagaraVariableBase>& OutVariables) const override
 	{
 		const FNiagaraStatelessGlobals& StatelessGlobals = FNiagaraStatelessGlobals::Get();
 		OutVariables.AddUnique(StatelessGlobals.UniqueIDVariable);
diff --git a/Engine/Plugins/FX/Niagara/Source/Niagara/Internal/Stateless/Modules/NiagaraStatelessModule_RotateAroundPoint.h b/Engine/Plugins/FX/Niagara/Source/Niagara/Internal/Stateless/Modules/NiagaraStatelessModule_RotateAroundPoint.h
--- a/Engine/Plugins/FX/Niagara/Source/Niagara/Internal/Stateless/Modules/NiagaraStatelessModule_RotateAroundPoint.h
+++ b/Engine/Plugins/FX/Niagara/Source/Niagara/Internal/Stateless/Modules/NiagaraStatelessModule_RotateAroundPoint.h
@@ -57,7 +57,7 @@ public:
 	virtual bool CanDisableModule() const override { return true; }
 #endif
 #if WITH_EDITORONLY_DATA
-	virtual void GetOutputVaruables(TArray<FNiagaraVariableBase>& OutVariables) const override
+	virtual void GetOutputVariables(TArray<FNiagaraVariableBase>& OutVariables) const override
 	{
 		const FNiagaraStatelessGlobals& StatelessGlobals = FNiagaraStatelessGlobals::Get();
 		OutVariables.AddUnique(StatelessGlobals.PositionVariable);
diff --git a/Engine/Plugins/FX/Niagara/Source/Niagara/Internal/Stateless/Modules/NiagaraStatelessModule_ScaleColor.h b/Engine/Plugins/FX/Niagara/Source/Niagara/Internal/Stateless/Modules/NiagaraStatelessModule_ScaleColor.h
--- a/Engine/Plugins/FX/Niagara/Source/Niagara/Internal/Stateless/Modules/NiagaraStatelessModule_ScaleColor.h
+++ b/Engine/Plugins/FX/Niagara/Source/Niagara/Internal/Stateless/Modules/NiagaraStatelessModule_ScaleColor.h
@@ -42,7 +42,7 @@ public:
 	virtual bool CanDisableModule() const override { return true; }
 #endif
 #if WITH_EDITORONLY_DATA
-	virtual void GetOutputVaruables(TArray<FNiagaraVariableBase>& OutVariables) const override
+	virtual void GetOutputVariables(TArray<FNiagaraVariableBase>& OutVariables) const override
 	{
 		const FNiagaraStatelessGlobals& StatelessGlobals = FNiagaraStatelessGlobals::Get();
 		OutVariables.AddUnique(StatelessGlobals.ColorVariable);
diff --git a/Engine/Plugins/FX/Niagara/Source/Niagara/Internal/Stateless/Modules/NiagaraStatelessModule_ScaleMeshSize.h b/Engine/Plugins/FX/Niagara/Source/Niagara/Internal/Stateless/Modules/NiagaraStatelessModule_ScaleMeshSize.h
--- a/Engine/Plugins/FX/Niagara/Source/Niagara/Internal/Stateless/Modules/NiagaraStatelessModule_ScaleMeshSize.h
+++ b/Engine/Plugins/FX/Niagara/Source/Niagara/Internal/Stateless/Modules/NiagaraStatelessModule_ScaleMeshSize.h
@@ -42,7 +42,7 @@ public:
 	virtual bool CanDisableModule() const override { return true; }
 #endif
 #if WITH_EDITORONLY_DATA
-	virtual void GetOutputVaruables(TArray<FNiagaraVariableBase>& OutVariables) const override
+	virtual void GetOutputVariables(TArray<FNiagaraVariableBase>& OutVariables) const override
 	{
 		const FNiagaraStatelessGlobals& StatelessGlobals = FNiagaraStatelessGlobals::Get();
 		OutVariables.AddUnique(StatelessGlobals.ScaleVariable);
diff --git a/Engine/Plugins/FX/Niagara/Source/Niagara/Internal/Stateless/Modules/NiagaraStatelessModule_ScaleMeshSizeBySpeed.h b/Engine/Plugins/FX/Niagara/Source/Niagara/Internal/Stateless/Modules/NiagaraStatelessModule_ScaleMeshSizeBySpeed.h
--- a/Engine/Plugins/FX/Niagara/Source/Niagara/Internal/Stateless/Modules/NiagaraStatelessModule_ScaleMeshSizeBySpeed.h
+++ b/Engine/Plugins/FX/Niagara/Source/Niagara/Internal/Stateless/Modules/NiagaraStatelessModule_ScaleMeshSizeBySpeed.h
@@ -61,7 +61,7 @@ public:
 	virtual bool CanDisableModule() const override { return true; }
 #endif
 #if WITH_EDITORONLY_DATA
-	virtual void GetOutputVaruables(TArray<FNiagaraVariableBase>& OutVariables) const override
+	virtual void GetOutputVariables(TArray<FNiagaraVariableBase>& OutVariables) const override
 	{
 		const FNiagaraStatelessGlobals& StatelessGlobals = FNiagaraStatelessGlobals::Get();
 		OutVariables.AddUnique(StatelessGlobals.ScaleVariable);
diff --git a/Engine/Plugins/FX/Niagara/Source/Niagara/Internal/Stateless/Modules/NiagaraStatelessModule_ScaleSpriteSize.h b/Engine/Plugins/FX/Niagara/Source/Niagara/Internal/Stateless/Modules/NiagaraStatelessModule_ScaleSpriteSize.h
--- a/Engine/Plugins/FX/Niagara/Source/Niagara/Internal/Stateless/Modules/NiagaraStatelessModule_ScaleSpriteSize.h
+++ b/Engine/Plugins/FX/Niagara/Source/Niagara/Internal/Stateless/Modules/NiagaraStatelessModule_ScaleSpriteSize.h
@@ -43,7 +43,7 @@ public:
 	virtual bool CanDisableModule() const override { return true; }
 #endif
 #if WITH_EDITORONLY_DATA
-	virtual void GetOutputVaruables(TArray<FNiagaraVariableBase>& OutVariables) const override
+	virtual void GetOutputVariables(TArray<FNiagaraVariableBase>& OutVariables) const override
 	{
 		const FNiagaraStatelessGlobals& StatelessGlobals = FNiagaraStatelessGlobals::Get();
 		OutVariables.AddUnique(StatelessGlobals.SpriteSizeVariable);
diff --git a/Engine/Plugins/FX/Niagara/Source/Niagara/Internal/Stateless/Modules/NiagaraStatelessModule_ScaleSpriteSizeBySpeed.h b/Engine/Plugins/FX/Niagara/Source/Niagara/Internal/Stateless/Modules/NiagaraStatelessModule_ScaleSpriteSizeBySpeed.h
--- a/Engine/Plugins/FX/Niagara/Source/Niagara/Internal/Stateless/Modules/NiagaraStatelessModule_ScaleSpriteSizeBySpeed.h
+++ b/Engine/Plugins/FX/Niagara/Source/Niagara/Internal/Stateless/Modules/NiagaraStatelessModule_ScaleSpriteSizeBySpeed.h
@@ -61,7 +61,7 @@ public:
 	virtual bool CanDisableModule() const override { return true; }
 #endif
 #if WITH_EDITORONLY_DATA
-	virtual void GetOutputVaruables(TArray<FNiagaraVariableBase>& OutVariables) const override
+	virtual void GetOutputVariables(TArray<FNiagaraVariableBase>& OutVariables) const override
 	{
 		const FNiagaraStatelessGlobals& StatelessGlobals = FNiagaraStatelessGlobals::Get();
 		OutVariables.AddUnique(StatelessGlobals.SpriteSizeVariable);
diff --git a/Engine/Plugins/FX/Niagara/Source/Niagara/Internal/Stateless/Modules/NiagaraStatelessModule_ShapeLocation.h b/Engine/Plugins/FX/Niagara/Source/Niagara/Internal/Stateless/Modules/NiagaraStatelessModule_ShapeLocation.h
--- a/Engine/Plugins/FX/Niagara/Source/Niagara/Internal/Stateless/Modules/NiagaraStatelessModule_ShapeLocation.h
+++ b/Engine/Plugins/FX/Niagara/Source/Niagara/Internal/Stateless/Modules/NiagaraStatelessModule_ShapeLocation.h
@@ -68,7 +68,7 @@ public:
 	virtual void DrawDebug(const FNiagaraStatelessDrawDebugContext& DrawDebugContext) const override;
 #endif
 #if WITH_EDITORONLY_DATA
-	virtual void GetOutputVaruables(TArray<FNiagaraVariableBase>& OutVariables) const override
+	virtual void GetOutputVariables(TArray<FNiagaraVariableBase>& OutVariables) const override
 	{
 		const FNiagaraStatelessGlobals& StatelessGlobals = FNiagaraStatelessGlobals::Get();
 		OutVariables.AddUnique(StatelessGlobals.PositionVariable);
diff --git a/Engine/Plugins/FX/Niagara/Source/Niagara/Internal/Stateless/Modules/NiagaraStatelessModule_SolveVelocitiesAndForces.h b/Engine/Plugins/FX/Niagara/Source/Niagara/Internal/Stateless/Modules/NiagaraStatelessModule_SolveVelocitiesAndForces.h
--- a/Engine/Plugins/FX/Niagara/Source/Niagara/Internal/Stateless/Modules/NiagaraStatelessModule_SolveVelocitiesAndForces.h
+++ b/Engine/Plugins/FX/Niagara/Source/Niagara/Internal/Stateless/Modules/NiagaraStatelessModule_SolveVelocitiesAndForces.h
@@ -130,7 +130,7 @@ public:
 	}
 
 #if WITH_EDITORONLY_DATA
-	virtual void GetOutputVaruables(TArray<FNiagaraVariableBase>& OutVariables) const override
+	virtual void GetOutputVariables(TArray<FNiagaraVariableBase>& OutVariables) const override
 	{
 		const FNiagaraStatelessGlobals& StatelessGlobals = FNiagaraStatelessGlobals::Get();
 		OutVariables.AddUnique(StatelessGlobals.PositionVariable);
diff --git a/Engine/Plugins/FX/Niagara/Source/Niagara/Internal/Stateless/Modules/NiagaraStatelessModule_SpriteFacingAndAlignment.h b/Engine/Plugins/FX/Niagara/Source/Niagara/Internal/Stateless/Modules/NiagaraStatelessModule_SpriteFacingAndAlignment.h
--- a/Engine/Plugins/FX/Niagara/Source/Niagara/Internal/Stateless/Modules/NiagaraStatelessModule_SpriteFacingAndAlignment.h
+++ b/Engine/Plugins/FX/Niagara/Source/Niagara/Internal/Stateless/Modules/NiagaraStatelessModule_SpriteFacingAndAlignment.h
@@ -39,7 +39,7 @@ public:
 	virtual bool CanDisableModule() const override { return true; }
 #endif
 #if WITH_EDITORONLY_DATA
-	virtual void GetOutputVaruables(TArray<FNiagaraVariableBase>& OutVariables) const override
+	virtual void GetOutputVariables(TArray<FNiagaraVariableBase>& OutVariables) const override
 	{
 		const FNiagaraStatelessGlobals& StatelessGlobals = FNiagaraStatelessGlobals::Get();
 		if (IsModuleEnabled())
diff --git a/Engine/Plugins/FX/Niagara/Source/Niagara/Internal/Stateless/Modules/NiagaraStatelessModule_SpriteRotationRate.h b/Engine/Plugins/FX/Niagara/Source/Niagara/Internal/Stateless/Modules/NiagaraStatelessModule_SpriteRotationRate.h
--- a/Engine/Plugins/FX/Niagara/Source/Niagara/Internal/Stateless/Modules/NiagaraStatelessModule_SpriteRotationRate.h
+++ b/Engine/Plugins/FX/Niagara/Source/Niagara/Internal/Stateless/Modules/NiagaraStatelessModule_SpriteRotationRate.h
@@ -39,7 +39,7 @@ public:
 	virtual bool CanDisableModule() const override { return true; }
 #endif
 #if WITH_EDITORONLY_DATA
-	virtual void GetOutputVaruables(TArray<FNiagaraVariableBase>& OutVariables) const override
+	virtual void GetOutputVariables(TArray<FNiagaraVariableBase>& OutVariables) const override
 	{
 		const FNiagaraStatelessGlobals& StatelessGlobals = FNiagaraStatelessGlobals::Get();
 		OutVariables.AddUnique(StatelessGlobals.SpriteRotationVariable);
diff --git a/Engine/Plugins/FX/Niagara/Source/Niagara/Internal/Stateless/Modules/NiagaraStatelessModule_SubUVAnimation.h b/Engine/Plugins/FX/Niagara/Source/Niagara/Internal/Stateless/Modules/NiagaraStatelessModule_SubUVAnimation.h
--- a/Engine/Plugins/FX/Niagara/Source/Niagara/Internal/Stateless/Modules/NiagaraStatelessModule_SubUVAnimation.h
+++ b/Engine/Plugins/FX/Niagara/Source/Niagara/Internal/Stateless/Modules/NiagaraStatelessModule_SubUVAnimation.h
@@ -111,7 +111,7 @@ public:
 	virtual bool CanDisableModule() const override { return true; }
 #endif
 #if WITH_EDITORONLY_DATA
-	virtual void GetOutputVaruables(TArray<FNiagaraVariableBase>& OutVariables) const override
+	virtual void GetOutputVariables(TArray<FNiagaraVariableBase>& OutVariables) const override
 	{
 		const FNiagaraStatelessGlobals& StatelessGlobals = FNiagaraStatelessGlobals::Get();
 		OutVariables.AddUnique(StatelessGlobals.SubImageIndexVariable);
diff --git a/Engine/Plugins/FX/Niagara/Source/Niagara/Internal/Stateless/NiagaraStatelessCommon.h b/Engine/Plugins/FX/Niagara/Source/Niagara/Internal/Stateless/NiagaraStatelessCommon.h
--- a/Engine/Plugins/FX/Niagara/Source/Niagara/Internal/Stateless/NiagaraStatelessCommon.h
+++ b/Engine/Plugins/FX/Niagara/Source/Niagara/Internal/Stateless/NiagaraStatelessCommon.h
@@ -29,6 +29,7 @@ using FNiagaraStatelessRangeColor	= FNiagaraStatelessRange<FLinearColor>;
 struct FNiagaraStatelessGlobals
 {
 	FNiagaraVariableBase	ColorVariable;
+	FNiagaraVariableBase	DynamicMaterialParameters0Variable;
 	FNiagaraVariableBase	MeshOrientationVariable;
 	FNiagaraVariableBase	PositionVariable;
 	FNiagaraVariableBase	RibbonWidthVariable;
@@ -42,6 +43,7 @@ struct FNiagaraStatelessGlobals
 	FNiagaraVariableBase	VelocityVariable;
 
 	//FNiagaraVariableBase	PreviousColorVariable;
+	//FNiagaraVariableBase	PreviousDynamicMaterialParameters0Variable;
 	FNiagaraVariableBase	PreviousMeshOrientationVariable;
 	FNiagaraVariableBase	PreviousPositionVariable;
 	FNiagaraVariableBase	PreviousRibbonWidthVariable;
@@ -53,6 +55,7 @@ struct FNiagaraStatelessGlobals
 	FNiagaraVariableBase	PreviousVelocityVariable;
 
 	inline static FLinearColor	GetDefaultColorValue() { return FLinearColor::White; }
+	inline static FVector4f		GetDefaultDynamicMaterialParameters0Value() { return FVector4f::Zero(); }
 	inline static float			GetDefaultLifetimeValue() { return 1.0f; }
 	inline static float			GetDefaultMassValue() { return 1.0f; }
 	inline static FQuat4f		GetDefaultMeshOrientationValue() { return FQuat4f::Identity; }
diff --git a/Engine/Plugins/FX/Niagara/Source/Niagara/Internal/Stateless/NiagaraStatelessEmitterDataBuildContext.h b/Engine/Plugins/FX/Niagara/Source/Niagara/Internal/Stateless/NiagaraStatelessEmitterDataBuildContext.h
--- a/Engine/Plugins/FX/Niagara/Source/Niagara/Internal/Stateless/NiagaraStatelessEmitterDataBuildContext.h
+++ b/Engine/Plugins/FX/Niagara/Source/Niagara/Internal/Stateless/NiagaraStatelessEmitterDataBuildContext.h
@@ -30,9 +30,12 @@ public:
 	template<typename T>
 	T* AllocateBuiltData()
 	{
+		static_assert(TIsTrivial<T>::Value, "Only trivial types can be used for built data");
+
 		const int32 Align = BuiltData.Num() % alignof(T);
 		const int32 Offset = BuiltData.AddZeroed(sizeof(T) + Align);
-		return reinterpret_cast<T*>(BuiltData.GetData() + Offset + Align);
+		void* NewData = BuiltData.GetData() + Offset + Align;
+		return new(NewData) T();
 	}
 
 	template<typename T>
diff --git a/Engine/Plugins/FX/Niagara/Source/Niagara/Internal/Stateless/NiagaraStatelessModule.h b/Engine/Plugins/FX/Niagara/Source/Niagara/Internal/Stateless/NiagaraStatelessModule.h
--- a/Engine/Plugins/FX/Niagara/Source/Niagara/Internal/Stateless/NiagaraStatelessModule.h
+++ b/Engine/Plugins/FX/Niagara/Source/Niagara/Internal/Stateless/NiagaraStatelessModule.h
@@ -120,7 +120,7 @@ public:
 #endif
 
 #if WITH_EDITORONLY_DATA
-	virtual void GetOutputVaruables(TArray<FNiagaraVariableBase>& OutVariables) const {}
+	virtual void GetOutputVariables(TArray<FNiagaraVariableBase>& OutVariables) const {}
 #endif
 
 	//~UObject interface Begin
diff --git a/Engine/Plugins/FX/Niagara/Source/Niagara/Private/Stateless/NiagaraStatelessCommon.cpp b/Engine/Plugins/FX/Niagara/Source/Niagara/Private/Stateless/NiagaraStatelessCommon.cpp
--- a/Engine/Plugins/FX/Niagara/Source/Niagara/Private/Stateless/NiagaraStatelessCommon.cpp
+++ b/Engine/Plugins/FX/Niagara/Source/Niagara/Private/Stateless/NiagaraStatelessCommon.cpp
@@ -10,20 +10,22 @@ namespace NiagaraStatelessCommon
 
 	void Initialize()
 	{
-		GGlobals.ColorVariable				= FNiagaraVariableBase(FNiagaraTypeDefinition::GetColorDef(), TEXT("Color"));
-		GGlobals.MeshOrientationVariable	= FNiagaraVariableBase(FNiagaraTypeDefinition::GetQuatDef(), TEXT("MeshOrientation"));
-		GGlobals.PositionVariable			= FNiagaraVariableBase(FNiagaraTypeDefinition::GetPositionDef(), TEXT("Position"));
-		GGlobals.RibbonWidthVariable		= FNiagaraVariableBase(FNiagaraTypeDefinition::GetFloatDef(), TEXT("RibbonWidth"));
-		GGlobals.ScaleVariable				= FNiagaraVariableBase(FNiagaraTypeDefinition::GetVec3Def(), TEXT("Scale"));
-		GGlobals.SpriteAlignmentVariable	= FNiagaraVariableBase(FNiagaraTypeDefinition::GetVec3Def(), TEXT("SpriteAlignment"));
-		GGlobals.SpriteFacingVariable		= FNiagaraVariableBase(FNiagaraTypeDefinition::GetVec3Def(), TEXT("SpriteFacing"));
-		GGlobals.SpriteSizeVariable			= FNiagaraVariableBase(FNiagaraTypeDefinition::GetVec2Def(), TEXT("SpriteSize"));
-		GGlobals.SpriteRotationVariable		= FNiagaraVariableBase(FNiagaraTypeDefinition::GetFloatDef(), TEXT("SpriteRotation"));
-		GGlobals.SubImageIndexVariable		= FNiagaraVariableBase(FNiagaraTypeDefinition::GetFloatDef(), TEXT("SubImageIndex"));
-		GGlobals.UniqueIDVariable			= FNiagaraVariableBase(FNiagaraTypeDefinition::GetIntDef(), TEXT("UniqueID"));
-		GGlobals.VelocityVariable			= FNiagaraVariableBase(FNiagaraTypeDefinition::GetVec3Def(), TEXT("Velocity"));
+		GGlobals.ColorVariable						= FNiagaraVariableBase(FNiagaraTypeDefinition::GetColorDef(), TEXT("Color"));
+		GGlobals.DynamicMaterialParameters0Variable = FNiagaraVariableBase(FNiagaraTypeDefinition::GetVec4Def(), TEXT("DynamicMaterialParameter"));
+		GGlobals.MeshOrientationVariable			= FNiagaraVariableBase(FNiagaraTypeDefinition::GetQuatDef(), TEXT("MeshOrientation"));
+		GGlobals.PositionVariable					= FNiagaraVariableBase(FNiagaraTypeDefinition::GetPositionDef(), TEXT("Position"));
+		GGlobals.RibbonWidthVariable				= FNiagaraVariableBase(FNiagaraTypeDefinition::GetFloatDef(), TEXT("RibbonWidth"));
+		GGlobals.ScaleVariable						= FNiagaraVariableBase(FNiagaraTypeDefinition::GetVec3Def(), TEXT("Scale"));
+		GGlobals.SpriteAlignmentVariable			= FNiagaraVariableBase(FNiagaraTypeDefinition::GetVec3Def(), TEXT("SpriteAlignment"));
+		GGlobals.SpriteFacingVariable				= FNiagaraVariableBase(FNiagaraTypeDefinition::GetVec3Def(), TEXT("SpriteFacing"));
+		GGlobals.SpriteSizeVariable					= FNiagaraVariableBase(FNiagaraTypeDefinition::GetVec2Def(), TEXT("SpriteSize"));
+		GGlobals.SpriteRotationVariable				= FNiagaraVariableBase(FNiagaraTypeDefinition::GetFloatDef(), TEXT("SpriteRotation"));
+		GGlobals.SubImageIndexVariable				= FNiagaraVariableBase(FNiagaraTypeDefinition::GetFloatDef(), TEXT("SubImageIndex"));
+		GGlobals.UniqueIDVariable					= FNiagaraVariableBase(FNiagaraTypeDefinition::GetIntDef(), TEXT("UniqueID"));
+		GGlobals.VelocityVariable					= FNiagaraVariableBase(FNiagaraTypeDefinition::GetVec3Def(), TEXT("Velocity"));
 
 		//GGlobals.PreviousColorVariable			= FNiagaraVariableBase(FNiagaraTypeDefinition::GetColorDef(), TEXT("Previous.Color"));
+		//GGlobals.PreviousDynamicMaterialParameters0Variable		= FNiagaraVariableBase(FNiagaraTypeDefinition::GetVec4Def(), TEXT("Previous.DynamicMaterialParameter"));
 		GGlobals.PreviousMeshOrientationVariable	= FNiagaraVariableBase(FNiagaraTypeDefinition::GetQuatDef(), TEXT("Previous.MeshOrientation"));
 		GGlobals.PreviousPositionVariable			= FNiagaraVariableBase(FNiagaraTypeDefinition::GetPositionDef(), TEXT("Previous.Position"));
 		GGlobals.PreviousRibbonWidthVariable		= FNiagaraVariableBase(FNiagaraTypeDefinition::GetFloatDef(), TEXT("Previous.RibbonWidth"));
diff --git a/Engine/Plugins/FX/Niagara/Source/Niagara/Private/Stateless/NiagaraStatelessEmitter.cpp b/Engine/Plugins/FX/Niagara/Source/Niagara/Private/Stateless/NiagaraStatelessEmitter.cpp
--- a/Engine/Plugins/FX/Niagara/Source/Niagara/Private/Stateless/NiagaraStatelessEmitter.cpp
+++ b/Engine/Plugins/FX/Niagara/Source/Niagara/Private/Stateless/NiagaraStatelessEmitter.cpp
@@ -303,7 +303,7 @@ void UNiagaraStatelessEmitter::BuildCompiledDataSet()
 		TArray<FNiagaraVariableBase> AvailableVariables;
 		for (UNiagaraStatelessModule* Module : Modules)
 		{
-			Module->GetOutputVaruables(AvailableVariables);
+			Module->GetOutputVariables(AvailableVariables);
 		}
 
 		// Remove any variables we don't output
diff --git a/Engine/Plugins/FX/Niagara/Source/Niagara/Private/Stateless/NiagaraStatelessEmitterTemplate.cpp b/Engine/Plugins/FX/Niagara/Source/Niagara/Private/Stateless/NiagaraStatelessEmitterTemplate.cpp
--- a/Engine/Plugins/FX/Niagara/Source/Niagara/Private/Stateless/NiagaraStatelessEmitterTemplate.cpp
+++ b/Engine/Plugins/FX/Niagara/Source/Niagara/Private/Stateless/NiagaraStatelessEmitterTemplate.cpp
@@ -27,6 +27,7 @@
 #include "Stateless/Modules/NiagaraStatelessModule_SpriteFacingAndAlignment.h"
 #include "Stateless/Modules/NiagaraStatelessModule_SpriteRotationRate.h"
 #include "Stateless/Modules/NiagaraStatelessModule_SubUVAnimation.h"
+#include "Stateless/Modules/NiagaraStatelessModule_DynamicMaterialParameters.h"
 
 namespace NiagaraStatelessEmitterTemplateImpl
 {
@@ -85,6 +86,7 @@ void UNiagaraStatelessEmitterDefault::InitModulesAndAttributes()
 		UNiagaraStatelessModule_SpriteFacingAndAlignment::StaticClass(),
 		UNiagaraStatelessModule_SpriteRotationRate::StaticClass(),
 		UNiagaraStatelessModule_SubUVAnimation::StaticClass(),
+		UNiagaraStatelessModule_DynamicMaterialParameters::StaticClass(),
 	};
 
 	const FNiagaraStatelessGlobals& StatelessGlobals = FNiagaraStatelessGlobals::Get();
@@ -93,6 +95,7 @@ void UNiagaraStatelessEmitterDefault::InitModulesAndAttributes()
 		StatelessGlobals.UniqueIDVariable,
 		StatelessGlobals.PositionVariable,
 		StatelessGlobals.ColorVariable,
+		StatelessGlobals.DynamicMaterialParameters0Variable,
 		StatelessGlobals.MeshOrientationVariable,
 		StatelessGlobals.RibbonWidthVariable,
 		StatelessGlobals.ScaleVariable,
@@ -136,6 +139,7 @@ void UNiagaraStatelessEmitterDefault::SetShaderParameters(uint8* ShaderParameter
 	ShaderParameters->Permutation_UniqueIDComponent					= ComponentOffsets[iComponent++];
 	ShaderParameters->Permutation_PositionComponent					= ComponentOffsets[iComponent++];
 	ShaderParameters->Permutation_ColorComponent					= ComponentOffsets[iComponent++];
+	ShaderParameters->Permutation_DynamicMaterialParameter0Component= ComponentOffsets[iComponent++];
 	ShaderParameters->Permutation_MeshOrientationComponent			= ComponentOffsets[iComponent++];
 	ShaderParameters->Permutation_RibbonWidthComponent				= ComponentOffsets[iComponent++];
 	ShaderParameters->Permutation_ScaleComponent					= ComponentOffsets[iComponent++];
diff --git a/Engine/Plugins/FX/Niagara/Source/NiagaraShader/Internal/Stateless/NiagaraStatelessEmitterShaders.h b/Engine/Plugins/FX/Niagara/Source/NiagaraShader/Internal/Stateless/NiagaraStatelessEmitterShaders.h
--- a/Engine/Plugins/FX/Niagara/Source/NiagaraShader/Internal/Stateless/NiagaraStatelessEmitterShaders.h
+++ b/Engine/Plugins/FX/Niagara/Source/NiagaraShader/Internal/Stateless/NiagaraStatelessEmitterShaders.h
@@ -40,10 +40,12 @@ namespace NiagaraStateless
 			SHADER_PARAMETER_STRUCT_INCLUDE(FSpriteFacingAndAlignmentModule_ShaderParameters,	Module9Params)
 			SHADER_PARAMETER_STRUCT_INCLUDE(FSpriteRotationRateModule_ShaderParameters,			Module10Params)
 			SHADER_PARAMETER_STRUCT_INCLUDE(FSubUVAnimationModule_ShaderParameters,				Module11Params)
+			SHADER_PARAMETER_STRUCT_INCLUDE(FDynamicMaterialParametersModule_ShaderParameters,	Module12Params)
 
 			SHADER_PARAMETER(int, Permutation_UniqueIDComponent)
 			SHADER_PARAMETER(int, Permutation_PositionComponent)
 			SHADER_PARAMETER(int, Permutation_ColorComponent)
+			SHADER_PARAMETER(int, Permutation_DynamicMaterialParameter0Component)
 			SHADER_PARAMETER(int, Permutation_MeshOrientationComponent)
 			SHADER_PARAMETER(int, Permutation_RibbonWidthComponent)
 			SHADER_PARAMETER(int, Permutation_ScaleComponent)
diff --git a/Engine/Plugins/FX/Niagara/Source/NiagaraShader/Internal/Stateless/NiagaraStatelessModuleShaderParameters.h b/Engine/Plugins/FX/Niagara/Source/NiagaraShader/Internal/Stateless/NiagaraStatelessModuleShaderParameters.h
--- a/Engine/Plugins/FX/Niagara/Source/NiagaraShader/Internal/Stateless/NiagaraStatelessModuleShaderParameters.h
+++ b/Engine/Plugins/FX/Niagara/Source/NiagaraShader/Internal/Stateless/NiagaraStatelessModuleShaderParameters.h
@@ -43,7 +43,15 @@ namespace NiagaraStateless
 		SHADER_PARAMETER(float,	RotateAroundPoint_InitialPhaseScale)
 		SHADER_PARAMETER(float,	RotateAroundPoint_InitialPhaseBias)
 	END_SHADER_PARAMETER_STRUCT()
-	
+
+	BEGIN_SHADER_PARAMETER_STRUCT(FDynamicMaterialParametersModule_ShaderParameters, NIAGARASHADER_API)
+		SHADER_PARAMETER(uint32,		DynamicMaterialParameters_ChannelMask)
+		SHADER_PARAMETER(FUintVector3,	DynamicMaterialParameters_Parameter0X)
+		SHADER_PARAMETER(FUintVector3,	DynamicMaterialParameters_Parameter0Y)
+		SHADER_PARAMETER(FUintVector3,	DynamicMaterialParameters_Parameter0Z)
+		SHADER_PARAMETER(FUintVector3,	DynamicMaterialParameters_Parameter0W)
+	END_SHADER_PARAMETER_STRUCT()
+
 	BEGIN_SHADER_PARAMETER_STRUCT(FScaleColorModule_ShaderParameters, NIAGARASHADER_API)
 		SHADER_PARAMETER(FUintVector3,	ScaleColor_Distribution)
 	END_SHADER_PARAMETER_STRUCT()
