fix crash when unmuting cache tracks- skip ResetData if grid is not allocated

diff --git a/Engine/Plugins/FX/Niagara/Source/Niagara/Private/NiagaraDataInterfaceRasterizationGrid3D.cpp b/Engine/Plugins/FX/Niagara/Source/Niagara/Private/NiagaraDataInterfaceRasterizationGrid3D.cpp
--- a/Engine/Plugins/FX/Niagara/Source/Niagara/Private/NiagaraDataInterfaceRasterizationGrid3D.cpp
+++ b/Engine/Plugins/FX/Niagara/Source/Niagara/Private/NiagaraDataInterfaceRasterizationGrid3D.cpp
@@ -1086,7 +1086,7 @@ void UNiagaraDataInterfaceRasterizationGrid3D::DestroyPerInstanceData(void* PerI
 void FNiagaraDataInterfaceProxyRasterizationGrid3D::ResetData(const FNDIGpuComputeResetContext& Context)
 {	
 	RasterizationGrid3DRWInstanceData* ProxyData = SystemInstancesToProxyData.Find(Context.GetSystemInstanceID());
-	if (!ProxyData)
+	if (!ProxyData || ProxyData->NeedsRealloc || !ProxyData->RasterizationTexture.IsValid())
 	{
 		return;
 	}
@@ -1108,7 +1108,7 @@ void FNiagaraDataInterfaceProxyRasterizationGrid3D::PreStage(const FNDIGpuComput
 	const uint32 NumTotalCells = InstanceData.NumCells.X * InstanceData.NumCells.Y * InstanceData.NumCells.Z * InstanceData.TotalNumAttributes;
 
 	// Resize was requested
-	if (InstanceData.NeedsRealloc)
+	if (InstanceData.NeedsRealloc || !InstanceData.RasterizationTexture.IsValid())
 	{
 		InstanceData.NeedsRealloc = false;
 		
