Cleaned up some code & shuffled property order in asset tags, removed some properties from editor as they aren't actively used anymore. Will deprecate & remove in the future

diff --git a/Engine/Plugins/FX/Niagara/Source/Niagara/Classes/NiagaraEmitter.h b/Engine/Plugins/FX/Niagara/Source/Niagara/Classes/NiagaraEmitter.h
--- a/Engine/Plugins/FX/Niagara/Source/Niagara/Classes/NiagaraEmitter.h
+++ b/Engine/Plugins/FX/Niagara/Source/Niagara/Classes/NiagaraEmitter.h
@@ -704,18 +704,18 @@ private:
 public:
 	
 #if WITH_EDITORONLY_DATA
-	UPROPERTY(EditAnywhere, Category = "Asset Options", meta=(ShowOnlyInnerProperties))
+	UPROPERTY()
 	TArray<FNiagaraAssetTagDefinitionReference> AssetTags;
 
 	/** If an emitter is inheritable, new emitters based on an inheritable emitter, or Niagara Systems using an inheritable emitter, will automatically inherit changes made to the original emitter. */
 	UPROPERTY(EditAnywhere, Category = "Asset Options", AssetRegistrySearchable)
 	bool bIsInheritable = true;
 	
-	UPROPERTY(EditAnywhere, AdvancedDisplay, Category = "Asset Options", AssetRegistrySearchable)
+	UPROPERTY(EditAnywhere, AdvancedDisplay, Category = "Asset Options", AssetRegistrySearchable, DisplayName="Asset Description")
 	FText TemplateAssetDescription;
 
 	/** Category to collate this emitter into for "add new emitter" dialogs.*/
-	UPROPERTY(AssetRegistrySearchable, EditAnywhere, Category = Script)
+	UPROPERTY(AssetRegistrySearchable)
 	FText Category;
 
 	///** The thumbnail image used for the asset. This is always the latest recorded thumbnail. This can be different from the thumbnails that are saved per emitter version in collapsed view. */
@@ -723,7 +723,7 @@ public:
 	TObjectPtr<UTexture2D> ThumbnailImage;
 	
 	/** If this emitter is exposed to the library, or should be explicitly hidden. */
-    UPROPERTY(EditAnywhere, AdvancedDisplay, Category = "Asset Options", AssetRegistrySearchable)
+    UPROPERTY(AssetRegistrySearchable)
     ENiagaraScriptLibraryVisibility LibraryVisibility = ENiagaraScriptLibraryVisibility::Unexposed;
 
 	/** This is used as a transient value to open a specific version in the editor */
diff --git a/Engine/Plugins/FX/Niagara/Source/Niagara/Classes/NiagaraSystem.h b/Engine/Plugins/FX/Niagara/Source/Niagara/Classes/NiagaraSystem.h
--- a/Engine/Plugins/FX/Niagara/Source/Niagara/Classes/NiagaraSystem.h
+++ b/Engine/Plugins/FX/Niagara/Source/Niagara/Classes/NiagaraSystem.h
@@ -448,7 +448,7 @@ public:
 	bool bExposeToLibrary_DEPRECATED;
 
 	/** If this system is exposed to the library, or should be explicitly hidden. */
-	UPROPERTY(EditAnywhere, AdvancedDisplay, Category = "Asset Options", AssetRegistrySearchable, meta = (SkipSystemResetOnChange = "true"))
+	UPROPERTY(AssetRegistrySearchable, meta = (SkipSystemResetOnChange = "true"))
 	ENiagaraScriptLibraryVisibility LibraryVisibility;
 	
 	/** Deprecated template asset bool. Use the TemplateSpecification enum instead. */
@@ -459,14 +459,14 @@ public:
 	UPROPERTY()
 	ENiagaraScriptTemplateSpecification TemplateSpecification_DEPRECATED;
 
-	UPROPERTY(EditAnywhere, AdvancedDisplay, Category = "Asset Options", AssetRegistrySearchable, meta = (SkipSystemResetOnChange = "true"))
+	UPROPERTY(EditAnywhere, AdvancedDisplay, Category = "Asset Options", AssetRegistrySearchable, DisplayName="Asset Description", meta = (SkipSystemResetOnChange = "true"))
 	FText TemplateAssetDescription;
 
 	/** Category of this system. */
-	UPROPERTY(EditAnywhere, AdvancedDisplay, Category = "Asset Options", AssetRegistrySearchable, meta = (SkipSystemResetOnChange = "true"))
+	UPROPERTY(AssetRegistrySearchable, meta = (SkipSystemResetOnChange = "true"))
 	FText Category;
 
-	UPROPERTY(EditAnywhere, AdvancedDisplay, Category = "Asset Options", meta = (SkipSystemResetOnChange = "true"))
+	UPROPERTY(meta = (SkipSystemResetOnChange = "true"))
 	TArray<FNiagaraAssetTagDefinitionReference> AssetTags;
 
 	UPROPERTY()
diff --git a/Engine/Plugins/FX/Niagara/Source/Niagara/Public/NiagaraAssetTagDefinitions.h b/Engine/Plugins/FX/Niagara/Source/Niagara/Public/NiagaraAssetTagDefinitions.h
--- a/Engine/Plugins/FX/Niagara/Source/Niagara/Public/NiagaraAssetTagDefinitions.h
+++ b/Engine/Plugins/FX/Niagara/Source/Niagara/Public/NiagaraAssetTagDefinitions.h
@@ -107,7 +107,7 @@ public:
 		return AssetTagDefinitionGuid == Other.AssetTagDefinitionGuid;
 	}
 private:
-	UPROPERTY()
+	UPROPERTY(VisibleAnywhere, Category="Asset Tags")
 	FGuid AssetTagDefinitionGuid;
 };
 
@@ -129,17 +129,17 @@ public:
 
 	int32 GetSortOrder() const { return SortOrder; }
 private:
-	UPROPERTY(EditAnywhere, Category="Properties", meta=(TitleProperty="AssetTag"))
-	TArray<FNiagaraAssetTagDefinition> TagDefinitions;
-
 	/** The display name to use when listing this asset in the Niagara Asset Browser */
 	UPROPERTY(EditAnywhere, Category="Properties")
 	FText DisplayName;
-
+	
 	/** A description for this group of tags. Used for tooltips. */
 	UPROPERTY(EditAnywhere, Category="Properties")
 	FText Description;
 	
+	UPROPERTY(EditAnywhere, Category="Properties", meta=(TitleProperty="AssetTag"))
+	TArray<FNiagaraAssetTagDefinition> TagDefinitions;
+	
 	/** If true, no 'parent' entry for this asset will be displayed in the Niagara Asset Browser. Instead a flat list of the contained tags will be added. */
 	UPROPERTY(EditAnywhere, Category="Properties")
 	bool bDisplayTagsAsFlatList = false;
