Added a validation rule to make sure some modules are only used once per script

#virtualized

diff --git a/Engine/Plugins/FX/Niagara/Source/NiagaraEditor/Private/NiagaraValidationRules.cpp b/Engine/Plugins/FX/Niagara/Source/NiagaraEditor/Private/NiagaraValidationRules.cpp
--- a/Engine/Plugins/FX/Niagara/Source/NiagaraEditor/Private/NiagaraValidationRules.cpp
+++ b/Engine/Plugins/FX/Niagara/Source/NiagaraEditor/Private/NiagaraValidationRules.cpp
@@ -1239,4 +1239,52 @@ void UNiagaraValidationRule_UserDataInterfaces::CheckValidity(const FNiagaraVali
 	}
 }
 
+void UNiagaraValidationRule_SingletonModule::CheckValidity(const FNiagaraValidationContext& Context, TArray<FNiagaraValidationResult>& OutResults) const
+{
+	// check to see if we're called from a module
+	if (UNiagaraStackModuleItem* SourceModule = Cast<UNiagaraStackModuleItem>(Context.Source))
+	{
+		if (SourceModule->GetIsEnabled())
+		{
+			UNiagaraScript* ModuleScript = SourceModule->GetModuleNode().FunctionScript;
+			TArray<UNiagaraStackModuleItem*> StackModuleItems =	NiagaraValidation::GetAllStackEntriesInSystem<UNiagaraStackModuleItem>(Context.ViewModel);
+			for (UNiagaraStackModuleItem* Module : StackModuleItems)
+			{
+				// if another module in the same stack calls the same script, report it
+				if (Module && Module != SourceModule && Module->GetIsEnabled() && Module->GetModuleNode().FunctionScript == ModuleScript && SourceModule->GetEmitterViewModel().Get() == Module->GetEmitterViewModel().Get())
+				{
+					if (bCheckDetailedUsageContext)
+					{
+						ENiagaraScriptUsage ModuleAUsage = FNiagaraStackGraphUtilities::GetOutputNodeUsage(SourceModule->GetModuleNode());
+						ENiagaraScriptUsage ModuleBUsage = FNiagaraStackGraphUtilities::GetOutputNodeUsage(Module->GetModuleNode());
+						if (ModuleAUsage != ModuleBUsage)
+						{
+							continue;
+						}
+					}
+					OutResults.Emplace_GetRef(
+						Severity,
+						LOCTEXT("SingletonModuleError", "Module can only be used once per stack"),
+						LOCTEXT("SingletonModuleErrorDetailed", "This module is intended to be used as a singleton, so only once per emitter or system stack.\nThis is usually the case when there is a data dependency between modules because they share written attributes."),
+						SourceModule
+					).Fixes.Emplace(
+						FNiagaraValidationFix(
+							LOCTEXT("SingletonModuleErrorFix", "Disable module"),
+							FNiagaraValidationFixDelegate::CreateLambda(
+								[WeakSourceModule=MakeWeakObjectPtr(SourceModule)]()
+								{
+									if (UNiagaraStackModuleItem* StackModule = WeakSourceModule.Get())
+									{
+										StackModule->SetEnabled(false);
+									}
+								}
+							)
+						)
+					);
+				}
+			}
+		}
+	}
+}
+
 #undef LOCTEXT_NAMESPACE
diff --git a/Engine/Plugins/FX/Niagara/Source/NiagaraEditor/Public/NiagaraValidationRules.h b/Engine/Plugins/FX/Niagara/Source/NiagaraEditor/Public/NiagaraValidationRules.h
--- a/Engine/Plugins/FX/Niagara/Source/NiagaraEditor/Public/NiagaraValidationRules.h
+++ b/Engine/Plugins/FX/Niagara/Source/NiagaraEditor/Public/NiagaraValidationRules.h
@@ -287,3 +287,21 @@ public:
 	UPROPERTY(EditAnywhere, Category = Validation)
 	ENiagaraValidationSeverity Severity = ENiagaraValidationSeverity::Warning;
 };
+
+/** This validation rule checks that a module is only used once per emitter/system stack. */
+UCLASS(Category = "Validation", DisplayName = "Singleton Module")
+class UNiagaraValidationRule_SingletonModule : public UNiagaraValidationRule
+{
+	GENERATED_BODY()
+
+public:
+	virtual void CheckValidity(const FNiagaraValidationContext& Context, TArray<FNiagaraValidationResult>& OutResults) const override;
+
+	/** How do we want to repro the error in the stack */
+	UPROPERTY(EditAnywhere, Category = Validation)
+	ENiagaraValidationSeverity Severity = ENiagaraValidationSeverity::Warning;
+
+	/** If true then the check is not emitter-wide, but only within the same context (e.g. particle update). */
+	UPROPERTY(EditAnywhere, Category = Validation)
+	bool bCheckDetailedUsageContext = false;
+};
\ No newline at end of file
