Fix Niagara compile error caused by mismached preprocessors
#rnx

diff --git a/Engine/Plugins/FX/Niagara/Source/Niagara/Private/NiagaraDataChannel.cpp b/Engine/Plugins/FX/Niagara/Source/Niagara/Private/NiagaraDataChannel.cpp
--- a/Engine/Plugins/FX/Niagara/Source/Niagara/Private/NiagaraDataChannel.cpp
+++ b/Engine/Plugins/FX/Niagara/Source/Niagara/Private/NiagaraDataChannel.cpp
@@ -750,7 +750,7 @@ int32 FNiagaraDataChannelData::ConsumePublishRequests(UNiagaraDataChannelHandler
 
 	CPUSimData->EndSimulate();
 
-#if WITH_NIAGARA_DEBUGGER
+#if WITH_NIAGARA_DEBUGGER && !UE_BUILD_SHIPPING
 	if(Owner->GetDataChannel()->GetVerboseLogging())
 	{
 		FString Label = FString::Printf(TEXT("Data Channel %s - CURR"), *Owner->GetDataChannel()->GetName());
