Update normal extraction from (non VT) landscape with Niagara DI.  Not all platforms seem to handle SampleLevel with an explicit index offset, so stole the implementation from the landscape PS using a pair of gathers instead.

#virtualized

diff --git a/Engine/Plugins/FX/Niagara/Shaders/Private/NiagaraDataInterfaceLandscape.ush b/Engine/Plugins/FX/Niagara/Shaders/Private/NiagaraDataInterfaceLandscape.ush
--- a/Engine/Plugins/FX/Niagara/Shaders/Private/NiagaraDataInterfaceLandscape.ush
+++ b/Engine/Plugins/FX/Niagara/Shaders/Private/NiagaraDataInterfaceLandscape.ush
@@ -5,6 +5,7 @@
 #endif
 
 #include "/Engine/Private/VirtualTextureCommon.ush"
+#include "/Engine/Private/Landscape/LandscapeCommon.ush"
 
 // Mirrored from FNiagaraDataInterfaceParametersCS_Landscape
 #define VIRTUAL_TEXTURE_UNPACK_NORMAL_NONE 0
@@ -50,7 +51,8 @@ float3				{ParameterName}_CachedHeightTextureLWCTile;
 float4x4			{ParameterName}_CachedHeightTextureWorldToUvTransform;
 float4x4			{ParameterName}_CachedHeightTextureUvToWorldTransform;
 float4				{ParameterName}_CachedHeightTextureUvScaleBias;
-float2				{ParameterName}_CachedHeightTextureGridSize;
+float2				{ParameterName}_CachedHeightTextureWorldGridSize;
+int2				{ParameterName}_CachedHeightTextureDimension;
 SamplerState		{ParameterName}_PointClampedSampler;
 Texture2D<uint>		{ParameterName}_CachedPhysMatTexture;
 int2				{ParameterName}_CachedPhysMatTextureDimension;
@@ -202,23 +204,71 @@ void GetWorldNormal_{ParameterName}(float3 InWorldPos, out float3 OutNormal, out
 	}
 	else if ({ParameterName}_CachedHeightTextureEnabled != 0)
 	{
-		float2 Uv = GetCachedHeightTextureUv_{ParameterName}(InWorldPos);
-
-		float Left = {ParameterName}_CachedHeightTexture.SampleLevel({ParameterName}_PointClampedSampler, Uv, 0, int2(-1,0)).x;
-		float Right = {ParameterName}_CachedHeightTexture.SampleLevel({ParameterName}_PointClampedSampler, Uv, 0, int2(1,0)).x;
-		float Top = {ParameterName}_CachedHeightTexture.SampleLevel({ParameterName}_PointClampedSampler, Uv, 0, int2(0,1)).x;
-		float Bottom = {ParameterName}_CachedHeightTexture.SampleLevel({ParameterName}_PointClampedSampler, Uv, 0, int2(0,-1)).x;
-
-		// cross product of our two gradient vectors
-		OutNormal = float3(
-			{ParameterName}_CachedHeightTextureGridSize.y * (Left - Right),
-			{ParameterName}_CachedHeightTextureGridSize.x * (Bottom - Top),
-			2.0f * {ParameterName}_CachedHeightTextureGridSize.x * {ParameterName}_CachedHeightTextureGridSize.y);
-
-		// next rotate the normal to make it world space
-		OutNormal = mul(float4(OutNormal, 0.0f), {ParameterName}_CachedHeightTextureUvToWorldTransform).xyz;
-		OutNormal = normalize(OutNormal);
-
+		float2 InTextureCoordinates = GetCachedHeightTextureUv_{ParameterName}(InWorldPos);
+		float2 TextureSize = {ParameterName}_CachedHeightTextureDimension.xy;
+
+		// taken from LandscapeLayersPS.usf -> CalculateNormalsFromHeights()
+		float2 TexelSize = 1.0 / TextureSize;
+
+		float MaskMinBorderX = (InTextureCoordinates.x <= TexelSize.x) ? 0.0f : 1.0f;			// first pixel X
+		float MaskMinBorderY = (InTextureCoordinates.y <= TexelSize.y) ? 0.0f : 1.0f;			// first pixel Y
+		float MaskMaxBorderX = (InTextureCoordinates.x >= 1.0f - TexelSize.x) ? 0.0f : 1.0f;	// last pixel X
+		float MaskMaxBorderY = (InTextureCoordinates.y >= 1.0f - TexelSize.y) ? 0.0f : 1.0f;	// last pixel Y
+
+		// The triangle topology is the following (where C = center, T = top, B = bottom, L = left, R = right and Nx the normals we need to interpolate):
+		// TL ------ TT
+		// | \       | \
+		// |  \      |  \
+		// |   \     |   \
+		// | N0 \ N1 | N3 \
+		// |     \   |     \
+		// |      \  |      \
+		// |       \ |       \
+		// LL ------ CC ------ RR
+		//   \       | \       |
+		//    \      |  \      |
+		//     \     |   \     |
+		//      \ N2 | N4 \ N5 |
+		//       \   |     \   |
+		//        \  |      \  |
+		//         \ |       \ |
+		//           BB ------ BR
+		// So we can retrieve all 7 samples using 4 gathers : 
+		// Make sure we gather the 4 pixels in the bottom right direction by sampling at (0.75, 0.75) of the pixel (gather retrieves the 4 samples that would be used for bilinear interpolation):
+
+		float2 GatherLocation = InTextureCoordinates + 0.25f * TexelSize; 
+		float4 Red0 = {ParameterName}_CachedHeightTexture.GatherRed({ParameterName}_CachedHeightTextureSampler, GatherLocation - TexelSize);
+		float4 Red1 = {ParameterName}_CachedHeightTexture.GatherRed({ParameterName}_CachedHeightTextureSampler, GatherLocation);
+
+		// [x,y,z] = World Position [w] = valid mask
+		float4 TL = float4(float2(-1, -1) * {ParameterName}_CachedHeightTextureWorldGridSize, Red0.w, 1.0f);
+		float4 TT = float4(float2(+0, -1) * {ParameterName}_CachedHeightTextureWorldGridSize, Red0.z, 1.0f);
+		float4 CC = float4(float2(+0, +0) * {ParameterName}_CachedHeightTextureWorldGridSize, Red0.y, 1.0f);
+		float4 LL = float4(float2(-1, +0) * {ParameterName}_CachedHeightTextureWorldGridSize, Red0.x, 1.0f);
+		float4 RR = float4(float2(+1, +0) * {ParameterName}_CachedHeightTextureWorldGridSize, Red1.z, 1.0f);
+		float4 BR = float4(float2(+1, +1) * {ParameterName}_CachedHeightTextureWorldGridSize, Red1.y, 1.0f);
+		float4 BB = float4(float2(+0, +1) * {ParameterName}_CachedHeightTextureWorldGridSize, Red1.x, 1.0f);
+
+		// mask out samples that are off the edge of the texture
+		TL.w *= MaskMinBorderX * MaskMinBorderY;
+		TT.w *=                  MaskMinBorderY;
+		// CC should never be off the edge
+		LL.w *= MaskMinBorderX;
+		RR.w *= MaskMaxBorderX;
+		BR.w *= MaskMaxBorderX * MaskMaxBorderY;
+		BB.w *=                  MaskMaxBorderY;
+
+		float3 N0 = ComputeNullableTriangleNormal(CC, LL, TL);
+		float3 N1 = ComputeNullableTriangleNormal(TL, TT, CC);
+		float3 N2 = ComputeNullableTriangleNormal(LL, CC, BB);
+		float3 N3 = ComputeNullableTriangleNormal(RR, CC, TT);
+		float3 N4 = ComputeNullableTriangleNormal(BR, BB, CC);
+		float3 N5 = ComputeNullableTriangleNormal(CC, RR, BR);
+
+		// average the normals
+		float3 FinalNormal = normalize(N0 + N1 + N2 + N3 + N4 + N5);
+
+		OutNormal = normalize(mul(float4(FinalNormal, 0.0f), {ParameterName}_CachedHeightTextureUvToWorldTransform).xyz);
 		OutIsValid = true;
 	}
 	else
diff --git a/Engine/Plugins/FX/Niagara/Source/Niagara/Classes/NiagaraDataInterfaceLandscape.h b/Engine/Plugins/FX/Niagara/Source/Niagara/Classes/NiagaraDataInterfaceLandscape.h
--- a/Engine/Plugins/FX/Niagara/Source/Niagara/Classes/NiagaraDataInterfaceLandscape.h
+++ b/Engine/Plugins/FX/Niagara/Source/Niagara/Classes/NiagaraDataInterfaceLandscape.h
@@ -4,6 +4,7 @@
 
 #include "NiagaraDataInterface.h"
 #include "NiagaraCommon.h"
+#include "NiagaraSimCacheCustomStorageInterface.h"
 #include "VectorVM.h"
 
 #include "NiagaraDataInterfaceLandscape.generated.h"
@@ -35,7 +36,7 @@ enum class ENDILandscape_SourceMode : uint8
 
 /** Data Interface allowing sampling of a Landscape */
 UCLASS(EditInlineNew, Category = "Landscape", CollapseCategories, meta = (DisplayName = "Landscape Sample"), MinimalAPI)
-class UNiagaraDataInterfaceLandscape : public UNiagaraDataInterface
+class UNiagaraDataInterfaceLandscape : public UNiagaraDataInterface, public INiagaraSimCacheCustomStorageInterface
 {
 	GENERATED_UCLASS_BODY()
 
@@ -62,9 +63,15 @@ public:
 #if WITH_EDITORONLY_DATA
 	NIAGARA_API virtual bool UpgradeFunctionCall(FNiagaraFunctionSignature& FunctionSignature) override;
 #endif
+	NIAGARA_API virtual bool Equals(const UNiagaraDataInterface* Other) const override;
 	//UNiagaraDataInterface Interface End
 
-	NIAGARA_API virtual bool Equals(const UNiagaraDataInterface* Other) const override;
+	//~ INiagaraSimCacheCustomStorageInterface interface BEGIN
+	NIAGARA_API virtual UObject* SimCacheBeginWrite(UObject* SimCache, FNiagaraSystemInstance* NiagaraSystemInstance, const void* OptionalPerInstanceData, FNiagaraSimCacheFeedbackContext& FeedbackContext) const override;
+	NIAGARA_API virtual bool SimCacheWriteFrame(UObject* StorageObject, int FrameIndex, FNiagaraSystemInstance* SystemInstance, const void* OptionalPerInstanceData, FNiagaraSimCacheFeedbackContext& FeedbackContext) const override;
+	//NIAGARA_API virtual bool SimCacheEndWrite(UObject* StorageObject) const override;
+	NIAGARA_API virtual bool SimCacheReadFrame(UObject* StorageObject, int FrameA, int FrameB, float Interp, FNiagaraSystemInstance* SystemInstance, void* OptionalPerInstanceData) override;
+	//~ UNiagaraDataInterface interface END
 
 	// GPU sim functionality
 #if WITH_EDITORONLY_DATA
@@ -95,3 +102,14 @@ protected:
 	static NIAGARA_API const FName GetWorldNormalName;
 	static NIAGARA_API const FName GetPhysicalMaterialIndexName;
 };
+
+
+UCLASS(MinimalAPI)
+class UNDILandscapeSimCacheData : public UObject
+{
+	GENERATED_BODY()
+
+public:
+	UPROPERTY()
+	TArray<TObjectPtr<UTexture2D>> HeightFieldTextures;
+};
diff --git a/Engine/Plugins/FX/Niagara/Source/Niagara/Private/NiagaraDataInterfaceLandscape.cpp b/Engine/Plugins/FX/Niagara/Source/Niagara/Private/NiagaraDataInterfaceLandscape.cpp
--- a/Engine/Plugins/FX/Niagara/Source/Niagara/Private/NiagaraDataInterfaceLandscape.cpp
+++ b/Engine/Plugins/FX/Niagara/Source/Niagara/Private/NiagaraDataInterfaceLandscape.cpp
@@ -88,7 +88,8 @@ namespace NiagaraDataInterfaceLandscape
 		SHADER_PARAMETER(FMatrix44f,				CachedHeightTextureWorldToUvTransform)
 		SHADER_PARAMETER(FMatrix44f,				CachedHeightTextureUvToWorldTransform)
 		SHADER_PARAMETER(FVector4f,					CachedHeightTextureUvScaleBias)
-		SHADER_PARAMETER(FVector2f,					CachedHeightTextureGridSize)
+		SHADER_PARAMETER(FVector2f,					CachedHeightTextureWorldGridSize)
+		SHADER_PARAMETER(FIntPoint,					CachedHeightTextureDimension)
 		SHADER_PARAMETER_SAMPLER(SamplerState,		PointClampedSampler)
 		SHADER_PARAMETER_TEXTURE(Texture2D<uint>,	CachedPhysMatTexture)
 		SHADER_PARAMETER(FIntPoint,					CachedPhysMatTextureDimension)
@@ -255,7 +256,7 @@ struct FNDILandscapeData_GameToRenderThread
 	FMatrix CachedHeightTextureWorldToUvTransform = FMatrix::Identity;
 	FMatrix CachedHeightTextureUvToWorldTransform = FMatrix::Identity;
 	FVector4 CachedHeightTextureUvScaleBias = FVector4::Zero();
-	FVector2D CachedHeightTextureGridSize = FVector2D(1.0f, 1.0f);
+	FVector2D CachedHeightTextureWorldGridSize = FVector2D(1.0f, 1.0f);
 
 	bool BaseColorVirtualTextureSRGB = false;
 };
@@ -508,34 +509,37 @@ struct FNDILandscapeData_RenderThread
 
 			if (HeightTexture || PhysMatTexture)
 			{
+				const FIntPoint TextureDimensions(LandscapeData.TextureResources->GetDimensions());
+
 				ShaderParameters->CachedHeightTextureLWCTile = LandscapeData.CachedHeightTextureLWCTile;
 				ShaderParameters->CachedHeightTextureWorldToUvTransform = (FMatrix44f)LandscapeData.CachedHeightTextureWorldToUvTransform;
 				ShaderParameters->CachedHeightTextureUvToWorldTransform = (FMatrix44f)LandscapeData.CachedHeightTextureUvToWorldTransform;
 				ShaderParameters->CachedHeightTextureUvScaleBias = (FVector4f)LandscapeData.CachedHeightTextureUvScaleBias;
-				ShaderParameters->CachedHeightTextureGridSize = (FVector2f)LandscapeData.CachedHeightTextureGridSize;
+				ShaderParameters->CachedHeightTextureWorldGridSize = (FVector2f)LandscapeData.CachedHeightTextureWorldGridSize;
 
 				if (HeightTexture)
 				{
 					ShaderParameters->CachedHeightTexture = HeightTexture;
 					ShaderParameters->CachedHeightTextureEnabled = 1;
+					ShaderParameters->CachedHeightTextureDimension = TextureDimensions;
 				}
 				else
 				{
 					ShaderParameters->CachedHeightTexture = GBlackTexture->TextureRHI;
 					ShaderParameters->CachedHeightTextureEnabled = 0;
+					ShaderParameters->CachedHeightTextureDimension = FIntPoint(ForceInitToZero);
 				}
 
 				if (PhysMatTexture)
 				{
-					const FIntPoint PhysMatDimensions(LandscapeData.TextureResources->GetDimensions());
 					ShaderParameters->CachedPhysMatTexture = PhysMatTexture;
-					ShaderParameters->CachedPhysMatTextureDimension = PhysMatDimensions;
+					ShaderParameters->CachedPhysMatTextureDimension = TextureDimensions;
 				}
 				else
 				{
 					const FIntPoint PhysMatDimensions(ForceInitToZero);
 					ShaderParameters->CachedPhysMatTexture = GBlackUintTexture->TextureRHI;
-					ShaderParameters->CachedPhysMatTextureDimension = PhysMatDimensions;
+					ShaderParameters->CachedPhysMatTextureDimension = FIntPoint(ForceInitToZero);
 				}
 
 				return true;
@@ -554,7 +558,7 @@ struct FNDILandscapeData_RenderThread
 		ShaderParameters->CachedHeightTextureWorldToUvTransform = FMatrix44f::Identity;
 		ShaderParameters->CachedHeightTextureUvToWorldTransform = FMatrix44f::Identity;
 		ShaderParameters->CachedHeightTextureUvScaleBias = DummyVector4;
-		ShaderParameters->CachedHeightTextureGridSize = FVector2f::ZeroVector;
+		ShaderParameters->CachedHeightTextureWorldGridSize = FVector2f::ZeroVector;
 		ShaderParameters->CachedHeightTextureEnabled = 0;
 
 		ShaderParameters->CachedPhysMatTexture = GBlackUintTexture->TextureRHI;
@@ -1176,6 +1180,244 @@ bool UNiagaraDataInterfaceLandscape::Equals(const UNiagaraDataInterface* Other)
 		&& OtherLandscape->bVirtualTexturesSupported == bVirtualTexturesSupported;
 }
 
+UObject* UNiagaraDataInterfaceLandscape::SimCacheBeginWrite(UObject* SimCache, FNiagaraSystemInstance* NiagaraSystemInstance, const void* OptionalPerInstanceData, FNiagaraSimCacheFeedbackContext& FeedbackContext) const
+{
+	return NewObject<UNDILandscapeSimCacheData>(SimCache);
+}
+
+bool UNiagaraDataInterfaceLandscape::SimCacheWriteFrame(UObject* StorageObject, int FrameIndex, FNiagaraSystemInstance* SystemInstance, const void* PerInstanceData, FNiagaraSimCacheFeedbackContext& FeedbackContext) const
+{
+	if (UNDILandscapeSimCacheData* SimCacheData = Cast<UNDILandscapeSimCacheData>(StorageObject))
+	{
+		SimCacheData->HeightFieldTextures.SetNum(1 + FrameIndex);
+
+		const FNDILandscapeData_GameThread& SourceData = *reinterpret_cast<const FNDILandscapeData_GameThread*>(PerInstanceData);
+
+		if (SourceData.SharedResourceHandle)
+		{
+			const FNDI_Landscape_SharedResource& SourceResource = SourceData.SharedResourceHandle.ReadResource();
+			if (const FRHITexture* HeightTexture = SourceResource.LandscapeTextures->GetHeightTexture())
+			{
+				UTexture2D* FrameTexture = NewObject<UTexture2D>(SimCacheData);
+
+				// copy texture
+
+				SimCacheData->HeightFieldTextures[FrameIndex] = FrameTexture;
+			}
+		}
+	}
+
+	return true;
+}
+
+bool UNiagaraDataInterfaceLandscape::SimCacheReadFrame(UObject* StorageObject, int FrameA, int FrameB, float Interp, FNiagaraSystemInstance* SystemInstance, void* OptionalPerInstanceData)
+{
+	return true;
+	#if 0
+	FRenderTargetVolumeRWInstanceData_GameThread* InstanceData_GT = reinterpret_cast<FRenderTargetVolumeRWInstanceData_GameThread*>(OptionalPerInstanceData);
+
+	if (const* UTexture2D HeightFieldTexture = Cast<UTexture2D>(StorageObject))
+	{
+		
+	}
+
+	return true;
+	#endif
+	#if 0
+	else if (Cast<UAnimatedSparseVolumeTexture>(StorageObject))
+	{
+		UAnimatedSparseVolumeTexture* SVT = Cast<UAnimatedSparseVolumeTexture>(StorageObject);
+
+		const int32 MipLevel = 0;
+		const bool bBlocking = true;
+		USparseVolumeTextureFrame* SVTFrame = USparseVolumeTextureFrame::GetFrameAndIssueStreamingRequest(SVT, FMath::RoundToInt(FrameA + Interp), MipLevel, bBlocking);
+
+		// The streaming manager normally ticks in FDeferredShadingSceneRenderer::Render(), but the SVT->DenseTexture conversion compute shader happens in a render command before that.
+		// At execution time of that command, the streamer hasn't had the chance to do any streaming yet, so we force another tick here.
+		// Assuming blocking requests are used, this guarantees that the requested frame is fully streamed in (if there is memory available).
+		UE::SVT::GetStreamingManager().Update_GameThread();
+
+		FIntVector VolumeResolution = SVT->GetVolumeResolution();
+
+		InstanceData_GT->Size = VolumeResolution;
+		InstanceData_GT->Format = SVT->GetFormat(0);
+		InstanceData_GT->Filter = TextureFilter::TF_Default;
+
+		PerInstanceTick(InstanceData_GT, SystemInstance, 0.0f);
+		PerInstanceTickPostSimulate(InstanceData_GT, SystemInstance, 0.0f);
+
+		FNiagaraDataInterfaceProxyRenderTargetVolumeProxy* RT_Proxy = GetProxyAs<FNiagaraDataInterfaceProxyRenderTargetVolumeProxy>();
+
+
+		// Execute compute shader
+		ENQUEUE_RENDER_COMMAND(NDIRenderTargetVolumeUpdate)
+			(
+				[RT_Proxy, RT_InstanceID = SystemInstance->GetId(), RT_Format = InstanceData_GT->Format,
+				RT_VolumeResolution = VolumeResolution, RT_SVTRenderResources = SVTFrame ? SVTFrame->GetTextureRenderResources() : nullptr,
+				FeatureLevel = SystemInstance->GetFeatureLevel()](FRHICommandListImmediate& RHICmdList)
+				{
+					if (RT_SVTRenderResources == nullptr)
+					{
+						return;
+					}
+
+		FUintVector4 CurrentPackedUniforms0 = FUintVector4();
+		FUintVector4 CurrentPackedUniforms1 = FUintVector4();
+		RT_SVTRenderResources->GetPackedUniforms(CurrentPackedUniforms0, CurrentPackedUniforms1);
+
+		if (FRenderTargetVolumeRWInstanceData_RenderThread* InstanceData_RT = RT_Proxy->SystemInstancesToProxyData_RT.Find(RT_InstanceID))
+		{
+			FRDGBuilder GraphBuilder(RHICmdList);
+
+			FIntVector ThreadGroupSize = FNiagaraShader::GetDefaultThreadGroupSize(ENiagaraGpuDispatchType::ThreeD);
+
+			TShaderMapRef<FNiagaraCopySVTToDenseBufferCS> ComputeShader(GetGlobalShaderMap(FeatureLevel));
+
+			const FIntVector NumThreadGroups(
+				FMath::DivideAndRoundUp(RT_VolumeResolution.X, ThreadGroupSize.X),
+				FMath::DivideAndRoundUp(RT_VolumeResolution.Y, ThreadGroupSize.Y),
+				FMath::DivideAndRoundUp(RT_VolumeResolution.Z, ThreadGroupSize.Z)
+			);
+
+			FNiagaraCopySVTToDenseBufferCS::FParameters* PassParameters = GraphBuilder.AllocParameters<FNiagaraCopySVTToDenseBufferCS::FParameters>();
+
+			if (InstanceData_RT->RenderTarget.IsValid() && InstanceData_RT->TransientRDGTexture == nullptr)
+			{
+				InstanceData_RT->TransientRDGTexture = GraphBuilder.RegisterExternalTexture(InstanceData_RT->RenderTarget);
+				InstanceData_RT->TransientRDGUAV = GraphBuilder.CreateUAV(InstanceData_RT->TransientRDGTexture);
+
+				// #todo(dmp): needed?							
+				//Context.GetRDGExternalAccessQueue().Add(InstanceData_RT->TransientRDGTexture);
+			}
+
+			if (InstanceData_RT->RenderTarget.IsValid() && InstanceData_RT->TransientRDGUAV != nullptr)
+			{
+				PassParameters->DestinationBuffer = InstanceData_RT->TransientRDGUAV;
+			}
+			else
+			{
+				PassParameters->DestinationBuffer = GraphBuilder.CreateUAV(GraphBuilder.CreateTexture(
+					FRDGTextureDesc::Create3D(FIntVector(1, 1, 1), RT_Format, FClearValueBinding::Black, ETextureCreateFlags::ShaderResource | ETextureCreateFlags::UAV),
+					TEXT("NiagaraEmptyTextureUAV::Texture3D")
+				),
+					ERDGUnorderedAccessViewFlags::SkipBarrier
+				);
+			}
+
+			FRHITexture* PageTableTexture = RT_SVTRenderResources->GetPageTableTexture();
+			FRHITexture* TextureA = RT_SVTRenderResources->GetPhysicalTileDataATexture();
+
+			PassParameters->TileDataTextureSampler = TStaticSamplerState<SF_Point, AM_Clamp, AM_Clamp, AM_Clamp>::GetRHI();
+			PassParameters->SparseVolumeTexturePageTable = PageTableTexture ? PageTableTexture : GBlackUintVolumeTexture->TextureRHI.GetReference();
+			PassParameters->SparseVolumeTextureA = TextureA ? TextureA : GBlackVolumeTexture->TextureRHI.GetReference();
+
+			PassParameters->PackedSVTUniforms0 = CurrentPackedUniforms0;
+			PassParameters->PackedSVTUniforms1 = CurrentPackedUniforms1;
+
+			PassParameters->TextureSize = RT_VolumeResolution;
+			PassParameters->MipLevel = 0;
+
+			GraphBuilder.AddPass(
+				// Friendly name of the pass for profilers using printf semantics.
+				RDG_EVENT_NAME("Copy SVT to Volume RT"),
+				// Parameters provided to RDG.
+				PassParameters,
+				// Issues compute commands.
+				ERDGPassFlags::Compute,
+				// This is deferred until Execute. May execute in parallel with other passes.
+				[PassParameters, ComputeShader, NumThreadGroups](FRHIComputeCommandList& RHICmdList)
+				{
+					FComputeShaderUtils::Dispatch(RHICmdList, ComputeShader, *PassParameters, NumThreadGroups);
+				});
+
+			// Execute the graph.
+			GraphBuilder.Execute();
+
+			InstanceData_RT->TransientRDGTexture = nullptr;
+			InstanceData_RT->TransientRDGSRV = nullptr;
+			InstanceData_RT->TransientRDGUAV = nullptr;
+		}
+				});
+	}
+	else
+	{
+		UNDIRenderTargetVolumeSimCacheData* SimCacheData = CastChecked<UNDIRenderTargetVolumeSimCacheData>(StorageObject);
+
+		const int FrameIndex = Interp >= 0.5f ? FrameB : FrameA;
+		if (!SimCacheData->Frames.IsValidIndex(FrameIndex))
+		{
+			return false;
+		}
+
+		FNDIRenderTargetVolumeSimCacheFrame* CacheFrame = &SimCacheData->Frames[FrameIndex];
+
+		InstanceData_GT->Size = CacheFrame->Size;
+		InstanceData_GT->Format = CacheFrame->Format;
+		InstanceData_GT->Filter = TextureFilter::TF_Default;
+
+		PerInstanceTick(InstanceData_GT, SystemInstance, 0.0f);
+		PerInstanceTickPostSimulate(InstanceData_GT, SystemInstance, 0.0f);
+
+		if (InstanceData_GT->TargetTexture && CacheFrame->GetPixelData() != nullptr)
+		{
+			FNiagaraDataInterfaceProxyRenderTargetVolumeProxy* RT_Proxy = GetProxyAs<FNiagaraDataInterfaceProxyRenderTargetVolumeProxy>();
+			FTextureRenderTargetResource* RT_TargetTexture = InstanceData_GT->TargetTexture->GameThread_GetRenderTargetResource();
+			ENQUEUE_RENDER_COMMAND(NDIRenderTargetVolumeUpdate)
+				(
+					[RT_Proxy, RT_InstanceID = SystemInstance->GetId(), RT_TargetTexture, RT_CacheFrame = CacheFrame, RT_CompressionType = SimCacheData->CompressionType, RT_Format = InstanceData_GT->Format](FRHICommandListImmediate& RHICmdList)
+					{
+						if (FRenderTargetVolumeRWInstanceData_RenderThread* InstanceData_RT = RT_Proxy->SystemInstancesToProxyData_RT.Find(RT_InstanceID))
+						{
+							FUpdateTextureRegion3D UpdateRegion(FIntVector::ZeroValue, FIntVector::ZeroValue, InstanceData_RT->Size);
+
+							FUpdateTexture3DData UpdateTexture = RHICmdList.BeginUpdateTexture3D(InstanceData_RT->RenderTarget->GetRHI(), 0, UpdateRegion);
+
+							uint32 BlockBytes = GPixelFormats[RT_Format].BlockBytes;
+
+							const int32 SrcRowPitch = InstanceData_RT->Size.X * BlockBytes;
+							const int32 SrcDepthPitch = InstanceData_RT->Size.Y * SrcRowPitch;
+							if (UpdateTexture.RowPitch == SrcRowPitch && UpdateTexture.DepthPitch == SrcDepthPitch)
+							{
+								if (RT_CacheFrame->CompressedSize > 0)
+								{
+									FCompression::UncompressMemory(RT_CompressionType, UpdateTexture.Data, UpdateTexture.DataSizeBytes, RT_CacheFrame->GetPixelData(), RT_CacheFrame->CompressedSize);
+								}
+								else
+								{
+									FMemory::Memcpy(UpdateTexture.Data, RT_CacheFrame->GetPixelData(), InstanceData_RT->Size.X * InstanceData_RT->Size.Y * sizeof(FFloat16Color));
+								}
+							}
+							else
+							{
+								TArray<uint8> Decompressed;
+								if (RT_CacheFrame->CompressedSize > 0)
+								{
+									Decompressed.AddUninitialized(BlockBytes * InstanceData_RT->Size.X * InstanceData_RT->Size.Y * InstanceData_RT->Size.Z);
+									FCompression::UncompressMemory(RT_CompressionType, Decompressed.GetData(), Decompressed.Num(), RT_CacheFrame->GetPixelData(), RT_CacheFrame->CompressedSize);
+								}
+
+								const uint8* SrcData = Decompressed.Num() > 0 ? Decompressed.GetData() : RT_CacheFrame->GetPixelData();
+								for (int32 z = 0; z < InstanceData_RT->Size.Z; ++z)
+								{
+									uint8* DstData = UpdateTexture.Data + (z * UpdateTexture.DepthPitch);
+									for (int32 y = 0; y < InstanceData_RT->Size.Y; ++y)
+									{
+										FMemory::Memcpy(DstData, SrcData, SrcRowPitch);
+										SrcData += SrcRowPitch;
+										DstData += UpdateTexture.RowPitch;
+									}
+								}
+							}
+							RHICmdList.EndUpdateTexture3D(UpdateTexture);
+						}
+					}
+			);
+		}
+	}
+	return true;
+	#endif
+}
+
 #if WITH_EDITORONLY_DATA
 void UNiagaraDataInterfaceLandscape::GetFunctionsInternal(TArray<FNiagaraFunctionSignature>& OutFunctions) const
 {
@@ -1281,7 +1523,7 @@ void UNiagaraDataInterfaceLandscape::ProvidePerInstanceDataForRenderThread(void*
 		TargetData->CachedHeightTextureLWCTile = SourceResource.LandscapeLWCTile;
 		TargetData->CachedHeightTextureWorldToUvTransform = SourceResource.WorldToActorTransform;
 		TargetData->CachedHeightTextureUvToWorldTransform = SourceResource.ActorToWorldTransform;
-		TargetData->CachedHeightTextureGridSize = SourceResource.TextureWorldGridSize;
+		TargetData->CachedHeightTextureWorldGridSize = SourceResource.TextureWorldGridSize;
 	}
 
 	TargetData->BaseColorVirtualTextureSRGB = SourceData.BaseColorVirtualTextureSRGB;
