Made Niagara note header wrappable & reduced max width for note tooltip as it was just too wide

diff --git a/Engine/Plugins/FX/Niagara/Source/NiagaraEditor/Private/NiagaraEditorUtilities.cpp b/Engine/Plugins/FX/Niagara/Source/NiagaraEditor/Private/NiagaraEditorUtilities.cpp
--- a/Engine/Plugins/FX/Niagara/Source/NiagaraEditor/Private/NiagaraEditorUtilities.cpp
+++ b/Engine/Plugins/FX/Niagara/Source/NiagaraEditor/Private/NiagaraEditorUtilities.cpp
@@ -4660,12 +4660,12 @@ TSharedRef<SToolTip> FNiagaraEditorUtilities::Tooltips::CreateStackNoteTooltip(U
 			]
 		]
 		+ SHorizontalBox::Slot()
-		.AutoWidth()
 		.Padding(2.f)
 		[
 			SNew(STextBlock)
 			.Text(StackNote.GetTargetStackNoteData().GetValue().MessageHeader)
 			.TextStyle(&FNiagaraEditorStyle::Get().GetWidgetStyle<FTextBlockStyle>("NiagaraEditor.Stack.Note.HeaderText"))
+			.AutoWrapText(true)
 		]
 	];
 
@@ -4685,7 +4685,7 @@ TSharedRef<SToolTip> FNiagaraEditorUtilities::Tooltips::CreateStackNoteTooltip(U
 	.Content()
 	[
 		SNew(SBox)
-		.MaxDesiredWidth(750.f)
+		.MaxDesiredWidth(500.f)
 		[
 			TooltipContent
 		]
diff --git a/Engine/Plugins/FX/Niagara/Source/NiagaraEditorWidgets/Private/SNiagaraStackNote.cpp b/Engine/Plugins/FX/Niagara/Source/NiagaraEditorWidgets/Private/SNiagaraStackNote.cpp
--- a/Engine/Plugins/FX/Niagara/Source/NiagaraEditorWidgets/Private/SNiagaraStackNote.cpp
+++ b/Engine/Plugins/FX/Niagara/Source/NiagaraEditorWidgets/Private/SNiagaraStackNote.cpp
@@ -135,6 +135,7 @@ void SNiagaraStackNote::Rebuild()
 					.OnTextCommitted(this, &SNiagaraStackNote::CommitStackNoteHeaderUpdate)
 					.Style(&FNiagaraEditorStyle::Get().GetWidgetStyle<FInlineEditableTextBlockStyle>("NiagaraEditor.Stack.Note.HeaderEditableText"))
 					.IsSelected(FIsSelected::CreateLambda([]() { return false; }))
+					.AutoWrapNonEditText(true)
 				]
 				+ SHorizontalBox::Slot()
 				.AutoWidth()
