Allow static variables to be linked to non-static inputs of the same underlying type

diff --git a/Engine/Plugins/FX/Niagara/Source/Niagara/Private/NiagaraCommon.cpp b/Engine/Plugins/FX/Niagara/Source/Niagara/Private/NiagaraCommon.cpp
--- a/Engine/Plugins/FX/Niagara/Source/Niagara/Private/NiagaraCommon.cpp
+++ b/Engine/Plugins/FX/Niagara/Source/Niagara/Private/NiagaraCommon.cpp
@@ -1592,16 +1592,17 @@ void FNiagaraUtilities::PrepareRapidIterationParameters(const TArray<UNiagaraScr
 	}
 }
 
-bool FNiagaraUtilities::AreTypesAssignable(const FNiagaraTypeDefinition& TypeA, const FNiagaraTypeDefinition& TypeB)
+bool FNiagaraUtilities::AreTypesAssignable(const FNiagaraTypeDefinition& FromType, const FNiagaraTypeDefinition& ToType)
 {
 	const UNiagaraSettings* Settings = GetDefault<UNiagaraSettings>();
+	bool bStrictAssignable = (FromType == ToType) || (FromType.IsStatic() && !ToType.IsStatic() && FromType == ToType.ToStaticDef());
 	if (Settings->bEnforceStrictStackTypes)
 	{
-		return TypeA == TypeB;
+		return bStrictAssignable;
 	}
-	return (TypeA == TypeB)
-		|| (TypeA == FNiagaraTypeDefinition::GetPositionDef() && TypeB == FNiagaraTypeDefinition::GetVec3Def())
-		|| (TypeB == FNiagaraTypeDefinition::GetPositionDef() && TypeA == FNiagaraTypeDefinition::GetVec3Def());
+	return bStrictAssignable
+		|| (FromType == FNiagaraTypeDefinition::GetPositionDef() && ToType == FNiagaraTypeDefinition::GetVec3Def())
+		|| (ToType == FNiagaraTypeDefinition::GetPositionDef() && FromType == FNiagaraTypeDefinition::GetVec3Def());
 }
 
 #endif
diff --git a/Engine/Plugins/FX/Niagara/Source/Niagara/Public/NiagaraCommon.h b/Engine/Plugins/FX/Niagara/Source/Niagara/Public/NiagaraCommon.h
--- a/Engine/Plugins/FX/Niagara/Source/Niagara/Public/NiagaraCommon.h
+++ b/Engine/Plugins/FX/Niagara/Source/Niagara/Public/NiagaraCommon.h
@@ -1525,7 +1525,7 @@ namespace FNiagaraUtilities
 	 */
 	void NIAGARA_API PrepareRapidIterationParameters(const TArray<UNiagaraScript*>& Scripts, const TMap<UNiagaraScript*, UNiagaraScript*>& ScriptDependencyMap, const TMap<UNiagaraScript*, FVersionedNiagaraEmitter>& ScriptToEmitterNameMap);
 
-	bool NIAGARA_API AreTypesAssignable(const FNiagaraTypeDefinition& TypeA, const FNiagaraTypeDefinition& TypeB);
+	bool NIAGARA_API AreTypesAssignable(const FNiagaraTypeDefinition& FromType, const FNiagaraTypeDefinition& ToType);
 #endif
 
 	void NIAGARA_API DumpHLSLText(const FString& SourceCode, const FString& DebugName);
diff --git a/Engine/Plugins/FX/Niagara/Source/NiagaraEditor/Private/NiagaraEditorUtilities.cpp b/Engine/Plugins/FX/Niagara/Source/NiagaraEditor/Private/NiagaraEditorUtilities.cpp
--- a/Engine/Plugins/FX/Niagara/Source/NiagaraEditor/Private/NiagaraEditorUtilities.cpp
+++ b/Engine/Plugins/FX/Niagara/Source/NiagaraEditor/Private/NiagaraEditorUtilities.cpp
@@ -993,9 +993,9 @@ FText FNiagaraEditorUtilities::GetTypeDefinitionCategory(const FNiagaraTypeDefin
 	return Category;
 }
 
-bool FNiagaraEditorUtilities::AreTypesAssignable(const FNiagaraTypeDefinition& TypeA, const FNiagaraTypeDefinition& TypeB)
+bool FNiagaraEditorUtilities::AreTypesAssignable(const FNiagaraTypeDefinition& FromType, const FNiagaraTypeDefinition& ToType)
 {
-	return FNiagaraUtilities::AreTypesAssignable(TypeA, TypeB);
+	return FNiagaraUtilities::AreTypesAssignable(FromType, ToType);
 }
 
 void FNiagaraEditorUtilities::MarkDependentCompilableAssetsDirty(TArray<UObject*> InObjects)
diff --git a/Engine/Plugins/FX/Niagara/Source/NiagaraEditor/Private/ViewModels/Stack/NiagaraStackGraphUtilities.cpp b/Engine/Plugins/FX/Niagara/Source/NiagaraEditor/Private/ViewModels/Stack/NiagaraStackGraphUtilities.cpp
--- a/Engine/Plugins/FX/Niagara/Source/NiagaraEditor/Private/ViewModels/Stack/NiagaraStackGraphUtilities.cpp
+++ b/Engine/Plugins/FX/Niagara/Source/NiagaraEditor/Private/ViewModels/Stack/NiagaraStackGraphUtilities.cpp
@@ -2228,6 +2228,16 @@ void FNiagaraStackGraphUtilities::SetLinkedValueHandleForFunctionInput(UEdGraphP
 			ParameterToRead = ParameterAlternative;
 		}
 	}
+	else
+	{
+		// check if a static var of the requested input exists and use that instead
+		FNiagaraTypeDefinition StaticInputType = InputType.ToStaticDef();
+		FNiagaraVariable StaticParameter(StaticInputType, LinkedParameterHandle.GetParameterHandleString());
+		if (KnownParameters.Contains(StaticParameter))
+		{
+			ParameterToRead = StaticParameter;
+		}
+	}
 	
 	UEdGraphPin* GetOutputPin = GetNode->RequestNewTypedPin(EGPD_Output, ParameterToRead.GetType(), LinkedParameterHandle.GetParameterHandleString());
 	MakeLinkTo(GetInputPin, PreviousStackNodeOutputPin);
diff --git a/Engine/Plugins/FX/Niagara/Source/NiagaraEditor/Public/NiagaraEditorUtilities.h b/Engine/Plugins/FX/Niagara/Source/NiagaraEditor/Public/NiagaraEditorUtilities.h
--- a/Engine/Plugins/FX/Niagara/Source/NiagaraEditor/Public/NiagaraEditorUtilities.h
+++ b/Engine/Plugins/FX/Niagara/Source/NiagaraEditor/Public/NiagaraEditorUtilities.h
@@ -176,7 +176,7 @@ namespace FNiagaraEditorUtilities
 
 	FText GetTypeDefinitionCategory(const FNiagaraTypeDefinition& TypeDefinition);
 
-	NIAGARAEDITOR_API bool AreTypesAssignable(const FNiagaraTypeDefinition& TypeA, const FNiagaraTypeDefinition& TypeB);
+	NIAGARAEDITOR_API bool AreTypesAssignable(const FNiagaraTypeDefinition& FromType, const FNiagaraTypeDefinition& ToType);
 
 	void MarkDependentCompilableAssetsDirty(TArray<UObject*> InObjects);
 	void MarkDependentCompilableAssetsDirty(const TArray<FAssetData>& InAssets);
