- Simulation Stage Loops Prototype
- Where the data is currently located will likely change, this is a proof of concept for shipping a more customizable / stable simulation

diff --git a/Engine/Plugins/FX/Niagara/Shaders/Private/NiagaraEmitterInstanceShader.usf b/Engine/Plugins/FX/Niagara/Shaders/Private/NiagaraEmitterInstanceShader.usf
--- a/Engine/Plugins/FX/Niagara/Shaders/Private/NiagaraEmitterInstanceShader.usf
+++ b/Engine/Plugins/FX/Niagara/Shaders/Private/NiagaraEmitterInstanceShader.usf
@@ -560,20 +560,23 @@ float4 NiagaraQuatSLerp(float4 Quat1, float4 Quat2, float Slerp)
 	
 	// X = Count Buffer Instance Count Offset (INDEX_NONE == Use Instance Count)
 	// Y = Instance Count
-	// Z = Iteration Index
-	// W = Num Iterations
-	int4 SimulationStageIterationInfo;
-	float SimulationStageNormalizedIterationIndex;
+	// Z = Iteration Num | Index
+	// W = Loop Num | Index
+	uint4 SimulationStageIterationInfo;
 
 	// Note: These are referenced from an asset that passes back the data to the user (see SimulationStageIterationInfo)
 	#if NIAGARA_DISPATCH_INDIRECT
 		int SimulationStage_GetInstanceCount() { return DispatchThreadIdBounds.x * DispatchThreadIdBounds.y * DispatchThreadIdBounds.z; }
 	#else
-		int SimulationStage_GetInstanceCount() { return SimulationStageIterationInfo.x == -1 ? SimulationStageIterationInfo.y : RWInstanceCounts[SimulationStageIterationInfo.x]; }
+		int SimulationStage_GetInstanceCount() { return SimulationStageIterationInfo.x == uint(-1) ? int(SimulationStageIterationInfo.y) : int(RWInstanceCounts[SimulationStageIterationInfo.x]); }
 	#endif
-	int SimulationStage_GetIterationIndex() { return SimulationStageIterationInfo.z; }
-	int SimulationStage_GetNumIterations() { return SimulationStageIterationInfo.w; }
-	float SimulationStage_GetNormalizedIterationIndex() { return SimulationStageNormalizedIterationIndex; }
+	int SimulationStage_GetNumIterations()  { return int((SimulationStageIterationInfo.z >> 16) & 0xffff); }
+	int SimulationStage_GetIterationIndex() { return int((SimulationStageIterationInfo.z >> 0 ) & 0xffff); }
+	int SimulationStage_GetNumLoops()  { return int((SimulationStageIterationInfo.w >> 16) & 0xffff); }
+	int SimulationStage_GetLoopIndex() { return int((SimulationStageIterationInfo.w >> 0 ) & 0xffff); }
+
+	float SimulationStage_GetNormalizedIterationIndex() { return float(SimulationStage_GetIterationIndex()) / float(max(SimulationStage_GetNumIterations(), 1) - 1); }
+	float SimulationStage_GetNormalizedLoopIndex() { return float(SimulationStage_GetLoopIndex()) / float(max(SimulationStage_GetNumLoops(), 1) - 1); }
 
 	// Where X = Parameter Binding, YZ = Inclusive Range
 	uint3 ParticleIterationStateInfo;
diff --git a/Engine/Plugins/FX/Niagara/Source/Niagara/Classes/NiagaraComputeExecutionContext.h b/Engine/Plugins/FX/Niagara/Source/Niagara/Classes/NiagaraComputeExecutionContext.h
--- a/Engine/Plugins/FX/Niagara/Source/Niagara/Classes/NiagaraComputeExecutionContext.h
+++ b/Engine/Plugins/FX/Niagara/Source/Niagara/Classes/NiagaraComputeExecutionContext.h
@@ -5,6 +5,7 @@
 #include "NiagaraCommon.h"
 #include "NiagaraScriptExecutionContext.h"
 #include "NiagaraSimStageData.h"
+#include "NiagaraSimStageExecutionData.h"
 #include "RHIGPUReadback.h"
 
 class FNiagaraGPUInstanceCountManager;
@@ -85,7 +86,7 @@ struct FNiagaraComputeExecutionContext : public INiagaraComputeDataBufferInterfa
 
 	void Reset(FNiagaraGpuComputeDispatchInterface* ComputeDispatchInterface);
 
-	void InitParams(UNiagaraScript* InGPUComputeScript, ENiagaraSimTarget InSimTarget);
+	void InitParams(UNiagaraScript* InGPUComputeScript, const FNiagaraSimStageExecutionDataPtr& InSimStageExecData, ENiagaraSimTarget InSimTarget);
 	void DirtyDataInterfaces();
 	bool Tick(FNiagaraSystemInstance* ParentSystemInstance);
 
@@ -206,7 +207,7 @@ public:
 	// The current maximum instances we should allocate on the RT
 	uint32 CurrentMaxAllocateInstances_RT = 0;
 
-	TArray<FSimulationStageMetaData> SimStageInfo;
+	FNiagaraSimStageExecutionDataPtr SimStageExecData;
 
 	bool IsOutputStage(FNiagaraDataInterfaceProxy* DIProxy, uint32 SimulationStageIndex) const;
 	bool IsInputStage(FNiagaraDataInterfaceProxy* DIProxy, uint32 SimulationStageIndex) const;
diff --git a/Engine/Plugins/FX/Niagara/Source/Niagara/Classes/NiagaraEmitter.h b/Engine/Plugins/FX/Niagara/Source/Niagara/Classes/NiagaraEmitter.h
--- a/Engine/Plugins/FX/Niagara/Source/Niagara/Classes/NiagaraEmitter.h
+++ b/Engine/Plugins/FX/Niagara/Source/Niagara/Classes/NiagaraEmitter.h
@@ -18,7 +18,7 @@
 #include "NiagaraParameterDefinitionsBase.h"
 #include "NiagaraParameterDefinitionsSubscriber.h"
 #include "NiagaraScratchPadContainer.h"
-
+#include "NiagaraSimStageExecutionData.h"
 #include "NiagaraDataInterfacePlatformSet.h"
 
 #include "NiagaraEmitter.generated.h"
@@ -385,6 +385,7 @@ struct FVersionedNiagaraEmitterData
 	TConstArrayView<TSharedPtr<FNiagaraBoundsCalculator>> GetBoundsCalculators() const { return MakeArrayView(BoundsCalculators); }
 	NIAGARA_API bool RequiresPersistentIDs() const;
 	const TArray<UNiagaraSimulationStageBase*>& GetSimulationStages() const { return SimulationStages; }
+	FNiagaraSimStageExecutionDataPtr GetSimStageExcecutionData() const { return SimStageExecutionData; }
 	FORCEINLINE const FNiagaraEmitterScalabilitySettings& GetScalabilitySettings()const { return CurrentScalabilitySettings; }
 	NIAGARA_API const FNiagaraEmitterScalabilityOverride& GetCurrentOverrideSettings() const;
 	NIAGARA_API UNiagaraSimulationStageBase* GetSimulationStageById(FGuid ScriptUsageId) const;
@@ -483,7 +484,7 @@ struct FVersionedNiagaraEmitterData
 #endif
 
 	NIAGARA_API void GatherCompiledParticleAttributes(TArray<FNiagaraVariableBase>& OutVariables) const;
-	
+
 private:
 	UPROPERTY()
 	TArray<TObjectPtr<UNiagaraRendererProperties>> RendererProperties;
@@ -491,6 +492,16 @@ private:
 	UPROPERTY(meta = (NiagaraNoMerge))
 	TArray<TObjectPtr<UNiagaraSimulationStageBase>> SimulationStages;
 
+	UPROPERTY()
+	TArray<FNiagaraSimStageExecutionLoopData> SimStageExecutionLoops;
+
+#if WITH_EDITORONLY_DATA
+	UPROPERTY(EditAnywhere, Advanceddisplay, Category = "Emitter", meta = (DisplayName = "Sim Stage Loops"))
+	TArray<FNiagaraSimStageExecutionLoopEditorData> SimStageExecutionLoopEditorData;
+#endif
+
+	FNiagaraSimStageExecutionDataPtr SimStageExecutionData;
+
 	UPROPERTY()
 	TObjectPtr<UNiagaraScript> GPUComputeScript = nullptr;
 	
diff --git a/Engine/Plugins/FX/Niagara/Source/Niagara/Classes/NiagaraGPUSystemTick.h b/Engine/Plugins/FX/Niagara/Source/Niagara/Classes/NiagaraGPUSystemTick.h
--- a/Engine/Plugins/FX/Niagara/Source/Niagara/Classes/NiagaraGPUSystemTick.h
+++ b/Engine/Plugins/FX/Niagara/Source/Niagara/Classes/NiagaraGPUSystemTick.h
@@ -25,22 +25,13 @@ struct FNiagaraComputeInstanceData
 
 	struct FPerStageInfo
 	{
-		FPerStageInfo() {}
-		FPerStageInfo(int32 InNumIterations, FIntVector InElementCountXYZ) : NumIterations(InNumIterations), ElementCountXYZ(InElementCountXYZ) {}
-
-		int32 NumIterations = 0;
-		FIntVector ElementCountXYZ = FIntVector::ZeroValue;
-
-		bool ShouldRunStage() const { return NumIterations > 0; }
+		uint16		SimStageIndex = 0;
+		uint16		NumIterations = 0;
+		uint16		LoopIndex = 0;
+		uint16		NumLoops = 0;
+		FIntVector	ElementCountXYZ = FIntVector::ZeroValue;
 	};
 
-	FNiagaraComputeInstanceData()
-	{
-		bResetData = false;
-		bStartNewOverlapGroup = false;
-		bHasMultipleStages = false;
-	}
-
 	FNiagaraGpuSpawnInfo SpawnInfo;
 	uint8* EmitterParamData = nullptr;
 	uint8* ExternalParamData = nullptr;
@@ -51,10 +42,11 @@ struct FNiagaraComputeInstanceData
 	TArray<FPerStageInfo, TInlineAllocator<1>> PerStageInfo;
 	uint32 ParticleCountFence = INDEX_NONE;
 	uint32 TotalDispatches = 0;
-	uint32 bResetData : 1;
-	uint32 bStartNewOverlapGroup : 1;
-	uint32 bHasMultipleStages : 1;
+	uint32 bResetData : 1 = false;
+	uint32 bStartNewOverlapGroup : 1 = false;
+	uint32 bHasMultipleStages : 1 = false;
 
+	FNiagaraComputeInstanceData() = default;
 	bool IsOutputStage(FNiagaraDataInterfaceProxy* DIProxy, uint32 CurrentStage) const;
 	bool IsInputStage(FNiagaraDataInterfaceProxy* DIProxy, uint32 CurrentStage) const;
 	bool IsIterationStage(FNiagaraDataInterfaceProxy* DIProxy, uint32 CurrentStage) const;
@@ -135,7 +127,6 @@ public:
 	uint32 InstanceCount = 0;
 	uint32 TotalDispatches = 0;
 	bool bIsFinalTick = false;
-	bool bHasMultipleStages = false;
 	bool bHasInterpolatedParameters = false;
 
 #if !UE_BUILD_SHIPPING && !UE_BUILD_TEST
diff --git a/Engine/Plugins/FX/Niagara/Source/Niagara/Private/NiagaraComputeExecutionContext.cpp b/Engine/Plugins/FX/Niagara/Source/Niagara/Private/NiagaraComputeExecutionContext.cpp
--- a/Engine/Plugins/FX/Niagara/Source/Niagara/Private/NiagaraComputeExecutionContext.cpp
+++ b/Engine/Plugins/FX/Niagara/Source/Niagara/Private/NiagaraComputeExecutionContext.cpp
@@ -38,22 +38,14 @@ void FNiagaraComputeExecutionContext::Reset(FNiagaraGpuComputeDispatchInterface*
 	);
 }
 
-void FNiagaraComputeExecutionContext::InitParams(UNiagaraScript* InGPUComputeScript, ENiagaraSimTarget InSimTarget)
+void FNiagaraComputeExecutionContext::InitParams(UNiagaraScript* InGPUComputeScript, const FNiagaraSimStageExecutionDataPtr& InSimStageExecData, ENiagaraSimTarget InSimTarget)
 {
 	GPUScript = InGPUComputeScript;
+	SimStageExecData = InSimStageExecData;
 	CombinedParamStore.InitFromOwningContext(InGPUComputeScript, InSimTarget, true);
 	
 	HasInterpolationParameters = GPUScript && GPUScript->GetComputedVMCompilationId().HasInterpolatedParameters();
 
-	if (InGPUComputeScript)
-	{
-		FNiagaraVMExecutableData& VMData = InGPUComputeScript->GetVMExecutableData();
-		if ( VMData.IsValid() )
-		{
-			SimStageInfo = VMData.SimulationStageMetaData;
-		}
-	}
-
 #if DO_CHECK
 	// DI Parameters are the same between all shader permutations so we can just get the first one
 	FNiagaraShaderRef Shader = InGPUComputeScript->GetRenderThreadScript()->GetShaderGameThread(0);
@@ -81,7 +73,7 @@ bool FNiagaraComputeExecutionContext::IsOutputStage(FNiagaraDataInterfaceProxy*
 {
 	if (DIProxy && !DIProxy->SourceDIName.IsNone())
 	{		
-		return SimStageInfo[SimulationStageIndex].OutputDestinations.Contains(DIProxy->SourceDIName);
+		return SimStageExecData->SimStageMetaData[SimulationStageIndex].OutputDestinations.Contains(DIProxy->SourceDIName);
 	}
 	return false;
 }
@@ -90,7 +82,7 @@ bool FNiagaraComputeExecutionContext::IsInputStage(FNiagaraDataInterfaceProxy* D
 {
 	if (DIProxy && !DIProxy->SourceDIName.IsNone())
 	{
-		return SimStageInfo[SimulationStageIndex].InputDataInterfaces.Contains(DIProxy->SourceDIName);
+		return SimStageExecData->SimStageMetaData[SimulationStageIndex].InputDataInterfaces.Contains(DIProxy->SourceDIName);
 	}
 	return false;
 }
@@ -99,7 +91,7 @@ bool FNiagaraComputeExecutionContext::IsIterationStage(FNiagaraDataInterfaceProx
 {
 	if (DIProxy && !DIProxy->SourceDIName.IsNone())
 	{
-		return SimStageInfo[SimulationStageIndex].IterationSourceType == ENiagaraIterationSource::DataInterface && (SimStageInfo[SimulationStageIndex].IterationDataInterface == DIProxy->SourceDIName);
+		return SimStageExecData->SimStageMetaData[SimulationStageIndex].IterationSourceType == ENiagaraIterationSource::DataInterface && (SimStageExecData->SimStageMetaData[SimulationStageIndex].IterationDataInterface == DIProxy->SourceDIName);
 	}
 	return false;
 }
@@ -107,20 +99,20 @@ bool FNiagaraComputeExecutionContext::IsIterationStage(FNiagaraDataInterfaceProx
 FNiagaraDataInterfaceProxyRW* FNiagaraComputeExecutionContext::FindIterationInterface(const TArray<FNiagaraDataInterfaceProxyRW*>& InProxies, uint32 SimulationStageIndex) const
 {
 	// Particle stage
-	if ( SimStageInfo[SimulationStageIndex].IterationSourceType != ENiagaraIterationSource::DataInterface )
+	if (SimStageExecData->SimStageMetaData[SimulationStageIndex].IterationSourceType != ENiagaraIterationSource::DataInterface )
 	{
 		return nullptr;
 	}
 
 	for (FNiagaraDataInterfaceProxyRW* Proxy : InProxies)
 	{
-		if (Proxy->SourceDIName == SimStageInfo[SimulationStageIndex].IterationDataInterface)
+		if (Proxy->SourceDIName == SimStageExecData->SimStageMetaData[SimulationStageIndex].IterationDataInterface)
 		{
 			return Proxy;
 		}
 	}
 
-	UE_LOG(LogNiagara, Verbose, TEXT("FNiagaraComputeExecutionContext::FindIterationInterface could not find IterationInterface %s"), *SimStageInfo[SimulationStageIndex].IterationDataInterface.ToString());
+	UE_LOG(LogNiagara, Verbose, TEXT("FNiagaraComputeExecutionContext::FindIterationInterface could not find IterationInterface %s"), *SimStageExecData->SimStageMetaData[SimulationStageIndex].IterationDataInterface.ToString());
 
 	return nullptr;
 }
diff --git a/Engine/Plugins/FX/Niagara/Source/Niagara/Private/NiagaraEmitter.cpp b/Engine/Plugins/FX/Niagara/Source/Niagara/Private/NiagaraEmitter.cpp
--- a/Engine/Plugins/FX/Niagara/Source/Niagara/Private/NiagaraEmitter.cpp
+++ b/Engine/Plugins/FX/Niagara/Source/Niagara/Private/NiagaraEmitter.cpp
@@ -1733,6 +1733,33 @@ void FVersionedNiagaraEmitterData::CacheFromCompiledData(const FNiagaraDataSetCo
 	}
 	UpdateDebugName(Emitter, CompiledData);
 
+	SimStageExecutionData = nullptr;
+	if (GPUComputeScript)
+	{
+		if (CompiledData)
+		{
+			for (FSimulationStageMetaData& SimStageMetaData : GPUComputeScript->GetVMExecutableData().SimulationStageMetaData)
+			{
+				if (SimStageMetaData.bParticleIterationStateEnabled)
+				{
+					if (const FNiagaraVariableLayoutInfo* VariableInfo = CompiledData->FindVariableLayoutInfo(FNiagaraVariableBase(FNiagaraTypeDefinition::GetIntDef(), SimStageMetaData.ParticleIterationStateBinding)))
+					{
+						SimStageMetaData.ParticleIterationStateComponentIndex = VariableInfo->GetInt32ComponentStart();
+					}
+				}
+			}
+		}
+
+		SimStageExecutionData = MakeShared<FNiagaraSimStageExecutionData>();
+#if WITH_EDITORONLY_DATA
+		SimStageExecutionData->Build(GPUComputeScript->GetVMExecutableData().SimulationStageMetaData, SimStageExecutionLoopEditorData);
+		SimStageExecutionLoops = SimStageExecutionData->ExecutionLoops;
+#else
+		SimStageExecutionData->ExecutionLoops = SimStageExecutionLoops;
+		SimStageExecutionData->SimStageMetaData = GPUComputeScript->GetVMExecutableData().SimulationStageMetaData;
+#endif
+	}
+
 	// Detemine if we are allowed to execute or not
 	bIsAllowedToExecute = IsAllowedByScalability();
 	if (bIsAllowedToExecute && (SimTarget == ENiagaraSimTarget::GPUComputeSim))
@@ -1789,20 +1816,6 @@ void FVersionedNiagaraEmitterData::UpdateDebugName(const UNiagaraEmitter& Emitte
 #endif
 
 	RebuildRendererBindings(Emitter);
-
-	if ( GPUComputeScript && CompiledData )
-	{
-		for (FSimulationStageMetaData& SimStageMetaData : GPUComputeScript->GetVMExecutableData().SimulationStageMetaData)
-		{
-			if (SimStageMetaData.bParticleIterationStateEnabled)
-			{
-				if (const FNiagaraVariableLayoutInfo* VariableInfo = CompiledData->FindVariableLayoutInfo(FNiagaraVariableBase(FNiagaraTypeDefinition::GetIntDef(), SimStageMetaData.ParticleIterationStateBinding)))
-				{
-					SimStageMetaData.ParticleIterationStateComponentIndex = VariableInfo->GetInt32ComponentStart();
-				}
-			}
-		}
-	}
 }
 
 bool FVersionedNiagaraEmitterData::BuildParameterStoreRendererBindings(FNiagaraParameterStore& ParameterStore) const
@@ -1844,6 +1857,13 @@ bool FVersionedNiagaraEmitterData::BuildParameterStoreRendererBindings(FNiagaraP
 				bAnyBindingsAdded |= ParameterStore.AddParameter(FNiagaraVariable(FNiagaraTypeDefinition::GetIntDef(), SimStageMetaData.NumIterationsBinding), false);
 			}
 		}
+		for ( const FNiagaraSimStageExecutionLoopData& LoopData : SimStageExecutionLoops)
+		{
+			if (!LoopData.NumLoopsBinding.IsNone())
+			{
+				bAnyBindingsAdded |= ParameterStore.AddParameter(FNiagaraVariable(FNiagaraTypeDefinition::GetIntDef(), LoopData.NumLoopsBinding), false);
+			}
+		}
 	}
 	return bAnyBindingsAdded;
 }
diff --git a/Engine/Plugins/FX/Niagara/Source/Niagara/Private/NiagaraEmitterInstanceImpl.cpp b/Engine/Plugins/FX/Niagara/Source/Niagara/Private/NiagaraEmitterInstanceImpl.cpp
--- a/Engine/Plugins/FX/Niagara/Source/Niagara/Private/NiagaraEmitterInstanceImpl.cpp
+++ b/Engine/Plugins/FX/Niagara/Source/Niagara/Private/NiagaraEmitterInstanceImpl.cpp
@@ -362,7 +362,7 @@ void FNiagaraEmitterInstanceImpl::Init(int32 InEmitterIndex)
 		if (EmitterData->SimTarget == ENiagaraSimTarget::GPUComputeSim)
 		{
 			GPUExecContext = new FNiagaraComputeExecutionContext();
-			GPUExecContext->InitParams(EmitterData->GetGPUComputeScript(), EmitterData->SimTarget);
+			GPUExecContext->InitParams(EmitterData->GetGPUComputeScript(), EmitterData->GetSimStageExcecutionData(), EmitterData->SimTarget);
 			GPUExecContext->SetDebugSimName(EmitterData->GetDebugSimName());
 			GPUExecContext->ProfilingComponentPtr = ParentSystemInstance ? ParentSystemInstance->GetAttachComponent() : nullptr;
 			GPUExecContext->ProfilingEmitterPtr = VersionedEmitter.ToWeakPtr();
diff --git a/Engine/Plugins/FX/Niagara/Source/Niagara/Private/NiagaraGPUSystemTick.cpp b/Engine/Plugins/FX/Niagara/Source/Niagara/Private/NiagaraGPUSystemTick.cpp
--- a/Engine/Plugins/FX/Niagara/Source/Niagara/Private/NiagaraGPUSystemTick.cpp
+++ b/Engine/Plugins/FX/Niagara/Source/Niagara/Private/NiagaraGPUSystemTick.cpp
@@ -180,7 +180,6 @@ void FNiagaraGPUSystemTick::Init(FNiagaraSystemInstance* InSystemInstance)
 
 			ParamDataBufferPtr = GPUContext->WriteConstantBufferInstanceData(ParamDataBufferPtr, *InstanceData);
 
-			bHasMultipleStages = InstanceData->bHasMultipleStages;
 			bHasInterpolatedParameters |= GPUContext->HasInterpolationParameters;
 
 			// Calling PostTick will push current -> previous parameters this must be done after copying the parameter data
@@ -208,67 +207,88 @@ void FNiagaraGPUSystemTick::Init(FNiagaraSystemInstance* InSystemInstance)
 
 			// Gather number of iterations for each stage, and if the stage should run or not
 			InstanceData->bHasMultipleStages = false;
-			InstanceData->PerStageInfo.Reserve(GPUContext->SimStageInfo.Num());
-			for ( FSimulationStageMetaData& SimStageMetaData : GPUContext->SimStageInfo )
+			InstanceData->PerStageInfo.Reserve(GPUContext->SimStageExecData->SimStageMetaData.Num());		// Note: This presize isn't accurate but in the general none looping case it's correct
+
+			const FNiagaraParameterStore& ParameterStore = EmitterInstance->GetRendererBoundVariables();
+			for (const FNiagaraSimStageExecutionLoopData& LoopData : GPUContext->SimStageExecData->ExecutionLoops)
 			{
-				int32 NumIterations = SimStageMetaData.NumIterations;
-				FIntVector ElementCountXYZ = FIntVector(1, 1, 1);
-				if (SimStageMetaData.ShouldRunStage(InstanceData->bResetData))
+				int32 NumLoops = LoopData.NumLoopsBinding.IsNone() ? LoopData.NumLoops : ParameterStore.GetParameterValueOrDefault(FNiagaraVariable(FNiagaraTypeDefinition::GetIntDef(), LoopData.NumLoopsBinding), LoopData.NumLoops);
+				NumLoops = FMath::Clamp(NumLoops, 0, TNumericLimits<uint16>::Max());
+				for (int32 LoopIndex=0; LoopIndex < NumLoops; ++LoopIndex)
 				{
-					InstanceData->bHasMultipleStages = true;
-					if (!SimStageMetaData.NumIterationsBinding.IsNone())
-					{
-						FNiagaraParameterStore& BoundParamStore = EmitterInstance->GetRendererBoundVariables();
-						if ( const uint8* ParameterData = BoundParamStore.GetParameterData(FNiagaraVariable(FNiagaraTypeDefinition::GetIntDef(), SimStageMetaData.NumIterationsBinding)) )
-						{
-							NumIterations = *reinterpret_cast<const int32*>(ParameterData);
-							NumIterations = FMath::Max(NumIterations, 0);
-						}
-					}
-					if (!SimStageMetaData.EnabledBinding.IsNone())
-					{
-						FNiagaraParameterStore& BoundParamStore = EmitterInstance->GetRendererBoundVariables();
-						if (const uint8* ParameterData = BoundParamStore.GetParameterData(FNiagaraVariable(FNiagaraTypeDefinition::GetBoolDef(), SimStageMetaData.EnabledBinding)))
-						{
-							const FNiagaraBool StageEnabled = *reinterpret_cast<const FNiagaraBool*>(ParameterData);
-							NumIterations = StageEnabled.GetValue() ? NumIterations : 0;
-						}
-					}
-					if (SimStageMetaData.IterationSourceType == ENiagaraIterationSource::DirectSet)
+					for (int32 SimStageIndex=LoopData.StartStageIndex; SimStageIndex <= LoopData.EndStageIndex; ++SimStageIndex)
 					{
-						FNiagaraParameterStore& BoundParamStore = EmitterInstance->GetRendererBoundVariables();
-						ElementCountXYZ = FIntVector(1, 1, 1);
+						const FSimulationStageMetaData& SimStageMetaData = GPUContext->SimStageExecData->SimStageMetaData[SimStageIndex];
 
-						// This will cause a PVS warning, if we add something before this enum entry add the code back in
-						//if (SimStageMetaData.GpuDispatchType >= ENiagaraGpuDispatchType::OneD)
+						// Should we run the stage?
+						if (!SimStageMetaData.ShouldRunStage(InstanceData->bResetData))
 						{
-							ElementCountXYZ.X = SimStageMetaData.ElementCountXBinding.IsNone() ? SimStageMetaData.ElementCount.X : BoundParamStore.GetParameterValueOrDefault(FNiagaraVariable(FNiagaraTypeDefinition::GetIntDef(), SimStageMetaData.ElementCountXBinding), 0);
+							continue;
 						}
-						if (SimStageMetaData.GpuDispatchType >= ENiagaraGpuDispatchType::TwoD)
+
+						// Is the stage enabled?
+						if (!SimStageMetaData.EnabledBinding.IsNone())
 						{
-							ElementCountXYZ.Y = SimStageMetaData.ElementCountYBinding.IsNone() ? SimStageMetaData.ElementCount.Y : BoundParamStore.GetParameterValueOrDefault(FNiagaraVariable(FNiagaraTypeDefinition::GetIntDef(), SimStageMetaData.ElementCountYBinding), 0);
+							if (const uint8* ParameterData = ParameterStore.GetParameterData(FNiagaraVariable(FNiagaraTypeDefinition::GetBoolDef(), SimStageMetaData.EnabledBinding)))
+							{
+								const FNiagaraBool StageEnabled = *reinterpret_cast<const FNiagaraBool*>(ParameterData);
+								if (StageEnabled.GetValue() == false)
+								{
+									continue;
+								}
+							}
 						}
-						if (SimStageMetaData.GpuDispatchType >= ENiagaraGpuDispatchType::ThreeD)
+
+						// Get number of iterations
+						int32 NumIterations = SimStageMetaData.NumIterationsBinding.IsNone() ? SimStageMetaData.NumIterations : ParameterStore.GetParameterValueOrDefault(FNiagaraVariable(FNiagaraTypeDefinition::GetIntDef(), SimStageMetaData.NumIterationsBinding), SimStageMetaData.NumIterations);
+						NumIterations = FMath::Clamp(NumIterations, 0, TNumericLimits<uint16>::Max());
+						if ( NumIterations <= 0 )
 						{
-							ElementCountXYZ.Z = SimStageMetaData.ElementCountZBinding.IsNone() ? SimStageMetaData.ElementCount.Z : BoundParamStore.GetParameterValueOrDefault(FNiagaraVariable(FNiagaraTypeDefinition::GetIntDef(), SimStageMetaData.ElementCountZBinding), 0);
+							continue;
 						}
 
-						// If any of the element count values are <= 0 we can kill the stage as it won't execute
-						if (ElementCountXYZ.GetMin() <= 0)
+						// Get element count
+						FIntVector ElementCountXYZ = FIntVector(1, 1, 1);
+						if (SimStageMetaData.IterationSourceType == ENiagaraIterationSource::DirectSet)
 						{
-							ElementCountXYZ = FIntVector::ZeroValue;
-							NumIterations = 0;
+							FNiagaraParameterStore& BoundParamStore = EmitterInstance->GetRendererBoundVariables();
+							ElementCountXYZ = FIntVector(1, 1, 1);
+
+							// This will cause a PVS warning, if we add something before this enum entry add the code back in
+							//if (SimStageMetaData.GpuDispatchType >= ENiagaraGpuDispatchType::OneD)
+							{
+								ElementCountXYZ.X = SimStageMetaData.ElementCountXBinding.IsNone() ? SimStageMetaData.ElementCount.X : BoundParamStore.GetParameterValueOrDefault(FNiagaraVariable(FNiagaraTypeDefinition::GetIntDef(), SimStageMetaData.ElementCountXBinding), 0);
+							}
+							if (SimStageMetaData.GpuDispatchType >= ENiagaraGpuDispatchType::TwoD)
+							{
+								ElementCountXYZ.Y = SimStageMetaData.ElementCountYBinding.IsNone() ? SimStageMetaData.ElementCount.Y : BoundParamStore.GetParameterValueOrDefault(FNiagaraVariable(FNiagaraTypeDefinition::GetIntDef(), SimStageMetaData.ElementCountYBinding), 0);
+							}
+							if (SimStageMetaData.GpuDispatchType >= ENiagaraGpuDispatchType::ThreeD)
+							{
+								ElementCountXYZ.Z = SimStageMetaData.ElementCountZBinding.IsNone() ? SimStageMetaData.ElementCount.Z : BoundParamStore.GetParameterValueOrDefault(FNiagaraVariable(FNiagaraTypeDefinition::GetIntDef(), SimStageMetaData.ElementCountZBinding), 0);
+							}
+
+							// If any of the element count values are <= 0 we can kill the stage as it won't execute
+							if (ElementCountXYZ.GetMin() <= 0)
+							{
+								continue;
+							}
 						}
+
+						// Stage is live we can add it
+						FNiagaraComputeInstanceData::FPerStageInfo& NewStageInfo = InstanceData->PerStageInfo.AddDefaulted_GetRef();
+						NewStageInfo.SimStageIndex		= SimStageIndex;
+						NewStageInfo.NumIterations		= NumIterations;
+						NewStageInfo.LoopIndex			= LoopIndex;
+						NewStageInfo.NumLoops			= NumLoops;
+						NewStageInfo.ElementCountXYZ	= ElementCountXYZ;
+
+						InstanceData->bHasMultipleStages = true;
+						InstanceData->TotalDispatches += NumIterations;
 					}
 				}
-				else
-				{
-					NumIterations = 0;
-				}
-
-				InstanceData->TotalDispatches += NumIterations;
-				InstanceData->PerStageInfo.Emplace(NumIterations, ElementCountXYZ);
 			}
+
 			TotalDispatches += InstanceData->TotalDispatches;
 		}
 	}
@@ -314,7 +334,7 @@ void FNiagaraGPUSystemTick::BuildUniformBuffers()
 		{
 			if ( ensure(ExternalCBufferLayout && bExternalLayoutValid) )
 			{
-				ExternalUnformBuffers_RT[iInstance] = RHICreateUniformBuffer(Instance.ExternalParamData, ExternalCBufferLayout, bHasMultipleStages ? EUniformBufferUsage::UniformBuffer_SingleFrame : EUniformBufferUsage::UniformBuffer_SingleDraw);
+				ExternalUnformBuffers_RT[iInstance] = RHICreateUniformBuffer(Instance.ExternalParamData, ExternalCBufferLayout, Instance.bHasMultipleStages ? EUniformBufferUsage::UniformBuffer_SingleFrame : EUniformBufferUsage::UniformBuffer_SingleDraw);
 			}
 		}
 		if ( Instance.Context->GPUScript_RT->IsExternalConstantBufferUsed_RenderThread(1) )
@@ -322,7 +342,7 @@ void FNiagaraGPUSystemTick::BuildUniformBuffers()
 			if (ensure(ExternalCBufferLayout && bExternalLayoutValid))
 			{
 				check(ExternalCBufferLayout->ConstantBufferSize + ExternalCBufferLayout->ConstantBufferSize <= Instance.ExternalParamDataSize);
-				ExternalUnformBuffers_RT[InstanceCount + iInstance] = RHICreateUniformBuffer(Instance.ExternalParamData + ExternalCBufferLayout->ConstantBufferSize, ExternalCBufferLayout, bHasMultipleStages ? EUniformBufferUsage::UniformBuffer_SingleFrame : EUniformBufferUsage::UniformBuffer_SingleDraw);
+				ExternalUnformBuffers_RT[InstanceCount + iInstance] = RHICreateUniformBuffer(Instance.ExternalParamData + ExternalCBufferLayout->ConstantBufferSize, ExternalCBufferLayout, Instance.bHasMultipleStages ? EUniformBufferUsage::UniformBuffer_SingleFrame : EUniformBufferUsage::UniformBuffer_SingleDraw);
 			}
 		}
 	}
diff --git a/Engine/Plugins/FX/Niagara/Source/Niagara/Private/NiagaraGpuComputeDispatch.cpp b/Engine/Plugins/FX/Niagara/Source/Niagara/Private/NiagaraGpuComputeDispatch.cpp
--- a/Engine/Plugins/FX/Niagara/Source/Niagara/Private/NiagaraGpuComputeDispatch.cpp
+++ b/Engine/Plugins/FX/Niagara/Source/Niagara/Private/NiagaraGpuComputeDispatch.cpp
@@ -827,16 +827,14 @@ void FNiagaraGpuComputeDispatch::PrepareTicksForProxy(FRHICommandListImmediate&
 
 			//-OPT: Do we need this test?  Can remove in favor of MaxUpdateIterations
 			bool bFirstStage = true;
-			for (int32 SimStageIndex=0; SimStageIndex < ComputeContext->SimStageInfo.Num(); ++SimStageIndex)
+			for (const FNiagaraComputeInstanceData::FPerStageInfo& PerStageInfo : InstanceData.PerStageInfo)
 			{
-				const FSimulationStageMetaData& SimStageMetaData = ComputeContext->SimStageInfo[SimStageIndex];
-				if (InstanceData.PerStageInfo[SimStageIndex].ShouldRunStage() == false)
-				{
-					continue;
-				}
+				const int32 SimStageIndex = PerStageInfo.SimStageIndex;
+				const FSimulationStageMetaData& SimStageMetaData = ComputeContext->SimStageExecData->SimStageMetaData[SimStageIndex];
 
 				FNiagaraDataInterfaceProxyRW* IterationInterface = InstanceData.FindIterationInterface(SimStageIndex);
-				for ( int32 IterationIndex=0; IterationIndex < InstanceData.PerStageInfo[SimStageIndex].NumIterations; ++IterationIndex )
+				const int32 NumIterations = PerStageInfo.NumIterations;
+				for ( int32 IterationIndex=0; IterationIndex < NumIterations; ++IterationIndex )
 				{
 					// Build SimStage data
 					FNiagaraGpuDispatchGroup& DispatchGroup = GpuDispatchList.DispatchGroups[iInstanceCurrDispatchGroup++];
@@ -844,8 +842,11 @@ void FNiagaraGpuComputeDispatch::PrepareTicksForProxy(FRHICommandListImmediate&
 					FNiagaraSimStageData& SimStageData = DispatchInstance.SimStageData;
 					SimStageData.bFirstStage = bFirstStage;
 					SimStageData.StageIndex = SimStageIndex;
+					SimStageData.NumIterations = NumIterations;
 					SimStageData.IterationIndex = IterationIndex;
-					SimStageData.DispatchArgs.ElementCount = InstanceData.PerStageInfo[SimStageIndex].ElementCountXYZ;
+					SimStageData.NumLoops = PerStageInfo.NumLoops;
+					SimStageData.LoopIndex = PerStageInfo.LoopIndex;
+					SimStageData.DispatchArgs.ElementCount = PerStageInfo.ElementCountXYZ;
 					SimStageData.StageMetaData = &SimStageMetaData;
 					SimStageData.AlternateIterationSource = IterationInterface;
 
@@ -1561,8 +1562,7 @@ void FNiagaraGpuComputeDispatch::DispatchStage(FRDGBuilder& GraphBuilder, const
 	// Z = Iteration Index
 	// W = Num Iterations
 	{
-		DispatchParameters->SimulationStageIterationInfo = FIntVector4(INDEX_NONE, -1, 0, 0);
-		DispatchParameters->SimulationStageNormalizedIterationIndex = 0.0f;
+		DispatchParameters->SimulationStageIterationInfo = FUintVector4(INDEX_NONE, 0, 0, 0);
 		switch (SimStageData.StageMetaData->IterationSourceType)
 		{
 			case ENiagaraIterationSource::Particles:
@@ -1582,11 +1582,11 @@ void FNiagaraGpuComputeDispatch::DispatchStage(FRDGBuilder& GraphBuilder, const
 				break;
 		}
 
-		const int32 NumIterations = InstanceData.PerStageInfo[SimStageData.StageIndex].NumIterations;
-		const int32 IterationIndex = SimStageData.IterationIndex;
-		DispatchParameters->SimulationStageIterationInfo.Z = IterationIndex;
-		DispatchParameters->SimulationStageIterationInfo.W = NumIterations;
-		DispatchParameters->SimulationStageNormalizedIterationIndex = NumIterations > 1 ? float(IterationIndex) / float(NumIterations - 1) : 1.0f;
+		DispatchParameters->SimulationStageIterationInfo.Z  = uint32(SimStageData.NumIterations) << 16;
+		DispatchParameters->SimulationStageIterationInfo.Z |= uint32(SimStageData.IterationIndex);
+
+		DispatchParameters->SimulationStageIterationInfo.W  = uint32(SimStageData.NumLoops) << 16;
+		DispatchParameters->SimulationStageIterationInfo.W |= uint32(SimStageData.LoopIndex);
 	}
 
 	// Set particle iteration state info
diff --git a/Engine/Plugins/FX/Niagara/Source/Niagara/Private/NiagaraSimStageExecutionData.cpp b/Engine/Plugins/FX/Niagara/Source/Niagara/Private/NiagaraSimStageExecutionData.cpp
new file mode 100644
--- /dev/null
+++ b/Engine/Plugins/FX/Niagara/Source/Niagara/Private/NiagaraSimStageExecutionData.cpp
@@ -0,0 +1,86 @@
+// Copyright Epic Games, Inc. All Rights Reserved.
+
+#include "NiagaraSimStageExecutionData.h"
+#include "NiagaraCommon.h"
+
+#if WITH_EDITORONLY_DATA
+bool FNiagaraSimStageExecutionData::Build(TConstArrayView<FSimulationStageMetaData> InSimStageMetaData, TConstArrayView<FNiagaraSimStageExecutionLoopEditorData> EditorExecutionLoops)
+{
+	ExecutionLoops.Empty();
+	SimStageMetaData = InSimStageMetaData;
+
+	// Nothing to do?
+	if ( SimStageMetaData.Num() == 0 )
+	{
+		return true;
+	}
+
+	// Create a per stage loop index
+	TArray<TOptional<int32>> PerStageLoopIndex;
+	PerStageLoopIndex.AddDefaulted(SimStageMetaData.Num());
+
+	bool bLoopsValid = true;
+	for ( int32 LoopIndex=0; LoopIndex < EditorExecutionLoops.Num(); ++LoopIndex)
+	{
+		const FNiagaraSimStageExecutionLoopEditorData& EditorLoopData = EditorExecutionLoops[LoopIndex];
+		if ( !EditorLoopData.bEnabled || EditorLoopData.StageNameStart.IsNone() || EditorLoopData.StageNameEnd.IsNone() )
+		{
+			continue;
+		}
+
+		const int32 IndexStart = SimStageMetaData.IndexOfByPredicate([StartName=EditorLoopData.StageNameStart](const FSimulationStageMetaData& Data) { return Data.SimulationStageName == StartName; });
+		const int32 IndexEnd = SimStageMetaData.IndexOfByPredicate([EndName=EditorLoopData.StageNameEnd](const FSimulationStageMetaData& Data) { return Data.SimulationStageName == EndName; });
+		if ( IndexStart == INDEX_NONE || IndexEnd == INDEX_NONE || IndexEnd < IndexStart )
+		{
+			UE_LOG(LogNiagara, Error, TEXT("Sim Stage Loop has invalid start / end name, all loops will be ignored"));
+			bLoopsValid = false;
+			break;
+		}
+
+		for ( int32 i=IndexStart; i <= IndexEnd; ++i )
+		{
+			if (PerStageLoopIndex[i].IsSet() )
+			{
+				UE_LOG(LogNiagara, Error, TEXT("Sim Stage Loop has overlap, all loops will be ignored"));
+				bLoopsValid = false;
+			}
+			PerStageLoopIndex[i] = LoopIndex;
+		}
+	}
+
+	// Invalid for some reason, just ignore all loops and run linear
+	if ( !bLoopsValid )
+	{
+		PerStageLoopIndex.Reset(SimStageMetaData.Num());
+		PerStageLoopIndex.AddDefaulted(SimStageMetaData.Num());
+	}
+
+	// Generate the loop data
+	FNiagaraSimStageExecutionLoopData* CurrentLoop = &ExecutionLoops.AddDefaulted_GetRef();
+	CurrentLoop->NumLoops = 1;
+
+	constexpr int32 InvalidLoopIndex = INDEX_NONE;
+	int32 CurrentLoopIndex = InvalidLoopIndex;
+	for ( int32 i=1; i < SimStageMetaData.Num(); ++i )
+	{
+		const int32 LoopIndex = PerStageLoopIndex[i].Get(InvalidLoopIndex);
+		if (LoopIndex != CurrentLoopIndex)
+		{
+			// Close current loop
+			CurrentLoop->EndStageIndex		= i - 1;
+
+			// Start new  loop
+			CurrentLoop						= &ExecutionLoops.AddDefaulted_GetRef();
+			CurrentLoop->NumLoopsBinding	= LoopIndex == InvalidLoopIndex ? NAME_None : EditorExecutionLoops[LoopIndex].NumLoopsBindingName;
+			CurrentLoop->NumLoops			= LoopIndex == InvalidLoopIndex ? 1 : EditorExecutionLoops[LoopIndex].NumLoops;
+			CurrentLoop->StartStageIndex	= i;
+			CurrentLoopIndex				= LoopIndex;
+		}
+	}
+
+	// Close the loop
+	CurrentLoop->EndStageIndex = SimStageMetaData.Num() - 1;
+
+	return bLoopsValid;
+}
+#endif //WITH_EDITORONLY_DATA
diff --git a/Engine/Plugins/FX/Niagara/Source/Niagara/Public/NiagaraSimStageData.h b/Engine/Plugins/FX/Niagara/Source/Niagara/Public/NiagaraSimStageData.h
--- a/Engine/Plugins/FX/Niagara/Source/Niagara/Public/NiagaraSimStageData.h
+++ b/Engine/Plugins/FX/Niagara/Source/Niagara/Public/NiagaraSimStageData.h
@@ -21,19 +21,17 @@ struct FNiagaraSimStageDispatchArgs
 
 struct FNiagaraSimStageData
 {
-	FNiagaraSimStageData()
-	{
-		bFirstStage = false;
-		bLastStage = false;
-		bSetDataToRender = false;
-	}
-
-	uint32 bFirstStage : 1;
-	uint32 bLastStage : 1;
-	uint32 bSetDataToRender : 1;
-
-	uint32 StageIndex = INDEX_NONE;
-	uint32 IterationIndex = 0;
+	uint16 bFirstStage : 1 = false;
+	uint16 bLastStage : 1 = false;
+	uint16 bSetDataToRender : 1 = false;
+	uint16 StageIndex = INDEX_NONE;
+
+	uint16 NumIterations = 0;
+	uint16 IterationIndex = 0;
+
+	uint16 NumLoops = 0;
+	uint16 LoopIndex = 0;
+
 	FNiagaraSimStageDispatchArgs DispatchArgs;
 
 	class FNiagaraDataBuffer* Source = nullptr;
diff --git a/Engine/Plugins/FX/Niagara/Source/Niagara/Public/NiagaraSimStageExecutionData.h b/Engine/Plugins/FX/Niagara/Source/Niagara/Public/NiagaraSimStageExecutionData.h
new file mode 100644
--- /dev/null
+++ b/Engine/Plugins/FX/Niagara/Source/Niagara/Public/NiagaraSimStageExecutionData.h
@@ -0,0 +1,63 @@
+// Copyright Epic Games, Inc. All Rights Reserved.
+
+#pragma once
+
+#include "HAL/Platform.h"
+#include "Math/UnrealMath.h"
+#include "Misc/CoreMiscDefines.h"
+#include "NiagaraScriptBase.h"
+
+#include "NiagaraSimStageExecutionData.generated.h"
+
+USTRUCT()
+struct FNiagaraSimStageExecutionLoopEditorData
+{
+	GENERATED_BODY()
+
+#if WITH_EDITORONLY_DATA
+	UPROPERTY(EditAnywhere, Category = "Stage Loop")
+	bool	bEnabled = true;
+
+	UPROPERTY(EditAnywhere, Category = "Stage Loop")
+	FName	NumLoopsBindingName;
+
+	UPROPERTY(EditAnywhere, Category = "Stage Loop")
+	int32	NumLoops = 1;
+
+	UPROPERTY(EditAnywhere, Category = "Stage Loop")
+	FName	StageNameStart;
+
+	UPROPERTY(EditAnywhere, Category = "Stage Loop")
+	FName	StageNameEnd;
+#endif
+};
+
+USTRUCT()
+struct FNiagaraSimStageExecutionLoopData
+{
+	GENERATED_USTRUCT_BODY()
+
+	UPROPERTY(EditAnywhere, Category = "SimStageLoop")
+	FName NumLoopsBinding;
+
+	UPROPERTY(EditAnywhere, Category = "SimStageLoop")
+	int32 NumLoops = 1;
+
+	UPROPERTY(EditAnywhere, Category = "SimStageLoop")
+	int32 StartStageIndex = 0;
+
+	UPROPERTY(EditAnywhere, Category = "SimStageLoop")
+	int32 EndStageIndex = 0;
+};
+
+struct FNiagaraSimStageExecutionData
+{
+	TArray<FNiagaraSimStageExecutionLoopData, TInlineAllocator<1>>	ExecutionLoops;
+	TArray<FSimulationStageMetaData>								SimStageMetaData;
+
+#if WITH_EDITORONLY_DATA
+	bool Build(TConstArrayView<FSimulationStageMetaData> InSimStageMetaDatas, TConstArrayView<FNiagaraSimStageExecutionLoopEditorData> InEditorExecutionLoops);
+#endif
+};
+
+using FNiagaraSimStageExecutionDataPtr = TSharedPtr<FNiagaraSimStageExecutionData>;
diff --git a/Engine/Plugins/FX/Niagara/Source/NiagaraShader/Private/NiagaraScriptBase.cpp b/Engine/Plugins/FX/Niagara/Source/NiagaraShader/Private/NiagaraScriptBase.cpp
--- a/Engine/Plugins/FX/Niagara/Source/NiagaraShader/Private/NiagaraScriptBase.cpp
+++ b/Engine/Plugins/FX/Niagara/Source/NiagaraShader/Private/NiagaraScriptBase.cpp
@@ -14,8 +14,6 @@ FSimulationStageMetaData::FSimulationStageMetaData()
 	, bPartialParticleUpdate(0)
 	, bParticleIterationStateEnabled(0)
 	, bGpuIndirectDispatch(0)
-	, GpuDispatchType(ENiagaraGpuDispatchType::OneD)
-	, GpuDispatchNumThreads(0, 0, 0)
 {
 }
 
diff --git a/Engine/Plugins/FX/Niagara/Source/NiagaraShader/Public/NiagaraScriptBase.h b/Engine/Plugins/FX/Niagara/Source/NiagaraShader/Public/NiagaraScriptBase.h
--- a/Engine/Plugins/FX/Niagara/Source/NiagaraShader/Public/NiagaraScriptBase.h
+++ b/Engine/Plugins/FX/Niagara/Source/NiagaraShader/Public/NiagaraScriptBase.h
@@ -59,7 +59,7 @@ USTRUCT()
 struct FSimulationStageMetaData
 {
 	GENERATED_USTRUCT_BODY()
-public:
+
 	NIAGARASHADER_API FSimulationStageMetaData();
 
 	/** User simulation stage name. */
@@ -151,14 +151,14 @@ public:
 
 	/** Dispatch type set for this stage. */
 	UPROPERTY()
-	ENiagaraGpuDispatchType GpuDispatchType;
+	ENiagaraGpuDispatchType GpuDispatchType = ENiagaraGpuDispatchType::OneD;
 
 	UPROPERTY()
 	ENiagaraDirectDispatchElementType GpuDirectDispatchElementType = ENiagaraDirectDispatchElementType::NumThreads;
 
 	/** When in custom mode this is the num threads. */
 	UPROPERTY()
-	FIntVector GpuDispatchNumThreads;
+	FIntVector GpuDispatchNumThreads = FIntVector(0, 0, 0);
 };
 
 UCLASS(MinimalAPI, abstract)
diff --git a/Engine/Plugins/FX/Niagara/Source/NiagaraShader/Public/NiagaraShader.h b/Engine/Plugins/FX/Niagara/Source/NiagaraShader/Public/NiagaraShader.h
--- a/Engine/Plugins/FX/Niagara/Source/NiagaraShader/Public/NiagaraShader.h
+++ b/Engine/Plugins/FX/Niagara/Source/NiagaraShader/Public/NiagaraShader.h
@@ -178,8 +178,7 @@ public:
 		SHADER_PARAMETER_SRV(Buffer<int>,							FreeIDList)
 		SHADER_PARAMETER_UAV(RWBuffer<int>,							RWIDToIndexTable)
 
-		SHADER_PARAMETER(FIntVector4,								SimulationStageIterationInfo)
-		SHADER_PARAMETER(float,										SimulationStageNormalizedIterationIndex)
+		SHADER_PARAMETER(FUintVector4,								SimulationStageIterationInfo)
 
 		SHADER_PARAMETER(FIntVector3,								ParticleIterationStateInfo)
 
