PR #9888: Fix use after free in Niagara data interface details

diff --git a/Engine/Plugins/FX/Niagara/Source/NiagaraEditorWidgets/Private/DetailCustomizations/NiagaraDataInterfaceDetails.cpp b/Engine/Plugins/FX/Niagara/Source/NiagaraEditorWidgets/Private/DetailCustomizations/NiagaraDataInterfaceDetails.cpp
--- a/Engine/Plugins/FX/Niagara/Source/NiagaraEditorWidgets/Private/DetailCustomizations/NiagaraDataInterfaceDetails.cpp
+++ b/Engine/Plugins/FX/Niagara/Source/NiagaraEditorWidgets/Private/DetailCustomizations/NiagaraDataInterfaceDetails.cpp
@@ -214,8 +214,8 @@ class FNiagaraDataInterfaceCustomNodeBuilder : public IDetailCustomNodeBuilder
 											 , public TSharedFromThis<FNiagaraDataInterfaceCustomNodeBuilder>
 {
 public:
-	FNiagaraDataInterfaceCustomNodeBuilder(IDetailLayoutBuilder* InDetailBuilder)
-		: DetailBuilder(InDetailBuilder)
+	FNiagaraDataInterfaceCustomNodeBuilder(TWeakPtr<IDetailLayoutBuilder> InDetailBuilder)
+		: WeakDetailBuilder(MoveTemp(InDetailBuilder))
 	{
 	}
 
@@ -316,21 +316,23 @@ private:
 	{
 		FProperty* PropertyPlaceholder = nullptr;  // we don't need to specify the property, all we need is to trigger the restart of the emitter
 		FPropertyChangedEvent ChangeEvent(PropertyPlaceholder, EPropertyChangeType::Unspecified);
-		if (DetailBuilder->GetPropertyUtilities()->GetNotifyHook() != nullptr)
+		if (TSharedPtr<IDetailLayoutBuilder> DetailBuilder = WeakDetailBuilder.Pin())
 		{
-			DetailBuilder->GetPropertyUtilities()->GetNotifyHook()->NotifyPostChange(ChangeEvent, PropertyPlaceholder);
+			if (DetailBuilder->GetPropertyUtilities()->GetNotifyHook() != nullptr)
+			{
+				DetailBuilder->GetPropertyUtilities()->GetNotifyHook()->NotifyPostChange(ChangeEvent, PropertyPlaceholder);
+			}
 		}
 	}
 
 private:
 	TWeakObjectPtr<UNiagaraDataInterface> DataInterface;
-	IDetailLayoutBuilder* DetailBuilder;
+	TWeakPtr<IDetailLayoutBuilder> WeakDetailBuilder;
 	FSimpleDelegate OnRebuildChildren;
 };
 
 void FNiagaraDataInterfaceDetailsBase::CustomizeDetails(IDetailLayoutBuilder& DetailBuilder)
 {	
-	Builder = &DetailBuilder;
 	PropertyUtilitiesWeak = DetailBuilder.GetPropertyUtilities();
 	TArray<TWeakObjectPtr<UObject>> SelectedObjects;
 	DetailBuilder.GetObjectsBeingCustomized(SelectedObjects);
@@ -340,17 +342,24 @@ void FNiagaraDataInterfaceDetailsBase::CustomizeDetails(IDetailLayoutBuilder& De
 	DataInterface->OnErrorsRefreshed().AddSP(this, &FNiagaraDataInterfaceDetailsBase::OnErrorsRefreshed);
 	IDetailCategoryBuilder& ErrorsBuilderRef = DetailBuilder.EditCategory(ErrorsCategoryName, LOCTEXT("ErrorsAndWarnings", "Errors And Warnings"), ECategoryPriority::Important);
 	ErrorsCategoryBuilder = &ErrorsBuilderRef;
-	CustomBuilder = MakeShared<FNiagaraDataInterfaceCustomNodeBuilder>(&DetailBuilder);
+	CustomBuilder = MakeShared<FNiagaraDataInterfaceCustomNodeBuilder>(Builder);
 	CustomBuilder->Initialize(*DataInterface);
 	ErrorsCategoryBuilder->AddCustomBuilder(CustomBuilder.ToSharedRef());
 	OnErrorsRefreshed();
 }
 
+void FNiagaraDataInterfaceDetailsBase::CustomizeDetails(const TSharedPtr<IDetailLayoutBuilder>& DetailBuilder)
+{
+	Builder = DetailBuilder;
+	CustomizeDetails(*DetailBuilder);
+}
+
 void FNiagaraDataInterfaceDetailsBase::OnErrorsRefreshed() // need to only refresh errors, and all will be good
 {
 	TSharedPtr<IPropertyUtilities> PropertyUtilities = PropertyUtilitiesWeak.Pin();
 	bool bStillValid =
 		DataInterface.IsValid() &&
+		Builder.IsValid() &&
 		PropertyUtilities.IsValid() &&
 		PropertyUtilities->GetSelectedObjects().Num() == 1 &&
 		PropertyUtilities->GetSelectedObjects()[0].IsValid() &&
diff --git a/Engine/Plugins/FX/Niagara/Source/NiagaraEditorWidgets/Private/DetailCustomizations/NiagaraDataInterfaceDetails.h b/Engine/Plugins/FX/Niagara/Source/NiagaraEditorWidgets/Private/DetailCustomizations/NiagaraDataInterfaceDetails.h
--- a/Engine/Plugins/FX/Niagara/Source/NiagaraEditorWidgets/Private/DetailCustomizations/NiagaraDataInterfaceDetails.h
+++ b/Engine/Plugins/FX/Niagara/Source/NiagaraEditorWidgets/Private/DetailCustomizations/NiagaraDataInterfaceDetails.h
@@ -18,6 +18,7 @@ class FNiagaraDataInterfaceDetailsBase : public IDetailCustomization
 public:
 	~FNiagaraDataInterfaceDetailsBase();
 	virtual void CustomizeDetails(IDetailLayoutBuilder& DetailBuilder) override;
+	virtual void CustomizeDetails(const TSharedPtr<IDetailLayoutBuilder>& DetailBuilder) override;
 	static TSharedRef<IDetailCustomization> MakeInstance();
 
 private:
@@ -27,6 +28,6 @@ private:
 	TWeakObjectPtr<UNiagaraDataInterface> DataInterface;
 	TSharedPtr<FNiagaraDataInterfaceCustomNodeBuilder> CustomBuilder;
 	IDetailCategoryBuilder* ErrorsCategoryBuilder;
-	IDetailLayoutBuilder* Builder;
+	TWeakPtr<IDetailLayoutBuilder> Builder;
 	TWeakPtr<IPropertyUtilities> PropertyUtilitiesWeak;
 };
