- Handle case where the graph is invalid / missing due to loading issue

diff --git a/Engine/Plugins/FX/Niagara/Source/NiagaraEditor/Private/NiagaraNodeFunctionCall.cpp b/Engine/Plugins/FX/Niagara/Source/NiagaraEditor/Private/NiagaraNodeFunctionCall.cpp
--- a/Engine/Plugins/FX/Niagara/Source/NiagaraEditor/Private/NiagaraNodeFunctionCall.cpp
+++ b/Engine/Plugins/FX/Niagara/Source/NiagaraEditor/Private/NiagaraNodeFunctionCall.cpp
@@ -62,34 +62,36 @@ void UNiagaraNodeFunctionCall::PostLoad()
 		if (Source)
 		{
 			Source->ConditionalPostLoad();
-			UNiagaraGraph* Graph = Source->NodeGraph;
-			Graph->ConditionalPostLoad();
-
-			// Fix up autogenerated default values if necessary.
-			if (NiagaraCustomVersion < FNiagaraCustomVersion::EnabledAutogeneratedDefaultValuesForFunctionCallNodes)
+			if (UNiagaraGraph* Graph = Source->NodeGraph)
 			{
-				FPinCollectorArray InputPins;
-				GetInputPins(InputPins);
-
-				TArray<UNiagaraNodeInput*> InputNodes;
-				UNiagaraGraph::FFindInputNodeOptions Options;
-				Options.bSort = true;
-				Options.bFilterDuplicates = true;
-				Graph->FindInputNodes(InputNodes, Options);
+				Graph->ConditionalPostLoad();
 
-				for (UEdGraphPin* InputPin : InputPins)
+				// Fix up autogenerated default values if necessary.
+				if (NiagaraCustomVersion < FNiagaraCustomVersion::EnabledAutogeneratedDefaultValuesForFunctionCallNodes)
 				{
-					auto FindInputNodeByName = [&](UNiagaraNodeInput* InputNode) { return InputNode->Input.GetName().ToString() == InputPin->PinName.ToString(); };
-					UNiagaraNodeInput** MatchingInputNodePtr = InputNodes.FindByPredicate(FindInputNodeByName);
-					if (MatchingInputNodePtr != nullptr)
-					{
-						UNiagaraNodeInput* MatchingInputNode = *MatchingInputNodePtr;
-						SetPinAutoGeneratedDefaultValue(*InputPin, *MatchingInputNode);
+					FPinCollectorArray InputPins;
+					GetInputPins(InputPins);
 
-						// If the default value wasn't set, update it with the new autogenerated default.
-						if (InputPin->DefaultValue.IsEmpty())
+					TArray<UNiagaraNodeInput*> InputNodes;
+					UNiagaraGraph::FFindInputNodeOptions Options;
+					Options.bSort = true;
+					Options.bFilterDuplicates = true;
+					Graph->FindInputNodes(InputNodes, Options);
+
+					for (UEdGraphPin* InputPin : InputPins)
+					{
+						auto FindInputNodeByName = [&](UNiagaraNodeInput* InputNode) { return InputNode->Input.GetName().ToString() == InputPin->PinName.ToString(); };
+						UNiagaraNodeInput** MatchingInputNodePtr = InputNodes.FindByPredicate(FindInputNodeByName);
+						if (MatchingInputNodePtr != nullptr)
 						{
-							InputPin->DefaultValue = InputPin->AutogeneratedDefaultValue;
+							UNiagaraNodeInput* MatchingInputNode = *MatchingInputNodePtr;
+							SetPinAutoGeneratedDefaultValue(*InputPin, *MatchingInputNode);
+
+							// If the default value wasn't set, update it with the new autogenerated default.
+							if (InputPin->DefaultValue.IsEmpty())
+							{
+								InputPin->DefaultValue = InputPin->AutogeneratedDefaultValue;
+							}
 						}
 					}
 				}
