- Fix crash in ribbon renderer if normalized age is invalid

diff --git a/Engine/Plugins/FX/Niagara/Source/Niagara/Private/NiagaraRendererRibbons.cpp b/Engine/Plugins/FX/Niagara/Source/Niagara/Private/NiagaraRendererRibbons.cpp
--- a/Engine/Plugins/FX/Niagara/Source/Niagara/Private/NiagaraRendererRibbons.cpp
+++ b/Engine/Plugins/FX/Niagara/Source/Niagara/Private/NiagaraRendererRibbons.cpp
@@ -1313,8 +1313,8 @@ void FNiagaraRendererRibbons::CalculateUVScaleAndOffsets(const FNiagaraRibbonUVS
 	float NormalizedLeadingSegmentOffset;
 	if (UVSettings.LeadingEdgeMode == ENiagaraRibbonUVEdgeMode::SmoothTransition)
 	{
-		const float FirstAge = NormalizedAgeReader[RibbonIndices[StartIndex]];
-		const float SecondAge = NormalizedAgeReader[RibbonIndices[StartIndex + 1]];
+		const float FirstAge = NormalizedAgeReader.GetSafe(RibbonIndices[StartIndex], 0.0f);
+		const float SecondAge = NormalizedAgeReader.GetSafe(RibbonIndices[StartIndex + 1], 0.0f);
 
 		const float StartTimeStep = SecondAge - FirstAge;
 		const float StartTimeOffset = FirstAge < StartTimeStep ? StartTimeStep - FirstAge : 0;
@@ -1334,8 +1334,8 @@ void FNiagaraRendererRibbons::CalculateUVScaleAndOffsets(const FNiagaraRibbonUVS
 	float NormalizedTrailingSegmentOffset;
 	if (UVSettings.TrailingEdgeMode == ENiagaraRibbonUVEdgeMode::SmoothTransition)
 	{
-		const float SecondToLastAge = NormalizedAgeReader[RibbonIndices[EndIndex - 1]];
-		const float LastAge = NormalizedAgeReader[RibbonIndices[EndIndex]];
+		const float SecondToLastAge = NormalizedAgeReader.GetSafe(RibbonIndices[EndIndex - 1], 0.0f);
+		const float LastAge = NormalizedAgeReader.GetSafe(RibbonIndices[EndIndex], 0.0f);
 
 		const float EndTimeStep = LastAge - SecondToLastAge;
 		const float EndTimeOffset = 1 - LastAge < EndTimeStep ? EndTimeStep - (1 - LastAge) : 0;
