Added integer-based equal/not equal ops for enum graph pins

diff --git a/Engine/Plugins/FX/Niagara/Source/NiagaraEditor/Private/NiagaraEditorCommon.cpp b/Engine/Plugins/FX/Niagara/Source/NiagaraEditor/Private/NiagaraEditorCommon.cpp
--- a/Engine/Plugins/FX/Niagara/Source/NiagaraEditor/Private/NiagaraEditorCommon.cpp
+++ b/Engine/Plugins/FX/Niagara/Source/NiagaraEditor/Private/NiagaraEditorCommon.cpp
@@ -1571,6 +1571,51 @@ void FNiagaraOpInfo::Init()
 	OpInfoMap.Add(Op->Name) = Idx;
 
 
+	Idx = OpInfos.AddDefaulted();
+	Op = &OpInfos[Idx];
+	Op->Category = IntCategory;
+	Op->FriendlyName = NSLOCTEXT("NiagaraOpInfo", "EnumEq Name", "Enum Equal");
+	Op->CompactName = FText::FromString("==");
+	Op->Description = NSLOCTEXT("NiagaraOpInfo", "EnumEq Desc", "Result = A == B");
+	Op->Keywords = FText::FromString(TEXT("=="));
+	Op->Inputs.Add(FNiagaraOpInOutInfo(A, IntType, AText, AText, Default_IntZero));
+	Op->Inputs.Add(FNiagaraOpInOutInfo(B, IntType, BText, BText, Default_IntZero));
+	Op->Outputs.Add(FNiagaraOpInOutInfo(Result, FNiagaraTypeDefinition::GetBoolDef(), ResultText, ResultText, Default_IntZero, TEXT("NiagaraAll({0} == {1})")));
+	Op->BuildName(TEXT("EnumEq"), IntCategoryName);
+	Op->bSupportsStaticResolution = true;
+	Op->StaticVariableResolveFunction.BindLambda([=](const TArray<int32>& InPinValues)
+	{
+		if (InPinValues.Num() != 2)
+			return 0;
+
+		int32 Return = InPinValues[0] == InPinValues[1];
+		return Return;
+	});
+	OpInfoMap.Add(Op->Name) = Idx;
+
+	Idx = OpInfos.AddDefaulted();
+	Op = &OpInfos[Idx];
+	Op->Category = IntCategory;
+	Op->FriendlyName = NSLOCTEXT("NiagaraOpInfo", "EnumNEq Name", "Enum Not Equal");
+	Op->CompactName = FText::FromString("!=");
+	Op->Description = NSLOCTEXT("NiagaraOpInfo", "EnumNEq Desc", "Result = A != B");
+	Op->Keywords = FText::FromString(TEXT("!="));
+	Op->Inputs.Add(FNiagaraOpInOutInfo(A, IntType, AText, AText, Default_IntZero));
+	Op->Inputs.Add(FNiagaraOpInOutInfo(B, IntType, BText, BText, Default_IntZero));
+	Op->Outputs.Add(FNiagaraOpInOutInfo(Result, FNiagaraTypeDefinition::GetBoolDef(), ResultText, ResultText, Default_IntZero, TEXT("NiagaraAll({0} != {1})")));
+	Op->BuildName(TEXT("EnumNEq"), IntCategoryName);
+	Op->bSupportsStaticResolution = true;
+	Op->StaticVariableResolveFunction.BindLambda([=](const TArray<int32>& InPinValues)
+	{
+		if (InPinValues.Num() != 2)
+			return 0;
+
+		int32 Return = InPinValues[0] != InPinValues[1];
+		return Return;
+	});
+	OpInfoMap.Add(Op->Name) = Idx;
+
+
 	//////////////////////////////////////////////////////////////////////////
 	// Boolean Only Ops
 	FString Default_BoolZero(TEXT("false"));
