- Don't attempt to capture sim cache on bake if it's complete already

diff --git a/Engine/Plugins/FX/Niagara/Source/NiagaraEditor/Private/NiagaraBakerRendererOutputSimCache.cpp b/Engine/Plugins/FX/Niagara/Source/NiagaraEditor/Private/NiagaraBakerRendererOutputSimCache.cpp
--- a/Engine/Plugins/FX/Niagara/Source/NiagaraEditor/Private/NiagaraBakerRendererOutputSimCache.cpp
+++ b/Engine/Plugins/FX/Niagara/Source/NiagaraEditor/Private/NiagaraBakerRendererOutputSimCache.cpp
@@ -72,15 +72,19 @@ void FNiagaraBakerRendererOutputSimCache::BakeFrame(FNiagaraBakerFeedbackContext
 		return;
 	}
 
-	FNiagaraSimCacheFeedbackContext SimCacheFeedbackContext(false);
-	if ( FrameIndex == 0 )
+	UNiagaraComponent* PreviewComponent = BakerRenderer.GetPreviewComponent();
+	if (PreviewComponent && !PreviewComponent->IsComplete())
 	{
-		BakeSimCache->BeginWrite(BakerOutput->CreateParameters, BakerRenderer.GetPreviewComponent(), SimCacheFeedbackContext);
+		FNiagaraSimCacheFeedbackContext SimCacheFeedbackContext(false);
+		if ( FrameIndex == 0 )
+		{
+			BakeSimCache->BeginWrite(BakerOutput->CreateParameters, BakerRenderer.GetPreviewComponent(), SimCacheFeedbackContext);
+		}
+		BakeSimCache->WriteFrame(BakerRenderer.GetPreviewComponent(), SimCacheFeedbackContext);
+
+		FeedbackContext.Errors.Append(SimCacheFeedbackContext.Errors);
+		FeedbackContext.Warnings.Append(SimCacheFeedbackContext.Warnings);
 	}
-	BakeSimCache->WriteFrame(BakerRenderer.GetPreviewComponent(), SimCacheFeedbackContext);
-
-	FeedbackContext.Errors.Append(SimCacheFeedbackContext.Errors);
-	FeedbackContext.Warnings.Append(SimCacheFeedbackContext.Warnings);
 }
 
 void FNiagaraBakerRendererOutputSimCache::EndBake(FNiagaraBakerFeedbackContext& FeedbackContext, UNiagaraBakerOutput* InBakerOutput)
