- Exclude SimCache editor components from target capture selection
- Change from raw to weak component pointers

diff --git a/Engine/Plugins/FX/Niagara/Source/NiagaraEditor/Private/Widgets/SNiagaraDebugCaptureView.cpp b/Engine/Plugins/FX/Niagara/Source/NiagaraEditor/Private/Widgets/SNiagaraDebugCaptureView.cpp
--- a/Engine/Plugins/FX/Niagara/Source/NiagaraEditor/Private/Widgets/SNiagaraDebugCaptureView.cpp
+++ b/Engine/Plugins/FX/Niagara/Source/NiagaraEditor/Private/Widgets/SNiagaraDebugCaptureView.cpp
@@ -38,7 +38,14 @@ void SNiagaraDebugCaptureView::CreateComponentSelectionMenuContent(FMenuBuilder&
 			continue;
 		}
 
-		if(!NiagaraComponent->GetWorld())
+		UWorld* World = NiagaraComponent->GetWorld();
+		if(!World)
+		{
+			continue;
+		}
+
+		// Exclude components in a preview world with a sim cache attached, this stops us being able to select sim cache editor components
+		if (World->IsPreviewWorld() && NiagaraComponent->GetSimCache() != nullptr)
 		{
 			continue;
 		}
@@ -51,9 +58,9 @@ void SNiagaraDebugCaptureView::CreateComponentSelectionMenuContent(FMenuBuilder&
 			ComponentName,
 			ComponentTooltip,
 			FSlateIcon(),
-			FUIAction(FExecuteAction::CreateLambda([&, NiagaraComponent]()
+			FUIAction(FExecuteAction::CreateLambda([&, WeakNiagaraComponent=MakeWeakObjectPtr(*NiagaraComponent)]()
 			{
-				TargetComponent = *NiagaraComponent;
+				WeakTargetComponent = WeakNiagaraComponent;
 			})));
 	}
 }
@@ -189,7 +196,7 @@ void SNiagaraDebugCaptureView::Construct(const FArguments& InArgs, const TShared
 {
 	NumFrames = FMath::Max(1, GetDefault<UNiagaraSettings>()->QuickSimCacheCaptureFrameCount);
 
-	TargetComponent = InSystemViewModel->GetPreviewComponent();
+	WeakTargetComponent = InSystemViewModel->GetPreviewComponent();
 	SimCacheViewModel = InSimCacheViewModel;
 	SystemViewModel = InSystemViewModel;
 
@@ -302,6 +309,7 @@ void SNiagaraDebugCaptureView::OnCaptureSelected()
 
 void SNiagaraDebugCaptureView::OnSingleFrameSelected()
 {
+	UNiagaraComponent* TargetComponent = WeakTargetComponent.Get();
 	if(!bIsCaptureActive && TargetComponent && CapturedCache.IsValid())
 	{
 		TSharedPtr<ISequencer> Sequencer = SystemViewModel.Get()->GetSequencer();
@@ -318,12 +326,9 @@ void SNiagaraDebugCaptureView::OnSingleFrameSelected()
 
 		if(OutCache)
 		{
-			
 			// 60fps
 			FFrameRate SystemFrameRate (60, 1);
-
 			FFrameTime NewTime = StartTime.ConvertTo(SystemFrameRate) + FQualifiedFrameTime(1, SystemFrameRate).Time;
-
 			
 			TargetComponent->SetDesiredAge(CurrentAge + 0.01666f);
 			
@@ -336,6 +341,7 @@ void SNiagaraDebugCaptureView::OnSingleFrameSelected()
 
 void SNiagaraDebugCaptureView::OnMultiFrameSelected()
 {
+	UNiagaraComponent* TargetComponent = WeakTargetComponent.Get();
 	if(!bIsCaptureActive && TargetComponent)
 	{
 		UNiagaraSimCache* MultiFrameCache = NewObject<UNiagaraSimCache>(GetTransientPackage(), FNiagaraEditorUtilities::GetUniqueObjectName<UNiagaraSimCache>(GetTransientPackage(), "TempCache"));
diff --git a/Engine/Plugins/FX/Niagara/Source/NiagaraEditor/Private/Widgets/SNiagaraDebugCaptureView.h b/Engine/Plugins/FX/Niagara/Source/NiagaraEditor/Private/Widgets/SNiagaraDebugCaptureView.h
--- a/Engine/Plugins/FX/Niagara/Source/NiagaraEditor/Private/Widgets/SNiagaraDebugCaptureView.h
+++ b/Engine/Plugins/FX/Niagara/Source/NiagaraEditor/Private/Widgets/SNiagaraDebugCaptureView.h
@@ -51,7 +51,7 @@ protected:
 	FNiagaraSimCacheCapture SimCacheCapture;
 	TWeakObjectPtr<UNiagaraSimCache> CapturedCache = nullptr;
 
-	UNiagaraComponent* TargetComponent = nullptr;
+	TWeakObjectPtr<UNiagaraComponent> WeakTargetComponent;
 	int32 NumFrames = 1;
 	bool bIsCaptureActive = false;
 	TSharedPtr<FNiagaraSimCacheViewModel> SimCacheViewModel;
