Fixed a bug where undo wouldn't work when deleting a parameter from a module script

diff --git a/Engine/Plugins/FX/Niagara/Source/NiagaraEditor/Private/ViewModels/NiagaraParameterPanelViewModel.cpp b/Engine/Plugins/FX/Niagara/Source/NiagaraEditor/Private/ViewModels/NiagaraParameterPanelViewModel.cpp
--- a/Engine/Plugins/FX/Niagara/Source/NiagaraEditor/Private/ViewModels/NiagaraParameterPanelViewModel.cpp
+++ b/Engine/Plugins/FX/Niagara/Source/NiagaraEditor/Private/ViewModels/NiagaraParameterPanelViewModel.cpp
@@ -2609,9 +2609,10 @@ void FNiagaraScriptToolkitParameterPanelViewModel::DeleteParameters(const TArray
 			{
 				VariableObjectSelection->ClearSelectedObjects();
 			}
+			Graph->Modify();
 			Graph->RemoveParameter(ItemToDelete.GetVariable());
 			
-			if(!bAnyChange)
+			if (!bAnyChange)
 			{
 				bAnyChange = true;
 			}
