Changed default Niagara note contents text to "Title" & "Text" respectively

diff --git a/Engine/Plugins/FX/Niagara/Source/NiagaraEditorWidgets/Private/Stack/SNiagaraStackTableRow.cpp b/Engine/Plugins/FX/Niagara/Source/NiagaraEditorWidgets/Private/Stack/SNiagaraStackTableRow.cpp
--- a/Engine/Plugins/FX/Niagara/Source/NiagaraEditorWidgets/Private/Stack/SNiagaraStackTableRow.cpp
+++ b/Engine/Plugins/FX/Niagara/Source/NiagaraEditorWidgets/Private/Stack/SNiagaraStackTableRow.cpp
@@ -556,8 +556,8 @@ void SNiagaraStackTableRow::AddStackNote() const
 		StackEntry->GetStackEditorData().Modify();
 		
 		FNiagaraStackNoteData NewStackNote;
-		NewStackNote.MessageHeader = LOCTEXT("DefaultStackNoteHeader", "My Note");
-		NewStackNote.Message = LOCTEXT("DefaultStackNoteMessage", "My Message");
+		NewStackNote.MessageHeader = LOCTEXT("DefaultStackNoteHeader", "Title");
+		NewStackNote.Message = LOCTEXT("DefaultStackNoteMessage", "Text");
 
 		StackEntry->GetStackEditorData().AddOrReplaceStackNote(StackEntry->GetStackEditorDataKey(), NewStackNote);
 		StackEntry->RefreshChildren();
