Nicer names for debug capture tabs

diff --git a/Engine/Plugins/FX/Niagara/Source/NiagaraEditor/Private/Widgets/SNiagaraDebugCaptureView.cpp b/Engine/Plugins/FX/Niagara/Source/NiagaraEditor/Private/Widgets/SNiagaraDebugCaptureView.cpp
--- a/Engine/Plugins/FX/Niagara/Source/NiagaraEditor/Private/Widgets/SNiagaraDebugCaptureView.cpp
+++ b/Engine/Plugins/FX/Niagara/Source/NiagaraEditor/Private/Widgets/SNiagaraDebugCaptureView.cpp
@@ -192,6 +192,13 @@ FSlateIcon SNiagaraDebugCaptureView::GetCaptureIcon()
 	}
 }
 
+
+
+FName GetTempCacheName(const FString& SystemName)
+{
+	return FNiagaraEditorUtilities::GetUniqueObjectName<UNiagaraSimCache>(GetTransientPackage(), TEXT("TempCache_") + SystemName);
+}
+
 void SNiagaraDebugCaptureView::Construct(const FArguments& InArgs, const TSharedRef<FNiagaraSystemViewModel> InSystemViewModel, const TSharedRef<FNiagaraSimCacheViewModel> InSimCacheViewModel)
 {
 	NumFrames = FMath::Max(1, GetDefault<UNiagaraSettings>()->QuickSimCacheCaptureFrameCount);
@@ -200,7 +207,7 @@ void SNiagaraDebugCaptureView::Construct(const FArguments& InArgs, const TShared
 	SimCacheViewModel = InSimCacheViewModel;
 	SystemViewModel = InSystemViewModel;
 
-	CapturedCache = NewObject<UNiagaraSimCache>(GetTransientPackage(), FNiagaraEditorUtilities::GetUniqueObjectName<UNiagaraSimCache>(GetTransientPackage(), "TempCache"));
+	CapturedCache = NewObject<UNiagaraSimCache>(GetTransientPackage(), GetTempCacheName(InSystemViewModel->GetSystem().GetName()));
 	SimCacheViewModel.Get()->Initialize(CapturedCache);
 	CapturedCache->SetFlags(RF_Transient);
 
@@ -344,7 +351,7 @@ void SNiagaraDebugCaptureView::OnMultiFrameSelected()
 	UNiagaraComponent* TargetComponent = WeakTargetComponent.Get();
 	if(!bIsCaptureActive && TargetComponent)
 	{
-		UNiagaraSimCache* MultiFrameCache = NewObject<UNiagaraSimCache>(GetTransientPackage(), FNiagaraEditorUtilities::GetUniqueObjectName<UNiagaraSimCache>(GetTransientPackage(), "TempCache"));
+		UNiagaraSimCache* MultiFrameCache = NewObject<UNiagaraSimCache>(GetTransientPackage(), GetTempCacheName(TargetComponent->GetFXSystemAsset()->GetName()));
 		MultiFrameCache->SetFlags(RF_Transient);
 		const FNiagaraSimCacheCreateParameters CreateParameters;
 		
