Fixed a bug where scratch pads didn't always enable the apply button after a change

diff --git a/Engine/Plugins/FX/Niagara/Source/NiagaraEditor/Private/ViewModels/NiagaraScratchPadScriptViewModel.cpp b/Engine/Plugins/FX/Niagara/Source/NiagaraEditor/Private/ViewModels/NiagaraScratchPadScriptViewModel.cpp
--- a/Engine/Plugins/FX/Niagara/Source/NiagaraEditor/Private/ViewModels/NiagaraScratchPadScriptViewModel.cpp
+++ b/Engine/Plugins/FX/Niagara/Source/NiagaraEditor/Private/ViewModels/NiagaraScratchPadScriptViewModel.cpp
@@ -3,16 +3,11 @@
 #include "ViewModels/NiagaraScratchPadScriptViewModel.h"
 #include "UObject/Linker.h"
 #include "ViewModels/Stack/NiagaraStackGraphUtilities.h"
-#include "ViewModels/Stack/NiagaraStackFunctionInputCollection.h"
 #include "ViewModels/NiagaraParameterPanelViewModel.h"
 #include "NiagaraNodeFunctionCall.h"
-#include "NiagaraEmitter.h"
 #include "NiagaraEditorUtilities.h"
-#include "NiagaraEditorModule.h"
 #include "NiagaraObjectSelection.h"
-#include "NiagaraParameterDefinitions.h"
 #include "NiagaraScriptSource.h"
-#include "NiagaraSystem.h"
 #include "ViewModels/NiagaraScratchPadUtilities.h"
 #include "ViewModels/NiagaraScriptGraphViewModel.h"
 #include "ViewModels/NiagaraSystemViewModel.h"
@@ -67,6 +62,8 @@ void FNiagaraScratchPadScriptViewModel::Initialize(UNiagaraScript* Script, UNiag
 		EditScript.Script = CastChecked<UNiagaraScript>(StaticDuplicateObject(Script, GetTransientPackage()));
 		SetScript(EditScript);
 	}
+	LastAppliedChangeID = GetScriptChangeID();
+	
 	UNiagaraScriptSource* EditScriptSource = CastChecked<UNiagaraScriptSource>(EditScript.Script->GetLatestSource());
 	OnGraphNeedsRecompileHandle = EditScriptSource->NodeGraph->AddOnGraphNeedsRecompileHandler(FOnGraphChanged::FDelegate::CreateSP(this, &FNiagaraScratchPadScriptViewModel::OnScriptGraphChanged));
 	EditScript.Script->OnPropertyChanged().AddSP(this, &FNiagaraScratchPadScriptViewModel::OnScriptPropertyChanged);
@@ -220,17 +217,16 @@ void FNiagaraScratchPadScriptViewModel::SetEditorHeight(float InEditorHeight)
 
 bool FNiagaraScratchPadScriptViewModel::HasUnappliedChanges() const
 {
-	return bHasPendingChanges;
+	return bHasPendingChanges || (LastAppliedChangeID.IsValid() && LastAppliedChangeID != GetScriptChangeID());
 }
 
 void FNiagaraScratchPadScriptViewModel::ApplyChanges()
 {
 	ResetLoaders(OriginalScript->GetOutermost()); // Make sure that we're not going to get invalid version number linkers into the package we are going into. 
 
-	OriginalScript = (UNiagaraScript*)StaticDuplicateObject(EditScript.Script, OriginalScript->GetOuter(), OriginalScript->GetFName(),
-		RF_AllFlags,
-		OriginalScript->GetClass());
+	OriginalScript = Cast<UNiagaraScript>(StaticDuplicateObject(EditScript.Script, OriginalScript->GetOuter(), OriginalScript->GetFName(), RF_AllFlags, OriginalScript->GetClass()));
 	bHasPendingChanges = false;
+	LastAppliedChangeID = GetScriptChangeID();
 
 	TArray<UNiagaraNodeFunctionCall*> FunctionCallNodesToRefresh;
 	FNiagaraEditorUtilities::GetReferencingFunctionCallNodes(OriginalScript, FunctionCallNodesToRefresh);
@@ -300,6 +296,15 @@ void FNiagaraScratchPadScriptViewModel::OnScriptPropertyChanged(FPropertyChanged
 	}
 }
 
+FGuid FNiagaraScratchPadScriptViewModel::GetScriptChangeID() const
+{
+	if (Source.IsValid() && Source->NodeGraph)
+	{
+		return Source->NodeGraph->GetChangeID();
+	}
+	return FGuid();
+}
+
 FNiagaraScratchPadScriptViewModel::FOnNodeIDFocusRequested& FNiagaraScratchPadScriptViewModel::OnNodeIDFocusRequested()
 {
 	return OnNodeIDFocusRequestedDelegate;
diff --git a/Engine/Plugins/FX/Niagara/Source/NiagaraEditor/Public/ViewModels/NiagaraScratchPadScriptViewModel.h b/Engine/Plugins/FX/Niagara/Source/NiagaraEditor/Public/ViewModels/NiagaraScratchPadScriptViewModel.h
--- a/Engine/Plugins/FX/Niagara/Source/NiagaraEditor/Public/ViewModels/NiagaraScratchPadScriptViewModel.h
+++ b/Engine/Plugins/FX/Niagara/Source/NiagaraEditor/Public/ViewModels/NiagaraScratchPadScriptViewModel.h
@@ -109,6 +109,8 @@ private:
 
 	NIAGARAEDITOR_API void OnScriptPropertyChanged(FPropertyChangedEvent& PropertyChangedEvent);
 
+	FGuid GetScriptChangeID() const;
+
 	bool bIsPendingRename;
 	bool bIsPinned;
 	float EditorHeight;
@@ -119,6 +121,7 @@ private:
 	TWeakPtr<FNiagaraSystemViewModel> SystemViewModel;
 
 	bool bHasPendingChanges;
+	FGuid LastAppliedChangeID;
 
 	TSharedPtr<FUICommandList> ParameterPanelCommands;
 	TSharedPtr<FNiagaraScriptToolkitParameterPanelViewModel> ParameterPaneViewModel;
