Fix issue causing changes to an Island NDC's systems to not be applied correctly until editor restart.

diff --git a/Engine/Plugins/FX/Niagara/Source/Niagara/Private/NiagaraDataChannel_Islands.cpp b/Engine/Plugins/FX/Niagara/Source/Niagara/Private/NiagaraDataChannel_Islands.cpp
--- a/Engine/Plugins/FX/Niagara/Source/Niagara/Private/NiagaraDataChannel_Islands.cpp
+++ b/Engine/Plugins/FX/Niagara/Source/Niagara/Private/NiagaraDataChannel_Islands.cpp
@@ -31,7 +31,7 @@ TConstArrayView<TObjectPtr<UNiagaraSystem>> UNiagaraDataChannel_Islands::GetSyst
 
 void UNiagaraDataChannel_Islands::AsyncLoadSystems()const
 {
-	if(INiagaraModule::DataChannelsEnabled() && SystemsInternal.Num() != Systems.Num() && Systems.Num() > 0 && !IsRunningDedicatedServer())
+	if(INiagaraModule::DataChannelsEnabled() && !IsRunningDedicatedServer())
 	{
 		TArray<FSoftObjectPath> Requests;
 		for (const TSoftObjectPtr<UNiagaraSystem>& SoftSys : Systems)
