Make NDC assets a BP Type so BPs can store refs and soft refs to them.

diff --git a/Engine/Plugins/FX/Niagara/Source/Niagara/Public/NiagaraDataChannelPublic.h b/Engine/Plugins/FX/Niagara/Source/Niagara/Public/NiagaraDataChannelPublic.h
--- a/Engine/Plugins/FX/Niagara/Source/Niagara/Public/NiagaraDataChannelPublic.h
+++ b/Engine/Plugins/FX/Niagara/Source/Niagara/Public/NiagaraDataChannelPublic.h
@@ -25,7 +25,7 @@ Blueprint and C++ code can also read from and write to data channels using its A
 EXPERIMENTAL:
 Data Channels are currently experimental and undergoing heavy development.
  */
-UCLASS(Experimental, DisplayName = "Niagara Data Channel", MinimalAPI)
+UCLASS(Experimental, BlueprintType, DisplayName = "Niagara Data Channel", MinimalAPI)
 class UNiagaraDataChannelAsset : public UObject
 {
 	GENERATED_BODY()
