Fixed issue where calls to DeactivateImmediate would not properly complete systems that were culled by scalability.

diff --git a/Engine/Plugins/FX/Niagara/Source/Niagara/Private/NiagaraComponent.cpp b/Engine/Plugins/FX/Niagara/Source/Niagara/Private/NiagaraComponent.cpp
--- a/Engine/Plugins/FX/Niagara/Source/Niagara/Private/NiagaraComponent.cpp
+++ b/Engine/Plugins/FX/Niagara/Source/Niagara/Private/NiagaraComponent.cpp
@@ -1530,7 +1530,7 @@ bool UNiagaraComponent::IsComplete() const
 void UNiagaraComponent::DeactivateImmediateInternal(bool bIsScalabilityCull)
 {
 	SCOPE_CYCLE_COUNTER(STAT_NiagaraComponentDeactivate);
-	Super::Deactivate();
+	CSV_SCOPED_TIMING_STAT_EXCLUSIVE(Effects);
 
 	bool bWasCulledByScalability = bIsCulledByScalability;
 
@@ -1549,16 +1549,28 @@ void UNiagaraComponent::DeactivateImmediateInternal(bool bIsScalabilityCull)
 		UnregisterWithScalabilityManager();
 	}
 
-	SetActiveFlag(false);
 
-	if (SystemInstanceController)
+	if (IsActive() && SystemInstanceController)
 	{
+		Super::Deactivate();
 		SystemInstanceController->Deactivate(true);
 	}
-	else if (bWasCulledByScalability && !bIsCulledByScalability)//We were culled by scalability but no longer, ensure we've handled completion correctly. E.g. returned to the pool etc.
+	else 
 	{
-		OnSystemComplete(true);
+		Super::Deactivate();
+
+		if (bWasCulledByScalability && !bIsCulledByScalability)//We were culled by scalability but no longer, ensure we've handled completion correctly. E.g. returned to the pool etc.
+		{
+			OnSystemComplete(true);
+		}
+	}
+
+#if !UE_BUILD_SHIPPING
+	if(IsActive())
+	{
+		UE_LOG(LogNiagara, Error, TEXT("NiagaraSystem is still active after it should have been deactivated.\nSystem: %s\nComponent: %s"), *GetNameSafe(GetAsset()), *GetNameSafe(this));
 	}
+#endif
 }
 
 bool UNiagaraComponent::ShouldPreCull()
