Change default for NDC DIs overriding spawn group when spawning particles.
Preferred method is to now use the GetNDCSpawnData function.

diff --git a/Engine/Plugins/FX/Niagara/Source/Niagara/Internal/DataInterface/NiagaraDataInterfaceDataChannelRead.h b/Engine/Plugins/FX/Niagara/Source/Niagara/Internal/DataInterface/NiagaraDataInterfaceDataChannelRead.h
--- a/Engine/Plugins/FX/Niagara/Source/Niagara/Internal/DataInterface/NiagaraDataInterfaceDataChannelRead.h
+++ b/Engine/Plugins/FX/Niagara/Source/Niagara/Internal/DataInterface/NiagaraDataInterfaceDataChannelRead.h
@@ -100,11 +100,12 @@ public:
 	Calling GetNDCSpawnInfo() in the particle spawn script to get the spawning NDC Index is prefered.
 	*/
 	UPROPERTY(EditAnywhere, Category = "Spawning", AdvancedDisplay)
-	bool bOverrideSpawnGroupToDataChannelIndex = true;
+	bool bOverrideSpawnGroupToDataChannelIndex = false;
 
 	//UObject Interface
 	NIAGARA_API virtual void PostInitProperties() override;
 	NIAGARA_API virtual void BeginDestroy() override;
+	NIAGARA_API virtual void Serialize(FArchive& Ar) override;
 	//UObject Interface End
 
 	//UNiagaraDataInterface Interface
diff --git a/Engine/Plugins/FX/Niagara/Source/Niagara/Private/DataInterface/NiagaraDataInterfaceDataChannelRead.cpp b/Engine/Plugins/FX/Niagara/Source/Niagara/Private/DataInterface/NiagaraDataInterfaceDataChannelRead.cpp
--- a/Engine/Plugins/FX/Niagara/Source/Niagara/Private/DataInterface/NiagaraDataInterfaceDataChannelRead.cpp
+++ b/Engine/Plugins/FX/Niagara/Source/Niagara/Private/DataInterface/NiagaraDataInterfaceDataChannelRead.cpp
@@ -9,6 +9,7 @@
 #include "NiagaraModule.h"
 #include "NiagaraShaderParametersBuilder.h"
 
+#include "NiagaraCustomVersion.h"
 #include "NiagaraCommon.h"
 #include "NiagaraSystem.h"
 #include "NiagaraWorldManager.h"
@@ -505,6 +506,19 @@ void UNiagaraDataInterfaceDataChannelRead::BeginDestroy()
 	Super::BeginDestroy();
 }
 
+void UNiagaraDataInterfaceDataChannelRead::Serialize(FArchive& Ar)
+{
+	Ar.UsingCustomVersion(FNiagaraCustomVersion::GUID);
+	const int32 NiagaraVersion = Ar.CustomVer(FNiagaraCustomVersion::GUID);
+
+	//Before we serialize in the properties we will restore any old default values from previous versions.
+	if(NiagaraVersion < FNiagaraCustomVersion::NDCSpawnGroupOverrideDisabledByDefault)
+	{
+		bOverrideSpawnGroupToDataChannelIndex = true;
+	}
+
+	Super::Serialize(Ar);
+}
 
 bool UNiagaraDataInterfaceDataChannelRead::InitPerInstanceData(void* PerInstanceData, FNiagaraSystemInstance* SystemInstance)
 {
diff --git a/Engine/Plugins/FX/Niagara/Source/NiagaraCore/Public/NiagaraCustomVersion.h b/Engine/Plugins/FX/Niagara/Source/NiagaraCore/Public/NiagaraCustomVersion.h
--- a/Engine/Plugins/FX/Niagara/Source/NiagaraCore/Public/NiagaraCustomVersion.h
+++ b/Engine/Plugins/FX/Niagara/Source/NiagaraCore/Public/NiagaraCustomVersion.h
@@ -200,6 +200,8 @@ struct FNiagaraCustomVersion
 		SimCache_BulkDataVersion1, // Sim Cache moved to bulk data by default
 
 		InheritanceUxRefactor, // Decoupling of 'Template' and 'Inheritance'
+
+		NDCSpawnGroupOverrideDisabledByDefault, // NDC Read DIs will not override spawn group by default when spawning particles. Old content will remain unchanged.
 		
 		// DO NOT ADD A NEW VERSION UNLESS YOU HAVE TALKED TO THE NIAGARA LEAD. Mismanagement of these versions can lead to data loss if it is adjusted in multiple streams simultaneously.
 		// -----<new versions can be added above this line>  -------------------------------------------------
