- Stateless merge [at] CL 31459790

diff --git a/Engine/Plugins/FX/Niagara/Source/Niagara/Internal/Stateless/NiagaraStatelessDistribution.h b/Engine/Plugins/FX/Niagara/Source/Niagara/Internal/Stateless/NiagaraStatelessDistribution.h
--- a/Engine/Plugins/FX/Niagara/Source/Niagara/Internal/Stateless/NiagaraStatelessDistribution.h
+++ b/Engine/Plugins/FX/Niagara/Source/Niagara/Internal/Stateless/NiagaraStatelessDistribution.h
@@ -43,6 +43,7 @@ struct FNiagaraDistributionBase
 	int32 MaxLutSampleCount = 64;
 
 	virtual bool AllowCurves() const { return true; }
+	virtual bool DisplayAsColor() const { return false; }
 	virtual void UpdateValuesFromDistribution() { }
 
 	static void PostEditChangeProperty(UObject* OwnerObject, FPropertyChangedEvent& PropertyChangedEvent);
@@ -215,6 +216,7 @@ struct FNiagaraDistributionColor : public FNiagaraDistributionBase
 	NIAGARA_API FNiagaraStatelessRangeColor CalculateRange(const FLinearColor& Default = FLinearColor::White) const;
 
 #if WITH_EDITORONLY_DATA
+	virtual bool DisplayAsColor() const override { return true; }
 	virtual void UpdateValuesFromDistribution() override;
 #endif
 };
diff --git a/Engine/Plugins/FX/Niagara/Source/NiagaraEditor/Private/NiagaraEditorStyle.cpp b/Engine/Plugins/FX/Niagara/Source/NiagaraEditor/Private/NiagaraEditorStyle.cpp
--- a/Engine/Plugins/FX/Niagara/Source/NiagaraEditor/Private/NiagaraEditorStyle.cpp
+++ b/Engine/Plugins/FX/Niagara/Source/NiagaraEditor/Private/NiagaraEditorStyle.cpp
@@ -440,10 +440,16 @@ void FNiagaraEditorStyle::InitIcons()
 
 	Set("NiagaraEditor.DistributionEditor.UniformConstant", new IMAGE_BRUSH_SVG("Icons/DistributionEditor/UniformConstant", Icon20x20));
 	Set("NiagaraEditor.DistributionEditor.NonUniformConstant", new IMAGE_BRUSH_SVG("Icons/DistributionEditor/NonUniformConstant", Icon20x20));
+	Set("NiagaraEditor.DistributionEditor.ColorConstant", new IMAGE_BRUSH_SVG("Icons/DistributionEditor/ColorConstant", Icon20x20));
+
 	Set("NiagaraEditor.DistributionEditor.UniformRange", new IMAGE_BRUSH_SVG("Icons/DistributionEditor/UniformRange", Icon20x20));
 	Set("NiagaraEditor.DistributionEditor.NonUniformRange", new IMAGE_BRUSH_SVG("Icons/DistributionEditor/NonUniformRange", Icon20x20));
+	Set("NiagaraEditor.DistributionEditor.ColorRange", new IMAGE_BRUSH_SVG("Icons/DistributionEditor/ColorRange", Icon20x20));
+
 	Set("NiagaraEditor.DistributionEditor.UniformCurve", new IMAGE_BRUSH_SVG("Icons/DistributionEditor/UniformCurve", Icon20x20));
 	Set("NiagaraEditor.DistributionEditor.NonUniformCurve", new IMAGE_BRUSH_SVG("Icons/DistributionEditor/NonUniformCurve", Icon20x20));
+	Set("NiagaraEditor.DistributionEditor.ColorGradient", new IMAGE_BRUSH_SVG("Icons/DistributionEditor/ColorGradient", Icon20x20));
+	
 	Set("NiagaraEditor.DistributionEditor.CurveKeyHandle", new IMAGE_BRUSH_SVG("Icons/DistributionEditor/CurveKeyHandle", Icon20x20));
 }
 
diff --git a/Engine/Plugins/FX/Niagara/Source/NiagaraEditor/Private/NiagaraEditorUtilities.cpp b/Engine/Plugins/FX/Niagara/Source/NiagaraEditor/Private/NiagaraEditorUtilities.cpp
--- a/Engine/Plugins/FX/Niagara/Source/NiagaraEditor/Private/NiagaraEditorUtilities.cpp
+++ b/Engine/Plugins/FX/Niagara/Source/NiagaraEditor/Private/NiagaraEditorUtilities.cpp
@@ -4629,6 +4629,12 @@ UNiagaraDataInterface* FNiagaraEditorUtilities::GetResolvedRuntimeInstanceForEdi
 			}
 		}
 	}
+	UNiagaraNodeFunctionCall* OuterFunctionCall = EditorDataInterfaceInstance.GetTypedOuter<UNiagaraNodeFunctionCall>();
+	if (OuterFunctionCall != nullptr)
+	{
+		// Placeholder data interfaces are outered to the function call defines their input.
+
+	}
 	return nullptr;
 }
 
diff --git a/Engine/Plugins/FX/Niagara/Source/NiagaraEditor/Private/SNiagaraParameterEditor.cpp b/Engine/Plugins/FX/Niagara/Source/NiagaraEditor/Private/SNiagaraParameterEditor.cpp
--- a/Engine/Plugins/FX/Niagara/Source/NiagaraEditor/Private/SNiagaraParameterEditor.cpp
+++ b/Engine/Plugins/FX/Niagara/Source/NiagaraEditor/Private/SNiagaraParameterEditor.cpp
@@ -2,7 +2,7 @@
 
 #include "SNiagaraParameterEditor.h"
 
-const float SNiagaraParameterEditor::DefaultInputSize = 110.0f;
+const float SNiagaraParameterEditor::DefaultInputSize = 125.0f;
 
 void SNiagaraParameterEditor::Construct(const FArguments& InArgs)
 {
diff --git a/Engine/Plugins/FX/Niagara/Source/NiagaraEditor/Private/Stateless/NiagaraDistributionPropertyCustomization.cpp b/Engine/Plugins/FX/Niagara/Source/NiagaraEditor/Private/Stateless/NiagaraDistributionPropertyCustomization.cpp
--- a/Engine/Plugins/FX/Niagara/Source/NiagaraEditor/Private/Stateless/NiagaraDistributionPropertyCustomization.cpp
+++ b/Engine/Plugins/FX/Niagara/Source/NiagaraEditor/Private/Stateless/NiagaraDistributionPropertyCustomization.cpp
@@ -28,17 +28,19 @@ public:
 	{
 		if (IsValid())
 		{
-			EditorMode = GetDistributionEditorModeFromSourceMode(SourceNumChannels, SourceDistribution->Mode);
-			EditorNumChannels = FNiagaraDistributionEditorUtilities::IsUniform(EditorMode) ? 1 : SourceNumChannels;
-
 			const FName DisableCurveDistributionName("DisableCurveDistribution");
 			const FName DisableUniformDistributionName("DisableUniformDistribution");
 			const FName DisableNonUniformDistributionName("DisableNonUniformDistribution");
 			const FName DisableRangeDistributionName("DisableRangeDistribution");
+			const FName DisplayAsColorDistributionName("DisplayAsColorDistribution");
 			bAllowUniform = PropertyHandle ? PropertyHandle->HasMetaData(DisableUniformDistributionName) == false : false;
 			bAllowNonUniform = PropertyHandle ? PropertyHandle->HasMetaData(DisableNonUniformDistributionName) == false : false;
 			bAllowRange = PropertyHandle ? PropertyHandle->HasMetaData(DisableRangeDistributionName) == false : false;
 			bAllowCurves = InDistribution->AllowCurves() && (PropertyHandle ? PropertyHandle->HasMetaData(DisableCurveDistributionName) == false : false);
+			bDisplayAsColor = InDistribution->DisplayAsColor() || (InNumChannels >= 3 && (PropertyHandle ? PropertyHandle->HasMetaData(DisplayAsColorDistributionName) : false));
+
+			EditorMode = GetDistributionEditorModeFromSourceMode(SourceNumChannels, bDisplayAsColor, SourceDistribution->Mode);
+			EditorNumChannels = FNiagaraDistributionEditorUtilities::IsUniform(EditorMode) ? 1 : SourceNumChannels;
 		}
 	}
 
@@ -48,16 +50,33 @@ public:
 
 	virtual FText GetChannelDisplayName(int32 ChannelIndex) const override 
 	{
-		switch(ChannelIndex)
+		if (bDisplayAsColor)
 		{
-		case 0:
-			return LOCTEXT("XChannelName", "X");
-		case 1:
-			return LOCTEXT("YChannelName", "Y");
-		case 2:
-			return LOCTEXT("ZChannelName", "Z");
-		case 3:
-			return LOCTEXT("WChannelName", "W");
+			switch (ChannelIndex)
+			{
+			case 0:
+				return LOCTEXT("RChannelName", "R");
+			case 1:
+				return LOCTEXT("GChannelName", "G");
+			case 2:
+				return LOCTEXT("BChannelName", "B");
+			case 3:
+				return LOCTEXT("AChannelName", "A");
+			}
+		}
+		else
+		{
+			switch(ChannelIndex)
+			{
+			case 0:
+				return LOCTEXT("XChannelName", "X");
+			case 1:
+				return LOCTEXT("YChannelName", "Y");
+			case 2:
+				return LOCTEXT("ZChannelName", "Z");
+			case 3:
+				return LOCTEXT("WChannelName", "W");
+			}
 		}
 		return FText();
 	}
@@ -83,7 +102,17 @@ public:
 
 	virtual void GetSupportedDistributionModes(TArray<ENiagaraDistributionEditorMode>& OutSupportedModes) const override
 	{
-		if (SourceNumChannels == 1)
+		if (bDisplayAsColor)
+		{
+			OutSupportedModes.Add(ENiagaraDistributionEditorMode::ColorConstant);
+			OutSupportedModes.Add(ENiagaraDistributionEditorMode::ColorRange);
+			if (bAllowCurves)
+			{
+				OutSupportedModes.Add(ENiagaraDistributionEditorMode::NonUniformCurve);
+				OutSupportedModes.Add(ENiagaraDistributionEditorMode::ColorGradient);
+			}
+		}
+		else if (SourceNumChannels == 1)
 		{
 			if (bAllowUniform)
 			{
@@ -127,9 +156,23 @@ public:
 		}
 	}
 
-	static ENiagaraDistributionEditorMode GetDistributionEditorModeFromSourceMode(int32 InSourceNumChannels, ENiagaraDistributionMode InSourceMode)
+	static ENiagaraDistributionEditorMode GetDistributionEditorModeFromSourceMode(int32 InSourceNumChannels, bool bInDisplayAsColor, ENiagaraDistributionMode InSourceMode)
 	{
-		if (InSourceNumChannels == 1)
+		if (bInDisplayAsColor)
+		{
+			switch (InSourceMode)
+			{
+			case ENiagaraDistributionMode::NonUniformConstant:
+				return ENiagaraDistributionEditorMode::ColorConstant;
+			case ENiagaraDistributionMode::NonUniformRange:
+				return ENiagaraDistributionEditorMode::ColorRange;
+			case ENiagaraDistributionMode::NonUniformCurve:
+				return ENiagaraDistributionEditorMode::ColorGradient;
+			default:
+				return ENiagaraDistributionEditorMode::ColorConstant;
+			}
+		}
+		else if (InSourceNumChannels == 1)
 		{
 			switch (InSourceMode)
 			{
@@ -173,16 +216,19 @@ public:
 		case ENiagaraDistributionEditorMode::UniformConstant:
 			return ENiagaraDistributionMode::UniformConstant;
 		case ENiagaraDistributionEditorMode::NonUniformConstant:
+		case ENiagaraDistributionEditorMode::ColorConstant:
 			return ENiagaraDistributionMode::NonUniformConstant;
 		case ENiagaraDistributionEditorMode::Range:
 		case ENiagaraDistributionEditorMode::UniformRange:
 			return ENiagaraDistributionMode::UniformRange;
 		case ENiagaraDistributionEditorMode::NonUniformRange:
+		case ENiagaraDistributionEditorMode::ColorRange:
 			return ENiagaraDistributionMode::NonUniformRange;
 		case ENiagaraDistributionEditorMode::Curve:
 		case ENiagaraDistributionEditorMode::UniformCurve:
 			return ENiagaraDistributionMode::UniformCurve;
 		case ENiagaraDistributionEditorMode::NonUniformCurve:
+		case ENiagaraDistributionEditorMode::ColorGradient:
 			return ENiagaraDistributionMode::NonUniformCurve;
 		default:
 			return ENiagaraDistributionMode::UniformConstant;
@@ -193,28 +239,34 @@ public:
 
 	virtual void SetDistributionMode(ENiagaraDistributionEditorMode InEditorMode) override
 	{
+		if (InEditorMode == EditorMode)
+		{
+			return;
+		}
+
 		UObject* OwnerObject = OwnerObjectWeak.Get();
 		if (OwnerObject == nullptr)
 		{
 			return;
 		}
 
-		ENiagaraDistributionMode NewMode = GetDistributionSourceModeFromEditorMode(InEditorMode);
-		if (SourceDistribution->Mode != NewMode)
-		{
-			int32 NewEditorNumChannels = FNiagaraDistributionEditorUtilities::IsUniform(InEditorMode) ? 1 : SourceNumChannels;
+		int32 NewEditorNumChannels = FNiagaraDistributionEditorUtilities::IsUniform(InEditorMode) ? 1 : SourceNumChannels;
 
+		ENiagaraDistributionMode NewSourceMode = GetDistributionSourceModeFromEditorMode(InEditorMode);
+		if (SourceDistribution->Mode != NewSourceMode)
+		{
 			const FScopedTransaction Transaction(LOCTEXT("SetDistributionModeTransaction", "Set distribution mode"));
 			OwnerObject->Modify();
 			PropertyHandle->NotifyPreChange();
 
-			MigrateDataFromModeChange(NewMode, NewEditorNumChannels);
-			SourceDistribution->Mode = NewMode;
-			EditorNumChannels = NewEditorNumChannels;
-			EditorMode = InEditorMode;
+			MigrateDataFromModeChange(NewSourceMode, NewEditorNumChannels);
+			SourceDistribution->Mode = NewSourceMode;
 
 			PropertyHandle->NotifyPostChange(EPropertyChangeType::ValueSet);
 		}
+
+		EditorNumChannels = NewEditorNumChannels;
+		EditorMode = InEditorMode;
 	}
 
 	float GetConstantOrRangeValue(int32 ChannelIndex, int32 ValueIndex) const override
@@ -234,50 +286,116 @@ public:
 	void SetConstantOrRangeValue(int32 ChannelIndex, int32 ValueIndex, float InValue) override
 	{
 		UObject* OwnerObject = OwnerObjectWeak.Get();
-		if (OwnerObject != nullptr &&
-			GetConstantOrRangeValue(ChannelIndex, ValueIndex) != InValue &&
-			ChannelIndex < EditorNumChannels)
+		if (OwnerObject == nullptr ||
+			ChannelIndex >= EditorNumChannels ||
+			GetConstantOrRangeValue(ChannelIndex, ValueIndex) == InValue)
 		{
-			FText TransactionText;
-			int32 RequiredValueCount = 0;
-			if (FNiagaraDistributionEditorUtilities::IsConstant(GetDistributionMode()))
-			{
-				TransactionText = LOCTEXT("SetConstantValueTransaction", "Set constant value");
-				RequiredValueCount = EditorNumChannels;
-			}
-			else if (FNiagaraDistributionEditorUtilities::IsRange(GetDistributionMode()))
-			{
-				TransactionText = ValueIndex == 0
-					? LOCTEXT("SetRangeMinValueTransaction", "Set range min value")
-					: LOCTEXT("SetRangeMaxValueTransaction", "Set range max value");
-				RequiredValueCount = EditorNumChannels * 2;
-			}
+			return;
+		}
 
-			if (bContinuousTransactionPending)
-			{
-				bContinuousTransactionPending = false;
-				ContinuousTransactionIndex = GEditor->BeginTransaction(TransactionText);
-			}
+		FText TransactionText;
+		int32 RequiredValueCount = 0;
+		if (FNiagaraDistributionEditorUtilities::IsConstant(GetDistributionMode()))
+		{
+			TransactionText = LOCTEXT("SetConstantValueTransaction", "Set constant value");
+			RequiredValueCount = EditorNumChannels;
+		}
+		else if (FNiagaraDistributionEditorUtilities::IsRange(GetDistributionMode()))
+		{
+			TransactionText = ValueIndex == 0
+				? LOCTEXT("SetRangeMinValueTransaction", "Set range min value")
+				: LOCTEXT("SetRangeMaxValueTransaction", "Set range max value");
+			RequiredValueCount = EditorNumChannels * 2;
+		}
 
-			const FScopedTransaction Transaction(TransactionText, ContinuousTransactionIndex.IsSet() == false);
-			OwnerObject->Modify();
-			PropertyHandle->NotifyPreChange();
+		if (bContinuousTransactionPending)
+		{
+			bContinuousTransactionPending = false;
+			ContinuousTransactionIndex = GEditor->BeginTransaction(TransactionText);
+		}
 
-			if (SourceDistribution->ChannelConstantsAndRanges.Num() < RequiredValueCount)
-			{
-				SourceDistribution->ChannelConstantsAndRanges.AddZeroed(RequiredValueCount - SourceDistribution->ChannelConstantsAndRanges.Num());
-			}
-			else if (SourceDistribution->ChannelConstantsAndRanges.Num() > RequiredValueCount)
+		const FScopedTransaction Transaction(TransactionText, ContinuousTransactionIndex.IsSet() == false);
+		OwnerObject->Modify();
+		PropertyHandle->NotifyPreChange();
+
+		if (SourceDistribution->ChannelConstantsAndRanges.Num() < RequiredValueCount)
+		{
+			SourceDistribution->ChannelConstantsAndRanges.AddZeroed(RequiredValueCount - SourceDistribution->ChannelConstantsAndRanges.Num());
+		}
+		else if (SourceDistribution->ChannelConstantsAndRanges.Num() > RequiredValueCount)
+		{
+			SourceDistribution->ChannelConstantsAndRanges.SetNum(RequiredValueCount);
+		}
+
+		SourceDistribution->ChannelConstantsAndRanges[ValueIndex * EditorNumChannels + ChannelIndex] = InValue;
+		PropertyHandle->NotifyPostChange(bContinuousChangeActive ? EPropertyChangeType::Interactive : EPropertyChangeType::ValueSet);
+	}
+
+	virtual void SetConstantOrRangeValues(int32 ValueIndex, const TArray<float>& InValues) override
+	{
+		UObject* OwnerObject = OwnerObjectWeak.Get();
+		if (OwnerObject == nullptr || InValues.Num() != EditorNumChannels)
+		{
+			return;
+		}
+
+		bool bAllChannelsTheSame = true;
+		for (int32 ChannelIndex = 0; ChannelIndex < EditorNumChannels; ++ChannelIndex)
+		{
+			if (GetConstantOrRangeValue(ChannelIndex, ValueIndex) != InValues[ChannelIndex])
 			{
-				SourceDistribution->ChannelConstantsAndRanges.SetNum(RequiredValueCount);
+				bAllChannelsTheSame = false;
+				break;
 			}
+		}
 
-			SourceDistribution->ChannelConstantsAndRanges[ValueIndex * EditorNumChannels + ChannelIndex] = InValue;
-			PropertyHandle->NotifyPostChange(bContinuousChangeActive ? EPropertyChangeType::Interactive : EPropertyChangeType::ValueSet);
+		if (bAllChannelsTheSame)
+		{
+			return;
+		}
+
+		FText TransactionText;
+		int32 RequiredValueCount = 0;
+		if (FNiagaraDistributionEditorUtilities::IsConstant(GetDistributionMode()))
+		{
+			TransactionText = LOCTEXT("SetConstantValueTransaction", "Set constant value");
+			RequiredValueCount = EditorNumChannels;
+		}
+		else if (FNiagaraDistributionEditorUtilities::IsRange(GetDistributionMode()))
+		{
+			TransactionText = ValueIndex == 0
+				? LOCTEXT("SetRangeMinValueTransaction", "Set range min value")
+				: LOCTEXT("SetRangeMaxValueTransaction", "Set range max value");
+			RequiredValueCount = EditorNumChannels * 2;
 		}
+
+		if (bContinuousTransactionPending)
+		{
+			bContinuousTransactionPending = false;
+			ContinuousTransactionIndex = GEditor->BeginTransaction(TransactionText);
+		}
+
+		const FScopedTransaction Transaction(TransactionText, ContinuousTransactionIndex.IsSet() == false);
+		OwnerObject->Modify();
+		PropertyHandle->NotifyPreChange();
+
+		if (SourceDistribution->ChannelConstantsAndRanges.Num() < RequiredValueCount)
+		{
+			SourceDistribution->ChannelConstantsAndRanges.AddZeroed(RequiredValueCount - SourceDistribution->ChannelConstantsAndRanges.Num());
+		}
+		else if (SourceDistribution->ChannelConstantsAndRanges.Num() > RequiredValueCount)
+		{
+			SourceDistribution->ChannelConstantsAndRanges.SetNum(RequiredValueCount);
+		}
+
+		for (int32 ChannelIndex = 0; ChannelIndex < EditorNumChannels; ++ChannelIndex)
+		{
+			SourceDistribution->ChannelConstantsAndRanges[ValueIndex * EditorNumChannels + ChannelIndex] = InValues[ChannelIndex];
+		}
+		PropertyHandle->NotifyPostChange(bContinuousChangeActive ? EPropertyChangeType::Interactive : EPropertyChangeType::ValueSet);
 	}
 
-	virtual FRichCurve* GetCurveValue(int32 ChannelIndex) const override
+	virtual const FRichCurve* GetCurveValue(int32 ChannelIndex) const override
 	{
 		if (ChannelIndex < EditorNumChannels &&
 			ChannelIndex < SourceDistribution->ChannelCurves.Num())
@@ -290,7 +408,7 @@ public:
 	virtual void SetCurveValue(int32 ChannelIndex, const FRichCurve& InValue) override
 	{
 		UObject* OwnerObject = OwnerObjectWeak.Get();
-		FRichCurve* CurrentValue = GetCurveValue(ChannelIndex);
+		const FRichCurve* CurrentValue = GetCurveValue(ChannelIndex);
 		if (OwnerObject != nullptr &&
 			(CurrentValue == nullptr || *CurrentValue != InValue) &&
 			ChannelIndex < EditorNumChannels)
@@ -340,6 +458,27 @@ public:
 		PropertyHandle->NotifyFinishedChangingProperties();
 	}
 
+	virtual void CancelContinuousChange() override
+	{
+		if (ContinuousTransactionIndex.IsSet())
+		{
+			GEditor->CancelTransaction(ContinuousTransactionIndex.GetValue());
+		}
+		bContinuousTransactionPending = false;
+		ContinuousTransactionIndex.Reset();
+		bContinuousChangeActive = false;
+		PropertyHandle->NotifyFinishedChangingProperties();
+	}
+
+	virtual void ModifyOwners() override
+	{
+		UObject* OwnerObject = OwnerObjectWeak.Get();
+		if (OwnerObject != nullptr)
+		{
+			OwnerObject->Modify();
+		}
+	}
+
 private:
 	void MigrateDataFromModeChange(ENiagaraDistributionMode NewMode, int32 NewEditorNumChannels)
 	{
@@ -375,7 +514,7 @@ private:
 		case ENiagaraDistributionMode::NonUniformCurve:
 			for (int32 ChannelIndex = 0; ChannelIndex < EditorNumChannels; ChannelIndex++)
 			{
-				FRichCurve* Curve = GetCurveValue(ChannelIndex);
+				const FRichCurve* Curve = GetCurveValue(ChannelIndex);
 				float ConstantValue = Curve != nullptr && Curve->GetNumKeys() > 0 ? Curve->GetFirstKey().Value : 0;
 				ConstantValues.Add(ConstantValue);
 			}
@@ -420,7 +559,7 @@ private:
 		case ENiagaraDistributionMode::NonUniformCurve:
 			for (int32 ChannelIndex = 0; ChannelIndex < EditorNumChannels; ChannelIndex++)
 			{
-				FRichCurve* Curve = GetCurveValue(ChannelIndex);
+				const FRichCurve* Curve = GetCurveValue(ChannelIndex);
 				if (Curve != nullptr && Curve->GetNumKeys() > 0)
 				{
 					MinValues.Add(Curve->GetFirstKey().Value);
@@ -477,7 +616,7 @@ private:
 		case ENiagaraDistributionMode::NonUniformCurve:
 			for (int32 ChannelIndex = 0; ChannelIndex < EditorNumChannels; ChannelIndex++)
 			{
-				FRichCurve* Curve = GetCurveValue(ChannelIndex);
+				const FRichCurve* Curve = GetCurveValue(ChannelIndex);
 				if (Curve != nullptr)
 				{
 					CurveValues.Add(*Curve);
@@ -521,6 +660,7 @@ private:
 	bool bAllowNonUniform = true;
 	bool bAllowRange = true;
 	bool bAllowCurves = true;
+	bool bDisplayAsColor = false;
 };
 
 class SNiagaraDistributionPropertyWidget : public SCompoundWidget, public FEditorUndoClient
diff --git a/Engine/Plugins/FX/Niagara/Source/NiagaraEditor/Private/TypeEditorUtilities/NiagaraColorTypeEditorUtilities.cpp b/Engine/Plugins/FX/Niagara/Source/NiagaraEditor/Private/TypeEditorUtilities/NiagaraColorTypeEditorUtilities.cpp
--- a/Engine/Plugins/FX/Niagara/Source/NiagaraEditor/Private/TypeEditorUtilities/NiagaraColorTypeEditorUtilities.cpp
+++ b/Engine/Plugins/FX/Niagara/Source/NiagaraEditor/Private/TypeEditorUtilities/NiagaraColorTypeEditorUtilities.cpp
@@ -12,6 +12,7 @@
 #include "Widgets/Input/SNumericEntryBox.h"
 #include "Widgets/Colors/SColorBlock.h"
 #include "Widgets/Layout/SWrapBox.h"
+#include "Widgets/SNiagaraColorEditor.h"
 
 class SNiagaraColorParameterEditor : public SNiagaraParameterEditor
 {
@@ -21,43 +22,19 @@ public:
 
 	void Construct(const FArguments& InArgs)
 	{
-		SNiagaraParameterEditor::Construct(SNiagaraParameterEditor::FArguments()
-			.HAlign(HAlign_Fill));
+		SNiagaraParameterEditor::Construct(SNiagaraParameterEditor::FArguments());
 
 		ChildSlot
 		[
-			SNew(SWrapBox)
-			.UseAllottedSize(true)
-			+ SWrapBox::Slot()
-			.VAlign(VAlign_Center)
-			.Padding(0, 0, 6, 0)
-			[
-				SAssignNew(ColorBlock, SColorBlock)
-				.Color(this, &SNiagaraColorParameterEditor::GetColor)
-				.ShowBackgroundForAlpha(true)
-				.OnMouseButtonDown(this, &SNiagaraColorParameterEditor::OnMouseButtonDownColorBlock)
-				.Size(FVector2D(40.0f, 16.0f))
-				.CornerRadius(FVector4(3.0f, 3.0f, 3.0f, 3.0f))
-			]
-			+ SWrapBox::Slot()
-			.Padding(0, 0, 4, 0)
-			[
-				ConstructComponentWidget(0, NSLOCTEXT("ColorParameterEditor", "RLabel", "R"))
-			]
-			+ SWrapBox::Slot()
-			.Padding(0, 0, 4, 0)
-			[
-				ConstructComponentWidget(1, NSLOCTEXT("ColorParameterEditor", "GLabel", "G"))
-			]
-			+ SWrapBox::Slot()
-			.Padding(0, 0, 4, 0)
-			[
-				ConstructComponentWidget(2, NSLOCTEXT("ColorParameterEditor", "BLabel", "B"))
-			]
-			+ SWrapBox::Slot()
-			[
-				ConstructComponentWidget(3, NSLOCTEXT("ColorParameterEditor", "ALabel", "A"))
-			]
+			SNew(SNiagaraColorEditor)
+			.Color(this, &SNiagaraColorParameterEditor::GetColor)
+			.OnColorChanged(this, &SNiagaraColorParameterEditor::SetColor)
+			.OnBeginEditing(this, &SNiagaraColorParameterEditor::BeginEditing)
+			.OnEndEditing(this, &SNiagaraColorParameterEditor::EndEditing)
+			.OnCancelEditing(this, &SNiagaraColorParameterEditor::CancelEditing)
+			.OnColorPickerOpened(this, &SNiagaraColorParameterEditor::ColorPickerOpened)
+			.OnColorPickerClosed(this, &SNiagaraColorParameterEditor::ColorPickerClosed)
+			.MinDesiredColorBlockWidth(SNiagaraParameterEditor::DefaultInputSize)
 		];
 	}
 
@@ -76,108 +53,28 @@ public:
 	virtual bool CanChangeContinuously() const override { return true; }
 
 private:
-	TSharedRef<SWidget> ConstructComponentWidget(int32 Index, FText ComponentLabel)
-	{
-		return SNew(SNumericEntryBox<float>)
-		.Font(FAppStyle::Get().GetFontStyle("PropertyWindow.NormalFont"))
-		.OverrideTextMargin(2)
-		.MinValue(TOptional<float>())
-		.MaxValue(TOptional<float>())
-		.MaxSliderValue(TOptional<float>())
-		.MinSliderValue(TOptional<float>())
-		.Delta(0.0f)
-		.Value(this, &SNiagaraColorParameterEditor::GetComponentValue, Index)
-		.OnValueChanged(this, &SNiagaraColorParameterEditor::ComponentValueChanged, Index)
-		.OnValueCommitted(this, &SNiagaraColorParameterEditor::ComponentValueCommitted, Index)
-		.OnBeginSliderMovement(this, &SNiagaraColorParameterEditor::BeginSliderMovement)
-		.OnEndSliderMovement(this, &SNiagaraColorParameterEditor::EndSliderMovement)
-		.AllowSpin(true)
-		.BroadcastValueChangesPerKey(!GetDefault<UNiagaraEditorSettings>()->GetUpdateStackValuesOnCommitOnly())
-		.LabelVAlign(EVerticalAlignment::VAlign_Center)
-		.MinDesiredValueWidth(30)
-		.Label()
-		[
-			SNew(STextBlock)
-			.TextStyle(FNiagaraEditorStyle::Get(), "NiagaraEditor.ParameterText")
-			.Text(ComponentLabel)
-		];
-	}
-
-	void BeginSliderMovement()
+	void BeginEditing()
 	{
 		ExecuteOnBeginValueChange();
 	}
 
-	void EndSliderMovement(float Value)
+	void EndEditing()
 	{
 		ExecuteOnEndValueChange();
 	}
 
-	TOptional<float> GetComponentValue(int32 Index) const
+	void CancelEditing(FLinearColor OriginalColor)
 	{
-		return TOptional<float>(ColorValue.Component(Index));
-	}
-
-	void ComponentValueChanged(float ComponentValue, int32 Index)
-	{
-		ColorValue.Component(Index) = ComponentValue;
+		ColorValue = OriginalColor;
 		ExecuteOnValueChanged();
 	}
 
-	void ComponentValueCommitted(float ComponentValue, ETextCommit::Type CommitInfo, int32 Index)
-	{
-		if (CommitInfo == ETextCommit::OnEnter || CommitInfo == ETextCommit::OnUserMovedFocus)
-		{
-			ComponentValueChanged(ComponentValue, Index);
-		}
-	}
-
-	FReply OnMouseButtonDownColorBlock(const FGeometry& MyGeometry, const FPointerEvent& MouseEvent)
+	void ColorPickerOpened()
 	{
-		if (MouseEvent.GetEffectingButton() != EKeys::LeftMouseButton)
-		{
-			return FReply::Unhandled();
-		}
-
-		FColorPickerArgs PickerArgs;
-		{
-			PickerArgs.bUseAlpha = true;
-			PickerArgs.bOnlyRefreshOnMouseUp = false;
-			PickerArgs.bOnlyRefreshOnOk = false;
-			PickerArgs.DisplayGamma = TAttribute<float>::Create(TAttribute<float>::FGetter::CreateUObject(GEngine, &UEngine::GetDisplayGamma));
-			PickerArgs.OnColorCommitted = FOnLinearColorValueChanged::CreateSP(this, &SNiagaraColorParameterEditor::SetColor);
-			PickerArgs.OnColorPickerCancelled = FOnColorPickerCancelled::CreateSP(this, &SNiagaraColorParameterEditor::ColorPickerCancelled);
-			PickerArgs.OnInteractivePickBegin = FSimpleDelegate::CreateSP(this, &SNiagaraColorParameterEditor::InteractivePickBegin);
-			PickerArgs.OnInteractivePickEnd = FSimpleDelegate::CreateSP(this, &SNiagaraColorParameterEditor::InteractivePickEnd);
-			PickerArgs.OnColorPickerWindowClosed = FOnWindowClosed::CreateSP(this, &SNiagaraColorParameterEditor::ColorPickerClosed);
-			PickerArgs.InitialColor = ColorValue;
-			PickerArgs.ParentWidget = ColorBlock;
-		}
-
-		OpenColorPicker(PickerArgs);
-		// Mark this parameter editor as editing exclusively so that the corresponding structure details view doesn't get updated
-		// since it closes all color pickers when it gets updated!
 		SetIsEditingExclusively(true);
-		return FReply::Handled();
-	}
-
-	void InteractivePickBegin()
-	{
-		ExecuteOnBeginValueChange();
-	}
-
-	void InteractivePickEnd()
-	{
-		ExecuteOnEndValueChange();
-	}
-
-	void ColorPickerCancelled(FLinearColor OriginalColor)
-	{
-		ColorValue = OriginalColor;
-		ExecuteOnValueChanged();
 	}
 
-	void ColorPickerClosed(const TSharedRef<SWindow>& Window)
+	void ColorPickerClosed()
 	{
 		SetIsEditingExclusively(false);
 	}
diff --git a/Engine/Plugins/FX/Niagara/Source/NiagaraEditor/Private/Widgets/INiagaraDistributionAdapter.h b/Engine/Plugins/FX/Niagara/Source/NiagaraEditor/Private/Widgets/INiagaraDistributionAdapter.h
--- a/Engine/Plugins/FX/Niagara/Source/NiagaraEditor/Private/Widgets/INiagaraDistributionAdapter.h
+++ b/Engine/Plugins/FX/Niagara/Source/NiagaraEditor/Private/Widgets/INiagaraDistributionAdapter.h
@@ -9,41 +9,100 @@
 
 struct FRichCurve;
 
+/** Defines different edit modes for a niagara distribution */
 enum class ENiagaraDistributionEditorMode
 {
+	/** A single constant value */
 	Constant,
+	/** A multi-channel value edited with a single constant value. */
 	UniformConstant,
+	/** A multi-channel value edited with a constant value per channel. */
 	NonUniformConstant,
+	/** A multi-channel value edited with a single constant color. */
+	ColorConstant,
+	/** A single value which is distributed between a single range of values. */
 	Range,
+	/** A multi-channel value which is distributed between a single range of values. */
 	UniformRange,
+	/** A multi-channel value which is distributed between ranges of values, one range per channel. */
 	NonUniformRange,
+	/** A multi-channel value which is distributed between a single range of colors. */
+	ColorRange,
+	/** A single value which has been mapped to a range of values along a curve. */
 	Curve,
+	/** A multi-channel value which has been mapped to a single range of values along a curve. */
 	UniformCurve,
-	NonUniformCurve
+	/** A multi-channel value which has been mapped to a range of values along a curve, one curve per channel. */
+	NonUniformCurve,
+	/** A multi-channel value which has been mapped to a single range of colors along a gradient. */
+	ColorGradient
 };
 
+/** Defines the interface for editing a distribution of values with different numbers of channels and different representations. */
 class INiagaraDistributionAdapter
 {
 public:
 	virtual ~INiagaraDistributionAdapter() { }
 
+	/** Returns whether or not this adapter valid for use */
 	virtual bool IsValid() const = 0;
 
+	/** Gets the number of channels currently supported by this adapter. This can change with different distribution modes. */
 	virtual int32 GetNumChannels() const = 0;
+
+	/** Gets the display name for a specific channel by index. */
 	virtual FText GetChannelDisplayName(int32 ChannelIndex) const = 0;
+
+	/** Gets the color used to display a channel by index. */
 	virtual FSlateColor GetChannelColor(int32 ChannelIndex) const = 0;
 
+	/** Gets the editor modes supported by this distribution. */
 	virtual void GetSupportedDistributionModes(TArray<ENiagaraDistributionEditorMode>& OutSupportedModes) const = 0;
 
+	/** Gets the current mode for this distribution. */
 	virtual ENiagaraDistributionEditorMode GetDistributionMode() const = 0;
+
+	/** Sets the current mode for this distribution. */
 	virtual void SetDistributionMode(ENiagaraDistributionEditorMode InMode) = 0;
 
+	/** 
+		Gets a constant or range value by channel index and value index. 
+		@param ChannelIndex the index of the channel to get.
+		@param ValueIndex the index of the value to get, 0 for constants, 0 for minimum range values, and 1 for maximum range values.
+	*/
 	virtual float GetConstantOrRangeValue(int32 ChannelIndex, int32 ValueIndex) const = 0;
+
+	/** 
+		Sets a constant or range value by channel index and value index. 
+		@param ChannelIndex the index of the channel to set.
+		@param ValueIndex the index of the value to set, 0 for constants, 0 for minimum range values, and 1 for maximum range values.
+		@param InValue the value to set.
+	*/
 	virtual void SetConstantOrRangeValue(int32 ChannelIndex, int32 ValueIndex, float InValue) = 0;
 
-	virtual FRichCurve* GetCurveValue(int32 ChannelIndex) const = 0;
+	/** 
+		Sets all channel values of a constant or range by value index. 
+		@param ValueIndex the index of the values to set, 0 for constants, 0 for minimum range values, and 1 for maximum range values.
+		@param InValues the values to set, one per channel.
+	*/
+	virtual void SetConstantOrRangeValues(int32 ValueIndex, const TArray<float>& InValues) = 0;
+
+	/** Gets a curve value for the specified channel */
+	virtual const FRichCurve* GetCurveValue(int32 ChannelIndex) const = 0;
+
+	/** Set a curve value for the specified channel */
 	virtual void SetCurveValue(int32 ChannelIndex, const FRichCurve& InValue) = 0;
 
+	/** Called to notify the distribution that a continuous change has begun. */
 	virtual void BeginContinuousChange() = 0;
+
+	/** Called to notify the distribution that a continuous change has ended. */
 	virtual void EndContinuousChange() = 0;
+
+	/** Called to notify the distribution that a continuous change has been cancelled. */
+	virtual void CancelContinuousChange() = 0;
+
+	/** Called to modify the distribution that an external editor has started a 
+		transaction and that owning objects should be modified to take part in the transaction. */
+	virtual void ModifyOwners() = 0;
 };
\ No newline at end of file
diff --git a/Engine/Plugins/FX/Niagara/Source/NiagaraEditor/Private/Widgets/NiagaraDistributionEditorUtilities.cpp b/Engine/Plugins/FX/Niagara/Source/NiagaraEditor/Private/Widgets/NiagaraDistributionEditorUtilities.cpp
--- a/Engine/Plugins/FX/Niagara/Source/NiagaraEditor/Private/Widgets/NiagaraDistributionEditorUtilities.cpp
+++ b/Engine/Plugins/FX/Niagara/Source/NiagaraEditor/Private/Widgets/NiagaraDistributionEditorUtilities.cpp
@@ -18,18 +18,27 @@ FText FNiagaraDistributionEditorUtilities::DistributionModeToDisplayName(ENiagar
 		return LOCTEXT("UniformContstantDisplayName", "Uniform Constant");
 	case ENiagaraDistributionEditorMode::NonUniformConstant:
 		return LOCTEXT("NonUniformConstantDisplayName", "Non-uniform Constant");
+	case ENiagaraDistributionEditorMode::ColorConstant:
+		return LOCTEXT("ColorConstantDisplayName", "Color Constant");
+
 	case ENiagaraDistributionEditorMode::Range:
 		return LOCTEXT("RangeDisplayName", "Range");
 	case ENiagaraDistributionEditorMode::UniformRange:
 		return LOCTEXT("UniformRangeDisplayName", "Uniform Range");
 	case ENiagaraDistributionEditorMode::NonUniformRange:
 		return LOCTEXT("NonUniformRangeDisplayName", "Non-uniform Range");
+	case ENiagaraDistributionEditorMode::ColorRange:
+		return LOCTEXT("ColorRangeDisplayName", "Color Range");
+
 	case ENiagaraDistributionEditorMode::Curve:
 		return LOCTEXT("CurveDisplayName", "Curve");
 	case ENiagaraDistributionEditorMode::UniformCurve:
 		return LOCTEXT("UniformCurveDisplayName", "Uniform Curve");
 	case ENiagaraDistributionEditorMode::NonUniformCurve:
 		return LOCTEXT("NonUniformCurveDisplayName", "Non-uniform Curve");
+	case ENiagaraDistributionEditorMode::ColorGradient:
+		return LOCTEXT("ColorGradientDisplayName", "Color Gradient");
+
 	default:
 		return LOCTEXT("UnknownDisplayName", "Unknown");
 	}
@@ -45,18 +54,27 @@ FText FNiagaraDistributionEditorUtilities::DistributionModeToToolTipText(ENiagar
 		return LOCTEXT("UniformContstantToolTip", "A constant value applied to all value components.");
 	case ENiagaraDistributionEditorMode::NonUniformConstant:
 		return LOCTEXT("NonUniformConstantToolTip", "Constant values which can be different for each value component.");
+	case ENiagaraDistributionEditorMode::ColorConstant:
+		return LOCTEXT("ColorConstantToolTip", "A constant color value.");
+
 	case ENiagaraDistributionEditorMode::Range:
 		return LOCTEXT("RangeToolTip", "A single min/max range.");
 	case ENiagaraDistributionEditorMode::UniformRange:
 		return LOCTEXT("UniformRangeToolTip", "A min/max range applied to all value components.");
 	case ENiagaraDistributionEditorMode::NonUniformRange:
 		return LOCTEXT("NonUniformRangeToolTip", "Min/max ranges which can be different for each value component.");
+	case ENiagaraDistributionEditorMode::ColorRange:
+		return LOCTEXT("ColorRangeToolTip", "A color min/max range.");
+
 	case ENiagaraDistributionEditorMode::Curve:
 		return LOCTEXT("CurveToolTip", "This value is driven by a curve.");
 	case ENiagaraDistributionEditorMode::UniformCurve:
 		return LOCTEXT("UniformCurveToolTip", "All components of this value are driven by the same curve.");
 	case ENiagaraDistributionEditorMode::NonUniformCurve:
 		return LOCTEXT("NonUniformCurveToolTip", "Each component of this value is driven by its own curve.");
+	case ENiagaraDistributionEditorMode::ColorGradient:
+		return LOCTEXT("ColorGradientToolTip", "This color value is driven by a gradient.");
+
 	default:
 		return LOCTEXT("UnknownToolTip", "Unknown");
 	}
@@ -71,16 +89,25 @@ FName FNiagaraDistributionEditorUtilities::DistributionModeToIconBrushName(ENiag
 		return "NiagaraEditor.DistributionEditor.UniformConstant";
 	case ENiagaraDistributionEditorMode::NonUniformConstant:
 		return "NiagaraEditor.DistributionEditor.NonUniformConstant";
+	case ENiagaraDistributionEditorMode::ColorConstant:
+		return "NiagaraEditor.DistributionEditor.ColorConstant";
+
 	case ENiagaraDistributionEditorMode::Range:
 	case ENiagaraDistributionEditorMode::UniformRange:
 		return "NiagaraEditor.DistributionEditor.UniformRange";
 	case ENiagaraDistributionEditorMode::NonUniformRange:
 		return "NiagaraEditor.DistributionEditor.NonUniformRange";
+	case ENiagaraDistributionEditorMode::ColorRange:
+		return "NiagaraEditor.DistributionEditor.ColorRange";
+
 	case ENiagaraDistributionEditorMode::Curve:
 	case ENiagaraDistributionEditorMode::UniformCurve:
 		return "NiagaraEditor.DistributionEditor.UniformCurve";
 	case ENiagaraDistributionEditorMode::NonUniformCurve:
 		return "NiagaraEditor.DistributionEditor.NonUniformCurve";
+	case ENiagaraDistributionEditorMode::ColorGradient:
+		return "NiagaraEditor.DistributionEditor.ColorGradient";
+
 	default:
 		return NAME_None;
 	}
@@ -107,12 +134,21 @@ bool FNiagaraDistributionEditorUtilities::IsUniform(ENiagaraDistributionEditorMo
 		InMode == ENiagaraDistributionEditorMode::UniformCurve;
 }
 
+bool FNiagaraDistributionEditorUtilities::IsColor(ENiagaraDistributionEditorMode InMode)
+{
+	return
+		InMode == ENiagaraDistributionEditorMode::ColorConstant ||
+		InMode == ENiagaraDistributionEditorMode::ColorRange ||
+		InMode == ENiagaraDistributionEditorMode::ColorGradient;
+}
+
 bool FNiagaraDistributionEditorUtilities::IsConstant(ENiagaraDistributionEditorMode InMode)
 {
 	return
 		InMode == ENiagaraDistributionEditorMode::Constant ||
 		InMode == ENiagaraDistributionEditorMode::UniformConstant ||
-		InMode == ENiagaraDistributionEditorMode::NonUniformConstant;
+		InMode == ENiagaraDistributionEditorMode::NonUniformConstant ||
+		InMode == ENiagaraDistributionEditorMode::ColorConstant;
 }
 
 bool FNiagaraDistributionEditorUtilities::IsRange(ENiagaraDistributionEditorMode InMode)
@@ -120,7 +156,8 @@ bool FNiagaraDistributionEditorUtilities::IsRange(ENiagaraDistributionEditorMode
 	return
 		InMode == ENiagaraDistributionEditorMode::Range ||
 		InMode == ENiagaraDistributionEditorMode::UniformRange ||
-		InMode == ENiagaraDistributionEditorMode::NonUniformRange;
+		InMode == ENiagaraDistributionEditorMode::NonUniformRange ||
+		InMode == ENiagaraDistributionEditorMode::ColorRange;
 }
 
 bool FNiagaraDistributionEditorUtilities::IsCurve(ENiagaraDistributionEditorMode InMode)
@@ -131,4 +168,9 @@ bool FNiagaraDistributionEditorUtilities::IsCurve(ENiagaraDistributionEditorMode
 		InMode == ENiagaraDistributionEditorMode::NonUniformCurve;
 }
 
+bool FNiagaraDistributionEditorUtilities::IsGradient(ENiagaraDistributionEditorMode InMode)
+{
+	return InMode == ENiagaraDistributionEditorMode::ColorGradient;
+}
+
 #undef LOCTEXT_NAMESPACE
\ No newline at end of file
diff --git a/Engine/Plugins/FX/Niagara/Source/NiagaraEditor/Private/Widgets/NiagaraDistributionEditorUtilities.h b/Engine/Plugins/FX/Niagara/Source/NiagaraEditor/Private/Widgets/NiagaraDistributionEditorUtilities.h
--- a/Engine/Plugins/FX/Niagara/Source/NiagaraEditor/Private/Widgets/NiagaraDistributionEditorUtilities.h
+++ b/Engine/Plugins/FX/Niagara/Source/NiagaraEditor/Private/Widgets/NiagaraDistributionEditorUtilities.h
@@ -20,7 +20,10 @@ namespace FNiagaraDistributionEditorUtilities
 	FSlateIcon DistributionModeToIcon(ENiagaraDistributionEditorMode InMode);
 
 	bool IsUniform(ENiagaraDistributionEditorMode InMode);
+	bool IsColor(ENiagaraDistributionEditorMode InMode);
+
 	bool IsConstant(ENiagaraDistributionEditorMode InMode);
 	bool IsRange(ENiagaraDistributionEditorMode InMode);
 	bool IsCurve(ENiagaraDistributionEditorMode InMode);
+	bool IsGradient(ENiagaraDistributionEditorMode InMode);
 }
\ No newline at end of file
diff --git a/Engine/Plugins/FX/Niagara/Source/NiagaraEditor/Private/Widgets/SNiagaraColorEditor.cpp b/Engine/Plugins/FX/Niagara/Source/NiagaraEditor/Private/Widgets/SNiagaraColorEditor.cpp
new file mode 100644
--- /dev/null
+++ b/Engine/Plugins/FX/Niagara/Source/NiagaraEditor/Private/Widgets/SNiagaraColorEditor.cpp
@@ -0,0 +1,282 @@
+// Copyright Epic Games, Inc. All Rights Reserved.
+
+#include "Widgets/SNiagaraColorEditor.h"
+
+#include "Engine/Engine.h"
+#include "Widgets/Colors/SColorBlock.h"
+#include "Widgets/Colors/SColorPicker.h"
+#include "Widgets/Input/SNumericEntryBox.h"
+#include "Widgets/Layout/SBox.h"
+#include "Widgets/Layout/SGridPanel.h"
+#include "Widgets/SNiagaraExpandedToggle.h"
+
+#define LOCTEXT_NAMESPACE "NiagaraColorEditor"
+
+void SNiagaraColorEditor::Construct(const FArguments& InArgs)
+{
+	bShowAlpha = InArgs._ShowAlpha;
+	bShowExpander = InArgs._ShowExpander;
+	Color = InArgs._Color;
+	ExpandComponents = InArgs._ExpandComponents;
+	LabelText = InArgs._LabelText;
+	MinDesiredColorBlockWidth = InArgs._MinDesiredColorBlockWidth;
+	OnColorChangedDelegate = InArgs._OnColorChanged;
+	OnBeginEditingDelegate = InArgs._OnBeginEditing;
+	OnEndEditingDelegate = InArgs._OnEndEditing;
+	OnCancelEditingDelegate = InArgs._OnCancelEditing;
+	OnExpandComponentsChangedDelegate = InArgs._OnExpandComponentsChanged;
+	OnColorPickerOpenedDelegate = InArgs._OnColorPickerOpened;
+	OnColorPickerClosedDelegate = InArgs._OnColorPickerClosed;
+
+	ColorBlock = SNew(SColorBlock)
+		.AlphaBackgroundBrush(FAppStyle::Get().GetBrush("ColorPicker.RoundedAlphaBackground"))
+		.Color(Color)
+		.ShowBackgroundForAlpha(true)
+		.AlphaDisplayMode(bShowAlpha ? EColorBlockAlphaDisplayMode::Separate : EColorBlockAlphaDisplayMode::Ignore)
+		.OnMouseButtonDown(this, &SNiagaraColorEditor::OnColorBlockMouseButtonDown)
+		.Size(FVector2D(75.0f, 18.0f))
+		.CornerRadius(FVector4(4.0f, 4.0f, 4.0f, 4.0f));
+
+	RootBox = SNew(SGridPanel)
+		.FillColumn(2, 1.0f);
+
+	if (bShowExpander)
+	{
+		RootBox->AddSlot(0, 0)
+		.VAlign(VAlign_Center)
+		.Padding(0, 0, 3, 0)
+		[
+			SNew(SNiagaraExpandedToggle)
+			.Expanded(ExpandComponents)
+			.OnExpandedChanged(this, &SNiagaraColorEditor::ExpanderExpandedChanged)
+		];
+	}
+
+	if (LabelText.IsBound() || LabelText.Get().IsEmptyOrWhitespace() == false)
+	{
+		RootBox->AddSlot(1, 0)
+		.VAlign(VAlign_Center)
+		.Padding(0, 0, 5, 0)
+		[
+			SNew(STextBlock)
+			.Font(FAppStyle::Get().GetFontStyle("PropertyWindow.NormalFont"))
+			.Text(LabelText)
+		];
+	}
+
+	TSharedPtr<SWidget> ColorBlockWidgetForSlot;
+	if (MinDesiredColorBlockWidth.IsSet())
+	{
+		ColorBlockWidgetForSlot = SNew(SBox)
+			.MinDesiredWidth(MinDesiredColorBlockWidth)
+			[
+				ColorBlock.ToSharedRef()
+			];
+	}
+	else
+	{
+		ColorBlockWidgetForSlot = ColorBlock.ToSharedRef();
+	}
+
+	RootBox->AddSlot(2, 0)
+	.Padding(0, 1, 1, 1)
+	[
+		ColorBlockWidgetForSlot.ToSharedRef()
+	];
+
+	ChildSlot
+	[
+		RootBox.ToSharedRef()
+	];
+
+	UpdateComponentWidgets();
+}
+
+void SNiagaraColorEditor::Tick(const FGeometry& AllottedGeometry, const double InCurrentTime, const float InDeltaTime)
+{
+	if (bComponentsConstructed == false && ExpandComponents.Get() == true)
+	{
+		UpdateComponentWidgets();
+	}
+	SCompoundWidget::Tick(AllottedGeometry, InCurrentTime, InDeltaTime);
+}
+
+void SNiagaraColorEditor::UpdateComponentWidgets()
+{
+	if (bComponentsConstructed == false && ExpandComponents.Get())
+	{
+		TSharedPtr<SGridPanel> ComponentGrid;
+		RootBox->AddSlot(2, 1)
+			.HAlign(HAlign_Left)
+			.Padding(0, 3, 0, 0)
+			[
+				SAssignNew(ComponentGrid, SGridPanel)
+				.Visibility(this, &SNiagaraColorEditor::GetComponentsVisibility)
+			];
+		ConstructComponentWidgets(ComponentGrid.ToSharedRef(), LOCTEXT("RedLabel", "Red"), 0);
+		ConstructComponentWidgets(ComponentGrid.ToSharedRef(), LOCTEXT("GreenLabel", "Green"), 1);
+		ConstructComponentWidgets(ComponentGrid.ToSharedRef(), LOCTEXT("BlueLabel", "Blue"), 2);
+		if (bShowAlpha)
+		{
+			ConstructComponentWidgets(ComponentGrid.ToSharedRef(), LOCTEXT("AlphaLabel", "Alpha"), 3);
+		}
+		bComponentsConstructed = true;
+	}
+}
+
+FReply SNiagaraColorEditor::OnColorBlockMouseButtonDown(const FGeometry& MyGeometry, const FPointerEvent& MouseEvent)
+{
+	if (MouseEvent.GetEffectingButton() != EKeys::LeftMouseButton)
+	{
+		return FReply::Unhandled();
+	}
+
+	OnBeginEditingDelegate.ExecuteIfBound();
+
+	FColorPickerArgs PickerArgs;
+	{
+		PickerArgs.bUseAlpha = bShowAlpha;
+		PickerArgs.bOnlyRefreshOnMouseUp = false;
+		PickerArgs.bOnlyRefreshOnOk = false;
+		PickerArgs.DisplayGamma = TAttribute<float>::Create(TAttribute<float>::FGetter::CreateUObject(GEngine, &UEngine::GetDisplayGamma));
+		PickerArgs.OnColorCommitted = FOnLinearColorValueChanged::CreateSP(this, &SNiagaraColorEditor::ColorPickerColorCommitted);
+		PickerArgs.OnColorPickerCancelled = FOnColorPickerCancelled::CreateSP(this, &SNiagaraColorEditor::OnColorPickerCancelled);
+		PickerArgs.OnColorPickerWindowClosed = FOnWindowClosed::CreateSP(this, &SNiagaraColorEditor::OnColorPickerClosed);
+		PickerArgs.InitialColor = Color.Get(FLinearColor::White);
+		PickerArgs.ParentWidget = ColorBlock;
+		FWidgetPath ParentWidgetPath;
+		if (FSlateApplication::Get().FindPathToWidget(ColorBlock.ToSharedRef(), ParentWidgetPath))
+		{
+			PickerArgs.bOpenAsMenu = FSlateApplication::Get().FindMenuInWidgetPath(ParentWidgetPath).IsValid();
+		}
+	}
+
+	OpenColorPicker(PickerArgs);
+	OnColorPickerOpenedDelegate.ExecuteIfBound();
+	return FReply::Handled();
+}
+
+void SNiagaraColorEditor::ExpanderExpandedChanged(bool bIsExpanded)
+{
+	if (ExpandComponents.Get() != bIsExpanded)
+	{
+		if (ExpandComponents.IsBound() == false)
+		{
+			ExpandComponents = bIsExpanded;
+		}
+		OnExpandComponentsChangedDelegate.ExecuteIfBound(bIsExpanded);
+	}
+}
+
+void SNiagaraColorEditor::ConstructComponentWidgets(TSharedRef<SGridPanel> ComponentGrid, FText ComponentLabelText, int32 ComponentIndex)
+{
+	ComponentGrid->AddSlot(0, ComponentIndex)
+	.VAlign(VAlign_Center)
+	.HAlign(HAlign_Left)
+	.Padding(0, 0, 5, 0)
+	[
+		SNew(STextBlock)
+		.Font(FAppStyle::Get().GetFontStyle("PropertyWindow.NormalFont"))
+		.Text(ComponentLabelText)
+	];
+	ComponentGrid->AddSlot(1, ComponentIndex)
+	.Padding(0, 2, 0, 2)
+	[
+		SNew(SNumericEntryBox<float>)
+		.Font(FAppStyle::Get().GetFontStyle("PropertyWindow.NormalFont"))
+		.Value(this, &SNiagaraColorEditor::GetComponentValue, ComponentIndex)
+		.OnValueChanged(this, &SNiagaraColorEditor::ComponentValueChanged, ComponentIndex)
+		.OnValueCommitted(this, &SNiagaraColorEditor::ComponentValueCommitted, ComponentIndex)
+		.OnBeginSliderMovement(this, &SNiagaraColorEditor::BeginComponentValueChange)
+		.OnEndSliderMovement(this, &SNiagaraColorEditor::EndComponentValueChange)
+		.AllowSpin(true)
+		.MinValue(TOptional<float>())
+		.MaxValue(TOptional<float>())
+		.MinSliderValue(TOptional<float>())
+		.MaxSliderValue(TOptional<float>())
+		.BroadcastValueChangesPerKey(false)
+		.MinDesiredValueWidth(50)
+	];
+}
+
+TOptional<float> SNiagaraColorEditor::GetComponentValue(int32 ComponentIndex) const
+{
+	FLinearColor CurrentColor = Color.Get(FLinearColor::White);
+	switch(ComponentIndex)
+	{
+	case 0:
+		return CurrentColor.R;
+	case 1:
+		return CurrentColor.G;
+	case 2:
+		return CurrentColor.B;
+	case 3:
+		return CurrentColor.A;
+	default:
+		return 0.0f;
+	}
+}
+
+void SNiagaraColorEditor::ComponentValueChanged(float Value, int32 ComponentIndex)
+{
+	FLinearColor CurrentColor = Color.Get(FLinearColor::White);
+	switch (ComponentIndex)
+	{
+	case 0:
+		CurrentColor.R = Value;
+		break;
+	case 1:
+		CurrentColor.G = Value;
+		break;
+	case 2:
+		CurrentColor.B = Value;
+		break;
+	case 3:
+		CurrentColor.A = Value;
+		break;
+	}
+	OnColorChangedDelegate.ExecuteIfBound(CurrentColor);
+}
+
+void SNiagaraColorEditor::ComponentValueCommitted(float Value, ETextCommit::Type CommitInfo, int32 ComponentIndex)
+{
+	if (CommitInfo == ETextCommit::OnEnter || CommitInfo == ETextCommit::OnUserMovedFocus)
+	{
+		ComponentValueChanged(Value, ComponentIndex);
+	}
+}
+
+void SNiagaraColorEditor::BeginComponentValueChange()
+{
+	OnBeginEditingDelegate.ExecuteIfBound();
+}
+
+void SNiagaraColorEditor::EndComponentValueChange(float Value)
+{
+	OnEndEditingDelegate.ExecuteIfBound();
+}
+
+EVisibility SNiagaraColorEditor::GetComponentsVisibility() const
+{
+	return ExpandComponents.Get() 
+		? EVisibility::Visible 
+		: EVisibility::Collapsed;
+}
+
+void SNiagaraColorEditor::ColorPickerColorCommitted(FLinearColor InColor)
+{
+	OnColorChangedDelegate.ExecuteIfBound(InColor);
+}
+
+void SNiagaraColorEditor::OnColorPickerClosed(const TSharedRef<SWindow>& InWindow)
+{
+	OnEndEditingDelegate.ExecuteIfBound();
+	OnColorPickerClosedDelegate.ExecuteIfBound();
+}
+
+void SNiagaraColorEditor::OnColorPickerCancelled(FLinearColor InColor)
+{
+	OnCancelEditingDelegate.ExecuteIfBound(InColor);
+}
+
+#undef LOCTEXT_NAMESPACE
diff --git a/Engine/Plugins/FX/Niagara/Source/NiagaraEditor/Private/Widgets/SNiagaraColorEditor.h b/Engine/Plugins/FX/Niagara/Source/NiagaraEditor/Private/Widgets/SNiagaraColorEditor.h
new file mode 100644
--- /dev/null
+++ b/Engine/Plugins/FX/Niagara/Source/NiagaraEditor/Private/Widgets/SNiagaraColorEditor.h
@@ -0,0 +1,84 @@
+// Copyright Epic Games, Inc. All Rights Reserved.
+
+#pragma once
+
+#include "Types/SlateStructs.h"
+#include "Widgets/SCompoundWidget.h"
+
+class SColorBlock;
+class SGridPanel;
+
+class SNiagaraColorEditor : public SCompoundWidget
+{
+public:
+	DECLARE_DELEGATE_OneParam(FOnValueChanged, FLinearColor /* Value */);
+	DECLARE_DELEGATE_OneParam(FOnExpandComponentsChanged, bool /* bIsExpanded */);
+	DECLARE_DELEGATE_OneParam(FOnCancelEditing, FLinearColor /* OriginalValue */);
+
+public:
+	SLATE_BEGIN_ARGS(SNiagaraColorEditor)
+		: _ShowAlpha(true)
+		, _ShowExpander(true)
+		, _ExpandComponents(false)
+		{ }
+		SLATE_ARGUMENT(bool, ShowAlpha)
+		SLATE_ARGUMENT(bool, ShowExpander)
+		SLATE_ATTRIBUTE(FLinearColor, Color)
+		SLATE_ATTRIBUTE(bool, ExpandComponents)
+		SLATE_ATTRIBUTE(FText, LabelText)
+		SLATE_ATTRIBUTE(FOptionalSize, MinDesiredColorBlockWidth)
+		SLATE_EVENT(FOnValueChanged, OnColorChanged)
+		SLATE_EVENT(FSimpleDelegate, OnBeginEditing)
+		SLATE_EVENT(FSimpleDelegate, OnEndEditing)
+		SLATE_EVENT(FOnCancelEditing, OnCancelEditing)
+		SLATE_EVENT(FOnExpandComponentsChanged, OnExpandComponentsChanged)
+		SLATE_EVENT(FSimpleDelegate, OnColorPickerOpened)
+		SLATE_EVENT(FSimpleDelegate, OnColorPickerClosed)
+	SLATE_END_ARGS();
+
+	void Construct(const FArguments& InArgs);
+	virtual void Tick(const FGeometry& AllottedGeometry, const double InCurrentTime, const float InDeltaTime) override;
+
+private:
+	void UpdateComponentWidgets();
+
+	FReply OnColorBlockMouseButtonDown(const FGeometry& MyGeometry, const FPointerEvent& MouseEvent);
+
+	void ExpanderExpandedChanged(bool bIsExpanded);
+
+	void ConstructComponentWidgets(TSharedRef<SGridPanel> ComponentGrid, FText ComponentLabelText, int32 ComponentIndex);
+
+	TOptional<float> GetComponentValue(int32 ComponentIndex) const;
+	void ComponentValueChanged(float Value, int32 ComponentIndex);
+	void ComponentValueCommitted(float Value, ETextCommit::Type CommitInfo, int32 ComponentIndex);
+
+	void BeginComponentValueChange();
+	void EndComponentValueChange(float Value);
+
+	EVisibility GetComponentsVisibility() const;
+
+	void ColorPickerColorCommitted(FLinearColor InColor);
+	void OnColorPickerClosed(const TSharedRef<SWindow>& InWindow);
+	void OnColorPickerCancelled(FLinearColor InColor);
+
+private:
+	bool bShowAlpha;
+	bool bShowExpander;
+
+	TAttribute<bool> ExpandComponents;
+	TAttribute<FLinearColor> Color;
+	TAttribute<FText> LabelText;
+	TAttribute<FOptionalSize> MinDesiredColorBlockWidth;
+
+	FOnValueChanged OnColorChangedDelegate;
+	FSimpleDelegate OnBeginEditingDelegate;
+	FSimpleDelegate OnEndEditingDelegate;
+	FOnCancelEditing OnCancelEditingDelegate;
+	FOnExpandComponentsChanged OnExpandComponentsChangedDelegate;
+	FSimpleDelegate OnColorPickerOpenedDelegate;
+	FSimpleDelegate OnColorPickerClosedDelegate;
+
+	bool bComponentsConstructed = false;
+	TSharedPtr<SGridPanel> RootBox;
+	TSharedPtr<SColorBlock> ColorBlock;
+};
\ No newline at end of file
diff --git a/Engine/Plugins/FX/Niagara/Source/NiagaraEditor/Private/Widgets/SNiagaraDistributionCurveEditor.cpp b/Engine/Plugins/FX/Niagara/Source/NiagaraEditor/Private/Widgets/SNiagaraDistributionCurveEditor.cpp
--- a/Engine/Plugins/FX/Niagara/Source/NiagaraEditor/Private/Widgets/SNiagaraDistributionCurveEditor.cpp
+++ b/Engine/Plugins/FX/Niagara/Source/NiagaraEditor/Private/Widgets/SNiagaraDistributionCurveEditor.cpp
@@ -500,6 +500,7 @@ void SNiagaraDistributionCurveEditor::Construct(const FArguments& InArgs, TShare
 				.Padding(FMargin(KeyHandleWidth / 2, 3, KeyHandleWidth / 2, 3))
 				[
 					SNew(SNiagaraCompactCurveView, &EditCurve)
+					.Width(1000)
 					.Height(50)
 					.CurveColor(InArgs._CurveColor)
 					.SelectedKeyHandle(this, &SNiagaraDistributionCurveEditor::GetSelectedKeyHandle)
diff --git a/Engine/Plugins/FX/Niagara/Source/NiagaraEditor/Private/Widgets/SNiagaraDistributionEditor.cpp b/Engine/Plugins/FX/Niagara/Source/NiagaraEditor/Private/Widgets/SNiagaraDistributionEditor.cpp
--- a/Engine/Plugins/FX/Niagara/Source/NiagaraEditor/Private/Widgets/SNiagaraDistributionEditor.cpp
+++ b/Engine/Plugins/FX/Niagara/Source/NiagaraEditor/Private/Widgets/SNiagaraDistributionEditor.cpp
@@ -2,17 +2,19 @@
 
 #include "Widgets/SNiagaraDistributionEditor.h"
 
-#include "Framework/Commands/UIAction.h"
+#include "Curves/CurveOwnerInterface.h"
 #include "Framework/MultiBox/MultiBoxBuilder.h"
 #include "NiagaraEditorStyle.h"
-#include "Textures/SlateIcon.h"
+#include "SColorGradientEditor.h"
 #include "Widgets/Input/SComboButton.h"
 #include "Widgets/Input/SNumericEntryBox.h"
 #include "Widgets/Input/SVectorInputBox.h"
 #include "Widgets/Layout/SBox.h"
 #include "Widgets/Layout/SGridPanel.h"
 #include "Widgets/NiagaraDistributionEditorUtilities.h"
+#include "Widgets/SNiagaraColorEditor.h"
 #include "Widgets/SNiagaraDistributionCurveEditor.h"
+#include "Widgets/SNiagaraExpandedToggle.h"
 
 #define LOCTEXT_NAMESPACE "NiagaraDistributionEditor"
 
@@ -113,7 +115,6 @@ private:
 	mutable FText ModeToolTipCache;
 };
 
-
 class SNiagaraDistributionValueEditor : public SCompoundWidget
 {
 public:
@@ -124,22 +125,42 @@ public:
 	{
 		DistributionAdapter = InDistributionAdapter;
 		bool bIsUniform = FNiagaraDistributionEditorUtilities::IsUniform(DistributionAdapter->GetDistributionMode());
+		bool bIsColor = FNiagaraDistributionEditorUtilities::IsColor(DistributionAdapter->GetDistributionMode());
 		bool bIsConstant = FNiagaraDistributionEditorUtilities::IsConstant(DistributionAdapter->GetDistributionMode());
+		bool bLayoutHorizontal = bIsUniform || bIsColor;
 
-		FText MinLabelText = LOCTEXT("MinLabel", "Min");
-		FText MaxLabelText = LOCTEXT("MaxLabel", "Max");
+		static const FText MinLabelText = LOCTEXT("MinLabel", "Min");
+		static const FText MaxLabelText = LOCTEXT("MaxLabel", "Max");
 
 		TSharedPtr<SWidget> ChildWidget;
-		if (bIsUniform)
+		if (bIsConstant)
 		{
-			if (bIsConstant)
+			TSharedPtr<SWidget> ValueWidget;
+			if (bIsUniform)
+			{
+				ValueWidget = ConstructFloatWidget(0, 0, FText());
+			}
+			else if (bIsColor)
 			{
-				ChildWidget = ConstructFloatWidget(0, 0, FText());
+				ValueWidget = ConstructColorWidget(DistributionAdapter->GetNumChannels(), 0, true, FText());
 			}
 			else
+			{
+				ValueWidget = ConstructVectorWidget(DistributionAdapter->GetNumChannels(), 0);
+			}
+			ChildWidget = SNew(SBox)
+				.HAlign(HAlign_Left)
+				[
+					ValueWidget.ToSharedRef()
+				];
+		}
+		else
+		{
+			if (bIsUniform)
 			{
 				ChildWidget = SNew(SHorizontalBox)
 					+ SHorizontalBox::Slot()
+					.Padding(0, 0, 10, 0)
 					[
 						ConstructFloatWidget(0, 0, MinLabelText)
 					]
@@ -148,41 +169,58 @@ public:
 						ConstructFloatWidget(0, 1, MaxLabelText)
 					];
 			}
-		}
-		else
-		{
-			if (bIsConstant)
+			else if (bIsColor)
 			{
-				ChildWidget = ConstructVectorWidget(DistributionAdapter->GetNumChannels(), 0);
-			}
-			else
-			{
-				ChildWidget = SNew(SGridPanel)
-					.FillColumn(1, 1)
-					+ SGridPanel::Slot(0, 0)
-					.VAlign(VAlign_Center)
-					.Padding(0, 0, 5, 3)
-					[
-						SNew(STextBlock)
-						.Text(MinLabelText)
-					]
-					+ SGridPanel::Slot(1, 0)
-					.Padding(0, 0, 0, 3)
-					[
-						ConstructVectorWidget(DistributionAdapter->GetNumChannels(), 0)
-					]
-					+ SGridPanel::Slot(0, 1)
-					.VAlign(VAlign_Center)
-					.Padding(0, 0, 5, 0)
+				ChildWidget = SNew(SHorizontalBox)
+					+ SHorizontalBox::Slot()
+					.AutoWidth()
+					.Padding(0, 0, 10, 0)
 					[
-						SNew(STextBlock)
-						.Text(MaxLabelText)
+						ConstructColorWidget(DistributionAdapter->GetNumChannels(), 0, true, MinLabelText)
 					]
-					+ SGridPanel::Slot(1, 1)
+					+ SHorizontalBox::Slot()
+					.AutoWidth()
 					[
-						ConstructVectorWidget(DistributionAdapter->GetNumChannels(), 1)
+						ConstructColorWidget(DistributionAdapter->GetNumChannels(), 1, false, MaxLabelText)
 					];
 			}
+			else
+			{
+				if (bIsConstant)
+				{
+					ChildWidget = ConstructVectorWidget(DistributionAdapter->GetNumChannels(), 0);
+				}
+				else
+				{
+					ChildWidget = SNew(SGridPanel)
+						.FillColumn(1, 1)
+						+ SGridPanel::Slot(0, 0)
+						.VAlign(VAlign_Center)
+						.Padding(0, 0, 5, 3)
+						[
+							SNew(STextBlock)
+							.Font(FAppStyle::Get().GetFontStyle("PropertyWindow.NormalFont"))
+							.Text(MinLabelText)
+						]
+						+ SGridPanel::Slot(1, 0)
+						.Padding(0, 0, 0, 3)
+						[
+							ConstructVectorWidget(DistributionAdapter->GetNumChannels(), 0)
+						]
+						+ SGridPanel::Slot(0, 1)
+						.VAlign(VAlign_Center)
+						.Padding(0, 0, 5, 0)
+						[
+							SNew(STextBlock)
+							.Font(FAppStyle::Get().GetFontStyle("PropertyWindow.NormalFont"))
+							.Text(MaxLabelText)
+						]
+						+ SGridPanel::Slot(1, 1)
+						[
+							ConstructVectorWidget(DistributionAdapter->GetNumChannels(), 1)
+						];
+				}
+			}
 		}
 
 		ChildSlot
@@ -192,6 +230,16 @@ public:
 	}
 
 private:
+	bool GetColorComponentsExpanded() const
+	{
+		return bColorComponentsExpanded;
+	}
+
+	void ColorComponentsExandedChanged(bool bExpanded)
+	{
+		bColorComponentsExpanded = bExpanded;
+	}
+
 	TSharedRef<SWidget> ConstructFloatWidget(int32 ChannelIndex, int32 ValueIndex, FText LabelText)
 	{
 		TSharedRef<SWidget> LabelWidget = LabelText.IsEmpty()
@@ -200,21 +248,22 @@ private:
 				.TextStyle(FNiagaraEditorStyle::Get(), "NiagaraEditor.ParameterText")
 				.Text(LabelText);
 
-		return SNew(SNumericEntryBox<float>)
+		return
+			SNew(SNumericEntryBox<float>)
 			.Font(FAppStyle::Get().GetFontStyle("PropertyWindow.NormalFont"))
 			.Value(this, &SNiagaraDistributionValueEditor::GetValue, ChannelIndex, ValueIndex)
 			.OnValueChanged(this, &SNiagaraDistributionValueEditor::ValueChanged, ChannelIndex, ValueIndex)
 			.OnValueCommitted(this, &SNiagaraDistributionValueEditor::ValueCommitted, ChannelIndex, ValueIndex)
-			.OnBeginSliderMovement(this, &SNiagaraDistributionValueEditor::BeginValueSliderMovement)
-			.OnEndSliderMovement(this, &SNiagaraDistributionValueEditor::EndValueSliderMovement)
+			.OnBeginSliderMovement(this, &SNiagaraDistributionValueEditor::BeginValueChange)
+			.OnEndSliderMovement(this, &SNiagaraDistributionValueEditor::EndValueChange)
 			.AllowSpin(true)
 			.MinValue(TOptional<float>())
 			.MaxValue(TOptional<float>())
 			.MinSliderValue(TOptional<float>())
 			.MaxSliderValue(TOptional<float>())
 			.BroadcastValueChangesPerKey(false)
+			.MinDesiredValueWidth(SNiagaraDistributionEditor::DefaultInputSize - 18)
 			.LabelVAlign(EVerticalAlignment::VAlign_Center)
-			.MinDesiredValueWidth(30)
 			.Label()
 			[
 				LabelWidget
@@ -229,61 +278,89 @@ private:
 	{
 		if (ChannelCount == 2)
 		{
-			return SNew(SNumericVectorInputBox2)
-				.Font(FAppStyle::Get().GetFontStyle("PropertyWindow.NormalFont"))
-				.AllowSpin(true)
-				.bColorAxisLabels(true)
-				.X(this, &SNiagaraDistributionValueEditor::GetValue, 0, ValueIndex)
-				.OnXChanged(this, &SNiagaraDistributionValueEditor::ValueChanged, 0, ValueIndex)
-				.OnXCommitted(this, &SNiagaraDistributionValueEditor::ValueCommitted, 0, ValueIndex)
-				.Y(this, &SNiagaraDistributionValueEditor::GetValue, 1, ValueIndex)
-				.OnYChanged(this, &SNiagaraDistributionValueEditor::ValueChanged, 1, ValueIndex)
-				.OnYCommitted(this, &SNiagaraDistributionValueEditor::ValueCommitted, 1, ValueIndex)
-				.OnBeginSliderMovement(this, &SNiagaraDistributionValueEditor::BeginValueSliderMovement)
-				.OnEndSliderMovement(this, &SNiagaraDistributionValueEditor::EndValueSliderMovement);
+			return SNew(SBox)
+				.MinDesiredWidth(2 * SNiagaraDistributionEditor::DefaultInputSize)
+				[
+					SNew(SNumericVectorInputBox2)
+					.Font(FAppStyle::Get().GetFontStyle("PropertyWindow.NormalFont"))
+					.AllowSpin(true)
+					.bColorAxisLabels(true)
+					.X(this, &SNiagaraDistributionValueEditor::GetValue, 0, ValueIndex)
+					.OnXChanged(this, &SNiagaraDistributionValueEditor::ValueChanged, 0, ValueIndex)
+					.OnXCommitted(this, &SNiagaraDistributionValueEditor::ValueCommitted, 0, ValueIndex)
+					.Y(this, &SNiagaraDistributionValueEditor::GetValue, 1, ValueIndex)
+					.OnYChanged(this, &SNiagaraDistributionValueEditor::ValueChanged, 1, ValueIndex)
+					.OnYCommitted(this, &SNiagaraDistributionValueEditor::ValueCommitted, 1, ValueIndex)
+					.OnBeginSliderMovement(this, &SNiagaraDistributionValueEditor::BeginValueChange)
+					.OnEndSliderMovement(this, &SNiagaraDistributionValueEditor::EndValueChange)
+				];
 		}
 		if (ChannelCount == 3)
 		{
-			return SNew(SNumericVectorInputBox3)
-				.Font(FAppStyle::Get().GetFontStyle("PropertyWindow.NormalFont"))
-				.AllowSpin(true)
-				.bColorAxisLabels(true)
-				.X(this, &SNiagaraDistributionValueEditor::GetValue, 0, ValueIndex)
-				.OnXChanged(this, &SNiagaraDistributionValueEditor::ValueChanged, 0, ValueIndex)
-				.OnXCommitted(this, &SNiagaraDistributionValueEditor::ValueCommitted, 0, ValueIndex)
-				.Y(this, &SNiagaraDistributionValueEditor::GetValue, 1, ValueIndex)
-				.OnYChanged(this, &SNiagaraDistributionValueEditor::ValueChanged, 1, ValueIndex)
-				.OnYCommitted(this, &SNiagaraDistributionValueEditor::ValueCommitted, 1, ValueIndex)
-				.Z(this, &SNiagaraDistributionValueEditor::GetValue, 2, ValueIndex)
-				.OnZChanged(this, &SNiagaraDistributionValueEditor::ValueChanged, 2, ValueIndex)
-				.OnZCommitted(this, &SNiagaraDistributionValueEditor::ValueCommitted, 2, ValueIndex)
-				.OnBeginSliderMovement(this, &SNiagaraDistributionValueEditor::BeginValueSliderMovement)
-				.OnEndSliderMovement(this, &SNiagaraDistributionValueEditor::EndValueSliderMovement);
+			return SNew(SBox)
+				.MinDesiredWidth(3 * SNiagaraDistributionEditor::DefaultInputSize)
+				[
+					SNew(SNumericVectorInputBox3)
+					.Font(FAppStyle::Get().GetFontStyle("PropertyWindow.NormalFont"))
+					.AllowSpin(true)
+					.bColorAxisLabels(true)
+					.X(this, &SNiagaraDistributionValueEditor::GetValue, 0, ValueIndex)
+					.OnXChanged(this, &SNiagaraDistributionValueEditor::ValueChanged, 0, ValueIndex)
+					.OnXCommitted(this, &SNiagaraDistributionValueEditor::ValueCommitted, 0, ValueIndex)
+					.Y(this, &SNiagaraDistributionValueEditor::GetValue, 1, ValueIndex)
+					.OnYChanged(this, &SNiagaraDistributionValueEditor::ValueChanged, 1, ValueIndex)
+					.OnYCommitted(this, &SNiagaraDistributionValueEditor::ValueCommitted, 1, ValueIndex)
+					.Z(this, &SNiagaraDistributionValueEditor::GetValue, 2, ValueIndex)
+					.OnZChanged(this, &SNiagaraDistributionValueEditor::ValueChanged, 2, ValueIndex)
+					.OnZCommitted(this, &SNiagaraDistributionValueEditor::ValueCommitted, 2, ValueIndex)
+					.OnBeginSliderMovement(this, &SNiagaraDistributionValueEditor::BeginValueChange)
+					.OnEndSliderMovement(this, &SNiagaraDistributionValueEditor::EndValueChange)
+				];
 		}
 		if (ChannelCount == 4)
 		{
-			return SNew(SNumericVectorInputBox4)
-				.Font(FAppStyle::Get().GetFontStyle("PropertyWindow.NormalFont"))
-				.AllowSpin(true)
-				.bColorAxisLabels(true)
-				.X(this, &SNiagaraDistributionValueEditor::GetValue, 0, ValueIndex)
-				.OnXChanged(this, &SNiagaraDistributionValueEditor::ValueChanged, 0, ValueIndex)
-				.OnXCommitted(this, &SNiagaraDistributionValueEditor::ValueCommitted, 0, ValueIndex)
-				.Y(this, &SNiagaraDistributionValueEditor::GetValue, 1, ValueIndex)
-				.OnYChanged(this, &SNiagaraDistributionValueEditor::ValueChanged, 1, ValueIndex)
-				.OnYCommitted(this, &SNiagaraDistributionValueEditor::ValueCommitted, 1, ValueIndex)
-				.Z(this, &SNiagaraDistributionValueEditor::GetValue, 2, ValueIndex)
-				.OnZChanged(this, &SNiagaraDistributionValueEditor::ValueChanged, 2, ValueIndex)
-				.OnZCommitted(this, &SNiagaraDistributionValueEditor::ValueCommitted, 2, ValueIndex)
-				.W(this, &SNiagaraDistributionValueEditor::GetValue, 3, ValueIndex)
-				.OnWChanged(this, &SNiagaraDistributionValueEditor::ValueChanged, 3, ValueIndex)
-				.OnWCommitted(this, &SNiagaraDistributionValueEditor::ValueCommitted, 3, ValueIndex)
-				.OnBeginSliderMovement(this, &SNiagaraDistributionValueEditor::BeginValueSliderMovement)
-				.OnEndSliderMovement(this, &SNiagaraDistributionValueEditor::EndValueSliderMovement);
+			return SNew(SBox)
+				.MinDesiredWidth(4 * SNiagaraDistributionEditor::DefaultInputSize)
+				[
+					SNew(SNumericVectorInputBox4)
+					.Font(FAppStyle::Get().GetFontStyle("PropertyWindow.NormalFont"))
+					.AllowSpin(true)
+					.bColorAxisLabels(true)
+					.X(this, &SNiagaraDistributionValueEditor::GetValue, 0, ValueIndex)
+					.OnXChanged(this, &SNiagaraDistributionValueEditor::ValueChanged, 0, ValueIndex)
+					.OnXCommitted(this, &SNiagaraDistributionValueEditor::ValueCommitted, 0, ValueIndex)
+					.Y(this, &SNiagaraDistributionValueEditor::GetValue, 1, ValueIndex)
+					.OnYChanged(this, &SNiagaraDistributionValueEditor::ValueChanged, 1, ValueIndex)
+					.OnYCommitted(this, &SNiagaraDistributionValueEditor::ValueCommitted, 1, ValueIndex)
+					.Z(this, &SNiagaraDistributionValueEditor::GetValue, 2, ValueIndex)
+					.OnZChanged(this, &SNiagaraDistributionValueEditor::ValueChanged, 2, ValueIndex)
+					.OnZCommitted(this, &SNiagaraDistributionValueEditor::ValueCommitted, 2, ValueIndex)
+					.W(this, &SNiagaraDistributionValueEditor::GetValue, 3, ValueIndex)
+					.OnWChanged(this, &SNiagaraDistributionValueEditor::ValueChanged, 3, ValueIndex)
+					.OnWCommitted(this, &SNiagaraDistributionValueEditor::ValueCommitted, 3, ValueIndex)
+					.OnBeginSliderMovement(this, &SNiagaraDistributionValueEditor::BeginValueChange)
+					.OnEndSliderMovement(this, &SNiagaraDistributionValueEditor::EndValueChange)
+				];
 		}
 		return SNullWidget::NullWidget;
 	}
 
+	TSharedRef<SWidget> ConstructColorWidget(int32 ChannelCount, int32 ValueIndex, bool bShowExpander, FText LabelText)
+	{
+		return SNew(SNiagaraColorEditor)
+			.ShowAlpha(ChannelCount == 4)
+			.ShowExpander(bShowExpander)
+			.ExpandComponents(this, &SNiagaraDistributionValueEditor::GetColorComponentsExpanded)
+			.Color(this, &SNiagaraDistributionValueEditor::GetColorValue, ValueIndex)
+			.LabelText(LabelText)
+			.MinDesiredColorBlockWidth(SNiagaraDistributionEditor::DefaultInputSize)
+			.OnColorChanged(this, &SNiagaraDistributionValueEditor::ColorValueChanged, ValueIndex)
+			.OnBeginEditing(this, &SNiagaraDistributionValueEditor::BeginValueChange)
+			.OnEndEditing(this, &SNiagaraDistributionValueEditor::EndValueChange, 0.0f)
+			.OnCancelEditing(this, &SNiagaraDistributionValueEditor::CancelValueChange, ValueIndex)
+			.OnExpandComponentsChanged(this, &SNiagaraDistributionValueEditor::ColorComponentsExandedChanged);
+	}
+
 	TOptional<float> GetValue(int32 ChannelIndex, int32 ValueIndex) const
 	{
 		return DistributionAdapter->GetConstantOrRangeValue(ChannelIndex, ValueIndex);
@@ -302,18 +379,178 @@ private:
 		}
 	}
 
-	void BeginValueSliderMovement()
+	FLinearColor GetColorValue(int32 ValueIndex) const
+	{
+		if (DistributionAdapter->GetNumChannels() == 3)
+		{
+			return FLinearColor(
+				DistributionAdapter->GetConstantOrRangeValue(0, ValueIndex),
+				DistributionAdapter->GetConstantOrRangeValue(1, ValueIndex), 
+				DistributionAdapter->GetConstantOrRangeValue(2, ValueIndex));
+		}
+		else if (DistributionAdapter->GetNumChannels() == 4)
+		{
+			return FLinearColor(
+				DistributionAdapter->GetConstantOrRangeValue(0, ValueIndex),
+				DistributionAdapter->GetConstantOrRangeValue(1, ValueIndex),
+				DistributionAdapter->GetConstantOrRangeValue(2, ValueIndex),
+				DistributionAdapter->GetConstantOrRangeValue(3, ValueIndex));
+		}
+		return FLinearColor::White;
+	}
+
+	void ColorValueChanged(FLinearColor InColorValue, int32 ValueIndex)
+	{
+		UpdateColorValue(ValueIndex, InColorValue);
+	}
+
+	void BeginValueChange()
 	{
 		DistributionAdapter->BeginContinuousChange();
 	}
 
-	void EndValueSliderMovement(float Value)
+	void EndValueChange(float Value)
 	{
 		DistributionAdapter->EndContinuousChange();
 	}
 
+	void CancelValueChange(FLinearColor OriginalColor, int32 ValueIndex)
+	{
+		UpdateColorValue(ValueIndex, OriginalColor);
+		DistributionAdapter->CancelContinuousChange();
+	}
+
+	void UpdateColorValue(int32 ValueIndex, const FLinearColor& InColorValue)
+	{
+		TArray<float> ColorValues{ InColorValue.R, InColorValue.G, InColorValue.B };
+		if (DistributionAdapter->GetNumChannels() == 4)
+		{
+			ColorValues.Add(InColorValue.A);
+		}
+		DistributionAdapter->SetConstantOrRangeValues(ValueIndex, ColorValues);
+	}
+
 private:
 	TSharedPtr<INiagaraDistributionAdapter> DistributionAdapter;
+	bool bColorComponentsExpanded = false;
+};
+
+class FNiagaraDistributionColorCurveOwner : public FCurveOwnerInterface
+{
+public:
+	FNiagaraDistributionColorCurveOwner(TSharedPtr<INiagaraDistributionAdapter> InDistributionAdapter)
+	{
+		DistributionAdapter = InDistributionAdapter;
+		if (DistributionAdapter->GetDistributionMode() == ENiagaraDistributionEditorMode::ColorGradient &&
+			DistributionAdapter->GetNumChannels() >= 3)
+		{
+			// The gradient editor requires an alpha curve, so we need to create 4 curve infos here even if the distribution
+			// only has 3 channels.
+			for (int32 ChannelIndex = 0; ChannelIndex < 4; ++ChannelIndex)
+			{
+				FRichCurve& EditCurve = EditCurves.Add_GetRef(ChannelIndex < InDistributionAdapter->GetNumChannels()
+					? *DistributionAdapter->GetCurveValue(ChannelIndex)
+					: FRichCurve());
+				ConstCurveInfos.Add(FRichCurveEditInfoConst(&EditCurve, *DistributionAdapter->GetChannelDisplayName(ChannelIndex).ToString()));
+				CurveInfos.Add(FRichCurveEditInfo(&EditCurve, *DistributionAdapter->GetChannelDisplayName(ChannelIndex).ToString()));
+			}
+		}
+	}
+
+	virtual TArray<FRichCurveEditInfoConst> GetCurves() const override { return ConstCurveInfos; }
+	virtual TArray<FRichCurveEditInfo> GetCurves() override { return CurveInfos; }
+
+	virtual void ModifyOwner() override
+	{
+		DistributionAdapter->ModifyOwners();
+	}
+
+	virtual TArray<const UObject*> GetOwners() const override
+	{
+		static TArray<const UObject*> Owners;
+		return Owners;
+	}
+
+	virtual void MakeTransactional() override { }
+
+	virtual void OnCurveChanged(const TArray<FRichCurveEditInfo>& ChangedCurveEditInfos) override
+	{
+		for (const FRichCurveEditInfo& ChangedCurveEditInfo : ChangedCurveEditInfos)
+		{
+			int32 CurveIndex = EditCurves.IndexOfByPredicate([ChangedCurveEditInfo]
+				(const FRichCurve& EditCurve) { return &EditCurve == ChangedCurveEditInfo.CurveToEdit; });
+			if (CurveIndex != INDEX_NONE && CurveIndex < DistributionAdapter->GetNumChannels())
+			{
+				DistributionAdapter->SetCurveValue(CurveIndex, EditCurves[CurveIndex]);
+			}
+		}
+	}
+
+	virtual bool IsLinearColorCurve() const override
+	{
+		return true;
+	}
+
+	virtual FLinearColor GetLinearColorValue(float InTime) const override
+	{
+		return FLinearColor(
+			EditCurves[0].Eval(InTime),
+			EditCurves[1].Eval(InTime),
+			EditCurves[2].Eval(InTime),
+			EditCurves[3].Eval(InTime));
+	}
+
+	virtual bool HasAnyAlphaKeys() const override
+	{
+		return EditCurves[3].GetNumKeys() > 0;
+	}
+
+	virtual bool IsValidCurve(FRichCurveEditInfo CurveInfo) override
+	{
+		return CurveInfos.Contains(CurveInfo);
+	}
+
+	virtual FLinearColor GetCurveColor(FRichCurveEditInfo CurveInfo) const override
+	{
+		return FLinearColor::White;
+	}
+
+private:
+	TSharedPtr<INiagaraDistributionAdapter> DistributionAdapter;
+	TArray<FRichCurve> EditCurves;
+	TArray<FRichCurveEditInfoConst> ConstCurveInfos;
+	TArray<FRichCurveEditInfo> CurveInfos;
+};
+
+class SNiagaraDistributionGradientEditor : public SCompoundWidget
+{
+public:
+	SLATE_BEGIN_ARGS(SNiagaraDistributionGradientEditor) { }
+	SLATE_END_ARGS();
+
+	void Construct(const FArguments& InArgs, TSharedRef<INiagaraDistributionAdapter> InDistributionAdapter)
+	{
+		DistributionAdapter = InDistributionAdapter;
+		ColorCurveOwner = MakeShared<FNiagaraDistributionColorCurveOwner>(DistributionAdapter);
+		TSharedRef<SColorGradientEditor> GradientEditor = SNew(SColorGradientEditor)
+			.ViewMinInput(0.0f)
+			.ViewMaxInput(1.0f)
+			.ClampStopsToViewRange(true);
+
+		GradientEditor->SetCurveOwner(ColorCurveOwner.Get());
+		ChildSlot
+		[
+			SNew(SBox)
+			.Padding(2)
+			[
+				GradientEditor
+			]
+		];
+	}
+
+private:
+	TSharedPtr<INiagaraDistributionAdapter> DistributionAdapter;
+	TSharedPtr<FNiagaraDistributionColorCurveOwner> ColorCurveOwner;
 };
 
 void SNiagaraDistributionEditor::Construct(const FArguments& InArgs, TSharedRef<INiagaraDistributionAdapter> InDistributionAdapter)
@@ -325,7 +562,7 @@ void SNiagaraDistributionEditor::Construct(const FArguments& InArgs, TSharedRef<
 		SNew(SHorizontalBox)
 		+ SHorizontalBox::Slot()
 		.AutoWidth()
-		.Padding(0, 0, 10, 0)
+		.Padding(0, 0, 5, 0)
 		[
 			SNew(SNiagaraDistributionModeSelector, InDistributionAdapter)
 			.OnDistributionModeChanged(this, &SNiagaraDistributionEditor::OnDistributionModeChanged)
@@ -340,6 +577,8 @@ void SNiagaraDistributionEditor::Construct(const FArguments& InArgs, TSharedRef<
 	];
 }
 
+const float SNiagaraDistributionEditor::DefaultInputSize = 125.0f;
+
 void  SNiagaraDistributionEditor::OnDistributionModeChanged()
 {
 	ContentBox->SetContent(ConstructContentForMode());
@@ -368,6 +607,10 @@ TSharedRef<SWidget> SNiagaraDistributionEditor::ConstructContentForMode()
 		}
 		return CurveBox;
 	}
+	else if (FNiagaraDistributionEditorUtilities::IsGradient(Mode))
+	{
+		return SNew(SNiagaraDistributionGradientEditor, DistributionAdapter.ToSharedRef());
+	}
 	else
 	{
 		return SNullWidget::NullWidget;
diff --git a/Engine/Plugins/FX/Niagara/Source/NiagaraEditor/Private/Widgets/SNiagaraDistributionEditor.h b/Engine/Plugins/FX/Niagara/Source/NiagaraEditor/Private/Widgets/SNiagaraDistributionEditor.h
--- a/Engine/Plugins/FX/Niagara/Source/NiagaraEditor/Private/Widgets/SNiagaraDistributionEditor.h
+++ b/Engine/Plugins/FX/Niagara/Source/NiagaraEditor/Private/Widgets/SNiagaraDistributionEditor.h
@@ -18,6 +18,8 @@ public:
 
 	void Construct(const FArguments& InArgs, TSharedRef<INiagaraDistributionAdapter> InDistributionAdapter);
 
+	static const float DefaultInputSize;
+
 private:
 	void OnDistributionModeChanged();
 
diff --git a/Engine/Plugins/FX/Niagara/Source/NiagaraEditor/Private/Widgets/SNiagaraExpandedToggle.cpp b/Engine/Plugins/FX/Niagara/Source/NiagaraEditor/Private/Widgets/SNiagaraExpandedToggle.cpp
new file mode 100644
--- /dev/null
+++ b/Engine/Plugins/FX/Niagara/Source/NiagaraEditor/Private/Widgets/SNiagaraExpandedToggle.cpp
@@ -0,0 +1,54 @@
+// Copyright Epic Games, Inc. All Rights Reserved.
+
+#include "Widgets/SNiagaraExpandedToggle.h"
+
+#include "Styling/AppStyle.h"
+#include "Widgets/Images/SImage.h"
+#include "Widgets/Input/SCheckBox.h"
+
+void SNiagaraExpandedToggle::Construct(const FArguments& InArgs)
+{
+	Expanded = InArgs._Expanded;
+	OnExpandedChangedDelegate = InArgs._OnExpandedChanged;
+	ChildSlot
+	[
+		SNew(SCheckBox)
+		.Style(&FAppStyle::Get().GetWidgetStyle<FCheckBoxStyle>("TransparentCheckBox"))
+		.IsChecked(this, &SNiagaraExpandedToggle::GetExpandedCheckBoxState)
+		.OnCheckStateChanged(this, &SNiagaraExpandedToggle::ExpandedCheckBoxStateChanged)
+		[
+			SNew(SImage)
+			.Image(this, &SNiagaraExpandedToggle::GetExpandedBrush)
+		]
+	];
+}
+
+ECheckBoxState SNiagaraExpandedToggle::GetExpandedCheckBoxState() const
+{
+	return Expanded.Get(false) ? ECheckBoxState::Checked : ECheckBoxState::Unchecked;
+}
+
+void SNiagaraExpandedToggle::ExpandedCheckBoxStateChanged(ECheckBoxState InCheckState)
+{
+	bool bExpanded = InCheckState == ECheckBoxState::Checked;
+	if (Expanded.Get(false) != bExpanded)
+	{
+		if (Expanded.IsBound() == false)
+		{
+			Expanded = bExpanded;
+		}
+		ExpandedBrushCache.Reset();
+		OnExpandedChangedDelegate.ExecuteIfBound(bExpanded);
+	}
+}
+
+const FSlateBrush* SNiagaraExpandedToggle::GetExpandedBrush() const
+{
+	if (ExpandedBrushCache.IsSet() == false)
+	{
+		ExpandedBrushCache = Expanded.Get(false)
+			? FCoreStyle::Get().GetBrush("TreeArrow_Expanded")
+			: FCoreStyle::Get().GetBrush("TreeArrow_Collapsed");
+	}
+	return ExpandedBrushCache.GetValue();
+}
\ No newline at end of file
diff --git a/Engine/Plugins/FX/Niagara/Source/NiagaraEditor/Private/Widgets/SNiagaraExpandedToggle.h b/Engine/Plugins/FX/Niagara/Source/NiagaraEditor/Private/Widgets/SNiagaraExpandedToggle.h
new file mode 100644
--- /dev/null
+++ b/Engine/Plugins/FX/Niagara/Source/NiagaraEditor/Private/Widgets/SNiagaraExpandedToggle.h
@@ -0,0 +1,32 @@
+// Copyright Epic Games, Inc. All Rights Reserved.
+
+#pragma once
+
+#include "Styling/SlateTypes.h"
+#include "Widgets/SCompoundWidget.h"
+
+class SNiagaraExpandedToggle : public SCompoundWidget
+{
+public:
+	DECLARE_DELEGATE_OneParam(FOnExpandedChanged, bool /* bExpanded */);
+
+public:
+	SLATE_BEGIN_ARGS(SNiagaraExpandedToggle)
+		: _Expanded(false)
+		{ }
+		SLATE_ATTRIBUTE(bool, Expanded)
+		SLATE_EVENT(FOnExpandedChanged, OnExpandedChanged)
+	SLATE_END_ARGS();
+
+	void Construct(const FArguments& InArgs);
+
+private:
+	ECheckBoxState GetExpandedCheckBoxState() const;
+	void ExpandedCheckBoxStateChanged(ECheckBoxState InCheckState);
+	const FSlateBrush* GetExpandedBrush() const;
+
+private:
+	TAttribute<bool> Expanded;
+	FOnExpandedChanged OnExpandedChangedDelegate;
+	mutable TOptional<const FSlateBrush*> ExpandedBrushCache;
+};
\ No newline at end of file
diff --git a/Engine/Plugins/FX/Niagara/Source/NiagaraEditorWidgets/Private/DetailCustomizations/NiagaraDataInterfaceCurveDetails.cpp b/Engine/Plugins/FX/Niagara/Source/NiagaraEditorWidgets/Private/DetailCustomizations/NiagaraDataInterfaceCurveDetails.cpp
--- a/Engine/Plugins/FX/Niagara/Source/NiagaraEditorWidgets/Private/DetailCustomizations/NiagaraDataInterfaceCurveDetails.cpp
+++ b/Engine/Plugins/FX/Niagara/Source/NiagaraEditorWidgets/Private/DetailCustomizations/NiagaraDataInterfaceCurveDetails.cpp
@@ -1148,13 +1148,14 @@ public:
 			ColorCurveOwner->OnColorCurveChanged().AddSP(this, &SNiagaraDataInterfaceGradientEditor::ColorCurveChanged);
 			TSharedRef<SColorGradientEditor> GradientEditor = SNew(SColorGradientEditor)
 				.ViewMinInput(0.0f)
-				.ViewMaxInput(1.0f);
+				.ViewMaxInput(1.0f)
+				.ClampStopsToViewRange(true);
 
 			GradientEditor->SetCurveOwner(ColorCurveOwner.Get());
 			ChildSlot
 			[
 				SNew(SBox)
-				.Padding(FMargin(0, 0, 10, 0))
+				.Padding(0, 2, 5, 2)
 				[
 					GradientEditor
 				]
diff --git a/Engine/Plugins/FX/Niagara/Source/NiagaraEditorWidgets/Private/SNiagaraStack.cpp b/Engine/Plugins/FX/Niagara/Source/NiagaraEditorWidgets/Private/SNiagaraStack.cpp
--- a/Engine/Plugins/FX/Niagara/Source/NiagaraEditorWidgets/Private/SNiagaraStack.cpp
+++ b/Engine/Plugins/FX/Niagara/Source/NiagaraEditorWidgets/Private/SNiagaraStack.cpp
@@ -1085,6 +1085,7 @@ SNiagaraStack::FRowWidgets SNiagaraStack::ConstructNameAndValueWidgetsForItem(UN
 
 			TSharedRef<SHorizontalBox> ValueWidget = SNew(SHorizontalBox)
 			+ SHorizontalBox::Slot()
+			.HAlign(HAlign_Left)
 			[
 				PropertyRowWidgets.ValueWidget.ToSharedRef()
 			];
diff --git a/Engine/Source/Editor/UnrealEd/Private/SColorGradientEditor.cpp b/Engine/Source/Editor/UnrealEd/Private/SColorGradientEditor.cpp
--- a/Engine/Source/Editor/UnrealEd/Private/SColorGradientEditor.cpp
+++ b/Engine/Source/Editor/UnrealEd/Private/SColorGradientEditor.cpp
@@ -29,7 +29,74 @@
 
 #define LOCTEXT_NAMESPACE "SColorGradientEditor"
 
-const FSlateRect SColorGradientEditor::HandleRect( 13.0f/2.0f, 0.0f, 13.0f, 16.0f );
+namespace FColorGradientEditorGeometry
+{
+	const float HandleWidth = 13.0f;
+	const float HandleHeight = 16.0f;
+	const float HandleColorWidth = 7.0f;
+	const float HandleColorHeight = 7.0f;
+	const float HandleColorXOffset = (HandleWidth - HandleColorWidth) / 2;
+	const float HandleColorYOffsetForColor = HandleColorXOffset + 1;
+	const float HandleColorYOffsetForAlpha = HandleHeight - HandleColorHeight - HandleColorXOffset - 1;
+
+	const FSlateRect HandleRect = FSlateRect(
+		-FColorGradientEditorGeometry::HandleWidth / 2.0f,
+		0.0f,
+		FColorGradientEditorGeometry::HandleWidth / 2.0f,
+		FColorGradientEditorGeometry::HandleHeight);
+
+	FGeometry GetColorMarkAreaPaintGeometry( const FGeometry& WidgetGeometry )
+	{
+		return WidgetGeometry.MakeChild(
+			FVector2D( WidgetGeometry.GetLocalSize().X - HandleWidth, HandleHeight ),
+			FSlateLayoutTransform(FVector2D(HandleWidth / 2, 0)) );
+	}
+
+	FGeometry GetColorMarkAreaMouseGeometry(const FGeometry& WidgetGeometry)
+	{
+		return WidgetGeometry.MakeChild(
+			FVector2D(WidgetGeometry.GetLocalSize().X, HandleHeight),
+			FSlateLayoutTransform() );
+	}
+
+	FGeometry GetAlphaMarkAreaPaintGeometry( const FGeometry& WidgetGeometry)
+	{
+		return WidgetGeometry.MakeChild(
+			FVector2D(WidgetGeometry.GetLocalSize().X - HandleWidth, HandleHeight ),
+			FSlateLayoutTransform(FVector2D(HandleWidth / 2, WidgetGeometry.GetLocalSize().Y - HandleHeight)) );
+	}
+
+	FGeometry GetAlphaMarkAreaMouseGeometry(const FGeometry& WidgetGeometry)
+	{
+		return WidgetGeometry.MakeChild(
+			FVector2D(WidgetGeometry.GetLocalSize().X, HandleHeight ),
+			FSlateLayoutTransform(FVector2D(0.0f, WidgetGeometry.GetLocalSize().Y - HandleHeight)) );
+	}
+
+	FGeometry GetGradientAreaPaintGeometry( const FGeometry& WidgetGeometry )
+	{
+		return WidgetGeometry.MakeChild(
+			FVector2D( WidgetGeometry.GetLocalSize().X - HandleWidth, WidgetGeometry.GetLocalSize().Y - (2 * HandleHeight) ),
+			FSlateLayoutTransform(FVector2D(HandleWidth / 2, HandleHeight)) );
+	}
+
+	FTrackScaleInfo GetTimeScaleForWidgetGeometry(const FGeometry& WidgetGeometry, float ViewMinInput, float ViewMaxInput)
+	{
+		return FTrackScaleInfo(ViewMinInput, ViewMaxInput, 0.0f, 1.0f, GetGradientAreaPaintGeometry(WidgetGeometry).GetLocalSize());
+	}
+
+	FTrackScaleInfo GetTimeScaleForAreaPaintGeomtry(const FGeometry& AreaGeometry, float ViewMinInput, float ViewMaxInput)
+	{
+		return FTrackScaleInfo(ViewMinInput, ViewMaxInput, 0.0f, 1.0f, AreaGeometry.GetLocalSize());
+	}
+
+	float GetTimeForAbsolutePosition(const FGeometry& WidgetGeometry, float ViewMinInput, float ViewMaxInput, const FVector2D& AbsolutePosition)
+	{
+		FGeometry GradientAreaPaintGeometry = GetGradientAreaPaintGeometry(WidgetGeometry);
+		FTrackScaleInfo TimeScaleInfo = FTrackScaleInfo(ViewMinInput, ViewMaxInput, 0.0f, 1.0f, GradientAreaPaintGeometry.GetLocalSize());
+		return TimeScaleInfo.LocalXToInput(GradientAreaPaintGeometry.AbsoluteToLocal(AbsolutePosition).X);
+	}
+}
 
 FGradientStopMark::FGradientStopMark()
 {}
@@ -135,6 +202,7 @@ void SColorGradientEditor::Construct( const FArguments& InArgs )
 	CurveOwner = NULL;
 	ViewMinInput = InArgs._ViewMinInput;
 	ViewMaxInput = InArgs._ViewMaxInput;
+	bClampStopsToViewRange = InArgs._ClampStopsToViewRange;
 	bDraggingAlphaValue = false;
 	bDraggingStop = false;
 	DistanceDragged = 0.0f;
@@ -148,20 +216,19 @@ int32 SColorGradientEditor::OnPaint( const FPaintArgs& Args, const FGeometry& Al
 	if( CurveOwner )
 	{
 		// Split the geometry into areas for stops and the gradient
-		FGeometry ColorMarkAreaGeometry = GetColorMarkAreaGeometry( AllottedGeometry );
-		FGeometry AlphaMarkAreaGeometry = GetAlphaMarkAreaGeometry( AllottedGeometry );
-
-		FGeometry GradientAreaGeometry = AllottedGeometry.MakeChild( FVector2D( AllottedGeometry.GetLocalSize().X, AllottedGeometry.GetLocalSize().Y - 30.0f ), FSlateLayoutTransform(FVector2D(0.0f, 16.0f)) );
+		FGeometry ColorMarkAreaGeometry = FColorGradientEditorGeometry::GetColorMarkAreaPaintGeometry( AllottedGeometry );
+		FGeometry AlphaMarkAreaGeometry = FColorGradientEditorGeometry::GetAlphaMarkAreaPaintGeometry( AllottedGeometry );
+		FGeometry GradientAreaGeometry = FColorGradientEditorGeometry::GetGradientAreaPaintGeometry( AllottedGeometry );
 
 		bool bEnabled = ShouldBeEnabled( bParentEnabled );
 		ESlateDrawEffect DrawEffects = bEnabled ? ESlateDrawEffect::None : ESlateDrawEffect::DisabledEffect;
 
 		// Pixel to value input converter
-		FTrackScaleInfo ScaleInfo(ViewMinInput.Get(),  ViewMaxInput.Get(), 0.0f, 1.0f, GradientAreaGeometry.GetLocalSize());
+		FTrackScaleInfo ScaleInfo = FColorGradientEditorGeometry::GetTimeScaleForWidgetGeometry(AllottedGeometry, ViewMinInput.Get(), ViewMaxInput.Get());
 
 		// The start and end location in slate units of the area to draw
 		int32 Start = 0;
-		int32 Finish = FMath::TruncToInt( AllottedGeometry.GetLocalSize().X );
+		int32 Finish = FMath::TruncToInt( GradientAreaGeometry.GetLocalSize().X );
 
 		TArray<FSlateGradientStop> Stops;
 
@@ -380,15 +447,15 @@ FReply SColorGradientEditor::OnMouseButtonDoubleClick( const FGeometry& InMyGeom
 
 FReply SColorGradientEditor::OnMouseMove( const FGeometry& MyGeometry, const FPointerEvent& MouseEvent )
 {
-	FGeometry ColorMarkAreaGeometry = GetColorMarkAreaGeometry(MyGeometry);
-	FGeometry AlphaMarkAreaGeometry = GetAlphaMarkAreaGeometry(MyGeometry);
+	FGeometry ColorMarkAreaMouseGeometry = FColorGradientEditorGeometry::GetColorMarkAreaMouseGeometry(MyGeometry);
+	FGeometry AlphaMarkAreaMouseGeometry = FColorGradientEditorGeometry::GetAlphaMarkAreaMouseGeometry(MyGeometry);
 
-	if (ColorMarkAreaGeometry.IsUnderLocation(MouseEvent.GetScreenSpacePosition()))
+	if (ColorMarkAreaMouseGeometry.IsUnderLocation(MouseEvent.GetScreenSpacePosition()))
 	{
 		bColorAreaHovered = true;
 		bAlphaAreaHovered = false;
 	}
-	else if (AlphaMarkAreaGeometry.IsUnderLocation(MouseEvent.GetScreenSpacePosition()))
+	else if (AlphaMarkAreaMouseGeometry.IsUnderLocation(MouseEvent.GetScreenSpacePosition()))
 	{
 		bColorAreaHovered = false;
 		bAlphaAreaHovered = true;
@@ -421,8 +488,11 @@ FReply SColorGradientEditor::OnMouseMove( const FGeometry& MyGeometry, const FPo
 			else
 			{
 				// Already dragging a stop, move it
-				FTrackScaleInfo ScaleInfo(ViewMinInput.Get(),  ViewMaxInput.Get(), 0.0f, 1.0f, MyGeometry.GetLocalSize());
-				float MouseTime = ScaleInfo.LocalXToInput( MyGeometry.AbsoluteToLocal( MouseEvent.GetScreenSpacePosition() ).X );
+				float MouseTime = FColorGradientEditorGeometry::GetTimeForAbsolutePosition(MyGeometry, ViewMinInput.Get(), ViewMaxInput.Get(), MouseEvent.GetScreenSpacePosition());
+				if (bClampStopsToViewRange)
+				{
+					MouseTime = FMath::Clamp(MouseTime, ViewMinInput.Get(), ViewMaxInput.Get());
+				}
 				MoveStop( SelectedStop, MouseTime );
 
 				return FReply::Handled();
@@ -448,10 +518,10 @@ FReply SColorGradientEditor::OnMouseButtonUp( const FGeometry& MyGeometry, const
 			}
 			else if( DistanceDragged < DragThresholdDist && !SelectedStop.IsValid( *CurveOwner ) )
 			{
-				FGeometry ColorMarkAreaGeometry = GetColorMarkAreaGeometry( MyGeometry );
-				FGeometry AlphaMarkAreaGeometry = GetAlphaMarkAreaGeometry( MyGeometry );
+				FGeometry ColorMarkAreaMouseGeometry = FColorGradientEditorGeometry::GetColorMarkAreaMouseGeometry( MyGeometry );
+				FGeometry AlphaMarkAreaMouseGeometry = FColorGradientEditorGeometry::GetAlphaMarkAreaMouseGeometry( MyGeometry );
 
-				if( ColorMarkAreaGeometry.IsUnderLocation( MouseEvent.GetScreenSpacePosition() ) )
+				if( ColorMarkAreaMouseGeometry.IsUnderLocation( MouseEvent.GetScreenSpacePosition() ) )
 				{
 					// Add a new color mark
 					bool bColorStop = true;
@@ -460,7 +530,7 @@ FReply SColorGradientEditor::OnMouseButtonUp( const FGeometry& MyGeometry, const
 					return FReply::Handled().CaptureMouse( SharedThis(this) );
 
 				}
-				else if( AlphaMarkAreaGeometry.IsUnderLocation( MouseEvent.GetScreenSpacePosition() ) )
+				else if( AlphaMarkAreaMouseGeometry.IsUnderLocation( MouseEvent.GetScreenSpacePosition() ) )
 				{
 					// Add a new alpha mark
 					bool bColorStop = false;
@@ -758,8 +828,6 @@ void SColorGradientEditor::DrawGradientStopMark( const FGradientStopMark& Mark,
 	static const FSlateBrush* AlphaStopBrush = FAppStyle::GetBrush("CurveEditor.Gradient.HandleUp");
 	static const FLinearColor SelectionColor = FAppStyle::GetSlateColor("SelectionColor").GetColor( InWidgetStyle );
 
-	const float HandleSize = 13.0f;
-
 	bool bSelected = Mark == SelectedStop;
 	if( bSelected )
 	{
@@ -768,11 +836,13 @@ void SColorGradientEditor::DrawGradientStopMark( const FGradientStopMark& Mark,
 	}
 
 	// Draw a box for the non colored area
+	FVector2f HandleSize = FColorGradientEditorGeometry::HandleRect.GetSize();
+	FVector2f HandlePosition = FColorGradientEditorGeometry::HandleRect.GetTopLeft2f() + FVector2f(XPos, 0);
 	FSlateDrawElement::MakeBox
 	( 
 		OutDrawElements,
 		LayerId,
-		Geometry.ToPaintGeometry( FVector2f( HandleRect.Right, HandleRect.Bottom ), FSlateLayoutTransform(FVector2f( XPos-HandleRect.Left, HandleRect.Top )) ),
+		Geometry.ToPaintGeometry( HandleSize, FSlateLayoutTransform(HandlePosition) ),
 		bColor ? ColorStopBrush : AlphaStopBrush,
 		DrawEffects,
 		bSelected ? SelectionColor : FLinearColor::White
@@ -780,76 +850,73 @@ void SColorGradientEditor::DrawGradientStopMark( const FGradientStopMark& Mark,
 
 	// Draw a box with the gradient stop color
 	//Slate MakeGradient call expects the linear colors to be pre-converted to sRGB
+	FVector2f HandleColorSize = FVector2f(FColorGradientEditorGeometry::HandleColorWidth, FColorGradientEditorGeometry::HandleColorHeight);
+	FVector2f HandleColorPosition = HandlePosition + FVector2f(
+		FColorGradientEditorGeometry::HandleColorXOffset, 
+		bColor 
+			? FColorGradientEditorGeometry::HandleColorYOffsetForColor 
+			: FColorGradientEditorGeometry::HandleColorYOffsetForAlpha);
 	FSlateDrawElement::MakeBox
 	( 
 		OutDrawElements,
 		LayerId+1,
-		Geometry.ToPaintGeometry( FVector2f( HandleRect.Right-6.f, HandleRect.Bottom-9.f ), FSlateLayoutTransform(FVector2f( XPos-HandleRect.Left+3.f, bColor ? HandleRect.Top+3.0f : HandleRect.Top+6.f)) ),
+		Geometry.ToPaintGeometry( HandleColorSize, FSlateLayoutTransform(HandleColorPosition) ),
 		WhiteBrush,
 		DrawEffects,
 		Color.ToFColor(true)
 	);
 }
-FGeometry SColorGradientEditor::GetColorMarkAreaGeometry( const FGeometry& InGeometry ) const
-{
-	return InGeometry.MakeChild( FVector2D( InGeometry.GetLocalSize().X, 16.0f ), FSlateLayoutTransform() );
-}
-
-FGeometry SColorGradientEditor::GetAlphaMarkAreaGeometry( const FGeometry& InGeometry ) const
-{
-	return InGeometry.MakeChild( FVector2D( InGeometry.GetLocalSize().X, 16.0f ), FSlateLayoutTransform(FVector2D( 0.0f, InGeometry.GetLocalSize().Y-14.0f)) );
-}
 
 FGradientStopMark SColorGradientEditor::GetGradientStopAtPoint( const FVector2D& MousePos, const FGeometry& MyGeometry )
 {
-	FGeometry ColorMarkAreaGeometry = GetColorMarkAreaGeometry( MyGeometry );
-	FGeometry AlphaMarkAreaGeometry = GetAlphaMarkAreaGeometry( MyGeometry );
-
-	FTrackScaleInfo ScaleInfo(ViewMinInput.Get(),  ViewMaxInput.Get(), 0.0f, 1.0f, MyGeometry.GetLocalSize());
+	FGradientStopMark StopUnderMouse;
+	FGeometry ColorMarkAreaMouseGeometry = FColorGradientEditorGeometry::GetColorMarkAreaMouseGeometry( MyGeometry );
+	FGeometry AlphaMarkAreaMouseGeometry = FColorGradientEditorGeometry::GetAlphaMarkAreaMouseGeometry( MyGeometry );
+	bool bUnderColorMarkAreaGeometry = ColorMarkAreaMouseGeometry.IsUnderLocation(MousePos);
+	bool bUnderAlphaMarkAreaGeometry = AlphaMarkAreaMouseGeometry.IsUnderLocation(MousePos);
 
-	if( ColorMarkAreaGeometry.IsUnderLocation( MousePos ) || AlphaMarkAreaGeometry.IsUnderLocation( MousePos ) )
+	if( bUnderColorMarkAreaGeometry || bUnderAlphaMarkAreaGeometry )
 	{
 		TArray<FGradientStopMark> ColorMarks;
 		TArray<FGradientStopMark> AlphaMarks;
 		GetGradientStopMarks( ColorMarks, AlphaMarks );
 
-		// See if any color stops are under the mouse
-		for( int32 ColorIndex = 0; ColorIndex < ColorMarks.Num(); ++ColorIndex )
+		auto GetMarkUnderMouse = [this](const FVector2D& MousePos, const TArray<FGradientStopMark>& Marks, const FGeometry& AreaPaintGeometry)
 		{
-			const FGradientStopMark& Mark = ColorMarks[ColorIndex];
+			FTrackScaleInfo TimeScaleInfo = FColorGradientEditorGeometry::GetTimeScaleForAreaPaintGeomtry(AreaPaintGeometry, ViewMinInput.Get(), ViewMaxInput.Get());
+			for (int32 MarkIndex = 0; MarkIndex < Marks.Num(); ++MarkIndex)
+			{
+				const FGradientStopMark& Mark = Marks[MarkIndex];
 
-			// Convert the time to a screen coordinate
-			float XVal = ScaleInfo.InputToLocalX( Mark.Time );
+				// Convert the time to a screen coordinate
+				float XVal = TimeScaleInfo.InputToLocalX(Mark.Time);
 
-			if( XVal >= 0 )
-			{
-				FGeometry MarkGeometry = ColorMarkAreaGeometry.MakeChild( FVector2D( HandleRect.Right, HandleRect.Bottom ), FSlateLayoutTransform(FVector2D( XVal-HandleRect.Left, HandleRect.Top )) );
-				if( MarkGeometry.IsUnderLocation( MousePos ) )
+				if (XVal >= 0)
 				{
-					return Mark;
+					FVector2f HandleSize = FColorGradientEditorGeometry::HandleRect.GetSize2f();
+					FVector2f HandlePosition = FColorGradientEditorGeometry::HandleRect.GetTopLeft2f() + FVector2f(XVal, 0.0f);
+					FGeometry HandleGeometry = AreaPaintGeometry.MakeChild(HandleSize, FSlateLayoutTransform(HandlePosition));
+					if (HandleGeometry.IsUnderLocation(MousePos))
+					{
+						return Mark;
+					}
 				}
 			}
-		}
+			return FGradientStopMark();
+		};
 
-		// See if any color stops are under the mouse
-		for( int32 ColorIndex = 0; ColorIndex < AlphaMarks.Num(); ++ColorIndex )
+		if (bUnderColorMarkAreaGeometry)
 		{
-			const FGradientStopMark& Mark = AlphaMarks[ColorIndex];
-
-			float XVal = ScaleInfo.InputToLocalX( Mark.Time );
-
-			if( XVal >= 0 )
-			{
-				FGeometry MarkGeometry = AlphaMarkAreaGeometry.MakeChild( FVector2D( HandleRect.Right, HandleRect.Bottom ), FSlateLayoutTransform(FVector2D( XVal-HandleRect.Left, HandleRect.Top )) );
-				if( MarkGeometry.IsUnderLocation( MousePos ) )
-				{
-					return Mark;
-				}
-			}
+			FGeometry ColorMarkAreaPaintGeometry = FColorGradientEditorGeometry::GetColorMarkAreaPaintGeometry(MyGeometry);
+			StopUnderMouse = GetMarkUnderMouse(MousePos, ColorMarks, ColorMarkAreaPaintGeometry);
+		}
+		else // bUnderAlphaMarkAreaGeometry
+		{
+			FGeometry AlphaMarkAreaPaintGeometry = FColorGradientEditorGeometry::GetAlphaMarkAreaPaintGeometry(MyGeometry);
+			StopUnderMouse = GetMarkUnderMouse(MousePos, AlphaMarks, AlphaMarkAreaPaintGeometry);
 		}
 	}
-
-	return FGradientStopMark();
+	return StopUnderMouse;
 }
 
 void SColorGradientEditor::GetGradientStopMarks( TArray<FGradientStopMark>& OutColorMarks, TArray<FGradientStopMark>& OutAlphaMarks ) const
@@ -928,12 +995,8 @@ FGradientStopMark SColorGradientEditor::AddStop( const FVector2D& Position, cons
 
 	CurveOwner->ModifyOwner();
 
-	FTrackScaleInfo ScaleInfo(ViewMinInput.Get(),  ViewMaxInput.Get(), 0.0f, 1.0f, MyGeometry.GetLocalSize());
-
-	FVector2D LocalPos = MyGeometry.AbsoluteToLocal( Position );
+	float NewStopTime = FColorGradientEditorGeometry::GetTimeForAbsolutePosition(MyGeometry, ViewMinInput.Get(), ViewMaxInput.Get(), Position);
 
-	float NewStopTime = ScaleInfo.LocalXToInput( LocalPos.X );
-			
 	TArray<FRichCurveEditInfo> Curves = CurveOwner->GetCurves();
 
 	FGradientStopMark NewStop;
diff --git a/Engine/Source/Editor/UnrealEd/Public/SColorGradientEditor.h b/Engine/Source/Editor/UnrealEd/Public/SColorGradientEditor.h
--- a/Engine/Source/Editor/UnrealEd/Public/SColorGradientEditor.h
+++ b/Engine/Source/Editor/UnrealEd/Public/SColorGradientEditor.h
@@ -49,10 +49,12 @@ public:
 		: _ViewMinInput(0.0f)
 		, _ViewMaxInput(0.0f)
 		, _IsEditingEnabled( true )
+		, _ClampStopsToViewRange( false )
 	{}
 		SLATE_ATTRIBUTE( float, ViewMinInput )
 		SLATE_ATTRIBUTE( float, ViewMaxInput )
 		SLATE_ATTRIBUTE( bool, IsEditingEnabled )
+		SLATE_ARGUMENT( bool, ClampStopsToViewRange )
 	SLATE_END_ARGS()
 
 
@@ -151,20 +153,6 @@ private:
 	 */
 	UNREALED_API void DrawGradientStopMark( const FGradientStopMark& Mark, const FGeometry& Geometry, float XPos, const FLinearColor& Color, FSlateWindowElementList& OutDrawElements, int32 LayerId, const FSlateRect& InClippingRect, ESlateDrawEffect DrawEffects, bool bColor, const FWidgetStyle& InWidgetStyle ) const;
 
-	/**
-	 * Calculates the geometry of the gradient stop color mark area
-	 *
-	 * @param InGeometry The parent geometry of the area
-	 */
-	UNREALED_API FGeometry GetColorMarkAreaGeometry( const FGeometry& InGeometry ) const;
-
-	/**
-	 * Calculates the geometry of the gradient stop alpha mark area
-	 *
-	 * @param InGeometry The parent geometry of the area
-	 */
-	UNREALED_API FGeometry GetAlphaMarkAreaGeometry( const FGeometry& InGeometry ) const;
-
 	/**
 	 * Gets the gradient stop (if any) of at the current mouse position
 	 *
@@ -204,9 +192,6 @@ private:
 	UNREALED_API void MoveStop( FGradientStopMark& Mark, float NewTime );
 
 private:
-	/** Local space rectangle of each gradient stop handle */
-	static UNREALED_API const FSlateRect HandleRect;
-
 	/** The currently selected stop */
 	FGradientStopMark SelectedStop;
 	/** The color that was last modified. New stops are added with this color */
@@ -219,6 +204,8 @@ private:
 	TAttribute<float> ViewMaxInput;
 	/** Whether or not the gradient is editable or just viewed */
 	TAttribute<bool> IsEditingEnabled;
+	/** Whether or not to clamp the time value of stops to the view range. */
+	bool bClampStopsToViewRange;
 	/** Cached position where context menus should appear */
 	FVector2D ContextMenuPosition;
 	/** Whether or not the color gradient stop area is hovered */
