- Priliminary support for generic bindings on distributions

diff --git a/Engine/Plugins/FX/Niagara/Shaders/Private/Stateless/Modules/NiagaraStatelessModule_InitializeParticle.ush b/Engine/Plugins/FX/Niagara/Shaders/Private/Stateless/Modules/NiagaraStatelessModule_InitializeParticle.ush
--- a/Engine/Plugins/FX/Niagara/Shaders/Private/Stateless/Modules/NiagaraStatelessModule_InitializeParticle.ush
+++ b/Engine/Plugins/FX/Niagara/Shaders/Private/Stateless/Modules/NiagaraStatelessModule_InitializeParticle.ush
@@ -6,7 +6,7 @@ static const uint EInitializeParticleModuleFlag_UniformSpriteSize	= 1 << 0;
 static const uint EInitializeParticleModuleFlag_UniformMeshScale	= 1 << 1;
 
 uint	InitializeParticle_ModuleFlags;
-float3	InitializeParticle_Position;
+uint3	InitializeParticle_InitialPosition;
 float4	InitializeParticle_ColorScale;
 float4	InitializeParticle_ColorBias;
 float2	InitializeParticle_SpriteSizeScale;
@@ -29,8 +29,11 @@ void InitializeParticle_Initialize(inout FStatelessModule_InitializeParticle Mod
 	const bool bUniformSpriteSize = (InitializeParticle_ModuleFlags & EInitializeParticleModuleFlag_UniformSpriteSize) != 0;
 	const bool bUniformMeshScale  = (InitializeParticle_ModuleFlags & EInitializeParticleModuleFlag_UniformMeshScale) != 0;
 
+	FStatelessDistributionFloat3 InitialPosition;
+	InitialPosition.Init(Particle, InitializeParticle_InitialPosition);
+		
 	Particle.UniqueID		= Particle.UniqueIndex;
-	Particle.Position		= InitializeParticle_Position;
+	Particle.Position		= InitialPosition.GetValue(0.0f);
 	Particle.Color			= RandomScaleBiasFloat(Particle, InitializeParticle_ColorScale, InitializeParticle_ColorBias);
 	Particle.RibbonWidth	= RandomScaleBiasFloat(Particle, InitializeParticle_RibbonWidthScale, InitializeParticle_RibbonWidthBias);
 	Particle.SpriteSize		= RandomScaleBiasFloat(Particle, InitializeParticle_SpriteSizeScale, InitializeParticle_SpriteSizeBias, bUniformSpriteSize);
diff --git a/Engine/Plugins/FX/Niagara/Shaders/Private/Stateless/NiagaraStatelessCommon.ush b/Engine/Plugins/FX/Niagara/Shaders/Private/Stateless/NiagaraStatelessCommon.ush
--- a/Engine/Plugins/FX/Niagara/Shaders/Private/Stateless/NiagaraStatelessCommon.ush
+++ b/Engine/Plugins/FX/Niagara/Shaders/Private/Stateless/NiagaraStatelessCommon.ush
@@ -27,6 +27,7 @@ uint				Common_OutputBufferStride;
 RWBuffer<float>		Common_FloatOutputBuffer;
 RWBuffer<int>		Common_IntOutputBuffer;
 Buffer<float>		Common_StaticFloatBuffer;
+Buffer<uint>		Common_ParameterBuffer;
 
 ////////////////////////////////////////////////////////////////////////////////////
 
@@ -126,6 +127,11 @@ float2 GetStaticFloat2(int BaseOffset, int Element)	{ return float2(Common_Stati
 float3 GetStaticFloat3(int BaseOffset, int Element)	{ return float3(Common_StaticFloatBuffer[BaseOffset + Element*3], Common_StaticFloatBuffer[BaseOffset + 1 + Element*3], Common_StaticFloatBuffer[BaseOffset + 2 + Element*3]); }
 float4 GetStaticFloat4(int BaseOffset, int Element)	{ return float4(Common_StaticFloatBuffer[BaseOffset + Element*4], Common_StaticFloatBuffer[BaseOffset + 1 + Element*4], Common_StaticFloatBuffer[BaseOffset + 2 + Element*4], Common_StaticFloatBuffer[BaseOffset + 3 + Element*4]); }
 
+float  GetParameterBufferFloat (int BaseOffset, int Element) { return asfloat(Common_ParameterBuffer[BaseOffset + Element]); }
+float2 GetParameterBufferFloat2(int BaseOffset, int Element) { return float2(asfloat(Common_ParameterBuffer[BaseOffset + Element*2]), asfloat(Common_ParameterBuffer[BaseOffset + 1 + Element*2])); }
+float3 GetParameterBufferFloat3(int BaseOffset, int Element) { return float3(asfloat(Common_ParameterBuffer[BaseOffset + Element*3]), asfloat(Common_ParameterBuffer[BaseOffset + 1 + Element*3]), asfloat(Common_ParameterBuffer[BaseOffset + 2 + Element*3])); }
+float4 GetParameterBufferFloat4(int BaseOffset, int Element) { return float4(asfloat(Common_ParameterBuffer[BaseOffset + Element*4]), asfloat(Common_ParameterBuffer[BaseOffset + 1 + Element*4]), asfloat(Common_ParameterBuffer[BaseOffset + 2 + Element*4]), asfloat(Common_ParameterBuffer[BaseOffset + 3 + Element*4])); }
+	
 float LerpStaticFloat (uint2 Parameters, float U)
 {
 	const float Offset = U * float(Parameters.y);
diff --git a/Engine/Plugins/FX/Niagara/Shaders/Private/Stateless/NiagaraStatelessDistribution.ush b/Engine/Plugins/FX/Niagara/Shaders/Private/Stateless/NiagaraStatelessDistribution.ush
--- a/Engine/Plugins/FX/Niagara/Shaders/Private/Stateless/NiagaraStatelessDistribution.ush
+++ b/Engine/Plugins/FX/Niagara/Shaders/Private/Stateless/NiagaraStatelessDistribution.ush
@@ -3,11 +3,13 @@
 
 #define StatelessDistributionFlag_Random	0x00000001
 #define StatelessDistributionFlag_Uniform	0x00000002
+#define StatelessDistributionFlag_Binding	0x00000004
 
 struct FStatelessDistributionFloat
 {
 	bool	bRandomDistribution;
 	bool	bUniformValue;
+	bool	bBindingValue;
 	uint	TableOffset;
 	float	TableLength;
 	float	RandomOffset;
@@ -16,6 +18,7 @@ struct FStatelessDistributionFloat
 	{
 		bRandomDistribution = (Parameters.x & StatelessDistributionFlag_Random) != 0;
 		bUniformValue		= (Parameters.x & StatelessDistributionFlag_Uniform) != 0;
+		bBindingValue		= (Parameters.x & StatelessDistributionFlag_Binding) != 0;
 		TableOffset			= Parameters.y;
 		TableLength			= asfloat(Parameters.z);
 		RandomOffset		= RandomFloat(Particle);
@@ -28,7 +31,11 @@ struct FStatelessDistributionFloat
 	
 	float GetValue(float NormalizedAge)
 	{
-		if ( bRandomDistribution )
+		if ( bBindingValue )
+		{
+			return GetParameterBufferFloat(TableOffset, 0);
+		}
+		else if ( bRandomDistribution )
 		{
 			const float	Value0 = GetStaticFloat(TableOffset, 0);
 			const float	Value1 = GetStaticFloat(TableOffset, 1);
@@ -48,6 +55,7 @@ struct FStatelessDistributionFloat2
 {
 	bool	bRandomDistribution;
 	bool	bUniformValue;
+	bool	bBindingValue;
 	uint	TableOffset;
 	float	TableLength;
 	float2	RandomOffset;
@@ -56,6 +64,7 @@ struct FStatelessDistributionFloat2
 	{
 		bRandomDistribution = (Parameters.x & StatelessDistributionFlag_Random) != 0;
 		bUniformValue		= (Parameters.x & StatelessDistributionFlag_Uniform) != 0;
+		bBindingValue		= (Parameters.x & StatelessDistributionFlag_Binding) != 0;
 		TableOffset			= Parameters.y;
 		TableLength			= asfloat(Parameters.z);
 		RandomOffset		= 0;
@@ -72,7 +81,11 @@ struct FStatelessDistributionFloat2
 	
 	float2 GetValue(float NormalizedAge)
 	{
-		if ( bRandomDistribution )
+		if ( bBindingValue )
+		{
+			return GetParameterBufferFloat2(TableOffset, 0);
+		}
+		else if ( bRandomDistribution )
 		{
 			const float2 Value0 = GetStaticFloat2(TableOffset, 0);
 			const float2 Value1 = GetStaticFloat2(TableOffset, 1);
@@ -95,6 +108,7 @@ struct FStatelessDistributionFloat3
 {
 	bool	bRandomDistribution;
 	bool	bUniformValue;
+	bool	bBindingValue;
 	uint	TableOffset;
 	float	TableLength;
 	float3	RandomOffset;
@@ -103,6 +117,7 @@ struct FStatelessDistributionFloat3
 	{
 		bRandomDistribution = (Parameters.x & StatelessDistributionFlag_Random) != 0;
 		bUniformValue		= (Parameters.x & StatelessDistributionFlag_Uniform) != 0;
+		bBindingValue		= (Parameters.x & StatelessDistributionFlag_Binding) != 0;
 		TableOffset			= Parameters.y;
 		TableLength			= asfloat(Parameters.z);
 		RandomOffset		= 0;
@@ -119,7 +134,11 @@ struct FStatelessDistributionFloat3
 	
 	float3 GetValue(float NormalizedAge)
 	{
-		if ( bRandomDistribution )
+		if ( bBindingValue )
+		{
+			return GetParameterBufferFloat3(TableOffset, 0);
+		}
+		else if ( bRandomDistribution )
 		{
 			const float3 Value0 = GetStaticFloat3(TableOffset, 0);
 			const float3 Value1 = GetStaticFloat3(TableOffset, 1);
@@ -143,6 +162,7 @@ struct FStatelessDistributionFloat4
 {
 	bool	bRandomDistribution;
 	bool	bUniformValue;
+	bool	bBindingValue;
 	uint	TableOffset;
 	float	TableLength;
 	float4	RandomOffset;
@@ -151,6 +171,7 @@ struct FStatelessDistributionFloat4
 	{
 		bRandomDistribution = (Parameters.x & StatelessDistributionFlag_Random) != 0;
 		bUniformValue		= (Parameters.x & StatelessDistributionFlag_Uniform) != 0;
+		bBindingValue		= (Parameters.x & StatelessDistributionFlag_Binding) != 0;
 		TableOffset			= Parameters.y;
 		TableLength			= asfloat(Parameters.z);
 		RandomOffset		= 0;
@@ -167,7 +188,11 @@ struct FStatelessDistributionFloat4
 	
 	float4 GetValue(float NormalizedAge)
 	{
-		if ( bRandomDistribution )
+		if ( bBindingValue )
+		{
+			return GetParameterBufferFloat4(TableOffset, 0);
+		}
+		else if ( bRandomDistribution )
 		{
 			const float4 Value0 = GetStaticFloat4(TableOffset, 0);
 			const float4 Value1 = GetStaticFloat4(TableOffset, 1);
diff --git a/Engine/Plugins/FX/Niagara/Source/Niagara/Internal/Stateless/Modules/NiagaraStatelessModule_InitializeParticle.h b/Engine/Plugins/FX/Niagara/Source/Niagara/Internal/Stateless/Modules/NiagaraStatelessModule_InitializeParticle.h
--- a/Engine/Plugins/FX/Niagara/Source/Niagara/Internal/Stateless/Modules/NiagaraStatelessModule_InitializeParticle.h
+++ b/Engine/Plugins/FX/Niagara/Source/Niagara/Internal/Stateless/Modules/NiagaraStatelessModule_InitializeParticle.h
@@ -23,7 +23,7 @@ class UNiagaraStatelessModule_InitializeParticle : public UNiagaraStatelessModul
 	struct FModuleBuiltData
 	{
 		uint32							ModuleFlags = 0;
-		int32							PositionParameterBinding = INDEX_NONE;
+		FUintVector3					InitialPosition = FUintVector3::ZeroValue;
 		FNiagaraStatelessRangeFloat		LifetimeRange;
 		FNiagaraStatelessRangeColor		ColorRange;
 		FNiagaraStatelessRangeFloat		MassRange;
@@ -60,23 +60,8 @@ public:
 	UPROPERTY(EditAnywhere, Category = "Parameters", meta = (DisplayName = "Ribbon Width", EditCondition="bWriteRibbonWidth"))
 	FNiagaraDistributionRangeFloat RibbonWidthDistribution = FNiagaraDistributionRangeFloat(FNiagaraStatelessGlobals::GetDefaultRibbonWidthValue());
 
-	UPROPERTY(EditAnywhere, Category = "Parameters")
-	FNiagaraParameterBindingWithValue	InitialPositionBinding;
-
-	//~Begin: UObject Interface
-#if WITH_EDITORONLY_DATA
-	virtual void PostInitProperties()
-	{
-		Super::PostInitProperties();
-		if (HasAnyFlags(RF_ClassDefaultObject) == false)
-		{
-			InitialPositionBinding.SetUsage(ENiagaraParameterBindingUsage::NotParticle);
-			InitialPositionBinding.SetAllowedTypeDefinitions({ FNiagaraTypeDefinition::GetVec3Def() });
-			InitialPositionBinding.SetDefaultParameter(FNiagaraTypeDefinition::GetVec3Def(), FVector3f::ZeroVector);
-		}
-	}
-#endif
-	//~End: UObject Interface
+	UPROPERTY(EditAnywhere, Category = "Parameters", meta = (DisableCurveDistribution))
+	FNiagaraDistributionVector3	InitialPosition = FNiagaraDistributionVector3(FVector3f::ZeroVector);
 
 	virtual void BuildEmitterData(FNiagaraStatelessEmitterDataBuildContext& BuildContext) const override
 	{
@@ -84,7 +69,7 @@ public:
 		BuiltData->ModuleFlags				 = SpriteSizeDistribution.IsUniform() ? EInitializeParticleModuleFlag_UniformSpriteSize : 0;
 		BuiltData->ModuleFlags				|= MeshScaleDistribution.IsUniform() ? EInitializeParticleModuleFlag_UniformMeshScale : 0;
 
-		BuiltData->PositionParameterBinding	= BuildContext.AddRendererBinding(InitialPositionBinding);
+		BuiltData->InitialPosition			= BuildContext.AddDistribution(InitialPosition, true);
 		BuiltData->LifetimeRange			= LifetimeDistribution.CalculateRange(FNiagaraStatelessGlobals::GetDefaultLifetimeValue());
 		BuiltData->ColorRange				= ColorDistribution.CalculateRange(FNiagaraStatelessGlobals::GetDefaultColorValue());
 		BuiltData->MassRange				= MassDistribution.CalculateRange(FNiagaraStatelessGlobals::GetDefaultMassValue());
@@ -102,8 +87,8 @@ public:
 		FParameters* Parameters = SetShaderParameterContext.GetParameterNestedStruct<FParameters>();
 		const FModuleBuiltData* ModuleBuiltData = SetShaderParameterContext.ReadBuiltData<FModuleBuiltData>();
 
-		SetShaderParameterContext.GetRendererParameterValue(Parameters->InitializeParticle_Position, ModuleBuiltData->PositionParameterBinding, InitialPositionBinding.GetDefaultValue<FVector3f>());
 		Parameters->InitializeParticle_ModuleFlags			= ModuleBuiltData->ModuleFlags;
+		Parameters->InitializeParticle_InitialPosition		= ModuleBuiltData->InitialPosition;
 		Parameters->InitializeParticle_ColorScale			= ModuleBuiltData->ColorRange.GetScale();
 		Parameters->InitializeParticle_ColorBias			= ModuleBuiltData->ColorRange.Min;
 		Parameters->InitializeParticle_SpriteSizeScale		= ModuleBuiltData->SpriteSizeRange.GetScale();
diff --git a/Engine/Plugins/FX/Niagara/Source/Niagara/Internal/Stateless/Modules/NiagaraStatelessModule_ScaleMeshSizeBySpeed.h b/Engine/Plugins/FX/Niagara/Source/Niagara/Internal/Stateless/Modules/NiagaraStatelessModule_ScaleMeshSizeBySpeed.h
--- a/Engine/Plugins/FX/Niagara/Source/Niagara/Internal/Stateless/Modules/NiagaraStatelessModule_ScaleMeshSizeBySpeed.h
+++ b/Engine/Plugins/FX/Niagara/Source/Niagara/Internal/Stateless/Modules/NiagaraStatelessModule_ScaleMeshSizeBySpeed.h
@@ -25,7 +25,7 @@ public:
 	UPROPERTY(EditAnywhere, Category = "Parameters", meta = (ClampMin = "0.01", UIMin = "0.01"))
 	float VelocityThreshold = 1000.0f;
 
-	UPROPERTY(EditAnywhere, Category = "Parameters", meta = (DisplayName = "Scale", DisableRangeDistribution))
+	UPROPERTY(EditAnywhere, Category = "Parameters", meta = (DisplayName = "Scale", DisableRangeDistribution, DisableBindingDistribution))
 	FNiagaraDistributionVector3 ScaleDistribution = FNiagaraDistributionVector3(1.0f);
 
 	virtual void BuildEmitterData(FNiagaraStatelessEmitterDataBuildContext& BuildContext) const override
diff --git a/Engine/Plugins/FX/Niagara/Source/Niagara/Internal/Stateless/Modules/NiagaraStatelessModule_ScaleSpriteSizeBySpeed.h b/Engine/Plugins/FX/Niagara/Source/Niagara/Internal/Stateless/Modules/NiagaraStatelessModule_ScaleSpriteSizeBySpeed.h
--- a/Engine/Plugins/FX/Niagara/Source/Niagara/Internal/Stateless/Modules/NiagaraStatelessModule_ScaleSpriteSizeBySpeed.h
+++ b/Engine/Plugins/FX/Niagara/Source/Niagara/Internal/Stateless/Modules/NiagaraStatelessModule_ScaleSpriteSizeBySpeed.h
@@ -25,7 +25,7 @@ public:
 	UPROPERTY(EditAnywhere, Category = "Parameters", meta = (ClampMin = "0.01", UIMin = "0.01"))
 	float VelocityThreshold = 1000.0f;
 
-	UPROPERTY(EditAnywhere, Category = "Parameters", meta = (DisplayName = "Scale", DisableRangeDistribution))
+	UPROPERTY(EditAnywhere, Category = "Parameters", meta = (DisplayName = "Scale", DisableRangeDistribution, DisableBindingDistribution))
 	FNiagaraDistributionVector2 ScaleDistribution = FNiagaraDistributionVector2(1.0f);
 
 	virtual void BuildEmitterData(FNiagaraStatelessEmitterDataBuildContext& BuildContext) const override
diff --git a/Engine/Plugins/FX/Niagara/Source/Niagara/Internal/Stateless/NiagaraStatelessDistribution.h b/Engine/Plugins/FX/Niagara/Source/Niagara/Internal/Stateless/NiagaraStatelessDistribution.h
--- a/Engine/Plugins/FX/Niagara/Source/Niagara/Internal/Stateless/NiagaraStatelessDistribution.h
+++ b/Engine/Plugins/FX/Niagara/Source/Niagara/Internal/Stateless/NiagaraStatelessDistribution.h
@@ -4,12 +4,14 @@
 
 #include "Curves/RichCurve.h"
 #include "NiagaraStatelessCommon.h"
+#include "NiagaraParameterBinding.h"
 
 #include "NiagaraStatelessDistribution.generated.h"
 
 UENUM()
 enum class ENiagaraDistributionMode
 {
+	Binding,
 	UniformConstant,
 	NonUniformConstant,
 	UniformRange,
@@ -28,6 +30,10 @@ struct FNiagaraDistributionBase
 	UPROPERTY(EditAnywhere, Category = "Parameters")
 	ENiagaraDistributionMode Mode = ENiagaraDistributionMode::UniformConstant;
 
+	UPROPERTY(EditAnywhere, Category = "Parameters")
+	FNiagaraVariableBase ParameterBinding;
+
+	bool IsBinding() const { return Mode == ENiagaraDistributionMode::Binding; }
 	bool IsConstant() const { return Mode == ENiagaraDistributionMode::UniformConstant || Mode == ENiagaraDistributionMode::NonUniformConstant; }
 	bool IsUniform() const { return Mode == ENiagaraDistributionMode::UniformConstant || Mode == ENiagaraDistributionMode::UniformRange; }
 	bool IsCurve() const { return Mode == ENiagaraDistributionMode::UniformCurve || Mode == ENiagaraDistributionMode::NonUniformCurve; }
@@ -42,10 +48,13 @@ struct FNiagaraDistributionBase
 	UPROPERTY(EditAnywhere, Category = "Parameters")
 	int32 MaxLutSampleCount = 64;
 
+	virtual bool AllowBinding() const { return true; }
 	virtual bool AllowCurves() const { return true; }
 	virtual bool DisplayAsColor() const { return false; }
 	virtual void UpdateValuesFromDistribution() { }
 
+	virtual FNiagaraTypeDefinition GetBindingTypeDef() { return FNiagaraTypeDefinition(); }
+
 	static void PostEditChangeProperty(UObject* OwnerObject, FPropertyChangedEvent& PropertyChangedEvent);
 #endif
 };
@@ -156,6 +165,7 @@ struct FNiagaraDistributionFloat : public FNiagaraDistributionBase
 
 #if WITH_EDITORONLY_DATA
 	virtual void UpdateValuesFromDistribution() override;
+	virtual FNiagaraTypeDefinition GetBindingTypeDef() { return FNiagaraTypeDefinition::GetFloatDef(); }
 #endif
 };
 
@@ -177,6 +187,7 @@ struct FNiagaraDistributionVector2 : public FNiagaraDistributionBase
 
 #if WITH_EDITORONLY_DATA
 	virtual void UpdateValuesFromDistribution() override;
+	virtual FNiagaraTypeDefinition GetBindingTypeDef() { return FNiagaraTypeDefinition::GetVec2Def(); }
 #endif
 };
 
@@ -198,6 +209,7 @@ struct FNiagaraDistributionVector3 : public FNiagaraDistributionBase
 
 #if WITH_EDITORONLY_DATA
 	virtual void UpdateValuesFromDistribution() override;
+	virtual FNiagaraTypeDefinition GetBindingTypeDef() { return FNiagaraTypeDefinition::GetVec3Def(); }
 #endif
 };
 
diff --git a/Engine/Plugins/FX/Niagara/Source/Niagara/Internal/Stateless/NiagaraStatelessEmitterDataBuildContext.h b/Engine/Plugins/FX/Niagara/Source/Niagara/Internal/Stateless/NiagaraStatelessEmitterDataBuildContext.h
--- a/Engine/Plugins/FX/Niagara/Source/Niagara/Internal/Stateless/NiagaraStatelessEmitterDataBuildContext.h
+++ b/Engine/Plugins/FX/Niagara/Source/Niagara/Internal/Stateless/NiagaraStatelessEmitterDataBuildContext.h
@@ -11,6 +11,10 @@ struct FNiagaraParameterStore;
 
 class FNiagaraStatelessEmitterDataBuildContext
 {
+	static constexpr uint32 StatelessDistributionFlag_Random	= 0x00000001;
+	static constexpr uint32 StatelessDistributionFlag_Uniform	= 0x00000002;
+	static constexpr uint32 StatelessDistributionFlag_Binding	= 0x00000004;
+
 public:
 	UE_NONCOPYABLE(FNiagaraStatelessEmitterDataBuildContext);
 
@@ -49,6 +53,7 @@ public:
 		return *reinterpret_cast<T*>(TransientObj->GetObject());
 	}
 
+	int32 AddRendererBinding(const FNiagaraVariableBase& Variable);
 	int32 AddRendererBinding(const FNiagaraParameterBinding& Binding);
 	int32 AddRendererBinding(const FNiagaraParameterBindingWithValue& Binding);	
 
@@ -59,11 +64,9 @@ public:
 		FUintVector3 Parameters = FUintVector3::ZeroValue;
 		if (bEnabled && Values.Num() > 0)
 		{
-			constexpr uint32 StatelessDistributionFlag_Random = 0x00000001;
-			constexpr uint32 StatelessDistributionFlag_Uniform = 0x00000002;
-
 			switch (Mode)
 			{
+				case ENiagaraDistributionMode::Binding:				checkNoEntry(); break;
 				case ENiagaraDistributionMode::UniformConstant:		Parameters.X = StatelessDistributionFlag_Random | StatelessDistributionFlag_Uniform; break;
 				case ENiagaraDistributionMode::NonUniformConstant:	Parameters.X = StatelessDistributionFlag_Random; break;
 				case ENiagaraDistributionMode::UniformRange:		Parameters.X = StatelessDistributionFlag_Random | StatelessDistributionFlag_Uniform; break;
@@ -83,7 +86,25 @@ public:
 	template<typename TDistribution>
 	FUintVector3 AddDistribution(const TDistribution& Distribution, bool bEnabled)
 	{
-		return AddDistribution(Distribution.Mode, MakeArrayView(Distribution.Values), bEnabled);
+		FUintVector3 Parameters = FUintVector3::ZeroValue;
+		if ( bEnabled )
+		{
+			if (Distribution.Mode == ENiagaraDistributionMode::Binding)
+			{
+				const int32 ParameterOffset = AddRendererBinding(Distribution.ParameterBinding);
+				if (ParameterOffset >= 0)
+				{
+					Parameters.X = StatelessDistributionFlag_Binding;
+					Parameters.Y = ParameterOffset;
+					reinterpret_cast<float&>(Parameters.Z) = 1.0f;
+				}
+			}
+			else
+			{
+				Parameters = AddDistribution(Distribution.Mode, MakeArrayView(Distribution.Values), bEnabled);
+			}
+		}
+		return Parameters;
 	}
 
 private:
diff --git a/Engine/Plugins/FX/Niagara/Source/Niagara/Internal/Stateless/NiagaraStatelessEmitterInstance.h b/Engine/Plugins/FX/Niagara/Source/Niagara/Internal/Stateless/NiagaraStatelessEmitterInstance.h
--- a/Engine/Plugins/FX/Niagara/Source/Niagara/Internal/Stateless/NiagaraStatelessEmitterInstance.h
+++ b/Engine/Plugins/FX/Niagara/Source/Niagara/Internal/Stateless/NiagaraStatelessEmitterInstance.h
@@ -26,6 +26,9 @@ namespace NiagaraStateless
 		TArray<FNiagaraStatelessRuntimeSpawnInfo>				SpawnInfos;
 		TUniquePtr<NiagaraStateless::FCommonShaderParameters>	ShaderParameters;
 
+		mutable TOptional<TArray<uint8>>						BindingBufferData;
+		mutable FReadBuffer										BindingBuffer;
+
 		// Begin: INiagaraComputeDataBufferInterface
 		virtual bool HasTranslucentDataToRender() const override { return false; }
 		virtual FNiagaraDataBuffer* GetDataToRender(bool bIsLowLatencyTranslucent) const override;
diff --git a/Engine/Plugins/FX/Niagara/Source/Niagara/Private/Stateless/NiagaraStatelessComputeManager.cpp b/Engine/Plugins/FX/Niagara/Source/Niagara/Private/Stateless/NiagaraStatelessComputeManager.cpp
--- a/Engine/Plugins/FX/Niagara/Source/Niagara/Private/Stateless/NiagaraStatelessComputeManager.cpp
+++ b/Engine/Plugins/FX/Niagara/Source/Niagara/Private/Stateless/NiagaraStatelessComputeManager.cpp
@@ -7,6 +7,7 @@
 
 #include "NiagaraGpuComputeDispatchInterface.h"
 #include "NiagaraGPUInstanceCountManager.h"
+#include "NiagaraRenderer.h"
 
 #include "GPUSortManager.h"
 #include "PrimitiveSceneProxy.h"
@@ -171,11 +172,23 @@ void FNiagaraStatelessComputeManager::OnPostPreRender(FRDGBuilder& GraphBuilder)
 			for (auto it = UsedData.CreateIterator(); it; ++it)
 			{
 				FStatelessDataCache* CacheData = it.Value().Get();
+				const NiagaraStateless::FEmitterInstance_RT* EmitterInstance = CacheData->EmitterInstance;
 				const FNiagaraStatelessEmitterData* EmitterData = CacheData->EmitterInstance->EmitterData.Get();
 
+				// Do we need to update the parameter buffer?
+				if (EmitterInstance->BindingBufferData.IsSet())
+				{
+					EmitterInstance->BindingBuffer.Release();
+					EmitterInstance->BindingBuffer.Initialize(RHICmdList, TEXT("FNiagaraStatelessEmitterInstance::BindingBuffer"), sizeof(uint32), EmitterInstance->BindingBufferData->Num() / sizeof(uint32), EPixelFormat::PF_R32_UINT, EBufferUsageFlags::Static);
+					void* LockedBuffer = RHICmdList.LockBuffer(EmitterInstance->BindingBuffer.Buffer, 0, EmitterInstance->BindingBuffer.NumBytes, RLM_WriteOnly);
+					FMemory::Memcpy(LockedBuffer, EmitterInstance->BindingBufferData->GetData(), EmitterInstance->BindingBuffer.NumBytes);
+					RHICmdList.UnlockBuffer(EmitterInstance->BindingBuffer.Buffer);
+					EmitterInstance->BindingBufferData.Reset();
+				}
+
 				// Update parameters for this compute invocation
-				NiagaraStateless::FCommonShaderParameters* ShaderParameters = CacheData->EmitterInstance->ShaderParameters.Get();
-				ShaderParameters->Common_SimulationTime			= CacheData->EmitterInstance->Age;
+				NiagaraStateless::FCommonShaderParameters* ShaderParameters = EmitterInstance->ShaderParameters.Get();
+				ShaderParameters->Common_SimulationTime			= EmitterInstance->Age;
 				ShaderParameters->Common_SimulationDeltaTime	= 1.0f / 60.0f;		//-TODO: Pull from view information, needs reworking of how we link to the dispatch
 				ShaderParameters->Common_SimulationInvDeltaTime	= 60.0f;
 				ShaderParameters->Common_OutputBufferStride		= CacheData->DataBuffer->GetFloatStride() / sizeof(float);
@@ -185,6 +198,7 @@ void FNiagaraStatelessComputeManager::OnPostPreRender(FRDGBuilder& GraphBuilder)
 				ShaderParameters->Common_IntOutputBuffer		= CacheData->DataBuffer->GetGPUBufferInt().UAV.IsValid() ? CacheData->DataBuffer->GetGPUBufferInt().UAV.GetReference() : EmptyIntBufferUAV;
 				ShaderParameters->Common_GPUCountBuffer			= CountBufferUAV;
 				ShaderParameters->Common_StaticFloatBuffer		= EmitterData->StaticFloatBuffer.SRV;
+				ShaderParameters->Common_ParameterBuffer		= FNiagaraRenderer::GetSrvOrDefaultUInt(EmitterInstance->BindingBuffer.SRV);
 
 				// Execute the simulation
 				TShaderRef<NiagaraStateless::FSimulationShader> ComputeShader = EmitterData->GetShader();
diff --git a/Engine/Plugins/FX/Niagara/Source/Niagara/Private/Stateless/NiagaraStatelessEmitterDataBuildContext.cpp b/Engine/Plugins/FX/Niagara/Source/Niagara/Private/Stateless/NiagaraStatelessEmitterDataBuildContext.cpp
--- a/Engine/Plugins/FX/Niagara/Source/Niagara/Private/Stateless/NiagaraStatelessEmitterDataBuildContext.cpp
+++ b/Engine/Plugins/FX/Niagara/Source/Niagara/Private/Stateless/NiagaraStatelessEmitterDataBuildContext.cpp
@@ -41,6 +41,18 @@ uint32 FNiagaraStatelessEmitterDataBuildContext::AddStaticData(TConstArrayView<F
 	return AddStaticData(MakeArrayView(reinterpret_cast<const float*>(FloatData.GetData()), FloatData.Num() * 4));
 }
 
+int32 FNiagaraStatelessEmitterDataBuildContext::AddRendererBinding(const FNiagaraVariableBase& Variable)
+{
+	int32 DataOffset = INDEX_NONE;
+	if (Variable.IsValid())
+	{
+		FNiagaraVariable Var(Variable);
+		RendererBindings.AddParameter(Var, false, false, &DataOffset);
+	}
+
+	return DataOffset;
+}
+
 int32 FNiagaraStatelessEmitterDataBuildContext::AddRendererBinding(const FNiagaraParameterBinding& Binding)
 {
 	int32 DataOffset = INDEX_NONE;
diff --git a/Engine/Plugins/FX/Niagara/Source/Niagara/Private/Stateless/NiagaraStatelessEmitterInstance.cpp b/Engine/Plugins/FX/Niagara/Source/Niagara/Private/Stateless/NiagaraStatelessEmitterInstance.cpp
--- a/Engine/Plugins/FX/Niagara/Source/Niagara/Private/Stateless/NiagaraStatelessEmitterInstance.cpp
+++ b/Engine/Plugins/FX/Niagara/Source/Niagara/Private/Stateless/NiagaraStatelessEmitterInstance.cpp
@@ -205,23 +205,22 @@ void FNiagaraStatelessEmitterInstance::Tick(float DeltaSeconds)
 
 	if (NiagaraStateless::FEmitterInstance_RT* RenderThreadData = RenderThreadDataPtr.Get())
 	{
-		NiagaraStateless::FCommonShaderParameters* NewShaderParameters = nullptr;
+		TOptional<TArray<uint8>> NewBindingBufferData;
 		if ( EmitterData->bModulesHaveRendererBindings && RendererBindings.GetParametersDirty() )
 		{
 			RendererBindings.Tick();
-
-			NewShaderParameters = WeakStatelessEmitter->AllocateShaderParameters(RendererBindings);
-			NewShaderParameters->Common_RandomSeed = RandomSeed;
+			NewBindingBufferData = RendererBindings.GetParameterDataArray();
+			check((NewBindingBufferData->Num() % sizeof(uint32)) == 0);
 		}
 
 		ENQUEUE_RENDER_COMMAND(UpdateStatelessAge)(
-			[RenderThreadData, AgeForRT=Age, ExecutionStateForRT=ExecutionState, NewShaderParameters](FRHICommandListImmediate& RHICmdList)
+			[RenderThreadData, AgeForRT=Age, ExecutionStateForRT=ExecutionState, NewBindingBufferDataRT=MoveTemp(NewBindingBufferData)](FRHICommandListImmediate& RHICmdList)
 			{
 				RenderThreadData->Age = AgeForRT;
 				RenderThreadData->ExecutionState = ExecutionStateForRT;
-				if (NewShaderParameters)
+				if (NewBindingBufferDataRT.IsSet())
 				{
-					RenderThreadData->ShaderParameters.Reset(NewShaderParameters);
+					RenderThreadData->BindingBufferData = NewBindingBufferDataRT;
 				}
 			}
 		);
diff --git a/Engine/Plugins/FX/Niagara/Source/NiagaraEditor/Private/NiagaraEditorStyle.cpp b/Engine/Plugins/FX/Niagara/Source/NiagaraEditor/Private/NiagaraEditorStyle.cpp
--- a/Engine/Plugins/FX/Niagara/Source/NiagaraEditor/Private/NiagaraEditorStyle.cpp
+++ b/Engine/Plugins/FX/Niagara/Source/NiagaraEditor/Private/NiagaraEditorStyle.cpp
@@ -438,6 +438,7 @@ void FNiagaraEditorStyle::InitIcons()
 	SlateBrush->TintColor = FLinearColor(0.01f, 0.01f, 0.01f);
 	Set("NiagaraEditor.Module.InputTypeBorder", SlateBrush);
 
+	Set("NiagaraEditor.DistributionEditor.Binding", new IMAGE_BRUSH_SVG("Icons/DistributionEditor/Binding", Icon20x20));
 	Set("NiagaraEditor.DistributionEditor.UniformConstant", new IMAGE_BRUSH_SVG("Icons/DistributionEditor/UniformConstant", Icon20x20));
 	Set("NiagaraEditor.DistributionEditor.NonUniformConstant", new IMAGE_BRUSH_SVG("Icons/DistributionEditor/NonUniformConstant", Icon20x20));
 	Set("NiagaraEditor.DistributionEditor.ColorConstant", new IMAGE_BRUSH_SVG("Icons/DistributionEditor/ColorConstant", Icon20x20));
diff --git a/Engine/Plugins/FX/Niagara/Source/NiagaraEditor/Private/Stateless/NiagaraDistributionPropertyCustomization.cpp b/Engine/Plugins/FX/Niagara/Source/NiagaraEditor/Private/Stateless/NiagaraDistributionPropertyCustomization.cpp
--- a/Engine/Plugins/FX/Niagara/Source/NiagaraEditor/Private/Stateless/NiagaraDistributionPropertyCustomization.cpp
+++ b/Engine/Plugins/FX/Niagara/Source/NiagaraEditor/Private/Stateless/NiagaraDistributionPropertyCustomization.cpp
@@ -15,6 +15,11 @@
 #include "Widgets/SCompoundWidget.h"
 #include "Widgets/SNiagaraDistributionEditor.h"
 
+#include "NiagaraNodeParameterMapBase.h"
+#include "NiagaraParameterMapHistory.h"
+#include "NiagaraScriptSource.h"
+#include "NiagaraSystem.h"
+
 #define LOCTEXT_NAMESPACE "NiagaraDistributionPropertyCustomization"
 
 class FNiagaraDistributionAdapter : public INiagaraDistributionAdapter
@@ -28,11 +33,13 @@ public:
 	{
 		if (IsValid())
 		{
+			const FName DisableBindingDistributionName("DisableBindingDistribution");
 			const FName DisableCurveDistributionName("DisableCurveDistribution");
 			const FName DisableUniformDistributionName("DisableUniformDistribution");
 			const FName DisableNonUniformDistributionName("DisableNonUniformDistribution");
 			const FName DisableRangeDistributionName("DisableRangeDistribution");
 			const FName DisplayAsColorDistributionName("DisplayAsColorDistribution");
+			bAllowBinding = InDistribution->AllowBinding() && InDistribution->GetBindingTypeDef().IsValid() && (PropertyHandle ? PropertyHandle->HasMetaData(DisableBindingDistributionName) == false : false);
 			bAllowUniform = PropertyHandle ? PropertyHandle->HasMetaData(DisableUniformDistributionName) == false : false;
 			bAllowNonUniform = PropertyHandle ? PropertyHandle->HasMetaData(DisableNonUniformDistributionName) == false : false;
 			bAllowRange = PropertyHandle ? PropertyHandle->HasMetaData(DisableRangeDistributionName) == false : false;
@@ -102,6 +109,10 @@ public:
 
 	virtual void GetSupportedDistributionModes(TArray<ENiagaraDistributionEditorMode>& OutSupportedModes) const override
 	{
+		if (bAllowBinding)
+		{
+			OutSupportedModes.Add(ENiagaraDistributionEditorMode::Binding);
+		}
 		if (bDisplayAsColor)
 		{
 			OutSupportedModes.Add(ENiagaraDistributionEditorMode::ColorConstant);
@@ -176,6 +187,8 @@ public:
 		{
 			switch (InSourceMode)
 			{
+			case ENiagaraDistributionMode::Binding:
+				return ENiagaraDistributionEditorMode::Binding;
 			case ENiagaraDistributionMode::UniformConstant:
 				return ENiagaraDistributionEditorMode::Constant;
 			case ENiagaraDistributionMode::UniformRange:
@@ -190,6 +203,8 @@ public:
 		{
 			switch (InSourceMode)
 			{
+			case ENiagaraDistributionMode::Binding:
+				return ENiagaraDistributionEditorMode::Binding;
 			case ENiagaraDistributionMode::UniformConstant:
 				return ENiagaraDistributionEditorMode::UniformConstant;
 			case ENiagaraDistributionMode::NonUniformConstant:
@@ -212,6 +227,8 @@ public:
 	{
 		switch (InEditorMode)
 		{
+		case ENiagaraDistributionEditorMode::Binding:
+			return ENiagaraDistributionMode::Binding;
 		case ENiagaraDistributionEditorMode::Constant:
 		case ENiagaraDistributionEditorMode::UniformConstant:
 			return ENiagaraDistributionMode::UniformConstant;
@@ -395,6 +412,75 @@ public:
 		PropertyHandle->NotifyPostChange(bContinuousChangeActive ? EPropertyChangeType::Interactive : EPropertyChangeType::ValueSet);
 	}
 
+	virtual FNiagaraVariableBase GetBindingValue() const
+	{
+		UObject* OwnerObject = OwnerObjectWeak.Get();
+		if (OwnerObject == nullptr)
+		{
+			return FNiagaraVariableBase();
+		}
+		return SourceDistribution->ParameterBinding;
+	}
+
+	virtual void SetBindingValue(FNiagaraVariableBase Binding)
+	{
+		UObject* OwnerObject = OwnerObjectWeak.Get();
+		if (OwnerObject == nullptr)
+		{
+			return;
+		}
+
+		FText TransactionText = LOCTEXT("SetBinding", "Set binding value");
+
+		if (bContinuousTransactionPending)
+		{
+			bContinuousTransactionPending = false;
+			ContinuousTransactionIndex = GEditor->BeginTransaction(TransactionText);
+		}
+
+		const FScopedTransaction Transaction(TransactionText, ContinuousTransactionIndex.IsSet() == false);
+		OwnerObject->Modify();
+		PropertyHandle->NotifyPreChange();
+
+		SourceDistribution->ParameterBinding = Binding;
+
+		PropertyHandle->NotifyPostChange(bContinuousChangeActive ? EPropertyChangeType::Interactive : EPropertyChangeType::ValueSet);
+	}
+
+	virtual TArray<FNiagaraVariableBase> GetAvailableBindings() const
+	{
+		TArray<FNiagaraVariableBase> AvailableBindings;
+		UObject* OwnerObject = OwnerObjectWeak.Get();
+		UNiagaraSystem* OwnerSystem = OwnerObject ? OwnerObject->GetTypedOuter<UNiagaraSystem>() : nullptr;
+		const FNiagaraTypeDefinition AllowedTypeDef = SourceDistribution->GetBindingTypeDef();
+		if ( OwnerSystem && AllowedTypeDef.IsValid() )
+		{
+			if (UNiagaraScriptSource* Source = Cast<UNiagaraScriptSource>(OwnerSystem->GetSystemUpdateScript()->GetLatestSource()))
+			{
+				TArray<FNiagaraParameterMapHistory> Histories = UNiagaraNodeParameterMapBase::GetParameterMaps(Source->NodeGraph);
+				for (const FNiagaraParameterMapHistory& History : Histories)
+				{
+					for (const FNiagaraVariable& Variable : History.Variables)
+					{
+						if (Variable.GetType() == AllowedTypeDef && Variable.IsInNameSpace(FNiagaraConstants::SystemNamespaceString))
+						{
+							AvailableBindings.Add(Variable);
+						}
+					}
+				}
+			}
+
+			for (const FNiagaraVariableBase& Variable : OwnerSystem->GetExposedParameters().ReadParameterVariables())
+			{
+				if (Variable.GetType() == AllowedTypeDef)
+				{
+					AvailableBindings.Add(Variable);
+				}
+			}
+		}
+		return AvailableBindings;
+	}
+
 	virtual const FRichCurve* GetCurveValue(int32 ChannelIndex) const override
 	{
 		if (ChannelIndex < EditorNumChannels &&
@@ -656,6 +742,7 @@ private:
 	TOptional<int32> ContinuousTransactionIndex;
 	bool bContinuousChangeActive = false;
 
+	bool bAllowBinding = true;
 	bool bAllowUniform = true;
 	bool bAllowNonUniform = true;
 	bool bAllowRange = true;
diff --git a/Engine/Plugins/FX/Niagara/Source/NiagaraEditor/Private/Widgets/INiagaraDistributionAdapter.h b/Engine/Plugins/FX/Niagara/Source/NiagaraEditor/Private/Widgets/INiagaraDistributionAdapter.h
--- a/Engine/Plugins/FX/Niagara/Source/NiagaraEditor/Private/Widgets/INiagaraDistributionAdapter.h
+++ b/Engine/Plugins/FX/Niagara/Source/NiagaraEditor/Private/Widgets/INiagaraDistributionAdapter.h
@@ -6,13 +6,15 @@
 #include "HAL/Platform.h"
 #include "Internationalization/Text.h"
 #include "Styling/SlateColor.h"
+#include "NiagaraCommon.h"
 
 struct FRichCurve;
 
 /** Defines different edit modes for a niagara distribution */
 enum class ENiagaraDistributionEditorMode
 {
-	/** A single constant value */
+	/* Bound to an script attribute. */
+	Binding,
 	Constant,
 	/** A multi-channel value edited with a single constant value. */
 	UniformConstant,
@@ -87,6 +89,13 @@ public:
 	*/
 	virtual void SetConstantOrRangeValues(int32 ValueIndex, const TArray<float>& InValues) = 0;
 
+	/** Get the name of the variable we are currently bound to. */
+	virtual FNiagaraVariableBase GetBindingValue() const = 0;
+	/** Set the name of the variable we want to bind to. */
+	virtual void SetBindingValue(FNiagaraVariableBase Binding) = 0;
+	/** Get a list of available variables we can bind to. */
+	virtual TArray<FNiagaraVariableBase> GetAvailableBindings() const = 0;
+
 	/** Gets a curve value for the specified channel */
 	virtual const FRichCurve* GetCurveValue(int32 ChannelIndex) const = 0;
 
diff --git a/Engine/Plugins/FX/Niagara/Source/NiagaraEditor/Private/Widgets/NiagaraDistributionEditorUtilities.cpp b/Engine/Plugins/FX/Niagara/Source/NiagaraEditor/Private/Widgets/NiagaraDistributionEditorUtilities.cpp
--- a/Engine/Plugins/FX/Niagara/Source/NiagaraEditor/Private/Widgets/NiagaraDistributionEditorUtilities.cpp
+++ b/Engine/Plugins/FX/Niagara/Source/NiagaraEditor/Private/Widgets/NiagaraDistributionEditorUtilities.cpp
@@ -12,6 +12,8 @@ FText FNiagaraDistributionEditorUtilities::DistributionModeToDisplayName(ENiagar
 {
 	switch (InMode)
 	{
+	case ENiagaraDistributionEditorMode::Binding:
+		return LOCTEXT("BindingDisplayName", "Binding");
 	case ENiagaraDistributionEditorMode::Constant:
 		return LOCTEXT("ConstantDisplayName", "Constant");
 	case ENiagaraDistributionEditorMode::UniformConstant:
@@ -48,6 +50,8 @@ FText FNiagaraDistributionEditorUtilities::DistributionModeToToolTipText(ENiagar
 {
 	switch (InMode)
 	{
+	case ENiagaraDistributionEditorMode::Binding:
+		return LOCTEXT("BindingToolTip", "Value bound to an attribute.");
 	case ENiagaraDistributionEditorMode::Constant:
 		return LOCTEXT("ConstantToolTip", "A constant single value.");
 	case ENiagaraDistributionEditorMode::UniformConstant:
@@ -84,6 +88,8 @@ FName FNiagaraDistributionEditorUtilities::DistributionModeToIconBrushName(ENiag
 {
 	switch (InMode)
 	{
+	case ENiagaraDistributionEditorMode::Binding:
+		return "NiagaraEditor.DistributionEditor.Binding";
 	case ENiagaraDistributionEditorMode::Constant:
 	case ENiagaraDistributionEditorMode::UniformConstant:
 		return "NiagaraEditor.DistributionEditor.UniformConstant";
@@ -123,6 +129,12 @@ FSlateIcon FNiagaraDistributionEditorUtilities::DistributionModeToIcon(ENiagaraD
 	return FSlateIcon(FNiagaraEditorStyle::Get().GetStyleSetName(), DistributionModeToIconBrushName(InMode));
 }
 
+bool FNiagaraDistributionEditorUtilities::IsBinding(ENiagaraDistributionEditorMode InMode)
+{
+	return
+		InMode == ENiagaraDistributionEditorMode::Binding;
+}
+
 bool FNiagaraDistributionEditorUtilities::IsUniform(ENiagaraDistributionEditorMode InMode)
 {
 	return
diff --git a/Engine/Plugins/FX/Niagara/Source/NiagaraEditor/Private/Widgets/NiagaraDistributionEditorUtilities.h b/Engine/Plugins/FX/Niagara/Source/NiagaraEditor/Private/Widgets/NiagaraDistributionEditorUtilities.h
--- a/Engine/Plugins/FX/Niagara/Source/NiagaraEditor/Private/Widgets/NiagaraDistributionEditorUtilities.h
+++ b/Engine/Plugins/FX/Niagara/Source/NiagaraEditor/Private/Widgets/NiagaraDistributionEditorUtilities.h
@@ -19,6 +19,7 @@ namespace FNiagaraDistributionEditorUtilities
 	const FSlateBrush* DistributionModeToIconBrush(ENiagaraDistributionEditorMode InMode);
 	FSlateIcon DistributionModeToIcon(ENiagaraDistributionEditorMode InMode);
 
+	bool IsBinding(ENiagaraDistributionEditorMode InMode);
 	bool IsUniform(ENiagaraDistributionEditorMode InMode);
 	bool IsColor(ENiagaraDistributionEditorMode InMode);
 
diff --git a/Engine/Plugins/FX/Niagara/Source/NiagaraEditor/Private/Widgets/SNiagaraDistributionEditor.cpp b/Engine/Plugins/FX/Niagara/Source/NiagaraEditor/Private/Widgets/SNiagaraDistributionEditor.cpp
--- a/Engine/Plugins/FX/Niagara/Source/NiagaraEditor/Private/Widgets/SNiagaraDistributionEditor.cpp
+++ b/Engine/Plugins/FX/Niagara/Source/NiagaraEditor/Private/Widgets/SNiagaraDistributionEditor.cpp
@@ -15,6 +15,7 @@
 #include "Widgets/SNiagaraColorEditor.h"
 #include "Widgets/SNiagaraDistributionCurveEditor.h"
 #include "Widgets/SNiagaraExpandedToggle.h"
+#include "Widgets/SNiagaraParameterName.h"
 
 #define LOCTEXT_NAMESPACE "NiagaraDistributionEditor"
 
@@ -115,6 +116,68 @@ private:
 	mutable FText ModeToolTipCache;
 };
 
+class SNiagaraDistributionBindingEditor : public SCompoundWidget
+{
+public:
+	SLATE_BEGIN_ARGS(SNiagaraDistributionBindingEditor) { }
+		//SLATE_EVENT(FSimpleDelegate, OnDistributionModeChanged)
+	SLATE_END_ARGS();
+
+	void Construct(const FArguments& InArgs, TSharedRef<INiagaraDistributionAdapter> InDistributionAdapter)
+	{
+		DistributionAdapter = InDistributionAdapter;
+
+		ChildSlot
+		[
+			SNew(SComboButton)
+			.ButtonStyle(FAppStyle::Get(), "HoverHintOnly")
+			.ContentPadding(FMargin(0))
+			.OnGetMenuContent(this, &SNiagaraDistributionBindingEditor::OnGetMenuContent)
+			.ButtonContent()
+			[
+				SNew(SNiagaraParameterName)
+				.ParameterName(this, &SNiagaraDistributionBindingEditor::GetBindingName)
+				.IsReadOnly(true)
+			]
+		];
+	}
+
+private:
+	FName GetBindingName() const
+	{
+		FNiagaraVariableBase Binding = DistributionAdapter->GetBindingValue();
+		return Binding.IsValid() ? Binding.GetName() : NAME_None;
+	}
+
+	void SetBinding(FNiagaraVariableBase Binding)
+	{
+		DistributionAdapter->SetBindingValue(Binding);
+	}
+
+	TSharedRef<SWidget> OnGetMenuContent()
+	{
+		FMenuBuilder MenuBuilder(true, nullptr);
+		for (FNiagaraVariableBase Variable : DistributionAdapter->GetAvailableBindings())
+		{
+			TSharedRef<SWidget> Widget = SNew(SNiagaraParameterName)
+				.ParameterName(Variable.GetName())
+				.IsReadOnly(true);
+
+			MenuBuilder.AddMenuEntry(
+				FUIAction(
+					FExecuteAction::CreateSP(this, &SNiagaraDistributionBindingEditor::SetBinding, Variable)
+				),
+				Widget
+			);
+		}
+
+		return MenuBuilder.MakeWidget();
+	}
+
+private:
+	TSharedPtr<INiagaraDistributionAdapter> DistributionAdapter;
+};
+
 class SNiagaraDistributionValueEditor : public SCompoundWidget
 {
 public:
@@ -587,7 +650,11 @@ void  SNiagaraDistributionEditor::OnDistributionModeChanged()
 TSharedRef<SWidget> SNiagaraDistributionEditor::ConstructContentForMode()
 {
 	ENiagaraDistributionEditorMode Mode = DistributionAdapter->GetDistributionMode();
-	if (FNiagaraDistributionEditorUtilities::IsConstant(Mode) || FNiagaraDistributionEditorUtilities::IsRange(Mode))
+	if (FNiagaraDistributionEditorUtilities::IsBinding(Mode))
+	{
+		return SNew(SNiagaraDistributionBindingEditor, DistributionAdapter.ToSharedRef());
+	}
+	else if (FNiagaraDistributionEditorUtilities::IsConstant(Mode) || FNiagaraDistributionEditorUtilities::IsRange(Mode))
 	{
 		return SNew(SNiagaraDistributionValueEditor, DistributionAdapter.ToSharedRef());
 	}
diff --git a/Engine/Plugins/FX/Niagara/Source/NiagaraShader/Internal/Stateless/NiagaraStatelessModuleShaderParameters.h b/Engine/Plugins/FX/Niagara/Source/NiagaraShader/Internal/Stateless/NiagaraStatelessModuleShaderParameters.h
--- a/Engine/Plugins/FX/Niagara/Source/NiagaraShader/Internal/Stateless/NiagaraStatelessModuleShaderParameters.h
+++ b/Engine/Plugins/FX/Niagara/Source/NiagaraShader/Internal/Stateless/NiagaraStatelessModuleShaderParameters.h
@@ -8,8 +8,8 @@
 namespace NiagaraStateless
 {
 	BEGIN_SHADER_PARAMETER_STRUCT(FInitializeParticleModule_ShaderParameters, NIAGARASHADER_API)
-		SHADER_PARAMETER(uint32,	InitializeParticle_ModuleFlags)
-		SHADER_PARAMETER(FVector3f,	InitializeParticle_Position)				// Unset / Direct Set / Simulation Position
+		SHADER_PARAMETER(uint32,		InitializeParticle_ModuleFlags)
+		SHADER_PARAMETER(FUintVector3,	InitializeParticle_InitialPosition)
 
 		SHADER_PARAMETER(FVector4f,	InitializeParticle_ColorScale)				// Unset / Direct Set / Random Range(Link RGBA / Link RGB|Link A / Random Channels) / Random Hue|Saturation|Value
 		SHADER_PARAMETER(FVector4f,	InitializeParticle_ColorBias)
diff --git a/Engine/Plugins/FX/Niagara/Source/NiagaraShader/Internal/Stateless/NiagaraStatelessSimulationShader.h b/Engine/Plugins/FX/Niagara/Source/NiagaraShader/Internal/Stateless/NiagaraStatelessSimulationShader.h
--- a/Engine/Plugins/FX/Niagara/Source/NiagaraShader/Internal/Stateless/NiagaraStatelessSimulationShader.h
+++ b/Engine/Plugins/FX/Niagara/Source/NiagaraShader/Internal/Stateless/NiagaraStatelessSimulationShader.h
@@ -33,6 +33,7 @@ namespace NiagaraStateless
 		SHADER_PARAMETER_UAV(RWBuffer<int>,		Common_IntOutputBuffer)
 		SHADER_PARAMETER_UAV(RWBuffer<int>,		Common_GPUCountBuffer)
 		SHADER_PARAMETER_SRV(Buffer<float>,		Common_StaticFloatBuffer)
+		SHADER_PARAMETER_SRV(Buffer<uint32>,	Common_ParameterBuffer)
 
 		SHADER_PARAMETER_STRUCT_INCLUDE(FSpawnInfoShaderParameters,	SpawnParameters)
 	END_SHADER_PARAMETER_STRUCT()
