Fix for disabling an override from an NPC Instance not being reflected immediately

diff --git a/Engine/Plugins/FX/Niagara/Source/NiagaraEditor/Private/ViewModels/NiagaraParameterCollectionAssetViewModel.cpp b/Engine/Plugins/FX/Niagara/Source/NiagaraEditor/Private/ViewModels/NiagaraParameterCollectionAssetViewModel.cpp
--- a/Engine/Plugins/FX/Niagara/Source/NiagaraEditor/Private/ViewModels/NiagaraParameterCollectionAssetViewModel.cpp
+++ b/Engine/Plugins/FX/Niagara/Source/NiagaraEditor/Private/ViewModels/NiagaraParameterCollectionAssetViewModel.cpp
@@ -423,6 +423,17 @@ void FNiagaraParameterCollectionAssetViewModel::OnParameterProvidedChanged(FNiag
 
 	RefreshParameterViewModels();
 	//CollectionChanged(false);
+
+	// Restore the value from the Collection if we are no longer overriding it.  Note that this will override the value that is in the instance
+	// so we don't currently have a way to store an unused override in the instance...something that would be nice to have.
+	if (Instance && Collection && !Instance->OverridesParameter(ParameterVariable))
+	{
+		if (const uint8* CollectionParameterValue = Collection->GetDefaultInstance()->GetParameterStore().GetParameterData(ParameterVariable))
+		{
+			Instance->GetParameterStore().SetParameterData(CollectionParameterValue, ParameterVariable);
+			Instance->GetParameterStore().Tick();
+		}
+	}
 }
 
 void FNiagaraParameterCollectionAssetViewModel::OnParameterValueChangedInternal(TSharedRef<FNiagaraCollectionParameterViewModel> ChangedParameter)
