Clean up how we collect the digested DI CDO for the new async task compilation path

diff --git a/Engine/Plugins/FX/Niagara/Source/NiagaraEditor/Private/NiagaraCompilationCopy.cpp b/Engine/Plugins/FX/Niagara/Source/NiagaraEditor/Private/NiagaraCompilationCopy.cpp
--- a/Engine/Plugins/FX/Niagara/Source/NiagaraEditor/Private/NiagaraCompilationCopy.cpp
+++ b/Engine/Plugins/FX/Niagara/Source/NiagaraEditor/Private/NiagaraCompilationCopy.cpp
@@ -505,7 +505,7 @@ void FNiagaraCompilationCopyData::InstantiateCompilationCopy(const FNiagaraCompi
 
 	if (InstantiatedGraph)
 	{
-		AggregatedDataInterfaceCDODuplicates.Append(ObjectPtrDecay(InstantiatedGraph->AggregateDataInterfaceCDODuplicates));
+		SourceGraph.CollectReferencedDataInterfaceCDO(AggregatedDataInterfaceCDODuplicates);
 	}
 }
 
diff --git a/Engine/Plugins/FX/Niagara/Source/NiagaraEditor/Private/NiagaraCompilationPrivate.h b/Engine/Plugins/FX/Niagara/Source/NiagaraEditor/Private/NiagaraCompilationPrivate.h
--- a/Engine/Plugins/FX/Niagara/Source/NiagaraEditor/Private/NiagaraCompilationPrivate.h
+++ b/Engine/Plugins/FX/Niagara/Source/NiagaraEditor/Private/NiagaraCompilationPrivate.h
@@ -195,7 +195,7 @@ public:
 
 	TArray<FSharedCompilationCopy> EmitterData;
 
-	TMap<UClass*, UNiagaraDataInterface*> AggregatedDataInterfaceCDODuplicates;
+	TMap<TObjectPtr<UClass>, TObjectPtr<UNiagaraDataInterface>> AggregatedDataInterfaceCDODuplicates;
 };
 
 //////////////////////////////////////////////////////////////////////////
diff --git a/Engine/Plugins/FX/Niagara/Source/NiagaraEditor/Private/NiagaraCompilationTasks.cpp b/Engine/Plugins/FX/Niagara/Source/NiagaraEditor/Private/NiagaraCompilationTasks.cpp
--- a/Engine/Plugins/FX/Niagara/Source/NiagaraEditor/Private/NiagaraCompilationTasks.cpp
+++ b/Engine/Plugins/FX/Niagara/Source/NiagaraEditor/Private/NiagaraCompilationTasks.cpp
@@ -549,20 +549,6 @@ void FNiagaraSystemCompilationTask::FCompileGroupInfo::InstantiateCompileGraph(c
 	TArray<TSharedPtr<FNiagaraCompilationCopyData, ESPMode::ThreadSafe>> DependentRequests;
 	FCompileConstantResolver EmptyResolver;
 
-	TArray<UClass*> DataInterfaceClasses;
-
-	auto CollectDataInterfaceClasses = [&](TConstArrayView<FNiagaraVariable> Variables)
-	{
-		// Collect classes for external encounterable variables
-		for (const FNiagaraVariable& EncounterableVariable : Variables)
-		{
-			if (EncounterableVariable.IsDataInterface())
-			{
-				DataInterfaceClasses.AddUnique(EncounterableVariable.GetType().GetClass());
-			}
-		}
-	};
-
 	BasePtr.ValidUsages = ValidUsages;
 
 	// First deep copy all the emitter graphs referenced by the system so that we can later hook up emitter handles in the system traversal.
@@ -610,8 +596,6 @@ void FNiagaraSystemCompilationTask::FCompileGroupInfo::InstantiateCompileGraph(c
 				BasePtr.AggregatedDataInterfaceCDODuplicates.Append(EmitterData->AggregatedDataInterfaceCDODuplicates);
 			}
 		}
-
-		CollectDataInterfaceClasses(EncounterableSystemVariables);
 	}
 
 	// Now we can finish off the emitters.
diff --git a/Engine/Plugins/FX/Niagara/Source/NiagaraEditor/Private/NiagaraGraphDigest.cpp b/Engine/Plugins/FX/Niagara/Source/NiagaraEditor/Private/NiagaraGraphDigest.cpp
--- a/Engine/Plugins/FX/Niagara/Source/NiagaraEditor/Private/NiagaraGraphDigest.cpp
+++ b/Engine/Plugins/FX/Niagara/Source/NiagaraEditor/Private/NiagaraGraphDigest.cpp
@@ -44,13 +44,15 @@ NIAGARA_GRAPH_DIGEST_NODE_TYPE_LIST;
 
 struct FNiagaraCompilationGraphCreateContext
 {
-	FNiagaraCompilationGraphCreateContext(FNiagaraCompilationGraphDigested& InParentGraph, const FNiagaraGraphChangeIdBuilder& InChangeIdBuilder)
+	FNiagaraCompilationGraphCreateContext(FNiagaraCompilationGraphDigested& InParentGraph, TArray<const FNiagaraCompilationGraphDigested*>& InDigestedChildGraphs, const FNiagaraGraphChangeIdBuilder& InChangeIdBuilder)
 		: ParentGraph(InParentGraph)
+		, DigestedChildGraphs(InDigestedChildGraphs)
 		, ChangeIdBuilder(InChangeIdBuilder)
 	{
 	}
 
 	FNiagaraCompilationGraphDigested& ParentGraph;
+	TArray<const FNiagaraCompilationGraphDigested*>& DigestedChildGraphs;
 	const FNiagaraGraphChangeIdBuilder& ChangeIdBuilder;
 };
 
@@ -546,7 +548,7 @@ void FNiagaraCompilationGraphDigested::Digest(const UNiagaraGraph* InGraph, cons
 
 	TMap<const UEdGraphNode*, int32> NodeIndexMap;
 
-	FNiagaraCompilationGraphCreateContext NodeContext(*this, ChangeIdBuilder);
+	FNiagaraCompilationGraphCreateContext NodeContext(*this, ChildGraphs, ChangeIdBuilder);
 
 	for (const UEdGraphNode* SourceNode : InGraph->Nodes)
 	{
@@ -724,6 +726,20 @@ void FNiagaraCompilationGraphDigested::RegisterObjectAsset(FName VariableName, U
 	}
 }
 
+void FNiagaraCompilationGraphDigested::CollectReferencedDataInterfaceCDO(FDataInterfaceCDOMap& Interfaces) const
+{
+	// add our collected DIs then iterate over all the child graphs
+	for (FDataInterfaceCDOMap::TConstIterator It = CachedDataInterfaceCDODuplicates.CreateConstIterator(); It; ++It)
+	{
+		Interfaces.FindOrAdd(It.Key(), It.Value());
+	}
+
+	for (const FNiagaraCompilationGraphDigested* ChildGraph : ChildGraphs)
+	{
+		ChildGraph->CollectReferencedDataInterfaceCDO(Interfaces);
+	}
+}
+
 void FNiagaraCompilationGraph::FindOutputNodes(TArray<const FNiagaraCompilationNodeOutput*>& OutputNodes) const
 {
 	OutputNodes.Reserve(OutputNodeIndices.Num());
@@ -1025,9 +1041,6 @@ TSharedPtr<FNiagaraCompilationGraphInstanced, ESPMode::ThreadSafe> FNiagaraCompi
 		SubGraph->StaticSwitchInputs.AddUnique(SwitchInput);
 	}
 
-	// make sure that we've collected the CDO for the registered DI as well
-	SubGraph->AggregateDataInterfaceCDODuplicates.Append(CachedDataInterfaceCDODuplicates);
-
 	return SubGraph;
 }
 
@@ -1282,8 +1295,6 @@ TSharedPtr<FNiagaraCompilationGraphInstanced, ESPMode::ThreadSafe> FNiagaraCompi
 
 				CalledInstantiatedGraph->Refine(InstantiationContext, FunctionToInstantiate);
 
-				InstantiatedGraph->AggregateDataInterfaces(CalledInstantiatedGraph.Get());
-
 				++TotalGraphCount;
 				TotalNodeCount += CalledInstantiatedGraph->Nodes.Num();
 				TotalCulledNodeCount += CalledDigestedGraph->Nodes.Num() - CalledInstantiatedGraph->Nodes.Num();
@@ -1415,13 +1426,6 @@ void FNiagaraCompilationGraph::NodeTraversal(
 	}
 }
 
-
-void FNiagaraCompilationGraphInstanced::AggregateDataInterfaces(const FNiagaraCompilationGraphInstanced* ChildGraph)
-{
-	//AggregateDataInterfaceCDODuplicates.Append(ObjectPtrDecay(ChildGraph->AggregateDataInterfaceCDODuplicates));
-	AggregateDataInterfaceCDODuplicates.Append(ChildGraph->AggregateDataInterfaceCDODuplicates);
-}
-
 void FNiagaraCompilationGraphInstanced::AggregateChildGraph(const FNiagaraCompilationGraphInstanced* ChildGraph)
 {
 	if (ChildGraph->bContainsStaticVariables)
@@ -2013,6 +2017,10 @@ FNiagaraCompilationNodeEmitter::FNiagaraCompilationNodeEmitter(const UNiagaraNod
 	if (const UNiagaraGraph* DependentGraph = InNode->GetCalledGraph())
 	{
 		CalledGraph = FNiagaraDigestDatabase::Get().CreateGraphDigest(DependentGraph, Context.ChangeIdBuilder);
+		if (const FNiagaraCompilationGraphDigested* DigestedGraph = CalledGraph->AsDigested())
+		{
+			Context.DigestedChildGraphs.AddUnique(DigestedGraph);
+		}
 	}
 
 	Usage = InNode->GetUsage();
@@ -2256,6 +2264,10 @@ FNiagaraCompilationNodeFunctionCall::FNiagaraCompilationNodeFunctionCall(const U
 	if (const UNiagaraGraph* DependentGraph = InNode->GetCalledGraph())
 	{
 		CalledGraph = FNiagaraDigestDatabase::Get().CreateGraphDigest(DependentGraph, Context.ChangeIdBuilder);
+		if (const FNiagaraCompilationGraphDigested* DigestedGraph = CalledGraph->AsDigested())
+		{
+			Context.DigestedChildGraphs.AddUnique(DigestedGraph);
+		}
 	}
 
 	// Top level functions (functions invoked in the root graph) will use the serialized DebugState value, all others will
diff --git a/Engine/Plugins/FX/Niagara/Source/NiagaraEditor/Private/NiagaraGraphDigest.h b/Engine/Plugins/FX/Niagara/Source/NiagaraEditor/Private/NiagaraGraphDigest.h
--- a/Engine/Plugins/FX/Niagara/Source/NiagaraEditor/Private/NiagaraGraphDigest.h
+++ b/Engine/Plugins/FX/Niagara/Source/NiagaraEditor/Private/NiagaraGraphDigest.h
@@ -332,9 +332,12 @@ public:
 
 	virtual FNiagaraCompilationGraphDigested* AsDigested() override { return this; }
 
+	using FDataInterfaceCDOMap = TMap<TObjectPtr<UClass>, TObjectPtr<UNiagaraDataInterface>>;
+	void CollectReferencedDataInterfaceCDO(FDataInterfaceCDOMap& Interfaces) const;
+
 	using FDataInterfaceDuplicateMap = TMap<TObjectKey<UNiagaraDataInterface>, TObjectPtr<UNiagaraDataInterface>>;
 	FDataInterfaceDuplicateMap CachedDataInterfaceDuplicates;
-	TMap<TObjectPtr<UClass>, TObjectPtr<UNiagaraDataInterface>> CachedDataInterfaceCDODuplicates;
+	FDataInterfaceCDOMap CachedDataInterfaceCDODuplicates;
 	TMap<FName, TObjectPtr<UObject>> CachedNamedObjectAssets;
 
 	TWeakObjectPtr<const UNiagaraGraph> SourceGraph;
@@ -345,6 +348,8 @@ protected:
 		const FNiagaraCompilationCopyData* CopyCompilationData,
 		const FNiagaraCompilationBranchMap& Branches,
 		TArray<FNiagaraCompilationNodeFunctionCall*>& PendingInstantiations) const;
+
+	TArray<const FNiagaraCompilationGraphDigested*> ChildGraphs;
 };
 
 // Instanced version of the compilation graph.  Created from a FNiagaraCompilationGraphDigested
@@ -353,7 +358,6 @@ protected:
 class FNiagaraCompilationGraphInstanced : public FNiagaraCompilationGraph
 {
 public:
-	void AggregateDataInterfaces(const FNiagaraCompilationGraphInstanced* ChildGraph);
 	void AggregateChildGraph(const FNiagaraCompilationGraphInstanced* ChildGraph);
 
 	void ResolveNumerics(FNiagaraCompilationGraphInstanceContext& Context);
@@ -363,10 +367,6 @@ public:
 
 	TSharedPtr<const FNiagaraCompilationGraph> InstantiationSourceGraph;
 
-	// the full collection of duplicated CDO data interfaces that are used by this graph (and all
-	// instanced graphs that it references).
-	TMap<TObjectPtr<UClass>, TObjectPtr<UNiagaraDataInterface>> AggregateDataInterfaceCDODuplicates;
-
 protected:
 	void ValidateRefinement() const;
 	void PatchGenericNumericsFromCaller(FNiagaraCompilationGraphInstanceContext& Context);
