Fix for new compilation mode not passing NPC references down through emitter nodes

diff --git a/Engine/Plugins/FX/Niagara/Source/NiagaraEditor/Private/NiagaraDigestDatabase.h b/Engine/Plugins/FX/Niagara/Source/NiagaraEditor/Private/NiagaraDigestDatabase.h
--- a/Engine/Plugins/FX/Niagara/Source/NiagaraEditor/Private/NiagaraDigestDatabase.h
+++ b/Engine/Plugins/FX/Niagara/Source/NiagaraEditor/Private/NiagaraDigestDatabase.h
@@ -88,6 +88,7 @@ protected:
 class FNiagaraDigestedParameterCollections
 {
 public:
+	TConstArrayView<FNiagaraCompilationNPCHandle> ReadCollections() const { return Collections; }
 	TArray<FNiagaraCompilationNPCHandle>& EditCollections() { return Collections; };
 
 	FNiagaraCompilationNPCHandle FindMatchingCollection(FName VariableName, bool bAllowPartialMatch, FNiagaraVariable& OutVar) const;
diff --git a/Engine/Plugins/FX/Niagara/Source/NiagaraEditor/Private/NiagaraGraphDigest.cpp b/Engine/Plugins/FX/Niagara/Source/NiagaraEditor/Private/NiagaraGraphDigest.cpp
--- a/Engine/Plugins/FX/Niagara/Source/NiagaraEditor/Private/NiagaraGraphDigest.cpp
+++ b/Engine/Plugins/FX/Niagara/Source/NiagaraEditor/Private/NiagaraGraphDigest.cpp
@@ -2110,6 +2110,7 @@ void FNiagaraCompilationNodeEmitter::BuildParameterMapHistory(FParameterMapHisto
 			TArray<FNiagaraVariable> LocalStaticVars;
 			FNiagaraParameterUtilities::FilterToRelevantStaticVariables(Builder.StaticVariables, LocalStaticVars, *EmitterUniqueName, TEXT("Emitter"), true);
 			ChildBuilder.RegisterExternalStaticVariables(LocalStaticVars);
+			ChildBuilder.AvailableCollections->EditCollections() = Builder.AvailableCollections->ReadCollections();
 
 			FString LocalEmitterName = TEXT("Emitter");
 			ChildBuilder.EnterEmitter(LocalEmitterName, CalledGraph.Get(), this);
