Fix for new compilation mode stomping GPU shadermaps by just evaluating their feature level and not the shader platform

diff --git a/Engine/Plugins/FX/Niagara/Source/Niagara/Private/NiagaraScript.cpp b/Engine/Plugins/FX/Niagara/Source/Niagara/Private/NiagaraScript.cpp
--- a/Engine/Plugins/FX/Niagara/Source/Niagara/Private/NiagaraScript.cpp
+++ b/Engine/Plugins/FX/Niagara/Source/Niagara/Private/NiagaraScript.cpp
@@ -3168,15 +3168,15 @@ void UNiagaraScript::SetComputeCompilationResults(
 
 	// for now we will use the CachedScriptResourcesForCooking for storing all compilation results and will update ScriptResource
 	// if it matches with the supplied parameters
-	auto ScriptMatchesFeatureLevel = [FeatureLevel](const TUniquePtr<FNiagaraShaderScript>& ShaderScript) -> bool
+	auto ScriptMatches = [FeatureLevel, ShaderPlatform](const TUniquePtr<FNiagaraShaderScript>& ShaderScript) -> bool
 	{
-		return ShaderScript->GetFeatureLevel() == FeatureLevel;
+		return ShaderScript->GetFeatureLevel() == FeatureLevel && ShaderScript->GetShaderPlatform() == ShaderPlatform;
 	};
 
 	TArray<TUniquePtr<FNiagaraShaderScript>>& CachedScripts = CachedScriptResourcesForCooking.FindOrAdd(TargetPlatform);
 	FNiagaraShaderScript* TargetShaderScript = nullptr;
 
-	if (TUniquePtr<FNiagaraShaderScript>* ExistingScript = CachedScripts.FindByPredicate(ScriptMatchesFeatureLevel))
+	if (TUniquePtr<FNiagaraShaderScript>* ExistingScript = CachedScripts.FindByPredicate(ScriptMatches))
 	{
 		TargetShaderScript = ExistingScript->Get();
 	}
