Unshelved from pending changelist '31527134':

fix issues with bodies entering and leaving domain

diff --git a/Engine/Plugins/FX/Niagara/Source/Niagara/Private/NiagaraDataInterfaceRigidMeshCollisionQuery.cpp b/Engine/Plugins/FX/Niagara/Source/Niagara/Private/NiagaraDataInterfaceRigidMeshCollisionQuery.cpp
--- a/Engine/Plugins/FX/Niagara/Source/Niagara/Private/NiagaraDataInterfaceRigidMeshCollisionQuery.cpp
+++ b/Engine/Plugins/FX/Niagara/Source/Niagara/Private/NiagaraDataInterfaceRigidMeshCollisionQuery.cpp
@@ -510,52 +510,56 @@ template<bool FirstStage>
 void RemapPreviousTransforms(
 	TConstArrayView<FNDIRigidMeshCollisionData::FComponentBodyCount> PreviousCounts,
 	TConstArrayView<FNDIRigidMeshCollisionData::FComponentBodyCount> CurrentCounts,
+	const FNDIRigidMeshCollisionElementOffset& PreviousOffsets,
+	const FNDIRigidMeshCollisionElementOffset& CurrentOffsets,
 	FNDIRigidMeshCollisionArrays* OutAssetArrays)
 {
-	const FNDIRigidMeshCollisionElementOffset& CurrentOffsets = OutAssetArrays->ElementOffsets;
-
-	uint32 BoxCount = 0;
-	uint32 SphereCount = 0;
-	uint32 CapsuleCount = 0;
-
-	auto CopyComponent = [&BoxCount, &SphereCount, &CapsuleCount]
-		(const TArray<FVector4f>& Src, TArray<FVector4f>& Dst, const FNDIRigidMeshCollisionElementOffset& Offsets, const FNDIRigidMeshCollisionData::FComponentBodyCount& BodyCounts) -> void
-	{
-		const uint32 BoxStartIndex = 3 * (Offsets.BoxOffset + BoxCount);
+	auto CopyComponent = []
+		(const TArray<FVector4f>& Src, TArray<FVector4f>& Dst, 
+			const uint32 SrcBoxCount, const uint32 SrcSphereCount, const uint32 SrcCapsuleCount,
+			const uint32 DstBoxCount, const uint32 DstSphereCount, const uint32 DstCapsuleCount,
+			const FNDIRigidMeshCollisionElementOffset& SrcOffsets, const FNDIRigidMeshCollisionElementOffset& DstOffsets,
+			const FNDIRigidMeshCollisionData::FComponentBodyCount& BodyCounts) -> void
+	{
+		const uint32 SrcBoxStartIndex = 3 * (SrcOffsets.BoxOffset + SrcBoxCount);
+		const uint32 DstBoxStartIndex = 3 * (DstOffsets.BoxOffset + DstBoxCount);
 		for (uint32 ElementIt = 0; ElementIt < 3 * BodyCounts.BoxCount; ++ElementIt)
 		{
-			Dst[BoxStartIndex + ElementIt] = Src[BoxStartIndex + ElementIt];
+			Dst[DstBoxStartIndex + ElementIt] = Src[SrcBoxStartIndex + ElementIt];
 		}
 
-		const uint32 SphereStartIndex = 3 * (Offsets.SphereOffset + SphereCount);
+		const uint32 SrcSphereStartIndex = 3 * (SrcOffsets.SphereOffset + SrcSphereCount);
+		const uint32 DstSphereStartIndex = 3 * (DstOffsets.SphereOffset + DstSphereCount);
 		for (uint32 ElementIt = 0; ElementIt < 3 * BodyCounts.SphereCount; ++ElementIt)
 		{
-			Dst[SphereStartIndex + ElementIt] = Src[SphereStartIndex + ElementIt];
+			Dst[DstSphereStartIndex + ElementIt] = Src[SrcSphereStartIndex + ElementIt];
 		}
 
-		const uint32 CapsuleStartIndex = 3 * (Offsets.CapsuleOffset + CapsuleCount);
+		const uint32 SrcCapsuleStartIndex = 3 * (SrcOffsets.CapsuleOffset + SrcCapsuleCount);
+		const uint32 DstCapsuleStartIndex = 3 * (DstOffsets.CapsuleOffset + DstCapsuleCount);
 		for (uint32 ElementIt = 0; ElementIt < 3 * BodyCounts.CapsuleCount; ++ElementIt)
 		{
-			Dst[CapsuleStartIndex + ElementIt] = Src[CapsuleStartIndex + ElementIt];
+			Dst[DstCapsuleStartIndex + ElementIt] = Src[SrcCapsuleStartIndex + ElementIt];
 		}
 	};
 
-	int32 PreviousIndex = 0;
+	uint32 CurrentBoxCount = 0;
+	uint32 CurrentSphereCount = 0;
+	uint32 CurrentCapsuleCount = 0;
+
 	for (const FNDIRigidMeshCollisionData::FComponentBodyCount& CurrentCount : CurrentCounts)
 	{
 		bool bPreviousValuesCopied = false;
 
+		uint32 PreviousBoxCount = 0;
+		uint32 PreviousSphereCount = 0;
+		uint32 PreviousCapsuleCount = 0;
+
+		int32 PreviousIndex = 0;
 		while (PreviousCounts.IsValidIndex(PreviousIndex) && !bPreviousValuesCopied)
 		{
 			const FNDIRigidMeshCollisionData::FComponentBodyCount& PreviousCount = PreviousCounts[PreviousIndex];
 
-			// skip through the previous components until we find a match
-			if (PreviousCount.ComponentHash < CurrentCount.ComponentHash)
-			{
-				++PreviousIndex;
-				continue;
-			}
-
 			if (PreviousCount.ComponentHash == CurrentCount.ComponentHash)
 			{
 				check(PreviousCount.BoxCount == CurrentCount.BoxCount);
@@ -564,26 +568,56 @@ void RemapPreviousTransforms(
 
 				bPreviousValuesCopied = true;
 
+				// For the body in the current tick, we found a match in the previous tick transforms.  Copy the previous tick transforms (stored in CurrentTransform)
+				// over to PreviousTransform.  Note the indices for the transforms might differ between current and previous as bodies enter/leave the arrays
 				if (FirstStage)
 				{
-					CopyComponent(OutAssetArrays->CurrentTransform, OutAssetArrays->PreviousTransform, CurrentOffsets, CurrentCount);
-					CopyComponent(OutAssetArrays->CurrentInverse, OutAssetArrays->PreviousInverse, CurrentOffsets, CurrentCount);
+					const TArray<FVector4f>& Src = OutAssetArrays->CurrentTransform;
+					TArray<FVector4f>& Dst = OutAssetArrays->PreviousTransform;
+					const TArray<FVector4f>& SrcInverse = OutAssetArrays->CurrentInverse;
+					TArray<FVector4f>& DstInverse = OutAssetArrays->PreviousInverse;
+
+					CopyComponent(Src, Dst,
+						PreviousBoxCount, PreviousSphereCount, PreviousCapsuleCount,
+						CurrentBoxCount, CurrentSphereCount, CurrentCapsuleCount,
+						PreviousOffsets, CurrentOffsets,
+						CurrentCount);
+
+					CopyComponent(SrcInverse, DstInverse,
+						PreviousBoxCount, PreviousSphereCount, PreviousCapsuleCount,
+						CurrentBoxCount, CurrentSphereCount, CurrentCapsuleCount,
+						PreviousOffsets, CurrentOffsets,
+						CurrentCount);
 				}
+
+				break;
 			}
 
-			break;
+			++PreviousIndex;
+			PreviousBoxCount += PreviousCount.BoxCount;
+			PreviousSphereCount += PreviousCount.SphereCount;
+			PreviousCapsuleCount += PreviousCount.CapsuleCount;
 		}
 
-		// copy the defaults over to the previous
+		// body is found in current tick but not in previous tick, so copy current transforms to previous buffer using indices and offsets from the current tick
 		if (!bPreviousValuesCopied && !FirstStage)
 		{
-			CopyComponent(OutAssetArrays->CurrentTransform, OutAssetArrays->PreviousTransform, CurrentOffsets, CurrentCount);
-			CopyComponent(OutAssetArrays->CurrentInverse, OutAssetArrays->PreviousInverse, CurrentOffsets, CurrentCount);
+			CopyComponent(OutAssetArrays->CurrentTransform, OutAssetArrays->PreviousTransform,
+				CurrentBoxCount, CurrentSphereCount, CurrentCapsuleCount,
+				CurrentBoxCount, CurrentSphereCount, CurrentCapsuleCount,
+				CurrentOffsets, CurrentOffsets,
+				CurrentCount);
+
+			CopyComponent(OutAssetArrays->CurrentInverse, OutAssetArrays->PreviousInverse,
+				CurrentBoxCount, CurrentSphereCount, CurrentCapsuleCount,
+				CurrentBoxCount, CurrentSphereCount, CurrentCapsuleCount,
+				CurrentOffsets, CurrentOffsets,
+				CurrentCount);
 		}
 
-		BoxCount += CurrentCount.BoxCount;
-		SphereCount += CurrentCount.SphereCount;
-		CapsuleCount += CurrentCount.CapsuleCount;
+		CurrentBoxCount += CurrentCount.BoxCount;
+		CurrentSphereCount += CurrentCount.SphereCount;
+		CurrentCapsuleCount += CurrentCount.CapsuleCount;
 	}
 }
 
@@ -635,6 +669,8 @@ bool UpdateInternalArrays(
 		return false;
 	}
 
+	FNDIRigidMeshCollisionElementOffset PreviousElementOffsets = OutAssetArrays->ElementOffsets;
+
 	OutAssetArrays->ElementOffsets.BoxOffset = 0;
 	OutAssetArrays->ElementOffsets.SphereOffset = OutAssetArrays->ElementOffsets.BoxOffset + TotalBoxCount;
 	OutAssetArrays->ElementOffsets.CapsuleOffset = OutAssetArrays->ElementOffsets.SphereOffset + TotalSphereCount;
@@ -646,12 +682,12 @@ bool UpdateInternalArrays(
 
 	// where possible PreviousTransform & PreviousInverse should be pulled from the current values of CurrentTransform &
 	// CurrentInverse based on the remapped entries
-	RemapPreviousTransforms<true>(BodyCounts, CurrentBodyCounts, OutAssetArrays);
+	RemapPreviousTransforms<true>(BodyCounts, CurrentBodyCounts, PreviousElementOffsets, OutAssetArrays->ElementOffsets, OutAssetArrays);
 
 	UpdateAssetArrays<UStaticMeshComponent, true>(StaticMeshView, LWCTile, OutAssetArrays, BoxIndex, SphereIndex, CapsuleIndex);
 	UpdateAssetArrays<USkeletalMeshComponent, true>(SkeletalMeshView, LWCTile, OutAssetArrays, BoxIndex, SphereIndex, CapsuleIndex);
 
-	RemapPreviousTransforms<false>(BodyCounts, CurrentBodyCounts, OutAssetArrays);
+	RemapPreviousTransforms<false>(BodyCounts, CurrentBodyCounts, PreviousElementOffsets, OutAssetArrays->ElementOffsets, OutAssetArrays);
 
 	BodyCounts = MoveTemp(CurrentBodyCounts);
 
@@ -1035,10 +1071,29 @@ void UNiagaraDataInterfaceRigidMeshCollisionQuery::DrawDebugHud(FNDIDrawDebugHud
 
 	DebugHudContext.GetOutputString().Appendf(TEXT("Boxes(%d) Spheres(%d) Capsules(%d)"), BoxCount, SphereCount, CapsuleCount);
 
-	auto GetCurrentTransform = [&](int32 ElementIndex)
+	auto GetTransformFromArray = [&](TArray<FVector4f> &TransformArray, int32 ElementIndex)
 	{
 		const uint32 ElementOffset = 3 * ElementIndex;
-		FVector4f* TransformVec = InstanceData_GT->AssetArrays->CurrentTransform.GetData() + ElementOffset;
+		FVector4f* TransformVec = TransformArray.GetData() + ElementOffset;
+
+		FMatrix ElementMatrix;
+		ElementMatrix.SetIdentity();
+
+		for (int32 RowIt = 0; RowIt < 3; ++RowIt)
+		{
+			for (int32 ColIt = 0; ColIt < 4; ++ColIt)
+			{
+				ElementMatrix.M[RowIt][ColIt] = TransformVec[RowIt][ColIt];
+			}
+		}
+
+		return ElementMatrix.GetTransposed();
+	};
+
+	auto GetPreviousTransform = [&](int32 ElementIndex)
+	{
+		const uint32 ElementOffset = 3 * ElementIndex;
+		FVector4f* TransformVec = InstanceData_GT->AssetArrays->PreviousTransform.GetData() + ElementOffset;
 
 		FMatrix ElementMatrix;
 		ElementMatrix.SetIdentity();
@@ -1058,28 +1113,40 @@ void UNiagaraDataInterfaceRigidMeshCollisionQuery::DrawDebugHud(FNDIDrawDebugHud
 	{
 		const UWorld* World = DebugHudContext.GetWorld();
 
+		TArray<FVector4f>& CurrentTransformArray = InstanceData_GT->AssetArrays->CurrentTransform;
+		TArray<FVector4f>& PreviousTransformArray = InstanceData_GT->AssetArrays->PreviousTransform;
+
 		// Boxes
 		for (uint32 BoxIt = 0; BoxIt < BoxCount; ++BoxIt)
 		{
 			const FVector3f HalfBoxExtent = 0.5f * InstanceData_GT->AssetArrays->ElementExtent[ElementOffsets.BoxOffset + BoxIt];
-			const FMatrix CurrentTransform = GetCurrentTransform(ElementOffsets.BoxOffset + BoxIt);
-			DrawDebugBox(World, CurrentTransform.TransformPosition(FVector::ZeroVector), FVector(HalfBoxExtent), CurrentTransform.ToQuat(), FColor::Blue);
+			const FMatrix CurrentTransform = GetTransformFromArray(CurrentTransformArray, ElementOffsets.BoxOffset + BoxIt);
+			DrawDebugBox(World, CurrentTransform.TransformPosition(FVector::ZeroVector), FVector(HalfBoxExtent), CurrentTransform.ToQuat(), FColor::Blue);			
+
+			const FMatrix PrevTransform = GetTransformFromArray(PreviousTransformArray, ElementOffsets.BoxOffset + BoxIt);
+			DrawDebugBox(World, PrevTransform.TransformPosition(FVector::ZeroVector), FVector(HalfBoxExtent), PrevTransform.ToQuat(), FColor::Red);
 		}
 
 		// Spheres
 		for (uint32 SphereIt = 0; SphereIt < SphereCount; ++SphereIt)
 		{
 			const float Radius = InstanceData_GT->AssetArrays->ElementExtent[ElementOffsets.SphereOffset + SphereIt].X;
-			const FMatrix CurrentTransform = GetCurrentTransform(ElementOffsets.SphereOffset + SphereIt);
+			const FMatrix CurrentTransform = GetTransformFromArray(CurrentTransformArray, ElementOffsets.SphereOffset + SphereIt);
 			DrawDebugSphere(World, CurrentTransform.TransformPosition(FVector::ZeroVector), Radius, 20, FColor::Blue);
+
+			const FMatrix PrevTransform = GetTransformFromArray(PreviousTransformArray, ElementOffsets.SphereOffset + SphereIt);
+			DrawDebugSphere(World, PrevTransform.TransformPosition(FVector::ZeroVector), Radius, 20, FColor::Red);
 		}
 
 		// Capsules
 		for (uint32 CapsuleIt = 0; CapsuleIt < CapsuleCount; ++CapsuleIt)
 		{
 			const FVector2f RadiusLength(InstanceData_GT->AssetArrays->ElementExtent[ElementOffsets.CapsuleOffset + CapsuleIt]);
-			const FMatrix CurrentTransform = GetCurrentTransform(ElementOffsets.CapsuleOffset + CapsuleIt);
+			const FMatrix CurrentTransform = GetTransformFromArray(CurrentTransformArray, ElementOffsets.CapsuleOffset + CapsuleIt);
 			DrawDebugCapsule(World, CurrentTransform.TransformPosition(FVector::ZeroVector), RadiusLength.Y * 0.5f, RadiusLength.X, CurrentTransform.ToQuat(), FColor::Blue);
+
+			const FMatrix PrevTransform = GetTransformFromArray(PreviousTransformArray, ElementOffsets.CapsuleOffset + CapsuleIt);
+			DrawDebugCapsule(World, PrevTransform.TransformPosition(FVector::ZeroVector), RadiusLength.Y * 0.5f, RadiusLength.X, CurrentTransform.ToQuat(), FColor::Red);
 		}
 
 		if (!InstanceData_GT->ExplicitActors.IsEmpty() || !InstanceData_GT->FoundActors.IsEmpty())
