- Enable bindings on various modules & properties
- Allow ranges to have bindings

diff --git a/Engine/Plugins/FX/Niagara/Source/Niagara/Internal/Stateless/Modules/NiagaraStatelessModule_AccelerationForce.h b/Engine/Plugins/FX/Niagara/Source/Niagara/Internal/Stateless/Modules/NiagaraStatelessModule_AccelerationForce.h
--- a/Engine/Plugins/FX/Niagara/Source/Niagara/Internal/Stateless/Modules/NiagaraStatelessModule_AccelerationForce.h
+++ b/Engine/Plugins/FX/Niagara/Source/Niagara/Internal/Stateless/Modules/NiagaraStatelessModule_AccelerationForce.h
@@ -15,7 +15,7 @@ class UNiagaraStatelessModule_AccelerationForce : public UNiagaraStatelessModule
 	GENERATED_BODY()
 
 public:
-	UPROPERTY(EditAnywhere, Category = "Parameters", meta = (DisplayName = "Acceleration", DisableUniformDistribution))
+	UPROPERTY(EditAnywhere, Category = "Parameters", meta = (DisplayName = "Acceleration", DisableUniformDistribution, DisableBindingDistribution))
 	FNiagaraDistributionRangeVector3 AccelerationDistribution = FNiagaraDistributionRangeVector3(FVector3f::ZeroVector);
 
 	virtual void BuildEmitterData(FNiagaraStatelessEmitterDataBuildContext& BuildContext) const override
diff --git a/Engine/Plugins/FX/Niagara/Source/Niagara/Internal/Stateless/Modules/NiagaraStatelessModule_AddVelocity.h b/Engine/Plugins/FX/Niagara/Source/Niagara/Internal/Stateless/Modules/NiagaraStatelessModule_AddVelocity.h
--- a/Engine/Plugins/FX/Niagara/Source/Niagara/Internal/Stateless/Modules/NiagaraStatelessModule_AddVelocity.h
+++ b/Engine/Plugins/FX/Niagara/Source/Niagara/Internal/Stateless/Modules/NiagaraStatelessModule_AddVelocity.h
@@ -26,12 +26,12 @@ public:
 	UPROPERTY(EditAnywhere, Category = "Parameters")
 	ENSM_VelocityType VelocityType = ENSM_VelocityType::Linear;
 
-	UPROPERTY(EditAnywhere, Category = "Parameters", meta = (DisplayName = "Velocity", DisableUniformDistribution, EditConditionHides, EditCondition = "VelocityType == ENSM_VelocityType::Linear"))
+	UPROPERTY(EditAnywhere, Category = "Parameters", meta = (DisplayName = "Velocity", DisableUniformDistribution, DisableBindingDistribution, EditConditionHides, EditCondition = "VelocityType == ENSM_VelocityType::Linear"))
 	FNiagaraDistributionRangeVector3 LinearVelocityDistribution = FNiagaraDistributionRangeVector3(FVector3f(0.0f, 0.0f, 100.0f));
 	UPROPERTY(EditAnywhere, Category = "Parameters", meta = (DisplayName = "Velocity Scale", EditConditionHides, EditCondition = "VelocityType == ENSM_VelocityType::Linear"))
 	float LinearVelocityScale = 1.0f;
 
-	UPROPERTY(EditAnywhere, Category = "Parameters", meta = (DisplayName = "Cone Velocity", EditConditionHides, EditCondition = "VelocityType == ENSM_VelocityType::InCone"))
+	UPROPERTY(EditAnywhere, Category = "Parameters", meta = (DisplayName = "Cone Velocity", DisableBindingDistribution, EditConditionHides, EditCondition = "VelocityType == ENSM_VelocityType::InCone"))
 	FNiagaraDistributionRangeFloat ConeVelocityDistribution = FNiagaraDistributionRangeFloat(500.0f);
 	UPROPERTY(EditAnywhere, Category = "Parameters", meta = (EditConditionHides, EditCondition = "VelocityType == ENSM_VelocityType::InCone"))
 	FRotator ConeRotation = FRotator::ZeroRotator;
@@ -40,7 +40,7 @@ public:
 	UPROPERTY(EditAnywhere, Category = "Parameters", meta = (EditConditionHides, EditCondition = "VelocityType == ENSM_VelocityType::InCone", ClampMin = "0.0", ClampMax = "360.0", UIMin = "0.0", UIMax = "360.0"))
 	float InnerCone = 0.0f;
 
-	UPROPERTY(EditAnywhere, Category = "Parameters", meta = (DisplayName = "Point Velocity", EditConditionHides, EditCondition = "VelocityType == ENSM_VelocityType::FromPoint"))
+	UPROPERTY(EditAnywhere, Category = "Parameters", meta = (DisplayName = "Point Velocity", DisableBindingDistribution, EditConditionHides, EditCondition = "VelocityType == ENSM_VelocityType::FromPoint"))
 	FNiagaraDistributionRangeFloat PointVelocityDistribution = FNiagaraDistributionRangeFloat(100.0f);
 	UPROPERTY(EditAnywhere, Category = "Parameters", meta = (EditConditionHides, EditCondition = "VelocityType == ENSM_VelocityType::FromPoint"))
 	FVector3f PointOrigin = FVector3f::ZeroVector;
diff --git a/Engine/Plugins/FX/Niagara/Source/Niagara/Internal/Stateless/Modules/NiagaraStatelessModule_Drag.h b/Engine/Plugins/FX/Niagara/Source/Niagara/Internal/Stateless/Modules/NiagaraStatelessModule_Drag.h
--- a/Engine/Plugins/FX/Niagara/Source/Niagara/Internal/Stateless/Modules/NiagaraStatelessModule_Drag.h
+++ b/Engine/Plugins/FX/Niagara/Source/Niagara/Internal/Stateless/Modules/NiagaraStatelessModule_Drag.h
@@ -17,7 +17,7 @@ class UNiagaraStatelessModule_Drag : public UNiagaraStatelessModule
 	static constexpr float DefaultDrag = 1.0f;
 
 public:
-	UPROPERTY(EditAnywhere, Category = "Parameters", meta = (DisplayName = "Drag"))
+	UPROPERTY(EditAnywhere, Category = "Parameters", meta = (DisplayName = "Drag", DisableBindingDistribution))
 	FNiagaraDistributionRangeFloat DragDistribution = FNiagaraDistributionRangeFloat(DefaultDrag);
 
 	virtual void BuildEmitterData(FNiagaraStatelessEmitterDataBuildContext& BuildContext) const override
diff --git a/Engine/Plugins/FX/Niagara/Source/Niagara/Internal/Stateless/Modules/NiagaraStatelessModule_GravityForce.h b/Engine/Plugins/FX/Niagara/Source/Niagara/Internal/Stateless/Modules/NiagaraStatelessModule_GravityForce.h
--- a/Engine/Plugins/FX/Niagara/Source/Niagara/Internal/Stateless/Modules/NiagaraStatelessModule_GravityForce.h
+++ b/Engine/Plugins/FX/Niagara/Source/Niagara/Internal/Stateless/Modules/NiagaraStatelessModule_GravityForce.h
@@ -15,7 +15,7 @@ class UNiagaraStatelessModule_GravityForce : public UNiagaraStatelessModule
 	GENERATED_BODY()
 
 public:
-	UPROPERTY(EditAnywhere, Category = "Parameters", meta = (DisplayName = "Gravity", DisableUniformDistribution))
+	UPROPERTY(EditAnywhere, Category = "Parameters", meta = (DisplayName = "Gravity", DisableUniformDistribution, DisableBindingDistribution))
 	FNiagaraDistributionRangeVector3 GravityDistribution = FNiagaraDistributionRangeVector3(GetDefaultValue());
 
 	virtual void BuildEmitterData(FNiagaraStatelessEmitterDataBuildContext& BuildContext) const override
diff --git a/Engine/Plugins/FX/Niagara/Source/Niagara/Internal/Stateless/Modules/NiagaraStatelessModule_InitializeParticle.h b/Engine/Plugins/FX/Niagara/Source/Niagara/Internal/Stateless/Modules/NiagaraStatelessModule_InitializeParticle.h
--- a/Engine/Plugins/FX/Niagara/Source/Niagara/Internal/Stateless/Modules/NiagaraStatelessModule_InitializeParticle.h
+++ b/Engine/Plugins/FX/Niagara/Source/Niagara/Internal/Stateless/Modules/NiagaraStatelessModule_InitializeParticle.h
@@ -40,7 +40,7 @@ public:
 	FNiagaraDistributionRangeFloat LifetimeDistribution = FNiagaraDistributionRangeFloat(FNiagaraStatelessGlobals::GetDefaultLifetimeValue());
 
 	UPROPERTY(EditAnywhere, Category = "Parameters", meta = (DisplayName = "Color", DisableCurveDistribution))
-	FNiagaraDistributionColor ColorDistribution = FNiagaraDistributionColor(FNiagaraStatelessGlobals::GetDefaultColorValue());
+	FNiagaraDistributionRangeColor ColorDistribution = FNiagaraDistributionRangeColor(FNiagaraStatelessGlobals::GetDefaultColorValue());
 
 	UPROPERTY(EditAnywhere, Category = "Parameters", meta = (DisplayName = "Mass"))
 	FNiagaraDistributionRangeFloat MassDistribution = FNiagaraDistributionRangeFloat(FNiagaraStatelessGlobals::GetDefaultMassValue());
@@ -71,12 +71,12 @@ public:
 
 		BuiltData->InitialPosition			= BuildContext.AddDistribution(InitialPosition, true);
 		BuiltData->LifetimeRange			= LifetimeDistribution.CalculateRange(FNiagaraStatelessGlobals::GetDefaultLifetimeValue());
-		BuiltData->ColorRange				= ColorDistribution.CalculateRange(FNiagaraStatelessGlobals::GetDefaultColorValue());
-		BuiltData->MassRange				= MassDistribution.CalculateRange(FNiagaraStatelessGlobals::GetDefaultMassValue());
-		BuiltData->SpriteSizeRange			= SpriteSizeDistribution.CalculateRange(FNiagaraStatelessGlobals::GetDefaultSpriteSizeValue());
-		BuiltData->SpriteRotationRange		= SpriteRotationDistribution.CalculateRange(FNiagaraStatelessGlobals::GetDefaultSpriteRotationValue());
-		BuiltData->MeshScaleRange			= MeshScaleDistribution.CalculateRange(FNiagaraStatelessGlobals::GetDefaultScaleValue());
-		BuiltData->RibbonWidthRange			= RibbonWidthDistribution.CalculateRange(FNiagaraStatelessGlobals::GetDefaultRibbonWidthValue());
+		BuiltData->ColorRange				= BuildContext.ConvertDistributionToRange(ColorDistribution, FNiagaraStatelessGlobals::GetDefaultColorValue());
+		BuiltData->MassRange				= BuildContext.ConvertDistributionToRange(MassDistribution, FNiagaraStatelessGlobals::GetDefaultMassValue());
+		BuiltData->SpriteSizeRange			= BuildContext.ConvertDistributionToRange(SpriteSizeDistribution, FNiagaraStatelessGlobals::GetDefaultSpriteSizeValue());
+		BuiltData->SpriteRotationRange		= BuildContext.ConvertDistributionToRange(SpriteRotationDistribution, FNiagaraStatelessGlobals::GetDefaultSpriteRotationValue());
+		BuiltData->MeshScaleRange			= BuildContext.ConvertDistributionToRange(MeshScaleDistribution, FNiagaraStatelessGlobals::GetDefaultScaleValue());
+		BuiltData->RibbonWidthRange			= BuildContext.ConvertDistributionToRange(RibbonWidthDistribution, FNiagaraStatelessGlobals::GetDefaultRibbonWidthValue());
 
 		NiagaraStateless::FPhysicsBuildData& PhysicsBuildData = BuildContext.GetTransientBuildData<NiagaraStateless::FPhysicsBuildData>();
 		PhysicsBuildData.MassRange			= MassDistribution.CalculateRange(FNiagaraStatelessGlobals::GetDefaultMassValue());
@@ -89,16 +89,11 @@ public:
 
 		Parameters->InitializeParticle_ModuleFlags			= ModuleBuiltData->ModuleFlags;
 		Parameters->InitializeParticle_InitialPosition		= ModuleBuiltData->InitialPosition;
-		Parameters->InitializeParticle_ColorScale			= ModuleBuiltData->ColorRange.GetScale();
-		Parameters->InitializeParticle_ColorBias			= ModuleBuiltData->ColorRange.Min;
-		Parameters->InitializeParticle_SpriteSizeScale		= ModuleBuiltData->SpriteSizeRange.GetScale();
-		Parameters->InitializeParticle_SpriteSizeBias		= ModuleBuiltData->SpriteSizeRange.Min;
-		Parameters->InitializeParticle_SpriteRotationScale	= ModuleBuiltData->SpriteRotationRange.GetScale();
-		Parameters->InitializeParticle_SpriteRotationBias	= ModuleBuiltData->SpriteRotationRange.Min;
-		Parameters->InitializeParticle_MeshScaleScale		= ModuleBuiltData->MeshScaleRange.GetScale();
-		Parameters->InitializeParticle_MeshScaleBias		= ModuleBuiltData->MeshScaleRange.Min;
-		Parameters->InitializeParticle_RibbonWidthScale		= ModuleBuiltData->RibbonWidthRange.GetScale();
-		Parameters->InitializeParticle_RibbonWidthBias		= ModuleBuiltData->RibbonWidthRange.Min;
+		SetShaderParameterContext.ConvertRangeToScaleBias(ModuleBuiltData->ColorRange,			Parameters->InitializeParticle_ColorScale, Parameters->InitializeParticle_ColorBias);
+		SetShaderParameterContext.ConvertRangeToScaleBias(ModuleBuiltData->SpriteSizeRange,		Parameters->InitializeParticle_SpriteSizeScale, Parameters->InitializeParticle_SpriteSizeBias);
+		SetShaderParameterContext.ConvertRangeToScaleBias(ModuleBuiltData->SpriteRotationRange, Parameters->InitializeParticle_SpriteRotationScale, Parameters->InitializeParticle_SpriteRotationBias);
+		SetShaderParameterContext.ConvertRangeToScaleBias(ModuleBuiltData->MeshScaleRange,		Parameters->InitializeParticle_MeshScaleScale, Parameters->InitializeParticle_MeshScaleBias);
+		SetShaderParameterContext.ConvertRangeToScaleBias(ModuleBuiltData->RibbonWidthRange,	Parameters->InitializeParticle_RibbonWidthScale, Parameters->InitializeParticle_RibbonWidthBias);
 	}
 
 #if WITH_EDITORONLY_DATA
diff --git a/Engine/Plugins/FX/Niagara/Source/Niagara/Internal/Stateless/Modules/NiagaraStatelessModule_SpriteRotationRate.h b/Engine/Plugins/FX/Niagara/Source/Niagara/Internal/Stateless/Modules/NiagaraStatelessModule_SpriteRotationRate.h
--- a/Engine/Plugins/FX/Niagara/Source/Niagara/Internal/Stateless/Modules/NiagaraStatelessModule_SpriteRotationRate.h
+++ b/Engine/Plugins/FX/Niagara/Source/Niagara/Internal/Stateless/Modules/NiagaraStatelessModule_SpriteRotationRate.h
@@ -13,26 +13,29 @@ class UNiagaraStatelessModule_SpriteRotationRate : public UNiagaraStatelessModul
 {
 	GENERATED_BODY()
 
+	struct FModuleBuiltData
+	{
+		FNiagaraStatelessRangeFloat		RotationRange;
+	};
+
 public:
 	using FParameters = NiagaraStateless::FSpriteRotationRateModule_ShaderParameters;
 
 	UPROPERTY(EditAnywhere, Category = "Parameters", meta = (DisplayName = "Rotation Rate"))
 	FNiagaraDistributionRangeFloat RotationRateDistribution = FNiagaraDistributionRangeFloat(FNiagaraStatelessGlobals::GetDefaultSpriteRotationValue());
 
+	virtual void BuildEmitterData(FNiagaraStatelessEmitterDataBuildContext& BuildContext) const override
+	{
+		FModuleBuiltData* BuiltData = BuildContext.AllocateBuiltData<FModuleBuiltData>();
+		BuiltData->RotationRange = BuildContext.ConvertDistributionToRange(RotationRateDistribution, 0.0f, IsModuleEnabled());
+	}
+
 	virtual void SetShaderParameters(const FNiagaraStatelessSetShaderParameterContext& SetShaderParameterContext) const override
 	{
 		FParameters* Parameters = SetShaderParameterContext.GetParameterNestedStruct<FParameters>();
-		if (IsModuleEnabled())
-		{
-			const FNiagaraStatelessRangeFloat RotationRateRange = RotationRateDistribution.CalculateRange(FNiagaraStatelessGlobals::GetDefaultSpriteRotationValue());
-			Parameters->SpriteRotationRate_Scale	= RotationRateRange.GetScale();
-			Parameters->SpriteRotationRate_Bias		= RotationRateRange.Min;
-		}
-		else
-		{
-			Parameters->SpriteRotationRate_Scale	= 0.0f;
-			Parameters->SpriteRotationRate_Bias		= 0.0f;
-		}
+		const FModuleBuiltData* ModuleBuiltData = SetShaderParameterContext.ReadBuiltData<FModuleBuiltData>();
+
+		SetShaderParameterContext.ConvertRangeToScaleBias(ModuleBuiltData->RotationRange, Parameters->SpriteRotationRate_Scale, Parameters->SpriteRotationRate_Bias);
 	}
 
 #if WITH_EDITOR
diff --git a/Engine/Plugins/FX/Niagara/Source/Niagara/Internal/Stateless/NiagaraStatelessCommon.h b/Engine/Plugins/FX/Niagara/Source/Niagara/Internal/Stateless/NiagaraStatelessCommon.h
--- a/Engine/Plugins/FX/Niagara/Source/Niagara/Internal/Stateless/NiagaraStatelessCommon.h
+++ b/Engine/Plugins/FX/Niagara/Source/Niagara/Internal/Stateless/NiagaraStatelessCommon.h
@@ -15,8 +15,9 @@ struct FNiagaraStatelessRange
 
 	ValueType GetScale() const { return Max - Min; }
 
-	ValueType Min = {};
-	ValueType Max = {};
+	int32		ParameterOffset = INDEX_NONE;
+	ValueType	Min = {};
+	ValueType	Max = {};
 };
 
 using FNiagaraStatelessRangeInt		= FNiagaraStatelessRange<int32>;
diff --git a/Engine/Plugins/FX/Niagara/Source/Niagara/Internal/Stateless/NiagaraStatelessDistribution.h b/Engine/Plugins/FX/Niagara/Source/Niagara/Internal/Stateless/NiagaraStatelessDistribution.h
--- a/Engine/Plugins/FX/Niagara/Source/Niagara/Internal/Stateless/NiagaraStatelessDistribution.h
+++ b/Engine/Plugins/FX/Niagara/Source/Niagara/Internal/Stateless/NiagaraStatelessDistribution.h
@@ -100,6 +100,7 @@ struct FNiagaraDistributionRangeFloat : public FNiagaraDistributionBase
 #if WITH_EDITORONLY_DATA
 	virtual bool AllowCurves() const override { return false; }
 	NIAGARA_API virtual void UpdateValuesFromDistribution() override;
+	virtual FNiagaraTypeDefinition GetBindingTypeDef() { return FNiagaraTypeDefinition::GetFloatDef(); }
 #endif
 };
 
@@ -123,6 +124,7 @@ struct FNiagaraDistributionRangeVector2 : public FNiagaraDistributionBase
 #if WITH_EDITORONLY_DATA
 	virtual bool AllowCurves() const override { return false; }
 	virtual void UpdateValuesFromDistribution() override;
+	virtual FNiagaraTypeDefinition GetBindingTypeDef() { return FNiagaraTypeDefinition::GetVec2Def(); }
 #endif
 };
 
@@ -146,6 +148,31 @@ struct FNiagaraDistributionRangeVector3 : public FNiagaraDistributionBase
 #if WITH_EDITORONLY_DATA
 	virtual bool AllowCurves() const override { return false; }
 	virtual void UpdateValuesFromDistribution() override;
+	virtual FNiagaraTypeDefinition GetBindingTypeDef() { return FNiagaraTypeDefinition::GetVec3Def(); }
+#endif
+};
+
+USTRUCT()
+struct FNiagaraDistributionRangeColor : public FNiagaraDistributionBase
+{
+	GENERATED_BODY()
+
+	FNiagaraDistributionRangeColor() = default;
+	explicit FNiagaraDistributionRangeColor(const FLinearColor& ConstantValue) { InitConstant(ConstantValue); }
+
+	UPROPERTY(EditAnywhere, Category = "Parameters")
+	FLinearColor Min = FLinearColor::White;
+
+	UPROPERTY(EditAnywhere, Category = "Parameters")
+	FLinearColor Max = FLinearColor::White;
+
+	NIAGARA_API void InitConstant(const FLinearColor& Value);
+	NIAGARA_API FNiagaraStatelessRangeColor CalculateRange(const FLinearColor& Default = FLinearColor::White) const;
+
+#if WITH_EDITORONLY_DATA
+	virtual bool AllowCurves() const override { return false; }
+	virtual void UpdateValuesFromDistribution() override;
+	virtual FNiagaraTypeDefinition GetBindingTypeDef() { return FNiagaraTypeDefinition::GetColorDef(); }
 #endif
 };
 
@@ -230,5 +257,6 @@ struct FNiagaraDistributionColor : public FNiagaraDistributionBase
 #if WITH_EDITORONLY_DATA
 	virtual bool DisplayAsColor() const override { return true; }
 	virtual void UpdateValuesFromDistribution() override;
+	virtual FNiagaraTypeDefinition GetBindingTypeDef() { return FNiagaraTypeDefinition::GetColorDef(); }
 #endif
 };
diff --git a/Engine/Plugins/FX/Niagara/Source/Niagara/Internal/Stateless/NiagaraStatelessEmitterDataBuildContext.h b/Engine/Plugins/FX/Niagara/Source/Niagara/Internal/Stateless/NiagaraStatelessEmitterDataBuildContext.h
--- a/Engine/Plugins/FX/Niagara/Source/Niagara/Internal/Stateless/NiagaraStatelessEmitterDataBuildContext.h
+++ b/Engine/Plugins/FX/Niagara/Source/Niagara/Internal/Stateless/NiagaraStatelessEmitterDataBuildContext.h
@@ -107,6 +107,34 @@ public:
 		return Parameters;
 	}
 
+	template<typename TRange, typename TDistribution, typename TDefaultValue>
+	TRange ConvertDistributionToRangeHelper(const TDistribution& Distribution, const TDefaultValue& DefaultValue, bool bEnabled)
+	{
+		TRange Range(DefaultValue);
+		if (bEnabled)
+		{
+			if (Distribution.Mode == ENiagaraDistributionMode::Binding)
+			{
+				Range.ParameterOffset = AddRendererBinding(Distribution.ParameterBinding);
+			}
+			else
+			{
+				Range = Distribution.CalculateRange(DefaultValue);
+			}
+		}
+		return Range;
+	}
+
+	FNiagaraStatelessRangeFloat   ConvertDistributionToRange(const FNiagaraDistributionFloat& Distribution, float DefaultValue, bool bEnabled = true) { return ConvertDistributionToRangeHelper<FNiagaraStatelessRangeFloat>(Distribution, DefaultValue, bEnabled); }
+	FNiagaraStatelessRangeVector2 ConvertDistributionToRange(const FNiagaraDistributionVector2& Distribution, const FVector2f& DefaultValue, bool bEnabled = true) { return ConvertDistributionToRangeHelper<FNiagaraStatelessRangeVector2>(Distribution, DefaultValue, bEnabled); }
+	FNiagaraStatelessRangeVector3 ConvertDistributionToRange(const FNiagaraDistributionVector3& Distribution, const FVector3f& DefaultValue, bool bEnabled = true) { return ConvertDistributionToRangeHelper<FNiagaraStatelessRangeVector3>(Distribution, DefaultValue, bEnabled); }
+	FNiagaraStatelessRangeColor   ConvertDistributionToRange(const FNiagaraDistributionColor& Distribution, const FLinearColor& DefaultValue, bool bEnabled = true) { return ConvertDistributionToRangeHelper<FNiagaraStatelessRangeColor>(Distribution, DefaultValue, bEnabled); }
+
+	FNiagaraStatelessRangeFloat   ConvertDistributionToRange(const FNiagaraDistributionRangeFloat& Distribution, float DefaultValue, bool bEnabled = true) { return ConvertDistributionToRangeHelper<FNiagaraStatelessRangeFloat>(Distribution, DefaultValue, bEnabled); }
+	FNiagaraStatelessRangeVector2 ConvertDistributionToRange(const FNiagaraDistributionRangeVector2& Distribution, const FVector2f& DefaultValue, bool bEnabled = true) { return ConvertDistributionToRangeHelper<FNiagaraStatelessRangeVector2>(Distribution, DefaultValue, bEnabled); }
+	FNiagaraStatelessRangeVector3 ConvertDistributionToRange(const FNiagaraDistributionRangeVector3& Distribution, const FVector3f& DefaultValue, bool bEnabled = true) { return ConvertDistributionToRangeHelper<FNiagaraStatelessRangeVector3>(Distribution, DefaultValue, bEnabled); }
+	FNiagaraStatelessRangeColor   ConvertDistributionToRange(const FNiagaraDistributionRangeColor& Distribution, const FLinearColor& DefaultValue, bool bEnabled = true) { return ConvertDistributionToRangeHelper<FNiagaraStatelessRangeColor>(Distribution, DefaultValue, bEnabled); }
+
 private:
 	FNiagaraParameterStore& RendererBindings;
 	TArray<uint8>&			BuiltData;
diff --git a/Engine/Plugins/FX/Niagara/Source/Niagara/Internal/Stateless/NiagaraStatelessModule.h b/Engine/Plugins/FX/Niagara/Source/Niagara/Internal/Stateless/NiagaraStatelessModule.h
--- a/Engine/Plugins/FX/Niagara/Source/Niagara/Internal/Stateless/NiagaraStatelessModule.h
+++ b/Engine/Plugins/FX/Niagara/Source/Niagara/Internal/Stateless/NiagaraStatelessModule.h
@@ -65,6 +65,11 @@ public:
 		}
 	}
 
+	void ConvertRangeToScaleBias(const FNiagaraStatelessRangeFloat& Range, float& OutScale, float& OutBias) const { OutScale = Range.GetScale(); GetRendererParameterValue(OutBias, Range.ParameterOffset, Range.Min); }
+	void ConvertRangeToScaleBias(const FNiagaraStatelessRangeVector2& Range, FVector2f& OutScale, FVector2f& OutBias) const { OutScale = Range.GetScale(); GetRendererParameterValue(OutBias, Range.ParameterOffset, Range.Min); }
+	void ConvertRangeToScaleBias(const FNiagaraStatelessRangeVector3& Range, FVector3f& OutScale, FVector3f& OutBias) const { OutScale = Range.GetScale(); GetRendererParameterValue(OutBias, Range.ParameterOffset, Range.Min); }
+	void ConvertRangeToScaleBias(const FNiagaraStatelessRangeColor& Range, FLinearColor& OutScale, FLinearColor& OutBias) const { OutScale = Range.GetScale(); GetRendererParameterValue(OutBias, Range.ParameterOffset, Range.Min); }
+
 protected:
 #if DO_CHECK
 	void ValidateIncludeStructType(uint32 StructOffset, const FShaderParametersMetadata* StructMetaData) const;
diff --git a/Engine/Plugins/FX/Niagara/Source/Niagara/Internal/Stateless/NiagaraStatelessSpawnInfo.h b/Engine/Plugins/FX/Niagara/Source/Niagara/Internal/Stateless/NiagaraStatelessSpawnInfo.h
--- a/Engine/Plugins/FX/Niagara/Source/Niagara/Internal/Stateless/NiagaraStatelessSpawnInfo.h
+++ b/Engine/Plugins/FX/Niagara/Source/Niagara/Internal/Stateless/NiagaraStatelessSpawnInfo.h
@@ -30,10 +30,10 @@ struct FNiagaraStatelessSpawnInfo
 	UPROPERTY(EditAnywhere, Category = "Spawn", meta = (EditConditionHides, EditCondition = "Type == ENiagaraStatelessSpawnInfoType::Burst"))
 	float SpawnTime = 0.0f;
 
-	UPROPERTY(EditAnywhere, Category = "Spawn", meta = (EditConditionHides, EditCondition = "Type == ENiagaraStatelessSpawnInfoType::Burst", ClampMin = "0"))
+	UPROPERTY(EditAnywhere, Category = "Spawn", meta = (EditConditionHides, EditCondition = "Type == ENiagaraStatelessSpawnInfoType::Burst", ClampMin = "0", DisableBindingDistribution))
 	FNiagaraDistributionRangeInt Amount = FNiagaraDistributionRangeInt(1);
 
-	UPROPERTY(EditAnywhere, Category = "Spawn", meta = (EditConditionHides, EditCondition = "Type == ENiagaraStatelessSpawnInfoType::Rate", ClampMin = "0.0"))
+	UPROPERTY(EditAnywhere, Category = "Spawn", meta = (EditConditionHides, EditCondition = "Type == ENiagaraStatelessSpawnInfoType::Rate", ClampMin = "0.0", DisableBindingDistribution))
 	FNiagaraDistributionRangeFloat Rate = FNiagaraDistributionRangeFloat(60.0f);
 
 	UPROPERTY(EditAnywhere, Category = "Spawn", meta = (InlineEditConditionToggle))
diff --git a/Engine/Plugins/FX/Niagara/Source/Niagara/Private/Stateless/NiagaraStatelessDistribution.cpp b/Engine/Plugins/FX/Niagara/Source/Niagara/Private/Stateless/NiagaraStatelessDistribution.cpp
--- a/Engine/Plugins/FX/Niagara/Source/Niagara/Private/Stateless/NiagaraStatelessDistribution.cpp
+++ b/Engine/Plugins/FX/Niagara/Source/Niagara/Private/Stateless/NiagaraStatelessDistribution.cpp
@@ -217,6 +217,22 @@ FNiagaraStatelessRangeVector3 FNiagaraDistributionRangeVector3::CalculateRange(c
 	return Mode == ENiagaraDistributionMode::UniformConstant ? FNiagaraStatelessRangeVector3(Min, Min) : FNiagaraStatelessRangeVector3(Min, Max);
 }
 
+void FNiagaraDistributionRangeColor::InitConstant(const FLinearColor& Value)
+{
+	Mode = ENiagaraDistributionMode::NonUniformConstant;
+	Min = Value;
+	Max = Value;
+#if WITH_EDITORONLY_DATA
+	ChannelConstantsAndRanges = TArray<float>({ Value.R, Value.G, Value.B, Value.A });
+#endif
+}
+
+FNiagaraStatelessRangeColor FNiagaraDistributionRangeColor::CalculateRange(const FLinearColor& Default) const
+{
+	return Mode == ENiagaraDistributionMode::UniformConstant ? FNiagaraStatelessRangeColor(Min, Min) : FNiagaraStatelessRangeColor(Min, Max);
+}
+
+
 void FNiagaraDistributionFloat::InitConstant(float Value)
 {
 	Mode = ENiagaraDistributionMode::UniformConstant;
@@ -407,6 +423,19 @@ void FNiagaraDistributionRangeVector3::UpdateValuesFromDistribution()
 	);
 }
 
+void FNiagaraDistributionRangeColor::UpdateValuesFromDistribution()
+{
+	NiagaraStatelessDistributionPrivate::UpdateDistributionValues(
+		Mode,
+		ChannelConstantsAndRanges,
+		ChannelCurves,
+		4,
+		[this](int32 Num) {},
+		[this](int32 ValueIndex, int32 ChannelIndex) -> float& { return ValueIndex == 0 ? Min.Component(ChannelIndex) : Max.Component(ChannelIndex); },
+		MaxLutSampleCount
+	);
+}
+
 void FNiagaraDistributionFloat::UpdateValuesFromDistribution()
 {
 	NiagaraStatelessDistributionPrivate::UpdateDistributionValues(
diff --git a/Engine/Plugins/FX/Niagara/Source/Niagara/Private/Stateless/NiagaraStatelessEmitterInstance.cpp b/Engine/Plugins/FX/Niagara/Source/Niagara/Private/Stateless/NiagaraStatelessEmitterInstance.cpp
--- a/Engine/Plugins/FX/Niagara/Source/Niagara/Private/Stateless/NiagaraStatelessEmitterInstance.cpp
+++ b/Engine/Plugins/FX/Niagara/Source/Niagara/Private/Stateless/NiagaraStatelessEmitterInstance.cpp
@@ -205,16 +205,20 @@ void FNiagaraStatelessEmitterInstance::Tick(float DeltaSeconds)
 
 	if (NiagaraStateless::FEmitterInstance_RT* RenderThreadData = RenderThreadDataPtr.Get())
 	{
+		NiagaraStateless::FCommonShaderParameters* NewShaderParameters = nullptr;
 		TOptional<TArray<uint8>> NewBindingBufferData;
 		if ( EmitterData->bModulesHaveRendererBindings && RendererBindings.GetParametersDirty() )
 		{
 			RendererBindings.Tick();
 			NewBindingBufferData = RendererBindings.GetParameterDataArray();
 			check((NewBindingBufferData->Num() % sizeof(uint32)) == 0);
+
+			NewShaderParameters = WeakStatelessEmitter->AllocateShaderParameters(RendererBindings);
+			NewShaderParameters->Common_RandomSeed = RandomSeed;
 		}
 
 		ENQUEUE_RENDER_COMMAND(UpdateStatelessAge)(
-			[RenderThreadData, AgeForRT=Age, ExecutionStateForRT=ExecutionState, NewBindingBufferDataRT=MoveTemp(NewBindingBufferData)](FRHICommandListImmediate& RHICmdList)
+			[RenderThreadData, AgeForRT=Age, ExecutionStateForRT=ExecutionState, NewBindingBufferDataRT=MoveTemp(NewBindingBufferData), NewShaderParameters](FRHICommandListImmediate& RHICmdList)
 			{
 				RenderThreadData->Age = AgeForRT;
 				RenderThreadData->ExecutionState = ExecutionStateForRT;
@@ -222,7 +226,11 @@ void FNiagaraStatelessEmitterInstance::Tick(float DeltaSeconds)
 				{
 					RenderThreadData->BindingBufferData = NewBindingBufferDataRT;
 				}
-			}
+				if (NewShaderParameters)
+				{
+					RenderThreadData->ShaderParameters.Reset(NewShaderParameters);
+				}
+		}
 		);
 	}
 }
diff --git a/Engine/Plugins/FX/Niagara/Source/NiagaraEditor/Private/ViewModels/Stack/NiagaraStackStatelessEmitterGroup.cpp b/Engine/Plugins/FX/Niagara/Source/NiagaraEditor/Private/ViewModels/Stack/NiagaraStackStatelessEmitterGroup.cpp
--- a/Engine/Plugins/FX/Niagara/Source/NiagaraEditor/Private/ViewModels/Stack/NiagaraStackStatelessEmitterGroup.cpp
+++ b/Engine/Plugins/FX/Niagara/Source/NiagaraEditor/Private/ViewModels/Stack/NiagaraStackStatelessEmitterGroup.cpp
@@ -121,6 +121,7 @@ void UNiagaraStackStatelessEmitterObjectItem::RefreshChildrenInternal(const TArr
 			StatelessEmitterStackObject->RegisterInstancedCustomPropertyTypeLayout(FNiagaraDistributionRangeFloat::StaticStruct()->GetFName(), FOnGetPropertyTypeCustomizationInstance::CreateStatic(&FNiagaraDistributionPropertyCustomization::MakeFloatInstance));
 			StatelessEmitterStackObject->RegisterInstancedCustomPropertyTypeLayout(FNiagaraDistributionRangeVector2::StaticStruct()->GetFName(), FOnGetPropertyTypeCustomizationInstance::CreateStatic(&FNiagaraDistributionPropertyCustomization::MakeVector2Instance));
 			StatelessEmitterStackObject->RegisterInstancedCustomPropertyTypeLayout(FNiagaraDistributionRangeVector3::StaticStruct()->GetFName(), FOnGetPropertyTypeCustomizationInstance::CreateStatic(&FNiagaraDistributionPropertyCustomization::MakeVector3Instance));
+			StatelessEmitterStackObject->RegisterInstancedCustomPropertyTypeLayout(FNiagaraDistributionRangeColor::StaticStruct()->GetFName(), FOnGetPropertyTypeCustomizationInstance::CreateStatic(&FNiagaraDistributionPropertyCustomization::MakeColorInstance));
 			StatelessEmitterStackObjectWeak = StatelessEmitterStackObject;
 		}
 		NewChildren.Add(StatelessEmitterStackObject);
diff --git a/Engine/Plugins/FX/Niagara/Source/NiagaraEditor/Private/ViewModels/Stack/NiagaraStackStatelessEmitterSimulateGroup.cpp b/Engine/Plugins/FX/Niagara/Source/NiagaraEditor/Private/ViewModels/Stack/NiagaraStackStatelessEmitterSimulateGroup.cpp
--- a/Engine/Plugins/FX/Niagara/Source/NiagaraEditor/Private/ViewModels/Stack/NiagaraStackStatelessEmitterSimulateGroup.cpp
+++ b/Engine/Plugins/FX/Niagara/Source/NiagaraEditor/Private/ViewModels/Stack/NiagaraStackStatelessEmitterSimulateGroup.cpp
@@ -92,6 +92,7 @@ void UNiagaraStackStatelessModuleItem::RefreshChildrenInternal(const TArray<UNia
 			ModuleObject->RegisterInstancedCustomPropertyTypeLayout(FNiagaraDistributionRangeFloat::StaticStruct()->GetFName(), FOnGetPropertyTypeCustomizationInstance::CreateStatic(&FNiagaraDistributionPropertyCustomization::MakeFloatInstance));
 			ModuleObject->RegisterInstancedCustomPropertyTypeLayout(FNiagaraDistributionRangeVector2::StaticStruct()->GetFName(), FOnGetPropertyTypeCustomizationInstance::CreateStatic(&FNiagaraDistributionPropertyCustomization::MakeVector2Instance));
 			ModuleObject->RegisterInstancedCustomPropertyTypeLayout(FNiagaraDistributionRangeVector3::StaticStruct()->GetFName(), FOnGetPropertyTypeCustomizationInstance::CreateStatic(&FNiagaraDistributionPropertyCustomization::MakeVector3Instance));
+			ModuleObject->RegisterInstancedCustomPropertyTypeLayout(FNiagaraDistributionRangeColor::StaticStruct()->GetFName(), FOnGetPropertyTypeCustomizationInstance::CreateStatic(&FNiagaraDistributionPropertyCustomization::MakeColorInstance));
 			ModuleObjectWeak = ModuleObject;
 		}
 		NewChildren.Add(ModuleObject);
diff --git a/Engine/Plugins/FX/Niagara/Source/NiagaraEditor/Private/ViewModels/Stack/NiagaraStackStatelessEmitterSpawnGroup.cpp b/Engine/Plugins/FX/Niagara/Source/NiagaraEditor/Private/ViewModels/Stack/NiagaraStackStatelessEmitterSpawnGroup.cpp
--- a/Engine/Plugins/FX/Niagara/Source/NiagaraEditor/Private/ViewModels/Stack/NiagaraStackStatelessEmitterSpawnGroup.cpp
+++ b/Engine/Plugins/FX/Niagara/Source/NiagaraEditor/Private/ViewModels/Stack/NiagaraStackStatelessEmitterSpawnGroup.cpp
@@ -185,6 +185,7 @@ void UNiagaraStackStatelessEmitterSpawnItem::RefreshChildrenInternal(const TArra
 			SpawnInfoObject->RegisterInstancedCustomPropertyTypeLayout(FNiagaraDistributionRangeFloat::StaticStruct()->GetFName(), FOnGetPropertyTypeCustomizationInstance::CreateStatic(&FNiagaraDistributionPropertyCustomization::MakeFloatInstance, StatelessEmitterObject));
 			//SpawnInfoObject->RegisterInstancedCustomPropertyTypeLayout(FNiagaraDistributionRangeVector2::StaticStruct()->GetFName(), FOnGetPropertyTypeCustomizationInstance::CreateStatic(&FNiagaraDistributionPropertyCustomization::MakeVector2Instance));
 			//SpawnInfoObject->RegisterInstancedCustomPropertyTypeLayout(FNiagaraDistributionRangeVector3::StaticStruct()->GetFName(), FOnGetPropertyTypeCustomizationInstance::CreateStatic(&FNiagaraDistributionPropertyCustomization::MakeVector3Instance));
+			//SpawnInfoObject->RegisterInstancedCustomPropertyTypeLayout(FNiagaraDistributionRangeColor::StaticStruct()->GetFName(), FOnGetPropertyTypeCustomizationInstance::CreateStatic(&FNiagaraDistributionPropertyCustomization::MakeColorInstance));
 			SpawnInfoObject->RegisterInstancedCustomPropertyTypeLayout(FNiagaraDistributionRangeInt::StaticStruct()->GetFName(), FOnGetPropertyTypeCustomizationInstance::CreateStatic(&FNiagaraDistributionIntPropertyCustomization::MakeIntInstance, StatelessEmitterObject));
 
 			SpawnInfoObjectWeak = SpawnInfoObject;
diff --git a/Engine/Plugins/FX/Niagara/Source/NiagaraShader/Internal/Stateless/NiagaraStatelessModuleShaderParameters.h b/Engine/Plugins/FX/Niagara/Source/NiagaraShader/Internal/Stateless/NiagaraStatelessModuleShaderParameters.h
--- a/Engine/Plugins/FX/Niagara/Source/NiagaraShader/Internal/Stateless/NiagaraStatelessModuleShaderParameters.h
+++ b/Engine/Plugins/FX/Niagara/Source/NiagaraShader/Internal/Stateless/NiagaraStatelessModuleShaderParameters.h
@@ -11,8 +11,8 @@ namespace NiagaraStateless
 		SHADER_PARAMETER(uint32,		InitializeParticle_ModuleFlags)
 		SHADER_PARAMETER(FUintVector3,	InitializeParticle_InitialPosition)
 
-		SHADER_PARAMETER(FVector4f,	InitializeParticle_ColorScale)				// Unset / Direct Set / Random Range(Link RGBA / Link RGB|Link A / Random Channels) / Random Hue|Saturation|Value
-		SHADER_PARAMETER(FVector4f,	InitializeParticle_ColorBias)
+		SHADER_PARAMETER(FLinearColor,	InitializeParticle_ColorScale)				// Unset / Direct Set / Random Range(Link RGBA / Link RGB|Link A / Random Channels) / Random Hue|Saturation|Value
+		SHADER_PARAMETER(FLinearColor,	InitializeParticle_ColorBias)
 
 		SHADER_PARAMETER(FVector2f, InitializeParticle_SpriteSizeScale)			// Unset / Uniform / Random Uniform / Non-Uniform / Random Non-Uniform
 		SHADER_PARAMETER(FVector2f, InitializeParticle_SpriteSizeBias)
