Add bFilterByObjectType option to the RigidMeshCollisionQuery DI to allow people to try what seems to be a more reliable mechanism for controlling which rigid bodies get added to the query.

diff --git a/Engine/Plugins/FX/Niagara/Source/Niagara/Private/NiagaraDataInterfaceRigidMeshCollisionQuery.cpp b/Engine/Plugins/FX/Niagara/Source/Niagara/Private/NiagaraDataInterfaceRigidMeshCollisionQuery.cpp
--- a/Engine/Plugins/FX/Niagara/Source/Niagara/Private/NiagaraDataInterfaceRigidMeshCollisionQuery.cpp
+++ b/Engine/Plugins/FX/Niagara/Source/Niagara/Private/NiagaraDataInterfaceRigidMeshCollisionQuery.cpp
@@ -1121,10 +1121,10 @@ void UNiagaraDataInterfaceRigidMeshCollisionQuery::DrawDebugHud(FNDIDrawDebugHud
 		{
 			const FVector3f HalfBoxExtent = 0.5f * InstanceData_GT->AssetArrays->ElementExtent[ElementOffsets.BoxOffset + BoxIt];
 			const FMatrix CurrentTransform = GetTransformFromArray(CurrentTransformArray, ElementOffsets.BoxOffset + BoxIt);
-			DrawDebugBox(World, CurrentTransform.TransformPosition(FVector::ZeroVector), FVector(HalfBoxExtent), CurrentTransform.ToQuat(), FColor::Blue);			
+			DrawDebugBox(World, CurrentTransform.TransformPosition(FVector::ZeroVector), FVector(HalfBoxExtent), CurrentTransform.Rotator().Quaternion(), FColor::Blue);			
 
 			const FMatrix PrevTransform = GetTransformFromArray(PreviousTransformArray, ElementOffsets.BoxOffset + BoxIt);
-			DrawDebugBox(World, PrevTransform.TransformPosition(FVector::ZeroVector), FVector(HalfBoxExtent), PrevTransform.ToQuat(), FColor::Red);
+			DrawDebugBox(World, PrevTransform.TransformPosition(FVector::ZeroVector), FVector(HalfBoxExtent), PrevTransform.Rotator().Quaternion(), FColor::Red);
 		}
 
 		// Spheres
@@ -1143,10 +1143,10 @@ void UNiagaraDataInterfaceRigidMeshCollisionQuery::DrawDebugHud(FNDIDrawDebugHud
 		{
 			const FVector2f RadiusLength(InstanceData_GT->AssetArrays->ElementExtent[ElementOffsets.CapsuleOffset + CapsuleIt]);
 			const FMatrix CurrentTransform = GetTransformFromArray(CurrentTransformArray, ElementOffsets.CapsuleOffset + CapsuleIt);
-			DrawDebugCapsule(World, CurrentTransform.TransformPosition(FVector::ZeroVector), RadiusLength.Y * 0.5f, RadiusLength.X, CurrentTransform.ToQuat(), FColor::Blue);
+			DrawDebugCapsule(World, CurrentTransform.TransformPosition(FVector::ZeroVector), RadiusLength.Y * 0.5f, RadiusLength.X, CurrentTransform.Rotator().Quaternion(), FColor::Blue);
 
 			const FMatrix PrevTransform = GetTransformFromArray(PreviousTransformArray, ElementOffsets.CapsuleOffset + CapsuleIt);
-			DrawDebugCapsule(World, PrevTransform.TransformPosition(FVector::ZeroVector), RadiusLength.Y * 0.5f, RadiusLength.X, CurrentTransform.ToQuat(), FColor::Red);
+			DrawDebugCapsule(World, PrevTransform.TransformPosition(FVector::ZeroVector), RadiusLength.Y * 0.5f, RadiusLength.X, CurrentTransform.Rotator().Quaternion(), FColor::Red);
 		}
 
 		if (!InstanceData_GT->ExplicitActors.IsEmpty() || !InstanceData_GT->FoundActors.IsEmpty())
@@ -1173,7 +1173,7 @@ void UNiagaraDataInterfaceRigidMeshCollisionQuery::DrawDebugHud(FNDIDrawDebugHud
 					Actor->GetActorBounds(true, ActorOrigin, ActorBoundsExtent);
 
 					const FMatrix CurrentTransform = FTranslationMatrix(ActorOrigin);
-					DrawDebugBox(World, CurrentTransform.TransformPosition(FVector::ZeroVector), ActorBoundsExtent, CurrentTransform.ToQuat(), FColor::Yellow);
+					DrawDebugBox(World, CurrentTransform.TransformPosition(FVector::ZeroVector), ActorBoundsExtent, FQuat::Identity, FColor::Yellow);
 					if (!ShouldClip(CurrentTransform, FSphere(FVector::ZeroVector, ActorBoundsExtent.Size())))
 					{
 						FString ActorLabel;
@@ -1316,6 +1316,7 @@ bool UNiagaraDataInterfaceRigidMeshCollisionQuery::CopyToInternal(UNiagaraDataIn
 	OtherTyped->SourceActors = SourceActors;
 	OtherTyped->OnlyUseMoveable = OnlyUseMoveable;
 	OtherTyped->UseComplexCollisions = UseComplexCollisions;
+	OtherTyped->bFilterByObjectType = bFilterByObjectType;
 PRAGMA_DISABLE_DEPRECATION_WARNINGS
 	OtherTyped->GlobalSearchAllowed = GlobalSearchAllowed;
 	OtherTyped->GlobalSearchForced = GlobalSearchForced;
@@ -1339,6 +1340,7 @@ bool UNiagaraDataInterfaceRigidMeshCollisionQuery::Equals(const UNiagaraDataInte
 		&& (OtherTyped->SourceActors == SourceActors)
 		&& (OtherTyped->OnlyUseMoveable == OnlyUseMoveable)
 		&& (OtherTyped->UseComplexCollisions == UseComplexCollisions)
+		&& (OtherTyped->bFilterByObjectType == bFilterByObjectType)
 PRAGMA_DISABLE_DEPRECATION_WARNINGS
 		&& (OtherTyped->GlobalSearchAllowed == GlobalSearchAllowed)
 		&& (OtherTyped->GlobalSearchForced == GlobalSearchForced)
@@ -1975,7 +1977,6 @@ bool UNiagaraDataInterfaceRigidMeshCollisionQuery::FindActors(UWorld* World, FND
 		FCollisionQueryParams Params(SCENE_QUERY_STAT(NiagaraRigidMeshCollisionQuery), UseComplexCollisions);
 
 		// it is not clear the best strategy for filtering out results between ByObjectType and ByChannel.  For now we're going to preserve existing behavior
-		constexpr bool bFilterByObjectType = false;
 		if (bFilterByObjectType)
 		{
 			FCollisionObjectQueryParams ObjectParams;
diff --git a/Engine/Plugins/FX/Niagara/Source/Niagara/Public/NiagaraDataInterfaceRigidMeshCollisionQuery.h b/Engine/Plugins/FX/Niagara/Source/Niagara/Public/NiagaraDataInterfaceRigidMeshCollisionQuery.h
--- a/Engine/Plugins/FX/Niagara/Source/Niagara/Public/NiagaraDataInterfaceRigidMeshCollisionQuery.h
+++ b/Engine/Plugins/FX/Niagara/Source/Niagara/Public/NiagaraDataInterfaceRigidMeshCollisionQuery.h
@@ -203,7 +203,11 @@ public:
 	UPROPERTY(EditAnywhere, Category = "Source")
 	bool UseComplexCollisions = false;
 
-	/** If enabled the global search can be executed dependeing on GlobalSearchForced and GlobalSearchFallback_Unscripted */
+	/** If enabled, FindActors will use filtering based on ObjectType instead of Channel. */
+	UPROPERTY(EditAnywhere, Category = "Source")
+	bool bFilterByObjectType = false;
+
+	/** If enabled the global search can be executed depending on GlobalSearchForced and GlobalSearchFallback_Unscripted */
 	UE_DEPRECATED(5.2, "Global search will be deprecated in favor of using a spatial search through the FindActors function.")
 	UPROPERTY(EditAnywhere, AdvancedDisplay, Category = "Source", meta = (DisplayName = "Global Search Allowed"))
 	bool GlobalSearchAllowed = false;
