Niagara - Update the stateless overview and stack UX for consistency and simplicity.
+  Update stateless modules to support an add/delete workflow, although under the hood they're just disabled and hidden from the UI.
+  Change the visibility flags so that individual spawn infos, modules, and renderers show up in the overview.
+  A a flag to the renderer creation infos which indicate whether they work with stateless emitters or not, and then filter out unsupported renderers in the UI.

diff --git a/Engine/Plugins/FX/Niagara/Source/NiagaraEditor/Private/NiagaraEditorModule.cpp b/Engine/Plugins/FX/Niagara/Source/NiagaraEditor/Private/NiagaraEditorModule.cpp
--- a/Engine/Plugins/FX/Niagara/Source/NiagaraEditor/Private/NiagaraEditorModule.cpp
+++ b/Engine/Plugins/FX/Niagara/Source/NiagaraEditor/Private/NiagaraEditorModule.cpp
@@ -2082,9 +2082,11 @@ void FNiagaraEditorModule::GetTargetSystemAndEmitterForDataInterface(UNiagaraDat
 
 void FNiagaraEditorModule::RegisterDefaultRendererFactories()
 {
+	bool bIsSupportedByStateless = true;
 	RegisterRendererCreationInfo(FNiagaraRendererCreationInfo(
 		UNiagaraMeshRendererProperties::StaticClass()->GetDisplayNameText(),
 		FText::FromString(UNiagaraMeshRendererProperties::StaticClass()->GetDescription()),
+		bIsSupportedByStateless,
 		UNiagaraMeshRendererProperties::StaticClass()->GetClassPathName(),
 		FNiagaraRendererCreationInfo::FRendererFactory::CreateLambda([](UObject* OuterEmitter)
 		{
@@ -2101,6 +2103,7 @@ void FNiagaraEditorModule::RegisterDefaultRendererFactories()
 	RegisterRendererCreationInfo(FNiagaraRendererCreationInfo(
 		UNiagaraSpriteRendererProperties::StaticClass()->GetDisplayNameText(),
 		FText::FromString(UNiagaraSpriteRendererProperties::StaticClass()->GetDescription()),
+		bIsSupportedByStateless,
 		UNiagaraSpriteRendererProperties::StaticClass()->GetClassPathName(),
 		FNiagaraRendererCreationInfo::FRendererFactory::CreateLambda([](UObject* OuterEmitter)
 		{
@@ -2113,6 +2116,7 @@ void FNiagaraEditorModule::RegisterDefaultRendererFactories()
 	RegisterRendererCreationInfo(FNiagaraRendererCreationInfo(
 		UNiagaraRibbonRendererProperties::StaticClass()->GetDisplayNameText(),
 		FText::FromString(UNiagaraRibbonRendererProperties::StaticClass()->GetDescription()),
+		bIsSupportedByStateless,
 		UNiagaraRibbonRendererProperties::StaticClass()->GetClassPathName(),
 		FNiagaraRendererCreationInfo::FRendererFactory::CreateLambda([](UObject* OuterEmitter)
 		{
diff --git a/Engine/Plugins/FX/Niagara/Source/NiagaraEditor/Private/NiagaraStackEditorData.cpp b/Engine/Plugins/FX/Niagara/Source/NiagaraEditor/Private/NiagaraStackEditorData.cpp
--- a/Engine/Plugins/FX/Niagara/Source/NiagaraEditor/Private/NiagaraStackEditorData.cpp
+++ b/Engine/Plugins/FX/Niagara/Source/NiagaraEditor/Private/NiagaraStackEditorData.cpp
@@ -271,3 +271,15 @@ const TArray<FString>& UNiagaraStackEditorData::GetDismissedStackIssueIds()
 	return DismissedStackIssueIds;
 }
 
+bool UNiagaraStackEditorData::GetStatelessModuleShowWhenDisabled(const FString& StackEntryKey) const
+{
+	const FNiagaraStatelessModuleEditorData* StatelessModuleEditorData = StackEntryKeyToStatelessModuleEditorData.Find(StackEntryKey);
+	return StatelessModuleEditorData != nullptr && StatelessModuleEditorData->bShowWhenDisabled;
+}
+
+void UNiagaraStackEditorData::SetStatelessModuleShowWhenDisabled(const FString& StackEntryKey, bool bInShowWhenDisabled)
+{
+	FNiagaraStatelessModuleEditorData& StatelessModuleEditorData = StackEntryKeyToStatelessModuleEditorData.FindOrAdd(StackEntryKey);
+	StatelessModuleEditorData.bShowWhenDisabled = bInShowWhenDisabled;
+}
+
diff --git a/Engine/Plugins/FX/Niagara/Source/NiagaraEditor/Private/ViewModels/Stack/NiagaraStackRenderItemGroup.cpp b/Engine/Plugins/FX/Niagara/Source/NiagaraEditor/Private/ViewModels/Stack/NiagaraStackRenderItemGroup.cpp
--- a/Engine/Plugins/FX/Niagara/Source/NiagaraEditor/Private/ViewModels/Stack/NiagaraStackRenderItemGroup.cpp
+++ b/Engine/Plugins/FX/Niagara/Source/NiagaraEditor/Private/ViewModels/Stack/NiagaraStackRenderItemGroup.cpp
@@ -77,7 +77,7 @@ public:
 		const UNiagaraEditorSettings* NiagaraEditorSettings = GetDefault<UNiagaraEditorSettings>();
 		for (const FNiagaraRendererCreationInfo& RendererCreationInfo : RendererCreationInfos)
 		{
-			if (NiagaraEditorSettings->IsAllowedClassPath(RendererCreationInfo.RendererClassPath))
+			if (NiagaraEditorSettings->IsAllowedClassPath(RendererCreationInfo.RendererClassPath) && RenderersOwner->IsRenderCreationInfoSupported(RendererCreationInfo))
 			{
 				OutAddActions.Add(MakeShared<FRenderItemGroupAddAction>(RendererCreationInfo));
 			}
diff --git a/Engine/Plugins/FX/Niagara/Source/NiagaraEditor/Private/ViewModels/Stack/NiagaraStackRenderersOwner.cpp b/Engine/Plugins/FX/Niagara/Source/NiagaraEditor/Private/ViewModels/Stack/NiagaraStackRenderersOwner.cpp
--- a/Engine/Plugins/FX/Niagara/Source/NiagaraEditor/Private/ViewModels/Stack/NiagaraStackRenderersOwner.cpp
+++ b/Engine/Plugins/FX/Niagara/Source/NiagaraEditor/Private/ViewModels/Stack/NiagaraStackRenderersOwner.cpp
@@ -3,6 +3,7 @@
 #include "ViewModels/Stack/NiagaraStackRenderersOwner.h"
 
 #include "Framework/Notifications/NotificationManager.h"
+#include "NiagaraEditorModule.h"
 #include "NiagaraEditorStyle.h"
 #include "NiagaraEmitter.h"
 #include "NiagaraRendererProperties.h"
@@ -180,6 +181,11 @@ void FNiagaraStackRenderersOwnerStateless::GetRenderers(TArray<UNiagaraRendererP
 	}
 }
 
+bool FNiagaraStackRenderersOwnerStateless::IsRenderCreationInfoSupported(const FNiagaraRendererCreationInfo& RendererCreationInfo) const
+{
+	return RendererCreationInfo.bIsSupportedByStateless;
+}
+
 void FNiagaraStackRenderersOwnerStateless::AddRenderer(UNiagaraRendererProperties* RendererToAdd)
 {
 	if (StatelessEmitterWeak.IsValid())
diff --git a/Engine/Plugins/FX/Niagara/Source/NiagaraEditor/Private/ViewModels/Stack/NiagaraStackRenderersOwner.h b/Engine/Plugins/FX/Niagara/Source/NiagaraEditor/Private/ViewModels/Stack/NiagaraStackRenderersOwner.h
--- a/Engine/Plugins/FX/Niagara/Source/NiagaraEditor/Private/ViewModels/Stack/NiagaraStackRenderersOwner.h
+++ b/Engine/Plugins/FX/Niagara/Source/NiagaraEditor/Private/ViewModels/Stack/NiagaraStackRenderersOwner.h
@@ -9,6 +9,7 @@
 
 class FNiagaraEmitterHandleViewModel;
 class FNiagaraEmitterViewModel;
+struct FNiagaraRendererCreationInfo;
 class UNiagaraEmitter;
 class UNiagaraRendererProperties;
 class UNiagaraStatelessEmitter;
@@ -21,6 +22,7 @@ public:
 	virtual bool IsValid() const = 0;
 	virtual UObject* GetOwnerObject() const = 0;
 	virtual void GetRenderers(TArray<UNiagaraRendererProperties*>& OutRenderers) const = 0;
+	virtual bool IsRenderCreationInfoSupported(const FNiagaraRendererCreationInfo& RendererCreationInfo) const = 0;
 
 	virtual void AddRenderer(UNiagaraRendererProperties* RendererToAdd) = 0;
 	virtual void RemoveRenderer(UNiagaraRendererProperties* Renderer) = 0;
@@ -41,6 +43,8 @@ public:
 class FNiagaraStackRenderersOwner : public INiagaraStackRenderersOwner
 {
 public:
+	virtual bool IsRenderCreationInfoSupported(const FNiagaraRendererCreationInfo& RendererCreationInfo) const override { return true; }
+
 	virtual bool HasBaseRenderer(UNiagaraRendererProperties* Renderer) const override { return false; }
 	virtual bool IsRendererDifferentFromBase(UNiagaraRendererProperties* Renderer) const override { return false; }
 	virtual void ResetRendererToBase(UNiagaraRendererProperties* Renderer) const override { }
@@ -92,10 +96,11 @@ public:
 	virtual bool IsValid() const override;
 	virtual UObject* GetOwnerObject() const override;
 	virtual void GetRenderers(TArray<UNiagaraRendererProperties*>& OutRenderers) const override;
+	virtual bool IsRenderCreationInfoSupported(const FNiagaraRendererCreationInfo& RendererCreationInfo) const override;
 	virtual void AddRenderer(UNiagaraRendererProperties* RendererToAdd) override;
 	virtual void RemoveRenderer(UNiagaraRendererProperties* Renderer) override;
 	virtual void MoveRenderer(UNiagaraRendererProperties* Renderer, int32 NewIndex) override;
-	virtual bool ShouldShowRendererItemsInOverview() const override { return false; }
+	virtual bool ShouldShowRendererItemsInOverview() const override { return true; }
 	virtual UNiagaraStackEntry::EIconMode GetSupportedIconMode() const override { return UNiagaraStackEntry::EIconMode::Brush; }
 	virtual const FSlateBrush* GetIconBrush() const override;
 
diff --git a/Engine/Plugins/FX/Niagara/Source/NiagaraEditor/Private/ViewModels/Stack/NiagaraStackStatelessEmitterSimulateGroup.cpp b/Engine/Plugins/FX/Niagara/Source/NiagaraEditor/Private/ViewModels/Stack/NiagaraStackStatelessEmitterSimulateGroup.cpp
--- a/Engine/Plugins/FX/Niagara/Source/NiagaraEditor/Private/ViewModels/Stack/NiagaraStackStatelessEmitterSimulateGroup.cpp
+++ b/Engine/Plugins/FX/Niagara/Source/NiagaraEditor/Private/ViewModels/Stack/NiagaraStackStatelessEmitterSimulateGroup.cpp
@@ -4,6 +4,7 @@
 
 #include "IDetailTreeNode.h"
 #include "NiagaraEditorStyle.h"
+#include "NiagaraStackEditorData.h"
 #include "PropertyHandle.h"
 #include "ScopedTransaction.h"
 #include "Stateless/NiagaraDistributionPropertyCustomization.h"
@@ -15,13 +16,87 @@
 
 #define LOCTEXT_NAMESPACE "NiagaraEmitterStatelessSimulateGroup"
 
+class FNiagaraStatelessEmitterAddModuleAction : public INiagaraStackItemGroupAddAction
+{
+
+public:
+	FNiagaraStatelessEmitterAddModuleAction(UNiagaraStatelessModule* StatelessModule)
+	{
+		StatelessModuleWeak = StatelessModule;
+		DisplayName = StatelessModule->GetClass()->GetDisplayNameText();
+	}
+
+	UNiagaraStatelessModule* GetModule() const { return StatelessModuleWeak.Get(); }
+
+	virtual TArray<FString> GetCategories() const override { return Categories; }
+	virtual FText GetDisplayName() const override { return DisplayName; }
+	virtual FText GetDescription() const override { return FText(); }
+	virtual FText GetKeywords() const override { return FText(); }
+
+private:
+	TWeakObjectPtr<UNiagaraStatelessModule> StatelessModuleWeak;
+	TArray<FString> Categories;
+	FText DisplayName;
+};
+
+class FNiagaraStatelessEmitterSimulateGroupAddUtilities : public TNiagaraStackItemGroupAddUtilities<UNiagaraStatelessModule*>
+{
+public:
+	FNiagaraStatelessEmitterSimulateGroupAddUtilities(UNiagaraStatelessEmitter& StatelessEmitter, UNiagaraStackEditorData& StackEditorData, FOnItemAdded InOnItemAdded)
+		: TNiagaraStackItemGroupAddUtilities<UNiagaraStatelessModule*>(LOCTEXT("ModuleName", "Module"), EAddMode::AddFromAction, true, false, InOnItemAdded)
+	{
+		StatelessEmitterWeak = &StatelessEmitter;
+		StackEditorDataWeak = &StackEditorData;
+	}
+
+	virtual void AddItemDirectly() { unimplemented(); };
+
+	virtual void GenerateAddActions(TArray<TSharedRef<INiagaraStackItemGroupAddAction>>& OutAddActions, const FNiagaraStackItemGroupAddOptions& AddProperties = FNiagaraStackItemGroupAddOptions()) const override
+	{
+		UNiagaraStatelessEmitter* StatelessEmitter = StatelessEmitterWeak.Get();
+		UNiagaraStackEditorData* StackEditorData = StackEditorDataWeak.Get();
+		if (StatelessEmitter != nullptr && StackEditorData != nullptr)
+		{
+			for (UNiagaraStatelessModule* StatelessModule : StatelessEmitter->GetModules())
+			{
+				if (StatelessModule->IsModuleEnabled() == false && 
+					StackEditorData->GetStatelessModuleShowWhenDisabled(UNiagaraStackStatelessModuleItem::GenerateStackEditorDataKey(StatelessModule)) == false)
+				{
+					OutAddActions.Add(MakeShared<FNiagaraStatelessEmitterAddModuleAction>(StatelessModule));
+				}
+			}
+		}
+	}
+
+	virtual void ExecuteAddAction(TSharedRef<INiagaraStackItemGroupAddAction> AddAction, int32 TargetIndex) override
+	{
+		TSharedRef<FNiagaraStatelessEmitterAddModuleAction> AddModuleAction = StaticCastSharedRef<FNiagaraStatelessEmitterAddModuleAction>(AddAction);
+		UNiagaraStatelessModule* StatelessModule = AddModuleAction->GetModule();
+		UNiagaraStackEditorData* StackEditorData = StackEditorDataWeak.Get();
+		if (StatelessModule != nullptr && StackEditorData != nullptr)
+		{
+			FScopedTransaction ScopedTransaction(LOCTEXT("AddStatelessModuleTransaction", "Add module."));
+			StatelessModule->Modify();
+			StatelessModule->SetIsModuleEnabled(true);
+			StatelessModule->PostEditChange();
+			OnItemAdded.ExecuteIfBound(StatelessModule);
+		}
+	}
+
+protected:
+	TWeakObjectPtr<UNiagaraStatelessEmitter> StatelessEmitterWeak;
+	TWeakObjectPtr<UNiagaraStackEditorData> StackEditorDataWeak;
+};
+
 void UNiagaraStackStatelessEmitterSimulateGroup::Initialize(FRequiredEntryData InRequiredEntryData, UNiagaraStatelessEmitter* InStatelessEmitter)
 {
+	AddUtilities = MakeShared<FNiagaraStatelessEmitterSimulateGroupAddUtilities>(*InStatelessEmitter, *InRequiredEntryData.StackEditorData,
+		FNiagaraStatelessEmitterSimulateGroupAddUtilities::FOnItemAdded::CreateUObject(this, &UNiagaraStackStatelessEmitterSimulateGroup::ModuleAdded));
 	Super::Initialize(
 		InRequiredEntryData, 
 		LOCTEXT("EmitterStatelessSimulateGroupDisplayName", "Simulate"),
 		LOCTEXT("EmitterStatelessSimulateGroupToolTip", "Data related to the simulation of the particles"),
-		nullptr);
+		AddUtilities.Get());
 	StatelessEmitterWeak = InStatelessEmitter;
 }
 
@@ -39,25 +114,83 @@ void UNiagaraStackStatelessEmitterSimulateGroup::RefreshChildrenInternal(const T
 	{
 		for (UNiagaraStatelessModule* StatelessModule : StatelessEmitter->GetModules())
 		{
+			FString ModuleStackEditorDataKey = UNiagaraStackStatelessModuleItem::GenerateStackEditorDataKey(StatelessModule);
+			if (StatelessModule->CanDisableModule() && 
+				StatelessModule->IsModuleEnabled() == false &&
+				GetStackEditorData().GetStatelessModuleShowWhenDisabled(ModuleStackEditorDataKey) == false)
+			{
+				// If a module is disabled and doesn't have the show when disabled flag set, filter it from the UI.
+				continue;
+			}
+
 			UNiagaraStackStatelessModuleItem* ModuleItem = FindCurrentChildOfTypeByPredicate<UNiagaraStackStatelessModuleItem>(CurrentChildren,
 				[StatelessModule](const UNiagaraStackStatelessModuleItem* CurrentChild) { return CurrentChild->GetStatelessModule() == StatelessModule; });
 			if (ModuleItem == nullptr)
 			{
 				ModuleItem = NewObject<UNiagaraStackStatelessModuleItem>(this);
 				ModuleItem->Initialize(CreateDefaultChildRequiredData(), StatelessModule);
+				ModuleItem->OnModifiedGroupItems().AddUObject(this, &UNiagaraStackStatelessEmitterSimulateGroup::ModuleModifiedGroupItems);
 			}
 			NewChildren.Add(ModuleItem);
 		}
 	}
 }
 
+void UNiagaraStackStatelessEmitterSimulateGroup::ModuleAdded(UNiagaraStatelessModule* StatelessModule)
+{
+	RefreshChildren();
+}
+
+void UNiagaraStackStatelessEmitterSimulateGroup::ModuleModifiedGroupItems()
+{
+	RefreshChildren();
+}
+
+FString UNiagaraStackStatelessModuleItem::GenerateStackEditorDataKey(const UNiagaraStatelessModule* InStatelessModule)
+{
+	return FString::Printf(TEXT("StatelessModuleItem-%s"), *InStatelessModule->GetName());
+}
+
 void UNiagaraStackStatelessModuleItem::Initialize(FRequiredEntryData InRequiredEntryData, UNiagaraStatelessModule* InStatelessModule)
 {
-	Super::Initialize(InRequiredEntryData, FString::Printf(TEXT("StatelessModuleItem-%s"), *InStatelessModule->GetName()));
+	Super::Initialize(InRequiredEntryData, GenerateStackEditorDataKey(InStatelessModule));
 	StatelessModuleWeak = InStatelessModule;
 	DisplayName = InStatelessModule->GetClass()->GetDisplayNameText();
 }
 
+bool UNiagaraStackStatelessModuleItem::TestCanDeleteWithMessage(FText& OutCanDeleteMessage) const
+{
+	UNiagaraStatelessModule* StatelessModule = StatelessModuleWeak.Get();
+	if (StatelessModule != nullptr)
+	{
+		if (StatelessModule->CanDisableModule())
+		{
+			OutCanDeleteMessage = LOCTEXT("DeleteStatelessModule", "Delete this module.");
+			return true;
+		}
+	}
+	OutCanDeleteMessage = LOCTEXT("DeleteStatelessModuleUnsupported", "This module does not support being deleted.");
+	return false;
+}
+
+FText UNiagaraStackStatelessModuleItem::GetDeleteTransactionText() const
+{
+	return LOCTEXT("DeleteStatelessModule", "Delete module from stateless emitter.");
+}
+
+void UNiagaraStackStatelessModuleItem::Delete()
+{
+	UNiagaraStatelessModule* StatelessModule = StatelessModuleWeak.Get();
+	if (StatelessModule != nullptr && StatelessModule->CanDisableModule())
+	{
+		StatelessModule->Modify();
+		StatelessModule->SetIsModuleEnabled(false);
+		GetStackEditorData().Modify();
+		GetStackEditorData().SetStatelessModuleShowWhenDisabled(GetStackEditorDataKey(), false);
+		OnModifiedGroupItems().Broadcast();
+	}
+}
+
 bool UNiagaraStackStatelessModuleItem::SupportsChangeEnabled() const
 {
 	UNiagaraStatelessModule* StatelessModule = StatelessModuleWeak.Get();
@@ -145,6 +278,7 @@ void UNiagaraStackStatelessModuleItem::SetIsEnabledInternal(bool bInIsEnabled)
 		StatelessModule->Modify();
 		StatelessModule->SetIsModuleEnabled(bInIsEnabled);
 		StatelessModule->PostEditChange();
+		GetStackEditorData().SetStatelessModuleShowWhenDisabled(GetStackEditorDataKey(), true);
 		TArray<UObject*> ChangedObjects = { StatelessModule };
 		OnDataObjectModified().Broadcast(ChangedObjects, ENiagaraDataObjectChange::Changed);
 		RefreshChildren();
diff --git a/Engine/Plugins/FX/Niagara/Source/NiagaraEditor/Public/NiagaraEditorModule.h b/Engine/Plugins/FX/Niagara/Source/NiagaraEditor/Public/NiagaraEditorModule.h
--- a/Engine/Plugins/FX/Niagara/Source/NiagaraEditor/Public/NiagaraEditorModule.h
+++ b/Engine/Plugins/FX/Niagara/Source/NiagaraEditor/Public/NiagaraEditorModule.h
@@ -117,12 +117,23 @@ struct FNiagaraRendererCreationInfo
 	FNiagaraRendererCreationInfo(FText InDisplayName, FText InDescription, const FTopLevelAssetPath& InRendererClassPath, FRendererFactory InFactory) : DisplayName(InDisplayName), Description(InDescription), RendererClassPath(InRendererClassPath), RendererFactory(InFactory)
 	{}
 
+	FNiagaraRendererCreationInfo(FText InDisplayName, bool bInIsSupportedByStateless, const FTopLevelAssetPath& InRendererClassPath, FRendererFactory InFactory) 
+		: DisplayName(InDisplayName), bIsSupportedByStateless(bInIsSupportedByStateless), RendererClassPath(InRendererClassPath), RendererFactory(InFactory)
+	{}
+
+	FNiagaraRendererCreationInfo(FText InDisplayName, FText InDescription, bool bInIsSupportedByStateless, const FTopLevelAssetPath& InRendererClassPath, FRendererFactory InFactory) 
+		: DisplayName(InDisplayName), Description(InDescription), bIsSupportedByStateless(bInIsSupportedByStateless), RendererClassPath(InRendererClassPath), RendererFactory(InFactory)
+	{}
+
 	UPROPERTY()
 	FText DisplayName;
 
 	UPROPERTY()
 	FText Description;
 
+	UPROPERTY()
+	bool bIsSupportedByStateless = false;
+
 	UPROPERTY()
 	FTopLevelAssetPath RendererClassPath;
 	
diff --git a/Engine/Plugins/FX/Niagara/Source/NiagaraEditor/Public/NiagaraStackEditorData.h b/Engine/Plugins/FX/Niagara/Source/NiagaraEditor/Public/NiagaraStackEditorData.h
--- a/Engine/Plugins/FX/Niagara/Source/NiagaraEditor/Public/NiagaraStackEditorData.h
+++ b/Engine/Plugins/FX/Niagara/Source/NiagaraEditor/Public/NiagaraStackEditorData.h
@@ -20,6 +20,16 @@ enum class ENiagaraStackEntryInlineDisplayMode
 	None
 };
 
+/* Defines stack options for stateless modules */
+USTRUCT()
+struct FNiagaraStatelessModuleEditorData
+{
+	GENERATED_BODY()
+
+	UPROPERTY()
+	bool bShowWhenDisabled = false;
+};
+
 /** Editor only UI data for emitters. */
 UCLASS()
 class UNiagaraStackEditorData : public UNiagaraEditorDataBase
@@ -187,6 +197,19 @@ public:
 	/* Gets a reference to the dismissed stack issue array */
 	NIAGARAEDITOR_API const TArray<FString>& GetDismissedStackIssueIds();
 
+	/*
+	* Gets whether or not a stateless module item should be shown when it's disabled.
+	* @param StackEntryKey A unique key for the entry.
+	*/
+	bool GetStatelessModuleShowWhenDisabled(const FString& StackEntryKey) const;
+
+	/*
+	* Sets whether or not a stateless module item should be shown when it's disabled.
+	* @param StackEntryKey A unique key for the entry.
+	* @param bInShowWhenDisabled Whether or not the item should be shown when disabled.
+	*/
+	void SetStatelessModuleShowWhenDisabled(const FString& StackEntryKey, bool bInShowWhenDisabled);
+
 	UPROPERTY()
 	bool bHideDisabledModules = false;
 
@@ -211,6 +234,9 @@ private:
 	UPROPERTY()
 	TMap<FString, FNiagaraStackNoteData> StackNotes;
 
+	UPROPERTY()
+	TMap<FString, FNiagaraStatelessModuleEditorData> StackEntryKeyToStatelessModuleEditorData;
+
 	/* Marking those FTexts explicitly as editoronly_data will make localization not pick these up.
 	 * This is a workaround. EditorDataBase in system & emitter is already flagged as editor only, but it doesn't propagate properly */
 #if WITH_EDITORONLY_DATA
diff --git a/Engine/Plugins/FX/Niagara/Source/NiagaraEditor/Public/ViewModels/Stack/NiagaraStackStatelessEmitterSimulateGroup.h b/Engine/Plugins/FX/Niagara/Source/NiagaraEditor/Public/ViewModels/Stack/NiagaraStackStatelessEmitterSimulateGroup.h
--- a/Engine/Plugins/FX/Niagara/Source/NiagaraEditor/Public/ViewModels/Stack/NiagaraStackStatelessEmitterSimulateGroup.h
+++ b/Engine/Plugins/FX/Niagara/Source/NiagaraEditor/Public/ViewModels/Stack/NiagaraStackStatelessEmitterSimulateGroup.h
@@ -9,6 +9,7 @@
 #include "NiagaraStackStatelessEmitterSimulateGroup.generated.h"
 
 class FNiagaraStackItemPropertyHeaderValue;
+class FNiagaraStatelessEmitterSimulateGroupAddUtilities;
 class UNiagaraStatelessEmitter;
 class UNiagaraStatelessModule;
 class UNiagaraStackObject;
@@ -24,7 +25,6 @@ public:
 	virtual EIconMode GetSupportedIconMode() const override { return EIconMode::Brush; }
 	virtual const FSlateBrush* GetIconBrush() const override;
 
-	virtual bool GetCanExpandInOverview() const override { return false; }
 	virtual bool GetShouldShowInStack() const override { return false; }
 
 	UNiagaraStatelessEmitter* GetStatelessEmitter() const { return StatelessEmitterWeak.Get(); }
@@ -32,8 +32,13 @@ public:
 protected:
 	virtual void RefreshChildrenInternal(const TArray<UNiagaraStackEntry*>& CurrentChildren, TArray<UNiagaraStackEntry*>& NewChildren, TArray<FStackIssue>& NewIssues) override;
 
+private:
+	void ModuleAdded(UNiagaraStatelessModule* StatelessModule);
+	void ModuleModifiedGroupItems();
+
 private:
 	TWeakObjectPtr<UNiagaraStatelessEmitter> StatelessEmitterWeak;
+	TSharedPtr<FNiagaraStatelessEmitterSimulateGroupAddUtilities> AddUtilities;
 };
 
 UCLASS(MinimalAPI)
@@ -42,10 +47,16 @@ class UNiagaraStackStatelessModuleItem : public UNiagaraStackItem
 	GENERATED_BODY()
 
 public:
+	static FString GenerateStackEditorDataKey(const UNiagaraStatelessModule* InStatelessModule);
+
 	void Initialize(FRequiredEntryData InRequiredEntryData, UNiagaraStatelessModule* InStatelessModule);
 
 	virtual FText GetDisplayName() const override { return DisplayName; }
-	virtual bool GetShouldShowInOverview() const override { return false; }
+
+	virtual bool SupportsDelete() const override { return true; }
+	virtual bool TestCanDeleteWithMessage(FText& OutCanDeleteMessage) const override;
+	virtual FText GetDeleteTransactionText() const override;
+	virtual void Delete() override;
 
 	virtual bool SupportsChangeEnabled() const override;
 	virtual bool GetIsEnabled() const override;
diff --git a/Engine/Plugins/FX/Niagara/Source/NiagaraEditor/Public/ViewModels/Stack/NiagaraStackStatelessEmitterSpawnGroup.h b/Engine/Plugins/FX/Niagara/Source/NiagaraEditor/Public/ViewModels/Stack/NiagaraStackStatelessEmitterSpawnGroup.h
--- a/Engine/Plugins/FX/Niagara/Source/NiagaraEditor/Public/ViewModels/Stack/NiagaraStackStatelessEmitterSpawnGroup.h
+++ b/Engine/Plugins/FX/Niagara/Source/NiagaraEditor/Public/ViewModels/Stack/NiagaraStackStatelessEmitterSpawnGroup.h
@@ -26,8 +26,7 @@ public:
 	virtual EIconMode GetSupportedIconMode() const override { return EIconMode::Brush; }
 	virtual const FSlateBrush* GetIconBrush() const override;
 
-	virtual bool GetCanExpandInOverview() const override { return false; }
-	virtual bool GetShouldShowInStack() const override { return true; }
+	virtual bool GetShouldShowInStack() const override { return false; }
 
 	UNiagaraStatelessEmitter* GetStatelessEmitter() const { return StatelessEmitterWeak.Get(); }
 
@@ -55,7 +54,6 @@ public:
 	void Initialize(FRequiredEntryData InRequiredEntryData, UNiagaraStatelessEmitter* InStatelessEmitter, int32 InIndex);
 
 	virtual FText GetDisplayName() const override;
-	virtual bool GetShouldShowInOverview() const override { return false; }
 
 	virtual bool SupportsDelete() const override { return true; }
 	virtual bool TestCanDeleteWithMessage(FText& OutCanDeleteMessage) const override;
