Fix crash in menu caused by transient / garbage NDCs

diff --git a/Engine/Plugins/FX/Niagara/Source/Niagara/Public/NiagaraDataChannel.h b/Engine/Plugins/FX/Niagara/Source/Niagara/Public/NiagaraDataChannel.h
--- a/Engine/Plugins/FX/Niagara/Source/Niagara/Public/NiagaraDataChannel.h
+++ b/Engine/Plugins/FX/Niagara/Source/Niagara/Public/NiagaraDataChannel.h
@@ -183,7 +183,10 @@ void UNiagaraDataChannel::ForEachDataChannel(TAction Func)
 {
 	for(TObjectIterator<UNiagaraDataChannel> It; It; ++It)
 	{
-		if(*It)
+		UNiagaraDataChannel* NDC = *It;
+		if (NDC && 
+			NDC->HasAnyFlags(RF_ClassDefaultObject | RF_Transient) == false &&
+			Cast<UNiagaraDataChannelAsset>(NDC->GetOuter()) != nullptr)
 		{
 			Func(*It);
 		}
diff --git a/Engine/Plugins/FX/Niagara/Source/NiagaraEditor/Private/NiagaraActions.cpp b/Engine/Plugins/FX/Niagara/Source/NiagaraEditor/Private/NiagaraActions.cpp
--- a/Engine/Plugins/FX/Niagara/Source/NiagaraEditor/Private/NiagaraActions.cpp
+++ b/Engine/Plugins/FX/Niagara/Source/NiagaraEditor/Private/NiagaraActions.cpp
@@ -756,37 +756,35 @@ void FNiagaraDataInterfaceNodeActionProvider_DataChannelWrite::GetNodeContextMen
 
 		auto InitForDataChannelSection = [&SubSection, WeakNode](UNiagaraDataChannel* DataChannel)
 		{
-			if (DataChannel->HasAnyFlags(RF_ClassDefaultObject) == false)
+			check(DataChannel);
+
+			TWeakObjectPtr<UNiagaraDataChannel> WeakChannel = DataChannel;
+			auto CreateDataChannelActionEntry = [WeakChannel, WeakNode]()
 			{
-				TWeakObjectPtr<UNiagaraDataChannel> WeakChannel = DataChannel;
-				auto CreateDataChannelActionEntry = [WeakChannel, WeakNode]()
+				UNiagaraDataChannel* Channel = WeakChannel.Get();
+				UNiagaraNodeFunctionCall* Node = const_cast<UNiagaraNodeFunctionCall*>(WeakNode.Get());
+				if (Channel && Node)
 				{
-					UNiagaraDataChannel* Channel = WeakChannel.Get();
-					UNiagaraNodeFunctionCall* Node = const_cast<UNiagaraNodeFunctionCall*>(WeakNode.Get());
-					if(Channel && Node)
+					Node->RemoveAllDynamicPins();
+					TConstArrayView<FNiagaraDataChannelVariable> ChannelVars = Channel->GetVariables();
+					for (const FNiagaraDataChannelVariable& Var : ChannelVars)
 					{
-						Node->RemoveAllDynamicPins();
-						TConstArrayView<FNiagaraDataChannelVariable> ChannelVars = Channel->GetVariables();
-						for (const FNiagaraDataChannelVariable& Var : ChannelVars)
+						FNiagaraTypeDefinition Type = Var.GetType();
+						if (Type.IsEnum() == false)
 						{
-							FNiagaraTypeDefinition Type = Var.GetType();
-							if (Type.IsEnum() == false)
-							{
-								Type = FNiagaraTypeDefinition(FNiagaraTypeHelper::GetSWCStruct(Var.GetType().GetScriptStruct()));
-							}
-							FNiagaraVariable SWCVar(Type, Var.GetName());
-							Node->AddParameter(SWCVar, EEdGraphPinDirection::EGPD_Input);
+							Type = FNiagaraTypeDefinition(FNiagaraTypeHelper::GetSWCStruct(Var.GetType().GetScriptStruct()));
 						}
+						FNiagaraVariable SWCVar(Type, Var.GetName());
+						Node->AddParameter(SWCVar, EEdGraphPinDirection::EGPD_Input);
 					}
-				};
-				SubSection.AddMenuEntry(
-					NAME_None,
-					FText::FromString(DataChannel->GetAsset()->GetName()),
-					FText::FromString(DataChannel->GetAsset()->GetName()),
-					FSlateIcon(FAppStyle::GetAppStyleSetName(), "Icons.Import"),
-					FUIAction(FExecuteAction::CreateLambda(CreateDataChannelActionEntry), FCanExecuteAction()));
-
-			}
+				}
+			};
+			SubSection.AddMenuEntry(
+				NAME_None,
+				FText::FromString(DataChannel->GetAsset()->GetName()),
+				FText::FromString(DataChannel->GetAsset()->GetName()),
+				FSlateIcon(FAppStyle::GetAppStyleSetName(), "Icons.Import"),
+				FUIAction(FExecuteAction::CreateLambda(CreateDataChannelActionEntry), FCanExecuteAction()));
 		};
 
 		UNiagaraDataChannel::ForEachDataChannel(InitForDataChannelSection);
@@ -915,34 +913,33 @@ void FNiagaraDataInterfaceNodeActionProvider_DataChannelRead::AddDataChannelInit
 	FToolMenuSection& MenuSection = ToolMenu->AddSection("DataChannelRead", NiagaraActionsLocal::InitForDataChannelHeaderText);
 	auto InitForDataChannelSection = [&MenuSection, WeakNode](UNiagaraDataChannel* DataChannel)
 	{
-		if (DataChannel->HasAnyFlags(RF_ClassDefaultObject) == false)
+		check(DataChannel);
+
+		TWeakObjectPtr<UNiagaraDataChannel> WeakChannel = DataChannel;
+		auto CreateDataChannelActionEntry = [WeakChannel, WeakNode]()
 		{
-			TWeakObjectPtr<UNiagaraDataChannel> WeakChannel = DataChannel;
-			auto CreateDataChannelActionEntry = [WeakChannel, WeakNode]()
-			{
-				UNiagaraDataChannel* Channel = WeakChannel.Get();
-				UNiagaraNodeFunctionCall* Node = const_cast<UNiagaraNodeFunctionCall*>(WeakNode.Get());
+			UNiagaraDataChannel* Channel = WeakChannel.Get();
+			UNiagaraNodeFunctionCall* Node = const_cast<UNiagaraNodeFunctionCall*>(WeakNode.Get());
 
-				if (Channel && Node)
+			if (Channel && Node)
+			{
+				Node->RemoveAllDynamicPins();
+				TConstArrayView<FNiagaraDataChannelVariable> ChannelVars = Channel->GetVariables();
+				for (const FNiagaraDataChannelVariable& Var : ChannelVars)
 				{
-					Node->RemoveAllDynamicPins();
-					TConstArrayView<FNiagaraDataChannelVariable> ChannelVars = Channel->GetVariables();
-					for (const FNiagaraDataChannelVariable& Var : ChannelVars)
+					FNiagaraTypeDefinition Type = Var.GetType();
+					if (Type.IsEnum() == false)
 					{
-						FNiagaraTypeDefinition Type = Var.GetType();
-						if (Type.IsEnum() == false)
-						{
-							Type = FNiagaraTypeDefinition(FNiagaraTypeHelper::GetSWCStruct(Var.GetType().GetScriptStruct()));
-						}
-						FNiagaraVariable SWCVar(Type, Var.GetName());
-						Node->AddParameter(SWCVar, EEdGraphPinDirection::EGPD_Output);
+						Type = FNiagaraTypeDefinition(FNiagaraTypeHelper::GetSWCStruct(Var.GetType().GetScriptStruct()));
 					}
+					FNiagaraVariable SWCVar(Type, Var.GetName());
+					Node->AddParameter(SWCVar, EEdGraphPinDirection::EGPD_Output);
 				}
-			};
+			}
+		};
 
-			MenuSection.AddMenuEntry(NAME_None, FText::FromString(DataChannel->GetAsset()->GetName()), FText::FromString(DataChannel->GetAsset()->GetName()), FSlateIcon(FAppStyle::GetAppStyleSetName(), "Icons.Import"),
-				FUIAction(FExecuteAction::CreateLambda(CreateDataChannelActionEntry), FCanExecuteAction()));
-		}
+		MenuSection.AddMenuEntry(NAME_None, FText::FromString(DataChannel->GetAsset()->GetName()), FText::FromString(DataChannel->GetAsset()->GetName()), FSlateIcon(FAppStyle::GetAppStyleSetName(), "Icons.Import"),
+			FUIAction(FExecuteAction::CreateLambda(CreateDataChannelActionEntry), FCanExecuteAction()));
 	};
 
 	UNiagaraDataChannel::ForEachDataChannel(InitForDataChannelSection);
