Lazy initialization of NDC handlers. Avoids issue where they're not initialized if used on the first frame after loading.
Optionally block pending async loads when the Island NDC is used.

diff --git a/Engine/Plugins/FX/Niagara/Source/Niagara/Private/NiagaraDataChannelManager.cpp b/Engine/Plugins/FX/Niagara/Source/Niagara/Private/NiagaraDataChannelManager.cpp
--- a/Engine/Plugins/FX/Niagara/Source/Niagara/Private/NiagaraDataChannelManager.cpp
+++ b/Engine/Plugins/FX/Niagara/Source/Niagara/Private/NiagaraDataChannelManager.cpp
@@ -11,6 +11,15 @@ DECLARE_CYCLE_STAT(TEXT("FNiagaraDataChannelManager::BeginFrame"), STAT_DataChan
 DECLARE_CYCLE_STAT(TEXT("FNiagaraDataChannelManager::EndFrame"), STAT_DataChannelManager_EndFrame, STATGROUP_NiagaraDataChannels);
 DECLARE_CYCLE_STAT(TEXT("FNiagaraDataChannelManager::Tick"), STAT_DataChannelManager_Tick, STATGROUP_NiagaraDataChannels);
 
+
+static int GNDCAllowLazyHandlerInit = 1;
+static FAutoConsoleVariableRef CVarNDCAllowLazyHandlerInit(
+	TEXT("fx.Niagara.DataChannels.AllowLazyHandlerInit"),
+	GNDCAllowLazyHandlerInit,
+	TEXT("True if we allow lazy initialization of NDC handlers."),
+	ECVF_Default
+);
+
 FNiagaraDataChannelManager::FNiagaraDataChannelManager(FNiagaraWorldManager* InWorldMan)
 	: WorldMan(InWorldMan)
 {
@@ -30,10 +39,7 @@ void FNiagaraDataChannelManager::RefreshDataChannels()
 
 	UNiagaraDataChannel::ForEachDataChannel([&](UNiagaraDataChannel* DataChannel)
 	{
-		if (DataChannel->HasAnyFlags(RF_ClassDefaultObject) == false)
-		{
-			InitDataChannel(DataChannel, false);
-		}
+		InitDataChannel(DataChannel, false);
 	});
 }
 
@@ -44,10 +50,7 @@ void FNiagaraDataChannelManager::Init()
 	//Initialize any existing data channels, more may be initialized later as they are loaded.
 	UNiagaraDataChannel::ForEachDataChannel([&](UNiagaraDataChannel* DataChannel)
 	{
-		if (DataChannel->HasAnyFlags(RF_ClassDefaultObject) == false)
-		{
-			InitDataChannel(DataChannel, true);
-		}
+		InitDataChannel(DataChannel, true);
 	});
 }
 
@@ -113,14 +116,28 @@ void FNiagaraDataChannelManager::Tick(float DeltaSeconds, ETickingGroup TickGrou
 UNiagaraDataChannelHandler* FNiagaraDataChannelManager::FindDataChannelHandler(const UNiagaraDataChannel* Channel)
 {
 	check(IsInGameThread());
-	if (TObjectPtr<UNiagaraDataChannelHandler>* Found = Channels.Find(Channel))
+	if(Channel == nullptr)
 	{
-		return (*Found).Get();
+		return nullptr;
+	}
+
+	if(GNDCAllowLazyHandlerInit != 0)
+	{
+		//Try to lazy init the channel.
+		//It's possible we've just loaded the channel and it's pending init.
+		return InitDataChannel(Channel, false);
+	}
+	else
+	{
+		if (TObjectPtr<UNiagaraDataChannelHandler>* Found = Channels.Find(Channel))
+		{
+			return (*Found).Get();
+		}
+		return nullptr;
 	}
-	return nullptr;
 }
 
-void FNiagaraDataChannelManager::InitDataChannel(const UNiagaraDataChannel* InChannel, bool bForce)
+UNiagaraDataChannelHandler* FNiagaraDataChannelManager::InitDataChannel(const UNiagaraDataChannel* InChannel, bool bForce)
 {
 	check(IsInGameThread());
 	UWorld* World = GetWorld();
@@ -132,7 +149,9 @@ void FNiagaraDataChannelManager::InitDataChannel(const UNiagaraDataChannel* InCh
 		{
 			Handler = InChannel->CreateHandler(WorldMan->GetWorld());
 		}
+		return Handler;
 	}
+	return nullptr;
 }
 
 void FNiagaraDataChannelManager::RemoveDataChannel(const UNiagaraDataChannel* InChannel)
diff --git a/Engine/Plugins/FX/Niagara/Source/Niagara/Private/NiagaraDataChannelManager.h b/Engine/Plugins/FX/Niagara/Source/Niagara/Private/NiagaraDataChannelManager.h
--- a/Engine/Plugins/FX/Niagara/Source/Niagara/Private/NiagaraDataChannelManager.h
+++ b/Engine/Plugins/FX/Niagara/Source/Niagara/Private/NiagaraDataChannelManager.h
@@ -24,7 +24,7 @@ public:
 	void Init();
 	void Cleanup();
 
-	void InitDataChannel(const UNiagaraDataChannel* Channel, bool bForce);
+	UNiagaraDataChannelHandler* InitDataChannel(const UNiagaraDataChannel* Channel, bool bForce);
 	void RemoveDataChannel(const UNiagaraDataChannel* Channel);
 
 	void BeginFrame(float DeltaSeconds);
diff --git a/Engine/Plugins/FX/Niagara/Source/Niagara/Private/NiagaraDataChannel_Islands.cpp b/Engine/Plugins/FX/Niagara/Source/Niagara/Private/NiagaraDataChannel_Islands.cpp
--- a/Engine/Plugins/FX/Niagara/Source/Niagara/Private/NiagaraDataChannel_Islands.cpp
+++ b/Engine/Plugins/FX/Niagara/Source/Niagara/Private/NiagaraDataChannel_Islands.cpp
@@ -14,6 +14,15 @@
 
 //////////////////////////////////////////////////////////////////////////
 
+namespace NDCIslands_Local
+{
+	int32 AllowAsyncLoad = 1;
+	static FAutoConsoleVariableRef CVarAllowAsyncLoad(TEXT("fx.Niagara.DataChannels.AllowAsyncLoad"), AllowAsyncLoad, TEXT("True if we should attempt to load systems etc asynchronosly."), ECVF_Default);
+
+	int32 BlockAsyncLoadOnUse = 1;
+	static FAutoConsoleVariableRef CVarBlockAsyncLoadOnUse(TEXT("fx.Niagara.DataChannels.BlockAsyncLoadOnUse"), BlockAsyncLoadOnUse, TEXT("True if we should block on any pending async loads when those assets are used."), ECVF_Default);
+}
+
 void UNiagaraDataChannel_Islands::PostLoad()
 {
 	Super::PostLoad();
@@ -26,6 +35,12 @@ void UNiagaraDataChannel_Islands::PostLoad()
 TConstArrayView<TObjectPtr<UNiagaraSystem>> UNiagaraDataChannel_Islands::GetSystems()const
 {
 	check(IsInGameThread());
+	if(NDCIslands_Local::BlockAsyncLoadOnUse && AsyncLoadHandle && AsyncLoadHandle->IsActive())
+	{
+		AsyncLoadHandle->WaitUntilComplete();
+		PostLoadSystems();
+	}
+
 	return SystemsInternal;
 }
 
@@ -44,9 +59,9 @@ void UNiagaraDataChannel_Islands::AsyncLoadSystems()const
 
 		if (Requests.Num() > 0)
 		{
-			if (UAssetManager::IsInitialized())
+			if (NDCIslands_Local::AllowAsyncLoad && UAssetManager::IsInitialized())
 			{
-				UAssetManager::GetStreamableManager().RequestAsyncLoad(Requests
+				AsyncLoadHandle = UAssetManager::GetStreamableManager().RequestAsyncLoad(Requests
 					, FStreamableDelegate::CreateUObject(this, &UNiagaraDataChannel_Islands::PostLoadSystems)
 					, FStreamableManager::AsyncLoadHighPriority
 					, false
@@ -77,6 +92,8 @@ void UNiagaraDataChannel_Islands::PostLoadSystems()const
 	{
 		SystemsInternal.Add(SoftSys.Get());
 	}
+
+	AsyncLoadHandle.Reset();
 }
 
 //////////////////////////////////////////////////////////////////////////
diff --git a/Engine/Plugins/FX/Niagara/Source/Niagara/Public/NiagaraDataChannel_Islands.h b/Engine/Plugins/FX/Niagara/Source/Niagara/Public/NiagaraDataChannel_Islands.h
--- a/Engine/Plugins/FX/Niagara/Source/Niagara/Public/NiagaraDataChannel_Islands.h
+++ b/Engine/Plugins/FX/Niagara/Source/Niagara/Public/NiagaraDataChannel_Islands.h
@@ -6,6 +6,8 @@
 #include "NiagaraDataChannelHandler.h"
 #include "NiagaraDataChannel_Islands.generated.h"
 
+struct FStreamableHandle;
+
 UENUM()
 enum class ENiagraDataChannel_IslandMode : uint8
 {
@@ -183,6 +185,8 @@ protected:
 
 	UPROPERTY(Transient)
 	mutable TArray<TObjectPtr<UNiagaraSystem>> SystemsInternal;
+
+	mutable TSharedPtr<FStreamableHandle> AsyncLoadHandle;
 };
 
 UCLASS(Experimental, BlueprintType, MinimalAPI)
