Update constant resolver to include emitter information when stepping through the emitter node to it's subgraph.  For the base compilation mode this finds the emitter for creating a new constant resolver, while the async task path incorpoates child scopes in the system constant resolver.

diff --git a/Engine/Plugins/FX/Niagara/Source/NiagaraEditor/Private/NiagaraCompilationCopy.cpp b/Engine/Plugins/FX/Niagara/Source/NiagaraEditor/Private/NiagaraCompilationCopy.cpp
--- a/Engine/Plugins/FX/Niagara/Source/NiagaraEditor/Private/NiagaraCompilationCopy.cpp
+++ b/Engine/Plugins/FX/Niagara/Source/NiagaraEditor/Private/NiagaraCompilationCopy.cpp
@@ -499,7 +499,7 @@ UNiagaraDataInterface* FNiagaraCompilationCopyData::GetDuplicatedDataInterfaceCD
 	return nullptr;
 }
 
-void FNiagaraCompilationCopyData::InstantiateCompilationCopy(const FNiagaraCompilationGraphDigested& SourceGraph, const FNiagaraPrecompileData* PrecompileData, ENiagaraScriptUsage InUsage, FNiagaraFixedConstantResolver ConstantResolver)
+void FNiagaraCompilationCopyData::InstantiateCompilationCopy(const FNiagaraCompilationGraphDigested& SourceGraph, const FNiagaraPrecompileData* PrecompileData, ENiagaraScriptUsage InUsage, const FNiagaraFixedConstantResolver& ConstantResolver)
 {
 	InstantiatedGraph = SourceGraph.Instantiate(PrecompileData, this, ValidUsages, ConstantResolver);
 
@@ -509,7 +509,7 @@ void FNiagaraCompilationCopyData::InstantiateCompilationCopy(const FNiagaraCompi
 	}
 }
 
-void FNiagaraCompilationCopyData::CreateParameterMapHistory(const FNiagaraSystemCompilationTask& CompilationTask, const TArray<FNiagaraVariable>& EncounterableVariables, const TArray<FNiagaraVariable>& InStaticVariables, FNiagaraFixedConstantResolver ConstantResolver, TConstArrayView<FNiagaraSimulationStageInfo> SimStages)
+void FNiagaraCompilationCopyData::CreateParameterMapHistory(const FNiagaraSystemCompilationTask& CompilationTask, const TArray<FNiagaraVariable>& EncounterableVariables, const TArray<FNiagaraVariable>& InStaticVariables, const FNiagaraFixedConstantResolver& ConstantResolver, TConstArrayView<FNiagaraSimulationStageInfo> SimStages)
 {
 	FNiagaraEditorModule& NiagaraEditorModule = FModuleManager::GetModuleChecked<FNiagaraEditorModule>("NiagaraEditor");
 	PrecompiledHistories.Empty();
diff --git a/Engine/Plugins/FX/Niagara/Source/NiagaraEditor/Private/NiagaraCompilationPrivate.h b/Engine/Plugins/FX/Niagara/Source/NiagaraEditor/Private/NiagaraCompilationPrivate.h
--- a/Engine/Plugins/FX/Niagara/Source/NiagaraEditor/Private/NiagaraCompilationPrivate.h
+++ b/Engine/Plugins/FX/Niagara/Source/NiagaraEditor/Private/NiagaraCompilationPrivate.h
@@ -46,6 +46,9 @@ public:
 
 	ENiagaraFunctionDebugState GetDebugState() const;
 
+	void AddNamedChildResolver(FName ScopeName, const FNiagaraFixedConstantResolver& ChildResolver);
+	const FNiagaraFixedConstantResolver* FindChildResolver(FName ScopeName) const;
+
 private:
 	void InitConstants();
 	void SetScriptUsage(ENiagaraScriptUsage ScriptUsage);
@@ -66,6 +69,9 @@ private:
 	};
 
 	TArray<FNiagaraVariable, TFixedAllocator<(uint8)EResolvedConstant::Count>> ResolvedConstants;
+
+	using FNamedResolverPair = TTuple<FName, FNiagaraFixedConstantResolver>;
+	TArray<FNamedResolverPair> ChildResolversByName;
 };
 
 struct FNiagaraSimulationStageInfo
@@ -178,8 +184,8 @@ public:
 	TArray<FParameterMapHistory>& GetPrecomputedHistories() { return PrecompiledHistories; }
 	const TArray<FParameterMapHistory>& GetPrecomputedHistories() const { return PrecompiledHistories; }
 
-	void InstantiateCompilationCopy(const FNiagaraCompilationGraphDigested& SourceGraph, const FNiagaraPrecompileData* PrecompileData, ENiagaraScriptUsage InUsage, FNiagaraFixedConstantResolver ConstantResolver);
-	void CreateParameterMapHistory(const FNiagaraSystemCompilationTask& CompilationTask, const TArray<FNiagaraVariable>& EncounterableVariables, const TArray<FNiagaraVariable>& InStaticVariables, FNiagaraFixedConstantResolver ConstantResolver, TConstArrayView<FNiagaraSimulationStageInfo> SimStages);
+	void InstantiateCompilationCopy(const FNiagaraCompilationGraphDigested& SourceGraph, const FNiagaraPrecompileData* PrecompileData, ENiagaraScriptUsage InUsage, const FNiagaraFixedConstantResolver& ConstantResolver);
+	void CreateParameterMapHistory(const FNiagaraSystemCompilationTask& CompilationTask, const TArray<FNiagaraVariable>& EncounterableVariables, const TArray<FNiagaraVariable>& InStaticVariables, const FNiagaraFixedConstantResolver& ConstantResolver, TConstArrayView<FNiagaraSimulationStageInfo> SimStages);
 
 	int32 GetDependentRequestCount() const { return EmitterData.Num(); }
 	FSharedCompilationCopy GetDependentRequest(int32 Index) { return EmitterData[Index]; }
diff --git a/Engine/Plugins/FX/Niagara/Source/NiagaraEditor/Private/NiagaraCompilationTasks.cpp b/Engine/Plugins/FX/Niagara/Source/NiagaraEditor/Private/NiagaraCompilationTasks.cpp
--- a/Engine/Plugins/FX/Niagara/Source/NiagaraEditor/Private/NiagaraCompilationTasks.cpp
+++ b/Engine/Plugins/FX/Niagara/Source/NiagaraEditor/Private/NiagaraCompilationTasks.cpp
@@ -1374,6 +1374,9 @@ void FNiagaraSystemCompilationTask::DigestSystemInfo()
 				EmitterInfo.SourceEmitterIndex = SourceEmitterIndex;
 				EmitterInfo.DigestedEmitterIndex = DigestedEmitterIndex;
 				EmitterInfo.SourceGraph = DigestDatabase.CreateGraphDigest(EmitterGraph, ChangeIdBuilder);
+
+				// be sure to incorporate our constant resolver into the top level SystemInfo.ConstantResolver
+				SystemInfo.ConstantResolver.AddNamedChildResolver(*EmitterInfo.UniqueEmitterName, EmitterInfo.ConstantResolver);
 			}
 
 			{
diff --git a/Engine/Plugins/FX/Niagara/Source/NiagaraEditor/Private/NiagaraGraphDigest.cpp b/Engine/Plugins/FX/Niagara/Source/NiagaraEditor/Private/NiagaraGraphDigest.cpp
--- a/Engine/Plugins/FX/Niagara/Source/NiagaraEditor/Private/NiagaraGraphDigest.cpp
+++ b/Engine/Plugins/FX/Niagara/Source/NiagaraEditor/Private/NiagaraGraphDigest.cpp
@@ -2027,6 +2027,7 @@ FNiagaraCompilationNodeEmitter::FNiagaraCompilationNodeEmitter(const UNiagaraNod
 	EmitterName = InNode->GetName();
 	EmitterPathName = InNode->GetPathName();
 	EmitterHandleIdString = InNode->GetEmitterHandleId().ToString(EGuidFormats::Digits);
+	EmitterUniqueFName = *EmitterUniqueName;
 }
 
 
@@ -2037,6 +2038,7 @@ FNiagaraCompilationNodeEmitter::FNiagaraCompilationNodeEmitter(const FNiagaraCom
 	, EmitterName(InNode.EmitterName)
 	, EmitterPathName(InNode.EmitterPathName)
 	, EmitterHandleIdString(InNode.EmitterHandleIdString)
+	, EmitterUniqueFName(InNode.EmitterUniqueFName)
 	, Usage(InNode.Usage)
 {
 	// we need to replace the CalledGraph here with the graph that has been instantiated already
@@ -2103,7 +2105,15 @@ void FNiagaraCompilationNodeEmitter::BuildParameterMapHistory(FParameterMapHisto
 
 			// Build up a new parameter map history with all the child graph nodes..
 			FParameterMapHistoryBuilder ChildBuilder;
-			*ChildBuilder.ConstantResolver = *Builder.ConstantResolver;
+			const FNiagaraFixedConstantResolver* ChildConstantResolver = Builder.ConstantResolver->FindChildResolver(EmitterUniqueFName);
+			if (ensure(ChildConstantResolver))
+			{
+				*ChildBuilder.ConstantResolver = *ChildConstantResolver;
+			}
+			else
+			{
+				*ChildBuilder.ConstantResolver = *Builder.ConstantResolver;
+			}
 			ChildBuilder.RegisterEncounterableVariables(Builder.GetEncounterableVariables());
 			ChildBuilder.EnableScriptAllowList(true, Usage);
 
diff --git a/Engine/Plugins/FX/Niagara/Source/NiagaraEditor/Private/NiagaraGraphDigest.h b/Engine/Plugins/FX/Niagara/Source/NiagaraEditor/Private/NiagaraGraphDigest.h
--- a/Engine/Plugins/FX/Niagara/Source/NiagaraEditor/Private/NiagaraGraphDigest.h
+++ b/Engine/Plugins/FX/Niagara/Source/NiagaraEditor/Private/NiagaraGraphDigest.h
@@ -398,6 +398,7 @@ public:
 	FString EmitterName;
 	FString EmitterPathName;
 	FString EmitterHandleIdString;
+	FName EmitterUniqueFName;
 	FNiagaraCompilationGraph::FSharedPtr CalledGraph;
 	ENiagaraScriptUsage Usage;
 };
diff --git a/Engine/Plugins/FX/Niagara/Source/NiagaraEditor/Private/NiagaraNodeEmitter.cpp b/Engine/Plugins/FX/Niagara/Source/NiagaraEditor/Private/NiagaraNodeEmitter.cpp
--- a/Engine/Plugins/FX/Niagara/Source/NiagaraEditor/Private/NiagaraNodeEmitter.cpp
+++ b/Engine/Plugins/FX/Niagara/Source/NiagaraEditor/Private/NiagaraNodeEmitter.cpp
@@ -342,6 +342,15 @@ void UNiagaraNodeEmitter::BuildParameterMapHistory(FNiagaraParameterMapHistoryBu
 		}
 	}
 
+	const FNiagaraEmitterHandle* EmitterHandle = nullptr;
+	if (const UNiagaraSystem* System = OutHistory.ConstantResolver->GetSystem())
+	{
+		EmitterHandle = System->GetEmitterHandles().FindByPredicate([this](const FNiagaraEmitterHandle& Handle) -> bool
+		{
+			return Handle.GetId() == EmitterHandleId;
+		});
+	}
+
 	FString EmitterUniqueName = GetEmitterUniqueName();
 	UNiagaraGraph* Graph = GetCalledGraph();
 	if (Graph && ParamMapIdx != INDEX_NONE && OutHistory.bShouldBuildSubHistories)
@@ -366,7 +375,14 @@ void UNiagaraNodeEmitter::BuildParameterMapHistory(FNiagaraParameterMapHistoryBu
 
 			// Build up a new parameter map history with all the child graph nodes..
 			FNiagaraParameterMapHistoryBuilder ChildBuilder;
-			*ChildBuilder.ConstantResolver = *OutHistory.ConstantResolver;
+			if (ensure(EmitterHandle))
+			{
+				*ChildBuilder.ConstantResolver = OutHistory.ConstantResolver->AsEmitter(EmitterHandle->GetInstance()).WithUsage(OutputNodeUsage);
+			}
+			else
+			{
+				*ChildBuilder.ConstantResolver = *OutHistory.ConstantResolver;
+			}
 			ChildBuilder.RegisterEncounterableVariables(OutHistory.GetEncounterableVariables());
 			ChildBuilder.EnableScriptAllowList(true, GetUsage());
 
diff --git a/Engine/Plugins/FX/Niagara/Source/NiagaraEditor/Private/NiagaraParameterMapHistory.cpp b/Engine/Plugins/FX/Niagara/Source/NiagaraEditor/Private/NiagaraParameterMapHistory.cpp
--- a/Engine/Plugins/FX/Niagara/Source/NiagaraEditor/Private/NiagaraParameterMapHistory.cpp
+++ b/Engine/Plugins/FX/Niagara/Source/NiagaraEditor/Private/NiagaraParameterMapHistory.cpp
@@ -2571,6 +2571,15 @@ FCompileConstantResolver FCompileConstantResolver::WithUsage(ENiagaraScriptUsage
 	return Copy;
 }
 
+FCompileConstantResolver FCompileConstantResolver::AsEmitter(const FVersionedNiagaraEmitter& InEmitter) const
+{
+	FCompileConstantResolver Copy = *this;
+	Copy.Emitter = InEmitter;
+	Copy.System = nullptr;
+	return Copy;
+}
+
+
 uint32 FCompileConstantResolver::BuildTypeHash() const
 {
 	union
diff --git a/Engine/Plugins/FX/Niagara/Source/NiagaraEditor/Private/NiagaraTraversalStateContext.cpp b/Engine/Plugins/FX/Niagara/Source/NiagaraEditor/Private/NiagaraTraversalStateContext.cpp
--- a/Engine/Plugins/FX/Niagara/Source/NiagaraEditor/Private/NiagaraTraversalStateContext.cpp
+++ b/Engine/Plugins/FX/Niagara/Source/NiagaraEditor/Private/NiagaraTraversalStateContext.cpp
@@ -332,3 +332,21 @@ ENiagaraFunctionDebugState FNiagaraFixedConstantResolver::GetDebugState() const
 	FNiagaraInt32 EnumValue = ResolvedConstants[(uint8)EResolvedConstant::FunctionDebugState].GetValue<FNiagaraInt32>();
 	return (ENiagaraFunctionDebugState)EnumValue.Value;
 }
+
+void FNiagaraFixedConstantResolver::AddNamedChildResolver(FName ScopeName, const FNiagaraFixedConstantResolver& ChildResolver)
+{
+	if (ensure(FindChildResolver(ScopeName) == nullptr))
+	{
+		ChildResolversByName.Emplace(ScopeName, ChildResolver);
+	}
+}
+
+const FNiagaraFixedConstantResolver* FNiagaraFixedConstantResolver::FindChildResolver(FName ScopeName) const
+{
+	const FNamedResolverPair* ChildResolver = ChildResolversByName.FindByPredicate([ScopeName](const FNamedResolverPair& NamedPair) -> bool
+	{
+			return NamedPair.Key == ScopeName;
+	});
+
+	return ChildResolver ? &ChildResolver->Value : nullptr;
+}
diff --git a/Engine/Plugins/FX/Niagara/Source/NiagaraEditor/Public/NiagaraParameterMapHistory.h b/Engine/Plugins/FX/Niagara/Source/NiagaraEditor/Public/NiagaraParameterMapHistory.h
--- a/Engine/Plugins/FX/Niagara/Source/NiagaraEditor/Public/NiagaraParameterMapHistory.h
+++ b/Engine/Plugins/FX/Niagara/Source/NiagaraEditor/Public/NiagaraParameterMapHistory.h
@@ -40,6 +40,7 @@ public:
 	ENiagaraFunctionDebugState CalculateDebugState() const;
 	FCompileConstantResolver WithDebugState(ENiagaraFunctionDebugState InDebugState) const;
 	FCompileConstantResolver WithUsage(ENiagaraScriptUsage ScriptUsage) const;
+	FCompileConstantResolver AsEmitter(const FVersionedNiagaraEmitter& InEmitter) const;
 
 	// returns a hash of the data that is used when resolving constants using this resolver.  The specific
 	// Emitter/System/Translator shouldn't be included but the values that it might use from those objects
