- Add missing shader types to Niagara script stats

#perflight 65cf94d06328b17dc95fcf57

diff --git a/Engine/Plugins/FX/Niagara/Source/NiagaraEditor/Private/ViewModels/NiagaraScriptStatsViewModel.cpp b/Engine/Plugins/FX/Niagara/Source/NiagaraEditor/Private/ViewModels/NiagaraScriptStatsViewModel.cpp
--- a/Engine/Plugins/FX/Niagara/Source/NiagaraEditor/Private/ViewModels/NiagaraScriptStatsViewModel.cpp
+++ b/Engine/Plugins/FX/Niagara/Source/NiagaraEditor/Private/ViewModels/NiagaraScriptStatsViewModel.cpp
@@ -780,7 +780,11 @@ EShaderPlatform FNiagaraScriptStatsViewModel::ColumnNameToShaderPlatform(const F
 void FNiagaraScriptStatsViewModel::BuildShaderPlatformDetails()
 {
 	TArray<EShaderPlatform, TInlineAllocator<SP_StaticPlatform_Last + 1>> StaticShaderPlatforms;
+#if PLATFORM_WINDOWS
 	StaticShaderPlatforms.Add(SP_PCD3D_SM5);
+	StaticShaderPlatforms.Add(SP_PCD3D_SM6);
+	StaticShaderPlatforms.Add(SP_PCD3D_ES3_1);
+#endif
 	StaticShaderPlatforms.Add(SP_VULKAN_SM5);
 	StaticShaderPlatforms.Add(SP_VULKAN_SM6);
 	StaticShaderPlatforms.Add(SP_OPENGL_ES3_1_ANDROID); 
