- Emitter & Renderer Count Validation Rules

diff --git a/Engine/Plugins/FX/Niagara/Source/NiagaraEditor/Private/NiagaraValidationRules.cpp b/Engine/Plugins/FX/Niagara/Source/NiagaraEditor/Private/NiagaraValidationRules.cpp
--- a/Engine/Plugins/FX/Niagara/Source/NiagaraEditor/Private/NiagaraValidationRules.cpp
+++ b/Engine/Plugins/FX/Niagara/Source/NiagaraEditor/Private/NiagaraValidationRules.cpp
@@ -160,48 +160,60 @@ namespace NiagaraValidation
 		);
 	}
 
-
-	FString GetPlatformConflictsString(const FNiagaraPlatformSet& PlatformSetA, const FNiagaraPlatformSet& PlatformSetB, int MaxPlatformsToShow = 4)
+	TArray<FNiagaraPlatformSetConflictInfo> GatherPlatformSetConflicts(const FNiagaraPlatformSet* SetA, const FNiagaraPlatformSet* SetB)
 	{
-		TArray<const FNiagaraPlatformSet*> CheckSets;
-		CheckSets.Add(&PlatformSetA);
-		CheckSets.Add(&PlatformSetB);
-
-		TArray<FNiagaraPlatformSetConflictInfo> ConflictInfos;
-		FNiagaraPlatformSet::GatherConflicts(CheckSets, ConflictInfos);
+		TArray<const FNiagaraPlatformSet*> PlatformSets = {SetA, SetB};
+		TArray<FNiagaraPlatformSetConflictInfo> Conflicts;
+		FNiagaraPlatformSet::GatherConflicts(PlatformSets, Conflicts);
+		return MoveTemp(Conflicts);
+	}
 
+	FString GetPlatformConflictsString(TConstArrayView<FNiagaraPlatformSetConflictInfo> ConflictInfos, int MaxPlatformsToShow = 4)
+	{
 		if (ConflictInfos.Num() > 0)
 		{
-			TSet<FName> BannedPlatformNames;
+			TSet<FName> ConflictPlatformNames;
 			for (const FNiagaraPlatformSetConflictInfo& ConflictInfo : ConflictInfos)
 			{
 				for (const FNiagaraPlatformSetConflictEntry& ConflictEntry : ConflictInfo.Conflicts)
 				{
-					BannedPlatformNames.Add(ConflictEntry.ProfileName);
+					ConflictPlatformNames.Add(ConflictEntry.ProfileName);
 				}
 			}
 
-			TStringBuilder<256> BannedPlatformsString;
+			TStringBuilder<256> ConflictPlatformsString;
 			int NumFounds = 0;
-			for (FName PlatformName : BannedPlatformNames)
+			for (FName PlatformName : ConflictPlatformNames)
 			{
 				if (NumFounds >= MaxPlatformsToShow)
 				{
-					BannedPlatformsString.Append(TEXT(", ..."));
+					ConflictPlatformsString.Append(TEXT(", ..."));
 					break;
 				}
 				if (NumFounds != 0)
 				{
-					BannedPlatformsString.Append(TEXT(", "));
+					ConflictPlatformsString.Append(TEXT(", "));
 				}
 				++NumFounds;
-				PlatformName.AppendString(BannedPlatformsString);
+				PlatformName.AppendString(ConflictPlatformsString);
 			}
-			return BannedPlatformsString.ToString();
+			return ConflictPlatformsString.ToString();
 		}
 		return FString();
 	}
 
+	FString GetPlatformConflictsString(const FNiagaraPlatformSet& PlatformSetA, const FNiagaraPlatformSet& PlatformSetB, int MaxPlatformsToShow = 4)
+	{
+		TArray<const FNiagaraPlatformSet*> CheckSets;
+		CheckSets.Add(&PlatformSetA);
+		CheckSets.Add(&PlatformSetB);
+
+		TArray<FNiagaraPlatformSetConflictInfo> ConflictInfos;
+		FNiagaraPlatformSet::GatherConflicts(CheckSets, ConflictInfos);
+
+		return GetPlatformConflictsString(ConflictInfos, MaxPlatformsToShow);
+	}
+
 	TSharedPtr<FNiagaraEmitterHandleViewModel> GetEmitterViewModel(const FNiagaraValidationContext& Context, UNiagaraEmitter* NiagaraEmitter)
 	{
 		if (NiagaraEmitter == nullptr)
@@ -448,6 +460,127 @@ bool IsEnabledForMaxQualityLevel(FNiagaraPlatformSet Platforms, int32 MaxQuality
 	return false;
 }
 
+void UNiagaraValidationRule_EmitterCount::CheckValidity(const FNiagaraValidationContext& Context, TArray<FNiagaraValidationResult>& OutResults) const
+{
+	const int32 NumEmitterCountLimits = EmitterCountLimits.Num();
+	if (NumEmitterCountLimits == 0)
+	{
+		return;
+	}
+
+	TArray<TArray<FNiagaraPlatformSetConflictInfo>, TInlineAllocator<8>> ConflictsPerLimit;
+	TArray<int32, TInlineAllocator<8>> EmitterCountPerLimit;
+
+	EmitterCountPerLimit.AddDefaulted(NumEmitterCountLimits);
+	ConflictsPerLimit.AddDefaulted(NumEmitterCountLimits);
+
+	UNiagaraSystem& System = Context.ViewModel->GetSystem();
+	TArray<TSharedRef<FNiagaraEmitterHandleViewModel>> EmitterHandleViewModels = Context.ViewModel->GetEmitterHandleViewModels();
+	for (TSharedRef<FNiagaraEmitterHandleViewModel> EmitterHandleModel : EmitterHandleViewModels)
+	{
+		FNiagaraEmitterHandle* EmitterHandle = EmitterHandleModel.Get().GetEmitterHandle();
+		if (!EmitterHandle->GetIsEnabled())
+		{
+			continue;
+		}
+		FVersionedNiagaraEmitterData* EmitterData = EmitterHandleModel.Get().GetEmitterHandle()->GetEmitterData();
+	
+		for (int32 i=0; i < NumEmitterCountLimits; ++i)
+		{
+			const TArray<FNiagaraPlatformSetConflictInfo> Conflicts = NiagaraValidation::GatherPlatformSetConflicts(&EmitterCountLimits[i].Platforms, &EmitterData->Platforms);
+			if (Conflicts.Num() > 0)
+			{
+				ConflictsPerLimit[i].Append(Conflicts);
+				++EmitterCountPerLimit[i];
+			}
+		}
+	}
+
+	for (int32 i=0; i < NumEmitterCountLimits; ++i)
+	{
+		const int32 EmitterCountLimit = EmitterCountLimits[i].EmitterCountLimit;
+		if (EmitterCountPerLimit[i] <= EmitterCountLimit)
+		{
+			continue;
+		}
+
+		const FString PlatformConflicts = NiagaraValidation::GetPlatformConflictsString(ConflictsPerLimit[i]);
+
+		FNiagaraValidationResult& Result = OutResults.AddDefaulted_GetRef();
+		Result.Severity = Severity;
+		Result.SummaryText = FText::Format(LOCTEXT("EmitterCountLimit", "Emitter count limit {0} exceeded {1}."), EmitterCountLimit, EmitterCountPerLimit[i]);
+		Result.Description = FText::Format(LOCTEXT("EmitterCountLimitDesc", "Emitter count limit {0} has been exceeded {1} for platforms '{2}' please reduce the emitter count to improve performance."), EmitterCountLimit, EmitterCountPerLimit[i], FText::FromString(PlatformConflicts));
+		Result.SourceObject = NiagaraValidation::GetStackEntry<UNiagaraStackSystemPropertiesItem>(Context.ViewModel->GetSystemStackViewModel());
+	}
+}
+
+void UNiagaraValidationRule_RendererCount::CheckValidity(const FNiagaraValidationContext& Context, TArray<FNiagaraValidationResult>& OutResults) const
+{
+	const int32 NumRendererCountLimits = RendererCountLimits.Num();
+	if (NumRendererCountLimits == 0)
+	{
+		return;
+	}
+
+	TArray<TArray<FNiagaraPlatformSetConflictInfo>, TInlineAllocator<8>> ConflictsPerLimit;
+	TArray<int32, TInlineAllocator<8>> RendererCountPerLimit;
+
+	RendererCountPerLimit.AddDefaulted(NumRendererCountLimits);
+	ConflictsPerLimit.AddDefaulted(NumRendererCountLimits);
+
+	UNiagaraSystem& System = Context.ViewModel->GetSystem();
+	TArray<TSharedRef<FNiagaraEmitterHandleViewModel>> EmitterHandleViewModels = Context.ViewModel->GetEmitterHandleViewModels();
+	for (TSharedRef<FNiagaraEmitterHandleViewModel> EmitterHandleModel : EmitterHandleViewModels)
+	{
+		FNiagaraEmitterHandle* EmitterHandle = EmitterHandleModel.Get().GetEmitterHandle();
+		if (!EmitterHandle->GetIsEnabled())
+		{
+			continue;
+		}
+		FVersionedNiagaraEmitterData* EmitterData = EmitterHandleModel.Get().GetEmitterHandle()->GetEmitterData();
+
+		for (int32 i=0; i < NumRendererCountLimits; ++i)
+		{
+			if ( NiagaraValidation::GatherPlatformSetConflicts(&RendererCountLimits[i].Platforms, &EmitterData->Platforms).Num() == 0 )
+			{
+				continue;
+			}
+
+			EmitterData->ForEachRenderer(
+				[this, i, &ConflictsPerLimit, &RendererCountPerLimit, &EmitterData](UNiagaraRendererProperties* RendererProperties)
+				{
+					if (RendererProperties->GetIsEnabled())
+					{
+						TArray<FNiagaraPlatformSetConflictInfo> Conflicts = NiagaraValidation::GatherPlatformSetConflicts(&RendererCountLimits[i].Platforms, &EmitterData->Platforms);
+						if (Conflicts.Num() > 0)
+						{
+							ConflictsPerLimit[i].Append(Conflicts);
+							++RendererCountPerLimit[i];
+						}
+					}
+				}
+			);
+		}
+	}
+
+	for (int32 i = 0; i < NumRendererCountLimits; ++i)
+	{
+		const int32 RendererCountLimit = RendererCountLimits[i].RendererCountLimit;
+		if (RendererCountPerLimit[i] <= RendererCountLimit)
+		{
+			continue;
+		}
+
+		const FString PlatformConflicts = NiagaraValidation::GetPlatformConflictsString(ConflictsPerLimit[i]);
+
+		FNiagaraValidationResult& Result = OutResults.AddDefaulted_GetRef();
+		Result.Severity = Severity;
+		Result.SummaryText = FText::Format(LOCTEXT("RendererCountLimit", "Renderer count limit {0} exceeded {1}."), RendererCountLimit, RendererCountPerLimit[i]);
+		Result.Description = FText::Format(LOCTEXT("RendererCountLimitDesc", "Renderer count limit {0} has been exceeded {1} for platforms '{2}' please reduce the renderer count to improve performance."), RendererCountLimit, RendererCountPerLimit[i], FText::FromString(PlatformConflicts));
+		Result.SourceObject = NiagaraValidation::GetStackEntry<UNiagaraStackSystemPropertiesItem>(Context.ViewModel->GetSystemStackViewModel());
+	}
+}
+
 void UNiagaraValidationRule_BannedRenderers::CheckValidity(const FNiagaraValidationContext& Context, TArray<FNiagaraValidationResult>& Results)  const
 {
 	UNiagaraSystem& System = Context.ViewModel->GetSystem();
@@ -459,12 +592,7 @@ void UNiagaraValidationRule_BannedRenderers::CheckValidity(const FNiagaraValidat
 		{
 			if (RendererProperties->GetIsEnabled() && BannedRenderers.Contains(RendererProperties->GetClass()))
 			{
-				TArray<const FNiagaraPlatformSet*> CheckSets;
-				CheckSets.Add(&Platforms);
-				CheckSets.Add(&RendererProperties->Platforms);
-
-				TArray<FNiagaraPlatformSetConflictInfo> Conflicts;
-				FNiagaraPlatformSet::GatherConflicts(CheckSets, Conflicts);
+				TArray<FNiagaraPlatformSetConflictInfo> Conflicts = NiagaraValidation::GatherPlatformSetConflicts(&Platforms, &RendererProperties->Platforms);
 				if (Conflicts.Num() > 0)
 				{
 					if ( UNiagaraStackRendererItem* StackItem = NiagaraValidation::GetRendererStackItem(EmitterHandleModel.Get().GetEmitterStackViewModel(), RendererProperties) )
@@ -520,11 +648,7 @@ void UNiagaraValidationRule_BannedModules::CheckValidity(const FNiagaraValidatio
 					if (EmitterData)
 					{
 						//If we're on an emitter, this emitter may be culled on the platforms the rule applies to.
-						TArray<const FNiagaraPlatformSet*> CheckSets;
-						CheckSets.Add(&Platforms);
-						CheckSets.Add(&EmitterData->Platforms);
-						TArray<FNiagaraPlatformSetConflictInfo> Conflicts;
-						FNiagaraPlatformSet::GatherConflicts(CheckSets, Conflicts);
+						const TArray<FNiagaraPlatformSetConflictInfo> Conflicts = NiagaraValidation::GatherPlatformSetConflicts(&Platforms, &EmitterData->Platforms);
 						bApplyBan = Conflicts.Num() > 0;
 					}
 
diff --git a/Engine/Plugins/FX/Niagara/Source/NiagaraEditor/Public/NiagaraValidationRules.h b/Engine/Plugins/FX/Niagara/Source/NiagaraEditor/Public/NiagaraValidationRules.h
--- a/Engine/Plugins/FX/Niagara/Source/NiagaraEditor/Public/NiagaraValidationRules.h
+++ b/Engine/Plugins/FX/Niagara/Source/NiagaraEditor/Public/NiagaraValidationRules.h
@@ -35,6 +35,62 @@ public:
 	virtual void CheckValidity(const FNiagaraValidationContext& Context, TArray<FNiagaraValidationResult>& OutResults) const override;
 };
 
+USTRUCT()
+struct FNiagaraValidationRule_EmitterCountAndPlatformSet
+{
+	GENERATED_BODY()
+
+	UPROPERTY(EditAnywhere, Category = Validation)
+	FNiagaraPlatformSet Platforms;
+
+	UPROPERTY(EditAnywhere, Category = Validation)
+	int32 EmitterCountLimit = 8;
+};
+
+/** This validation rule can be used to apply budgets for emitter count. */
+UCLASS(Category = "Validation", DisplayName = "Emitter Count")
+class UNiagaraValidationRule_EmitterCount : public UNiagaraValidationRule
+{
+	GENERATED_BODY()
+
+public:
+	UPROPERTY(EditAnywhere, Category = Validation)
+	ENiagaraValidationSeverity Severity = ENiagaraValidationSeverity::Warning;
+
+	UPROPERTY(EditAnywhere, Category = Validation)
+	TArray<FNiagaraValidationRule_EmitterCountAndPlatformSet> EmitterCountLimits;
+
+	virtual void CheckValidity(const FNiagaraValidationContext& Context, TArray<FNiagaraValidationResult>& OutResults) const override;
+};
+
+USTRUCT()
+struct FNiagaraValidationRule_RendererCountAndPlatformSet
+{
+	GENERATED_BODY()
+
+	UPROPERTY(EditAnywhere, Category = Validation)
+	FNiagaraPlatformSet Platforms;
+
+	UPROPERTY(EditAnywhere, Category = Validation)
+	int32 RendererCountLimit = 8;
+};
+
+/** This validation rule can be used to apply budgets for renderer count. */
+UCLASS(Category = "Validation", DisplayName = "Renderer Count")
+class UNiagaraValidationRule_RendererCount : public UNiagaraValidationRule
+{
+	GENERATED_BODY()
+
+public:
+	UPROPERTY(EditAnywhere, Category = Validation)
+	ENiagaraValidationSeverity Severity = ENiagaraValidationSeverity::Warning;
+
+	UPROPERTY(EditAnywhere, Category = Validation)
+	TArray<FNiagaraValidationRule_RendererCountAndPlatformSet> RendererCountLimits;
+
+	virtual void CheckValidity(const FNiagaraValidationContext& Context, TArray<FNiagaraValidationResult>& OutResults) const override;
+};
+
 /** This validation rule can ban the use of certain renderers on all or a subset of platforms. */
 UCLASS(Category = "Validation", DisplayName = "Banned Renderers")
 class UNiagaraValidationRule_BannedRenderers : public UNiagaraValidationRule
