- Fix lack of notify when adding / removing stateless modules
- Fix color range not using color picker

diff --git a/Engine/Plugins/FX/Niagara/Source/Niagara/Internal/Stateless/NiagaraStatelessDistribution.h b/Engine/Plugins/FX/Niagara/Source/Niagara/Internal/Stateless/NiagaraStatelessDistribution.h
--- a/Engine/Plugins/FX/Niagara/Source/Niagara/Internal/Stateless/NiagaraStatelessDistribution.h
+++ b/Engine/Plugins/FX/Niagara/Source/Niagara/Internal/Stateless/NiagaraStatelessDistribution.h
@@ -171,6 +171,7 @@ struct FNiagaraDistributionRangeColor : public FNiagaraDistributionBase
 
 #if WITH_EDITORONLY_DATA
 	virtual bool AllowCurves() const override { return false; }
+	virtual bool DisplayAsColor() const { return true; }
 	virtual void UpdateValuesFromDistribution() override;
 	virtual FNiagaraTypeDefinition GetBindingTypeDef() { return FNiagaraTypeDefinition::GetColorDef(); }
 #endif
diff --git a/Engine/Plugins/FX/Niagara/Source/NiagaraEditor/Private/ViewModels/Stack/NiagaraStackStatelessEmitterSimulateGroup.cpp b/Engine/Plugins/FX/Niagara/Source/NiagaraEditor/Private/ViewModels/Stack/NiagaraStackStatelessEmitterSimulateGroup.cpp
--- a/Engine/Plugins/FX/Niagara/Source/NiagaraEditor/Private/ViewModels/Stack/NiagaraStackStatelessEmitterSimulateGroup.cpp
+++ b/Engine/Plugins/FX/Niagara/Source/NiagaraEditor/Private/ViewModels/Stack/NiagaraStackStatelessEmitterSimulateGroup.cpp
@@ -138,6 +138,7 @@ void UNiagaraStackStatelessEmitterSimulateGroup::RefreshChildrenInternal(const T
 
 void UNiagaraStackStatelessEmitterSimulateGroup::ModuleAdded(UNiagaraStatelessModule* StatelessModule)
 {
+	OnDataObjectModified().Broadcast({ StatelessModule }, ENiagaraDataObjectChange::Changed);
 	RefreshChildren();
 }
 
@@ -185,6 +186,8 @@ void UNiagaraStackStatelessModuleItem::Delete()
 	{
 		StatelessModule->Modify();
 		StatelessModule->SetIsModuleEnabled(false);
+		StatelessModule->PostEditChange();
+		OnDataObjectModified().Broadcast({ StatelessModule }, ENiagaraDataObjectChange::Changed);
 		GetStackEditorData().Modify();
 		GetStackEditorData().SetStatelessModuleShowWhenDisabled(GetStackEditorDataKey(), false);
 		OnModifiedGroupItems().Broadcast();
@@ -279,8 +282,7 @@ void UNiagaraStackStatelessModuleItem::SetIsEnabledInternal(bool bInIsEnabled)
 		StatelessModule->SetIsModuleEnabled(bInIsEnabled);
 		StatelessModule->PostEditChange();
 		GetStackEditorData().SetStatelessModuleShowWhenDisabled(GetStackEditorDataKey(), true);
-		TArray<UObject*> ChangedObjects = { StatelessModule };
-		OnDataObjectModified().Broadcast(ChangedObjects, ENiagaraDataObjectChange::Changed);
+		OnDataObjectModified().Broadcast({ StatelessModule }, ENiagaraDataObjectChange::Changed);
 		RefreshChildren();
 	}
 }
