Niagara - Stateless UI fixes
+ Fix debug draw properties showing up when they are available in the header.
+ Fix debug draw header widgets not actually enabling debug draw.
+ Fix automatic selection of newly added modules and spawn infos.

diff --git a/Engine/Plugins/FX/Niagara/Source/Niagara/Internal/Stateless/NiagaraStatelessModule.h b/Engine/Plugins/FX/Niagara/Source/Niagara/Internal/Stateless/NiagaraStatelessModule.h
--- a/Engine/Plugins/FX/Niagara/Source/Niagara/Internal/Stateless/NiagaraStatelessModule.h
+++ b/Engine/Plugins/FX/Niagara/Source/Niagara/Internal/Stateless/NiagaraStatelessModule.h
@@ -2,9 +2,9 @@
 
 #pragma once
 
+#include "NiagaraMergeable.h"
 #include "NiagaraStatelessCommon.h"
 #include "NiagaraStatelessDistribution.h"
-
 #include "ShaderParameterStruct.h"
 
 #include "NiagaraStatelessModule.generated.h"
@@ -87,7 +87,7 @@ private:
 ////////////////////////////////////////////////////////////////////////////////////////////////////////////////////////////////////////////////////////////////
 
 UCLASS(MinimalAPI, abstract, EditInlineNew)
-class UNiagaraStatelessModule : public UObject
+class UNiagaraStatelessModule : public UNiagaraMergeable
 {
 	GENERATED_BODY()
 
@@ -96,7 +96,7 @@ protected:
 	uint32 bModuleEnabled : 1 = true;
 
 #if WITH_EDITORONLY_DATA
-	UPROPERTY(EditAnywhere, Transient, Category = "Parameters", meta = (DisplayPriority = 0, StackItemHeaderIcon = "Icons.Visibility"))
+	UPROPERTY(EditAnywhere, Transient, Category = "Parameters", meta = (DisplayPriority = 0, HideInStack, StackItemHeaderIcon = "Icons.Visibility"))
 	uint32 bDebugDrawEnabled : 1 = false;
 #endif
 
diff --git a/Engine/Plugins/FX/Niagara/Source/NiagaraEditor/Private/ViewModels/Stack/NiagaraStackStatelessEmitterSimulateGroup.cpp b/Engine/Plugins/FX/Niagara/Source/NiagaraEditor/Private/ViewModels/Stack/NiagaraStackStatelessEmitterSimulateGroup.cpp
--- a/Engine/Plugins/FX/Niagara/Source/NiagaraEditor/Private/ViewModels/Stack/NiagaraStackStatelessEmitterSimulateGroup.cpp
+++ b/Engine/Plugins/FX/Niagara/Source/NiagaraEditor/Private/ViewModels/Stack/NiagaraStackStatelessEmitterSimulateGroup.cpp
@@ -11,6 +11,8 @@
 #include "Stateless/NiagaraStatelessCommon.h"
 #include "Stateless/NiagaraStatelessEmitter.h"
 #include "Stateless/NiagaraStatelessModule.h"
+#include "ViewModels/NiagaraSystemSelectionViewModel.h"
+#include "ViewModels/NiagaraSystemViewModel.h"
 #include "ViewModels/Stack/NiagaraStackItemPropertyHeaderValueShared.h"
 #include "ViewModels/Stack/NiagaraStackObject.h"
 
@@ -138,6 +140,8 @@ void UNiagaraStackStatelessEmitterSimulateGroup::RefreshChildrenInternal(const T
 
 void UNiagaraStackStatelessEmitterSimulateGroup::ModuleAdded(UNiagaraStatelessModule* StatelessModule)
 {
+	GetSystemViewModel()->GetSelectionViewModel()->EmptySelection();
+	GetSystemViewModel()->GetSelectionViewModel()->AddEntryToSelectionByDisplayedObjectDeferred(StatelessModule);
 	OnDataObjectModified().Broadcast({ StatelessModule }, ENiagaraDataObjectChange::Changed);
 	RefreshChildren();
 }
@@ -252,7 +256,7 @@ void UNiagaraStackStatelessModuleItem::RefreshChildrenInternal(const TArray<UNia
 				if (DebugDrawProperty != nullptr)
 				{
 					HeaderValueHandlers.Add(MakeShared<FNiagaraStackItemPropertyHeaderValue>(
-						*ModuleObject, nullptr, *DebugDrawProperty, 
+						*StatelessModule, nullptr, *DebugDrawProperty, 
 						FSimpleDelegate::CreateUObject(this, &UNiagaraStackStatelessModuleItem::OnHeaderValueChanged)));
 				}
 			}
diff --git a/Engine/Plugins/FX/Niagara/Source/NiagaraEditor/Private/ViewModels/Stack/NiagaraStackStatelessEmitterSpawnGroup.cpp b/Engine/Plugins/FX/Niagara/Source/NiagaraEditor/Private/ViewModels/Stack/NiagaraStackStatelessEmitterSpawnGroup.cpp
--- a/Engine/Plugins/FX/Niagara/Source/NiagaraEditor/Private/ViewModels/Stack/NiagaraStackStatelessEmitterSpawnGroup.cpp
+++ b/Engine/Plugins/FX/Niagara/Source/NiagaraEditor/Private/ViewModels/Stack/NiagaraStackStatelessEmitterSpawnGroup.cpp
@@ -10,6 +10,8 @@
 #include "Stateless/NiagaraDistributionIntPropertyCustomization.h"
 #include "Stateless/NiagaraStatelessEmitter.h"
 #include "Styling/AppStyle.h"
+#include "ViewModels/NiagaraSystemSelectionViewModel.h"
+#include "ViewModels/NiagaraSystemViewModel.h"
 #include "ViewModels/Stack/NiagaraStackItemPropertyHeaderValueShared.h"
 #include "ViewModels/Stack/NiagaraStackObject.h"
 
@@ -98,6 +100,13 @@ void UNiagaraStackStatelessEmitterSpawnGroup::RefreshChildrenInternal(const TArr
 
 void UNiagaraStackStatelessEmitterSpawnGroup::OnSpawnInfoAdded(FGuid AddedItemId)
 {
+	GetSystemViewModel()->GetSelectionViewModel()->EmptySelection();
+	GetSystemViewModel()->GetSelectionViewModel()->AddEntryToSelectionBySelectionIdDeferred(AddedItemId);
+	UNiagaraStatelessEmitter* StatelessEmitter = StatelessEmitterWeak.Get();
+	if (StatelessEmitter != nullptr)
+	{
+		OnDataObjectModified().Broadcast({ StatelessEmitter }, ENiagaraDataObjectChange::Changed);
+	}
 	RefreshChildren();
 }
 
@@ -130,6 +139,11 @@ FText UNiagaraStackStatelessEmitterSpawnItem::GetDisplayName() const
 	return LOCTEXT("EmitterSpawnDisplayName", "Spawn Info");
 }
 
+FGuid UNiagaraStackStatelessEmitterSpawnItem::GetSelectionId() const
+{
+	return SourceId;
+}
+
 bool UNiagaraStackStatelessEmitterSpawnItem::TestCanDeleteWithMessage(FText& OutCanDeleteMessage) const
 {
 	OutCanDeleteMessage = LOCTEXT("DeleteSpawnDataMessage", "Delete this spawn data.");
diff --git a/Engine/Plugins/FX/Niagara/Source/NiagaraEditor/Public/ViewModels/Stack/NiagaraStackStatelessEmitterSpawnGroup.h b/Engine/Plugins/FX/Niagara/Source/NiagaraEditor/Public/ViewModels/Stack/NiagaraStackStatelessEmitterSpawnGroup.h
--- a/Engine/Plugins/FX/Niagara/Source/NiagaraEditor/Public/ViewModels/Stack/NiagaraStackStatelessEmitterSpawnGroup.h
+++ b/Engine/Plugins/FX/Niagara/Source/NiagaraEditor/Public/ViewModels/Stack/NiagaraStackStatelessEmitterSpawnGroup.h
@@ -54,6 +54,7 @@ public:
 	void Initialize(FRequiredEntryData InRequiredEntryData, UNiagaraStatelessEmitter* InStatelessEmitter, int32 InIndex);
 
 	virtual FText GetDisplayName() const override;
+	virtual FGuid GetSelectionId() const override;
 
 	virtual bool SupportsDelete() const override { return true; }
 	virtual bool TestCanDeleteWithMessage(FText& OutCanDeleteMessage) const override;
