- Emitter binding changed to be case insensitive
- The unique emitter name is an FString and binding is an FName, the FName can be impacted by the first user so case is not guaranteed

diff --git a/Engine/Plugins/FX/Niagara/Source/Niagara/Private/NiagaraDataInterfaceEmitterBinding.cpp b/Engine/Plugins/FX/Niagara/Source/Niagara/Private/NiagaraDataInterfaceEmitterBinding.cpp
--- a/Engine/Plugins/FX/Niagara/Source/Niagara/Private/NiagaraDataInterfaceEmitterBinding.cpp
+++ b/Engine/Plugins/FX/Niagara/Source/Niagara/Private/NiagaraDataInterfaceEmitterBinding.cpp
@@ -81,7 +81,7 @@ FNiagaraEmitterInstance* FNiagaraDataInterfaceEmitterBinding::Resolve(const FNia
 			if (const UNiagaraEmitter* CachedEmitter = EmitterInstance->GetEmitter())
 			{
 				//-TODO: UniqueEmitterName should probably be a FName?
-				if (EmitterNameStringView.Equals(CachedEmitter->GetUniqueEmitterName()) )
+				if (EmitterNameStringView.Equals(CachedEmitter->GetUniqueEmitterName(), ESearchCase::IgnoreCase) )
 				{
 					return &EmitterInstance.Get();
 				}
@@ -161,7 +161,7 @@ const FNiagaraEmitterHandle* FNiagaraDataInterfaceEmitterBinding::ResolveHandle(
 				if (UNiagaraEmitter* NiagaraEmitter = EmitterHandle.GetInstance().Emitter)
 				{
 					//-TODO: UniqueEmitterName should probably be a FName?
-					if (EmitterNameStringView.Equals(NiagaraEmitter->GetUniqueEmitterName()))
+					if (EmitterNameStringView.Equals(NiagaraEmitter->GetUniqueEmitterName(), ESearchCase::IgnoreCase))
 					{
 						return &EmitterHandle;
 					}
