Niagara - When updating compiled data interfaces, modify them so that the updates are captured in the transaction buffer if active.  This fixes an issue where undoing a change to a data interface would not update the running simulation.

diff --git a/Engine/Plugins/FX/Niagara/Source/NiagaraEditor/Private/ViewModels/NiagaraSystemViewModel.cpp b/Engine/Plugins/FX/Niagara/Source/NiagaraEditor/Private/ViewModels/NiagaraSystemViewModel.cpp
--- a/Engine/Plugins/FX/Niagara/Source/NiagaraEditor/Private/ViewModels/NiagaraSystemViewModel.cpp
+++ b/Engine/Plugins/FX/Niagara/Source/NiagaraEditor/Private/ViewModels/NiagaraSystemViewModel.cpp
@@ -2577,6 +2577,7 @@ void FNiagaraSystemViewModel::UpdateCompiledDataInterfaces(UNiagaraDataInterface
 	UNiagaraDataInterface* ResolvedCompileDataInterface = FNiagaraEditorUtilities::GetResolvedRuntimeInstanceForEditorDataInterfaceInstance(GetSystem(), *ChangedDataInterface);
 	if (ResolvedCompileDataInterface != nullptr)
 	{
+		ResolvedCompileDataInterface->Modify(false);
 		ChangedDataInterface->CopyTo(ResolvedCompileDataInterface);
 	}
 	else
