Re-enable fixed error messages for unset variables being read

diff --git a/Engine/Plugins/FX/Niagara/Source/NiagaraEditor/Private/NiagaraHlslTranslator.cpp b/Engine/Plugins/FX/Niagara/Source/NiagaraEditor/Private/NiagaraHlslTranslator.cpp
--- a/Engine/Plugins/FX/Niagara/Source/NiagaraEditor/Private/NiagaraHlslTranslator.cpp
+++ b/Engine/Plugins/FX/Niagara/Source/NiagaraEditor/Private/NiagaraHlslTranslator.cpp
@@ -5432,8 +5432,7 @@ void TNiagaraHlslTranslator<GraphBridge>::InitializeParameterMapDefaults(int32 P
 				FNiagaraVariable SourceForInitialValue = FNiagaraParameterUtilities::GetSourceForInitialValue(Var);
 				if (!UniqueVars.Contains(SourceForInitialValue))
 				{
-					//@todo(ng) disabled pending investigation UE-150159
-					//Error(FText::Format(LOCTEXT("MissingInitialValueSource", "Variable {0} is used, but its source variable {1} is not set!"), FText::FromName(Var.GetName()), FText::FromName(SourceForInitialValue.GetName())), nullptr, nullptr);
+					Error(FText::Format(LOCTEXT("MissingInitialValueSource", "Variable {0} is used, but its source variable {1} is not set!"), FText::FromName(Var.GetName()), FText::FromName(SourceForInitialValue.GetName())), nullptr, nullptr);
 				}
 				InitialNamespaceVariablesMissingDefault.Add(Var);
 			}
@@ -7076,8 +7075,7 @@ void TNiagaraHlslTranslator<GraphBridge>::HandleParameterRead(int32 ParamMapHist
 				}
 				else
 				{
-					//@todo(ng) disabled pending investigation UE-150159
-					//Error(FText::Format(LOCTEXT("MissingInitialValueSource", "Variable {0} is used, but its source variable {1} is not set!"), FText::FromName(Var.GetName()), FText::FromName(SourceForInitialValue.GetName())), nullptr, nullptr);
+					Error(FText::Format(LOCTEXT("MissingInitialValueSourceNode", "Variable {0} is used, but its source variable {1} is not set!"), FText::FromName(Var.GetName()), FText::FromName(SourceForInitialValue.GetName())), ErrorNode, nullptr);
 				}
 			}
 			else if (UniqueVars.Contains(Var) && UniqueVarToChunk.Contains(Var))
