Include all GPU Tick init as part of the main CSV Effects time.

diff --git a/Engine/Plugins/FX/Niagara/Source/Niagara/Private/NiagaraSystemSimulation.cpp b/Engine/Plugins/FX/Niagara/Source/Niagara/Private/NiagaraSystemSimulation.cpp
--- a/Engine/Plugins/FX/Niagara/Source/Niagara/Private/NiagaraSystemSimulation.cpp
+++ b/Engine/Plugins/FX/Niagara/Source/Niagara/Private/NiagaraSystemSimulation.cpp
@@ -369,6 +369,8 @@ struct FNiagaraSystemInstanceTickConcurrentTask
 	void DoTask(ENamedThreads::Type CurrentThread, const FGraphEventRef& MyCompletionGraphEvent)
 	{
 		PARTICLE_PERF_STAT_CYCLES_GT(FParticlePerfStatsContext(WorldContext, SystemSimulation->GetSystem()), TickConcurrent);
+		CSV_SCOPED_TIMING_STAT_EXCLUSIVE(Effects);
+
 #if NIAGARA_SYSTEMSIMULATION_DEBUGGING
 		NiagaraSystemSimulationLocal::DebugDelayInstancesTask();
 #endif
