* Implicit lwc testing for engine tests
* Fixed failing engine test
* Data channel test comparison logic

#virtualized

diff --git a/Engine/Plugins/FX/Niagara/Source/Niagara/Classes/NiagaraSimCacheCustomStorageInterface.h b/Engine/Plugins/FX/Niagara/Source/Niagara/Classes/NiagaraSimCacheCustomStorageInterface.h
--- a/Engine/Plugins/FX/Niagara/Source/Niagara/Classes/NiagaraSimCacheCustomStorageInterface.h
+++ b/Engine/Plugins/FX/Niagara/Source/Niagara/Classes/NiagaraSimCacheCustomStorageInterface.h
@@ -58,7 +58,7 @@ public:
 	virtual void SimCachePostReadFrame(void* OptionalPerInstanceData, FNiagaraSystemInstance* SystemInstance) {}
 
 	/**
-	Called to compare a frame between two separate simulation cache storages
+	Called to compare a frame between two separate simulation cache storages. Mainly useful for unit testing.
 	This will be called on the CDO object since we do not have the actual data interface.
 	*/
 	virtual bool SimCacheCompareFrame(UObject* LhsStorageObject, UObject* RhsStorageObject, int FrameIndex, TOptional<float> Tolerance, FString& OutErrors) const { OutErrors = TEXT("Compare not implemented"); return false; }
diff --git a/Engine/Plugins/FX/Niagara/Source/Niagara/Internal/DataInterface/NiagaraDataInterfaceDataChannelWrite.h b/Engine/Plugins/FX/Niagara/Source/Niagara/Internal/DataInterface/NiagaraDataInterfaceDataChannelWrite.h
--- a/Engine/Plugins/FX/Niagara/Source/Niagara/Internal/DataInterface/NiagaraDataInterfaceDataChannelWrite.h
+++ b/Engine/Plugins/FX/Niagara/Source/Niagara/Internal/DataInterface/NiagaraDataInterfaceDataChannelWrite.h
@@ -168,9 +168,9 @@ public:
 	virtual ETickingGroup CalculateFinalTickGroup(const void* PerInstanceData) const override;
 
 	//We cannot overlap frames as we must correctly sync up with the data channel manager on Begin/End frame etc.
-	virtual bool PostSimulateCanOverlapFrames() const { return false; }
+	virtual bool PostSimulateCanOverlapFrames() const override { return false; }
 	//We cannot have post stage overlap tick groups so that the write DI can publish it's contents to the data channel at the correct time to allow same frame reads.
-	virtual bool PostStageCanOverlapTickGroups() const { return false; }
+	virtual bool PostStageCanOverlapTickGroups() const override { return false; }
 	//UNiagaraDataInterface Interface
 
 	NIAGARA_API void Num(FVectorVMExternalFunctionContext& Context);
@@ -186,6 +186,7 @@ public:
 	virtual bool SimCacheWriteFrame(UObject* StorageObject, int FrameIndex, FNiagaraSystemInstance* SystemInstance, const void* OptionalPerInstanceData, FNiagaraSimCacheFeedbackContext& FeedbackContext) const override;
 	virtual bool SimCacheReadFrame(UObject* StorageObject, int FrameA, int FrameB, float Interp, FNiagaraSystemInstance* SystemInstance, void* OptionalPerInstanceData) override;
 	virtual void SimCachePostReadFrame(void* OptionalPerInstanceData, FNiagaraSystemInstance* SystemInstance) override;
+	virtual bool SimCacheCompareFrame(UObject* LhsStorageObject, UObject* RhsStorageObject, int FrameIndex, TOptional<float> Tolerance, FString& OutErrors) const override;
 protected:
 #if WITH_EDITORONLY_DATA
 	NIAGARA_API virtual void GetFunctionsInternal(TArray<FNiagaraFunctionSignature>& OutFunctions) const override;
diff --git a/Engine/Plugins/FX/Niagara/Source/Niagara/Private/DataInterface/NiagaraDataInterfaceDataChannelWrite.cpp b/Engine/Plugins/FX/Niagara/Source/Niagara/Private/DataInterface/NiagaraDataInterfaceDataChannelWrite.cpp
--- a/Engine/Plugins/FX/Niagara/Source/Niagara/Private/DataInterface/NiagaraDataInterfaceDataChannelWrite.cpp
+++ b/Engine/Plugins/FX/Niagara/Source/Niagara/Private/DataInterface/NiagaraDataInterfaceDataChannelWrite.cpp
@@ -691,6 +691,63 @@ void UNiagaraDataInterfaceDataChannelWrite::SimCachePostReadFrame(void* Optional
 	// send data to data channel
 }
 
+bool UNiagaraDataInterfaceDataChannelWrite::SimCacheCompareFrame(UObject* LhsStorageObject, UObject* RhsStorageObject, int FrameIndex, TOptional<float> Tolerance, FString& OutErrors) const
+{
+	UNDIDataChannelWriteSimCacheData* Storage1 = Cast<UNDIDataChannelWriteSimCacheData>(LhsStorageObject);
+	UNDIDataChannelWriteSimCacheData* Storage2 = Cast<UNDIDataChannelWriteSimCacheData>(RhsStorageObject);
+
+	if (Storage1 == nullptr && Storage2 == nullptr)
+	{
+		return true;
+	}
+	if (Storage1 == nullptr || Storage2 == nullptr)
+	{
+		OutErrors = TEXT("Recevied nullptr storage object for comparison");
+		return false;
+	}
+	if (Storage1->DataChannelReference != Storage2->DataChannelReference)
+	{
+		OutErrors = TEXT("Different source data channel assets");
+		return false;
+	}
+	if (Storage1->FrameData.Num() != Storage2->FrameData.Num())
+	{
+		OutErrors = FString::Format(TEXT("Different frame data count. {0} vs {1}"), {Storage1->FrameData.Num(), Storage2->FrameData.Num()});
+		return false;
+	}
+
+	bool bEqual = true;
+	for (int i = 0; i < Storage1->FrameData.Num(); i++)
+	{
+		FNDIDataChannelWriteSimCacheFrame& Frame1 = Storage1->FrameData[i];
+		FNDIDataChannelWriteSimCacheFrame& Frame2 = Storage2->FrameData[i];
+
+		if (Frame1.NumElements != Frame2.NumElements)
+		{
+			bEqual = false;
+			OutErrors += FString::Format(TEXT("Frame {0}: different number of elements in data channel store, {1} vs {2}\n"), {i, Frame1.NumElements, Frame2.NumElements});
+		}
+		else if (Frame1.VariableData.Num() != Frame2.VariableData.Num())
+		{
+			bEqual = false;
+			OutErrors += FString::Format(TEXT("Frame {0}: different number of variables in data channel store, {1} vs {2}\n"), {i, Frame1.VariableData.Num(), Frame2.VariableData.Num()});
+		}
+		else
+		{
+			for (int k = 0; k < Frame1.VariableData.Num(); k++)
+			{
+				FNDIDataChannelWriteSimCacheFrameBuffer& Buffer1 = Frame1.VariableData[k];
+				FNDIDataChannelWriteSimCacheFrameBuffer& Buffer2 = Frame2.VariableData[k];
+				if (Buffer1.SourceVar != Buffer2.SourceVar || Buffer1.Data != Buffer2.Data)
+				{
+					OutErrors += FString::Format(TEXT("Frame {0}: different buffers in data channel store for source var {1}\n"), {i, *Buffer1.SourceVar.GetName().ToString()});
+				}
+			}
+		}
+	}
+	return bEqual;
+}
+
 bool UNiagaraDataInterfaceDataChannelWrite::CopyToInternal(UNiagaraDataInterface* Destination)const
 {
 	if (!Super::CopyToInternal(Destination))
diff --git a/Engine/Plugins/FX/Niagara/Source/Niagara/Private/NiagaraSimCache.cpp b/Engine/Plugins/FX/Niagara/Source/Niagara/Private/NiagaraSimCache.cpp
--- a/Engine/Plugins/FX/Niagara/Source/Niagara/Private/NiagaraSimCache.cpp
+++ b/Engine/Plugins/FX/Niagara/Source/Niagara/Private/NiagaraSimCache.cpp
@@ -24,6 +24,11 @@
 UNiagaraSimCache::FOnCacheBeginWrite	UNiagaraSimCache::OnCacheBeginWrite;
 UNiagaraSimCache::FOnCacheEndWrite		UNiagaraSimCache::OnCacheEndWrite;
 
+namespace FNiagaraSimCacheInternal
+{
+	FNiagaraDataInterfaceUtilities::FDataInterfaceSearchOptions DISearchOptions = { true };
+}
+
 ////////////////////////////////////////////////////////////////////////////////////////////////////////////////////////////////////////////////////
 
 bool FNiagaraSimCacheVariable::operator==(const FNiagaraSimCacheVariable& Other) const
@@ -393,7 +398,7 @@ bool UNiagaraSimCache::BeginWrite(FNiagaraSimCacheCreateParameters InCreateParam
 			CreateParameters.ExplicitCaptureAttributes.AddUnique(NameBuilder.ToString());
 		}
 
-		FNiagaraDataInterfaceUtilities::ForEachDataInterface(
+		ForEachDataInterface(
 			Helper.SystemInstance,
 			[&](const FNiagaraDataInterfaceUtilities::FDataInterfaceUsageContext& UsageContext)
 			{
@@ -405,7 +410,8 @@ bool UNiagaraSimCache::BeginWrite(FNiagaraSimCacheCreateParameters InCreateParam
 					}
 				}
 				return true;
-			}
+			},
+			FNiagaraSimCacheInternal::DISearchOptions
 		);
 	}
 
@@ -423,7 +429,7 @@ bool UNiagaraSimCache::BeginWrite(FNiagaraSimCacheCreateParameters InCreateParam
 	if ( CreateParameters.bAllowDataInterfaceCaching )
 	{
 		TSet<FNiagaraVariableBase> VisitedDataInterfaces;
-		FNiagaraDataInterfaceUtilities::ForEachDataInterface(
+		ForEachDataInterface(
 			Helper.SystemInstance,
 			[&](const FNiagaraVariableBase& Variable, UNiagaraDataInterface* DataInterface)
 			{
@@ -448,7 +454,8 @@ bool UNiagaraSimCache::BeginWrite(FNiagaraSimCacheCreateParameters InCreateParam
 					}
 				}
 				return true;
-			}
+			},
+			FNiagaraSimCacheInternal::DISearchOptions
 		);
 	}
 
@@ -686,7 +693,7 @@ bool UNiagaraSimCache::WriteFrame(UNiagaraComponent* NiagaraComponent, FNiagaraS
 		FString DataInterfaceName;
 		TSet<FNiagaraVariableBase> VisitedDataInterfaces;
 
-		FNiagaraDataInterfaceUtilities::ForEachDataInterface(
+		ForEachDataInterface(
 			Helper.SystemInstance,
 			[&](const FNiagaraVariableBase& Variable, UNiagaraDataInterface* DataInterface)
 			{
@@ -707,7 +714,8 @@ bool UNiagaraSimCache::WriteFrame(UNiagaraComponent* NiagaraComponent, FNiagaraS
 					}
 				}
 				return true;
-			}
+			},
+			FNiagaraSimCacheInternal::DISearchOptions
 		);
 
 		// A data interface failed to write information
@@ -1028,7 +1036,7 @@ bool UNiagaraSimCache::ReadFrame(int32 FrameIndex, float FrameFraction, FNiagara
 		TSet<FNiagaraVariableBase> VisitedDataInterfaces;
 		bool bDataInterfacesSucess = true;
 
-		FNiagaraDataInterfaceUtilities::ForEachDataInterface(
+		ForEachDataInterface(
 			Helper.SystemInstance,
 			[&](const FNiagaraVariableBase& Variable, UNiagaraDataInterface* DataInterface)
 			{
@@ -1047,7 +1055,8 @@ bool UNiagaraSimCache::ReadFrame(int32 FrameIndex, float FrameFraction, FNiagara
 					}
 				}
 				return true;
-			}
+			},
+			FNiagaraSimCacheInternal::DISearchOptions
 		);
 
 
