- Fix cache playback not respecting dataset limit

diff --git a/Engine/Plugins/FX/Niagara/Source/Niagara/Private/NiagaraSimCacheHelper.h b/Engine/Plugins/FX/Niagara/Source/Niagara/Private/NiagaraSimCacheHelper.h
--- a/Engine/Plugins/FX/Niagara/Source/Niagara/Private/NiagaraSimCacheHelper.h
+++ b/Engine/Plugins/FX/Niagara/Source/Niagara/Private/NiagaraSimCacheHelper.h
@@ -684,9 +684,9 @@ struct FNiagaraSimCacheHelper
 		}
 	}
 
-	static void ReadFloatBuffers(int32& iComponent, uint32 FloatCount, TConstArrayView<uint16> ComponentMappingsToDataBuffer, const FNiagaraSimCacheDataBuffers& CacheBuffer, TArrayView<uint8> DestBuffer, uint32 DestStride)
+	static void ReadFloatBuffers(uint32 NumInstances, int32& iComponent, uint32 FloatCount, TConstArrayView<uint16> ComponentMappingsToDataBuffer, const FNiagaraSimCacheDataBuffers& CacheBuffer, TArrayView<uint8> DestBuffer, uint32 DestStride)
 	{
-		const int32 NumInstances = CacheBuffer.NumInstances;
+		const uint32 ComponentStride = CacheBuffer.NumInstances;
 		for (uint32 i = 0; i < FloatCount; ++i)
 		{
 			const uint32 Component = ComponentMappingsToDataBuffer[iComponent++];
@@ -694,15 +694,15 @@ struct FNiagaraSimCacheHelper
 			{
 				continue;
 			}
-			const uint8* SourceFloats = CacheBuffer.FloatData.GetData() + (i * NumInstances * sizeof(float));
+			const uint8* SourceFloats = CacheBuffer.FloatData.GetData() + (i * ComponentStride * sizeof(float));
 			uint8* DestFloats = DestBuffer.GetData() + (DestStride * Component);
 			CheckedMemcpy(DestBuffer, DestFloats, CacheBuffer.FloatData, SourceFloats, sizeof(float) * NumInstances);
 		}
 	}
 
-	static void ReadHalfBuffers(int32& iComponent, uint32 HalfCount, TConstArrayView<uint16> ComponentMappingsToDataBuffer, const FNiagaraSimCacheDataBuffers& CacheBuffer, TArrayView<uint8> DestBuffer, uint32 DestStride)
+	static void ReadHalfBuffers(uint32 NumInstances, int32& iComponent, uint32 HalfCount, TConstArrayView<uint16> ComponentMappingsToDataBuffer, const FNiagaraSimCacheDataBuffers& CacheBuffer, TArrayView<uint8> DestBuffer, uint32 DestStride)
 	{
-		const int32 NumInstances = CacheBuffer.NumInstances;
+		const uint32 ComponentStride = CacheBuffer.NumInstances;
 		for (uint32 i = 0; i < HalfCount; ++i)
 		{
 			const uint32 Component = ComponentMappingsToDataBuffer[iComponent++];
@@ -710,15 +710,15 @@ struct FNiagaraSimCacheHelper
 			{
 				continue;
 			}
-			const uint8* SourceHalfs = CacheBuffer.HalfData.GetData() + (i * NumInstances * sizeof(FFloat16));
+			const uint8* SourceHalfs = CacheBuffer.HalfData.GetData() + (i * ComponentStride * sizeof(FFloat16));
 			uint8* DestHalfs = DestBuffer.GetData() + (DestStride * Component);
 			CheckedMemcpy(DestBuffer, DestHalfs, CacheBuffer.HalfData, SourceHalfs, sizeof(FFloat16) * NumInstances);
 		}
 	}
 
-	static void ReadInt32Buffers(int32& iComponent, uint32 Int32Count, TConstArrayView<uint16> ComponentMappingsToDataBuffer, const FNiagaraSimCacheDataBuffers& CacheBuffer, TArrayView<uint8> DestBuffer, uint32 DestStride)
+	static void ReadInt32Buffers(uint32 NumInstances, int32& iComponent, uint32 Int32Count, TConstArrayView<uint16> ComponentMappingsToDataBuffer, const FNiagaraSimCacheDataBuffers& CacheBuffer, TArrayView<uint8> DestBuffer, uint32 DestStride)
 	{
-		const int32 NumInstances = CacheBuffer.NumInstances;
+		const uint32 ComponentStride = CacheBuffer.NumInstances;
 		for (uint32 i = 0; i < Int32Count; ++i)
 		{
 			const uint32 Component = ComponentMappingsToDataBuffer[iComponent++];
@@ -726,13 +726,13 @@ struct FNiagaraSimCacheHelper
 			{
 				continue;
 			}
-			const uint8* SourceInt32s = CacheBuffer.Int32Data.GetData() + (i * NumInstances * sizeof(int32));
+			const uint8* SourceInt32s = CacheBuffer.Int32Data.GetData() + (i * ComponentStride * sizeof(int32));
 			uint8* DestInt32s = DestBuffer.GetData() + (DestStride * Component);
 			CheckedMemcpy(DestBuffer, DestInt32s, CacheBuffer.Int32Data, SourceInt32s, sizeof(int32) * NumInstances);
 		}
 	}
 
-	static void ReadCustomBuffers(float FrameFraction, float FrameDeltaSeconds, float SimDeltaSeconds, const FTransform& RebaseTransform, uint16 ComponentVelocity, const TConstArrayView<FNiagaraSimCacheDataBuffersLayout::FVariableCopyMapping>& VariableCopyMappingsToDataBuffer, const FNiagaraSimCacheDataBuffers& CacheBufferA, const FNiagaraSimCacheDataBuffers& CacheBufferB, uint8* DestBuffer, uint32 DestStride)
+	static void ReadCustomBuffers(uint32 NumInstances, float FrameFraction, float FrameDeltaSeconds, float SimDeltaSeconds, const FTransform& RebaseTransform, uint16 ComponentVelocity, const TConstArrayView<FNiagaraSimCacheDataBuffersLayout::FVariableCopyMapping>& VariableCopyMappingsToDataBuffer, const FNiagaraSimCacheDataBuffers& CacheBufferA, const FNiagaraSimCacheDataBuffers& CacheBufferB, uint8* DestBuffer, uint32 DestStride)
 	{
 		if (VariableCopyMappingsToDataBuffer.Num() == 0)
 		{
@@ -744,7 +744,7 @@ struct FNiagaraSimCacheHelper
 		VariableCopyDataContext.FrameDeltaSeconds	= FrameDeltaSeconds;
 		VariableCopyDataContext.SimDeltaSeconds		= SimDeltaSeconds;
 		VariableCopyDataContext.PrevFrameFraction	= 1.0f - FMath::Clamp(SimDeltaSeconds / (FrameDeltaSeconds * (1.0f + FrameFraction)), 0.0f, 1.0f);		//-TODO: We are assuming both frames have the same DT here
-		VariableCopyDataContext.NumInstances		= CacheBufferA.NumInstances;
+		VariableCopyDataContext.NumInstances		= NumInstances;
 		VariableCopyDataContext.RebaseTransform		= RebaseTransform;
 		VariableCopyDataContext.InterpMappings		= CacheBufferA.InterpMapping;
 		VariableCopyDataContext.DestStride			= DestStride;
@@ -772,16 +772,18 @@ struct FNiagaraSimCacheHelper
 
 	void ReadDataBuffer(float FrameFraction, float FrameDeltaSeconds, float SimDeltaSeconds, const FTransform& RebaseTransform, const FNiagaraSimCacheDataBuffersLayout& CacheLayout, const FNiagaraSimCacheDataBuffers& CacheBufferA, const FNiagaraSimCacheDataBuffers& CacheBufferB, FNiagaraDataSet& DataSet)
 	{
+		const uint32 NumInstances = FMath::Min(CacheBufferA.NumInstances, DataSet.GetMaxAllocationCount());
+
 		FNiagaraDataBuffer& DataBuffer = DataSet.BeginSimulate();
-		DataBuffer.Allocate(CacheBufferA.NumInstances);
-		DataBuffer.SetNumInstances(CacheBufferA.NumInstances);
+		DataBuffer.Allocate(NumInstances);
+		DataBuffer.SetNumInstances(NumInstances);
 		if (CacheBufferA.NumInstances > 0 )
 		{
 			int32 iComponent = 0;
-			ReadFloatBuffers(iComponent, CacheLayout.FloatCount, CacheLayout.CacheBufferReadInfo.ComponentMappingsToDataBuffer, CacheBufferA, MakeArrayView(DataBuffer.GetComponentPtrFloat(0), DataBuffer.GetFloatBuffer().Num()), DataBuffer.GetFloatStride());
-			ReadHalfBuffers(iComponent, CacheLayout.HalfCount, CacheLayout.CacheBufferReadInfo.ComponentMappingsToDataBuffer, CacheBufferA, MakeArrayView(DataBuffer.GetComponentPtrHalf(0), DataBuffer.GetHalfBuffer().Num()), DataBuffer.GetHalfStride());
-			ReadInt32Buffers(iComponent, CacheLayout.Int32Count, CacheLayout.CacheBufferReadInfo.ComponentMappingsToDataBuffer, CacheBufferA, MakeArrayView(DataBuffer.GetComponentPtrInt32(0), DataBuffer.GetInt32Buffer().Num()), DataBuffer.GetInt32Stride());
-			ReadCustomBuffers(FrameFraction, FrameDeltaSeconds, SimDeltaSeconds, RebaseTransform, CacheLayout.ComponentVelocity, CacheLayout.CacheBufferReadInfo.VariableCopyMappingsToDataBuffer, CacheBufferA, CacheBufferB, DataBuffer.GetComponentPtrFloat(0), DataBuffer.GetFloatStride());
+			ReadFloatBuffers(NumInstances, iComponent, CacheLayout.FloatCount, CacheLayout.CacheBufferReadInfo.ComponentMappingsToDataBuffer, CacheBufferA, MakeArrayView(DataBuffer.GetComponentPtrFloat(0), DataBuffer.GetFloatBuffer().Num()), DataBuffer.GetFloatStride());
+			ReadHalfBuffers(NumInstances, iComponent, CacheLayout.HalfCount, CacheLayout.CacheBufferReadInfo.ComponentMappingsToDataBuffer, CacheBufferA, MakeArrayView(DataBuffer.GetComponentPtrHalf(0), DataBuffer.GetHalfBuffer().Num()), DataBuffer.GetHalfStride());
+			ReadInt32Buffers(NumInstances, iComponent, CacheLayout.Int32Count, CacheLayout.CacheBufferReadInfo.ComponentMappingsToDataBuffer, CacheBufferA, MakeArrayView(DataBuffer.GetComponentPtrInt32(0), DataBuffer.GetInt32Buffer().Num()), DataBuffer.GetInt32Stride());
+			ReadCustomBuffers(NumInstances, FrameFraction, FrameDeltaSeconds, SimDeltaSeconds, RebaseTransform, CacheLayout.ComponentVelocity, CacheLayout.CacheBufferReadInfo.VariableCopyMappingsToDataBuffer, CacheBufferA, CacheBufferB, DataBuffer.GetComponentPtrFloat(0), DataBuffer.GetFloatStride());
 		}
 
 		//-TODO:DestinationDataBuffer.SetIDTable(CacheBufferA.IDToIndexTable);
@@ -817,7 +819,7 @@ struct FNiagaraSimCacheHelper
 		ENQUEUE_RENDER_COMMAND(NiagaraSimCacheGpuReadFrame)(
 			[DispathInterface, GPUExecContext=EmitterInstance.GetGPUContext(), FrameFraction=InFrameFraction, FrameDeltaSeconds=InFrameDeltaSeconds, SimDeltaSeconds=InSimDeltaSeconds, RebaseTransform=InRebaseTransform, CacheLayout=&InCacheLayout, CacheBufferA=&InCacheBufferA, CacheBufferB=&InCacheBufferB, DataSet=&InDataSet, PendingCommandsCounter=&InPendingCommandsCounter](FRHICommandListImmediate& RHICmdList)
 			{
-				const int32 NumInstances = CacheBufferA->NumInstances;
+				const uint32 NumInstances = FMath::Min(CacheBufferA->NumInstances, DataSet->GetMaxAllocationCount());
 
 				// Set Instance Count
 				{
@@ -850,8 +852,8 @@ struct FNiagaraSimCacheHelper
 					{
 						FRWBuffer& RWBuffer = DataBuffer.GetGPUBufferFloat();
 						uint8* RWBufferMemory = reinterpret_cast<uint8*>(RHICmdList.LockBuffer(RWBuffer.Buffer, 0, RWBuffer.NumBytes, RLM_WriteOnly));
-						ReadFloatBuffers(iComponent, CacheLayout->FloatCount, CacheLayout->CacheBufferReadInfo_RT.ComponentMappingsToDataBuffer, *CacheBufferA, MakeArrayView(RWBufferMemory, RWBuffer.NumBytes), DataBuffer.GetFloatStride());
-						ReadCustomBuffers(FrameFraction, FrameDeltaSeconds, SimDeltaSeconds, RebaseTransform, CacheLayout->ComponentVelocity, CacheLayout->CacheBufferReadInfo_RT.VariableCopyMappingsToDataBuffer, *CacheBufferA, *CacheBufferB, RWBufferMemory, DataBuffer.GetFloatStride());
+						ReadFloatBuffers(NumInstances, iComponent, CacheLayout->FloatCount, CacheLayout->CacheBufferReadInfo_RT.ComponentMappingsToDataBuffer, *CacheBufferA, MakeArrayView(RWBufferMemory, RWBuffer.NumBytes), DataBuffer.GetFloatStride());
+						ReadCustomBuffers(NumInstances, FrameFraction, FrameDeltaSeconds, SimDeltaSeconds, RebaseTransform, CacheLayout->ComponentVelocity, CacheLayout->CacheBufferReadInfo_RT.VariableCopyMappingsToDataBuffer, *CacheBufferA, *CacheBufferB, RWBufferMemory, DataBuffer.GetFloatStride());
 						RHICmdList.UnlockBuffer(RWBuffer.Buffer);
 					}
 
@@ -860,7 +862,7 @@ struct FNiagaraSimCacheHelper
 					{
 						FRWBuffer& RWBuffer = DataBuffer.GetGPUBufferHalf();
 						uint8* RWBufferMemory = reinterpret_cast<uint8*>(RHICmdList.LockBuffer(RWBuffer.Buffer, 0, RWBuffer.NumBytes, RLM_WriteOnly));
-						ReadHalfBuffers(iComponent, CacheLayout->HalfCount, CacheLayout->CacheBufferReadInfo_RT.ComponentMappingsToDataBuffer, *CacheBufferA, MakeArrayView(RWBufferMemory, RWBuffer.NumBytes), DataBuffer.GetHalfStride());
+						ReadHalfBuffers(NumInstances, iComponent, CacheLayout->HalfCount, CacheLayout->CacheBufferReadInfo_RT.ComponentMappingsToDataBuffer, *CacheBufferA, MakeArrayView(RWBufferMemory, RWBuffer.NumBytes), DataBuffer.GetHalfStride());
 						RHICmdList.UnlockBuffer(RWBuffer.Buffer);
 					}
 
@@ -869,7 +871,7 @@ struct FNiagaraSimCacheHelper
 					{
 						FRWBuffer& RWBuffer = DataBuffer.GetGPUBufferInt();
 						uint8* RWBufferMemory = reinterpret_cast<uint8*>(RHICmdList.LockBuffer(RWBuffer.Buffer, 0, RWBuffer.NumBytes, RLM_WriteOnly));
-						ReadInt32Buffers(iComponent, CacheLayout->Int32Count, CacheLayout->CacheBufferReadInfo_RT.ComponentMappingsToDataBuffer, *CacheBufferA, MakeArrayView(RWBufferMemory, RWBuffer.NumBytes), DataBuffer.GetInt32Stride());
+						ReadInt32Buffers(NumInstances, iComponent, CacheLayout->Int32Count, CacheLayout->CacheBufferReadInfo_RT.ComponentMappingsToDataBuffer, *CacheBufferA, MakeArrayView(RWBufferMemory, RWBuffer.NumBytes), DataBuffer.GetInt32Stride());
 						RHICmdList.UnlockBuffer(RWBuffer.Buffer);
 					}
 				}
