- Fix for baker not capturing anything if the system was already complete

diff --git a/Engine/Plugins/FX/Niagara/Source/NiagaraEditor/Private/NiagaraBakerRenderer.cpp b/Engine/Plugins/FX/Niagara/Source/NiagaraEditor/Private/NiagaraBakerRenderer.cpp
--- a/Engine/Plugins/FX/Niagara/Source/NiagaraEditor/Private/NiagaraBakerRenderer.cpp
+++ b/Engine/Plugins/FX/Niagara/Source/NiagaraEditor/Private/NiagaraBakerRenderer.cpp
@@ -312,6 +312,11 @@ void FNiagaraBakerRenderer::SetAbsoluteTime(float AbsoluteTime, bool bShouldTick
 		return;
 	}
 
+	if (!PreviewComponent->IsActive() && (AbsoluteTime < PreviewComponent->GetDesiredAge()))
+	{
+		PreviewComponent->ReinitializeSystem();
+	}
+
 	PreviewComponent->SetSeekDelta(BakerSettings->GetSeekDelta());
 	PreviewComponent->SeekToDesiredAge(AbsoluteTime);
 
