- Fix duplicated loc text key

diff --git a/Engine/Plugins/FX/Niagara/Source/NiagaraEditor/Private/ViewModels/Stack/NiagaraStackStatelessEmitterSimulateGroup.cpp b/Engine/Plugins/FX/Niagara/Source/NiagaraEditor/Private/ViewModels/Stack/NiagaraStackStatelessEmitterSimulateGroup.cpp
--- a/Engine/Plugins/FX/Niagara/Source/NiagaraEditor/Private/ViewModels/Stack/NiagaraStackStatelessEmitterSimulateGroup.cpp
+++ b/Engine/Plugins/FX/Niagara/Source/NiagaraEditor/Private/ViewModels/Stack/NiagaraStackStatelessEmitterSimulateGroup.cpp
@@ -180,7 +180,7 @@ bool UNiagaraStackStatelessModuleItem::TestCanDeleteWithMessage(FText& OutCanDel
 
 FText UNiagaraStackStatelessModuleItem::GetDeleteTransactionText() const
 {
-	return LOCTEXT("DeleteStatelessModule", "Delete module from stateless emitter.");
+	return LOCTEXT("DeleteStatelessModuleTransaction", "Delete module from stateless emitter.");
 }
 
 void UNiagaraStackStatelessModuleItem::Delete()
