- Fix override parameters not working with cache playback

diff --git a/Engine/Plugins/FX/Niagara/Source/Niagara/Private/NiagaraSystemInstance.cpp b/Engine/Plugins/FX/Niagara/Source/Niagara/Private/NiagaraSystemInstance.cpp
--- a/Engine/Plugins/FX/Niagara/Source/Niagara/Private/NiagaraSystemInstance.cpp
+++ b/Engine/Plugins/FX/Niagara/Source/Niagara/Private/NiagaraSystemInstance.cpp
@@ -2273,7 +2273,13 @@ void FNiagaraSystemInstance::SimCacheTick_GameThread(UNiagaraSimCache* SimCache,
 	CachedDeltaSeconds = DeltaSeconds;
 	//FixedBounds_CNC = FixedBounds_GT;
 
+	// We still need to tick the override parameters as these might pass through to material bindings
 	//TickInstanceParameters_GameThread(DeltaSeconds);
+	if (OverrideParameters)
+	{
+		OverrideParameters->ResolvePositions(GetLWCConverter());
+		OverrideParameters->Tick();
+	}
 
 	//TickDataInterfaces(DeltaSeconds, false);
 
