- Fix for use after free
- Resource names for Niagara must be character arrays not other types of strings, as the internal pool may hold a reference beyond the lifetime of the Niagara data

diff --git a/Engine/Plugins/FX/Niagara/Source/Niagara/Classes/NiagaraDataInterfaceGrid2DCollection.h b/Engine/Plugins/FX/Niagara/Source/Niagara/Classes/NiagaraDataInterfaceGrid2DCollection.h
--- a/Engine/Plugins/FX/Niagara/Source/Niagara/Classes/NiagaraDataInterfaceGrid2DCollection.h
+++ b/Engine/Plugins/FX/Niagara/Source/Niagara/Classes/NiagaraDataInterfaceGrid2DCollection.h
@@ -71,10 +71,6 @@ struct FGrid2DCollectionRWInstanceData_RenderThread
 	TArray<int32> VarComponents;
 	TArray<uint32> Offsets;
 
-#if !UE_BUILD_SHIPPING
-	FString GridTextureResourceName;
-#endif
-
 #if WITH_EDITORONLY_DATA
 	bool bPreviewGrid = false;
 	FIntVector4 PreviewAttribute = FIntVector4(INDEX_NONE, INDEX_NONE, INDEX_NONE, INDEX_NONE);
diff --git a/Engine/Plugins/FX/Niagara/Source/Niagara/Private/NiagaraDataInterfaceGrid2DCollection.cpp b/Engine/Plugins/FX/Niagara/Source/Niagara/Private/NiagaraDataInterfaceGrid2DCollection.cpp
--- a/Engine/Plugins/FX/Niagara/Source/Niagara/Private/NiagaraDataInterfaceGrid2DCollection.cpp
+++ b/Engine/Plugins/FX/Niagara/Source/Niagara/Private/NiagaraDataInterfaceGrid2DCollection.cpp
@@ -2042,21 +2042,6 @@ bool UNiagaraDataInterfaceGrid2DCollection::InitPerInstanceData(void* PerInstanc
 			TargetData->VarComponents.Add(RT_InstanceData.Vars[i].GetType().GetSize() / sizeof(float));
 		}
 
-#if !UE_BUILD_SHIPPING
-		{
-			TStringBuilder<128> StringBuilder;
-			TargetData->SourceDIName.AppendString(StringBuilder);
-			StringBuilder.Append("_");
-			for (FName CurrName : TargetData->Vars)
-			{
-				CurrName.AppendString(StringBuilder);
-				StringBuilder.Append("_");
-			}
-			StringBuilder.Append("Grid2DCollection");
-			TargetData->GridTextureResourceName = StringBuilder;
-		}
-#endif // WITH_SAHDER_NAMES
-
 #if WITH_EDITORONLY_DATA
 		TargetData->bPreviewGrid = RT_InstanceData.bPreviewGrid;
 		TargetData->PreviewAttribute = RT_InstanceData.PreviewAttribute;
@@ -2680,12 +2665,7 @@ void FGrid2DCollectionRWInstanceData_RenderThread::BeginSimulate(FRDGBuilder& Gr
 		Buffers.Emplace(DestinationData);
 
 		const FRDGTextureDesc TextureDesc = FRDGTextureDesc::Create2DArray(NumCells, PixelFormat.GetValue(), FClearValueBinding::Black, ETextureCreateFlags::ShaderResource | ETextureCreateFlags::UAV, uint16(NumAttributes));
-
-#if UE_BUILD_SHIPPING
 		DestinationData->Initialize(GraphBuilder, TEXT("Grid2D::GridTexture"), TextureDesc);
-#else
-		DestinationData->Initialize(GraphBuilder, *GridTextureResourceName, TextureDesc);
-#endif
 
 		// This destination buffer will sometimes have old data in it.  Force it to clear.
 		AddClearUAVPass(GraphBuilder, DestinationData->GetOrCreateUAV(GraphBuilder), FVector4f(ForceInitToZero));
diff --git a/Engine/Plugins/FX/Niagara/Source/Niagara/Private/NiagaraDataInterfaceGrid3DCollection.cpp b/Engine/Plugins/FX/Niagara/Source/Niagara/Private/NiagaraDataInterfaceGrid3DCollection.cpp
--- a/Engine/Plugins/FX/Niagara/Source/Niagara/Private/NiagaraDataInterfaceGrid3DCollection.cpp
+++ b/Engine/Plugins/FX/Niagara/Source/Niagara/Private/NiagaraDataInterfaceGrid3DCollection.cpp
@@ -3991,8 +3991,8 @@ void FGrid3DCollectionRWInstanceData_RenderThread::BeginSimulate(FRDGBuilder& Gr
 		const FIntVector TextureSize(NumCells.X * NumTiles.X, NumCells.Y * NumTiles.Y, NumCells.Z * NumTiles.Z);
 		const FRDGTextureDesc TextureDesc = FRDGTextureDesc::Create3D(TextureSize, PixelFormat.GetValue(), FClearValueBinding::Black, ETextureCreateFlags::ShaderResource | ETextureCreateFlags::UAV);
 
-		const TCHAR* GridTextureName = TEXT("Grid3D::GridTexture");
 	#if 0
+		// Note: Invalid as the render target pool holds onto a pointer and can exist beyond the lifetime of the grid
 		TStringBuilder<128> StringBuilder;
 		SourceDIName.AppendString(StringBuilder);
 		StringBuilder.Append("_");
@@ -4004,7 +4004,7 @@ void FGrid3DCollectionRWInstanceData_RenderThread::BeginSimulate(FRDGBuilder& Gr
 		StringBuilder.Append("Grid3DCollection");
 		GridTextureName = *StringBuilder;
 	#endif
-		DestinationData->Initialize(GraphBuilder, GridTextureName, TextureDesc);
+		DestinationData->Initialize(GraphBuilder, TEXT("Grid3D::GridTexture"), TextureDesc);
 
 		// This destination buffer will sometimes have old data in it.  Force it to clear.
 		AddClearUAVPass(GraphBuilder, DestinationData->GetOrCreateUAV(GraphBuilder), FVector4f(ForceInitToZero));
diff --git a/Engine/Plugins/FX/Niagara/Source/Niagara/Private/NiagaraRenderGraphUtils.cpp b/Engine/Plugins/FX/Niagara/Source/Niagara/Private/NiagaraRenderGraphUtils.cpp
--- a/Engine/Plugins/FX/Niagara/Source/Niagara/Private/NiagaraRenderGraphUtils.cpp
+++ b/Engine/Plugins/FX/Niagara/Source/Niagara/Private/NiagaraRenderGraphUtils.cpp
@@ -6,7 +6,7 @@
 
 //////////////////////////////////////////////////////////////////////////
 
-void FNiagaraPooledRWBuffer::Initialize(FRDGBuilder& GraphBuilder, const TCHAR* ResourceName, EPixelFormat InPixelFormat, const FRDGBufferDesc& BufferDesc)
+void FNiagaraPooledRWBuffer::InitializeInternal(FRDGBuilder& GraphBuilder, const TCHAR* ResourceName, EPixelFormat InPixelFormat, const FRDGBufferDesc& BufferDesc)
 {
 	Release();
 
@@ -20,11 +20,11 @@ void FNiagaraPooledRWBuffer::Initialize(FRDGBuilder& GraphBuilder, const TCHAR*
 #endif
 }
 
-void FNiagaraPooledRWBuffer::Initialize(FRDGBuilder& GraphBuilder, const TCHAR* ResourceName, EPixelFormat InPixelFormat, const uint32 BytesPerElemenet, const uint32 NumElements, EBufferUsageFlags UsageFlags)
+void FNiagaraPooledRWBuffer::InitializeInternal(FRDGBuilder& GraphBuilder, const TCHAR* ResourceName, EPixelFormat InPixelFormat, const uint32 BytesPerElemenet, const uint32 NumElements, EBufferUsageFlags UsageFlags)
 {
 	FRDGBufferDesc BufferDesc = FRDGBufferDesc::CreateBufferDesc(BytesPerElemenet, NumElements);
 	BufferDesc.Usage |= UsageFlags;
-	Initialize(GraphBuilder, ResourceName, InPixelFormat, BufferDesc);
+	InitializeInternal(GraphBuilder, ResourceName, InPixelFormat, BufferDesc);
 }
 
 void FNiagaraPooledRWBuffer::Release()
@@ -89,7 +89,7 @@ void FNiagaraPooledRWBuffer::EndGraphUsage()
 
 //////////////////////////////////////////////////////////////////////////
 
-void FNiagaraPooledRWTexture::Initialize(FRDGBuilder& GraphBuilder, const TCHAR* ResourceName, const FRDGTextureDesc& TextureDesc)
+void FNiagaraPooledRWTexture::InitializeInternal(FRDGBuilder& GraphBuilder, const TCHAR* ResourceName, const FRDGTextureDesc& TextureDesc)
 {
 	Release();
 
diff --git a/Engine/Plugins/FX/Niagara/Source/Niagara/Public/NiagaraRenderGraphUtils.h b/Engine/Plugins/FX/Niagara/Source/Niagara/Public/NiagaraRenderGraphUtils.h
--- a/Engine/Plugins/FX/Niagara/Source/Niagara/Public/NiagaraRenderGraphUtils.h
+++ b/Engine/Plugins/FX/Niagara/Source/Niagara/Public/NiagaraRenderGraphUtils.h
@@ -9,8 +9,23 @@ struct FNiagaraPooledRWBuffer
 {
 	~FNiagaraPooledRWBuffer() { Release(); }
 
-	NIAGARA_API void Initialize(FRDGBuilder& GraphBuilder, const TCHAR* ResourceName, EPixelFormat PixelFormat, const FRDGBufferDesc& BufferDesc);
-	NIAGARA_API void Initialize(FRDGBuilder& GraphBuilder, const TCHAR* ResourceName, EPixelFormat PixelFormat, const uint32 BytesPerElemenet, const uint32 NumElements, EBufferUsageFlags UsageFlags = EBufferUsageFlags::None);
+	template<typename TResourceName>
+	void Initialize(FRDGBuilder& GraphBuilder, const TResourceName& ResourceName, EPixelFormat InPixelFormat, const FRDGBufferDesc& BufferDesc)
+	{
+		static_assert(TIsArrayOrRefOfTypeByPredicate<TResourceName, TIsCharEncodingCompatibleWithTCHAR>::Value, "Resource name must be a character array.");
+		InitializeInternal(GraphBuilder, ResourceName, InPixelFormat, BufferDesc);
+	}
+	template<typename TResourceName>
+	void Initialize(FRDGBuilder& GraphBuilder, const TResourceName& ResourceName, EPixelFormat InPixelFormat, const uint32 BytesPerElemenet, const uint32 NumElements, EBufferUsageFlags UsageFlags = EBufferUsageFlags::None)
+	{
+		static_assert(TIsArrayOrRefOfTypeByPredicate<TResourceName, TIsCharEncodingCompatibleWithTCHAR>::Value, "Resource name must be a character array.");
+		InitializeInternal(GraphBuilder, ResourceName, InPixelFormat, BytesPerElemenet, NumElements, UsageFlags);
+	}
+
+private:
+	NIAGARA_API void InitializeInternal(FRDGBuilder& GraphBuilder, const TCHAR* ResourceName, EPixelFormat PixelFormat, const FRDGBufferDesc& BufferDesc);
+	NIAGARA_API void InitializeInternal(FRDGBuilder& GraphBuilder, const TCHAR* ResourceName, EPixelFormat PixelFormat, const uint32 BytesPerElemenet, const uint32 NumElements, EBufferUsageFlags UsageFlags = EBufferUsageFlags::None);
+public:
 	NIAGARA_API void Release();
 
 	bool IsValid() const { return Buffer.IsValid(); }
@@ -42,7 +57,15 @@ struct FNiagaraPooledRWTexture
 {
 	~FNiagaraPooledRWTexture() { Release(); }
 
-	NIAGARA_API void Initialize(FRDGBuilder& GraphBuilder, const TCHAR* ResourceName, const FRDGTextureDesc& TextureDesc);
+	template<typename TResourceName>
+	void Initialize(FRDGBuilder& GraphBuilder, const TResourceName& ResourceName, const FRDGTextureDesc& TextureDesc)
+	{
+		static_assert(TIsArrayOrRefOfTypeByPredicate<TResourceName, TIsCharEncodingCompatibleWithTCHAR>::Value, "Resource name must be a character array.");
+		InitializeInternal(GraphBuilder, ResourceName, TextureDesc);
+	}
+private:
+	void InitializeInternal(FRDGBuilder& GraphBuilder, const TCHAR* ResourceName, const FRDGTextureDesc& TextureDesc);
+public:
 	NIAGARA_API void Release();
 
 	bool IsValid() const { return Texture.IsValid(); }
