[Backout] - CL31517967
Original CL Desc
-----------------------------------------------------------------
Fixed issue where calls to DeactivateImmediate would not properly complete systems that were culled by scalability.

diff --git a/Engine/Plugins/FX/Niagara/Source/Niagara/Private/NiagaraComponent.cpp b/Engine/Plugins/FX/Niagara/Source/Niagara/Private/NiagaraComponent.cpp
--- a/Engine/Plugins/FX/Niagara/Source/Niagara/Private/NiagaraComponent.cpp
+++ b/Engine/Plugins/FX/Niagara/Source/Niagara/Private/NiagaraComponent.cpp
@@ -1532,6 +1532,8 @@ void UNiagaraComponent::DeactivateImmediateInternal(bool bIsScalabilityCull)
 	SCOPE_CYCLE_COUNTER(STAT_NiagaraComponentDeactivate);
 	CSV_SCOPED_TIMING_STAT_EXCLUSIVE(Effects);
 
+	Super::Deactivate();
+
 	bool bWasCulledByScalability = bIsCulledByScalability;
 
 	//UE_LOG(LogNiagara, Log, TEXT("DeactivateImmediate: %p - %s - %s"), this, *Asset->GetName(), bIsScalabilityCull ? TEXT("Scalability") : TEXT(""));
@@ -1549,28 +1551,16 @@ void UNiagaraComponent::DeactivateImmediateInternal(bool bIsScalabilityCull)
 		UnregisterWithScalabilityManager();
 	}
 
+	SetActiveFlag(false);
 
-	if (IsActive() && SystemInstanceController)
+	if (SystemInstanceController)
 	{
-		Super::Deactivate();
 		SystemInstanceController->Deactivate(true);
 	}
-	else 
+	else if (bWasCulledByScalability && !bIsCulledByScalability)//We were culled by scalability but no longer, ensure we've handled completion correctly. E.g. returned to the pool etc.
 	{
-		Super::Deactivate();
-
-		if (bWasCulledByScalability && !bIsCulledByScalability)//We were culled by scalability but no longer, ensure we've handled completion correctly. E.g. returned to the pool etc.
-		{
-			OnSystemComplete(true);
-		}
-	}
-
-#if !UE_BUILD_SHIPPING
-	if(IsActive())
-	{
-		UE_LOG(LogNiagara, Error, TEXT("NiagaraSystem is still active after it should have been deactivated.\nSystem: %s\nComponent: %s"), *GetNameSafe(GetAsset()), *GetNameSafe(this));
+		OnSystemComplete(true);
 	}
-#endif
 }
 
 bool UNiagaraComponent::ShouldPreCull()
