- Remove all other preview world components from the capture list as it's confusing
- Add information about which world type the component comes from

diff --git a/Engine/Plugins/FX/Niagara/Source/NiagaraEditor/Private/Widgets/SNiagaraDebugCaptureView.cpp b/Engine/Plugins/FX/Niagara/Source/NiagaraEditor/Private/Widgets/SNiagaraDebugCaptureView.cpp
--- a/Engine/Plugins/FX/Niagara/Source/NiagaraEditor/Private/Widgets/SNiagaraDebugCaptureView.cpp
+++ b/Engine/Plugins/FX/Niagara/Source/NiagaraEditor/Private/Widgets/SNiagaraDebugCaptureView.cpp
@@ -23,11 +23,13 @@ void SNiagaraDebugCaptureView::OnNumFramesChanged(int32 InNumFrames)
 
 void SNiagaraDebugCaptureView::CreateComponentSelectionMenuContent(FMenuBuilder& MenuBuilder) 
 {
-	for(TObjectIterator<UNiagaraComponent> NiagaraComponent; NiagaraComponent; ++NiagaraComponent)
+	for(TObjectIterator<UNiagaraComponent> NiagaraComponentIt; NiagaraComponentIt; ++NiagaraComponentIt)
 	{
+		UNiagaraComponent* NiagaraComponent = *NiagaraComponentIt;
+
 		// Ignore dying or CDO versions of data..
 		// No need to check the unreachable flag here as TObjectIterator already does that
-		if(!IsValid(*NiagaraComponent) || NiagaraComponent->HasAnyFlags(RF_ClassDefaultObject))
+		if(!IsValid(NiagaraComponent) || NiagaraComponent->HasAnyFlags(RF_ClassDefaultObject))
 		{
 			continue;
 		}
@@ -44,21 +46,22 @@ void SNiagaraDebugCaptureView::CreateComponentSelectionMenuContent(FMenuBuilder&
 			continue;
 		}
 
-		// Exclude components in a preview world with a sim cache attached, this stops us being able to select sim cache editor components
-		if (World->IsPreviewWorld() && NiagaraComponent->GetSimCache() != nullptr)
+		// Only allow the component from our preview world to exist in the component list
+		// Without this test things like sim cache previews, or the baker's component will show up in the capture list which is confusing
+		if (World->IsPreviewWorld() && NiagaraComponent != SystemViewModel->GetPreviewComponent())
 		{
 			continue;
 		}
 
 		FText ComponentName;
 		FText ComponentTooltip;
-		GetComponentNameAndTooltip(*NiagaraComponent, ComponentName, ComponentTooltip);
+		GetComponentNameAndTooltip(NiagaraComponent, ComponentName, ComponentTooltip);
 
 		MenuBuilder.AddMenuEntry(
 			ComponentName,
 			ComponentTooltip,
 			FSlateIcon(),
-			FUIAction(FExecuteAction::CreateLambda([&, WeakNiagaraComponent=MakeWeakObjectPtr(*NiagaraComponent)]()
+			FUIAction(FExecuteAction::CreateLambda([&, WeakNiagaraComponent=MakeWeakObjectPtr(NiagaraComponent)]()
 			{
 				WeakTargetComponent = WeakNiagaraComponent;
 			})));
@@ -82,8 +85,9 @@ void SNiagaraDebugCaptureView::GetComponentNameAndTooltip(const UNiagaraComponen
     else
     {
 	    const UWorld* World = InComponent->GetWorld();
+		const EWorldType::Type WorldType = World ? EWorldType::Type(World->WorldType) : EWorldType::None;
     	const AActor* Actor = InComponent->GetOwner();
-    	OutName = FText::Format(LOCTEXT("SourceComponentLabel","World: \"{0}\" Actor: \"{1}\""), World ? FText::FromString(World->GetName()) : FText::GetEmpty(), Actor ? FText::FromString(Actor->GetActorNameOrLabel()) : FText::GetEmpty());
+    	OutName = FText::Format(LOCTEXT("SourceComponentLabel","World: \"{0} - {1}\" Actor: \"{2}\""), World ? FText::FromString(World->GetName()) : FText::GetEmpty(), FText::FromString(LexToString(WorldType)), Actor ? FText::FromString(Actor->GetActorNameOrLabel()) : FText::GetEmpty());
     	OutTooltip = OutName;
     }
 }
