Fixes for ensures related to emitter constant resolvers.  Handles both disabled emitters and when a constant resolver isn't really provided to the parameter map builder.

diff --git a/Engine/Plugins/FX/Niagara/Source/NiagaraEditor/Private/NiagaraGraphDigest.cpp b/Engine/Plugins/FX/Niagara/Source/NiagaraEditor/Private/NiagaraGraphDigest.cpp
--- a/Engine/Plugins/FX/Niagara/Source/NiagaraEditor/Private/NiagaraGraphDigest.cpp
+++ b/Engine/Plugins/FX/Niagara/Source/NiagaraEditor/Private/NiagaraGraphDigest.cpp
@@ -2012,6 +2012,7 @@ FNiagaraCompilationNodeEmitter::FNiagaraCompilationNodeEmitter(const UNiagaraNod
 : FNiagaraCompilationNode(ENodeType::Emitter, InNode, Context)
 {
 	EmitterID = InNode->GetEmitterID();
+	EmitterHandleID = InNode->GetEmitterHandleId();
 	EmitterUniqueName = InNode->GetEmitterUniqueName();
 
 	if (const UNiagaraGraph* DependentGraph = InNode->GetCalledGraph())
@@ -2026,7 +2027,7 @@ FNiagaraCompilationNodeEmitter::FNiagaraCompilationNodeEmitter(const UNiagaraNod
 	Usage = InNode->GetUsage();
 	EmitterName = InNode->GetName();
 	EmitterPathName = InNode->GetPathName();
-	EmitterHandleIdString = InNode->GetEmitterHandleId().ToString(EGuidFormats::Digits);
+	EmitterHandleIdString = EmitterHandleID.ToString(EGuidFormats::Digits);
 	EmitterUniqueFName = *EmitterUniqueName;
 }
 
@@ -2034,6 +2035,7 @@ FNiagaraCompilationNodeEmitter::FNiagaraCompilationNodeEmitter(const UNiagaraNod
 FNiagaraCompilationNodeEmitter::FNiagaraCompilationNodeEmitter(const FNiagaraCompilationNodeEmitter& InNode, FNiagaraCompilationGraphDuplicateContext& Context)
 	: FNiagaraCompilationNode(InNode, Context)
 	, EmitterID(InNode.EmitterID)
+	, EmitterHandleID(InNode.EmitterHandleID)
 	, EmitterUniqueName(InNode.EmitterUniqueName)
 	, EmitterName(InNode.EmitterName)
 	, EmitterPathName(InNode.EmitterPathName)
@@ -2068,6 +2070,21 @@ void FNiagaraCompilationNodeEmitter::BuildParameterMapHistory(FParameterMapHisto
 		return;
 	}
 
+	if (Builder.ExclusiveEmitterHandle.IsSet() && Builder.ExclusiveEmitterHandle != EmitterHandleID)
+	{
+		RouteParameterMapAroundMe(Builder);
+		return;
+	}
+
+	const FNiagaraFixedConstantResolver* ChildConstantResolver = Builder.ConstantResolver->FindChildResolver(EmitterUniqueFName);
+	if (!ChildConstantResolver)
+	{
+		// if no child resolver was found for the specified emitter, that means that the emitter is likely not enabled and so we can proceed without
+		// processing it
+		RouteParameterMapAroundMe(Builder);
+		return;
+	}
+
 	int32 ParamMapIdx = INDEX_NONE;
 	if (InputPins[0].LinkedTo)
 	{
@@ -2105,15 +2122,7 @@ void FNiagaraCompilationNodeEmitter::BuildParameterMapHistory(FParameterMapHisto
 
 			// Build up a new parameter map history with all the child graph nodes..
 			FParameterMapHistoryBuilder ChildBuilder;
-			const FNiagaraFixedConstantResolver* ChildConstantResolver = Builder.ConstantResolver->FindChildResolver(EmitterUniqueFName);
-			if (ensure(ChildConstantResolver))
-			{
-				*ChildBuilder.ConstantResolver = *ChildConstantResolver;
-			}
-			else
-			{
-				*ChildBuilder.ConstantResolver = *Builder.ConstantResolver;
-			}
+			*ChildBuilder.ConstantResolver = *ChildConstantResolver;
 			ChildBuilder.RegisterEncounterableVariables(Builder.GetEncounterableVariables());
 			ChildBuilder.EnableScriptAllowList(true, Usage);
 
diff --git a/Engine/Plugins/FX/Niagara/Source/NiagaraEditor/Private/NiagaraGraphDigest.h b/Engine/Plugins/FX/Niagara/Source/NiagaraEditor/Private/NiagaraGraphDigest.h
--- a/Engine/Plugins/FX/Niagara/Source/NiagaraEditor/Private/NiagaraGraphDigest.h
+++ b/Engine/Plugins/FX/Niagara/Source/NiagaraEditor/Private/NiagaraGraphDigest.h
@@ -394,6 +394,7 @@ public:
 	virtual void Compile(FTranslator* Translator, TArray<int32>& Outputs) const;
 
 	FNiagaraEmitterID EmitterID;
+	FGuid EmitterHandleID;
 	FString EmitterUniqueName;
 	FString EmitterName;
 	FString EmitterPathName;
diff --git a/Engine/Plugins/FX/Niagara/Source/NiagaraEditor/Private/NiagaraNodeEmitter.cpp b/Engine/Plugins/FX/Niagara/Source/NiagaraEditor/Private/NiagaraNodeEmitter.cpp
--- a/Engine/Plugins/FX/Niagara/Source/NiagaraEditor/Private/NiagaraNodeEmitter.cpp
+++ b/Engine/Plugins/FX/Niagara/Source/NiagaraEditor/Private/NiagaraNodeEmitter.cpp
@@ -351,6 +351,10 @@ void UNiagaraNodeEmitter::BuildParameterMapHistory(FNiagaraParameterMapHistoryBu
 		});
 	}
 
+	const FCompileConstantResolver ChildConstantResolver = EmitterHandle
+		? OutHistory.ConstantResolver->AsEmitter(EmitterHandle->GetInstance())
+		: *OutHistory.ConstantResolver;
+
 	FString EmitterUniqueName = GetEmitterUniqueName();
 	UNiagaraGraph* Graph = GetCalledGraph();
 	if (Graph && ParamMapIdx != INDEX_NONE && OutHistory.bShouldBuildSubHistories)
@@ -375,14 +379,7 @@ void UNiagaraNodeEmitter::BuildParameterMapHistory(FNiagaraParameterMapHistoryBu
 
 			// Build up a new parameter map history with all the child graph nodes..
 			FNiagaraParameterMapHistoryBuilder ChildBuilder;
-			if (ensure(EmitterHandle))
-			{
-				*ChildBuilder.ConstantResolver = OutHistory.ConstantResolver->AsEmitter(EmitterHandle->GetInstance()).WithUsage(OutputNodeUsage);
-			}
-			else
-			{
-				*ChildBuilder.ConstantResolver = *OutHistory.ConstantResolver;
-			}
+			*ChildBuilder.ConstantResolver = ChildConstantResolver.WithUsage(OutputNodeUsage);
 			ChildBuilder.RegisterEncounterableVariables(OutHistory.GetEncounterableVariables());
 			ChildBuilder.EnableScriptAllowList(true, GetUsage());
 
