- Warning to log about zero attributes on rasterization grid

diff --git a/Engine/Plugins/FX/Niagara/Source/Niagara/Private/NiagaraDataInterfaceRasterizationGrid3D.cpp b/Engine/Plugins/FX/Niagara/Source/Niagara/Private/NiagaraDataInterfaceRasterizationGrid3D.cpp
--- a/Engine/Plugins/FX/Niagara/Source/Niagara/Private/NiagaraDataInterfaceRasterizationGrid3D.cpp
+++ b/Engine/Plugins/FX/Niagara/Source/Niagara/Private/NiagaraDataInterfaceRasterizationGrid3D.cpp
@@ -908,7 +908,7 @@ bool UNiagaraDataInterfaceRasterizationGrid3D::InitPerInstanceData(void* PerInst
 
 	if (NumAttribChannelsFound == 0 && !Proxy->SourceDIName.IsNone())
 	{
-		UE_LOG(LogNiagara, Warning, TEXT("Zero attributes defined on %s"), *Proxy->SourceDIName.ToString());
+		UE_LOG(LogNiagara, Log, TEXT("Zero attributes defined on %s"), *Proxy->SourceDIName.ToString());
 	}
 
 //	check(InstanceData->NumTiles.X > 0);
