When pulling results from the DDC make sure we properly handle user variables with SWC structs (need to convert them to LWC) in order to mirror what we would have found with a precompile.

diff --git a/Engine/Plugins/FX/Niagara/Source/Niagara/Private/NiagaraAsyncCompile.cpp b/Engine/Plugins/FX/Niagara/Source/Niagara/Private/NiagaraAsyncCompile.cpp
--- a/Engine/Plugins/FX/Niagara/Source/Niagara/Private/NiagaraAsyncCompile.cpp
+++ b/Engine/Plugins/FX/Niagara/Source/Niagara/Private/NiagaraAsyncCompile.cpp
@@ -492,6 +492,7 @@ void FNiagaraAsyncCompileTask::CheckDDCResult()
 			{
 				ScriptPair.CompileResults = ExeData;
 				ScriptPair.bResultsReady = true;
+				bResultsFromDDC = true;
 				MoveToState(ENiagaraCompilationState::Finished);
 				UE_LOG(LogNiagara, Verbose, TEXT("Compilation data for %s could be pulled from the ddc."), *AssetPath);
 			}
@@ -919,14 +920,58 @@ void FNiagaraActiveCompilationDefault::Apply(const FNiagaraQueryCompilationOptio
 
 			// we also grab from the data contained in the ExeData, since that might be all we have if grabbed results from the DDC rather
 			// than actually issuing a precompile
-			auto IsUserVariable = [](const FNiagaraVariableBase& InVariable) -> bool
+			if (AsyncTask->bResultsFromDDC)
 			{
-				return InVariable.IsInNameSpace(FStringView(PARAM_MAP_USER_STR));
-			};
+				auto IsValidUserVariable = [](const FNiagaraVariableBase& InVariable) -> bool
+				{
+					if (!InVariable.IsInNameSpace(FStringView(PARAM_MAP_USER_STR)))
+					{
+						return false;
+					}
+
+					// check if this is a struct type, if it is make sure that we ignore the SWC version
+					if (const UStruct* Struct = InVariable.GetType().GetStruct())
+					{
+						if (FNiagaraTypeHelper::IsConvertedSWCStructure(Struct))
+						{
+							return false;
+						}
+					}
+					return true;
+				};
+
+				TArray<FNiagaraVariable> ScriptExposedVariables;
+				for (const FNiagaraVariable& Parameter : ExeData->Parameters.Parameters)
+				{
+					if (!Parameter.IsInNameSpace(FStringView(PARAM_MAP_USER_STR)))
+					{
+						continue;
+					}
+
+					// we need to also deal with SWC -> LWC conversions
+					if (const UStruct* Struct = Parameter.GetType().GetStruct())
+					{
+						if (FNiagaraTypeHelper::IsConvertedSWCStructure(Struct))
+						{
+							FNiagaraVariable ConvertedVariable(FNiagaraTypeHelper::GetLWCType(Parameter.GetType()), Parameter.GetName());
+							ConvertedVariable.AllocateData();
+
+							if (Parameter.IsDataAllocated())
+							{
+								FNiagaraLwcStructConverter DataConverter = FNiagaraTypeRegistry::GetStructConverter(ConvertedVariable.GetType());
+								DataConverter.ConvertDataFromSimulation(ConvertedVariable.GetData(), Parameter.GetData());
+							}
+							AddUniqueExposedVariables({ConvertedVariable});
+						}
+					}
+					else
+					{
+						AddUniqueExposedVariables({Parameter});
+					}
+				}
 
-			TArray<FNiagaraVariable> ScriptExposedVariables;
-			Algo::CopyIf(ExeData->Parameters.Parameters, ScriptExposedVariables, IsUserVariable);
-			AddUniqueExposedVariables(ScriptExposedVariables);
+				AddUniqueExposedVariables(ScriptExposedVariables);
+			}
 		}
 	}
 }
diff --git a/Engine/Plugins/FX/Niagara/Source/Niagara/Private/NiagaraAsyncCompile.h b/Engine/Plugins/FX/Niagara/Source/Niagara/Private/NiagaraAsyncCompile.h
--- a/Engine/Plugins/FX/Niagara/Source/Niagara/Private/NiagaraAsyncCompile.h
+++ b/Engine/Plugins/FX/Niagara/Source/Niagara/Private/NiagaraAsyncCompile.h
@@ -72,6 +72,7 @@ public:
 	bool bUsedShaderCompilerWorker = false;
 	bool bFetchedGCObjects = false;
 	bool bCompilableScript = false;
+	bool bResultsFromDDC = false;
 
 	// in order to coordinate between tasks associated with a system CheckDDC can be handled external to the task.
 	// In that case we disable this flag
diff --git a/Engine/Plugins/FX/Niagara/Source/Niagara/Private/NiagaraTypes.cpp b/Engine/Plugins/FX/Niagara/Source/Niagara/Private/NiagaraTypes.cpp
--- a/Engine/Plugins/FX/Niagara/Source/Niagara/Private/NiagaraTypes.cpp
+++ b/Engine/Plugins/FX/Niagara/Source/Niagara/Private/NiagaraTypes.cpp
@@ -296,6 +296,7 @@ FNiagaraTypeDefinition FNiagaraTypeHelper::VectorDef;
 FNiagaraTypeDefinition FNiagaraTypeHelper::Vector4Def;
 FNiagaraTypeDefinition FNiagaraTypeHelper::QuatDef;
 FNiagaraTypeDefinition FNiagaraTypeHelper::DoubleDef;
+FString FNiagaraTypeHelper::ConvertedSWCStructSuffix = TEXT("_SWC");
 
 FRWLock FNiagaraTypeHelper::RemapTableLock;
 TMap<TWeakObjectPtr<UScriptStruct>, FNiagaraTypeHelper::FRemapEntry> FNiagaraTypeHelper::RemapTable;
@@ -469,7 +470,7 @@ FNiagaraLwcStructConverter BuildSWCStructure(UScriptStruct* NewStruct, UScriptSt
 	return StructConverter;
 }
 
-bool FNiagaraTypeHelper::IsLWCStructure(UStruct* InStruct)
+bool FNiagaraTypeHelper::IsLWCStructure(const UStruct* InStruct)
 {
 	for (const FField* ChildProperty = InStruct->ChildProperties; ChildProperty; ChildProperty = ChildProperty->Next)
 	{
@@ -492,6 +493,30 @@ bool FNiagaraTypeHelper::IsLWCStructure(UStruct* InStruct)
 	return false;
 }
 
+bool FNiagaraTypeHelper::IsConvertedSWCStructure(const UStruct* InStruct)
+{
+	if (InStruct)
+	{
+		FNameBuilder StructName(InStruct->GetFName());
+		if (!StructName.ToView().EndsWith(ConvertedSWCStructSuffix))
+		{
+			return false;
+		}
+	}
+
+	{
+		FReadScopeLock ReadLock(RemapTableLock);
+		for (auto It = RemapTable.CreateConstIterator(); It; ++It)
+		{
+			if (It->Value.Struct.Get() == InStruct)
+			{
+				return true;
+			}
+		}
+	}
+	return false;
+}
+
 bool FNiagaraTypeHelper::IsLWCType(const FNiagaraTypeDefinition& InType)
 {
 	return InType.IsValid() && !InType.IsEnum() && !InType.IsUObject() && IsLWCStructure(InType.GetStruct());
@@ -576,7 +601,7 @@ UScriptStruct* FNiagaraTypeHelper::GetSWCStruct(UScriptStruct* LWCStruct)
 				//More?
 				else
 				{
-					FName SWCName = FName(LWCStruct->GetName() + TEXT("_SWC"));
+					FName SWCName = FName(LWCStruct->GetName() + ConvertedSWCStructSuffix);
 
 					//check if the remap table contains a previous entry for that struct. This might happen when the lwc struct was changed at runtime.
 					// In that case we want to reuse the existing swc struct object, since it might already be referenced by compile results.
diff --git a/Engine/Plugins/FX/Niagara/Source/Niagara/Public/NiagaraTypes.h b/Engine/Plugins/FX/Niagara/Source/Niagara/Public/NiagaraTypes.h
--- a/Engine/Plugins/FX/Niagara/Source/Niagara/Public/NiagaraTypes.h
+++ b/Engine/Plugins/FX/Niagara/Source/Niagara/Public/NiagaraTypes.h
@@ -416,7 +416,8 @@ class FNiagaraTypeHelper
 {
 public:
 	static NIAGARA_API FString ToString(const uint8* ValueData, const UObject* StructOrEnum);
-	static NIAGARA_API bool IsLWCStructure(UStruct* InStruct);
+	static NIAGARA_API bool IsLWCStructure(const UStruct* InStruct);
+	static NIAGARA_API bool IsConvertedSWCStructure(const UStruct* InStruct);
 	static NIAGARA_API bool IsLWCType(const FNiagaraTypeDefinition& InType);
 	static NIAGARA_API UScriptStruct* FindNiagaraFriendlyTopLevelStruct(UScriptStruct* InStruct, ENiagaraStructConversion StructConversion);
 	static NIAGARA_API bool IsNiagaraFriendlyTopLevelStruct(UScriptStruct* InStruct, ENiagaraStructConversion StructConversion);
@@ -468,6 +469,8 @@ private:
 	static FNiagaraTypeDefinition Vector4Def;
 	static FNiagaraTypeDefinition QuatDef;
 	static FNiagaraTypeDefinition DoubleDef;
+
+	static FString ConvertedSWCStructSuffix;
 };
 
 /** Information about how this type should be laid out in an FNiagaraDataSet */
