Improve GetClosestPointMeshDistanceField so that it will produce valid results when the search point is outside of the volume of the mesh distance field

diff --git a/Engine/Plugins/FX/Niagara/Shaders/Private/NiagaraDataInterfaceRigidMeshCollisionQuery.ush b/Engine/Plugins/FX/Niagara/Shaders/Private/NiagaraDataInterfaceRigidMeshCollisionQuery.ush
--- a/Engine/Plugins/FX/Niagara/Shaders/Private/NiagaraDataInterfaceRigidMeshCollisionQuery.ush
+++ b/Engine/Plugins/FX/Niagara/Shaders/Private/NiagaraDataInterfaceRigidMeshCollisionQuery.ush
@@ -477,6 +477,68 @@ void GetClosestDistance_{ParameterName}(in float3 WorldPosition, in float TimeFr
 	}
 }
 
+bool SampleClosestPointMeshDistanceField(int DFIndex, const FDFVector3 LWCWorldPosition, float MaxDistance, out float3 OutClosestPosition, out float3 OutClosestNormal)
+{
+	FDFObjectData DFObjectData = LoadDFObjectData(DFIndex);
+
+	// DistanceToNearestSurfaceForObject()
+	const float3 VolumePosition = DFMultiplyDemote(LWCWorldPosition, DFObjectData.WorldToVolume);
+	const float3 ClampedVolumePosition = clamp(VolumePosition, -DFObjectData.VolumePositionExtent, DFObjectData.VolumePositionExtent);
+
+	const float3 ToBox = (abs(ClampedVolumePosition) - DFObjectData.VolumePositionExtent) * DFObjectData.VolumeToWorldScale;
+	const float DistanceToBoxOutside = length(max(0.0f, ToBox));
+	const float DistanceToBoxInside = min(0.0f, max3(ToBox.x, ToBox.y, ToBox.z));
+	const float DistanceToBox = DistanceToBoxOutside + DistanceToBoxInside;
+
+	float DistanceToOccluder = MaxDistance;
+
+	BRANCH
+	if (DistanceToBox < MaxDistance)
+	{
+		uint NumMips = LoadDFAssetData(DFObjectData.AssetIndex, 0).NumMips;
+		FDFAssetData DFAssetData = LoadDFAssetData(DFObjectData.AssetIndex, NumMips - 1);
+
+		DistanceToOccluder = DistanceToMeshSurfaceStandalone(ClampedVolumePosition, DFObjectData) * DFObjectData.VolumeScale;
+
+		// Don't allow negative SDF distance to stick out of the bounding box
+		DistanceToOccluder = max(DistanceToOccluder, 0.0f);
+
+		float3 VolumeGradient = CalculateMeshSDFGradient(ClampedVolumePosition, DFAssetData);
+		float VolumeGradientLength = length(VolumeGradient);
+		float3 VolumeNormal = VolumeGradientLength > 0.0f ? VolumeGradient / VolumeGradientLength : 0;
+
+		// Transform by transposed inverse to handle non-uniform scaling
+		float4x4 WorldToVolume = DFHackToFloat(DFObjectData.WorldToVolume);
+		float3 WorldGradient = mul(VolumeNormal, transpose((float3x3)WorldToVolume));
+		float WorldGradientLength = length(WorldGradient);
+		const bool bNormalIsValid = WorldGradientLength > 0.0f;
+			
+		// output 			
+		if (bNormalIsValid)
+		{
+			OutClosestNormal = WorldGradient / WorldGradientLength;
+
+			const float3 WorldClampedPosition = DFDemote(DFMultiply(ClampedVolumePosition, DFObjectData.VolumeToWorld));
+			OutClosestPosition = WorldClampedPosition - OutClosestNormal * DistanceToOccluder;
+
+			return true;
+		}
+	}
+
+	return false;
+}
+
+void EvaluateWorldPositionOnElement_{ParameterName}(int ElementIndex, float3 WorldPosition, float DeltaTime, float TimeFraction, out float3 OutPosition, out float3 OutVelocity)
+{
+	const float3 LocalPosition = mul(GetInverseTransform_{ParameterName}(ElementIndex), float4(WorldPosition,1.0)).xyz;			
+
+	const float3 PreviousPosition = mul(GetPreviousTransform_{ParameterName}(ElementIndex), float4(LocalPosition,1.0)).xyz;
+	const float3 CurrentPosition = mul(GetCurrentTransform_{ParameterName}(ElementIndex), float4(LocalPosition,1.0)).xyz;
+	
+	OutVelocity = ( CurrentPosition - PreviousPosition ) / DeltaTime;
+	OutPosition = PreviousPosition + TimeFraction * (CurrentPosition-PreviousPosition);
+}
+
 // Given a world space position (WorldPosition) compute the static mesh closest point (position,normal,velocity)
 void GetClosestPointMeshDistanceField_{ParameterName}(in float3 WorldPosition, in float DeltaTime, in float TimeFraction, in float MaxDistance, out float OutClosestDistance, out float3 OutClosestPosition, out float3 OutClosestNormal, out float3 OutClosestVelocity, out bool NormalIsValid)
 {
@@ -501,36 +563,15 @@ void GetClosestPointMeshDistanceField_{ParameterName}(in float3 WorldPosition, i
 
 		if (DFIndex >= 0 && DFIndex < NumSceneObjects)
 		{
-			// #todo(dmp): this does a dynamic branch based on intersecting the bbox.  Maybe we can factor that out due to the broadphase here?
-			// #todo(dmp): Check the distance from the physics asset for a speedup?
-			MaxDistance = max(.001, MaxDistance);
-			OutClosestDistance = DistanceToNearestSurfaceForObject(DFIndex, LWCWorldPosition, MaxDistance);
-
-			FDFObjectData DFObjectData = LoadDFObjectData(DFIndex);
-			uint NumMips = LoadDFAssetData(DFObjectData.AssetIndex, 0).NumMips;
-			FDFAssetData DFAssetData = LoadDFAssetData(DFObjectData.AssetIndex, NumMips - 1);
-
-			float3 VolumePosition = DFMultiplyDemote(LWCWorldPosition, DFObjectData.WorldToVolume);
-			float3 VolumeGradient = CalculateMeshSDFGradient(VolumePosition, DFAssetData);
-			float VolumeGradientLength = length(VolumeGradient);
-			float3 VolumeNormal = VolumeGradientLength > 0.0f ? VolumeGradient / VolumeGradientLength : 0;
-			// Transform by transposed inverse to handle non-uniform scaling
-			float4x4 WorldToVolume = DFHackToFloat(DFObjectData.WorldToVolume);
-			float3 WorldGradient = mul(VolumeNormal, transpose((float3x3)WorldToVolume));
-			float WorldGradientLength = length(WorldGradient);
-			OutClosestNormal = WorldGradientLength > 0.0f ? WorldGradient / WorldGradientLength : 0;
-			
-			CollisionPosition = WorldPosition - OutClosestNormal * OutClosestDistance;			
-			const float3 LocalCollisionPosition = mul(GetInverseTransform_{ParameterName}(ElementIndex), float4(CollisionPosition,1.0)).xyz;			
-			
-			const float3 PreviousPosition = mul(GetPreviousTransform_{ParameterName}(ElementIndex), float4(LocalCollisionPosition,1.0)).xyz;
-			const float3 CurrentPosition = mul(GetCurrentTransform_{ParameterName}(ElementIndex), float4(LocalCollisionPosition,1.0)).xyz;
-			OutClosestVelocity = ( CurrentPosition - PreviousPosition ) / DeltaTime;
-			OutClosestPosition = PreviousPosition + TimeFraction * (CurrentPosition-PreviousPosition);
+			MaxDistance = max(0.001f, MaxDistance);
 
-			// output 			
-			NormalIsValid =  all(step(-DFObjectData.VolumePositionExtent, VolumePosition) * step(VolumePosition, DFObjectData.VolumePositionExtent));
-			OutClosestPosition = NormalIsValid ? OutClosestPosition : float3(0,0,0);
+			float3 SampledWorldPosition;
+			if (SampleClosestPointMeshDistanceField(DFIndex, LWCWorldPosition, MaxDistance, SampledWorldPosition, OutClosestNormal))
+			{
+				EvaluateWorldPositionOnElement_{ParameterName}(ElementIndex, SampledWorldPosition, DeltaTime, TimeFraction, OutClosestPosition, OutClosestVelocity);
+				OutClosestDistance = length(WorldPosition - SampledWorldPosition);
+				NormalIsValid = true;
+			}
 		}
 	}
 }
@@ -561,32 +602,13 @@ void GetClosestPointMeshDistanceFieldAccurate_{ParameterName}(in float3 WorldPos
 
 		if (OutClosestDistanceTmp < OutClosestDistance)
 		{
-			OutClosestDistance = OutClosestDistanceTmp;
-			FDFObjectData DFObjectData = LoadDFObjectData(DFIndex);
-			uint NumMips = LoadDFAssetData(DFObjectData.AssetIndex, 0).NumMips;
-			FDFAssetData DFAssetData = LoadDFAssetData(DFObjectData.AssetIndex, NumMips - 1);
-
-			float3 VolumePosition = DFMultiplyDemote(LWCWorldPosition, DFObjectData.WorldToVolume);
-			float3 VolumeGradient = CalculateMeshSDFGradient(VolumePosition, DFAssetData);
-			float VolumeGradientLength = length(VolumeGradient);
-			float3 VolumeNormal = VolumeGradientLength > 0.0f ? VolumeGradient / VolumeGradientLength : 0;
-			// Transform by transposed inverse to handle non-uniform scaling
-			float4x4 WorldToVolume = DFHackToFloat(DFObjectData.WorldToVolume);
-			float3 WorldGradient = mul(VolumeNormal, transpose((float3x3)WorldToVolume));
-			float WorldGradientLength = length(WorldGradient);
-			OutClosestNormal = WorldGradientLength > 0.0f ? WorldGradient / WorldGradientLength : 0;
-			
-			CollisionPosition = WorldPosition - OutClosestNormal * OutClosestDistance;			
-			const float3 LocalCollisionPosition = mul(GetInverseTransform_{ParameterName}(ElementIndex), float4(CollisionPosition,1.0)).xyz;			
-			
-			const float3 PreviousPosition = mul(GetPreviousTransform_{ParameterName}(ElementIndex), float4(LocalCollisionPosition,1.0)).xyz;
-			const float3 CurrentPosition = mul(GetCurrentTransform_{ParameterName}(ElementIndex), float4(LocalCollisionPosition,1.0)).xyz;
-			OutClosestVelocity = ( CurrentPosition - PreviousPosition ) / DeltaTime;
-			OutClosestPosition = PreviousPosition + TimeFraction * (CurrentPosition-PreviousPosition);
-
-			// output 			
-			NormalIsValid =  all(step(-DFObjectData.VolumePositionExtent, VolumePosition) * step(VolumePosition, DFObjectData.VolumePositionExtent));
-			OutClosestPosition = NormalIsValid ? OutClosestPosition : float3(0,0,0);
+			float3 SampledWorldPosition;
+			if (SampleClosestPointMeshDistanceField(DFIndex, LWCWorldPosition, OutClosestDistanceTmp, SampledWorldPosition, OutClosestNormal))
+			{
+				EvaluateWorldPositionOnElement_{ParameterName}(ElementIndex, SampledWorldPosition, DeltaTime, TimeFraction, OutClosestPosition, OutClosestVelocity);
+				OutClosestDistance = length(WorldPosition - SampledWorldPosition);
+				NormalIsValid = true;
+			}
 		}
 	}	
 }
@@ -632,4 +654,4 @@ void GetClosestPointMeshDistanceFieldNoNormal_{ParameterName}(in float3 WorldPos
 
 #undef BOX_INDEX
 #undef SPHERE_INDEX
-#undef CAPSULE_INDEX
+#undef CAPSULE_INDEX
\ No newline at end of file
