Flush the ShaderFileCache when we are force compiling a niagara system when r.ShaderDevelopmentMode is enabled.  Allows quick iteration of DI ush files during development.

diff --git a/Engine/Plugins/FX/Niagara/Source/Niagara/Private/NiagaraSystem.cpp b/Engine/Plugins/FX/Niagara/Source/Niagara/Private/NiagaraSystem.cpp
--- a/Engine/Plugins/FX/Niagara/Source/Niagara/Private/NiagaraSystem.cpp
+++ b/Engine/Plugins/FX/Niagara/Source/Niagara/Private/NiagaraSystem.cpp
@@ -3434,6 +3434,13 @@ bool UNiagaraSystem::RequestCompile(bool bForce, FNiagaraSystemUpdateContext* Op
 	if (bForce)
 	{
 		ForceGraphToRecompileOnNextCheck();
+
+		// if we're forcing a recompile in development mode then flush the shader file cache to catch any datainterface files that may have been edited
+		static IConsoleVariable* CVarShaderDevMode = IConsoleManager::Get().FindConsoleVariable(TEXT("r.ShaderDevelopmentMode"));
+		if (CVarShaderDevMode && CVarShaderDevMode->GetInt() != 0)
+		{
+			FlushShaderFileCache();
+		}
 	}
 
 	// we can't compile systems that have been cooked without editor data
