- Fix for invalidating asset when create a new exposed curve texture

diff --git a/Engine/Plugins/FX/Niagara/Source/Niagara/Private/NiagaraDataInterfaceCurveBase.cpp b/Engine/Plugins/FX/Niagara/Source/Niagara/Private/NiagaraDataInterfaceCurveBase.cpp
--- a/Engine/Plugins/FX/Niagara/Source/Niagara/Private/NiagaraDataInterfaceCurveBase.cpp
+++ b/Engine/Plugins/FX/Niagara/Source/Niagara/Private/NiagaraDataInterfaceCurveBase.cpp
@@ -373,42 +373,48 @@ void UNiagaraDataInterfaceCurveBase::UpdateExposedTexture()
 		return;
 	}
 
-	const int32 CurveWidth = 256;
+	constexpr int32 CurveWidth = 256;
+
+	// Build LUT
+	TArray<FFloat16Color> Texels;
+	{
+		const int32 NumElements = GetCurveNumElems();
+		const FLinearColor DefaultColors(1.0f, 1.0f, 1.0f, 1.0f);
+
+		TArray<float> TempLUT = BuildLUT(CurveWidth);
+		Texels.AddDefaulted(CurveWidth);
+		for (int32 i = 0; i < CurveWidth; ++i)
+		{
+			Texels[i].R = TempLUT[(i * NumElements) + 0];
+			Texels[i].G = NumElements >= 2 ? TempLUT[(i * NumElements) + 1] : DefaultColors.G;
+			Texels[i].B = NumElements >= 3 ? TempLUT[(i * NumElements) + 2] : DefaultColors.B;
+			Texels[i].A = NumElements >= 4 ? TempLUT[(i * NumElements) + 3] : DefaultColors.A;
+		}
+	}
 
 	if (ExposedTexture == nullptr)
 	{
 		ExposedTexture = NewObject<UTexture2D>(this);
-		ExposedTexture->Source.Init(CurveWidth, 1, 1, 1, ETextureSourceFormat::TSF_RGBA16F);
+		ExposedTexture->Source.Init(CurveWidth, 1, 1, 1, ETextureSourceFormat::TSF_RGBA16F, reinterpret_cast<const uint8*>(Texels.GetData()));
 		ExposedTexture->SRGB = false;
 		ExposedTexture->CompressionNone = true;
 		ExposedTexture->MipGenSettings = TMGS_NoMipmaps;
 		ExposedTexture->AddressX = TA_Clamp;
 		ExposedTexture->AddressY = TA_Clamp;
 		ExposedTexture->LODGroup = TEXTUREGROUP_EffectsNotFiltered;
-	}
-	
-	ExposedTexture->PreEditChange(nullptr);
-	FFloat16Color* TexData = reinterpret_cast<FFloat16Color*>(ExposedTexture->Source.LockMip(0));
+		ExposedTexture->SetDeterministicLightingGuid();
 
-	const int32 NumElements = GetCurveNumElems();
-	const FLinearColor DefaultColors(1.0f, 1.0f, 1.0f, 1.0f);
-
-	TArray<float> TempLUT = BuildLUT(CurveWidth);
-	for (int32 i=0; i < CurveWidth; ++i)
-	{
-		TexData[i].R = TempLUT[(i * NumElements) + 0];
-		TexData[i].G = NumElements >= 2 ? TempLUT[(i * NumElements) + 1] : DefaultColors.G;
-		TexData[i].B = NumElements >= 3 ? TempLUT[(i * NumElements) + 2] : DefaultColors.B;
-		TexData[i].A = NumElements >= 4 ? TempLUT[(i * NumElements) + 3] : DefaultColors.A;
+		ExposedTexture->UpdateResource();
 	}
+	else
+	{
+		ExposedTexture->PreEditChange(nullptr);
+		ExposedTexture->Source.Init(CurveWidth, 1, 1, 1, ETextureSourceFormat::TSF_RGBA16F, reinterpret_cast<const uint8*>(Texels.GetData()));
+		ExposedTexture->PostEditChange();
 
-	ExposedTexture->Source.UnlockMip(0);
-	// UseHashAsGuid is done by UnlockMip
-	//ExposedTexture->Source.UseHashAsGuid();
-	ExposedTexture->PostEditChange();
-
-	// PostEditChange() will assign a random GUID to the texture, which leads to non-deterministic builds.
-	ExposedTexture->SetDeterministicLightingGuid();
+		// PostEditChange() will assign a random GUID to the texture, which leads to non-deterministic builds.
+		ExposedTexture->SetDeterministicLightingGuid();
+	}
 }
 #endif
 
