Make sure to PostLoad the RendererBindings so that they can be sorted & searchable in cooked builds.  Addresses some sporadic issues where static bools used for bindings might not be found and assumed to be false.

diff --git a/Engine/Plugins/FX/Niagara/Source/Niagara/Private/NiagaraEmitter.cpp b/Engine/Plugins/FX/Niagara/Source/Niagara/Private/NiagaraEmitter.cpp
--- a/Engine/Plugins/FX/Niagara/Source/Niagara/Private/NiagaraEmitter.cpp
+++ b/Engine/Plugins/FX/Niagara/Source/Niagara/Private/NiagaraEmitter.cpp
@@ -790,7 +790,7 @@ void FVersionedNiagaraEmitterData::PostLoad(UNiagaraEmitter& Emitter, int32 Niag
 		INiagaraModule& NiagaraModule = FModuleManager::GetModuleChecked<INiagaraModule>("Niagara");
 		EditorData = NiagaraModule.GetEditorOnlyDataUtilities().CreateDefaultEditorData(&Emitter);
 	}
-	
+
 	if (!GPUComputeScript)
 	{
 		GPUComputeScript = NewObject<UNiagaraScript>(&Emitter, "GPUComputeScript", RF_Transactional);
@@ -859,6 +859,9 @@ void FVersionedNiagaraEmitterData::PostLoad(UNiagaraEmitter& Emitter, int32 Niag
 	check(GPUComputeScript == nullptr || SimTarget == ENiagaraSimTarget::GPUComputeSim);
 #endif
 
+	// make sure to PostLoad any ParameterStore so that they remains searchable (sorted)
+	RendererBindings.PostLoad(&Emitter);
+
 	//Temporarily disabling interpolated spawn if the script type and flag don't match.
 	if (SpawnScriptProps.Script)
 	{
diff --git a/Engine/Plugins/FX/Niagara/Source/Niagara/Private/NiagaraVolumeRendererProperties.cpp b/Engine/Plugins/FX/Niagara/Source/Niagara/Private/NiagaraVolumeRendererProperties.cpp
--- a/Engine/Plugins/FX/Niagara/Source/Niagara/Private/NiagaraVolumeRendererProperties.cpp
+++ b/Engine/Plugins/FX/Niagara/Source/Niagara/Private/NiagaraVolumeRendererProperties.cpp
@@ -256,7 +256,6 @@ void UNiagaraVolumeRendererProperties::CacheFromCompiledData(const FNiagaraDataS
 	RendererVisibilityTagAccessor.Init(CompiledData, RendererVisibilityTagBinding.GetDataSetBindableVariable().GetName());
 	VolumeResolutionMaxAxisAccessor.Init(CompiledData, VolumeResolutionMaxAxisBinding.GetDataSetBindableVariable().GetName());
 	VolumeWorldSpaceSizeAccessor.Init(CompiledData, VolumeWorldSpaceSizeBinding.GetDataSetBindableVariable().GetName());
-	RendererVisibilityTagAccessor.Init(CompiledData, RendererVisibilityTagBinding.GetDataSetBindableVariable().GetName());
 }
 
 void UNiagaraVolumeRendererProperties::UpdateSourceModeDerivates(ENiagaraRendererSourceDataMode InSourceMode, bool bFromPropertyEdit)
