Refactored Niagara's external 'is class allowed' API into 'is class visible' and 'is class referenceable' APIs
Visible is for UI filter purposes, referenceable is for validation purposes.
There can be invisible, but referenceable types (engine provided content might use some types we don't want users to use, but it will still validate)

diff --git a/Engine/Plugins/FX/Niagara/Source/Niagara/Private/NiagaraAssetTagDefinitions.cpp b/Engine/Plugins/FX/Niagara/Source/Niagara/Private/NiagaraAssetTagDefinitions.cpp
--- a/Engine/Plugins/FX/Niagara/Source/Niagara/Private/NiagaraAssetTagDefinitions.cpp
+++ b/Engine/Plugins/FX/Niagara/Source/Niagara/Private/NiagaraAssetTagDefinitions.cpp
@@ -67,24 +67,6 @@ FString FNiagaraAssetTagDefinition::GetGuidAsString() const
 	return TagGuid.ToString(EGuidFormats::DigitsWithHyphens);
 }
 
-void FNiagaraAssetTagDefinition::AddTagToAssetRegistryTags(TArray<UObject::FAssetRegistryTag>& OutTags) const
-{
-	UObject::FAssetRegistryTag* FoundAssetTag = OutTags.FindByPredicate([this](const UObject::FAssetRegistryTag& Candidate)
-	{
-		return Candidate.Name == FName(GetGuidAsString());
-	});
-
-	if(FoundAssetTag)
-	{
-		return;
-	}
-	else
-	{
-		// Value has to be >1 of data to be considered non-empty to not get thrown away
-		OutTags.Add(UObject::FAssetRegistryTag(FName(GetGuidAsString()), "  ", UObject::FAssetRegistryTag::TT_Alphabetical));
-	}
-}
-
 FNiagaraAssetTagDefinitionReference::FNiagaraAssetTagDefinitionReference(const FNiagaraAssetTagDefinition& InTagDefinition)
 {
 	SetTagDefinitionReference(InTagDefinition);
diff --git a/Engine/Plugins/FX/Niagara/Source/Niagara/Public/NiagaraAssetTagDefinitions.h b/Engine/Plugins/FX/Niagara/Source/Niagara/Public/NiagaraAssetTagDefinitions.h
--- a/Engine/Plugins/FX/Niagara/Source/Niagara/Public/NiagaraAssetTagDefinitions.h
+++ b/Engine/Plugins/FX/Niagara/Source/Niagara/Public/NiagaraAssetTagDefinitions.h
@@ -77,7 +77,6 @@ public:
 	NIAGARA_API bool DoesAssetDataContainTag(const FAssetData& AssetData) const;
 
 	NIAGARA_API FString GetGuidAsString() const;
-	void AddTagToAssetRegistryTags(TArray<UObject::FAssetRegistryTag>& OutTags) const;
 };
 
 inline uint32 GetTypeHash(const FNiagaraAssetTagDefinition& AssetTagDefinition)
diff --git a/Engine/Plugins/FX/Niagara/Source/NiagaraEditor/Private/EdGraphSchema_Niagara.cpp b/Engine/Plugins/FX/Niagara/Source/NiagaraEditor/Private/EdGraphSchema_Niagara.cpp
--- a/Engine/Plugins/FX/Niagara/Source/NiagaraEditor/Private/EdGraphSchema_Niagara.cpp
+++ b/Engine/Plugins/FX/Niagara/Source/NiagaraEditor/Private/EdGraphSchema_Niagara.cpp
@@ -230,7 +230,7 @@ TSharedPtr<FNiagaraAction_NewNode> AddNewNodeMenuAction(
 {
 	const UNiagaraEditorSettings* NiagaraEditorSettings = GetDefault<UNiagaraEditorSettings>();
 	TSharedPtr<FNiagaraAction_NewNode> NewAction = MakeShared<FNiagaraAction_NewNode>(DisplayName, Section, NestedCategories, Tooltip, Keywords);
-	if (ensureMsgf(InNodeTemplate == nullptr || NiagaraEditorSettings->IsAllowedClass(InNodeTemplate->GetClass()),
+	if (ensureMsgf(InNodeTemplate == nullptr || NiagaraEditorSettings->IsVisibleClass(InNodeTemplate->GetClass()),
 		TEXT("Can not create a menu action for a node of class %s."), *InNodeTemplate->GetClass()->GetName()))
 	{
 		NewAction->NodeTemplate = InNodeTemplate;
@@ -443,7 +443,7 @@ TArray<TSharedPtr<FNiagaraAction_NewNode>> UEdGraphSchema_Niagara::GetGraphActio
 	// Add operations (add / mul / etc)
 	if (bAllowOpNodes &&
 		(GbAllowAllNiagaraNodesInEmitterGraphs || bModuleGraph || bFunctionGraph || bSystemGraph) &&
-		NiagaraEditorSettings->IsAllowedClass(UNiagaraNodeOp::StaticClass()))
+		NiagaraEditorSettings->IsVisibleClass(UNiagaraNodeOp::StaticClass()))
 	{
 		const TArray<FNiagaraOpInfo>& OpInfos = FNiagaraOpInfo::GetOpInfoArray();
 
@@ -495,7 +495,7 @@ TArray<TSharedPtr<FNiagaraAction_NewNode>> UEdGraphSchema_Niagara::GetGraphActio
 	}
 
 	// Add custom code
-	if (bAllowCustomNode && NiagaraEditorSettings->IsAllowedClass(UNiagaraNodeCustomHlsl::StaticClass()))
+	if (bAllowCustomNode && NiagaraEditorSettings->IsVisibleClass(UNiagaraNodeCustomHlsl::StaticClass()))
 	{
 		const FText DisplayName = LOCTEXT("CustomHLSLNode","Custom Hlsl");
 		const FText TooltipDesc = LOCTEXT("CustomHlslPopupTooltip", "Add a node with custom hlsl content");
@@ -537,7 +537,7 @@ PRAGMA_ENABLE_DEPRECATION_WARNINGS
 	//Add functions
 	if (bAllowFunctionNodes &&
 		(GbAllowAllNiagaraNodesInEmitterGraphs || bModuleGraph || bFunctionGraph || bDynamicInputGraph) &&
-		NiagaraEditorSettings->IsAllowedClass(UNiagaraNodeFunctionCall::StaticClass()))
+		NiagaraEditorSettings->IsVisibleClass(UNiagaraNodeFunctionCall::StaticClass()))
 	{
 		TArray<FAssetData> FunctionScriptAssets;
 		FNiagaraEditorUtilities::FGetFilteredScriptAssetsOptions FunctionScriptFilterOptions;
@@ -565,7 +565,7 @@ PRAGMA_ENABLE_DEPRECATION_WARNINGS
 
 	//Add modules
 	if (bAllowModuleNodes && !bFunctionGraph && !bModuleGraph && !bDynamicInputGraph &&
-		NiagaraEditorSettings->IsAllowedClass(UNiagaraNodeFunctionCall::StaticClass()))
+		NiagaraEditorSettings->IsVisibleClass(UNiagaraNodeFunctionCall::StaticClass()))
 	{
 		TArray<FAssetData> ModuleScriptAssets;
 		FNiagaraEditorUtilities::FGetFilteredScriptAssetsOptions ModuleScriptFilterOptions;
@@ -582,8 +582,8 @@ PRAGMA_ENABLE_DEPRECATION_WARNINGS
 
 	//Add event read and writes nodes
 	if (bAllowEventNodes && bModuleGraph &&
-		NiagaraEditorSettings->IsAllowedClass(UNiagaraNodeReadDataSet::StaticClass()) &&
-		NiagaraEditorSettings->IsAllowedClass(UNiagaraNodeWriteDataSet::StaticClass()))
+		NiagaraEditorSettings->IsVisibleClass(UNiagaraNodeReadDataSet::StaticClass()) &&
+		NiagaraEditorSettings->IsVisibleClass(UNiagaraNodeWriteDataSet::StaticClass()))
 	{
 		const FText MenuCat = LOCTEXT("NiagaraEventMenuCat", "Events");
 		TArray<FString> ReadCategories = {MenuCat.ToString(), LOCTEXT("NiagaraEventCategory_Read", "Read").ToString()};
@@ -640,7 +640,7 @@ PRAGMA_ENABLE_DEPRECATION_WARNINGS
 		UsageTypesToAdd.Add(ENiagaraScriptUsage::SystemUpdateScript);
 	}
 
-	if (UsageTypesToAdd.Num() != 0 && NiagaraEditorSettings->IsAllowedClass(UNiagaraNodeOutput::StaticClass()))
+	if (UsageTypesToAdd.Num() != 0 && NiagaraEditorSettings->IsVisibleClass(UNiagaraNodeOutput::StaticClass()))
 	{
 		for (ENiagaraScriptUsage Usage : UsageTypesToAdd)
 		{
@@ -667,7 +667,7 @@ PRAGMA_ENABLE_DEPRECATION_WARNINGS
 
 
 	// Add Convert Nodes
-	if (NiagaraEditorSettings->IsAllowedClass(UNiagaraNodeConvert::StaticClass()))
+	if (NiagaraEditorSettings->IsVisibleClass(UNiagaraNodeConvert::StaticClass()))
 	{
 		FNiagaraTypeDefinition PinType = FNiagaraTypeDefinition::GetGenericNumericDef();
 
@@ -856,31 +856,31 @@ PRAGMA_ENABLE_DEPRECATION_WARNINGS
 	if (bAllowParameterMapGetSetNodes)
 	{
 		FText MenuCat = FText::FromString("Parameter Map");
-		if (NiagaraEditorSettings->IsAllowedClass(UNiagaraNodeParameterMapGet::StaticClass()))
+		if (NiagaraEditorSettings->IsVisibleClass(UNiagaraNodeParameterMapGet::StaticClass()))
 		{
 			FString Name = TEXT("Parameter Map Get");
 			UNiagaraNodeParameterMapGet* BaseNode = NewObject<UNiagaraNodeParameterMapGet>(OwnerOfTemporaries);
 			AddNewNodeMenuAction(NewActions, BaseNode, FText::FromString(Name), ENiagaraMenuSections::Suggested, {MenuCat.ToString()}, FText::GetEmpty(), FText::GetEmpty());
 		}
-		if (NiagaraEditorSettings->IsAllowedClass(UNiagaraNodeParameterMapSet::StaticClass()))
+		if (NiagaraEditorSettings->IsVisibleClass(UNiagaraNodeParameterMapSet::StaticClass()))
 		{
 			FString Name = TEXT("Parameter Map Set");
 			UNiagaraNodeParameterMapSet* BaseNode = NewObject<UNiagaraNodeParameterMapSet>(OwnerOfTemporaries);
 			AddNewNodeMenuAction(NewActions, BaseNode, FText::FromString(Name), ENiagaraMenuSections::Suggested, {MenuCat.ToString()}, FText::GetEmpty(), FText::GetEmpty());
 		}
-		if (NiagaraEditorSettings->IsAllowedClass(UNiagaraNodeParameterMapFor::StaticClass()))
+		if (NiagaraEditorSettings->IsVisibleClass(UNiagaraNodeParameterMapFor::StaticClass()))
 		{
 			FString Name = TEXT("Parameter Map For");
 			UNiagaraNodeParameterMapFor* BaseNode = NewObject<UNiagaraNodeParameterMapFor>(OwnerOfTemporaries);
 			AddNewNodeMenuAction(NewActions, BaseNode, FText::FromString(Name), ENiagaraMenuSections::General, {MenuCat.ToString()}, FText::GetEmpty(), FText::GetEmpty());
 		}
-		if (NiagaraEditorSettings->IsAllowedClass(UNiagaraNodeParameterMapForWithContinue::StaticClass()))
+		if (NiagaraEditorSettings->IsVisibleClass(UNiagaraNodeParameterMapForWithContinue::StaticClass()))
 		{
 			FString Name = TEXT("Parameter Map For With Continue");
 			UNiagaraNodeParameterMapForWithContinue* BaseNode = NewObject<UNiagaraNodeParameterMapForWithContinue>(OwnerOfTemporaries);
 			AddNewNodeMenuAction(NewActions, BaseNode, FText::FromString(Name), ENiagaraMenuSections::General, { MenuCat.ToString() }, FText::GetEmpty(), FText::GetEmpty());
 		}
-		if (NiagaraEditorSettings->IsAllowedClass(UNiagaraNodeParameterMapForIndex::StaticClass()))
+		if (NiagaraEditorSettings->IsVisibleClass(UNiagaraNodeParameterMapForIndex::StaticClass()))
 		{
 			FString Name = TEXT("Parameter Map For Current Index");
 			UNiagaraNodeParameterMapForIndex* BaseNode = NewObject<UNiagaraNodeParameterMapForIndex>(OwnerOfTemporaries);
@@ -889,7 +889,7 @@ PRAGMA_ENABLE_DEPRECATION_WARNINGS
 	}
 
 	// Handle comment nodes
-	if (NiagaraEditorSettings->IsAllowedClass(UEdGraphNode_Comment::StaticClass()))
+	if (NiagaraEditorSettings->IsVisibleClass(UEdGraphNode_Comment::StaticClass()))
 	{
 		FText MenuCat = FText::FromString("Comments");
 
@@ -901,7 +901,7 @@ PRAGMA_ENABLE_DEPRECATION_WARNINGS
 	}
 
 	// Handle output tag nodes
-	if (bAllowOutputTagNodes && NiagaraEditorSettings->IsAllowedClass(UNiagaraNodeOutputTag::StaticClass()))
+	if (bAllowOutputTagNodes && NiagaraEditorSettings->IsVisibleClass(UNiagaraNodeOutputTag::StaticClass()))
 	{
 		FText MenuCat = FText::FromString("Special Purpose Parameters");
 
@@ -913,7 +913,7 @@ PRAGMA_ENABLE_DEPRECATION_WARNINGS
 	}
 	
 	//Add all input node options for input pins or no pin.
-	if ((FromPin == nullptr || FromPin->Direction == EGPD_Input) && NiagaraEditorSettings->IsAllowedClass(UNiagaraNodeInput::StaticClass()))
+	if ((FromPin == nullptr || FromPin->Direction == EGPD_Input) && NiagaraEditorSettings->IsVisibleClass(UNiagaraNodeInput::StaticClass()))
 	{
 		TArray<UNiagaraNodeInput*> InputNodes;
 		NiagaraGraph->GetNodesOfClass(InputNodes);
@@ -1023,7 +1023,7 @@ PRAGMA_ENABLE_DEPRECATION_WARNINGS
 	const FText UtilMenuCat = LOCTEXT("NiagaraUsageSelectorMenuCat", "Utility");
 
 	// Add reroute node
-	if (NiagaraEditorSettings->IsAllowedClass(UNiagaraNodeReroute::StaticClass()))
+	if (NiagaraEditorSettings->IsVisibleClass(UNiagaraNodeReroute::StaticClass()))
 	{
 		const FText RerouteMenuDesc = LOCTEXT("NiagaraRerouteMenuDesc", "Reroute");
 		
@@ -1033,7 +1033,7 @@ PRAGMA_ENABLE_DEPRECATION_WARNINGS
 	
 	// Add select node
 	// Note: Data Interfaces are not supported for this type
-	if (bAllowSelectNodes && NiagaraEditorSettings->IsAllowedClass(UNiagaraNodeSelect::StaticClass()))
+	if (bAllowSelectNodes && NiagaraEditorSettings->IsVisibleClass(UNiagaraNodeSelect::StaticClass()))
 	{
 		const FText SelectMenuDesc = LOCTEXT("NiagaraSelectMenuDesc", "Select / If");
 		
@@ -1043,7 +1043,7 @@ PRAGMA_ENABLE_DEPRECATION_WARNINGS
 
 	// Add static switch node
 	// Note: Data Interfaces are not supported for this type
-	if (bAllowStaticSwitchNodes && NiagaraEditorSettings->IsAllowedClass(UNiagaraNodeStaticSwitch::StaticClass()))
+	if (bAllowStaticSwitchNodes && NiagaraEditorSettings->IsVisibleClass(UNiagaraNodeStaticSwitch::StaticClass()))
 	{
 		const FText UsageSelectorMenuDesc = LOCTEXT("NiagaraStaticSwitchMenuDesc", "Static Switch");
 		
diff --git a/Engine/Plugins/FX/Niagara/Source/NiagaraEditor/Private/NiagaraEditorSettings.cpp b/Engine/Plugins/FX/Niagara/Source/NiagaraEditor/Private/NiagaraEditorSettings.cpp
--- a/Engine/Plugins/FX/Niagara/Source/NiagaraEditor/Private/NiagaraEditorSettings.cpp
+++ b/Engine/Plugins/FX/Niagara/Source/NiagaraEditor/Private/NiagaraEditorSettings.cpp
@@ -660,14 +660,14 @@ UNiagaraEditorSettings::FOnNiagaraEditorSettingsChanged& UNiagaraEditorSettings:
 	return GetMutableDefault<UNiagaraEditorSettings>()->SettingsChangedDelegate;
 }
 
-void UNiagaraEditorSettings::SetOnIsClassAllowed(const FOnIsClassAllowed& InOnIsClassAllowed)
+void UNiagaraEditorSettings::SetOnIsClassVisible(const FOnIsClassAllowed& InOnIsClassAllowed)
 {
-	OnIsClassAllowedDelegate = InOnIsClassAllowed;
+	OnIsClassVisibleDelegate = InOnIsClassAllowed;
 }
 
-void UNiagaraEditorSettings::SetOnIsClassPathAllowed(const FOnIsClassPathAllowed& InOnIsClassPathAllowed)
+void UNiagaraEditorSettings::SetOnIsClassReferenceable(const FOnIsClassAllowed& InOnIsClassAllowed)
 {
-	OnIsClassPathAllowedDelegate = InOnIsClassPathAllowed;
+	OnIsClassReferenceableDelegate = InOnIsClassAllowed;
 }
 
 void UNiagaraEditorSettings::SetOnShouldFilterAssetByClassUsage(const FOnShouldFilterAsset& InOnShouldFilterAssetByClassUsage)
@@ -680,19 +680,19 @@ void UNiagaraEditorSettings::SetOnShouldFilterAssetInNiagaraAssetBrowser(const F
 	OnShouldFilterAssetInNiagaraAssetBrowser = InOnShouldFilterAssetInNiagaraAssetBrowser;
 }
 
-bool UNiagaraEditorSettings::IsAllowedClass(const UClass* InClass) const
+bool UNiagaraEditorSettings::IsVisibleClass(const UClass* InClass) const
 {
-	return OnIsClassAllowedDelegate.IsBound() == false || OnIsClassAllowedDelegate.Execute(InClass);
+	return OnIsClassVisibleDelegate.IsBound() == false || OnIsClassVisibleDelegate.Execute(InClass);
 }
 
-bool UNiagaraEditorSettings::IsAllowedClassPath(const FTopLevelAssetPath& InClassPath) const
+bool UNiagaraEditorSettings::IsReferenceableClass(const UClass* InClass) const
 {
-	return OnIsClassPathAllowedDelegate.IsBound() == false || OnIsClassPathAllowedDelegate.Execute(InClassPath);
+	return OnIsClassReferenceableDelegate.IsBound() == false || OnIsClassReferenceableDelegate.Execute(InClass);
 }
 
-bool UNiagaraEditorSettings::IsAllowedTypeDefinition(const FNiagaraTypeDefinition& InTypeDefinition) const
+bool UNiagaraEditorSettings::IsVisibleTypeDefinition(const FNiagaraTypeDefinition& InTypeDefinition) const
 {
-	return InTypeDefinition.GetClass() == nullptr || IsAllowedClass(InTypeDefinition.GetClass());
+	return InTypeDefinition.GetClass() == nullptr || IsVisibleClass(InTypeDefinition.GetClass());
 }
 
 const FName ClassUsageListTagName = "ClassUsageList";
@@ -743,7 +743,7 @@ bool UNiagaraEditorSettings::IsAllowedAssetObjectByClassUsageInternal(const UObj
 	GetObjectsWithPackage(AssetObject.GetOutermost(), ObjectsInPackage);
 	for (UObject* ObjectInPackage : ObjectsInPackage)
 	{
-		if (ShouldTrackClassUsage(ObjectInPackage->GetClass()) && IsAllowedClass(ObjectInPackage->GetClass()) == false)
+		if (ShouldTrackClassUsage(ObjectInPackage->GetClass()) && IsReferenceableClass(ObjectInPackage->GetClass()) == false)
 		{
 			if (GbLogFoundButNotAllowedAssets)
 			{
@@ -779,7 +779,7 @@ bool UNiagaraEditorSettings::IsAllowedObjectByClassUsageInternal(const UObject&
 	
 	for (const UObject* ObjectInPackage : ObjectsWithOuter)
 	{
-		if (ShouldTrackClassUsage(ObjectInPackage->GetClass()) && IsAllowedClass(ObjectInPackage->GetClass()) == false)
+		if (ShouldTrackClassUsage(ObjectInPackage->GetClass()) && IsReferenceableClass(ObjectInPackage->GetClass()) == false)
 		{
 			if (GbLogFoundButNotAllowedAssets)
 			{
@@ -809,7 +809,7 @@ bool UNiagaraEditorSettings::IsAllowedObjectByClassUsageInternal(const UObject&
 bool UNiagaraEditorSettings::IsAllowedAssetByClassUsage(const FAssetData& InAssetData) const
 {
 	if (OnShouldFilterAssetByClassUsage.IsBound() == false ||
-		OnIsClassAllowedDelegate.IsBound() == false ||
+		OnIsClassVisibleDelegate.IsBound() == false ||
 		OnShouldFilterAssetByClassUsage.Execute(FTopLevelAssetPath(InAssetData.GetObjectPathString())) == false)
 	{
 		return true;
@@ -835,7 +835,7 @@ bool UNiagaraEditorSettings::IsAllowedAssetByClassUsage(const FAssetData& InAsse
 			{
 				FTopLevelAssetPath ClassPath(ClassUsageListLine);
 				UClass* UsedClass = FindObject<UClass>(ClassPath);
-				if (UsedClass == nullptr || (ShouldTrackClassUsage(UsedClass) && IsAllowedClass(UsedClass) == false))
+				if (UsedClass == nullptr || (ShouldTrackClassUsage(UsedClass) && IsVisibleClass(UsedClass) == false))
 				{
 					bInvalidClassFound = true;
 					if (GbLogFoundButNotAllowedAssets)
@@ -865,13 +865,13 @@ bool UNiagaraEditorSettings::IsAllowedAssetByClassUsage(const FAssetData& InAsse
 bool UNiagaraEditorSettings::IsAllowedAssetObjectByClassUsage(const UObject& InAssetObject) const
 {
 	TSet<const UObject*> CheckedAssetObjects;
-	if (OnIsClassAllowedDelegate.IsBound() && InAssetObject.IsAsset() == false)
+	if (OnIsClassVisibleDelegate.IsBound() && InAssetObject.IsAsset() == false)
 	{
 		return IsAllowedObjectByClassUsageInternal(InAssetObject, CheckedAssetObjects);
 	}
 	
 	if (OnShouldFilterAssetByClassUsage.IsBound() == false ||
-		OnIsClassAllowedDelegate.IsBound() == false ||
+		OnIsClassVisibleDelegate.IsBound() == false ||
 		OnShouldFilterAssetByClassUsage.Execute(FTopLevelAssetPath(InAssetObject.GetPathName())) == false)
 	{
 		return true;
diff --git a/Engine/Plugins/FX/Niagara/Source/NiagaraEditor/Private/NiagaraEditorUtilities.cpp b/Engine/Plugins/FX/Niagara/Source/NiagaraEditor/Private/NiagaraEditorUtilities.cpp
--- a/Engine/Plugins/FX/Niagara/Source/NiagaraEditor/Private/NiagaraEditorUtilities.cpp
+++ b/Engine/Plugins/FX/Niagara/Source/NiagaraEditor/Private/NiagaraEditorUtilities.cpp
@@ -4504,7 +4504,7 @@ void FNiagaraEditorUtilities::GetAllowedUserVariableTypes(TArray<FNiagaraTypeDef
 	const UNiagaraEditorSettings* NiagaraEditorSettings = GetDefault<UNiagaraEditorSettings>();
 	for (const FNiagaraTypeDefinition& RegisteredUserVariableType : FNiagaraTypeRegistry::GetRegisteredUserVariableTypes())
 	{
-		if (NiagaraEditorSettings->IsAllowedTypeDefinition(RegisteredUserVariableType))
+		if (NiagaraEditorSettings->IsVisibleTypeDefinition(RegisteredUserVariableType))
 		{
 			OutAllowedTypes.Add(RegisteredUserVariableType);
 		}
@@ -4516,7 +4516,7 @@ void FNiagaraEditorUtilities::GetAllowedSystemVariableTypes(TArray<FNiagaraTypeD
 	const UNiagaraEditorSettings* NiagaraEditorSettings = GetDefault<UNiagaraEditorSettings>();
 	for (const FNiagaraTypeDefinition& RegisteredSystemVariableType : FNiagaraTypeRegistry::GetRegisteredSystemVariableTypes())
 	{
-		if (NiagaraEditorSettings->IsAllowedTypeDefinition(RegisteredSystemVariableType))
+		if (NiagaraEditorSettings->IsVisibleTypeDefinition(RegisteredSystemVariableType))
 		{
 			OutAllowedTypes.Add(RegisteredSystemVariableType);
 		}
@@ -4528,7 +4528,7 @@ void FNiagaraEditorUtilities::GetAllowedEmitterVariableTypes(TArray<FNiagaraType
 	const UNiagaraEditorSettings* NiagaraEditorSettings = GetDefault<UNiagaraEditorSettings>();
 	for (const FNiagaraTypeDefinition& RegisteredEmitterVariableType : FNiagaraTypeRegistry::GetRegisteredEmitterVariableTypes())
 	{
-		if (NiagaraEditorSettings->IsAllowedTypeDefinition(RegisteredEmitterVariableType))
+		if (NiagaraEditorSettings->IsVisibleTypeDefinition(RegisteredEmitterVariableType))
 		{
 			OutAllowedTypes.Add(RegisteredEmitterVariableType);
 		}
@@ -4540,7 +4540,7 @@ void FNiagaraEditorUtilities::GetAllowedParticleVariableTypes(TArray<FNiagaraTyp
 	const UNiagaraEditorSettings* NiagaraEditorSettings = GetDefault<UNiagaraEditorSettings>();
 	for (const FNiagaraTypeDefinition& RegisteredParticleVariableType : FNiagaraTypeRegistry::GetRegisteredParticleVariableTypes())
 	{
-		if (NiagaraEditorSettings->IsAllowedTypeDefinition(RegisteredParticleVariableType))
+		if (NiagaraEditorSettings->IsVisibleTypeDefinition(RegisteredParticleVariableType))
 		{
 			OutAllowedTypes.Add(RegisteredParticleVariableType);
 		}
@@ -4552,7 +4552,7 @@ void FNiagaraEditorUtilities::GetAllowedParameterTypes(TArray<FNiagaraTypeDefini
 	const UNiagaraEditorSettings* NiagaraEditorSettings = GetDefault<UNiagaraEditorSettings>();
 	for (const FNiagaraTypeDefinition& RegisteredParameterType : FNiagaraTypeRegistry::GetRegisteredParameterTypes())
 	{
-		if (NiagaraEditorSettings->IsAllowedTypeDefinition(RegisteredParameterType))
+		if (NiagaraEditorSettings->IsVisibleTypeDefinition(RegisteredParameterType))
 		{
 			OutAllowedTypes.Add(RegisteredParameterType);
 		}
@@ -4564,7 +4564,7 @@ void FNiagaraEditorUtilities::GetAllowedPayloadTypes(TArray<FNiagaraTypeDefiniti
 	const UNiagaraEditorSettings* NiagaraEditorSettings = GetDefault<UNiagaraEditorSettings>();
 	for (const FNiagaraTypeDefinition& RegisteredPayloadType : FNiagaraTypeRegistry::GetRegisteredPayloadTypes())
 	{
-		if (NiagaraEditorSettings->IsAllowedTypeDefinition(RegisteredPayloadType))
+		if (NiagaraEditorSettings->IsVisibleTypeDefinition(RegisteredPayloadType))
 		{
 			OutAllowedTypes.Add(RegisteredPayloadType);
 		}
diff --git a/Engine/Plugins/FX/Niagara/Source/NiagaraEditor/Private/ViewModels/Stack/NiagaraStackEmitterPropertiesGroup.cpp b/Engine/Plugins/FX/Niagara/Source/NiagaraEditor/Private/ViewModels/Stack/NiagaraStackEmitterPropertiesGroup.cpp
--- a/Engine/Plugins/FX/Niagara/Source/NiagaraEditor/Private/ViewModels/Stack/NiagaraStackEmitterPropertiesGroup.cpp
+++ b/Engine/Plugins/FX/Niagara/Source/NiagaraEditor/Private/ViewModels/Stack/NiagaraStackEmitterPropertiesGroup.cpp
@@ -97,7 +97,7 @@ public:
 		GetDerivedClasses(UNiagaraSimulationStageBase::StaticClass(), SimulationStageClasses);
 		for (UClass* SimulationStageClass : SimulationStageClasses)
 		{
-			if (NiagaraEditorSettings->IsAllowedClass(SimulationStageClass))
+			if (NiagaraEditorSettings->IsVisibleClass(SimulationStageClass))
 			{
 				OutAddActions.Add(MakeShared<FAddEmitterStageAction>(
 					EEmitterAddMode::SimulationStage,
diff --git a/Engine/Plugins/FX/Niagara/Source/NiagaraEditor/Private/ViewModels/Stack/NiagaraStackFunctionInput.cpp b/Engine/Plugins/FX/Niagara/Source/NiagaraEditor/Private/ViewModels/Stack/NiagaraStackFunctionInput.cpp
--- a/Engine/Plugins/FX/Niagara/Source/NiagaraEditor/Private/ViewModels/Stack/NiagaraStackFunctionInput.cpp
+++ b/Engine/Plugins/FX/Niagara/Source/NiagaraEditor/Private/ViewModels/Stack/NiagaraStackFunctionInput.cpp
@@ -3535,7 +3535,7 @@ bool UNiagaraStackFunctionInput::OpenSourceAsset() const
 bool UNiagaraStackFunctionInput::SupportsCustomExpressions() const
 {
 	const UNiagaraEditorSettings* NiagaraEditorSettings = GetDefault<UNiagaraEditorSettings>();
-	return NiagaraEditorSettings->IsAllowedClass(UNiagaraNodeCustomHlsl::StaticClass());
+	return NiagaraEditorSettings->IsVisibleClass(UNiagaraNodeCustomHlsl::StaticClass());
 }
 
 #undef LOCTEXT_NAMESPACE
diff --git a/Engine/Plugins/FX/Niagara/Source/NiagaraEditor/Private/ViewModels/Stack/NiagaraStackObject.cpp b/Engine/Plugins/FX/Niagara/Source/NiagaraEditor/Private/ViewModels/Stack/NiagaraStackObject.cpp
--- a/Engine/Plugins/FX/Niagara/Source/NiagaraEditor/Private/ViewModels/Stack/NiagaraStackObject.cpp
+++ b/Engine/Plugins/FX/Niagara/Source/NiagaraEditor/Private/ViewModels/Stack/NiagaraStackObject.cpp
@@ -321,7 +321,7 @@ void UNiagaraStackObject::RefreshChildrenInternal(const TArray<UNiagaraStackEntr
 	}
 
 	const UNiagaraEditorSettings* NiagaraEditorSettings = GetDefault<UNiagaraEditorSettings>();
-	if (Object != nullptr && NiagaraEditorSettings->IsAllowedClass(Object->GetClass()) == false)
+	if (Object != nullptr && NiagaraEditorSettings->IsReferenceableClass(Object->GetClass()) == false)
 	{
 		NewIssues.Add(FStackIssue(
 			EStackIssueSeverity::Error,
diff --git a/Engine/Plugins/FX/Niagara/Source/NiagaraEditor/Private/ViewModels/Stack/NiagaraStackRenderItemGroup.cpp b/Engine/Plugins/FX/Niagara/Source/NiagaraEditor/Private/ViewModels/Stack/NiagaraStackRenderItemGroup.cpp
--- a/Engine/Plugins/FX/Niagara/Source/NiagaraEditor/Private/ViewModels/Stack/NiagaraStackRenderItemGroup.cpp
+++ b/Engine/Plugins/FX/Niagara/Source/NiagaraEditor/Private/ViewModels/Stack/NiagaraStackRenderItemGroup.cpp
@@ -77,7 +77,8 @@ public:
 		const UNiagaraEditorSettings* NiagaraEditorSettings = GetDefault<UNiagaraEditorSettings>();
 		for (const FNiagaraRendererCreationInfo& RendererCreationInfo : RendererCreationInfos)
 		{
-			if (NiagaraEditorSettings->IsAllowedClassPath(RendererCreationInfo.RendererClassPath) && RenderersOwner->IsRenderCreationInfoSupported(RendererCreationInfo))
+			FSoftClassPath SoftClassPath(RendererCreationInfo.RendererClassPath.ToString());
+			if (NiagaraEditorSettings->IsVisibleClass(SoftClassPath.ResolveClass()) && RenderersOwner->IsRenderCreationInfoSupported(RendererCreationInfo))
 			{
 				OutAddActions.Add(MakeShared<FRenderItemGroupAddAction>(RendererCreationInfo));
 			}
diff --git a/Engine/Plugins/FX/Niagara/Source/NiagaraEditor/Private/Widgets/SNiagaraGraphNode.cpp b/Engine/Plugins/FX/Niagara/Source/NiagaraEditor/Private/Widgets/SNiagaraGraphNode.cpp
--- a/Engine/Plugins/FX/Niagara/Source/NiagaraEditor/Private/Widgets/SNiagaraGraphNode.cpp
+++ b/Engine/Plugins/FX/Niagara/Source/NiagaraEditor/Private/Widgets/SNiagaraGraphNode.cpp
@@ -84,7 +84,7 @@ void SNiagaraGraphNode::UpdateErrorInfo()
 	ErrorMsg.Reset();
 	if (NiagaraNode.IsValid())
 	{
-		if (NiagaraEditorSettings->IsAllowedClass(GraphNode->GetClass()) == false || 
+		if (NiagaraEditorSettings->IsReferenceableClass(GraphNode->GetClass()) == false || 
 			(NiagaraNode->GetReferencedAsset() != nullptr && NiagaraEditorSettings->IsAllowedAssetObjectByClassUsage(*NiagaraNode->GetReferencedAsset()) == false))
 		{
 			ErrorMsg = FString(TEXT("UNSUPPORTED!"));
diff --git a/Engine/Plugins/FX/Niagara/Source/NiagaraEditor/Public/NiagaraEditorSettings.h b/Engine/Plugins/FX/Niagara/Source/NiagaraEditor/Public/NiagaraEditorSettings.h
--- a/Engine/Plugins/FX/Niagara/Source/NiagaraEditor/Public/NiagaraEditorSettings.h
+++ b/Engine/Plugins/FX/Niagara/Source/NiagaraEditor/Public/NiagaraEditorSettings.h
@@ -474,11 +474,13 @@ public:
 	DECLARE_DELEGATE_RetVal_OneParam(bool, FOnIsClassPathAllowed, const FTopLevelAssetPath& /*InClassPath*/);
 	DECLARE_DELEGATE_RetVal_OneParam(bool, FOnShouldFilterAsset, const FTopLevelAssetPath& /*InAssetPath*/)
 	
-	/** Sets a delegate that allows external code to restrict which features can be used within the niagara editor by filtering which classes are allowed. */
-	NIAGARAEDITOR_API void SetOnIsClassAllowed(const FOnIsClassAllowed& InOnIsClassAllowed);
+	/** Sets a delegate that allows external code to restrict which features can be used within the niagara editor by filtering which classes are visible in various menus.
+	 * A non-visible class might still be referenceable, and can be selected through actions as copy pasting content. */
+	NIAGARAEDITOR_API void SetOnIsClassVisible(const FOnIsClassAllowed& InOnIsClassAllowed);
 
-	/** Sets a delegate that allows external code to restrict which features can be used within the niagara editor by filtering which classes are allowed by class path. */
-	NIAGARAEDITOR_API void SetOnIsClassPathAllowed(const FOnIsClassPathAllowed& InOnIsClassPathAllowed);
+	/** Sets a delegate that allows external code to restrict which features can be used within the niagara editor by filtering which classes are actually referenceable.
+	 * A referenceable class is considered 'valid' content, but might or might not be visible in menus. */
+	NIAGARAEDITOR_API void SetOnIsClassReferenceable(const FOnIsClassAllowed& InOnIsClassAllowed);
 
 	/** Sets a delegate that allows external code to restrict which features can be used within the niagara editor by filtering which assets are allowed. */
 	NIAGARAEDITOR_API void SetOnShouldFilterAssetByClassUsage(const FOnShouldFilterAsset& InOnShouldFilterAssetByClassUsage);
@@ -486,14 +488,14 @@ public:
 	/** Sets a delegate that allows external code to restrict what assets will show up in the Niagara Asset Browser. */
 	NIAGARAEDITOR_API void SetOnShouldFilterAssetInNiagaraAssetBrowser(const FOnShouldFilterAsset& InOnShouldFilterAssetByClassUsage);
 
-	/** Returns whether or not the supplied class can be used in the current editor context. */
-	NIAGARAEDITOR_API bool IsAllowedClass(const UClass* InClass) const;
+	/** Returns whether or not the supplied class is visible for UI purposes in the current editor context. */
+	NIAGARAEDITOR_API bool IsVisibleClass(const UClass* InClass) const;
 
-	/** Returns whether or not the class referenced by the supplied class path can be used in the current editor context. */
-	NIAGARAEDITOR_API bool IsAllowedClassPath(const FTopLevelAssetPath& InClassPath) const;
+	/** Returns whether or not the supplied class is valid to be used in the current editor context. */
+	NIAGARAEDITOR_API bool IsReferenceableClass(const UClass* InClass) const;
 
 	/** Returns whether or not the supplied niagara type definition can be used in the current editor context. */
-	NIAGARAEDITOR_API bool IsAllowedTypeDefinition(const FNiagaraTypeDefinition& InTypeDefinition) const;
+	NIAGARAEDITOR_API bool IsVisibleTypeDefinition(const FNiagaraTypeDefinition& InTypeDefinition) const;
 
 	NIAGARAEDITOR_API FAssetRegistryTag CreateClassUsageAssetRegistryTag(const UObject* SourceObject) const;
 
@@ -594,7 +596,8 @@ private:
 	UPROPERTY(config)
 	bool bForceSilentLoadingOfCachedAssets;
 	
-	FOnIsClassAllowed OnIsClassAllowedDelegate;
+	FOnIsClassAllowed OnIsClassVisibleDelegate;
+	FOnIsClassAllowed OnIsClassReferenceableDelegate;
 	FOnIsClassPathAllowed OnIsClassPathAllowedDelegate;
 	FOnShouldFilterAsset OnShouldFilterAssetByClassUsage;
 	FOnShouldFilterAsset OnShouldFilterAssetInNiagaraAssetBrowser;
