- When calculating the camera offset normalize the tile LWCNormalizeTile as the 0.5 tile offset is causing precision issues

diff --git a/Engine/Plugins/FX/Niagara/Shaders/Private/NiagaraMeshParticleUtils.ush b/Engine/Plugins/FX/Niagara/Shaders/Private/NiagaraMeshParticleUtils.ush
--- a/Engine/Plugins/FX/Niagara/Shaders/Private/NiagaraMeshParticleUtils.ush
+++ b/Engine/Plugins/FX/Niagara/Shaders/Private/NiagaraMeshParticleUtils.ush
@@ -326,11 +326,14 @@ float3x3 NiagaraGetMeshFacingMatrix(FLWCVector3 ParticleSimPosition, NiagaraMesh
 // Calculates the simulation-space CameraOffset of the mesh particle 
 float3 NiagaraCalculateCameraOffset(FLWCVector3 ParticleSimPosition, NiagaraMeshParticleSRTParams Params)
 {
-	FLWCVector3 WorldPosition = NiagaraSimToWorldPos(ParticleSimPosition, Params);
-	float Offset = NiagaraSafeGetFloat(Params.CameraOffsetDataOffset, Params.ParticleIndex, Params.DefaultCameraOffset);
-	float3 CameraVector = NiagaraWorldToSimVec(LWCNormalize(LWCSubtract(Params.CameraOrigin, WorldPosition)), Params);
+	const FLWCVector3 WorldPosition = NiagaraSimToWorldPos(ParticleSimPosition, Params);
+	const float Offset = NiagaraSafeGetFloat(Params.CameraOffsetDataOffset, Params.ParticleIndex, Params.DefaultCameraOffset);
 
-	return CameraVector * Offset;
+	//-Note: LWCNormalizeTile should not be required but a precision issue is causing PLAY-32984, this is a temporary workaround
+	const FLWCVector3 CameraVec = LWCNormalizeTile(LWCSubtract(Params.CameraOrigin, WorldPosition));
+	const float3 CameraDir = NiagaraWorldToSimVec(LWCNormalize(CameraVec), Params);
+		
+	return CameraDir * Offset;
 }
 
 // Calculates scale, rotation, and translation of a mesh particle. NOTE: translation and rotation are in particle simulation space
