Made several properties on the cached Niagara Script data Transient.  This prevents conflicts with deterministic DDC results (DDC-Verify) and reduces the size of our DDC payload substantially (about 50% smaller).  The issue is that translation of a niagara graph doesn't generate deterministic hlsl because functions/variables are named using FNames of UObjects created for the translation (and so their numbers will be constantly increasing).  Means that the generated code view will only be present for locally compiled scripts.

diff --git a/Engine/Plugins/FX/Niagara/Source/Niagara/Classes/NiagaraScript.h b/Engine/Plugins/FX/Niagara/Source/Niagara/Classes/NiagaraScript.h
--- a/Engine/Plugins/FX/Niagara/Source/Niagara/Classes/NiagaraScript.h
+++ b/Engine/Plugins/FX/Niagara/Source/Niagara/Classes/NiagaraScript.h
@@ -522,16 +522,16 @@ public:
 	TArray<FNiagaraStatScope> StatScopes;
 
 #if WITH_EDITORONLY_DATA
-	UPROPERTY()
+	UPROPERTY(Transient)
 	FString LastHlslTranslation;
 
-	UPROPERTY()
+	UPROPERTY(Transient)
 	FString LastHlslTranslationGPU;
 
-	UPROPERTY()
+	UPROPERTY(Transient)
 	FString LastAssemblyTranslation;
 
-	UPROPERTY()
+	UPROPERTY(Transient)
 	uint32 LastOpCount;
 #endif
 
diff --git a/Engine/Plugins/FX/Niagara/Source/NiagaraEditor/Private/Widgets/SNiagaraGeneratedCodeView.cpp b/Engine/Plugins/FX/Niagara/Source/NiagaraEditor/Private/Widgets/SNiagaraGeneratedCodeView.cpp
--- a/Engine/Plugins/FX/Niagara/Source/NiagaraEditor/Private/Widgets/SNiagaraGeneratedCodeView.cpp
+++ b/Engine/Plugins/FX/Niagara/Source/NiagaraEditor/Private/Widgets/SNiagaraGeneratedCodeView.cpp
@@ -168,7 +168,7 @@ void SNiagaraGeneratedCodeView::Construct(const FArguments& InArgs, TSharedRef<F
 				+ SVerticalBox::Slot()
 				[
 					SNew(STextBlock)
-					.Text(LOCTEXT("NoDataText", "Failed to compile or has not been compiled."))
+					.Text(LOCTEXT("NoDataText", "Generated code not available.  Either data was pulled from DDC (do a Full Rebuild of the System) or the script failed to be generated."))
 					.Visibility_Lambda([&]() {
 						if (TabHasScriptData())
 							return EVisibility::Collapsed;
diff --git a/Engine/Source/Runtime/Core/Private/UObject/FortniteReleaseBranchCustomObjectVersion.cpp b/Engine/Source/Runtime/Core/Private/UObject/FortniteReleaseBranchCustomObjectVersion.cpp
--- a/Engine/Source/Runtime/Core/Private/UObject/FortniteReleaseBranchCustomObjectVersion.cpp
+++ b/Engine/Source/Runtime/Core/Private/UObject/FortniteReleaseBranchCustomObjectVersion.cpp
@@ -6,7 +6,7 @@ TMap<FGuid, FGuid> FFortniteReleaseBranchCustomObjectVersion::GetSystemGuids()
 	TMap<FGuid, FGuid> SystemGuids;
 	const FDevSystemGuids& DevGuids = FDevSystemGuids::Get();
 
-	SystemGuids.Add(DevGuids.Niagara_LatestScriptCompileVersion, FGuid("495EDD7CBCED44B5B456F6CB2C76074C"));
+	SystemGuids.Add(DevGuids.Niagara_LatestScriptCompileVersion, FGuid("BF16049E73584021A385C4918F5C85BE"));
 	SystemGuids.Add(DevGuids.SkeletalMeshDerivedDataVersion, FGuid("5705956EC7134274A5A1BD3FC74DB3A9"));
 	return SystemGuids;
 }
