- DisableAdvancedFeatures in Niagara Preview Viewport
- This saves a lot of memory per Niagara VP that is created by disabling things like Lumen

diff --git a/Engine/Plugins/FX/Niagara/Source/NiagaraEditor/Private/Widgets/SNiagaraSystemViewport.cpp b/Engine/Plugins/FX/Niagara/Source/NiagaraEditor/Private/Widgets/SNiagaraSystemViewport.cpp
--- a/Engine/Plugins/FX/Niagara/Source/NiagaraEditor/Private/Widgets/SNiagaraSystemViewport.cpp
+++ b/Engine/Plugins/FX/Niagara/Source/NiagaraEditor/Private/Widgets/SNiagaraSystemViewport.cpp
@@ -627,6 +627,8 @@ void FNiagaraSystemViewportClient::LoadSharedSettingsFromConfig()
 	FEngineShowFlags EditorShowFlags(ESFIM_Editor);
 	FEngineShowFlags GameShowFlags(ESFIM_Game);
 
+	EditorShowFlags.DisableAdvancedFeatures();
+
 	if(!ViewportSharedSettings.EditorShowFlagsString.IsEmpty())
 	{
 		EditorShowFlags.SetFromString(*ViewportSharedSettings.EditorShowFlagsString);
