- Remove unnecessary include

diff --git a/Engine/Plugins/FX/Niagara/Source/Niagara/Private/NiagaraRenderer.cpp b/Engine/Plugins/FX/Niagara/Source/Niagara/Private/NiagaraRenderer.cpp
--- a/Engine/Plugins/FX/Niagara/Source/Niagara/Private/NiagaraRenderer.cpp
+++ b/Engine/Plugins/FX/Niagara/Source/Niagara/Private/NiagaraRenderer.cpp
@@ -15,9 +15,6 @@
 #include "NiagaraSystemInstanceController.h"
 #include "Materials/MaterialInstanceDynamic.h"
 #include "PrimitiveViewRelevance.h"
-//-TODO:Stateless: Consider a TFunction or something else to get the data
-#include "Stateless/NiagaraStatelessEmitterInstance.h"
-//-TODO:Stateless: Consider a TFunction or something else to get the data
 
 DECLARE_CYCLE_STAT(TEXT("Sort Particles"), STAT_NiagaraSortParticles, STATGROUP_Niagara);
 DECLARE_CYCLE_STAT(TEXT("Global Float Alloc - All"), STAT_NiagaraAllocateGlobalFloatAll, STATGROUP_Niagara);
