- Disable Niagara ribbons for RayTracing
- This is not supported currently as the index buffer can potentially be packed and will cause a GPU hang when used with RT

diff --git a/Engine/Plugins/FX/Niagara/Source/Niagara/Private/NiagaraRendererRibbons.cpp b/Engine/Plugins/FX/Niagara/Source/Niagara/Private/NiagaraRendererRibbons.cpp
--- a/Engine/Plugins/FX/Niagara/Source/Niagara/Private/NiagaraRendererRibbons.cpp
+++ b/Engine/Plugins/FX/Niagara/Source/Niagara/Private/NiagaraRendererRibbons.cpp
@@ -160,7 +160,7 @@ static FAutoConsoleVariableRef CVarNiagaraRibbonGpuBufferAlign(
 
 static TAutoConsoleVariable<int32> CVarRayTracingNiagaraRibbons(
 	TEXT("r.RayTracing.Geometry.NiagaraRibbons"),
-	1,
+	0,
 	TEXT("Include Niagara ribbons in ray tracing effects (default = 1 (Niagara ribbons enabled in ray tracing))"));
 
 // max absolute error 9.0x10^-3
