expose get num elements as a di function

diff --git a/Engine/Plugins/FX/Niagara/Shaders/Private/NiagaraDataInterfaceRigidMeshCollisionQuery.ush b/Engine/Plugins/FX/Niagara/Shaders/Private/NiagaraDataInterfaceRigidMeshCollisionQuery.ush
--- a/Engine/Plugins/FX/Niagara/Shaders/Private/NiagaraDataInterfaceRigidMeshCollisionQuery.ush
+++ b/Engine/Plugins/FX/Niagara/Shaders/Private/NiagaraDataInterfaceRigidMeshCollisionQuery.ush
@@ -12,6 +12,7 @@ NiagaraDataInterfaceRigidMeshCollisionQuery.ush
 #define BOX_INDEX 0
 #define SPHERE_INDEX 1
 #define CAPSULE_INDEX 2
+#define NUM_ELEMENTS_INDEX 3
 
 uint			{ParameterName}_MaxTransforms;
 uint			{ParameterName}_CurrentOffset;
@@ -45,6 +46,11 @@ void GetNumCapsules_{ParameterName}(out int OutNumCapsules)
 	OutNumCapsules = {ParameterName}_ElementOffsets[CAPSULE_INDEX+1] - {ParameterName}_ElementOffsets[CAPSULE_INDEX];
 }
 
+void GetNumElements_{ParameterName}(out int OutNumElements)
+{
+	OutNumElements = {ParameterName}_ElementOffsets[NUM_ELEMENTS_INDEX];
+}
+
 /* -----------------------------------------------------------------
  * Position / Transform utils
  * -----------------------------------------------------------------
@@ -591,7 +597,7 @@ void GetClosestPointMeshDistanceFieldAccurate_{ParameterName}(in float3 WorldPos
 	
 	NormalIsValid = false;
 
-	for (int ElementIndex = 0; ElementIndex < {ParameterName}_ElementOffsets[3]; ++ElementIndex)
+	for (int ElementIndex = 0; ElementIndex < {ParameterName}_ElementOffsets[NUM_ELEMENTS_INDEX]; ++ElementIndex)
 	{
 		uint DFIndex = {ParameterName}_DFIndexBuffer[ElementIndex];
 
diff --git a/Engine/Plugins/FX/Niagara/Source/Niagara/Private/NiagaraDataInterfaceRigidMeshCollisionQuery.cpp b/Engine/Plugins/FX/Niagara/Source/Niagara/Private/NiagaraDataInterfaceRigidMeshCollisionQuery.cpp
--- a/Engine/Plugins/FX/Niagara/Source/Niagara/Private/NiagaraDataInterfaceRigidMeshCollisionQuery.cpp
+++ b/Engine/Plugins/FX/Niagara/Source/Niagara/Private/NiagaraDataInterfaceRigidMeshCollisionQuery.cpp
@@ -87,6 +87,7 @@ static const FName FindActorsName(TEXT("FindActors"));
 static const FName GetNumBoxesName(TEXT("GetNumBoxes"));
 static const FName GetNumSpheresName(TEXT("GetNumSpheres"));
 static const FName GetNumCapsulesName(TEXT("GetNumCapsules"));
+static const FName GetNumElementsName(TEXT("GetNumElements"));
 
 //------------------------------------------------------------------------------------------------------------
 
@@ -1438,6 +1439,19 @@ void UNiagaraDataInterfaceRigidMeshCollisionQuery::GetFunctionsInternal(TArray<F
 
 		OutFunctions.Add(Sig);
 	}
+	{
+		FNiagaraFunctionSignature Sig;
+		Sig.Name = GetNumElementsName;
+		Sig.SetDescription(LOCTEXT("GetNumElementsNameDescription", "Returns the number of primitive elements for the collection of static meshes the DI represents."));
+		Sig.SetFunctionVersion(FNiagaraRigidMeshCollisionDIFunctionVersion::LatestVersion);
+		Sig.bSupportsGPU = true;
+		Sig.bSupportsCPU = true;
+		Sig.bMemberFunction = true;
+		Sig.Inputs.Add(FNiagaraVariable(FNiagaraTypeDefinition(GetClass()), TEXT("Collision DI")));
+		Sig.Outputs.Add(FNiagaraVariable(FNiagaraTypeDefinition::GetIntDef(), TEXT("Num Elements")));
+
+		OutFunctions.Add(Sig);
+	}
 	{
 		FNiagaraFunctionSignature Sig;
 		Sig.Name = GetClosestPointName;
@@ -1625,6 +1639,7 @@ void UNiagaraDataInterfaceRigidMeshCollisionQuery::GetFunctionsInternal(TArray<F
 #endif
 
 DEFINE_NDI_DIRECT_FUNC_BINDER(UNiagaraDataInterfaceRigidMeshCollisionQuery, FindActorsCPU);
+DEFINE_NDI_DIRECT_FUNC_BINDER(UNiagaraDataInterfaceRigidMeshCollisionQuery, GetNumElementsCPU);
 
 void UNiagaraDataInterfaceRigidMeshCollisionQuery::GetVMExternalFunction(const FVMExternalFunctionBindingInfo& BindingInfo, void* InstanceData, FVMExternalFunction& OutFunc)
 {
@@ -1634,6 +1649,10 @@ void UNiagaraDataInterfaceRigidMeshCollisionQuery::GetVMExternalFunction(const F
 	{
 		NDI_FUNC_BINDER(UNiagaraDataInterfaceRigidMeshCollisionQuery, FindActorsCPU)::Bind(this, OutFunc);
 	}
+	else if (BindingInfo.Name == GetNumElementsName)
+	{
+		NDI_FUNC_BINDER(UNiagaraDataInterfaceRigidMeshCollisionQuery, GetNumElementsCPU)::Bind(this, OutFunc);
+	}
 	else
 	{
 		UE_LOG(LogNiagara, Display, TEXT("Could not find data interface external function in %s. %s\n"),
@@ -1649,6 +1668,7 @@ bool UNiagaraDataInterfaceRigidMeshCollisionQuery::GetFunctionHLSL(const FNiagar
 	if ((FunctionInfo.DefinitionName == GetNumBoxesName) ||
 		(FunctionInfo.DefinitionName == GetNumCapsulesName) ||
 		(FunctionInfo.DefinitionName == GetNumSpheresName) ||
+		(FunctionInfo.DefinitionName == GetNumElementsName) ||
 		(FunctionInfo.DefinitionName == GetClosestPointName) ||
 		(FunctionInfo.DefinitionName == GetClosestPointSimpleName) ||		
 		(FunctionInfo.DefinitionName == GetClosestElementName) ||
@@ -2082,6 +2102,22 @@ void UNiagaraDataInterfaceRigidMeshCollisionQuery::FindActorsCPU(FVectorVMExtern
 	}
 }
 
+void UNiagaraDataInterfaceRigidMeshCollisionQuery::GetNumElementsCPU(FVectorVMExternalFunctionContext& Context)
+{
+	VectorVM::FUserPtrHandler<FNDIRigidMeshCollisionData> InstanceData(Context);
+
+	FNDIOutputParam<int32> NumElementsParam(Context);
+
+	if (InstanceData->SystemInstance != nullptr && InstanceData->AssetArrays != nullptr)
+	{
+		NumElementsParam.SetAndAdvance(InstanceData->AssetArrays->ElementOffsets.NumElements);
+	}
+	else
+	{
+		NumElementsParam.SetAndAdvance(0);
+	}
+}
+
 #if WITH_EDITOR
 
 void UNiagaraDataInterfaceRigidMeshCollisionQuery::GetFeedback(UNiagaraSystem* InAsset, UNiagaraComponent* InComponent, TArray<FNiagaraDataInterfaceError>& OutErrors, TArray<FNiagaraDataInterfaceFeedback>& OutWarnings, TArray<FNiagaraDataInterfaceFeedback>& OutInfo)
diff --git a/Engine/Plugins/FX/Niagara/Source/Niagara/Public/NiagaraDataInterfaceRigidMeshCollisionQuery.h b/Engine/Plugins/FX/Niagara/Source/Niagara/Public/NiagaraDataInterfaceRigidMeshCollisionQuery.h
--- a/Engine/Plugins/FX/Niagara/Source/Niagara/Public/NiagaraDataInterfaceRigidMeshCollisionQuery.h
+++ b/Engine/Plugins/FX/Niagara/Source/Niagara/Public/NiagaraDataInterfaceRigidMeshCollisionQuery.h
@@ -272,6 +272,7 @@ public:
 	bool GlobalFindActors(UWorld* World, FNDIRigidMeshCollisionData& InstanceData) const;
 
 	void FindActorsCPU(FVectorVMExternalFunctionContext& Context);
+	void GetNumElementsCPU(FVectorVMExternalFunctionContext& Context);
 
 #if WITH_EDITOR
 	virtual void GetFeedback(UNiagaraSystem* InAsset, UNiagaraComponent* InComponent, TArray<FNiagaraDataInterfaceError>& OutErrors, TArray<FNiagaraDataInterfaceFeedback>& OutWarnings, TArray<FNiagaraDataInterfaceFeedback>& OutInfo) override;
