- Replace custom widget with SDebuggerSuggestionTextBox
- Fixes an issue where focus didn't correctly close the suggestion list
- Fixes an issue where text box would sometimes get confused when using wildcards and leave characters in place

diff --git a/Engine/Plugins/FX/Niagara/Source/NiagaraEditor/Private/Customizations/NiagaraDebugHUDCustomization.cpp b/Engine/Plugins/FX/Niagara/Source/NiagaraEditor/Private/Customizations/NiagaraDebugHUDCustomization.cpp
--- a/Engine/Plugins/FX/Niagara/Source/NiagaraEditor/Private/Customizations/NiagaraDebugHUDCustomization.cpp
+++ b/Engine/Plugins/FX/Niagara/Source/NiagaraEditor/Private/Customizations/NiagaraDebugHUDCustomization.cpp
@@ -2,37 +2,6 @@
 
 #include "NiagaraDebugHUDCustomization.h"
 
-//#include "CoreMinimal.h"
-//#include "NiagaraConstants.h"
-//#include "NiagaraConstants.h"
-//#include "NiagaraEditorStyle.h"
-//#include "NiagaraEmitter.h"
-//#include "NiagaraNodeOutput.h"
-//#include "NiagaraNodeParameterMapBase.h"
-//#include "NiagaraParameterMapHistory.h"
-//#include "NiagaraPlatformSet.h"
-//#include "NiagaraRendererProperties.h"
-//#include "NiagaraScriptSource.h"
-//#include "NiagaraScriptVariable.h"
-//#include "NiagaraSystem.h"
-//#include "NiagaraTypes.h"
-//#include "PlatformInfo.h"
-//#include "PropertyHandle.h"
-//#include "SGraphActionMenu.h"
-//#include "Scalability.h"
-//#include "ScopedTransaction.h"
-//#include "DeviceProfiles/DeviceProfile.h"
-//#include "DeviceProfiles/DeviceProfileManager.h"
-//#include "Framework/Application/SlateApplication.h"
-//#include "Widgets/Input/SButton.h"
-//#include "Widgets/Input/SComboButton.h"
-//#include "Widgets/Input/STextComboBox.h"
-//#include "Widgets/Layout/SWrapBox.h"
-//#include "NiagaraSimulationStageBase.h"
-//#include "Widgets/Text/STextBlock.h"
-//#include "NiagaraDataInterfaceRW.h"
-//#include "NiagaraSettings.h"
-
 #include "Modules/ModuleManager.h"
 
 //Customization
@@ -49,6 +18,7 @@
 #include "Widgets/Views/STreeView.h"
 #include "Widgets/Input/SMenuAnchor.h"
 #include "Widgets/Input/SSearchBox.h"
+#include "Widgets/Input/SSuggestionTextBox.h"
 #include "Widgets/Layout/SScrollBorder.h"
 #include "Widgets/Input/SMultiLineEditableTextBox.h"
 ///Niagara
@@ -59,6 +29,108 @@
 
 #define LOCTEXT_NAMESPACE "NiagaraDebugHUDCustomization"
 
+//////////////////////////////////////////////////////////////////////////
+
+namespace NiagaraDebugHUDSettingsDetailsCustomizationInternal
+{
+
+class SDebuggerSuggestionTextBox : public SSuggestionTextBox
+{
+public:
+	SLATE_BEGIN_ARGS(SDebuggerSuggestionTextBox) {}
+		SLATE_ARGUMENT(TSharedPtr<IPropertyHandle>, PropertyHandle)
+		SLATE_ARGUMENT(UClass*, ObjectClass)
+	SLATE_END_ARGS()
+
+	void Construct(const FArguments& InArgs)
+	{
+		FNiagaraEditorModule& NiagaraEditorModule = FModuleManager::GetModuleChecked<FNiagaraEditorModule>("NiagaraEditor");
+		Debugger = NiagaraEditorModule.GetDebugger();
+		if ( Debugger )
+		{
+			Debugger->GetOnSimpleClientInfoChanged().AddSP(this, &SDebuggerSuggestionTextBox::OnSimpleClientInfoChanged);
+		}
+
+		PropertyHandle = InArgs._PropertyHandle;
+		ObjectClass = InArgs._ObjectClass;
+
+		FText CurrentValue;
+		PropertyHandle->GetValueAsFormattedText(CurrentValue);
+
+		SSuggestionTextBox::Construct(
+			SSuggestionTextBox::FArguments()
+			.ForegroundColor(FSlateColor::UseForeground())
+			.Font(IDetailLayoutBuilder::GetDetailFont())
+			.Text(CurrentValue)
+			.OnTextCommitted(this, &SDebuggerSuggestionTextBox::OnDebuggerTextCommitted)
+			.OnShowingSuggestions(this, &SDebuggerSuggestionTextBox::GetDebuggerSuggestions)
+		);
+	}
+
+	void OnDebuggerTextCommitted(const FText& NewText, ETextCommit::Type CommitInfo)
+	{
+		if (PropertyHandle)
+		{
+			PropertyHandle->SetValueFromFormattedString(NewText.ToString());
+		}
+	}
+
+	void GetDebuggerSuggestions(const FString& CurrText, TArray<FString>& OutSuggestions)
+	{
+		if (!Debugger.IsValid() || !ObjectClass)
+		{
+			return;
+		}
+
+		const FNiagaraSimpleClientInfo& SimpleClientInfo = Debugger->GetSimpleClientInfo();
+		if (ObjectClass == UNiagaraSystem::StaticClass())
+		{
+			OutSuggestions.Append(SimpleClientInfo.Systems);
+		}
+		else if (ObjectClass == UNiagaraEmitter::StaticClass())
+		{
+			OutSuggestions.Append(SimpleClientInfo.Emitters);
+		}
+		else if (ObjectClass == AActor::StaticClass())
+		{
+			OutSuggestions.Append(SimpleClientInfo.Actors);
+		}
+		else if (ObjectClass == UNiagaraComponent::StaticClass())
+		{
+			OutSuggestions.Append(SimpleClientInfo.Components);
+		}
+	}
+
+	void OnSimpleClientInfoChanged(const FNiagaraSimpleClientInfo& ClientInfo)
+	{
+		bWaitingUpdate = false;
+	}
+
+	virtual void Tick(const FGeometry& AllottedGeometry, const double InCurrentTime, const float InDeltaTime) override
+	{
+		SSuggestionTextBox::Tick(AllottedGeometry, InCurrentTime, InDeltaTime);
+
+		if (Debugger)
+		{
+			const bool bHasFocus = FSlateApplication::Get().HasFocusedDescendants(SharedThis(this));
+			if (bHasFocus && !bWaitingUpdate)
+			{
+				bWaitingUpdate = true;
+				Debugger->RequestUpdatedClientInfo();
+			}
+		}
+	}
+
+	TSharedPtr<FNiagaraDebugger>	Debugger;
+	TSharedPtr<IPropertyHandle>		PropertyHandle;
+	UClass*							ObjectClass = nullptr;
+	bool							bWaitingUpdate = false;
+};
+
+} // namespace
+
+//////////////////////////////////////////////////////////////////////////
+
 void FNiagaraDebugHUDVariableCustomization::CustomizeHeader(TSharedRef<IPropertyHandle> StructPropertyHandle, FDetailWidgetRow& HeaderRow, IPropertyTypeCustomizationUtils& CustomizationUtils)
 {
 	EnabledPropertyHandle = StructPropertyHandle->GetChildHandle(GET_MEMBER_NAME_CHECKED(FNiagaraDebugHUDVariable, bEnabled));
@@ -124,647 +196,79 @@ bool FNiagaraDebugHUDVariableCustomization::IsTextEditable() const
 	return bEnabled;
 }
 
-//////////////////////////////////////////////////////////////////////////
-
-/**
- * Input box for referencing objects on the client being debugged from the Niagara debugger.
- */
-class SNiagaraDebuggerObjectInputBox : public SCompoundWidget
-{
-public:
-	SLATE_BEGIN_ARGS(SNiagaraDebuggerObjectInputBox)
-		: _SuggestionListPlacement(MenuPlacement_BelowAnchor)
-	{}
-
-	/** Where to place the suggestion list */
-	SLATE_ARGUMENT(TSharedPtr<IPropertyHandle>, OwnerProperty)
-	SLATE_ARGUMENT(TSharedPtr<FNiagaraDebugger>, Debugger)
-	SLATE_ARGUMENT(UClass*, ObjectClass)
-	SLATE_ARGUMENT(bool, bAllowWildcards)
-	SLATE_ARGUMENT(EMenuPlacement, SuggestionListPlacement)
-	SLATE_EVENT(FOnTextCommitted, OnTextCommittedEvent)
-	SLATE_EVENT(FSimpleDelegate, OnCloseConsole)
-	SLATE_END_ARGS()
-
-	/** Protected console input box widget constructor, called by Slate */
-	SNiagaraDebuggerObjectInputBox();
-
-	/**
-	 * Construct this widget.  Called by the SNew() Slate macro.
-	 *
-	 * @param	InArgs	Declaration used by the SNew() macro to construct this widget
-	 */
-	void Construct(const FArguments& InArgs);
-
-	/** Returns the editable text box associated with this widget.  Used to set focus directly. */
-	TSharedRef< SMultiLineEditableTextBox > GetEditableTextBox()
-	{
-		return InputText.ToSharedRef();
-	}
-
-	/** SWidget interface */
-	virtual void Tick(const FGeometry& AllottedGeometry, const double InCurrentTime, const float InDeltaTime) override;
-
-	void OnSimpleClientInfoChanged(const FNiagaraSimpleClientInfo& ClientInfo);
-protected:
-
-	virtual bool SupportsKeyboardFocus() const override { return true; }
-
-	// e.g. Tab or Key_Up
-	virtual FReply OnPreviewKeyDown(const FGeometry& MyGeometry, const FKeyEvent& KeyEvent) override;
-
-	/** Handles entering in a command */
-	void OnTextCommitted(const FText& InText, ETextCommit::Type CommitInfo);
-
-	void OnTextChanged(const FText& InText);
-
-	/** Get the maximum width of the selection list */
-	FOptionalSize GetSelectionListMaxWidth() const;
-
-	/** Makes the widget for the suggestions messages in the list view */
-	TSharedRef<ITableRow> MakeSuggestionListItemWidget(TSharedPtr<FString> Message, const TSharedRef<STableViewBase>& OwnerTable);
-
-	void SuggestionSelectionChanged(TSharedPtr<FString> NewValue, ESelectInfo::Type SelectInfo);
-
-	void SetSuggestions(TArray<FString>& Elements, FText Highlight, bool bOpen);
-
-	void MarkActiveSuggestion();
-
-	void ClearSuggestions();
-
-	FText GetText() const;
-
-	void UpdateSuggestions(bool bOpen);
-
-	FReply OnKeyDownHandler(const FGeometry& MyGeometry, const FKeyEvent& InKeyEvent);
-
-	FReply OnKeyCharHandler(const FGeometry& MyGeometry, const FCharacterEvent& InCharacterEvent);
-
-private:
-
-	struct FSuggestions
-	{
-		FSuggestions()
-			: SelectedSuggestion(INDEX_NONE)
-		{
-		}
-
-		void Reset()
-		{
-			SelectedSuggestion = INDEX_NONE;
-			SuggestionsList.Reset();
-			SuggestionsHighlight = FText::GetEmpty();
-		}
-
-		bool HasSuggestions() const
-		{
-			return SuggestionsList.Num() > 0;
-		}
-
-		bool HasSelectedSuggestion() const
-		{
-			return SuggestionsList.IsValidIndex(SelectedSuggestion);
-		}
-
-		void StepSelectedSuggestion(const int32 Step)
-		{
-			SelectedSuggestion += Step;
-			if (SelectedSuggestion < 0)
-			{
-				SelectedSuggestion = SuggestionsList.Num() - 1;
-			}
-			else if (SelectedSuggestion >= SuggestionsList.Num())
-			{
-				SelectedSuggestion = 0;
-			}
-		}
-
-		TSharedPtr<FString> GetSelectedSuggestion() const
-		{
-			return SuggestionsList.IsValidIndex(SelectedSuggestion) ? SuggestionsList[SelectedSuggestion] : nullptr;
-		}
-
-		/** INDEX_NONE if not set, otherwise index into SuggestionsList */
-		int32 SelectedSuggestion;
-
-		/** All log messages stored in this widget for the list view */
-		TArray<TSharedPtr<FString>> SuggestionsList;
-
-		/** Highlight text to use for the suggestions list */
-		FText SuggestionsHighlight;
-	};
-
-	/** Editable text widget */
-	TSharedPtr< SMultiLineEditableTextBox > InputText;
-
-	/** Auto completion elements */
-	TSharedPtr< SMenuAnchor > SuggestionBox;
-
-	/** List view for showing suggestions. */
-	TSharedPtr< SListView< TSharedPtr<FString> > > SuggestionListView;
-
-	/** Active list of suggestions */
-	FSuggestions Suggestions;
-
-	/** Delegate to call to close the console */
-	FSimpleDelegate OnCloseConsole;
-
-	FOnTextCommitted OnTextCommittedEvent;
-
-	/** to prevent recursive calls in UI callback */
-	bool bIgnoreUIUpdate;
-
-	/** true if this widget has been Ticked at least once */
-	bool bHasTicked;
-	
-	bool bHadFocusLastFrame = false;
-
-	/** True if we consumed a tab key in OnPreviewKeyDown, so we can ignore it in OnKeyCharHandler as well */
-	bool bConsumeTab;
-
-	TSharedPtr<IPropertyHandle> OwnerProperty;
-	TSharedPtr<FNiagaraDebugger> Debugger;
-	UClass* ObjectClass = nullptr;
-	bool bAllowWildcards = true;
-};
-
-SNiagaraDebuggerObjectInputBox::SNiagaraDebuggerObjectInputBox()
-	: bIgnoreUIUpdate(false)
-	, bHasTicked(false)
-{
-}
-
-void SNiagaraDebuggerObjectInputBox::Construct(const FArguments& InArgs)
+FNiagaraDebugHUDSettingsDetailsCustomization::FNiagaraDebugHUDSettingsDetailsCustomization(UNiagaraDebugHUDSettings* InSettings)
+	: WeakSettings(InSettings)
 {
-	OnTextCommittedEvent = InArgs._OnTextCommittedEvent;
-	OnCloseConsole = InArgs._OnCloseConsole;
-	
-	OwnerProperty = InArgs._OwnerProperty;
-	Debugger = InArgs._Debugger;
-	ObjectClass = InArgs._ObjectClass;
-	bAllowWildcards = InArgs._bAllowWildcards;
-
-	Debugger->RequestUpdatedClientInfo();
-	Debugger->GetOnSimpleClientInfoChanged().AddSP(this, &SNiagaraDebuggerObjectInputBox::OnSimpleClientInfoChanged);
-
-	SAssignNew(InputText, SMultiLineEditableTextBox)
-		.AllowMultiLine(false)
-		.ClearTextSelectionOnFocusLoss(true)
-		.SelectAllTextWhenFocused(true)
-		.Text(TAttribute<FText>::Create(TAttribute<FText>::FGetter::CreateSP(this, &SNiagaraDebuggerObjectInputBox::GetText)))
-		.OnTextCommitted(this, &SNiagaraDebuggerObjectInputBox::OnTextCommitted)
-		.OnTextChanged(this, &SNiagaraDebuggerObjectInputBox::OnTextChanged)
-		.OnKeyCharHandler(this, &SNiagaraDebuggerObjectInputBox::OnKeyCharHandler)
-		.OnKeyDownHandler(this, &SNiagaraDebuggerObjectInputBox::OnKeyDownHandler)
-		.OnIsTypedCharValid(FOnIsTypedCharValid::CreateLambda([](const TCHAR InCh) { return true; })) // allow tabs to be typed into the field
-		.ClearKeyboardFocusOnCommit(true)
-		.ModiferKeyForNewLine(EModifierKey::Shift);
-
-
-	EPopupMethod PopupMethod = GIsEditor ? EPopupMethod::CreateNewWindow : EPopupMethod::UseCurrentWindow;
-	ChildSlot
-		[
-			SAssignNew(SuggestionBox, SMenuAnchor)
-			.Method(PopupMethod)
-			.Placement(InArgs._SuggestionListPlacement)
-			[
-				SNew(SHorizontalBox)
-				+ SHorizontalBox::Slot()
-				[
-					InputText.ToSharedRef()
-				]
-			]
-			.MenuContent
-			(
-				SNew(SBorder)
-				.BorderImage(FAppStyle::GetBrush("Menu.Background"))
-				.Padding(FMargin(2))
-				[
-					SNew(SBox)
-					.HeightOverride(250) // avoids flickering, ideally this would be adaptive to the content without flickering
-					.MinDesiredWidth(300)
-					.MaxDesiredWidth(this, &SNiagaraDebuggerObjectInputBox::GetSelectionListMaxWidth)
-					[
-						SAssignNew(SuggestionListView, SListView< TSharedPtr<FString> >)
-						.ListItemsSource(&Suggestions.SuggestionsList)
-						.SelectionMode(ESelectionMode::Single)							// Ideally the mouse over would not highlight while keyboard controls the UI
-						.OnGenerateRow(this, &SNiagaraDebuggerObjectInputBox::MakeSuggestionListItemWidget)
-						.OnSelectionChanged(this, &SNiagaraDebuggerObjectInputBox::SuggestionSelectionChanged)
-						.ItemHeight(18)
-					]
-				]
-			)
-		];
 }
 
-void SNiagaraDebuggerObjectInputBox::Tick(const FGeometry& AllottedGeometry, const double InCurrentTime, const float InDeltaTime)
+void FNiagaraDebugHUDSettingsDetailsCustomization::CustomizeDetails(IDetailLayoutBuilder& DetailBuilder)
 {
-	bHasTicked = true;
-	
-	bool bHasFocus = FSlateApplication::Get().HasFocusedDescendants(SharedThis(this));
-	if (bHasFocus != bHadFocusLastFrame)
-	{
-		bHadFocusLastFrame = bHasFocus;
-
-		if (Debugger && bHasFocus)
-		{
-			Debugger->RequestUpdatedClientInfo();
-		}
-	}
-
-	if (!GIntraFrameDebuggingGameThread && !IsEnabled())
-	{
-		SetEnabled(true);
-	}
-	else if (GIntraFrameDebuggingGameThread && IsEnabled())
+	// Customize General
 	{
-		SetEnabled(false);
+		IDetailCategoryBuilder& GeneralCategory = DetailBuilder.EditCategory("Debug General");
 	}
-}
 
-void SNiagaraDebuggerObjectInputBox::OnSimpleClientInfoChanged(const FNiagaraSimpleClientInfo& ClientInfo)
-{
-	UpdateSuggestions(false);
-}
-
-void SNiagaraDebuggerObjectInputBox::UpdateSuggestions(bool bOpen)
-{
-	TArray<FString> AutoCompleteList;
-
-	const FNiagaraSimpleClientInfo& SimpleClientInfo = Debugger->GetSimpleClientInfo();
-	const TArray<FString>* Names;
-	if (ObjectClass == UNiagaraSystem::StaticClass()) Names = &SimpleClientInfo.Systems;
-	else if (ObjectClass == UNiagaraEmitter::StaticClass()) Names = &SimpleClientInfo.Emitters;
-	else if (ObjectClass == AActor::StaticClass()) Names = &SimpleClientInfo.Actors;
-	else if (ObjectClass == UNiagaraComponent::StaticClass()) Names = &SimpleClientInfo.Components;
-	else
+	// Customize Overview
 	{
-		static TArray<FString> DummyNames;
-		Names = &DummyNames;
+		IDetailCategoryBuilder& OverviewCategory = DetailBuilder.EditCategory("Debug Overview");
 	}
 
-	const FString& InputTextStr = InputText->GetText().ToString();
-	for (const FString& SysName : *Names)
+	// Customize Filters
 	{
-		if (InputTextStr.IsEmpty() || SysName.Contains(InputTextStr) || (bAllowWildcards && SysName.MatchesWildcard(InputTextStr)))
-		{
-			AutoCompleteList.Add(SysName);
-		}
+		IDetailCategoryBuilder& FilterCategory = DetailBuilder.EditCategory("Debug Filter");
+		MakeCustomAssetSearch(DetailBuilder, FilterCategory, DetailBuilder.GetProperty(GET_MEMBER_NAME_CHECKED(FNiagaraDebugHUDSettingsData, SystemFilter), FNiagaraDebugHUDSettingsData::StaticStruct()), UNiagaraSystem::StaticClass(), [&]() -> bool& { return WeakSettings.Get()->Data.bSystemFilterEnabled; });
+		MakeCustomAssetSearch(DetailBuilder, FilterCategory, DetailBuilder.GetProperty(GET_MEMBER_NAME_CHECKED(FNiagaraDebugHUDSettingsData, EmitterFilter), FNiagaraDebugHUDSettingsData::StaticStruct()), UNiagaraEmitter::StaticClass(), [&]() -> bool& { return WeakSettings.Get()->Data.bEmitterFilterEnabled; });
+		MakeCustomAssetSearch(DetailBuilder, FilterCategory, DetailBuilder.GetProperty(GET_MEMBER_NAME_CHECKED(FNiagaraDebugHUDSettingsData, ActorFilter), FNiagaraDebugHUDSettingsData::StaticStruct()), AActor::StaticClass(), [&]() -> bool& { return WeakSettings.Get()->Data.bActorFilterEnabled; });
+		MakeCustomAssetSearch(DetailBuilder, FilterCategory, DetailBuilder.GetProperty(GET_MEMBER_NAME_CHECKED(FNiagaraDebugHUDSettingsData, ComponentFilter), FNiagaraDebugHUDSettingsData::StaticStruct()), UNiagaraComponent::StaticClass(), [&]() -> bool& { return WeakSettings.Get()->Data.bComponentFilterEnabled; });
 	}
-
-	AutoCompleteList.Sort([InputTextStr](const FString& A, const FString& B)
-		{
-			if (A.StartsWith(InputTextStr))
-			{
-				if (!B.StartsWith(InputTextStr))
-				{
-					return true;
-				}
-			}
-			else
-			{
-				if (B.StartsWith(InputTextStr))
-				{
-					return false;
-				}
-			}
-
-			return A < B;
-
-		});
-
-
-	SetSuggestions(AutoCompleteList, FText::FromString(InputTextStr), bOpen);
 }
 
-void SNiagaraDebuggerObjectInputBox::SuggestionSelectionChanged(TSharedPtr<FString> NewValue, ESelectInfo::Type SelectInfo)
+void FNiagaraDebugHUDSettingsDetailsCustomization::MakeCustomAssetSearch(IDetailLayoutBuilder& DetailBuilder, IDetailCategoryBuilder& DetailCategory, TSharedRef<IPropertyHandle> PropertyHandle, UClass* ObjRefClass, TFunction<bool&()> GetEditBool)
 {
-	if (bIgnoreUIUpdate)
+	if ( !PropertyHandle->IsValidHandle() || (ObjRefClass == nullptr) )
 	{
 		return;
 	}
 
-	Suggestions.SelectedSuggestion = Suggestions.SuggestionsList.IndexOfByPredicate([&NewValue](const TSharedPtr<FString>& InSuggestion)
-		{
-			return InSuggestion == NewValue;
-		});
-
-	MarkActiveSuggestion();
-
-	// If the user selected this suggestion by clicking on it, then go ahead and close the suggestion
-	// box as they've chosen the suggestion they're interested in.
-	if (SelectInfo == ESelectInfo::OnMouseClick)
-	{
-		FString CommitString = *NewValue;
-		OnTextCommitted(FText::FromString(CommitString), ETextCommit::Default);
-		SuggestionBox->SetIsOpen(false);
-	}
-
-	// Ideally this would set the focus back to the edit control
-//	FWidgetPath WidgetToFocusPath;
-//	FSlateApplication::Get().GeneratePathToWidgetUnchecked( InputText.ToSharedRef(), WidgetToFocusPath );
-//	FSlateApplication::Get().SetKeyboardFocus( WidgetToFocusPath, EFocusCause::SetDirectly );
-}
-
-FOptionalSize SNiagaraDebuggerObjectInputBox::GetSelectionListMaxWidth() const
-{
-	// Limit the width of the suggestions list to the work area that this widget currently resides on
-	const FSlateRect WidgetRect(GetCachedGeometry().GetAbsolutePosition(), GetCachedGeometry().GetAbsolutePosition() + GetCachedGeometry().GetAbsoluteSize());
-	const FSlateRect WidgetWorkArea = FSlateApplication::Get().GetWorkArea(WidgetRect);
-	return FMath::Max(300.0f, WidgetWorkArea.GetSize().X - 12.0f);
-}
-
-TSharedRef<ITableRow> SNiagaraDebuggerObjectInputBox::MakeSuggestionListItemWidget(TSharedPtr<FString> Text, const TSharedRef<STableViewBase>& OwnerTable)
-{
-	check(Text.IsValid());
+	FText CurrentValue;
+	PropertyHandle->GetValueAsFormattedText(CurrentValue);
 
-	FString SanitizedText = *Text;
-	SanitizedText.ReplaceInline(TEXT("\r\n"), TEXT("\n"), ESearchCase::CaseSensitive);
-	SanitizedText.ReplaceInline(TEXT("\r"), TEXT(" "), ESearchCase::CaseSensitive);
-	SanitizedText.ReplaceInline(TEXT("\n"), TEXT(" "), ESearchCase::CaseSensitive);
+	DetailBuilder.HideProperty(PropertyHandle);
 
-	return
-		SNew(STableRow< TSharedPtr<FString> >, OwnerTable)
+	DetailCategory.AddCustomRow(PropertyHandle->GetPropertyDisplayName())
+	.NameContent()
+	[
+		SNew(SHorizontalBox)
+		+ SHorizontalBox::Slot()
+		.AutoWidth()
 		[
-			SNew(STextBlock)
-			.Text(FText::FromString(SanitizedText))
-			.TextStyle(FAppStyle::Get(), "Log.Normal")
-			.HighlightText(Suggestions.SuggestionsHighlight)
-		];
-}
-
-void SNiagaraDebuggerObjectInputBox::OnTextChanged(const FText& InText)
-{
-	if (bIgnoreUIUpdate)
-	{
-		return;
-	}
-	
-	InputText->SetText(InText);
-
-	UpdateSuggestions(true);
-}
-
-void SNiagaraDebuggerObjectInputBox::OnTextCommitted(const FText& InText, ETextCommit::Type CommitInfo)
-{
-	OwnerProperty->SetValueFromFormattedString(InText.ToString());
-	OnTextCommittedEvent.ExecuteIfBound(InText, CommitInfo);
-}
-
-FReply SNiagaraDebuggerObjectInputBox::OnPreviewKeyDown(const FGeometry& MyGeometry, const FKeyEvent& KeyEvent)
-{
-	if (SuggestionBox->IsOpen())
-	{
-		if (KeyEvent.GetKey() == EKeys::Up || KeyEvent.GetKey() == EKeys::Down)
-		{
-			Suggestions.StepSelectedSuggestion(KeyEvent.GetKey() == EKeys::Up ? -1 : +1);
-			MarkActiveSuggestion();
-
-			return FReply::Handled();
-		}
-		else if (KeyEvent.GetKey() == EKeys::Tab)
-		{
-			if (Suggestions.HasSuggestions())
-			{
-				if (Suggestions.HasSelectedSuggestion())
-				{
-					Suggestions.StepSelectedSuggestion(KeyEvent.IsShiftDown() ? -1 : +1);
-				}
-				else
-				{
-					Suggestions.SelectedSuggestion = 0;
-				}
-				MarkActiveSuggestion();
-			}
-
-			bConsumeTab = true;
-			return FReply::Handled();
-		}
-		else if (KeyEvent.GetKey() == EKeys::Escape)
-		{
-			SuggestionBox->SetIsOpen(false);
-			return FReply::Handled();
-		}
-		else if (KeyEvent.GetKey() == EKeys::Enter)
-		{
-			if (Suggestions.GetSelectedSuggestion().IsValid())
-			{
-				FText SelectionText = FText::FromString(*Suggestions.GetSelectedSuggestion().Get());
-				OnTextChanged(SelectionText);
-				OnTextCommitted(SelectionText, ETextCommit::OnEnter);
-				FSlateApplication::Get().ClearKeyboardFocus(EFocusCause::Cleared);
-			}
-			SuggestionBox->SetIsOpen(false);
-			return FReply::Handled();
-		}
-	}
-	else
-	{
-		if (KeyEvent.GetKey() == EKeys::Up || KeyEvent.GetKey() == EKeys::Tab)
-		{
-			bConsumeTab = true;
-			UpdateSuggestions(true);
-		}
-		else if (KeyEvent.GetKey() == EKeys::Escape)
-		{
-			if (InputText->GetText().IsEmpty())
-			{
-				OnCloseConsole.ExecuteIfBound();
-			}
-			else
-			{
-				// Clear the console input area
-				bIgnoreUIUpdate = true;
-				InputText->SetText(FText::GetEmpty());
-				bIgnoreUIUpdate = false;
-
-				ClearSuggestions();
-			}
-
-			return FReply::Handled();
-		}
-	}
-
-	return FReply::Unhandled();
-}
-
-void SNiagaraDebuggerObjectInputBox::SetSuggestions(TArray<FString>& Elements, FText Highlight, bool bOpen)
-{
-	FString SelectionText;
-	if (Suggestions.HasSelectedSuggestion())
-	{
-		SelectionText = *Suggestions.GetSelectedSuggestion();
-	}
-
-	Suggestions.Reset();
-	Suggestions.SuggestionsHighlight = Highlight;
-
-	for (int32 i = 0; i < Elements.Num(); ++i)
-	{
-		Suggestions.SuggestionsList.Add(MakeShared<FString>(Elements[i]));
-
-		if (Elements[i] == SelectionText)
-		{
-			Suggestions.SelectedSuggestion = i;
-		}
-	}
-	SuggestionListView->RequestListRefresh();
-
-	if (Suggestions.HasSuggestions())
-	{
-		// Ideally if the selection box is open the output window is not changing it's window title (flickers)
-		SuggestionBox->SetIsOpen(bOpen, false);
-		if (Suggestions.HasSelectedSuggestion())
-		{
-			SuggestionListView->RequestScrollIntoView(Suggestions.GetSelectedSuggestion());
-		}
-		else
-		{
-			SuggestionListView->ScrollToTop();
-		}
-	}
-	else
-	{
-		SuggestionBox->SetIsOpen(false);
-	}
-}
-
-void SNiagaraDebuggerObjectInputBox::MarkActiveSuggestion()
-{
-	bIgnoreUIUpdate = true;
-	if (Suggestions.HasSelectedSuggestion())
-	{
-		TSharedPtr<FString> SelectedSuggestion = Suggestions.GetSelectedSuggestion();
-
-		SuggestionListView->SetSelection(SelectedSuggestion);
-		SuggestionListView->RequestScrollIntoView(SelectedSuggestion);	// Ideally this would only scroll if outside of the view
-
-		InputText->SetText(FText::FromString(*SelectedSuggestion));
-	}
-	else
-	{
-		SuggestionListView->ClearSelection();
-	}
-	bIgnoreUIUpdate = false;
-}
-
-void SNiagaraDebuggerObjectInputBox::ClearSuggestions()
-{
-	SuggestionBox->SetIsOpen(false);
-	Suggestions.Reset();
-}
-
-FText SNiagaraDebuggerObjectInputBox::GetText() const
-{
-	if (OwnerProperty.IsValid())
-	{
-		FText Ret;
-		OwnerProperty->GetValueAsFormattedText(Ret);
-		return Ret;
-	}
-	return FText::GetEmpty();
-}
-
-FReply SNiagaraDebuggerObjectInputBox::OnKeyDownHandler(const FGeometry& MyGeometry, const FKeyEvent& InKeyEvent)
-{
-	return FReply::Unhandled();
-}
-
-FReply SNiagaraDebuggerObjectInputBox::OnKeyCharHandler(const FGeometry& MyGeometry, const FCharacterEvent& InCharacterEvent)
-{
-	// A printable key may be used to open the console, so consume all characters before our first Tick
-	if (!bHasTicked)
-	{
-		return FReply::Handled();
-	}
-
-	// Intercept tab if used for auto-complete
-	if (InCharacterEvent.GetCharacter() == '\t' && bConsumeTab)
-	{
-		bConsumeTab = false;
-		return FReply::Handled();
-	}
-
-	return FReply::Unhandled();
-}
-
-//////////////////////////////////////////////////////////////////////////
-
-FNiagaraDebugHUDSettingsDetailsCustomization::FNiagaraDebugHUDSettingsDetailsCustomization(UNiagaraDebugHUDSettings* InSettings)
-	: WeakSettings(InSettings)
-{
-}
-
-void FNiagaraDebugHUDSettingsDetailsCustomization::CustomizeDetails(class IDetailLayoutBuilder& DetailBuilder)
-{
-	FNiagaraEditorModule& NiagaraEditorModule = FModuleManager::GetModuleChecked<FNiagaraEditorModule>("NiagaraEditor");
-	auto CustomAssetSearch =
-		[&](IDetailCategoryBuilder& DetailCategory, TSharedRef<IPropertyHandle> PropertyHandle, UClass* ObjRefClass, TFunction<bool&()> GetEditBool)
-		{
-			if ( !PropertyHandle->IsValidHandle() || (ObjRefClass == nullptr) )
-			{
-				return;
-			}
-
-			DetailBuilder.HideProperty(PropertyHandle);
-
-			DetailCategory.AddCustomRow(PropertyHandle->GetPropertyDisplayName())
-			.NameContent()
+			SNew(SCheckBox)
+			.IsChecked_Lambda([GetEditBool]() -> ECheckBoxState { return GetEditBool() ? ECheckBoxState::Checked : ECheckBoxState::Unchecked; } )
+			.OnCheckStateChanged_Lambda([this, GetEditBool](ECheckBoxState NewState) { GetEditBool() = NewState == ECheckBoxState::Checked; WeakSettings.Get()->NotifyPropertyChanged(); })
+		]
+		+ SHorizontalBox::Slot()
+		[
+			SNew(SHorizontalBox)
+			.IsEnabled_Lambda([=]() { return GetEditBool(); })
+			+ SHorizontalBox::Slot()
 			[
-				SNew(SHorizontalBox)
-				+ SHorizontalBox::Slot()
-				.AutoWidth()
-				[
-					SNew(SCheckBox)
-					.IsChecked_Lambda([GetEditBool]() -> ECheckBoxState { return GetEditBool() ? ECheckBoxState::Checked : ECheckBoxState::Unchecked; } )
-					.OnCheckStateChanged_Lambda([this, GetEditBool](ECheckBoxState NewState) { GetEditBool() = NewState == ECheckBoxState::Checked; WeakSettings.Get()->NotifyPropertyChanged(); })
-				]
-				+ SHorizontalBox::Slot()
-				[
-					SNew(SHorizontalBox)
-					.IsEnabled_Lambda([=]() { return GetEditBool(); })
-					+ SHorizontalBox::Slot()
-					[
-						PropertyHandle->CreatePropertyNameWidget()
-					]
-				]
+				PropertyHandle->CreatePropertyNameWidget()
 			]
-			.ValueContent()
-			[
-				SNew(SHorizontalBox)
-				.IsEnabled_Lambda([=]() { return GetEditBool(); })
-				+ SHorizontalBox::Slot()
-				[
-					SNew(SNiagaraDebuggerObjectInputBox)
-					.OwnerProperty(PropertyHandle)
-					.Debugger(NiagaraEditorModule.GetDebugger())
-					.ObjectClass(ObjRefClass)
-					.bAllowWildcards(true)
-					.SuggestionListPlacement(EMenuPlacement::MenuPlacement_BelowAnchor)
-				]
-			];
-		};
-
-	// Customize General
-	{
-		IDetailCategoryBuilder& GeneralCategory = DetailBuilder.EditCategory("Debug General");
-	}
-
-	// Customize Overview
-	{
-		IDetailCategoryBuilder& OverviewCategory = DetailBuilder.EditCategory("Debug Overview");
-	}
-
-	// Customize Filters
-	{
-		IDetailCategoryBuilder& FilterCategory = DetailBuilder.EditCategory("Debug Filter");
-		CustomAssetSearch(FilterCategory, DetailBuilder.GetProperty(GET_MEMBER_NAME_CHECKED(FNiagaraDebugHUDSettingsData, SystemFilter), FNiagaraDebugHUDSettingsData::StaticStruct()), UNiagaraSystem::StaticClass(), [&]() -> bool& { return WeakSettings.Get()->Data.bSystemFilterEnabled; });
-		CustomAssetSearch(FilterCategory, DetailBuilder.GetProperty(GET_MEMBER_NAME_CHECKED(FNiagaraDebugHUDSettingsData, EmitterFilter), FNiagaraDebugHUDSettingsData::StaticStruct()), UNiagaraEmitter::StaticClass(), [&]() -> bool& { return WeakSettings.Get()->Data.bEmitterFilterEnabled; });
-		CustomAssetSearch(FilterCategory, DetailBuilder.GetProperty(GET_MEMBER_NAME_CHECKED(FNiagaraDebugHUDSettingsData, ActorFilter), FNiagaraDebugHUDSettingsData::StaticStruct()), AActor::StaticClass(), [&]() -> bool& { return WeakSettings.Get()->Data.bActorFilterEnabled; });
-		CustomAssetSearch(FilterCategory, DetailBuilder.GetProperty(GET_MEMBER_NAME_CHECKED(FNiagaraDebugHUDSettingsData, ComponentFilter), FNiagaraDebugHUDSettingsData::StaticStruct()), UNiagaraComponent::StaticClass(), [&]() -> bool& { return WeakSettings.Get()->Data.bComponentFilterEnabled; });
-	}
+		]
+	]
+	.ValueContent()
+	[
+		SNew(SHorizontalBox)
+		.IsEnabled_Lambda([=]() { return GetEditBool(); })
+		+ SHorizontalBox::Slot()
+		[
+			SNew(NiagaraDebugHUDSettingsDetailsCustomizationInternal::SDebuggerSuggestionTextBox)
+			.PropertyHandle(PropertyHandle)
+			.ObjectClass(ObjRefClass)
+		]
+	];
 }
 
 #undef LOCTEXT_NAMESPACE
 
-#endif
\ No newline at end of file
+#endif
diff --git a/Engine/Plugins/FX/Niagara/Source/NiagaraEditor/Private/Customizations/NiagaraDebugHUDCustomization.h b/Engine/Plugins/FX/Niagara/Source/NiagaraEditor/Private/Customizations/NiagaraDebugHUDCustomization.h
--- a/Engine/Plugins/FX/Niagara/Source/NiagaraEditor/Private/Customizations/NiagaraDebugHUDCustomization.h
+++ b/Engine/Plugins/FX/Niagara/Source/NiagaraEditor/Private/Customizations/NiagaraDebugHUDCustomization.h
@@ -10,6 +10,8 @@
 #include "Layout/Visibility.h"
 #include "NiagaraDebuggerCommon.h"
 
+class IDetailCategoryBuilder;
+class IDetailLayoutBuilder;
 class FDetailWidgetRow;
 class IPropertyHandle;
 class IPropertyHandleArray;
@@ -57,9 +59,12 @@ public:
 	}
 
 	// Begin IDetailCustomization interface
-	virtual void CustomizeDetails(class IDetailLayoutBuilder& DetailBuilder) override;
+	virtual void CustomizeDetails(IDetailLayoutBuilder& DetailBuilder) override;
 	// End IDetailCustomization interface
 
+private:
+	void MakeCustomAssetSearch(IDetailLayoutBuilder& DetailBuilder, IDetailCategoryBuilder& DetailCategory, TSharedRef<IPropertyHandle> PropertyHandle, UClass* ObjRefClass, TFunction<bool& ()> GetEditBool);
+
 private:
 	TWeakObjectPtr<UNiagaraDebugHUDSettings> WeakSettings = nullptr;
 };
