- Add an optional name to validation rules
- This makes it easier to know what you are breaking

#changelist validated
#virtualized

diff --git a/Engine/Plugins/FX/Niagara/Source/NiagaraEditor/Private/NiagaraValidationRules.cpp b/Engine/Plugins/FX/Niagara/Source/NiagaraEditor/Private/NiagaraValidationRules.cpp
--- a/Engine/Plugins/FX/Niagara/Source/NiagaraEditor/Private/NiagaraValidationRules.cpp
+++ b/Engine/Plugins/FX/Niagara/Source/NiagaraEditor/Private/NiagaraValidationRules.cpp
@@ -506,10 +506,20 @@ void UNiagaraValidationRule_EmitterCount::CheckValidity(const FNiagaraValidation
 
 		const FString PlatformConflicts = NiagaraValidation::GetPlatformConflictsString(ConflictsPerLimit[i]);
 
+		FText RuleName;
+		if (EmitterCountLimits[i].RuleName.IsEmpty())
+		{
+			RuleName = LOCTEXT("EmitterCountLimitExceeded", "Emitter count limit exceeded");
+		}
+		else
+		{
+			RuleName = FText::Format(LOCTEXT("EmitterCountLimitExceededFmt", "Emitter count limit '{0}' exceeded"), FText::FromString(EmitterCountLimits[i].RuleName));
+		}
+
 		FNiagaraValidationResult& Result = OutResults.AddDefaulted_GetRef();
 		Result.Severity = Severity;
-		Result.SummaryText = FText::Format(LOCTEXT("EmitterCountLimit", "Emitter count limit {0} exceeded {1}."), EmitterCountLimit, EmitterCountPerLimit[i]);
-		Result.Description = FText::Format(LOCTEXT("EmitterCountLimitDesc", "Emitter count limit {0} has been exceeded {1} for platforms '{2}' please reduce the emitter count to improve performance."), EmitterCountLimit, EmitterCountPerLimit[i], FText::FromString(PlatformConflicts));
+		Result.SummaryText = FText::Format(LOCTEXT("EmitterCountLimit", "{0} {1}/{2}."), RuleName, EmitterCountPerLimit[i], EmitterCountLimit);
+		Result.Description = FText::Format(LOCTEXT("EmitterCountLimitDesc", "{0} {1}/{2} for platforms '{3}' please reduce the emitter count to improve performance."), RuleName, EmitterCountPerLimit[i], EmitterCountLimit, FText::FromString(PlatformConflicts));
 		Result.SourceObject = NiagaraValidation::GetStackEntry<UNiagaraStackSystemPropertiesItem>(Context.ViewModel->GetSystemStackViewModel());
 	}
 }
@@ -573,10 +583,20 @@ void UNiagaraValidationRule_RendererCount::CheckValidity(const FNiagaraValidatio
 
 		const FString PlatformConflicts = NiagaraValidation::GetPlatformConflictsString(ConflictsPerLimit[i]);
 
+		FText RuleName;
+		if (RendererCountLimits[i].RuleName.IsEmpty())
+		{
+			RuleName = LOCTEXT("RendererCountLimitExceeded", "Renderer count limit exceeded");
+		}
+		else
+		{
+			RuleName = FText::Format(LOCTEXT("RendererCountLimitExceededFmt", "Renderer count limit '{0}' exceeded"), FText::FromString(RendererCountLimits[i].RuleName));
+		}
+
 		FNiagaraValidationResult& Result = OutResults.AddDefaulted_GetRef();
 		Result.Severity = Severity;
-		Result.SummaryText = FText::Format(LOCTEXT("RendererCountLimit", "Renderer count limit {0} exceeded {1}."), RendererCountLimit, RendererCountPerLimit[i]);
-		Result.Description = FText::Format(LOCTEXT("RendererCountLimitDesc", "Renderer count limit {0} has been exceeded {1} for platforms '{2}' please reduce the renderer count to improve performance."), RendererCountLimit, RendererCountPerLimit[i], FText::FromString(PlatformConflicts));
+		Result.SummaryText = FText::Format(LOCTEXT("RendererCountLimit", "{0} {1}/{2}."), RuleName, RendererCountPerLimit[i], RendererCountLimit);
+		Result.Description = FText::Format(LOCTEXT("RendererCountLimitDesc", "{0} {1}/{2} for platforms '{3}' please reduce the renderer count to improve performance."), RuleName, RendererCountPerLimit[i], RendererCountLimit, FText::FromString(PlatformConflicts));
 		Result.SourceObject = NiagaraValidation::GetStackEntry<UNiagaraStackSystemPropertiesItem>(Context.ViewModel->GetSystemStackViewModel());
 	}
 }
diff --git a/Engine/Plugins/FX/Niagara/Source/NiagaraEditor/Public/NiagaraValidationRules.h b/Engine/Plugins/FX/Niagara/Source/NiagaraEditor/Public/NiagaraValidationRules.h
--- a/Engine/Plugins/FX/Niagara/Source/NiagaraEditor/Public/NiagaraValidationRules.h
+++ b/Engine/Plugins/FX/Niagara/Source/NiagaraEditor/Public/NiagaraValidationRules.h
@@ -40,6 +40,10 @@ struct FNiagaraValidationRule_EmitterCountAndPlatformSet
 {
 	GENERATED_BODY()
 
+	/** Name to display if we fail the limit check */
+	UPROPERTY(EditAnywhere, Category = Validation)
+	FString RuleName;
+
 	UPROPERTY(EditAnywhere, Category = Validation)
 	FNiagaraPlatformSet Platforms;
 
@@ -68,6 +72,10 @@ struct FNiagaraValidationRule_RendererCountAndPlatformSet
 {
 	GENERATED_BODY()
 
+	/** Name to display if we fail the limit check */
+	UPROPERTY(EditAnywhere, Category = Validation)
+	FString RuleName;
+
 	UPROPERTY(EditAnywhere, Category = Validation)
 	FNiagaraPlatformSet Platforms;
 
