Couple of fixes for handling GPUComputeScript within cooked editor builds.
-allow LastHlslTranslationGPU to be preserved so that compilation can still happen
-prevent ScriptResource being discarded when serialized in cooked editor

diff --git a/Engine/Plugins/FX/Niagara/Source/Niagara/Classes/NiagaraScript.h b/Engine/Plugins/FX/Niagara/Source/Niagara/Classes/NiagaraScript.h
--- a/Engine/Plugins/FX/Niagara/Source/Niagara/Classes/NiagaraScript.h
+++ b/Engine/Plugins/FX/Niagara/Source/Niagara/Classes/NiagaraScript.h
@@ -1361,6 +1361,7 @@ private:
 
 	void ProcessSerializedShaderMaps();
 	void SerializeNiagaraShaderMaps(FArchive& Ar, int32 NiagaraVer, bool IsValidShaderScript);
+	void CreateScriptResource(const FNiagaraShaderScript* Source);
 
 	/** Return the expected SimTarget for this script. Only returns a valid target if there is valid data to run with. */
 	TOptional<ENiagaraSimTarget> GetSimTarget() const;
diff --git a/Engine/Plugins/FX/Niagara/Source/Niagara/Private/NiagaraScript.cpp b/Engine/Plugins/FX/Niagara/Source/Niagara/Private/NiagaraScript.cpp
--- a/Engine/Plugins/FX/Niagara/Source/Niagara/Private/NiagaraScript.cpp
+++ b/Engine/Plugins/FX/Niagara/Source/Niagara/Private/NiagaraScript.cpp
@@ -1015,9 +1015,6 @@ UNiagaraScript::UNiagaraScript(const FObjectInitializer& ObjectInitializer)
 #endif
 {
 #if WITH_EDITORONLY_DATA
-	ScriptResource = MakeShared<FNiagaraShaderScript>();
-	ScriptResource->OnCompilationComplete().AddUniqueDynamic(this, &UNiagaraScript::RaiseOnGPUCompilationComplete);
-
 	RapidIterationParameters.DebugName = *GetFullName();
 #endif
 }
@@ -1932,7 +1929,10 @@ void UNiagaraScript::Serialize(FArchive& Ar)
 		auto SwapDebugData = [](FNiagaraVMExecutableData& SourceData, FNiagaraVMExecutableData& TargetData) -> void
 		{
 			Swap(TargetData.LastHlslTranslation, SourceData.LastHlslTranslation);
-			Swap(TargetData.LastHlslTranslationGPU, SourceData.LastHlslTranslationGPU);
+
+			// We must preserve the LastTranslationGPU as it is used to compile the GPUComputeScript when using the standard compilation mode
+			// Swap(TargetData.LastHlslTranslationGPU, SourceData.LastHlslTranslationGPU);
+
 			Swap(TargetData.LastAssemblyTranslation, SourceData.LastAssemblyTranslation);
 			Swap(TargetData.CompileTagsEditorOnly, SourceData.CompileTagsEditorOnly);
 		};
@@ -3010,7 +3010,7 @@ void UNiagaraScript::SetVMCompilationResults(const FNiagaraVMExecutableDataId& I
 	CachedScriptVM.OptimizationTask.State = nullptr;
 	CachedParameterCollectionReferences.Empty();
 	// Proactively clear out the script resource, because it might be stale now.
-	ScriptResource->Invalidate();
+	ScriptResource = nullptr;
 
 	if (CachedScriptVM.LastCompileStatus == ENiagaraScriptCompileStatus::NCS_Error)
 	{
@@ -3211,7 +3211,7 @@ void UNiagaraScript::SetComputeCompilationResults(
 			{
 				if (!ScriptResource.IsValid())
 				{
-					ScriptResource = MakeShared<FNiagaraShaderScript>(*TargetShaderScript);
+					CreateScriptResource(TargetShaderScript);
 				}
 				else
 				{
@@ -3908,34 +3908,49 @@ void UNiagaraScript::CacheResourceShadersForRendering(bool bRegenerateId, bool b
 		}
 	}
 
-	//UpdateResourceAllocations();
-
 	if (CanBeRunOnGpu())
 	{
+#if WITH_EDITORONLY_DATA
+		const ERHIFeatureLevel::Type CacheFeatureLevel = NiagaraScriptInternal::PreviewFeatureLevel.Get(GMaxRHIFeatureLevel);
+#else
+		const ERHIFeatureLevel::Type CacheFeatureLevel = GMaxRHIFeatureLevel;
+#endif
+		const EShaderPlatform ShaderPlatform = GShaderPlatformForFeatureLevel[CacheFeatureLevel];
+
+		if (ScriptResource.IsValid())
+		{
+			if (ScriptResource->GetFeatureLevel() != CacheFeatureLevel
+				|| !ScriptResource->IsShaderMapComplete()
+				// double check that the shader platform also matches (note that ScriptResource->GetShaderPlatform() doesn't always seem valid
+				|| ScriptResource->GetGameThreadShaderMap()->GetShaderPlatform() != ShaderPlatform)
+			{
+				bForceRecompile = true;
+			}
+		}
+
+		if (bForceRecompile)
+		{
+			ScriptResource = nullptr;
+		}
+
 		// Need to make sure the owner supports GPU scripts, otherwise this is a wasted compile.
 		UNiagaraScriptSourceBase* Source = GetLatestScriptData()->Source;
-		if (Source && OwnerCanBeRunOnGpu())
+		if (Source && OwnerCanBeRunOnGpu() && !ScriptResource.IsValid())
 		{
-		#if WITH_EDITORONLY_DATA
-			const ERHIFeatureLevel::Type CacheFeatureLevel = NiagaraScriptInternal::PreviewFeatureLevel.Get(GMaxRHIFeatureLevel);
-		#else
-			const ERHIFeatureLevel::Type CacheFeatureLevel = GMaxRHIFeatureLevel;
-		#endif
-			const EShaderPlatform ShaderPlatform = GShaderPlatformForFeatureLevel[CacheFeatureLevel];
+			CreateScriptResource(nullptr);
 
 			ScriptResource->SetScript(this, CacheFeatureLevel, ShaderPlatform, CachedScriptVMId.CompilerVersionID, CachedScriptVMId.AdditionalDefines,
 				CachedScriptVMId.GetAdditionalVariableStrings(),
 				CachedScriptVMId.BaseScriptCompileHash, CachedScriptVMId.ReferencedCompileHashes,
 				CachedScriptVMId.bUsesRapidIterationParams, GetFriendlyName());
 
-			ScriptResource->BuildScriptParametersMetadata(CachedScriptVM.ShaderScriptParametersMetadata);
-
 			CacheShadersForResources(ScriptResource.Get(), true);
-			ScriptResourcesByFeatureLevel[CacheFeatureLevel] = ScriptResource;
 		}
-		else
+
+		if (ScriptResource.IsValid())
 		{
-			ScriptResource->Invalidate();
+			ScriptResource->BuildScriptParametersMetadata(CachedScriptVM.ShaderScriptParametersMetadata);
+			ScriptResourcesByFeatureLevel[CacheFeatureLevel] = ScriptResource;
 		}
 	}
 }
@@ -4167,7 +4182,7 @@ bool UNiagaraScript::SynchronizeExecutablesWithCompilation(const UNiagaraScript*
 	if (Id == Script->GetVMExecutableDataCompilationId())
 	{
 		CachedScriptVM.Reset();
-		ScriptResource->Invalidate();
+		ScriptResource = nullptr;
 
 		CachedScriptVM = Script->CachedScriptVM;
 		CachedScriptVMId = Script->CachedScriptVMId;
@@ -4208,7 +4223,7 @@ void UNiagaraScript::InvalidateCompileResults(const FString& Reason)
 {
 	UE_LOG(LogNiagara, Verbose, TEXT("InvalidateCompileResults Script:%s Reason:%s"), *GetPathName(), *Reason);
 	CachedScriptVM.Reset();
-	ScriptResource->Invalidate();
+	ScriptResource = nullptr;
 	CachedScriptVMId.Invalidate();
 	if (VersionData.Num() > 0)
 	{
@@ -4445,7 +4460,7 @@ void UNiagaraScript::SerializeNiagaraShaderMaps(FArchive& Ar, int32 NiagaraVer,
 					{
 						if (GMaxRHIShaderPlatform == ShaderMap->GetShaderPlatform())
 						{
-							ScriptResource = MakeShared<FNiagaraShaderScript>(Resource);
+							CreateScriptResource(&Resource);
 						}
 					}
 				}
@@ -4454,21 +4469,25 @@ void UNiagaraScript::SerializeNiagaraShaderMaps(FArchive& Ar, int32 NiagaraVer,
 	}
 }
 
+void UNiagaraScript::CreateScriptResource(const FNiagaraShaderScript* Source)
+{
+	ScriptResource = Source ? MakeShared<FNiagaraShaderScript>(*Source) : MakeShared<FNiagaraShaderScript>();
+#if WITH_EDITORONLY_DATA
+	ScriptResource->OnCompilationComplete().AddUniqueDynamic(this, &UNiagaraScript::RaiseOnGPUCompilationComplete);
+#endif
+}
+
 void UNiagaraScript::ProcessSerializedShaderMaps()
 {
 	check(IsInGameThread());
 
-	bool HasScriptResource = false;
-
 #if WITH_EDITORONLY_DATA
 	for (FNiagaraShaderScript& LoadedResource : LoadedScriptResources)
 	{
 		FNiagaraShaderMap* LoadedShaderMap = LoadedResource.GetGameThreadShaderMap();
 		if (LoadedShaderMap && LoadedShaderMap->GetShaderPlatform() == GMaxRHIShaderPlatform)
 		{
-			HasScriptResource = true;
-			ScriptResource = MakeShared<FNiagaraShaderScript>(LoadedResource);
-			ScriptResource->OnCompilationComplete().AddUniqueDynamic(this, &UNiagaraScript::RaiseOnGPUCompilationComplete);
+			CreateScriptResource(&LoadedResource);
 
 			ERHIFeatureLevel::Type LoadedFeatureLevel = LoadedShaderMap->GetShaderMapId().FeatureLevel;
 			if (!ScriptResourcesByFeatureLevel[LoadedFeatureLevel])
@@ -4480,11 +4499,9 @@ void UNiagaraScript::ProcessSerializedShaderMaps()
 			break;
 		}
 	}
-#else
-	HasScriptResource = ScriptResource.IsValid();
 #endif
 
-	if (HasScriptResource)
+	if (ScriptResource.IsValid())
 	{
 		ScriptResource->BuildScriptParametersMetadata(CachedScriptVM.ShaderScriptParametersMetadata);
 	}
@@ -4604,18 +4621,21 @@ ENiagaraScriptCompileStatus UNiagaraScript::GetLastCompileStatus() const
 	bool bHasWarnings = false;
 	if (Usage == ENiagaraScriptUsage::ParticleGPUComputeScript)
 	{
-		const TArray<FString>& Errors = ScriptResource->GetCompileErrors();
-		if (Errors.Num() > 0)
+		if (ScriptResource.IsValid())
 		{
-			for (const FString& ErrorStr : Errors)
+			const TArray<FString>& Errors = ScriptResource->GetCompileErrors();
+			if (Errors.Num() > 0)
 			{
-				if (ErrorStr.Contains(TEXT("err0r")))
+				for (const FString& ErrorStr : Errors)
 				{
-					return ENiagaraScriptCompileStatus::NCS_Error;
+					if (ErrorStr.Contains(TEXT("err0r")))
+					{
+						return ENiagaraScriptCompileStatus::NCS_Error;
+					}
 				}
+				bHasWarnings = true;
 			}
-			bHasWarnings = true;
-		}	
+		}
 	}
 	if (CachedScriptVM.IsValid())
 	{
