- Fix for Spawn Rate changes not propagating to stateless emitters

diff --git a/Engine/Plugins/FX/Niagara/Source/NiagaraEditor/Private/ViewModels/Stack/NiagaraStackStatelessEmitterSpawnGroup.cpp b/Engine/Plugins/FX/Niagara/Source/NiagaraEditor/Private/ViewModels/Stack/NiagaraStackStatelessEmitterSpawnGroup.cpp
--- a/Engine/Plugins/FX/Niagara/Source/NiagaraEditor/Private/ViewModels/Stack/NiagaraStackStatelessEmitterSpawnGroup.cpp
+++ b/Engine/Plugins/FX/Niagara/Source/NiagaraEditor/Private/ViewModels/Stack/NiagaraStackStatelessEmitterSpawnGroup.cpp
@@ -276,6 +276,9 @@ void UNiagaraStackStatelessEmitterSpawnItem::OnSpawnInfoModified(TArray<UObject*
 	{
 		FNiagaraStatelessSpawnInfo* SpawnInfo = GetSpawnInfo();
 		SpawnInfo->Rate.UpdateValuesFromDistribution();
+
+		FPropertyChangedEvent EmptyPropertyUpdateStruct(nullptr);
+		StatelessEmitter->PostEditChangeProperty(EmptyPropertyUpdateStruct);
 	}
 }
 
