[Backout] - CL31950735
Original CL Desc
-----------------------------------------------------------------
Couple of fixes for handling GPUComputeScript within cooked editor builds.
-allow LastHlslTranslationGPU to be preserved so that compilation can still happen
-prevent ScriptResource being discarded when serialized in cooked editor

diff --git a/Engine/Plugins/FX/Niagara/Source/Niagara/Classes/NiagaraScript.h b/Engine/Plugins/FX/Niagara/Source/Niagara/Classes/NiagaraScript.h
--- a/Engine/Plugins/FX/Niagara/Source/Niagara/Classes/NiagaraScript.h
+++ b/Engine/Plugins/FX/Niagara/Source/Niagara/Classes/NiagaraScript.h
@@ -1361,7 +1361,6 @@ private:
 
 	void ProcessSerializedShaderMaps();
 	void SerializeNiagaraShaderMaps(FArchive& Ar, int32 NiagaraVer, bool IsValidShaderScript);
-	void CreateScriptResource(const FNiagaraShaderScript* Source);
 
 	/** Return the expected SimTarget for this script. Only returns a valid target if there is valid data to run with. */
 	TOptional<ENiagaraSimTarget> GetSimTarget() const;
diff --git a/Engine/Plugins/FX/Niagara/Source/Niagara/Private/NiagaraScript.cpp b/Engine/Plugins/FX/Niagara/Source/Niagara/Private/NiagaraScript.cpp
--- a/Engine/Plugins/FX/Niagara/Source/Niagara/Private/NiagaraScript.cpp
+++ b/Engine/Plugins/FX/Niagara/Source/Niagara/Private/NiagaraScript.cpp
@@ -1015,6 +1015,9 @@ UNiagaraScript::UNiagaraScript(const FObjectInitializer& ObjectInitializer)
 #endif
 {
 #if WITH_EDITORONLY_DATA
+	ScriptResource = MakeShared<FNiagaraShaderScript>();
+	ScriptResource->OnCompilationComplete().AddUniqueDynamic(this, &UNiagaraScript::RaiseOnGPUCompilationComplete);
+
 	RapidIterationParameters.DebugName = *GetFullName();
 #endif
 }
@@ -1929,10 +1932,7 @@ void UNiagaraScript::Serialize(FArchive& Ar)
 		auto SwapDebugData = [](FNiagaraVMExecutableData& SourceData, FNiagaraVMExecutableData& TargetData) -> void
 		{
 			Swap(TargetData.LastHlslTranslation, SourceData.LastHlslTranslation);
-
-			// We must preserve the LastTranslationGPU as it is used to compile the GPUComputeScript when using the standard compilation mode
-			// Swap(TargetData.LastHlslTranslationGPU, SourceData.LastHlslTranslationGPU);
-
+			Swap(TargetData.LastHlslTranslationGPU, SourceData.LastHlslTranslationGPU);
 			Swap(TargetData.LastAssemblyTranslation, SourceData.LastAssemblyTranslation);
 			Swap(TargetData.CompileTagsEditorOnly, SourceData.CompileTagsEditorOnly);
 		};
@@ -3010,7 +3010,7 @@ void UNiagaraScript::SetVMCompilationResults(const FNiagaraVMExecutableDataId& I
 	CachedScriptVM.OptimizationTask.State = nullptr;
 	CachedParameterCollectionReferences.Empty();
 	// Proactively clear out the script resource, because it might be stale now.
-	ScriptResource = nullptr;
+	ScriptResource->Invalidate();
 
 	if (CachedScriptVM.LastCompileStatus == ENiagaraScriptCompileStatus::NCS_Error)
 	{
@@ -3211,7 +3211,7 @@ void UNiagaraScript::SetComputeCompilationResults(
 			{
 				if (!ScriptResource.IsValid())
 				{
-					CreateScriptResource(TargetShaderScript);
+					ScriptResource = MakeShared<FNiagaraShaderScript>(*TargetShaderScript);
 				}
 				else
 				{
@@ -3908,49 +3908,34 @@ void UNiagaraScript::CacheResourceShadersForRendering(bool bRegenerateId, bool b
 		}
 	}
 
+	//UpdateResourceAllocations();
+
 	if (CanBeRunOnGpu())
 	{
-#if WITH_EDITORONLY_DATA
-		const ERHIFeatureLevel::Type CacheFeatureLevel = NiagaraScriptInternal::PreviewFeatureLevel.Get(GMaxRHIFeatureLevel);
-#else
-		const ERHIFeatureLevel::Type CacheFeatureLevel = GMaxRHIFeatureLevel;
-#endif
-		const EShaderPlatform ShaderPlatform = GShaderPlatformForFeatureLevel[CacheFeatureLevel];
-
-		if (ScriptResource.IsValid())
-		{
-			if (ScriptResource->GetFeatureLevel() != CacheFeatureLevel
-				|| !ScriptResource->IsShaderMapComplete()
-				// double check that the shader platform also matches (note that ScriptResource->GetShaderPlatform() doesn't always seem valid
-				|| ScriptResource->GetGameThreadShaderMap()->GetShaderPlatform() != ShaderPlatform)
-			{
-				bForceRecompile = true;
-			}
-		}
-
-		if (bForceRecompile)
-		{
-			ScriptResource = nullptr;
-		}
-
 		// Need to make sure the owner supports GPU scripts, otherwise this is a wasted compile.
 		UNiagaraScriptSourceBase* Source = GetLatestScriptData()->Source;
-		if (Source && OwnerCanBeRunOnGpu() && !ScriptResource.IsValid())
+		if (Source && OwnerCanBeRunOnGpu())
 		{
-			CreateScriptResource(nullptr);
+		#if WITH_EDITORONLY_DATA
+			const ERHIFeatureLevel::Type CacheFeatureLevel = NiagaraScriptInternal::PreviewFeatureLevel.Get(GMaxRHIFeatureLevel);
+		#else
+			const ERHIFeatureLevel::Type CacheFeatureLevel = GMaxRHIFeatureLevel;
+		#endif
+			const EShaderPlatform ShaderPlatform = GShaderPlatformForFeatureLevel[CacheFeatureLevel];
 
 			ScriptResource->SetScript(this, CacheFeatureLevel, ShaderPlatform, CachedScriptVMId.CompilerVersionID, CachedScriptVMId.AdditionalDefines,
 				CachedScriptVMId.GetAdditionalVariableStrings(),
 				CachedScriptVMId.BaseScriptCompileHash, CachedScriptVMId.ReferencedCompileHashes,
 				CachedScriptVMId.bUsesRapidIterationParams, GetFriendlyName());
 
+			ScriptResource->BuildScriptParametersMetadata(CachedScriptVM.ShaderScriptParametersMetadata);
+
 			CacheShadersForResources(ScriptResource.Get(), true);
+			ScriptResourcesByFeatureLevel[CacheFeatureLevel] = ScriptResource;
 		}
-
-		if (ScriptResource.IsValid())
+		else
 		{
-			ScriptResource->BuildScriptParametersMetadata(CachedScriptVM.ShaderScriptParametersMetadata);
-			ScriptResourcesByFeatureLevel[CacheFeatureLevel] = ScriptResource;
+			ScriptResource->Invalidate();
 		}
 	}
 }
@@ -4182,7 +4167,7 @@ bool UNiagaraScript::SynchronizeExecutablesWithCompilation(const UNiagaraScript*
 	if (Id == Script->GetVMExecutableDataCompilationId())
 	{
 		CachedScriptVM.Reset();
-		ScriptResource = nullptr;
+		ScriptResource->Invalidate();
 
 		CachedScriptVM = Script->CachedScriptVM;
 		CachedScriptVMId = Script->CachedScriptVMId;
@@ -4223,7 +4208,7 @@ void UNiagaraScript::InvalidateCompileResults(const FString& Reason)
 {
 	UE_LOG(LogNiagara, Verbose, TEXT("InvalidateCompileResults Script:%s Reason:%s"), *GetPathName(), *Reason);
 	CachedScriptVM.Reset();
-	ScriptResource = nullptr;
+	ScriptResource->Invalidate();
 	CachedScriptVMId.Invalidate();
 	if (VersionData.Num() > 0)
 	{
@@ -4460,7 +4445,7 @@ void UNiagaraScript::SerializeNiagaraShaderMaps(FArchive& Ar, int32 NiagaraVer,
 					{
 						if (GMaxRHIShaderPlatform == ShaderMap->GetShaderPlatform())
 						{
-							CreateScriptResource(&Resource);
+							ScriptResource = MakeShared<FNiagaraShaderScript>(Resource);
 						}
 					}
 				}
@@ -4469,25 +4454,21 @@ void UNiagaraScript::SerializeNiagaraShaderMaps(FArchive& Ar, int32 NiagaraVer,
 	}
 }
 
-void UNiagaraScript::CreateScriptResource(const FNiagaraShaderScript* Source)
-{
-	ScriptResource = Source ? MakeShared<FNiagaraShaderScript>(*Source) : MakeShared<FNiagaraShaderScript>();
-#if WITH_EDITORONLY_DATA
-	ScriptResource->OnCompilationComplete().AddUniqueDynamic(this, &UNiagaraScript::RaiseOnGPUCompilationComplete);
-#endif
-}
-
 void UNiagaraScript::ProcessSerializedShaderMaps()
 {
 	check(IsInGameThread());
 
+	bool HasScriptResource = false;
+
 #if WITH_EDITORONLY_DATA
 	for (FNiagaraShaderScript& LoadedResource : LoadedScriptResources)
 	{
 		FNiagaraShaderMap* LoadedShaderMap = LoadedResource.GetGameThreadShaderMap();
 		if (LoadedShaderMap && LoadedShaderMap->GetShaderPlatform() == GMaxRHIShaderPlatform)
 		{
-			CreateScriptResource(&LoadedResource);
+			HasScriptResource = true;
+			ScriptResource = MakeShared<FNiagaraShaderScript>(LoadedResource);
+			ScriptResource->OnCompilationComplete().AddUniqueDynamic(this, &UNiagaraScript::RaiseOnGPUCompilationComplete);
 
 			ERHIFeatureLevel::Type LoadedFeatureLevel = LoadedShaderMap->GetShaderMapId().FeatureLevel;
 			if (!ScriptResourcesByFeatureLevel[LoadedFeatureLevel])
@@ -4499,9 +4480,11 @@ void UNiagaraScript::ProcessSerializedShaderMaps()
 			break;
 		}
 	}
+#else
+	HasScriptResource = ScriptResource.IsValid();
 #endif
 
-	if (ScriptResource.IsValid())
+	if (HasScriptResource)
 	{
 		ScriptResource->BuildScriptParametersMetadata(CachedScriptVM.ShaderScriptParametersMetadata);
 	}
@@ -4621,21 +4604,18 @@ ENiagaraScriptCompileStatus UNiagaraScript::GetLastCompileStatus() const
 	bool bHasWarnings = false;
 	if (Usage == ENiagaraScriptUsage::ParticleGPUComputeScript)
 	{
-		if (ScriptResource.IsValid())
+		const TArray<FString>& Errors = ScriptResource->GetCompileErrors();
+		if (Errors.Num() > 0)
 		{
-			const TArray<FString>& Errors = ScriptResource->GetCompileErrors();
-			if (Errors.Num() > 0)
+			for (const FString& ErrorStr : Errors)
 			{
-				for (const FString& ErrorStr : Errors)
+				if (ErrorStr.Contains(TEXT("err0r")))
 				{
-					if (ErrorStr.Contains(TEXT("err0r")))
-					{
-						return ENiagaraScriptCompileStatus::NCS_Error;
-					}
+					return ENiagaraScriptCompileStatus::NCS_Error;
 				}
-				bHasWarnings = true;
 			}
-		}
+			bHasWarnings = true;
+		}	
 	}
 	if (CachedScriptVM.IsValid())
 	{
