- Don't persistent some debugger settings between sessions

diff --git a/Engine/Plugins/FX/Niagara/Source/Niagara/Private/NiagaraDebuggerCommon.cpp b/Engine/Plugins/FX/Niagara/Source/Niagara/Private/NiagaraDebuggerCommon.cpp
--- a/Engine/Plugins/FX/Niagara/Source/Niagara/Private/NiagaraDebuggerCommon.cpp
+++ b/Engine/Plugins/FX/Niagara/Source/Niagara/Private/NiagaraDebuggerCommon.cpp
@@ -28,6 +28,17 @@ FNiagaraDebugHUDSettingsData::FNiagaraDebugHUDSettingsData()
 
 //////////////////////////////////////////////////////////////////////////
 
+void UNiagaraDebugHUDSettings::PostInitProperties()
+{
+	Super::PostInitProperties();
+
+	// Reset the properties we don't want to carry between editor sessions, we have to do this as all the properties are inside a
+	// data structure, and the whole data structure will go into the config rather than individual properties that specify if they should.
+	FNiagaraDebugHUDSettingsData Defaults;
+	Data.PlaybackMode			= Defaults.PlaybackMode;
+	Data.bPlaybackRateEnabled	= Defaults.bPlaybackRateEnabled;
+}
+
 void UNiagaraDebugHUDSettings::NotifyPropertyChanged()
 {
 	OnChangedDelegate.Broadcast();
diff --git a/Engine/Plugins/FX/Niagara/Source/Niagara/Public/NiagaraDebuggerCommon.h b/Engine/Plugins/FX/Niagara/Source/Niagara/Public/NiagaraDebuggerCommon.h
--- a/Engine/Plugins/FX/Niagara/Source/Niagara/Public/NiagaraDebuggerCommon.h
+++ b/Engine/Plugins/FX/Niagara/Source/Niagara/Public/NiagaraDebuggerCommon.h
@@ -528,83 +528,83 @@ struct FNiagaraDebugHUDSettingsData
 	Wildcard filter which is compared against the Components Actor name to narrow down the detailed information.
 	For example, "*Water*" would match all actors that contain the string "water".
 	*/
-	UPROPERTY(Config, EditAnywhere, Category = "Debug Filter", meta = (EditCondition = "bActorFilterEnabled"))
+	UPROPERTY(EditAnywhere, Category = "Debug Filter", meta = (EditCondition = "bActorFilterEnabled"))
 	FString ActorFilter;
 
-	UPROPERTY(Config, EditAnywhere, Category = "Debug Filter", meta = (InlineEditConditionToggle))
+	UPROPERTY(EditAnywhere, Category = "Debug Filter", meta = (InlineEditConditionToggle))
 	bool bComponentFilterEnabled = false;
 
 	/**
 	Wildcard filter for the components to show more detailed information about.
 	For example, "*MyComp*" would match all components that contain MyComp.
 	*/
-	UPROPERTY(Config, EditAnywhere, Category = "Debug Filter", meta = (EditCondition = "bComponentFilterEnabled"))
+	UPROPERTY(EditAnywhere, Category = "Debug Filter", meta = (EditCondition = "bComponentFilterEnabled"))
 	FString ComponentFilter;
 
-	UPROPERTY(Config, EditAnywhere, Category = "Debug Filter", meta = (InlineEditConditionToggle))
+	UPROPERTY(EditAnywhere, Category = "Debug Filter", meta = (InlineEditConditionToggle))
 	bool bSystemFilterEnabled = false;
 
 	/**
 	Wildcard filter for the systems to show more detailed information about.
 	For example,. "NS_*" would match all systems starting with NS_.
 	*/
-	UPROPERTY(Config, EditAnywhere, Category = "Debug Filter", meta = (EditCondition = "bSystemFilterEnabled"))
+	UPROPERTY(EditAnywhere, Category = "Debug Filter", meta = (EditCondition = "bSystemFilterEnabled"))
 	FString SystemFilter;
 
-	UPROPERTY(Config, EditAnywhere, Category = "Debug Filter", meta = (InlineEditConditionToggle))
+	UPROPERTY(EditAnywhere, Category = "Debug Filter", meta = (InlineEditConditionToggle))
 	bool bEmitterFilterEnabled = false;
 
 	/**
 	Wildcard filter used to match emitters when generating particle attribute view.
 	For example,. "Fluid*" would match all emtiters starting with Fluid and only particle attributes for those would be visible.
 	*/
-	UPROPERTY(Config, EditAnywhere, Category = "Debug Filter", meta = (EditCondition = "bEmitterFilterEnabled"))
+	UPROPERTY(EditAnywhere, Category = "Debug Filter", meta = (EditCondition = "bEmitterFilterEnabled"))
 	FString EmitterFilter;
 
-	UPROPERTY(Config, EditAnywhere, Category = "Debug Filter", meta = (InlineEditConditionToggle))
+	UPROPERTY(EditAnywhere, Category = "Debug Filter", meta = (InlineEditConditionToggle))
 	bool bActorFilterEnabled = false;
 
 	/** When enabled system debug information will be displayed in world. */
-	UPROPERTY(Config, EditAnywhere, Category = "Debug System")
+	UPROPERTY(EditAnywhere, Category = "Debug System")
 	ENiagaraDebugHudVerbosity SystemDebugVerbosity = ENiagaraDebugHudVerbosity::Basic;
 
 	/** When enabled we show information about emitter / particle counts. */
-	UPROPERTY(Config, EditAnywhere, Category = "Debug System", meta = (EditCondition = "SystemDebugVerbosity != ENiagaraDebugHudVerbosity::None"))
+	UPROPERTY(EditAnywhere, Category = "Debug System", meta = (EditCondition = "SystemDebugVerbosity != ENiagaraDebugHudVerbosity::None"))
 	ENiagaraDebugHudVerbosity SystemEmitterVerbosity = ENiagaraDebugHudVerbosity::Basic;
 
 	/** When enabled allows data interfaces to include additional debugging information. */
-	UPROPERTY(Config, EditAnywhere, Category = "Debug System", meta = (EditCondition = "SystemDebugVerbosity != ENiagaraDebugHudVerbosity::None"))
+	UPROPERTY(EditAnywhere, Category = "Debug System", meta = (EditCondition = "SystemDebugVerbosity != ENiagaraDebugHudVerbosity::None"))
 	ENiagaraDebugHudVerbosity DataInterfaceVerbosity = ENiagaraDebugHudVerbosity::None;
 
 	/** When enabled will show the system bounds for all filtered systems. */
-	UPROPERTY(Config, EditAnywhere, Category = "Debug System")
+	UPROPERTY(EditAnywhere, Category = "Debug System")
 	bool bSystemShowBounds = false;
 
 	/** When bounds display is enabled allows you to draw a solid box if alpha is > 0. */
-	UPROPERTY(Config, EditAnywhere, Category = "Debug System", meta = (EditCondition = "bSystemShowBounds"))
+	UPROPERTY(EditAnywhere, Category = "Debug System", meta = (EditCondition = "bSystemShowBounds"))
 	float SystemBoundsSolidBoxAlpha = 0.0f;
 
 	/** When disabled in world rendering will show systems deactivated by scalability. */
-	UPROPERTY(Config, EditAnywhere, Category = "Debug System", meta = (EditCondition = "SystemDebugVerbosity != ENiagaraDebugHudVerbosity::None"))
+	UPROPERTY(EditAnywhere, Category = "Debug System", meta = (EditCondition = "SystemDebugVerbosity != ENiagaraDebugHudVerbosity::None"))
 	bool bSystemShowActiveOnlyInWorld = true;
 
 	/** Should we display the system attributes. */
-	UPROPERTY(Config, EditAnywhere, Category = "Debug System", meta = (EditCondition = "SystemDebugVerbosity != ENiagaraDebugHudVerbosity::None", DisplayName="Show System Attributes"))
+	UPROPERTY(EditAnywhere, Category = "Debug System", meta = (EditCondition = "SystemDebugVerbosity != ENiagaraDebugHudVerbosity::None", DisplayName="Show System Attributes"))
 	bool bShowSystemVariables = true;
 
 	/**
 	List of attributes to show about the system, each entry uses wildcard matching.
 	For example, "System.*" would match all system attributes.
 	*/
-	UPROPERTY(Config, EditAnywhere, Category = "Debug System", meta = (EditCondition = "SystemDebugVerbosity != ENiagaraDebugHudVerbosity::None && bShowSystemVariables", DisplayName="System Attributes"))
+	UPROPERTY(EditAnywhere, Category = "Debug System", meta = (EditCondition = "SystemDebugVerbosity != ENiagaraDebugHudVerbosity::None && bShowSystemVariables", DisplayName="System Attributes"))
 	TArray<FNiagaraDebugHUDVariable> SystemVariables;
 
 	/** Sets display text options for system information. */
-	UPROPERTY(Config, EditAnywhere, Category = "Debug System", meta=(EditCondition="SystemDebugVerbosity != ENiagaraDebugHudVerbosity::None"))
+	UPROPERTY(EditAnywhere, Category = "Debug System", meta=(EditCondition="SystemDebugVerbosity != ENiagaraDebugHudVerbosity::None"))
 	FNiagaraDebugHudTextOptions SystemTextOptions;
 
 	/** When enabled will show particle attributes from the list. */
-	UPROPERTY(Config, EditAnywhere, Category = "Debug Particles", meta = (DisplayName="Show Particle Attributes"))
+	UPROPERTY(EditAnywhere, Category = "Debug Particles", meta = (DisplayName="Show Particle Attributes"))
 	bool bShowParticleVariables = true;
 
 	/**
@@ -612,102 +612,102 @@ struct FNiagaraDebugHUDSettingsData
 	Warning: This has an impact on performance & memory since we copy the whole buffer.
 	The displayed data is latent since we are seeing what happened a few frames ago.
 	*/
-	UPROPERTY(Config, EditAnywhere, Category = "Debug Particles", meta = (EditCondition = "bShowParticleVariables"))
+	UPROPERTY(EditAnywhere, Category = "Debug Particles", meta = (EditCondition = "bShowParticleVariables"))
 	bool bEnableGpuParticleReadback = false;
 
 	/**
 	When enabled the particle index will be displayed along with any attributes.
 	Note: This is the index in the particle data buffer and not the persistent ID index.
 	*/
-	UPROPERTY(Config, EditAnywhere, Category = "Debug Particles", meta = (EditCondition = "bShowParticleVariables"))
+	UPROPERTY(EditAnywhere, Category = "Debug Particles", meta = (EditCondition = "bShowParticleVariables"))
 	bool bShowParticleIndex = false;
 
 	/**
 	List of attributes to show per particle, each entry uses wildcard matching.
 	For example, "*Position" would match all attributes that end in Position.
 	*/
-	UPROPERTY(Config, EditAnywhere, Category = "Debug Particles", meta = (EditCondition = "bShowParticleVariables", DisplayName="Particle Attributes"))
+	UPROPERTY(EditAnywhere, Category = "Debug Particles", meta = (EditCondition = "bShowParticleVariables", DisplayName="Particle Attributes"))
 	TArray<FNiagaraDebugHUDVariable> ParticlesVariables;
 
 	/** Sets display text options for particle information. */
-	UPROPERTY(Config, EditAnywhere, Category = "Debug Particles", meta = (EditCondition = "bShowParticleVariables"))
+	UPROPERTY(EditAnywhere, Category = "Debug Particles", meta = (EditCondition = "bShowParticleVariables"))
 	FNiagaraDebugHudTextOptions ParticleTextOptions;
 
 	/**
 	When enabled particle attributes will display with the system information
 	rather than in world at the particle location.
 	*/
-	UPROPERTY(Config, EditAnywhere, Category = "Debug Particles", meta = (EditCondition = "bShowParticleVariables", DisplayName="Show Particles Attributes With System"))
+	UPROPERTY(EditAnywhere, Category = "Debug Particles", meta = (EditCondition = "bShowParticleVariables", DisplayName="Show Particles Attributes With System"))
 	bool bShowParticlesVariablesWithSystem = false;
 
-	UPROPERTY(Config, EditAnywhere, Category = "Debug Particles", meta = (EditCondition = "bShowParticleVariables", DisplayName = "Show Particle Attributes Vertical"))
+	UPROPERTY(EditAnywhere, Category = "Debug Particles", meta = (EditCondition = "bShowParticleVariables", DisplayName = "Show Particle Attributes Vertical"))
 	bool bShowParticleVariablesVertical = false;
 
-	UPROPERTY(Config, EditAnywhere, Category = "Debug Particles", meta = (EditCondition = "bShowParticleVariables"))
+	UPROPERTY(EditAnywhere, Category = "Debug Particles", meta = (EditCondition = "bShowParticleVariables"))
 	bool bUseMaxParticlesToDisplay = true;
 
 	/** When enabled we use the clip planes to narrow down which particles to display */
-	UPROPERTY(Config, EditAnywhere, Category = "Debug Particles", meta = (EditCondition = "bShowParticleVariables"))
+	UPROPERTY(EditAnywhere, Category = "Debug Particles", meta = (EditCondition = "bShowParticleVariables"))
 	bool bUseParticleDisplayClip = false;
 
 	/** Clipping planes used to display particle attributes. */
-	UPROPERTY(Config, EditAnywhere, Category = "Debug Particles", meta = (EditCondition = "bShowParticleVariables", UIMin = "0", ClampMin = "0"))
+	UPROPERTY(EditAnywhere, Category = "Debug Particles", meta = (EditCondition = "bShowParticleVariables", UIMin = "0", ClampMin = "0"))
 	FVector2D ParticleDisplayClip = FVector2D(0.0f, 10000.0f);
 
 	/** When enabled we use a radius from the display center to avoid showing too many particle attributes. */
-	UPROPERTY(Config, EditAnywhere, Category = "Debug Particles", meta = (EditCondition = "bShowParticleVariables"))
+	UPROPERTY(EditAnywhere, Category = "Debug Particles", meta = (EditCondition = "bShowParticleVariables"))
 	bool bUseParticleDisplayCenterRadius = false;
 
 	/** Radius from screen center where 0 is center to 1.0 is edge to avoid display too many particle attributes. */
-	UPROPERTY(Config, EditAnywhere, Category = "Debug Particles", meta = (EditCondition = "bShowParticleVariables", UIMin = "0", ClampMin = "0"))
+	UPROPERTY(EditAnywhere, Category = "Debug Particles", meta = (EditCondition = "bShowParticleVariables", UIMin = "0", ClampMin = "0"))
 	float ParticleDisplayCenterRadius = 1.0f;
 
 	/**
 	When enabled, the maximum number of particles to show information about.
 	When disabled all particles will show attributes, this can result in poor performance & potential OOM on some platforms.
 	*/
-	UPROPERTY(Config, EditAnywhere, Category = "Debug Particles", meta = (EditCondition = "bUseMaxParticlesToDisplay && bShowParticleVariables", UIMin="1", ClampMin="1"))
+	UPROPERTY(EditAnywhere, Category = "Debug Particles", meta = (EditCondition = "bUseMaxParticlesToDisplay && bShowParticleVariables", UIMin="1", ClampMin="1"))
 	int32 MaxParticlesToDisplay = 32;
 
 	/** How many frames to capture in between updates to the max and average perf report values. */
-	UPROPERTY(Config, EditAnywhere, Category = "Perf Overview")
+	UPROPERTY(EditAnywhere, Category = "Perf Overview")
 	int32 PerfReportFrames = 60;
 
-	UPROPERTY(Config, EditAnywhere, Category = "Perf Overview")
+	UPROPERTY(EditAnywhere, Category = "Perf Overview")
 	ENiagaraDebugHUDPerfSampleMode PerfSampleMode = ENiagaraDebugHUDPerfSampleMode::FrameTotal;
 
-	UPROPERTY(Config, EditAnywhere, Category = "Perf Overview")
+	UPROPERTY(EditAnywhere, Category = "Perf Overview")
 	ENiagaraDebugHUDPerfUnits PerfUnits = ENiagaraDebugHUDPerfUnits::Microseconds;
 
 	/** Time range of the Y Axis of the perf graph */
-	UPROPERTY(Config, EditAnywhere, Category = "Perf Overview")
+	UPROPERTY(EditAnywhere, Category = "Perf Overview")
 	ENiagaraDebugHUDPerfGraphMode PerfGraphMode = ENiagaraDebugHUDPerfGraphMode::GameThread;
 
 	/** How many frames of history to display in the perf graphs. */
-	UPROPERTY(Config, EditAnywhere, Category = "Perf Overview")
+	UPROPERTY(EditAnywhere, Category = "Perf Overview")
 	int32 PerfHistoryFrames = 600;
 
 	/** Use the specified user range when enabled, otherwise we will auto detect the range to use. */
-	UPROPERTY(Config, EditAnywhere, Category = "Perf Overview", meta = (InlineEditConditionToggle))
+	UPROPERTY(EditAnywhere, Category = "Perf Overview", meta = (InlineEditConditionToggle))
 	bool bUsePerfGraphTimeRange = false;
 
 	/** Time range of the Y Axis of the perf graph */
-	UPROPERTY(Config, EditAnywhere, Category = "Perf Overview", meta = (EditCondition = "bUsePerfGraphTimeRange"))
+	UPROPERTY(EditAnywhere, Category = "Perf Overview", meta = (EditCondition = "bUsePerfGraphTimeRange"))
 	float PerfGraphTimeRange = 500.0f;
 
 	/** Pixel size of the perf graph. */
-	UPROPERTY(Config, EditAnywhere, Category = "Perf Overview")
+	UPROPERTY(EditAnywhere, Category = "Perf Overview")
 	FVector2D PerfGraphSize = FVector2D(500,500);
 
-	UPROPERTY(Config, EditAnywhere, Category = "Perf Overview")
+	UPROPERTY(EditAnywhere, Category = "Perf Overview")
 	FLinearColor PerfGraphAxisColor = FLinearColor::White;
 
 	// True if perf graphs should be smoothed.
-	UPROPERTY(Config, EditAnywhere, Category = "Perf Overview", meta = (InlineEditConditionToggle))
+	UPROPERTY(EditAnywhere, Category = "Perf Overview", meta = (InlineEditConditionToggle))
 	bool bEnableSmoothing = true;
 
 	//Number of samples to use either size of a value when smoothing perf graphs.
-	UPROPERTY(Config, EditAnywhere, Category = "Perf Overview", meta = (EditCondition = "bEnableSmoothing"))
+	UPROPERTY(EditAnywhere, Category = "Perf Overview", meta = (EditCondition = "bEnableSmoothing"))
 	int32 SmoothingWidth = 4;
 
 	// Default background color used generally for panels
@@ -746,19 +746,19 @@ struct FNiagaraDebugHUDSettingsData
 	FLinearColor MessageErrorTextColor = FLinearColor(1.0f, 0.4f, 0.3f, 1.0f);
 	
 	/** Opacity of the system color background tile in overview table rows. */
-	UPROPERTY(Config, EditAnywhere, Category = "Colors")
+	UPROPERTY(EditAnywhere, Category = "Colors")
 	float SystemColorTableOpacity = 0.2f;
 
 	/** Additional seed value for random system colors. Useful if current colors of systems are too similar. */
-	UPROPERTY(Config, EditAnywhere, Category = "Colors")
+	UPROPERTY(EditAnywhere, Category = "Colors")
 	uint32 SystemColorSeed = 0;
 
 	/** Minimum HSV values for the random colors generated for each System. */
-	UPROPERTY(Config, EditAnywhere, Category = "Colors")
+	UPROPERTY(EditAnywhere, Category = "Colors")
 	FVector SystemColorHSVMin = FVector(0, 200, 200);
 
 	/** Maximum HSV values for the random colors generated for each System. */
-	UPROPERTY(Config, EditAnywhere, Category = "Colors")
+	UPROPERTY(EditAnywhere, Category = "Colors")
 	FVector SystemColorHSVMax = FVector(255, 255, 255);
 
 	UPROPERTY()
@@ -767,16 +767,16 @@ struct FNiagaraDebugHUDSettingsData
 	UPROPERTY()
 	bool bPlaybackRateEnabled = false;
 
-	UPROPERTY(Config)
+	UPROPERTY()
 	float PlaybackRate = 0.25f;
 
-	UPROPERTY(Config)
+	UPROPERTY()
 	bool bLoopTimeEnabled = false;
 
-	UPROPERTY(Config)
+	UPROPERTY()
 	float LoopTime = 1.0f;
 
-	UPROPERTY(Config, EditAnywhere, Category = "Performance")
+	UPROPERTY(EditAnywhere, Category = "Performance")
 	bool bShowGlobalBudgetInfo = false;
 };
 
@@ -793,6 +793,8 @@ class UNiagaraDebugHUDSettings : public UObject, public FNotifyHook
 	GENERATED_BODY()
 
 public:
+	virtual void PostInitProperties() override;
+
 	DECLARE_MULTICAST_DELEGATE(FOnChanged);
 	FOnChanged OnChangedDelegate;
 
