Fix for Niagara's generate code view when shaders become unavailable (and the SourceText is empty)

diff --git a/Engine/Plugins/FX/Niagara/Source/NiagaraEditor/Private/Widgets/SNiagaraGeneratedCodeView.cpp b/Engine/Plugins/FX/Niagara/Source/NiagaraEditor/Private/Widgets/SNiagaraGeneratedCodeView.cpp
--- a/Engine/Plugins/FX/Niagara/Source/NiagaraEditor/Private/Widgets/SNiagaraGeneratedCodeView.cpp
+++ b/Engine/Plugins/FX/Niagara/Source/NiagaraEditor/Private/Widgets/SNiagaraGeneratedCodeView.cpp
@@ -481,7 +481,7 @@ void SNiagaraGeneratedCodeView::UpdateUI()
 					GeneratedCode[i].UsageName = FText::Format(LOCTEXT("UsageName", "{0}{1}"), ScriptEnum->GetDisplayNameTextByValue((int64)Scripts[i]->Usage), AssemblyIdText);
 				}
 
-				if (SourceText)
+				if (SourceText && SourceText->Len() > 0)
 				{
 					if (bAddLineNumbers)
 					{
