- Allow ribbons to share generated data where possible
- CPU will share the generated data across views (i.e. split screen) when opaque or not using multi ribbons
- GPU will share across view families in the same frame (i.e. split screen / shadow passes), this is because we don't support multi ribbon sorting on the GPU
- Niagara.Ribbon.ShareGeneratedData can be used to enable / disable the optimization

diff --git a/Engine/Plugins/FX/Niagara/Source/Niagara/Private/NiagaraRendererRibbons.cpp b/Engine/Plugins/FX/Niagara/Source/Niagara/Private/NiagaraRendererRibbons.cpp
--- a/Engine/Plugins/FX/Niagara/Source/Niagara/Private/NiagaraRendererRibbons.cpp
+++ b/Engine/Plugins/FX/Niagara/Source/Niagara/Private/NiagaraRendererRibbons.cpp
@@ -158,6 +158,14 @@ static FAutoConsoleVariableRef CVarNiagaraRibbonGpuBufferAlign(
 	ECVF_Default
 );
 
+static bool GNiagaraRibbonShareGeneratedData = true;
+static FAutoConsoleVariableRef CVarNiagaraRibbonShareGeneratedData(
+	TEXT("Niagara.Ribbon.ShareGeneratedData"),
+	GNiagaraRibbonShareGeneratedData,
+	TEXT("Allow ribbons to share the generate data where possible."),
+	ECVF_Default
+);
+
 static TAutoConsoleVariable<int32> CVarRayTracingNiagaraRibbons(
 	TEXT("r.RayTracing.Geometry.NiagaraRibbons"),
 	0,
@@ -323,14 +331,14 @@ struct FNiagaraRibbonRenderingFrameViewResources
 	FNiagaraRibbonVertexFactory				VertexFactory;
 	FNiagaraRibbonUniformBufferRef			UniformBuffer;
 	TSharedPtr<FNiagaraRibbonIndexBuffer>	IndexBuffer;
-	FRWBuffer								IndirectDrawBuffer;
+	TSharedPtr<FRWBuffer>					IndirectDrawBuffer;
 	FNiagaraIndexGenerationInput			IndexGenerationSettings;
+	bool									bNeedsIndexBufferGeneration = true;
 
 	~FNiagaraRibbonRenderingFrameViewResources()
 	{
 		UniformBuffer.SafeRelease();
 		VertexFactory.ReleaseResource();
-		IndirectDrawBuffer.Release();
 	}
 };
 
@@ -529,8 +537,8 @@ class FNiagaraGpuRibbonsDataManager final : public FNiagaraGpuComputeDataManager
 {
 	struct FIndirectDrawBufferEntry
 	{
-		uint64		FrameUsed = 0;
-		FRWBuffer	Buffer;
+		uint64					FrameUsed = 0;
+		TSharedPtr<FRWBuffer>	Buffer;
 	};
 
 	struct FIndexBufferEntry
@@ -560,13 +568,30 @@ public:
 	}
 
 	void RegisterRenderer(const FNiagaraRendererRibbons* InRenderer, const FNiagaraDataBuffer* InSourceParticleData, const TSharedPtr<FNiagaraRibbonRenderingFrameResources>& InRenderingResources)
+	{
+		if (InRenderingResources->ViewResources.Num() > 0)
+		{
+			const int32 GenerateIndex = InSourceParticleData->GetGPUDataReadyStage() == ENiagaraGpuComputeTickStage::PostOpaqueRender ? 1 : 0;
+			RenderersToGeneratePerStage[GenerateIndex].Emplace(InRenderer, InSourceParticleData, InRenderingResources);
+		}
+	}
+
+	FNiagaraRibbonRenderingFrameViewResources* FindExistingRendererViewData(const FNiagaraRendererRibbons* InRenderer, const FNiagaraDataBuffer* InSourceParticleData)
 	{
 		const int32 GenerateIndex = InSourceParticleData->GetGPUDataReadyStage() == ENiagaraGpuComputeTickStage::PostOpaqueRender ? 1 : 0;
-		RenderersToGeneratePerStage[GenerateIndex].Emplace(InRenderer, InSourceParticleData, InRenderingResources);
+		for (FNiagaraRibbonGPUInitParameters& ExistingData : RenderersToGeneratePerStage[GenerateIndex])
+		{
+			if (ExistingData.Renderer == InRenderer)
+			{
+				TSharedPtr<FNiagaraRibbonRenderingFrameResources> RenderingResources = ExistingData.RenderingResources.Pin();
+				return RenderingResources && RenderingResources->ViewResources.Num() > 0 ? RenderingResources->ViewResources[0].Get() : nullptr;
+			}
+		}
+		return nullptr;
 	}
 
 	//-OPT: These caches should be more central and are as a simple solution to reduce memory thrashing / poor performance for ribbons
-	FRWBuffer GetOrAllocateIndirectDrawBuffer(FRHICommandListBase& RHICmdList)
+	const TSharedPtr<FRWBuffer>& GetOrAllocateIndirectDrawBuffer(FRHICommandListBase& RHICmdList)
 	{
 		FIndirectDrawBufferEntry* BufferEntry = IndirectDrawBufferCache.FindByPredicate(
 			[&](const FIndirectDrawBufferEntry& ExistingBuffer)
@@ -577,13 +602,14 @@ public:
 		if (BufferEntry == nullptr)
 		{
 			BufferEntry = &IndirectDrawBufferCache.AddDefaulted_GetRef();
-			BufferEntry->Buffer.Initialize(RHICmdList, TEXT("RibbonIndirectDrawBuffer"), sizeof(uint32), FNiagaraRibbonIndirectDrawBufferLayout::NumElements, EPixelFormat::PF_R32_UINT, ERHIAccess::IndirectArgs | ERHIAccess::SRVMask, BUF_Static | BUF_DrawIndirect);
+			BufferEntry->Buffer = MakeShared<FRWBuffer>();
+			BufferEntry->Buffer->Initialize(RHICmdList, TEXT("RibbonIndirectDrawBuffer"), sizeof(uint32), FNiagaraRibbonIndirectDrawBufferLayout::NumElements, EPixelFormat::PF_R32_UINT, ERHIAccess::IndirectArgs | ERHIAccess::SRVMask, BUF_Static | BUF_DrawIndirect);
 		}
 		BufferEntry->FrameUsed = FrameCounter;
 		return BufferEntry->Buffer;
 	}
 
-	TSharedPtr<FNiagaraRibbonIndexBuffer>& GetOrAllocateIndexBuffer(FRHICommandListBase& RHICmdList, int32 NumIndices, int32 MaxIndicesEstimate)
+	const TSharedPtr<FNiagaraRibbonIndexBuffer>& GetOrAllocateIndexBuffer(FRHICommandListBase& RHICmdList, int32 NumIndices, int32 MaxIndicesEstimate)
 	{
 		if (GNiagaraRibbonGpuBufferCachePurgeCounter >= 0)
 		{
@@ -840,6 +866,22 @@ void FNiagaraRendererRibbons::GetDynamicMeshElements(const TArray<const FSceneVi
 	FNiagaraGpuComputeDispatchInterface* ComputeDispatchInterface = SceneProxy->GetComputeDispatchInterface();
 	FNiagaraGpuRibbonsDataManager& GpuRibbonDataManager = ComputeDispatchInterface->GetOrCreateDataManager<FNiagaraGpuRibbonsDataManager>();
 
+	// Can we share the generated data across different view families for this frame?
+	// Note: We only handle this for GPU currently
+	const bool bShareAcrossViewFamilies = GNiagaraRibbonShareGeneratedData && (DynamicData->bUseGPUInit || DynamicData->bIsGPUSystem);
+
+	// Do we need to generate a per view buffer (i.e. split screen data)
+	// Note: GPU can always share as we don't generate anything per view dependent and CPU can only share if not using multi ribbons or opaque
+	const bool bNeedsBufferPerView =
+		!GNiagaraRibbonShareGeneratedData ||
+		(
+			!DynamicData->bUseGPUInit && !DynamicData->bIsGPUSystem &&
+			bTranslucentMaterial &&
+			(DynamicData->GenerationOutput && DynamicData->GenerationOutput->RibbonInfoLookup.Num() > 0)
+		);
+
+	bool bNeedsVertexIndexGeneration = false;
+
 	// Compute the per-view uniform buffers.
 	for (int32 ViewIndex = 0; ViewIndex < Views.Num(); ViewIndex++)
 	{
@@ -862,16 +904,37 @@ void FNiagaraRendererRibbons::GetDynamicMeshElements(const TArray<const FSceneVi
 			}
 
 			FMeshBatch& MeshBatch = Collector.AllocateMesh();
-			
+	
 			const FVector ViewOriginForDistanceCulling = View->ViewMatrices.GetViewOrigin();
 
-			auto& RenderingViewResources = RenderingResources.RibbonResources->ViewResources.Add_GetRef(MakeShared<FNiagaraRibbonRenderingFrameViewResources>());
-			RenderingViewResources->IndexGenerationSettings = CalculateIndexBufferConfiguration(DynamicData->GenerationOutput, SourceParticleData, SceneProxy, View, ViewOriginForDistanceCulling, DynamicData->bUseGPUInit, DynamicData->bIsGPUSystem);
-			
-			GenerateIndexBufferForView(RHICmdList, GpuRibbonDataManager, Collector, RenderingViewResources->IndexGenerationSettings, DynamicData, RenderingViewResources, View, ViewOriginForDistanceCulling);
-			
+			FNiagaraRibbonRenderingFrameViewResources* RenderingViewResources = RenderingResources.RibbonResources->ViewResources.Add_GetRef(MakeShared<FNiagaraRibbonRenderingFrameViewResources>()).Get();
+			FNiagaraRibbonRenderingFrameViewResources* SharedRenderingViewResources = nullptr;
+
+			if (bShareAcrossViewFamilies)
+			{
+				SharedRenderingViewResources = GpuRibbonDataManager.FindExistingRendererViewData(this, SourceParticleData);
+			}
+			else if (!bNeedsBufferPerView && RenderingResources.RibbonResources->ViewResources.Num() > 1)
+			{
+				SharedRenderingViewResources = RenderingResources.RibbonResources->ViewResources[0].Get();
+			}
+
+			if (SharedRenderingViewResources)
+			{
+				RenderingViewResources->IndexBuffer = SharedRenderingViewResources->IndexBuffer;
+				RenderingViewResources->IndirectDrawBuffer = SharedRenderingViewResources->IndirectDrawBuffer;
+				RenderingViewResources->IndexGenerationSettings = SharedRenderingViewResources->IndexGenerationSettings;
+				RenderingViewResources->bNeedsIndexBufferGeneration = false;
+			}
+			else
+			{
+				bNeedsVertexIndexGeneration = true;
+				RenderingViewResources->IndexGenerationSettings = CalculateIndexBufferConfiguration(DynamicData->GenerationOutput, SourceParticleData, SceneProxy, View, ViewOriginForDistanceCulling, DynamicData->bUseGPUInit, DynamicData->bIsGPUSystem);
+				GenerateIndexBufferForView(RHICmdList, GpuRibbonDataManager, Collector, RenderingViewResources->IndexGenerationSettings, DynamicData, RenderingViewResources, View, ViewOriginForDistanceCulling);
+			}
+
 			SetupPerViewUniformBuffer(RenderingViewResources->IndexGenerationSettings, View, ViewFamily, SceneProxy, RenderingViewResources->UniformBuffer);
-			
+
 			SetupMeshBatchAndCollectorResourceForView(RHICmdList, RenderingViewResources->IndexGenerationSettings, DynamicData, SourceParticleData, View, ViewFamily, SceneProxy, RenderingResources.RibbonResources, RenderingViewResources, MeshBatch, DynamicData->bUseGPUInit);
 
 			Collector.AddMesh(ViewIndex, MeshBatch);
@@ -879,7 +942,7 @@ void FNiagaraRendererRibbons::GetDynamicMeshElements(const TArray<const FSceneVi
 	}
 
 	// Register this renderer for generation this frame if we're a gpu system or using gpu init
-	if (DynamicData->bUseGPUInit || DynamicData->bIsGPUSystem)
+	if (bNeedsVertexIndexGeneration && (DynamicData->bUseGPUInit || DynamicData->bIsGPUSystem))
 	{
 		GpuRibbonDataManager.RegisterRenderer(this, SourceParticleData, RenderingResources.RibbonResources);
 	}
@@ -1055,7 +1118,7 @@ void FNiagaraRendererRibbons::GetDynamicRayTracingInstances(FRayTracingMaterialG
 	const FVector ViewOriginForDistanceCulling = View->ViewMatrices.GetViewOrigin();
 	
 	FNiagaraRibbonMeshCollectorResources& RenderingResources = Context.RayTracingMeshResourceCollector.AllocateOneFrameResource<FNiagaraRibbonMeshCollectorResources>();
-	auto& RenderingViewResources = RenderingResources.RibbonResources->ViewResources.Add_GetRef(MakeShared<FNiagaraRibbonRenderingFrameViewResources>());
+	FNiagaraRibbonRenderingFrameViewResources* RenderingViewResources = RenderingResources.RibbonResources->ViewResources.Add_GetRef(MakeShared<FNiagaraRibbonRenderingFrameViewResources>()).Get();
 	RenderingViewResources->IndexGenerationSettings = CalculateIndexBufferConfiguration(DynamicDataRibbon->GenerationOutput, SourceParticleData, SceneProxy, View, ViewOriginForDistanceCulling, DynamicDataRibbon->bUseGPUInit, DynamicDataRibbon->bIsGPUSystem);
 	
 	if (!RenderingViewResources->VertexFactory.GetType()->SupportsRayTracingDynamicGeometry())
@@ -1875,7 +1938,7 @@ void FNiagaraRendererRibbons::GenerateIndexBufferForView(
 	FNiagaraGpuRibbonsDataManager& GpuRibbonsDataManager,
 	FMeshElementCollector& Collector,
 	FNiagaraIndexGenerationInput& GeneratedData, FNiagaraDynamicDataRibbon* DynamicDataRibbon,
-    const TSharedPtr<FNiagaraRibbonRenderingFrameViewResources>& RenderingViewResources, const FSceneView* View,
+    FNiagaraRibbonRenderingFrameViewResources* RenderingViewResources, const FSceneView* View,
     const FVector& ViewOriginForDistanceCulling
 ) const
 {
@@ -2064,7 +2127,7 @@ void FNiagaraRendererRibbons::SetupPerViewUniformBuffer(FNiagaraIndexGenerationI
 }
 
 inline void FNiagaraRendererRibbons::SetupMeshBatchAndCollectorResourceForView(FRHICommandListBase& RHICmdList, const FNiagaraIndexGenerationInput& GeneratedData, FNiagaraDynamicDataRibbon* DynamicDataRibbon, const FNiagaraDataBuffer* SourceParticleData, const FSceneView* View,
-    const FSceneViewFamily& ViewFamily, const FNiagaraSceneProxy* SceneProxy, const TSharedPtr<FNiagaraRibbonRenderingFrameResources>& RenderingResources, const TSharedPtr<FNiagaraRibbonRenderingFrameViewResources>& RenderingViewResources,
+    const FSceneViewFamily& ViewFamily, const FNiagaraSceneProxy* SceneProxy, const TSharedPtr<FNiagaraRibbonRenderingFrameResources>& RenderingResources, FNiagaraRibbonRenderingFrameViewResources* RenderingViewResources,
     FMeshBatch& OutMeshBatch, bool bShouldUseGPUInitIndices) const
 {
 	const bool bIsWireframe = ViewFamily.EngineShowFlags.Wireframe;
@@ -2086,7 +2149,7 @@ inline void FNiagaraRendererRibbons::SetupMeshBatchAndCollectorResourceForView(F
 	VFLooseParams.NeedsPreciseMotionVectors = GenerationConfig.NeedsPreciseMotionVectors();
 	VFLooseParams.UseGeometryNormals = (ShapeState.Shape != ENiagaraRibbonShapeMode::Plane || bUseGeometryNormals) ? 1 : 0;
 
-	VFLooseParams.IndirectDrawOutput = bShouldUseGPUInitIndices ? (FRHIShaderResourceView*)RenderingViewResources->IndirectDrawBuffer.SRV : GetDummyUIntBuffer();
+	VFLooseParams.IndirectDrawOutput = bShouldUseGPUInitIndices ? (FRHIShaderResourceView*)RenderingViewResources->IndirectDrawBuffer->SRV : GetDummyUIntBuffer();
 	VFLooseParams.IndirectDrawOutputOffset = bShouldUseGPUInitIndices ? 0 : -1;
 
 	// Collector.AllocateOneFrameResource uses default ctor, initialize the vertex factory
@@ -2119,7 +2182,7 @@ inline void FNiagaraRendererRibbons::SetupMeshBatchAndCollectorResourceForView(F
 	if (bShouldUseGPUInitIndices)
 	{
 		MeshElement.NumPrimitives = 0;
-		MeshElement.IndirectArgsBuffer = RenderingViewResources->IndirectDrawBuffer.Buffer;
+		MeshElement.IndirectArgsBuffer = RenderingViewResources->IndirectDrawBuffer->Buffer;
 		MeshElement.IndirectArgsOffset = View->IsInstancedStereoPass() ? FNiagaraRibbonIndirectDrawBufferLayout::StereoDrawIndirectParametersByteOffset : FNiagaraRibbonIndirectDrawBufferLayout::DrawIndirectParametersByteOffset;
 	}
 	else
@@ -2136,11 +2199,11 @@ inline void FNiagaraRendererRibbons::SetupMeshBatchAndCollectorResourceForView(F
 	MeshElement.PrimitiveUniformBuffer = SceneProxy->GetCustomUniformBuffer(RHICmdList, false);	// Note: Ribbons don't generate accurate velocities so disabling	
 }
 
-void FNiagaraRendererRibbons::InitializeViewIndexBuffersGPU(FRHICommandListImmediate& RHICmdList, FNiagaraGpuComputeDispatchInterface* ComputeDispatchInterface, const FNiagaraRibbonGPUInitParameters& GpuInitParameters, const TSharedPtr<FNiagaraRibbonRenderingFrameViewResources>& RenderingViewResources) const
+void FNiagaraRendererRibbons::InitializeViewIndexBuffersGPU(FRHICommandListImmediate& RHICmdList, FNiagaraGpuComputeDispatchInterface* ComputeDispatchInterface, const FNiagaraRibbonGPUInitParameters& GpuInitParameters, const FNiagaraRibbonRenderingFrameViewResources* RenderingViewResources) const
 {
 	SCOPE_CYCLE_COUNTER(STAT_NiagaraRenderRibbonsGenIndiciesGPU);
 	
-	if (!RenderingViewResources->IndirectDrawBuffer.Buffer.IsValid())
+	if (!RenderingViewResources->IndirectDrawBuffer->Buffer.IsValid())
 	{
 		return;
 	}
@@ -2158,7 +2221,7 @@ void FNiagaraRendererRibbons::InitializeViewIndexBuffersGPU(FRHICommandListImmed
 		FNiagaraRibbonInitializeIndices Params;
 		FMemory::Memzero(Params);
 
-		Params.IndirectDrawOutput = RenderingViewResources->IndirectDrawBuffer.UAV;
+		Params.IndirectDrawOutput = RenderingViewResources->IndirectDrawBuffer->UAV;
 		Params.VertexGenerationResults = VertexBuffers.GPUComputeCommandBuffer.SRV;
 
 		// Total particle Count
@@ -2186,9 +2249,9 @@ void FNiagaraRendererRibbons::InitializeViewIndexBuffersGPU(FRHICommandListImmed
 		Params.GNiagaraRibbonTessellationEnabled = GNiagaraRibbonTessellationEnabled ? 1 : 0;
 		Params.GNiagaraRibbonTessellationMinDisplacementError = GNiagaraRibbonTessellationMinDisplacementError;
 
-		RHICmdList.Transition(FRHITransitionInfo(RenderingViewResources->IndirectDrawBuffer.UAV, ERHIAccess::SRVMask | ERHIAccess::IndirectArgs, ERHIAccess::UAVCompute));
+		RHICmdList.Transition(FRHITransitionInfo(RenderingViewResources->IndirectDrawBuffer->UAV, ERHIAccess::SRVMask | ERHIAccess::IndirectArgs, ERHIAccess::UAVCompute));
 		FComputeShaderUtils::Dispatch(RHICmdList, ComputeShader, Params, FIntVector(1, 1, 1));
-		RHICmdList.Transition(FRHITransitionInfo(RenderingViewResources->IndirectDrawBuffer.UAV, ERHIAccess::UAVCompute, ERHIAccess::SRVMask | ERHIAccess::IndirectArgs));
+		RHICmdList.Transition(FRHITransitionInfo(RenderingViewResources->IndirectDrawBuffer->UAV, ERHIAccess::UAVCompute, ERHIAccess::SRVMask | ERHIAccess::IndirectArgs));
 	}
 	
 	// Not possible to have a valid ribbon with less than 2 particles, abort!
@@ -2221,7 +2284,7 @@ void FNiagaraRendererRibbons::InitializeViewIndexBuffersGPU(FRHICommandListImmed
 		Params.MultiRibbonIndices = VertexBuffers.MultiRibbonIndicesBuffer.SRV;
 		Params.Segments = VertexBuffers.SegmentsBuffer.SRV;
 
-		Params.IndirectDrawInfo = RenderingViewResources->IndirectDrawBuffer.SRV;
+		Params.IndirectDrawInfo = RenderingViewResources->IndirectDrawBuffer->SRV;
 		Params.TriangleToVertexIds = ShapeState.SliceTriangleToVertexIdsBuffer.SRV;
 
 		// Total particle Count
@@ -2245,7 +2308,7 @@ void FNiagaraRendererRibbons::InitializeViewIndexBuffersGPU(FRHICommandListImmed
 		Params.SubSegmentBitMask = RenderingViewResources->IndexGenerationSettings.SubSegmentBitMask;
 		
 		RHICmdList.Transition(FRHITransitionInfo(RenderingViewResources->IndexBuffer->UAV, ERHIAccess::VertexOrIndexBuffer, ERHIAccess::UAVCompute));
-		FComputeShaderUtils::DispatchIndirect(RHICmdList, ComputeShader, Params, RenderingViewResources->IndirectDrawBuffer.Buffer, IndirectDispatchArgsOffset);
+		FComputeShaderUtils::DispatchIndirect(RHICmdList, ComputeShader, Params, RenderingViewResources->IndirectDrawBuffer->Buffer, IndirectDispatchArgsOffset);
 		RHICmdList.Transition(FRHITransitionInfo(RenderingViewResources->IndexBuffer->UAV, ERHIAccess::UAVCompute, ERHIAccess::VertexOrIndexBuffer));
 	}
 }
@@ -2739,7 +2802,10 @@ void FNiagaraGpuRibbonsDataManager::GenerateAllGPUData(FRHICommandListImmediate&
 		{
 			for (const TSharedPtr<FNiagaraRibbonRenderingFrameViewResources>& RenderingResourcesView : RenderingResources->ViewResources)
 			{
-				RendererToGen.Renderer->InitializeViewIndexBuffersGPU(RHICmdList, ComputeDispatchInterface, RendererToGen, RenderingResourcesView);
+				if (RenderingResourcesView->bNeedsIndexBufferGeneration)
+				{
+					RendererToGen.Renderer->InitializeViewIndexBuffersGPU(RHICmdList, ComputeDispatchInterface, RendererToGen, RenderingResourcesView.Get());
+				}
 			}
 		}
 	}
diff --git a/Engine/Plugins/FX/Niagara/Source/Niagara/Public/NiagaraRendererRibbons.h b/Engine/Plugins/FX/Niagara/Source/Niagara/Public/NiagaraRendererRibbons.h
--- a/Engine/Plugins/FX/Niagara/Source/Niagara/Public/NiagaraRendererRibbons.h
+++ b/Engine/Plugins/FX/Niagara/Source/Niagara/Public/NiagaraRendererRibbons.h
@@ -338,18 +338,17 @@ protected:
 	template<typename IDType, typename ReaderType>
 	void GenerateVertexBufferForMultiRibbon(const FNiagaraGenerationInputDataCPUAccessors& CPUData, const ReaderType& IDReader, FNiagaraRibbonCPUGeneratedVertexData& OutputData) const;
 	
-	NIAGARA_API void GenerateVertexBufferCPU(const FNiagaraGenerationInputDataCPUAccessors& CPUData, FNiagaraRibbonCPUGeneratedVertexData& OutputData) const;
+	void GenerateVertexBufferCPU(const FNiagaraGenerationInputDataCPUAccessors& CPUData, FNiagaraRibbonCPUGeneratedVertexData& OutputData) const;
 
-
-	NIAGARA_API int32 CalculateTessellationFactor(const FNiagaraSceneProxy* SceneProxy, const FSceneView* View, const FVector& ViewOriginForDistanceCulling) const;
-	NIAGARA_API FNiagaraIndexGenerationInput CalculateIndexBufferConfiguration(const TSharedPtr<FNiagaraRibbonCPUGeneratedVertexData>& GeneratedVertices, const FNiagaraDataBuffer* SourceParticleData,
+	int32 CalculateTessellationFactor(const FNiagaraSceneProxy* SceneProxy, const FSceneView* View, const FVector& ViewOriginForDistanceCulling) const;
+	FNiagaraIndexGenerationInput CalculateIndexBufferConfiguration(const TSharedPtr<FNiagaraRibbonCPUGeneratedVertexData>& GeneratedVertices, const FNiagaraDataBuffer* SourceParticleData,
 		const FNiagaraSceneProxy* SceneProxy, const FSceneView* View, const FVector& ViewOriginForDistanceCulling, bool bShouldUseGPUInitIndices, bool bIsGPUSim) const;
 	
-	NIAGARA_API void GenerateIndexBufferForView(
+	void GenerateIndexBufferForView(
 		FRHICommandListBase& RHICmdList,
 		FNiagaraGpuRibbonsDataManager& GpuRibbonsDataManager, FMeshElementCollector& Collector,
 		FNiagaraIndexGenerationInput& GeneratedData, FNiagaraDynamicDataRibbon* DynamicDataRibbon,
-		const TSharedPtr<FNiagaraRibbonRenderingFrameViewResources>& RenderingViewResources, const FSceneView* View, const FVector& ViewOriginForDistanceCulling
+		FNiagaraRibbonRenderingFrameViewResources* RenderingViewResources, const FSceneView* View, const FVector& ViewOriginForDistanceCulling
 	) const;
 	
 	template <typename TValue>
@@ -361,24 +360,24 @@ protected:
 	template <typename TValue>
 	static TValue* AppendToIndexBufferCPU(TValue* OutIndices, const FNiagaraIndexGenerationInput& GeneratedData, const FNiagaraRibbonShapeGeometryData& ShapeState, const TArrayView<uint32>& SegmentData, bool bInvertOrder);
 	
-	NIAGARA_API void SetupPerViewUniformBuffer(FNiagaraIndexGenerationInput& GeneratedData,
+	void SetupPerViewUniformBuffer(FNiagaraIndexGenerationInput& GeneratedData,
 	                               const FSceneView* View, const FSceneViewFamily& ViewFamily, const FNiagaraSceneProxy* SceneProxy, FNiagaraRibbonUniformBufferRef& OutUniformBuffer) const;
 
-	NIAGARA_API void SetupMeshBatchAndCollectorResourceForView(FRHICommandListBase& RHICmdList, const FNiagaraIndexGenerationInput& GeneratedData, FNiagaraDynamicDataRibbon* DynamicDataRibbon,
+	void SetupMeshBatchAndCollectorResourceForView(FRHICommandListBase& RHICmdList, const FNiagaraIndexGenerationInput& GeneratedData, FNiagaraDynamicDataRibbon* DynamicDataRibbon,
 	                                               const FNiagaraDataBuffer* SourceParticleData, const FSceneView* View, const FSceneViewFamily& ViewFamily, const FNiagaraSceneProxy* SceneProxy,
-	                                               const TSharedPtr<FNiagaraRibbonRenderingFrameResources>& RenderingResources, const TSharedPtr<FNiagaraRibbonRenderingFrameViewResources>& RenderingViewResources, FMeshBatch& OutMeshBatch, bool bShouldUseGPUInitIndices) const;
+	                                               const TSharedPtr<FNiagaraRibbonRenderingFrameResources>& RenderingResources, FNiagaraRibbonRenderingFrameViewResources* RenderingViewResources, FMeshBatch& OutMeshBatch, bool bShouldUseGPUInitIndices) const;
 
 
-	NIAGARA_API void InitializeViewIndexBuffersGPU(FRHICommandListImmediate& RHICmdList, FNiagaraGpuComputeDispatchInterface* ComputeDispatchInterface, const FNiagaraRibbonGPUInitParameters& GpuInitParameters,
-		const TSharedPtr<FNiagaraRibbonRenderingFrameViewResources>& RenderingViewResources) const;
+	void InitializeViewIndexBuffersGPU(FRHICommandListImmediate& RHICmdList, FNiagaraGpuComputeDispatchInterface* ComputeDispatchInterface, const FNiagaraRibbonGPUInitParameters& GpuInitParameters,
+		const FNiagaraRibbonRenderingFrameViewResources* RenderingViewResources) const;
 
-	NIAGARA_API void InitializeVertexBuffersResources(FRHICommandListBase& RHICmdList, const FNiagaraDynamicDataRibbon* DynamicDataRibbon, FNiagaraDataBuffer* SourceParticleData,
+	void InitializeVertexBuffersResources(FRHICommandListBase& RHICmdList, const FNiagaraDynamicDataRibbon* DynamicDataRibbon, FNiagaraDataBuffer* SourceParticleData,
 	                                      FGlobalDynamicReadBuffer& DynamicReadBuffer, const TSharedPtr<FNiagaraRibbonRenderingFrameResources>& RenderingResources, bool bShouldUseGPUInit) const;
 	
-	NIAGARA_API void InitializeVertexBuffersGPU(FRHICommandListImmediate& RHICmdList, FNiagaraGpuComputeDispatchInterface* ComputeDispatchInterface, const FNiagaraRibbonGPUInitParameters& GpuInitParameters,
+	void InitializeVertexBuffersGPU(FRHICommandListImmediate& RHICmdList, FNiagaraGpuComputeDispatchInterface* ComputeDispatchInterface, const FNiagaraRibbonGPUInitParameters& GpuInitParameters,
 		FNiagaraRibbonGPUInitComputeBuffers& TempBuffers, const TSharedPtr<FNiagaraRibbonRenderingFrameResources>& RenderingResources) const;
 
-	NIAGARA_API FRibbonComputeUniformParameters SetupComputeVertexGenParams(FNiagaraGpuComputeDispatchInterface* ComputeDispatchInterface, const TSharedPtr<FNiagaraRibbonRenderingFrameResources>& RenderingResources, const FNiagaraRibbonGPUInitParameters& GpuInitParameters) const;
+	FRibbonComputeUniformParameters SetupComputeVertexGenParams(FNiagaraGpuComputeDispatchInterface* ComputeDispatchInterface, const TSharedPtr<FNiagaraRibbonRenderingFrameResources>& RenderingResources, const FNiagaraRibbonGPUInitParameters& GpuInitParameters) const;
 
 	FNiagaraRibbonGenerationConfig GenerationConfig;
 	
