- Remove duplicate enum entries from list

diff --git a/Engine/Plugins/FX/Niagara/Source/NiagaraEditor/Private/Customizations/NiagaraStaticSwitchNodeDetails.cpp b/Engine/Plugins/FX/Niagara/Source/NiagaraEditor/Private/Customizations/NiagaraStaticSwitchNodeDetails.cpp
--- a/Engine/Plugins/FX/Niagara/Source/NiagaraEditor/Private/Customizations/NiagaraStaticSwitchNodeDetails.cpp
+++ b/Engine/Plugins/FX/Niagara/Source/NiagaraEditor/Private/Customizations/NiagaraStaticSwitchNodeDetails.cpp
@@ -33,13 +33,16 @@ FNiagaraStaticSwitchNodeDetails::FNiagaraStaticSwitchNodeDetails()
 	DropdownOptions.Add(MakeShareable(new SwitchDropdownOption("Integer Constant")));
 	DropdownOptions.Add(MakeShareable(new SwitchDropdownOption("Enum Constant")));
 
+	// Note: Due to static variables Enums can exist in the type registery as both dynamic & static variables so we need to deduplicate them
+	TSet<UEnum*> FoundEnums;
 	TArray<FNiagaraTypeDefinition> ParameterTypes;
 	FNiagaraEditorUtilities::GetAllowedParameterTypes(ParameterTypes);
 	for (FNiagaraTypeDefinition Type : ParameterTypes)
 	{
-		if (Type.IsEnum())
+		if (Type.IsEnum() && !FoundEnums.Contains(Type.GetEnum()))
 		{
 			DropdownOptions.Add(MakeShareable(new SwitchDropdownOption(Type.GetEnum()->GetName(), Type.GetEnum())));
+			FoundEnums.Add(Type.GetEnum());
 		}
 	}
 }
