Preserve LastHlslTranslationGPU on NiagaraScripts for cooked editor so that they can compile GPU simulation scripts.

diff --git a/Engine/Plugins/FX/Niagara/Source/Niagara/Private/NiagaraScript.cpp b/Engine/Plugins/FX/Niagara/Source/Niagara/Private/NiagaraScript.cpp
--- a/Engine/Plugins/FX/Niagara/Source/Niagara/Private/NiagaraScript.cpp
+++ b/Engine/Plugins/FX/Niagara/Source/Niagara/Private/NiagaraScript.cpp
@@ -1932,7 +1932,10 @@ void UNiagaraScript::Serialize(FArchive& Ar)
 		auto SwapDebugData = [](FNiagaraVMExecutableData& SourceData, FNiagaraVMExecutableData& TargetData) -> void
 		{
 			Swap(TargetData.LastHlslTranslation, SourceData.LastHlslTranslation);
-			Swap(TargetData.LastHlslTranslationGPU, SourceData.LastHlslTranslationGPU);
+
+			// We must preserve the LastTranslationGPU as it is used to compile the GPUComputeScript when using the standard compilation mode
+			// Swap(TargetData.LastHlslTranslationGPU, SourceData.LastHlslTranslationGPU);
+
 			Swap(TargetData.LastAssemblyTranslation, SourceData.LastAssemblyTranslation);
 			Swap(TargetData.CompileTagsEditorOnly, SourceData.CompileTagsEditorOnly);
 		};
