- Fix for DDC population issue on specific platforms

diff --git a/Engine/Plugins/FX/Niagara/Source/NiagaraBlueprintNodes/Private/K2Node_DataChannelBase.cpp b/Engine/Plugins/FX/Niagara/Source/NiagaraBlueprintNodes/Private/K2Node_DataChannelBase.cpp
--- a/Engine/Plugins/FX/Niagara/Source/NiagaraBlueprintNodes/Private/K2Node_DataChannelBase.cpp
+++ b/Engine/Plugins/FX/Niagara/Source/NiagaraBlueprintNodes/Private/K2Node_DataChannelBase.cpp
@@ -97,7 +97,7 @@ void UK2Node_DataChannelBase::PreloadRequiredAssets()
 {
 	Super::PreloadRequiredAssets();
 
-	if (HasValidDataChannel())
+	if (DataChannel)
 	{
 		PreloadObject(DataChannel);
 		PreloadObject(DataChannel->Get());
