compute bounding box of all objects and reject samples that fall outside of it

diff --git a/Engine/Plugins/FX/Niagara/Shaders/Private/NiagaraDataInterfaceRigidMeshCollisionQuery.ush b/Engine/Plugins/FX/Niagara/Shaders/Private/NiagaraDataInterfaceRigidMeshCollisionQuery.ush
--- a/Engine/Plugins/FX/Niagara/Shaders/Private/NiagaraDataInterfaceRigidMeshCollisionQuery.ush
+++ b/Engine/Plugins/FX/Niagara/Shaders/Private/NiagaraDataInterfaceRigidMeshCollisionQuery.ush
@@ -25,6 +25,8 @@ Buffer<float4>	{ParameterName}_ElementExtentBuffer;
 Buffer<uint>	{ParameterName}_PhysicsTypeBuffer;
 Buffer<uint>	{ParameterName}_DFIndexBuffer;
 float3			{ParameterName}_SystemLWCTile;
+float3			{ParameterName}_CombinedBBoxWorldMin;
+float3			{ParameterName}_CombinedBBoxWorldMax;
 
 /* -----------------------------------------------------------------
  * Element accessors
@@ -85,6 +87,19 @@ float3 GetMeshScale_{ParameterName}(in int ElementIndex)
 	return {ParameterName}_MeshScaleBuffer[ElementIndex].xyz;
 }
 
+void IsWorldPositionInsideCombinedBounds_{ParameterName}(in float3 WorldPosition, out bool IsInside)
+{
+	const FDFVector3 LWCWorldPosition = DFFromTileOffset_Hack(MakeLWCVector3({ParameterName}_SystemLWCTile, WorldPosition));
+
+	IsInside =  
+		WorldPosition.x >= {ParameterName}_CombinedBBoxWorldMin.x && 
+		WorldPosition.y >= {ParameterName}_CombinedBBoxWorldMin.y && 
+		WorldPosition.z >= {ParameterName}_CombinedBBoxWorldMin.z && 
+		WorldPosition.x <= {ParameterName}_CombinedBBoxWorldMax.x && 
+		WorldPosition.y <= {ParameterName}_CombinedBBoxWorldMax.y && 
+		WorldPosition.z <= {ParameterName}_CombinedBBoxWorldMax.z;
+}
+
 /* -----------------------------------------------------------------
  * Collision detection utils
  * -----------------------------------------------------------------
diff --git a/Engine/Plugins/FX/Niagara/Source/Niagara/Private/NiagaraDataInterfaceRigidMeshCollisionQuery.cpp b/Engine/Plugins/FX/Niagara/Source/Niagara/Private/NiagaraDataInterfaceRigidMeshCollisionQuery.cpp
--- a/Engine/Plugins/FX/Niagara/Source/Niagara/Private/NiagaraDataInterfaceRigidMeshCollisionQuery.cpp
+++ b/Engine/Plugins/FX/Niagara/Source/Niagara/Private/NiagaraDataInterfaceRigidMeshCollisionQuery.cpp
@@ -74,6 +74,8 @@ BEGIN_SHADER_PARAMETER_STRUCT(FShaderParameters,)
 	SHADER_PARAMETER_SRV(Buffer<uint32>,	PhysicsTypeBuffer)
 	SHADER_PARAMETER_SRV(Buffer<uint32>,	DFIndexBuffer)
 	SHADER_PARAMETER(FVector3f,				SystemLWCTile)
+	SHADER_PARAMETER(FVector3f,				CombinedBBoxWorldMin)
+	SHADER_PARAMETER(FVector3f,				CombinedBBoxWorldMax)
 END_SHADER_PARAMETER_STRUCT()
 
 static const TCHAR* TemplateShaderFile = TEXT("/Plugin/FX/Niagara/Private/NiagaraDataInterfaceRigidMeshCollisionQuery.ush");
@@ -90,7 +92,7 @@ static const FName GetNumCapsulesName(TEXT("GetNumCapsules"));
 static const FName GetNumElementsName(TEXT("GetNumElements"));
 
 //------------------------------------------------------------------------------------------------------------
-
+static const FName IsWorldPositionInsideCombinedBoundsName(TEXT("IsWorldPositionInsideCombinedBounds"));
 static const FName GetClosestElementName(TEXT("GetClosestElement"));
 static const FName GetElementPointName(TEXT("GetElementPoint"));
 static const FName GetElementPointMeshDistanceFieldNoNormalName(TEXT("GetElementPointMeshDistanceFieldNoNormal"));
@@ -385,7 +387,7 @@ FTransform CreateElementTransform<USkeletalMeshComponent>(const USkeletalMeshCom
 }
 
 template<typename TComponentType, bool InitializeStatics>
-void UpdateAssetArrays(TConstArrayView<TComponentType*> Components, const FVector& LWCTile, FNDIRigidMeshCollisionArrays* OutAssetArrays, uint32& BoxIndex, uint32& SphereIndex, uint32& CapsuleIndex)
+void UpdateAssetArrays(TConstArrayView<TComponentType*> Components, const FVector& LWCTile, FNDIRigidMeshCollisionArrays* OutAssetArrays, uint32& BoxIndex, uint32& SphereIndex, uint32& CapsuleIndex, FBoxSphereBounds::Builder &CombinedWorldBoundsBuilder)
 {
 	auto UpdateAssetPredicate = [&](TComponentType* Component, const UBodySetup* BodySetup)
 	{
@@ -500,9 +502,13 @@ void UpdateAssetArrays(TConstArrayView<TComponentType*> Components, const FVecto
 			++BoxIndex;
 		}
 	};
-
+	
 	for (TComponentType* Component : Components)
-	{
+	{				
+		FBoxSphereBounds LocalBounds = static_cast<USceneComponent*>(Component)->GetLocalBounds();
+		FBoxSphereBounds WorldBounds = LocalBounds.TransformBy(Component->GetComponentTransform());
+		CombinedWorldBoundsBuilder += WorldBounds;		
+
 		ForEachBodySetup(Component, UpdateAssetPredicate);
 	}
 }
@@ -623,6 +629,7 @@ void RemapPreviousTransforms(
 }
 
 bool UpdateInternalArrays(
+	const FNiagaraSystemInstance* SystemInstance,
 	TConstArrayView<UStaticMeshComponent*> StaticMeshView,
 	TConstArrayView<USkeletalMeshComponent*> SkeletalMeshView,
 	FVector LWCTile,
@@ -647,9 +654,24 @@ bool UpdateInternalArrays(
 		uint32 BoxIndex = OutAssetArrays->ElementOffsets.BoxOffset;
 		uint32 SphereIndex = OutAssetArrays->ElementOffsets.SphereOffset;
 		uint32 CapsuleIndex = OutAssetArrays->ElementOffsets.CapsuleOffset;
+		
+		OutAssetArrays->CombinedBBoxWorldMin = FVector3f(TNumericLimits<float>::Max());
+		OutAssetArrays->CombinedBBoxWorldMax = FVector3f(TNumericLimits<float>::Lowest());
+
+		FBoxSphereBounds::Builder CombinedWorldBoundsBuilder;
+
+		UpdateAssetArrays<UStaticMeshComponent, true>(StaticMeshView, LWCTile, OutAssetArrays, BoxIndex, SphereIndex, CapsuleIndex, CombinedWorldBoundsBuilder);
+		UpdateAssetArrays<USkeletalMeshComponent, true>(SkeletalMeshView, LWCTile, OutAssetArrays, BoxIndex, SphereIndex, CapsuleIndex, CombinedWorldBoundsBuilder);
+
+		const FBoxSphereBounds CombinedWorldBounds(CombinedWorldBoundsBuilder);
+		const FBox CombinedWorldBoundingBox = CombinedWorldBounds.GetBox();
 
-		UpdateAssetArrays<UStaticMeshComponent, false>(StaticMeshView, LWCTile, OutAssetArrays, BoxIndex, SphereIndex, CapsuleIndex);
-		UpdateAssetArrays<USkeletalMeshComponent, false>(SkeletalMeshView, LWCTile, OutAssetArrays, BoxIndex, SphereIndex, CapsuleIndex);
+		OutAssetArrays->CombinedBBoxWorldMin = (FVector3f)CombinedWorldBoundingBox.Min;
+		OutAssetArrays->CombinedBBoxWorldMax = (FVector3f)CombinedWorldBoundingBox.Max;
+
+		FNiagaraLWCConverter LWCConverter = SystemInstance->GetLWCConverter();
+		OutAssetArrays->CombinedBBoxWorldMin = LWCConverter.ConvertWorldToSimulationVector(CombinedWorldBoundingBox.Min);
+		OutAssetArrays->CombinedBBoxWorldMax = LWCConverter.ConvertWorldToSimulationVector(CombinedWorldBoundingBox.Max);
 
 		return true;
 	}
@@ -685,9 +707,25 @@ bool UpdateInternalArrays(
 	// CurrentInverse based on the remapped entries
 	RemapPreviousTransforms<true>(BodyCounts, CurrentBodyCounts, PreviousElementOffsets, OutAssetArrays->ElementOffsets, OutAssetArrays);
 
-	UpdateAssetArrays<UStaticMeshComponent, true>(StaticMeshView, LWCTile, OutAssetArrays, BoxIndex, SphereIndex, CapsuleIndex);
-	UpdateAssetArrays<USkeletalMeshComponent, true>(SkeletalMeshView, LWCTile, OutAssetArrays, BoxIndex, SphereIndex, CapsuleIndex);
+	// Initialize world bounds
+	OutAssetArrays->CombinedBBoxWorldMin = FVector3f(TNumericLimits<float>::Max());
+	OutAssetArrays->CombinedBBoxWorldMax = FVector3f(TNumericLimits<float>::Lowest());
+
+	FBoxSphereBounds::Builder CombinedWorldBoundsBuilder;
+
+	UpdateAssetArrays<UStaticMeshComponent, true>(StaticMeshView, LWCTile, OutAssetArrays, BoxIndex, SphereIndex, CapsuleIndex, CombinedWorldBoundsBuilder);
+	UpdateAssetArrays<USkeletalMeshComponent, true>(SkeletalMeshView, LWCTile, OutAssetArrays, BoxIndex, SphereIndex, CapsuleIndex, CombinedWorldBoundsBuilder);
+
+	const FBoxSphereBounds CombinedWorldBounds(CombinedWorldBoundsBuilder);
+	const FBox CombinedWorldBoundingBox = CombinedWorldBounds.GetBox();
 
+	OutAssetArrays->CombinedBBoxWorldMin = (FVector3f) CombinedWorldBoundingBox.Min;
+	OutAssetArrays->CombinedBBoxWorldMax = (FVector3f) CombinedWorldBoundingBox.Max;
+
+	FNiagaraLWCConverter LWCConverter = SystemInstance->GetLWCConverter();
+	OutAssetArrays->CombinedBBoxWorldMin = LWCConverter.ConvertWorldToSimulationVector(CombinedWorldBoundingBox.Min);
+	OutAssetArrays->CombinedBBoxWorldMax = LWCConverter.ConvertWorldToSimulationVector(CombinedWorldBoundingBox.Max);
+			
 	RemapPreviousTransforms<false>(BodyCounts, CurrentBodyCounts, PreviousElementOffsets, OutAssetArrays->ElementOffsets, OutAssetArrays);
 
 	BodyCounts = MoveTemp(CurrentBodyCounts);
@@ -918,6 +956,7 @@ PRAGMA_ENABLE_DEPRECATION_WARNINGS
 		TConstArrayView<USkeletalMeshComponent*> SkeletalMeshView = MakeArrayView(SkeletalMeshes.GetData(), SkeletalMeshes.Num());
 
 		const bool bArraysUpdated = UpdateInternalArrays(
+			SystemInstance,
 			StaticMeshView,
 			SkeletalMeshView,
 			FVector(SystemInstance->GetLWCTile()),
@@ -954,6 +993,9 @@ struct FNDIRigidMeshCollisionProxy : public FNiagaraDataInterfaceProxy
 		TArray<FPrimitiveComponentId> UniqueComponentIds;
 
 		FNDIRigidMeshCollisionBuffer* AssetBuffer = nullptr;
+
+		FVector3f CombinedBBoxWorldMin;
+		FVector3f CombinedBBoxWorldMax;
 	};
 
 	void RemoveInstance(const FNiagaraSystemInstanceID& Instance)
@@ -1091,25 +1133,6 @@ void UNiagaraDataInterfaceRigidMeshCollisionQuery::DrawDebugHud(FNDIDrawDebugHud
 		return ElementMatrix.GetTransposed();
 	};
 
-	auto GetPreviousTransform = [&](int32 ElementIndex)
-	{
-		const uint32 ElementOffset = 3 * ElementIndex;
-		FVector4f* TransformVec = InstanceData_GT->AssetArrays->PreviousTransform.GetData() + ElementOffset;
-
-		FMatrix ElementMatrix;
-		ElementMatrix.SetIdentity();
-
-		for (int32 RowIt = 0; RowIt < 3; ++RowIt)
-		{
-			for (int32 ColIt = 0; ColIt < 4; ++ColIt)
-			{
-				ElementMatrix.M[RowIt][ColIt] = TransformVec[RowIt][ColIt];
-			}
-		}
-
-		return ElementMatrix.GetTransposed();
-	};
-
 	if (DebugHudContext.IsVerbose())
 	{
 		const UWorld* World = DebugHudContext.GetWorld();
@@ -1117,6 +1140,11 @@ void UNiagaraDataInterfaceRigidMeshCollisionQuery::DrawDebugHud(FNDIDrawDebugHud
 		TArray<FVector4f>& CurrentTransformArray = InstanceData_GT->AssetArrays->CurrentTransform;
 		TArray<FVector4f>& PreviousTransformArray = InstanceData_GT->AssetArrays->PreviousTransform;
 
+		// Inclusive world bounds for DI
+		FVector3f Center = (InstanceData_GT->AssetArrays->CombinedBBoxWorldMin + InstanceData_GT->AssetArrays->CombinedBBoxWorldMax) * .5;
+		FVector3f HalfExtent = (InstanceData_GT->AssetArrays->CombinedBBoxWorldMax - InstanceData_GT->AssetArrays->CombinedBBoxWorldMin) * .5;
+		DrawDebugBox(World, (FVector) Center, (FVector) HalfExtent, FQuat::Identity, FColor::Purple);
+
 		// Boxes
 		for (uint32 BoxIt = 0; BoxIt < BoxCount; ++BoxIt)
 		{
@@ -1452,6 +1480,21 @@ void UNiagaraDataInterfaceRigidMeshCollisionQuery::GetFunctionsInternal(TArray<F
 
 		OutFunctions.Add(Sig);
 	}
+	{
+		FNiagaraFunctionSignature Sig;
+		Sig.Name = IsWorldPositionInsideCombinedBoundsName;
+		Sig.SetDescription(LOCTEXT("IsWorldPositionInsideCombinedBoundsDescription", "Returns true if the query point is inside the bounding box of all the actors"));
+		Sig.SetFunctionVersion(FNiagaraRigidMeshCollisionDIFunctionVersion::LatestVersion);
+		Sig.bSupportsGPU = true;
+		Sig.bSupportsCPU = true;
+		Sig.bMemberFunction = true;
+		Sig.Inputs.Add(FNiagaraVariable(FNiagaraTypeDefinition(GetClass()), TEXT("Collision DI")));
+		Sig.Inputs.Add(FNiagaraVariable(FNiagaraTypeDefinition::GetPositionDef(), TEXT("World Position")));
+		Sig.Outputs.Add(FNiagaraVariable(FNiagaraTypeDefinition::GetBoolDef(), TEXT("IsInside")));
+
+		OutFunctions.Add(Sig);
+	}
+
 	{
 		FNiagaraFunctionSignature Sig;
 		Sig.Name = GetClosestPointName;
@@ -1669,6 +1712,7 @@ bool UNiagaraDataInterfaceRigidMeshCollisionQuery::GetFunctionHLSL(const FNiagar
 		(FunctionInfo.DefinitionName == GetNumCapsulesName) ||
 		(FunctionInfo.DefinitionName == GetNumSpheresName) ||
 		(FunctionInfo.DefinitionName == GetNumElementsName) ||
+		(FunctionInfo.DefinitionName == IsWorldPositionInsideCombinedBoundsName) ||
 		(FunctionInfo.DefinitionName == GetClosestPointName) ||
 		(FunctionInfo.DefinitionName == GetClosestPointSimpleName) ||		
 		(FunctionInfo.DefinitionName == GetClosestElementName) ||
@@ -1855,6 +1899,9 @@ void UNiagaraDataInterfaceRigidMeshCollisionQuery::SetShaderParameters(const FNi
 		ShaderParameters->ElementOffsets.W = ProxyData->ElementOffsets.NumElements;
 
 		bBindDistanceFieldData = true;
+
+		ShaderParameters->CombinedBBoxWorldMin = ProxyData->CombinedBBoxWorldMin;
+		ShaderParameters->CombinedBBoxWorldMax = ProxyData->CombinedBBoxWorldMax;
 	}
 	else
 	{
@@ -1869,6 +1916,9 @@ void UNiagaraDataInterfaceRigidMeshCollisionQuery::SetShaderParameters(const FNi
 		ShaderParameters->CurrentOffset = 0;
 		ShaderParameters->PreviousOffset = 0;
 		ShaderParameters->ElementOffsets = FUintVector4(0, 0, 0, 0);
+
+		ShaderParameters->CombinedBBoxWorldMin = FVector3f(0, 0, 0);
+		ShaderParameters->CombinedBBoxWorldMax = FVector3f(0, 0, 0);
 	}
 
 	if (bDistanceFieldDataBound)
@@ -1935,6 +1985,9 @@ void UNiagaraDataInterfaceRigidMeshCollisionQuery::ProvidePerInstanceDataForRend
 			RenderThreadData->UniqueComponentIds = GameThreadData->AssetArrays->UniqueCompnentId;
 			RenderThreadData->MaxPrimitiveCount = ElementCount;
 			RenderThreadData->AssetBuffer = GameThreadData->AssetBuffer;
+
+			RenderThreadData->CombinedBBoxWorldMin = GameThreadData->AssetArrays->CombinedBBoxWorldMin;
+			RenderThreadData->CombinedBBoxWorldMax = GameThreadData->AssetArrays->CombinedBBoxWorldMax;
 		}
 	}
 	check(Proxy);
@@ -2118,6 +2171,38 @@ void UNiagaraDataInterfaceRigidMeshCollisionQuery::GetNumElementsCPU(FVectorVMEx
 	}
 }
 
+void UNiagaraDataInterfaceRigidMeshCollisionQuery::IsWorldPositionInsideCombinedBoundsCPU(FVectorVMExternalFunctionContext& Context)
+{
+	VectorVM::FUserPtrHandler<FNDIRigidMeshCollisionData> InstanceData(Context);
+	
+	FNDIInputParam<FNiagaraPosition> WorldPositionParam(Context);
+	FNDIOutputParam<bool> IsInsideParam(Context);
+
+	FNDIOutputParam<int32> NumElementsParam(Context);
+
+	if (InstanceData->SystemInstance != nullptr && InstanceData->AssetArrays != nullptr)
+	{
+		FNiagaraPosition WorldPosition = WorldPositionParam.GetAndAdvance();
+
+		FNiagaraLWCConverter LWCConverter = InstanceData->SystemInstance->GetLWCConverter();
+		const FVector ConvertedWorldPosition = LWCConverter.ConvertSimulationPositionToWorld(WorldPosition);
+
+		const bool IsInside =
+			ConvertedWorldPosition.X >= InstanceData->AssetArrays->CombinedBBoxWorldMin.X &&
+			ConvertedWorldPosition.Y >= InstanceData->AssetArrays->CombinedBBoxWorldMin.Y &&
+			ConvertedWorldPosition.Z >= InstanceData->AssetArrays->CombinedBBoxWorldMin.Z &&
+			ConvertedWorldPosition.X <= InstanceData->AssetArrays->CombinedBBoxWorldMax.X &&
+			ConvertedWorldPosition.Y <= InstanceData->AssetArrays->CombinedBBoxWorldMax.Y &&
+			ConvertedWorldPosition.Z <= InstanceData->AssetArrays->CombinedBBoxWorldMax.Z;
+
+		NumElementsParam.SetAndAdvance(IsInside);
+	}
+	else
+	{
+		NumElementsParam.SetAndAdvance(false);
+	}
+}
+
 #if WITH_EDITOR
 
 void UNiagaraDataInterfaceRigidMeshCollisionQuery::GetFeedback(UNiagaraSystem* InAsset, UNiagaraComponent* InComponent, TArray<FNiagaraDataInterfaceError>& OutErrors, TArray<FNiagaraDataInterfaceFeedback>& OutWarnings, TArray<FNiagaraDataInterfaceFeedback>& OutInfo)
diff --git a/Engine/Plugins/FX/Niagara/Source/Niagara/Public/NiagaraDataInterfaceRigidMeshCollisionQuery.h b/Engine/Plugins/FX/Niagara/Source/Niagara/Public/NiagaraDataInterfaceRigidMeshCollisionQuery.h
--- a/Engine/Plugins/FX/Niagara/Source/Niagara/Public/NiagaraDataInterfaceRigidMeshCollisionQuery.h
+++ b/Engine/Plugins/FX/Niagara/Source/Niagara/Public/NiagaraDataInterfaceRigidMeshCollisionQuery.h
@@ -43,6 +43,9 @@ struct FNDIRigidMeshCollisionArrays
 	TArray<int32> ComponentIdIndex;
 	TArray<FPrimitiveComponentId> UniqueCompnentId;
 
+	FVector3f CombinedBBoxWorldMin;
+	FVector3f CombinedBBoxWorldMax;
+
 	FNDIRigidMeshCollisionArrays() = delete;
 	FNDIRigidMeshCollisionArrays(uint32 Num)
 		: MaxPrimitives(Num)
@@ -267,12 +270,13 @@ public:
 
 	bool GetExplicitActors(FNDIRigidMeshCollisionData& InstanceData);
 	bool FindActors(UWorld* World, FNDIRigidMeshCollisionData& InstanceData, ECollisionChannel Channel, const FVector& OverlapLocation, const FVector& OverlapExtent, const FQuat& OverlapRotation) const;
-
+	
 	UE_DEPRECATED(5.2, "Global search will be deprecated in favor of using a spatial search through the FindActors function.")
 	bool GlobalFindActors(UWorld* World, FNDIRigidMeshCollisionData& InstanceData) const;
 
 	void FindActorsCPU(FVectorVMExternalFunctionContext& Context);
 	void GetNumElementsCPU(FVectorVMExternalFunctionContext& Context);
+	void IsWorldPositionInsideCombinedBoundsCPU(FVectorVMExternalFunctionContext& Context);
 
 #if WITH_EDITOR
 	virtual void GetFeedback(UNiagaraSystem* InAsset, UNiagaraComponent* InComponent, TArray<FNiagaraDataInterfaceError>& OutErrors, TArray<FNiagaraDataInterfaceFeedback>& OutWarnings, TArray<FNiagaraDataInterfaceFeedback>& OutInfo) override;
