fixed issue with vm version of checking combined bounds

diff --git a/Engine/Plugins/FX/Niagara/Source/Niagara/Private/NiagaraDataInterfaceRigidMeshCollisionQuery.cpp b/Engine/Plugins/FX/Niagara/Source/Niagara/Private/NiagaraDataInterfaceRigidMeshCollisionQuery.cpp
--- a/Engine/Plugins/FX/Niagara/Source/Niagara/Private/NiagaraDataInterfaceRigidMeshCollisionQuery.cpp
+++ b/Engine/Plugins/FX/Niagara/Source/Niagara/Private/NiagaraDataInterfaceRigidMeshCollisionQuery.cpp
@@ -1683,6 +1683,7 @@ void UNiagaraDataInterfaceRigidMeshCollisionQuery::GetFunctionsInternal(TArray<F
 
 DEFINE_NDI_DIRECT_FUNC_BINDER(UNiagaraDataInterfaceRigidMeshCollisionQuery, FindActorsCPU);
 DEFINE_NDI_DIRECT_FUNC_BINDER(UNiagaraDataInterfaceRigidMeshCollisionQuery, GetNumElementsCPU);
+DEFINE_NDI_DIRECT_FUNC_BINDER(UNiagaraDataInterfaceRigidMeshCollisionQuery, IsWorldPositionInsideCombinedBoundsCPU);
 
 void UNiagaraDataInterfaceRigidMeshCollisionQuery::GetVMExternalFunction(const FVMExternalFunctionBindingInfo& BindingInfo, void* InstanceData, FVMExternalFunction& OutFunc)
 {
@@ -1696,6 +1697,10 @@ void UNiagaraDataInterfaceRigidMeshCollisionQuery::GetVMExternalFunction(const F
 	{
 		NDI_FUNC_BINDER(UNiagaraDataInterfaceRigidMeshCollisionQuery, GetNumElementsCPU)::Bind(this, OutFunc);
 	}
+	else if (BindingInfo.Name == IsWorldPositionInsideCombinedBoundsName)
+	{
+		NDI_FUNC_BINDER(UNiagaraDataInterfaceRigidMeshCollisionQuery, IsWorldPositionInsideCombinedBoundsCPU)::Bind(this, OutFunc);
+	}
 	else
 	{
 		UE_LOG(LogNiagara, Display, TEXT("Could not find data interface external function in %s. %s\n"),
@@ -2178,8 +2183,6 @@ void UNiagaraDataInterfaceRigidMeshCollisionQuery::IsWorldPositionInsideCombined
 	FNDIInputParam<FNiagaraPosition> WorldPositionParam(Context);
 	FNDIOutputParam<bool> IsInsideParam(Context);
 
-	FNDIOutputParam<int32> NumElementsParam(Context);
-
 	if (InstanceData->SystemInstance != nullptr && InstanceData->AssetArrays != nullptr)
 	{
 		FNiagaraPosition WorldPosition = WorldPositionParam.GetAndAdvance();
@@ -2195,11 +2198,11 @@ void UNiagaraDataInterfaceRigidMeshCollisionQuery::IsWorldPositionInsideCombined
 			ConvertedWorldPosition.Y <= InstanceData->AssetArrays->CombinedBBoxWorldMax.Y &&
 			ConvertedWorldPosition.Z <= InstanceData->AssetArrays->CombinedBBoxWorldMax.Z;
 
-		NumElementsParam.SetAndAdvance(IsInside);
+		IsInsideParam.SetAndAdvance(IsInside);
 	}
 	else
 	{
-		NumElementsParam.SetAndAdvance(false);
+		IsInsideParam.SetAndAdvance(false);
 	}
 }
 
