Niagara - Fix several issues related to dependencies.
+ When refreshing a stack module item, validate that the nodes graph is correct, and that the node is still in the graph to prevent crashes for nodes that have been removed.
+ When generating dependency issues, find the index of the module data using only the module node pointer and not the module node pointer or the name.  Allowing matching by name can end up returning nodes from the system script when evaluating nodes in the emitter or particle scripts, since they are in different graphs and can have duplicate names.
+ When adding a module to the stack and trying to find an appropriate index, use the supplied index from the arguments when the FindBestIndexForModuleInStack isn't able to find a better index.
+ Update FindBestIndexForModuleInStack so that it doesn't assume it's dealing with emitter scripts only and handles system scripts correctly.
+ When finding an index in FindBestIndexForModuleInStack, restrict the search to a specific usage and usage id since the index stored on the stack module data is relative to a specific script, and checking indices across multiple scripts will result in inconsistent and sometimes wrong results.

diff --git a/Engine/Plugins/FX/Niagara/Source/NiagaraEditor/Private/ViewModels/NiagaraSystemViewModel.cpp b/Engine/Plugins/FX/Niagara/Source/NiagaraEditor/Private/ViewModels/NiagaraSystemViewModel.cpp
--- a/Engine/Plugins/FX/Niagara/Source/NiagaraEditor/Private/ViewModels/NiagaraSystemViewModel.cpp
+++ b/Engine/Plugins/FX/Niagara/Source/NiagaraEditor/Private/ViewModels/NiagaraSystemViewModel.cpp
@@ -1220,10 +1220,13 @@ const TArray<FNiagaraStackModuleData>& FNiagaraSystemViewModel::GetStackModuleDa
 	}
 
 	FVersionedNiagaraEmitter Emitter;
-	TSharedPtr<FNiagaraEmitterHandleViewModel> EmitterHandleViewModel = GetEmitterHandleViewModelById(EmitterHandleId);
-	if (ensureMsgf(EmitterHandleViewModel.IsValid(), TEXT("Failed to get the emitter handle view model for emitter handle id %s while getting stack module data."), *EmitterHandleId.ToString(EGuidFormats::DigitsWithHyphens)))
+	if (EmitterHandleId.IsValid())
 	{
-		Emitter = EmitterHandleViewModel->GetEmitterViewModel()->GetEmitter();
+		TSharedPtr<FNiagaraEmitterHandleViewModel> EmitterHandleViewModel = GetEmitterHandleViewModelById(EmitterHandleId);
+		if (ensureMsgf(EmitterHandleViewModel.IsValid(), TEXT("Failed to get the emitter handle view model for emitter handle id %s while getting stack module data."), *EmitterHandleId.ToString(EGuidFormats::DigitsWithHyphens)))
+		{
+			Emitter = EmitterHandleViewModel->GetEmitterViewModel()->GetEmitter();
+		}
 	}
 
 	return BuildAndCacheStackModuleData(EmitterHandleId, Emitter);
diff --git a/Engine/Plugins/FX/Niagara/Source/NiagaraEditor/Private/ViewModels/Stack/NiagaraStackGraphUtilities.cpp b/Engine/Plugins/FX/Niagara/Source/NiagaraEditor/Private/ViewModels/Stack/NiagaraStackGraphUtilities.cpp
--- a/Engine/Plugins/FX/Niagara/Source/NiagaraEditor/Private/ViewModels/Stack/NiagaraStackGraphUtilities.cpp
+++ b/Engine/Plugins/FX/Niagara/Source/NiagaraEditor/Private/ViewModels/Stack/NiagaraStackGraphUtilities.cpp
@@ -2617,11 +2617,12 @@ UNiagaraNodeFunctionCall* FNiagaraStackGraphUtilities::AddScriptModuleToStack(co
 	}
 
 	// If specified, find the nearest index to TargetIndex that satisfies the new module script's order dependencies.
-	int32 FinalTargetIndex = TargetIndex;
+	int32 BestFoundIndex = INDEX_NONE;
 	if (bFixupTargetIndex)
 	{
-		FinalTargetIndex = DependencyUtilities::FindBestIndexForModuleInStack(*NewModuleNode, *Graph);
+		BestFoundIndex = DependencyUtilities::FindBestIndexForModuleInStack(*NewModuleNode, *TargetOutputNode);
 	}
+	int32 FinalTargetIndex = BestFoundIndex != INDEX_NONE ? BestFoundIndex : TargetIndex;
 
 	ConnectModuleNode(*NewModuleNode, *TargetOutputNode, FinalTargetIndex);
 	return NewModuleNode;
@@ -4556,7 +4557,7 @@ void FNiagaraStackGraphUtilities::GetEmitterHandleAndCompiledScriptsForStackNode
 	}
 }
 
-int32 FNiagaraStackGraphUtilities::DependencyUtilities::FindBestIndexForModuleInStack(UNiagaraNodeFunctionCall& ModuleNode, UNiagaraGraph& EmitterScriptGraph)
+int32 FNiagaraStackGraphUtilities::DependencyUtilities::FindBestIndexForModuleInStack(UNiagaraNodeFunctionCall& ModuleNode, const UNiagaraNodeOutput& TargetOutputNode)
 {
 	// Check if the new module node has any dependencies to begin with. If not, early exit.
 	FVersionedNiagaraScriptData* ScriptData = ModuleNode.GetScriptData();
@@ -4565,26 +4566,44 @@ int32 FNiagaraStackGraphUtilities::DependencyUtilities::FindBestIndexForModuleIn
 		return INDEX_NONE;
 	}
 
+	// Get and validate the owning graph.
+	UNiagaraGraph* Graph = ModuleNode.GetNiagaraGraph();
+	if (Graph == nullptr || Graph->Nodes.Contains(&ModuleNode) == false)
+	{
+		return INDEX_NONE;
+	}
+
 	// Get the Emitter and System the emitter script script graph is outered to.
-	UNiagaraSystem* System = EmitterScriptGraph.GetTypedOuter<UNiagaraSystem>();
-	FVersionedNiagaraEmitter OuterEmitter = EmitterScriptGraph.GetOwningEmitter();
-	if (System == nullptr || OuterEmitter.Emitter == nullptr)
+	UNiagaraSystem* System = ModuleNode.GetTypedOuter<UNiagaraSystem>();
+	FVersionedNiagaraEmitter OuterEmitter = Graph->GetOwningEmitter();
+	if (System == nullptr && OuterEmitter.Emitter == nullptr)
 	{
 		return INDEX_NONE;
 	}
 
-	// Get the stack module data for the emitter stack to find dependencies.
+	// Get the stack module data for the stack to find dependencies.
 	TSharedPtr<FNiagaraSystemViewModel> SystemViewModel = TNiagaraViewModelManager<UNiagaraSystem, FNiagaraSystemViewModel>::GetExistingViewModelForObject(System);
 	if (!ensureMsgf(SystemViewModel.IsValid(), TEXT("Failed to get systemviewmodel for valid system when getting best index in stack for module!")))
 	{
 		return INDEX_NONE;
 	}
-	TSharedPtr<FNiagaraEmitterHandleViewModel> EmitterHandleViewModel = SystemViewModel->GetEmitterHandleViewModelForEmitter(OuterEmitter);
-	if (!ensureMsgf(EmitterHandleViewModel.IsValid(), TEXT("Failed to get emitterhandleviewmodel for valid emitter when getting best index in stack for module!")))
+
+	FGuid EmitterHandleId;
+	if (OuterEmitter.Emitter != nullptr)
 	{
-		return INDEX_NONE;
+		TSharedPtr<FNiagaraEmitterHandleViewModel> EmitterHandleViewModel = SystemViewModel->GetEmitterHandleViewModelForEmitter(OuterEmitter);
+		if (!ensureMsgf(EmitterHandleViewModel.IsValid(), TEXT("Failed to get emitterhandleviewmodel for valid emitter when getting best index in stack for module!")))
+		{
+			return INDEX_NONE;
+		}
+		else
+		{
+			EmitterHandleId = EmitterHandleViewModel->GetId();
+		}
 	}
-	const TArray<FNiagaraStackModuleData>& StackModuleData = SystemViewModel->GetStackModuleDataByEmitterHandleId(EmitterHandleViewModel->GetId());
+	
+	// Calling this with an invalid handle will retrieve the system stack module data.
+	const TArray<FNiagaraStackModuleData>& StackModuleData = SystemViewModel->GetStackModuleDataByEmitterHandleId(EmitterHandleId);
 
 
 	// Find the greatest and least indices for the stack first.
@@ -4592,6 +4611,11 @@ int32 FNiagaraStackGraphUtilities::DependencyUtilities::FindBestIndexForModuleIn
 	int32 GreatestIndex = INDEX_NONE;
 	for (const FNiagaraStackModuleData& CurrentStackModuleData : StackModuleData)
 	{
+		// Stack module data indices are only valid for a specific script which is identified by the usage and usage id.
+		if (CurrentStackModuleData.Usage != TargetOutputNode.GetUsage() || CurrentStackModuleData.UsageId != TargetOutputNode.GetUsageId())
+		{
+			continue;
+		}
 		int32 Index = CurrentStackModuleData.Index;
 		LeastIndex = LeastIndex > Index ? Index : LeastIndex;
 		GreatestIndex = GreatestIndex < Index ? Index : GreatestIndex;
@@ -4618,6 +4642,11 @@ int32 FNiagaraStackGraphUtilities::DependencyUtilities::FindBestIndexForModuleIn
 
 	for (const FNiagaraStackModuleData& CurrentStackModuleData : StackModuleData)
 	{
+		// Stack module data indices are only valid for a specific script which is identified by the usage and usage id.
+		if (CurrentStackModuleData.Usage != TargetOutputNode.GetUsage() || CurrentStackModuleData.UsageId != TargetOutputNode.GetUsageId())
+		{
+			continue;
+		}
 		for (const FNiagaraModuleDependency& RequiredDependency : NewModuleScriptRequiredDependencies)
 		{
 			if (DoesStackModuleProvideDependency(CurrentStackModuleData, RequiredDependency, *GetOutputNodeForStackModuleData(CurrentStackModuleData)))
@@ -4675,6 +4704,11 @@ int32 FNiagaraStackGraphUtilities::DependencyUtilities::FindBestIndexForModuleIn
 
 	for (const FNiagaraStackModuleData& CurrentStackModuleData : StackModuleData)
 	{
+		// Stack module data indices are only valid for a specific script which is identified by the usage and usage id.
+		if (CurrentStackModuleData.Usage != TargetOutputNode.GetUsage() || CurrentStackModuleData.UsageId != TargetOutputNode.GetUsageId())
+		{
+			continue;
+		}
 		for (const FNiagaraModuleDependency& ProvidedDependency : GetStackModuleDataRequiredDependenciesBeingProvided(CurrentStackModuleData))
 		{
 			if (ProvidedDependency.Type == ENiagaraModuleDependencyType::PostDependency)
diff --git a/Engine/Plugins/FX/Niagara/Source/NiagaraEditor/Private/ViewModels/Stack/NiagaraStackModuleItem.cpp b/Engine/Plugins/FX/Niagara/Source/NiagaraEditor/Private/ViewModels/Stack/NiagaraStackModuleItem.cpp
--- a/Engine/Plugins/FX/Niagara/Source/NiagaraEditor/Private/ViewModels/Stack/NiagaraStackModuleItem.cpp
+++ b/Engine/Plugins/FX/Niagara/Source/NiagaraEditor/Private/ViewModels/Stack/NiagaraStackModuleItem.cpp
@@ -142,7 +142,10 @@ void UNiagaraStackModuleItem::RefreshChildrenInternal(const TArray<UNiagaraStack
 	bIsScratchModuleCache.Reset();
 	DisplayNameCache.Reset();
 
-	if (FunctionCallNode != nullptr && (FunctionCallNode->HasValidScriptAndGraph() || FunctionCallNode->Signature.IsValid()))
+	if (FunctionCallNode != nullptr &&
+		(FunctionCallNode->HasValidScriptAndGraph() || FunctionCallNode->Signature.IsValid()) &&
+		FunctionCallNode->GetGraph() != nullptr && 
+		FunctionCallNode->GetGraph()->Nodes.Contains(FunctionCallNode))
 	{
 		// Determine if meta-data requires that we add our own refresh button here.
 		if (FunctionCallNode->HasValidScriptAndGraph())
@@ -204,11 +207,11 @@ void UNiagaraStackModuleItem::RefreshChildrenInternal(const TArray<UNiagaraStack
 				NewChildren.Add(EmptyAssignmentNodeMessage);
 			}
 		}
+		RefreshIsEnabled();
+		RefreshIssues(NewIssues);
 	}
 
-	RefreshIsEnabled();
 	Super::RefreshChildrenInternal(CurrentChildren, NewChildren, NewIssues);
-	RefreshIssues(NewIssues);
 }
 
 
@@ -749,7 +752,7 @@ namespace NiagaraStackModuleItemIssues {
 		}
 
 		FVersionedNiagaraScriptData* ScriptData = SourceModuleNode.FunctionScript->GetScriptData(SourceModuleNode.SelectedScriptVersion);
-		int32 ModuleIndex = SourceStackModuleData.IndexOfByPredicate([&SourceModuleNode](const FNiagaraStackModuleData& ModuleData) { return ModuleData.ModuleNode == &SourceModuleNode || ModuleData.ModuleNode->GetFunctionName() == SourceModuleNode.GetFunctionName(); });
+		int32 ModuleIndex = SourceStackModuleData.IndexOfByPredicate([&SourceModuleNode](const FNiagaraStackModuleData& ModuleData) { return ModuleData.ModuleNode == &SourceModuleNode; });
 		if (ensureMsgf(ModuleIndex != INDEX_NONE, TEXT("In system %s, module %s (%s) did not exist in the stack module data."),
 			*SourceSystemViewModel->GetSystem().GetPathName(), *SourceModuleNode.GetFunctionName(), *SourceModuleNode.GetName()))
 		{
diff --git a/Engine/Plugins/FX/Niagara/Source/NiagaraEditor/Public/ViewModels/Stack/NiagaraStackGraphUtilities.h b/Engine/Plugins/FX/Niagara/Source/NiagaraEditor/Public/ViewModels/Stack/NiagaraStackGraphUtilities.h
--- a/Engine/Plugins/FX/Niagara/Source/NiagaraEditor/Public/ViewModels/Stack/NiagaraStackGraphUtilities.h
+++ b/Engine/Plugins/FX/Niagara/Source/NiagaraEditor/Public/ViewModels/Stack/NiagaraStackGraphUtilities.h
@@ -377,6 +377,6 @@ namespace FNiagaraStackGraphUtilities
 
 		void GetModuleScriptAssetsByDependencyProvided(FName DependencyName, TOptional<ENiagaraScriptUsage> RequiredUsage, TArray<FAssetData>& OutAssets);
 
-		int32 FindBestIndexForModuleInStack(UNiagaraNodeFunctionCall& ModuleNode, UNiagaraGraph& EmitterScriptGraph);
+		int32 FindBestIndexForModuleInStack(UNiagaraNodeFunctionCall& ModuleNode, const UNiagaraNodeOutput& TargetOutputNode);
 	}
 }
