- Fix for toolbar buttons been broken in emitter editor

diff --git a/Engine/Plugins/FX/Niagara/Source/NiagaraEditor/Private/Toolkits/NiagaraSystemToolkit.cpp b/Engine/Plugins/FX/Niagara/Source/NiagaraEditor/Private/Toolkits/NiagaraSystemToolkit.cpp
--- a/Engine/Plugins/FX/Niagara/Source/NiagaraEditor/Private/Toolkits/NiagaraSystemToolkit.cpp
+++ b/Engine/Plugins/FX/Niagara/Source/NiagaraEditor/Private/Toolkits/NiagaraSystemToolkit.cpp
@@ -219,6 +219,9 @@ void FNiagaraSystemToolkit::InitializeWithEmitter(const EToolkitMode::Type Mode,
 	UNiagaraSystemFactoryNew::InitializeSystem(System, true);
 	System->EnsureFullyLoaded();
 
+	// order of registering commands matters. SetupCommands before InitAssetEditor will make the toolkit prioritize niagara commands
+	SetupCommands();
+
 	InEmitter.UpdateEmitterAfterLoad();
 	Emitter = &InEmitter;
 
