- Fix for incorrect diagnostic naming in WM

diff --git a/Engine/Plugins/FX/Niagara/Source/Niagara/Private/NiagaraWorldManager.cpp b/Engine/Plugins/FX/Niagara/Source/Niagara/Private/NiagaraWorldManager.cpp
--- a/Engine/Plugins/FX/Niagara/Source/Niagara/Private/NiagaraWorldManager.cpp
+++ b/Engine/Plugins/FX/Niagara/Source/Niagara/Private/NiagaraWorldManager.cpp
@@ -327,12 +327,12 @@ FString FNiagaraWorldManagerTickFunction::DiagnosticMessage()
 {
 	static const UEnum* EnumType = FindObjectChecked<UEnum>(nullptr, TEXT("/Script/Engine.ETickingGroup"));
 
-	return TEXT("FParticleSystemManager::Tick(") + EnumType->GetNameStringByIndex(static_cast<uint32>(TickGroup)) + TEXT(")");
+	return TEXT("FNiagaraWorldManagerTickFunction::Tick(") + EnumType->GetNameStringByIndex(static_cast<uint32>(TickGroup)) + TEXT(")");
 }
 
 FName FNiagaraWorldManagerTickFunction::DiagnosticContext(bool bDetailed)
 {
-	return FName(TEXT("ParticleSystemManager"));
+	return FName(TEXT("FNiagaraWorldManagerTickFunction"));
 }
 
 //////////////////////////////////////////////////////////////////////////
