- Enable Cascade performance counters on debug hud

diff --git a/Engine/Plugins/FX/Niagara/Source/Niagara/Private/NiagaraDebugHud.cpp b/Engine/Plugins/FX/Niagara/Source/Niagara/Private/NiagaraDebugHud.cpp
--- a/Engine/Plugins/FX/Niagara/Source/Niagara/Private/NiagaraDebugHud.cpp
+++ b/Engine/Plugins/FX/Niagara/Source/Niagara/Private/NiagaraDebugHud.cpp
@@ -3441,7 +3441,7 @@ bool FNiagaraDebugHUDStatsListener::Tick()
 		TWeakObjectPtr<const UFXSystemAsset> WeakSystem = HUDStatsPair.Key;
 		TSharedPtr<FNiagaraDebugHUDPerfStats>& HUDStats = HUDStatsPair.Value;
 
-		const UNiagaraSystem* System = Cast<const UNiagaraSystem>(WeakSystem.Get());
+		const UFXSystemAsset* System = WeakSystem.Get();
 		if(System == nullptr)
 		{
 			continue;
@@ -3540,7 +3540,7 @@ void FNiagaraDebugHUDStatsListener::TickRT()
 		TWeakObjectPtr<const UFXSystemAsset> WeakSystem = HUDStatsPair.Key;
 		TSharedPtr<FNiagaraDebugHUDPerfStats>& HUDStats = HUDStatsPair.Value;
 
-		const UNiagaraSystem* System = Cast<const UNiagaraSystem>(WeakSystem.Get());
+		const UFXSystemAsset* System = WeakSystem.Get();
 		if (System == nullptr)
 		{
 			continue;
