Cleaned up logging when exceeding primitive counts with the NiagaraDataInterfaceRigidMeshCollisionQuery

diff --git a/Engine/Plugins/FX/Niagara/Source/Niagara/Private/NiagaraDataInterfaceRigidMeshCollisionQuery.cpp b/Engine/Plugins/FX/Niagara/Source/Niagara/Private/NiagaraDataInterfaceRigidMeshCollisionQuery.cpp
--- a/Engine/Plugins/FX/Niagara/Source/Niagara/Private/NiagaraDataInterfaceRigidMeshCollisionQuery.cpp
+++ b/Engine/Plugins/FX/Niagara/Source/Niagara/Private/NiagaraDataInterfaceRigidMeshCollisionQuery.cpp
@@ -635,8 +635,15 @@ bool UpdateInternalArrays(
 	FVector LWCTile,
 	bool bFullUpdate,
 	TArray<FNDIRigidMeshCollisionData::FComponentBodyCount>& BodyCounts,
-	FNDIRigidMeshCollisionArrays* OutAssetArrays)
+	FNDIRigidMeshCollisionArrays* OutAssetArrays,
+	uint32& TotalBoxCount,
+	uint32& TotalSphereCount,
+	uint32& TotalCapsuleCount)
 {
+	TotalBoxCount = 0;
+	TotalSphereCount = 0;
+	TotalCapsuleCount = 0;
+
 	if (OutAssetArrays == nullptr || OutAssetArrays->ElementOffsets.NumElements >= OutAssetArrays->MaxPrimitives)
 	{
 		return false;
@@ -679,16 +686,11 @@ bool UpdateInternalArrays(
 
 	TArray<FNDIRigidMeshCollisionData::FComponentBodyCount> CurrentBodyCounts;
 
-	uint32 TotalBoxCount = 0;
-	uint32 TotalSphereCount = 0;
-	uint32 TotalCapsuleCount = 0;
-
 	CountCollisionPrimitives(StaticMeshView, CurrentBodyCounts, TotalBoxCount, TotalSphereCount, TotalCapsuleCount);
 	CountCollisionPrimitives(SkeletalMeshView, CurrentBodyCounts, TotalBoxCount, TotalSphereCount, TotalCapsuleCount);
 
 	if ((TotalBoxCount + TotalSphereCount + TotalCapsuleCount) >= OutAssetArrays->MaxPrimitives)
 	{
-		UE_LOG(LogRigidMeshCollision, Error, TEXT("Number of Collision DI primitives is higher than the %d limit.  Please increase it."), OutAssetArrays->MaxPrimitives);
 		return false;
 	}
 
@@ -955,6 +957,10 @@ PRAGMA_ENABLE_DEPRECATION_WARNINGS
 		TConstArrayView<UStaticMeshComponent*> StaticMeshView = MakeArrayView(StaticMeshes.GetData(), StaticMeshes.Num());
 		TConstArrayView<USkeletalMeshComponent*> SkeletalMeshView = MakeArrayView(SkeletalMeshes.GetData(), SkeletalMeshes.Num());
 
+		uint32 TotalBoxCount = 0;
+		uint32 TotalSphereCount = 0;
+		uint32 TotalCapsuleCount = 0;
+
 		const bool bArraysUpdated = UpdateInternalArrays(
 			SystemInstance,
 			StaticMeshView,
@@ -962,11 +968,30 @@ PRAGMA_ENABLE_DEPRECATION_WARNINGS
 			FVector(SystemInstance->GetLWCTile()),
 			bRequiresFullUpdate,
 			MeshBodyCounts,
-			AssetArrays.Get());
+			AssetArrays.Get(),
+			TotalBoxCount,
+			TotalSphereCount,
+			TotalCapsuleCount);
 
 		if (bArraysUpdated)
 		{
 			bRequiresFullUpdate = false;
+			bExceedingComponentLimits = false;
+		}
+		else
+		{
+			if (TotalBoxCount + TotalSphereCount + TotalCapsuleCount > AssetArrays->MaxPrimitives)
+			{
+				if (!bExceedingComponentLimits)
+				{
+					UE_LOG(LogRigidMeshCollision, Warning, TEXT("Number of Collision DI primitives (%d boxes, %d spheres, %d capsules) is higher than the %d limit for System[%s].  Please increase it."),
+						TotalBoxCount, TotalSphereCount, TotalCapsuleCount,
+						AssetArrays->MaxPrimitives,
+						*GetNameSafe(SystemInstance->GetSystem()));
+
+					bExceedingComponentLimits = true;
+				}
+			}
 		}
 	}
 	else
@@ -1112,7 +1137,14 @@ void UNiagaraDataInterfaceRigidMeshCollisionQuery::DrawDebugHud(FNDIDrawDebugHud
 	const uint32 SphereCount = ElementOffsets.CapsuleOffset - ElementOffsets.SphereOffset;
 	const uint32 CapsuleCount = ElementOffsets.NumElements - ElementOffsets.CapsuleOffset;
 
-	DebugHudContext.GetOutputString().Appendf(TEXT("Boxes(%d) Spheres(%d) Capsules(%d)"), BoxCount, SphereCount, CapsuleCount);
+	if (InstanceData_GT->bExceedingComponentLimits)
+	{
+		DebugHudContext.GetOutputString().Appendf(TEXT("EXCEEDED COMPONENT LIMIT - %d"), InstanceData_GT->AssetArrays->MaxPrimitives);
+	}
+	else
+	{
+		DebugHudContext.GetOutputString().Appendf(TEXT("Boxes(%d) Spheres(%d) Capsules(%d)"), BoxCount, SphereCount, CapsuleCount);
+	}
 
 	auto GetTransformFromArray = [&](TArray<FVector4f> &TransformArray, int32 ElementIndex)
 	{
diff --git a/Engine/Plugins/FX/Niagara/Source/Niagara/Public/NiagaraDataInterfaceRigidMeshCollisionQuery.h b/Engine/Plugins/FX/Niagara/Source/Niagara/Public/NiagaraDataInterfaceRigidMeshCollisionQuery.h
--- a/Engine/Plugins/FX/Niagara/Source/Niagara/Public/NiagaraDataInterfaceRigidMeshCollisionQuery.h
+++ b/Engine/Plugins/FX/Niagara/Source/Niagara/Public/NiagaraDataInterfaceRigidMeshCollisionQuery.h
@@ -170,6 +170,9 @@ struct FNDIRigidMeshCollisionData
 
 	/** Indicates that a full update of the arrays is required */
 	bool bRequiresFullUpdate = false;
+
+	/** Indicates that we are currently in a failure state because we've found more components in the actors than we can support */
+	bool bExceedingComponentLimits = false;
 };
 
 /** Data Interface used to collide against static meshes - whether it is the mesh distance field or a physics asset's collision primitive */
