- Temporary workaround for validation on platforms that strip the system on cook

diff --git a/Engine/Plugins/FX/Niagara/Source/Niagara/Classes/NiagaraSystem.h b/Engine/Plugins/FX/Niagara/Source/Niagara/Classes/NiagaraSystem.h
--- a/Engine/Plugins/FX/Niagara/Source/Niagara/Classes/NiagaraSystem.h
+++ b/Engine/Plugins/FX/Niagara/Source/Niagara/Classes/NiagaraSystem.h
@@ -793,6 +793,8 @@ public:
 
 	/** Updates the rapid iteration parameters for all scripts referenced by the system. */
 	NIAGARA_API void PrepareRapidIterationParametersForCompilation();
+
+	bool AllowValidation() const { return bAllowValidation; }
 #endif
 
 	/** Can we run the code only system state path, i.e. we don't need to invoke the VVM / store per instance data set? */
@@ -1044,6 +1046,9 @@ protected:
 	uint32 bSystemStateFastPathEnabled : 1 = true;
 
 #if WITH_EDITORONLY_DATA
+	/* In some cases we can not run validation as the data has been stripped.  This is a temporary workaround to skip validation in those cases. */
+	uint32 bAllowValidation : 1 = true;
+
 	/** Messages associated with the System asset. */
 	UPROPERTY()
 	TMap<FGuid, TObjectPtr<UNiagaraMessageDataBase>> MessageKeyToMessageMap_DEPRECATED;
diff --git a/Engine/Plugins/FX/Niagara/Source/Niagara/Private/NiagaraSystem.cpp b/Engine/Plugins/FX/Niagara/Source/Niagara/Private/NiagaraSystem.cpp
--- a/Engine/Plugins/FX/Niagara/Source/Niagara/Private/NiagaraSystem.cpp
+++ b/Engine/Plugins/FX/Niagara/Source/Niagara/Private/NiagaraSystem.cpp
@@ -540,6 +540,7 @@ void UNiagaraSystem::UpdateSystemAfterLoad()
 	const bool bTargetRequiresAvData = WillNeedAudioVisualData();
 	if (bIsDedicatedServer || !bTargetRequiresAvData)
 	{
+		bAllowValidation = false;
 		ResetToEmptySystem();
 		return;
 	}
