- Add missing CSV stat around finalize tasks

diff --git a/Engine/Plugins/FX/Niagara/Source/Niagara/Private/NiagaraSystemSimulation.cpp b/Engine/Plugins/FX/Niagara/Source/Niagara/Private/NiagaraSystemSimulation.cpp
--- a/Engine/Plugins/FX/Niagara/Source/Niagara/Private/NiagaraSystemSimulation.cpp
+++ b/Engine/Plugins/FX/Niagara/Source/Niagara/Private/NiagaraSystemSimulation.cpp
@@ -449,6 +449,7 @@ struct FNiagaraSystemInstanceFinalizeTask
 
 	void DoTask(ENamedThreads::Type CurrentThread, const FGraphEventRef& MyCompletionGraphEvent)
 	{
+		CSV_SCOPED_TIMING_STAT_EXCLUSIVE(Effects);
 		check(CurrentThread == ENamedThreads::GameThread);
 
 		if ( TickHandlingMode == ENiagaraGPUTickHandlingMode::GameThreadBatched )
