- Basic Niagara Component leak detector
- Looks for constant growth in component counts / active components
- This will be refined over time and is likely to not catch all cases / produce false positives in some cases

diff --git a/Engine/Plugins/FX/Niagara/Source/Niagara/Private/NiagaraComponentLeakDetector.cpp b/Engine/Plugins/FX/Niagara/Source/Niagara/Private/NiagaraComponentLeakDetector.cpp
new file mode 100644
--- /dev/null
+++ b/Engine/Plugins/FX/Niagara/Source/Niagara/Private/NiagaraComponentLeakDetector.cpp
@@ -0,0 +1,204 @@
+// Copyright Epic Games, Inc. All Rights Reserved.
+
+#include "NiagaraComponentLeakDetector.h"
+#include "NiagaraDebugHud.h"
+#include "NiagaraSystem.h"
+#include "NiagaraWorldManager.h"
+
+#if WITH_NIAGARA_LEAK_DETECTOR
+
+DEFINE_LOG_CATEGORY_STATIC(LogNiagaraLeakDetector, Log, All);
+
+namespace NiagaraComponentLeakDetectorPrivate
+{
+	static bool		GEnabled = false;
+	static int32	GImmediateReport = 2;
+	static int32	GGCReport = 1;
+	static float	GTickDeltaSeconds = 1.0f;
+	static int32	GGrowthCountThreshold = 16;
+	static float	GDebugMessageTime = 5.0f;
+
+	static FAutoConsoleVariableRef CVarEnabled(
+		TEXT("fx.Niagara.LeakDetector.Enabled"),
+		GEnabled,
+		TEXT("Enables or disables the leak detector."),
+		ECVF_Default
+	);
+
+	static FAutoConsoleVariableRef CVarTickDeltaSeconds(
+		TEXT("fx.Niagara.LeakDetector.TickDeltaSeconds"),
+		GTickDeltaSeconds,
+		TEXT("The time in seconds that must pass before we sample the component information."),
+		ECVF_Default
+	);
+
+	static FAutoConsoleVariableRef CVarGrowthCountThreshold(
+		TEXT("fx.Niagara.LeakDetector.GrowthCountThreshold"),
+		GGrowthCountThreshold,
+		TEXT("We need to see growth this many times without a drop in count before we consider it a leak."),
+		ECVF_Default
+	);
+
+	static FAutoConsoleVariableRef CVarImmediateReport(
+		TEXT("fx.Niagara.LeakDetector.ImmediateReport"),
+		GImmediateReport,
+		TEXT("Controls how we report leaks as we scan information.\n")
+		TEXT("0 - Never report during gameplay.\n")
+		TEXT("1 - Report all leaks.")
+		TEXT("2 - Report active leaks only. (default).\n"),
+		ECVF_Default
+	);
+
+	static FAutoConsoleVariableRef CVarGGCReport(
+		TEXT("fx.Niagara.LeakDetector.GCReport"),
+		GGCReport,
+		TEXT("Controls how we report leak after GC\n")
+		TEXT("0 - Never report after GC.\n")
+		TEXT("1 - Report all leaks. (default)")
+		TEXT("2 - Report active leaks only.\n"),
+		ECVF_Default
+	);
+
+	static FAutoConsoleVariableRef CVarDebugMessageTime(
+		TEXT("fx.Niagara.LeakDetector.DebugMessageTime"),
+		GDebugMessageTime,
+		TEXT("Time we display the debug message for on screen."),
+		ECVF_Default
+	);
+
+	void AddLeakWarning(UWorld* World, FName MessageKey, const FString& Message)
+	{
+		UE_LOG(LogNiagaraLeakDetector, Warning, TEXT("%s"), *Message);
+#if WITH_NIAGARA_DEBUGGER
+		if (FNiagaraWorldManager* WorldManager = FNiagaraWorldManager::Get(World))
+		{
+			if (FNiagaraDebugHud* DebugHud = WorldManager->GetNiagaraDebugHud())
+			{
+				DebugHud->AddMessage(MessageKey, FNiagaraDebugMessage(ENiagaraDebugMessageType::Warning, Message, GDebugMessageTime));
+			}
+		}
+#endif
+	}
+}
+
+void FNiagaraComponentLeakDetector::Tick(UWorld* World)
+{
+	using namespace NiagaraComponentLeakDetectorPrivate;
+
+	if (!GEnabled)
+	{
+		return;
+	}
+
+	const double CurrentTime = FPlatformTime::Seconds();
+	if (CurrentTime < NextUpdateTime)
+	{
+		return;
+	}
+	NextUpdateTime = CurrentTime + double(GTickDeltaSeconds);
+
+	for (TObjectIterator<UNiagaraComponent> ComponentIt; ComponentIt; ++ComponentIt)
+	{
+		UNiagaraComponent* Component = *ComponentIt;
+		UNiagaraSystem* System = Component ? Component->GetAsset() : nullptr;
+		if (!System)
+		{
+			continue;
+		}
+
+		if ( Component->GetWorld() != World || Component->PoolingMethod == ENCPoolMethod::FreeInPool)
+		{
+			continue;
+		}
+
+		const FName SystemName = System->GetFName();
+		FSystemData& SystemData = PerSystemData.FindOrAdd(SystemName);
+		++SystemData.TotalCurrCount;
+		SystemData.ActiveCurrCount += Component->IsActive();
+	}
+
+	const bool bReportTotalLeaks = GImmediateReport == 1;
+	const bool bReportActiveLeaks = GImmediateReport != 0;
+
+	for (auto SystemDataIt=PerSystemData.CreateIterator(); SystemDataIt; ++SystemDataIt)
+	{
+		const FName SystemName = SystemDataIt.Key();
+		FSystemData& SystemData = SystemDataIt.Value();
+
+		// Monitor total count
+		SystemData.TotalGrowthCounter += SystemData.TotalCurrCount > SystemData.TotalPrevCount ? 1 : 0;
+		SystemData.TotalShrinkCounter += SystemData.TotalCurrCount < SystemData.TotalPrevCount ? 1 : 0;
+		SystemData.TotalPrevCount = SystemData.TotalCurrCount;
+		SystemData.TotalCurrCount = 0;
+
+		// Monitor active count
+		SystemData.ActiveGrowthCounter += SystemData.ActiveCurrCount > SystemData.ActivePrevCount ? 1 : 0;
+		SystemData.ActiveShrinkCounter += SystemData.ActiveCurrCount < SystemData.ActivePrevCount ? 1 : 0;
+		SystemData.ActivePrevCount = SystemData.ActiveCurrCount;
+		SystemData.ActiveCurrCount = 0;
+
+		// Report any leaks
+		ReportLeak(World, SystemName, SystemData, bReportTotalLeaks, bReportActiveLeaks);
+	}
+}
+
+void FNiagaraComponentLeakDetector::ReportLeaks(UWorld* World)
+{
+	using namespace NiagaraComponentLeakDetectorPrivate;
+
+	if (!GEnabled || GGCReport == 0)
+	{
+		return;
+	}
+
+	const bool bReportTotalLeaks = GGCReport == 1;
+	const bool bReportActiveLeaks = GGCReport != 0;
+
+	for (auto SystemDataIt=PerSystemData.CreateIterator(); SystemDataIt; ++SystemDataIt)
+	{
+		const FName SystemName = SystemDataIt.Key();
+		FSystemData& SystemData = SystemDataIt.Value();
+		ReportLeak(World, SystemName, SystemData, bReportTotalLeaks, bReportActiveLeaks);
+	}
+}
+
+void FNiagaraComponentLeakDetector::ReportLeak(UWorld* World, FName SystemName, FSystemData& SystemData, bool bReportTotalLeaks, bool bReportActiveLeaks)
+{
+	using namespace NiagaraComponentLeakDetectorPrivate;
+
+	if (bReportTotalLeaks)
+	{
+		if (!SystemData.bTotalHasWarned)
+		{
+			if (SystemData.TotalShrinkCounter == 0 && SystemData.TotalGrowthCounter >= GGrowthCountThreshold)
+			{
+				SystemData.bTotalHasWarned = true;
+				AddLeakWarning(World, SystemName, FString::Printf(TEXT("Potential componment leak System(%s) (Total:%d), please investigate."), *SystemName.ToString(), SystemData.TotalPrevCount));
+			}
+		}
+		else if (!SystemData.bTotalFalseWarning && (SystemData.TotalShrinkCounter > 0))
+		{
+			SystemData.bTotalFalseWarning = true;
+			AddLeakWarning(World, SystemName, FString::Printf(TEXT("Potential invalid component leak reported for System(%s) (Total:%d)."), *SystemName.ToString(), SystemData.TotalPrevCount));
+		}
+	}
+
+	if (bReportActiveLeaks)
+	{
+		if (!SystemData.bActiveHasWarned)
+		{
+			if (SystemData.ActiveShrinkCounter == 0 && SystemData.ActiveGrowthCounter >= GGrowthCountThreshold)
+			{
+				SystemData.bActiveHasWarned = true;
+				AddLeakWarning(World, SystemName, FString::Printf(TEXT("Potential Active Component leak System(%s) (Active:%d), please investigate."), *SystemName.ToString(), SystemData.ActivePrevCount));
+			}
+		}
+		else if ( !SystemData.bActiveFalseWarning && (SystemData.ActiveShrinkCounter > 0) )
+		{
+			SystemData.bActiveFalseWarning = true;
+			AddLeakWarning(World, SystemName, FString::Printf(TEXT("Potential invalid active component leak reported for System(% s) (Active: % d)."), *SystemName.ToString(), SystemData.ActivePrevCount));
+		}
+	}
+}
+
+#endif //WITH_NIAGARA_LEAK_DETECTOR
diff --git a/Engine/Plugins/FX/Niagara/Source/Niagara/Private/NiagaraComponentLeakDetector.h b/Engine/Plugins/FX/Niagara/Source/Niagara/Private/NiagaraComponentLeakDetector.h
new file mode 100644
--- /dev/null
+++ b/Engine/Plugins/FX/Niagara/Source/Niagara/Private/NiagaraComponentLeakDetector.h
@@ -0,0 +1,44 @@
+// Copyright Epic Games, Inc. All Rights Reserved.
+
+#pragma once
+
+#include "NiagaraCommon.h"
+#include "NiagaraComponent.h"
+
+#if WITH_NIAGARA_LEAK_DETECTOR
+
+// Very simple leak detector, will be revised over time.
+class FNiagaraComponentLeakDetector
+{
+	struct FSystemData
+	{
+		static constexpr uint32 HistoryLength = 4;
+
+		bool	bTotalHasWarned = false;
+		bool	bTotalFalseWarning = false;
+		int32	TotalGrowthCounter = 0;
+		int32	TotalShrinkCounter = 0;
+		uint32	TotalCurrCount = 0;
+		uint32	TotalPrevCount = 0;
+
+		bool	bActiveHasWarned = false;
+		bool	bActiveFalseWarning = false;
+		int32	ActiveGrowthCounter = 0;
+		int32	ActiveShrinkCounter = 0;
+		uint32	ActiveCurrCount = 0;
+		uint32	ActivePrevCount = 0;
+	};
+
+public:
+	void Tick(UWorld* World);
+	void ReportLeaks(UWorld* World);
+
+private:
+	void ReportLeak(UWorld* World, FName SystemName, FSystemData& SystemData, bool bReportTotalLeaks, bool bReportActiveLeaks);
+
+private:
+	double						NextUpdateTime = 0.0;
+	TMap<FName, FSystemData>	PerSystemData;
+};
+
+#endif //WITH_NIAGARA_LEAK_DETECTOR
diff --git a/Engine/Plugins/FX/Niagara/Source/Niagara/Private/NiagaraWorldManager.cpp b/Engine/Plugins/FX/Niagara/Source/Niagara/Private/NiagaraWorldManager.cpp
--- a/Engine/Plugins/FX/Niagara/Source/Niagara/Private/NiagaraWorldManager.cpp
+++ b/Engine/Plugins/FX/Niagara/Source/Niagara/Private/NiagaraWorldManager.cpp
@@ -20,6 +20,7 @@
 #include "NiagaraDebugHud.h"
 #include "NiagaraGpuComputeDispatchInterface.h"
 #include "NiagaraSimpleObjectPool.h"
+#include "NiagaraComponentLeakDetector.h"
 #include "Particles/FXBudget.h"
 #include "NiagaraCullProxyComponent.h"
 #include "GameDelegates.h"
@@ -403,6 +404,12 @@ void FNiagaraWorldManager::Init(UWorld* InWorld)
 #if WITH_NIAGARA_DEBUGGER
 	NiagaraDebugHud.Reset(new FNiagaraDebugHud(World));
 #endif
+#if WITH_NIAGARA_LEAK_DETECTOR
+	if ( World->IsGameWorld() )
+	{
+		ComponentLeakDetector.Reset(new FNiagaraComponentLeakDetector());
+	}
+#endif
 
 	// Make sure we update our component settings, this includes ban lists, etc
 	FNiagaraComponentSettings::UpdateSettings();
@@ -794,6 +801,13 @@ void FNiagaraWorldManager::PostGarbageCollect()
 {
 	//Clear out and scalability managers who's EffectTypes have been GCd.
 	while (ScalabilityManagers.Remove(nullptr)) {}
+
+#if WITH_NIAGARA_LEAK_DETECTOR
+	if (ComponentLeakDetector.IsValid() && HasActiveWorld())
+	{
+		ComponentLeakDetector->ReportLeaks(World);
+	}
+#endif
 }
 
 void FNiagaraWorldManager::PreGarbageCollectBeginDestroy()
@@ -1013,6 +1027,13 @@ void FNiagaraWorldManager::TickStart(float DeltaSeconds)
 	SCOPE_CYCLE_COUNTER(STAT_NiagaraWorldManTick);
 
 	DataChannelManager->BeginFrame(DeltaSeconds);
+
+#if WITH_NIAGARA_LEAK_DETECTOR
+	if (ComponentLeakDetector.IsValid() && HasActiveWorld())
+	{
+		ComponentLeakDetector->Tick(World);
+	}
+#endif
 }
 
 void FNiagaraWorldManager::PreActorTick(ELevelTick InLevelTick, float InDeltaSeconds)
diff --git a/Engine/Plugins/FX/Niagara/Source/Niagara/Public/NiagaraCommon.h b/Engine/Plugins/FX/Niagara/Source/Niagara/Public/NiagaraCommon.h
--- a/Engine/Plugins/FX/Niagara/Source/Niagara/Public/NiagaraCommon.h
+++ b/Engine/Plugins/FX/Niagara/Source/Niagara/Public/NiagaraCommon.h
@@ -65,6 +65,10 @@ using FNiagaraShaderMapRef = TRefCountPtr<FNiagaraShaderMap>;
 /** When enabled allows builds to warn about system / emitter information such as hitting buffer caps.  Small performance overhead for copying the name around so RT has access. */
 #define WITH_NIAGARA_DEBUG_EMITTER_NAME	1//!UE_BUILD_SHIPPING
 
+#ifndef WITH_NIAGARA_LEAK_DETECTOR
+	#define WITH_NIAGARA_LEAK_DETECTOR	1
+#endif
+
 #define INTERPOLATED_PARAMETER_PREFIX TEXT("PREV_")
 
 /** Defines The maximum ThreadGroup size we allow in Niagara.  This is important for how memory is allocated as we always need to round this and the final instance is used to avoid overflowing the buffer. */
diff --git a/Engine/Plugins/FX/Niagara/Source/Niagara/Public/NiagaraWorldManager.h b/Engine/Plugins/FX/Niagara/Source/Niagara/Public/NiagaraWorldManager.h
--- a/Engine/Plugins/FX/Niagara/Source/Niagara/Public/NiagaraWorldManager.h
+++ b/Engine/Plugins/FX/Niagara/Source/Niagara/Public/NiagaraWorldManager.h
@@ -10,6 +10,7 @@
 #include "Particles/ParticlePerfStats.h"
 #include "NiagaraParameterCollection.h"
 #include "UObject/GCObject.h"
+#include "NiagaraCommon.h"
 #include "NiagaraDataSet.h"
 #include "NiagaraScriptExecutionContext.h"
 #include "NiagaraSystemSimulation.h"
@@ -412,6 +413,9 @@ private:
 #if WITH_NIAGARA_DEBUGGER
 	TUniquePtr<class FNiagaraDebugHud> NiagaraDebugHud;
 #endif
+#if WITH_NIAGARA_LEAK_DETECTOR
+	TUniquePtr<class FNiagaraComponentLeakDetector> ComponentLeakDetector;
+#endif
 
 	TMap<TObjectPtr<UNiagaraSystem>, TObjectPtr<UNiagaraCullProxyComponent>> CullProxyMap;
 	TArray<TObjectPtr<UObject>> ReferencedObjects;
