- Add mesh rotation rate module to stateless
- Make units consistent with initial mesh orientation

diff --git a/Engine/Plugins/FX/Niagara/Shaders/Private/Stateless/Modules/NiagaraStatelessModule_MeshRotationRate.ush b/Engine/Plugins/FX/Niagara/Shaders/Private/Stateless/Modules/NiagaraStatelessModule_MeshRotationRate.ush
new file mode 100644
--- /dev/null
+++ b/Engine/Plugins/FX/Niagara/Shaders/Private/Stateless/Modules/NiagaraStatelessModule_MeshRotationRate.ush
@@ -0,0 +1,23 @@
+// Copyright Epic Games, Inc. All Rights Reserved.
+
+#include "../NiagaraStatelessCommon.ush"
+
+float3 MeshRotationRate_Scale;
+float3 MeshRotationRate_Bias;
+
+struct FStatelessModule_MeshRotationRate
+{
+	float3 RotationRate;
+};
+
+void MeshRotationRate_Initialize(inout FStatelessModule_MeshRotationRate ModuleData, inout FStatelessParticle Particle)
+{
+	ModuleData.RotationRate = RandomScaleBiasFloat(Particle, MeshRotationRate_Scale, MeshRotationRate_Bias);
+}
+	
+void MeshRotationRate_Simulate(in FStatelessModule_MeshRotationRate ModuleData, inout FStatelessParticle Particle)
+{
+	Particle.MeshOrientation			= MultiplyQuat(Particle.MeshOrientation,			RotatorToQuat(Particle.Age * ModuleData.RotationRate));
+	Particle.PreviousMeshOrientation	= MultiplyQuat(Particle.PreviousMeshOrientation,	RotatorToQuat(Particle.PreviousAge * ModuleData.RotationRate));
+}
+
diff --git a/Engine/Plugins/FX/Niagara/Shaders/Private/Stateless/NiagaraStatelessSimulationDefault.usf b/Engine/Plugins/FX/Niagara/Shaders/Private/Stateless/NiagaraStatelessSimulationDefault.usf
--- a/Engine/Plugins/FX/Niagara/Shaders/Private/Stateless/NiagaraStatelessSimulationDefault.usf
+++ b/Engine/Plugins/FX/Niagara/Shaders/Private/Stateless/NiagaraStatelessSimulationDefault.usf
@@ -34,6 +34,7 @@
 	PARTICLE_MODULE(ScaleSpriteSize) \
 	PARTICLE_MODULE(SpriteFacingAndAlignment) \
 	PARTICLE_MODULE(SpriteRotationRate) \
+	PARTICLE_MODULE(MeshRotationRate) \
 	PARTICLE_MODULE(ScaleMeshSize) \
 	PARTICLE_MODULE(SolveVelocitiesAndForces) \
 	PARTICLE_MODULE(ScaleMeshSizeBySpeed) \
diff --git a/Engine/Plugins/FX/Niagara/Shaders/Private/Stateless/NiagaraStatelessSimulationTemplate.ush b/Engine/Plugins/FX/Niagara/Shaders/Private/Stateless/NiagaraStatelessSimulationTemplate.ush
--- a/Engine/Plugins/FX/Niagara/Shaders/Private/Stateless/NiagaraStatelessSimulationTemplate.ush
+++ b/Engine/Plugins/FX/Niagara/Shaders/Private/Stateless/NiagaraStatelessSimulationTemplate.ush
@@ -15,6 +15,7 @@
 #include "Modules/NiagaraStatelessModule_DynamicMaterialParameters.ush"
 #include "Modules/NiagaraStatelessModule_InitializeParticle.ush"
 #include "Modules/NiagaraStatelessModule_InitialMeshOrientation.ush"
+#include "Modules/NiagaraStatelessModule_MeshRotationRate.ush"
 #include "Modules/NiagaraStatelessModule_RotateAroundPoint.ush"
 #include "Modules/NiagaraStatelessModule_ShapeLocation.ush"
 #include "Modules/NiagaraStatelessModule_ScaleColor.ush"
diff --git a/Engine/Plugins/FX/Niagara/Source/Niagara/Internal/Stateless/Modules/NiagaraStatelessModule_InitialMeshOrientation.h b/Engine/Plugins/FX/Niagara/Source/Niagara/Internal/Stateless/Modules/NiagaraStatelessModule_InitialMeshOrientation.h
--- a/Engine/Plugins/FX/Niagara/Source/Niagara/Internal/Stateless/Modules/NiagaraStatelessModule_InitialMeshOrientation.h
+++ b/Engine/Plugins/FX/Niagara/Source/Niagara/Internal/Stateless/Modules/NiagaraStatelessModule_InitialMeshOrientation.h
@@ -27,7 +27,7 @@ public:
 		FParameters* Parameters = SetShaderParameterContext.GetParameterNestedStruct<FParameters>();
 		if (IsModuleEnabled())
 		{
-			Parameters->InitialMeshOrientation_Rotation			= Rotation;
+			Parameters->InitialMeshOrientation_Rotation			= Rotation / 360.0f;
 			Parameters->InitialMeshOrientation_RandomRangeScale	= RandomRotationRange / 360.0f;
 		}
 		else
diff --git a/Engine/Plugins/FX/Niagara/Source/Niagara/Internal/Stateless/Modules/NiagaraStatelessModule_MeshRotationRate.h b/Engine/Plugins/FX/Niagara/Source/Niagara/Internal/Stateless/Modules/NiagaraStatelessModule_MeshRotationRate.h
new file mode 100644
--- /dev/null
+++ b/Engine/Plugins/FX/Niagara/Source/Niagara/Internal/Stateless/Modules/NiagaraStatelessModule_MeshRotationRate.h
@@ -0,0 +1,54 @@
+// Copyright Epic Games, Inc. All Rights Reserved.
+
+#pragma once
+
+#include "Stateless/NiagaraStatelessModule.h"
+#include "Stateless/NiagaraStatelessEmitterDataBuildContext.h"
+#include "Stateless/NiagaraStatelessModuleShaderParameters.h"
+
+#include "NiagaraStatelessModule_MeshRotationRate.generated.h"
+
+UCLASS(MinimalAPI, EditInlineNew, meta = (DisplayName = "Mesh Rotation Rate"))
+class UNiagaraStatelessModule_MeshRotationRate : public UNiagaraStatelessModule
+{
+	GENERATED_BODY()
+
+	struct FModuleBuiltData
+	{
+		FNiagaraStatelessRangeVector3	RotationRange;
+	};
+
+public:
+	using FParameters = NiagaraStateless::FMeshRotationRateModule_ShaderParameters;
+
+	UPROPERTY(EditAnywhere, Category = "Parameters", meta = (DisplayName = "Rotation Rate"))
+	FNiagaraDistributionRangeVector3 RotationRateDistribution = FNiagaraDistributionRangeVector3(FVector3f::ZeroVector);
+
+	virtual void BuildEmitterData(FNiagaraStatelessEmitterDataBuildContext& BuildContext) const override
+	{
+		FModuleBuiltData* BuiltData = BuildContext.AllocateBuiltData<FModuleBuiltData>();
+		BuiltData->RotationRange = BuildContext.ConvertDistributionToRange(RotationRateDistribution, FVector3f::ZeroVector, IsModuleEnabled());
+		BuiltData->RotationRange.Min *= 1.0f / 360.0f;
+		BuiltData->RotationRange.Max *= 1.0f / 360.0f;
+	}
+
+	virtual void SetShaderParameters(const FNiagaraStatelessSetShaderParameterContext& SetShaderParameterContext) const override
+	{
+		FParameters* Parameters = SetShaderParameterContext.GetParameterNestedStruct<FParameters>();
+		const FModuleBuiltData* ModuleBuiltData = SetShaderParameterContext.ReadBuiltData<FModuleBuiltData>();
+
+		SetShaderParameterContext.ConvertRangeToScaleBias(ModuleBuiltData->RotationRange, Parameters->MeshRotationRate_Scale, Parameters->MeshRotationRate_Bias);
+	}
+
+#if WITH_EDITOR
+	virtual bool CanDisableModule() const override { return true; }
+#endif
+#if WITH_EDITORONLY_DATA
+	virtual void GetOutputVariables(TArray<FNiagaraVariableBase>& OutVariables) const override
+	{
+		const FNiagaraStatelessGlobals& StatelessGlobals = FNiagaraStatelessGlobals::Get();
+		OutVariables.AddUnique(StatelessGlobals.MeshOrientationVariable);
+		OutVariables.AddUnique(StatelessGlobals.PreviousMeshOrientationVariable);
+	}
+#endif
+};
diff --git a/Engine/Plugins/FX/Niagara/Source/Niagara/Private/Stateless/NiagaraStatelessEmitterTemplate.cpp b/Engine/Plugins/FX/Niagara/Source/Niagara/Private/Stateless/NiagaraStatelessEmitterTemplate.cpp
--- a/Engine/Plugins/FX/Niagara/Source/Niagara/Private/Stateless/NiagaraStatelessEmitterTemplate.cpp
+++ b/Engine/Plugins/FX/Niagara/Source/Niagara/Private/Stateless/NiagaraStatelessEmitterTemplate.cpp
@@ -16,6 +16,7 @@
 #include "Stateless/Modules/NiagaraStatelessModule_GravityForce.h"
 #include "Stateless/Modules/NiagaraStatelessModule_InitializeParticle.h"
 #include "Stateless/Modules/NiagaraStatelessModule_InitialMeshOrientation.h"
+#include "Stateless/Modules/NiagaraStatelessModule_MeshRotationRate.h"
 #include "Stateless/Modules/NiagaraStatelessModule_RotateAroundPoint.h"
 #include "Stateless/Modules/NiagaraStatelessModule_ScaleColor.h"
 #include "Stateless/Modules/NiagaraStatelessModule_ScaleMeshSize.h"
@@ -77,6 +78,7 @@ void UNiagaraStatelessEmitterDefault::InitModulesAndAttributes()
 		UNiagaraStatelessModule_ScaleSpriteSizeBySpeed::StaticClass(),
 		UNiagaraStatelessModule_ScaleMeshSize::StaticClass(),
 		UNiagaraStatelessModule_ScaleMeshSizeBySpeed::StaticClass(),
+		UNiagaraStatelessModule_MeshRotationRate::StaticClass(),
 		UNiagaraStatelessModule_AddVelocity::StaticClass(),
 		UNiagaraStatelessModule_AccelerationForce::StaticClass(),
 		UNiagaraStatelessModule_CurlNoiseForce::StaticClass(),
diff --git a/Engine/Plugins/FX/Niagara/Source/NiagaraShader/Internal/Stateless/NiagaraStatelessEmitterShaders.h b/Engine/Plugins/FX/Niagara/Source/NiagaraShader/Internal/Stateless/NiagaraStatelessEmitterShaders.h
--- a/Engine/Plugins/FX/Niagara/Source/NiagaraShader/Internal/Stateless/NiagaraStatelessEmitterShaders.h
+++ b/Engine/Plugins/FX/Niagara/Source/NiagaraShader/Internal/Stateless/NiagaraStatelessEmitterShaders.h
@@ -36,11 +36,12 @@ namespace NiagaraStateless
 			SHADER_PARAMETER_STRUCT_INCLUDE(FScaleSpriteSizeBySpeedModule_ShaderParameters,		Module5Params)
 			SHADER_PARAMETER_STRUCT_INCLUDE(FScaleMeshSizeModule_ShaderParameters,				Module6Params)
 			SHADER_PARAMETER_STRUCT_INCLUDE(FScaleMeshSizeBySpeedModule_ShaderParameters,		Module7Params)
-			SHADER_PARAMETER_STRUCT_INCLUDE(FSolveVelocitiesAndForcesModule_ShaderParameters,	Module8Params)
-			SHADER_PARAMETER_STRUCT_INCLUDE(FSpriteFacingAndAlignmentModule_ShaderParameters,	Module9Params)
-			SHADER_PARAMETER_STRUCT_INCLUDE(FSpriteRotationRateModule_ShaderParameters,			Module10Params)
-			SHADER_PARAMETER_STRUCT_INCLUDE(FSubUVAnimationModule_ShaderParameters,				Module11Params)
-			SHADER_PARAMETER_STRUCT_INCLUDE(FDynamicMaterialParametersModule_ShaderParameters,	Module12Params)
+			SHADER_PARAMETER_STRUCT_INCLUDE(FMeshRotationRateModule_ShaderParameters,			Module8Params)
+			SHADER_PARAMETER_STRUCT_INCLUDE(FSolveVelocitiesAndForcesModule_ShaderParameters,	Module9Params)
+			SHADER_PARAMETER_STRUCT_INCLUDE(FSpriteFacingAndAlignmentModule_ShaderParameters,	Module10Params)
+			SHADER_PARAMETER_STRUCT_INCLUDE(FSpriteRotationRateModule_ShaderParameters,			Module11Params)
+			SHADER_PARAMETER_STRUCT_INCLUDE(FSubUVAnimationModule_ShaderParameters,				Module12Params)
+			SHADER_PARAMETER_STRUCT_INCLUDE(FDynamicMaterialParametersModule_ShaderParameters,	Module13Params)
 
 			SHADER_PARAMETER(int, Permutation_UniqueIDComponent)
 			SHADER_PARAMETER(int, Permutation_PositionComponent)
diff --git a/Engine/Plugins/FX/Niagara/Source/NiagaraShader/Internal/Stateless/NiagaraStatelessModuleShaderParameters.h b/Engine/Plugins/FX/Niagara/Source/NiagaraShader/Internal/Stateless/NiagaraStatelessModuleShaderParameters.h
--- a/Engine/Plugins/FX/Niagara/Source/NiagaraShader/Internal/Stateless/NiagaraStatelessModuleShaderParameters.h
+++ b/Engine/Plugins/FX/Niagara/Source/NiagaraShader/Internal/Stateless/NiagaraStatelessModuleShaderParameters.h
@@ -65,6 +65,11 @@ namespace NiagaraStateless
 		SHADER_PARAMETER(float,			ScaleMeshSizeBySpeed_VelocityNorm)
 	END_SHADER_PARAMETER_STRUCT()
 
+	BEGIN_SHADER_PARAMETER_STRUCT(FMeshRotationRateModule_ShaderParameters, NIAGARASHADER_API)
+		SHADER_PARAMETER(FVector3f,		MeshRotationRate_Scale)
+		SHADER_PARAMETER(FVector3f,		MeshRotationRate_Bias)
+	END_SHADER_PARAMETER_STRUCT()
+
 	BEGIN_SHADER_PARAMETER_STRUCT(FScaleSpriteSizeModule_ShaderParameters, NIAGARASHADER_API)
 		SHADER_PARAMETER(FUintVector3, ScaleSpriteSize_Distribution)
 	END_SHADER_PARAMETER_STRUCT()
