- Don't allow particle read DI to be a user parameter
- It doesn't function, as in you can not change the target emitter

diff --git a/Engine/Plugins/FX/Niagara/Source/Niagara/Private/NiagaraDataInterfaceParticleRead.cpp b/Engine/Plugins/FX/Niagara/Source/Niagara/Private/NiagaraDataInterfaceParticleRead.cpp
--- a/Engine/Plugins/FX/Niagara/Source/Niagara/Private/NiagaraDataInterfaceParticleRead.cpp
+++ b/Engine/Plugins/FX/Niagara/Source/Niagara/Private/NiagaraDataInterfaceParticleRead.cpp
@@ -391,7 +391,7 @@ void UNiagaraDataInterfaceParticleRead::PostInitProperties()
 
 	if (HasAnyFlags(RF_ClassDefaultObject))
 	{
-		ENiagaraTypeRegistryFlags Flags = ENiagaraTypeRegistryFlags::AllowAnyVariable | ENiagaraTypeRegistryFlags::AllowParameter;
+		const ENiagaraTypeRegistryFlags Flags = ENiagaraTypeRegistryFlags::AllowNotUserVariable | ENiagaraTypeRegistryFlags::AllowParameter;
 		FNiagaraTypeRegistry::Register(FNiagaraTypeDefinition(GetClass()), Flags);
 	}
 }
diff --git a/Engine/Plugins/FX/Niagara/Source/Niagara/Public/NiagaraTypes.h b/Engine/Plugins/FX/Niagara/Source/Niagara/Public/NiagaraTypes.h
--- a/Engine/Plugins/FX/Niagara/Source/Niagara/Public/NiagaraTypes.h
+++ b/Engine/Plugins/FX/Niagara/Source/Niagara/Public/NiagaraTypes.h
@@ -1156,6 +1156,7 @@ enum class ENiagaraTypeRegistryFlags : uint32
 	AllowEmitterVariable	= (1 << 2),
 	AllowParticleVariable	= (1 << 3),
 	AllowAnyVariable		= (AllowUserVariable | AllowSystemVariable | AllowEmitterVariable | AllowParticleVariable),
+	AllowNotUserVariable	= (AllowSystemVariable | AllowEmitterVariable | AllowParticleVariable),
 
 	AllowParameter			= (1 << 4),
 	AllowPayload			= (1 << 5),
