Fixed a bug that prevented users to edit any complex structs in user parameters (e.g. matrices)

diff --git a/Engine/Plugins/FX/Niagara/Source/NiagaraEditor/Private/Customizations/UserParameters/NiagaraUserParametersBuilderBase.cpp b/Engine/Plugins/FX/Niagara/Source/NiagaraEditor/Private/Customizations/UserParameters/NiagaraUserParametersBuilderBase.cpp
--- a/Engine/Plugins/FX/Niagara/Source/NiagaraEditor/Private/Customizations/UserParameters/NiagaraUserParametersBuilderBase.cpp
+++ b/Engine/Plugins/FX/Niagara/Source/NiagaraEditor/Private/Customizations/UserParameters/NiagaraUserParametersBuilderBase.cpp
@@ -124,8 +124,9 @@ IDetailPropertyRow* FNiagaraUserParameterNodeBuilder::AddValueParameterAsRow(IDe
 	TSharedPtr<SWidget> DefaultNameWidget;
 	TSharedPtr<SWidget> DefaultValueWidget;
 	DetailPropertyRow->GetDefaultWidgets(DefaultNameWidget, DefaultValueWidget);
-	
-	FDetailWidgetRow& DetailWidgetRow = DetailPropertyRow->CustomWidget();
+
+	bool bHasCustomEditor = NiagaraParameterEditors.Contains(ChoppedUserParameter);
+	FDetailWidgetRow& DetailWidgetRow = DetailPropertyRow->CustomWidget(!bHasCustomEditor);
 	DetailWidgetRow
 	.FilterString(FText::FromName(ChoppedUserParameter.GetName()))
 	.NameContent()
@@ -134,7 +135,7 @@ IDetailPropertyRow* FNiagaraUserParameterNodeBuilder::AddValueParameterAsRow(IDe
 	]
 	.ValueContent()
 	[
-		NiagaraParameterEditors.Contains(ChoppedUserParameter) ? NiagaraParameterEditors[ChoppedUserParameter] : DefaultValueWidget.ToSharedRef()
+		bHasCustomEditor ? NiagaraParameterEditors[ChoppedUserParameter] : DefaultValueWidget.ToSharedRef()
 	];
 
 	AddCustomMenuActionsForParameter(DetailWidgetRow, ChoppedUserParameter);
