- Add camera offset module to stateless

diff --git a/Engine/Plugins/FX/Niagara/Shaders/Private/Stateless/Modules/NiagaraStatelessModule_CameraOffset.ush b/Engine/Plugins/FX/Niagara/Shaders/Private/Stateless/Modules/NiagaraStatelessModule_CameraOffset.ush
new file mode 100644
--- /dev/null
+++ b/Engine/Plugins/FX/Niagara/Shaders/Private/Stateless/Modules/NiagaraStatelessModule_CameraOffset.ush
@@ -0,0 +1,24 @@
+// Copyright Epic Games, Inc. All Rights Reserved.
+
+#include "../NiagaraStatelessCommon.ush"
+
+uint3 CameraOffset_Distribution;
+
+struct FStatelessModule_CameraOffset
+{
+	FStatelessDistributionFloat	CameraOffsetDistribution;
+};
+
+void CameraOffset_Initialize(inout FStatelessModule_CameraOffset ModuleData, inout FStatelessParticle Particle)
+{
+	ModuleData.CameraOffsetDistribution.Init(Particle, CameraOffset_Distribution);
+}
+	
+void CameraOffset_Simulate(in FStatelessModule_CameraOffset ModuleData, inout FStatelessParticle Particle)
+{
+	if (ModuleData.CameraOffsetDistribution.IsValid())
+	{
+		Particle.CameraOffset			= ModuleData.CameraOffsetDistribution.GetValue(Particle.NormalizedAge);
+		Particle.PreviousCameraOffset	= ModuleData.CameraOffsetDistribution.GetValue(Particle.PreviousNormalizedAge);
+	}
+}
diff --git a/Engine/Plugins/FX/Niagara/Shaders/Private/Stateless/NiagaraStatelessSimulationDefault.usf b/Engine/Plugins/FX/Niagara/Shaders/Private/Stateless/NiagaraStatelessSimulationDefault.usf
--- a/Engine/Plugins/FX/Niagara/Shaders/Private/Stateless/NiagaraStatelessSimulationDefault.usf
+++ b/Engine/Plugins/FX/Niagara/Shaders/Private/Stateless/NiagaraStatelessSimulationDefault.usf
@@ -3,6 +3,7 @@
 #define PARTICLE_ATTRIBUTES_OUTPUT \
 	PARTICLE_ATTRIBUTE(int,		UniqueID) \
 	PARTICLE_ATTRIBUTE(float3,	Position) \
+	PARTICLE_ATTRIBUTE(float,	CameraOffset) \
 	PARTICLE_ATTRIBUTE(float4,	Color) \
 	PARTICLE_ATTRIBUTE(float4,	DynamicMaterialParameter0) \
 	PARTICLE_ATTRIBUTE(float4,	MeshOrientation) \
@@ -14,6 +15,7 @@
 	PARTICLE_ATTRIBUTE(float,	SpriteRotation) \
 	PARTICLE_ATTRIBUTE(float,	SubImageIndex) \
 	PARTICLE_ATTRIBUTE(float3,	Velocity) \
+	PARTICLE_ATTRIBUTE(float,	PreviousCameraOffset) \
 	PARTICLE_ATTRIBUTE(float3,	PreviousPosition) \
 	PARTICLE_ATTRIBUTE(float4,	PreviousMeshOrientation) \
 	PARTICLE_ATTRIBUTE(float,	PreviousRibbonWidth) \
@@ -30,6 +32,7 @@
 	PARTICLE_MODULE(InitializeParticle) \
 	PARTICLE_MODULE(InitialMeshOrientation) \
 	PARTICLE_MODULE(ShapeLocation) \
+	PARTICLE_MODULE(CameraOffset) \
 	PARTICLE_MODULE(ScaleColor) \
 	PARTICLE_MODULE(ScaleSpriteSize) \
 	PARTICLE_MODULE(SpriteFacingAndAlignment) \
diff --git a/Engine/Plugins/FX/Niagara/Shaders/Private/Stateless/NiagaraStatelessSimulationExample1.usf b/Engine/Plugins/FX/Niagara/Shaders/Private/Stateless/NiagaraStatelessSimulationExample1.usf
--- a/Engine/Plugins/FX/Niagara/Shaders/Private/Stateless/NiagaraStatelessSimulationExample1.usf
+++ b/Engine/Plugins/FX/Niagara/Shaders/Private/Stateless/NiagaraStatelessSimulationExample1.usf
@@ -17,11 +17,13 @@
 	PARTICLE_ATTRIBUTE(float3,	PreviousVelocity) \
 
 #define PARTICLE_ATTRIBUTES_TRANSIENT \
+	PARTICLE_ATTRIBUTE(float,	CameraOffset) \
 	PARTICLE_ATTRIBUTE(float4,	DynamicMaterialParameter0) \
 	PARTICLE_ATTRIBUTE(float4,	MeshOrientation) \
 	PARTICLE_ATTRIBUTE(float3,	SpriteFacing) \
 	PARTICLE_ATTRIBUTE(float3,	SpriteAlignment) \
 	PARTICLE_ATTRIBUTE(float,	SubImageIndex) \
+	PARTICLE_ATTRIBUTE(float,	PreviousCameraOffset) \
 	PARTICLE_ATTRIBUTE(float4,	PreviousMeshOrientation) \
 	PARTICLE_ATTRIBUTE(float3,	PreviousSpriteFacing) \
 	PARTICLE_ATTRIBUTE(float3,	PreviousSpriteAlignment) \
diff --git a/Engine/Plugins/FX/Niagara/Shaders/Private/Stateless/NiagaraStatelessSimulationTemplate.ush b/Engine/Plugins/FX/Niagara/Shaders/Private/Stateless/NiagaraStatelessSimulationTemplate.ush
--- a/Engine/Plugins/FX/Niagara/Shaders/Private/Stateless/NiagaraStatelessSimulationTemplate.ush
+++ b/Engine/Plugins/FX/Niagara/Shaders/Private/Stateless/NiagaraStatelessSimulationTemplate.ush
@@ -12,6 +12,7 @@
 /////////////////////////////////////////////////////////////////////////////////////////////////////////////////////////
 // TBD this might produce sub optimal code, needs investigate and is temporary
 #include "Modules/NiagaraStatelessModule_CalculateAccurateVelocity.ush"
+#include "Modules/NiagaraStatelessModule_CameraOffset.ush"
 #include "Modules/NiagaraStatelessModule_DynamicMaterialParameters.ush"
 #include "Modules/NiagaraStatelessModule_InitializeParticle.ush"
 #include "Modules/NiagaraStatelessModule_InitialMeshOrientation.ush"
diff --git a/Engine/Plugins/FX/Niagara/Source/Niagara/Internal/Stateless/Modules/NiagaraStatelessModule_CameraOffset.h b/Engine/Plugins/FX/Niagara/Source/Niagara/Internal/Stateless/Modules/NiagaraStatelessModule_CameraOffset.h
new file mode 100644
--- /dev/null
+++ b/Engine/Plugins/FX/Niagara/Source/Niagara/Internal/Stateless/Modules/NiagaraStatelessModule_CameraOffset.h
@@ -0,0 +1,52 @@
+// Copyright Epic Games, Inc. All Rights Reserved.
+
+#pragma once
+
+#include "Stateless/NiagaraStatelessModule.h"
+#include "Stateless/NiagaraStatelessModuleShaderParameters.h"
+
+#include "NiagaraStatelessModule_CameraOffset.generated.h"
+
+UCLASS(MinimalAPI, EditInlineNew, meta = (DisplayName = "Camera Offset"))
+class UNiagaraStatelessModule_CameraOffset : public UNiagaraStatelessModule
+{
+public:
+	GENERATED_BODY()
+
+	struct FModuleBuiltData
+	{
+		FUintVector3	DistributionParameters = FUintVector3::ZeroValue;
+	};
+
+public:
+	using FParameters = NiagaraStateless::FCameraOffsetModule_ShaderParameters;
+
+	UPROPERTY(EditAnywhere, Category = "Parameters", meta = (DisplayName = "CameraOffset"))
+	FNiagaraDistributionFloat CameraOffsetDistribution = FNiagaraDistributionFloat(0.0f);
+
+	virtual void BuildEmitterData(FNiagaraStatelessEmitterDataBuildContext& BuildContext) const override
+	{
+		FModuleBuiltData* BuiltData = BuildContext.AllocateBuiltData<FModuleBuiltData>();
+		BuiltData->DistributionParameters = BuildContext.AddDistribution(CameraOffsetDistribution, IsModuleEnabled());
+	}
+
+	virtual void SetShaderParameters(const FNiagaraStatelessSetShaderParameterContext& SetShaderParameterContext) const override
+	{
+		const FModuleBuiltData* ModuleBuiltData = SetShaderParameterContext.ReadBuiltData<FModuleBuiltData>();
+
+		FParameters* Parameters = SetShaderParameterContext.GetParameterNestedStruct<FParameters>();
+		Parameters->CameraOffset_Distribution = ModuleBuiltData->DistributionParameters;
+	}
+
+#if WITH_EDITOR
+	virtual bool CanDisableModule() const override { return true; }
+#endif
+#if WITH_EDITORONLY_DATA
+	virtual void GetOutputVariables(TArray<FNiagaraVariableBase>& OutVariables) const override
+	{
+		const FNiagaraStatelessGlobals& StatelessGlobals = FNiagaraStatelessGlobals::Get();
+		OutVariables.AddUnique(StatelessGlobals.CameraOffsetVariable);
+		OutVariables.AddUnique(StatelessGlobals.PreviousCameraOffsetVariable);
+	}
+#endif
+};
diff --git a/Engine/Plugins/FX/Niagara/Source/Niagara/Internal/Stateless/NiagaraStatelessCommon.h b/Engine/Plugins/FX/Niagara/Source/Niagara/Internal/Stateless/NiagaraStatelessCommon.h
--- a/Engine/Plugins/FX/Niagara/Source/Niagara/Internal/Stateless/NiagaraStatelessCommon.h
+++ b/Engine/Plugins/FX/Niagara/Source/Niagara/Internal/Stateless/NiagaraStatelessCommon.h
@@ -29,6 +29,7 @@ using FNiagaraStatelessRangeColor	= FNiagaraStatelessRange<FLinearColor>;
 
 struct FNiagaraStatelessGlobals
 {
+	FNiagaraVariableBase	CameraOffsetVariable;
 	FNiagaraVariableBase	ColorVariable;
 	FNiagaraVariableBase	DynamicMaterialParameters0Variable;
 	FNiagaraVariableBase	MeshOrientationVariable;
@@ -43,6 +44,7 @@ struct FNiagaraStatelessGlobals
 	FNiagaraVariableBase	UniqueIDVariable;
 	FNiagaraVariableBase	VelocityVariable;
 
+	FNiagaraVariableBase	PreviousCameraOffsetVariable;
 	//FNiagaraVariableBase	PreviousColorVariable;
 	//FNiagaraVariableBase	PreviousDynamicMaterialParameters0Variable;
 	FNiagaraVariableBase	PreviousMeshOrientationVariable;
diff --git a/Engine/Plugins/FX/Niagara/Source/Niagara/Private/Stateless/NiagaraStatelessCommon.cpp b/Engine/Plugins/FX/Niagara/Source/Niagara/Private/Stateless/NiagaraStatelessCommon.cpp
--- a/Engine/Plugins/FX/Niagara/Source/Niagara/Private/Stateless/NiagaraStatelessCommon.cpp
+++ b/Engine/Plugins/FX/Niagara/Source/Niagara/Private/Stateless/NiagaraStatelessCommon.cpp
@@ -10,6 +10,7 @@ namespace NiagaraStatelessCommon
 
 	void Initialize()
 	{
+		GGlobals.CameraOffsetVariable				= FNiagaraVariableBase(FNiagaraTypeDefinition::GetFloatDef(), TEXT("CameraOffset"));
 		GGlobals.ColorVariable						= FNiagaraVariableBase(FNiagaraTypeDefinition::GetColorDef(), TEXT("Color"));
 		GGlobals.DynamicMaterialParameters0Variable = FNiagaraVariableBase(FNiagaraTypeDefinition::GetVec4Def(), TEXT("DynamicMaterialParameter"));
 		GGlobals.MeshOrientationVariable			= FNiagaraVariableBase(FNiagaraTypeDefinition::GetQuatDef(), TEXT("MeshOrientation"));
@@ -24,6 +25,7 @@ namespace NiagaraStatelessCommon
 		GGlobals.UniqueIDVariable					= FNiagaraVariableBase(FNiagaraTypeDefinition::GetIntDef(), TEXT("UniqueID"));
 		GGlobals.VelocityVariable					= FNiagaraVariableBase(FNiagaraTypeDefinition::GetVec3Def(), TEXT("Velocity"));
 
+		GGlobals.PreviousCameraOffsetVariable		= FNiagaraVariableBase(FNiagaraTypeDefinition::GetFloatDef(), TEXT("Previous.CameraOffset"));
 		//GGlobals.PreviousColorVariable			= FNiagaraVariableBase(FNiagaraTypeDefinition::GetColorDef(), TEXT("Previous.Color"));
 		//GGlobals.PreviousDynamicMaterialParameters0Variable		= FNiagaraVariableBase(FNiagaraTypeDefinition::GetVec4Def(), TEXT("Previous.DynamicMaterialParameter"));
 		GGlobals.PreviousMeshOrientationVariable	= FNiagaraVariableBase(FNiagaraTypeDefinition::GetQuatDef(), TEXT("Previous.MeshOrientation"));
diff --git a/Engine/Plugins/FX/Niagara/Source/Niagara/Private/Stateless/NiagaraStatelessEmitterTemplate.cpp b/Engine/Plugins/FX/Niagara/Source/Niagara/Private/Stateless/NiagaraStatelessEmitterTemplate.cpp
--- a/Engine/Plugins/FX/Niagara/Source/Niagara/Private/Stateless/NiagaraStatelessEmitterTemplate.cpp
+++ b/Engine/Plugins/FX/Niagara/Source/Niagara/Private/Stateless/NiagaraStatelessEmitterTemplate.cpp
@@ -11,6 +11,7 @@
 #include "Stateless/Modules/NiagaraStatelessModule_AddVelocity.h"
 #include "Stateless/Modules/NiagaraStatelessModule_AccelerationForce.h"
 #include "Stateless/Modules/NiagaraStatelessModule_CalculateAccurateVelocity.h"
+#include "Stateless/Modules/NiagaraStatelessModule_CameraOffset.h"
 #include "Stateless/Modules/NiagaraStatelessModule_CurlNoiseForce.h"
 #include "Stateless/Modules/NiagaraStatelessModule_Drag.h"
 #include "Stateless/Modules/NiagaraStatelessModule_GravityForce.h"
@@ -73,6 +74,7 @@ void UNiagaraStatelessEmitterDefault::InitModulesAndAttributes()
 		UNiagaraStatelessModule_InitializeParticle::StaticClass(),
 		UNiagaraStatelessModule_InitialMeshOrientation::StaticClass(),
 		UNiagaraStatelessModule_ShapeLocation::StaticClass(),
+		UNiagaraStatelessModule_CameraOffset::StaticClass(),
 		UNiagaraStatelessModule_ScaleColor::StaticClass(),
 		UNiagaraStatelessModule_ScaleSpriteSize::StaticClass(),
 		UNiagaraStatelessModule_ScaleSpriteSizeBySpeed::StaticClass(),
@@ -96,6 +98,7 @@ void UNiagaraStatelessEmitterDefault::InitModulesAndAttributes()
 	{
 		StatelessGlobals.UniqueIDVariable,
 		StatelessGlobals.PositionVariable,
+		StatelessGlobals.CameraOffsetVariable,
 		StatelessGlobals.ColorVariable,
 		StatelessGlobals.DynamicMaterialParameters0Variable,
 		StatelessGlobals.MeshOrientationVariable,
@@ -108,6 +111,7 @@ void UNiagaraStatelessEmitterDefault::InitModulesAndAttributes()
 		StatelessGlobals.SubImageIndexVariable,
 		StatelessGlobals.VelocityVariable,
 		StatelessGlobals.PreviousPositionVariable,
+		StatelessGlobals.PreviousCameraOffsetVariable,
 		StatelessGlobals.PreviousMeshOrientationVariable,
 		StatelessGlobals.PreviousRibbonWidthVariable,
 		StatelessGlobals.PreviousScaleVariable,
@@ -140,6 +144,7 @@ void UNiagaraStatelessEmitterDefault::SetShaderParameters(uint8* ShaderParameter
 	int iComponent = 0;
 	ShaderParameters->Permutation_UniqueIDComponent					= ComponentOffsets[iComponent++];
 	ShaderParameters->Permutation_PositionComponent					= ComponentOffsets[iComponent++];
+	ShaderParameters->Permutation_CameraOffsetComponent				= ComponentOffsets[iComponent++];
 	ShaderParameters->Permutation_ColorComponent					= ComponentOffsets[iComponent++];
 	ShaderParameters->Permutation_DynamicMaterialParameter0Component= ComponentOffsets[iComponent++];
 	ShaderParameters->Permutation_MeshOrientationComponent			= ComponentOffsets[iComponent++];
@@ -152,6 +157,7 @@ void UNiagaraStatelessEmitterDefault::SetShaderParameters(uint8* ShaderParameter
 	ShaderParameters->Permutation_SubImageIndexComponent			= ComponentOffsets[iComponent++];
 	ShaderParameters->Permutation_VelocityComponent					= ComponentOffsets[iComponent++];
 	ShaderParameters->Permutation_PreviousPositionComponent			= ComponentOffsets[iComponent++];
+	ShaderParameters->Permutation_PreviousCameraOffsetComponent		= ComponentOffsets[iComponent++];
 	ShaderParameters->Permutation_PreviousMeshOrientationComponent	= ComponentOffsets[iComponent++];
 	ShaderParameters->Permutation_PreviousRibbonWidthComponent		= ComponentOffsets[iComponent++];
 	ShaderParameters->Permutation_PreviousScaleComponent			= ComponentOffsets[iComponent++];
diff --git a/Engine/Plugins/FX/Niagara/Source/NiagaraShader/Internal/Stateless/NiagaraStatelessEmitterShaders.h b/Engine/Plugins/FX/Niagara/Source/NiagaraShader/Internal/Stateless/NiagaraStatelessEmitterShaders.h
--- a/Engine/Plugins/FX/Niagara/Source/NiagaraShader/Internal/Stateless/NiagaraStatelessEmitterShaders.h
+++ b/Engine/Plugins/FX/Niagara/Source/NiagaraShader/Internal/Stateless/NiagaraStatelessEmitterShaders.h
@@ -31,20 +31,22 @@ namespace NiagaraStateless
 			SHADER_PARAMETER_STRUCT_INCLUDE(FInitializeParticleModule_ShaderParameters,			Module0Params)
 			SHADER_PARAMETER_STRUCT_INCLUDE(FInitialMeshOrientationModule_ShaderParameters,		Module1Params)
 			SHADER_PARAMETER_STRUCT_INCLUDE(FShapeLocationModule_ShaderParameters,				Module2Params)
-			SHADER_PARAMETER_STRUCT_INCLUDE(FScaleColorModule_ShaderParameters,					Module3Params)
-			SHADER_PARAMETER_STRUCT_INCLUDE(FScaleSpriteSizeModule_ShaderParameters,			Module4Params)
-			SHADER_PARAMETER_STRUCT_INCLUDE(FScaleSpriteSizeBySpeedModule_ShaderParameters,		Module5Params)
-			SHADER_PARAMETER_STRUCT_INCLUDE(FScaleMeshSizeModule_ShaderParameters,				Module6Params)
-			SHADER_PARAMETER_STRUCT_INCLUDE(FScaleMeshSizeBySpeedModule_ShaderParameters,		Module7Params)
-			SHADER_PARAMETER_STRUCT_INCLUDE(FMeshRotationRateModule_ShaderParameters,			Module8Params)
-			SHADER_PARAMETER_STRUCT_INCLUDE(FSolveVelocitiesAndForcesModule_ShaderParameters,	Module9Params)
-			SHADER_PARAMETER_STRUCT_INCLUDE(FSpriteFacingAndAlignmentModule_ShaderParameters,	Module10Params)
-			SHADER_PARAMETER_STRUCT_INCLUDE(FSpriteRotationRateModule_ShaderParameters,			Module11Params)
-			SHADER_PARAMETER_STRUCT_INCLUDE(FSubUVAnimationModule_ShaderParameters,				Module12Params)
-			SHADER_PARAMETER_STRUCT_INCLUDE(FDynamicMaterialParametersModule_ShaderParameters,	Module13Params)
+			SHADER_PARAMETER_STRUCT_INCLUDE(FCameraOffsetModule_ShaderParameters,				Module4Params)
+			SHADER_PARAMETER_STRUCT_INCLUDE(FScaleColorModule_ShaderParameters,					Module5Params)
+			SHADER_PARAMETER_STRUCT_INCLUDE(FScaleSpriteSizeModule_ShaderParameters,			Module6Params)
+			SHADER_PARAMETER_STRUCT_INCLUDE(FScaleSpriteSizeBySpeedModule_ShaderParameters,		Module7Params)
+			SHADER_PARAMETER_STRUCT_INCLUDE(FScaleMeshSizeModule_ShaderParameters,				Module8Params)
+			SHADER_PARAMETER_STRUCT_INCLUDE(FScaleMeshSizeBySpeedModule_ShaderParameters,		Module9Params)
+			SHADER_PARAMETER_STRUCT_INCLUDE(FMeshRotationRateModule_ShaderParameters,			Module10Params)
+			SHADER_PARAMETER_STRUCT_INCLUDE(FSolveVelocitiesAndForcesModule_ShaderParameters,	Module11Params)
+			SHADER_PARAMETER_STRUCT_INCLUDE(FSpriteFacingAndAlignmentModule_ShaderParameters,	Module12Params)
+			SHADER_PARAMETER_STRUCT_INCLUDE(FSpriteRotationRateModule_ShaderParameters,			Module13Params)
+			SHADER_PARAMETER_STRUCT_INCLUDE(FSubUVAnimationModule_ShaderParameters,				Module14Params)
+			SHADER_PARAMETER_STRUCT_INCLUDE(FDynamicMaterialParametersModule_ShaderParameters,	Module15Params)
 
 			SHADER_PARAMETER(int, Permutation_UniqueIDComponent)
 			SHADER_PARAMETER(int, Permutation_PositionComponent)
+			SHADER_PARAMETER(int, Permutation_CameraOffsetComponent)
 			SHADER_PARAMETER(int, Permutation_ColorComponent)
 			SHADER_PARAMETER(int, Permutation_DynamicMaterialParameter0Component)
 			SHADER_PARAMETER(int, Permutation_MeshOrientationComponent)
@@ -57,6 +59,7 @@ namespace NiagaraStateless
 			SHADER_PARAMETER(int, Permutation_SubImageIndexComponent)
 			SHADER_PARAMETER(int, Permutation_VelocityComponent)
 			SHADER_PARAMETER(int, Permutation_PreviousPositionComponent)
+			SHADER_PARAMETER(int, Permutation_PreviousCameraOffsetComponent)
 			SHADER_PARAMETER(int, Permutation_PreviousMeshOrientationComponent)
 			SHADER_PARAMETER(int, Permutation_PreviousRibbonWidthComponent)
 			SHADER_PARAMETER(int, Permutation_PreviousScaleComponent)
diff --git a/Engine/Plugins/FX/Niagara/Source/NiagaraShader/Internal/Stateless/NiagaraStatelessModuleShaderParameters.h b/Engine/Plugins/FX/Niagara/Source/NiagaraShader/Internal/Stateless/NiagaraStatelessModuleShaderParameters.h
--- a/Engine/Plugins/FX/Niagara/Source/NiagaraShader/Internal/Stateless/NiagaraStatelessModuleShaderParameters.h
+++ b/Engine/Plugins/FX/Niagara/Source/NiagaraShader/Internal/Stateless/NiagaraStatelessModuleShaderParameters.h
@@ -52,6 +52,10 @@ namespace NiagaraStateless
 		SHADER_PARAMETER(FUintVector3,	DynamicMaterialParameters_Parameter0W)
 	END_SHADER_PARAMETER_STRUCT()
 
+	BEGIN_SHADER_PARAMETER_STRUCT(FCameraOffsetModule_ShaderParameters, NIAGARASHADER_API)
+		SHADER_PARAMETER(FUintVector3,	CameraOffset_Distribution)
+	END_SHADER_PARAMETER_STRUCT()
+
 	BEGIN_SHADER_PARAMETER_STRUCT(FScaleColorModule_ShaderParameters, NIAGARASHADER_API)
 		SHADER_PARAMETER(FUintVector3,	ScaleColor_Distribution)
 	END_SHADER_PARAMETER_STRUCT()
