RayTracing: Fix shader warnings from new DXC compiler

diff --git a/Engine/Plugins/FX/Niagara/Shaders/Private/NiagaraRayTracingShaders.usf b/Engine/Plugins/FX/Niagara/Shaders/Private/NiagaraRayTracingShaders.usf
--- a/Engine/Plugins/FX/Niagara/Shaders/Private/NiagaraRayTracingShaders.usf
+++ b/Engine/Plugins/FX/Niagara/Shaders/Private/NiagaraRayTracingShaders.usf
@@ -31,7 +31,7 @@ bool FindCollisionGroup(int GPUSceneIndexId, out int CollisionGroup)
 		CollisionGroup = HashToCollisionGroups[Index];
 		return true;
 	}
-
+	CollisionGroup = 0;
 	return false;
 }
 #endif
diff --git a/Engine/Shaders/Private/RayTracing/RayTracingDebug.usf b/Engine/Shaders/Private/RayTracing/RayTracingDebug.usf
--- a/Engine/Shaders/Private/RayTracing/RayTracingDebug.usf
+++ b/Engine/Shaders/Private/RayTracing/RayTracingDebug.usf
@@ -267,6 +267,10 @@ float4 DebugRayTracingInstance(FRayDesc Ray, out float OutDepth)
 		OutDepth = ClipPosition.z / ClipPosition.w;
 		Result.a = 0.0;
 	}
+	else
+	{
+		OutDepth = 0.0;
+	}
 
 	return Result;
 }
