- Change the color of stateless emitters

diff --git a/Engine/Plugins/FX/Niagara/Source/NiagaraEditor/Private/NiagaraOverviewNode.cpp b/Engine/Plugins/FX/Niagara/Source/NiagaraEditor/Private/NiagaraOverviewNode.cpp
--- a/Engine/Plugins/FX/Niagara/Source/NiagaraEditor/Private/NiagaraOverviewNode.cpp
+++ b/Engine/Plugins/FX/Niagara/Source/NiagaraEditor/Private/NiagaraOverviewNode.cpp
@@ -24,6 +24,7 @@
 
 bool UNiagaraOverviewNode::bColorsAreInitialized = false;
 FLinearColor UNiagaraOverviewNode::EmitterColor;
+FLinearColor UNiagaraOverviewNode::StatelessEmitterColor;
 FLinearColor UNiagaraOverviewNode::SystemColor;
 FLinearColor UNiagaraOverviewNode::IsolatedColor;
 FLinearColor UNiagaraOverviewNode::NotIsolatedColor;
@@ -117,19 +118,28 @@ FLinearColor UNiagaraOverviewNode::GetNodeTitleColor() const
 	{
 		FNiagaraEditorModule& NiagaraEditorModule = FModuleManager::LoadModuleChecked<FNiagaraEditorModule>("NiagaraEditor");
 		EmitterColor = NiagaraEditorModule.GetWidgetProvider()->GetColorForExecutionCategory(UNiagaraStackEntry::FExecutionCategoryNames::Emitter);
+		StatelessEmitterColor = NiagaraEditorModule.GetWidgetProvider()->GetColorForExecutionCategory(UNiagaraStackEntry::FExecutionCategoryNames::StatelessEmitter);
 		SystemColor = NiagaraEditorModule.GetWidgetProvider()->GetColorForExecutionCategory(UNiagaraStackEntry::FExecutionCategoryNames::System);
 		IsolatedColor = FNiagaraEditorStyle::Get().GetColor("NiagaraEditor.OverviewNode.IsolatedColor");
 		NotIsolatedColor = FNiagaraEditorStyle::Get().GetColor("NiagaraEditor.OverviewNode.NotIsolatedColor");
+		bColorsAreInitialized = true;
 	}
 
 	if (EmitterHandleGuid.IsValid())
 	{
-		if (OwningSystem != nullptr && OwningSystem->GetIsolateEnabled())
+		if (OwningSystem != nullptr)
 		{
 			const FNiagaraEmitterHandle* Handle = FindEmitterHandleByID(OwningSystem, EmitterHandleGuid);
 			if (ensureMsgf(Handle != nullptr, TEXT("Failed to find matching emitter handle for existing overview node!")))
 			{
-				return Handle->IsIsolated() ? IsolatedColor : NotIsolatedColor;
+				if (OwningSystem->GetIsolateEnabled())
+				{
+					return Handle->IsIsolated() ? IsolatedColor : NotIsolatedColor;
+				}
+				else if ( Handle->GetEmitterMode() == ENiagaraEmitterMode::Stateless )
+				{
+					return StatelessEmitterColor;
+				}
 			}
 		}
 		
diff --git a/Engine/Plugins/FX/Niagara/Source/NiagaraEditor/Private/ViewModels/Stack/NiagaraStackEntry.cpp b/Engine/Plugins/FX/Niagara/Source/NiagaraEditor/Private/ViewModels/Stack/NiagaraStackEntry.cpp
--- a/Engine/Plugins/FX/Niagara/Source/NiagaraEditor/Private/ViewModels/Stack/NiagaraStackEntry.cpp
+++ b/Engine/Plugins/FX/Niagara/Source/NiagaraEditor/Private/ViewModels/Stack/NiagaraStackEntry.cpp
@@ -20,6 +20,7 @@
 class UNiagaraStackItemGroup;
 const FName UNiagaraStackEntry::FExecutionCategoryNames::System = TEXT("System");
 const FName UNiagaraStackEntry::FExecutionCategoryNames::Emitter = TEXT("Emitter");
+const FName UNiagaraStackEntry::FExecutionCategoryNames::StatelessEmitter = TEXT("StatelessEmitter");
 const FName UNiagaraStackEntry::FExecutionCategoryNames::Particle = TEXT("Particle");
 const FName UNiagaraStackEntry::FExecutionCategoryNames::Render = TEXT("Render");
 
diff --git a/Engine/Plugins/FX/Niagara/Source/NiagaraEditor/Public/NiagaraOverviewNode.h b/Engine/Plugins/FX/Niagara/Source/NiagaraEditor/Public/NiagaraOverviewNode.h
--- a/Engine/Plugins/FX/Niagara/Source/NiagaraEditor/Public/NiagaraOverviewNode.h
+++ b/Engine/Plugins/FX/Niagara/Source/NiagaraEditor/Public/NiagaraOverviewNode.h
@@ -58,6 +58,7 @@ private:
 	static NIAGARAEDITOR_API bool bColorsAreInitialized;
 	static NIAGARAEDITOR_API FLinearColor SystemColor;
 	static NIAGARAEDITOR_API FLinearColor EmitterColor;
+	static NIAGARAEDITOR_API FLinearColor StatelessEmitterColor;
 	static NIAGARAEDITOR_API FLinearColor IsolatedColor;
 	static NIAGARAEDITOR_API FLinearColor NotIsolatedColor;
 };
diff --git a/Engine/Plugins/FX/Niagara/Source/NiagaraEditor/Public/ViewModels/Stack/NiagaraStackEntry.h b/Engine/Plugins/FX/Niagara/Source/NiagaraEditor/Public/ViewModels/Stack/NiagaraStackEntry.h
--- a/Engine/Plugins/FX/Niagara/Source/NiagaraEditor/Public/ViewModels/Stack/NiagaraStackEntry.h
+++ b/Engine/Plugins/FX/Niagara/Source/NiagaraEditor/Public/ViewModels/Stack/NiagaraStackEntry.h
@@ -115,6 +115,7 @@ public:
 	{
 		static NIAGARAEDITOR_API const FName System;
 		static NIAGARAEDITOR_API const FName Emitter;
+		static NIAGARAEDITOR_API const FName StatelessEmitter;
 		static NIAGARAEDITOR_API const FName Particle;
 		static NIAGARAEDITOR_API const FName Render;
 	};
diff --git a/Engine/Plugins/FX/Niagara/Source/NiagaraEditorWidgets/Private/NiagaraEditorWidgetsStyle.cpp b/Engine/Plugins/FX/Niagara/Source/NiagaraEditorWidgets/Private/NiagaraEditorWidgetsStyle.cpp
--- a/Engine/Plugins/FX/Niagara/Source/NiagaraEditorWidgets/Private/NiagaraEditorWidgetsStyle.cpp
+++ b/Engine/Plugins/FX/Niagara/Source/NiagaraEditorWidgets/Private/NiagaraEditorWidgetsStyle.cpp
@@ -116,6 +116,7 @@ void FNiagaraEditorWidgetsStyle::InitColors()
 
 	Set("NiagaraEditor.Stack.IconColor.System", FLinearColor(FColor(1, 202, 252)));
 	Set("NiagaraEditor.Stack.IconColor.Emitter", FLinearColor(FColor(241, 99, 6)));
+	Set("NiagaraEditor.Stack.IconColor.StatelessEmitter", FLinearColor(FColor(255, 41, 2)));
 	Set("NiagaraEditor.Stack.IconColor.Particle", FLinearColor(FColor(131, 218, 9)));
 	Set("NiagaraEditor.Stack.IconColor.Render", FLinearColor(FColor(230, 102, 102)));
 	Set("NiagaraEditor.Stack.IconColor.VersionUpgrade", FLinearColor(FColor(255, 181, 0)));
diff --git a/Engine/Plugins/FX/Niagara/Source/NiagaraEditorWidgets/Private/NiagaraEditorWidgetsUtilities.cpp b/Engine/Plugins/FX/Niagara/Source/NiagaraEditorWidgets/Private/NiagaraEditorWidgetsUtilities.cpp
--- a/Engine/Plugins/FX/Niagara/Source/NiagaraEditorWidgets/Private/NiagaraEditorWidgetsUtilities.cpp
+++ b/Engine/Plugins/FX/Niagara/Source/NiagaraEditorWidgets/Private/NiagaraEditorWidgetsUtilities.cpp
@@ -125,6 +125,10 @@ FName FNiagaraStackEditorWidgetsUtilities::GetIconColorNameForExecutionCategory(
 	{
 		return "NiagaraEditor.Stack.IconColor.Emitter";
 	}
+	else if (ExecutionCategoryName == UNiagaraStackEntry::FExecutionCategoryNames::StatelessEmitter)
+	{
+		return "NiagaraEditor.Stack.IconColor.StatelessEmitter";
+	}
 	else if (ExecutionCategoryName == UNiagaraStackEntry::FExecutionCategoryNames::Particle)
 	{
 		return "NiagaraEditor.Stack.IconColor.Particle";
