- Change spawn name in stack to make it more clear it's rate / burst

diff --git a/Engine/Plugins/FX/Niagara/Source/NiagaraEditor/Private/ViewModels/Stack/NiagaraStackStatelessEmitterSpawnGroup.cpp b/Engine/Plugins/FX/Niagara/Source/NiagaraEditor/Private/ViewModels/Stack/NiagaraStackStatelessEmitterSpawnGroup.cpp
--- a/Engine/Plugins/FX/Niagara/Source/NiagaraEditor/Private/ViewModels/Stack/NiagaraStackStatelessEmitterSpawnGroup.cpp
+++ b/Engine/Plugins/FX/Niagara/Source/NiagaraEditor/Private/ViewModels/Stack/NiagaraStackStatelessEmitterSpawnGroup.cpp
@@ -136,7 +136,20 @@ void UNiagaraStackStatelessEmitterSpawnItem::Initialize(FRequiredEntryData InReq
 
 FText UNiagaraStackStatelessEmitterSpawnItem::GetDisplayName() const
 {
-	return LOCTEXT("EmitterSpawnDisplayName", "Spawn Info");
+	const FNiagaraStatelessSpawnInfo* SpawnInfo = GetSpawnInfo();
+	const ENiagaraStatelessSpawnInfoType SpawnInfoType = SpawnInfo ? SpawnInfo->Type : ENiagaraStatelessSpawnInfoType::Burst;
+	switch (SpawnInfoType)
+	{
+		case ENiagaraStatelessSpawnInfoType::Burst:
+			return LOCTEXT("EmitterSpawnBurstDisplayName", "Spawn Burst Instantaneous");
+
+		case ENiagaraStatelessSpawnInfoType::Rate:
+			return LOCTEXT("EmitterSpawnRateDisplayName", "Spawn Rate");
+
+		default:
+			checkNoEntry();
+			return LOCTEXT("EmitterSpawnUnknownDisplayName", "Unknown");
+	}
 }
 
 FGuid UNiagaraStackStatelessEmitterSpawnItem::GetSelectionId() const
