- Add some stateless emitter information to the viewport
- Fix a bug in fast path detection where it could be enabled for regular emitters

diff --git a/Engine/Plugins/FX/Niagara/Source/NiagaraEditor/Private/NiagaraEditorCommands.cpp b/Engine/Plugins/FX/Niagara/Source/NiagaraEditor/Private/NiagaraEditorCommands.cpp
--- a/Engine/Plugins/FX/Niagara/Source/NiagaraEditor/Private/NiagaraEditorCommands.cpp
+++ b/Engine/Plugins/FX/Niagara/Source/NiagaraEditor/Private/NiagaraEditorCommands.cpp
@@ -27,6 +27,7 @@ void FNiagaraEditorCommands::RegisterCommands()
 	UI_COMMAND(ToggleEmitterExecutionOrder, "Emitter Execution Order", "Display Emitter Execution Order", EUserInterfaceActionType::ToggleButton, FInputChord());
 	UI_COMMAND(ToggleGpuTickInformation, "Gpu Tick Information", "Display Gpu Tick Information", EUserInterfaceActionType::ToggleButton, FInputChord());
 	UI_COMMAND(ToggleMemoryInfo, "Memory", "Show approximate memory usage", EUserInterfaceActionType::ToggleButton, FInputChord());
+	UI_COMMAND(ToggleStatelessInfo, "Stateless", "Show information about stateless emitters and execution path", EUserInterfaceActionType::ToggleButton, FInputChord());
 	UI_COMMAND(TogglePreviewBackground, "Background", "Toggles the preview pane's background.", EUserInterfaceActionType::ToggleButton, FInputChord());
 	UI_COMMAND(ToggleUnlockToChanges, "Lock/Unlock To Changes", "Toggles whether or not changes in the source asset get pulled into this asset automatically.", EUserInterfaceActionType::Button, FInputChord());
 	UI_COMMAND(ToggleStatPerformance, "Performance", "Show runtime performance for particle scripts. Not supported for optimized VM byte code.", EUserInterfaceActionType::ToggleButton, FInputChord());
diff --git a/Engine/Plugins/FX/Niagara/Source/NiagaraEditor/Private/NiagaraEditorModule.cpp b/Engine/Plugins/FX/Niagara/Source/NiagaraEditor/Private/NiagaraEditorModule.cpp
--- a/Engine/Plugins/FX/Niagara/Source/NiagaraEditor/Private/NiagaraEditorModule.cpp
+++ b/Engine/Plugins/FX/Niagara/Source/NiagaraEditor/Private/NiagaraEditorModule.cpp
@@ -311,7 +311,7 @@ class FNiagaraEditorOnlyDataUtilities : public INiagaraEditorOnlyDataUtilities
 		// We can perhaps look at this again, but we always write Emitter.RandomSeed currently even with an empty script
 		for (const FNiagaraEmitterHandle& EmitterHandle : System.GetEmitterHandles())
 		{
-			if ( !EmitterHandle.GetIsEnabled() && EmitterHandle.GetEmitterMode() != ENiagaraEmitterMode::Stateless )
+			if ( EmitterHandle.GetIsEnabled() && EmitterHandle.GetEmitterMode() != ENiagaraEmitterMode::Stateless )
 			{
 				return SystemStateData;
 			}
@@ -324,7 +324,10 @@ class FNiagaraEditorOnlyDataUtilities : public INiagaraEditorOnlyDataUtilities
 		{
 			const TCHAR* SystemStateName = TEXT("/Niagara/Modules/System/SystemState.SystemState");
 			TArray<UNiagaraNodeFunctionCall*> Nodes;
-			ScriptSource->NodeGraph->GetNodesOfClass<UNiagaraNodeFunctionCall>(Nodes);
+			if (ensure(ScriptSource->NodeGraph))
+			{
+				ScriptSource->NodeGraph->GetNodesOfClass<UNiagaraNodeFunctionCall>(Nodes);
+			}
 			Nodes.RemoveAll([](UNiagaraNodeFunctionCall* Node) { return !Node || !Node->IsNodeEnabled(); });
 
 			// No function calls, we can enable fast path with the empty system state
diff --git a/Engine/Plugins/FX/Niagara/Source/NiagaraEditor/Private/NiagaraEditorSettings.cpp b/Engine/Plugins/FX/Niagara/Source/NiagaraEditor/Private/NiagaraEditorSettings.cpp
--- a/Engine/Plugins/FX/Niagara/Source/NiagaraEditor/Private/NiagaraEditorSettings.cpp
+++ b/Engine/Plugins/FX/Niagara/Source/NiagaraEditor/Private/NiagaraEditorSettings.cpp
@@ -390,6 +390,20 @@ void UNiagaraEditorSettings::SetShowMemoryInfo(bool bInShowInfo)
 	}
 }
 
+bool UNiagaraEditorSettings::IsShowStatelessInfo() const
+{
+	return ViewportSettings.bShowStatelessInfo;
+}
+
+void UNiagaraEditorSettings::SetShowStatelessInfo(bool bInShowInfo)
+{
+	if (ViewportSettings.bShowStatelessInfo != bInShowInfo)
+	{
+		ViewportSettings.bShowStatelessInfo = bInShowInfo;
+		SaveConfig();
+	}
+}
+
 TArray<float> UNiagaraEditorSettings::GetPlaybackSpeeds() const
 {
 	if(!CachedPlaybackSpeeds.IsSet())
diff --git a/Engine/Plugins/FX/Niagara/Source/NiagaraEditor/Private/NiagaraScriptMergeManager.cpp b/Engine/Plugins/FX/Niagara/Source/NiagaraEditor/Private/NiagaraScriptMergeManager.cpp
--- a/Engine/Plugins/FX/Niagara/Source/NiagaraEditor/Private/NiagaraScriptMergeManager.cpp
+++ b/Engine/Plugins/FX/Niagara/Source/NiagaraEditor/Private/NiagaraScriptMergeManager.cpp
@@ -1565,6 +1565,8 @@ INiagaraMergeManager::FMergeEmitterResults FNiagaraScriptMergeManager::MergeEmit
 	}
 	else if (State == EmitterMergeState::DuplicateParent)
 	{
+		TRACE_CPUPROFILER_EVENT_SCOPE(EmitterMergeState::DuplicateParent);
+
 		// If there were differences from the parent or the parent diff failed we can just return a copy of the parent as the merged instance since there
 		// were no changes in the instance which need to be applied.
 		MergeResults.MergeResult = EMergeEmitterResult::SucceededDifferencesApplied;
@@ -1575,6 +1577,8 @@ INiagaraMergeManager::FMergeEmitterResults FNiagaraScriptMergeManager::MergeEmit
 	}
 	else if (State == EmitterMergeState::Failed)
 	{
+		TRACE_CPUPROFILER_EVENT_SCOPE(EmitterMergeState::Failed);
+
 		MergeResults.MergeResult = EMergeEmitterResult::FailedToDiff;
 		MergeResults.ErrorMessages = DiffResults.GetErrorMessages();
 
@@ -1604,6 +1608,8 @@ INiagaraMergeManager::FMergeEmitterResults FNiagaraScriptMergeManager::MergeEmit
 	}
 	else if (State == EmitterMergeState::CopyProperties) //-V547
 	{
+		TRACE_CPUPROFILER_EVENT_SCOPE(EmitterMergeState::CopyProperties);
+
 		UNiagaraEmitter* MergedInstance = Parent.Emitter->DuplicateWithoutMerging(GetTransientPackage());
 		MergedInstance->DisableVersioning(Parent.Version);
 		FVersionedNiagaraEmitter VersionedMergedInstance = FVersionedNiagaraEmitter(MergedInstance, Parent.Version);
@@ -1644,12 +1650,19 @@ INiagaraMergeManager::FMergeEmitterResults FNiagaraScriptMergeManager::MergeEmit
 			}
 		};
 
-		PopulateIdToNodeMap(MakeShared<FNiagaraEmitterMergeAdapter>(Parent), ParentFunctionIdToNodeMap);
-		PopulateIdToNodeMap(DiffResults.BaseEmitterAdapter.ToSharedRef(), LastMergedParentFunctionIdToNodeMap);
-		PopulateIdToNodeMap(DiffResults.OtherEmitterAdapter.ToSharedRef(), InstanceFunctionIdToNodeMap);
+
+		{
+			TRACE_CPUPROFILER_EVENT_SCOPE(Chunk1);
+			PopulateIdToNodeMap(MakeShared<FNiagaraEmitterMergeAdapter>(Parent), ParentFunctionIdToNodeMap);
+			PopulateIdToNodeMap(DiffResults.BaseEmitterAdapter.ToSharedRef(), LastMergedParentFunctionIdToNodeMap);
+			PopulateIdToNodeMap(DiffResults.OtherEmitterAdapter.ToSharedRef(), InstanceFunctionIdToNodeMap);
+		}
 
 		TMap<FGuid, FGuid> FunctionIdToForcedChangeId;
-		GetForcedChangeIds(ParentFunctionIdToNodeMap, LastMergedParentFunctionIdToNodeMap, InstanceFunctionIdToNodeMap, FunctionIdToForcedChangeId);
+		{
+			TRACE_CPUPROFILER_EVENT_SCOPE(Chunk2);
+			GetForcedChangeIds(ParentFunctionIdToNodeMap, LastMergedParentFunctionIdToNodeMap, InstanceFunctionIdToNodeMap, FunctionIdToForcedChangeId);
+		}
 
 		FVersionedNiagaraEmitterData* MergedEmitterData = VersionedMergedInstance.GetEmitterData();
 		MergedEmitterData->ParentScratchPads->AppendScripts(MergedEmitterData->ScratchPads);
@@ -2835,6 +2848,8 @@ FNiagaraScriptMergeManager::FApplyDiffResults FNiagaraScriptMergeManager::AddMod
 	UNiagaraNodeOutput& TargetOutputNode,
 	TSharedRef<FNiagaraStackFunctionMergeAdapter> AddModule) const
 {
+	TRACE_CPUPROFILER_EVENT_SCOPE(FNiagaraScriptMergeManager::AddModule);
+
 	FApplyDiffResults Results;
 
 	UNiagaraNodeFunctionCall* AddedModuleNode = nullptr;
@@ -2891,6 +2906,8 @@ FNiagaraScriptMergeManager::FApplyDiffResults FNiagaraScriptMergeManager::AddMod
 
 FNiagaraScriptMergeManager::FApplyDiffResults FNiagaraScriptMergeManager::RemoveInputOverride(UNiagaraScript& OwningScript, TSharedRef<FNiagaraStackFunctionInputOverrideMergeAdapter> OverrideToRemove) const
 {
+	TRACE_CPUPROFILER_EVENT_SCOPE(FNiagaraScriptMergeManager::RemoveInputOverride);
+
 	FApplyDiffResults Results;
 	Results.bSucceeded = true;
 	Results.bModifiedGraph = false;
@@ -2954,6 +2971,8 @@ FNiagaraScriptMergeManager::FApplyDiffResults FNiagaraScriptMergeManager::AddInp
 	UNiagaraNodeFunctionCall& TargetFunctionCall,
 	TSharedRef<FNiagaraStackFunctionInputOverrideMergeAdapter> OverrideToAdd) const
 {
+	TRACE_CPUPROFILER_EVENT_SCOPE(FNiagaraScriptMergeManager::AddInputOverride);
+
 	FApplyDiffResults Results;
 
 	// If an assignment node, make sure that we have an assignment target for the input override.
@@ -3221,6 +3240,8 @@ FNiagaraScriptMergeManager::FApplyDiffResults FNiagaraScriptMergeManager::ApplyS
 	const FNiagaraScriptStackDiffResults& DiffResults,
 	const bool bNoParentAtLastMerge) const
 {
+	TRACE_CPUPROFILER_EVENT_SCOPE(FNiagaraScriptMergeManager::ApplyScriptStackDiff);
+
 	FApplyDiffResults Results;
 
 	if (DiffResults.IsEmpty())
@@ -3267,6 +3288,8 @@ FNiagaraScriptMergeManager::FApplyDiffResults FNiagaraScriptMergeManager::ApplyS
 
 	for (TSharedRef<FNiagaraStackFunctionInputOverrideMergeAdapter> RemovedInputOverrideAdapter : DiffResults.RemovedBaseInputOverrides)
 	{
+		TRACE_CPUPROFILER_EVENT_SCOPE(FNiagaraStackGraphUtilities::RemovedInputOverrideAdapter);
+
 		TSharedPtr<FNiagaraStackFunctionMergeAdapter> MatchingModuleAdapter = BaseScriptStackAdapter->GetModuleFunctionById(RemovedInputOverrideAdapter->GetOwningFunctionCall()->NodeGuid);
 		if (MatchingModuleAdapter.IsValid())
 		{
@@ -3281,6 +3304,8 @@ FNiagaraScriptMergeManager::FApplyDiffResults FNiagaraScriptMergeManager::ApplyS
 
 	for (TSharedRef<FNiagaraStackFunctionInputOverrideMergeAdapter> AddedInputOverrideAdapter : DiffResults.AddedOtherInputOverrides)
 	{
+		TRACE_CPUPROFILER_EVENT_SCOPE(FNiagaraStackGraphUtilities::AddedInputOverrideAdapter);
+
 		TSharedPtr<FNiagaraStackFunctionMergeAdapter> MatchingModuleAdapter = BaseScriptStackAdapter->GetModuleFunctionById(AddedInputOverrideAdapter->GetOwningFunctionCall()->NodeGuid);
 		if (MatchingModuleAdapter.IsValid())
 		{
@@ -3300,6 +3325,8 @@ FNiagaraScriptMergeManager::FApplyDiffResults FNiagaraScriptMergeManager::ApplyS
 
 	for (TSharedRef<FNiagaraStackFunctionInputOverrideMergeAdapter> ModifiedInputOverrideAdapter : DiffResults.ModifiedOtherInputOverrides)
 	{
+		TRACE_CPUPROFILER_EVENT_SCOPE(FNiagaraStackGraphUtilities::ModifiedInputOverrideAdapter);
+
 		TSharedPtr<FNiagaraStackFunctionMergeAdapter> MatchingModuleAdapter = BaseScriptStackAdapter->GetModuleFunctionById(ModifiedInputOverrideAdapter->GetOwningFunctionCall()->NodeGuid);
 		if (MatchingModuleAdapter.IsValid())
 		{
@@ -3343,6 +3370,8 @@ FNiagaraScriptMergeManager::FApplyDiffResults FNiagaraScriptMergeManager::ApplyS
 	// Apply the graph actions.
 	for (TSharedRef<FNiagaraStackFunctionMergeAdapter> RemoveModule : RemoveModules)
 	{
+		TRACE_CPUPROFILER_EVENT_SCOPE(FNiagaraStackGraphUtilities::RemoveModuleFromStack);
+
 		bool bRemoveResults = FNiagaraStackGraphUtilities::RemoveModuleFromStack(*BaseScriptStackAdapter->GetScript(), *RemoveModule->GetFunctionCallNode());
 		if (bRemoveResults == false)
 		{
@@ -3381,6 +3410,7 @@ FNiagaraScriptMergeManager::FApplyDiffResults FNiagaraScriptMergeManager::ApplyS
 	}
 
 	// Apply enabled state last so that it applies to function calls added  from input overrides;
+	TRACE_CPUPROFILER_EVENT_SCOPE(FNiagaraScriptMergeManager::SetModuleIsEnabled);
 	for (TSharedRef<FNiagaraStackFunctionMergeAdapter> EnableModule : EnableModules)
 	{
 		FNiagaraStackGraphUtilities::SetModuleIsEnabled(*EnableModule->GetFunctionCallNode(), true);
@@ -3399,6 +3429,8 @@ FNiagaraScriptMergeManager::FApplyDiffResults FNiagaraScriptMergeManager::ApplyE
 	const FNiagaraEmitterDiffResults& DiffResults,
 	const bool bNoParentAtLastMerge) const
 {
+	TRACE_CPUPROFILER_EVENT_SCOPE(FNiagaraScriptMergeManager::ApplyEventHandlerDiff);
+
 	FApplyDiffResults Results;
 	if (DiffResults.RemovedBaseEventHandlers.Num() > 0)
 	{
@@ -3507,6 +3539,8 @@ FNiagaraScriptMergeManager::FApplyDiffResults FNiagaraScriptMergeManager::ApplyS
 	const FNiagaraEmitterDiffResults& DiffResults,
 	const bool bNoParentAtLastMerge) const
 {
+	TRACE_CPUPROFILER_EVENT_SCOPE(FNiagaraScriptMergeManager::ApplySimulationStageDiff);
+
 	FApplyDiffResults Results;
 	if (DiffResults.RemovedBaseSimulationStages.Num() > 0)
 	{
@@ -3614,6 +3648,8 @@ FNiagaraScriptMergeManager::FApplyDiffResults FNiagaraScriptMergeManager::ApplyS
 
 FNiagaraScriptMergeManager::FApplyDiffResults FNiagaraScriptMergeManager::ApplyRendererDiff(const FVersionedNiagaraEmitter& BaseEmitter, const FNiagaraEmitterDiffResults& DiffResults, const bool bNoParentAtLastMerge) const
 {
+	TRACE_CPUPROFILER_EVENT_SCOPE(FNiagaraScriptMergeManager::ApplyRendererDiff);
+
 	TArray<UNiagaraRendererProperties*> RenderersToRemove;
 	TArray<UNiagaraRendererProperties*> RenderersToAdd;
 	TArray<UNiagaraRendererProperties*> RenderersToDisable;
@@ -3678,6 +3714,8 @@ FNiagaraScriptMergeManager::FApplyDiffResults FNiagaraScriptMergeManager::ApplyR
 
 FNiagaraScriptMergeManager::FApplyDiffResults FNiagaraScriptMergeManager::ApplyEmitterSummaryDiff(const FVersionedNiagaraEmitter& BaseEmitter, const FNiagaraEmitterDiffResults& DiffResults) const
 {
+	TRACE_CPUPROFILER_EVENT_SCOPE(FNiagaraScriptMergeManager::ApplyEmitterSummaryDiff);
+
 	FVersionedNiagaraEmitterData* BaseEmitterData = BaseEmitter.GetEmitterData();
 	UNiagaraEmitterEditorData* EditorData = Cast<UNiagaraEmitterEditorData>(BaseEmitterData->GetEditorData());
 
@@ -3769,6 +3807,7 @@ FNiagaraScriptMergeManager::FApplyDiffResults FNiagaraScriptMergeManager::ApplyE
 
 FNiagaraScriptMergeManager::FApplyDiffResults FNiagaraScriptMergeManager::ApplyStackEntryDisplayNameDiffs(FVersionedNiagaraEmitter BaseEmitter, const FNiagaraEmitterDiffResults& DiffResults) const
 {
+	TRACE_CPUPROFILER_EVENT_SCOPE(FNiagaraScriptMergeManager::ApplyStackEntryDisplayNameDiffs);
 	if (DiffResults.ModifiedStackEntryDisplayNames.Num() > 0)
 	{
 		UNiagaraEmitterEditorData* EditorData = Cast<UNiagaraEmitterEditorData>(BaseEmitter.GetEmitterData()->GetEditorData());
@@ -3787,6 +3826,7 @@ FNiagaraScriptMergeManager::FApplyDiffResults FNiagaraScriptMergeManager::ApplyS
 
 FNiagaraScriptMergeManager::FApplyDiffResults FNiagaraScriptMergeManager::ApplyStackNoteDiffs(FVersionedNiagaraEmitter BaseEmitter, const FNiagaraEmitterDiffResults& DiffResults) const
 {
+	TRACE_CPUPROFILER_EVENT_SCOPE(FNiagaraScriptMergeManager::ApplyStackNoteDiffs);
 	if(DiffResults.AddedOrModifiedStackNotesInOther.Num() > 0)
 	{
 		if(UNiagaraEmitterEditorData* EmitterEditorData = Cast<UNiagaraEmitterEditorData>(BaseEmitter.GetEmitterData()->GetEditorData()))
diff --git a/Engine/Plugins/FX/Niagara/Source/NiagaraEditor/Private/Widgets/SNiagaraSystemViewport.cpp b/Engine/Plugins/FX/Niagara/Source/NiagaraEditor/Private/Widgets/SNiagaraSystemViewport.cpp
--- a/Engine/Plugins/FX/Niagara/Source/NiagaraEditor/Private/Widgets/SNiagaraSystemViewport.cpp
+++ b/Engine/Plugins/FX/Niagara/Source/NiagaraEditor/Private/Widgets/SNiagaraSystemViewport.cpp
@@ -69,6 +69,7 @@ public:
 	void DrawEmitterExecutionOrder(UNiagaraComponent* Component, FCanvas* Canvas, float& CurrentX, float& CurrentY, UFont* Font, const float FontHeight);
 	void DrawGpuTickInformation(UNiagaraComponent* Component, FCanvas* Canvas, float& CurrentX, float& CurrentY, UFont* Font, const float FontHeight);
 	void DrawMemoryInfo(UNiagaraComponent* Component, FCanvas* Canvas, float& CurrentX, float& CurrentY, UFont* Font, const float FontHeight);
+	void DrawStatelessInfo(UNiagaraComponent* Component, FViewport* InViewport, FCanvas* Canvas, UFont* Font, const float FontHeight);
 	void SetUpdateViewportFocus(bool bUpdate) { bUpdateViewportFocus = bUpdate; }
 
 	TWeakPtr<SNiagaraSystemViewport> NiagaraViewportPtr;
@@ -251,6 +252,11 @@ void FNiagaraSystemViewportClient::Draw(FViewport* InViewport, FCanvas* Canvas)
 			DrawParticleCounts(Component, Canvas, CurrentX, CurrentY, Font, FontHeight);
 			CurrentY += FontHeight;
 		}
+		if (NiagaraViewport->GetDrawElement(SNiagaraSystemViewport::EDrawElements::StatelessInfo) && Component)
+		{
+			UFont* TinyFont = GEngine->GetTinyFont();
+			DrawStatelessInfo(Component, InViewport, Canvas, TinyFont, TinyFont->GetMaxCharHeight() * 1.1f);
+		}
 	}
 
 	if (bCaptureScreenShot && ScreenShotOwner.IsValid() && OnScreenShotCaptured.IsBound())
@@ -351,19 +357,21 @@ void FNiagaraSystemViewportClient::DrawParticleCounts(UNiagaraComponent* Compone
 
 		for (const FNiagaraEmitterInstanceRef& EmitterInstance : SystemInstance->GetEmitters())
 		{
-			FVersionedNiagaraEmitterData* EmitterData = EmitterInstance->GetEmitterHandle().GetEmitterData();
+			const FNiagaraEmitterHandle& EmitterHandle = EmitterInstance->GetEmitterHandle();
+			FVersionedNiagaraEmitterData* EmitterData = EmitterHandle.GetEmitterData();
 
-			const FName EmitterName = EmitterInstance->GetEmitterHandle().GetName();
+			const FName EmitterName = EmitterHandle.GetName();
 			const int32 CurrentCount = EmitterInstance->GetNumParticles();
 			const int32 MaxCount = EmitterData ? EmitterData->GetMaxParticleCountEstimate() : 0;
-			const bool IsIsolated = EmitterInstance->GetEmitterHandle().IsIsolated();
-			const bool IsEnabled = EmitterInstance->GetEmitterHandle().GetIsEnabled();
+			const bool IsIsolated = EmitterHandle.IsIsolated();
+			const bool IsEnabled = EmitterHandle.GetIsEnabled();
 			const ENiagaraExecutionState ExecutionState = EmitterInstance->GetExecutionState();
 			const FString EmitterExecutionString = UEnum::GetValueAsString(ExecutionState);
 			const int32 EmitterExecutionStringValueIndex = EmitterExecutionString.Find(TEXT("::"));
 			const TCHAR* EmitterExecutionText = EmitterExecutionStringValueIndex == INDEX_NONE ? *EmitterExecutionString : *EmitterExecutionString + EmitterExecutionStringValueIndex + 2;
+			const TCHAR* EmitterMode = EmitterHandle.GetEmitterMode() == ENiagaraEmitterMode::Stateless ? TEXT("[Stateless]") : TEXT("");
 
-			TextItem.Text = FText::FromString(FString::Printf(TEXT("%i Current, %i Max (est.) - [%s] [%s]"), CurrentCount, MaxCount, *EmitterName.ToString(), EmitterExecutionText));
+			TextItem.Text = FText::FromString(FString::Printf(TEXT("%i Current, %i Max (est.) - [%s] [%s] %s"), CurrentCount, MaxCount, *EmitterName.ToString(), EmitterExecutionText, EmitterMode));
 			TextItem.Position = FVector2D(CurrentX, CurrentY);
 			TextItem.bOutlined = IsIsolated;
 			TextItem.OutlineColor = FLinearColor(0.7f, 0.0f, 0.0f);
@@ -531,6 +539,51 @@ void FNiagaraSystemViewportClient::DrawMemoryInfo(UNiagaraComponent* Component,
 	CurrentY += FontHeight;
 }
 
+void FNiagaraSystemViewportClient::DrawStatelessInfo(UNiagaraComponent* Component, FViewport* InViewport, FCanvas* Canvas, UFont* Font, const float FontHeight)
+{
+	UNiagaraSystem* NiagaraSystem = Component->GetAsset();
+	if (!NiagaraSystem)
+	{
+		return;
+	}
+
+	int32 NumStatelessEmitters = 0;
+	int32 NumEmitters = 0;
+
+	for (const FNiagaraEmitterHandle& EmitterHandle : NiagaraSystem->GetEmitterHandles())
+	{
+		if (EmitterHandle.GetIsEnabled())
+		{
+			++NumEmitters;
+			UNiagaraStatelessEmitter* StatelessEmitter = EmitterHandle.GetEmitterMode() == ENiagaraEmitterMode::Stateless ? EmitterHandle.GetStatelessEmitter() : nullptr;
+			if (StatelessEmitter)
+			{
+				++NumStatelessEmitters;
+			}
+		}
+	}
+
+	if (NumStatelessEmitters > 0)
+	{
+		const FVector2D ScaledViewportSize = FVector2D(InViewport->GetSizeXY()) / Canvas->GetDPIScale();
+		FCanvasTextItem TextItem(FVector2D::ZeroVector, FText::GetEmpty(), Font, FLinearColor::White);
+
+		FString StatelessInfo = FString::Printf(TEXT("Stateless: %d / %d Emitters"), NumStatelessEmitters, NumEmitters);
+		if (NiagaraSystem->SystemStateFastPathEnabled())
+		{
+			StatelessInfo.Append(TEXT(" [FastPath]"));
+		}
+
+		int32 Width = 0;
+		int32 Height = 0;
+		UCanvas::ClippedStrLen(Font, 1.0f, 1.0f, Width, Height, *StatelessInfo);
+
+		TextItem.Text = FText::FromString(StatelessInfo);
+		TextItem.EnableShadow(FLinearColor::Black);
+		TextItem.Draw(Canvas, ScaledViewportSize.X - 5 - float(Width), ScaledViewportSize.Y - 5 - float(Height));
+	}
+}
+
 void FNiagaraSystemViewportClient::SetOrbitModeFromSettings()
 {
 	const UNiagaraSettings* Settings = GetDefault<UNiagaraSettings>();
@@ -705,6 +758,7 @@ void SNiagaraSystemViewport::Construct(const FArguments& InArgs, TSharedRef<FNia
 	DrawFlags |= Settings->IsShowEmitterExecutionOrder() ? EDrawElements::EmitterExecutionOrder : 0;
 	DrawFlags |= Settings->IsShowGpuTickInformation() ? EDrawElements::GpuTickInformation : 0;
 	DrawFlags |= Settings->IsShowMemoryInfo() ? EDrawElements::MemoryInfo : 0;
+	DrawFlags |= Settings->IsShowStatelessInfo() ? EDrawElements::StatelessInfo : 0;
 
 	bShowBackground = false;
 	PreviewComponent = nullptr;
@@ -994,6 +1048,18 @@ void SNiagaraSystemViewport::BindCommands()
 		FIsActionChecked::CreateLambda([Viewport = this]() -> bool { return Viewport->GetDrawElement(EDrawElements::MemoryInfo); })
 	);
 
+	CommandList->MapAction(
+		Commands.ToggleStatelessInfo,
+		FExecuteAction::CreateLambda([Viewport = this]()
+		{
+			Viewport->ToggleDrawElement(EDrawElements::StatelessInfo);
+			GetMutableDefault<UNiagaraEditorSettings>()->SetShowStatelessInfo(Viewport->GetDrawElement(EDrawElements::StatelessInfo));
+			Viewport->RefreshViewport();
+		}),
+		FCanExecuteAction(),
+		FIsActionChecked::CreateLambda([Viewport = this]() -> bool { return Viewport->GetDrawElement(EDrawElements::StatelessInfo); })
+	);
+
 	CommandList->MapAction(
 		Commands.TogglePreviewBackground,
 		FExecuteAction::CreateSP( this, &SNiagaraSystemViewport::TogglePreviewBackground ),
diff --git a/Engine/Plugins/FX/Niagara/Source/NiagaraEditor/Private/Widgets/SNiagaraSystemViewport.h b/Engine/Plugins/FX/Niagara/Source/NiagaraEditor/Private/Widgets/SNiagaraSystemViewport.h
--- a/Engine/Plugins/FX/Niagara/Source/NiagaraEditor/Private/Widgets/SNiagaraSystemViewport.h
+++ b/Engine/Plugins/FX/Niagara/Source/NiagaraEditor/Private/Widgets/SNiagaraSystemViewport.h
@@ -84,6 +84,7 @@ public:
 		EmitterExecutionOrder	= 0x100,
 		GpuTickInformation		= 0x200,
 		MemoryInfo				= 0x400,
+		StatelessInfo			= 0x800,
 	};
 
 	bool GetDrawElement(EDrawElements Element) const;
diff --git a/Engine/Plugins/FX/Niagara/Source/NiagaraEditor/Private/Widgets/SNiagaraSystemViewportToolBar.cpp b/Engine/Plugins/FX/Niagara/Source/NiagaraEditor/Private/Widgets/SNiagaraSystemViewportToolBar.cpp
--- a/Engine/Plugins/FX/Niagara/Source/NiagaraEditor/Private/Widgets/SNiagaraSystemViewportToolBar.cpp
+++ b/Engine/Plugins/FX/Niagara/Source/NiagaraEditor/Private/Widgets/SNiagaraSystemViewportToolBar.cpp
@@ -43,6 +43,7 @@ TSharedRef<SWidget> SNiagaraSystemViewportToolBar::GenerateShowMenu() const
 		ShowMenuBuilder.AddMenuEntry(Commands.ToggleInstructionCounts);
 		ShowMenuBuilder.AddMenuEntry(Commands.ToggleMemoryInfo);
 		ShowMenuBuilder.AddMenuEntry(Commands.ToggleParticleCounts);
+		ShowMenuBuilder.AddMenuEntry(Commands.ToggleStatelessInfo);
 	}
 
 	return ShowMenuBuilder.MakeWidget();
diff --git a/Engine/Plugins/FX/Niagara/Source/NiagaraEditor/Public/NiagaraEditorCommands.h b/Engine/Plugins/FX/Niagara/Source/NiagaraEditor/Public/NiagaraEditorCommands.h
--- a/Engine/Plugins/FX/Niagara/Source/NiagaraEditor/Public/NiagaraEditorCommands.h
+++ b/Engine/Plugins/FX/Niagara/Source/NiagaraEditor/Public/NiagaraEditorCommands.h
@@ -42,6 +42,7 @@ public:
 	TSharedPtr<FUICommandInfo> ToggleEmitterExecutionOrder;
 	TSharedPtr<FUICommandInfo> ToggleGpuTickInformation;
 	TSharedPtr<FUICommandInfo> ToggleMemoryInfo;
+	TSharedPtr<FUICommandInfo> ToggleStatelessInfo;
 
 	/** Toggles the preview pane's background */
 	TSharedPtr< FUICommandInfo > TogglePreviewBackground;
diff --git a/Engine/Plugins/FX/Niagara/Source/NiagaraEditor/Public/NiagaraEditorSettings.h b/Engine/Plugins/FX/Niagara/Source/NiagaraEditor/Public/NiagaraEditorSettings.h
--- a/Engine/Plugins/FX/Niagara/Source/NiagaraEditor/Public/NiagaraEditorSettings.h
+++ b/Engine/Plugins/FX/Niagara/Source/NiagaraEditor/Public/NiagaraEditorSettings.h
@@ -330,6 +330,9 @@ struct FNiagaraViewportSharedSettings
 
 	UPROPERTY(config)
 	bool bShowMemoryInfo = false;
+
+	UPROPERTY(config)
+	bool bShowStatelessInfo = true;
 };
 
 FORCEINLINE uint32 GetTypeHash(const FNiagaraNamespaceMetadata& NamespaceMetaData)
@@ -613,4 +616,6 @@ public:
 	NIAGARAEDITOR_API void SetShowGpuTickInformation(bool bShowGpuTickInformation);
 	NIAGARAEDITOR_API bool IsShowMemoryInfo() const;
 	NIAGARAEDITOR_API void SetShowMemoryInfo(bool bInShowInfo);
+	NIAGARAEDITOR_API bool IsShowStatelessInfo() const;
+	NIAGARAEDITOR_API void SetShowStatelessInfo(bool bInShowInfo);
 };
