- Leak detector for scalability instance count vs execution active state count

diff --git a/Engine/Plugins/FX/Niagara/Source/Niagara/Private/NiagaraComponentLeakDetector.cpp b/Engine/Plugins/FX/Niagara/Source/Niagara/Private/NiagaraComponentLeakDetector.cpp
--- a/Engine/Plugins/FX/Niagara/Source/Niagara/Private/NiagaraComponentLeakDetector.cpp
+++ b/Engine/Plugins/FX/Niagara/Source/Niagara/Private/NiagaraComponentLeakDetector.cpp
@@ -2,6 +2,7 @@
 
 #include "NiagaraComponentLeakDetector.h"
 #include "NiagaraDebugHud.h"
+#include "NiagaraEffectType.h"
 #include "NiagaraSystem.h"
 #include "NiagaraWorldManager.h"
 
@@ -11,8 +12,17 @@ DEFINE_LOG_CATEGORY_STATIC(LogNiagaraLeakDetector, Log, All);
 
 namespace NiagaraComponentLeakDetectorPrivate
 {
+	enum class EReportLeakType
+	{
+		Never = 0,
+		Immediate,
+		PostGC,
+	};
+
 	static bool		GEnabled = false;
-	static int32	GImmediateReport = 2;
+	static int32	GReportActiveLeaks = (int32)EReportLeakType::Immediate;
+	static int32	GReportTotalLeaks = (int32)EReportLeakType::PostGC;
+	static int32	GReportScalabilityIssues = (int32)EReportLeakType::Immediate;
 	static int32	GGCReport = 1;
 	static float	GTickDeltaSeconds = 1.0f;
 	static int32	GGrowthCountThreshold = 16;
@@ -39,23 +49,23 @@ namespace NiagaraComponentLeakDetectorPrivate
 		ECVF_Default
 	);
 
-	static FAutoConsoleVariableRef CVarImmediateReport(
-		TEXT("fx.Niagara.LeakDetector.ImmediateReport"),
-		GImmediateReport,
-		TEXT("Controls how we report leaks as we scan information.\n")
-		TEXT("0 - Never report during gameplay.\n")
-		TEXT("1 - Report all leaks.")
-		TEXT("2 - Report active leaks only. (default).\n"),
+	static FAutoConsoleVariableRef CVarReportActiveLeaks(
+		TEXT("fx.Niagara.LeakDetector.ReportActiveLeaks"),
+		GReportActiveLeaks,
+		TEXT("How do we report active components leaks?")
+		TEXT("0 - Never report.")
+		TEXT("1 - Report immediately. (default)")
+		TEXT("2 - Report on GC.\n"),
 		ECVF_Default
 	);
 
-	static FAutoConsoleVariableRef CVarGGCReport(
-		TEXT("fx.Niagara.LeakDetector.GCReport"),
-		GGCReport,
-		TEXT("Controls how we report leak after GC\n")
-		TEXT("0 - Never report after GC.\n")
-		TEXT("1 - Report all leaks. (default)")
-		TEXT("2 - Report active leaks only.\n"),
+	static FAutoConsoleVariableRef CVarReportTotalLeaks(
+		TEXT("fx.Niagara.LeakDetector.ReportTotalLeaks"),
+		GReportTotalLeaks,
+		TEXT("How do we report total components leaks?")
+		TEXT("0 - Never report.")
+		TEXT("1 - Report immediately.")
+		TEXT("2 - Report on GC. (default)\n"),
 		ECVF_Default
 	);
 
@@ -79,6 +89,54 @@ namespace NiagaraComponentLeakDetectorPrivate
 		}
 #endif
 	}
+
+	void ReportLeak(EReportLeakType ReportLeakType, UWorld* World, FName SystemName, FNiagaraComponentLeakDetector::FSystemData& SystemData)
+	{
+		using namespace NiagaraComponentLeakDetectorPrivate;
+
+		if (int32(ReportLeakType) == GReportTotalLeaks)
+		{
+			if (!SystemData.bTotalHasWarned)
+			{
+				if (SystemData.TotalShrinkCounter == 0 && SystemData.TotalGrowthCounter >= GGrowthCountThreshold)
+				{
+					SystemData.bTotalHasWarned = true;
+					AddLeakWarning(World, SystemName, FString::Printf(TEXT("Potential componment leak System(%s) (Total:%d), please investigate."), *SystemName.ToString(), SystemData.TotalPrevCount));
+				}
+			}
+			else if (!SystemData.bTotalFalseWarning && (SystemData.TotalShrinkCounter > 0))
+			{
+				SystemData.bTotalFalseWarning = true;
+				AddLeakWarning(World, SystemName, FString::Printf(TEXT("Potential invalid component leak reported for System(%s) (Total:%d)."), *SystemName.ToString(), SystemData.TotalPrevCount));
+			}
+		}
+
+		if (int32(ReportLeakType) == GReportActiveLeaks)
+		{
+			if (!SystemData.bActiveHasWarned)
+			{
+				if (SystemData.ActiveShrinkCounter == 0 && SystemData.ActiveGrowthCounter >= GGrowthCountThreshold)
+				{
+					SystemData.bActiveHasWarned = true;
+					AddLeakWarning(World, SystemName, FString::Printf(TEXT("Potential Active Component leak System(%s) (Active:%d), please investigate."), *SystemName.ToString(), SystemData.ActivePrevCount));
+				}
+			}
+			else if (!SystemData.bActiveFalseWarning && (SystemData.ActiveShrinkCounter > 0))
+			{
+				SystemData.bActiveFalseWarning = true;
+				AddLeakWarning(World, SystemName, FString::Printf(TEXT("Potential invalid active component leak reported for System(%s) (Active: %d)."), *SystemName.ToString(), SystemData.ActivePrevCount));
+			}
+		}
+
+		if ( int32(ReportLeakType) == GReportScalabilityIssues )
+		{
+			if (!SystemData.bScalabilityHasWarned && SystemData.ScalabilityAllowedActiveExecState > 0 && SystemData.ScalabilityMaxActiveExecState > SystemData.ScalabilityAllowedActiveExecState )
+			{
+				SystemData.bScalabilityHasWarned = true;
+				AddLeakWarning(World, SystemName, FString::Printf(TEXT("Scalability instance count limit blown for System(%s) (Active: %d) (Limit: %d)."), *SystemName.ToString(), SystemData.ScalabilityMaxActiveExecState, SystemData.ScalabilityAllowedActiveExecState));
+			}
+		}
+	}
 }
 
 void FNiagaraComponentLeakDetector::Tick(UWorld* World)
@@ -112,13 +170,22 @@ void FNiagaraComponentLeakDetector::Tick(UWorld* World)
 		}
 
 		const FName SystemName = System->GetFName();
-		FSystemData& SystemData = PerSystemData.FindOrAdd(SystemName);
-		++SystemData.TotalCurrCount;
-		SystemData.ActiveCurrCount += Component->IsActive();
-	}
+		FSystemData* SystemData = PerSystemData.Find(SystemName);
+		if (!SystemData)
+		{
+			UNiagaraEffectType* EffectType = System->GetEffectType();
+
+			SystemData = &PerSystemData.Add(SystemName);
+			SystemData->ScalabilityAllowedActiveExecState = EffectType ? EffectType->GetActiveSystemScalabilitySettings().MaxSystemInstances : 0;
+		}
+
+		const bool bIsComponentActive = Component->IsActive();
+		const bool bIsExecutionActive = bIsComponentActive && (Component->GetExecutionState() == ENiagaraExecutionState::Active);
 
-	const bool bReportTotalLeaks = GImmediateReport == 1;
-	const bool bReportActiveLeaks = GImmediateReport != 0;
+		SystemData->TotalCurrCount += 1;
+		SystemData->ActiveCurrCount += bIsComponentActive ? 1 : 0;
+		SystemData->ScalabilityCurrActiveExecState += bIsExecutionActive ? 1 : 0;
+	}
 
 	for (auto SystemDataIt=PerSystemData.CreateIterator(); SystemDataIt; ++SystemDataIt)
 	{
@@ -137,8 +204,12 @@ void FNiagaraComponentLeakDetector::Tick(UWorld* World)
 		SystemData.ActivePrevCount = SystemData.ActiveCurrCount;
 		SystemData.ActiveCurrCount = 0;
 
+		// Monitor for scalability issues
+		SystemData.ScalabilityMaxActiveExecState = FMath::Max(SystemData.ScalabilityMaxActiveExecState, SystemData.ScalabilityCurrActiveExecState);
+		SystemData.ScalabilityCurrActiveExecState = 0;
+
 		// Report any leaks
-		ReportLeak(World, SystemName, SystemData, bReportTotalLeaks, bReportActiveLeaks);
+		ReportLeak(EReportLeakType::Immediate, World, SystemName, SystemData);
 	}
 }
 
@@ -151,53 +222,11 @@ void FNiagaraComponentLeakDetector::ReportLeaks(UWorld* World)
 		return;
 	}
 
-	const bool bReportTotalLeaks = GGCReport == 1;
-	const bool bReportActiveLeaks = GGCReport != 0;
-
 	for (auto SystemDataIt=PerSystemData.CreateIterator(); SystemDataIt; ++SystemDataIt)
 	{
 		const FName SystemName = SystemDataIt.Key();
 		FSystemData& SystemData = SystemDataIt.Value();
-		ReportLeak(World, SystemName, SystemData, bReportTotalLeaks, bReportActiveLeaks);
-	}
-}
-
-void FNiagaraComponentLeakDetector::ReportLeak(UWorld* World, FName SystemName, FSystemData& SystemData, bool bReportTotalLeaks, bool bReportActiveLeaks)
-{
-	using namespace NiagaraComponentLeakDetectorPrivate;
-
-	if (bReportTotalLeaks)
-	{
-		if (!SystemData.bTotalHasWarned)
-		{
-			if (SystemData.TotalShrinkCounter == 0 && SystemData.TotalGrowthCounter >= GGrowthCountThreshold)
-			{
-				SystemData.bTotalHasWarned = true;
-				AddLeakWarning(World, SystemName, FString::Printf(TEXT("Potential componment leak System(%s) (Total:%d), please investigate."), *SystemName.ToString(), SystemData.TotalPrevCount));
-			}
-		}
-		else if (!SystemData.bTotalFalseWarning && (SystemData.TotalShrinkCounter > 0))
-		{
-			SystemData.bTotalFalseWarning = true;
-			AddLeakWarning(World, SystemName, FString::Printf(TEXT("Potential invalid component leak reported for System(%s) (Total:%d)."), *SystemName.ToString(), SystemData.TotalPrevCount));
-		}
-	}
-
-	if (bReportActiveLeaks)
-	{
-		if (!SystemData.bActiveHasWarned)
-		{
-			if (SystemData.ActiveShrinkCounter == 0 && SystemData.ActiveGrowthCounter >= GGrowthCountThreshold)
-			{
-				SystemData.bActiveHasWarned = true;
-				AddLeakWarning(World, SystemName, FString::Printf(TEXT("Potential Active Component leak System(%s) (Active:%d), please investigate."), *SystemName.ToString(), SystemData.ActivePrevCount));
-			}
-		}
-		else if ( !SystemData.bActiveFalseWarning && (SystemData.ActiveShrinkCounter > 0) )
-		{
-			SystemData.bActiveFalseWarning = true;
-			AddLeakWarning(World, SystemName, FString::Printf(TEXT("Potential invalid active component leak reported for System(% s) (Active: % d)."), *SystemName.ToString(), SystemData.ActivePrevCount));
-		}
+		ReportLeak(EReportLeakType::PostGC, World, SystemName, SystemData);
 	}
 }
 
diff --git a/Engine/Plugins/FX/Niagara/Source/Niagara/Private/NiagaraComponentLeakDetector.h b/Engine/Plugins/FX/Niagara/Source/Niagara/Private/NiagaraComponentLeakDetector.h
--- a/Engine/Plugins/FX/Niagara/Source/Niagara/Private/NiagaraComponentLeakDetector.h
+++ b/Engine/Plugins/FX/Niagara/Source/Niagara/Private/NiagaraComponentLeakDetector.h
@@ -10,6 +10,7 @@
 // Very simple leak detector, will be revised over time.
 class FNiagaraComponentLeakDetector
 {
+public:
 	struct FSystemData
 	{
 		static constexpr uint32 HistoryLength = 4;
@@ -27,15 +28,17 @@ class FNiagaraComponentLeakDetector
 		int32	ActiveShrinkCounter = 0;
 		uint32	ActiveCurrCount = 0;
 		uint32	ActivePrevCount = 0;
+		
+		bool	bScalabilityHasWarned = false;
+		int32	ScalabilityCurrActiveExecState = 0;
+		int32	ScalabilityMaxActiveExecState = 0;
+		int32	ScalabilityAllowedActiveExecState = 0;
 	};
 
 public:
 	void Tick(UWorld* World);
 	void ReportLeaks(UWorld* World);
 
-private:
-	void ReportLeak(UWorld* World, FName SystemName, FSystemData& SystemData, bool bReportTotalLeaks, bool bReportActiveLeaks);
-
 private:
 	double						NextUpdateTime = 0.0;
 	TMap<FName, FSystemData>	PerSystemData;
