Niagara - Don't try to resolve required scripts when post load updating cooked systems. Also move some old debug code to make things more readable.

diff --git a/Engine/Plugins/FX/Niagara/Source/Niagara/Private/NiagaraSystem.cpp b/Engine/Plugins/FX/Niagara/Source/Niagara/Private/NiagaraSystem.cpp
--- a/Engine/Plugins/FX/Niagara/Source/Niagara/Private/NiagaraSystem.cpp
+++ b/Engine/Plugins/FX/Niagara/Source/Niagara/Private/NiagaraSystem.cpp
@@ -497,6 +497,30 @@
 	return false;
 }
 
+#if 0
+namespace FNiagaraSystemParameterUtilities
+{
+	void DumpParameterStateToLog(FString Message) const
+	{
+		UE_LOG(LogNiagara, Log, *Message);
+		for (FNiagaraEmitterHandle& EmitterHandle : EmitterHandles)
+		{
+			UE_LOG(LogNiagara, Log, TEXT("Emitter Handle: %s"), *EmitterHandle.GetUniqueInstanceName());
+			UNiagaraScript* UpdateScript = EmitterHandle.GetInstance()->GetScript(ENiagaraScriptUsage::ParticleUpdateScript, FGuid());
+			UNiagaraScript* SpawnScript = EmitterHandle.GetInstance()->GetScript(ENiagaraScriptUsage::ParticleSpawnScript, FGuid());
+			UE_LOG(LogNiagara, Log, TEXT("Spawn Parameters"));
+			SpawnScript->GetVMExecutableData().Parameters.DumpParameters();
+			UE_LOG(LogNiagara, Log, TEXT("Spawn RI Parameters"));
+			SpawnScript->RapidIterationParameters.DumpParameters();
+			UE_LOG(LogNiagara, Log, TEXT("Update Parameters"));
+			UpdateScript->GetVMExecutableData().Parameters.DumpParameters();
+			UE_LOG(LogNiagara, Log, TEXT("Update RI Parameters"));
+			UpdateScript->RapidIterationParameters.DumpParameters();
+		}
+	}
+}
+#endif
+
 void UNiagaraSystem::UpdateSystemAfterLoad()
 {
 	// guard against deadlocks by having wait called on it during the update
@@ -574,24 +598,6 @@
 		// Synchronize with parameter definitions
 		PostLoadDefinitionsSubscriptions();
 
-#if 0
-		UE_LOG(LogNiagara, Log, TEXT("PreMerger"));
-		for (FNiagaraEmitterHandle& EmitterHandle : EmitterHandles)
-		{
-			UE_LOG(LogNiagara, Log, TEXT("Emitter Handle: %s"), *EmitterHandle.GetUniqueInstanceName());
-			UNiagaraScript* UpdateScript = EmitterHandle.GetInstance()->GetScript(ENiagaraScriptUsage::ParticleUpdateScript, FGuid());
-			UNiagaraScript* SpawnScript = EmitterHandle.GetInstance()->GetScript(ENiagaraScriptUsage::ParticleSpawnScript, FGuid());
-			UE_LOG(LogNiagara, Log, TEXT("Spawn Parameters"));
-			SpawnScript->GetVMExecutableData().Parameters.DumpParameters();
-			UE_LOG(LogNiagara, Log, TEXT("Spawn RI Parameters"));
-			SpawnScript->RapidIterationParameters.DumpParameters();
-			UE_LOG(LogNiagara, Log, TEXT("Update Parameters"));
-			UpdateScript->GetVMExecutableData().Parameters.DumpParameters();
-			UE_LOG(LogNiagara, Log, TEXT("Update RI Parameters"));
-			UpdateScript->RapidIterationParameters.DumpParameters();
-		}
-#endif
-
 		if (UNiagaraEmitter::GetForceCompileOnLoad())
 		{
 			ForceGraphToRecompileOnNextCheck();
@@ -608,39 +614,7 @@
 			InitSystemCompiledData();
 		}
 
-#if 0
-		UE_LOG(LogNiagara, Log, TEXT("Before"));
-		for (FNiagaraEmitterHandle& EmitterHandle : EmitterHandles)
-		{
-			UE_LOG(LogNiagara, Log, TEXT("Emitter Handle: %s"), *EmitterHandle.GetUniqueInstanceName());
-			UNiagaraScript* UpdateScript = EmitterHandle.GetInstance()->GetScript(ENiagaraScriptUsage::ParticleUpdateScript, FGuid());
-			UNiagaraScript* SpawnScript = EmitterHandle.GetInstance()->GetScript(ENiagaraScriptUsage::ParticleSpawnScript, FGuid());
-			UE_LOG(LogNiagara, Log, TEXT("Spawn Parameters"));
-			SpawnScript->GetVMExecutableData().Parameters.DumpParameters();
-			UE_LOG(LogNiagara, Log, TEXT("Spawn RI Parameters"));
-			SpawnScript->RapidIterationParameters.DumpParameters();
-			UE_LOG(LogNiagara, Log, TEXT("Update Parameters"));
-			UpdateScript->GetVMExecutableData().Parameters.DumpParameters();
-			UE_LOG(LogNiagara, Log, TEXT("Update RI Parameters"));
-			UpdateScript->RapidIterationParameters.DumpParameters();
-		}
-		
-		UE_LOG(LogNiagara, Log, TEXT("After"));
-		for (FNiagaraEmitterHandle& EmitterHandle : EmitterHandles)
-		{
-			UE_LOG(LogNiagara, Log, TEXT("Emitter Handle: %s"), *EmitterHandle.GetUniqueInstanceName());
-			UNiagaraScript* UpdateScript = EmitterHandle.GetInstance()->GetScript(ENiagaraScriptUsage::ParticleUpdateScript, FGuid());
-			UNiagaraScript* SpawnScript = EmitterHandle.GetInstance()->GetScript(ENiagaraScriptUsage::ParticleSpawnScript, FGuid());
-			UE_LOG(LogNiagara, Log, TEXT("Spawn Parameters"));
-			SpawnScript->GetVMExecutableData().Parameters.DumpParameters();
-			UE_LOG(LogNiagara, Log, TEXT("Spawn RI Parameters"));
-			SpawnScript->RapidIterationParameters.DumpParameters();
-			UE_LOG(LogNiagara, Log, TEXT("Update Parameters"));
-			UpdateScript->GetVMExecutableData().Parameters.DumpParameters();
-			UE_LOG(LogNiagara, Log, TEXT("Update RI Parameters"));
-			UpdateScript->RapidIterationParameters.DumpParameters();
-		}
-#endif
+		ResolveRequiresScripts();
 	}
 #endif
 
@@ -650,8 +624,6 @@
 		bIsReadyToRunCached = IsReadyToRunInternal();
 	}
 
-	ResolveRequiresScripts();
-
 	ResolveScalabilitySettings();
 
 	ComputeEmittersExecutionOrder();
