Niagara - HashCombine() return value was being ignored during hash computation
-> uint32 GetTypeHash(const FNiagaraHierarchyIdentity& Identity)
-> uint32 GetTypeHash(const FInputDataCacheKey& InputDataCacheKey)

diff --git a/Engine/Plugins/FX/Niagara/Source/NiagaraEditor/Public/ViewModels/HierarchyEditor/NiagaraHierarchyViewModelBase.h b/Engine/Plugins/FX/Niagara/Source/NiagaraEditor/Public/ViewModels/HierarchyEditor/NiagaraHierarchyViewModelBase.h
--- a/Engine/Plugins/FX/Niagara/Source/NiagaraEditor/Public/ViewModels/HierarchyEditor/NiagaraHierarchyViewModelBase.h
+++ b/Engine/Plugins/FX/Niagara/Source/NiagaraEditor/Public/ViewModels/HierarchyEditor/NiagaraHierarchyViewModelBase.h
@@ -71,12 +71,12 @@ FORCEINLINE uint32 GetTypeHash(const FNiagaraHierarchyIdentity& Identity)
 	
 	for(const FGuid& Guid : Identity.Guids)
 	{
-		HashCombine(Hash, GetTypeHash(Guid));
+		Hash = HashCombine(Hash, GetTypeHash(Guid));
 	}
 	
 	for(const FName& Name : Identity.Names)
 	{
-		HashCombine(Hash, GetTypeHash(Name));
+		Hash = HashCombine(Hash, GetTypeHash(Name));
 	}
 	
 	return Hash;
diff --git a/Engine/Plugins/FX/Niagara/Source/NiagaraEditor/Public/ViewModels/Stack/NiagaraStackGraphUtilities.h b/Engine/Plugins/FX/Niagara/Source/NiagaraEditor/Public/ViewModels/Stack/NiagaraStackGraphUtilities.h
--- a/Engine/Plugins/FX/Niagara/Source/NiagaraEditor/Public/ViewModels/Stack/NiagaraStackGraphUtilities.h
+++ b/Engine/Plugins/FX/Niagara/Source/NiagaraEditor/Public/ViewModels/Stack/NiagaraStackGraphUtilities.h
@@ -171,8 +171,8 @@ namespace FNiagaraStackGraphUtilities
 	FORCEINLINE uint32 GetTypeHash(const FInputDataCacheKey& InputDataCacheKey)
 	{
 		uint32 Hash = 0;
-		HashCombine(Hash, GetTypeHash(InputDataCacheKey.NodeGuid));
-		HashCombine(Hash, GetTypeHash(InputDataCacheKey.VariableGuid));
+		Hash = HashCombine(Hash, GetTypeHash(InputDataCacheKey.NodeGuid));
+		Hash = HashCombine(Hash, GetTypeHash(InputDataCacheKey.VariableGuid));
 		return Hash;
 	}
 
