- Change Spline DI to use soft actor references

diff --git a/Engine/Plugins/FX/Niagara/Source/Niagara/Classes/NiagaraDataInterfaceSpline.h b/Engine/Plugins/FX/Niagara/Source/Niagara/Classes/NiagaraDataInterfaceSpline.h
--- a/Engine/Plugins/FX/Niagara/Source/Niagara/Classes/NiagaraDataInterfaceSpline.h
+++ b/Engine/Plugins/FX/Niagara/Source/Niagara/Classes/NiagaraDataInterfaceSpline.h
@@ -169,8 +169,13 @@ class UNiagaraDataInterfaceSpline : public UNiagaraDataInterface
 
 public:
 	/** The source actor from which to sample.  Note that this can only be set when used as a user variable on a component in the world.*/
-	UPROPERTY(EditAnywhere, Category = "Spline")
-	TObjectPtr<AActor> Source;
+	UPROPERTY(EditAnywhere, Category = "Mesh", meta = (DisplayName = "Source Actor"))
+	TSoftObjectPtr<AActor> SoftSourceActor;
+
+#if WITH_EDITORONLY_DATA
+	UPROPERTY()
+	TObjectPtr<AActor> Source_DEPRECATED;
+#endif
 
 	/** Reference to a user parameter if we're reading one. This should  be an Object user parameter that is either a USplineComponent or an AActor containing a USplineComponent. */
 	UPROPERTY(EditAnywhere, Category = "Spline")
@@ -183,7 +188,8 @@ public:
 	int32 NumLUTSteps;
 	
 	//UObject Interface
-	NIAGARA_API virtual void PostInitProperties()override;
+	NIAGARA_API virtual void PostInitProperties() override;
+	NIAGARA_API virtual void PostLoad() override;
 	//UObject Interface End
 
 	//UNiagaraDataInterface Interface
diff --git a/Engine/Plugins/FX/Niagara/Source/Niagara/Private/NiagaraDataInterfaceSpline.cpp b/Engine/Plugins/FX/Niagara/Source/Niagara/Private/NiagaraDataInterfaceSpline.cpp
--- a/Engine/Plugins/FX/Niagara/Source/Niagara/Private/NiagaraDataInterfaceSpline.cpp
+++ b/Engine/Plugins/FX/Niagara/Source/Niagara/Private/NiagaraDataInterfaceSpline.cpp
@@ -111,8 +111,11 @@ namespace NDISplineLocal
 
 UNiagaraDataInterfaceSpline::UNiagaraDataInterfaceSpline(FObjectInitializer const& ObjectInitializer)
 	: Super(ObjectInitializer)
-	, Source(nullptr)
-	, bUseLUT(false)
+	, SoftSourceActor(nullptr)
+#if WITH_EDITORONLY_DATA
+	, Source_DEPRECATED(nullptr)
+#endif
+, bUseLUT(false)
 	, NumLUTSteps(256)
 {
 	FNiagaraTypeDefinition Def(UObject::StaticClass());
@@ -134,6 +137,18 @@ void UNiagaraDataInterfaceSpline::PostInitProperties()
 	}
 }
 
+void UNiagaraDataInterfaceSpline::PostLoad()
+{
+	Super::PostLoad();
+
+#if WITH_EDITORONLY_DATA
+	if (Source_DEPRECATED != nullptr)
+	{
+		SoftSourceActor = Source_DEPRECATED;
+	}
+#endif
+}
+
 #if WITH_EDITORONLY_DATA
 void UNiagaraDataInterfaceSpline::GetFunctionsInternal(TArray<FNiagaraFunctionSignature>& OutFunctions) const
 {
@@ -388,7 +403,7 @@ bool UNiagaraDataInterfaceSpline::CopyToInternal(UNiagaraDataInterface* Destinat
 	}
 
 	UNiagaraDataInterfaceSpline* OtherTyped = CastChecked<UNiagaraDataInterfaceSpline>(Destination);
-	OtherTyped->Source = Source;
+	OtherTyped->SoftSourceActor = SoftSourceActor;
 	OtherTyped->SplineUserParameter = SplineUserParameter;
 	
 	OtherTyped->bUseLUT = bUseLUT;
@@ -517,7 +532,7 @@ bool UNiagaraDataInterfaceSpline::Equals(const UNiagaraDataInterface* Other) con
 		return false;
 	}
 	const UNiagaraDataInterfaceSpline* OtherTyped = CastChecked<const UNiagaraDataInterfaceSpline>(Other);
-	return OtherTyped->Source == Source && OtherTyped->SplineUserParameter == SplineUserParameter && OtherTyped->bUseLUT == bUseLUT && OtherTyped->NumLUTSteps == NumLUTSteps;
+	return OtherTyped->SoftSourceActor == SoftSourceActor && OtherTyped->SplineUserParameter == SplineUserParameter && OtherTyped->bUseLUT == bUseLUT && OtherTyped->NumLUTSteps == NumLUTSteps;
 }
 
 int32 UNiagaraDataInterfaceSpline::PerInstanceDataSize()const
@@ -599,9 +614,9 @@ bool UNiagaraDataInterfaceSpline::PerInstanceTick(void* PerInstanceData, FNiagar
 						SplineComponent = UserSplineComp;
 					}
 				}
-				else if (Cast<AActor>(UserParamObject))
+				else if (AActor* UserParamActor = Cast<AActor>(UserParamObject))
 				{
-					SplineComponent = Source->FindComponentByClass<USplineComponent>();
+					SplineComponent = UserParamActor->FindComponentByClass<USplineComponent>();
 				}
 				else
 				{
@@ -617,9 +632,9 @@ bool UNiagaraDataInterfaceSpline::PerInstanceTick(void* PerInstanceData, FNiagar
 				// The binding exists, but no object is bound. Not warning here in case the user knows what they're doing.
 			}
 		}
-		else if (Source != nullptr)
+		else if (SoftSourceActor != nullptr)
 		{
-			SplineComponent = Source->FindComponentByClass<USplineComponent>();
+			SplineComponent = SoftSourceActor->FindComponentByClass<USplineComponent>();
 		}
 		else if (USceneComponent* AttachComp = SystemInstance->GetAttachComponent())
 		{
