Fixed a bug where the context menu in module scripts would show invalid conversion operations for primitive types

diff --git a/Engine/Plugins/FX/Niagara/Source/NiagaraEditor/Private/EdGraphSchema_Niagara.cpp b/Engine/Plugins/FX/Niagara/Source/NiagaraEditor/Private/EdGraphSchema_Niagara.cpp
--- a/Engine/Plugins/FX/Niagara/Source/NiagaraEditor/Private/EdGraphSchema_Niagara.cpp
+++ b/Engine/Plugins/FX/Niagara/Source/NiagaraEditor/Private/EdGraphSchema_Niagara.cpp
@@ -705,6 +705,11 @@ PRAGMA_ENABLE_DEPRECATION_WARNINGS
 			TSet<FNiagaraTypeDefinition> NoMakeAllowed;
 			NoMakeAllowed.Add(FNiagaraTypeDefinition::GetGenericNumericDef());
 			NoMakeAllowed.Add(FNiagaraTypeDefinition::GetParameterMapDef());
+			NoMakeAllowed.Add(FNiagaraTypeHelper::GetDoubleDef());
+			NoMakeAllowed.Add(FNiagaraTypeHelper::GetQuatDef());
+			NoMakeAllowed.Add(FNiagaraTypeHelper::GetVector4Def());
+			NoMakeAllowed.Add(FNiagaraTypeHelper::GetVectorDef());
+			NoMakeAllowed.Add(FNiagaraTypeHelper::GetVector2DDef());
 			
 			// Dynamic Makes from type registry. Should include the pin type as well.
 			TArray<FNiagaraTypeDefinition> CandidateTypes;
