- Add missing transaction for adding a scratchpad

diff --git a/Engine/Plugins/FX/Niagara/Source/NiagaraEditorWidgets/Private/SNiagaraScratchPadScriptManager.cpp b/Engine/Plugins/FX/Niagara/Source/NiagaraEditorWidgets/Private/SNiagaraScratchPadScriptManager.cpp
--- a/Engine/Plugins/FX/Niagara/Source/NiagaraEditorWidgets/Private/SNiagaraScratchPadScriptManager.cpp
+++ b/Engine/Plugins/FX/Niagara/Source/NiagaraEditorWidgets/Private/SNiagaraScratchPadScriptManager.cpp
@@ -25,6 +25,7 @@
 #include "EditorFontGlyphs.h"
 #include "NiagaraConstants.h"
 #include "DragAndDrop/AssetDragDropOp.h"
+#include "ScopedTransaction.h"
 
 #define LOCTEXT_NAMESPACE "NiagaraScratchPad"
 
@@ -395,6 +396,7 @@ private:
 
 	FReply ScriptSelectorAddButtonClicked(ENiagaraScriptUsage Usage)
 	{
+		FScopedTransaction ScopedTransaction(LOCTEXT("MakeScratchpad", "Make new scratchpad"));
 		TSharedPtr<FNiagaraScratchPadScriptViewModel> NewScriptViewModel = ViewModel->CreateNewScript(Usage, ENiagaraScriptUsage::ParticleUpdateScript, FNiagaraTypeDefinition());
 		if (NewScriptViewModel.IsValid())
 		{
