- Scale Mesh / Sprite Size can have a curve scale
- This is useful if you want to drive a parameter for the scale but still have a curve over the life

diff --git a/Engine/Plugins/FX/Niagara/Shaders/Private/Stateless/Modules/NiagaraStatelessModule_ScaleMeshSize.ush b/Engine/Plugins/FX/Niagara/Shaders/Private/Stateless/Modules/NiagaraStatelessModule_ScaleMeshSize.ush
--- a/Engine/Plugins/FX/Niagara/Shaders/Private/Stateless/Modules/NiagaraStatelessModule_ScaleMeshSize.ush
+++ b/Engine/Plugins/FX/Niagara/Shaders/Private/Stateless/Modules/NiagaraStatelessModule_ScaleMeshSize.ush
@@ -2,7 +2,9 @@
 
 #include "../NiagaraStatelessCommon.ush"
 
-uint3 ScaleMeshSize_Distribution;
+uint3	ScaleMeshSize_Distribution;
+float3	ScaleMeshSize_CurveScale;
+int		ScaleMeshSize_CurveScaleOffset;
 	
 struct FStatelessModule_ScaleMeshSize
 {
@@ -18,7 +20,9 @@ void ScaleMeshSize_Simulate(in FStatelessModule_ScaleMeshSize ModuleData, inout
 {
 	if (ModuleData.MeshSizeDistribution.IsValid())
 	{
-		Particle.Scale			*= ModuleData.MeshSizeDistribution.GetValue(Particle.NormalizedAge);
-		Particle.PreviousScale	*= ModuleData.MeshSizeDistribution.GetValue(Particle.PreviousNormalizedAge);
+		const float3 ScaleFactor = GetParameterBufferFloat3(ScaleMeshSize_CurveScaleOffset, 0, ScaleMeshSize_CurveScale);
+			
+		Particle.Scale			*= ModuleData.MeshSizeDistribution.GetValue(Particle.NormalizedAge) * ScaleFactor;
+		Particle.PreviousScale	*= ModuleData.MeshSizeDistribution.GetValue(Particle.PreviousNormalizedAge) * ScaleFactor;
 	}
 }
diff --git a/Engine/Plugins/FX/Niagara/Shaders/Private/Stateless/Modules/NiagaraStatelessModule_ScaleSpriteSize.ush b/Engine/Plugins/FX/Niagara/Shaders/Private/Stateless/Modules/NiagaraStatelessModule_ScaleSpriteSize.ush
--- a/Engine/Plugins/FX/Niagara/Shaders/Private/Stateless/Modules/NiagaraStatelessModule_ScaleSpriteSize.ush
+++ b/Engine/Plugins/FX/Niagara/Shaders/Private/Stateless/Modules/NiagaraStatelessModule_ScaleSpriteSize.ush
@@ -2,7 +2,9 @@
 
 #include "../NiagaraStatelessCommon.ush"
 
-uint3 ScaleSpriteSize_Distribution;
+uint3	ScaleSpriteSize_Distribution;
+float2	ScaleSpriteSize_CurveScale;
+int		ScaleSpriteSize_CurveScaleOffset;
 
 struct FStatelessModule_ScaleSpriteSize
 {
@@ -18,7 +20,9 @@ void ScaleSpriteSize_Simulate(in FStatelessModule_ScaleSpriteSize ModuleData, in
 {
 	if (ModuleData.ScaleDistribution.IsValid())
 	{
-		Particle.SpriteSize			*= ModuleData.ScaleDistribution.GetValue(Particle.NormalizedAge);
-		Particle.PreviousSpriteSize	*= ModuleData.ScaleDistribution.GetValue(Particle.PreviousNormalizedAge);
+		const float2 ScaleFactor = GetParameterBufferFloat2(ScaleSpriteSize_CurveScaleOffset, 0, ScaleSpriteSize_CurveScale);
+
+		Particle.SpriteSize			*= ModuleData.ScaleDistribution.GetValue(Particle.NormalizedAge) * ScaleFactor;
+		Particle.PreviousSpriteSize	*= ModuleData.ScaleDistribution.GetValue(Particle.PreviousNormalizedAge) * ScaleFactor;
 	}
 }
diff --git a/Engine/Plugins/FX/Niagara/Shaders/Private/Stateless/NiagaraStatelessCommon.ush b/Engine/Plugins/FX/Niagara/Shaders/Private/Stateless/NiagaraStatelessCommon.ush
--- a/Engine/Plugins/FX/Niagara/Shaders/Private/Stateless/NiagaraStatelessCommon.ush
+++ b/Engine/Plugins/FX/Niagara/Shaders/Private/Stateless/NiagaraStatelessCommon.ush
@@ -131,6 +131,11 @@ float  GetParameterBufferFloat (int BaseOffset, int Element) { return asfloat(Co
 float2 GetParameterBufferFloat2(int BaseOffset, int Element) { return float2(asfloat(Common_ParameterBuffer[BaseOffset + Element*2]), asfloat(Common_ParameterBuffer[BaseOffset + 1 + Element*2])); }
 float3 GetParameterBufferFloat3(int BaseOffset, int Element) { return float3(asfloat(Common_ParameterBuffer[BaseOffset + Element*3]), asfloat(Common_ParameterBuffer[BaseOffset + 1 + Element*3]), asfloat(Common_ParameterBuffer[BaseOffset + 2 + Element*3])); }
 float4 GetParameterBufferFloat4(int BaseOffset, int Element) { return float4(asfloat(Common_ParameterBuffer[BaseOffset + Element*4]), asfloat(Common_ParameterBuffer[BaseOffset + 1 + Element*4]), asfloat(Common_ParameterBuffer[BaseOffset + 2 + Element*4]), asfloat(Common_ParameterBuffer[BaseOffset + 3 + Element*4])); }
+
+float  GetParameterBufferFloat (int BaseOffset, int Element, float  DefaultValue) { return BaseOffset >= 0 ? asfloat(Common_ParameterBuffer[BaseOffset + Element]) : DefaultValue; }
+float2 GetParameterBufferFloat2(int BaseOffset, int Element, float2 DefaultValue) { return BaseOffset >= 0 ? float2(asfloat(Common_ParameterBuffer[BaseOffset + Element*2]), asfloat(Common_ParameterBuffer[BaseOffset + 1 + Element*2])) : DefaultValue; }
+float3 GetParameterBufferFloat3(int BaseOffset, int Element, float3 DefaultValue) { return BaseOffset >= 0 ? float3(asfloat(Common_ParameterBuffer[BaseOffset + Element*3]), asfloat(Common_ParameterBuffer[BaseOffset + 1 + Element*3]), asfloat(Common_ParameterBuffer[BaseOffset + 2 + Element*3])) : DefaultValue; }
+float4 GetParameterBufferFloat4(int BaseOffset, int Element, float4 DefaultValue) { return BaseOffset >= 0 ? float4(asfloat(Common_ParameterBuffer[BaseOffset + Element*4]), asfloat(Common_ParameterBuffer[BaseOffset + 1 + Element*4]), asfloat(Common_ParameterBuffer[BaseOffset + 2 + Element*4]), asfloat(Common_ParameterBuffer[BaseOffset + 3 + Element*4])) : DefaultValue; }
 	
 float LerpStaticFloat (uint2 Parameters, float U)
 {
diff --git a/Engine/Plugins/FX/Niagara/Source/Niagara/Internal/Stateless/Modules/NiagaraStatelessModule_ScaleMeshSize.h b/Engine/Plugins/FX/Niagara/Source/Niagara/Internal/Stateless/Modules/NiagaraStatelessModule_ScaleMeshSize.h
--- a/Engine/Plugins/FX/Niagara/Source/Niagara/Internal/Stateless/Modules/NiagaraStatelessModule_ScaleMeshSize.h
+++ b/Engine/Plugins/FX/Niagara/Source/Niagara/Internal/Stateless/Modules/NiagaraStatelessModule_ScaleMeshSize.h
@@ -6,6 +6,8 @@
 #include "Stateless/NiagaraStatelessEmitterDataBuildContext.h"
 #include "Stateless/NiagaraStatelessModuleShaderParameters.h"
 
+#include "NiagaraParameterBinding.h"
+
 #include "NiagaraStatelessModule_ScaleMeshSize.generated.h"
 
 UCLASS(MinimalAPI, EditInlineNew, meta = (DisplayName = "Scale Mesh Size"))
@@ -16,6 +18,8 @@ class UNiagaraStatelessModule_ScaleMeshSize : public UNiagaraStatelessModule
 	struct FModuleBuiltData
 	{
 		FUintVector3	DistributionParameters = FUintVector3::ZeroValue;
+		FVector3f		CurveScale = FVector3f::OneVector;
+		int32			CurveScaleOffset = INDEX_NONE;
 	};
 
 public:
@@ -24,20 +28,47 @@ public:
 	UPROPERTY(EditAnywhere, Category = "Parameters", meta = (DisplayName = "Scale"))
 	FNiagaraDistributionVector3 ScaleDistribution = FNiagaraDistributionVector3(FVector3f::OneVector);
 
+	UPROPERTY(EditAnywhere, Category = "Parameters", meta = (EditConditionHides, EditCondition = "UseScaleCurveRange()"))
+	FNiagaraParameterBindingWithValue ScaleCurveRange;
+
+	virtual void PostInitProperties() override
+	{
+		Super::PostInitProperties();
+
+	#if WITH_EDITORONLY_DATA
+		if (HasAnyFlags(RF_ClassDefaultObject) == false)
+		{
+			ScaleCurveRange.SetUsage(ENiagaraParameterBindingUsage::NotParticle);
+			ScaleCurveRange.SetAllowedTypeDefinitions({ FNiagaraTypeDefinition::GetVec3Def() });
+			ScaleCurveRange.SetDefaultParameter(FNiagaraTypeDefinition::GetVec3Def(), FVector3f::OneVector);
+		}
+	#endif
+	}
+
 	virtual void BuildEmitterData(FNiagaraStatelessEmitterDataBuildContext& BuildContext) const override
 	{
 		FModuleBuiltData* BuiltData = BuildContext.AllocateBuiltData<FModuleBuiltData>();
 		BuiltData->DistributionParameters = BuildContext.AddDistribution(ScaleDistribution, IsModuleEnabled());
+		if (IsModuleEnabled() && UseScaleCurveRange())
+		{
+			BuiltData->CurveScaleOffset = BuildContext.AddRendererBinding(ScaleCurveRange.ResolvedParameter);
+			BuiltData->CurveScale = ScaleCurveRange.GetDefaultValue<FVector3f>();
+		}
 	}
 
 	virtual void SetShaderParameters(const FNiagaraStatelessSetShaderParameterContext& SetShaderParameterContext) const override
 	{
 		const FModuleBuiltData* ModuleBuiltData = SetShaderParameterContext.ReadBuiltData<FModuleBuiltData>();
 
-		FParameters* Parameters = SetShaderParameterContext.GetParameterNestedStruct<FParameters>();
-		Parameters->ScaleMeshSize_Distribution = ModuleBuiltData->DistributionParameters;
+		FParameters* Parameters						= SetShaderParameterContext.GetParameterNestedStruct<FParameters>();
+		Parameters->ScaleMeshSize_Distribution		= ModuleBuiltData->DistributionParameters;
+		Parameters->ScaleMeshSize_CurveScale		= ModuleBuiltData->CurveScale;
+		Parameters->ScaleMeshSize_CurveScaleOffset	= ModuleBuiltData->CurveScaleOffset;
 	}
 
+	UFUNCTION()
+	bool UseScaleCurveRange() const { return ScaleDistribution.IsCurve(); }
+
 #if WITH_EDITOR
 	virtual bool CanDisableModule() const override { return true; }
 #endif
diff --git a/Engine/Plugins/FX/Niagara/Source/Niagara/Internal/Stateless/Modules/NiagaraStatelessModule_ScaleSpriteSize.h b/Engine/Plugins/FX/Niagara/Source/Niagara/Internal/Stateless/Modules/NiagaraStatelessModule_ScaleSpriteSize.h
--- a/Engine/Plugins/FX/Niagara/Source/Niagara/Internal/Stateless/Modules/NiagaraStatelessModule_ScaleSpriteSize.h
+++ b/Engine/Plugins/FX/Niagara/Source/Niagara/Internal/Stateless/Modules/NiagaraStatelessModule_ScaleSpriteSize.h
@@ -17,6 +17,8 @@ class UNiagaraStatelessModule_ScaleSpriteSize : public UNiagaraStatelessModule
 	struct FModuleBuiltData
 	{
 		FUintVector3	DistributionParameters = FUintVector3::ZeroValue;
+		FVector2f		CurveScale = FVector2f::One();
+		int32			CurveScaleOffset = INDEX_NONE;
 	};
 
 public:
@@ -24,21 +26,48 @@ public:
 
 	UPROPERTY(EditAnywhere, Category = "Parameters", meta = (DisplayName = "Scale"))
 	FNiagaraDistributionVector2 ScaleDistribution = FNiagaraDistributionVector2(FVector2f::One());
-	
+
+	UPROPERTY(EditAnywhere, Category = "Parameters", meta = (EditConditionHides, EditCondition = "UseScaleCurveRange()"))
+	FNiagaraParameterBindingWithValue ScaleCurveRange;
+
+	virtual void PostInitProperties() override
+	{
+		Super::PostInitProperties();
+
+	#if WITH_EDITORONLY_DATA
+		if (HasAnyFlags(RF_ClassDefaultObject) == false)
+		{
+			ScaleCurveRange.SetUsage(ENiagaraParameterBindingUsage::NotParticle);
+			ScaleCurveRange.SetAllowedTypeDefinitions({ FNiagaraTypeDefinition::GetVec2Def() });
+			ScaleCurveRange.SetDefaultParameter(FNiagaraTypeDefinition::GetVec2Def(), FVector2f::One());
+		}
+	#endif
+	}
+
 	virtual void BuildEmitterData(FNiagaraStatelessEmitterDataBuildContext& BuildContext) const override
 	{
 		FModuleBuiltData* BuiltData = BuildContext.AllocateBuiltData<FModuleBuiltData>();
 		BuiltData->DistributionParameters = BuildContext.AddDistribution(ScaleDistribution, IsModuleEnabled());
+		if (IsModuleEnabled() && UseScaleCurveRange())
+		{
+			BuiltData->CurveScaleOffset = BuildContext.AddRendererBinding(ScaleCurveRange.ResolvedParameter);
+			BuiltData->CurveScale = ScaleCurveRange.GetDefaultValue<FVector2f>();
+		}
 	}
 
 	virtual void SetShaderParameters(const FNiagaraStatelessSetShaderParameterContext& SetShaderParameterContext) const override
 	{
 		const FModuleBuiltData* ModuleBuiltData = SetShaderParameterContext.ReadBuiltData<FModuleBuiltData>();
 
-		FParameters* Parameters = SetShaderParameterContext.GetParameterNestedStruct<FParameters>();
-		Parameters->ScaleSpriteSize_Distribution = ModuleBuiltData->DistributionParameters;
+		FParameters* Parameters						= SetShaderParameterContext.GetParameterNestedStruct<FParameters>();
+		Parameters->ScaleSpriteSize_Distribution	= ModuleBuiltData->DistributionParameters;
+		Parameters->ScaleSpriteSize_CurveScale		= ModuleBuiltData->CurveScale;
+		Parameters->ScaleSpriteSize_CurveScaleOffset= ModuleBuiltData->CurveScaleOffset;
 	}
 
+	UFUNCTION()
+	bool UseScaleCurveRange() const { return ScaleDistribution.IsCurve(); }
+
 #if WITH_EDITOR
 	virtual bool CanDisableModule() const override { return true; }
 #endif
diff --git a/Engine/Plugins/FX/Niagara/Source/NiagaraShader/Internal/Stateless/NiagaraStatelessModuleShaderParameters.h b/Engine/Plugins/FX/Niagara/Source/NiagaraShader/Internal/Stateless/NiagaraStatelessModuleShaderParameters.h
--- a/Engine/Plugins/FX/Niagara/Source/NiagaraShader/Internal/Stateless/NiagaraStatelessModuleShaderParameters.h
+++ b/Engine/Plugins/FX/Niagara/Source/NiagaraShader/Internal/Stateless/NiagaraStatelessModuleShaderParameters.h
@@ -62,6 +62,8 @@ namespace NiagaraStateless
 
 	BEGIN_SHADER_PARAMETER_STRUCT(FScaleMeshSizeModule_ShaderParameters, NIAGARASHADER_API)
 		SHADER_PARAMETER(FUintVector3,	ScaleMeshSize_Distribution)
+		SHADER_PARAMETER(FVector3f,		ScaleMeshSize_CurveScale)
+		SHADER_PARAMETER(int32,			ScaleMeshSize_CurveScaleOffset)
 	END_SHADER_PARAMETER_STRUCT()
 
 	BEGIN_SHADER_PARAMETER_STRUCT(FScaleMeshSizeBySpeedModule_ShaderParameters, NIAGARASHADER_API)
@@ -75,7 +77,9 @@ namespace NiagaraStateless
 	END_SHADER_PARAMETER_STRUCT()
 
 	BEGIN_SHADER_PARAMETER_STRUCT(FScaleSpriteSizeModule_ShaderParameters, NIAGARASHADER_API)
-		SHADER_PARAMETER(FUintVector3, ScaleSpriteSize_Distribution)
+		SHADER_PARAMETER(FUintVector3,	ScaleSpriteSize_Distribution)
+		SHADER_PARAMETER(FVector2f,		ScaleSpriteSize_CurveScale)
+		SHADER_PARAMETER(int32,			ScaleSpriteSize_CurveScaleOffset)
 	END_SHADER_PARAMETER_STRUCT()
 
 	BEGIN_SHADER_PARAMETER_STRUCT(FScaleSpriteSizeBySpeedModule_ShaderParameters, NIAGARASHADER_API)
