- Improve Debug HUD display for cull proxies

diff --git a/Engine/Plugins/FX/Niagara/Source/Niagara/Private/NiagaraDebugHud.cpp b/Engine/Plugins/FX/Niagara/Source/Niagara/Private/NiagaraDebugHud.cpp
--- a/Engine/Plugins/FX/Niagara/Source/Niagara/Private/NiagaraDebugHud.cpp
+++ b/Engine/Plugins/FX/Niagara/Source/Niagara/Private/NiagaraDebugHud.cpp
@@ -6,6 +6,7 @@
 #include "DrawDebugHelpers.h"
 #include "GameFramework/Pawn.h"
 #include "NiagaraComponent.h"
+#include "NiagaraCullProxyComponent.h"
 #include "NiagaraDataSetDebugAccessor.h"
 #include "NiagaraDataSetReadback.h"
 #include "NiagaraEmitterInstance.h"
@@ -1047,7 +1048,6 @@ void FNiagaraDebugHud::GatherSystemInfo()
 		FNiagaraSystemInstanceControllerPtr SystemInstanceController = NiagaraComponent ? NiagaraComponent->GetSystemInstanceController() : nullptr;
 		FNiagaraSystemInstance* SystemInstance = SystemInstanceController.IsValid() ? SystemInstanceController->GetSystemInstance_Unsafe() : nullptr;
 		
-
 		if (!IsValidChecked(FXComponent) || (!IsValid(NiagaraComponent) && !IsValid(CascadeComponent)) || FXComponent->IsUnreachable() || FXComponent->HasAnyFlags(EObjectFlags::RF_ClassDefaultObject))
 		{
 			continue;
@@ -1063,6 +1063,8 @@ void FNiagaraDebugHud::GatherSystemInfo()
 
 		const bool bIsActive = FXComponent->IsActive();
 		const bool bHasScalability = NiagaraComponent ? NiagaraComponent->IsRegisteredWithScalabilityManager() : CascadeComponent->bIsManagingSignificance;
+		const bool bUsingCullProxy = NiagaraComponent ? NiagaraComponent->IsUsingCullProxy() : false;
+		const bool bIsCullProxy = FXComponent->IsA<UNiagaraCullProxyComponent>();
 
 		FSystemDebugInfo& SystemDebugInfo = PerSystemDebugInfo.FindOrAdd(FXComponent->GetFXSystemAsset()->GetFName());
 		if (SystemDebugInfo.SystemName.IsEmpty())
@@ -1076,7 +1078,12 @@ void FNiagaraDebugHud::GatherSystemInfo()
 		SystemDebugInfo.bShowInWorld = Settings.bSystemFilterEnabled && SystemDebugInfo.SystemName.MatchesWildcard(Settings.SystemFilter);
 		SystemDebugInfo.bPassesSystemFilter = !Settings.bSystemFilterEnabled || SystemDebugInfo.SystemName.MatchesWildcard(Settings.SystemFilter);
 
-		if (SystemDebugInfo.bShowInWorld && (bIsActive || !Settings.bSystemShowActiveOnlyInWorld))
+		const bool bCanShowInWorld = 
+			SystemDebugInfo.bShowInWorld &&
+			((bIsActive || !Settings.bSystemShowActiveOnlyInWorld) || bUsingCullProxy) &&
+			!bIsCullProxy;
+
+		if (bCanShowInWorld)
 		{
 			bool bIsMatch = true;
 
@@ -3168,42 +3175,48 @@ void FNiagaraDebugHud::DrawComponents(FNiagaraWorldManager* WorldManager, UCanva
 						}
 					}
 				}
+				// Instance is considered inactive
 				else
 				{
-					//-TODO: Put a reason why here (either grab from manager or push from manager)
-					//FNiagaraScalabilityState* ScalabilityState = SystemInstance->GetWorldManager()->GetScalabilityState(NiagaraComponent);
 					if (Settings.SystemDebugVerbosity >= ENiagaraDebugHudVerbosity::Basic)
 					{
-						StringBuilder.Appendf(TEXT("Deactivated by Scalability - %s "), *GetNameSafe(NiagaraSystem->GetEffectType()));
-						if (Settings.SystemDebugVerbosity >= ENiagaraDebugHudVerbosity::Verbose)
+						if (NiagaraComponent->IsUsingCullProxy())
 						{
-							FNiagaraScalabilityState ScalabilityState;
-							if (WorldManager->GetScalabilityState(NiagaraComponent, ScalabilityState))
+							StringBuilder.Appendf(TEXT("Using Cull Proxy"));
+						}
+						else
+						{
+							StringBuilder.Appendf(TEXT("Deactivated by Scalability - %s "), *GetNameSafe(NiagaraSystem->GetEffectType()));
+							if (Settings.SystemDebugVerbosity >= ENiagaraDebugHudVerbosity::Verbose)
 							{
-								StringBuilder.Appendf(TEXT("- Significance(%.2f)"), ScalabilityState.Significance);
-#if DEBUG_SCALABILITY_STATE
-								if (ScalabilityState.bCulledByDistance)
-								{
-									StringBuilder.Append(TEXT(" DistanceCulled"));
-								}
-								if (ScalabilityState.bCulledByInstanceCount)
-								{
-									StringBuilder.Append(TEXT(" InstanceCulled"));
-								}
-								if (ScalabilityState.bCulledByVisibility)
+								FNiagaraScalabilityState ScalabilityState;
+								if (WorldManager->GetScalabilityState(NiagaraComponent, ScalabilityState))
 								{
-									StringBuilder.Append(TEXT(" VisibilityCulled"));
+									StringBuilder.Appendf(TEXT("- Significance(%.2f)"), ScalabilityState.Significance);
+#if DEBUG_SCALABILITY_STATE
+									if (ScalabilityState.bCulledByDistance)
+									{
+										StringBuilder.Append(TEXT(" DistanceCulled"));
+									}
+									if (ScalabilityState.bCulledByInstanceCount)
+									{
+										StringBuilder.Append(TEXT(" InstanceCulled"));
+									}
+									if (ScalabilityState.bCulledByVisibility)
+									{
+										StringBuilder.Append(TEXT(" VisibilityCulled"));
+									}
+									if (ScalabilityState.bCulledByGlobalBudget)
+									{
+										StringBuilder.Append(TEXT(" GlobalBudgetCulled"));
+									}
+#endif
+									StringBuilder.Append(TEXT("\n"));
 								}
-								if (ScalabilityState.bCulledByGlobalBudget)
+								else
 								{
-									StringBuilder.Append(TEXT(" GlobalBudgetCulled"));
+									StringBuilder.Appendf(TEXT("- Scalability State Unknown\n"));
 								}
-#endif
-								StringBuilder.Append(TEXT("\n"));
-							}
-							else
-							{
-								StringBuilder.Appendf(TEXT("- Scalability State Unknown\n"));
 							}
 						}
 					}
