- SetEmitterEnable support for stateless

diff --git a/Engine/Plugins/FX/Niagara/Source/Niagara/Classes/NiagaraEmitterInstance.h b/Engine/Plugins/FX/Niagara/Source/Niagara/Classes/NiagaraEmitterInstance.h
--- a/Engine/Plugins/FX/Niagara/Source/Niagara/Classes/NiagaraEmitterInstance.h
+++ b/Engine/Plugins/FX/Niagara/Source/Niagara/Classes/NiagaraEmitterInstance.h
@@ -32,6 +32,7 @@ public:
 	//~Begin: Define Emitter Interface
 	virtual void Init(int32 InEmitterIdx);
 	virtual void ResetSimulation(bool bKillExisting = true) = 0;
+	virtual void SetEmitterEnable(bool bNewEnableState) = 0;
 	virtual void OnPooledReuse() = 0;
 	virtual bool HandleCompletion(bool bForce = false) = 0;
 
@@ -149,13 +150,6 @@ protected:
 	}
 
 public:
-	////////////////////////////////////////////////////////////////////////////////////////////////////
-	//-TODO: Functionality that needs implementing for stateless
-	void SetEmitterEnable(bool bNewEnableState)
-	{
-		bAllowSpawning_GT = bNewEnableState;
-	}
-
 	////////////////////////////////////////////////////////////////////////////////////////////////////
 	// Deprecated functionality code needs to be updated
 	UE_DEPRECATED(5.4, "Please update your code to handle execution path for different emitter types")
@@ -224,8 +218,6 @@ protected:
 	uint8								bDeterministic : 1 = true;
 	uint8								bNeedsPartialDepthTexture : 1 = false;
 	uint8								bNeedsEarlyViewUniformBuffer : 1 = false;
-	uint8								bAllowSpawning_GT : 1 = true;
-	uint8								bAllowSpawning_CNC : 1 = true;
 	ENiagaraSimTarget					SimTarget = ENiagaraSimTarget::CPUSim;
 	ENiagaraExecutionState				ExecutionState = ENiagaraExecutionState::Active;
 
diff --git a/Engine/Plugins/FX/Niagara/Source/Niagara/Internal/NiagaraEmitterInstanceImpl.h b/Engine/Plugins/FX/Niagara/Source/Niagara/Internal/NiagaraEmitterInstanceImpl.h
--- a/Engine/Plugins/FX/Niagara/Source/Niagara/Internal/NiagaraEmitterInstanceImpl.h
+++ b/Engine/Plugins/FX/Niagara/Source/Niagara/Internal/NiagaraEmitterInstanceImpl.h
@@ -50,6 +50,7 @@ public:
 	// FNiagaraEmitterInstance Impl
 	virtual void Init(int32 InEmitterIndex) override;
 	virtual void ResetSimulation(bool bKillExisting = true) override;
+	virtual void SetEmitterEnable(bool bNewEnableState) override;
 	virtual bool HandleCompletion(bool bForce = false) override;
 	virtual void OnPooledReuse() override;
 	virtual FNiagaraEmitterInstanceImpl* AsStateful() override { return this; }
@@ -158,7 +159,12 @@ private:
 	uint32 MaxInstanceCount = 0;
 
 	/** Typical resets must be deferred until the tick as the RT could still be using the current buffer. */
-	uint32 bResetPending : 1;
+	uint32 bResetPending : 1 = false;
+
+	/** Used with SetEmitterEnable */
+	uint32 bAllowSpawning_GT : 1 = true;
+	/** Used with SetEmitterEnable */
+	uint32 bAllowSpawning_CNC : 1 = true;
 
 	// This is used to keep track which particles have spawned a component. This is needed when the bOnlyCreateComponentsOnParticleSpawn flag is set in the renderer.
 	// Without this bookkeeping, the particles would lose their components when the render state is recreated or the visibility tag flips them off and on again.
diff --git a/Engine/Plugins/FX/Niagara/Source/Niagara/Internal/Stateless/NiagaraStatelessEmitterInstance.h b/Engine/Plugins/FX/Niagara/Source/Niagara/Internal/Stateless/NiagaraStatelessEmitterInstance.h
--- a/Engine/Plugins/FX/Niagara/Source/Niagara/Internal/Stateless/NiagaraStatelessEmitterInstance.h
+++ b/Engine/Plugins/FX/Niagara/Source/Niagara/Internal/Stateless/NiagaraStatelessEmitterInstance.h
@@ -56,6 +56,7 @@ public:
 	// FNiagaraEmitterInstance Impl
 	virtual void Init(int32 InEmitterIndex) override;
 	virtual void ResetSimulation(bool bKillExisting) override;
+	virtual void SetEmitterEnable(bool bNewEnableState) override;
 	virtual void OnPooledReuse() override {}
 	virtual bool HandleCompletion(bool bForce) override;
 	virtual int32 GetNumParticles() const override;
@@ -78,6 +79,7 @@ private:
 	void InitEmitterState();
 	void TickEmitterState();
 	void CalculateBounds();
+	void SendRenderData();
 
 	void InitSpawnInfos();
 	void InitSpawnInfosForLoop();
@@ -87,7 +89,10 @@ private:
 	void SetExecutionStateInternal(ENiagaraExecutionState InExecutionState);
 
 private:
-	bool										bCanEverExecute = false;
+	uint32										bCanEverExecute : 1 = false;
+	uint32										bEmitterEnabled_GT : 1 = true;
+	uint32										bEmitterEnabled_CNC : 1 = true;
+	uint32										bSpawnInfosDirty : 1 = false;
 
 	int32										RandomSeed = 0;
 	FRandomStream								RandomStream;
diff --git a/Engine/Plugins/FX/Niagara/Source/Niagara/Private/NiagaraEmitterInstanceImpl.cpp b/Engine/Plugins/FX/Niagara/Source/Niagara/Private/NiagaraEmitterInstanceImpl.cpp
--- a/Engine/Plugins/FX/Niagara/Source/Niagara/Private/NiagaraEmitterInstanceImpl.cpp
+++ b/Engine/Plugins/FX/Niagara/Source/Niagara/Private/NiagaraEmitterInstanceImpl.cpp
@@ -113,7 +113,6 @@ struct FNiagaraEditorOnlyCycleTimer
 
 FNiagaraEmitterInstanceImpl::FNiagaraEmitterInstanceImpl(FNiagaraSystemInstance* InParentSystemInstance)
 	: FNiagaraEmitterInstance(InParentSystemInstance)
-	, bResetPending(false)
 {
 	ParticleDataSet = new FNiagaraDataSet();
 
@@ -982,6 +981,11 @@ void FNiagaraEmitterInstanceImpl::PostTick()
 	}
 }
 
+void FNiagaraEmitterInstanceImpl::SetEmitterEnable(bool bNewEnableState)
+{
+	bAllowSpawning_GT = bNewEnableState;
+}
+
 bool FNiagaraEmitterInstanceImpl::HandleCompletion(bool bForce)
 {
 	if (IsDisabled())
diff --git a/Engine/Plugins/FX/Niagara/Source/Niagara/Private/Stateless/NiagaraStatelessEmitterInstance.cpp b/Engine/Plugins/FX/Niagara/Source/Niagara/Private/Stateless/NiagaraStatelessEmitterInstance.cpp
--- a/Engine/Plugins/FX/Niagara/Source/Niagara/Private/Stateless/NiagaraStatelessEmitterInstance.cpp
+++ b/Engine/Plugins/FX/Niagara/Source/Niagara/Private/Stateless/NiagaraStatelessEmitterInstance.cpp
@@ -109,6 +109,7 @@ void FNiagaraStatelessEmitterInstance::ResetSimulation(bool bKillExisting)
 	if (bKillExisting)
 	{
 		SpawnInfos.Empty();
+		bSpawnInfosDirty = true;
 	}
 	else
 	{
@@ -124,6 +125,7 @@ void FNiagaraStatelessEmitterInstance::ResetSimulation(bool bKillExisting)
 
 	Age = 0.0f;
 	UniqueIndexOffset = 0;
+	bEmitterEnabled_CNC = bEmitterEnabled_GT;
 
 	InitEmitterState();
 	InitSpawnInfos();
@@ -141,6 +143,11 @@ void FNiagaraStatelessEmitterInstance::ResetSimulation(bool bKillExisting)
 	}
 }
 
+void FNiagaraStatelessEmitterInstance::SetEmitterEnable(bool bNewEnableState)
+{
+	bEmitterEnabled_GT = bNewEnableState;
+}
+
 bool FNiagaraStatelessEmitterInstance::HandleCompletion(bool bForce)
 {
 	bool bIsComplete = IsComplete();
@@ -202,37 +209,7 @@ void FNiagaraStatelessEmitterInstance::Tick(float DeltaSeconds)
 	TickSpawnInfos();
 	TickEmitterState();
 	CalculateBounds();
-
-	if (NiagaraStateless::FEmitterInstance_RT* RenderThreadData = RenderThreadDataPtr.Get())
-	{
-		NiagaraStateless::FCommonShaderParameters* NewShaderParameters = nullptr;
-		TOptional<TArray<uint8>> NewBindingBufferData;
-		if ( EmitterData->bModulesHaveRendererBindings && RendererBindings.GetParametersDirty() )
-		{
-			RendererBindings.Tick();
-			NewBindingBufferData = RendererBindings.GetParameterDataArray();
-			check((NewBindingBufferData->Num() % sizeof(uint32)) == 0);
-
-			NewShaderParameters = WeakStatelessEmitter->AllocateShaderParameters(RendererBindings);
-			NewShaderParameters->Common_RandomSeed = RandomSeed;
-		}
-
-		ENQUEUE_RENDER_COMMAND(UpdateStatelessAge)(
-			[RenderThreadData, AgeForRT=Age, ExecutionStateForRT=ExecutionState, NewBindingBufferDataRT=MoveTemp(NewBindingBufferData), NewShaderParameters](FRHICommandListImmediate& RHICmdList)
-			{
-				RenderThreadData->Age = AgeForRT;
-				RenderThreadData->ExecutionState = ExecutionStateForRT;
-				if (NewBindingBufferDataRT.IsSet())
-				{
-					RenderThreadData->BindingBufferData = NewBindingBufferDataRT;
-				}
-				if (NewShaderParameters)
-				{
-					RenderThreadData->ShaderParameters.Reset(NewShaderParameters);
-				}
-		}
-		);
-	}
+	SendRenderData();
 }
 
 void FNiagaraStatelessEmitterInstance::InitEmitterState()
@@ -327,28 +304,105 @@ void FNiagaraStatelessEmitterInstance::CalculateBounds()
 	}
 }
 
+void FNiagaraStatelessEmitterInstance::SendRenderData()
+{
+	NiagaraStateless::FEmitterInstance_RT* RenderThreadData = RenderThreadDataPtr.Get();
+	if (RenderThreadData == nullptr)
+	{
+		return;
+	}
+
+	struct FDataForRenderThread
+	{
+		float Age = 0.0f;
+		ENiagaraExecutionState ExecutionState = ENiagaraExecutionState::Disabled;
+
+		NiagaraStateless::FCommonShaderParameters* ShaderParameters = nullptr;
+
+		bool bHasBindingBufferData = false;
+		TArray<uint8> BindingBufferData;
+
+		bool bHasSpawnInfoData = false;
+		TArray<FNiagaraStatelessRuntimeSpawnInfo>	SpawnInfos;
+	};
+
+	FDataForRenderThread DataForRenderThread;
+	DataForRenderThread.Age				= Age;
+	DataForRenderThread.ExecutionState	= ExecutionState;
+
+	if ( EmitterData->bModulesHaveRendererBindings && RendererBindings.GetParametersDirty() )
+	{
+		RendererBindings.Tick();
+		DataForRenderThread.bHasBindingBufferData = true;
+		DataForRenderThread.BindingBufferData = RendererBindings.GetParameterDataArray();
+		check((DataForRenderThread.BindingBufferData.Num() % sizeof(uint32)) == 0);
+
+		DataForRenderThread.ShaderParameters = WeakStatelessEmitter->AllocateShaderParameters(RendererBindings);
+		DataForRenderThread.ShaderParameters->Common_RandomSeed = RandomSeed;
+	}
+
+	if (bSpawnInfosDirty)
+	{
+		DataForRenderThread.bHasSpawnInfoData = true;
+		DataForRenderThread.SpawnInfos = SpawnInfos;
+		bSpawnInfosDirty = false;
+	}
+
+	ENQUEUE_RENDER_COMMAND(UpdateStatelessAge)(
+		[RenderThreadData, EmitterData=MoveTemp(DataForRenderThread)](FRHICommandListImmediate& RHICmdList) mutable
+		{
+			RenderThreadData->Age				= EmitterData.Age;
+			RenderThreadData->ExecutionState	= EmitterData.ExecutionState;
+
+			if (EmitterData.ShaderParameters)
+			{
+				RenderThreadData->ShaderParameters.Reset(EmitterData.ShaderParameters);
+			}
+
+			if (EmitterData.bHasBindingBufferData)
+			{
+				RenderThreadData->BindingBufferData = MoveTemp(EmitterData.BindingBufferData);
+			}
+
+			if (EmitterData.bHasSpawnInfoData)
+			{
+				RenderThreadData->SpawnInfos = MoveTemp(EmitterData.SpawnInfos);
+			}
+		}
+	);
+}
+
 void FNiagaraStatelessEmitterInstance::InitSpawnInfos()
 {
 	UniqueIndexOffset = 0;
-	for (const FNiagaraStatelessSpawnInfo& SpawnInfo : EmitterData->SpawnInfos)
+
+	if (bEmitterEnabled_CNC)
 	{
-		if (SpawnInfo.Type == ENiagaraStatelessSpawnInfoType::Rate)
+		for (const FNiagaraStatelessSpawnInfo& SpawnInfo : EmitterData->SpawnInfos)
 		{
-			const float SpawnRate = RandomStream.FRandRange(SpawnInfo.Rate.Min, SpawnInfo.Rate.Max);
-			if (SpawnRate > 0.0f)
+			if (SpawnInfo.Type == ENiagaraStatelessSpawnInfoType::Rate)
 			{
-				FActiveSpawnRate& ActiveSpawnRate = ActiveSpawnRates.AddDefaulted_GetRef();
-				ActiveSpawnRate.Rate		= SpawnRate;
-				ActiveSpawnRate.SpawnTime	= CurrentLoopDelay;
+				const float SpawnRate = RandomStream.FRandRange(SpawnInfo.Rate.Min, SpawnInfo.Rate.Max);
+				if (SpawnRate > 0.0f)
+				{
+					FActiveSpawnRate& ActiveSpawnRate = ActiveSpawnRates.AddDefaulted_GetRef();
+					ActiveSpawnRate.Rate = SpawnRate;
+					ActiveSpawnRate.SpawnTime = CurrentLoopDelay;
+				}
 			}
 		}
-	}
 
-	InitSpawnInfosForLoop();
+		InitSpawnInfosForLoop();
+	}
 }
 
 void FNiagaraStatelessEmitterInstance::InitSpawnInfosForLoop()
 {
+	if (bEmitterEnabled_CNC == false)
+	{
+		return;
+	}
+
 	// Add the next chunk for any active spawn rates
 	for (const FActiveSpawnRate& SpawnInfo : ActiveSpawnRates)
 	{
@@ -371,6 +425,7 @@ void FNiagaraStatelessEmitterInstance::InitSpawnInfosForLoop()
 		const int32 NumSpawned		= FMath::FloorToInt(ActiveDuration * SpawnInfo.Rate);
 
 		UniqueIndexOffset += NumSpawned;
+		bSpawnInfosDirty = true;
 	}
 
 	// Add bursts that fit within the loop duration (due to loop random they might not)
@@ -400,22 +455,63 @@ void FNiagaraStatelessEmitterInstance::InitSpawnInfosForLoop()
 		NewSpawnInfo.Amount			= SpawnAmount;
 
 		UniqueIndexOffset += SpawnAmount;
+		bSpawnInfosDirty = true;
 	}
+}
 
-	//-TODO: Better way to send data to renderer
-	if (NiagaraStateless::FEmitterInstance_RT* RenderThreadData = RenderThreadDataPtr.Get())
+void FNiagaraStatelessEmitterInstance::TickSpawnInfos()
+{
+	if (bEmitterEnabled_CNC != bEmitterEnabled_GT)
 	{
-		ENQUEUE_RENDER_COMMAND(FInitStatelessEmitter)(
-			[RenderThreadData=RenderThreadDataPtr.Get(), SpawnInfos_RT=SpawnInfos](FRHICommandListImmediate& RHICmdList) mutable
+		bEmitterEnabled_CNC = bEmitterEnabled_GT;
+
+		if (bEmitterEnabled_CNC)
+		{
+			if ( ActiveSpawnRates.Num() > 0 )
 			{
-				RenderThreadData->SpawnInfos = MoveTemp(SpawnInfos_RT);
+				for (const FActiveSpawnRate& SpawnInfo : ActiveSpawnRates)
+				{
+					float SpawnTime = FMath::Max(Age - SpawnInfo.SpawnTime, 0.0f);
+					SpawnTime = FMath::CeilToFloat(SpawnTime * SpawnInfo.Rate) / SpawnInfo.Rate;
+					SpawnTime += SpawnInfo.SpawnTime;
+					if (SpawnTime >= CurrentLoopAgeEnd)
+					{
+						continue;
+					}
+
+					FNiagaraStatelessRuntimeSpawnInfo& NewSpawnInfo = SpawnInfos.AddDefaulted_GetRef();
+					NewSpawnInfo.Type = ENiagaraStatelessSpawnInfoType::Rate;
+					NewSpawnInfo.UniqueOffset = UniqueIndexOffset;
+					NewSpawnInfo.SpawnTimeStart = SpawnTime;
+					NewSpawnInfo.SpawnTimeEnd = CurrentLoopAgeEnd;
+					NewSpawnInfo.Rate = SpawnInfo.Rate;
+
+					const float ActiveDuration = CurrentLoopAgeEnd - SpawnTime;
+					const int32 NumSpawned = FMath::FloorToInt(ActiveDuration * SpawnInfo.Rate);
+
+					UniqueIndexOffset += NumSpawned;
+					bSpawnInfosDirty = true;
+				}
 			}
-		);
+		}
+		else
+		{
+			for (auto SpawnInfoIt=SpawnInfos.CreateIterator(); SpawnInfoIt; ++SpawnInfoIt)
+			{
+				FNiagaraStatelessRuntimeSpawnInfo& SpawnInfo = *SpawnInfoIt;
+				if (SpawnInfo.SpawnTimeStart > Age)
+				{
+					SpawnInfoIt.RemoveCurrent();
+				}
+				else
+				{
+					SpawnInfo.SpawnTimeEnd = FMath::Min(SpawnInfo.SpawnTimeEnd, Age);
+				}
+			}
+			bSpawnInfosDirty = true;
+		}
 	}
-}
 
-void FNiagaraStatelessEmitterInstance::TickSpawnInfos()
-{
 	const float MaxLifetime = EmitterData->LifetimeRange.Max;
 	SpawnInfos.RemoveAll(
 		[this, &MaxLifetime](const FNiagaraStatelessRuntimeSpawnInfo& SpawnInfo)
