Change some code around to work around issue with static analysis and it being unable to handle TStaticArray

diff --git a/Engine/Plugins/FX/Niagara/Source/Niagara/Private/NiagaraUvQuadTree.cpp b/Engine/Plugins/FX/Niagara/Source/Niagara/Private/NiagaraUvQuadTree.cpp
--- a/Engine/Plugins/FX/Niagara/Source/Niagara/Private/NiagaraUvQuadTree.cpp
+++ b/Engine/Plugins/FX/Niagara/Source/Niagara/Private/NiagaraUvQuadTree.cpp
@@ -103,8 +103,10 @@ int32 FNiagaraUvQuadTree::FSubTree::Freeze(const FNiagaraUvQuadTree & QuadTree,
 
 	if (bInternal)
 	{
-		FChildArray ChildOffsets;
-		for (int32 ChildIt = 0; ChildIt < ChildOffsets.Num(); ++ChildIt) //-V621 //-V654
+		constexpr uint32 ChildCount = 4;
+		int32 ChildOffsets[ChildCount];
+
+		for (int32 ChildIt = 0; ChildIt < ChildCount; ++ChildIt)
 		{
 			ChildOffsets[ChildIt] = QuadTree.ChildTrees[SubTreeIndices[ChildIt]].Freeze(QuadTree, Ar);
 		}
@@ -112,7 +114,7 @@ int32 FNiagaraUvQuadTree::FSubTree::Freeze(const FNiagaraUvQuadTree & QuadTree,
 		const int64 EndOffset = Ar.Tell();
 		Ar.Seek(StartOffset);
 
-		for (int32 ChildIt = 0; ChildIt < ChildOffsets.Num(); ++ChildIt) //-V621 //-V654
+		for (int32 ChildIt = 0; ChildIt < ChildCount; ++ChildIt)
 		{
 			Ar << ChildOffsets[ChildIt];
 		}
