Fix for default values for variables using StackContext.  Before, the variables would be left as default initialized to 0.  Some modules will now have subtle behavioral differences with variables having better default values in their spawn script (quats will be 0,0,0,1 instead of 0,0,0,0), only to be assigned newer values in update; but also, the use of StackContext within update modules will have their defaults properly set.

diff --git a/Engine/Plugins/FX/Niagara/Source/NiagaraEditor/Private/NiagaraHlslTranslator.cpp b/Engine/Plugins/FX/Niagara/Source/NiagaraEditor/Private/NiagaraHlslTranslator.cpp
--- a/Engine/Plugins/FX/Niagara/Source/NiagaraEditor/Private/NiagaraHlslTranslator.cpp
+++ b/Engine/Plugins/FX/Niagara/Source/NiagaraEditor/Private/NiagaraHlslTranslator.cpp
@@ -5321,7 +5321,9 @@ bool FNiagaraHlslTranslator::IsWriteAllowedForNamespace(const FNiagaraVariable&
 template<typename GraphBridge>
 bool FNiagaraHlslTranslationStage::IsRelevantToSpawnForStage(const typename GraphBridge::FParamMapHistory& InHistory, const FNiagaraVariable& InAliasedVar, const FNiagaraVariable& InVar) const
 {
-	if (InHistory.IsPrimaryDataSetOutput(InAliasedVar, ScriptUsage) && (UNiagaraScript::IsSpawnScript(ScriptUsage) || bShouldUpdateInitialAttributeValues))
+	const FNiagaraVariableBase DataSetVariable = InAliasedVar.IsInNameSpace(FNiagaraConstants::StackContextNamespace) ? InVar : InAliasedVar;
+
+	if (InHistory.IsPrimaryDataSetOutput(DataSetVariable, ScriptUsage) && (UNiagaraScript::IsSpawnScript(ScriptUsage) || bShouldUpdateInitialAttributeValues))
 	{
 		return true;
 	}
@@ -5330,7 +5332,7 @@ bool FNiagaraHlslTranslationStage::IsRelevantToSpawnForStage(const typename Grap
 	{
 		switch (IterationSourceType)
 		{
-			case ENiagaraIterationSource::Particles:		return InHistory.IsPrimaryDataSetOutput(InAliasedVar, ENiagaraScriptUsage::EmitterSpawnScript); 
+			case ENiagaraIterationSource::Particles:		return InHistory.IsPrimaryDataSetOutput(DataSetVariable, ENiagaraScriptUsage::EmitterSpawnScript);
 			case ENiagaraIterationSource::DataInterface:	return InVar.IsInNameSpace(IterationDataInterface) && !InVar.IsDataInterface();
 			case ENiagaraIterationSource::DirectSet:		return false;
 			default:										check(false);
diff --git a/Engine/Plugins/FX/Niagara/Source/NiagaraEditor/Private/NiagaraParameterMapHistory.cpp b/Engine/Plugins/FX/Niagara/Source/NiagaraEditor/Private/NiagaraParameterMapHistory.cpp
--- a/Engine/Plugins/FX/Niagara/Source/NiagaraEditor/Private/NiagaraParameterMapHistory.cpp
+++ b/Engine/Plugins/FX/Niagara/Source/NiagaraEditor/Private/NiagaraParameterMapHistory.cpp
@@ -755,13 +755,13 @@ FName FNiagaraParameterUtilities::GetSourceForPreviousValue(const FName& InVaria
 }
 
 template<typename GraphBridge>
-bool TNiagaraParameterMapHistory<GraphBridge>::IsPrimaryDataSetOutput(const FNiagaraVariable& InVar, const UNiagaraScript* InScript,  bool bAllowDataInterfaces, bool bAllowStatics) const
+bool TNiagaraParameterMapHistory<GraphBridge>::IsPrimaryDataSetOutput(const FNiagaraVariableBase& InVar, const UNiagaraScript* InScript,  bool bAllowDataInterfaces, bool bAllowStatics) const
 {
 	return IsPrimaryDataSetOutput(InVar, InScript->GetUsage(),  bAllowDataInterfaces);
 }
 
 template<typename GraphBridge>
-bool TNiagaraParameterMapHistory<GraphBridge>::IsPrimaryDataSetOutput(const FNiagaraVariable& InVar, ENiagaraScriptUsage Usage, bool bAllowDataInterfaces, bool bAllowStatics) const
+bool TNiagaraParameterMapHistory<GraphBridge>::IsPrimaryDataSetOutput(const FNiagaraVariableBase& InVar, ENiagaraScriptUsage Usage, bool bAllowDataInterfaces, bool bAllowStatics) const
 {
 	if (bAllowDataInterfaces == false && InVar.GetType().GetClass() != nullptr)
 	{
@@ -778,20 +778,20 @@ bool TNiagaraParameterMapHistory<GraphBridge>::IsPrimaryDataSetOutput(const FNia
 	{
 		// In the case of system/emitter scripts we must include the variables in the overall system namespace as well as any of 
 		// the child emitters that were encountered.
-		for (FString EmitterEncounteredNamespace : EmitterNamespacesEncountered)
+		for (const FString& EmitterEncounteredNamespace : EmitterNamespacesEncountered)
 		{
-			if (FNiagaraParameterUtilities::IsInNamespace(InVar, EmitterEncounteredNamespace))
+			if (InVar.IsInNameSpace(EmitterEncounteredNamespace))
 			{
 				return true;
 			}
 		}
-		return FNiagaraParameterUtilities::IsInNamespace(InVar, PARAM_MAP_SYSTEM_STR) || FNiagaraParameterUtilities::IsInNamespace(InVar, PARAM_MAP_EMITTER_STR);
+		return InVar.IsInNameSpace(FStringView(PARAM_MAP_SYSTEM_STR)) || InVar.IsInNameSpace(FStringView(PARAM_MAP_EMITTER_STR));
 	}
 	else if (Usage == ENiagaraScriptUsage::Module || Usage == ENiagaraScriptUsage::Function)
 	{
-		return FNiagaraParameterUtilities::IsInNamespace(InVar, PARAM_MAP_MODULE_STR);
+		return InVar.IsInNameSpace(FStringView(PARAM_MAP_MODULE_STR));
 	}
-	return FNiagaraParameterUtilities::IsInNamespace(InVar, PARAM_MAP_ATTRIBUTE_STR);
+	return InVar.IsInNameSpace(FStringView(PARAM_MAP_ATTRIBUTE_STR));
 }
 
 bool FNiagaraParameterUtilities::IsWrittenToScriptUsage(const FNiagaraVariable& InVar, ENiagaraScriptUsage Usage, bool bAllowDataInterfaces)
diff --git a/Engine/Plugins/FX/Niagara/Source/NiagaraEditor/Public/NiagaraParameterMapHistory.h b/Engine/Plugins/FX/Niagara/Source/NiagaraEditor/Public/NiagaraParameterMapHistory.h
--- a/Engine/Plugins/FX/Niagara/Source/NiagaraEditor/Public/NiagaraParameterMapHistory.h
+++ b/Engine/Plugins/FX/Niagara/Source/NiagaraEditor/Public/NiagaraParameterMapHistory.h
@@ -354,8 +354,8 @@ public:
 
 
 	/** Called to determine if a given variable should be output from a script. It is not static as it requires the overall context to include emitter namespaces visited for system scripts.*/
-	bool IsPrimaryDataSetOutput(const FNiagaraVariable& InVar, const UNiagaraScript* InScript,   bool bAllowDataInterfaces = false,  bool bAllowStatics = false) const;
-	bool IsPrimaryDataSetOutput(const FNiagaraVariable& InVar, ENiagaraScriptUsage InUsage, bool bAllowDataInterfaces = false, bool bAllowStatics = false) const;
+	bool IsPrimaryDataSetOutput(const FNiagaraVariableBase& InVar, const UNiagaraScript* InScript,   bool bAllowDataInterfaces = false,  bool bAllowStatics = false) const;
+	bool IsPrimaryDataSetOutput(const FNiagaraVariableBase& InVar, ENiagaraScriptUsage InUsage, bool bAllowDataInterfaces = false, bool bAllowStatics = false) const;
 
 	/** Get the output node associated with this graph.*/
 	const FOutputNode* GetFinalOutputNode() const;
diff --git a/Engine/Source/Runtime/Core/Private/UObject/FortniteMainBranchObjectVersion.cpp b/Engine/Source/Runtime/Core/Private/UObject/FortniteMainBranchObjectVersion.cpp
--- a/Engine/Source/Runtime/Core/Private/UObject/FortniteMainBranchObjectVersion.cpp
+++ b/Engine/Source/Runtime/Core/Private/UObject/FortniteMainBranchObjectVersion.cpp
@@ -11,7 +11,7 @@ TMap<FGuid, FGuid> FFortniteMainBranchObjectVersion::GetSystemGuids()
 	SystemGuids.Add(DevGuids.MATERIALSHADERMAP_DERIVEDDATA_VER, FGuid("E416BA6E4D48491994BB91E4B06E43AC"));
 	SystemGuids.Add(DevGuids.NANITE_DERIVEDDATA_VER, FGuid("D28C003F825F4C178F698938FB0E47F7"));
 	SystemGuids.Add(DevGuids.NIAGARASHADERMAP_DERIVEDDATA_VER, FGuid("C417A50271B3427D98785388266E3B66"));
-	SystemGuids.Add(DevGuids.Niagara_LatestScriptCompileVersion, FGuid("48A35E7D14134FBE9DF89B3E485E5762"));
+	SystemGuids.Add(DevGuids.Niagara_LatestScriptCompileVersion, FGuid("E267F600682A4E94885E8AEC5CE1DC91"));
 	SystemGuids.Add(DevGuids.SkeletalMeshDerivedDataVersion, FGuid("B9AEAA3EE7AC4FCFB3860E8F3E06DD75"));
 	SystemGuids.Add(DevGuids.STATICMESH_DERIVEDDATA_VER, FGuid("543843B103794E0E9BA4BE80FB602F79"));
 	SystemGuids.Add(DevGuids.MaterialTranslationDDCVersion, FGuid("BBD650804D314CE7A44ADF9BB4F3F5F9"));
