- Fix for not generating cutout information on SubImageSize change

#tests sprite renderer set cutout then modify sub image size

diff --git a/Engine/Plugins/FX/Niagara/Source/Niagara/Private/NiagaraSpriteRendererProperties.cpp b/Engine/Plugins/FX/Niagara/Source/Niagara/Private/NiagaraSpriteRendererProperties.cpp
--- a/Engine/Plugins/FX/Niagara/Source/Niagara/Private/NiagaraSpriteRendererProperties.cpp
+++ b/Engine/Plugins/FX/Niagara/Source/Niagara/Private/NiagaraSpriteRendererProperties.cpp
@@ -480,6 +480,7 @@ void UNiagaraSpriteRendererProperties::PostEditChangeProperty(struct FPropertyCh
 			PropertyName == TEXT("BoundingMode") ||
 			PropertyName == TEXT("OpacitySourceMode") ||
 			PropertyName == TEXT("AlphaThreshold") ||
+			MemberPropertyName == TEXT("SubImageSize") ||
 			(bUseMaterialCutoutTexture && PropertyName == TEXT("Material"));
 
 		if (bUpdateCutoutDDC)
