Adapted parts of the Blueprint wiring style to Niagara
Main thing are tangents & hovering on a connection deemphasizes other connections

diff --git a/Engine/Plugins/FX/Niagara/Source/NiagaraEditor/Private/EdGraphSchema_Niagara.cpp b/Engine/Plugins/FX/Niagara/Source/NiagaraEditor/Private/EdGraphSchema_Niagara.cpp
--- a/Engine/Plugins/FX/Niagara/Source/NiagaraEditor/Private/EdGraphSchema_Niagara.cpp
+++ b/Engine/Plugins/FX/Niagara/Source/NiagaraEditor/Private/EdGraphSchema_Niagara.cpp
@@ -2446,29 +2446,28 @@ void UEdGraphSchema_Niagara::GetContextMenuActions(UToolMenu* Menu, UGraphNodeCo
 
 FNiagaraConnectionDrawingPolicy::FNiagaraConnectionDrawingPolicy(int32 InBackLayerID, int32 InFrontLayerID, float InZoomFactor, const FSlateRect& InClippingRect, FSlateWindowElementList& InDrawElements, UEdGraph* InGraph)
 	: FConnectionDrawingPolicy(InBackLayerID, InFrontLayerID, InZoomFactor, InClippingRect, InDrawElements)
-	, Graph(CastChecked<UNiagaraGraph>(InGraph))
+	, GraphObj(CastChecked<UNiagaraGraph>(InGraph))
 {
 	ArrowImage = nullptr;
 	ArrowRadius = FVector2D::ZeroVector;
+	
+	DefaultDataWireThickness = Settings->DefaultDataWireThickness;
+	DefaultExecutionWireThickness = Settings->DefaultExecutionWireThickness;
 }
 
 void FNiagaraConnectionDrawingPolicy::DetermineWiringStyle(UEdGraphPin* OutputPin, UEdGraphPin* InputPin, /*inout*/ FConnectionParams& Params)
 {
 	FConnectionDrawingPolicy::DetermineWiringStyle(OutputPin, InputPin, Params);
-	if (HoveredPins.Contains(InputPin) && HoveredPins.Contains(OutputPin))
-	{
-		Params.WireThickness = Params.WireThickness * 5;
-	}
-
-	if (Graph)
+	
+	if (GraphObj)
 	{
-		const UEdGraphSchema_Niagara* NSchema = Cast<UEdGraphSchema_Niagara>(Graph->GetSchema());
+		const UEdGraphSchema_Niagara* NSchema = Cast<UEdGraphSchema_Niagara>(GraphObj->GetSchema());
 		if (NSchema && OutputPin)
 		{
 			Params.WireColor = NSchema->GetPinTypeColor(OutputPin->PinType);
 			if (NSchema->PinToTypeDefinition(OutputPin) == FNiagaraTypeDefinition::GetGenericNumericDef())
 			{
-				FNiagaraTypeDefinition NewDef = Graph->GetCachedNumericConversion(OutputPin);
+				FNiagaraTypeDefinition NewDef = GraphObj->GetCachedNumericConversion(OutputPin);
 				if (NewDef.IsValid())
 				{
 					FEdGraphPinType NewPinType = NSchema->TypeDefinitionToPinType(NewDef);
@@ -2476,7 +2475,7 @@ void FNiagaraConnectionDrawingPolicy::DetermineWiringStyle(UEdGraphPin* OutputPi
 				}
 			}
 		}
-
+	
 		if(OutputPin && InputPin)
 		{
 			if(OutputPin->bOrphanedPin || InputPin->bOrphanedPin)
@@ -2485,6 +2484,156 @@ void FNiagaraConnectionDrawingPolicy::DetermineWiringStyle(UEdGraphPin* OutputPi
 			}
 		}
 	}
+
+	Params.WireThickness = DefaultDataWireThickness;
+	Params.AssociatedPin1 = OutputPin;
+	Params.AssociatedPin2 = InputPin;
+	
+	// Get the schema and grab the default color from it
+	check(OutputPin);
+	check(GraphObj);
+	
+	UEdGraphNode* OutputNode = OutputPin->GetOwningNode();
+	UEdGraphNode* InputNode = (InputPin != nullptr) ? InputPin->GetOwningNode() : nullptr;
+	
+	const bool bDeemphasizeUnhoveredPins = HoveredPins.Num() > 0;
+	
+	// If the output or input connect to a knot that is going backwards, we will flip the direction on values going into them
+	if (UNiagaraNodeReroute* OutputRerouteNode = Cast<UNiagaraNodeReroute>(OutputNode))
+	{
+		if (ShouldChangeTangentForReroute(OutputRerouteNode))
+		{
+			Params.StartDirection = EGPD_Input;
+		}
+	}
+	
+	if (UNiagaraNodeReroute* InputRerouteNode = Cast<UNiagaraNodeReroute>(InputNode))
+	{
+		if (ShouldChangeTangentForReroute(InputRerouteNode))
+		{
+			Params.EndDirection = EGPD_Output;
+		}
+	}
+	
+	// Container types should draw thicker
+	if ((InputPin && InputPin->PinType.IsContainer()) || (OutputPin && OutputPin->PinType.IsContainer()))
+	{
+		Params.WireThickness = DefaultExecutionWireThickness;
+	}
+	
+	// If either end of the connection is not enabled (and not a passthru to something else), draw the wire differently
+	bool bWireIsOnDisabledNodeAndNotPassthru = false;
+	if (OutputNode && (OutputNode->IsDisplayAsDisabledForced() || !OutputNode->IsNodeEnabled()))
+	{
+		if (OutputNode->GetPassThroughPin(OutputPin) == nullptr)
+		{
+			bWireIsOnDisabledNodeAndNotPassthru = true;
+		}
+	}
+	
+	if (InputNode && (InputNode->IsDisplayAsDisabledForced() || !InputNode->IsNodeEnabled()))
+	{
+		if (InputNode->GetPassThroughPin(InputPin) == nullptr)
+		{
+			bWireIsOnDisabledNodeAndNotPassthru = true;
+		}
+	}
+	
+	if ((OutputPin && OutputPin->GetOwningNode()->IsNodeUnrelated()) || (InputPin && InputPin->GetOwningNode()->IsNodeUnrelated()))
+	{
+		bWireIsOnDisabledNodeAndNotPassthru = true;
+	}
+	
+	if (bWireIsOnDisabledNodeAndNotPassthru)
+	{
+		Params.WireColor *= 0.5f;
+		Params.WireThickness = 0.5f;
+	}
+	
+	if (bDeemphasizeUnhoveredPins)
+	{
+		ApplyHoverDeemphasis(OutputPin, InputPin, /*inout*/ Params.WireThickness, /*inout*/ Params.WireColor);
+	}
+}
+
+bool FNiagaraConnectionDrawingPolicy::ShouldChangeTangentForReroute(UNiagaraNodeReroute* Node)
+{
+	if (bool* pResult = RerouteToReversedDirectionMap.Find(Node))
+	{
+		return *pResult;
+	}
+	else
+	{
+		bool bPinReversed = false;
+
+		FVector2D AverageLeftPin;
+		FVector2D AverageRightPin;
+		FVector2D CenterPin;
+		bool bCenterValid = FindPinCenter(Node->GetOutputPin(0), /*out*/ CenterPin);
+		bool bLeftValid = GetAverageConnectedPosition(Node, EGPD_Input, /*out*/ AverageLeftPin);
+		bool bRightValid = GetAverageConnectedPosition(Node, EGPD_Output, /*out*/ AverageRightPin);
+
+		if (bLeftValid && bRightValid)
+		{
+			bPinReversed = AverageRightPin.X < AverageLeftPin.X;
+		}
+		else if (bCenterValid)
+		{
+			if (bLeftValid)
+			{
+				bPinReversed = CenterPin.X < AverageLeftPin.X;
+			}
+			else if (bRightValid)
+			{
+				bPinReversed = AverageRightPin.X < CenterPin.X;
+			}
+		}
+
+		RerouteToReversedDirectionMap.Add(Node, bPinReversed);
+
+		return bPinReversed;
+	}
+}
+
+bool FNiagaraConnectionDrawingPolicy::GetAverageConnectedPosition(UNiagaraNodeReroute* Reroute, EEdGraphPinDirection Direction, FVector2D& OutPos) const
+{
+	FVector2D Result = FVector2D::ZeroVector;
+	int32 ResultCount = 0;
+
+	UEdGraphPin* Pin = (Direction == EGPD_Input) ? Reroute->GetInputPin(0) : Reroute->GetOutputPin(0);
+	for (UEdGraphPin* LinkedPin : Pin->LinkedTo)
+	{
+		FVector2D CenterPoint;
+		if (FindPinCenter(LinkedPin, /*out*/ CenterPoint))
+		{
+			Result += CenterPoint;
+			ResultCount++;
+		}
+	}
+
+	if (ResultCount > 0)
+	{
+		OutPos = Result * (1.0f / ResultCount);
+		return true;
+	}
+	else
+	{
+		return false;
+	}
+}
+
+bool FNiagaraConnectionDrawingPolicy::FindPinCenter(UEdGraphPin* Pin, FVector2D& OutCenter) const
+{
+	if (const TSharedPtr<SGraphPin>* pPinWidget = PinToPinWidgetMap.Find(Pin))
+	{
+		if (FArrangedWidget* pPinEntry = PinGeometries->Find((*pPinWidget).ToSharedRef()))
+		{
+			OutCenter = FGeometryHelper::CenterOf(pPinEntry->Geometry);
+			return true;
+		}
+	}
+
+	return false;
 }
 
 #undef LOCTEXT_NAMESPACE
diff --git a/Engine/Plugins/FX/Niagara/Source/NiagaraEditor/Public/EdGraphSchema_Niagara.h b/Engine/Plugins/FX/Niagara/Source/NiagaraEditor/Public/EdGraphSchema_Niagara.h
--- a/Engine/Plugins/FX/Niagara/Source/NiagaraEditor/Public/EdGraphSchema_Niagara.h
+++ b/Engine/Plugins/FX/Niagara/Source/NiagaraEditor/Public/EdGraphSchema_Niagara.h
@@ -207,7 +207,17 @@ public:
 	FNiagaraConnectionDrawingPolicy(int32 InBackLayerID, int32 InFrontLayerID, float InZoomFactor, const FSlateRect& InClippingRect, FSlateWindowElementList& InDrawElements, UEdGraph* InGraph);
 	virtual void DetermineWiringStyle(UEdGraphPin* OutputPin, UEdGraphPin* InputPin, /*inout*/ FConnectionParams& Params) override;
 
-private:
-	class UNiagaraGraph* Graph;
+protected:
+	bool ShouldChangeTangentForReroute(class UNiagaraNodeReroute* Node);
+	bool GetAverageConnectedPosition(class UNiagaraNodeReroute* Reroute, EEdGraphPinDirection Direction, FVector2D& OutPos) const;
+	bool FindPinCenter(UEdGraphPin* Pin, FVector2D& OutCenter) const;
+protected:
+	class UNiagaraGraph* GraphObj;
+	
+	float DefaultDataWireThickness;
+	float DefaultExecutionWireThickness;
+	
+	// Each time a Reroute is encountered, input geometry is compared to output geometry to see if the pins on the reroute need to be reversed
+	TMap<class UNiagaraNodeReroute*, bool> RerouteToReversedDirectionMap;
 };
 
