Fixed compile issue

diff --git a/Engine/Plugins/FX/Niagara/Source/NiagaraEditor/Public/NiagaraSystemFactoryNew.h b/Engine/Plugins/FX/Niagara/Source/NiagaraEditor/Public/NiagaraSystemFactoryNew.h
--- a/Engine/Plugins/FX/Niagara/Source/NiagaraEditor/Public/NiagaraSystemFactoryNew.h
+++ b/Engine/Plugins/FX/Niagara/Source/NiagaraEditor/Public/NiagaraSystemFactoryNew.h
@@ -27,6 +27,9 @@ class UNiagaraSystemFactoryNew : public UFactory
 
 public:
 	NIAGARAEDITOR_API static void InitializeSystem(UNiagaraSystem* System, bool bCreateDefaultNodes);
+
+private:
+	void TryAssignDefaultEffectType(UNiagaraSystem* System);
 };
 
 
