Fixed a bug where renaming a parameter in a module script would in some edge cases not correctly propagate to existing usages

diff --git a/Engine/Plugins/FX/Niagara/Source/Niagara/Public/NiagaraCommon.h b/Engine/Plugins/FX/Niagara/Source/Niagara/Public/NiagaraCommon.h
--- a/Engine/Plugins/FX/Niagara/Source/Niagara/Public/NiagaraCommon.h
+++ b/Engine/Plugins/FX/Niagara/Source/Niagara/Public/NiagaraCommon.h
@@ -360,12 +360,12 @@ struct FNiagaraFunctionSignature
 	/** Name of the function. */
 	UPROPERTY()
 	FName Name;
-	/** Input parameters to this function. */
+	/** Input parameters to this function. The data stored in the variables is used for default values. */
 	UPROPERTY()
-	TArray<FNiagaraVariable> Inputs; //TODO: Can these be FNiagaraVariableBase?
-	/** Input parameters of this function. */
+	TArray<FNiagaraVariable> Inputs;
+	/** Output parameters of this function. The data stored in the variables is used for default values. */
 	UPROPERTY()
-	TArray<FNiagaraVariable> Outputs; //TODO: Can these be FNiagaraVariableBase?
+	TArray<FNiagaraVariable> Outputs;
 	/** Id of the owner is this is a member function. */
 	UPROPERTY()
 	FName OwnerName;
diff --git a/Engine/Plugins/FX/Niagara/Source/NiagaraEditor/Private/NiagaraNodeParameterMapGet.cpp b/Engine/Plugins/FX/Niagara/Source/NiagaraEditor/Private/NiagaraNodeParameterMapGet.cpp
--- a/Engine/Plugins/FX/Niagara/Source/NiagaraEditor/Private/NiagaraNodeParameterMapGet.cpp
+++ b/Engine/Plugins/FX/Niagara/Source/NiagaraEditor/Private/NiagaraNodeParameterMapGet.cpp
@@ -595,11 +595,19 @@ bool UNiagaraNodeParameterMapGet::CommitEditablePinName(const FText& InName, UEd
 		Modify();
 		InGraphPinObj->Modify();
 
+		if (bSuppressEvents == false)
+		{
+			// we refresh the parameter references before changing the pin name, otherwise we might run into weird edge cases where a dirty graph tries to rename the wrong pin or doesn't copy the metadata correctly
+			GetNiagaraGraph()->ConditionalRefreshParameterReferences();
+		}
+
 		InGraphPinObj->PinName = *InName.ToString();
 		InGraphPinObj->PinFriendlyName = InName;
 
 		if (!bSuppressEvents)
+		{
 			OnPinRenamed(InGraphPinObj, OldPinName);
+		}
 
 		return true;
 	}
diff --git a/Engine/Plugins/FX/Niagara/Source/NiagaraEditor/Private/NiagaraNodeParameterMapSet.cpp b/Engine/Plugins/FX/Niagara/Source/NiagaraEditor/Private/NiagaraNodeParameterMapSet.cpp
--- a/Engine/Plugins/FX/Niagara/Source/NiagaraEditor/Private/NiagaraNodeParameterMapSet.cpp
+++ b/Engine/Plugins/FX/Niagara/Source/NiagaraEditor/Private/NiagaraNodeParameterMapSet.cpp
@@ -182,10 +182,18 @@ bool UNiagaraNodeParameterMapSet::CommitEditablePinName(const FText& InName, UEd
 		Modify();
 		InGraphPinObj->Modify();
 
+		if (bSuppressEvents == false)
+		{
+			// we refresh the parameter references before changing the pin name, otherwise we might run into weird edge cases where a dirty graph tries to rename the wrong pin or doesn't copy the metadata correctly
+			GetNiagaraGraph()->ConditionalRefreshParameterReferences();
+		}
+
 		InGraphPinObj->PinName = *NewPinName;
 		InGraphPinObj->PinFriendlyName = InName;
-		if (bSuppressEvents == false)	
+		if (bSuppressEvents == false)
+		{
 			OnPinRenamed(InGraphPinObj, OldPinName);
+		}
 
 		return true;
 	}
