- Add an option to disable creating actors in the VP from systems that do not have an effect type assigned
- This will likely be removed in a future release

diff --git a/Engine/Plugins/FX/Niagara/Source/Niagara/Public/NiagaraSettings.h b/Engine/Plugins/FX/Niagara/Source/Niagara/Public/NiagaraSettings.h
--- a/Engine/Plugins/FX/Niagara/Source/Niagara/Public/NiagaraSettings.h
+++ b/Engine/Plugins/FX/Niagara/Source/Niagara/Public/NiagaraSettings.h
@@ -194,6 +194,10 @@ class UNiagaraSettings : public UDeveloperSettings
 	UPROPERTY(config, EditAnywhere, Category = Niagara, meta = (AllowedClasses = "/Script/Niagara.NiagaraEffectType"))
 	FSoftObjectPath DefaultEffectType;
 
+	/** Should we allow placing a Niagara System in the editor into a level which has no effect type assigned? */
+	UPROPERTY(config, EditAnywhere, Category = Niagara)
+	bool bAllowCreateActorFromSystemWithNoEffectType = true;
+
 	/** Position pin type color. The other pin colors are defined in the general editor settings. */
 	UPROPERTY(config, EditAnywhere, Category=Niagara)
 	FLinearColor PositionPinTypeColor;
diff --git a/Engine/Plugins/FX/Niagara/Source/NiagaraEditor/Private/NiagaraEditorCommon.cpp b/Engine/Plugins/FX/Niagara/Source/NiagaraEditor/Private/NiagaraEditorCommon.cpp
--- a/Engine/Plugins/FX/Niagara/Source/NiagaraEditor/Private/NiagaraEditorCommon.cpp
+++ b/Engine/Plugins/FX/Niagara/Source/NiagaraEditor/Private/NiagaraEditorCommon.cpp
@@ -7,6 +7,7 @@
 #include "ActorFactoryNiagara.h"
 #include "NiagaraActor.h"
 #include "NiagaraComponent.h"
+#include "NiagaraSettings.h"
 #include "Misc/StringFormatter.h"
 
 #include UE_INLINE_GENERATED_CPP_BY_NAME(NiagaraEditorCommon)
@@ -1941,6 +1942,17 @@ bool UActorFactoryNiagara::CanCreateActorFrom(const FAssetData& AssetData, FText
 		return false;
 	}
 
+	const UNiagaraSettings* NiagarSettings = GetDefault<UNiagaraSettings>();
+	if (NiagarSettings && !NiagarSettings->bAllowCreateActorFromSystemWithNoEffectType)
+	{
+		const UNiagaraSystem* NiagaraSystem = Cast<UNiagaraSystem>(AssetData.GetAsset());
+		if (NiagaraSystem && NiagaraSystem->GetEffectType() == nullptr)
+		{
+			OutErrorMsg = NSLOCTEXT("CanCreateActor", "RequiredEffectType", "An effect type is required to create actor instances.");
+			return false;
+		}
+	}
+
 	return true;
 }
 
