- Improve error message with mismatched stateless shader parameters

diff --git a/Engine/Plugins/FX/Niagara/Source/Niagara/Private/Stateless/NiagaraStatelessModule.cpp b/Engine/Plugins/FX/Niagara/Source/Niagara/Private/Stateless/NiagaraStatelessModule.cpp
--- a/Engine/Plugins/FX/Niagara/Source/Niagara/Private/Stateless/NiagaraStatelessModule.cpp
+++ b/Engine/Plugins/FX/Niagara/Source/Niagara/Private/Stateless/NiagaraStatelessModule.cpp
@@ -22,7 +22,8 @@ void FNiagaraStatelessSetShaderParameterContext::ValidateIncludeStructType(uint3
 			return;
 		}
 
-		UE_LOG(LogNiagara, Fatal, TEXT("Shader parameter struct member (%s) at offset (%u) is not of type (%s)"), Member.GetName(), StructOffset, StructMetaData->GetStructTypeName());
+		const TCHAR* StructType = Member.GetStructMetadata() ? Member.GetStructMetadata()->GetStructTypeName() : TEXT("null");
+		UE_LOG(LogNiagara, Fatal, TEXT("Shader parameter struct member (%s) at offset (%u) is not of type (%s) struct type is (%s)"), Member.GetName(), StructOffset, StructMetaData->GetStructTypeName(), StructType);
 		return;
 	}
 
