- When building components don't include disabled modules
- Reduces memory usage of stateless emitters

diff --git a/Engine/Plugins/FX/Niagara/Source/Niagara/Internal/Stateless/Modules/NiagaraStatelessModule_DynamicMaterialParameters.h b/Engine/Plugins/FX/Niagara/Source/Niagara/Internal/Stateless/Modules/NiagaraStatelessModule_DynamicMaterialParameters.h
--- a/Engine/Plugins/FX/Niagara/Source/Niagara/Internal/Stateless/Modules/NiagaraStatelessModule_DynamicMaterialParameters.h
+++ b/Engine/Plugins/FX/Niagara/Source/Niagara/Internal/Stateless/Modules/NiagaraStatelessModule_DynamicMaterialParameters.h
@@ -134,13 +134,10 @@ public:
 	virtual void GetOutputVariables(TArray<FNiagaraVariableBase>& OutVariables) const override
 	{
 		const FNiagaraStatelessGlobals& StatelessGlobals = FNiagaraStatelessGlobals::Get();
-		if (IsModuleEnabled())
+		//-TODO: Channel masking, etc
+		if (GetParameterChannelMask(0) != 0)
 		{
-			//-TODO: Channel masking, etc
-			if (GetParameterChannelMask(0) != 0)
-			{
-				OutVariables.AddUnique(StatelessGlobals.DynamicMaterialParameters0Variable);
-			}
+			OutVariables.AddUnique(StatelessGlobals.DynamicMaterialParameters0Variable);
 		}
 	}
 #endif
diff --git a/Engine/Plugins/FX/Niagara/Source/Niagara/Internal/Stateless/Modules/NiagaraStatelessModule_SpriteFacingAndAlignment.h b/Engine/Plugins/FX/Niagara/Source/Niagara/Internal/Stateless/Modules/NiagaraStatelessModule_SpriteFacingAndAlignment.h
--- a/Engine/Plugins/FX/Niagara/Source/Niagara/Internal/Stateless/Modules/NiagaraStatelessModule_SpriteFacingAndAlignment.h
+++ b/Engine/Plugins/FX/Niagara/Source/Niagara/Internal/Stateless/Modules/NiagaraStatelessModule_SpriteFacingAndAlignment.h
@@ -42,18 +42,15 @@ public:
 	virtual void GetOutputVariables(TArray<FNiagaraVariableBase>& OutVariables) const override
 	{
 		const FNiagaraStatelessGlobals& StatelessGlobals = FNiagaraStatelessGlobals::Get();
-		if (IsModuleEnabled())
+		if (bSpriteFacingEnabled)
 		{
-			if (bSpriteFacingEnabled)
-			{
-				OutVariables.AddUnique(StatelessGlobals.SpriteFacingVariable);
-				OutVariables.AddUnique(StatelessGlobals.PreviousSpriteFacingVariable);
-			}
-			if (bSpriteAlignmentEnabled)
-			{
-				OutVariables.AddUnique(StatelessGlobals.SpriteAlignmentVariable);
-				OutVariables.AddUnique(StatelessGlobals.PreviousSpriteAlignmentVariable);
-			}
+			OutVariables.AddUnique(StatelessGlobals.SpriteFacingVariable);
+			OutVariables.AddUnique(StatelessGlobals.PreviousSpriteFacingVariable);
+		}
+		if (bSpriteAlignmentEnabled)
+		{
+			OutVariables.AddUnique(StatelessGlobals.SpriteAlignmentVariable);
+			OutVariables.AddUnique(StatelessGlobals.PreviousSpriteAlignmentVariable);
 		}
 	}
 #endif
diff --git a/Engine/Plugins/FX/Niagara/Source/Niagara/Private/Stateless/NiagaraStatelessEmitter.cpp b/Engine/Plugins/FX/Niagara/Source/Niagara/Private/Stateless/NiagaraStatelessEmitter.cpp
--- a/Engine/Plugins/FX/Niagara/Source/Niagara/Private/Stateless/NiagaraStatelessEmitter.cpp
+++ b/Engine/Plugins/FX/Niagara/Source/Niagara/Private/Stateless/NiagaraStatelessEmitter.cpp
@@ -303,7 +303,10 @@ void UNiagaraStatelessEmitter::BuildCompiledDataSet()
 		TArray<FNiagaraVariableBase> AvailableVariables;
 		for (UNiagaraStatelessModule* Module : Modules)
 		{
-			Module->GetOutputVariables(AvailableVariables);
+			if (Module->IsModuleEnabled())
+			{
+				Module->GetOutputVariables(AvailableVariables);
+			}
 		}
 
 		// Remove any variables we don't output
