- Mesh Index module for stateless
- Missing the ability to animate mesh index over life
- Allows for weighted random mesh index, this is done using an array currently but should be moved into a widget

diff --git a/Engine/Plugins/FX/Niagara/Shaders/Private/Stateless/Modules/NiagaraStatelessModule_MeshIndex.ush b/Engine/Plugins/FX/Niagara/Shaders/Private/Stateless/Modules/NiagaraStatelessModule_MeshIndex.ush
new file mode 100644
--- /dev/null
+++ b/Engine/Plugins/FX/Niagara/Shaders/Private/Stateless/Modules/NiagaraStatelessModule_MeshIndex.ush
@@ -0,0 +1,41 @@
+// Copyright Epic Games, Inc. All Rights Reserved.
+
+#include "../NiagaraStatelessCommon.ush"
+
+int MeshIndex_Index;
+int MeshIndex_TableOffset;
+int MeshIndex_TableNumElements;	// -1 to avoid added instructions
+
+struct FStatelessModule_MeshIndex
+{
+};
+
+void MeshIndex_Initialize(inout FStatelessModule_MeshIndex ModuleData, inout FStatelessParticle Particle)
+{
+}
+	
+void MeshIndex_Simulate(in FStatelessModule_MeshIndex ModuleData, inout FStatelessParticle Particle)
+{
+	const bool bIsParameter = (MeshIndex_Index & 0x80000000) != 0;
+	int MeshIndex = MeshIndex_Index & ~0x80000000;
+	
+	if (bIsParameter)
+	{
+		MeshIndex = asint(GetParameterBufferFloat(MeshIndex, 0));
+	}
+	else
+	{
+		if ( MeshIndex_TableNumElements > 0 )
+		{
+			const float2 Rand = RandomFloat2(Particle);
+			MeshIndex = round(Rand.x * float(MeshIndex_TableNumElements));
+			const float Probability = GetStaticFloat(MeshIndex_TableOffset, (MeshIndex * 2) + 0);
+			if (Rand.y > Probability)
+			{
+				MeshIndex = int(GetStaticFloat(MeshIndex_TableOffset, (MeshIndex * 2) + 1));
+			}
+		}
+	}
+	Particle.MeshIndex = MeshIndex;
+}
+
diff --git a/Engine/Plugins/FX/Niagara/Shaders/Private/Stateless/NiagaraStatelessSimulationDefault.usf b/Engine/Plugins/FX/Niagara/Shaders/Private/Stateless/NiagaraStatelessSimulationDefault.usf
--- a/Engine/Plugins/FX/Niagara/Shaders/Private/Stateless/NiagaraStatelessSimulationDefault.usf
+++ b/Engine/Plugins/FX/Niagara/Shaders/Private/Stateless/NiagaraStatelessSimulationDefault.usf
@@ -6,6 +6,7 @@
 	PARTICLE_ATTRIBUTE(float,	CameraOffset) \
 	PARTICLE_ATTRIBUTE(float4,	Color) \
 	PARTICLE_ATTRIBUTE(float4,	DynamicMaterialParameter0) \
+	PARTICLE_ATTRIBUTE(int,		MeshIndex) \
 	PARTICLE_ATTRIBUTE(float4,	MeshOrientation) \
 	PARTICLE_ATTRIBUTE(float,	RibbonWidth) \
 	PARTICLE_ATTRIBUTE(float3,	Scale) \
@@ -37,6 +38,7 @@
 	PARTICLE_MODULE(ScaleSpriteSize) \
 	PARTICLE_MODULE(SpriteFacingAndAlignment) \
 	PARTICLE_MODULE(SpriteRotationRate) \
+	PARTICLE_MODULE(MeshIndex) \
 	PARTICLE_MODULE(MeshRotationRate) \
 	PARTICLE_MODULE(ScaleMeshSize) \
 	PARTICLE_MODULE(SolveVelocitiesAndForces) \
diff --git a/Engine/Plugins/FX/Niagara/Shaders/Private/Stateless/NiagaraStatelessSimulationExample1.usf b/Engine/Plugins/FX/Niagara/Shaders/Private/Stateless/NiagaraStatelessSimulationExample1.usf
--- a/Engine/Plugins/FX/Niagara/Shaders/Private/Stateless/NiagaraStatelessSimulationExample1.usf
+++ b/Engine/Plugins/FX/Niagara/Shaders/Private/Stateless/NiagaraStatelessSimulationExample1.usf
@@ -19,6 +19,7 @@
 #define PARTICLE_ATTRIBUTES_TRANSIENT \
 	PARTICLE_ATTRIBUTE(float,	CameraOffset) \
 	PARTICLE_ATTRIBUTE(float4,	DynamicMaterialParameter0) \
+	PARTICLE_ATTRIBUTE(int,		MeshIndex) \
 	PARTICLE_ATTRIBUTE(float4,	MeshOrientation) \
 	PARTICLE_ATTRIBUTE(float3,	SpriteFacing) \
 	PARTICLE_ATTRIBUTE(float3,	SpriteAlignment) \
diff --git a/Engine/Plugins/FX/Niagara/Shaders/Private/Stateless/NiagaraStatelessSimulationTemplate.ush b/Engine/Plugins/FX/Niagara/Shaders/Private/Stateless/NiagaraStatelessSimulationTemplate.ush
--- a/Engine/Plugins/FX/Niagara/Shaders/Private/Stateless/NiagaraStatelessSimulationTemplate.ush
+++ b/Engine/Plugins/FX/Niagara/Shaders/Private/Stateless/NiagaraStatelessSimulationTemplate.ush
@@ -16,6 +16,7 @@
 #include "Modules/NiagaraStatelessModule_DynamicMaterialParameters.ush"
 #include "Modules/NiagaraStatelessModule_InitializeParticle.ush"
 #include "Modules/NiagaraStatelessModule_InitialMeshOrientation.ush"
+#include "Modules/NiagaraStatelessModule_MeshIndex.ush"
 #include "Modules/NiagaraStatelessModule_MeshRotationRate.ush"
 #include "Modules/NiagaraStatelessModule_RotateAroundPoint.ush"
 #include "Modules/NiagaraStatelessModule_ShapeLocation.ush"
diff --git a/Engine/Plugins/FX/Niagara/Source/Niagara/Internal/Stateless/Modules/NiagaraStatelessModule_MeshIndex.h b/Engine/Plugins/FX/Niagara/Source/Niagara/Internal/Stateless/Modules/NiagaraStatelessModule_MeshIndex.h
new file mode 100644
--- /dev/null
+++ b/Engine/Plugins/FX/Niagara/Source/Niagara/Internal/Stateless/Modules/NiagaraStatelessModule_MeshIndex.h
@@ -0,0 +1,42 @@
+// Copyright Epic Games, Inc. All Rights Reserved.
+
+#pragma once
+
+#include "Stateless/NiagaraStatelessModule.h"
+#include "Stateless/NiagaraStatelessEmitterDataBuildContext.h"
+#include "Stateless/NiagaraStatelessModuleShaderParameters.h"
+
+#include "NiagaraStatelessModule_MeshIndex.generated.h"
+
+UCLASS(MinimalAPI, EditInlineNew, meta = (DisplayName = "Mesh Index"))
+class UNiagaraStatelessModule_MeshIndex : public UNiagaraStatelessModule
+{
+	GENERATED_BODY()
+
+public:
+	using FParameters = NiagaraStateless::FMeshIndexModule_ShaderParameters;
+
+	UPROPERTY(EditAnywhere, Category = "Parameters")
+	FNiagaraDistributionRangeInt MeshIndex = FNiagaraDistributionRangeInt(0);
+
+	/* Weight for each potential mesh index when using a range. */
+	UPROPERTY(EditAnywhere, Category = "Parameters", meta = (EditConditionHides, EditCondition="NeedsMeshIndexWeights()"))
+	TArray<float> MeshIndexWeight;
+
+	virtual void BuildEmitterData(FNiagaraStatelessEmitterDataBuildContext& BuildContext) const override;
+	virtual void SetShaderParameters(const FNiagaraStatelessSetShaderParameterContext& SetShaderParameterContext) const override;
+
+	UFUNCTION()
+	bool NeedsMeshIndexWeights() const { return MeshIndex.IsRange(); }
+
+#if WITH_EDITOR
+	virtual bool CanDisableModule() const override { return true; }
+#endif
+#if WITH_EDITORONLY_DATA
+	virtual void GetOutputVariables(TArray<FNiagaraVariableBase>& OutVariables) const override
+	{
+		const FNiagaraStatelessGlobals& StatelessGlobals = FNiagaraStatelessGlobals::Get();
+		OutVariables.AddUnique(StatelessGlobals.MeshIndexVariable);
+	}
+#endif
+};
diff --git a/Engine/Plugins/FX/Niagara/Source/Niagara/Internal/Stateless/NiagaraStatelessCommon.h b/Engine/Plugins/FX/Niagara/Source/Niagara/Internal/Stateless/NiagaraStatelessCommon.h
--- a/Engine/Plugins/FX/Niagara/Source/Niagara/Internal/Stateless/NiagaraStatelessCommon.h
+++ b/Engine/Plugins/FX/Niagara/Source/Niagara/Internal/Stateless/NiagaraStatelessCommon.h
@@ -32,6 +32,7 @@ struct FNiagaraStatelessGlobals
 	FNiagaraVariableBase	CameraOffsetVariable;
 	FNiagaraVariableBase	ColorVariable;
 	FNiagaraVariableBase	DynamicMaterialParameters0Variable;
+	FNiagaraVariableBase	MeshIndexVariable;
 	FNiagaraVariableBase	MeshOrientationVariable;
 	FNiagaraVariableBase	PositionVariable;
 	FNiagaraVariableBase	RibbonWidthVariable;
diff --git a/Engine/Plugins/FX/Niagara/Source/Niagara/Internal/Stateless/NiagaraStatelessDistribution.h b/Engine/Plugins/FX/Niagara/Source/Niagara/Internal/Stateless/NiagaraStatelessDistribution.h
--- a/Engine/Plugins/FX/Niagara/Source/Niagara/Internal/Stateless/NiagaraStatelessDistribution.h
+++ b/Engine/Plugins/FX/Niagara/Source/Niagara/Internal/Stateless/NiagaraStatelessDistribution.h
@@ -37,6 +37,7 @@ struct FNiagaraDistributionBase
 	bool IsConstant() const { return Mode == ENiagaraDistributionMode::UniformConstant || Mode == ENiagaraDistributionMode::NonUniformConstant; }
 	bool IsUniform() const { return Mode == ENiagaraDistributionMode::UniformConstant || Mode == ENiagaraDistributionMode::UniformRange; }
 	bool IsCurve() const { return Mode == ENiagaraDistributionMode::UniformCurve || Mode == ENiagaraDistributionMode::NonUniformCurve; }
+	bool IsRange() const { return Mode == ENiagaraDistributionMode::UniformRange || Mode == ENiagaraDistributionMode::NonUniformRange; }
 
 #if WITH_EDITORONLY_DATA
 	UPROPERTY(EditAnywhere, Category = "Parameters")
@@ -78,6 +79,12 @@ struct FNiagaraDistributionRangeInt
 
 	NIAGARA_API void InitConstant(int32 Value);
 	NIAGARA_API FNiagaraStatelessRangeInt CalculateRange(const int32 Default = 0) const;
+
+	bool IsBinding() const { return Mode == ENiagaraDistributionMode::Binding; }
+	bool IsConstant() const { return Mode == ENiagaraDistributionMode::UniformConstant || Mode == ENiagaraDistributionMode::NonUniformConstant; }
+	bool IsUniform() const { return Mode == ENiagaraDistributionMode::UniformConstant || Mode == ENiagaraDistributionMode::UniformRange; }
+	bool IsCurve() const { return Mode == ENiagaraDistributionMode::UniformCurve || Mode == ENiagaraDistributionMode::NonUniformCurve; }
+	bool IsRange() const { return Mode == ENiagaraDistributionMode::UniformRange || Mode == ENiagaraDistributionMode::NonUniformRange; }
 };
 
 USTRUCT()
diff --git a/Engine/Plugins/FX/Niagara/Source/Niagara/Private/Stateless/Modules/NiagaraStatelessModule_MeshIndex.cpp b/Engine/Plugins/FX/Niagara/Source/Niagara/Private/Stateless/Modules/NiagaraStatelessModule_MeshIndex.cpp
new file mode 100644
--- /dev/null
+++ b/Engine/Plugins/FX/Niagara/Source/Niagara/Private/Stateless/Modules/NiagaraStatelessModule_MeshIndex.cpp
@@ -0,0 +1,94 @@
+// Copyright Epic Games, Inc. All Rights Reserved.
+
+#include "Stateless/Modules/NiagaraStatelessModule_MeshIndex.h"
+
+#include "WeightedRandomSampler.h"
+
+namespace NSMMeshIndexPrivate
+{
+	struct FMeshIndexWeightedSampler : public FWeightedRandomSampler
+	{
+		FMeshIndexWeightedSampler(int32 IndexRange, TConstArrayView<float> IndexWeights)
+		{
+			TotalWeight = 0.0f;
+			PerIndexWeight.Reserve(IndexRange);
+			for (int32 i = 0; i < IndexRange; ++i)
+			{
+				if (IndexWeights.Num() > 0)
+				{
+					PerIndexWeight.Add(IndexWeights.IsValidIndex(i) ? FMath::Max(IndexWeights[i], 0.0f) : 0.0f);
+				}
+				else
+				{
+					PerIndexWeight.Add(1.0f);
+				}
+				TotalWeight += PerIndexWeight.Last();
+			}
+		}
+
+		virtual float GetWeights(TArray<float>& OutWeights) override
+		{
+			OutWeights = PerIndexWeight;
+			return TotalWeight;
+		}
+
+		TArray<float> PerIndexWeight;
+		float TotalWeight = 0.0f;
+	};
+
+	struct FModuleBuiltData
+	{
+		int	Index				= 0;
+		int	TableOffset			= 0;
+		int	TableNumElements	= 0;
+	};
+}
+
+void UNiagaraStatelessModule_MeshIndex::BuildEmitterData(FNiagaraStatelessEmitterDataBuildContext& BuildContext) const
+{
+	using namespace NSMMeshIndexPrivate;
+
+	FModuleBuiltData* BuiltData = BuildContext.AllocateBuiltData<FModuleBuiltData>();
+	if (IsModuleEnabled())
+	{
+		//-TODO: We don't have an int binding at the moment
+		//const FNiagaraStatelessRangeInt MeshIndexRange = BuildContext.ConvertDistributionToRange(MeshIndex, 0, true);
+		//if (MeshIndexRange.ParameterOffset != INDEX_NONE)
+		//{
+		//}
+
+		const FNiagaraStatelessRangeInt MeshIndexRange = MeshIndex.CalculateRange();
+		BuiltData->Index = MeshIndexRange.Min;
+ 		if (MeshIndexRange.GetScale() > 0 && MeshIndexRange.GetScale() < 256)
+		{
+			FMeshIndexWeightedSampler Sampler(MeshIndexRange.GetScale() + 1, MeshIndexWeight);
+			Sampler.Initialize();
+
+			const int32 NumTableEntries = Sampler.GetNumEntries();
+			if (NumTableEntries > 1)
+			{
+				BuiltData->TableNumElements = NumTableEntries - 1;
+
+				TArray<float, TInlineAllocator<16>> StaticData;
+				StaticData.AddUninitialized(NumTableEntries * 2);
+				for (int32 i=0; i < NumTableEntries; ++i)
+				{
+					StaticData[i * 2 + 0] = Sampler.GetProb()[i];
+					StaticData[i * 2 + 1] = float(MeshIndexRange.Min + Sampler.GetAlias()[i]);
+				}
+				BuiltData->TableOffset = BuildContext.AddStaticData(StaticData);
+			}
+		}
+	}
+}
+
+void UNiagaraStatelessModule_MeshIndex::SetShaderParameters(const FNiagaraStatelessSetShaderParameterContext& SetShaderParameterContext) const
+{
+	using namespace NSMMeshIndexPrivate;
+
+	FParameters* Parameters = SetShaderParameterContext.GetParameterNestedStruct<FParameters>();
+	const FModuleBuiltData* ModuleBuiltData = SetShaderParameterContext.ReadBuiltData<FModuleBuiltData>();
+	Parameters->MeshIndex_Index				= ModuleBuiltData->Index;
+	Parameters->MeshIndex_TableOffset		= ModuleBuiltData->TableOffset;
+	Parameters->MeshIndex_TableNumElements	= ModuleBuiltData->TableNumElements;
+}
diff --git a/Engine/Plugins/FX/Niagara/Source/Niagara/Private/Stateless/NiagaraStatelessCommon.cpp b/Engine/Plugins/FX/Niagara/Source/Niagara/Private/Stateless/NiagaraStatelessCommon.cpp
--- a/Engine/Plugins/FX/Niagara/Source/Niagara/Private/Stateless/NiagaraStatelessCommon.cpp
+++ b/Engine/Plugins/FX/Niagara/Source/Niagara/Private/Stateless/NiagaraStatelessCommon.cpp
@@ -13,6 +13,7 @@ namespace NiagaraStatelessCommon
 		GGlobals.CameraOffsetVariable				= FNiagaraVariableBase(FNiagaraTypeDefinition::GetFloatDef(), TEXT("CameraOffset"));
 		GGlobals.ColorVariable						= FNiagaraVariableBase(FNiagaraTypeDefinition::GetColorDef(), TEXT("Color"));
 		GGlobals.DynamicMaterialParameters0Variable = FNiagaraVariableBase(FNiagaraTypeDefinition::GetVec4Def(), TEXT("DynamicMaterialParameter"));
+		GGlobals.MeshIndexVariable					= FNiagaraVariableBase(FNiagaraTypeDefinition::GetIntDef(), TEXT("MeshIndex"));
 		GGlobals.MeshOrientationVariable			= FNiagaraVariableBase(FNiagaraTypeDefinition::GetQuatDef(), TEXT("MeshOrientation"));
 		GGlobals.PositionVariable					= FNiagaraVariableBase(FNiagaraTypeDefinition::GetPositionDef(), TEXT("Position"));
 		GGlobals.RibbonWidthVariable				= FNiagaraVariableBase(FNiagaraTypeDefinition::GetFloatDef(), TEXT("RibbonWidth"));
diff --git a/Engine/Plugins/FX/Niagara/Source/Niagara/Private/Stateless/NiagaraStatelessEmitterTemplate.cpp b/Engine/Plugins/FX/Niagara/Source/Niagara/Private/Stateless/NiagaraStatelessEmitterTemplate.cpp
--- a/Engine/Plugins/FX/Niagara/Source/Niagara/Private/Stateless/NiagaraStatelessEmitterTemplate.cpp
+++ b/Engine/Plugins/FX/Niagara/Source/Niagara/Private/Stateless/NiagaraStatelessEmitterTemplate.cpp
@@ -17,6 +17,7 @@
 #include "Stateless/Modules/NiagaraStatelessModule_GravityForce.h"
 #include "Stateless/Modules/NiagaraStatelessModule_InitializeParticle.h"
 #include "Stateless/Modules/NiagaraStatelessModule_InitialMeshOrientation.h"
+#include "Stateless/Modules/NiagaraStatelessModule_MeshIndex.h"
 #include "Stateless/Modules/NiagaraStatelessModule_MeshRotationRate.h"
 #include "Stateless/Modules/NiagaraStatelessModule_RotateAroundPoint.h"
 #include "Stateless/Modules/NiagaraStatelessModule_ScaleColor.h"
@@ -80,6 +81,7 @@ void UNiagaraStatelessEmitterDefault::InitModulesAndAttributes()
 		UNiagaraStatelessModule_ScaleSpriteSizeBySpeed::StaticClass(),
 		UNiagaraStatelessModule_ScaleMeshSize::StaticClass(),
 		UNiagaraStatelessModule_ScaleMeshSizeBySpeed::StaticClass(),
+		UNiagaraStatelessModule_MeshIndex::StaticClass(),
 		UNiagaraStatelessModule_MeshRotationRate::StaticClass(),
 		UNiagaraStatelessModule_AddVelocity::StaticClass(),
 		UNiagaraStatelessModule_AccelerationForce::StaticClass(),
@@ -101,6 +103,7 @@ void UNiagaraStatelessEmitterDefault::InitModulesAndAttributes()
 		StatelessGlobals.CameraOffsetVariable,
 		StatelessGlobals.ColorVariable,
 		StatelessGlobals.DynamicMaterialParameters0Variable,
+		StatelessGlobals.MeshIndexVariable,
 		StatelessGlobals.MeshOrientationVariable,
 		StatelessGlobals.RibbonWidthVariable,
 		StatelessGlobals.ScaleVariable,
@@ -147,6 +150,7 @@ void UNiagaraStatelessEmitterDefault::SetShaderParameters(uint8* ShaderParameter
 	ShaderParameters->Permutation_CameraOffsetComponent				= ComponentOffsets[iComponent++];
 	ShaderParameters->Permutation_ColorComponent					= ComponentOffsets[iComponent++];
 	ShaderParameters->Permutation_DynamicMaterialParameter0Component= ComponentOffsets[iComponent++];
+	ShaderParameters->Permutation_MeshIndexComponent				= ComponentOffsets[iComponent++];
 	ShaderParameters->Permutation_MeshOrientationComponent			= ComponentOffsets[iComponent++];
 	ShaderParameters->Permutation_RibbonWidthComponent				= ComponentOffsets[iComponent++];
 	ShaderParameters->Permutation_ScaleComponent					= ComponentOffsets[iComponent++];
diff --git a/Engine/Plugins/FX/Niagara/Source/NiagaraEditor/Private/Stateless/NiagaraDistributionIntPropertyCustomization.h b/Engine/Plugins/FX/Niagara/Source/NiagaraEditor/Private/Stateless/NiagaraDistributionIntPropertyCustomization.h
--- a/Engine/Plugins/FX/Niagara/Source/NiagaraEditor/Private/Stateless/NiagaraDistributionIntPropertyCustomization.h
+++ b/Engine/Plugins/FX/Niagara/Source/NiagaraEditor/Private/Stateless/NiagaraDistributionIntPropertyCustomization.h
@@ -10,6 +10,7 @@ class FNiagaraDistributionIntPropertyCustomization : public IPropertyTypeCustomi
 {
 public:
 	static TSharedRef<IPropertyTypeCustomization> MakeIntInstance(UObject* OptionalOuter);
+	static TSharedRef<IPropertyTypeCustomization> MakeIntInstance() { return MakeIntInstance(nullptr); }
 
 	virtual void CustomizeHeader(TSharedRef<IPropertyHandle> PropertyHandle, FDetailWidgetRow& HeaderRow, IPropertyTypeCustomizationUtils& CustomizationUtils) override;
 	virtual void CustomizeChildren(TSharedRef<IPropertyHandle> PropertyHandle, IDetailChildrenBuilder& ChildBuilder, IPropertyTypeCustomizationUtils& CustomizationUtils) override;
diff --git a/Engine/Plugins/FX/Niagara/Source/NiagaraEditor/Private/ViewModels/Stack/NiagaraStackStatelessEmitterGroup.cpp b/Engine/Plugins/FX/Niagara/Source/NiagaraEditor/Private/ViewModels/Stack/NiagaraStackStatelessEmitterGroup.cpp
--- a/Engine/Plugins/FX/Niagara/Source/NiagaraEditor/Private/ViewModels/Stack/NiagaraStackStatelessEmitterGroup.cpp
+++ b/Engine/Plugins/FX/Niagara/Source/NiagaraEditor/Private/ViewModels/Stack/NiagaraStackStatelessEmitterGroup.cpp
@@ -5,6 +5,7 @@
 #include "IDetailTreeNode.h"
 #include "PropertyHandle.h"
 #include "Stateless/NiagaraDistributionPropertyCustomization.h"
+#include "Stateless/NiagaraDistributionIntPropertyCustomization.h"
 #include "Stateless/NiagaraStatelessEmitter.h"
 #include "Styling/AppStyle.h"
 #include "ViewModels/Stack/NiagaraStackObject.h"
@@ -122,6 +123,7 @@ void UNiagaraStackStatelessEmitterObjectItem::RefreshChildrenInternal(const TArr
 			StatelessEmitterStackObject->RegisterInstancedCustomPropertyTypeLayout(FNiagaraDistributionRangeVector2::StaticStruct()->GetFName(), FOnGetPropertyTypeCustomizationInstance::CreateStatic(&FNiagaraDistributionPropertyCustomization::MakeVector2Instance));
 			StatelessEmitterStackObject->RegisterInstancedCustomPropertyTypeLayout(FNiagaraDistributionRangeVector3::StaticStruct()->GetFName(), FOnGetPropertyTypeCustomizationInstance::CreateStatic(&FNiagaraDistributionPropertyCustomization::MakeVector3Instance));
 			StatelessEmitterStackObject->RegisterInstancedCustomPropertyTypeLayout(FNiagaraDistributionRangeColor::StaticStruct()->GetFName(), FOnGetPropertyTypeCustomizationInstance::CreateStatic(&FNiagaraDistributionPropertyCustomization::MakeColorInstance));
+			StatelessEmitterStackObject->RegisterInstancedCustomPropertyTypeLayout(FNiagaraDistributionRangeInt::StaticStruct()->GetFName(), FOnGetPropertyTypeCustomizationInstance::CreateStatic(&FNiagaraDistributionIntPropertyCustomization::MakeIntInstance));
 			StatelessEmitterStackObjectWeak = StatelessEmitterStackObject;
 		}
 		NewChildren.Add(StatelessEmitterStackObject);
diff --git a/Engine/Plugins/FX/Niagara/Source/NiagaraEditor/Private/ViewModels/Stack/NiagaraStackStatelessEmitterSimulateGroup.cpp b/Engine/Plugins/FX/Niagara/Source/NiagaraEditor/Private/ViewModels/Stack/NiagaraStackStatelessEmitterSimulateGroup.cpp
--- a/Engine/Plugins/FX/Niagara/Source/NiagaraEditor/Private/ViewModels/Stack/NiagaraStackStatelessEmitterSimulateGroup.cpp
+++ b/Engine/Plugins/FX/Niagara/Source/NiagaraEditor/Private/ViewModels/Stack/NiagaraStackStatelessEmitterSimulateGroup.cpp
@@ -8,6 +8,7 @@
 #include "PropertyHandle.h"
 #include "ScopedTransaction.h"
 #include "Stateless/NiagaraDistributionPropertyCustomization.h"
+#include "Stateless/NiagaraDistributionIntPropertyCustomization.h"
 #include "Stateless/NiagaraStatelessCommon.h"
 #include "Stateless/NiagaraStatelessEmitter.h"
 #include "Stateless/NiagaraStatelessModule.h"
@@ -233,6 +234,7 @@ void UNiagaraStackStatelessModuleItem::RefreshChildrenInternal(const TArray<UNia
 			ModuleObject->RegisterInstancedCustomPropertyTypeLayout(FNiagaraDistributionRangeVector2::StaticStruct()->GetFName(), FOnGetPropertyTypeCustomizationInstance::CreateStatic(&FNiagaraDistributionPropertyCustomization::MakeVector2Instance));
 			ModuleObject->RegisterInstancedCustomPropertyTypeLayout(FNiagaraDistributionRangeVector3::StaticStruct()->GetFName(), FOnGetPropertyTypeCustomizationInstance::CreateStatic(&FNiagaraDistributionPropertyCustomization::MakeVector3Instance));
 			ModuleObject->RegisterInstancedCustomPropertyTypeLayout(FNiagaraDistributionRangeColor::StaticStruct()->GetFName(), FOnGetPropertyTypeCustomizationInstance::CreateStatic(&FNiagaraDistributionPropertyCustomization::MakeColorInstance));
+			ModuleObject->RegisterInstancedCustomPropertyTypeLayout(FNiagaraDistributionRangeInt::StaticStruct()->GetFName(), FOnGetPropertyTypeCustomizationInstance::CreateStatic(&FNiagaraDistributionIntPropertyCustomization::MakeIntInstance));
 			ModuleObjectWeak = ModuleObject;
 		}
 		NewChildren.Add(ModuleObject);
diff --git a/Engine/Plugins/FX/Niagara/Source/NiagaraShader/Internal/Stateless/NiagaraStatelessEmitterShaders.h b/Engine/Plugins/FX/Niagara/Source/NiagaraShader/Internal/Stateless/NiagaraStatelessEmitterShaders.h
--- a/Engine/Plugins/FX/Niagara/Source/NiagaraShader/Internal/Stateless/NiagaraStatelessEmitterShaders.h
+++ b/Engine/Plugins/FX/Niagara/Source/NiagaraShader/Internal/Stateless/NiagaraStatelessEmitterShaders.h
@@ -12,43 +12,39 @@
 
 namespace NiagaraStateless
 {
+	#define ADD_STATELESS_MODULE(NAME) SHADER_PARAMETER_STRUCT_INCLUDE(NAME, NAME##Params)
+
 	class FSimulationShaderDefaultCS : public FSimulationShader
 	{
 	public:
 		DECLARE_EXPORTED_GLOBAL_SHADER(FSimulationShaderDefaultCS, NIAGARASHADER_API);
 		SHADER_USE_PARAMETER_STRUCT(FSimulationShaderDefaultCS, FSimulationShader);
 
-		//BEGIN_SHADER_PARAMETER_STRUCT(FOutputParameters, NIAGARASHADER_API)
-		//	SHADER_PARAMETER(int,	Output_PositionComponent)
-		//	SHADER_PARAMETER(int,	Output_ColorComponent)
-		//	SHADER_PARAMETER(int,	Output_SpriteSizeComponent)
-		//	SHADER_PARAMETER(int,	Output_PreviousPositionComponent)
-		//	SHADER_PARAMETER(int,	Output_PreviousSpriteSizeComponent)
-		//END_SHADER_PARAMETER_STRUCT()
-
 		BEGIN_SHADER_PARAMETER_STRUCT(FParameters, NIAGARASHADER_API)
-			SHADER_PARAMETER_STRUCT_INCLUDE(FCommonShaderParameters,							CommonParameters)
-			SHADER_PARAMETER_STRUCT_INCLUDE(FInitializeParticleModule_ShaderParameters,			Module0Params)
-			SHADER_PARAMETER_STRUCT_INCLUDE(FInitialMeshOrientationModule_ShaderParameters,		Module1Params)
-			SHADER_PARAMETER_STRUCT_INCLUDE(FShapeLocationModule_ShaderParameters,				Module2Params)
-			SHADER_PARAMETER_STRUCT_INCLUDE(FCameraOffsetModule_ShaderParameters,				Module4Params)
-			SHADER_PARAMETER_STRUCT_INCLUDE(FScaleColorModule_ShaderParameters,					Module5Params)
-			SHADER_PARAMETER_STRUCT_INCLUDE(FScaleSpriteSizeModule_ShaderParameters,			Module6Params)
-			SHADER_PARAMETER_STRUCT_INCLUDE(FScaleSpriteSizeBySpeedModule_ShaderParameters,		Module7Params)
-			SHADER_PARAMETER_STRUCT_INCLUDE(FScaleMeshSizeModule_ShaderParameters,				Module8Params)
-			SHADER_PARAMETER_STRUCT_INCLUDE(FScaleMeshSizeBySpeedModule_ShaderParameters,		Module9Params)
-			SHADER_PARAMETER_STRUCT_INCLUDE(FMeshRotationRateModule_ShaderParameters,			Module10Params)
-			SHADER_PARAMETER_STRUCT_INCLUDE(FSolveVelocitiesAndForcesModule_ShaderParameters,	Module11Params)
-			SHADER_PARAMETER_STRUCT_INCLUDE(FSpriteFacingAndAlignmentModule_ShaderParameters,	Module12Params)
-			SHADER_PARAMETER_STRUCT_INCLUDE(FSpriteRotationRateModule_ShaderParameters,			Module13Params)
-			SHADER_PARAMETER_STRUCT_INCLUDE(FSubUVAnimationModule_ShaderParameters,				Module14Params)
-			SHADER_PARAMETER_STRUCT_INCLUDE(FDynamicMaterialParametersModule_ShaderParameters,	Module15Params)
+			SHADER_PARAMETER_STRUCT_INCLUDE(FCommonShaderParameters, CommonParameters)
+			ADD_STATELESS_MODULE(FInitializeParticleModule_ShaderParameters)
+			ADD_STATELESS_MODULE(FInitialMeshOrientationModule_ShaderParameters)
+			ADD_STATELESS_MODULE(FShapeLocationModule_ShaderParameters)
+			ADD_STATELESS_MODULE(FCameraOffsetModule_ShaderParameters)
+			ADD_STATELESS_MODULE(FScaleColorModule_ShaderParameters)
+			ADD_STATELESS_MODULE(FScaleSpriteSizeModule_ShaderParameters)
+			ADD_STATELESS_MODULE(FScaleSpriteSizeBySpeedModule_ShaderParameters)
+			ADD_STATELESS_MODULE(FScaleMeshSizeModule_ShaderParameters)
+			ADD_STATELESS_MODULE(FScaleMeshSizeBySpeedModule_ShaderParameters)
+			ADD_STATELESS_MODULE(FMeshIndexModule_ShaderParameters)
+			ADD_STATELESS_MODULE(FMeshRotationRateModule_ShaderParameters)
+			ADD_STATELESS_MODULE(FSolveVelocitiesAndForcesModule_ShaderParameters)
+			ADD_STATELESS_MODULE(FSpriteFacingAndAlignmentModule_ShaderParameters)
+			ADD_STATELESS_MODULE(FSpriteRotationRateModule_ShaderParameters)
+			ADD_STATELESS_MODULE(FSubUVAnimationModule_ShaderParameters)
+			ADD_STATELESS_MODULE(FDynamicMaterialParametersModule_ShaderParameters)
 
 			SHADER_PARAMETER(int, Permutation_UniqueIDComponent)
 			SHADER_PARAMETER(int, Permutation_PositionComponent)
 			SHADER_PARAMETER(int, Permutation_CameraOffsetComponent)
 			SHADER_PARAMETER(int, Permutation_ColorComponent)
 			SHADER_PARAMETER(int, Permutation_DynamicMaterialParameter0Component)
+			SHADER_PARAMETER(int, Permutation_MeshIndexComponent)
 			SHADER_PARAMETER(int, Permutation_MeshOrientationComponent)
 			SHADER_PARAMETER(int, Permutation_RibbonWidthComponent)
 			SHADER_PARAMETER(int, Permutation_ScaleComponent)
@@ -79,12 +75,12 @@ namespace NiagaraStateless
 
 		BEGIN_SHADER_PARAMETER_STRUCT(FParameters, NIAGARASHADER_API)
 			SHADER_PARAMETER_STRUCT_INCLUDE(FCommonShaderParameters, CommonParameters)
-			SHADER_PARAMETER_STRUCT_INCLUDE(FInitializeParticleModule_ShaderParameters,			Module0Params)
-			SHADER_PARAMETER_STRUCT_INCLUDE(FShapeLocationModule_ShaderParameters,				Module1Params)
-			SHADER_PARAMETER_STRUCT_INCLUDE(FScaleColorModule_ShaderParameters,					Module2Params)
-			SHADER_PARAMETER_STRUCT_INCLUDE(FScaleSpriteSizeModule_ShaderParameters,			Module3Params)
-			SHADER_PARAMETER_STRUCT_INCLUDE(FRotateAroundPointModule_ShaderParameters,			Module4Params)
-			SHADER_PARAMETER_STRUCT_INCLUDE(FSolveVelocitiesAndForcesModule_ShaderParameters,	Module5Params)
+			ADD_STATELESS_MODULE(FInitializeParticleModule_ShaderParameters)
+			ADD_STATELESS_MODULE(FShapeLocationModule_ShaderParameters)
+			ADD_STATELESS_MODULE(FScaleColorModule_ShaderParameters)
+			ADD_STATELESS_MODULE(FScaleSpriteSizeModule_ShaderParameters)
+			ADD_STATELESS_MODULE(FRotateAroundPointModule_ShaderParameters)
+			ADD_STATELESS_MODULE(FSolveVelocitiesAndForcesModule_ShaderParameters)
 
 			SHADER_PARAMETER(int, Permutation_UniqueIDComponent)
 			SHADER_PARAMETER(int, Permutation_PositionComponent)
@@ -102,4 +98,6 @@ namespace NiagaraStateless
 			SHADER_PARAMETER(int, Permutation_PreviousVelocityComponent)
 		END_SHADER_PARAMETER_STRUCT()
 	};
+
+	#undef ADD_STATELESS_MODULE
 }
diff --git a/Engine/Plugins/FX/Niagara/Source/NiagaraShader/Internal/Stateless/NiagaraStatelessModuleShaderParameters.h b/Engine/Plugins/FX/Niagara/Source/NiagaraShader/Internal/Stateless/NiagaraStatelessModuleShaderParameters.h
--- a/Engine/Plugins/FX/Niagara/Source/NiagaraShader/Internal/Stateless/NiagaraStatelessModuleShaderParameters.h
+++ b/Engine/Plugins/FX/Niagara/Source/NiagaraShader/Internal/Stateless/NiagaraStatelessModuleShaderParameters.h
@@ -70,6 +70,12 @@ namespace NiagaraStateless
 		SHADER_PARAMETER(FUintVector2,	ScaleMeshSizeBySpeed_ScaleDistribution)
 		SHADER_PARAMETER(float,			ScaleMeshSizeBySpeed_VelocityNorm)
 	END_SHADER_PARAMETER_STRUCT()
+		
+	BEGIN_SHADER_PARAMETER_STRUCT(FMeshIndexModule_ShaderParameters, NIAGARASHADER_API)
+		SHADER_PARAMETER(int32,		MeshIndex_Index)
+		SHADER_PARAMETER(int32,		MeshIndex_TableOffset)
+		SHADER_PARAMETER(int32,		MeshIndex_TableNumElements)
+	END_SHADER_PARAMETER_STRUCT()
 
 	BEGIN_SHADER_PARAMETER_STRUCT(FMeshRotationRateModule_ShaderParameters, NIAGARASHADER_API)
 		SHADER_PARAMETER(FVector3f,		MeshRotationRate_Scale)
