- Change stateless module macro to reduce generated string length

#tests Used stateless emitters

diff --git a/Engine/Plugins/FX/Niagara/Source/NiagaraShader/Internal/Stateless/NiagaraStatelessEmitterShaders.h b/Engine/Plugins/FX/Niagara/Source/NiagaraShader/Internal/Stateless/NiagaraStatelessEmitterShaders.h
--- a/Engine/Plugins/FX/Niagara/Source/NiagaraShader/Internal/Stateless/NiagaraStatelessEmitterShaders.h
+++ b/Engine/Plugins/FX/Niagara/Source/NiagaraShader/Internal/Stateless/NiagaraStatelessEmitterShaders.h
@@ -12,7 +12,9 @@
 
 namespace NiagaraStateless
 {
-	#define ADD_STATELESS_MODULE(NAME) SHADER_PARAMETER_STRUCT_INCLUDE(NAME, NAME##Params)
+	#define STATELESS_MODULE_COMBINE2(A, B) A##B
+	#define STATELESS_MODULE_COMBINE(A, B) STATELESS_MODULE_COMBINE2(A, B)
+	#define ADD_STATELESS_MODULE(NAME) SHADER_PARAMETER_STRUCT_INCLUDE(NAME, STATELESS_MODULE_COMBINE(ModuleParam,__LINE__))
 
 	class FSimulationShaderDefaultCS : public FSimulationShader
 	{
@@ -100,4 +102,6 @@ namespace NiagaraStateless
 	};
 
 	#undef ADD_STATELESS_MODULE
+	#undef STATELESS_MODULE_COMBINE
+	#undef STATELESS_MODULE_COMBINE2
 }
