- Stateless Int Range Binding support

#tests Used binding with mesh index module

diff --git a/Engine/Plugins/FX/Niagara/Source/Niagara/Internal/Stateless/NiagaraStatelessDistribution.h b/Engine/Plugins/FX/Niagara/Source/Niagara/Internal/Stateless/NiagaraStatelessDistribution.h
--- a/Engine/Plugins/FX/Niagara/Source/Niagara/Internal/Stateless/NiagaraStatelessDistribution.h
+++ b/Engine/Plugins/FX/Niagara/Source/Niagara/Internal/Stateless/NiagaraStatelessDistribution.h
@@ -54,7 +54,7 @@ struct FNiagaraDistributionBase
 	virtual bool DisplayAsColor() const { return false; }
 	virtual void UpdateValuesFromDistribution() { }
 
-	virtual FNiagaraTypeDefinition GetBindingTypeDef() { return FNiagaraTypeDefinition(); }
+	virtual FNiagaraTypeDefinition GetBindingTypeDef() const { return FNiagaraTypeDefinition(); }
 
 	static void PostEditChangeProperty(UObject* OwnerObject, FPropertyChangedEvent& PropertyChangedEvent);
 #endif
@@ -71,6 +71,9 @@ struct FNiagaraDistributionRangeInt
 	UPROPERTY(EditAnywhere, Category = "Parameters")
 	ENiagaraDistributionMode Mode = ENiagaraDistributionMode::UniformConstant;
 
+	UPROPERTY(EditAnywhere, Category = "Parameters")
+	FNiagaraVariableBase ParameterBinding;
+
 	UPROPERTY(EditAnywhere, Category = "Parameters")
 	int32 Min = 0;
 
@@ -85,6 +88,10 @@ struct FNiagaraDistributionRangeInt
 	bool IsUniform() const { return Mode == ENiagaraDistributionMode::UniformConstant || Mode == ENiagaraDistributionMode::UniformRange; }
 	bool IsCurve() const { return Mode == ENiagaraDistributionMode::UniformCurve || Mode == ENiagaraDistributionMode::NonUniformCurve; }
 	bool IsRange() const { return Mode == ENiagaraDistributionMode::UniformRange || Mode == ENiagaraDistributionMode::NonUniformRange; }
+
+#if WITH_EDITORONLY_DATA
+	FNiagaraTypeDefinition GetBindingTypeDef() const { return FNiagaraTypeDefinition::GetIntDef(); }
+#endif
 };
 
 USTRUCT()
@@ -107,7 +114,7 @@ struct FNiagaraDistributionRangeFloat : public FNiagaraDistributionBase
 #if WITH_EDITORONLY_DATA
 	virtual bool AllowCurves() const override { return false; }
 	NIAGARA_API virtual void UpdateValuesFromDistribution() override;
-	virtual FNiagaraTypeDefinition GetBindingTypeDef() { return FNiagaraTypeDefinition::GetFloatDef(); }
+	virtual FNiagaraTypeDefinition GetBindingTypeDef() const { return FNiagaraTypeDefinition::GetFloatDef(); }
 #endif
 };
 
@@ -131,7 +138,7 @@ struct FNiagaraDistributionRangeVector2 : public FNiagaraDistributionBase
 #if WITH_EDITORONLY_DATA
 	virtual bool AllowCurves() const override { return false; }
 	virtual void UpdateValuesFromDistribution() override;
-	virtual FNiagaraTypeDefinition GetBindingTypeDef() { return FNiagaraTypeDefinition::GetVec2Def(); }
+	virtual FNiagaraTypeDefinition GetBindingTypeDef() const { return FNiagaraTypeDefinition::GetVec2Def(); }
 #endif
 };
 
@@ -155,7 +162,7 @@ struct FNiagaraDistributionRangeVector3 : public FNiagaraDistributionBase
 #if WITH_EDITORONLY_DATA
 	virtual bool AllowCurves() const override { return false; }
 	virtual void UpdateValuesFromDistribution() override;
-	virtual FNiagaraTypeDefinition GetBindingTypeDef() { return FNiagaraTypeDefinition::GetVec3Def(); }
+	virtual FNiagaraTypeDefinition GetBindingTypeDef() const { return FNiagaraTypeDefinition::GetVec3Def(); }
 #endif
 };
 
@@ -180,7 +187,7 @@ struct FNiagaraDistributionRangeColor : public FNiagaraDistributionBase
 	virtual bool AllowCurves() const override { return false; }
 	virtual bool DisplayAsColor() const { return true; }
 	virtual void UpdateValuesFromDistribution() override;
-	virtual FNiagaraTypeDefinition GetBindingTypeDef() { return FNiagaraTypeDefinition::GetColorDef(); }
+	virtual FNiagaraTypeDefinition GetBindingTypeDef() const { return FNiagaraTypeDefinition::GetColorDef(); }
 #endif
 };
 
@@ -200,7 +207,7 @@ struct FNiagaraDistributionFloat : public FNiagaraDistributionBase
 
 #if WITH_EDITORONLY_DATA
 	virtual void UpdateValuesFromDistribution() override;
-	virtual FNiagaraTypeDefinition GetBindingTypeDef() { return FNiagaraTypeDefinition::GetFloatDef(); }
+	virtual FNiagaraTypeDefinition GetBindingTypeDef() const { return FNiagaraTypeDefinition::GetFloatDef(); }
 #endif
 };
 
@@ -222,7 +229,7 @@ struct FNiagaraDistributionVector2 : public FNiagaraDistributionBase
 
 #if WITH_EDITORONLY_DATA
 	virtual void UpdateValuesFromDistribution() override;
-	virtual FNiagaraTypeDefinition GetBindingTypeDef() { return FNiagaraTypeDefinition::GetVec2Def(); }
+	virtual FNiagaraTypeDefinition GetBindingTypeDef() const { return FNiagaraTypeDefinition::GetVec2Def(); }
 #endif
 };
 
@@ -244,7 +251,7 @@ struct FNiagaraDistributionVector3 : public FNiagaraDistributionBase
 
 #if WITH_EDITORONLY_DATA
 	virtual void UpdateValuesFromDistribution() override;
-	virtual FNiagaraTypeDefinition GetBindingTypeDef() { return FNiagaraTypeDefinition::GetVec3Def(); }
+	virtual FNiagaraTypeDefinition GetBindingTypeDef() const { return FNiagaraTypeDefinition::GetVec3Def(); }
 #endif
 };
 
@@ -265,6 +272,6 @@ struct FNiagaraDistributionColor : public FNiagaraDistributionBase
 #if WITH_EDITORONLY_DATA
 	virtual bool DisplayAsColor() const override { return true; }
 	virtual void UpdateValuesFromDistribution() override;
-	virtual FNiagaraTypeDefinition GetBindingTypeDef() { return FNiagaraTypeDefinition::GetColorDef(); }
+	virtual FNiagaraTypeDefinition GetBindingTypeDef() const { return FNiagaraTypeDefinition::GetColorDef(); }
 #endif
 };
diff --git a/Engine/Plugins/FX/Niagara/Source/Niagara/Internal/Stateless/NiagaraStatelessEmitterDataBuildContext.h b/Engine/Plugins/FX/Niagara/Source/Niagara/Internal/Stateless/NiagaraStatelessEmitterDataBuildContext.h
--- a/Engine/Plugins/FX/Niagara/Source/Niagara/Internal/Stateless/NiagaraStatelessEmitterDataBuildContext.h
+++ b/Engine/Plugins/FX/Niagara/Source/Niagara/Internal/Stateless/NiagaraStatelessEmitterDataBuildContext.h
@@ -134,6 +134,7 @@ public:
 	FNiagaraStatelessRangeVector2 ConvertDistributionToRange(const FNiagaraDistributionRangeVector2& Distribution, const FVector2f& DefaultValue, bool bEnabled = true) { return ConvertDistributionToRangeHelper<FNiagaraStatelessRangeVector2>(Distribution, DefaultValue, bEnabled); }
 	FNiagaraStatelessRangeVector3 ConvertDistributionToRange(const FNiagaraDistributionRangeVector3& Distribution, const FVector3f& DefaultValue, bool bEnabled = true) { return ConvertDistributionToRangeHelper<FNiagaraStatelessRangeVector3>(Distribution, DefaultValue, bEnabled); }
 	FNiagaraStatelessRangeColor   ConvertDistributionToRange(const FNiagaraDistributionRangeColor& Distribution, const FLinearColor& DefaultValue, bool bEnabled = true) { return ConvertDistributionToRangeHelper<FNiagaraStatelessRangeColor>(Distribution, DefaultValue, bEnabled); }
+	FNiagaraStatelessRangeInt     ConvertDistributionToRange(const FNiagaraDistributionRangeInt& Distribution, int32 DefaultValue, bool bEnabled = true) { return ConvertDistributionToRangeHelper<FNiagaraStatelessRangeInt>(Distribution, DefaultValue, bEnabled); }
 
 private:
 	FNiagaraParameterStore& RendererBindings;
diff --git a/Engine/Plugins/FX/Niagara/Source/Niagara/Private/Stateless/Modules/NiagaraStatelessModule_MeshIndex.cpp b/Engine/Plugins/FX/Niagara/Source/Niagara/Private/Stateless/Modules/NiagaraStatelessModule_MeshIndex.cpp
--- a/Engine/Plugins/FX/Niagara/Source/Niagara/Private/Stateless/Modules/NiagaraStatelessModule_MeshIndex.cpp
+++ b/Engine/Plugins/FX/Niagara/Source/Niagara/Private/Stateless/Modules/NiagaraStatelessModule_MeshIndex.cpp
@@ -51,32 +51,33 @@ void UNiagaraStatelessModule_MeshIndex::BuildEmitterData(FNiagaraStatelessEmitte
 	FModuleBuiltData* BuiltData = BuildContext.AllocateBuiltData<FModuleBuiltData>();
 	if (IsModuleEnabled())
 	{
-		//-TODO: We don't have an int binding at the moment
-		//const FNiagaraStatelessRangeInt MeshIndexRange = BuildContext.ConvertDistributionToRange(MeshIndex, 0, true);
-		//if (MeshIndexRange.ParameterOffset != INDEX_NONE)
-		//{
-		//}
-
-		const FNiagaraStatelessRangeInt MeshIndexRange = MeshIndex.CalculateRange();
-		BuiltData->Index = MeshIndexRange.Min;
- 		if (MeshIndexRange.GetScale() > 0 && MeshIndexRange.GetScale() < 256)
+		const FNiagaraStatelessRangeInt MeshIndexRange = BuildContext.ConvertDistributionToRange(MeshIndex, 0, true);
+		if (MeshIndexRange.ParameterOffset != INDEX_NONE)
 		{
-			FMeshIndexWeightedSampler Sampler(MeshIndexRange.GetScale() + 1, MeshIndexWeight);
-			Sampler.Initialize();
-
-			const int32 NumTableEntries = Sampler.GetNumEntries();
-			if (NumTableEntries > 1)
+			BuiltData->Index = MeshIndexRange.ParameterOffset | 0x80000000;
+		}
+		else
+		{
+			BuiltData->Index = MeshIndexRange.Min;
+			if (MeshIndexRange.GetScale() > 0 && MeshIndexRange.GetScale() < 256)
 			{
-				BuiltData->TableNumElements = NumTableEntries - 1;
+				FMeshIndexWeightedSampler Sampler(MeshIndexRange.GetScale() + 1, MeshIndexWeight);
+				Sampler.Initialize();
 
-				TArray<float, TInlineAllocator<16>> StaticData;
-				StaticData.AddUninitialized(NumTableEntries * 2);
-				for (int32 i=0; i < NumTableEntries; ++i)
+				const int32 NumTableEntries = Sampler.GetNumEntries();
+				if (NumTableEntries > 1)
 				{
-					StaticData[i * 2 + 0] = Sampler.GetProb()[i];
-					StaticData[i * 2 + 1] = float(MeshIndexRange.Min + Sampler.GetAlias()[i]);
+					BuiltData->TableNumElements = NumTableEntries - 1;
+
+					TArray<float, TInlineAllocator<16>> StaticData;
+					StaticData.AddUninitialized(NumTableEntries * 2);
+					for (int32 i = 0; i < NumTableEntries; ++i)
+					{
+						StaticData[i * 2 + 0] = Sampler.GetProb()[i];
+						StaticData[i * 2 + 1] = float(MeshIndexRange.Min + Sampler.GetAlias()[i]);
+					}
+					BuiltData->TableOffset = BuildContext.AddStaticData(StaticData);
 				}
-				BuiltData->TableOffset = BuildContext.AddStaticData(StaticData);
 			}
 		}
 	}
diff --git a/Engine/Plugins/FX/Niagara/Source/NiagaraEditor/Private/Stateless/NiagaraDistributionIntPropertyCustomization.cpp b/Engine/Plugins/FX/Niagara/Source/NiagaraEditor/Private/Stateless/NiagaraDistributionIntPropertyCustomization.cpp
--- a/Engine/Plugins/FX/Niagara/Source/NiagaraEditor/Private/Stateless/NiagaraDistributionIntPropertyCustomization.cpp
+++ b/Engine/Plugins/FX/Niagara/Source/NiagaraEditor/Private/Stateless/NiagaraDistributionIntPropertyCustomization.cpp
@@ -6,6 +6,10 @@
 #include "Editor.h"
 #include "EditorUndoClient.h"
 #include "NiagaraEditorStyle.h"
+#include "NiagaraNodeParameterMapBase.h"
+#include "NiagaraParameterMapHistory.h"
+#include "NiagaraScriptSource.h"
+#include "NiagaraSystem.h"
 #include "PropertyHandle.h"
 #include "ScopedTransaction.h"
 
@@ -16,6 +20,7 @@
 #include "Widgets/NiagaraDistributionEditorUtilities.h"
 #include "Widgets/SBoxPanel.h"
 #include "Widgets/SCompoundWidget.h"
+#include "Widgets/SNiagaraParameterName.h"
 #include "Widgets/INiagaraDistributionAdapter.h"
 #include "Widgets/Input/SComboButton.h"
 #include "Widgets/Input/SNumericEntryBox.h"
@@ -72,7 +77,13 @@ class SNiagaraDistributionIntPropertyWidget : public SCompoundWidget
 		UObject* OwnerObject = WeakOwnerObject.Get();
 		if (PropertyHandle.IsValid() && OwnerObject && Distribution)
 		{
-			return Distribution->Mode == ENiagaraDistributionMode::UniformConstant ? ENiagaraDistributionEditorMode::UniformConstant : ENiagaraDistributionEditorMode::UniformRange;
+			switch (Distribution->Mode)
+			{
+				case ENiagaraDistributionMode::Binding:			return ENiagaraDistributionEditorMode::Binding;
+				case ENiagaraDistributionMode::UniformConstant: return ENiagaraDistributionEditorMode::UniformConstant;
+				case ENiagaraDistributionMode::UniformRange:	return ENiagaraDistributionEditorMode::UniformRange;
+				default:										break;
+			}
 		}
 		return ENiagaraDistributionEditorMode::UniformConstant;
 	}
@@ -87,7 +98,13 @@ class SNiagaraDistributionIntPropertyWidget : public SCompoundWidget
 			OwnerObject->Modify();
 			PropertyHandle->NotifyPreChange();
 
-			Distribution->Mode = NewMode == ENiagaraDistributionEditorMode::UniformConstant ? ENiagaraDistributionMode::UniformConstant : ENiagaraDistributionMode::UniformRange;
+			switch (NewMode)
+			{
+				case ENiagaraDistributionEditorMode::Binding:			Distribution->Mode = ENiagaraDistributionMode::Binding; break;
+				case ENiagaraDistributionEditorMode::UniformConstant:	Distribution->Mode = ENiagaraDistributionMode::UniformConstant; break;
+				case ENiagaraDistributionEditorMode::UniformRange:		Distribution->Mode = ENiagaraDistributionMode::UniformRange; break;
+				default:												break;
+			}
 			ContentBox->SetContent(ConstructContentForMode());
 
 			PropertyHandle->NotifyPostChange(bContinuousChangeActive ? EPropertyChangeType::Interactive : EPropertyChangeType::ValueSet);
@@ -102,7 +119,21 @@ class SNiagaraDistributionIntPropertyWidget : public SCompoundWidget
 	TSharedRef<SWidget> OnGetDistributionModeMenuContent()
 	{
 		FMenuBuilder MenuBuilder(true, nullptr);
-		TArray<ENiagaraDistributionEditorMode> SupportedModes { ENiagaraDistributionEditorMode::UniformConstant, ENiagaraDistributionEditorMode::UniformRange };
+
+		TArray<ENiagaraDistributionEditorMode, TInlineAllocator<3>> SupportedModes;
+
+		TSharedPtr<IPropertyHandle> PropertyHandle = WeakPropertyHandle.Pin();
+		if (PropertyHandle.IsValid())
+		{
+			const FName DisableBindingDistributionName("DisableBindingDistribution");
+			if (!PropertyHandle->HasMetaData(DisableBindingDistributionName))
+			{
+				SupportedModes.Add(ENiagaraDistributionEditorMode::Binding);
+			}
+			SupportedModes.Add(ENiagaraDistributionEditorMode::UniformConstant);
+			SupportedModes.Add(ENiagaraDistributionEditorMode::UniformRange);
+		}
+
 		for (ENiagaraDistributionEditorMode SupportedMode : SupportedModes)
 		{
 			MenuBuilder.AddMenuEntry(
@@ -135,7 +166,15 @@ class SNiagaraDistributionIntPropertyWidget : public SCompoundWidget
 	TSharedRef<SWidget> ConstructContentForMode()
 	{
 		const ENiagaraDistributionEditorMode Mode = GetDistributionMode();
-		if (Mode == ENiagaraDistributionEditorMode::UniformConstant)
+		if (Mode == ENiagaraDistributionEditorMode::Binding)
+		{
+			return SNew(SHorizontalBox)
+				+ SHorizontalBox::Slot()
+				[
+					ConstructBindingWidget()
+				];
+		}
+		else if (Mode == ENiagaraDistributionEditorMode::UniformConstant)
 		{
 			return SNew(SHorizontalBox)
 				+ SHorizontalBox::Slot()
@@ -143,7 +182,7 @@ class SNiagaraDistributionIntPropertyWidget : public SCompoundWidget
 					ConstructIntWidget(0, FText())
 				];
 		}
-		else
+		else if (Mode == ENiagaraDistributionEditorMode::UniformRange)
 		{
 			return SNew(SHorizontalBox)
 				+ SHorizontalBox::Slot()
@@ -155,6 +194,7 @@ class SNiagaraDistributionIntPropertyWidget : public SCompoundWidget
 					ConstructIntWidget(1, LOCTEXT("MaxLabel", "Max"))
 				];
 		}
+		return SNullWidget::NullWidget;
 	}
 
 	TSharedRef<SWidget> ConstructIntWidget(int32 ValueIndex, FText LabelText)
@@ -236,6 +276,99 @@ class SNiagaraDistributionIntPropertyWidget : public SCompoundWidget
 		bContinuousChangeActive = false;
 	}
 
+	TSharedRef<SWidget> ConstructBindingWidget()
+	{
+		return SNew(SComboButton)
+			.ButtonStyle(FAppStyle::Get(), "HoverHintOnly")
+			.ContentPadding(FMargin(0))
+			.OnGetMenuContent(this, &SNiagaraDistributionIntPropertyWidget::OnGetBindingMenuContent)
+			.ButtonContent()
+			[
+				SNew(SNiagaraParameterName)
+				.ParameterName(this, &SNiagaraDistributionIntPropertyWidget::GetBindingValueName)
+				.IsReadOnly(true)
+			];
+	}
+
+	virtual TArray<FNiagaraVariableBase> GetAvailableBindings() const
+	{
+		TArray<FNiagaraVariableBase> AvailableBindings;
+		
+		UObject* OwnerObject = WeakOwnerObject.Get();
+		UNiagaraSystem* OwnerSystem = OwnerObject ? OwnerObject->GetTypedOuter<UNiagaraSystem>() : nullptr;
+		const FNiagaraTypeDefinition AllowedTypeDef = Distribution->GetBindingTypeDef();
+		if (OwnerSystem && AllowedTypeDef.IsValid())
+		{
+			if (UNiagaraScriptSource* Source = Cast<UNiagaraScriptSource>(OwnerSystem->GetSystemUpdateScript()->GetLatestSource()))
+			{
+				TArray<FNiagaraParameterMapHistory> Histories = UNiagaraNodeParameterMapBase::GetParameterMaps(Source->NodeGraph);
+				for (const FNiagaraParameterMapHistory& History : Histories)
+				{
+					for (const FNiagaraVariable& Variable : History.Variables)
+					{
+						if (Variable.GetType() == AllowedTypeDef && Variable.IsInNameSpace(FNiagaraConstants::SystemNamespaceString))
+						{
+							AvailableBindings.Add(Variable);
+						}
+					}
+				}
+			}
+
+			for (const FNiagaraVariableBase& Variable : OwnerSystem->GetExposedParameters().ReadParameterVariables())
+			{
+				if (Variable.GetType() == AllowedTypeDef)
+				{
+					AvailableBindings.Add(Variable);
+				}
+			}
+		}
+		return AvailableBindings;
+	}
+
+	TSharedRef<SWidget> OnGetBindingMenuContent()
+	{
+		FMenuBuilder MenuBuilder(true, nullptr);
+		for (const FNiagaraVariableBase& Variable : GetAvailableBindings())
+		{
+			TSharedRef<SWidget> Widget = SNew(SNiagaraParameterName)
+				.ParameterName(Variable.GetName())
+				.IsReadOnly(true);
+
+			MenuBuilder.AddMenuEntry(
+				FUIAction(
+					FExecuteAction::CreateSP(this, &SNiagaraDistributionIntPropertyWidget::SetBindingValue, Variable)
+				),
+				Widget
+			);
+		}
+
+		return MenuBuilder.MakeWidget();
+	}
+
+	FName GetBindingValueName() const
+	{
+		UObject* OwnerObject = WeakOwnerObject.Get();
+		return OwnerObject ? Distribution->ParameterBinding.GetName() : NAME_None;
+	}
+
+	void SetBindingValue(FNiagaraVariableBase Binding)
+	{
+		UObject* OwnerObject = WeakOwnerObject.Get();
+		TSharedPtr<IPropertyHandle> PropertyHandle = WeakPropertyHandle.Pin();
+		if (OwnerObject == nullptr || !PropertyHandle.IsValid())
+		{
+			return;
+		}
+
+		const FScopedTransaction Transaction(LOCTEXT("SetBinding", "Set binding value"));
+		OwnerObject->Modify();
+		PropertyHandle->NotifyPreChange();
+
+		Distribution->ParameterBinding = Binding;
+
+		PropertyHandle->NotifyPostChange(bContinuousChangeActive ? EPropertyChangeType::Interactive : EPropertyChangeType::ValueSet);
+	}
+
 private:
 	TWeakPtr<IPropertyHandle>		WeakPropertyHandle;
 	TWeakObjectPtr<UObject>			WeakOwnerObject;
