Small optimization (up to 1s) in editor startup by using the batch API in NiagaraEditor to receive asset updates from the registry.

diff --git a/Engine/Plugins/FX/Niagara/Source/NiagaraEditor/Private/NiagaraEditorModule.cpp b/Engine/Plugins/FX/Niagara/Source/NiagaraEditor/Private/NiagaraEditorModule.cpp
--- a/Engine/Plugins/FX/Niagara/Source/NiagaraEditor/Private/NiagaraEditorModule.cpp
+++ b/Engine/Plugins/FX/Niagara/Source/NiagaraEditor/Private/NiagaraEditorModule.cpp
@@ -1071,13 +1071,32 @@ void FNiagaraEditorModule::StartupModule()
 			ParameterCollectionAssetCache.RefreshCache(true /*bAllowLoading*/);
 			ParameterDefinitionsAssetCache.RefreshCache(true /*bAllowLoading*/);
 		});
-		AssetRegistryModule.Get().OnAssetAdded().AddLambda([this](const FAssetData& InAssetData)
+		AssetRegistryModule.Get().OnAssetsAdded().AddLambda([this](TConstArrayView<FAssetData> InAssets)
 		{
-			if (InAssetData.IsInstanceOf(UNiagaraParameterCollection::StaticClass()))
+			bool FoundParameterCollection = false;
+			bool FoundParameterDefinitions = false;
+			for (const FAssetData& Asset : InAssets)
+			{
+				if (!FoundParameterCollection && Asset.IsInstanceOf(UNiagaraParameterCollection::StaticClass()))
+				{
+					FoundParameterCollection = true;
+					continue;
+				}
+				else if (!FoundParameterDefinitions && Asset.IsInstanceOf(UNiagaraParameterDefinitions::StaticClass()))
+				{
+					FoundParameterDefinitions = true;
+					continue;
+				}
+				if (FoundParameterDefinitions && FoundParameterCollection)
+				{
+					break;
+				}
+			}
+			if (FoundParameterCollection)
 			{
 				ParameterCollectionAssetCache.RefreshCache(false);
 			}
-			else if (InAssetData.IsInstanceOf(UNiagaraParameterDefinitions::StaticClass()))
+			if (FoundParameterDefinitions)
 			{
 				ParameterDefinitionsAssetCache.RefreshCache(false);
 			}
