- Fix static analysis issue

diff --git a/Engine/Plugins/FX/Niagara/Source/Niagara/Private/Stateless/NiagaraStatelessModule.cpp b/Engine/Plugins/FX/Niagara/Source/Niagara/Private/Stateless/NiagaraStatelessModule.cpp
--- a/Engine/Plugins/FX/Niagara/Source/Niagara/Private/Stateless/NiagaraStatelessModule.cpp
+++ b/Engine/Plugins/FX/Niagara/Source/Niagara/Private/Stateless/NiagaraStatelessModule.cpp
@@ -17,7 +17,7 @@ void FNiagaraStatelessSetShaderParameterContext::ValidateIncludeStructType(uint3
 			continue;
 		}
 
-		if (Member.GetBaseType() == UBMT_INCLUDED_STRUCT && Member.GetStructMetadata()->GetLayout() == StructMetaData->GetLayout())
+		if (Member.GetBaseType() == UBMT_INCLUDED_STRUCT && Member.GetStructMetadata() && Member.GetStructMetadata()->GetLayout() == StructMetaData->GetLayout())
 		{
 			return;
 		}
