Fixed SA issue by adding nullptr check

diff --git a/Engine/Plugins/FX/Niagara/Source/NiagaraEditor/Private/EdGraphSchema_Niagara.cpp b/Engine/Plugins/FX/Niagara/Source/NiagaraEditor/Private/EdGraphSchema_Niagara.cpp
--- a/Engine/Plugins/FX/Niagara/Source/NiagaraEditor/Private/EdGraphSchema_Niagara.cpp
+++ b/Engine/Plugins/FX/Niagara/Source/NiagaraEditor/Private/EdGraphSchema_Niagara.cpp
@@ -2493,7 +2493,7 @@ void FNiagaraConnectionDrawingPolicy::DetermineWiringStyle(UEdGraphPin* OutputPi
 	check(OutputPin);
 	check(GraphObj);
 	
-	UEdGraphNode* OutputNode = OutputPin->GetOwningNode();
+	UEdGraphNode* OutputNode = (OutputPin != nullptr) ? OutputPin->GetOwningNode() : nullptr;
 	UEdGraphNode* InputNode = (InputPin != nullptr) ? InputPin->GetOwningNode() : nullptr;
 	
 	const bool bDeemphasizeUnhoveredPins = HoveredPins.Num() > 0;
