- Stateless -> Lightweight renaming for user facing parts

#tests Observing creating emitter, project properties, etc

diff --git a/Engine/Plugins/FX/Niagara/Source/Niagara/Internal/Stateless/NiagaraStatelessEmitter.h b/Engine/Plugins/FX/Niagara/Source/Niagara/Internal/Stateless/NiagaraStatelessEmitter.h
--- a/Engine/Plugins/FX/Niagara/Source/Niagara/Internal/Stateless/NiagaraStatelessEmitter.h
+++ b/Engine/Plugins/FX/Niagara/Source/Niagara/Internal/Stateless/NiagaraStatelessEmitter.h
@@ -106,16 +106,16 @@ protected:
 	UPROPERTY()
 	FString UniqueEmitterName;
 
-	UPROPERTY(EditAnywhere, Category = "Stateless", meta = (AllowedClasses = "/Script/Niagara.NiagaraStatelessEmitterTemplate", HideInStack))
+	UPROPERTY(EditAnywhere, Category = "General", meta = (AllowedClasses = "/Script/Niagara.NiagaraStatelessEmitterTemplate", HideInStack))
 	TObjectPtr<UClass> EmitterTemplateClass;
 
-	UPROPERTY(EditAnywhere, Category = "Stateless")
+	UPROPERTY(EditAnywhere, Category = "General")
 	bool bDeterministic = false;
 
-	UPROPERTY(EditAnywhere, Category = "Stateless")
+	UPROPERTY(EditAnywhere, Category = "General")
 	int32 RandomSeed = 0;
 
-	UPROPERTY(EditAnywhere, Category = "Stateless")
+	UPROPERTY(EditAnywhere, Category = "General")
 	FBox FixedBounds = FBox(FVector(-100), FVector(100));
 
 	UPROPERTY(EditAnywhere, Category = "Emitter State")
diff --git a/Engine/Plugins/FX/Niagara/Source/Niagara/Private/Stateless/NiagaraStatelessEmitterInstance.cpp b/Engine/Plugins/FX/Niagara/Source/Niagara/Private/Stateless/NiagaraStatelessEmitterInstance.cpp
--- a/Engine/Plugins/FX/Niagara/Source/Niagara/Private/Stateless/NiagaraStatelessEmitterInstance.cpp
+++ b/Engine/Plugins/FX/Niagara/Source/Niagara/Private/Stateless/NiagaraStatelessEmitterInstance.cpp
@@ -531,7 +531,7 @@ void FNiagaraStatelessEmitterInstance::SetExecutionStateInternal(ENiagaraExecuti
 	switch (RequestedExecutionState)
 	{
 		case ENiagaraExecutionState::Active:
-			UE_LOG(LogNiagara, Error, TEXT("Stateless Emitter: Was requested to go Active and we do not supoprt that."));
+			UE_LOG(LogNiagara, Error, TEXT("Lightweight Emitter: Was requested to go Active and we do not supoprt that."));
 			break;
 
 		case ENiagaraExecutionState::Inactive:
diff --git a/Engine/Plugins/FX/Niagara/Source/Niagara/Public/NiagaraSettings.h b/Engine/Plugins/FX/Niagara/Source/Niagara/Public/NiagaraSettings.h
--- a/Engine/Plugins/FX/Niagara/Source/Niagara/Public/NiagaraSettings.h
+++ b/Engine/Plugins/FX/Niagara/Source/Niagara/Public/NiagaraSettings.h
@@ -165,13 +165,13 @@ class UNiagaraSettings : public UDeveloperSettings
 	bool bExperimentalVMEnabled = false;
 
 	/**
-	Enables Stateless Emitters experimental feature.
+	Enables Lightweight Emitters experimental feature.
 	Statless emitters are lightweight fixed function emitters, they are not fully programmable like regular emitters and do not run scripts on the CPU.
 	Particle data is extrapolated per frame for the current particle age.  This means we never store particle data, we only generate it on demand.
-	Systems that contain only stateless emitters and no system script modules can take advantage of a much faster path to execute.
-	** There is no guarantee on backwards compatability for this feature currently.  Do not ship stateless content. **
+	Systems that contain only lightweight emitters and no system script modules can take advantage of a much faster path to execute.
+	** There is no guarantee on backwards compatability for this feature currently.  Do not ship lightweight content. **
 	*/
-	UPROPERTY(config, EditAnywhere, Category = Niagara, meta = (DisplayName = "Enable Stateless Emitters (Experimental)", ConfigRestartRequired = true))
+	UPROPERTY(config, EditAnywhere, Category = Niagara, meta = (DisplayName = "Enable Lightweight Emitters (Experimental)", ConfigRestartRequired = true))
 	bool bStatelessEmittersEnabled = false;
 
 	/** If set to true, quaternion attributes will be interpolated via slerp instead of lerp in interpolated spawn scripts. */
diff --git a/Engine/Plugins/FX/Niagara/Source/NiagaraEditor/Private/NiagaraEditorCommands.cpp b/Engine/Plugins/FX/Niagara/Source/NiagaraEditor/Private/NiagaraEditorCommands.cpp
--- a/Engine/Plugins/FX/Niagara/Source/NiagaraEditor/Private/NiagaraEditorCommands.cpp
+++ b/Engine/Plugins/FX/Niagara/Source/NiagaraEditor/Private/NiagaraEditorCommands.cpp
@@ -27,7 +27,7 @@ void FNiagaraEditorCommands::RegisterCommands()
 	UI_COMMAND(ToggleEmitterExecutionOrder, "Emitter Execution Order", "Display Emitter Execution Order", EUserInterfaceActionType::ToggleButton, FInputChord());
 	UI_COMMAND(ToggleGpuTickInformation, "Gpu Tick Information", "Display Gpu Tick Information", EUserInterfaceActionType::ToggleButton, FInputChord());
 	UI_COMMAND(ToggleMemoryInfo, "Memory", "Show approximate memory usage", EUserInterfaceActionType::ToggleButton, FInputChord());
-	UI_COMMAND(ToggleStatelessInfo, "Stateless", "Show information about stateless emitters and execution path", EUserInterfaceActionType::ToggleButton, FInputChord());
+	UI_COMMAND(ToggleStatelessInfo, "Lightweight", "Show information about lightweight emitters and execution path", EUserInterfaceActionType::ToggleButton, FInputChord());
 	UI_COMMAND(TogglePreviewBackground, "Background", "Toggles the preview pane's background.", EUserInterfaceActionType::ToggleButton, FInputChord());
 	UI_COMMAND(ToggleUnlockToChanges, "Lock/Unlock To Changes", "Toggles whether or not changes in the source asset get pulled into this asset automatically.", EUserInterfaceActionType::Button, FInputChord());
 	UI_COMMAND(ToggleStatPerformance, "Performance", "Show runtime performance for particle scripts. Not supported for optimized VM byte code.", EUserInterfaceActionType::ToggleButton, FInputChord());
diff --git a/Engine/Plugins/FX/Niagara/Source/NiagaraEditor/Private/NiagaraOverviewNode.cpp b/Engine/Plugins/FX/Niagara/Source/NiagaraEditor/Private/NiagaraOverviewNode.cpp
--- a/Engine/Plugins/FX/Niagara/Source/NiagaraEditor/Private/NiagaraOverviewNode.cpp
+++ b/Engine/Plugins/FX/Niagara/Source/NiagaraEditor/Private/NiagaraOverviewNode.cpp
@@ -337,8 +337,8 @@ void UNiagaraOverviewNode::GetNodeContextMenuActions(class UToolMenu* Menu, clas
 								);
 								SubMenuSection.AddMenuEntry(
 									"SetEmitterModeStateless",
-									LOCTEXT("SetEmitterStatelessModeMenuLabel", "Stateless Mode"),
-									LOCTEXT("SetEmitterStatelessModeMenuToolTip", "Set this emitter to use stateless mode."),
+									LOCTEXT("SetEmitterStatelessModeMenuLabel", "Lightweight Mode"),
+									LOCTEXT("SetEmitterStatelessModeMenuToolTip", "Set this emitter to use lightweight mode."),
 									FSlateIcon(),
 									FUIAction(
 										FExecuteAction::CreateStatic(&FNiagaraEditorUtilities::SetSelectedEmittersEmitterMode, OwningSystemViewModel.ToSharedRef(), ENiagaraEmitterMode::Stateless),
diff --git a/Engine/Plugins/FX/Niagara/Source/NiagaraEditor/Private/ViewModels/NiagaraSystemViewModel.cpp b/Engine/Plugins/FX/Niagara/Source/NiagaraEditor/Private/ViewModels/NiagaraSystemViewModel.cpp
--- a/Engine/Plugins/FX/Niagara/Source/NiagaraEditor/Private/ViewModels/NiagaraSystemViewModel.cpp
+++ b/Engine/Plugins/FX/Niagara/Source/NiagaraEditor/Private/ViewModels/NiagaraSystemViewModel.cpp
@@ -650,7 +650,7 @@ TSharedPtr<FNiagaraEmitterHandleViewModel> FNiagaraSystemViewModel::AddStateless
 		{
 			EmitterHandleNames.Add(EmitterHandle.GetName());
 		}
-		FName NewEmitterName = FName("Stateless");
+		FName NewEmitterName = FName("Lightweight");
 		NewEmitterName = FNiagaraUtilities::GetUniqueName(NewEmitterName, EmitterHandleNames);
 		NewEmitter = NewObject<UNiagaraStatelessEmitter>(&NiagaraSystem, *NewEmitterName.ToString(), RF_Transactional);
 		SetupStatelessEmitter(NewEmitter, NewEmitterName);
diff --git a/Engine/Plugins/FX/Niagara/Source/NiagaraEditor/Private/ViewModels/Stack/NiagaraStackStatelessEmitterGroup.cpp b/Engine/Plugins/FX/Niagara/Source/NiagaraEditor/Private/ViewModels/Stack/NiagaraStackStatelessEmitterGroup.cpp
--- a/Engine/Plugins/FX/Niagara/Source/NiagaraEditor/Private/ViewModels/Stack/NiagaraStackStatelessEmitterGroup.cpp
+++ b/Engine/Plugins/FX/Niagara/Source/NiagaraEditor/Private/ViewModels/Stack/NiagaraStackStatelessEmitterGroup.cpp
@@ -17,7 +17,7 @@ void UNiagaraStackStatelessEmitterGroup::Initialize(FRequiredEntryData InRequire
 	Super::Initialize(
 		InRequiredEntryData,
 		LOCTEXT("EmitterStatelessGroupDisplayName", "Properties"), 
-		LOCTEXT("EmitterStatelessGroupToolTip", "Properties for this stateless emitter"),
+		LOCTEXT("EmitterStatelessGroupToolTip", "Properties for this lightweight emitter"),
 		nullptr);
 	StatelessEmitterWeak = InStatelessEmitter;
 }
diff --git a/Engine/Plugins/FX/Niagara/Source/NiagaraEditor/Private/ViewModels/Stack/NiagaraStackStatelessEmitterSimulateGroup.cpp b/Engine/Plugins/FX/Niagara/Source/NiagaraEditor/Private/ViewModels/Stack/NiagaraStackStatelessEmitterSimulateGroup.cpp
--- a/Engine/Plugins/FX/Niagara/Source/NiagaraEditor/Private/ViewModels/Stack/NiagaraStackStatelessEmitterSimulateGroup.cpp
+++ b/Engine/Plugins/FX/Niagara/Source/NiagaraEditor/Private/ViewModels/Stack/NiagaraStackStatelessEmitterSimulateGroup.cpp
@@ -181,7 +181,7 @@ bool UNiagaraStackStatelessModuleItem::TestCanDeleteWithMessage(FText& OutCanDel
 
 FText UNiagaraStackStatelessModuleItem::GetDeleteTransactionText() const
 {
-	return LOCTEXT("DeleteStatelessModuleTransaction", "Delete module from stateless emitter.");
+	return LOCTEXT("DeleteStatelessModuleTransaction", "Delete module from lightweight emitter.");
 }
 
 void UNiagaraStackStatelessModuleItem::Delete()
diff --git a/Engine/Plugins/FX/Niagara/Source/NiagaraEditor/Private/Widgets/SNiagaraSystemViewport.cpp b/Engine/Plugins/FX/Niagara/Source/NiagaraEditor/Private/Widgets/SNiagaraSystemViewport.cpp
--- a/Engine/Plugins/FX/Niagara/Source/NiagaraEditor/Private/Widgets/SNiagaraSystemViewport.cpp
+++ b/Engine/Plugins/FX/Niagara/Source/NiagaraEditor/Private/Widgets/SNiagaraSystemViewport.cpp
@@ -369,7 +369,7 @@ void FNiagaraSystemViewportClient::DrawParticleCounts(UNiagaraComponent* Compone
 			const FString EmitterExecutionString = UEnum::GetValueAsString(ExecutionState);
 			const int32 EmitterExecutionStringValueIndex = EmitterExecutionString.Find(TEXT("::"));
 			const TCHAR* EmitterExecutionText = EmitterExecutionStringValueIndex == INDEX_NONE ? *EmitterExecutionString : *EmitterExecutionString + EmitterExecutionStringValueIndex + 2;
-			const TCHAR* EmitterMode = EmitterHandle.GetEmitterMode() == ENiagaraEmitterMode::Stateless ? TEXT("[Stateless]") : TEXT("");
+			const TCHAR* EmitterMode = EmitterHandle.GetEmitterMode() == ENiagaraEmitterMode::Stateless ? TEXT("[Lightweight]") : TEXT("");
 
 			TextItem.Text = FText::FromString(FString::Printf(TEXT("%i Current, %i Max (est.) - [%s] [%s] %s"), CurrentCount, MaxCount, *EmitterName.ToString(), EmitterExecutionText, EmitterMode));
 			TextItem.Position = FVector2D(CurrentX, CurrentY);
diff --git a/Engine/Plugins/FX/Niagara/Source/NiagaraEditorWidgets/Private/SNiagaraOverviewGraph.cpp b/Engine/Plugins/FX/Niagara/Source/NiagaraEditorWidgets/Private/SNiagaraOverviewGraph.cpp
--- a/Engine/Plugins/FX/Niagara/Source/NiagaraEditorWidgets/Private/SNiagaraOverviewGraph.cpp
+++ b/Engine/Plugins/FX/Niagara/Source/NiagaraEditorWidgets/Private/SNiagaraOverviewGraph.cpp
@@ -245,8 +245,8 @@ FActionMenuContent SNiagaraOverviewGraph::OnCreateGraphActionMenu(UEdGraph* InGr
 			if (GetDefault<UNiagaraSettings>()->bStatelessEmittersEnabled)
 			{
 				MenuBuilder.AddMenuEntry(
-					LOCTEXT("AddStatelessEmitter", "Add stateless emitter"),
-					LOCTEXT("AddStatelessEmitterToolTip", "Adds a stateless emitter without any modules or renderers."),
+					LOCTEXT("AddStatelessEmitter", "Add lightweight emitter"),
+					LOCTEXT("AddStatelessEmitterToolTip", "Adds a lightweight emitter without any modules or renderers."),
 					FSlateIcon(),
 					FExecuteAction::CreateSP(this, &SNiagaraOverviewGraph::OnCreateStatelessEmitter));
 			}
diff --git a/Engine/Plugins/FX/Niagara/Source/NiagaraEditorWidgets/Private/SNiagaraOverviewStackNode.cpp b/Engine/Plugins/FX/Niagara/Source/NiagaraEditorWidgets/Private/SNiagaraOverviewStackNode.cpp
--- a/Engine/Plugins/FX/Niagara/Source/NiagaraEditorWidgets/Private/SNiagaraOverviewStackNode.cpp
+++ b/Engine/Plugins/FX/Niagara/Source/NiagaraEditorWidgets/Private/SNiagaraOverviewStackNode.cpp
@@ -498,7 +498,7 @@ TSharedRef<SWidget> SNiagaraOverviewStackNode::CreateTitleWidget_Default(TShared
 				.AutoHeight()
 				[
 					SNew(STextBlock)
-						.Text(LOCTEXT("StatelessNotEnabled", "Stateless not enabled in project settings."))
+						.Text(LOCTEXT("StatelessNotEnabled", "Lightweight not enabled in project settings."))
 						.Font(FCoreStyle::GetDefaultFontStyle("Regular", 8.f))
 						.OverflowPolicy(ETextOverflowPolicy::Ellipsis)
 						.AutoWrapText(true)
