Fix minor shader warning

diff --git a/Engine/Plugins/FX/Niagara/Shaders/Private/Stateless/Modules/NiagaraStatelessModule_MeshIndex.ush b/Engine/Plugins/FX/Niagara/Shaders/Private/Stateless/Modules/NiagaraStatelessModule_MeshIndex.ush
--- a/Engine/Plugins/FX/Niagara/Shaders/Private/Stateless/Modules/NiagaraStatelessModule_MeshIndex.ush
+++ b/Engine/Plugins/FX/Niagara/Shaders/Private/Stateless/Modules/NiagaraStatelessModule_MeshIndex.ush
@@ -17,7 +17,7 @@ void MeshIndex_Initialize(inout FStatelessModule_MeshIndex ModuleData, inout FSt
 void MeshIndex_Simulate(in FStatelessModule_MeshIndex ModuleData, inout FStatelessParticle Particle)
 {
 	const bool bIsParameter = (MeshIndex_Index & 0x80000000) != 0;
-	int MeshIndex = MeshIndex_Index & ~0x80000000;
+	int MeshIndex = MeshIndex_Index & ~0x80000000u;
 	
 	if (bIsParameter)
 	{
