- Fix inconsistent usage of bytes vs uints for parameter bindings in stateless

diff --git a/Engine/Plugins/FX/Niagara/Shaders/Private/Stateless/Modules/NiagaraStatelessModule_MeshIndex.ush b/Engine/Plugins/FX/Niagara/Shaders/Private/Stateless/Modules/NiagaraStatelessModule_MeshIndex.ush
--- a/Engine/Plugins/FX/Niagara/Shaders/Private/Stateless/Modules/NiagaraStatelessModule_MeshIndex.ush
+++ b/Engine/Plugins/FX/Niagara/Shaders/Private/Stateless/Modules/NiagaraStatelessModule_MeshIndex.ush
@@ -21,7 +21,7 @@ void MeshIndex_Simulate(in FStatelessModule_MeshIndex ModuleData, inout FStatele
 	
 	if (bIsParameter)
 	{
-		MeshIndex = asint(GetParameterBufferFloat(MeshIndex, 0));
+		MeshIndex = GetParameterBufferInt(MeshIndex, 0);
 	}
 	else
 	{
diff --git a/Engine/Plugins/FX/Niagara/Shaders/Private/Stateless/NiagaraStatelessCommon.ush b/Engine/Plugins/FX/Niagara/Shaders/Private/Stateless/NiagaraStatelessCommon.ush
--- a/Engine/Plugins/FX/Niagara/Shaders/Private/Stateless/NiagaraStatelessCommon.ush
+++ b/Engine/Plugins/FX/Niagara/Shaders/Private/Stateless/NiagaraStatelessCommon.ush
@@ -127,11 +127,15 @@ float2 GetStaticFloat2(int BaseOffset, int Element)	{ return float2(Common_Stati
 float3 GetStaticFloat3(int BaseOffset, int Element)	{ return float3(Common_StaticFloatBuffer[BaseOffset + Element*3], Common_StaticFloatBuffer[BaseOffset + 1 + Element*3], Common_StaticFloatBuffer[BaseOffset + 2 + Element*3]); }
 float4 GetStaticFloat4(int BaseOffset, int Element)	{ return float4(Common_StaticFloatBuffer[BaseOffset + Element*4], Common_StaticFloatBuffer[BaseOffset + 1 + Element*4], Common_StaticFloatBuffer[BaseOffset + 2 + Element*4], Common_StaticFloatBuffer[BaseOffset + 3 + Element*4]); }
 
+uint   GetParameterBufferUint  (int BaseOffset, int Element) { return Common_ParameterBuffer[BaseOffset + Element]; }
+int    GetParameterBufferInt   (int BaseOffset, int Element) { return asint(Common_ParameterBuffer[BaseOffset + Element]); }
 float  GetParameterBufferFloat (int BaseOffset, int Element) { return asfloat(Common_ParameterBuffer[BaseOffset + Element]); }
 float2 GetParameterBufferFloat2(int BaseOffset, int Element) { return float2(asfloat(Common_ParameterBuffer[BaseOffset + Element*2]), asfloat(Common_ParameterBuffer[BaseOffset + 1 + Element*2])); }
 float3 GetParameterBufferFloat3(int BaseOffset, int Element) { return float3(asfloat(Common_ParameterBuffer[BaseOffset + Element*3]), asfloat(Common_ParameterBuffer[BaseOffset + 1 + Element*3]), asfloat(Common_ParameterBuffer[BaseOffset + 2 + Element*3])); }
 float4 GetParameterBufferFloat4(int BaseOffset, int Element) { return float4(asfloat(Common_ParameterBuffer[BaseOffset + Element*4]), asfloat(Common_ParameterBuffer[BaseOffset + 1 + Element*4]), asfloat(Common_ParameterBuffer[BaseOffset + 2 + Element*4]), asfloat(Common_ParameterBuffer[BaseOffset + 3 + Element*4])); }
 
+uint   GetParameterBufferUint  (int BaseOffset, int Element, uint DefaultValue)   { return BaseOffset >= 0 ? Common_ParameterBuffer[BaseOffset + Element] : DefaultValue; }
+int    GetParameterBufferInt   (int BaseOffset, int Element, int DefaultValue)    { return BaseOffset >= 0 ? asint(Common_ParameterBuffer[BaseOffset + Element]) : DefaultValue; }
 float  GetParameterBufferFloat (int BaseOffset, int Element, float  DefaultValue) { return BaseOffset >= 0 ? asfloat(Common_ParameterBuffer[BaseOffset + Element]) : DefaultValue; }
 float2 GetParameterBufferFloat2(int BaseOffset, int Element, float2 DefaultValue) { return BaseOffset >= 0 ? float2(asfloat(Common_ParameterBuffer[BaseOffset + Element*2]), asfloat(Common_ParameterBuffer[BaseOffset + 1 + Element*2])) : DefaultValue; }
 float3 GetParameterBufferFloat3(int BaseOffset, int Element, float3 DefaultValue) { return BaseOffset >= 0 ? float3(asfloat(Common_ParameterBuffer[BaseOffset + Element*3]), asfloat(Common_ParameterBuffer[BaseOffset + 1 + Element*3]), asfloat(Common_ParameterBuffer[BaseOffset + 2 + Element*3])) : DefaultValue; }
diff --git a/Engine/Plugins/FX/Niagara/Source/Niagara/Internal/Stateless/NiagaraStatelessEmitterDataBuildContext.h b/Engine/Plugins/FX/Niagara/Source/Niagara/Internal/Stateless/NiagaraStatelessEmitterDataBuildContext.h
--- a/Engine/Plugins/FX/Niagara/Source/Niagara/Internal/Stateless/NiagaraStatelessEmitterDataBuildContext.h
+++ b/Engine/Plugins/FX/Niagara/Source/Niagara/Internal/Stateless/NiagaraStatelessEmitterDataBuildContext.h
@@ -53,6 +53,9 @@ public:
 		return *reinterpret_cast<T*>(TransientObj->GetObject());
 	}
 
+	// Adds a binding to the renderer parameter store
+	// This allows you to read the parameter data inside the simulation process
+	// The returned value is INDEX_NONE is the variables is index otherwise the offset in DWORDs
 	int32 AddRendererBinding(const FNiagaraVariableBase& Variable);
 	int32 AddRendererBinding(const FNiagaraParameterBinding& Binding);
 	int32 AddRendererBinding(const FNiagaraParameterBindingWithValue& Binding);	
diff --git a/Engine/Plugins/FX/Niagara/Source/Niagara/Internal/Stateless/NiagaraStatelessModule.h b/Engine/Plugins/FX/Niagara/Source/Niagara/Internal/Stateless/NiagaraStatelessModule.h
--- a/Engine/Plugins/FX/Niagara/Source/Niagara/Internal/Stateless/NiagaraStatelessModule.h
+++ b/Engine/Plugins/FX/Niagara/Source/Niagara/Internal/Stateless/NiagaraStatelessModule.h
@@ -56,6 +56,7 @@ public:
 	{
 		if (Offset != INDEX_NONE)
 		{
+			Offset *= sizeof(uint32);
 			check(Offset >= 0 && Offset + sizeof(T) <= RendererParameterData.Num());
 			FMemory::Memcpy(&OutValue, RendererParameterData.GetData() + Offset, sizeof(T));
 		}
diff --git a/Engine/Plugins/FX/Niagara/Source/Niagara/Private/Stateless/NiagaraStatelessEmitterDataBuildContext.cpp b/Engine/Plugins/FX/Niagara/Source/Niagara/Private/Stateless/NiagaraStatelessEmitterDataBuildContext.cpp
--- a/Engine/Plugins/FX/Niagara/Source/Niagara/Private/Stateless/NiagaraStatelessEmitterDataBuildContext.cpp
+++ b/Engine/Plugins/FX/Niagara/Source/Niagara/Private/Stateless/NiagaraStatelessEmitterDataBuildContext.cpp
@@ -48,6 +48,7 @@ int32 FNiagaraStatelessEmitterDataBuildContext::AddRendererBinding(const FNiagar
 	{
 		FNiagaraVariable Var(Variable);
 		RendererBindings.AddParameter(Var, false, false, &DataOffset);
+		DataOffset /= sizeof(uint32);
 	}
 
 	return DataOffset;
@@ -60,6 +61,7 @@ int32 FNiagaraStatelessEmitterDataBuildContext::AddRendererBinding(const FNiagar
 	{
 		FNiagaraVariable Var(Binding.ResolvedParameter);
 		RendererBindings.AddParameter(Var, false, false, &DataOffset);
+		DataOffset /= sizeof(uint32);
 	}
 
 	return DataOffset;
@@ -80,6 +82,8 @@ int32 FNiagaraStatelessEmitterDataBuildContext::AddRendererBinding(const FNiagar
 			check(DefaultValue.Num() == Var.GetSizeInBytes());
 			RendererBindings.SetParameterData(DefaultValue.GetData(), DataOffset, DefaultValue.Num());
 		}
+
+		DataOffset /= sizeof(uint32);
 	}
 
 	return DataOffset;
