improve conversion to unit space for 2D grids to avoid float precision issues causing bugs with border cells in cubic interpolation

diff --git a/Engine/Plugins/FX/Niagara/Shaders/Private/NiagaraDataInterfaceGrid2DCollection.ush b/Engine/Plugins/FX/Niagara/Shaders/Private/NiagaraDataInterfaceGrid2DCollection.ush
--- a/Engine/Plugins/FX/Niagara/Shaders/Private/NiagaraDataInterfaceGrid2DCollection.ush
+++ b/Engine/Plugins/FX/Niagara/Shaders/Private/NiagaraDataInterfaceGrid2DCollection.ush
@@ -48,16 +48,16 @@ float Basic2DCubic_{ParameterName}(SamplerState Sampler, float3 UVW, int MipLeve
 	// identify the lower-left-hand corner of the cell
 	const int2 GridCell = floor(GridPos);
 								
-	const int2 MaxCell = {NumCellsName} - int2(2,2);
+	const int2 MaxCell = {NumCellsName} - int2(3,3);
 
-	if (any((GridCell < int2(1,1)) || (GridCell >= MaxCell)))
+	if (any((GridCell < int2(2,2)) || (GridCell >= MaxCell)))
 	{
 		// revert to bilinear hardware sampling at the boundary cells.
 		return {GridName}.SampleLevel(Sampler, UVW, MipLevel);
 	}
 	else 
 	{
-		const float2 t = GridPos - GridCell;
+		const float2 t = frac(GridPos);
 		float4 GridVals;
 		
 		float minv = 3.402823466e+38;
@@ -134,57 +134,4 @@ float Monotonic2DCubic_{ParameterName}(SamplerState Sampler, float3 UVW, int Mip
 float SampleBiCubic_{ParameterName}(SamplerState Sampler, float3 UVW, int MipLevel)
 {		   
 	return Basic2DCubic_{ParameterName}(Sampler, UVW, MipLevel);
-
-	const int AttributeIndex = round(UVW[2]);
-
-	const float2 GridPos = UVW.xy * {NumCellsName}.xy - .5;
-								
-	// identify the lower-left-hand corner of the cell
-	const int2 GridCell = floor(GridPos);
-								
-	const int2 MaxCell = {NumCellsName} - int2(2,2);
-
-	if (any((GridCell < int2(1,1)) || (GridCell >= MaxCell)))
-	{
-		// revert to bilinear hardware sampling at the boundary cells.
-		return {GridName}.SampleLevel(Sampler, UVW, MipLevel);
-	}
-	else 
-	{
-		// sample point offset from lower left
-		float2 delta = GridPos - GridCell;
-		
-		float4x2 Weights;
-
-		Weights[0] = -1./3 * delta + 1./2 * delta * delta - 1./6 * delta * delta * delta;
-		Weights[1] = 1. - delta * delta + 1./2 * (delta * delta * delta - delta);
-		Weights[2] = delta + 1./2 * (delta * delta - delta * delta * delta);
-		Weights[3] = 1./6 * (delta * delta * delta - delta);		
-
-		float minv = 3.402823466e+38;
-		float maxv = -3.402823466e+38;
-
-		float BiCubicValue = 0;		
-		for (int j = 0; j < 4; ++j)
-		{
-			float qj = 0;
-
-			for (int i = 0; i < 4; ++i)
-			{
-				float Val = {GridName}.Load(int4(GridCell.x + i - 1, GridCell.y + j - 1, AttributeIndex, MipLevel)); 
-				qj += Weights[i][0] * Val;
-
-				minv = min(Val, minv);
-				maxv = max(Val, maxv);
-			}
-
-			BiCubicValue += Weights[j][1] * qj;
-		}
-
-		if (BiCubicValue < minv || BiCubicValue > maxv)
-		{
-			return {GridName}.SampleLevel(Sampler, UVW, MipLevel);	
-		}	
-		return BiCubicValue;
-	}
 }
\ No newline at end of file
diff --git a/Engine/Plugins/FX/Niagara/Source/Niagara/Private/NiagaraDataInterfaceRW.cpp b/Engine/Plugins/FX/Niagara/Source/Niagara/Private/NiagaraDataInterfaceRW.cpp
--- a/Engine/Plugins/FX/Niagara/Source/Niagara/Private/NiagaraDataInterfaceRW.cpp
+++ b/Engine/Plugins/FX/Niagara/Source/Niagara/Private/NiagaraDataInterfaceRW.cpp
@@ -1039,7 +1039,7 @@ bool UNiagaraDataInterfaceGrid2D::GetFunctionHLSL(const FNiagaraDataInterfaceGPU
 			void {FunctionName}(out float2 Out_Unit)
 			{
 				#if NIAGARA_DISPATCH_TYPE == NIAGARA_DISPATCH_TYPE_TWO_D || NIAGARA_DISPATCH_TYPE == NIAGARA_DISPATCH_TYPE_CUSTOM
-					Out_Unit = (float2(GDispatchThreadId.x, GDispatchThreadId.y) + .5) * {UnitToUVName};			
+					Out_Unit = (float2(GDispatchThreadId.x, GDispatchThreadId.y) + .5) * rcp(1.f/{NumCellsName});
 				#else
 					const uint Linear = GLinearThreadId;
 					const uint IndexX = Linear % {NumCellsName}.x;
