Tweaked display of the scalability indicator for system nodes

diff --git a/Engine/Plugins/FX/Niagara/Source/NiagaraEditorWidgets/Private/SNiagaraOverviewStackNode.cpp b/Engine/Plugins/FX/Niagara/Source/NiagaraEditorWidgets/Private/SNiagaraOverviewStackNode.cpp
--- a/Engine/Plugins/FX/Niagara/Source/NiagaraEditorWidgets/Private/SNiagaraOverviewStackNode.cpp
+++ b/Engine/Plugins/FX/Niagara/Source/NiagaraEditorWidgets/Private/SNiagaraOverviewStackNode.cpp
@@ -902,7 +902,9 @@ EVisibility SNiagaraOverviewStackNode::GetScalabilityIndicatorVisibility() const
 
 	if(UNiagaraSystem* System = OverviewStackNode->GetOwningSystem())
 	{
-		bool bIsScalabilitySetup = System->GetOverrideScalabilitySettings();
+		bool bIsQualityLevelMaskSetup = System->GetScalabilityPlatformSet().QualityLevelMask != INDEX_NONE;
+		bool bIsScalabilityOverridden = System->GetOverrideScalabilitySettings();
+		bool bIsScalabilitySetup = bIsScalabilityOverridden || (System->GetScalabilityOverrides().Overrides.Num() != 0 || (bIsQualityLevelMaskSetup && System->GetScalabilityPlatformSet().QualityLevelMask != FNiagaraPlatformSet::GetFullQualityLevelMask(GetDefault<UNiagaraSettings>()->QualityLevels.Num())));
 		return bIsScalabilitySetup ? EVisibility::Visible : EVisibility::Collapsed; 
 	}
 	
