Removed the custom context menu for the Niagara Asset Browser for now, as the only option 'Find in Content Browser' would result in bugged asset creation due to the modal nature of the window.
Also added the config argument to disallow showing folders

diff --git a/Engine/Plugins/FX/Niagara/Source/NiagaraEditor/Private/Widgets/AssetBrowser/SNiagaraAssetBrowser.cpp b/Engine/Plugins/FX/Niagara/Source/NiagaraEditor/Private/Widgets/AssetBrowser/SNiagaraAssetBrowser.cpp
--- a/Engine/Plugins/FX/Niagara/Source/NiagaraEditor/Private/Widgets/AssetBrowser/SNiagaraAssetBrowser.cpp
+++ b/Engine/Plugins/FX/Niagara/Source/NiagaraEditor/Private/Widgets/AssetBrowser/SNiagaraAssetBrowser.cpp
@@ -233,6 +233,7 @@ void SNiagaraAssetBrowser::PopulateAssetBrowserContentSlot()
 	FAssetPickerConfig Config;
 	Config.Filter = GetCurrentBackendFilter();
 	Config.bCanShowClasses = false;
+	Config.bCanShowFolders = false;
 	Config.bAddFilterUI = true;
 	Config.DefaultFilterMenuExpansion = EAssetTypeCategories::FX;
 	Config.ExtraFrontendFilters = OnGetExtraFrontendFilters();
@@ -250,7 +251,8 @@ void SNiagaraAssetBrowser::PopulateAssetBrowserContentSlot()
 	Config.bForceShowPluginContent = true;
 	Config.SelectionMode = AssetSelectionMode;
 	Config.bAllowDragging = false;
-	Config.OnGetAssetContextMenu = FOnGetAssetContextMenu::CreateSP(this, &SNiagaraAssetBrowser::OnGetAssetContextMenu);
+	// Disabled for now as the only option 'Find in Content Browser' leads to bugs due to modal creation window
+	// Config.OnGetAssetContextMenu = FOnGetAssetContextMenu::CreateSP(this, &SNiagaraAssetBrowser::OnGetAssetContextMenu);
 
 	if(SaveSettingsName.IsSet())
 	{
