Added property diff tool support for some Niagara types

#tests Did a smoke test with basic asset types (asset tag, effect type, validation rule, sim cached test) by connecting to perforce and then diffing to the depot version after a small change

diff --git a/Engine/Plugins/FX/Niagara/Source/NiagaraEditor/Private/AssetDefinitions/AssetDefinition_NiagaraAssetTagDefinitions.cpp b/Engine/Plugins/FX/Niagara/Source/NiagaraEditor/Private/AssetDefinitions/AssetDefinition_NiagaraAssetTagDefinitions.cpp
--- a/Engine/Plugins/FX/Niagara/Source/NiagaraEditor/Private/AssetDefinitions/AssetDefinition_NiagaraAssetTagDefinitions.cpp
+++ b/Engine/Plugins/FX/Niagara/Source/NiagaraEditor/Private/AssetDefinitions/AssetDefinition_NiagaraAssetTagDefinitions.cpp
@@ -2,6 +2,8 @@
 
 #include "AssetDefinition_NiagaraAssetTagDefinitions.h"
 
+#include "SDetailsDiff.h"
+
 #define LOCTEXT_NAMESPACE "UAssetDefinition_NiagaraAssetTagDefinitions"
 
 UAssetDefinition_NiagaraAssetTagDefinitions::UAssetDefinition_NiagaraAssetTagDefinitions()
@@ -14,4 +16,20 @@ TConstArrayView<FAssetCategoryPath> UAssetDefinition_NiagaraAssetTagDefinitions:
 	return Categories;
 }
 
+EAssetCommandResult UAssetDefinition_NiagaraAssetTagDefinitions::PerformAssetDiff(const FAssetDiffArgs& DiffArgs) const
+{
+	if (DiffArgs.OldAsset == nullptr && DiffArgs.NewAsset == nullptr)
+	{
+		return EAssetCommandResult::Unhandled;
+	}
+	
+	const TSharedRef<SDetailsDiff> DetailsDiff = SDetailsDiff::CreateDiffWindow(DiffArgs.OldAsset, DiffArgs.NewAsset, DiffArgs.OldRevision, DiffArgs.NewRevision, UNiagaraAssetTagDefinitions::StaticClass());
+	// allow users to edit NewAsset if it's a local asset
+	if (!FPackageName::IsTempPackage(DiffArgs.NewAsset->GetPackage()->GetName()))
+	{
+		DetailsDiff->SetOutputObject(DiffArgs.NewAsset);
+	}
+	return EAssetCommandResult::Handled;
+}
+
 #undef LOCTEXT_NAMESPACE
diff --git a/Engine/Plugins/FX/Niagara/Source/NiagaraEditor/Private/AssetDefinitions/AssetDefinition_NiagaraAssetTagDefinitions.h b/Engine/Plugins/FX/Niagara/Source/NiagaraEditor/Private/AssetDefinitions/AssetDefinition_NiagaraAssetTagDefinitions.h
--- a/Engine/Plugins/FX/Niagara/Source/NiagaraEditor/Private/AssetDefinitions/AssetDefinition_NiagaraAssetTagDefinitions.h
+++ b/Engine/Plugins/FX/Niagara/Source/NiagaraEditor/Private/AssetDefinitions/AssetDefinition_NiagaraAssetTagDefinitions.h
@@ -18,4 +18,5 @@ class NIAGARAEDITOR_API UAssetDefinition_NiagaraAssetTagDefinitions : public UAs
 	virtual FLinearColor GetAssetColor() const override { return FLinearColor::White; }
 	virtual TSoftClassPtr<UObject> GetAssetClass() const override { return UNiagaraAssetTagDefinitions::StaticClass(); }
 	virtual TConstArrayView<FAssetCategoryPath> GetAssetCategories() const override;
+	virtual EAssetCommandResult PerformAssetDiff(const FAssetDiffArgs& DiffArgs) const override;
 };
diff --git a/Engine/Plugins/FX/Niagara/Source/NiagaraEditor/Private/AssetDefinitions/AssetDefinition_NiagaraEffectType.cpp b/Engine/Plugins/FX/Niagara/Source/NiagaraEditor/Private/AssetDefinitions/AssetDefinition_NiagaraEffectType.cpp
--- a/Engine/Plugins/FX/Niagara/Source/NiagaraEditor/Private/AssetDefinitions/AssetDefinition_NiagaraEffectType.cpp
+++ b/Engine/Plugins/FX/Niagara/Source/NiagaraEditor/Private/AssetDefinitions/AssetDefinition_NiagaraEffectType.cpp
@@ -2,6 +2,7 @@
 
 #include "AssetDefinition_NiagaraEffectType.h"
 #include "NiagaraEditorStyle.h"
+#include "SDetailsDiff.h"
 
 #define LOCTEXT_NAMESPACE "AssetTypeActions_NiagaraEffectType"
 
@@ -10,4 +11,20 @@ FLinearColor UAssetDefinition_NiagaraEffectType::GetAssetColor() const
 	return FNiagaraEditorStyle::Get().GetColor("NiagaraEditor.AssetColors.EffectType");
 }
 
+EAssetCommandResult UAssetDefinition_NiagaraEffectType::PerformAssetDiff(const FAssetDiffArgs& DiffArgs) const
+{
+	if (DiffArgs.OldAsset == nullptr && DiffArgs.NewAsset == nullptr)
+	{
+		return EAssetCommandResult::Unhandled;
+	}
+	
+	const TSharedRef<SDetailsDiff> DetailsDiff = SDetailsDiff::CreateDiffWindow(DiffArgs.OldAsset, DiffArgs.NewAsset, DiffArgs.OldRevision, DiffArgs.NewRevision, UNiagaraEffectType::StaticClass());
+	// allow users to edit NewAsset if it's a local asset
+	if (!FPackageName::IsTempPackage(DiffArgs.NewAsset->GetPackage()->GetName()))
+	{
+		DetailsDiff->SetOutputObject(DiffArgs.NewAsset);
+	}
+	return EAssetCommandResult::Handled;
+}
+
 #undef LOCTEXT_NAMESPACE
diff --git a/Engine/Plugins/FX/Niagara/Source/NiagaraEditor/Private/AssetDefinitions/AssetDefinition_NiagaraEffectType.h b/Engine/Plugins/FX/Niagara/Source/NiagaraEditor/Private/AssetDefinitions/AssetDefinition_NiagaraEffectType.h
--- a/Engine/Plugins/FX/Niagara/Source/NiagaraEditor/Private/AssetDefinitions/AssetDefinition_NiagaraEffectType.h
+++ b/Engine/Plugins/FX/Niagara/Source/NiagaraEditor/Private/AssetDefinitions/AssetDefinition_NiagaraEffectType.h
@@ -26,5 +26,6 @@ public:
 		
 		return Categories;
 	}
+	virtual EAssetCommandResult PerformAssetDiff(const FAssetDiffArgs& DiffArgs) const override;
 	// UAssetDefinition End
 };
diff --git a/Engine/Plugins/FX/Niagara/Source/NiagaraEditor/Private/AssetDefinitions/AssetDefinition_NiagaraValidationRuleSet.cpp b/Engine/Plugins/FX/Niagara/Source/NiagaraEditor/Private/AssetDefinitions/AssetDefinition_NiagaraValidationRuleSet.cpp
--- a/Engine/Plugins/FX/Niagara/Source/NiagaraEditor/Private/AssetDefinitions/AssetDefinition_NiagaraValidationRuleSet.cpp
+++ b/Engine/Plugins/FX/Niagara/Source/NiagaraEditor/Private/AssetDefinitions/AssetDefinition_NiagaraValidationRuleSet.cpp
@@ -2,8 +2,25 @@
 
 #include "AssetDefinition_NiagaraValidationRuleSet.h"
 #include "NiagaraEditorStyle.h"
+#include "SDetailsDiff.h"
 
 FLinearColor UAssetDefinition_NiagaraValidationRuleSet::GetAssetColor() const
 {
 	return FNiagaraEditorStyle::Get().GetColor("NiagaraEditor.AssetColors.ValidationRuleSet");
 }
+
+EAssetCommandResult UAssetDefinition_NiagaraValidationRuleSet::PerformAssetDiff(const FAssetDiffArgs& DiffArgs) const
+{
+	if (DiffArgs.OldAsset == nullptr && DiffArgs.NewAsset == nullptr)
+	{
+		return EAssetCommandResult::Unhandled;
+	}
+	
+	const TSharedRef<SDetailsDiff> DetailsDiff = SDetailsDiff::CreateDiffWindow(DiffArgs.OldAsset, DiffArgs.NewAsset, DiffArgs.OldRevision, DiffArgs.NewRevision, UNiagaraValidationRuleSet::StaticClass());
+	// allow users to edit NewAsset if it's a local asset
+	if (!FPackageName::IsTempPackage(DiffArgs.NewAsset->GetPackage()->GetName()))
+	{
+		DetailsDiff->SetOutputObject(DiffArgs.NewAsset);
+	}
+	return EAssetCommandResult::Handled;
+}
diff --git a/Engine/Plugins/FX/Niagara/Source/NiagaraEditor/Private/AssetDefinitions/AssetDefinition_NiagaraValidationRuleSet.h b/Engine/Plugins/FX/Niagara/Source/NiagaraEditor/Private/AssetDefinitions/AssetDefinition_NiagaraValidationRuleSet.h
--- a/Engine/Plugins/FX/Niagara/Source/NiagaraEditor/Private/AssetDefinitions/AssetDefinition_NiagaraValidationRuleSet.h
+++ b/Engine/Plugins/FX/Niagara/Source/NiagaraEditor/Private/AssetDefinitions/AssetDefinition_NiagaraValidationRuleSet.h
@@ -26,5 +26,6 @@ public:
 		
 		return Categories;
 	}
+	virtual EAssetCommandResult PerformAssetDiff(const FAssetDiffArgs& DiffArgs) const override;
 	// UAssetDefinition End
 };
