Change Niagara selection model to use weak pointers

#tests Just clicking through the Niagara ui a bit to make sure nothing broke. No extensive testing since the original bug couldn't be reliably reproduced anyway.

diff --git a/Engine/Plugins/FX/Niagara/Source/NiagaraEditor/Private/Toolkits/NiagaraScriptToolkit.cpp b/Engine/Plugins/FX/Niagara/Source/NiagaraEditor/Private/Toolkits/NiagaraScriptToolkit.cpp
--- a/Engine/Plugins/FX/Niagara/Source/NiagaraEditor/Private/Toolkits/NiagaraScriptToolkit.cpp
+++ b/Engine/Plugins/FX/Niagara/Source/NiagaraEditor/Private/Toolkits/NiagaraScriptToolkit.cpp
@@ -350,7 +350,7 @@ void FNiagaraScriptToolkit::InitViewWithVersionedData()
 	}
 	if (DetailsView)
 	{
-		DetailsView->SetObjects(DetailsScriptSelection->GetSelectedObjects().Array(), true);
+		DetailsView->SetObjects(DetailsScriptSelection->GetSelectedObjectsResolved().Array(), true);
 	}
 
 	// add listeners
@@ -458,7 +458,7 @@ TSharedRef<SDockTab> FNiagaraScriptToolkit::SpawnTabScriptDetails(const FSpawnTa
 	DetailsView = PropertyEditorModule.CreateDetailView(DetailsViewArgs);
 
 	DetailsView->OnFinishedChangingProperties().AddRaw(this, &FNiagaraScriptToolkit::OnEditedScriptPropertyFinishedChanging);
-	DetailsView->SetObjects(DetailsScriptSelection->GetSelectedObjects().Array());
+	DetailsView->SetObjects(DetailsScriptSelection->GetSelectedObjectsResolved().Array());
 
 	return SNew(SDockTab)
 		.Label(LOCTEXT("ScriptDetailsTabLabel", "Script Details"))
diff --git a/Engine/Plugins/FX/Niagara/Source/NiagaraEditor/Private/Toolkits/SystemToolkitModes/NiagaraSystemToolkitModeBase.cpp b/Engine/Plugins/FX/Niagara/Source/NiagaraEditor/Private/Toolkits/SystemToolkitModes/NiagaraSystemToolkitModeBase.cpp
--- a/Engine/Plugins/FX/Niagara/Source/NiagaraEditor/Private/Toolkits/SystemToolkitModes/NiagaraSystemToolkitModeBase.cpp
+++ b/Engine/Plugins/FX/Niagara/Source/NiagaraEditor/Private/Toolkits/SystemToolkitModes/NiagaraSystemToolkitModeBase.cpp
@@ -164,10 +164,10 @@ void FNiagaraSystemToolkitModeBase::OnParameterPanelViewModelExternalSelectionCh
 		TSharedPtr<FNiagaraObjectSelection> ParamSel = Toolkit->GetSystemViewModel()->GetParameterPanelViewModel()->GetVariableObjectSelection();
 		if (ParamSel.IsValid())
 		{
-			TArray<UObject*> SelectedNodes = ParamSel->GetSelectedObjects().Array();
+			TArray<UObject*> SelectedNodes = ParamSel->GetSelectedObjectsResolved().Array();
 			if (SelectedNodes.Num() != 0)
 			{
-				const TSet<UObject*>& CurrentlySelectedObjects = ObjectSelection->GetSelectedObjects();
+				TSet<UObject*> CurrentlySelectedObjects = ObjectSelection->GetSelectedObjectsResolved();
 				if(CurrentlySelectedObjects.Array() != SelectedNodes)
 				{
 					ObjectSelection->SetSelectedObject(SelectedNodes[0]);
@@ -200,7 +200,7 @@ void FNiagaraSystemToolkitModeBase::UpdateSelectionForActiveDocument()
 	TSharedPtr<FNiagaraScratchPadScriptViewModel> ScratchScriptVM = Toolkit->GetSystemViewModel()->GetDocumentViewModel()->GetActiveScratchPadViewModelIfSet();
 	if (ScratchScriptVM.IsValid())
 	{
-		TArray<UObject*> SelectedNodes = ScratchScriptVM->GetGraphViewModel()->GetNodeSelection()->GetSelectedObjects().Array();
+		TArray<UObject*> SelectedNodes = ScratchScriptVM->GetGraphViewModel()->GetNodeSelection()->GetSelectedObjectsResolved().Array();
 		if (SelectedNodes.Num() == 0)
 		{
 			if (ScratchScriptVM)
diff --git a/Engine/Plugins/FX/Niagara/Source/NiagaraEditor/Private/ViewModels/NiagaraOverviewGraphViewModel.cpp b/Engine/Plugins/FX/Niagara/Source/NiagaraEditor/Private/ViewModels/NiagaraOverviewGraphViewModel.cpp
--- a/Engine/Plugins/FX/Niagara/Source/NiagaraEditor/Private/ViewModels/NiagaraOverviewGraphViewModel.cpp
+++ b/Engine/Plugins/FX/Niagara/Source/NiagaraEditor/Private/ViewModels/NiagaraOverviewGraphViewModel.cpp
@@ -192,7 +192,7 @@ void FNiagaraOverviewGraphViewModel::DeleteSelectedNodes()
 		const FScopedTransaction Transaction(FGenericCommands::Get().Delete->GetDescription());
 		Graph->Modify();
 
-		TArray<UObject*> ObjectsToDelete = NodeSelection->GetSelectedObjects().Array();
+		TSet<UObject*> ObjectsToDelete = NodeSelection->GetSelectedObjectsResolved();
 		NodeSelection->ClearSelectedObjects();
 
 		TSet<FGuid> EmitterGuidsToDelete;
@@ -225,7 +225,7 @@ bool FNiagaraOverviewGraphViewModel::CanDeleteNodes() const
 	UEdGraph* Graph = GetGraph();
 	if (Graph != nullptr)
 	{
-		for (UObject* SelectedNode : NodeSelection->GetSelectedObjects())
+		for (UObject* SelectedNode : NodeSelection->GetSelectedObjectsResolved())
 		{
 			UEdGraphNode* SelectedGraphNode = Cast<UEdGraphNode>(SelectedNode);
 			if (SelectedGraphNode != nullptr && SelectedGraphNode->CanUserDeleteNode())
@@ -242,7 +242,7 @@ void FNiagaraOverviewGraphViewModel::CutSelectedNodes()
 	// Collect nodes which can not be delete or duplicated so they can be reselected.
 	TSet<UObject*> CanBeDuplicatedAndDeleted;
 	TSet<UObject*> CanNotBeDuplicatedAndDeleted;
-	for (UObject* SelectedNode : NodeSelection->GetSelectedObjects())
+	for (UObject* SelectedNode : NodeSelection->GetSelectedObjectsResolved())
 	{
 		UEdGraphNode* SelectedGraphNode = Cast<UEdGraphNode>(SelectedNode);
 		if (SelectedGraphNode != nullptr)
@@ -273,7 +273,7 @@ bool FNiagaraOverviewGraphViewModel::CanCutNodes() const
 void FNiagaraOverviewGraphViewModel::CopySelectedNodes()
 {
 	TSet<UObject*> NodesToCopy;
-	for (UObject* SelectedNode : NodeSelection->GetSelectedObjects())
+	for (UObject* SelectedNode : NodeSelection->GetSelectedObjectsResolved())
 	{
 		UEdGraphNode* SelectedGraphNode = Cast<UEdGraphNode>(SelectedNode);
 		if (SelectedGraphNode != nullptr)
@@ -305,7 +305,7 @@ bool FNiagaraOverviewGraphViewModel::CanCopyNodes() const
 
 	if (Graph != nullptr)
 	{
-		for (UObject* SelectedNode : NodeSelection->GetSelectedObjects())
+		for (UObject* SelectedNode : NodeSelection->GetSelectedObjectsResolved())
 		{
 			UEdGraphNode* SelectedGraphNode = Cast<UEdGraphNode>(SelectedNode);
 			if (SelectedGraphNode != nullptr && SelectedGraphNode->CanDuplicateNode())
@@ -400,20 +400,15 @@ bool FNiagaraOverviewGraphViewModel::CanDuplicateNodes() const
 
 void FNiagaraOverviewGraphViewModel::RenameNode()
 {
-	UEdGraph* Graph = GetGraph();
-	if (Graph != nullptr)
+	if (GetGraph())
 	{
-		const TSet<UObject*>& SelectedNodes = NodeSelection->GetSelectedObjects();
-		if (SelectedNodes.Num() > 0)
+		for (UObject* SelectedNode : NodeSelection->GetSelectedObjectsResolved())
 		{
-			for (UObject* SelectedNode : SelectedNodes)
+			UNiagaraOverviewNode* SelectedOverviewNode = Cast<UNiagaraOverviewNode>(SelectedNode);
+			if (SelectedOverviewNode != nullptr)
 			{
-				UNiagaraOverviewNode* SelectedOverviewNode = Cast<UNiagaraOverviewNode>(SelectedNode);
-				if (SelectedOverviewNode != nullptr)
-				{
-					SelectedOverviewNode->RequestRename();
-					return;
-				}
+				SelectedOverviewNode->RequestRename();
+				return;
 			}
 		}
 	}
@@ -421,10 +416,9 @@ void FNiagaraOverviewGraphViewModel::RenameNode()
 
 bool FNiagaraOverviewGraphViewModel::CanRenameNode() const
 {
-	UEdGraph* Graph = GetGraph();
-	if (Graph != nullptr)
+	if (GetGraph())
 	{
-		const TSet<UObject*>& SelectedNodes = NodeSelection->GetSelectedObjects();
+		TSet<UObject*> SelectedNodes = NodeSelection->GetSelectedObjectsResolved();
 		if (SelectedNodes.Num() == 1)
 		{
 			for (UObject* SelectedNode : SelectedNodes)
@@ -491,7 +485,7 @@ void FNiagaraOverviewGraphViewModel::GraphSelectionChanged()
 		TArray<FGuid> SelectedGuids;
 		TArray<TObjectPtr<UNiagaraStackObject>> NewTempEntries;
 
-		for (UObject* SelectedObject : NodeSelection->GetSelectedObjects())
+		for (UObject* SelectedObject : NodeSelection->GetSelectedObjectsResolved())
 		{
 			if (UNiagaraOverviewNode* SelectedOverviewNode = Cast<UNiagaraOverviewNode>(SelectedObject))
 			{
diff --git a/Engine/Plugins/FX/Niagara/Source/NiagaraEditor/Private/ViewModels/NiagaraParameterPanelViewModel.cpp b/Engine/Plugins/FX/Niagara/Source/NiagaraEditor/Private/ViewModels/NiagaraParameterPanelViewModel.cpp
--- a/Engine/Plugins/FX/Niagara/Source/NiagaraEditor/Private/ViewModels/NiagaraParameterPanelViewModel.cpp
+++ b/Engine/Plugins/FX/Niagara/Source/NiagaraEditor/Private/ViewModels/NiagaraParameterPanelViewModel.cpp
@@ -1076,8 +1076,7 @@ bool FNiagaraSystemToolkitParameterPanelViewModel::IsVariableSelected(FNiagaraVa
 {
 	if (VariableObjectSelection.IsValid())
 	{
-		const TSet<UObject*>& Objects = VariableObjectSelection->GetSelectedObjects();
-		for (UObject* Obj : Objects)
+		for (UObject* Obj : VariableObjectSelection->GetSelectedObjectsResolved())
 		{
 			UNiagaraScriptVariable* ScriptVar = Cast<UNiagaraScriptVariable>(Obj);
 			if (ScriptVar && ScriptVar->Variable.IsEquivalent(InVar, false))
diff --git a/Engine/Plugins/FX/Niagara/Source/NiagaraEditor/Private/ViewModels/NiagaraScratchPadViewModel.cpp b/Engine/Plugins/FX/Niagara/Source/NiagaraEditor/Private/ViewModels/NiagaraScratchPadViewModel.cpp
--- a/Engine/Plugins/FX/Niagara/Source/NiagaraEditor/Private/ViewModels/NiagaraScratchPadViewModel.cpp
+++ b/Engine/Plugins/FX/Niagara/Source/NiagaraEditor/Private/ViewModels/NiagaraScratchPadViewModel.cpp
@@ -761,7 +761,7 @@ void UNiagaraScratchPadViewModel::ScriptGraphNodeSelectionChanged(TWeakPtr<FNiag
 	TSharedPtr<FNiagaraScratchPadScriptViewModel> InScriptViewModel = InScriptViewModelWeak.Pin();
 	if (InScriptViewModel.IsValid())
 	{
-		TArray<UObject*> SelectedNodes = InScriptViewModel->GetGraphViewModel()->GetNodeSelection()->GetSelectedObjects().Array();
+		TSet<UObject*> SelectedNodes = InScriptViewModel->GetGraphViewModel()->GetNodeSelection()->GetSelectedObjectsResolved();
 		if (SelectedNodes.Num() > 0)
 		{
 			ObjectSelection->SetSelectedObjects(SelectedNodes);
@@ -895,7 +895,7 @@ void UNiagaraScratchPadViewModel::ScriptViewModelVariableSelectionChanged(TWeakP
 	TSharedPtr<FNiagaraScratchPadScriptViewModel> ScriptViewModel = ScriptViewModelWeak.Pin();
 	if (ScriptViewModel.IsValid())
 	{
-		ObjectSelection->SetSelectedObjects(ScriptViewModel->GetVariableSelection()->GetSelectedObjects().Array());
+		ObjectSelection->SetSelectedObjects(ScriptViewModel->GetVariableSelection()->GetSelectedObjectsResolved());
 	}
 }
 
diff --git a/Engine/Plugins/FX/Niagara/Source/NiagaraEditor/Private/ViewModels/NiagaraScriptGraphViewModel.cpp b/Engine/Plugins/FX/Niagara/Source/NiagaraEditor/Private/ViewModels/NiagaraScriptGraphViewModel.cpp
--- a/Engine/Plugins/FX/Niagara/Source/NiagaraEditor/Private/ViewModels/NiagaraScriptGraphViewModel.cpp
+++ b/Engine/Plugins/FX/Niagara/Source/NiagaraEditor/Private/ViewModels/NiagaraScriptGraphViewModel.cpp
@@ -190,7 +190,7 @@ void FNiagaraScriptGraphViewModel::DeleteSelectedNodes()
 		const FScopedTransaction Transaction(FGenericCommands::Get().Delete->GetDescription());
 		Graph->Modify();
 
-		TArray<UObject*> NodesToDelete = NodeSelection->GetSelectedObjects().Array();
+		TSet<UObject*> NodesToDelete = NodeSelection->GetSelectedObjectsResolved();
 		NodeSelection->ClearSelectedObjects();
 
 		for (UObject* NodeToDelete : NodesToDelete)
@@ -210,7 +210,7 @@ bool FNiagaraScriptGraphViewModel::CanDeleteNodes() const
 	UNiagaraGraph* Graph = GetGraph();
 	if (Graph != nullptr)
 	{
-		for (UObject* SelectedNode : NodeSelection->GetSelectedObjects())
+		for (UObject* SelectedNode : NodeSelection->GetSelectedObjectsResolved())
 		{
 			UEdGraphNode* SelectedGraphNode = Cast<UEdGraphNode>(SelectedNode);
 			if (SelectedGraphNode != nullptr && SelectedGraphNode->CanUserDeleteNode())
@@ -227,7 +227,7 @@ void FNiagaraScriptGraphViewModel::CutSelectedNodes()
 	// Collect nodes which can not be delete or duplicated so they can be reselected.
 	TSet<UObject*> CanBeDuplicatedAndDeleted;
 	TSet<UObject*> CanNotBeDuplicatedAndDeleted;
-	for (UObject* SelectedNode : NodeSelection->GetSelectedObjects())
+	for (UObject* SelectedNode : NodeSelection->GetSelectedObjectsResolved())
 	{
 		UEdGraphNode* SelectedGraphNode = Cast<UEdGraphNode>(SelectedNode);
 		if (SelectedGraphNode != nullptr)
@@ -261,7 +261,7 @@ void FNiagaraScriptGraphViewModel::CopySelectedNodes()
 
 	// we put all nodes we want to copy into the clipboard content. We also cache a set of nodes for more performant lookup later.
 	TSet<UObject*> CachedNodes;
-	for (UObject* SelectedNode : NodeSelection->GetSelectedObjects())
+	for (UObject* SelectedNode : NodeSelection->GetSelectedObjectsResolved())
 	{
 		UEdGraphNode* SelectedGraphNode = Cast<UEdGraphNode>(SelectedNode);
 		if (SelectedGraphNode != nullptr)
@@ -317,7 +317,7 @@ bool FNiagaraScriptGraphViewModel::CanCopyNodes() const
 	UNiagaraGraph* Graph = GetGraph();
 	if (Graph != nullptr)
 	{
-		for (UObject* SelectedNode : NodeSelection->GetSelectedObjects())
+		for (UObject* SelectedNode : NodeSelection->GetSelectedObjectsResolved())
 		{
 			UEdGraphNode* SelectedGraphNode = Cast<UEdGraphNode>(SelectedNode);
 			if (SelectedGraphNode != nullptr && SelectedGraphNode->CanDuplicateNode())
diff --git a/Engine/Plugins/FX/Niagara/Source/NiagaraEditor/Private/ViewModels/NiagaraScriptViewModel.cpp b/Engine/Plugins/FX/Niagara/Source/NiagaraEditor/Private/ViewModels/NiagaraScriptViewModel.cpp
--- a/Engine/Plugins/FX/Niagara/Source/NiagaraEditor/Private/ViewModels/NiagaraScriptViewModel.cpp
+++ b/Engine/Plugins/FX/Niagara/Source/NiagaraEditor/Private/ViewModels/NiagaraScriptViewModel.cpp
@@ -491,7 +491,7 @@ void FNiagaraScriptViewModel::GraphViewModelSelectedNodesChanged()
 		bUpdatingSelectionInternally = true;
 		{
 			TSet<FName> SelectedInputIds;
-			for (UObject* SelectedObject : GraphViewModel->GetNodeSelection()->GetSelectedObjects())
+			for (UObject* SelectedObject : GraphViewModel->GetNodeSelection()->GetSelectedObjectsResolved())
 			{
 				UNiagaraNodeInput* InputNode = Cast<UNiagaraNodeInput>(SelectedObject);
 				if (InputNode != nullptr)
diff --git a/Engine/Plugins/FX/Niagara/Source/NiagaraEditor/Private/ViewModels/Stack/NiagaraStackCommentCollection.cpp b/Engine/Plugins/FX/Niagara/Source/NiagaraEditor/Private/ViewModels/Stack/NiagaraStackCommentCollection.cpp
--- a/Engine/Plugins/FX/Niagara/Source/NiagaraEditor/Private/ViewModels/Stack/NiagaraStackCommentCollection.cpp
+++ b/Engine/Plugins/FX/Niagara/Source/NiagaraEditor/Private/ViewModels/Stack/NiagaraStackCommentCollection.cpp
@@ -35,9 +35,9 @@ UNiagaraStackObject* UNiagaraStackCommentCollection::FindStackObjectForCommentNo
 
 void UNiagaraStackCommentCollection::RefreshChildrenInternal(const TArray<UNiagaraStackEntry*>& CurrentChildren, TArray<UNiagaraStackEntry*>& NewChildren, TArray<FStackIssue>& NewIssues)
 {
-	const TSet<UObject*>& SelectedNodes = GetSystemViewModel()->GetOverviewGraphViewModel()->GetNodeSelection()->GetSelectedObjects();
+	TSet<UObject*> SelectedNodes = GetSystemViewModel()->GetOverviewGraphViewModel()->GetNodeSelection()->GetSelectedObjectsResolved();
 
-	for(UObject* SelectedNode : SelectedNodes)
+	for (UObject* SelectedNode : SelectedNodes)
 	{
 		if(UEdGraphNode_Comment* CommentNode = Cast<UEdGraphNode_Comment>(SelectedNode))
 		{
diff --git a/Engine/Plugins/FX/Niagara/Source/NiagaraEditor/Private/Widgets/SNiagaraScriptGraph.cpp b/Engine/Plugins/FX/Niagara/Source/NiagaraEditor/Private/Widgets/SNiagaraScriptGraph.cpp
--- a/Engine/Plugins/FX/Niagara/Source/NiagaraEditor/Private/Widgets/SNiagaraScriptGraph.cpp
+++ b/Engine/Plugins/FX/Niagara/Source/NiagaraEditor/Private/Widgets/SNiagaraScriptGraph.cpp
@@ -294,11 +294,11 @@ TSharedRef<SGraphEditor> SNiagaraScriptGraph::ConstructGraphEditor()
 
 void SNiagaraScriptGraph::ViewModelSelectedNodesChanged()
 {
-	if (FNiagaraEditorUtilities::SetsMatch(GraphEditor->GetSelectedNodes(), ViewModel->GetNodeSelection()->GetSelectedObjects()) == false)
+	if (FNiagaraEditorUtilities::SetsMatch(GraphEditor->GetSelectedNodes(), ViewModel->GetNodeSelection()->GetSelectedObjectsResolved()) == false)
 	{
 		bUpdatingGraphSelectionFromViewModel = true;
 		GraphEditor->ClearSelectionSet();
-		for (UObject* SelectedNode : ViewModel->GetNodeSelection()->GetSelectedObjects())
+		for (UObject* SelectedNode : ViewModel->GetNodeSelection()->GetSelectedObjectsResolved())
 		{
 			UEdGraphNode* GraphNode = Cast<UEdGraphNode>(SelectedNode);
 			if (GraphNode != nullptr)
diff --git a/Engine/Plugins/FX/Niagara/Source/NiagaraEditor/Private/Widgets/SNiagaraSelectedObjectsDetails.cpp b/Engine/Plugins/FX/Niagara/Source/NiagaraEditor/Private/Widgets/SNiagaraSelectedObjectsDetails.cpp
--- a/Engine/Plugins/FX/Niagara/Source/NiagaraEditor/Private/Widgets/SNiagaraSelectedObjectsDetails.cpp
+++ b/Engine/Plugins/FX/Niagara/Source/NiagaraEditor/Private/Widgets/SNiagaraSelectedObjectsDetails.cpp
@@ -28,7 +28,7 @@ void SNiagaraSelectedObjectsDetails::Construct(const FArguments& InArgs, TShared
 	DetailsViewArgs.NameAreaSettings = FDetailsViewArgs::HideNameArea;
 	DetailsViewArgs.bHideSelectionTip = true;
 	DetailsView = PropertyEditorModule.CreateDetailView(DetailsViewArgs);
-	DetailsView->SetObjects(SelectedObjectsArray[0]->GetSelectedObjects().Array());
+	DetailsView->SetObjects(SelectedObjectsArray[0]->GetSelectedObjectsResolved().Array());
 	DetailsView->SetIsPropertyReadOnlyDelegate(FIsPropertyReadOnly::CreateSP(this, &SNiagaraSelectedObjectsDetails::PropertyIsReadOnly));
 	DetailsView->SetEnabled(TAttribute<bool>::Create(TAttribute<bool>::FGetter::CreateSP(this, &SNiagaraSelectedObjectsDetails::DetailsPanelIsEnabled)));
 	//@todo(ng) re-enable once this is implemented.
@@ -56,7 +56,7 @@ void SNiagaraSelectedObjectsDetails::Construct(const FArguments& InArgs, TShared
 	DetailsViewArgs.NameAreaSettings = FDetailsViewArgs::HideNameArea;
 	DetailsViewArgs.bHideSelectionTip = true;
 	DetailsView = PropertyEditorModule.CreateDetailView(DetailsViewArgs);
-	DetailsView->SetObjects(SelectedObjectsArray[0]->GetSelectedObjects().Array());
+	DetailsView->SetObjects(SelectedObjectsArray[0]->GetSelectedObjectsResolved().Array());
 	DetailsView->SetIsPropertyReadOnlyDelegate(FIsPropertyReadOnly::CreateSP(this, &SNiagaraSelectedObjectsDetails::PropertyIsReadOnly));
 	DetailsView->SetEnabled(TAttribute<bool>::Create(TAttribute<bool>::FGetter::CreateSP(this, &SNiagaraSelectedObjectsDetails::DetailsPanelIsEnabled)));
 	//@todo(ng) re-enable once this is implemented.
@@ -81,7 +81,7 @@ void SNiagaraSelectedObjectsDetails::SelectedObjectsChanged()
 	{
 		UpdateSelectedObjectInfoFlags(SelectedObjectsArray[LastSetSelectedObjectsArrayIdx]);
 	}
-	DetailsView->SetObjects(SelectedObjectsArray[LastSetSelectedObjectsArrayIdx]->GetSelectedObjects().Array());
+	DetailsView->SetObjects(SelectedObjectsArray[LastSetSelectedObjectsArrayIdx]->GetSelectedObjectsResolved().Array());
 }
 
 void SNiagaraSelectedObjectsDetails::SelectedObjectsChangedFirst()
@@ -91,7 +91,7 @@ void SNiagaraSelectedObjectsDetails::SelectedObjectsChangedFirst()
 	{
 		UpdateSelectedObjectInfoFlags(SelectedObjectsArray[0]);
 	}
-	DetailsView->SetObjects(SelectedObjectsArray[0]->GetSelectedObjects().Array());
+	DetailsView->SetObjects(SelectedObjectsArray[0]->GetSelectedObjectsResolved().Array());
 	LastSetSelectedObjectsArrayIdx = 0;
 }
 
@@ -108,7 +108,7 @@ void SNiagaraSelectedObjectsDetails::SelectedObjectsChangedSecond()
 	{
 		UpdateSelectedObjectInfoFlags(SelectedObjectsArray[1]);
 	}
-	DetailsView->SetObjects(SelectedObjectsArray[1]->GetSelectedObjects().Array());
+	DetailsView->SetObjects(SelectedObjectsArray[1]->GetSelectedObjectsResolved().Array());
 	LastSetSelectedObjectsArrayIdx = 1;
 }
 
@@ -227,7 +227,7 @@ bool SNiagaraSelectedObjectsDetails::CustomRowIsReadOnly(const FName InRowName,
 
 void SNiagaraSelectedObjectsDetails::UpdateSelectedObjectInfoFlags(const TSharedPtr<FNiagaraObjectSelection>& SelectedObjects)
 {
-	for (const UObject* Obj : SelectedObjects->GetSelectedObjects())
+	for (const UObject* Obj : SelectedObjects->GetSelectedObjectsResolved())
 	{
 		if (const UNiagaraScriptVariable* ScriptVar = Cast<const UNiagaraScriptVariable>(Obj))
 		{
@@ -247,7 +247,7 @@ const UNiagaraScriptVariable* SNiagaraSelectedObjectsDetails::GetSelectedScriptV
 	//@todo(ng) avoid this check during level travel! 
 	for(const TSharedPtr<FNiagaraObjectSelection>& ObjectSelection : SelectedObjectsArray)
 	{
-		for (const UObject* Obj : ObjectSelection->GetSelectedObjects())
+		for (const UObject* Obj : ObjectSelection->GetSelectedObjectsResolved())
 		{
 			if (const UNiagaraScriptVariable* ScriptVar = Cast<const UNiagaraScriptVariable>(Obj))
 			{
diff --git a/Engine/Plugins/FX/Niagara/Source/NiagaraEditor/Private/Widgets/Scalability/SNiagaraScalabilityContext.cpp b/Engine/Plugins/FX/Niagara/Source/NiagaraEditor/Private/Widgets/Scalability/SNiagaraScalabilityContext.cpp
--- a/Engine/Plugins/FX/Niagara/Source/NiagaraEditor/Private/Widgets/Scalability/SNiagaraScalabilityContext.cpp
+++ b/Engine/Plugins/FX/Niagara/Source/NiagaraEditor/Private/Widgets/Scalability/SNiagaraScalabilityContext.cpp
@@ -112,7 +112,7 @@ void SNiagaraScalabilityContext::UpdateScalabilityContent()
 	TArray<UNiagaraStackEntry*> StackEntries;
 	ScalabilityViewModel->GetSystemViewModel().Pin()->GetSelectionViewModel()->GetSelectedEntries(StackEntries);
 
-	TSet<UObject*> SelectedNodes = ScalabilityViewModel->GetSystemViewModel().Pin()->GetOverviewGraphViewModel()->GetNodeSelection()->GetSelectedObjects();
+	TSet<UObject*> SelectedNodes = ScalabilityViewModel->GetSystemViewModel().Pin()->GetOverviewGraphViewModel()->GetNodeSelection()->GetSelectedObjectsResolved();
 
 	// if we select a node, the stack entries will be updated to a niagara stack root, which we want to ignore
 	if(StackEntries.Num() == 1 && !StackEntries[0]->IsA<UNiagaraStackRoot>())
diff --git a/Engine/Plugins/FX/Niagara/Source/NiagaraEditor/Private/Widgets/Scalability/SNiagaraSystemResolvedScalabilitySettings.cpp b/Engine/Plugins/FX/Niagara/Source/NiagaraEditor/Private/Widgets/Scalability/SNiagaraSystemResolvedScalabilitySettings.cpp
--- a/Engine/Plugins/FX/Niagara/Source/NiagaraEditor/Private/Widgets/Scalability/SNiagaraSystemResolvedScalabilitySettings.cpp
+++ b/Engine/Plugins/FX/Niagara/Source/NiagaraEditor/Private/Widgets/Scalability/SNiagaraSystemResolvedScalabilitySettings.cpp
@@ -278,7 +278,7 @@ void SNiagaraSystemResolvedScalabilitySettings::RebuildWidget()
 	
 	TSet<UNiagaraOverviewNode*> SelectedOverviewNodes;
 
-	for(UObject* SelectedNode : SystemViewModel->GetOverviewGraphViewModel()->GetNodeSelection()->GetSelectedObjects())
+	for(UObject* SelectedNode : SystemViewModel->GetOverviewGraphViewModel()->GetNodeSelection()->GetSelectedObjectsResolved())
 	{
 		if(UNiagaraOverviewNode* OverviewNode = Cast<UNiagaraOverviewNode>(SelectedNode))
 		{
diff --git a/Engine/Plugins/FX/Niagara/Source/NiagaraEditor/Public/NiagaraObjectSelection.h b/Engine/Plugins/FX/Niagara/Source/NiagaraEditor/Public/NiagaraObjectSelection.h
--- a/Engine/Plugins/FX/Niagara/Source/NiagaraEditor/Public/NiagaraObjectSelection.h
+++ b/Engine/Plugins/FX/Niagara/Source/NiagaraEditor/Public/NiagaraObjectSelection.h
@@ -11,7 +11,6 @@ class TNiagaraSelection
 public:
 	DECLARE_MULTICAST_DELEGATE(FOnSelectedObjectsChanged);
 
-public:
 	/** Gets the set of selected objects. */
 	const TSet<SelectedItemType>& GetSelectedObjects() const
 	{
@@ -93,7 +92,7 @@ public:
 		OnSelectedObjectsChangedDelegate.Broadcast();
 	}
 
-private:
+protected:
 	/** The set of selected objects. */
 	TSet<SelectedItemType> SelectedObjects;
 
@@ -104,6 +103,84 @@ private:
 	const void* AdditionalSelectionInfo = nullptr;
 };
 
-class FNiagaraObjectSelection : public TNiagaraSelection<UObject*>
+class FNiagaraObjectSelection : public TNiagaraSelection<FWeakObjectPtr>
 {
+public:
+
+	UObject* GetFirstSelectedObject() const
+	{
+		for (const FWeakObjectPtr& Obj : SelectedObjects)
+		{
+			if (Obj.IsValid())
+			{
+				return Obj.Get();
+			}
+		}
+		return nullptr;
+	}
+	
+	TSet<UObject*> GetSelectedObjectsResolved() const
+	{
+		TSet<UObject*> Objects;
+		for (const FWeakObjectPtr& Obj : SelectedObjects)
+		{
+			if (Obj.IsValid())
+			{
+				Objects.Add(Obj.Get());
+			}
+		}
+		return Objects;
+	}
+	
+	void SetSelectedObject(UObject* SelectedObject, const void* InSelectionInfo)
+	{
+		AdditionalSelectionInfo = InSelectionInfo;
+		if (SelectedObjects.Num() == 1 && SelectedObjects.Contains(SelectedObject))
+		{
+			// Refresh the delegate, in case a different object selection has been used in 
+			// a shared panel (but using a different selection, so this selection would not change)
+			OnSelectedObjectsChangedDelegate.Broadcast();
+			return;
+		}
+
+		SelectedObjects.Empty();
+		SelectedObjects.Add(SelectedObject);
+		OnSelectedObjectsChangedDelegate.Broadcast();
+	}
+
+	/** Replaces the currently selected set of objects with the supplied object. */
+	void SetSelectedObject(UObject* SelectedObject)
+	{
+		SetSelectedObject(SelectedObject, nullptr);
+	}
+
+	/** Replaces the currently selected set of objects with the supplied set. */
+	void SetSelectedObjects(const TSet<UObject*>& InSelectedObjects)
+	{
+		if (FNiagaraEditorUtilities::SetsMatch(GetSelectedObjectsResolved(), InSelectedObjects) == false)
+		{
+			SelectedObjects.Empty();
+			AdditionalSelectionInfo = nullptr;
+			for (UObject* Obj : InSelectedObjects)
+			{
+				SelectedObjects.Add(Obj);
+			}
+			OnSelectedObjectsChangedDelegate.Broadcast();
+		}
+	}
+
+	/** Replaces the currently selected set of objects with the supplied array. */
+	void SetSelectedObjects(const TArray<UObject*>& InSelectedObjects)
+	{
+		if (FNiagaraEditorUtilities::ArrayMatchesSet(InSelectedObjects, GetSelectedObjectsResolved()) == false)
+		{
+			SelectedObjects.Empty();
+			AdditionalSelectionInfo = nullptr;
+			for (UObject* Obj : InSelectedObjects)
+			{
+				SelectedObjects.Add(Obj);
+			}
+			OnSelectedObjectsChangedDelegate.Broadcast();
+		}
+	}
 };
\ No newline at end of file
diff --git a/Engine/Plugins/FX/Niagara/Source/NiagaraEditorWidgets/Private/SNiagaraOverviewGraph.cpp b/Engine/Plugins/FX/Niagara/Source/NiagaraEditorWidgets/Private/SNiagaraOverviewGraph.cpp
--- a/Engine/Plugins/FX/Niagara/Source/NiagaraEditorWidgets/Private/SNiagaraOverviewGraph.cpp
+++ b/Engine/Plugins/FX/Niagara/Source/NiagaraEditorWidgets/Private/SNiagaraOverviewGraph.cpp
@@ -171,11 +171,11 @@ void SNiagaraOverviewGraph::ViewModelSelectionChanged()
 {
 	if (bUpdatingViewModelSelectionFromGraph == false)
 	{
-		if (FNiagaraEditorUtilities::SetsMatch(GraphEditor->GetSelectedNodes(), ViewModel->GetNodeSelection()->GetSelectedObjects()) == false)
+		if (FNiagaraEditorUtilities::SetsMatch(GraphEditor->GetSelectedNodes(), ViewModel->GetNodeSelection()->GetSelectedObjectsResolved()) == false)
 		{
 			TGuardValue<bool> UpdateGuard(bUpdatingGraphSelectionFromViewModel, true);
 			GraphEditor->ClearSelectionSet();
-			for (UObject* SelectedNode : ViewModel->GetNodeSelection()->GetSelectedObjects())
+			for (UObject* SelectedNode : ViewModel->GetNodeSelection()->GetSelectedObjectsResolved())
 			{
 				UEdGraphNode* GraphNode = Cast<UEdGraphNode>(SelectedNode);
 				if (GraphNode != nullptr)
diff --git a/Engine/Plugins/FX/Niagara/Source/NiagaraEditorWidgets/Private/SNiagaraScratchPadScriptManager.cpp b/Engine/Plugins/FX/Niagara/Source/NiagaraEditorWidgets/Private/SNiagaraScratchPadScriptManager.cpp
--- a/Engine/Plugins/FX/Niagara/Source/NiagaraEditorWidgets/Private/SNiagaraScratchPadScriptManager.cpp
+++ b/Engine/Plugins/FX/Niagara/Source/NiagaraEditorWidgets/Private/SNiagaraScratchPadScriptManager.cpp
@@ -761,7 +761,7 @@ void SNiagaraScratchPadScriptManager::ObjectSelectionChanged()
 	}
 	else if (SelectionCount == 1)
 	{
-		UObject* SelectedObject = ViewModel->GetObjectSelection()->GetSelectedObjects().Array()[0];
+		UObject* SelectedObject = ViewModel->GetObjectSelection()->GetFirstSelectedObject();
 		if (SelectedObject->IsA<UEdGraphNode>())
 		{
 			UEdGraphNode* SelectedGraphNode = CastChecked<UEdGraphNode>(SelectedObject);
