- Remove ensure as this path still occurs during validation

diff --git a/Engine/Plugins/FX/Niagara/Source/NiagaraEditor/Private/NiagaraEditorModule.cpp b/Engine/Plugins/FX/Niagara/Source/NiagaraEditor/Private/NiagaraEditorModule.cpp
--- a/Engine/Plugins/FX/Niagara/Source/NiagaraEditor/Private/NiagaraEditorModule.cpp
+++ b/Engine/Plugins/FX/Niagara/Source/NiagaraEditor/Private/NiagaraEditorModule.cpp
@@ -324,7 +324,7 @@ class FNiagaraEditorOnlyDataUtilities : public INiagaraEditorOnlyDataUtilities
 		{
 			const TCHAR* SystemStateName = TEXT("/Niagara/Modules/System/SystemState.SystemState");
 			TArray<UNiagaraNodeFunctionCall*> Nodes;
-			if (ensure(ScriptSource->NodeGraph))
+			if (ScriptSource->NodeGraph)
 			{
 				ScriptSource->NodeGraph->GetNodesOfClass<UNiagaraNodeFunctionCall>(Nodes);
 			}
