Added Default Empty Emitter property to niagara editor settings.
All 'Add Empty Emitter' UI actions & menus try to use that before falling back to a fully empty emitter

diff --git a/Engine/Plugins/FX/Niagara/Config/BaseNiagara.ini b/Engine/Plugins/FX/Niagara/Config/BaseNiagara.ini
--- a/Engine/Plugins/FX/Niagara/Config/BaseNiagara.ini
+++ b/Engine/Plugins/FX/Niagara/Config/BaseNiagara.ini
@@ -65,6 +65,7 @@ DefaultScript=/Niagara/DefaultAssets/DefaultModule.DefaultModule
 DefaultDynamicInputScript=/Niagara/DefaultAssets/DefaultDynamicInputs.DefaultDynamicInputs
 DefaultFunctionScript=/Niagara/DefaultAssets/DefaultFunction.DefaultFunction
 DefaultModuleScript=/Niagara/DefaultAssets/DefaultModule.DefaultModule
+DefaultEmptyEmitter=/Niagara/DefaultAssets/Templates/Emitters/Empty.Empty
 RequiredSystemUpdateScript=/Niagara/Modules/System/SystemState.SystemState
 +GraphCreationShortcuts=(Name="Numeric::Add",Input=(Key=A,bShift=False,bCtrl=False,bAlt=False,bCmd=False))
 +GraphCreationShortcuts=(Name="Numeric::Divide",Input=(Key=D,bShift=False,bCtrl=False,bAlt=False,bCmd=False))
diff --git a/Engine/Plugins/FX/Niagara/Source/NiagaraEditor/Private/NiagaraEmitterFactoryNew.cpp b/Engine/Plugins/FX/Niagara/Source/NiagaraEditor/Private/NiagaraEmitterFactoryNew.cpp
--- a/Engine/Plugins/FX/Niagara/Source/NiagaraEditor/Private/NiagaraEmitterFactoryNew.cpp
+++ b/Engine/Plugins/FX/Niagara/Source/NiagaraEditor/Private/NiagaraEmitterFactoryNew.cpp
@@ -60,7 +60,17 @@ bool UNiagaraEmitterFactoryNew::ConfigureProperties()
 	}
 
 	TArray<FAssetData> SelectedAssetData = CreateAssetBrowserWindow->GetSelectedAssets();
-	
+
+	// Give the default empty emitter a chance to be selected if the user chose 'Empty Emitter'
+	if(SelectedAssetData.Num() == 0)
+	{
+		FSoftObjectPath DefaultEmptyEmitter = GetDefault<UNiagaraEditorSettings>()->DefaultEmptyEmitter;
+		if(DefaultEmptyEmitter.IsValid() && DefaultEmptyEmitter.IsAsset())
+		{
+			SelectedAssetData.Add(FAssetData(DefaultEmptyEmitter.TryLoad()));
+		}
+	}
+
 	if(SelectedAssetData.Num() == 1)
 	{
 		FAssetData SelectedAsset = SelectedAssetData[0];
@@ -73,6 +83,7 @@ bool UNiagaraEmitterFactoryNew::ConfigureProperties()
 	}
 	else
 	{
+		// If we haven't selected an emitter asset, nor had a valid default empty emitter, we add a truly empty emitter
 		EmitterToCopy = nullptr;
 		bUseInheritance = true;
 		bAddDefaultModulesAndRenderersToEmptyEmitter = false;
diff --git a/Engine/Plugins/FX/Niagara/Source/NiagaraEditor/Private/ViewModels/NiagaraSystemViewModel.cpp b/Engine/Plugins/FX/Niagara/Source/NiagaraEditor/Private/ViewModels/NiagaraSystemViewModel.cpp
--- a/Engine/Plugins/FX/Niagara/Source/NiagaraEditor/Private/ViewModels/NiagaraSystemViewModel.cpp
+++ b/Engine/Plugins/FX/Niagara/Source/NiagaraEditor/Private/ViewModels/NiagaraSystemViewModel.cpp
@@ -593,6 +593,7 @@ TSharedPtr<FNiagaraEmitterHandleViewModel> FNiagaraSystemViewModel::AddEmitter(U
 
 TSharedPtr<FNiagaraEmitterHandleViewModel> FNiagaraSystemViewModel::AddEmptyEmitter()
 {
+	// Otherwise we use an actually empty emitter
 	UNiagaraEmitter* EmptyEmitter = NewObject<UNiagaraEmitter>(GetTransientPackage());
 	bool bAddDefaultModulesAndRenderers = false;
 	UNiagaraEmitterFactoryNew::InitializeEmitter(EmptyEmitter, bAddDefaultModulesAndRenderers);
@@ -685,6 +686,20 @@ TSharedPtr<FNiagaraEmitterHandleViewModel> FNiagaraSystemViewModel::AddEmitter(c
 	return AddEmitter(*VersionedEmitter.Emitter, VersionedEmitter.Version);
 }
 
+TSharedPtr<FNiagaraEmitterHandleViewModel> FNiagaraSystemViewModel::AddDefaultEmptyEmitter()
+{
+	FSoftObjectPath DefaultEmptyEmitter = GetDefault<UNiagaraEditorSettings>()->DefaultEmptyEmitter;
+	if(DefaultEmptyEmitter.IsValid() && DefaultEmptyEmitter.IsAsset())
+	{
+		if(UNiagaraEmitter* Emitter = Cast<UNiagaraEmitter>(DefaultEmptyEmitter.TryLoad()))
+		{
+			return AddEmitterFromAssetData(FAssetData(Emitter));
+		}
+	}
+
+	return AddEmptyEmitter();
+}
+
 void FNiagaraSystemViewModel::DuplicateEmitters(TArray<FEmitterHandleToDuplicate> EmitterHandlesToDuplicate)
 {
 	if (EmitterHandlesToDuplicate.Num() <= 0)
diff --git a/Engine/Plugins/FX/Niagara/Source/NiagaraEditor/Private/Widgets/AssetBrowser/SNiagaraAssetBrowser.cpp b/Engine/Plugins/FX/Niagara/Source/NiagaraEditor/Private/Widgets/AssetBrowser/SNiagaraAssetBrowser.cpp
--- a/Engine/Plugins/FX/Niagara/Source/NiagaraEditor/Private/Widgets/AssetBrowser/SNiagaraAssetBrowser.cpp
+++ b/Engine/Plugins/FX/Niagara/Source/NiagaraEditor/Private/Widgets/AssetBrowser/SNiagaraAssetBrowser.cpp
@@ -998,16 +998,7 @@ void SNiagaraAddEmitterToSystemWindow::Construct(const FArguments& InArgs, TShar
 		   .HAlign(HAlign_Left)
 		   [
 			   SNew(SButton)
-			   .OnClicked_Lambda([this]()
-			   {
-			   		if(WeakSystemViewModel.IsValid())
-			   		{
-						WeakSystemViewModel.Pin()->AddEmptyEmitter();
-			   			RequestDestroyWindow();
-			   		}
-			   	
-				   return FReply::Handled();
-			   })
+			   .OnClicked(this, &SNiagaraAddEmitterToSystemWindow::AddEmptyEmitter)
 			   [
 				   SNew(SHorizontalBox)
 				   + SHorizontalBox::Slot()
@@ -1069,7 +1060,7 @@ void SNiagaraAddEmitterToSystemWindow::OnAssetsActivatedInternal(const TArray<FA
 	{
 		if(AssetData.Num() == 0)
 		{
-			WeakSystemViewModel.Pin()->AddEmptyEmitter();
+			WeakSystemViewModel.Pin()->AddDefaultEmptyEmitter();
 		}
 		else
 		{
diff --git a/Engine/Plugins/FX/Niagara/Source/NiagaraEditor/Public/NiagaraEditorSettings.h b/Engine/Plugins/FX/Niagara/Source/NiagaraEditor/Public/NiagaraEditorSettings.h
--- a/Engine/Plugins/FX/Niagara/Source/NiagaraEditor/Public/NiagaraEditorSettings.h
+++ b/Engine/Plugins/FX/Niagara/Source/NiagaraEditor/Public/NiagaraEditorSettings.h
@@ -362,6 +362,10 @@ public:
 	UPROPERTY(config, EditAnywhere, Category = Niagara)
 	FSoftObjectPath DefaultModuleScript;
 
+	/** Niagara emitter to duplicate as the base for empty emitters. Useful to set up default state. If empty, a truly empty emitter will be used. */
+	UPROPERTY(config, EditAnywhere, Category = Niagara, meta=(AllowedClasses="/Script/Niagara.NiagaraEmitter"))
+	FSoftObjectPath DefaultEmptyEmitter;
+
 	/** Niagara script which is required in the system update script to control system state. */
 	UPROPERTY(config, EditAnywhere, Category = Niagara)
 	FSoftObjectPath RequiredSystemUpdateScript;
diff --git a/Engine/Plugins/FX/Niagara/Source/NiagaraEditor/Public/ViewModels/NiagaraSystemViewModel.h b/Engine/Plugins/FX/Niagara/Source/NiagaraEditor/Public/ViewModels/NiagaraSystemViewModel.h
--- a/Engine/Plugins/FX/Niagara/Source/NiagaraEditor/Public/ViewModels/NiagaraSystemViewModel.h
+++ b/Engine/Plugins/FX/Niagara/Source/NiagaraEditor/Public/ViewModels/NiagaraSystemViewModel.h
@@ -242,7 +242,10 @@ public:
 	NIAGARAEDITOR_API TSharedPtr<FNiagaraEmitterHandleViewModel> AddEmitter(UNiagaraEmitter& Emitter, FGuid EmitterVersion);
 	NIAGARAEDITOR_API TSharedPtr<FNiagaraEmitterHandleViewModel> AddEmitter(const FVersionedNiagaraEmitter& Emitter);
 
-	/** Adds an empty emitter to the system. */
+	/** Adds the default empty emitter to the system, or a truly empty emitter as a fallback. */
+	NIAGARAEDITOR_API TSharedPtr<FNiagaraEmitterHandleViewModel> AddDefaultEmptyEmitter();
+
+	/** Adds a truly empty emitter to the system. */
 	NIAGARAEDITOR_API TSharedPtr<FNiagaraEmitterHandleViewModel> AddEmptyEmitter();
 
 	/** Adds a stateless emitter to the system. */
diff --git a/Engine/Plugins/FX/Niagara/Source/NiagaraEditorWidgets/Private/SNiagaraOverviewGraph.cpp b/Engine/Plugins/FX/Niagara/Source/NiagaraEditorWidgets/Private/SNiagaraOverviewGraph.cpp
--- a/Engine/Plugins/FX/Niagara/Source/NiagaraEditorWidgets/Private/SNiagaraOverviewGraph.cpp
+++ b/Engine/Plugins/FX/Niagara/Source/NiagaraEditorWidgets/Private/SNiagaraOverviewGraph.cpp
@@ -286,7 +286,7 @@ FActionMenuContent SNiagaraOverviewGraph::OnCreateGraphActionMenu(UEdGraph* InGr
 
 void SNiagaraOverviewGraph::OnCreateEmptyEmitter()
 {
-	ViewModel->GetSystemViewModel()->AddEmptyEmitter();
+	ViewModel->GetSystemViewModel()->AddDefaultEmptyEmitter();
 }
 
 void SNiagaraOverviewGraph::OnCreateStatelessEmitter()
