Added supports for profiles for the ContentBrowser.AssetViewOptions menu
The activated profile can be specified using the FAssetPickerConfig struct

diff --git a/Engine/Plugins/FX/Niagara/Source/NiagaraEditor/Private/NiagaraEditorMenuHelpers.cpp b/Engine/Plugins/FX/Niagara/Source/NiagaraEditor/Private/NiagaraEditorMenuHelpers.cpp
--- a/Engine/Plugins/FX/Niagara/Source/NiagaraEditor/Private/NiagaraEditorMenuHelpers.cpp
+++ b/Engine/Plugins/FX/Niagara/Source/NiagaraEditor/Private/NiagaraEditorMenuHelpers.cpp
@@ -126,6 +126,8 @@ void NiagaraEditorMenuHelpers::RegisterToolMenus()
 
 void NiagaraEditorMenuHelpers::RegisterMenuExtensions()
 {
+	RegisterAssetBrowserViewOptionsProfile();
+	
 	using namespace FNiagaraEditorUtilities::AssetBrowser;
 	
 	UToolMenus::Get()->RegisterStartupCallback(FSimpleMulticastDelegate::FDelegate::CreateLambda([]()
@@ -445,4 +447,18 @@ void NiagaraEditorMenuHelpers::RegisterMenuExtensions()
 	}));
 }
 
+void NiagaraEditorMenuHelpers::RegisterAssetBrowserViewOptionsProfile()
+{
+	UToolMenus::Get()->RegisterStartupCallback(FSimpleMulticastDelegate::FDelegate::CreateLambda([]()
+	{
+		const FName ProfileName("NiagaraAssetBrowser");
+
+		FToolMenuProfile* NiagaraAssetBrowserProfile = UToolMenus::Get()->AddRuntimeMenuProfile(FName("ContentBrowser.AssetViewOptions"), ProfileName);
+		if(FCustomizedToolMenuSection* Section = NiagaraAssetBrowserProfile->AddSection("View"))
+		{
+			Section->Visibility = ECustomizedToolMenuVisibility::Hidden;
+		}
+	}));
+}
+
 #undef LOCTEXT_NAMESPACE
diff --git a/Engine/Plugins/FX/Niagara/Source/NiagaraEditor/Private/NiagaraEditorMenuHelpers.h b/Engine/Plugins/FX/Niagara/Source/NiagaraEditor/Private/NiagaraEditorMenuHelpers.h
--- a/Engine/Plugins/FX/Niagara/Source/NiagaraEditor/Private/NiagaraEditorMenuHelpers.h
+++ b/Engine/Plugins/FX/Niagara/Source/NiagaraEditor/Private/NiagaraEditorMenuHelpers.h
@@ -6,4 +6,6 @@ namespace NiagaraEditorMenuHelpers
 {
 	void RegisterToolMenus();
 	void RegisterMenuExtensions();
+
+	void RegisterAssetBrowserViewOptionsProfile();
 }
\ No newline at end of file
diff --git a/Engine/Plugins/FX/Niagara/Source/NiagaraEditor/Private/Widgets/AssetBrowser/SNiagaraAssetBrowserContent.cpp b/Engine/Plugins/FX/Niagara/Source/NiagaraEditor/Private/Widgets/AssetBrowser/SNiagaraAssetBrowserContent.cpp
--- a/Engine/Plugins/FX/Niagara/Source/NiagaraEditor/Private/Widgets/AssetBrowser/SNiagaraAssetBrowserContent.cpp
+++ b/Engine/Plugins/FX/Niagara/Source/NiagaraEditor/Private/Widgets/AssetBrowser/SNiagaraAssetBrowserContent.cpp
@@ -19,6 +19,7 @@ void SNiagaraAssetBrowserContent::Construct(const FArguments& InArgs)
 	Config.SyncToAssetsDelegates.Add(&SyncToAssetsDelegate);
 	Config.SetFilterDelegates.Add(&SetNewFilterDelegate);
 	Config.bCanShowRealTimeThumbnails = true;
+	Config.AssetViewOptionsProfile = FName("NiagaraAssetBrowser");
 	
 	FContentBrowserModule& ContentBrowserModule = FModuleManager::Get().LoadModuleChecked<FContentBrowserModule>("ContentBrowser");
 	ChildSlot
diff --git a/Engine/Source/Editor/ContentBrowser/Private/SAssetPicker.cpp b/Engine/Source/Editor/ContentBrowser/Private/SAssetPicker.cpp
--- a/Engine/Source/Editor/ContentBrowser/Private/SAssetPicker.cpp
+++ b/Engine/Source/Editor/ContentBrowser/Private/SAssetPicker.cpp
@@ -362,6 +362,7 @@ void SAssetPicker::Construct( const FArguments& InArgs )
 		.CustomColumns(InArgs._AssetPickerConfig.CustomColumns)
 		.OnSearchOptionsChanged(this, &SAssetPicker::HandleSearchSettingsChanged)
 		.InitialThumbnailSize(InArgs._AssetPickerConfig.InitialThumbnailSize)
+		.AssetViewOptionsProfile(InArgs._AssetPickerConfig.AssetViewOptionsProfile)
 	];
 	// clang-format on
 
diff --git a/Engine/Source/Editor/ContentBrowser/Private/SAssetView.cpp b/Engine/Source/Editor/ContentBrowser/Private/SAssetView.cpp
--- a/Engine/Source/Editor/ContentBrowser/Private/SAssetView.cpp
+++ b/Engine/Source/Editor/ContentBrowser/Private/SAssetView.cpp
@@ -323,7 +323,8 @@ void SAssetView::Construct( const FArguments& InArgs )
 	OnSearchOptionsChanged = InArgs._OnSearchOptionsChanged;
 	bShowPathViewFilters = InArgs._bShowPathViewFilters;
 	OnExtendAssetViewOptionsMenuContext = InArgs._OnExtendAssetViewOptionsMenuContext;
-
+	AssetViewOptionsProfile = InArgs._AssetViewOptionsProfile;
+	
 	if ( InArgs._InitialViewType >= 0 && InArgs._InitialViewType < EAssetViewType::MAX )
 	{
 		CurrentViewType = InArgs._InitialViewType;
@@ -2606,6 +2607,13 @@ TSharedRef<SWidget> SAssetView::GetViewButtonContent()
 	TSharedPtr<FExtender> MenuExtender = FExtender::Combine(Extenders);
 	FToolMenuContext MenuContext(nullptr, MenuExtender, Context);
 
+	if(AssetViewOptionsProfile.IsSet())
+	{
+		UToolMenuProfileContext* ProfileContext = NewObject<UToolMenuProfileContext>();
+		ProfileContext->ActiveProfiles.Add(AssetViewOptionsProfile.GetValue());
+		MenuContext.AddObject(ProfileContext);
+	}
+
 	if (OnExtendAssetViewOptionsMenuContext.IsBound())
 	{
 		OnExtendAssetViewOptionsMenuContext.Execute(MenuContext);
diff --git a/Engine/Source/Editor/ContentBrowser/Public/IContentBrowserSingleton.h b/Engine/Source/Editor/ContentBrowser/Public/IContentBrowserSingleton.h
--- a/Engine/Source/Editor/ContentBrowser/Public/IContentBrowserSingleton.h
+++ b/Engine/Source/Editor/ContentBrowser/Public/IContentBrowserSingleton.h
@@ -312,6 +312,9 @@ struct FAssetPickerConfig
 	/** If we display filters & set to true, we will add sections instead of sub-menus for other filters. Useful if the number of additional filters is small. */
 	bool bUseSectionsForCustomFilterCategories;
 
+	/** If specified, the asset view options menu will use this tool menu profile for customization purposes. */
+	TOptional<FName> AssetViewOptionsProfile;
+
 	FAssetPickerConfig()
 		: SelectionMode( ESelectionMode::Multi )
 		, ThumbnailLabel( EThumbnailLabel::ClassName )
diff --git a/Engine/Source/Editor/ContentBrowser/Public/SAssetView.h b/Engine/Source/Editor/ContentBrowser/Public/SAssetView.h
--- a/Engine/Source/Editor/ContentBrowser/Public/SAssetView.h
+++ b/Engine/Source/Editor/ContentBrowser/Public/SAssetView.h
@@ -288,6 +288,9 @@ public:
 
 		/** The content browser that owns this view if any */
 		SLATE_ARGUMENT(TSharedPtr<SContentBrowser>, OwningContentBrowser)
+
+		/** The menu profile to use for the asset view options. The profile needs to be registered with the ToolMenus API. */
+		SLATE_ARGUMENT(TOptional<FName>, AssetViewOptionsProfile)
 	SLATE_END_ARGS()
 
 	~SAssetView();
@@ -995,6 +998,9 @@ private:
 	/** Called when opening view options menu */
 	FOnExtendAssetViewOptionsMenuContext OnExtendAssetViewOptionsMenuContext;
 
+	/** An optional profile name for the asset view options menu. */
+	TOptional<FName> AssetViewOptionsProfile;
+	
 	/** When true, filtered list items will be sorted next tick. Provided another sort hasn't happened recently or we are renaming an asset */
 	bool bPendingSortFilteredItems;
 	double CurrentTime;
