Fixed crash when ValidateDefaultPins was executing mid-transaction

diff --git a/Engine/Plugins/FX/Niagara/Source/NiagaraEditor/Private/NiagaraGraph.cpp b/Engine/Plugins/FX/Niagara/Source/NiagaraEditor/Private/NiagaraGraph.cpp
--- a/Engine/Plugins/FX/Niagara/Source/NiagaraEditor/Private/NiagaraGraph.cpp
+++ b/Engine/Plugins/FX/Niagara/Source/NiagaraEditor/Private/NiagaraGraph.cpp
@@ -2931,7 +2931,11 @@ void UNiagaraGraph::ScriptVariableChanged(FNiagaraVariable Variable)
 		MapGet->SynchronizeDefaultPins();
 	}
 
-	ValidateDefaultPins();
+	if(GIsTransacting == false)
+	{
+		ValidateDefaultPins();
+	}
+	
 	NotifyGraphChanged();
 }
 
