fix math for index->unit conversion

diff --git a/Engine/Plugins/FX/Niagara/Source/Niagara/Private/NiagaraDataInterfaceRW.cpp b/Engine/Plugins/FX/Niagara/Source/Niagara/Private/NiagaraDataInterfaceRW.cpp
--- a/Engine/Plugins/FX/Niagara/Source/Niagara/Private/NiagaraDataInterfaceRW.cpp
+++ b/Engine/Plugins/FX/Niagara/Source/Niagara/Private/NiagaraDataInterfaceRW.cpp
@@ -1039,7 +1039,7 @@ bool UNiagaraDataInterfaceGrid2D::GetFunctionHLSL(const FNiagaraDataInterfaceGPU
 			void {FunctionName}(out float2 Out_Unit)
 			{
 				#if NIAGARA_DISPATCH_TYPE == NIAGARA_DISPATCH_TYPE_TWO_D || NIAGARA_DISPATCH_TYPE == NIAGARA_DISPATCH_TYPE_CUSTOM
-					Out_Unit = (float2(GDispatchThreadId.x, GDispatchThreadId.y) + .5) * rcp(1.f/{NumCellsName});
+					Out_Unit = (float2(GDispatchThreadId.x, GDispatchThreadId.y) + .5) * rcp({NumCellsName});
 				#else
 					const uint Linear = GLinearThreadId;
 					const uint IndexX = Linear % {NumCellsName}.x;
