- Fix for system leaking when GPU is disabled using fx.NiagaraAllowGPUParticles

diff --git a/Engine/Plugins/FX/Niagara/Source/Niagara/Private/NiagaraCommon.cpp b/Engine/Plugins/FX/Niagara/Source/Niagara/Private/NiagaraCommon.cpp
--- a/Engine/Plugins/FX/Niagara/Source/Niagara/Private/NiagaraCommon.cpp
+++ b/Engine/Plugins/FX/Niagara/Source/Niagara/Private/NiagaraCommon.cpp
@@ -6,6 +6,7 @@
 #include "RenderUtils.h"
 #include "Misc/StringBuilder.h"
 #include "NiagaraComponent.h"
+#include "NiagaraComponentSettings.h"
 #include "NiagaraConstants.h"
 #include "NiagaraCustomVersion.h"
 #include "NiagaraEmitter.h"
@@ -24,7 +25,6 @@
 
 DECLARE_CYCLE_STAT(TEXT("Niagara - Utilities - PrepareRapidIterationParameters"), STAT_Niagara_Utilities_PrepareRapidIterationParameters, STATGROUP_Niagara);
 
-
 //////////////////////////////////////////////////////////////////////////
 
 int32 GNiagaraAllowComputeShaders = 1;
@@ -32,6 +32,7 @@ FAutoConsoleVariableRef CVarAllowComputeShaders(
 	TEXT("fx.NiagaraAllowComputeShaders"),
 	GNiagaraAllowComputeShaders,
 	TEXT("If true, allow the usage compute shaders within Niagara."),
+	FConsoleVariableDelegate::CreateStatic(FNiagaraComponentSettings::RequestUpdateSettings),
 	ECVF_Default);
 
 int32 GNiagaraAllowGPUParticles = 1;
@@ -39,6 +40,7 @@ FAutoConsoleVariableRef CVarAllowGPUParticles(
 	TEXT("fx.NiagaraAllowGPUParticles"),
 	GNiagaraAllowGPUParticles,
 	TEXT("If true, allow the usage of GPU particles for Niagara."),
+	FConsoleVariableDelegate::CreateStatic(FNiagaraComponentSettings::RequestUpdateSettings),
 	ECVF_Scalability | ECVF_Default);
 
 int32 GNiagaraGPUCulling = 1;
@@ -940,24 +942,24 @@ bool FNiagaraUtilities::LogVerboseWarnings()
 }
 #endif
 
-bool FNiagaraUtilities::AllowGPUParticles(EShaderPlatform ShaderPlatform)
+bool FNiagaraUtilities::AllowGPUParticles()
 {
 	return GNiagaraAllowGPUParticles && GNiagaraAllowComputeShaders && GRHISupportsDrawIndirect;
 }
 
-bool FNiagaraUtilities::AllowComputeShaders(EShaderPlatform ShaderPlatform)
+bool FNiagaraUtilities::AllowComputeShaders()
 {
 	return GNiagaraAllowComputeShaders && GRHISupportsDrawIndirect;
 }
 
-bool FNiagaraUtilities::AllowGPUSorting(EShaderPlatform ShaderPlatform)
+bool FNiagaraUtilities::AllowGPUSorting()
 {
 	return FXConsoleVariables::bAllowGPUSorting != 0;
 }
 
-bool FNiagaraUtilities::AllowGPUCulling(EShaderPlatform ShaderPlatform)
+bool FNiagaraUtilities::AllowGPUCulling()
 {
-	return GNiagaraGPUCulling && AllowGPUSorting(ShaderPlatform) && AllowComputeShaders(ShaderPlatform);
+	return GNiagaraGPUCulling && AllowGPUSorting() && AllowComputeShaders();
 }
 
 bool FNiagaraUtilities::AreBufferSRVsAlwaysCreated(EShaderPlatform ShaderPlatform)
diff --git a/Engine/Plugins/FX/Niagara/Source/Niagara/Private/NiagaraComponentSettings.cpp b/Engine/Plugins/FX/Niagara/Source/Niagara/Private/NiagaraComponentSettings.cpp
--- a/Engine/Plugins/FX/Niagara/Source/Niagara/Private/NiagaraComponentSettings.cpp
+++ b/Engine/Plugins/FX/Niagara/Source/Niagara/Private/NiagaraComponentSettings.cpp
@@ -200,8 +200,8 @@ namespace FNiagaraComponentSettings
 		}
 		bNeedsSettingsUpdate = false;
 
-		bool bShouldAllowGpuEmitters = true;
-		if (GpuRHIDenyListString.Len() > 0)
+		bool bShouldAllowGpuEmitters = FNiagaraUtilities::AllowGPUParticles();
+		if (bShouldAllowGpuEmitters && GpuRHIDenyListString.Len() > 0)
 		{
 			TArray<FString> BanNames;
 			GpuRHIDenyListString.ParseIntoArray(BanNames, TEXT(","));
diff --git a/Engine/Plugins/FX/Niagara/Source/Niagara/Public/NiagaraCommon.h b/Engine/Plugins/FX/Niagara/Source/Niagara/Public/NiagaraCommon.h
--- a/Engine/Plugins/FX/Niagara/Source/Niagara/Public/NiagaraCommon.h
+++ b/Engine/Plugins/FX/Niagara/Source/Niagara/Public/NiagaraCommon.h
@@ -1488,16 +1488,22 @@ namespace FNiagaraUtilities
 #endif
 
 	// Whether GPU particles are currently allowed. Could change depending on config and runtime switches.
-	bool AllowGPUParticles(EShaderPlatform ShaderPlatform);
+	bool AllowGPUParticles();
 
 	// Whether compute shaders are allowed. Could change depending on config and runtime switches.
-	bool AllowComputeShaders(EShaderPlatform ShaderPlatform);
+	bool AllowComputeShaders();
 
 	// Are we able to use the GPU for culling?
-	bool AllowGPUCulling(EShaderPlatform ShaderPlatform);
+	bool AllowGPUCulling();
 
 	// Are we able to use the GPU for sorting?
-	bool AllowGPUSorting(EShaderPlatform ShaderPlatform);
+	bool AllowGPUSorting();
+
+	//-TODO: Consider deprecating these methods if we don't need the ShaderPlatform
+	inline bool AllowGPUParticles(EShaderPlatform ShaderPlatform) { return AllowGPUParticles(); }
+	inline bool AllowComputeShaders(EShaderPlatform ShaderPlatform) { return AllowComputeShaders(); }
+	inline bool AllowGPUCulling(EShaderPlatform ShaderPlatform) { return AllowGPUCulling(); }
+	inline bool AllowGPUSorting(EShaderPlatform ShaderPlatform) { return AllowGPUSorting();  }
 
 	// Helper function to detect if SRVs are always created for buffers or not
 	bool AreBufferSRVsAlwaysCreated(EShaderPlatform ShaderPlatform);
diff --git a/Engine/Plugins/FX/Niagara/Source/Niagara/Public/NiagaraComponentSettings.h b/Engine/Plugins/FX/Niagara/Source/Niagara/Public/NiagaraComponentSettings.h
--- a/Engine/Plugins/FX/Niagara/Source/Niagara/Public/NiagaraComponentSettings.h
+++ b/Engine/Plugins/FX/Niagara/Source/Niagara/Public/NiagaraComponentSettings.h
@@ -11,6 +11,7 @@ struct FVersionedNiagaraEmitterData;
 namespace FNiagaraComponentSettings
 {
 	extern void UpdateSettings();
+	extern void RequestUpdateSettings(IConsoleVariable*);
 	extern NIAGARA_API bool IsSystemAllowedToRun(const UNiagaraSystem* System);
 	extern NIAGARA_API bool IsEmitterAllowedToRun(const FVersionedNiagaraEmitterData& EmitterData, const UNiagaraEmitter& NiagaraEmitter);
 };
