Moved Niagara passes to execute off the render thread and RHI thread.

diff --git a/Engine/Plugins/FX/Niagara/Source/Niagara/Private/NiagaraGPUProfiler.cpp b/Engine/Plugins/FX/Niagara/Source/Niagara/Private/NiagaraGPUProfiler.cpp
--- a/Engine/Plugins/FX/Niagara/Source/Niagara/Private/NiagaraGPUProfiler.cpp
+++ b/Engine/Plugins/FX/Niagara/Source/Niagara/Private/NiagaraGPUProfiler.cpp
@@ -34,7 +34,7 @@ FNiagaraGPUProfiler::~FNiagaraGPUProfiler()
 	}
 }
 
-void FNiagaraGPUProfiler::BeginFrame(FRHICommandListImmediate& RHICmdList)
+void FNiagaraGPUProfiler::BeginFrame(FRHICommandList& RHICmdList)
 {
 	// Process all frames until we run out
 	while (FGpuFrameData* ReadFrame = GetReadFrame())
@@ -105,7 +105,7 @@ void FNiagaraGPUProfiler::EndDispatch(FRHICommandList& RHICmdList)
 	RHICmdList.EndRenderQuery(DispatchTimer.EndQuery.GetQuery());
 }
 
-bool FNiagaraGPUProfiler::ProcessFrame(FRHICommandListImmediate& /*RHICmdList*/, FGpuFrameData& ReadFrame)
+bool FNiagaraGPUProfiler::ProcessFrame(FRHICommandList& /*RHICmdList*/, FGpuFrameData& ReadFrame)
 {
 	// Frame ready to process?
 	//-OPT: We can just look at the last write stage end timer here, but that relies on the batcher always executing
diff --git a/Engine/Plugins/FX/Niagara/Source/Niagara/Private/NiagaraGPUProfiler.h b/Engine/Plugins/FX/Niagara/Source/Niagara/Private/NiagaraGPUProfiler.h
--- a/Engine/Plugins/FX/Niagara/Source/Niagara/Private/NiagaraGPUProfiler.h
+++ b/Engine/Plugins/FX/Niagara/Source/Niagara/Private/NiagaraGPUProfiler.h
@@ -41,7 +41,7 @@ public:
 	FNiagaraGPUProfiler(uintptr_t InOwnerContext);
 	~FNiagaraGPUProfiler();
 
-	void BeginFrame(FRHICommandListImmediate& RHICmdList);
+	void BeginFrame(FRHICommandList& RHICmdList);
 	void EndFrame(FRHICommandList& RHICmdList);
 
 	void BeginDispatch(FRHICommandList& RHICmdList, const FNiagaraGpuProfileEvent& Event);
@@ -50,7 +50,7 @@ public:
 private:
 	FGpuFrameData* GetReadFrame() { check(CurrentReadFrame >= 0 && CurrentReadFrame < UE_ARRAY_COUNT(GpuFrames)); return GpuFrames[CurrentReadFrame].CanRead() ? &GpuFrames[CurrentReadFrame] : nullptr; }
 	FGpuFrameData* GetWriteFrame() { check(CurrentWriteFrame >= 0 && CurrentWriteFrame < UE_ARRAY_COUNT(GpuFrames)); return GpuFrames[CurrentWriteFrame].CanWrite() ? &GpuFrames[CurrentWriteFrame] : nullptr; }
-	bool ProcessFrame(FRHICommandListImmediate& RHICmdList, FGpuFrameData& ReadFrame);
+	bool ProcessFrame(FRHICommandList& RHICmdList, FGpuFrameData& ReadFrame);
 
 private:
 	uintptr_t				OwnerContext = 0;
diff --git a/Engine/Plugins/FX/Niagara/Source/Niagara/Private/NiagaraGpuComputeDispatch.cpp b/Engine/Plugins/FX/Niagara/Source/Niagara/Private/NiagaraGpuComputeDispatch.cpp
--- a/Engine/Plugins/FX/Niagara/Source/Niagara/Private/NiagaraGpuComputeDispatch.cpp
+++ b/Engine/Plugins/FX/Niagara/Source/Niagara/Private/NiagaraGpuComputeDispatch.cpp
@@ -1029,7 +1029,7 @@ void FNiagaraGpuComputeDispatch::ExecuteTicks(FRDGBuilder& GraphBuilder, TConstS
 			GraphBuilder.AddPass(
 				RDG_EVENT_NAME("Niagara::ExecuteTicksPre"),
 				ERDGPassFlags::None,
-				[this, TickStage, NumDispatchGroups = DispatchListPerStage[TickStage].DispatchGroups.Num(), GPUIndex](FRHICommandListImmediate& RHICmdList)
+				[this, TickStage, NumDispatchGroups = DispatchListPerStage[TickStage].DispatchGroups.Num(), GPUIndex](FRHICommandList& RHICmdList)
 				{
 					WaitForMultiGPUBuffers(RHICmdList, GPUIndex);
 				}
@@ -1042,7 +1042,7 @@ void FNiagaraGpuComputeDispatch::ExecuteTicks(FRDGBuilder& GraphBuilder, TConstS
 			GraphBuilder.AddPass(
 				RDG_EVENT_NAME("Niagara::MultiViewPreviousDataClear"),
 				ERDGPassFlags::None,
-				[this](FRHICommandListImmediate& RHICmdList)
+				[this](FRHICommandList& RHICmdList)
 				{
 					// Handle bookkeeping for multi-view, clearing MultiViewPreviousDataToRender pointer and running a couple bits of cleanup logic,
 					// which we had to skip earlier to keep necessary values valid when generating render data for additional view families.
@@ -1187,7 +1187,7 @@ void FNiagaraGpuComputeDispatch::ExecuteTicks(FRDGBuilder& GraphBuilder, TConstS
  		GraphBuilder.AddPass(
 			{},//RDG_EVENT_NAME("Niagara::ExecuteTicks::DispatchGroupPre"),
 			ERDGPassFlags::None,
-			[this, PreStageTransitions=MoveTemp(PreStageTransitions), PreStageIDToIndexInit=MoveTemp(PreStageIDToIndexInit), DispatchInstances=MakeArrayView(DispatchGroup.DispatchInstances)](FRHICommandListImmediate& RHICmdList)
+			[this, PreStageTransitions=MoveTemp(PreStageTransitions), PreStageIDToIndexInit=MoveTemp(PreStageIDToIndexInit), DispatchInstances=MakeArrayView(DispatchGroup.DispatchInstances)](FRHICommandList& RHICmdList)
 			{
 				// Execute Before Transitions
 				RHICmdList.Transition(PreStageTransitions);
@@ -1228,7 +1228,7 @@ void FNiagaraGpuComputeDispatch::ExecuteTicks(FRDGBuilder& GraphBuilder, TConstS
  		GraphBuilder.AddPass(
 			{},//RDG_EVENT_NAME("Niagara::ExecuteTicks::DispatchGroupPost"),
 			ERDGPassFlags::None,
-			[this, PostStageTransitions=MoveTemp(PostStageTransitions), DispatchInstances=MakeArrayView(DispatchGroup.DispatchInstances)](FRHICommandListImmediate& RHICmdList)
+			[this, PostStageTransitions=MoveTemp(PostStageTransitions), DispatchInstances=MakeArrayView(DispatchGroup.DispatchInstances)](FRHICommandList& RHICmdList)
 			{
 				RHICmdList.EndUAVOverlap(GPUInstanceCounterManager.GetInstanceCountBuffer().UAV);
 
@@ -1324,7 +1324,7 @@ void FNiagaraGpuComputeDispatch::ExecuteTicks(FRDGBuilder& GraphBuilder, TConstS
  			GraphBuilder.AddPass(
 				{},//RDG_EVENT_NAME("Niagara::ExecuteTicks::FreeIDUpdates"),
 				ERDGPassFlags::None,
-				[this, FreeIDUpdates=MakeArrayView(DispatchGroup.FreeIDUpdates)](FRHICommandListImmediate& RHICmdList)
+				[this, FreeIDUpdates=MakeArrayView(DispatchGroup.FreeIDUpdates)](FRHICommandList& RHICmdList)
 				{
 					const uint32 NumFreeIDUpdates = FreeIDUpdates.Num();
 
@@ -1373,16 +1373,12 @@ void FNiagaraGpuComputeDispatch::ExecuteTicks(FRDGBuilder& GraphBuilder, TConstS
 	}
 
 	const int32 StageTotalDispatches = TotalDispatchesThisFrame - StageStartTotalDispatches;
-	GraphBuilder.AddPass(
-		RDG_EVENT_NAME("Niagara::ExecuteTicksPost"),
-		ERDGPassFlags::None,
-		[this, StageTotalDispatches, TickStage, DispatchListPtr=&DispatchList](FRHICommandListImmediate& RHICmdList)
-		{
-			// Clear dispatch groups
-			// We do not release the counts as we won't do that until we finish the dispatches
-			DispatchListPtr->DispatchGroups.Empty();
-		}
-	);
+	GraphBuilder.AddPostExecuteCallback([&DispatchList]
+	{
+		// Clear dispatch groups
+		// We do not release the counts as we won't do that until we finish the dispatches
+		DispatchList.DispatchGroups.Empty();
+	});
 
 	// Tear down for tick pass
 	SimulationSceneViews = TConstStridedView<FSceneView>();
@@ -1672,7 +1668,7 @@ void FNiagaraGpuComputeDispatch::DispatchStage(FRDGBuilder& GraphBuilder, const
 			ShaderParametersMetadata,
 			DispatchParameters,
 			ERDGPassFlags::Compute | ERDGPassFlags::NeverCull,
-			[this, ShaderParametersMetadata, DispatchParameters, ComputeShader, TickPtr=&Tick, InstanceDataPtr=&InstanceData, SimStageDataPtr=&SimStageData](FRHICommandListImmediate& RHICmdList)		//-TODO:RDG: When legacy is removed this can be FRHIComputeCommandList
+			[this, ShaderParametersMetadata, DispatchParameters, ComputeShader, TickPtr=&Tick, InstanceDataPtr=&InstanceData, SimStageDataPtr=&SimStageData](FRHICommandList& RHICmdList)
 			{
 				FRHIComputeShader* RHIComputeShader = ComputeShader.GetComputeShader();
 				SetComputePipelineState(RHICmdList, RHIComputeShader);
@@ -1765,7 +1761,7 @@ void FNiagaraGpuComputeDispatch::DispatchStage(FRDGBuilder& GraphBuilder, const
 			ShaderParametersMetadata,
 			DispatchParameters,
 			ERDGPassFlags::Compute | ERDGPassFlags::NeverCull,
-			[this, ShaderParametersMetadata, DispatchParameters, ComputeShader, ThreadGroupCount, TickPtr=&Tick, InstanceDataPtr=&InstanceData, SimStageDataPtr=&SimStageData](FRHICommandListImmediate& RHICmdList)		//-TODO:RDG: When legacy is removed this can be FRHIComputeCommandList
+			[this, ShaderParametersMetadata, DispatchParameters, ComputeShader, ThreadGroupCount, TickPtr=&Tick, InstanceDataPtr=&InstanceData, SimStageDataPtr=&SimStageData](FRHICommandList& RHICmdList)
 			{
 				FRHIComputeShader* RHIComputeShader = ComputeShader.GetComputeShader();
 				SetComputePipelineState(RHICmdList, RHIComputeShader);
@@ -1782,8 +1778,6 @@ void FNiagaraGpuComputeDispatch::DispatchStage(FRDGBuilder& GraphBuilder, const
 					BatchedParameters.SetShaderUniformBuffer(ComputeShader->ExternalConstantBufferParam[1].GetBaseIndex(), TickPtr->GetExternalUniformBuffer(*InstanceDataPtr, true));
 				}
 
-				FNiagaraEmptyUAVPoolScopedAccess UAVPoolAccessScope(GetEmptyUAVPool());
-
 				FNiagaraGpuProfileScope GpuProfileDispatchScope(RHICmdList, this, FNiagaraGpuProfileEvent(*InstanceDataPtr, *SimStageDataPtr, InstanceDataPtr == &TickPtr->GetInstances()[0]));
 
 				SetShaderParameters(BatchedParameters, ComputeShader, ShaderParametersMetadata, *DispatchParameters);
@@ -1885,10 +1879,7 @@ void FNiagaraGpuComputeDispatch::PreInitViews(FRDGBuilder& GraphBuilder, bool bA
 #endif
 	if ((ViewFamilies.Num() > 1) && bIsFirstViewFamily)
 	{
-		AddPass(GraphBuilder, RDG_EVENT_NAME("Niagara::CopyToMultiViewCountBuffer"), [this](FRHICommandListImmediate& RHICmdList)
-		{
-			GPUInstanceCounterManager.CopyToMultiViewCountBuffer(RHICmdList);
-		});
+		GPUInstanceCounterManager.CopyToMultiViewCountBuffer(GraphBuilder.RHICmdList);
 	}
 
 	EmptyUAVPoolPtr->Tick();
@@ -2063,8 +2054,6 @@ void FNiagaraGpuComputeDispatch::PostRenderOpaque(FRDGBuilder& GraphBuilder, TCo
 			{
 				if (bAllowGPUParticleUpdate)
 				{
-					FinishDispatches();
-
 					AsyncGpuTraceHelper->EndFrame(RHICmdList, this, PassParameters->Scene->GetRHIRef());
 				}
 
@@ -2082,6 +2071,15 @@ void FNiagaraGpuComputeDispatch::PostRenderOpaque(FRDGBuilder& GraphBuilder, TCo
 			}
 		);
 	}
+
+	if (bAllowGPUParticleUpdate)
+	{
+		GraphBuilder.AddPostExecuteCallback([this]
+		{
+			FinishDispatches();
+		});
+	}
+
 	bRequiresReadback = false;
 
 	OnPostRenderEvent.Broadcast(GraphBuilder);
@@ -2092,7 +2090,7 @@ void FNiagaraGpuComputeDispatch::PostRenderOpaque(FRDGBuilder& GraphBuilder, TCo
 	}
 }
 
-void FNiagaraGpuComputeDispatch::ProcessDebugReadbacks(FRHICommandListImmediate& RHICmdList, bool bWaitCompletion)
+void FNiagaraGpuComputeDispatch::ProcessDebugReadbacks(FRHICommandList& RHICmdList, bool bWaitCompletion)
 {
 #if !UE_BUILD_SHIPPING
 	// Execute any pending readbacks as the ticks have now all been processed
@@ -2149,7 +2147,7 @@ void FNiagaraGpuComputeDispatch::ProcessDebugReadbacks(FRHICommandListImmediate&
 
 	if (bWaitCompletion)
 	{
-		GpuReadbackManagerPtr->WaitCompletion(RHICmdList);
+		GpuReadbackManagerPtr->WaitCompletion(RHICmdList.GetAsImmediate());
 	}
 #endif
 }
diff --git a/Engine/Plugins/FX/Niagara/Source/Niagara/Private/NiagaraGpuComputeDispatch.h b/Engine/Plugins/FX/Niagara/Source/Niagara/Private/NiagaraGpuComputeDispatch.h
--- a/Engine/Plugins/FX/Niagara/Source/Niagara/Private/NiagaraGpuComputeDispatch.h
+++ b/Engine/Plugins/FX/Niagara/Source/Niagara/Private/NiagaraGpuComputeDispatch.h
@@ -139,7 +139,7 @@ public:
 	void ProcessPendingTicksFlush(FRHICommandListImmediate& RHICmdList, bool bForceFlush);
 
 	/** Processes all pending readbacks */
-	virtual void ProcessDebugReadbacks(FRHICommandListImmediate& RHICmdList, bool bWaitCompletion) override;
+	virtual void ProcessDebugReadbacks(FRHICommandList& RHICmdList, bool bWaitCompletion) override;
 
 	virtual bool AddSortedGPUSimulation(FRHICommandListBase& RHICmdList, FNiagaraGPUSortInfo& SortInfo) override;
 	virtual const FGlobalDistanceFieldParameterData* GetGlobalDistanceFieldData() const override;
diff --git a/Engine/Plugins/FX/Niagara/Source/Niagara/Public/NiagaraGpuComputeDispatchInterface.h b/Engine/Plugins/FX/Niagara/Source/Niagara/Public/NiagaraGpuComputeDispatchInterface.h
--- a/Engine/Plugins/FX/Niagara/Source/Niagara/Public/NiagaraGpuComputeDispatchInterface.h
+++ b/Engine/Plugins/FX/Niagara/Source/Niagara/Public/NiagaraGpuComputeDispatchInterface.h
@@ -150,7 +150,7 @@ public:
 	virtual void AddDebugReadback(FNiagaraSystemInstanceID InstanceID, TSharedPtr<FNiagaraScriptDebuggerInfo, ESPMode::ThreadSafe> DebugInfo, FNiagaraComputeExecutionContext* Context) = 0;
 
 	/** Processes all pending debug readbacks */
-	virtual void ProcessDebugReadbacks(FRHICommandListImmediate& RHICmdList, bool bWaitCompletion) = 0;
+	virtual void ProcessDebugReadbacks(FRHICommandList& RHICmdList, bool bWaitCompletion) = 0;
 
 	virtual FNiagaraAsyncGpuTraceHelper& GetAsyncGpuTraceHelper() const = 0;
 
diff --git a/Engine/Source/Runtime/Engine/Private/GPUSortManager.cpp b/Engine/Source/Runtime/Engine/Private/GPUSortManager.cpp
--- a/Engine/Source/Runtime/Engine/Private/GPUSortManager.cpp
+++ b/Engine/Source/Runtime/Engine/Private/GPUSortManager.cpp
@@ -202,7 +202,7 @@ void FCopyUIntBufferCS::End(FRHICommandList& RHICmdList)
 
 //*****************************************************************************
 
-void CopyUIntBufferToTargets(FRHICommandListImmediate& RHICmdList, ERHIFeatureLevel::Type FeatureLevel,  FRHIShaderResourceView* SourceSRV, FRHIUnorderedAccessView*const* TargetUAVs, int32* TargetSizes, int32 StartingOffset, int32 NumTargets)
+void CopyUIntBufferToTargets(FRHICommandList& RHICmdList, ERHIFeatureLevel::Type FeatureLevel,  FRHIShaderResourceView* SourceSRV, FRHIUnorderedAccessView*const* TargetUAVs, int32* TargetSizes, int32 StartingOffset, int32 NumTargets)
 {
 	// No that resource transition must be made outside this call as we don't know how the content of those have been generated, and will be used.
 
diff --git a/Engine/Source/Runtime/Engine/Public/GPUSortManager.h b/Engine/Source/Runtime/Engine/Public/GPUSortManager.h
--- a/Engine/Source/Runtime/Engine/Public/GPUSortManager.h
+++ b/Engine/Source/Runtime/Engine/Public/GPUSortManager.h
@@ -30,7 +30,7 @@
  * @param StartingOffset - The starting position at which the copy starts. Applies for both the source and targets.
  * @param NumTargets - The number of elements in TargetUAVs and TargetSizes.
  */
-ENGINE_API void CopyUIntBufferToTargets(FRHICommandListImmediate& RHICmdList, ERHIFeatureLevel::Type FeatureLevel, FRHIShaderResourceView* SourceSRV, FRHIUnorderedAccessView*const* TargetUAVs, int32* TargetSizes, int32 StartingOffset, int32 NumTargets);
+ENGINE_API void CopyUIntBufferToTargets(FRHICommandList& RHICmdList, ERHIFeatureLevel::Type FeatureLevel, FRHIShaderResourceView* SourceSRV, FRHIUnorderedAccessView*const* TargetUAVs, int32* TargetSizes, int32 StartingOffset, int32 NumTargets);
 
 struct FGPUSortBuffers;
 
