Removed "is asset == false" check to strip thumbnail data during cooks.
This should remove the thumbnails from our cooked emitters

diff --git a/Engine/Plugins/FX/Niagara/Source/NiagaraEditor/Private/NiagaraEmitterEditorData.cpp b/Engine/Plugins/FX/Niagara/Source/NiagaraEditor/Private/NiagaraEmitterEditorData.cpp
--- a/Engine/Plugins/FX/Niagara/Source/NiagaraEditor/Private/NiagaraEmitterEditorData.cpp
+++ b/Engine/Plugins/FX/Niagara/Source/NiagaraEditor/Private/NiagaraEmitterEditorData.cpp
@@ -35,11 +35,10 @@ UNiagaraEmitterEditorData::UNiagaraEmitterEditorData(const FObjectInitializer& O
 void UNiagaraEmitterEditorData::Serialize(FArchive& Ar)
 {
 #if WITH_EDITORONLY_DATA
-	// When cooking an emitter that's not an asset, clear out the thumbnail image to prevent issues
-	// with cooked editor data.
-	bool bCookingNonAssetEmitter = Ar.IsCooking() && GetTypedOuter<UNiagaraEmitter>() && GetTypedOuter<UNiagaraEmitter>()->IsAsset() == false;
+	// We cook out the thumbnail texture if we are cooking as cooked textures don't behave well when duplicated into non-cooked systems
+	bool bCookingEmitter = Ar.IsCooking() && GetTypedOuter<UNiagaraEmitter>();
 	UTexture2D* CachedThumbnail = nullptr;
-	if (bCookingNonAssetEmitter)
+	if (bCookingEmitter)
 	{
 		CachedThumbnail = EmitterThumbnail;
 		EmitterThumbnail = nullptr;
@@ -50,7 +49,7 @@ void UNiagaraEmitterEditorData::Serialize(FArchive& Ar)
 
 #if WITH_EDITORONLY_DATA
 	// Restore the thumbnail image that was cleared before serialize.
-	if (bCookingNonAssetEmitter)
+	if (bCookingEmitter)
 	{
 		EmitterThumbnail = CachedThumbnail;
 	}
