- Remove stateless project setting, it's on by default

diff --git a/Engine/Plugins/FX/Niagara/Source/Niagara/Private/Stateless/NiagaraStatelessEmitter.cpp b/Engine/Plugins/FX/Niagara/Source/Niagara/Private/Stateless/NiagaraStatelessEmitter.cpp
--- a/Engine/Plugins/FX/Niagara/Source/Niagara/Private/Stateless/NiagaraStatelessEmitter.cpp
+++ b/Engine/Plugins/FX/Niagara/Source/Niagara/Private/Stateless/NiagaraStatelessEmitter.cpp
@@ -132,8 +132,7 @@ bool UNiagaraStatelessEmitter::NeedsLoadForTargetPlatform(const ITargetPlatform*
 		return false;
 	}
 
-	const bool bStatelessEnabled = GetDefault<UNiagaraSettings>()->bStatelessEmittersEnabled;
-	const bool bIsEnabled = Platforms.IsEnabledForPlatform(TargetPlatform->IniPlatformName()) && bStatelessEnabled;
+	const bool bIsEnabled = Platforms.IsEnabledForPlatform(TargetPlatform->IniPlatformName());
 	if (!bIsEnabled)
 	{
 		UE_LOG(LogNiagara, Verbose, TEXT("Pruned emitter %s for platform %s"), *GetFullName(), *TargetPlatform->DisplayName().ToString())
@@ -255,7 +254,6 @@ void UNiagaraStatelessEmitter::CacheFromCompiledData()
 
 	StatelessEmitterData->bCanEverExecute = NiagaraStatelessInternal::IsValid(*StatelessEmitterData);
 	StatelessEmitterData->bCanEverExecute &= Platforms.IsActive();
-	StatelessEmitterData->bCanEverExecute &= GetDefault<UNiagaraSettings>()->bStatelessEmittersEnabled;
 
 	// Build buffers that are shared across all instances
 	//-OPT: We should be able to build and serialize this data as part of the UNiagaraStatelessEmitter, potentially all of this data even since it's immutable and does not change at runtime
diff --git a/Engine/Plugins/FX/Niagara/Source/Niagara/Public/NiagaraRendererProperties.h b/Engine/Plugins/FX/Niagara/Source/Niagara/Public/NiagaraRendererProperties.h
--- a/Engine/Plugins/FX/Niagara/Source/Niagara/Public/NiagaraRendererProperties.h
+++ b/Engine/Plugins/FX/Niagara/Source/Niagara/Public/NiagaraRendererProperties.h
@@ -365,7 +365,7 @@ public:
 
 	NIAGARA_API virtual void GetRendererWidgets(const FNiagaraEmitterInstance* InEmitter, TArray<TSharedPtr<SWidget>>& OutWidgets, TSharedPtr<FAssetThumbnailPool> InThumbnailPool) const PURE_VIRTUAL(UNiagaraRendererProperties::GetRendererWidgets, );
 	NIAGARA_API virtual void GetRendererTooltipWidgets(const FNiagaraEmitterInstance* InEmitter, TArray<TSharedPtr<SWidget>>& OutWidgets, TSharedPtr<FAssetThumbnailPool> InThumbnailPool) const PURE_VIRTUAL(UNiagaraRendererProperties::GetRendererTooltipWidgets, );
-	virtual void GetRendererFeedback(const FVersionedNiagaraEmitter& InEmitter, TArray<FText>& OutErrors, TArray<FText>& OutWarnings, TArray<FText>& OutInfo) const {};
+	virtual void GetRendererFeedback(const FVersionedNiagaraEmitter& InEmitter, TArray<FText>& OutErrors, TArray<FText>& OutWarnings, TArray<FText>& OutInfo) const {}
 	NIAGARA_API virtual void GetRendererFeedback(const FVersionedNiagaraEmitter& InEmitter, TArray<FNiagaraRendererFeedback>& OutErrors, TArray<FNiagaraRendererFeedback>& OutWarnings, TArray<FNiagaraRendererFeedback>& OutInfo) const;
 
 	// The icon to display in the niagara stack widget under the renderer section
diff --git a/Engine/Plugins/FX/Niagara/Source/Niagara/Public/NiagaraSettings.h b/Engine/Plugins/FX/Niagara/Source/Niagara/Public/NiagaraSettings.h
--- a/Engine/Plugins/FX/Niagara/Source/Niagara/Public/NiagaraSettings.h
+++ b/Engine/Plugins/FX/Niagara/Source/Niagara/Public/NiagaraSettings.h
@@ -164,16 +164,6 @@ class UNiagaraSettings : public UDeveloperSettings
 	UPROPERTY(config, EditAnywhere, Category = Niagara, meta = (DisplayName = "Enable building data for Optimized VM"))
 	bool bExperimentalVMEnabled = false;
 
-	/**
-	Enables Lightweight Emitters experimental feature.
-	Statless emitters are lightweight fixed function emitters, they are not fully programmable like regular emitters and do not run scripts on the CPU.
-	Particle data is extrapolated per frame for the current particle age.  This means we never store particle data, we only generate it on demand.
-	Systems that contain only lightweight emitters and no system script modules can take advantage of a much faster path to execute.
-	** There is no guarantee on backwards compatability for this feature currently.  Do not ship lightweight content. **
-	*/
-	UPROPERTY(config, EditAnywhere, Category = Niagara, meta = (DisplayName = "Enable Lightweight Emitters (Experimental)", ConfigRestartRequired = true))
-	bool bStatelessEmittersEnabled = false;
-
 	/** If set to true, quaternion attributes will be interpolated via slerp instead of lerp in interpolated spawn scripts. */
 	UPROPERTY(config, EditAnywhere, AdvancedDisplay, Category = Niagara)
 	bool bAccurateQuatInterpolation = true;
diff --git a/Engine/Plugins/FX/Niagara/Source/NiagaraEditor/Private/NiagaraEditorModule.cpp b/Engine/Plugins/FX/Niagara/Source/NiagaraEditor/Private/NiagaraEditorModule.cpp
--- a/Engine/Plugins/FX/Niagara/Source/NiagaraEditor/Private/NiagaraEditorModule.cpp
+++ b/Engine/Plugins/FX/Niagara/Source/NiagaraEditor/Private/NiagaraEditorModule.cpp
@@ -301,12 +301,6 @@ class FNiagaraEditorOnlyDataUtilities : public INiagaraEditorOnlyDataUtilities
 	{
 		TOptional<FNiagaraSystemStateData> SystemStateData;
 
-		// Never allow if stateless is not enabled
-		if (!GetDefault<UNiagaraSettings>()->bStatelessEmittersEnabled)
-		{
-			return SystemStateData;
-		}
-
 		// All emitters must be stateless currently
 		// We can perhaps look at this again, but we always write Emitter.RandomSeed currently even with an empty script
 		for (const FNiagaraEmitterHandle& EmitterHandle : System.GetEmitterHandles())
diff --git a/Engine/Plugins/FX/Niagara/Source/NiagaraEditor/Private/NiagaraOverviewNode.cpp b/Engine/Plugins/FX/Niagara/Source/NiagaraEditor/Private/NiagaraOverviewNode.cpp
--- a/Engine/Plugins/FX/Niagara/Source/NiagaraEditor/Private/NiagaraOverviewNode.cpp
+++ b/Engine/Plugins/FX/Niagara/Source/NiagaraEditor/Private/NiagaraOverviewNode.cpp
@@ -66,17 +66,6 @@ void UNiagaraOverviewNode::Initialize(UNiagaraSystem* InOwningSystem, FGuid InEm
 
 void UNiagaraOverviewNode::UpdateStatus()
 {
-	if(FNiagaraEmitterHandle* EmitterHandle = FindEmitterHandleByID(OwningSystem, EmitterHandleGuid))
-	{
-		if(!GetDefault<UNiagaraSettings>()->bStatelessEmittersEnabled && EmitterHandle->GetEmitterMode() == ENiagaraEmitterMode::Stateless)
-		{
-			SetForceDisplayAsDisabled(true);
-		}
-		else
-		{
-			SetForceDisplayAsDisabled(false);
-		}
-	}
 }
 
 const FGuid UNiagaraOverviewNode::GetEmitterHandleGuid() const
@@ -310,47 +299,44 @@ void UNiagaraOverviewNode::GetNodeContextMenuActions(class UToolMenu* Menu, clas
 						)
 					)
 				);
-			#if 0
-				if (GetDefault<UNiagaraSettings>()->bStatelessEmittersEnabled)
-				{
-					Section.AddSubMenu(
-						"EmitterMode",
-						LOCTEXT("EmitterModeSubMenuLabel", "Emitter Mode..."),
-						FText(),
-						FNewToolMenuDelegate::CreateLambda([OwningSystemViewModel](UToolMenu* InMenu)
+
+				Section.AddSubMenu(
+					"EmitterMode",
+					LOCTEXT("EmitterModeSubMenuLabel", "Emitter Mode..."),
+					FText(),
+					FNewToolMenuDelegate::CreateLambda([OwningSystemViewModel](UToolMenu* InMenu)
+					{
+						if (OwningSystemViewModel.IsValid())
 						{
-							if (OwningSystemViewModel.IsValid())
-							{
-								bool bCanExecute = OwningSystemViewModel->GetSelectionViewModel()->GetSelectedEmitterHandleIds().Num() > 0;
-								FToolMenuSection& SubMenuSection = InMenu->AddSection("EmitterModeOptions", LOCTEXT("EmitterModeSubMenuHeader", "Emitter Mode"));
-								SubMenuSection.AddMenuEntry(
-									"SetEmitterModeStandard",
-									LOCTEXT("SetEmitterStandardModeMenuLabel", "Standard Mode"),
-									LOCTEXT("SetEmitterStandardModeMenuToolTip", "Set this emitter to use standard mode."),
-									FSlateIcon(),
-									FUIAction(
-										FExecuteAction::CreateStatic(&FNiagaraEditorUtilities::SetSelectedEmittersEmitterMode, OwningSystemViewModel.ToSharedRef(), ENiagaraEmitterMode::Standard),
-										FCanExecuteAction::CreateLambda([bCanExecute]() { return bCanExecute; }),
-										FGetActionCheckState::CreateStatic(&FNiagaraEditorUtilities::GetSelectedEmittersEmitterModeCheckState, OwningSystemViewModel.ToSharedRef(), ENiagaraEmitterMode::Standard)
-									),
-									EUserInterfaceActionType::RadioButton
-								);
-								SubMenuSection.AddMenuEntry(
-									"SetEmitterModeStateless",
-									LOCTEXT("SetEmitterStatelessModeMenuLabel", "Lightweight Mode"),
-									LOCTEXT("SetEmitterStatelessModeMenuToolTip", "Set this emitter to use lightweight mode."),
-									FSlateIcon(),
-									FUIAction(
-										FExecuteAction::CreateStatic(&FNiagaraEditorUtilities::SetSelectedEmittersEmitterMode, OwningSystemViewModel.ToSharedRef(), ENiagaraEmitterMode::Stateless),
-										FCanExecuteAction::CreateLambda([bCanExecute]() { return bCanExecute; }),
-										FGetActionCheckState::CreateStatic(&FNiagaraEditorUtilities::GetSelectedEmittersEmitterModeCheckState, OwningSystemViewModel.ToSharedRef(), ENiagaraEmitterMode::Stateless)
-									),
-									EUserInterfaceActionType::RadioButton
-								);
-							}
-						}));
-				}
-			#endif
+							bool bCanExecute = OwningSystemViewModel->GetSelectionViewModel()->GetSelectedEmitterHandleIds().Num() > 0;
+							FToolMenuSection& SubMenuSection = InMenu->AddSection("EmitterModeOptions", LOCTEXT("EmitterModeSubMenuHeader", "Emitter Mode"));
+							SubMenuSection.AddMenuEntry(
+								"SetEmitterModeStandard",
+								LOCTEXT("SetEmitterStandardModeMenuLabel", "Standard Mode"),
+								LOCTEXT("SetEmitterStandardModeMenuToolTip", "Set this emitter to use standard mode."),
+								FSlateIcon(),
+								FUIAction(
+									FExecuteAction::CreateStatic(&FNiagaraEditorUtilities::SetSelectedEmittersEmitterMode, OwningSystemViewModel.ToSharedRef(), ENiagaraEmitterMode::Standard),
+									FCanExecuteAction::CreateLambda([bCanExecute]() { return bCanExecute; }),
+									FGetActionCheckState::CreateStatic(&FNiagaraEditorUtilities::GetSelectedEmittersEmitterModeCheckState, OwningSystemViewModel.ToSharedRef(), ENiagaraEmitterMode::Standard)
+								),
+								EUserInterfaceActionType::RadioButton
+							);
+							SubMenuSection.AddMenuEntry(
+								"SetEmitterModeStateless",
+								LOCTEXT("SetEmitterStatelessModeMenuLabel", "Lightweight Mode"),
+								LOCTEXT("SetEmitterStatelessModeMenuToolTip", "Set this emitter to use lightweight mode."),
+								FSlateIcon(),
+								FUIAction(
+									FExecuteAction::CreateStatic(&FNiagaraEditorUtilities::SetSelectedEmittersEmitterMode, OwningSystemViewModel.ToSharedRef(), ENiagaraEmitterMode::Stateless),
+									FCanExecuteAction::CreateLambda([bCanExecute]() { return bCanExecute; }),
+									FGetActionCheckState::CreateStatic(&FNiagaraEditorUtilities::GetSelectedEmittersEmitterModeCheckState, OwningSystemViewModel.ToSharedRef(), ENiagaraEmitterMode::Stateless)
+								),
+								EUserInterfaceActionType::RadioButton
+							);
+						}
+					})
+				);
 			}
 		}
 	}
diff --git a/Engine/Plugins/FX/Niagara/Source/NiagaraEditorWidgets/Private/SNiagaraOverviewGraph.cpp b/Engine/Plugins/FX/Niagara/Source/NiagaraEditorWidgets/Private/SNiagaraOverviewGraph.cpp
--- a/Engine/Plugins/FX/Niagara/Source/NiagaraEditorWidgets/Private/SNiagaraOverviewGraph.cpp
+++ b/Engine/Plugins/FX/Niagara/Source/NiagaraEditorWidgets/Private/SNiagaraOverviewGraph.cpp
@@ -242,14 +242,11 @@ FActionMenuContent SNiagaraOverviewGraph::OnCreateGraphActionMenu(UEdGraph* InGr
 				LOCTEXT("AddEmitterToolTip", "Adds an empty emitter without any modules or renderers."),
 				FSlateIcon(),
 				FExecuteAction::CreateSP(this, &SNiagaraOverviewGraph::OnCreateEmptyEmitter));
-			if (GetDefault<UNiagaraSettings>()->bStatelessEmittersEnabled)
-			{
-				MenuBuilder.AddMenuEntry(
-					LOCTEXT("AddStatelessEmitter", "Add lightweight emitter"),
-					LOCTEXT("AddStatelessEmitterToolTip", "Adds a lightweight emitter without any modules or renderers."),
-					FSlateIcon(),
-					FExecuteAction::CreateSP(this, &SNiagaraOverviewGraph::OnCreateStatelessEmitter));
-			}
+			MenuBuilder.AddMenuEntry(
+				LOCTEXT("AddStatelessEmitter", "Add lightweight emitter"),
+				LOCTEXT("AddStatelessEmitterToolTip", "Adds a lightweight emitter without any modules or renderers."),
+				FSlateIcon(),
+				FExecuteAction::CreateSP(this, &SNiagaraOverviewGraph::OnCreateStatelessEmitter));
 			MenuBuilder.AddMenuEntry(
 				LOCTEXT("CommentsLabel", "Add Comment"),
 				LOCTEXT("AddCommentBoxToolTip", "Add a comment box"),
diff --git a/Engine/Plugins/FX/Niagara/Source/NiagaraEditorWidgets/Private/SNiagaraOverviewStackNode.cpp b/Engine/Plugins/FX/Niagara/Source/NiagaraEditorWidgets/Private/SNiagaraOverviewStackNode.cpp
--- a/Engine/Plugins/FX/Niagara/Source/NiagaraEditorWidgets/Private/SNiagaraOverviewStackNode.cpp
+++ b/Engine/Plugins/FX/Niagara/Source/NiagaraEditorWidgets/Private/SNiagaraOverviewStackNode.cpp
@@ -485,31 +485,6 @@ TSharedRef<SWidget> SNiagaraOverviewStackNode::CreateTitleWidget_Default(TShared
 			];
 	}
 
-	if (!GetDefault<UNiagaraSettings>()->bStatelessEmittersEnabled)
-	{
-		TSharedPtr<FNiagaraEmitterHandleViewModel> EmitterHandleViewModel = EmitterHandleViewModelWeak.Pin();
-		FNiagaraEmitterHandle* EmitterHandle = EmitterHandleViewModel.IsValid() ? EmitterHandleViewModel->GetEmitterHandle() : nullptr;
-		if (EmitterHandle && EmitterHandle->GetEmitterMode() != ENiagaraEmitterMode::Standard)
-		{
-			DefaultTitle =
-				SNew(SVerticalBox)
-				+ SVerticalBox::Slot()
-				.Padding(0, 0, 5, 0)
-				.AutoHeight()
-				[
-					SNew(STextBlock)
-						.Text(LOCTEXT("StatelessNotEnabled", "Lightweight not enabled in project settings."))
-						.Font(FCoreStyle::GetDefaultFontStyle("Regular", 8.f))
-						.OverflowPolicy(ETextOverflowPolicy::Ellipsis)
-						.AutoWrapText(true)
-				]
-				+ SVerticalBox::Slot()
-				[
-					DefaultTitle
-				];
-		}
-	}
-	
 	TSharedPtr<SWidget> TitleWidget = SNew(SHorizontalBox)
 	// Summary View Controls
 	+ SHorizontalBox::Slot()
