Stripping out emitter thumbnail texture for all cooked emitters on post load, removing the serialize thumbnail stripping

diff --git a/Engine/Plugins/FX/Niagara/Source/Niagara/Private/NiagaraEmitter.cpp b/Engine/Plugins/FX/Niagara/Source/Niagara/Private/NiagaraEmitter.cpp
--- a/Engine/Plugins/FX/Niagara/Source/Niagara/Private/NiagaraEmitter.cpp
+++ b/Engine/Plugins/FX/Niagara/Source/Niagara/Private/NiagaraEmitter.cpp
@@ -505,27 +505,10 @@ void UNiagaraEmitter::Serialize(FArchive& Ar)
 			}
 		}
 	}
-
-	// When cooking an emitter that's not an asset, clear out the thumbnail image to prevent issues
-	// with cooked editor data.
-	bool bCookingNonAssetEmitter = Ar.IsCooking() && this->IsAsset() == false;
-	UTexture2D* CachedThumbnail = nullptr;
-	if (bCookingNonAssetEmitter)
-	{
-		CachedThumbnail = ThumbnailImage;
-		ThumbnailImage = nullptr;
-	}
 	
 #endif
 	Super::Serialize(Ar);
-
-#if WITH_EDITORONLY_DATA
-	// Restore the thumbnail image that was cleared before serialize.
-	if (bCookingNonAssetEmitter)
-	{
-		ThumbnailImage = CachedThumbnail;
-	}
-#endif
+	
 
 	Ar.UsingCustomVersion(FNiagaraCustomVersion::GUID);
 	Ar.UsingCustomVersion(FUE5MainStreamObjectVersion::GUID);
@@ -667,10 +650,13 @@ void UNiagaraEmitter::PostLoad()
 		MessageKeyToMessageMap_DEPRECATED.Empty();
 	}
 
-	if (this->GetOuter()->IsA<UNiagaraSystem>() || this->GetOuter()->IsA<UNiagaraEmitter>())
+	if(UPackage* Package = GetPackage())
 	{
-		// Remove thunbnails for non-asset emitters to prevent problems with cooked for editor emitters being referenced by uncooked emitters and systems.
-		ThumbnailImage = nullptr;
+		if(Package->HasAnyPackageFlags(PKG_Cooked))
+		{
+			// We remove the thumbnail for cooked emitters as it can cause issues with cooked emitter being put into uncooked systems
+			ThumbnailImage = nullptr;
+		}
 	}
 #endif
 
diff --git a/Engine/Plugins/FX/Niagara/Source/NiagaraEditor/Private/NiagaraEmitterEditorData.cpp b/Engine/Plugins/FX/Niagara/Source/NiagaraEditor/Private/NiagaraEmitterEditorData.cpp
--- a/Engine/Plugins/FX/Niagara/Source/NiagaraEditor/Private/NiagaraEmitterEditorData.cpp
+++ b/Engine/Plugins/FX/Niagara/Source/NiagaraEditor/Private/NiagaraEmitterEditorData.cpp
@@ -32,30 +32,6 @@ UNiagaraEmitterEditorData::UNiagaraEmitterEditorData(const FObjectInitializer& O
 	PlaybackRangeMax = 10;
 }
 
-void UNiagaraEmitterEditorData::Serialize(FArchive& Ar)
-{
-#if WITH_EDITORONLY_DATA
-	// We cook out the thumbnail texture if we are cooking as cooked textures don't behave well when duplicated into non-cooked systems
-	bool bCookingEmitter = Ar.IsCooking() && GetTypedOuter<UNiagaraEmitter>();
-	UTexture2D* CachedThumbnail = nullptr;
-	if (bCookingEmitter)
-	{
-		CachedThumbnail = EmitterThumbnail;
-		EmitterThumbnail = nullptr;
-	}
-	
-#endif
-	Super::Serialize(Ar);
-
-#if WITH_EDITORONLY_DATA
-	// Restore the thumbnail image that was cleared before serialize.
-	if (bCookingEmitter)
-	{
-		EmitterThumbnail = CachedThumbnail;
-	}
-#endif
-}
-
 void UNiagaraEmitterEditorData::PostLoad_TransferSummaryDataToNewFormat()
 {
 	// ATTENTION
@@ -452,6 +428,15 @@ void UNiagaraEmitterEditorData::PostLoad()
 	{
 		SummaryViewRoot = NewObject<UNiagaraHierarchyRoot>(this, TEXT("SummaryViewRoot"), RF_Transactional);
 	}
+
+	if(UPackage* Package = GetPackage())
+	{
+		if(Package->HasAnyPackageFlags(PKG_Cooked))
+		{
+			// We remove the thumbnail for cooked emitters as it can cause issues with cooked emitter being put into uncooked systems
+			EmitterThumbnail = nullptr;
+		}
+	}
 }
 
 void UNiagaraEmitterEditorData::PostLoadFromOwner(UObject* InOwner)
