- Fix for Niagara profiler crash with parallel execution

diff --git a/Engine/Plugins/FX/Niagara/Source/Niagara/Private/NiagaraGPUProfiler.cpp b/Engine/Plugins/FX/Niagara/Source/Niagara/Private/NiagaraGPUProfiler.cpp
--- a/Engine/Plugins/FX/Niagara/Source/Niagara/Private/NiagaraGPUProfiler.cpp
+++ b/Engine/Plugins/FX/Niagara/Source/Niagara/Private/NiagaraGPUProfiler.cpp
@@ -14,6 +14,8 @@ static FAutoConsoleVariableRef CVarNiagaraGpuProfilingEnabled(
 
 #if WITH_NIAGARA_GPU_PROFILER
 
+thread_local uint32 FNiagaraGPUProfiler::ActiveQueryIndex = INDEX_NONE;
+
 FNiagaraGPUProfiler::FNiagaraGPUProfiler(uintptr_t InOwnerContext)
 	: OwnerContext(InOwnerContext)
 {
@@ -82,9 +84,11 @@ void FNiagaraGPUProfiler::BeginDispatch(FRHICommandList& RHICmdList, const FNiag
 	{
 		return;
 	}
-	check(bDispatchRecursionGuard == false);
-	bDispatchRecursionGuard = true;
 
+
+	UE::TScopeLock Lock(Mutex);
+	check(ActiveQueryIndex == INDEX_NONE);
+	ActiveQueryIndex = ActiveWriteFrame->DispatchTimers.Num();
 	FGpuDispatchTimer& DispatchTimer = ActiveWriteFrame->DispatchTimers.Emplace_GetRef(Event);
 
 	DispatchTimer.StartQuery = QueryPool->AllocateQuery();
@@ -97,12 +101,15 @@ void FNiagaraGPUProfiler::EndDispatch(FRHICommandList& RHICmdList)
 	{
 		return;
 	}
-	check(bDispatchRecursionGuard == true);
-	bDispatchRecursionGuard = false;
+	
+	UE::TScopeLock Lock(Mutex);
+	check(ActiveWriteFrame->DispatchTimers.IsValidIndex(ActiveQueryIndex));
 
-	FGpuDispatchTimer& DispatchTimer = ActiveWriteFrame->DispatchTimers.Last();
+	FGpuDispatchTimer& DispatchTimer = ActiveWriteFrame->DispatchTimers[ActiveQueryIndex];
 	DispatchTimer.EndQuery = QueryPool->AllocateQuery();
 	RHICmdList.EndRenderQuery(DispatchTimer.EndQuery.GetQuery());
+
+	ActiveQueryIndex = INDEX_NONE;
 }
 
 bool FNiagaraGPUProfiler::ProcessFrame(FRHICommandList& /*RHICmdList*/, FGpuFrameData& ReadFrame)
diff --git a/Engine/Plugins/FX/Niagara/Source/Niagara/Private/NiagaraGPUProfiler.h b/Engine/Plugins/FX/Niagara/Source/Niagara/Private/NiagaraGPUProfiler.h
--- a/Engine/Plugins/FX/Niagara/Source/Niagara/Private/NiagaraGPUProfiler.h
+++ b/Engine/Plugins/FX/Niagara/Source/Niagara/Private/NiagaraGPUProfiler.h
@@ -53,16 +53,17 @@ private:
 	bool ProcessFrame(FRHICommandList& RHICmdList, FGpuFrameData& ReadFrame);
 
 private:
-	uintptr_t				OwnerContext = 0;
+	uintptr_t					OwnerContext = 0;
 
-	int32					CurrentReadFrame = 0;					// Index of the next frame to read from
-	int32					CurrentWriteFrame = 0;					// Index of the next frame to write into
-	FGpuFrameData			GpuFrames[NumBufferFrames];
+	int32						CurrentReadFrame = 0;					// Index of the next frame to read from
+	int32						CurrentWriteFrame = 0;					// Index of the next frame to write into
+	FGpuFrameData				GpuFrames[NumBufferFrames];
 
-	FGpuFrameData*			ActiveWriteFrame = nullptr;				// Not null while we are generating a frame of data, otherwise null
-	bool					bDispatchRecursionGuard = false;		// We don't support timing dispatches inside one another
+	FGpuFrameData*				ActiveWriteFrame = nullptr;				// Not null while we are generating a frame of data, otherwise null
+	static thread_local uint32	ActiveQueryIndex;					// We don't support timing dispatches inside one another
+	UE::FMutex					Mutex;
 
-	FRenderQueryPoolRHIRef	QueryPool;
+	FRenderQueryPoolRHIRef		QueryPool;
 };
 
 #endif //WITH_NIAGARA_GPU_PROFILER
diff --git a/Engine/Plugins/FX/Niagara/Source/Niagara/Private/NiagaraGpuComputeDispatch.cpp b/Engine/Plugins/FX/Niagara/Source/Niagara/Private/NiagaraGpuComputeDispatch.cpp
--- a/Engine/Plugins/FX/Niagara/Source/Niagara/Private/NiagaraGpuComputeDispatch.cpp
+++ b/Engine/Plugins/FX/Niagara/Source/Niagara/Private/NiagaraGpuComputeDispatch.cpp
@@ -1901,10 +1901,7 @@ void FNiagaraGpuComputeDispatch::PreInitViews(FRDGBuilder& GraphBuilder, bool bA
 
 	// Add pass to begin the gpu profiler frame
 #if WITH_NIAGARA_GPU_PROFILER
-	AddPass(GraphBuilder, RDG_EVENT_NAME("Niagara::GPUProfiler_BeginFrame"), [this](FRHICommandListImmediate& RHICmdList)
-	{
-		GPUProfilerPtr->BeginFrame(RHICmdList);
-	});
+	GPUProfilerPtr->BeginFrame(GraphBuilder.RHICmdList);
 #endif
 
 	// Reset the list of GPUSort tasks and release any resources they hold on to.
@@ -2064,21 +2061,23 @@ void FNiagaraGpuComputeDispatch::PostRenderOpaque(FRDGBuilder& GraphBuilder, TCo
 					check(!GPUInstanceCounterManager.HasPendingGPUReadback());
 					GPUInstanceCounterManager.EnqueueGPUReadback(RHICmdList);
 				}
-
-			#if WITH_NIAGARA_GPU_PROFILER
-				GPUProfilerPtr->EndFrame(RHICmdList);
-			#endif
 			}
 		);
 	}
 
-	if (bAllowGPUParticleUpdate)
-	{
-		GraphBuilder.AddPostExecuteCallback([this]
+	GraphBuilder.AddPostExecuteCallback(
+		[this, &RHICmdList=GraphBuilder.RHICmdList, bAllowGPUParticleUpdate]
 		{
-			FinishDispatches();
-		});
-	}
+		#if WITH_NIAGARA_GPU_PROFILER
+			GPUProfilerPtr->EndFrame(RHICmdList);
+		#endif
+		
+			if (bAllowGPUParticleUpdate)
+			{
+				FinishDispatches();
+			}
+		}
+	);
 
 	bRequiresReadback = false;
 
