Fix out of bound access caused by both sides of ternary operator being executed in HLSL.

#rnx

diff --git a/Engine/Plugins/FX/Niagara/Shaders/Private/Ribbons/NiagaraRibbonCommon.ush b/Engine/Plugins/FX/Niagara/Shaders/Private/Ribbons/NiagaraRibbonCommon.ush
--- a/Engine/Plugins/FX/Niagara/Shaders/Private/Ribbons/NiagaraRibbonCommon.ush
+++ b/Engine/Plugins/FX/Niagara/Shaders/Private/Ribbons/NiagaraRibbonCommon.ush
@@ -96,7 +96,11 @@ int RibbonLinkOrderDataOffset;
 // Gets the num particles from the indirect buffer if available, or if not from the fallback direct input
 int GetTotalNumParticles()
 {
-	return EmitterParticleCountsBufferOffset != INDEX_NONE? EmitterParticleCountsBuffer[EmitterParticleCountsBufferOffset] : TotalNumParticlesDirect;
+	if (EmitterParticleCountsBufferOffset != INDEX_NONE)
+	{
+		return EmitterParticleCountsBuffer[EmitterParticleCountsBufferOffset];
+	}
+	return TotalNumParticlesDirect;
 }
 
 
