- Combine stateless spawn infos if possible with each loop
- Since we have limited space to store spawn infos this avoid running out with long lifetimes and short loops

diff --git a/Engine/Plugins/FX/Niagara/Source/Niagara/Private/Stateless/NiagaraStatelessEmitterInstance.cpp b/Engine/Plugins/FX/Niagara/Source/Niagara/Private/Stateless/NiagaraStatelessEmitterInstance.cpp
--- a/Engine/Plugins/FX/Niagara/Source/Niagara/Private/Stateless/NiagaraStatelessEmitterInstance.cpp
+++ b/Engine/Plugins/FX/Niagara/Source/Niagara/Private/Stateless/NiagaraStatelessEmitterInstance.cpp
@@ -414,12 +414,34 @@ void FNiagaraStatelessEmitterInstance::InitSpawnInfosForLoop()
 			continue;
 		}
 
-		FNiagaraStatelessRuntimeSpawnInfo& NewSpawnInfo = SpawnInfos.AddDefaulted_GetRef();
-		NewSpawnInfo.Type			= ENiagaraStatelessSpawnInfoType::Rate;
-		NewSpawnInfo.UniqueOffset	= UniqueIndexOffset;
-		NewSpawnInfo.SpawnTimeStart	= SpawnTime;
-		NewSpawnInfo.SpawnTimeEnd	= CurrentLoopAgeEnd;
-		NewSpawnInfo.Rate			= SpawnInfo.Rate;
+		// Try and append to the last info in the list if it's a rate type
+		// We do this to reduce the number of spawn infos in the common case of having a single rate info
+		bool bDidAppend = false;
+		if ( SpawnInfos.Num() > 0 )
+		{
+			FNiagaraStatelessRuntimeSpawnInfo& ExistingInfo = SpawnInfos.Last();
+			//-TODO: Validate time start better in the case of active / deactivate
+			if ( (ExistingInfo.Type == ENiagaraStatelessSpawnInfoType::Rate) && (ExistingInfo.Rate == SpawnInfo.Rate) && (ExistingInfo.SpawnTimeStart == SpawnInfo.SpawnTime))
+			{
+				const float ActiveDuration = ExistingInfo.SpawnTimeEnd - ExistingInfo.SpawnTimeStart;
+				const int32 NumSpawned = FMath::FloorToInt(ActiveDuration * ExistingInfo.Rate);
+				if (ExistingInfo.UniqueOffset + NumSpawned == UniqueIndexOffset)
+				{
+					ExistingInfo.SpawnTimeEnd = CurrentLoopAgeEnd;
+					bDidAppend = true;
+				}
+			}
+		}
+
+		if (!bDidAppend)
+		{
+			FNiagaraStatelessRuntimeSpawnInfo& NewSpawnInfo = SpawnInfos.AddDefaulted_GetRef();
+			NewSpawnInfo.Type			= ENiagaraStatelessSpawnInfoType::Rate;
+			NewSpawnInfo.UniqueOffset	= UniqueIndexOffset;
+			NewSpawnInfo.SpawnTimeStart	= SpawnTime;
+			NewSpawnInfo.SpawnTimeEnd	= CurrentLoopAgeEnd;
+			NewSpawnInfo.Rate			= SpawnInfo.Rate;
+		}
 
 		const float ActiveDuration	= CurrentLoopAgeEnd - SpawnTime;
 		const int32 NumSpawned		= FMath::FloorToInt(ActiveDuration * SpawnInfo.Rate);
