- Clamp loop duration to avoid infinite loop

diff --git a/Engine/Plugins/FX/Niagara/Source/Niagara/Private/Stateless/NiagaraStatelessEmitter.cpp b/Engine/Plugins/FX/Niagara/Source/Niagara/Private/Stateless/NiagaraStatelessEmitter.cpp
--- a/Engine/Plugins/FX/Niagara/Source/Niagara/Private/Stateless/NiagaraStatelessEmitter.cpp
+++ b/Engine/Plugins/FX/Niagara/Source/Niagara/Private/Stateless/NiagaraStatelessEmitter.cpp
@@ -215,6 +215,9 @@ void UNiagaraStatelessEmitter::CacheFromCompiledData()
 	// Setup emitter state
 	StatelessEmitterData->EmitterState = EmitterState;
 
+	StatelessEmitterData->EmitterState.LoopDuration.Min = FMath::Max(StatelessEmitterData->EmitterState.LoopDuration.Min, UE_KINDA_SMALL_NUMBER);
+	StatelessEmitterData->EmitterState.LoopDuration.Max = FMath::Max(StatelessEmitterData->EmitterState.LoopDuration.Max, UE_KINDA_SMALL_NUMBER);
+
 	// Find lifetime values
 	//-Note: We could abstact this out to be a more general modules implements Lifetime but this mirrors core Niagara where you always have Initialize Particle in 99% of cases
 	if (StatelessEmitterData->EmitterTemplate)
