improvements to SVT DI

diff --git a/Engine/Plugins/FX/Niagara/Shaders/Private/NiagaraDataInterfaceSparseVolumeTextureTemplate.ush b/Engine/Plugins/FX/Niagara/Shaders/Private/NiagaraDataInterfaceSparseVolumeTextureTemplate.ush
--- a/Engine/Plugins/FX/Niagara/Shaders/Private/NiagaraDataInterfaceSparseVolumeTextureTemplate.ush
+++ b/Engine/Plugins/FX/Niagara/Shaders/Private/NiagaraDataInterfaceSparseVolumeTextureTemplate.ush
@@ -10,6 +10,7 @@ uint4				{ParameterName}_PackedUniforms0;
 uint4				{ParameterName}_PackedUniforms1;
 int3				{ParameterName}_TextureSize;
 int					{ParameterName}_MipLevels;
+int					{ParameterName}_NumFrames;
 
 
 void LoadSparseVolumeTexture_{ParameterName}(in int TexelX, in int TexelY, in int TexelZ, in int MipLevel, out float4 OutAttributesA, out float4 OutAttributesB)
@@ -28,14 +29,19 @@ void SampleSparseVolumeTexture_{ParameterName}(in float3 UVW, in int MipLevel, o
 	OutAttributesB = SparseVolumeTextureSamplePhysicalTileData({ParameterName}_PhysicalTileDataATexture, {ParameterName}_PhysicalTileDataBTexture, {ParameterName}_TileDataTextureSampler, VoxelUVW, 1);
 }
 
-void GetTextureDimensions_{ParameterName}(int MipLevel, out float3 OutValue)
+void GetTextureDimensions_{ParameterName}(int MipLevel, out int SizeX, out int SizeY, out int SizeZ)
 {
-	OutValue.x = float(max({ParameterName}_TextureSize.x >> MipLevel, 1));
-	OutValue.y = float(max({ParameterName}_TextureSize.y >> MipLevel, 1));
-	OutValue.z = float(max({ParameterName}_TextureSize.z >> MipLevel, 1));
+	SizeX = float(max({ParameterName}_TextureSize.x >> MipLevel, 1));
+	SizeY = float(max({ParameterName}_TextureSize.y >> MipLevel, 1));
+	SizeZ = float(max({ParameterName}_TextureSize.z >> MipLevel, 1));
 }
 
 void GetNumMipLevels_{ParameterName}(out int OutMipLevels)
 {
 	OutMipLevels = {ParameterName}_MipLevels;
 }
+
+void GetNumFrames_{ParameterName}(out int OutNumFrames)
+{
+	OutNumFrames = {ParameterName}_NumFrames;
+}
\ No newline at end of file
diff --git a/Engine/Plugins/FX/Niagara/Source/Niagara/Classes/NiagaraDataInterfaceSparseVolumeTexture.h b/Engine/Plugins/FX/Niagara/Source/Niagara/Classes/NiagaraDataInterfaceSparseVolumeTexture.h
--- a/Engine/Plugins/FX/Niagara/Source/Niagara/Classes/NiagaraDataInterfaceSparseVolumeTexture.h
+++ b/Engine/Plugins/FX/Niagara/Source/Niagara/Classes/NiagaraDataInterfaceSparseVolumeTexture.h
@@ -24,6 +24,7 @@ class UNiagaraDataInterfaceSparseVolumeTexture : public UNiagaraDataInterface
 		SHADER_PARAMETER(FUintVector4, PackedUniforms1)
 		SHADER_PARAMETER(FIntVector3, TextureSize)
 		SHADER_PARAMETER(int32, MipLevels)
+		SHADER_PARAMETER(int32, NumFrames)
 	END_SHADER_PARAMETER_STRUCT()
 
 public:
@@ -32,6 +33,9 @@ public:
 
 	UPROPERTY(EditAnywhere, Category = "SparseVolumeTexture", meta = (ToolTip = "When valid the user parameter is used as the texture rather than the one on the data interface"))
 	FNiagaraUserParameterBinding SparseVolumeTextureUserParameter;
+	
+	UPROPERTY(EditAnywhere, Category = "SparseVolumeTexture")
+	bool BlockingStreamingRequests;
 
 	//UObject Interface
 	NIAGARA_API virtual void PostInitProperties() override;
@@ -65,6 +69,8 @@ public:
 
 	NIAGARA_API void VMGetTextureDimensions(FVectorVMExternalFunctionContext& Context);
 	NIAGARA_API void VMGetNumMipLevels(FVectorVMExternalFunctionContext& Context);
+	NIAGARA_API void VMRequestSparseVolumeTextureFrame(FVectorVMExternalFunctionContext& Context);	
+	NIAGARA_API void VMGetNumFrames(FVectorVMExternalFunctionContext& Context);
 
 	NIAGARA_API void SetTexture(USparseVolumeTexture* InSparseVolumeTexture);
 
@@ -80,4 +86,6 @@ protected:
 	static NIAGARA_API const FName SampleSparseVolumeTextureName;
 	static NIAGARA_API const FName GetTextureDimensionsName;
 	static NIAGARA_API const FName GetNumMipLevelsName;
+	static NIAGARA_API const FName RequestSparseVolumeTextureFrameName;
+	static NIAGARA_API const FName GetNumFramesName;
 };
diff --git a/Engine/Plugins/FX/Niagara/Source/Niagara/Private/NiagaraDataInterfaceSparseVolumeTexture.cpp b/Engine/Plugins/FX/Niagara/Source/Niagara/Private/NiagaraDataInterfaceSparseVolumeTexture.cpp
--- a/Engine/Plugins/FX/Niagara/Source/Niagara/Private/NiagaraDataInterfaceSparseVolumeTexture.cpp
+++ b/Engine/Plugins/FX/Niagara/Source/Niagara/Private/NiagaraDataInterfaceSparseVolumeTexture.cpp
@@ -11,6 +11,7 @@
 #include "RHIStaticStates.h"
 #include "ShaderCompilerCore.h"
 #include "GlobalRenderResources.h"
+#include "SparseVolumeTexture/ISparseVolumeTextureStreamingManager.h"
 
 #include UE_INLINE_GENERATED_CPP_BY_NAME(NiagaraDataInterfaceSparseVolumeTexture)
 
@@ -21,6 +22,8 @@ const FName UNiagaraDataInterfaceSparseVolumeTexture::LoadSparseVolumeTextureNam
 const FName UNiagaraDataInterfaceSparseVolumeTexture::SampleSparseVolumeTextureName(TEXT("SampleSparseVolumeTexture"));
 const FName UNiagaraDataInterfaceSparseVolumeTexture::GetTextureDimensionsName(TEXT("GetSparseVolumeTextureDimensions"));
 const FName UNiagaraDataInterfaceSparseVolumeTexture::GetNumMipLevelsName(TEXT("GetSparseVolumeTextureNumMipLevels"));
+const FName UNiagaraDataInterfaceSparseVolumeTexture::RequestSparseVolumeTextureFrameName(TEXT("RequestSparseVolumeTextureFrame"));
+const FName UNiagaraDataInterfaceSparseVolumeTexture::GetNumFramesName(TEXT("GetSparseVolumeTextureNumFrames"));
 
 struct FNDISparseVolumeTextureFunctionVersion
 {
@@ -40,6 +43,9 @@ struct FNDISparseVolumeTextureInstanceData_GameThread
 	FIntVector3 CurrentTextureSize = FIntVector3::ZeroValue;
 	int32 CurrentTextureMipLevels = 0;
 	FNiagaraParameterDirectBinding<UObject*> UserParamBinding;
+	int32 CurrentFrame = 0;
+	float FrameRate = 30.f;
+	int32 NumFrames = 0;
 };
 
 struct FNDISparseVolumeTextureInstanceData_RenderThread
@@ -47,6 +53,7 @@ struct FNDISparseVolumeTextureInstanceData_RenderThread
 	const UE::SVT::FTextureRenderResources* RenderResources = nullptr;
 	FIntVector3 TextureSize = FIntVector3::ZeroValue;
 	int32 MipLevels = 0;
+	int32 NumFrames = 0;
 };
 
 struct FNiagaraDataInterfaceProxySparseVolumeTexture : public FNiagaraDataInterfaceProxy
@@ -60,6 +67,7 @@ struct FNiagaraDataInterfaceProxySparseVolumeTexture : public FNiagaraDataInterf
 UNiagaraDataInterfaceSparseVolumeTexture::UNiagaraDataInterfaceSparseVolumeTexture(FObjectInitializer const& ObjectInitializer)
 	: Super(ObjectInitializer)
 	, SparseVolumeTexture(nullptr)
+	, BlockingStreamingRequests(true)
 {
 	Proxy.Reset(new FNiagaraDataInterfaceProxySparseVolumeTexture());
 
@@ -124,19 +132,61 @@ void UNiagaraDataInterfaceSparseVolumeTexture::GetFunctionsInternal(TArray<FNiag
 	{
 		FNiagaraFunctionSignature& Sig = OutFunctions.AddDefaulted_GetRef();
 		Sig.Name = GetTextureDimensionsName;
+		Sig.bMemberFunction = true;
+		Sig.bRequiresContext = false;
 		Sig.bSupportsCPU = true;
+		Sig.bSupportsGPU = true;
+		Sig.SetFunctionVersion(FNDISparseVolumeTextureFunctionVersion::LatestVersion);
+		Sig.Inputs.Emplace(FNiagaraTypeDefinition(GetClass()), TEXT("SparseVolumeTexture"));		
 		Sig.Inputs.Emplace(FNiagaraTypeDefinition::GetIntDef(), TEXT("MipLevel"));
-		Sig.Outputs.Emplace(FNiagaraTypeDefinition::GetVec3Def(), TEXT("Dimensions3D"));
+		Sig.Outputs.Emplace(FNiagaraTypeDefinition::GetIntDef(), TEXT("SizeX"));
+		Sig.Outputs.Emplace(FNiagaraTypeDefinition::GetIntDef(), TEXT("SizeY"));
+		Sig.Outputs.Emplace(FNiagaraTypeDefinition::GetIntDef(), TEXT("SizeZ"));
 		Sig.SetDescription(LOCTEXT("SparseVolumeTextureDimsDesc", "Get the dimensions of the provided mip level."));
 	}
 
 	{
 		FNiagaraFunctionSignature& Sig = OutFunctions.AddDefaulted_GetRef();
 		Sig.Name = GetNumMipLevelsName;
+		Sig.bMemberFunction = true;
+		Sig.bRequiresContext = false;
 		Sig.bSupportsCPU = true;
+		Sig.bSupportsGPU = true;
+		Sig.SetFunctionVersion(FNDISparseVolumeTextureFunctionVersion::LatestVersion);
+		Sig.Inputs.Emplace(FNiagaraTypeDefinition(GetClass()), TEXT("SparseVolumeTexture"));		
 		Sig.Outputs.Emplace(FNiagaraTypeDefinition::GetIntDef(), TEXT("NumMipLevels"));
 		Sig.SetDescription(LOCTEXT("SparseVolumeGetNumMipLevelsDesc", "Get the number of mip levels."));
 	}
+
+	{
+		FNiagaraFunctionSignature& Sig = OutFunctions.AddDefaulted_GetRef();
+		Sig.Name = RequestSparseVolumeTextureFrameName;
+		Sig.ModuleUsageBitmask = ENiagaraScriptUsageMask::Emitter | ENiagaraScriptUsageMask::System;		
+		Sig.bMemberFunction = true;
+		Sig.bRequiresExecPin = true;
+		Sig.bRequiresContext = false;
+		Sig.bSupportsCPU = true;
+		Sig.bSupportsGPU = false;
+		Sig.SetFunctionVersion(FNDISparseVolumeTextureFunctionVersion::LatestVersion);
+		Sig.Inputs.Emplace(FNiagaraTypeDefinition(GetClass()), TEXT("SparseVolumeTexture"));		
+		Sig.Inputs.Emplace(FNiagaraTypeDefinition::GetFloatDef(), TEXT("FrameRate"));
+		Sig.Inputs.Emplace(FNiagaraTypeDefinition::GetIntDef(), TEXT("Frame"));
+		Sig.Outputs.Add(FNiagaraVariable(FNiagaraTypeDefinition::GetBoolDef(), TEXT("Success")));
+		Sig.SetDescription(LOCTEXT("RequestSparseVolumeTextureFrameDesc", "Queue up the frame to load on tick"));
+	}
+
+	{
+		FNiagaraFunctionSignature& Sig = OutFunctions.AddDefaulted_GetRef();
+		Sig.Name = GetNumFramesName;
+		Sig.bMemberFunction = true;
+		Sig.bRequiresContext = false;
+		Sig.bSupportsCPU = true;
+		Sig.bSupportsGPU = true;
+		Sig.SetFunctionVersion(FNDISparseVolumeTextureFunctionVersion::LatestVersion);
+		Sig.Inputs.Emplace(FNiagaraTypeDefinition(GetClass()), TEXT("SparseVolumeTexture"));		
+		Sig.Outputs.Emplace(FNiagaraTypeDefinition::GetIntDef(), TEXT("NumFrames"));		
+		Sig.SetDescription(LOCTEXT("SparseVolumeNumFramesDesc", "Get the number of frames."));
+	}
 }
 #endif
 
@@ -152,6 +202,16 @@ void UNiagaraDataInterfaceSparseVolumeTexture::GetVMExternalFunction(const FVMEx
 		check(BindingInfo.GetNumInputs() == 1 && BindingInfo.GetNumOutputs() == 1);
 		OutFunc = FVMExternalFunction::CreateUObject(this, &UNiagaraDataInterfaceSparseVolumeTexture::VMGetNumMipLevels);
 	}
+	else if (BindingInfo.Name == RequestSparseVolumeTextureFrameName)
+	{
+		check(BindingInfo.GetNumInputs() == 3 && BindingInfo.GetNumOutputs() == 1);
+		OutFunc = FVMExternalFunction::CreateUObject(this, &UNiagaraDataInterfaceSparseVolumeTexture::VMRequestSparseVolumeTextureFrame);
+	}	
+	else if (BindingInfo.Name == GetNumFramesName)
+	{
+		check(BindingInfo.GetNumInputs() == 1 && BindingInfo.GetNumOutputs() == 1);
+		OutFunc = FVMExternalFunction::CreateUObject(this, &UNiagaraDataInterfaceSparseVolumeTexture::VMGetNumFrames);
+	}
 }
 
 int32 UNiagaraDataInterfaceSparseVolumeTexture::PerInstanceDataSize() const
@@ -183,45 +243,55 @@ void UNiagaraDataInterfaceSparseVolumeTexture::DestroyPerInstanceData(void* PerI
 bool UNiagaraDataInterfaceSparseVolumeTexture::PerInstanceTick(void* PerInstanceData, FNiagaraSystemInstance* SystemInstance, float DeltaSeconds)
 {
 	FNDISparseVolumeTextureInstanceData_GameThread* InstanceData = static_cast<FNDISparseVolumeTextureInstanceData_GameThread*>(PerInstanceData);
+	
+	USparseVolumeTexture* CurrentTexture = InstanceData->UserParamBinding.GetValueOrDefault<USparseVolumeTexture>(SparseVolumeTexture);	
 
-	USparseVolumeTexture* CurrentTexture = InstanceData->UserParamBinding.GetValueOrDefault<USparseVolumeTexture>(SparseVolumeTexture);
-	const UE::SVT::FTextureRenderResources* CurrentRenderResources = nullptr;
-	FIntVector3 CurrentTextureSize = FIntVector3::ZeroValue;
-	int32 CurrentTextureMipLevels = 0;
 	if (CurrentTexture)
 	{
-		CurrentRenderResources = CurrentTexture->GetTextureRenderResources();
-		CurrentTextureSize = FIntVector3(CurrentTexture->GetVolumeResolution());
-		CurrentTextureMipLevels = CurrentTexture->GetNumMipLevels();
-	}
+		InstanceData->NumFrames = CurrentTexture->GetNumFrames();
 
-	if ((InstanceData->CurrentTexture != CurrentTexture) 
-		|| (InstanceData->CurrentRenderResources != CurrentRenderResources)
-		|| (InstanceData->CurrentTextureSize != CurrentTextureSize) 
-		|| (InstanceData->CurrentTextureMipLevels != CurrentTextureMipLevels))
-	{
-		InstanceData->CurrentTexture = CurrentTexture;
-		InstanceData->CurrentRenderResources = CurrentRenderResources;
-		InstanceData->CurrentTextureSize = CurrentTextureSize;
-		InstanceData->CurrentTextureMipLevels = CurrentTextureMipLevels;
-
-		ENQUEUE_RENDER_COMMAND(NDISparseVolumeTexture_UpdateInstance)
-		(
-			[RT_Proxy = GetProxyAs<FNiagaraDataInterfaceProxySparseVolumeTexture>(),
-			RT_InstanceID = SystemInstance->GetId(),
-			RT_RenderResources = CurrentRenderResources,
-			RT_TextureSize = CurrentTextureSize, 
-			RT_MipLevels = CurrentTextureMipLevels] 
-		(FRHICommandListImmediate&)
+		// todo: we only support querying mip 0 for now
+		USparseVolumeTextureFrame* SparseVolumeTextureFrame = 
+			USparseVolumeTextureFrame::GetFrameAndIssueStreamingRequest(CurrentTexture, GetTypeHash(InstanceData), InstanceData->FrameRate, InstanceData->CurrentFrame, 0, BlockingStreamingRequests, true);
+		UE::SVT::GetStreamingManager().Update_GameThread();
+
+		if (SparseVolumeTextureFrame)
+		{
+			const UE::SVT::FTextureRenderResources* CurrentRenderResources = SparseVolumeTextureFrame->GetTextureRenderResources();
+			FIntVector3 CurrentTextureSize = FIntVector3(SparseVolumeTextureFrame->GetVolumeResolution());
+			int32 CurrentTextureMipLevels = SparseVolumeTextureFrame->GetNumMipLevels();
+
+			// if the loaded SVT frame is different from what the instance data specifies
+			if ((InstanceData->CurrentTexture != SparseVolumeTextureFrame)
+				|| (InstanceData->CurrentRenderResources != CurrentRenderResources)
+				|| (InstanceData->CurrentTextureSize != CurrentTextureSize) 
+				|| (InstanceData->CurrentTextureMipLevels != CurrentTextureMipLevels))
 			{
-				FNDISparseVolumeTextureInstanceData_RenderThread& RTInstanceData = RT_Proxy->InstanceData_RT.FindOrAdd(RT_InstanceID);
-				RTInstanceData.RenderResources = RT_RenderResources;
-				RTInstanceData.TextureSize = RT_TextureSize;
-				RTInstanceData.MipLevels = RT_MipLevels;
+				InstanceData->CurrentTexture = SparseVolumeTextureFrame;
+				InstanceData->CurrentRenderResources = CurrentRenderResources;
+				InstanceData->CurrentTextureSize = CurrentTextureSize;
+				InstanceData->CurrentTextureMipLevels = CurrentTextureMipLevels;
+
+				ENQUEUE_RENDER_COMMAND(NDISparseVolumeTexture_UpdateInstance)
+				(
+					[RT_Proxy = GetProxyAs<FNiagaraDataInterfaceProxySparseVolumeTexture>(),
+					RT_InstanceID = SystemInstance->GetId(),
+					RT_RenderResources = CurrentRenderResources,
+					RT_TextureSize = CurrentTextureSize, 
+					RT_MipLevels = CurrentTextureMipLevels,
+					RT_NumFrames = InstanceData->NumFrames]
+				(FRHICommandListImmediate&)
+					{
+						FNDISparseVolumeTextureInstanceData_RenderThread& RTInstanceData = RT_Proxy->InstanceData_RT.FindOrAdd(RT_InstanceID);
+						RTInstanceData.RenderResources = RT_RenderResources;
+						RTInstanceData.TextureSize = RT_TextureSize;
+						RTInstanceData.MipLevels = RT_MipLevels;
+						RTInstanceData.NumFrames = RT_NumFrames;
+					}
+				);
 			}
-		);
+		}
 	}
-
 	return false;
 }
 
@@ -232,7 +302,7 @@ bool UNiagaraDataInterfaceSparseVolumeTexture::Equals(const UNiagaraDataInterfac
 		return false;
 	}
 	const UNiagaraDataInterfaceSparseVolumeTexture* OtherSparseVolumeTexture = CastChecked<const UNiagaraDataInterfaceSparseVolumeTexture>(Other);
-	return OtherSparseVolumeTexture->SparseVolumeTexture == SparseVolumeTexture && OtherSparseVolumeTexture->SparseVolumeTextureUserParameter == SparseVolumeTextureUserParameter;
+	return OtherSparseVolumeTexture->SparseVolumeTexture == SparseVolumeTexture && OtherSparseVolumeTexture->SparseVolumeTextureUserParameter == SparseVolumeTextureUserParameter && OtherSparseVolumeTexture->BlockingStreamingRequests == BlockingStreamingRequests;
 }
 
 #if WITH_EDITORONLY_DATA
@@ -261,6 +331,7 @@ bool UNiagaraDataInterfaceSparseVolumeTexture::GetFunctionHLSL(const FNiagaraDat
 		SampleSparseVolumeTextureName,
 		GetTextureDimensionsName,
 		GetNumMipLevelsName,
+		GetNumFramesName,
 	};
 
 	return ValidGpuFunctions.Contains(FunctionInfo.DefinitionName);
@@ -291,6 +362,7 @@ void UNiagaraDataInterfaceSparseVolumeTexture::SetShaderParameters(const FNiagar
 	Parameters->PackedUniforms1 = FUintVector4();
 	Parameters->TextureSize = FIntVector3::ZeroValue;
 	Parameters->MipLevels = 0;
+	Parameters->NumFrames = 0;
 	
 	if (RTInstanceData && RTInstanceData->RenderResources)
 	{
@@ -304,6 +376,7 @@ void UNiagaraDataInterfaceSparseVolumeTexture::SetShaderParameters(const FNiagar
 		RTInstanceData->RenderResources->GetPackedUniforms(Parameters->PackedUniforms0, Parameters->PackedUniforms1);
 		Parameters->TextureSize = RTInstanceData->TextureSize;
 		Parameters->MipLevels = RTInstanceData->MipLevels;
+		Parameters->NumFrames = RTInstanceData->NumFrames;
 	}
 }
 
@@ -311,17 +384,17 @@ void UNiagaraDataInterfaceSparseVolumeTexture::VMGetTextureDimensions(FVectorVME
 {
 	VectorVM::FUserPtrHandler<FNDISparseVolumeTextureInstanceData_GameThread> InstData(Context);
 	FNDIInputParam<int32>		InMipLevel(Context);
-	FNDIOutputParam<FVector3f>	OutSize(Context);
+	FNDIOutputParam<int32>		OutSizeX(Context);
+	FNDIOutputParam<int32>		OutSizeY(Context);
+	FNDIOutputParam<int32>		OutSizeZ(Context);
 
 	for (int32 i = 0; i < Context.GetNumInstances(); ++i)
 	{
 		const int32 MipLevel = InMipLevel.GetAndAdvance();
-		const FVector3f TextureSize(
-			float(FMath::Max(InstData->CurrentTextureSize.X >> MipLevel, 1)),
-			float(FMath::Max(InstData->CurrentTextureSize.Y >> MipLevel, 1)),
-			float(FMath::Max(InstData->CurrentTextureSize.Z >> MipLevel, 1))
-		);
-		OutSize.SetAndAdvance(TextureSize);
+
+		OutSizeX.SetAndAdvance(FMath::Max(InstData->CurrentTextureSize.X >> MipLevel, 1));
+		OutSizeY.SetAndAdvance(FMath::Max(InstData->CurrentTextureSize.Y >> MipLevel, 1));
+		OutSizeZ.SetAndAdvance(FMath::Max(InstData->CurrentTextureSize.Z >> MipLevel, 1));
 	}
 }
 
@@ -336,6 +409,33 @@ void UNiagaraDataInterfaceSparseVolumeTexture::VMGetNumMipLevels(FVectorVMExtern
 	}
 }
 
+void UNiagaraDataInterfaceSparseVolumeTexture::VMRequestSparseVolumeTextureFrame(FVectorVMExternalFunctionContext& Context)
+{
+	VectorVM::FUserPtrHandler<FNDISparseVolumeTextureInstanceData_GameThread> InstData(Context);
+	FNDIInputParam<float>	InFrameRate(Context);
+	FNDIInputParam<int>		InFrame(Context);
+	FNDIOutputParam<bool> OutSuccess(Context);
+
+	InstData->FrameRate = InFrameRate.GetAndAdvance();
+	InstData->CurrentFrame = InFrame.GetAndAdvance();
+	
+	for (int32 i = 0; i < Context.GetNumInstances(); ++i)
+	{
+		OutSuccess.SetAndAdvance(true);
+	}
+}
+
+void UNiagaraDataInterfaceSparseVolumeTexture::VMGetNumFrames(FVectorVMExternalFunctionContext& Context)
+{
+	VectorVM::FUserPtrHandler<FNDISparseVolumeTextureInstanceData_GameThread> InstData(Context);
+	FNDIOutputParam<int32> OutNumFrames(Context);
+
+	for (int32 i = 0; i < Context.GetNumInstances(); ++i)
+	{
+		OutNumFrames.SetAndAdvance(InstData->NumFrames);
+	}
+}
+
 void UNiagaraDataInterfaceSparseVolumeTexture::SetTexture(USparseVolumeTexture* InSparseVolumeTexture)
 {
 	SparseVolumeTexture = InSparseVolumeTexture;
@@ -350,6 +450,7 @@ bool UNiagaraDataInterfaceSparseVolumeTexture::CopyToInternal(UNiagaraDataInterf
 	UNiagaraDataInterfaceSparseVolumeTexture* DestinationTexture = CastChecked<UNiagaraDataInterfaceSparseVolumeTexture>(Destination);
 	DestinationTexture->SparseVolumeTexture = SparseVolumeTexture;
 	DestinationTexture->SparseVolumeTextureUserParameter = SparseVolumeTextureUserParameter;
+	DestinationTexture->BlockingStreamingRequests = BlockingStreamingRequests;
 
 	return true;
 }
