- Make stateless spawn infos more like regular emitters
- Add action allows you to select from available types
- Infos can be enabled / disabled.  Disabled infos are automatically culled when converting from editor -> runtime data

diff --git a/Engine/Plugins/FX/Niagara/Source/Niagara/Internal/Stateless/NiagaraStatelessSpawnInfo.h b/Engine/Plugins/FX/Niagara/Source/Niagara/Internal/Stateless/NiagaraStatelessSpawnInfo.h
--- a/Engine/Plugins/FX/Niagara/Source/Niagara/Internal/Stateless/NiagaraStatelessSpawnInfo.h
+++ b/Engine/Plugins/FX/Niagara/Source/Niagara/Internal/Stateless/NiagaraStatelessSpawnInfo.h
@@ -37,7 +37,10 @@ struct FNiagaraStatelessSpawnInfo
 	FNiagaraDistributionRangeFloat Rate = FNiagaraDistributionRangeFloat(60.0f);
 
 	UPROPERTY(EditAnywhere, Category = "Spawn", meta = (InlineEditConditionToggle))
-	bool bSpawnProbabilityEnabled = false;
+	uint8 bEnabled : 1 = true;
+
+	UPROPERTY(EditAnywhere, Category = "Spawn", meta = (InlineEditConditionToggle))
+	uint8 bSpawnProbabilityEnabled : 1 = false;
 
 	UPROPERTY(EditAnywhere, Category = "Spawn", meta = (EditCondition = "bSpawnProbabilityEnabled", ClampMin = "0.0", ClampMax = "1.0", UIMin = "0.0", UIMax = "1.0"))
 	float SpawnProbability = 1.0f;
diff --git a/Engine/Plugins/FX/Niagara/Source/Niagara/Private/Stateless/NiagaraStatelessSpawnInfo.cpp b/Engine/Plugins/FX/Niagara/Source/Niagara/Private/Stateless/NiagaraStatelessSpawnInfo.cpp
--- a/Engine/Plugins/FX/Niagara/Source/Niagara/Private/Stateless/NiagaraStatelessSpawnInfo.cpp
+++ b/Engine/Plugins/FX/Niagara/Source/Niagara/Private/Stateless/NiagaraStatelessSpawnInfo.cpp
@@ -4,16 +4,20 @@
 
 bool FNiagaraStatelessSpawnInfo::IsValid(TOptional<float> LoopDuration) const
 {
-	switch (Type)
+	if (bEnabled)
 	{
-		case ENiagaraStatelessSpawnInfoType::Burst:
-			return (Amount.Min + Amount.Max) > 0 && SpawnTime >= 0.0f && SpawnTime < LoopDuration.Get(SpawnTime + UE_SMALL_NUMBER);
+		switch (Type)
+		{
+			case ENiagaraStatelessSpawnInfoType::Burst:
+				return (Amount.Min + Amount.Max) > 0 && SpawnTime >= 0.0f && SpawnTime < LoopDuration.Get(SpawnTime + UE_SMALL_NUMBER);
 
-		case ENiagaraStatelessSpawnInfoType::Rate:
-			return (Rate.Min + Rate.Max) > 0.0f;
+			case ENiagaraStatelessSpawnInfoType::Rate:
+				return (Rate.Min + Rate.Max) > 0.0f;
 
-		default:
-			checkNoEntry();
-			return false;
+			default:
+				checkNoEntry();
+				return false;
+		}
 	}
+	return false;
 }
diff --git a/Engine/Plugins/FX/Niagara/Source/NiagaraEditor/Private/ViewModels/Stack/NiagaraStackStatelessEmitterSpawnGroup.cpp b/Engine/Plugins/FX/Niagara/Source/NiagaraEditor/Private/ViewModels/Stack/NiagaraStackStatelessEmitterSpawnGroup.cpp
--- a/Engine/Plugins/FX/Niagara/Source/NiagaraEditor/Private/ViewModels/Stack/NiagaraStackStatelessEmitterSpawnGroup.cpp
+++ b/Engine/Plugins/FX/Niagara/Source/NiagaraEditor/Private/ViewModels/Stack/NiagaraStackStatelessEmitterSpawnGroup.cpp
@@ -17,38 +17,67 @@
 
 #define LOCTEXT_NAMESPACE "NiagaraEmitterStatelessSpawnGroup"
 
+class FNiagaraStatelessEmitterAddSpawnInfoAction : public INiagaraStackItemGroupAddAction
+{
+public:
+	explicit FNiagaraStatelessEmitterAddSpawnInfoAction(ENiagaraStatelessSpawnInfoType InSpawnInfoType)
+		: SpawnInfoType(InSpawnInfoType)
+		, DisplayName(UNiagaraStackStatelessEmitterSpawnItem::GetDisplayName(InSpawnInfoType))
+	{
+	}
+
+	ENiagaraStatelessSpawnInfoType GetSpawnInfoType() const { return SpawnInfoType; }
+
+	virtual TArray<FString> GetCategories() const override { return Categories; }
+	virtual FText GetDisplayName() const override { return DisplayName; }
+	virtual FText GetDescription() const override { return FText(); }
+	virtual FText GetKeywords() const override { return FText(); }
+
+private:
+	ENiagaraStatelessSpawnInfoType SpawnInfoType;
+	TWeakObjectPtr<UNiagaraStatelessModule> StatelessModuleWeak;
+	TArray<FString> Categories;
+	FText DisplayName;
+};
+
 class FNiagaraStackStatelessEmitterSpawnGroupAddUtilities : public TNiagaraStackItemGroupAddUtilities<FGuid>
 {
 public:
 	FNiagaraStackStatelessEmitterSpawnGroupAddUtilities(UNiagaraStatelessEmitter* StatelessEmitter, FOnItemAdded InOnItemAdded)
-		: TNiagaraStackItemGroupAddUtilities(LOCTEXT("AddUtilitiesName", "Spawn Data"), EAddMode::AddDirectly, false, false, InOnItemAdded)
+		: TNiagaraStackItemGroupAddUtilities(LOCTEXT("AddUtilitiesName", "Spawn Data"), EAddMode::AddFromAction, true, false, InOnItemAdded)
 	{
 		StatelessEmitterWeak = StatelessEmitter;
 	}
 
 	virtual void AddItemDirectly() override 
 	{ 
+		unimplemented();
+	}
+
+	virtual void GenerateAddActions(TArray<TSharedRef<INiagaraStackItemGroupAddAction>>& OutAddActions, const FNiagaraStackItemGroupAddOptions& AddProperties) const override
+	{
+		OutAddActions.Emplace(MakeShared<FNiagaraStatelessEmitterAddSpawnInfoAction>(ENiagaraStatelessSpawnInfoType::Burst));
+		OutAddActions.Emplace(MakeShared<FNiagaraStatelessEmitterAddSpawnInfoAction>(ENiagaraStatelessSpawnInfoType::Rate));
+	}
+
+	virtual void ExecuteAddAction(TSharedRef<INiagaraStackItemGroupAddAction> AddAction, int32 TargetIndex) override
+	{
 		UNiagaraStatelessEmitter* StatelessEmitter = StatelessEmitterWeak.Get();
 		if (StatelessEmitter != nullptr)
 		{
+			TSharedRef<FNiagaraStatelessEmitterAddSpawnInfoAction> AddSpawnInfoAction = StaticCastSharedRef<FNiagaraStatelessEmitterAddSpawnInfoAction>(AddAction);
+
 			FScopedTransaction ScopedTransaction(LOCTEXT("AddNewSpawnInfoTransaction", "Add new spawn data"));
 			StatelessEmitter->Modify();
+
 			FNiagaraStatelessSpawnInfo& SpawnInfo = StatelessEmitter->AddSpawnInfo();
+			SpawnInfo.Type = AddSpawnInfoAction->GetSpawnInfoType();
+
 			SpawnInfo.SourceId = FGuid::NewGuid();
 			OnItemAdded.ExecuteIfBound(SpawnInfo.SourceId);
 		}
 	}
 
-	virtual void GenerateAddActions(TArray<TSharedRef<INiagaraStackItemGroupAddAction>>& OutAddActions, const FNiagaraStackItemGroupAddOptions& AddProperties) const override
-	{
-		unimplemented();
-	}
-
-	virtual void ExecuteAddAction(TSharedRef<INiagaraStackItemGroupAddAction> AddAction, int32 TargetIndex) override
-	{
-		unimplemented();
-	}
-
 private:
 	TWeakObjectPtr<UNiagaraStatelessEmitter> StatelessEmitterWeak;
 };
@@ -134,24 +163,23 @@ void UNiagaraStackStatelessEmitterSpawnItem::Initialize(FRequiredEntryData InReq
 	OnDataObjectModified().AddUObject(this, &UNiagaraStackStatelessEmitterSpawnItem::OnSpawnInfoModified);
 }
 
-FText UNiagaraStackStatelessEmitterSpawnItem::GetDisplayName() const
+FText UNiagaraStackStatelessEmitterSpawnItem::GetDisplayName(ENiagaraStatelessSpawnInfoType SpawnInfoType)
 {
-	const FNiagaraStatelessSpawnInfo* SpawnInfo = GetSpawnInfo();
-	const ENiagaraStatelessSpawnInfoType SpawnInfoType = SpawnInfo ? SpawnInfo->Type : ENiagaraStatelessSpawnInfoType::Burst;
 	switch (SpawnInfoType)
 	{
-		case ENiagaraStatelessSpawnInfoType::Burst:
-			return LOCTEXT("EmitterSpawnBurstDisplayName", "Spawn Burst Instantaneous");
-
-		case ENiagaraStatelessSpawnInfoType::Rate:
-			return LOCTEXT("EmitterSpawnRateDisplayName", "Spawn Rate");
-
-		default:
-			checkNoEntry();
-			return LOCTEXT("EmitterSpawnUnknownDisplayName", "Unknown");
+		case ENiagaraStatelessSpawnInfoType::Burst:	return LOCTEXT("EmitterSpawnBurstDisplayName", "Spawn Burst Instantaneous");
+		case ENiagaraStatelessSpawnInfoType::Rate:	return LOCTEXT("EmitterSpawnRateDisplayName", "Spawn Rate");
+		default:									checkNoEntry();	return LOCTEXT("EmitterSpawnUnknownDisplayName", "Unknown");
 	}
 }
 
+FText UNiagaraStackStatelessEmitterSpawnItem::GetDisplayName() const
+{
+	const FNiagaraStatelessSpawnInfo* SpawnInfo = GetSpawnInfo();
+	const ENiagaraStatelessSpawnInfoType SpawnInfoType = SpawnInfo ? SpawnInfo->Type : ENiagaraStatelessSpawnInfoType::Burst;
+	return GetDisplayName(SpawnInfoType);
+}
+
 FGuid UNiagaraStackStatelessEmitterSpawnItem::GetSelectionId() const
 {
 	return SourceId;
@@ -173,6 +201,17 @@ void UNiagaraStackStatelessEmitterSpawnItem::Delete()
 	OnRequestDeleteDelegate.ExecuteIfBound(SourceId);
 }
 
+bool UNiagaraStackStatelessEmitterSpawnItem::SupportsChangeEnabled() const
+{
+	return true;
+}
+
+bool UNiagaraStackStatelessEmitterSpawnItem::GetIsEnabled() const
+{
+	const FNiagaraStatelessSpawnInfo* SpawnInfo = GetSpawnInfo();
+	return SpawnInfo ? SpawnInfo->bEnabled : false;
+}
+
 void UNiagaraStackStatelessEmitterSpawnItem::GetHeaderValueHandlers(TArray<TSharedRef<INiagaraStackItemHeaderValueHandler>>& OutHeaderValueHandlers) const
 {
 	if (GetSpawnInfo() != nullptr)
@@ -240,6 +279,25 @@ void UNiagaraStackStatelessEmitterSpawnItem::RefreshChildrenInternal(const TArra
 	}
 }
 
+void UNiagaraStackStatelessEmitterSpawnItem::SetIsEnabledInternal(bool bInIsEnabled)
+{
+	UNiagaraStatelessEmitter* StatelessEmitter = StatelessEmitterWeak.Get();
+	FNiagaraStatelessSpawnInfo* SpawnInfo = GetSpawnInfo();
+	if (StatelessEmitter && SpawnInfo && SpawnInfo->bEnabled != bInIsEnabled)
+	{
+		FScopedTransaction ScopedTransaction(LOCTEXT("ChangeStatelessSpawnInfoEnabledTransaction", "Change spawn info enabled"));
+		StatelessEmitter->Modify();
+		SpawnInfo->bEnabled = bInIsEnabled;
+		StatelessEmitter->PostEditChange();
+
+		TArray<UObject*> ChangedObjects;
+		ChangedObjects.Add(StatelessEmitter);
+		OnDataObjectModified().Broadcast(ChangedObjects, ENiagaraDataObjectChange::Changed);
+
+		RefreshChildren();
+	}
+}
+
 FNiagaraStatelessSpawnInfo* UNiagaraStackStatelessEmitterSpawnItem::GetSpawnInfo() const
 {
 	UNiagaraStatelessEmitter* StatelessEmitter = StatelessEmitterWeak.Get();
diff --git a/Engine/Plugins/FX/Niagara/Source/NiagaraEditor/Public/ViewModels/Stack/NiagaraStackStatelessEmitterSpawnGroup.h b/Engine/Plugins/FX/Niagara/Source/NiagaraEditor/Public/ViewModels/Stack/NiagaraStackStatelessEmitterSpawnGroup.h
--- a/Engine/Plugins/FX/Niagara/Source/NiagaraEditor/Public/ViewModels/Stack/NiagaraStackStatelessEmitterSpawnGroup.h
+++ b/Engine/Plugins/FX/Niagara/Source/NiagaraEditor/Public/ViewModels/Stack/NiagaraStackStatelessEmitterSpawnGroup.h
@@ -10,6 +10,7 @@
 
 class FNiagaraStackItemPropertyHeaderValue;
 struct FNiagaraStatelessSpawnInfo;
+enum class ENiagaraStatelessSpawnInfoType;
 class UNiagaraStatelessEmitter;
 class UNiagaraStackObject;
 
@@ -61,6 +62,9 @@ public:
 	virtual FText GetDeleteTransactionText() const override;
 	virtual void Delete() override;
 
+	virtual bool SupportsChangeEnabled() const override;
+	virtual bool GetIsEnabled() const override;
+
 	virtual bool SupportsHeaderValues() const override { return true; }
 	virtual void GetHeaderValueHandlers(TArray<TSharedRef<INiagaraStackItemHeaderValueHandler>>& OutHeaderValueHandlers) const override;
 
@@ -70,9 +74,13 @@ public:
 
 	FOnRequestDelete& OnRequestDelete() { return OnRequestDeleteDelegate; }
 
+	static FText GetDisplayName(ENiagaraStatelessSpawnInfoType SpawnInfoType);
+
 protected:
 	virtual void RefreshChildrenInternal(const TArray<UNiagaraStackEntry*>& CurrentChildren, TArray<UNiagaraStackEntry*>& NewChildren, TArray<FStackIssue>& NewIssues) override;
 
+	virtual void SetIsEnabledInternal(bool bInIsEnabled) override;
+
 private:
 	FNiagaraStatelessSpawnInfo* GetSpawnInfo() const;
 
