Ray Tracing support in FNiagaraRendererMeshes using GPU Scene instead of using GPUTransformsSRV.

diff --git a/Engine/Plugins/FX/Niagara/Source/Niagara/Private/NiagaraRendererMeshes.cpp b/Engine/Plugins/FX/Niagara/Source/Niagara/Private/NiagaraRendererMeshes.cpp
--- a/Engine/Plugins/FX/Niagara/Source/Niagara/Private/NiagaraRendererMeshes.cpp
+++ b/Engine/Plugins/FX/Niagara/Source/Niagara/Private/NiagaraRendererMeshes.cpp
@@ -317,8 +317,7 @@ void FNiagaraRendererMeshes::PrepareParticleMeshRenderData(FParticleMeshRenderDa
 	ParticleMeshRenderData.Collector = &Collector;
 
 #if NIAGARA_ENABLE_GPU_SCENE_MESHES
-	ParticleMeshRenderData.bUseGPUScene = !bRayTracing
-		&& FeatureLevel > ERHIFeatureLevel::ES3_1	
+	ParticleMeshRenderData.bUseGPUScene = FeatureLevel > ERHIFeatureLevel::ES3_1	
 		&& UseGPUScene(SceneProxy->GetScene().GetShaderPlatform(), FeatureLevel);
 #else
 	ParticleMeshRenderData.bUseGPUScene = false;
@@ -1774,6 +1773,8 @@ void FNiagaraRendererMeshes::GetDynamicRayTracingInstances(FRayTracingMaterialGa
 #endif
 
 			RayTracingInstance.Materials.Add(MoveTemp(MeshBatch));
+
+			Context.RayTracingMeshResourceCollector.AddMesh(ViewIndex, RayTracingInstance.Materials.Last());
 		}
 
 		if (RayTracingInstance.Materials.Num() == 0 || LODModel.Sections.Num() != RayTracingInstance.Materials.Num())
@@ -1786,6 +1787,8 @@ void FNiagaraRendererMeshes::GetDynamicRayTracingInstances(FRayTracingMaterialGa
 		const bool bUseLocalSpace = UseLocalSpace(SceneProxy);
 		if (SourceMode == ENiagaraRendererSourceDataMode::Emitter)
 		{
+			// TODO: Get instance transform from GPU Scene when using ENiagaraRendererSourceDataMode::Emitter instead of using the logic below
+
 			FVector Pos = bUseLocalSpace ? FVector::ZeroVector : LocalTransform.GetOrigin();
 			FVector3f Scale{ 1.0f, 1.0f, 1.0f };
 			FQuat Rot = FQuat::Identity;
@@ -1854,160 +1857,8 @@ void FNiagaraRendererMeshes::GetDynamicRayTracingInstances(FRayTracingMaterialGa
 		}
 		else
 		{
-			TConstArrayView<FNiagaraRendererVariableInfo> VFVariables = ParticleMeshRenderData.RendererLayout->GetVFVariables_RenderThread();
-			if (SimTarget == ENiagaraSimTarget::CPUSim)
-			{
-				const int32 TotalFloatSize = ParticleMeshRenderData.RendererLayout->GetTotalFloatComponents_RenderThread() * ParticleMeshRenderData.SourceParticleData->GetNumInstances();
-				const int32 ComponentStrideDest = ParticleMeshRenderData.SourceParticleData->GetNumInstances() * sizeof(float);
-
-				//ENiagaraMeshVFLayout::Transform just contains a Quat, not the whole transform
-				const FNiagaraRendererVariableInfo& VarPositionInfo = VFVariables[ENiagaraMeshVFLayout::Position];
-				const FNiagaraRendererVariableInfo& VarScaleInfo = VFVariables[ENiagaraMeshVFLayout::Scale];
-				const FNiagaraRendererVariableInfo& VarTransformInfo = VFVariables[ENiagaraMeshVFLayout::Rotation];
-
-				const int32 PositionBaseCompOffset = VarPositionInfo.GetRawDatasetOffset();
-				const int32 ScaleBaseCompOffset = VarScaleInfo.GetRawDatasetOffset();
-				const int32 TransformBaseCompOffset = VarTransformInfo.GetRawDatasetOffset();
-
-				const float* RESTRICT PositionX = reinterpret_cast<const float*>(ParticleMeshRenderData.SourceParticleData->GetComponentPtrFloat(PositionBaseCompOffset));
-				const float* RESTRICT PositionY = reinterpret_cast<const float*>(ParticleMeshRenderData.SourceParticleData->GetComponentPtrFloat(PositionBaseCompOffset + 1));
-				const float* RESTRICT PositionZ = reinterpret_cast<const float*>(ParticleMeshRenderData.SourceParticleData->GetComponentPtrFloat(PositionBaseCompOffset + 2));
-
-				const float* RESTRICT ScaleX = reinterpret_cast<const float*>(ParticleMeshRenderData.SourceParticleData->GetComponentPtrFloat(ScaleBaseCompOffset));
-				const float* RESTRICT ScaleY = reinterpret_cast<const float*>(ParticleMeshRenderData.SourceParticleData->GetComponentPtrFloat(ScaleBaseCompOffset + 1));
-				const float* RESTRICT ScaleZ = reinterpret_cast<const float*>(ParticleMeshRenderData.SourceParticleData->GetComponentPtrFloat(ScaleBaseCompOffset + 2));
-
-				const float* RESTRICT QuatArrayX = reinterpret_cast<const float*>(ParticleMeshRenderData.SourceParticleData->GetComponentPtrFloat(TransformBaseCompOffset));
-				const float* RESTRICT QuatArrayY = reinterpret_cast<const float*>(ParticleMeshRenderData.SourceParticleData->GetComponentPtrFloat(TransformBaseCompOffset + 1));
-				const float* RESTRICT QuatArrayZ = reinterpret_cast<const float*>(ParticleMeshRenderData.SourceParticleData->GetComponentPtrFloat(TransformBaseCompOffset + 2));
-				const float* RESTRICT QuatArrayW = reinterpret_cast<const float*>(ParticleMeshRenderData.SourceParticleData->GetComponentPtrFloat(TransformBaseCompOffset + 3));
-
-				const int32* RESTRICT RenderVisibilityData = ParticleRendererVisTagOffset == INDEX_NONE ? nullptr : reinterpret_cast<const int32*>(ParticleMeshRenderData.SourceParticleData->GetComponentPtrInt32(ParticleRendererVisTagOffset));
-				const int32* RESTRICT MeshIndexData = ParticleMeshIndexOffset == INDEX_NONE ? nullptr : reinterpret_cast<const int32*>(ParticleMeshRenderData.SourceParticleData->GetComponentPtrInt32(ParticleMeshIndexOffset));
-
-				auto GetInstancePosition = [&PositionX, &PositionY, &PositionZ](int32 Idx)
-				{
-					return FVector(PositionX[Idx], PositionY[Idx], PositionZ[Idx]);
-				};
-
-				auto GetInstanceScale = [&ScaleX, &ScaleY, &ScaleZ](int32 Idx)
-				{
-					return FVector(ScaleX[Idx], ScaleY[Idx], ScaleZ[Idx]);
-				};
-
-				auto GetInstanceQuat = [&QuatArrayX, &QuatArrayY, &QuatArrayZ, &QuatArrayW](int32 Idx)
-				{
-					return FQuat(QuatArrayX[Idx], QuatArrayY[Idx], QuatArrayZ[Idx], QuatArrayW[Idx]);
-				};
-
-				//#dxr_todo: handle MESH_FACING_VELOCITY, MESH_FACING_CAMERA_POSITION, MESH_FACING_CAMERA_PLANE
-				//#dxr_todo: handle half floats
-				const bool bHasPosition = PositionBaseCompOffset > 0;
-				const bool bHasRotation = TransformBaseCompOffset > 0;
-				const bool bHasScale = ScaleBaseCompOffset > 0;
-
-				const FMatrix NullInstanceTransform(FVector::ZeroVector, FVector::ZeroVector, FVector::ZeroVector, FVector::ZeroVector);
-				const FQuat MeshRotation = FQuat(MeshData.Rotation);
-				const FVector MeshScale = FVector(MeshData.Scale);
-				for (uint32 InstanceIndex = 0; InstanceIndex < NumInstances; InstanceIndex++)
-				{
-					if ( RenderVisibilityData && (RenderVisibilityData[InstanceIndex] != RendererVisibility) )
-					{
-						RayTracingInstance.InstanceTransforms.Add(NullInstanceTransform);
-						continue;
-					}
-					if ( MeshIndexData && (MeshIndexData[InstanceIndex] != MeshData.SourceMeshIndex) )
-					{
-						RayTracingInstance.InstanceTransforms.Add(NullInstanceTransform);
-						continue;
-					}
-
-					FVector InstancePosition = bHasPosition ? GetInstancePosition(InstanceIndex) : FVector::ZeroVector;
-					if (!bLocalSpace)
-					{
-						// Handle LWC
-						InstancePosition += FVector(SceneProxy->GetLWCRenderTile()) * FLargeWorldRenderScalar::GetTileSize();
-					}
-
-					const FQuat InstanceRotation = bHasRotation ? GetInstanceQuat(InstanceIndex).GetNormalized() * MeshRotation : MeshRotation;
-					FMatrix InstanceTransform = FQuatRotationTranslationMatrix::Make(InstanceRotation, InstancePosition);
-
-					const FVector InstanceScale = bHasScale ? GetInstanceScale(InstanceIndex) * MeshScale : MeshScale;
-					InstanceTransform = FScaleMatrix(InstanceScale) * InstanceTransform;
-
-					if (bLocalSpace)
-					{
-						InstanceTransform = InstanceTransform * LocalTransform;
-					}
-
-					RayTracingInstance.InstanceTransforms.Add(InstanceTransform);
-				}
-			}
-			// Gpu Target
-			else if (FNiagaraUtilities::AllowComputeShaders(GShaderPlatformForFeatureLevel[FeatureLevel]) && FDataDrivenShaderPlatformInfo::GetSupportsRayTracingIndirectInstanceData(GShaderPlatformForFeatureLevel[FeatureLevel]) )
-			{
-				RayTracingInstance.NumTransforms = NumInstances;
-
-				FRDGBufferRef InstanceGPUTransformsBufferRef = Context.GraphBuilder.CreateBuffer(
-					FRDGBufferDesc::CreateStructuredDesc(4 * sizeof(float), 3 * NumInstances),
-					TEXT("InstanceGPUTransformsBuffer"));
-
-				const FLargeWorldRenderPosition AbsoluteViewOrigin(View->ViewMatrices.GetViewOrigin());
-				const FVector ViewTileOffset = AbsoluteViewOrigin.GetTileOffset();
-				const FVector RelativePreViewTranslation = View->ViewMatrices.GetPreViewTranslation() + ViewTileOffset;
-				const FVector3f ViewTilePosition = AbsoluteViewOrigin.GetTile();
-
-				const FLargeWorldRenderPosition LocalTransformOrigin(LocalTransform.GetOrigin());
-
-				FNiagaraGPURayTracingTransformsCS::FParameters* PassParameters = Context.GraphBuilder.AllocParameters< FNiagaraGPURayTracingTransformsCS::FParameters>();
-				{
-					PassParameters->ParticleDataFloatStride		= ParticleMeshRenderData.ParticleFloatDataStride;
-					//PassParameters.ParticleDataHalfStride		= ParticleMeshRenderData.ParticleHalfDataStride;
-					PassParameters->ParticleDataIntStride		= ParticleMeshRenderData.ParticleIntDataStride;
-					PassParameters->CPUNumInstances				= NumInstances;
-					PassParameters->InstanceCountOffset			= ParticleMeshRenderData.SourceParticleData->GetGPUInstanceCountBufferOffset();
-					PassParameters->SystemLWCTile				= SceneProxy->GetLWCRenderTile();
-					PassParameters->PositionDataOffset			= VFVariables[ENiagaraMeshVFLayout::Position].GetGPUOffset();
-					PassParameters->RotationDataOffset			= VFVariables[ENiagaraMeshVFLayout::Rotation].GetGPUOffset();
-					PassParameters->ScaleDataOffset				= VFVariables[ENiagaraMeshVFLayout::Scale].GetGPUOffset();
-					PassParameters->bLocalSpace					= bUseLocalSpace ? 1 : 0;
-					PassParameters->RenderVisibilityOffset		= ParticleRendererVisTagOffset;
-					PassParameters->MeshIndexOffset				= ParticleMeshIndexOffset;
-					PassParameters->RenderVisibilityValue		= RendererVisibility;
-					PassParameters->MeshIndexValue				= MeshData.SourceMeshIndex;
-					PassParameters->LocalTransform				= FLargeWorldRenderScalar::MakeToRelativeWorldMatrix(LocalTransformOrigin.GetTileOffset(), LocalTransform);
-					PassParameters->LocalTransformTile			= LocalTransformOrigin.GetTile();
-					PassParameters->DefaultPosition				= CommonParameters.DefaultPosition;
-					PassParameters->DefaultRotation				= FVector4f(0.0f, 0.0f, 0.0f, 1.0f);
-					PassParameters->DefaultScale				= FVector3f(1.0f, 1.0f, 1.0f);
-					PassParameters->MeshScale					= MeshData.Scale;
-					PassParameters->MeshRotation				= FVector4f(MeshData.Rotation.X, MeshData.Rotation.Y, MeshData.Rotation.Z, MeshData.Rotation.W);
-					PassParameters->ParticleDataFloatBuffer		= ParticleMeshRenderData.ParticleFloatSRV;
-					//PassParameters.ParticleDataHalfBuffer		= ParticleMeshRenderData.ParticleHalfSRV;
-					PassParameters->ParticleDataIntBuffer		= ParticleMeshRenderData.ParticleIntSRV;
-					PassParameters->GPUInstanceCountBuffer		= ComputeDispatchInterface->GetGPUInstanceCounterManager().GetInstanceCountBuffer().SRV;
-					PassParameters->TLASTransforms				= Context.GraphBuilder.CreateUAV(FRDGBufferUAVDesc(InstanceGPUTransformsBufferRef));
-
-					PassParameters->ViewTilePosition			= ViewTilePosition;
-					PassParameters->RelativePreViewTranslation	= FVector3f(RelativePreViewTranslation);
-				}
-
-				Context.GraphBuilder.AddPass(
-					RDG_EVENT_NAME("NiagaraGPURayTracingTransforms"),
-					PassParameters,
-					ERDGPassFlags::Compute,
-					[PassParameters, Pass_FeatureLevel = FeatureLevel, NumInstances](FRHICommandList& RHICmdList)
-				{
-					FNiagaraGPURayTracingTransformsCS::FPermutationDomain PermutationVector;
-					TShaderMapRef<FNiagaraGPURayTracingTransformsCS> ComputeShader(GetGlobalShaderMap(Pass_FeatureLevel), PermutationVector);
-					FComputeShaderUtils::Dispatch(RHICmdList, ComputeShader, *PassParameters, FComputeShaderUtils::GetGroupCount(int32(NumInstances), int32(FNiagaraGPURayTracingTransformsCS::ThreadGroupSize)));
-				});
-
-				Context.GraphBuilder.UseExternalAccessMode(InstanceGPUTransformsBufferRef, ERHIAccess::SRVMask);
-
-				const TRefCountPtr<FRDGPooledBuffer>& ExternalBuffer = Context.GraphBuilder.ConvertToExternalBuffer(InstanceGPUTransformsBufferRef);
-				RayTracingInstance.InstanceGPUTransformsSRV = ExternalBuffer->GetOrCreateSRV(RHICmdList, FRHIBufferSRVCreateInfo());
-			}
+			// only set number of transforms so actual transform data is copied from GPU Scene
+			RayTracingInstance.NumTransforms = NumInstances;
 		}
 
 		OutRayTracingInstances.Add(MoveTemp(RayTracingInstance));
diff --git a/Engine/Plugins/FX/Niagara/Source/NiagaraVertexFactories/Private/NiagaraMeshVertexFactory.cpp b/Engine/Plugins/FX/Niagara/Source/NiagaraVertexFactories/Private/NiagaraMeshVertexFactory.cpp
--- a/Engine/Plugins/FX/Niagara/Source/NiagaraVertexFactories/Private/NiagaraMeshVertexFactory.cpp
+++ b/Engine/Plugins/FX/Niagara/Source/NiagaraVertexFactories/Private/NiagaraMeshVertexFactory.cpp
@@ -178,17 +178,8 @@ void FNiagaraMeshVertexFactory::ModifyCompilationEnvironment(const FVertexFactor
 	const bool bUseGPUScene = UseGPUScene(Parameters.Platform, MaxSupportedFeatureLevel) && MaxSupportedFeatureLevel > ERHIFeatureLevel::ES3_1;
 	const bool bSupportsPrimitiveIdStream = Parameters.VertexFactoryType->SupportsPrimitiveIdStream();
 	
-	// TODO: Support GPU Scene for raytracing
-	if (bSupportsPrimitiveIdStream && bUseGPUScene)
-	{
-		OutEnvironment.SetDefine(TEXT("VF_SUPPORTS_PRIMITIVE_SCENE_DATA"), TEXT("!(RAYHITGROUPSHADER)"));
-		OutEnvironment.SetDefine(TEXT("VF_REQUIRES_PER_INSTANCE_CUSTOM_DATA"), TEXT("!(RAYHITGROUPSHADER)"));
-	}
-	else
-	{
-		OutEnvironment.SetDefine(TEXT("VF_SUPPORTS_PRIMITIVE_SCENE_DATA"), 0);
-		OutEnvironment.SetDefine(TEXT("VF_REQUIRES_PER_INSTANCE_CUSTOM_DATA"), 0);
-	}
+	OutEnvironment.SetDefine(TEXT("VF_SUPPORTS_PRIMITIVE_SCENE_DATA"), bSupportsPrimitiveIdStream && bUseGPUScene);
+	OutEnvironment.SetDefine(TEXT("VF_REQUIRES_PER_INSTANCE_CUSTOM_DATA"), bSupportsPrimitiveIdStream && bUseGPUScene);
 #endif
 
 	if (RHISupportsManualVertexFetch(Parameters.Platform))
diff --git a/Engine/Source/Runtime/Engine/Private/SceneManagement.cpp b/Engine/Source/Runtime/Engine/Private/SceneManagement.cpp
--- a/Engine/Source/Runtime/Engine/Private/SceneManagement.cpp
+++ b/Engine/Source/Runtime/Engine/Private/SceneManagement.cpp
@@ -348,6 +348,7 @@ FRayTracingMaterialGatheringContext::FRayTracingMaterialGatheringContext(
 	const FSceneViewFamily& InReferenceViewFamily,
 	FRDGBuilder& InGraphBuilder,
 	FRayTracingMeshResourceCollector& InRayTracingMeshResourceCollector,
+	FGPUScenePrimitiveCollector& InDynamicPrimitiveCollector,
 	FGlobalDynamicReadBuffer& InDynamicReadBuffer)
 	: Scene(InScene)
 	, ReferenceView(InReferenceView)
@@ -360,6 +361,16 @@ FRayTracingMaterialGatheringContext::FRayTracingMaterialGatheringContext(
 	, DynamicReadBuffer(InDynamicReadBuffer)
 {
 	RayTracingMeshResourceCollector.Start(RHICmdList, DynamicVertexBuffer, DynamicIndexBuffer, DynamicReadBuffer);
+
+	RayTracingMeshResourceCollector.AddViewMeshArrays(
+		ReferenceView,
+		nullptr,
+		nullptr,
+		&InDynamicPrimitiveCollector
+#if UE_ENABLE_DEBUG_DRAWING
+		, nullptr
+#endif
+	);
 }
 
 FRayTracingMaterialGatheringContext::~FRayTracingMaterialGatheringContext()
@@ -368,6 +379,11 @@ FRayTracingMaterialGatheringContext::~FRayTracingMaterialGatheringContext()
 	DynamicReadBuffer.Commit(GraphBuilder.RHICmdList);
 }
 
+void FRayTracingMaterialGatheringContext::SetPrimitive(const FPrimitiveSceneProxy* InPrimitiveSceneProxy)
+{
+	RayTracingMeshResourceCollector.SetPrimitive(InPrimitiveSceneProxy, FHitProxyId::InvisibleHitProxyId);
+}
+
 #endif
 
 FMeshElementCollector::FMeshElementCollector(ERHIFeatureLevel::Type InFeatureLevel, FSceneRenderingBulkObjectAllocator& InBulkAllocator, ECommitFlags InCommitFlags) :
@@ -396,7 +412,10 @@ void FMeshElementCollector::SetPrimitive(const FPrimitiveSceneProxy* InPrimitive
 
 	for (int32 ViewIndex = 0; ViewIndex < SimpleElementCollectors.Num(); ViewIndex++)
 	{
-		SimpleElementCollectors[ViewIndex]->HitProxyId = DefaultHitProxyId;
+		if (SimpleElementCollectors[ViewIndex])
+		{
+			SimpleElementCollectors[ViewIndex]->HitProxyId = DefaultHitProxyId;
+		}
 	}
 
 	for (int32 ViewIndex = 0; ViewIndex < MeshIdInPrimitivePerView.Num(); ++ViewIndex)
@@ -567,7 +586,10 @@ void FMeshElementCollector::AddMesh(int32 ViewIndex, FMeshBatch& MeshBatch)
 
 	NumMeshBatchElementsPerView[ViewIndex] += MeshBatch.Elements.Num();
 
-	MeshBatches[ViewIndex]->Emplace(MeshBatch, PrimitiveSceneProxy, FeatureLevel);
+	if (MeshBatches[ViewIndex])
+	{
+		MeshBatches[ViewIndex]->Emplace(MeshBatch, PrimitiveSceneProxy, FeatureLevel);
+	}
 }
 
 FDynamicPrimitiveUniformBuffer::FDynamicPrimitiveUniformBuffer() = default;
diff --git a/Engine/Source/Runtime/Engine/Public/SceneManagement.h b/Engine/Source/Runtime/Engine/Public/SceneManagement.h
--- a/Engine/Source/Runtime/Engine/Public/SceneManagement.h
+++ b/Engine/Source/Runtime/Engine/Public/SceneManagement.h
@@ -2119,7 +2119,6 @@ class FRayTracingMeshResourceCollector : public FMeshElementCollector
 public:
 	// No MeshBatch should be allocated from an FRayTracingMeshResourceCollector.
 	inline FMeshBatch& AllocateMesh() = delete;
-	void AddMesh(int32 ViewIndex, FMeshBatch& MeshBatch) = delete;
 	void RegisterOneFrameMaterialProxy(FMaterialRenderProxy* Proxy) = delete;
 
 	FRayTracingMeshResourceCollector(
@@ -2172,9 +2171,12 @@ struct FRayTracingMaterialGatheringContext
 		const FSceneViewFamily& InReferenceViewFamily,
 		FRDGBuilder& InGraphBuilder,
 		FRayTracingMeshResourceCollector& InRayTracingMeshResourceCollector,
+		FGPUScenePrimitiveCollector& InDynamicPrimitiveCollector,
 		FGlobalDynamicReadBuffer& InGlobalDynamicReadBuffer);
 
 	ENGINE_API virtual ~FRayTracingMaterialGatheringContext();
+
+	ENGINE_API void SetPrimitive(const FPrimitiveSceneProxy* InPrimitiveSceneProxy);
 };
 #endif
 
diff --git a/Engine/Source/Runtime/Renderer/Private/GPUScene.cpp b/Engine/Source/Runtime/Renderer/Private/GPUScene.cpp
--- a/Engine/Source/Runtime/Renderer/Private/GPUScene.cpp
+++ b/Engine/Source/Runtime/Renderer/Private/GPUScene.cpp
@@ -1564,7 +1564,7 @@ struct FUploadDataSourceAdapterDynamicPrimitives
 	TArray<uint32, SceneRenderingAllocator> PrimitivesIds;
 };
 
-void FGPUScene::UploadDynamicPrimitiveShaderDataForViewInternal(FRDGBuilder& GraphBuilder, FViewInfo& View, UE::Renderer::Private::IShadowInvalidatingInstances *ShadowInvalidatingInstances)
+void FGPUScene::UploadDynamicPrimitiveShaderDataForViewInternal(FRDGBuilder& GraphBuilder, FViewInfo& View, bool bRayTracing, UE::Renderer::Private::IShadowInvalidatingInstances *ShadowInvalidatingInstances)
 {
 	LLM_SCOPE_BYTAG(GPUScene);
 
@@ -1574,7 +1574,7 @@ void FGPUScene::UploadDynamicPrimitiveShaderDataForViewInternal(FRDGBuilder& Gra
 	ensure(bInBeginEndBlock);
 	ensure(DynamicPrimitivesOffset >= Scene.GetMaxPersistentPrimitiveIndex());
 
-	FGPUScenePrimitiveCollector& Collector = View.DynamicPrimitiveCollector;
+	FGPUScenePrimitiveCollector& Collector = bRayTracing ? View.RayTracingDynamicPrimitiveCollector : View.DynamicPrimitiveCollector;
 
 	// Auto-commit if not done (should usually not be done, but sometimes the UploadDynamicPrimitiveShaderDataForViewInternal is called to ensure the 
 	// CachedViewUniformShaderParameters is set on the view.
@@ -1775,13 +1775,13 @@ void FGPUScene::Update(FRDGBuilder& GraphBuilder, FSceneUniformBuffer& SceneUB,
 	}
 }
 
-void FGPUScene::UploadDynamicPrimitiveShaderDataForView(FRDGBuilder& GraphBuilder, FViewInfo& View, UE::Renderer::Private::IShadowInvalidatingInstances *ShadowInvalidatingInstances)
+void FGPUScene::UploadDynamicPrimitiveShaderDataForView(FRDGBuilder& GraphBuilder, FViewInfo& View, bool bRayTracing, UE::Renderer::Private::IShadowInvalidatingInstances *ShadowInvalidatingInstances)
 {
 	if (bIsEnabled)
 	{
 		RDG_GPU_MASK_SCOPE(GraphBuilder, FRHIGPUMask::All());
 
-		UploadDynamicPrimitiveShaderDataForViewInternal(GraphBuilder, View, ShadowInvalidatingInstances);
+		UploadDynamicPrimitiveShaderDataForViewInternal(GraphBuilder, View, bRayTracing, ShadowInvalidatingInstances);
 	}
 }
 
diff --git a/Engine/Source/Runtime/Renderer/Private/GPUScene.h b/Engine/Source/Runtime/Renderer/Private/GPUScene.h
--- a/Engine/Source/Runtime/Renderer/Private/GPUScene.h
+++ b/Engine/Source/Runtime/Renderer/Private/GPUScene.h
@@ -220,7 +220,7 @@ public:
 	/**
 	 * Upload primitives from View.DynamicPrimitiveCollector.
 	 */
-	void UploadDynamicPrimitiveShaderDataForView(FRDGBuilder& GraphBuilder, FViewInfo& View, UE::Renderer::Private::IShadowInvalidatingInstances *ShadowInvalidatingInstances = nullptr);
+	void UploadDynamicPrimitiveShaderDataForView(FRDGBuilder& GraphBuilder, FViewInfo& View, bool bRayTracing = false, UE::Renderer::Private::IShadowInvalidatingInstances *ShadowInvalidatingInstances = nullptr);
 
 	/**
 	 * Modifies the GPUScene specific scene UB parameters to the current versions. Returns true if any of the parameters changed.
@@ -428,7 +428,7 @@ private:
 
 	static void InitLightData(const FLightSceneInfoCompact& LightInfoCompact, bool bAllowStaticLighting, FLightSceneData& DataOut);
 
-	void UploadDynamicPrimitiveShaderDataForViewInternal(FRDGBuilder& GraphBuilder, FViewInfo& View, UE::Renderer::Private::IShadowInvalidatingInstances *ShadowInvalidatingInstances);
+	void UploadDynamicPrimitiveShaderDataForViewInternal(FRDGBuilder& GraphBuilder, FViewInfo& View, bool bRayTracing, UE::Renderer::Private::IShadowInvalidatingInstances *ShadowInvalidatingInstances);
 
 	void UpdateInternal(FRDGBuilder& GraphBuilder, FSceneUniformBuffer& SceneUB, FRDGExternalAccessQueue& ExternalAccessQueue, const UE::Tasks::FTask& UpdateTaskPrerequisites);
 
diff --git a/Engine/Source/Runtime/Renderer/Private/RayTracing/RayTracing.cpp b/Engine/Source/Runtime/Renderer/Private/RayTracing/RayTracing.cpp
--- a/Engine/Source/Runtime/Renderer/Private/RayTracing/RayTracing.cpp
+++ b/Engine/Source/Runtime/Renderer/Private/RayTracing/RayTracing.cpp
@@ -683,16 +683,6 @@ namespace RayTracing
 
 		View.RayTracingCullingParameters.Init(View);
 
-		FRayTracingMaterialGatheringContext MaterialGatheringContext
-		(
-			&Scene,
-			&View,
-			*View.Family,
-			GraphBuilder,
-			*View.RayTracingMeshResourceCollector,
-			InDynamicReadBuffer
-		);
-
 		const float CurrentWorldTime = View.Family->Time.GetWorldTimeSeconds();
 
 		// Consume output of the relevant primitive gathering task
@@ -796,32 +786,76 @@ namespace RayTracing
 					NumPendingMeshBatches = 0;
 				};
 
+			// Need to process dynamic primitives in 3 passes to support dynamic primitives/instances in GPU Scene
+			// 1 - gather all dynamic ray tracing instances
+			// 2 - upload dynamic primitive/instance data to GPU scene
+			// 3 - process dynamic ray tracing instances (primitive IDs and instance scene data offset of dynamic primitives is now valid)
+
 			// Local temporary array of instances used for GetDynamicRayTracingInstances()
-			TArray<FRayTracingInstance> TempRayTracingInstances;
+			TArray<FRayTracingInstance> DynamicRayTracingInstances;
+			TArray<TRange<int32>> PrimitivesDynamicRayTracingInstances;
 
-			for (int32 PrimitiveIndex : RelevantPrimitiveList.DynamicPrimitives)
 			{
-				FPrimitiveSceneProxy* SceneProxy = Scene.PrimitiveSceneProxies[PrimitiveIndex];
-				FPrimitiveSceneInfo* SceneInfo = Scene.Primitives[PrimitiveIndex];
-				const FPersistentPrimitiveIndex PersistentPrimitiveIndex = SceneInfo->GetPersistentIndex();
+				FRayTracingMaterialGatheringContext MaterialGatheringContext
+				(
+					&Scene,
+					&View,
+					*View.Family,
+					GraphBuilder,
+					*View.RayTracingMeshResourceCollector,
+					View.RayTracingDynamicPrimitiveCollector,
+					InDynamicReadBuffer
+				);
+
+				for (int32 PrimitiveIndex : RelevantPrimitiveList.DynamicPrimitives)
+				{
+					check(MaterialGatheringContext.DynamicRayTracingGeometriesToUpdate.IsEmpty());
 
-				TempRayTracingInstances.Reset();
-				MaterialGatheringContext.DynamicRayTracingGeometriesToUpdate.Reset();
+					FPrimitiveSceneProxy* SceneProxy = Scene.PrimitiveSceneProxies[PrimitiveIndex];
+					FPrimitiveSceneInfo* SceneInfo = Scene.Primitives[PrimitiveIndex];
+					const FPersistentPrimitiveIndex PersistentPrimitiveIndex = SceneInfo->GetPersistentIndex();
 
-				SceneProxy->GetDynamicRayTracingInstances(MaterialGatheringContext, TempRayTracingInstances);
+					MaterialGatheringContext.SetPrimitive(SceneProxy);
 
-				for (const FRayTracingDynamicGeometryUpdateParams& DynamicRayTracingGeometryUpdate : MaterialGatheringContext.DynamicRayTracingGeometriesToUpdate)
-				{
-					Scene.GetRayTracingDynamicGeometryCollection()->AddDynamicMeshBatchForGeometryUpdate(
-						GraphBuilder.RHICmdList,
-						&Scene,
-						&View,
-						SceneProxy,
-						DynamicRayTracingGeometryUpdate,
-						PersistentPrimitiveIndex.Index
-					);
+					int32 BaseRayTracingInstance = DynamicRayTracingInstances.Num();
+
+					SceneProxy->GetDynamicRayTracingInstances(MaterialGatheringContext, DynamicRayTracingInstances);
+
+					for (const FRayTracingDynamicGeometryUpdateParams& DynamicRayTracingGeometryUpdate : MaterialGatheringContext.DynamicRayTracingGeometriesToUpdate)
+					{
+						Scene.GetRayTracingDynamicGeometryCollection()->AddDynamicMeshBatchForGeometryUpdate(
+							GraphBuilder.RHICmdList,
+							&Scene,
+							&View,
+							SceneProxy,
+							DynamicRayTracingGeometryUpdate,
+							PersistentPrimitiveIndex.Index
+						);
+					}
+					MaterialGatheringContext.DynamicRayTracingGeometriesToUpdate.Reset();
+
+					PrimitivesDynamicRayTracingInstances.Add(TRange<int32>(BaseRayTracingInstance, DynamicRayTracingInstances.Num()));
 				}
 
+				// FRayTracingMaterialGatheringContext destructor handles committing dynamic mesh batches to GPU Scene
+			}
+
+			Scene.GPUScene.UploadDynamicPrimitiveShaderDataForView(GraphBuilder, View, true);
+
+			const int32 ViewDynamicPrimitiveId = View.RayTracingDynamicPrimitiveCollector.GetPrimitiveIdRange().GetLowerBoundValue();
+			const int32 ViewInstanceSceneDataOffset = View.RayTracingDynamicPrimitiveCollector.GetInstanceSceneDataOffset();
+
+			for (int32 Index = 0; Index < RelevantPrimitiveList.DynamicPrimitives.Num(); ++Index)
+			{
+				const int32 PrimitiveIndex = RelevantPrimitiveList.DynamicPrimitives[Index];
+				FPrimitiveSceneProxy* SceneProxy = Scene.PrimitiveSceneProxies[PrimitiveIndex];
+				FPrimitiveSceneInfo* SceneInfo = Scene.Primitives[PrimitiveIndex];
+				const FPersistentPrimitiveIndex PersistentPrimitiveIndex = SceneInfo->GetPersistentIndex();
+
+				TArrayView<FRayTracingInstance> TempRayTracingInstances = MakeArrayView(
+					DynamicRayTracingInstances.GetData() + PrimitivesDynamicRayTracingInstances[Index].GetLowerBound().GetValue(),
+					PrimitivesDynamicRayTracingInstances[Index].Size<int32>());
+
 				if (TempRayTracingInstances.Num() > 0)
 				{
 					for (FRayTracingInstance& Instance : TempRayTracingInstances)
@@ -906,10 +940,20 @@ namespace RayTracing
 							continue;
 						}
 
+						int32 PrimitiveId = PersistentPrimitiveIndex.Index;
+						int32 InstanceSceneDataOffset = SceneInfo->GetInstanceSceneDataOffset();
+
+						if (Instance.Materials.Num() > 0 && Instance.Materials[0].Elements.Num() > 0 && Instance.Materials[0].Elements[0].DynamicPrimitiveData != nullptr)
+						{
+							check(Instance.NumTransforms == Instance.Materials[0].Elements[0].NumInstances);
+							PrimitiveId = ViewDynamicPrimitiveId + Instance.Materials[0].Elements[0].DynamicPrimitiveIndex;
+							InstanceSceneDataOffset = ViewInstanceSceneDataOffset + Instance.Materials[0].Elements[0].DynamicPrimitiveInstanceSceneDataOffset;
+						}
+
 						FRayTracingGeometryInstance RayTracingInstance;
 						RayTracingInstance.GeometryRHI = Geometry->GetRHI();
 						checkf(RayTracingInstance.GeometryRHI, TEXT("Ray tracing instance must have a valid geometry."));
-						RayTracingInstance.DefaultUserData = PersistentPrimitiveIndex.Index;
+						RayTracingInstance.DefaultUserData = PrimitiveId;
 						RayTracingInstance.bApplyLocalBoundsTransform = Instance.bApplyLocalBoundsTransform;
 						RayTracingInstance.LayerIndex = (uint8)ERayTracingSceneLayer::Base;
 						RayTracingInstance.Mask = SceneInfo->CachedRayTracingInstance.Mask;
@@ -921,7 +965,7 @@ namespace RayTracing
 							RayTracingInstance.NumTransforms = Instance.NumTransforms;
 							RayTracingInstance.GPUTransformsSRV = Instance.InstanceGPUTransformsSRV;
 						}
-						else
+						else if(!Instance.GetTransforms().IsEmpty())
 						{
 							if (Instance.OwnsTransforms())
 							{
@@ -942,6 +986,12 @@ namespace RayTracing
 								RayTracingInstance.Transforms = Instance.InstanceTransformsView;
 							}
 						}
+						else
+						{
+							// If neither InstanceGPUTransformsSRV nor array of transforms was provided, get the instance transforms from GPU Scene
+							RayTracingInstance.NumTransforms = Instance.NumTransforms;
+							RayTracingInstance.BaseInstanceSceneDataOffset = InstanceSceneDataOffset;
+						}
 
 						uint32 InstanceIndex = INDEX_NONE;
 						if (!SceneInfo->bCachedRayTracingInstanceAllSegmentsDecal)
diff --git a/Engine/Source/Runtime/Renderer/Private/SceneRendering.cpp b/Engine/Source/Runtime/Renderer/Private/SceneRendering.cpp
--- a/Engine/Source/Runtime/Renderer/Private/SceneRendering.cpp
+++ b/Engine/Source/Runtime/Renderer/Private/SceneRendering.cpp
@@ -2589,6 +2589,7 @@ FSceneRenderer::FSceneRenderer(const FSceneViewFamily* InViewFamily, FHitProxyCo
 
 		// Must initialize to have a GPUScene connected to be able to collect dynamic primitives.
 		ViewInfo->DynamicPrimitiveCollector = FGPUScenePrimitiveCollector(&GPUSceneDynamicContext);
+		ViewInfo->RayTracingDynamicPrimitiveCollector = FGPUScenePrimitiveCollector(&GPUSceneDynamicContext);
 
 		check(ViewInfo->ViewRect.Area() == 0);
 
diff --git a/Engine/Source/Runtime/Renderer/Private/SceneRendering.h b/Engine/Source/Runtime/Renderer/Private/SceneRendering.h
--- a/Engine/Source/Runtime/Renderer/Private/SceneRendering.h
+++ b/Engine/Source/Runtime/Renderer/Private/SceneRendering.h
@@ -1810,6 +1810,8 @@ public:
 	 * Collector for view-dependent data.
 	 */
 	FGPUScenePrimitiveCollector DynamicPrimitiveCollector;
+	FGPUScenePrimitiveCollector RayTracingDynamicPrimitiveCollector;
+
 private:
 	// Cache of TEXTUREGROUP_World to create view's samplers on render thread.
 	// may not have a valid value if FViewInfo is created on the render thread.
diff --git a/Engine/Source/Runtime/Renderer/Private/ShadowSetup.cpp b/Engine/Source/Runtime/Renderer/Private/ShadowSetup.cpp
--- a/Engine/Source/Runtime/Renderer/Private/ShadowSetup.cpp
+++ b/Engine/Source/Runtime/Renderer/Private/ShadowSetup.cpp
@@ -6615,7 +6615,7 @@ void FSceneRenderer::FinishDynamicShadowMeshPassSetup(FRDGBuilder& GraphBuilder,
 	}
 	for (FProjectedShadowInfo* ProjectedShadowInfo : SortedShadowsForShadowDepthPass.VirtualShadowMapShadows)
 	{
-		Scene->GPUScene.UploadDynamicPrimitiveShaderDataForView(GraphBuilder, *ProjectedShadowInfo->ShadowDepthView, GetShadowInvalidatingInstancesInterface(ProjectedShadowInfo->DependentView) );
+		Scene->GPUScene.UploadDynamicPrimitiveShaderDataForView(GraphBuilder, *ProjectedShadowInfo->ShadowDepthView, false, GetShadowInvalidatingInstancesInterface(ProjectedShadowInfo->DependentView) );
 	}
 
 	DynamicReadBufferForShadows.Commit(GraphBuilder.RHICmdList);
