support querying a frame directly rather than a sequence of svts
request frame as a float

#changelist validated
#virtualized

diff --git a/Engine/Plugins/FX/Niagara/Source/Niagara/Private/NiagaraDataInterfaceSparseVolumeTexture.cpp b/Engine/Plugins/FX/Niagara/Source/Niagara/Private/NiagaraDataInterfaceSparseVolumeTexture.cpp
--- a/Engine/Plugins/FX/Niagara/Source/Niagara/Private/NiagaraDataInterfaceSparseVolumeTexture.cpp
+++ b/Engine/Plugins/FX/Niagara/Source/Niagara/Private/NiagaraDataInterfaceSparseVolumeTexture.cpp
@@ -43,7 +43,7 @@ struct FNDISparseVolumeTextureInstanceData_GameThread
 	FIntVector3 CurrentTextureSize = FIntVector3::ZeroValue;
 	int32 CurrentTextureMipLevels = 0;
 	FNiagaraParameterDirectBinding<UObject*> UserParamBinding;
-	int32 CurrentFrame = 0;
+	float CurrentFrame = 0;
 	float FrameRate = 30.f;
 	int32 NumFrames = 0;
 };
@@ -170,7 +170,7 @@ void UNiagaraDataInterfaceSparseVolumeTexture::GetFunctionsInternal(TArray<FNiag
 		Sig.SetFunctionVersion(FNDISparseVolumeTextureFunctionVersion::LatestVersion);
 		Sig.Inputs.Emplace(FNiagaraTypeDefinition(GetClass()), TEXT("SparseVolumeTexture"));		
 		Sig.Inputs.Emplace(FNiagaraTypeDefinition::GetFloatDef(), TEXT("FrameRate"));
-		Sig.Inputs.Emplace(FNiagaraTypeDefinition::GetIntDef(), TEXT("Frame"));
+		Sig.Inputs.Emplace(FNiagaraTypeDefinition::GetFloatDef(), TEXT("Frame"));
 		Sig.Outputs.Add(FNiagaraVariable(FNiagaraTypeDefinition::GetBoolDef(), TEXT("Success")));
 		Sig.SetDescription(LOCTEXT("RequestSparseVolumeTextureFrameDesc", "Queue up the frame to load on tick"));
 	}
@@ -247,13 +247,18 @@ bool UNiagaraDataInterfaceSparseVolumeTexture::PerInstanceTick(void* PerInstance
 	USparseVolumeTexture* CurrentTexture = InstanceData->UserParamBinding.GetValueOrDefault<USparseVolumeTexture>(SparseVolumeTexture);	
 
 	if (CurrentTexture)
-	{
-		InstanceData->NumFrames = CurrentTexture->GetNumFrames();
+	{	
+		InstanceData->NumFrames = CurrentTexture->GetNumFrames();		
+
+		// skip frame request if we are passed in a frame (ie: from BP)
+		USparseVolumeTextureFrame* SparseVolumeTextureFrame = Cast<USparseVolumeTextureFrame>(CurrentTexture);
 
-		// todo: we only support querying mip 0 for now
-		USparseVolumeTextureFrame* SparseVolumeTextureFrame = 
-			USparseVolumeTextureFrame::GetFrameAndIssueStreamingRequest(CurrentTexture, GetTypeHash(InstanceData), InstanceData->FrameRate, InstanceData->CurrentFrame, 0, BlockingStreamingRequests, true);
-		UE::SVT::GetStreamingManager().Update_GameThread();
+		if (SparseVolumeTextureFrame == nullptr)
+		{			
+			// todo: we only support querying mip 0 for now
+			SparseVolumeTextureFrame = USparseVolumeTextureFrame::GetFrameAndIssueStreamingRequest(CurrentTexture, GetTypeHash(InstanceData), InstanceData->FrameRate, InstanceData->CurrentFrame, 0, BlockingStreamingRequests, true);
+			UE::SVT::GetStreamingManager().Update_GameThread();
+		}
 
 		if (SparseVolumeTextureFrame)
 		{
@@ -413,7 +418,7 @@ void UNiagaraDataInterfaceSparseVolumeTexture::VMRequestSparseVolumeTextureFrame
 {
 	VectorVM::FUserPtrHandler<FNDISparseVolumeTextureInstanceData_GameThread> InstData(Context);
 	FNDIInputParam<float>	InFrameRate(Context);
-	FNDIInputParam<int>		InFrame(Context);
+	FNDIInputParam<float>	InFrame(Context);
 	FNDIOutputParam<bool> OutSuccess(Context);
 
 	InstData->FrameRate = InFrameRate.GetAndAdvance();
