- Remove incorrect clamp on drag.
- Drag must be non zero, but can go above 1.0

diff --git a/Engine/Plugins/FX/Niagara/Source/Niagara/Internal/Stateless/Modules/NiagaraStatelessModule_SolveVelocitiesAndForces.h b/Engine/Plugins/FX/Niagara/Source/Niagara/Internal/Stateless/Modules/NiagaraStatelessModule_SolveVelocitiesAndForces.h
--- a/Engine/Plugins/FX/Niagara/Source/Niagara/Internal/Stateless/Modules/NiagaraStatelessModule_SolveVelocitiesAndForces.h
+++ b/Engine/Plugins/FX/Niagara/Source/Niagara/Internal/Stateless/Modules/NiagaraStatelessModule_SolveVelocitiesAndForces.h
@@ -28,8 +28,8 @@ public:
 
 		FModuleBuiltData* BuiltData		= BuildContext.AllocateBuiltData<FModuleBuiltData>();
 		BuiltData->MassRange			= PhysicsBuildData.MassRange;
-		BuiltData->DragRange.Min		= FMath::Clamp(PhysicsBuildData.DragRange.Min, 0.01f, 1.0f);
-		BuiltData->DragRange.Max		= FMath::Clamp(PhysicsBuildData.DragRange.Max, 0.01f, 1.0f);
+		BuiltData->DragRange.Min		= FMath::Max(PhysicsBuildData.DragRange.Min, 0.01f);
+		BuiltData->DragRange.Max		= FMath::Max(PhysicsBuildData.DragRange.Max, 0.01f);
 		BuiltData->VelocityRange		= PhysicsBuildData.VelocityRange;
 		BuiltData->WindRange			= PhysicsBuildData.WindRange;
 		BuiltData->AccelerationRange	= PhysicsBuildData.AccelerationRange;
