- Move free ID buffer update out of passes to fix issue with parallel building

#rnx

diff --git a/Engine/Plugins/FX/Niagara/Source/Niagara/Private/NiagaraGpuComputeDispatch.cpp b/Engine/Plugins/FX/Niagara/Source/Niagara/Private/NiagaraGpuComputeDispatch.cpp
--- a/Engine/Plugins/FX/Niagara/Source/Niagara/Private/NiagaraGpuComputeDispatch.cpp
+++ b/Engine/Plugins/FX/Niagara/Source/Niagara/Private/NiagaraGpuComputeDispatch.cpp
@@ -933,6 +933,7 @@ void FNiagaraGpuComputeDispatch::PrepareTicksForProxy(FRHICommandListImmediate&
 				if ( ComputeContext->MainDataSet->RequiresPersistentIDs() && ComputeContext->bHasTickedThisFrame_RT )
 				{
 					FinalDispatchGroup.FreeIDUpdates.Emplace(ComputeContext);
+					NumRequiredFreeIDListSizes = FMath::Max<uint32>(NumRequiredFreeIDListSizes, FinalDispatchGroup.FreeIDUpdates.Num());
 				}
 			}
 		}
@@ -1005,6 +1006,8 @@ void FNiagaraGpuComputeDispatch::PrepareTicksForProxy(FRHICommandListImmediate&
 
 void FNiagaraGpuComputeDispatch::PrepareAllTicks(FRHICommandListImmediate& RHICmdList)
 {
+	NumRequiredFreeIDListSizes = 0;
+
 	for (int iTickStage=0; iTickStage < ENiagaraGpuComputeTickStage::Max; ++iTickStage)
 	{
 		for (FNiagaraSystemGpuComputeProxy* ComputeProxy : ProxiesPerStage[iTickStage])
@@ -1012,6 +1015,14 @@ void FNiagaraGpuComputeDispatch::PrepareAllTicks(FRHICommandListImmediate& RHICm
 			PrepareTicksForProxy(RHICmdList, ComputeProxy, DispatchListPerStage[iTickStage]);
 		}
 	}
+
+	if (NumRequiredFreeIDListSizes > NumAllocatedFreeIDListSizes)
+	{
+		constexpr uint32 ALLOC_CHUNK_SIZE = 128;
+		NumAllocatedFreeIDListSizes = Align(NumRequiredFreeIDListSizes, ALLOC_CHUNK_SIZE);
+		FreeIDListSizesBuffer.Release();
+		FreeIDListSizesBuffer.Initialize(RHICmdList, TEXT("NiagaraFreeIDListSizes"), sizeof(uint32), NumAllocatedFreeIDListSizes, EPixelFormat::PF_R32_SINT, ERHIAccess::UAVCompute, BUF_Static);
+	}
 }
 
 void FNiagaraGpuComputeDispatch::ExecuteTicks(FRDGBuilder& GraphBuilder, TConstStridedView<FSceneView> Views, ENiagaraGpuComputeTickStage::Type TickStage)
@@ -1330,17 +1341,6 @@ void FNiagaraGpuComputeDispatch::ExecuteTicks(FRDGBuilder& GraphBuilder, TConstS
 
 					// Initialize the free ID size buffer
 					{
-						if (NumFreeIDUpdates > NumAllocatedFreeIDListSizes)
-						{
-							constexpr uint32 ALLOC_CHUNK_SIZE = 128;
-							NumAllocatedFreeIDListSizes = Align(NumFreeIDUpdates, ALLOC_CHUNK_SIZE);
-							if (FreeIDListSizesBuffer.Buffer)
-							{
-								FreeIDListSizesBuffer.Release();
-							}
-							FreeIDListSizesBuffer.Initialize(RHICmdList, TEXT("NiagaraFreeIDListSizes"), sizeof(uint32), NumAllocatedFreeIDListSizes, EPixelFormat::PF_R32_SINT, ERHIAccess::UAVCompute, BUF_Static);
-						}
-
 						SCOPED_DRAW_EVENT(RHICmdList, NiagaraGPUComputeClearFreeIDListSizes);
 						RHICmdList.Transition(FRHITransitionInfo(FreeIDListSizesBuffer.UAV, ERHIAccess::UAVCompute, ERHIAccess::UAVCompute));
 						NiagaraFillGPUIntBuffer(RHICmdList, FeatureLevel, FreeIDListSizesBuffer.UAV, FreeIDListSizesBuffer.NumBytes / sizeof(uint32), 0);
diff --git a/Engine/Plugins/FX/Niagara/Source/Niagara/Private/NiagaraGpuComputeDispatch.h b/Engine/Plugins/FX/Niagara/Source/Niagara/Private/NiagaraGpuComputeDispatch.h
--- a/Engine/Plugins/FX/Niagara/Source/Niagara/Private/NiagaraGpuComputeDispatch.h
+++ b/Engine/Plugins/FX/Niagara/Source/Niagara/Private/NiagaraGpuComputeDispatch.h
@@ -229,6 +229,7 @@ private:
 	/** A buffer of list sizes used by UpdateFreeIDBuffers to allow overlapping several dispatches. */
 	FRWBuffer FreeIDListSizesBuffer;
 	uint32 NumAllocatedFreeIDListSizes = 0;
+	uint32 NumRequiredFreeIDListSizes = 0;
 
 	uint32 NumProxiesThatRequireGlobalDistanceField = 0;
 	uint32 NumProxiesThatRequireDepthBuffer = 0;
