Added path property to Niagara asset browser and reworked property display

diff --git a/Engine/Plugins/FX/Niagara/Source/NiagaraEditor/Private/Widgets/AssetBrowser/SNiagaraAssetBrowser.cpp b/Engine/Plugins/FX/Niagara/Source/NiagaraEditor/Private/Widgets/AssetBrowser/SNiagaraAssetBrowser.cpp
--- a/Engine/Plugins/FX/Niagara/Source/NiagaraEditor/Private/Widgets/AssetBrowser/SNiagaraAssetBrowser.cpp
+++ b/Engine/Plugins/FX/Niagara/Source/NiagaraEditor/Private/Widgets/AssetBrowser/SNiagaraAssetBrowser.cpp
@@ -542,30 +542,30 @@ void SNiagaraAssetBrowser::OnAssetSelected(const FAssetData& AssetData)
 			.OnAssetTagActivated(this, &SNiagaraAssetBrowser::OnAssetTagActivated)
 			.OnAssetTagActivatedTooltip(LOCTEXT("SecondaryAssetTagButtonTooltip", "\n\nClicking this tag will activate/deactivate its corresponding filter."));
 		
-		TSharedRef<SWidget> Details = SNew(SVerticalBox)
-		+ SVerticalBox::Slot()
-		.AutoHeight()
-		.HAlign(HAlign_Right)
-		[
-			SNew(SCheckBox)
-			.IsChecked(this, &SNiagaraAssetBrowser::OnShouldDisplayViewport)
-			.OnCheckStateChanged(this, &SNiagaraAssetBrowser::OnShouldDisplayViewportChanged)
-			.ToolTipText(this, &SNiagaraAssetBrowser::OnGetShouldDisplayViewportTooltip)
-			.Visibility(this, &SNiagaraAssetBrowser::OnGetShouldDisplayVisibilityCheckbox)
-			[
-				SNew(STextBlock).Text(LOCTEXT("DisplayViewport", "Display Viewport"))
-			]
-		]
-		+ SVerticalBox::Slot()
-		.AutoHeight()
-		[
-			PreviewViewport.ToSharedRef()
-		]
-		+ SVerticalBox::Slot()
-		.AutoHeight()
-		[
-			SelectedAssetDetails
-		];
+		// TSharedRef<SWidget> Details = SNew(SVerticalBox)
+		// + SVerticalBox::Slot()
+		// .AutoHeight()
+		// .HAlign(HAlign_Right)
+		// [
+		// 	SNew(SCheckBox)
+		// 	.IsChecked(this, &SNiagaraAssetBrowser::OnShouldDisplayViewport)
+		// 	.OnCheckStateChanged(this, &SNiagaraAssetBrowser::OnShouldDisplayViewportChanged)
+		// 	.ToolTipText(this, &SNiagaraAssetBrowser::OnGetShouldDisplayViewportTooltip)
+		// 	.Visibility(this, &SNiagaraAssetBrowser::OnGetShouldDisplayVisibilityCheckbox)
+		// 	[
+		// 		SNew(STextBlock).Text(LOCTEXT("DisplayViewport", "Display Viewport"))
+		// 	]
+		// ]
+		// + SVerticalBox::Slot()
+		// .AutoHeight()
+		// [
+		// 	PreviewViewport.ToSharedRef()
+		// ]
+		// + SVerticalBox::Slot()
+		// .AutoHeight()
+		// [
+		// 	SelectedAssetDetails
+		// ];
 
 		// We don't attach the preview currently until some of the issues are resolved
 		AssetBrowserDetailsSlot->AttachWidget(SelectedAssetDetails);
diff --git a/Engine/Plugins/FX/Niagara/Source/NiagaraEditor/Private/Widgets/AssetBrowser/SNiagaraSelectedAssetDetails.cpp b/Engine/Plugins/FX/Niagara/Source/NiagaraEditor/Private/Widgets/AssetBrowser/SNiagaraSelectedAssetDetails.cpp
--- a/Engine/Plugins/FX/Niagara/Source/NiagaraEditor/Private/Widgets/AssetBrowser/SNiagaraSelectedAssetDetails.cpp
+++ b/Engine/Plugins/FX/Niagara/Source/NiagaraEditor/Private/Widgets/AssetBrowser/SNiagaraSelectedAssetDetails.cpp
@@ -13,6 +13,7 @@
 #include "Widgets/Layout/SSpacer.h"
 #include "Widgets/Layout/SSeparator.h"
 #include "Widgets/Layout/SScrollBox.h"
+#include "HAL/PlatformApplicationMisc.h"
 
 #define LOCTEXT_NAMESPACE "SNiagaraSelectedAssetPreview"
 
@@ -22,69 +23,85 @@ BEGIN_SLATE_FUNCTION_BUILD_OPTIMIZATION
 
 void FNiagaraAssetDetailDatabase::Init()
 {
-	FName SimulationTargetTag("HasGPUEmitter");
 	
 	// Emitters
 	{
 		FNiagaraAssetDetailClassInfo EmitterClassInfo;
-		EmitterClassInfo.DisplayedProperties =
+
+		// Inheritance
+		FDisplayedPropertyData InheritanceProperty;
+		InheritanceProperty.ShouldDisplayPropertyDelegate = FDisplayedPropertyData::FShouldDisplayProperty::CreateLambda([](const FAssetData& AssetData)
+		{
+			if(AssetData.FindTag(GET_MEMBER_NAME_CHECKED(UNiagaraEmitter, bIsInheritable)))
+			{
+				return true;
+			}
+			
+			return false; 
+		});
+		InheritanceProperty.NameWidgetDelegate = FDisplayedPropertyData::FGenerateWidget::CreateLambda([](const FAssetData& AssetData) -> TSharedRef<SWidget>
+		{
+			return SNew(STextBlock)
+			.Text(LOCTEXT("Inheritance", "Inheritance"));
+		});
+		InheritanceProperty.ValueWidgetDelegate = FDisplayedPropertyData::FGenerateWidget::CreateLambda([](const FAssetData& AssetData) -> TSharedRef<SWidget>
+		{
+			bool bUseInheritance = false;
+			if(FNiagaraEditorUtilities::GetIsInheritableFromAssetRegistryTags(AssetData, bUseInheritance))
 			{
-				// Inheritance
+				return SNew(STextBlock).Text(bUseInheritance
+					? LOCTEXT("Emitter_UseInheritance_Yes", "Yes")
+					: LOCTEXT("Emitter_UseInheritance_No", "No")
+					);
+			}
+
+			return SNullWidget::NullWidget;
+		});
+		
+		// Simulation Target
+		FName SimulationTargetTag("HasGPUEmitter");
+
+		FDisplayedPropertyData SimulationTargetProperty;
+		SimulationTargetProperty.ShouldDisplayPropertyDelegate = FDisplayedPropertyData::FShouldDisplayProperty::CreateLambda([SimulationTargetTag](const FAssetData& Asset)
+		{
+			return Asset.FindTag(SimulationTargetTag);
+		});
+		SimulationTargetProperty.NameWidgetDelegate = FDisplayedPropertyData::FGenerateWidget::CreateLambda([](const FAssetData& AssetData) -> TSharedRef<SWidget>
+		{
+			return SNew(STextBlock)
+			.Text(LOCTEXT("SimulationTargetLabel", "Runs on"));
+		});
+		SimulationTargetProperty.ValueWidgetDelegate = FDisplayedPropertyData::FGenerateWidget::CreateLambda([SimulationTargetTag](const FAssetData& Asset)
+		{
+			return SNew(SImage)
+			.Image_Lambda([Asset, SimulationTargetTag]() -> const FSlateBrush*
+			{
+				if(Asset.FindTag(SimulationTargetTag))
 				{
-					FDisplayedPropertyData::FShouldDisplayProperty(),
-					FDisplayedPropertyData::FGenerateWidget::CreateLambda([](const FAssetData& AssetData) -> TSharedRef<SWidget>
+					FString Value;
+					if(Asset.GetTagValue(SimulationTargetTag, Value))
 					{
-						return SNew(STextBlock)
-						.Text(LOCTEXT("Inheritance", "Inheritance"));
-					}),
-					FDisplayedPropertyData::FGenerateWidget::CreateLambda([](const FAssetData& AssetData) -> TSharedRef<SWidget>
-					{
-						bool bUseInheritance = false;
-						if(FNiagaraEditorUtilities::GetIsInheritableFromAssetRegistryTags(AssetData, bUseInheritance))
+						if(Value == "True")
 						{
-							return SNew(STextBlock).Text(bUseInheritance
-								? LOCTEXT("Emitter_UseInheritance_Yes", "Yes")
-								: LOCTEXT("Emitter_UseInheritance_No", "No")
-								);
+							return FNiagaraEditorStyle::Get().GetBrush("NiagaraEditor.Stack.GPUIcon");
 						}
-
-						return SNullWidget::NullWidget;
-					})
-				},
-				// Simulation Target
-				{
-					FDisplayedPropertyData::FShouldDisplayProperty(),
-					FDisplayedPropertyData::FGenerateWidget::CreateLambda([](const FAssetData& AssetData) -> TSharedRef<SWidget>
-					{
-						return SNew(STextBlock)
-						.Text(LOCTEXT("SimulationTargetLabel", "Runs on"));
-					}),
-					FDisplayedPropertyData::FGenerateWidget::CreateLambda([SimulationTargetTag](const FAssetData& Asset)
-					{
-						return SNew(SImage)
-						.Image_Lambda([Asset, SimulationTargetTag]() -> const FSlateBrush*
+						else
 						{
-							if(Asset.FindTag(SimulationTargetTag))
-							{
-								FString Value;
-								if(Asset.GetTagValue(SimulationTargetTag, Value))
-								{
-									if(Value == "True")
-									{
-										return FNiagaraEditorStyle::Get().GetBrush("NiagaraEditor.Stack.GPUIcon");
-									}
-									else
-									{
-										return FNiagaraEditorStyle::Get().GetBrush("NiagaraEditor.Stack.CPUIcon");
-									}
-								}
-							}
-
-							return FAppStyle::GetNoBrush();
-						});
-					})
+							return FNiagaraEditorStyle::Get().GetBrush("NiagaraEditor.Stack.CPUIcon");
+						}
+					}
 				}
+
+				return FAppStyle::GetNoBrush();
+			});
+		});
+
+		EmitterClassInfo.DisplayedProperties =
+		{
+			InheritanceProperty,
+			SimulationTargetProperty
 		};
+		
 		EmitterClassInfo.GetDescriptionDelegate = FNiagaraAssetDetailClassInfo::FGetDescription::CreateLambda([](const FAssetData& AssetData)
 		{
 			FName DescriptionTagName("TemplateAssetDescription");
@@ -104,31 +121,34 @@ void FNiagaraAssetDetailDatabase::Init()
 	// Systems
 	{
 		FNiagaraAssetDetailClassInfo SystemClassInfo;
-		SystemClassInfo.DisplayedProperties =
+
+		// Number of emitters
+		FDisplayedPropertyData NumEmittersProperty;
+		NumEmittersProperty.ShouldDisplayPropertyDelegate = FDisplayedPropertyData::FShouldDisplayProperty::CreateLambda([](const FAssetData& AssetData) -> bool
+		{
+			return AssetData.FindTag("NumEmitters");
+		});
+		NumEmittersProperty.NameWidgetDelegate = FDisplayedPropertyData::FGenerateWidget::CreateLambda([](const FAssetData& AssetData) -> TSharedRef<SWidget>
 		{
-			// Number of emitters
+			return SNew(STextBlock)
+			.Text(LOCTEXT("NumberOfEmitters", "Number of Emitters"));
+		});
+		NumEmittersProperty.ValueWidgetDelegate =  FDisplayedPropertyData::FGenerateWidget::CreateLambda([](const FAssetData& AssetData) -> TSharedRef<SWidget>
+		{
+			int32 NumEmitters = INDEX_NONE;
+			if(AssetData.GetTagValue("NumEmitters", NumEmitters))
 			{
-				FDisplayedPropertyData::FShouldDisplayProperty::CreateLambda([](const FAssetData& AssetData) -> bool
-				{
-					return AssetData.FindTag("NumEmitters");
-				}),
-				FDisplayedPropertyData::FGenerateWidget::CreateLambda([](const FAssetData& AssetData) -> TSharedRef<SWidget>
-				{
-					return SNew(STextBlock)
-					.Text(LOCTEXT("NumberOfEmitters", "Number of Emitters"));
-				}),
-				FDisplayedPropertyData::FGenerateWidget::CreateLambda([](const FAssetData& AssetData) -> TSharedRef<SWidget>
-				{
-					int32 NumEmitters = INDEX_NONE;
-					if(AssetData.GetTagValue("NumEmitters", NumEmitters))
-					{
-						return SNew(STextBlock).Text(FText::AsNumber(NumEmitters));
-					}
-
-					return SNullWidget::NullWidget;
-				})
+				return SNew(STextBlock).Text(FText::AsNumber(NumEmitters));
 			}
+
+			return SNullWidget::NullWidget;
+		});
+		
+		SystemClassInfo.DisplayedProperties =
+		{
+			NumEmittersProperty
 		};
+		
 		SystemClassInfo.GetDescriptionDelegate = FNiagaraAssetDetailClassInfo::FGetDescription::CreateLambda([](const FAssetData& AssetData)
 		{
 			FName DescriptionTagName("TemplateAssetDescription");
@@ -337,7 +357,17 @@ void SNiagaraSelectedAssetDetails::Construct(const FArguments& InArgs, const FAs
 			.AutoHeight()
 			.Padding(0.f, 6.f)
 			[
-				CreateOptionalPropertiesList()
+				SNew(SBox)
+				.MaxDesiredWidth(InArgs._MaxDesiredPropertiesWidth)
+				[
+					CreateOptionalPropertiesList()
+				]
+			]
+			+ SVerticalBox::Slot()
+			.AutoHeight()
+			.Padding(0.f, 6.f)
+			[
+				CreatePathWidget()
 			]
 			+ SVerticalBox::Slot()
 			.AutoHeight()
@@ -349,6 +379,12 @@ void SNiagaraSelectedAssetDetails::Construct(const FArguments& InArgs, const FAs
 	]; 
 }
 
+FReply SNiagaraSelectedAssetDetails::CopyAssetPathToClipboard() const
+{
+	FPlatformApplicationMisc::ClipboardCopy(*AssetData.PackagePath.ToString());
+	return FReply::Handled();
+}
+
 TSharedRef<SWidget> SNiagaraSelectedAssetDetails::CreateAssetThumbnailWidget()
 {
 	CurrentAssetThumbnail = MakeShared<FAssetThumbnail>(AssetData, 256.f, 256.f, UThumbnailManager::Get().GetSharedThumbnailPool());
@@ -408,6 +444,31 @@ TSharedRef<SWidget> SNiagaraSelectedAssetDetails::CreateTypeWidget()
 	];
 }
 
+TSharedRef<SWidget> SNiagaraSelectedAssetDetails::CreatePathWidget()
+{
+	return SNew(SHorizontalBox)
+		+ SHorizontalBox::Slot()
+		.AutoWidth()
+		[
+			SNew(STextBlock)
+			.Text(LOCTEXT("Path", "Path"))
+			.Font(FCoreStyle::GetDefaultFontStyle("Regular", 8))
+		]
+		+ SHorizontalBox::Slot()
+		.Padding(5.f, 0.f)
+		[
+			SNew(SSpacer)
+		]
+		+ SHorizontalBox::Slot()
+		.AutoWidth()
+		[
+			SNew(STextBlock)
+			.Text(FText::FromName(AssetData.PackagePath))
+			.Font(FCoreStyle::GetDefaultFontStyle("Regular", 8))
+			.AutoWrapText(true)
+		];
+}
+
 TSharedRef<SWidget> SNiagaraSelectedAssetDetails::CreateDescriptionWidget()
 {
 	if(FNiagaraAssetDetailDatabase::NiagaraAssetDetailDatabase.Contains(AssetData.GetClass()))
@@ -437,12 +498,8 @@ TSharedRef<SWidget> SNiagaraSelectedAssetDetails::CreateOptionalPropertiesList()
 					continue;
 				}
 			}
-			
-			Result->AddSlot()
-			.AutoHeight()
-			.Padding(10.f, 3.f)
-			[
-				SNew(SHorizontalBox)
+
+			TSharedRef<SHorizontalBox> DisplayedPropertyBox = SNew(SHorizontalBox)
 				+ SHorizontalBox::Slot()
 				.HAlign(HAlign_Left)
 				.AutoWidth()
@@ -451,14 +508,24 @@ TSharedRef<SWidget> SNiagaraSelectedAssetDetails::CreateOptionalPropertiesList()
 				]
 				+ SHorizontalBox::Slot()
 				[
-					SNew(SSpacer)
+					SNew(SBox)
+					.MinDesiredWidth(20.f)
+					[
+						SNew(SSpacer)
+					]
 				]
 				+ SHorizontalBox::Slot()
 				.HAlign(HAlign_Right)
-				.AutoWidth()
+				.FillContentWidth(2.f)
 				[
 					DisplayedProperty.ValueWidgetDelegate.Execute(AssetData)
-				]
+				];
+			
+			Result->AddSlot()
+			.AutoHeight()
+			.Padding(10.f, 3.f)
+			[
+				DisplayedPropertyBox
 			];
 
 			Result->AddSlot()
diff --git a/Engine/Plugins/FX/Niagara/Source/NiagaraEditor/Public/Widgets/AssetBrowser/SNiagaraSelectedAssetDetails.h b/Engine/Plugins/FX/Niagara/Source/NiagaraEditor/Public/Widgets/AssetBrowser/SNiagaraSelectedAssetDetails.h
--- a/Engine/Plugins/FX/Niagara/Source/NiagaraEditor/Public/Widgets/AssetBrowser/SNiagaraSelectedAssetDetails.h
+++ b/Engine/Plugins/FX/Niagara/Source/NiagaraEditor/Public/Widgets/AssetBrowser/SNiagaraSelectedAssetDetails.h
@@ -80,6 +80,7 @@ public:
 	SLATE_BEGIN_ARGS(SNiagaraSelectedAssetDetails)
 		: _ShowThumbnail(EVisibility::Visible)
 		, _MaxDesiredDescriptionWidth(300.f)
+		, _MaxDesiredPropertiesWidth(300.f)
 		{
 		}
 		SLATE_ATTRIBUTE(EVisibility, ShowThumbnail)
@@ -91,10 +92,16 @@ public:
 		 * Actual width might be increased due to big asset names.
 		 */
 		SLATE_ARGUMENT(float, MaxDesiredDescriptionWidth)
+		SLATE_ARGUMENT(float, MaxDesiredPropertiesWidth)
+
 	SLATE_END_ARGS()
 	
 	/** Constructs this widget with InArgs */
 	void Construct(const FArguments& InArgs, const FAssetData& Asset);
+
+private:
+	FReply CopyAssetPathToClipboard() const;
+	
 private:
 	FAssetData AssetData;
 	TAttribute<EVisibility> ShowThumbnail;
@@ -105,6 +112,7 @@ private:
 	TSharedRef<SWidget> CreateAssetThumbnailWidget();
 	TSharedRef<SWidget> CreateTitleWidget();
 	TSharedRef<SWidget> CreateTypeWidget();
+	TSharedRef<SWidget> CreatePathWidget();
 	TSharedRef<SWidget> CreateDescriptionWidget();
 	TSharedRef<SWidget> CreateOptionalPropertiesList();
 	TSharedRef<SWidget> CreateAssetTagRow();
