- Add extra clarity to tooltip

#rnx
#tests editor to observe new tooptip

diff --git a/Engine/Plugins/FX/Niagara/Source/Niagara/Classes/NiagaraEffectType.h b/Engine/Plugins/FX/Niagara/Source/Niagara/Classes/NiagaraEffectType.h
--- a/Engine/Plugins/FX/Niagara/Source/Niagara/Classes/NiagaraEffectType.h
+++ b/Engine/Plugins/FX/Niagara/Source/Niagara/Classes/NiagaraEffectType.h
@@ -418,7 +418,10 @@ class UNiagaraEffectType : public UObject
 	UPROPERTY(EditAnywhere, Category = "Scalability", meta=(DisplayInSystemScalability, ScalabilityBarDisplayName="Frequency"))
 	ENiagaraScalabilityUpdateFrequency UpdateFrequency;
 
-	/** How effects of this type react when they fail the cull checks. */
+	/**
+	How effects of this type react when they fail the cull checks.
+	Applied to all effects using this effect type and can not be overridden per effect.
+	*/
 	UPROPERTY(EditAnywhere, Category = "Scalability", meta=(DisplayInSystemScalability, ScalabilityBarDisplayName="Reaction"))
 	ENiagaraCullReaction CullReaction;
 
