- Detect particle reads from system / emitter script reading GPU script data
- Only error about CPU <-> GPU if using persistent / index functions as the others are safe

#tests Tested in system / emitter / particle scripts using CPU & GPU emitters with persistent / index functions both valid attribute and invalid

diff --git a/Engine/Plugins/FX/Niagara/Source/Niagara/Private/NiagaraDataInterfaceParticleRead.cpp b/Engine/Plugins/FX/Niagara/Source/Niagara/Private/NiagaraDataInterfaceParticleRead.cpp
--- a/Engine/Plugins/FX/Niagara/Source/Niagara/Private/NiagaraDataInterfaceParticleRead.cpp
+++ b/Engine/Plugins/FX/Niagara/Source/Niagara/Private/NiagaraDataInterfaceParticleRead.cpp
@@ -2405,10 +2405,12 @@ void UNiagaraDataInterfaceParticleRead::GetFeedback(UNiagaraSystem* Asset, UNiag
 	}
 
 	// Now check if any script uses functions that require persisitent ID access
-	TArray<FNiagaraFunctionSignature> CPUFunctions;
-	GetPersistentIDFunctions(CPUFunctions);
+	TArray<FNiagaraFunctionSignature> PersistentIDFunctions;
+	TArray<FNiagaraFunctionSignature> IndexFunctions;
+	GetPersistentIDFunctions(PersistentIDFunctions);
+	GetIndexFunctions(IndexFunctions);
 
-	bool bHasPersistenIDAccessWarning = [this, &Scripts, &CPUFunctions]()
+	bool bHasPersistenIDAccessWarning = [this, &Scripts, &PersistentIDFunctions]()
 	{
 		for (const auto Script : Scripts)
 		{
@@ -2422,7 +2424,7 @@ void UNiagaraDataInterfaceParticleRead::GetFeedback(UNiagaraSystem* Asset, UNiag
 						{
 							return CPUSig.Name == Func.Name;
 						};
-						if (CPUFunctions.FindByPredicate(Filter))
+						if (PersistentIDFunctions.FindByPredicate(Filter))
 						{
 							return true;
 						}
@@ -2441,74 +2443,78 @@ void UNiagaraDataInterfaceParticleRead::GetFeedback(UNiagaraSystem* Asset, UNiag
 		for (int32 Idx = 0; Idx < Script->GetVMExecutableData().DataInterfaceInfo.Num(); Idx++)
 		{
 			const auto& DIInfo = Script->GetVMExecutableData().DataInterfaceInfo[Idx];
-			if (DIInfo.MatchesClass(GetClass()))
+			if (!DIInfo.MatchesClass(GetClass()))
 			{
-				bool bMatchFound = false;
-				// We assume that if the properties match, it's a valid match for us.
-				if (CachedDefaultDIs.IsValidIndex(Idx) && CachedDefaultDIs[Idx].DataInterface != nullptr && CachedDefaultDIs[Idx].DataInterface->Equals(this))
+				continue;
+			}
+
+			if (!CachedDefaultDIs.IsValidIndex(Idx) || CachedDefaultDIs[Idx].DataInterface == nullptr || !CachedDefaultDIs[Idx].DataInterface->Equals(this))
+			{
+				continue;
+			}
+
+			bool bNeedsCPUGPUCheck = false;
+			for (const auto& Func : DIInfo.RegisteredFunctions)
+			{
+				const FName* AttributeName = Func.FunctionSpecifiers.Find(UNiagaraDataInterfaceRWBase::NAME_Attribute);
+				ENiagaraParticleDataValueType AttributeType = ENiagaraParticleDataValueType::Invalid;
+				if (AttributeName != nullptr)
 				{
-					bMatchFound = true;
-					FVersionedNiagaraEmitter OuterEmitter = Script->GetOuterEmitter();
-					if (OuterEmitter.GetEmitterData() && FoundSourceEmitter.GetEmitterData())
-					{
-						if (OuterEmitter.GetEmitterData()->SimTarget != FoundSourceEmitter.GetEmitterData()->SimTarget)
-						{
-							FText Msg = FText::Format(LOCTEXT("SourceEmitterSimTypeMismatchError", "Emitter \"{0}\" SimTarget not compatible (CPU vs GPU)!"), FText::FromName(OuterEmitter.Emitter->GetFName()));
-							FNiagaraDataInterfaceError SourceEmitterNotFoundError(
-								Msg, Msg,
-								FNiagaraDataInterfaceFix());
-							OutErrors.AddUnique(SourceEmitterNotFoundError);
-						}
-					}
+					AttributeType = GetValueTypeFromFuncName(Func.Name);
 				}
 
-				if (bMatchFound)
+				if (!bNeedsCPUGPUCheck)
 				{
-					for (const auto& Func : DIInfo.RegisteredFunctions)
+					bNeedsCPUGPUCheck = 
+						PersistentIDFunctions.ContainsByPredicate([Func](const FNiagaraFunctionSignature& Sig) { return Sig.Name == Func.Name;}) ||
+						IndexFunctions.ContainsByPredicate([Func](const FNiagaraFunctionSignature& Sig) { return Sig.Name == Func.Name; });
+				}
+
+				if (AttributeName && FoundSourceEmitter.GetEmitterData())
+				{
+					if (AttributeType != ENiagaraParticleDataValueType::Invalid)
 					{
-						const FName* AttributeName = Func.FunctionSpecifiers.Find(UNiagaraDataInterfaceRWBase::NAME_Attribute);
-						ENiagaraParticleDataValueType AttributeType = ENiagaraParticleDataValueType::Invalid;
-						if (AttributeName != nullptr)
+						auto AttribFilter = [AttributeName](const FNiagaraVariable& Var)
 						{
-							AttributeType = GetValueTypeFromFuncName(Func.Name);
-						}
+							return Var.GetName() == *AttributeName;
+						};
 
-						if (AttributeName && FoundSourceEmitter.GetEmitterData())
+						TArray<FNiagaraVariableBase> Variables;
+						FoundSourceEmitter.GetEmitterData()->GatherCompiledParticleAttributes(Variables);
+						const FNiagaraVariableBase* FoundVar = Variables.FindByPredicate(AttribFilter);
+						if (FoundVar && !CheckVariableType(FoundVar->GetType(), AttributeType))
 						{
-							if (AttributeType != ENiagaraParticleDataValueType::Invalid)
-							{
-								auto AttribFilter = [AttributeName](const FNiagaraVariable& Var)
-								{
-									return Var.GetName() == *AttributeName;
-								};
-
-								TArray<FNiagaraVariableBase> Variables;
-								FoundSourceEmitter.GetEmitterData()->GatherCompiledParticleAttributes(Variables);
-								const FNiagaraVariableBase* FoundVar = Variables.FindByPredicate(AttribFilter);
-								if (FoundVar && !CheckVariableType(FoundVar->GetType(), AttributeType))
-								{
-									FText Msg = FText::Format(LOCTEXT("SourceEmitterTypeMismatchError", "Source Emitter has attribute named, \"{0}\" but the type isn't compatible with the function \"{1}\", and will not succeed."), FText::FromName(*AttributeName), FText::FromName(Func.Name));
-
-									FNiagaraDataInterfaceFeedback MissingByType(Msg, Msg,
-										FNiagaraDataInterfaceFix());
-
-									Info.AddUnique(MissingByType);
-
-								}
-								else if (!FoundVar)
-								{
-									FText Msg = FText::Format(LOCTEXT("SourceEmitterNameMismatchError", "Source Emitter does not have attribute named, \"{0}\" referenced by function \"{1}\", and will not succeed."), FText::FromName(*AttributeName), FText::FromName(Func.Name));
-									FNiagaraDataInterfaceFeedback MissingByName(
-										Msg, Msg,
-										FNiagaraDataInterfaceFix());
-
-									Info.AddUnique(MissingByName);
-								}
-							}
+							FText Msg = FText::Format(LOCTEXT("SourceEmitterTypeMismatchError", "Source Emitter has attribute named, \"{0}\" but the type isn't compatible with the function \"{1}\", and will not succeed."), FText::FromName(*AttributeName), FText::FromName(Func.Name));
+							Info.AddUnique(FNiagaraDataInterfaceFeedback(Msg, Msg, FNiagaraDataInterfaceFix()));
+						}
+						else if (!FoundVar)
+						{
+							FText Msg = FText::Format(LOCTEXT("SourceEmitterNameMismatchError", "Source Emitter does not have attribute named, \"{0}\" referenced by function \"{1}\", and will not succeed."), FText::FromName(*AttributeName), FText::FromName(Func.Name));
+							Info.AddUnique(FNiagaraDataInterfaceFeedback(Msg, Msg, FNiagaraDataInterfaceFix()));
 						}
 					}
 				}
 			}
+
+			// We assume that if the properties match, it's a valid match for us.
+			if (bNeedsCPUGPUCheck)
+			{
+				FVersionedNiagaraEmitter OuterEmitter = Script->GetOuterEmitter();
+				if (OuterEmitter.GetEmitterData() && FoundSourceEmitter.GetEmitterData())
+				{
+					if (OuterEmitter.GetEmitterData()->SimTarget != FoundSourceEmitter.GetEmitterData()->SimTarget)
+					{
+						FText Msg = FText::Format(LOCTEXT("SourceEmitterSimTypeMismatchError", "Emitter \"{0}\" SimTarget not compatible (CPU vs GPU)!"), FText::FromName(OuterEmitter.Emitter->GetFName()));
+						OutErrors.AddUnique(FNiagaraDataInterfaceError(Msg, Msg, FNiagaraDataInterfaceFix()));
+					}
+				}
+				// This is likely from a system script
+				else if (FoundSourceEmitter.GetEmitterData() && FoundSourceEmitter.GetEmitterData()->SimTarget == ENiagaraSimTarget::GPUComputeSim)
+				{
+					FText Msg = LOCTEXT("SourceSystemSimTypeMismatchError", "Sytem script can not read from GPU emitter!");
+					OutErrors.AddUnique(FNiagaraDataInterfaceError(Msg, Msg, FNiagaraDataInterfaceFix()));
+				}
+			}
 		}
 	}
 
