- Check the velocity texture has been produced before using in gbuffer DI

#tests Made NS with GBuffer DI and dropped into level

diff --git a/Engine/Plugins/FX/Niagara/Source/Niagara/Private/DataInterface/NiagaraDataInterfaceGBuffer.cpp b/Engine/Plugins/FX/Niagara/Source/Niagara/Private/DataInterface/NiagaraDataInterfaceGBuffer.cpp
--- a/Engine/Plugins/FX/Niagara/Source/Niagara/Private/DataInterface/NiagaraDataInterfaceGBuffer.cpp
+++ b/Engine/Plugins/FX/Niagara/Source/Niagara/Private/DataInterface/NiagaraDataInterfaceGBuffer.cpp
@@ -184,7 +184,7 @@ void UNiagaraDataInterfaceGBuffer::SetShaderParameters(const FNiagaraDataInterfa
 			VelocityTexture = UE::FXRenderingUtils::GetSceneVelocityTexture(SimulationSceneViews[0]);
 		}
 		
-		if (VelocityTexture == nullptr)
+		if (VelocityTexture == nullptr || !VelocityTexture->HasBeenProduced())
 		{
 			VelocityTexture = Context.GetComputeDispatchInterface().GetBlackTexture(Context.GetGraphBuilder(), ETextureDimension::Texture2D);
 		}
