Remove dependency on FGraphActionListBuilderBase::ActionGroup's RootCategory from niagara parameter menu, which is the only producer

diff --git a/Engine/Plugins/FX/Niagara/Source/NiagaraEditor/Private/NiagaraActions.cpp b/Engine/Plugins/FX/Niagara/Source/NiagaraEditor/Private/NiagaraActions.cpp
--- a/Engine/Plugins/FX/Niagara/Source/NiagaraEditor/Private/NiagaraActions.cpp
+++ b/Engine/Plugins/FX/Niagara/Source/NiagaraEditor/Private/NiagaraActions.cpp
@@ -52,9 +52,9 @@ void FNiagaraMenuAction::SetParameterVariable(const FNiagaraVariable& InParamete
 	ParameterVariable = InParameterVariable;
 }
 
-void FNiagaraMenuActionCollector::AddAction(TSharedPtr<FNiagaraMenuAction> Action, int32 SortOrder, const FString& Category)
+void FNiagaraMenuActionCollector::AddAction(TSharedPtr<FNiagaraMenuAction> Action, int32 SortOrder)
 {
-	Actions.Add({Action, SortOrder, Category});
+	Actions.Add({Action, SortOrder});
 }
 
 void FNiagaraMenuActionCollector::AddAllActionsTo(FGraphActionListBuilderBase& ActionBuilder)
@@ -97,7 +97,7 @@ void FNiagaraMenuActionCollector::AddAllActionsTo(FGraphActionListBuilderBase& A
 	
 	for (const FCollectedAction& Entry : Actions)
 	{
-		ActionBuilder.AddAction(Entry.Action, Entry.Category);
+		ActionBuilder.AddAction(Entry.Action);
 	}
 }
 
diff --git a/Engine/Plugins/FX/Niagara/Source/NiagaraEditor/Public/NiagaraActions.h b/Engine/Plugins/FX/Niagara/Source/NiagaraEditor/Public/NiagaraActions.h
--- a/Engine/Plugins/FX/Niagara/Source/NiagaraEditor/Public/NiagaraActions.h
+++ b/Engine/Plugins/FX/Niagara/Source/NiagaraEditor/Public/NiagaraActions.h
@@ -93,7 +93,7 @@ private:
 class FNiagaraMenuActionCollector
 {
 public:
-	void AddAction(TSharedPtr<FNiagaraMenuAction> Action, int32 SortOrder, const FString& Category = FString());
+	void AddAction(TSharedPtr<FNiagaraMenuAction> Action, int32 SortOrder);
 	void AddAllActionsTo(FGraphActionListBuilderBase& ActionBuilder);
 
 private:
@@ -101,7 +101,6 @@ private:
 	{
 		TSharedPtr<FNiagaraMenuAction> Action;
 		int32 SortOrder;
-		FString Category;
 	};
 
 	TArray<FCollectedAction> Actions;
diff --git a/Engine/Plugins/FX/Niagara/Source/NiagaraEditor/Public/Widgets/SNiagaraParameterMenu.cpp b/Engine/Plugins/FX/Niagara/Source/NiagaraEditor/Public/Widgets/SNiagaraParameterMenu.cpp
--- a/Engine/Plugins/FX/Niagara/Source/NiagaraEditor/Public/Widgets/SNiagaraParameterMenu.cpp
+++ b/Engine/Plugins/FX/Niagara/Source/NiagaraEditor/Public/Widgets/SNiagaraParameterMenu.cpp
@@ -28,7 +28,37 @@
 #include "Widgets/SNiagaraActionMenuExpander.h"
 
 #define LOCTEXT_NAMESPACE "SNiagaraParameterMenu"
+namespace UE::Private
+{
+	template<typename... Strings>
+	FText ConcatinateCategoriesIfNotEmpty(Strings... InCategories);
+	FText ConcatinateCategoriesIfNotEmptyImpl(TConstArrayView<FStringView> Paths);
+}
+
+template<typename... Strings>
+FText UE::Private::ConcatinateCategoriesIfNotEmpty(Strings... InCategories)
+{
+	return ConcatinateCategoriesIfNotEmptyImpl({ FStringView(InCategories)... });
+}
 
+FText UE::Private::ConcatinateCategoriesIfNotEmptyImpl(TConstArrayView<FStringView> Paths)
+{
+	FString Result;
+	bool bIsFirst = true;
+	for (FStringView String : Paths)
+	{
+		if (!String.IsEmpty())
+		{
+			if (!bIsFirst)
+			{
+				Result.Append(TEXT("|"));
+			}
+			Result.Append(String);
+			bIsFirst = false;
+		}
+	}
+	return FText::FromString(MoveTemp(Result));
+}
 ///////////////////////////////////////////////////////////////////////////////
 /// Base Parameter Menu														///
 ///////////////////////////////////////////////////////////////////////////////
@@ -204,8 +234,7 @@ void SNiagaraAddParameterFromPanelMenu::AddParameterGroup(
 	TArray<FNiagaraVariable>& Variables,
 	const FGuid& InNamespaceId /*= FGuid()*/,
 	const FText& Category /*= FText::GetEmpty()*/,
-	int32 SortOrder /*= 0*/,
-	const FString& RootCategory /*= FString()*/)
+	int32 SortOrder /*= 0*/)
 {
 	for (const FNiagaraVariable& Variable : Variables)
 	{
@@ -239,12 +268,16 @@ void SNiagaraAddParameterFromPanelMenu::AddParameterGroup(
 			}
 		}
 
-		Collector.AddAction(Action, SortOrder, RootCategory);
+		Collector.AddAction(Action, SortOrder);
 	}
 }
 
-void SNiagaraAddParameterFromPanelMenu::AddMakeNewGroup(FNiagaraMenuActionCollector& Collector,
-	TArray<FNiagaraTypeDefinition>& TypeDefinitions, const FGuid& InNamespaceId, const FText& Category, int32 SortOrder,
+void SNiagaraAddParameterFromPanelMenu::AddMakeNewGroup(
+	FNiagaraMenuActionCollector& Collector,
+	TArray<FNiagaraTypeDefinition>& TypeDefinitions, 
+	const FGuid& InNamespaceId, 
+	const FText& Category, 
+	int32 SortOrder,
 	const FString& RootCategory)
 {
 	for (const FNiagaraTypeDefinition& TypeDefinition : TypeDefinitions)
@@ -252,14 +285,13 @@ void SNiagaraAddParameterFromPanelMenu::AddMakeNewGroup(FNiagaraMenuActionCollec
 		const FText DisplayName = TypeDefinition.GetNameText();
 		FText Tooltip = FText::GetEmpty();
 
-
 		if (const UStruct* VariableStruct = TypeDefinition.GetStruct())
 		{
 			Tooltip = VariableStruct->GetToolTipText(true);
 		}
 
 		FText SubCategory = FNiagaraEditorUtilities::GetTypeDefinitionCategory(TypeDefinition);
-		FText FullCategory = SubCategory.IsEmpty() ? Category : FText::Format(FText::FromString("{0}|{1}"), Category, SubCategory);
+		FText FullCategory = UE::Private::ConcatinateCategoriesIfNotEmpty(RootCategory, Category.ToString(), SubCategory.ToString() );
 		TSharedPtr<FNiagaraMenuAction> Action(new FNiagaraMenuAction(FullCategory, DisplayName, Tooltip, 0, FText(),
 			FNiagaraMenuAction::FOnExecuteStackAction::CreateSP(this, &SNiagaraAddParameterFromPanelMenu::NewParameterSelected, TypeDefinition, InNamespaceId)));
 
@@ -271,7 +303,7 @@ void SNiagaraAddParameterFromPanelMenu::AddMakeNewGroup(FNiagaraMenuActionCollec
 			}
 		}
 
-		Collector.AddAction(Action, SortOrder, RootCategory);
+		Collector.AddAction(Action, SortOrder);
 	}
 }
 
@@ -283,7 +315,7 @@ void SNiagaraAddParameterFromPanelMenu::CollectParameterCollectionsActions(FNiag
 	FNiagaraEditorUtilities::GetAvailableParameterCollections(AvailableParameterCollections);
 	for (UNiagaraParameterCollection* Collection : AvailableParameterCollections)
 	{
-		AddParameterGroup(Collector, Collection->GetParameters(), FNiagaraEditorGuids::ParameterCollectionNamespaceMetaDataGuid, Category, 10, FString());
+		AddParameterGroup(Collector, Collection->GetParameters(), FNiagaraEditorGuids::ParameterCollectionNamespaceMetaDataGuid, Category, 10);
 	}
 }
 
@@ -294,28 +326,24 @@ void SNiagaraAddParameterFromPanelMenu::CollectAllActions(FGraphActionListBuilde
 		const FNiagaraNamespaceMetadata NamespaceMetaData = FNiagaraEditorUtilities::GetNamespaceMetaDataForId(FNiagaraEditorGuids::EngineNamespaceMetaDataGuid);
 		TArray<FNiagaraVariable> Variables = FNiagaraConstants::GetEngineConstants();
 		const FText CategoryText = bShowNamespaceCategory ? GetNamespaceCategoryText(NamespaceMetaData) : LOCTEXT("EngineConstantNamespaceCategory", "Add Engine Constant");
-		const FString RootCategoryStr = FString();
 		AddParameterGroup(
 			Collector,
 			Variables,
 			FNiagaraEditorGuids::EngineNamespaceMetaDataGuid,
 			CategoryText,
-			4,
-			RootCategoryStr
+			4
 		);
 	};
 	auto CollectEmitterNamespaceParameterActions = [this, &Collector]() {
 		const FNiagaraNamespaceMetadata NamespaceMetaData = FNiagaraEditorUtilities::GetNamespaceMetaDataForId(FNiagaraEditorGuids::EmitterNamespaceMetaDataGuid);
 		TArray<FNiagaraVariable> Variables = FNiagaraConstants::GetEngineConstants().FilterByPredicate([](const FNiagaraVariable& Var) { return Var.IsInNameSpace(FNiagaraConstants::EmitterNamespaceString); });
 		const FText CategoryText = bShowNamespaceCategory ? GetNamespaceCategoryText(NamespaceMetaData) : LOCTEXT("EmitterConstantNamespaceCategory", "Add Emitter Constant");
-		const FString RootCategoryStr = FString();
 		AddParameterGroup(
 			Collector,
 			Variables,
 			FNiagaraEditorGuids::EngineNamespaceMetaDataGuid,
 			CategoryText,
-			4,
-			RootCategoryStr
+			4
 		);
 	};
 
@@ -401,7 +429,7 @@ void SNiagaraAddParameterFromPanelMenu::CollectAllActions(FGraphActionListBuilde
 	{
 		TArray<FNiagaraVariable> Variables;
 		Variables.Add(SYS_PARAM_INSTANCE_ALIVE);
-		AddParameterGroup(Collector, Variables, FNiagaraEditorGuids::DataInstanceNamespaceMetaDataGuid, FText(), 3, FString());
+		AddParameterGroup(Collector, Variables, FNiagaraEditorGuids::DataInstanceNamespaceMetaDataGuid, FText(), 3);
 	}
 	// No NamespaceId set but still collecting engine namespace parameters (e.g. map get/set node menu.)
 	else if (NamespaceId.IsValid() == false && bForceCollectEngineNamespaceParameterActions)
@@ -412,7 +440,7 @@ void SNiagaraAddParameterFromPanelMenu::CollectAllActions(FGraphActionListBuilde
 		// Special case; collect DataInstance.Alive so that it is an option if we are selecting a parameter from a map node in a script.
 		TArray<FNiagaraVariable> Variables;
 		Variables.Add(SYS_PARAM_INSTANCE_ALIVE);
-		AddParameterGroup(Collector, Variables, FNiagaraEditorGuids::DataInstanceNamespaceMetaDataGuid, FText(), 3, FString());
+		AddParameterGroup(Collector, Variables, FNiagaraEditorGuids::DataInstanceNamespaceMetaDataGuid, FText(), 3);
 	}
 
 	// Any other "unreserved" namespace
@@ -493,8 +521,7 @@ void SNiagaraAddParameterFromPanelMenu::CollectAllActions(FGraphActionListBuilde
 	for (UNiagaraParameterDefinitions* ParameterDefinitions : AvailableParameterDefinitions)
 	{
 		bool bTopLevelCategory = ParameterDefinitions->GetIsPromotedToTopInAddMenus();
-		const FText TopLevelCategory = FText::FromString(*ParameterDefinitions->GetName());
-		const FText Category = bTopLevelCategory ? FText() : TopLevelCategory;
+		const FText Category = bTopLevelCategory ? FText::FromString(*ParameterDefinitions->GetName()) : FText();
 		for (const UNiagaraScriptVariable* ScriptVar : ParameterDefinitions->GetParametersConst())
 		{
 			// Only add parameters in the same namespace as the target namespace id if bOnlyShowParametersInNamespaceId is set.
@@ -523,7 +550,7 @@ void SNiagaraAddParameterFromPanelMenu::CollectAllActions(FGraphActionListBuilde
 			
 			if (bTopLevelCategory)
 			{
-				Collector.AddAction(Action, ParameterDefinitions->GetMenuSortOrder(), TopLevelCategory.ToString());
+				Collector.AddAction(Action, ParameterDefinitions->GetMenuSortOrder());
 			}
 			else
 			{ 
diff --git a/Engine/Plugins/FX/Niagara/Source/NiagaraEditor/Public/Widgets/SNiagaraParameterMenu.h b/Engine/Plugins/FX/Niagara/Source/NiagaraEditor/Public/Widgets/SNiagaraParameterMenu.h
--- a/Engine/Plugins/FX/Niagara/Source/NiagaraEditor/Public/Widgets/SNiagaraParameterMenu.h
+++ b/Engine/Plugins/FX/Niagara/Source/NiagaraEditor/Public/Widgets/SNiagaraParameterMenu.h
@@ -106,8 +106,7 @@ public:
 		TArray<FNiagaraVariable>& Variables,
 		const FGuid& InNamespaceId = FGuid(),
 		const FText& Category = FText::GetEmpty(),
-		int32 SortOrder = 0,
-		const FString& RootCategory = FString());
+		int32 SortOrder = 0);
 
 	void AddMakeNewGroup(
 		FNiagaraMenuActionCollector& Collector,
