Re-enable some commented out RDG pass names in Niagara

diff --git a/Engine/Plugins/FX/Niagara/Source/Niagara/Private/NiagaraGpuComputeDispatch.cpp b/Engine/Plugins/FX/Niagara/Source/Niagara/Private/NiagaraGpuComputeDispatch.cpp
--- a/Engine/Plugins/FX/Niagara/Source/Niagara/Private/NiagaraGpuComputeDispatch.cpp
+++ b/Engine/Plugins/FX/Niagara/Source/Niagara/Private/NiagaraGpuComputeDispatch.cpp
@@ -1196,7 +1196,7 @@ void FNiagaraGpuComputeDispatch::ExecuteTicks(FRDGBuilder& GraphBuilder, TConstS
 
 		// Execute Transitions
  		GraphBuilder.AddPass(
-			{},//RDG_EVENT_NAME("Niagara::ExecuteTicks::DispatchGroupPre"),
+			RDG_EVENT_NAME("Niagara::ExecuteTicks::DispatchGroupPre"),
 			ERDGPassFlags::None,
 			[this, PreStageTransitions=MoveTemp(PreStageTransitions), PreStageIDToIndexInit=MoveTemp(PreStageIDToIndexInit), DispatchInstances=MakeArrayView(DispatchGroup.DispatchInstances)](FRHICommandList& RHICmdList)
 			{
@@ -1237,7 +1237,7 @@ void FNiagaraGpuComputeDispatch::ExecuteTicks(FRDGBuilder& GraphBuilder, TConstS
 
 		// Execute legacy Post Stage
  		GraphBuilder.AddPass(
-			{},//RDG_EVENT_NAME("Niagara::ExecuteTicks::DispatchGroupPost"),
+			RDG_EVENT_NAME("Niagara::ExecuteTicks::DispatchGroupPost"),
 			ERDGPassFlags::None,
 			[this, PostStageTransitions=MoveTemp(PostStageTransitions), DispatchInstances=MakeArrayView(DispatchGroup.DispatchInstances)](FRHICommandList& RHICmdList)
 			{
@@ -1333,7 +1333,7 @@ void FNiagaraGpuComputeDispatch::ExecuteTicks(FRDGBuilder& GraphBuilder, TConstS
 			}
 
  			GraphBuilder.AddPass(
-				{},//RDG_EVENT_NAME("Niagara::ExecuteTicks::FreeIDUpdates"),
+				RDG_EVENT_NAME("Niagara::ExecuteTicks::FreeIDUpdates"),
 				ERDGPassFlags::None,
 				[this, FreeIDUpdates=MakeArrayView(DispatchGroup.FreeIDUpdates)](FRHICommandList& RHICmdList)
 				{
