Fixed potentially uninitialized variables.

diff --git a/Engine/Plugins/Enterprise/DatasmithC4DImporter/Source/DatasmithC4DTranslator/Private/DatasmithC4DImporter.cpp b/Engine/Plugins/Enterprise/DatasmithC4DImporter/Source/DatasmithC4DTranslator/Private/DatasmithC4DImporter.cpp
--- a/Engine/Plugins/Enterprise/DatasmithC4DImporter/Source/DatasmithC4DTranslator/Private/DatasmithC4DImporter.cpp
+++ b/Engine/Plugins/Enterprise/DatasmithC4DImporter/Source/DatasmithC4DTranslator/Private/DatasmithC4DImporter.cpp
@@ -930,7 +930,7 @@ namespace
 	// Gets a color weighted by its brightness
 	FVector MelangeGetLayerColor(cineware::BaseList2D* MelangeObject, cineware::Int32 ColorAttributeID, cineware::Int32 BrightnessAttributeID)
 	{
-		FVector Result;
+		FVector Result = FVector::ZeroVector;
 		if (MelangeObject)
 		{
 			float Brightness = MelangeGetFloat(MelangeObject, BrightnessAttributeID);
@@ -943,7 +943,7 @@ namespace
 	// In here instead of utils because it depends on the document color profile
 	FVector MelangeGetColor(cineware::BaseList2D* MelangeObject, cineware::Int32 MelangeDescId)
 	{
-		FVector Result;
+		FVector Result = FVector::ZeroVector;
 		if (MelangeObject)
 		{
 			Result = ToLinearColor(MelangeGetVector(MelangeObject, MelangeDescId));
diff --git a/Engine/Plugins/Experimental/Avalanche/Source/AvalancheComponentVisualizers/Private/AvaVisBase.cpp b/Engine/Plugins/Experimental/Avalanche/Source/AvalancheComponentVisualizers/Private/AvaVisBase.cpp
--- a/Engine/Plugins/Experimental/Avalanche/Source/AvalancheComponentVisualizers/Private/AvaVisBase.cpp
+++ b/Engine/Plugins/Experimental/Avalanche/Source/AvalancheComponentVisualizers/Private/AvaVisBase.cpp
@@ -388,7 +388,7 @@ EAxisList::Type FAvaVisualizerBase::GetViewportWidgetAxisList(FEditorViewportCli
 
 FVector FAvaVisualizerBase::GetLocalVector(FEditorViewportClient* InViewport, const FVector& InVector) const
 {
-	FVector OutVector;
+	FVector OutVector = FVector::ZeroVector;
 	FMatrix RotMatrix;
 
 	if (GetCustomInputCoordinateSystem(InViewport, RotMatrix))
diff --git a/Engine/Plugins/Experimental/Avalanche/Source/AvalancheTextEditor/Private/DetailsView/AvaLinearGradientSettingsCustomization.cpp b/Engine/Plugins/Experimental/Avalanche/Source/AvalancheTextEditor/Private/DetailsView/AvaLinearGradientSettingsCustomization.cpp
--- a/Engine/Plugins/Experimental/Avalanche/Source/AvalancheTextEditor/Private/DetailsView/AvaLinearGradientSettingsCustomization.cpp
+++ b/Engine/Plugins/Experimental/Avalanche/Source/AvalancheTextEditor/Private/DetailsView/AvaLinearGradientSettingsCustomization.cpp
@@ -207,7 +207,7 @@ bool FAvaLinearGradientSettingsCustomization::IsCustomDirectionEnabled() const
 FLinearColor FAvaLinearGradientSettingsCustomization::GetLinearColorFromProperty(const TSharedPtr<IPropertyHandle>& InColorPropertyHandle)
 {
 	// Default to full alpha in case the alpha component is disabled.
-	FLinearColor OutColor;
+	FLinearColor OutColor = FLinearColor::Black;
 
 	FString StringValue;
 	const FPropertyAccess::Result Result = InColorPropertyHandle->GetValueAsFormattedString(StringValue);
diff --git a/Engine/Plugins/Experimental/Avalanche/Source/AvalancheTextEditor/Private/Visualizer/AvaTextVisualizer.cpp b/Engine/Plugins/Experimental/Avalanche/Source/AvalancheTextEditor/Private/Visualizer/AvaTextVisualizer.cpp
--- a/Engine/Plugins/Experimental/Avalanche/Source/AvalancheTextEditor/Private/Visualizer/AvaTextVisualizer.cpp
+++ b/Engine/Plugins/Experimental/Avalanche/Source/AvalancheTextEditor/Private/Visualizer/AvaTextVisualizer.cpp
@@ -817,7 +817,7 @@ FVector FAvaTextVisualizer::GetHeightHandleLocation(const UText3DComponent* InTe
 		? BoundsMax.Min.Y - 10.f
 		: BoundsMax.Max.Y + 10.f;
 
-	FVector IconLocation;
+	FVector IconLocation = FVector::ZeroVector;
 
 	switch (InText3DComp->GetVerticalAlignment())
 	{
diff --git a/Engine/Plugins/Experimental/FloatingProperties/Source/FloatingProperties/Private/Data/FloatingPropertiesPropertyNode.cpp b/Engine/Plugins/Experimental/FloatingProperties/Source/FloatingProperties/Private/Data/FloatingPropertiesPropertyNode.cpp
--- a/Engine/Plugins/Experimental/FloatingProperties/Source/FloatingProperties/Private/Data/FloatingPropertiesPropertyNode.cpp
+++ b/Engine/Plugins/Experimental/FloatingProperties/Source/FloatingProperties/Private/Data/FloatingPropertiesPropertyNode.cpp
@@ -504,7 +504,7 @@ EFloatingPropertiesUpdateResult FFloatingPropertiesPropertyNode::UpdatePropertyN
 
 FVector2f FFloatingPropertiesPropertyNode::CalculateAnchorMultiplier(const FFloatingPropertiesClassPropertyPosition& InPropertyPosition)
 {
-	FVector2f AnchorMultiplier;
+	FVector2f AnchorMultiplier = FVector2f::ZeroVector;
 
 	switch (InPropertyPosition.HorizontalAnchor.GetValue())
 	{
diff --git a/Engine/Plugins/FX/Niagara/Source/Niagara/Public/NiagaraParameterStore.h b/Engine/Plugins/FX/Niagara/Source/Niagara/Public/NiagaraParameterStore.h
--- a/Engine/Plugins/FX/Niagara/Source/Niagara/Public/NiagaraParameterStore.h
+++ b/Engine/Plugins/FX/Niagara/Source/Niagara/Public/NiagaraParameterStore.h
@@ -1019,7 +1019,7 @@ struct FNiagaraParameterDirectBinding<FMatrix44f>
 		checkSlow(BoundVariable.GetSizeInBytes() == sizeof(FMatrix44f));
 		checkfSlow(LayoutVersion == BoundStore->GetLayoutVersion(), TEXT("This binding is invalid, its bound parameter store's layout was changed since it was created"));
 #endif
-		FMatrix44f Ret;
+		FMatrix44f Ret = FMatrix44f::Identity;
 		if (ValuePtr)
 		{
 			FMemory::Memcpy(&Ret, ValuePtr, sizeof(FMatrix44f));//Temp annoyance until we fix the alignment issues with parameter stores.
@@ -1135,7 +1135,7 @@ struct FNiagaraParameterDirectBinding<FQuat4f>
 		checkSlow(BoundVariable.GetSizeInBytes() == sizeof(FQuat4f));
 		checkfSlow(LayoutVersion == BoundStore->GetLayoutVersion(), TEXT("This binding is invalid, its bound parameter store's layout was changed since it was created"));
 #endif
-		FQuat4f Ret;
+		FQuat4f Ret = FQuat4f::Identity;
 		if (ValuePtr)
 		{
 			FMemory::Memcpy(&Ret, ValuePtr, sizeof(FQuat4f));//Temp annoyance until we fix the alignment issues with parameter stores.
diff --git a/Engine/Plugins/Runtime/Soundscape/Source/SoundScape/Private/SoundscapeSubsystem.cpp b/Engine/Plugins/Runtime/Soundscape/Source/SoundScape/Private/SoundscapeSubsystem.cpp
--- a/Engine/Plugins/Runtime/Soundscape/Source/SoundScape/Private/SoundscapeSubsystem.cpp
+++ b/Engine/Plugins/Runtime/Soundscape/Source/SoundScape/Private/SoundscapeSubsystem.cpp
@@ -868,7 +868,7 @@ void USoundscapeSubsystem::DrawDebugCell(FVector Location, bool bSuccess)
 TPair<FVector, FVector> USoundscapeSubsystem::CalculateDebugCellDimensions(FVector Location, ESoundscapeLOD SoundscapeLOD)
 {
 	USoundscapeColorPointHashMap* LODHashMap = nullptr;
-	TPair<FVector, FVector> BoxResults;
+	TPair<FVector, FVector> BoxResults = { FVector::ZeroVector, FVector::ZeroVector };
 
 	// If relative distance is less than LOD 1 cutoff distance, use LOD 1 Map
 	switch (SoundscapeLOD)
