- Change GetMaterial to concurrent version

diff --git a/Engine/Plugins/FX/Niagara/Source/Niagara/Private/NiagaraRendererMeshes.cpp b/Engine/Plugins/FX/Niagara/Source/Niagara/Private/NiagaraRendererMeshes.cpp
--- a/Engine/Plugins/FX/Niagara/Source/Niagara/Private/NiagaraRendererMeshes.cpp
+++ b/Engine/Plugins/FX/Niagara/Source/Niagara/Private/NiagaraRendererMeshes.cpp
@@ -1953,7 +1953,7 @@ bool FNiagaraRendererMeshes::IsMaterialValid(const UMaterialInterface* Mat)const
 	bool bIsMaterialValid = false;
 	if (Mat)
 	{
-		bIsMaterialValid = (bIsHeterogeneousVolume && Mat->GetMaterial() && (Mat->GetMaterial()->MaterialDomain == MD_Volume)) ?
+		bIsMaterialValid = (bIsHeterogeneousVolume && Mat->GetMaterial_Concurrent() && (Mat->GetMaterial_Concurrent()->MaterialDomain == MD_Volume)) ?
 			Mat->CheckMaterialUsage_Concurrent(MATUSAGE_HeterogeneousVolumes) :
 			Mat->CheckMaterialUsage_Concurrent(MATUSAGE_NiagaraMeshParticles);
 	}
