Niagara - Allow copy/paste between module inputs and properties of objects shown in the stack.

diff --git a/Engine/Plugins/FX/Niagara/Source/Niagara/Internal/Stateless/NiagaraStatelessDistribution.h b/Engine/Plugins/FX/Niagara/Source/Niagara/Internal/Stateless/NiagaraStatelessDistribution.h
--- a/Engine/Plugins/FX/Niagara/Source/Niagara/Internal/Stateless/NiagaraStatelessDistribution.h
+++ b/Engine/Plugins/FX/Niagara/Source/Niagara/Internal/Stateless/NiagaraStatelessDistribution.h
@@ -49,9 +49,20 @@ struct FNiagaraDistributionBase
 	UPROPERTY(EditAnywhere, Category = "Parameters")
 	int32 MaxLutSampleCount = 64;
 
+	bool operator==(const FNiagaraDistributionBase& Other) const
+	{
+		return
+			Mode == Other.Mode &&
+			ParameterBinding == Other.ParameterBinding &&
+			ChannelConstantsAndRanges == Other.ChannelConstantsAndRanges &&
+			ChannelCurves == Other.ChannelCurves &&
+			MaxLutSampleCount == Other.MaxLutSampleCount;
+	}
+
 	virtual bool AllowBinding() const { return true; }
 	virtual bool AllowCurves() const { return true; }
 	virtual bool DisplayAsColor() const { return false; }
+	virtual int32 GetBaseNumberOfChannels() const { return 0; }
 	virtual void UpdateValuesFromDistribution() { }
 
 	virtual FNiagaraTypeDefinition GetBindingTypeDef() const { return FNiagaraTypeDefinition(); }
@@ -113,6 +124,7 @@ struct FNiagaraDistributionRangeFloat : public FNiagaraDistributionBase
 
 #if WITH_EDITORONLY_DATA
 	virtual bool AllowCurves() const override { return false; }
+	virtual int32 GetBaseNumberOfChannels() const override { return 1; }
 	NIAGARA_API virtual void UpdateValuesFromDistribution() override;
 	virtual FNiagaraTypeDefinition GetBindingTypeDef() const { return FNiagaraTypeDefinition::GetFloatDef(); }
 #endif
@@ -137,6 +149,7 @@ struct FNiagaraDistributionRangeVector2 : public FNiagaraDistributionBase
 
 #if WITH_EDITORONLY_DATA
 	virtual bool AllowCurves() const override { return false; }
+	virtual int32 GetBaseNumberOfChannels() const override { return 2; }
 	virtual void UpdateValuesFromDistribution() override;
 	virtual FNiagaraTypeDefinition GetBindingTypeDef() const { return FNiagaraTypeDefinition::GetVec2Def(); }
 #endif
@@ -161,6 +174,7 @@ struct FNiagaraDistributionRangeVector3 : public FNiagaraDistributionBase
 
 #if WITH_EDITORONLY_DATA
 	virtual bool AllowCurves() const override { return false; }
+	virtual int32 GetBaseNumberOfChannels() const override { return 3; }
 	virtual void UpdateValuesFromDistribution() override;
 	virtual FNiagaraTypeDefinition GetBindingTypeDef() const { return FNiagaraTypeDefinition::GetVec3Def(); }
 #endif
@@ -186,6 +200,7 @@ struct FNiagaraDistributionRangeColor : public FNiagaraDistributionBase
 #if WITH_EDITORONLY_DATA
 	virtual bool AllowCurves() const override { return false; }
 	virtual bool DisplayAsColor() const { return true; }
+	virtual int32 GetBaseNumberOfChannels() const override { return 4; }
 	virtual void UpdateValuesFromDistribution() override;
 	virtual FNiagaraTypeDefinition GetBindingTypeDef() const { return FNiagaraTypeDefinition::GetColorDef(); }
 #endif
@@ -206,7 +221,12 @@ struct FNiagaraDistributionFloat : public FNiagaraDistributionBase
 	NIAGARA_API FNiagaraStatelessRangeFloat CalculateRange(const float Default = 0.0f) const;
 
 #if WITH_EDITORONLY_DATA
-	virtual void UpdateValuesFromDistribution() override;
+	bool operator==(const FNiagaraDistributionFloat& Other) const
+	{
+		return (FNiagaraDistributionBase)*this == (FNiagaraDistributionBase)Other && Values == Other.Values;
+	}
+	virtual int32 GetBaseNumberOfChannels() const override { return 1; }
+	NIAGARA_API virtual void UpdateValuesFromDistribution() override;
 	virtual FNiagaraTypeDefinition GetBindingTypeDef() const { return FNiagaraTypeDefinition::GetFloatDef(); }
 #endif
 };
@@ -228,7 +248,8 @@ struct FNiagaraDistributionVector2 : public FNiagaraDistributionBase
 	NIAGARA_API FNiagaraStatelessRangeVector2 CalculateRange(const FVector2f& Default = FVector2f::ZeroVector) const;
 
 #if WITH_EDITORONLY_DATA
-	virtual void UpdateValuesFromDistribution() override;
+	virtual int32 GetBaseNumberOfChannels() const override { return 2; }
+	NIAGARA_API virtual void UpdateValuesFromDistribution() override;
 	virtual FNiagaraTypeDefinition GetBindingTypeDef() const { return FNiagaraTypeDefinition::GetVec2Def(); }
 #endif
 };
@@ -250,7 +271,8 @@ struct FNiagaraDistributionVector3 : public FNiagaraDistributionBase
 	NIAGARA_API FNiagaraStatelessRangeVector3 CalculateRange(const FVector3f& Default = FVector3f::ZeroVector) const;
 
 #if WITH_EDITORONLY_DATA
-	virtual void UpdateValuesFromDistribution() override;
+	virtual int32 GetBaseNumberOfChannels() const override { return 3; }
+	NIAGARA_API virtual void UpdateValuesFromDistribution() override;
 	virtual FNiagaraTypeDefinition GetBindingTypeDef() const { return FNiagaraTypeDefinition::GetVec3Def(); }
 #endif
 };
@@ -271,7 +293,8 @@ struct FNiagaraDistributionColor : public FNiagaraDistributionBase
 
 #if WITH_EDITORONLY_DATA
 	virtual bool DisplayAsColor() const override { return true; }
-	virtual void UpdateValuesFromDistribution() override;
+	virtual int32 GetBaseNumberOfChannels() const override { return 4; }
+	NIAGARA_API virtual void UpdateValuesFromDistribution() override;
 	virtual FNiagaraTypeDefinition GetBindingTypeDef() const { return FNiagaraTypeDefinition::GetColorDef(); }
 #endif
 };
diff --git a/Engine/Plugins/FX/Niagara/Source/NiagaraEditor/Private/NiagaraClipboard.cpp b/Engine/Plugins/FX/Niagara/Source/NiagaraEditor/Private/NiagaraClipboard.cpp
--- a/Engine/Plugins/FX/Niagara/Source/NiagaraEditor/Private/NiagaraClipboard.cpp
+++ b/Engine/Plugins/FX/Niagara/Source/NiagaraEditor/Private/NiagaraClipboard.cpp
@@ -13,6 +13,10 @@
 #include "NiagaraNodeFunctionCall.h"
 #include "Engine/UserDefinedEnum.h"
 #include "Engine/UserDefinedStruct.h"
+#include "NiagaraEditorModule.h"
+#include "INiagaraEditorTypeUtilities.h"
+#include "Logging/LogScopedVerbosityOverride.h"
+#include "PropertyHandle.h"
 
 #include UE_INLINE_GENERATED_CPP_BY_NAME(NiagaraClipboard)
 
@@ -171,6 +175,162 @@ UNiagaraClipboardFunction* UNiagaraClipboardFunction::CreateAssignmentFunction(U
 	return NewFunction;
 }
 
+FNiagaraClipboardPortableValue FNiagaraClipboardPortableValue::CreateFromStructValue(const UScriptStruct& TargetStruct, uint8* StructMemory)
+{
+	FNiagaraClipboardPortableValue PortableValue;
+	TargetStruct.ExportText(PortableValue.ValueString, StructMemory, nullptr, nullptr, PPF_Copy, nullptr);
+	return PortableValue;
+}
+
+FNiagaraClipboardPortableValue FNiagaraClipboardPortableValue::CreateFromTypedValue(const FNiagaraTypeDefinition& InType, const FNiagaraVariant& InValue)
+{
+	FNiagaraClipboardPortableValue PortableValue;
+
+	TSharedPtr<INiagaraEditorTypeUtilities> InputTypeUtilities = FNiagaraEditorModule::Get().GetTypeUtilities(InType);
+	if (InputTypeUtilities.IsValid() &&	InputTypeUtilities->SupportsClipboardPortableValues())
+	{
+		if (InputTypeUtilities->TryUpdateClipboardPortableValueFromTypedValue(InType, InValue, PortableValue) == false)
+		{
+			PortableValue.Reset();
+		}
+	}
+
+	if (PortableValue.IsValid() == false)
+	{
+		if (InType.GetStruct() != nullptr &&
+			InType.GetStruct()->IsA<UScriptStruct>() &&
+			InType.GetSize() == InValue.GetNumBytes())
+		{
+			UScriptStruct* ValueStruct = CastChecked<UScriptStruct>(InType.GetStruct());
+			PortableValue = CreateFromStructValue(*ValueStruct, InValue.GetBytes());
+		}
+	}
+
+	if (PortableValue.IsValid() == false)
+	{
+		PortableValue.Reset();
+	}
+	return PortableValue;
+}
+
+FNiagaraClipboardPortableValue FNiagaraClipboardPortableValue::CreateFromPropertyHandle(const IPropertyHandle& InPropertyHandle)
+{
+	FNiagaraClipboardPortableValue PortableValue;
+	FProperty* Property = InPropertyHandle.GetProperty();
+	if (Property != nullptr)
+	{
+		UScriptStruct* ValueStruct = nullptr;
+		if (Property->IsA(FStructProperty::StaticClass()))
+		{
+			FStructProperty* StructProperty = CastField<FStructProperty>(Property);
+			ValueStruct = StructProperty->Struct;
+			TSharedPtr<INiagaraEditorPropertyUtilities, ESPMode::ThreadSafe> PropertyUtilities = FNiagaraEditorModule::Get().GetPropertyUtilities(*ValueStruct);
+			if (PropertyUtilities.IsValid() && PropertyUtilities->SupportsClipboardPortableValues())
+			{
+				if (PropertyUtilities->TryUpdateClipboardPortableValueFromProperty(InPropertyHandle, PortableValue) == false)
+				{
+					PortableValue.Reset();
+				}
+			}
+		}
+
+		if (PortableValue.IsValid() == false)
+		{
+			if (InPropertyHandle.GetValueAsFormattedString(PortableValue.ValueString, PPF_Copy) != FPropertyAccess::Success)
+			{
+				PortableValue.Reset();
+			}
+		}
+	}
+
+	return PortableValue;
+}
+
+bool FNiagaraClipboardPortableValue::CanUpdateTypedValue(const FNiagaraTypeDefinition& InTargetType) const
+{
+	if (IsValid() == false)
+	{
+		return false;
+	}
+
+	TSharedPtr<INiagaraEditorTypeUtilities> TargetTypeUtilities = FNiagaraEditorModule::Get().GetTypeUtilities(InTargetType);
+	if (TargetTypeUtilities.IsValid() && TargetTypeUtilities->SupportsClipboardPortableValues())
+	{
+		return TargetTypeUtilities->CanUpdateTypedValueFromClipboardPortableValue(*this, InTargetType);
+	}
+	if (InTargetType.GetStruct() != nullptr && InTargetType.GetStruct()->IsA<UScriptStruct>())
+	{
+		UScriptStruct* TargetStruct = CastChecked<UScriptStruct>(InTargetType.GetStruct());
+		TArray<uint8> ValueBytes;
+		ValueBytes.AddDefaulted(InTargetType.GetSize());
+		return TryUpdateStructValue(*TargetStruct, ValueBytes.GetData());
+	}
+	return false;
+}
+
+bool FNiagaraClipboardPortableValue::TryUpdateStructValue(const UScriptStruct& TargetStruct, uint8* StructMemory) const
+{
+	LOG_SCOPE_VERBOSITY_OVERRIDE(LogExec, ELogVerbosity::Verbose);
+	FErrorPipe ErrorPipe;
+	return TargetStruct.ImportText(*ValueString, StructMemory, nullptr, PPF_Copy, &ErrorPipe, TargetStruct.GetName()) != nullptr && ErrorPipe.NumErrors == 0;
+}
+
+bool FNiagaraClipboardPortableValue::TryUpdateTypedValue(const FNiagaraTypeDefinition& InTargetType, FNiagaraVariant& InTargetValue) const
+{
+	if (IsValid() == false)
+	{
+		return false;
+	}
+
+	TSharedPtr<INiagaraEditorTypeUtilities> TargetTypeUtilities = FNiagaraEditorModule::Get().GetTypeUtilities(InTargetType);
+	if (TargetTypeUtilities.IsValid() && TargetTypeUtilities->SupportsClipboardPortableValues())
+	{
+		if (TargetTypeUtilities->TryUpdateTypedValueFromClipboardPortableValue(*this, InTargetType, InTargetValue))
+		{
+			return true;
+		}
+	}
+	if (InTargetType.GetStruct() != nullptr && InTargetType.GetStruct()->IsA<UScriptStruct>())
+	{
+		UScriptStruct* TargetStruct = CastChecked<UScriptStruct>(InTargetType.GetStruct());
+		TArray<uint8> ValueBytes;
+		ValueBytes.AddDefaulted(TargetStruct->GetStructureSize());
+		if (TryUpdateStructValue(*TargetStruct, ValueBytes.GetData()))
+		{
+			InTargetValue.SetBytes(ValueBytes.GetData(), ValueBytes.Num());
+			return true;
+		}
+	}
+	return false;
+}
+
+bool FNiagaraClipboardPortableValue::TryUpdatePropertyHandle(IPropertyHandle& InTargetPropertyHandle) const
+{
+	if (IsValid() == false)
+	{
+		return false;
+	}
+
+	FProperty* Property = InTargetPropertyHandle.GetProperty();
+	if (Property != nullptr)
+	{
+		if (Property->IsA(FStructProperty::StaticClass()))
+		{
+			FStructProperty* StructProperty = CastField<FStructProperty>(Property);
+			TSharedPtr<INiagaraEditorPropertyUtilities, ESPMode::ThreadSafe> PropertyUtilities = FNiagaraEditorModule::Get().GetPropertyUtilities(*StructProperty->Struct);
+			if (PropertyUtilities.IsValid() && PropertyUtilities->SupportsClipboardPortableValues())
+			{
+				if (PropertyUtilities->TryUpdatePropertyFromClipboardPortableValue(*this, InTargetPropertyHandle))
+				{
+					return true;
+				}
+			}
+		}
+		return InTargetPropertyHandle.SetValueFromFormattedString(ValueString) == FPropertyAccess::Success;
+	}
+	return false;
+}
+
 UNiagaraClipboardContent* UNiagaraClipboardContent::Create()
 {
 	return NewObject<UNiagaraClipboardContent>(GetTransientPackage());
diff --git a/Engine/Plugins/FX/Niagara/Source/NiagaraEditor/Private/NiagaraEditorModule.cpp b/Engine/Plugins/FX/Niagara/Source/NiagaraEditor/Private/NiagaraEditorModule.cpp
--- a/Engine/Plugins/FX/Niagara/Source/NiagaraEditor/Private/NiagaraEditorModule.cpp
+++ b/Engine/Plugins/FX/Niagara/Source/NiagaraEditor/Private/NiagaraEditorModule.cpp
@@ -40,6 +40,7 @@
 #include "TypeEditorUtilities/NiagaraColorTypeEditorUtilities.h"
 #include "TypeEditorUtilities/NiagaraMatrixTypeEditorUtilities.h"
 #include "TypeEditorUtilities/NiagaraDataInterfaceCurveTypeEditorUtilities.h"
+#include "TypeEditorUtilities/NiagaraDistributionPropertyEditorUtilities.h"
 
 #include "NiagaraSystemCompilingManager.h"
 #include "NiagaraEditorStyle.h"
@@ -1346,6 +1347,16 @@ void FNiagaraEditorModule::StartupModule()
 	RegisterTypeUtilities(FNiagaraTypeDefinition(UNiagaraDataInterfaceVector4Curve::StaticClass()), MakeShared<FNiagaraDataInterfaceVectorCurveTypeEditorUtilities, ESPMode::ThreadSafe>());
 	RegisterTypeUtilities(FNiagaraTypeDefinition(UNiagaraDataInterfaceColorCurve::StaticClass()), MakeShared<FNiagaraDataInterfaceColorCurveTypeEditorUtilities, ESPMode::ThreadSafe>());
 
+	TSharedRef<FNiagaraDistributionPropertyEditorUtilities, ESPMode::ThreadSafe> DistributionPropertyUtilities = MakeShared<FNiagaraDistributionPropertyEditorUtilities, ESPMode::ThreadSafe>();
+	RegisterPropertyUtilities(FNiagaraDistributionFloat::StaticStruct(), MakeShared<FNiagaraDistributionPropertyEditorUtilities, ESPMode::ThreadSafe>());
+	RegisterPropertyUtilities(FNiagaraDistributionVector2::StaticStruct(), MakeShared<FNiagaraDistributionPropertyEditorUtilities, ESPMode::ThreadSafe>());
+	RegisterPropertyUtilities(FNiagaraDistributionVector3::StaticStruct(), MakeShared<FNiagaraDistributionPropertyEditorUtilities, ESPMode::ThreadSafe>());
+	RegisterPropertyUtilities(FNiagaraDistributionColor::StaticStruct(), MakeShared<FNiagaraDistributionPropertyEditorUtilities, ESPMode::ThreadSafe>());
+	RegisterPropertyUtilities(FNiagaraDistributionRangeFloat::StaticStruct(), MakeShared<FNiagaraDistributionPropertyEditorUtilities, ESPMode::ThreadSafe>());
+	RegisterPropertyUtilities(FNiagaraDistributionRangeVector2::StaticStruct(), MakeShared<FNiagaraDistributionPropertyEditorUtilities, ESPMode::ThreadSafe>());
+	RegisterPropertyUtilities(FNiagaraDistributionRangeVector3::StaticStruct(), MakeShared<FNiagaraDistributionPropertyEditorUtilities, ESPMode::ThreadSafe>());
+	RegisterPropertyUtilities(FNiagaraDistributionRangeColor::StaticStruct(), MakeShared<FNiagaraDistributionPropertyEditorUtilities, ESPMode::ThreadSafe>());
+
 	FEdGraphUtilities::RegisterVisualPinFactory(GraphPanelPinFactory);
 
 	FNiagaraOpInfo::Init();
@@ -1800,6 +1811,12 @@ void FNiagaraEditorModule::RegisterTypeUtilities(FNiagaraTypeDefinition Type, TS
 	TypeEditorsCS.Unlock();
 }
 
+void FNiagaraEditorModule::RegisterPropertyUtilities(const UScriptStruct* InStruct, TSharedRef<INiagaraEditorPropertyUtilities, ESPMode::ThreadSafe> InPropertyUtilities)
+{
+	TypeEditorsCS.Lock();
+	StructToPropertyUtilitiesMap.Add(InStruct, InPropertyUtilities);
+	TypeEditorsCS.Unlock();
+}
 
 TSharedPtr<INiagaraEditorTypeUtilities, ESPMode::ThreadSafe> FNiagaraEditorModule::GetTypeUtilities(const FNiagaraTypeDefinition& Type)
 {
@@ -1820,6 +1837,15 @@ TSharedPtr<INiagaraEditorTypeUtilities, ESPMode::ThreadSafe> FNiagaraEditorModul
 	return TSharedPtr<INiagaraEditorTypeUtilities, ESPMode::ThreadSafe>();
 }
 
+TSharedPtr<INiagaraEditorPropertyUtilities, ESPMode::ThreadSafe> FNiagaraEditorModule::GetPropertyUtilities(const UScriptStruct& Struct)
+{
+	TypeEditorsCS.Lock();
+	TSharedRef<INiagaraEditorPropertyUtilities, ESPMode::ThreadSafe>* PropertyUtilities = StructToPropertyUtilitiesMap.Find(&Struct);
+	TypeEditorsCS.Unlock();
+
+	return PropertyUtilities != nullptr ? *PropertyUtilities : TSharedPtr<INiagaraEditorPropertyUtilities>();
+}
+
 void FNiagaraEditorModule::RegisterWidgetProvider(TSharedRef<INiagaraEditorWidgetProvider> InWidgetProvider)
 {
 	checkf(WidgetProvider.IsValid() == false, TEXT("Widget provider has already been set."));
diff --git a/Engine/Plugins/FX/Niagara/Source/NiagaraEditor/Private/Stateless/NiagaraDistributionPropertyCustomization.cpp b/Engine/Plugins/FX/Niagara/Source/NiagaraEditor/Private/Stateless/NiagaraDistributionPropertyCustomization.cpp
--- a/Engine/Plugins/FX/Niagara/Source/NiagaraEditor/Private/Stateless/NiagaraDistributionPropertyCustomization.cpp
+++ b/Engine/Plugins/FX/Niagara/Source/NiagaraEditor/Private/Stateless/NiagaraDistributionPropertyCustomization.cpp
@@ -22,15 +22,23 @@
 
 #define LOCTEXT_NAMESPACE "NiagaraDistributionPropertyCustomization"
 
-class FNiagaraDistributionAdapter : public INiagaraDistributionAdapter
+class FNiagaraDistributionAdapter : public INiagaraDistributionAdapter, public TSharedFromThis<FNiagaraDistributionAdapter>
 {
 public:
-	FNiagaraDistributionAdapter(TSharedPtr<IPropertyHandle> InPropertyHandle, UObject* InOwningObject, FNiagaraDistributionBase* InDistribution, int32 InNumChannels)
-		: PropertyHandle(InPropertyHandle)
-		, OwnerObjectWeak(TWeakObjectPtr<UObject>(InOwningObject))
-		, SourceDistribution(InDistribution)
-		, SourceNumChannels(InNumChannels)
+	FNiagaraDistributionAdapter()
+		: PropertyHandle(nullptr)
+		, OwnerObjectWeak(nullptr)
+		, SourceDistribution(nullptr)
+		, SourceNumChannels(INDEX_NONE)
 	{
+	}
+
+	void Initialize(TSharedPtr<IPropertyHandle> InPropertyHandle, UObject* InOwningObject, FNiagaraDistributionBase* InDistribution, int32 InNumChannels)
+	{
+		PropertyHandle = InPropertyHandle;
+		OwnerObjectWeak = TWeakObjectPtr<UObject>(InOwningObject);
+		SourceDistribution = InDistribution;
+		SourceNumChannels = InNumChannels;
 		if (IsValid())
 		{
 			const FName DisableBindingDistributionName("DisableBindingDistribution");
@@ -45,13 +53,13 @@ public:
 			bAllowRange = PropertyHandle ? PropertyHandle->HasMetaData(DisableRangeDistributionName) == false : false;
 			bAllowCurves = InDistribution->AllowCurves() && (PropertyHandle ? PropertyHandle->HasMetaData(DisableCurveDistributionName) == false : false);
 			bDisplayAsColor = InDistribution->DisplayAsColor() || (InNumChannels >= 3 && (PropertyHandle ? PropertyHandle->HasMetaData(DisplayAsColorDistributionName) : false));
-
+			PropertyHandle->SetOnPropertyValueChanged(FSimpleDelegate::CreateSP(this->AsShared(), &FNiagaraDistributionAdapter::DistributionPropertyChanged));
 			EditorMode = GetDistributionEditorModeFromSourceMode(SourceNumChannels, bDisplayAsColor, SourceDistribution->Mode);
 			EditorNumChannels = FNiagaraDistributionEditorUtilities::IsUniform(EditorMode) ? 1 : SourceNumChannels;
 		}
 	}
 
-	virtual bool IsValid() const override { return SourceDistribution != nullptr; }
+	virtual bool IsValid() const override { return PropertyHandle.IsValid() && SourceDistribution != nullptr; }
 
 	virtual int32 GetNumChannels() const override { return EditorNumChannels; }
 
@@ -267,6 +275,7 @@ public:
 			return;
 		}
 
+		TGuardValue<bool> UpdateGuard(bUpdatingDistributionModeInternally, true);
 		int32 NewEditorNumChannels = FNiagaraDistributionEditorUtilities::IsUniform(InEditorMode) ? 1 : SourceNumChannels;
 
 		ENiagaraDistributionMode NewSourceMode = GetDistributionSourceModeFromEditorMode(InEditorMode);
@@ -284,6 +293,12 @@ public:
 
 		EditorNumChannels = NewEditorNumChannels;
 		EditorMode = InEditorMode;
+		OnDistributionEditorModeChangedDelegate.Broadcast();
+	}
+
+	virtual FSimpleMulticastDelegate& OnDistributionEditorModeChanged() override
+	{
+		return OnDistributionEditorModeChangedDelegate;
 	}
 
 	float GetConstantOrRangeValue(int32 ChannelIndex, int32 ValueIndex) const override
@@ -566,6 +581,17 @@ public:
 	}
 
 private:
+	void DistributionPropertyChanged()
+	{
+		// Watch for editor mode changes from external sources like copy/paste.
+		ENiagaraDistributionEditorMode NewEditorMode = GetDistributionEditorModeFromSourceMode(SourceNumChannels, bDisplayAsColor, SourceDistribution->Mode);
+		if (NewEditorMode != EditorMode && bUpdatingDistributionModeInternally == false)
+		{
+			EditorMode = NewEditorMode;
+			OnDistributionEditorModeChangedDelegate.Broadcast();
+		}
+	}
+
 	void MigrateDataFromModeChange(ENiagaraDistributionMode NewMode, int32 NewEditorNumChannels)
 	{
 		if (NewMode == ENiagaraDistributionMode::UniformConstant || NewMode == ENiagaraDistributionMode::NonUniformConstant)
@@ -748,6 +774,9 @@ private:
 	bool bAllowRange = true;
 	bool bAllowCurves = true;
 	bool bDisplayAsColor = false;
+
+	FSimpleMulticastDelegate OnDistributionEditorModeChangedDelegate;
+	bool bUpdatingDistributionModeInternally = false;
 };
 
 class SNiagaraDistributionPropertyWidget : public SCompoundWidget, public FEditorUndoClient
@@ -820,9 +849,11 @@ TSharedRef<IPropertyTypeCustomization> FNiagaraDistributionPropertyCustomization
 		if (OuterObjects.Num() == 1 && FloatDistributionPropertyHandle->GetValueData(ValueData) == FPropertyAccess::Success)
 		{
 			FNiagaraDistributionBase* FloatDistribution = static_cast<FNiagaraDistributionBase*>(ValueData);
-			return MakeShared<FNiagaraDistributionAdapter>(FloatDistributionPropertyHandle, OuterObjects[0], FloatDistribution, 1);
+			TSharedRef<FNiagaraDistributionAdapter> DistributionAdapter = MakeShared<FNiagaraDistributionAdapter>();
+			DistributionAdapter->Initialize(FloatDistributionPropertyHandle, OuterObjects[0], FloatDistribution, 1);
+			return DistributionAdapter;
 		}
-		return MakeShared<FNiagaraDistributionAdapter>(nullptr, nullptr, nullptr, INDEX_NONE);
+		return MakeShared<FNiagaraDistributionAdapter>();
 	});
 
 	return MakeShareable<FNiagaraDistributionPropertyCustomization>(new FNiagaraDistributionPropertyCustomization(FloatDistributionPropertyHandleToDistributionAdapter));
@@ -843,9 +874,11 @@ TSharedRef<IPropertyTypeCustomization> FNiagaraDistributionPropertyCustomization
 		if (OuterObjects.Num() == 1 && Vector2DistributionPropertyHandle->GetValueData(ValueData) == FPropertyAccess::Success)
 		{
 			FNiagaraDistributionBase* Vector2Distribution = static_cast<FNiagaraDistributionBase*>(ValueData);
-			return MakeShared<FNiagaraDistributionAdapter>(Vector2DistributionPropertyHandle, OuterObjects[0], Vector2Distribution, 2);
+			TSharedRef<FNiagaraDistributionAdapter> DistributionAdapter = MakeShared<FNiagaraDistributionAdapter>();
+			DistributionAdapter->Initialize(Vector2DistributionPropertyHandle, OuterObjects[0], Vector2Distribution, 2);
+			return DistributionAdapter;
 		}
-		return MakeShared<FNiagaraDistributionAdapter>(nullptr, nullptr, nullptr, INDEX_NONE);
+		return MakeShared<FNiagaraDistributionAdapter>();
 	});
 
 	return MakeShareable<FNiagaraDistributionPropertyCustomization>(new FNiagaraDistributionPropertyCustomization(Vector2DistributionPropertyHandleToDistributionAdapter));
@@ -861,9 +894,11 @@ TSharedRef<IPropertyTypeCustomization> FNiagaraDistributionPropertyCustomization
 		if (OuterObjects.Num() == 1 && Vector3DistributionPropertyHandle->GetValueData(ValueData) == FPropertyAccess::Success)
 		{
 			FNiagaraDistributionBase* Vector3Distribution = static_cast<FNiagaraDistributionBase*>(ValueData);
-			return MakeShared<FNiagaraDistributionAdapter>(Vector3DistributionPropertyHandle, OuterObjects[0], Vector3Distribution, 3);
+			TSharedRef<FNiagaraDistributionAdapter> DistributionAdapter = MakeShared<FNiagaraDistributionAdapter>();
+			DistributionAdapter->Initialize(Vector3DistributionPropertyHandle, OuterObjects[0], Vector3Distribution, 3);
+			return DistributionAdapter;
 		}
-		return MakeShared<FNiagaraDistributionAdapter>(nullptr, nullptr, nullptr, INDEX_NONE);
+		return MakeShared<FNiagaraDistributionAdapter>();
 	});
 
 	return MakeShareable<FNiagaraDistributionPropertyCustomization>(new FNiagaraDistributionPropertyCustomization(Vector3DistributionPropertyHandleToDistributionAdapter));
@@ -879,9 +914,11 @@ TSharedRef<IPropertyTypeCustomization> FNiagaraDistributionPropertyCustomization
 		if (OuterObjects.Num() == 1 && ColorDistributionPropertyHandle->GetValueData(ValueData) == FPropertyAccess::Success)
 		{
 			FNiagaraDistributionBase* ColorDistribution = static_cast<FNiagaraDistributionBase*>(ValueData);
-			return MakeShared<FNiagaraDistributionAdapter>(ColorDistributionPropertyHandle, OuterObjects[0], ColorDistribution, 4);
+			TSharedRef<FNiagaraDistributionAdapter> DistributionAdapter = MakeShared<FNiagaraDistributionAdapter>();
+			DistributionAdapter->Initialize(ColorDistributionPropertyHandle, OuterObjects[0], ColorDistribution, 4);
+			return DistributionAdapter;
 		}
-		return MakeShared<FNiagaraDistributionAdapter>(nullptr, nullptr, nullptr, INDEX_NONE);
+		return MakeShared<FNiagaraDistributionAdapter>();
 	});
 
 	return MakeShareable<FNiagaraDistributionPropertyCustomization>(new FNiagaraDistributionPropertyCustomization(ColorDistributionPropertyHandleToDistributionAdapter));
diff --git a/Engine/Plugins/FX/Niagara/Source/NiagaraEditor/Private/TypeEditorUtilities/NiagaraBoolTypeEditorUtilities.cpp b/Engine/Plugins/FX/Niagara/Source/NiagaraEditor/Private/TypeEditorUtilities/NiagaraBoolTypeEditorUtilities.cpp
--- a/Engine/Plugins/FX/Niagara/Source/NiagaraEditor/Private/TypeEditorUtilities/NiagaraBoolTypeEditorUtilities.cpp
+++ b/Engine/Plugins/FX/Niagara/Source/NiagaraEditor/Private/TypeEditorUtilities/NiagaraBoolTypeEditorUtilities.cpp
@@ -1,7 +1,9 @@
 // Copyright Epic Games, Inc. All Rights Reserved.
 
 #include "NiagaraBoolTypeEditorUtilities.h"
+#include "NiagaraClipboard.h"
 #include "NiagaraTypes.h"
+#include "NiagaraVariant.h"
 #include "SNiagaraParameterEditor.h"
 #include "Widgets/DeclarativeSyntaxSupport.h"
 #include "Widgets/SBoxPanel.h"
@@ -87,3 +89,32 @@ bool FNiagaraEditorBoolTypeUtilities::SetValueFromPinDefaultString(const FString
 	}
 	return false;
 }
+
+bool FNiagaraEditorBoolTypeUtilities::TryUpdateClipboardPortableValueFromTypedValue(const FNiagaraTypeDefinition& InSourceType, const FNiagaraVariant& InSourceValue, FNiagaraClipboardPortableValue& InTargetClipboardPortableValue) const
+{
+	if (InSourceType == FNiagaraTypeDefinition::GetBoolDef() && InSourceValue.GetNumBytes() == FNiagaraTypeDefinition::GetBoolDef().GetSize())
+	{
+		FNiagaraVariable Temp(InSourceType, NAME_None);
+		Temp.SetData(InSourceValue.GetBytes());
+		bool BoolValue = Temp.GetValue<FNiagaraBool>().GetValue();
+		InTargetClipboardPortableValue.ValueString = LexToString(BoolValue);
+		return true;
+	}
+	return false;
+}
+
+bool FNiagaraEditorBoolTypeUtilities::TryUpdateTypedValueFromClipboardPortableValue(const FNiagaraClipboardPortableValue& InSourceClipboardPortableValue, const FNiagaraTypeDefinition& InTargetType, FNiagaraVariant& InTargetValue) const
+{
+	bool BoolValue;
+	if (InTargetType == FNiagaraTypeDefinition::GetBoolDef() &&
+		LexTryParseString(BoolValue, *InSourceClipboardPortableValue.ValueString))
+	{
+		FNiagaraBool NiagaraBoolValue;
+		NiagaraBoolValue.SetValue(BoolValue);
+		FNiagaraVariable Temp(InTargetType, NAME_None);
+		Temp.SetValue<FNiagaraBool>(NiagaraBoolValue);
+		InTargetValue.SetBytes(Temp.GetData(), Temp.GetSizeInBytes());
+		return true;
+	}
+	return false;
+}
diff --git a/Engine/Plugins/FX/Niagara/Source/NiagaraEditor/Private/TypeEditorUtilities/NiagaraBoolTypeEditorUtilities.h b/Engine/Plugins/FX/Niagara/Source/NiagaraEditor/Private/TypeEditorUtilities/NiagaraBoolTypeEditorUtilities.h
--- a/Engine/Plugins/FX/Niagara/Source/NiagaraEditor/Private/TypeEditorUtilities/NiagaraBoolTypeEditorUtilities.h
+++ b/Engine/Plugins/FX/Niagara/Source/NiagaraEditor/Private/TypeEditorUtilities/NiagaraBoolTypeEditorUtilities.h
@@ -16,4 +16,7 @@ public:
 	virtual bool CanHandlePinDefaults() const override;
 	virtual FString GetPinDefaultStringFromValue(const FNiagaraVariable& AllocatedVariable) const override;
 	virtual bool SetValueFromPinDefaultString(const FString& StringValue, FNiagaraVariable& Variable) const override;
+	virtual bool SupportsClipboardPortableValues() const { return true; }
+	virtual bool TryUpdateClipboardPortableValueFromTypedValue(const FNiagaraTypeDefinition& InSourceType, const FNiagaraVariant& InSourceValue, FNiagaraClipboardPortableValue& InTargetClipboardPortableValue) const override;
+	virtual bool TryUpdateTypedValueFromClipboardPortableValue(const FNiagaraClipboardPortableValue& InSourceClipboardPortableValue, const FNiagaraTypeDefinition& InTargetType, FNiagaraVariant& InTargetValue) const override;
 };
\ No newline at end of file
diff --git a/Engine/Plugins/FX/Niagara/Source/NiagaraEditor/Private/TypeEditorUtilities/NiagaraDataInterfaceCurveTypeEditorUtilities.cpp b/Engine/Plugins/FX/Niagara/Source/NiagaraEditor/Private/TypeEditorUtilities/NiagaraDataInterfaceCurveTypeEditorUtilities.cpp
--- a/Engine/Plugins/FX/Niagara/Source/NiagaraEditor/Private/TypeEditorUtilities/NiagaraDataInterfaceCurveTypeEditorUtilities.cpp
+++ b/Engine/Plugins/FX/Niagara/Source/NiagaraEditor/Private/TypeEditorUtilities/NiagaraDataInterfaceCurveTypeEditorUtilities.cpp
@@ -2,12 +2,14 @@
 
 #include "NiagaraDataInterfaceCurveTypeEditorUtilities.h"
 
+#include "NiagaraClipboard.h"
 #include "NiagaraDataInterface.h"
 #include "NiagaraDataInterfaceCurve.h"
 #include "NiagaraDataInterfaceVector2DCurve.h"
 #include "NiagaraDataInterfaceVectorCurve.h"
 #include "NiagaraDataInterfaceVector4Curve.h"
 #include "NiagaraDataInterfaceColorCurve.h"
+#include "NiagaraVariant.h"
 #include "IContentBrowserSingleton.h"
 #include "ContentBrowserModule.h"
 #include "Modules/ModuleManager.h"
@@ -82,6 +84,58 @@ void FNiagaraDataInterfaceCurveTypeEditorUtilitiesBase::CurveAssetSelected(const
 	}
 }
 
+bool FNiagaraDataInterfaceCurveTypeEditorUtilitiesBase::TryUpdateClipboardPortableValueFromTypedValue(const FNiagaraTypeDefinition& InSourceType, const FNiagaraVariant& InSourceValue, FNiagaraClipboardPortableValue& InTargetClipboardPortableValue) const
+{
+	if (InSourceType.GetClass()->IsChildOf(UNiagaraDataInterfaceCurveBase::StaticClass()) &&
+		InSourceValue.GetDataInterface() != nullptr &&
+		InSourceValue.GetDataInterface()->IsA<UNiagaraDataInterfaceCurveBase>())
+	{
+		UNiagaraDataInterfaceCurveBase* CurveDataInterface = CastChecked<UNiagaraDataInterfaceCurveBase>(InSourceValue.GetDataInterface());
+		FNiagaraClipboardCurveCollection CurveCollection;
+		TArray<UNiagaraDataInterfaceCurveBase::FCurveData> CurveData;
+		CurveDataInterface->GetCurveData(CurveData);
+		for (const UNiagaraDataInterfaceCurveBase::FCurveData& CurveDataItem : CurveData)
+		{
+			CurveCollection.Curves.Add(*CurveDataItem.Curve);
+		}
+		InTargetClipboardPortableValue = FNiagaraClipboardPortableValue::CreateFromStructValue(*FNiagaraClipboardCurveCollection::StaticStruct(), (uint8*)&CurveCollection);
+		return true;
+	}
+	return false;
+}
+
+bool FNiagaraDataInterfaceCurveTypeEditorUtilitiesBase::CanUpdateTypedValueFromClipboardPortableValue(const FNiagaraClipboardPortableValue& InSourceClipboardPortableValue, const FNiagaraTypeDefinition& InTargetType) const
+{
+	UScriptStruct* CurveCollectionStruct = FNiagaraClipboardCurveCollection::StaticStruct();
+	FNiagaraClipboardCurveCollection CurveCollection;
+	return InSourceClipboardPortableValue.TryUpdateStructValue(*CurveCollectionStruct, (uint8*)&CurveCollection);
+}
+
+bool FNiagaraDataInterfaceCurveTypeEditorUtilitiesBase::TryUpdateTypedValueFromClipboardPortableValue(const FNiagaraClipboardPortableValue& InSourceClipboardPortableValue, const FNiagaraTypeDefinition& InTargetType, FNiagaraVariant& InTargetValue) const
+{
+	if (InTargetType.GetClass()->IsChildOf(UNiagaraDataInterfaceCurveBase::StaticClass()))
+	{
+		UScriptStruct* CurveCollectionStruct = FNiagaraClipboardCurveCollection::StaticStruct();
+		FNiagaraClipboardCurveCollection CurveCollection;
+		if (InSourceClipboardPortableValue.TryUpdateStructValue(*CurveCollectionStruct, (uint8*)&CurveCollection))
+		{
+			UNiagaraDataInterfaceCurveBase* CurveDataInterface = NewObject<UNiagaraDataInterfaceCurveBase>(GetTransientPackage(), InTargetType.GetClass());
+			TArray<UNiagaraDataInterfaceCurveBase::FCurveData> CurveData;
+			CurveDataInterface->GetCurveData(CurveData);
+			if (CurveData.Num() == CurveCollection.Curves.Num())
+			{
+				for (int32 CurveIndex = 0; CurveIndex < CurveData.Num(); CurveIndex++)
+				{
+					*CurveData[CurveIndex].Curve = CurveCollection.Curves[CurveIndex];
+				}
+				InTargetValue.SetDataInterface(CurveDataInterface);
+				return true;
+			}
+		}
+	}
+	return false;
+}
+
 FTopLevelAssetPath FNiagaraDataInterfaceCurveTypeEditorUtilities::GetSupportedAssetClassName() const
 {
 	return UCurveFloat::StaticClass()->GetClassPathName();
diff --git a/Engine/Plugins/FX/Niagara/Source/NiagaraEditor/Private/TypeEditorUtilities/NiagaraDataInterfaceCurveTypeEditorUtilities.h b/Engine/Plugins/FX/Niagara/Source/NiagaraEditor/Private/TypeEditorUtilities/NiagaraDataInterfaceCurveTypeEditorUtilities.h
--- a/Engine/Plugins/FX/Niagara/Source/NiagaraEditor/Private/TypeEditorUtilities/NiagaraDataInterfaceCurveTypeEditorUtilities.h
+++ b/Engine/Plugins/FX/Niagara/Source/NiagaraEditor/Private/TypeEditorUtilities/NiagaraDataInterfaceCurveTypeEditorUtilities.h
@@ -15,6 +15,11 @@ public:
 	virtual bool CanCreateDataInterfaceEditor() const override;
 	virtual TSharedPtr<SWidget> CreateDataInterfaceEditor(UObject* DataInterface, FNotifyValueChanged DataInterfaceChangedHandler) const override;
 
+	virtual bool SupportsClipboardPortableValues() const { return true; }
+	virtual bool TryUpdateClipboardPortableValueFromTypedValue(const FNiagaraTypeDefinition& InSourceType, const FNiagaraVariant& InSourceValue, FNiagaraClipboardPortableValue& InTargetClipboardPortableValue) const override;
+	virtual bool CanUpdateTypedValueFromClipboardPortableValue(const FNiagaraClipboardPortableValue& InSourceClipboardPortableValue, const FNiagaraTypeDefinition& InTargetType) const override;
+	virtual bool TryUpdateTypedValueFromClipboardPortableValue(const FNiagaraClipboardPortableValue& InSourceClipboardPortableValue, const FNiagaraTypeDefinition& InTargetType, FNiagaraVariant& InTargetValue) const override;
+
 protected:
 	/** Gets the class name used for filtering the curve asset picker. */
 	virtual FTopLevelAssetPath GetSupportedAssetClassName() const = 0;
diff --git a/Engine/Plugins/FX/Niagara/Source/NiagaraEditor/Private/TypeEditorUtilities/NiagaraDistributionPropertyEditorUtilities.cpp b/Engine/Plugins/FX/Niagara/Source/NiagaraEditor/Private/TypeEditorUtilities/NiagaraDistributionPropertyEditorUtilities.cpp
new file mode 100644
--- /dev/null
+++ b/Engine/Plugins/FX/Niagara/Source/NiagaraEditor/Private/TypeEditorUtilities/NiagaraDistributionPropertyEditorUtilities.cpp
@@ -0,0 +1,183 @@
+// Copyright Epic Games, Inc. All Rights Reserved.
+
+#include "TypeEditorUtilities/NiagaraDistributionPropertyEditorUtilities.h"
+
+#include "NiagaraClipboard.h"
+#include "PropertyHandle.h"
+#include "ScopedTransaction.h"
+#include "Stateless/NiagaraStatelessDistribution.h"
+
+bool FNiagaraDistributionPropertyEditorUtilities::SupportsClipboardPortableValues() const
+{ 
+	return true;
+}
+
+bool FNiagaraDistributionPropertyEditorUtilities::TryUpdateClipboardPortableValueFromProperty(const IPropertyHandle& InPropertyHandle, FNiagaraClipboardPortableValue& InTargetClipboardPortableValue) const
+{
+	FStructProperty* StructProperty = CastField<FStructProperty>(InPropertyHandle.GetProperty());
+	if (StructProperty != nullptr && StructProperty->Struct->IsChildOf(FNiagaraDistributionBase::StaticStruct()))
+	{
+		void* ValueData;
+		if (InPropertyHandle.GetValueData(ValueData) == FPropertyAccess::Success)
+		{
+			const FNiagaraDistributionBase* Distribution = (const FNiagaraDistributionBase*)ValueData;
+			if (Distribution->IsConstant())
+			{
+				switch (Distribution->ChannelConstantsAndRanges.Num())
+				{
+				case 1:
+				{
+					float FloatValue = Distribution->ChannelConstantsAndRanges[0];
+					InTargetClipboardPortableValue.ValueString = LexToString(FloatValue);
+					return true;
+				}
+				case 2:
+				{
+					UScriptStruct* Vector2Struct = TVariantStructure<FVector2f>::Get();
+					FVector2f Vector2Value(Distribution->ChannelConstantsAndRanges[0], Distribution->ChannelConstantsAndRanges[1]);
+					InTargetClipboardPortableValue = FNiagaraClipboardPortableValue::CreateFromStructValue(*Vector2Struct, (uint8*)&Vector2Value);
+					return true;
+				}
+				case 3:
+				{
+					UScriptStruct* Vector3Struct = TVariantStructure<FVector3f>::Get();
+					FVector3f Vector3Value(Distribution->ChannelConstantsAndRanges[0], Distribution->ChannelConstantsAndRanges[1], Distribution->ChannelConstantsAndRanges[2]);
+					InTargetClipboardPortableValue = FNiagaraClipboardPortableValue::CreateFromStructValue(*Vector3Struct, (uint8*)&Vector3Value);
+					return true;
+				}
+				case 4:
+				{
+					if (Distribution->DisplayAsColor())
+					{
+						UScriptStruct* ColorStruct = TBaseStructure<FLinearColor>::Get();
+						FLinearColor ColorValue(Distribution->ChannelConstantsAndRanges[0], Distribution->ChannelConstantsAndRanges[1], Distribution->ChannelConstantsAndRanges[2], Distribution->ChannelConstantsAndRanges[3]);
+						InTargetClipboardPortableValue = FNiagaraClipboardPortableValue::CreateFromStructValue(*ColorStruct, (uint8*)&ColorValue);
+						return true;
+					}
+					else
+					{
+						UScriptStruct* Vector4Struct = TVariantStructure<FVector4f>::Get();
+						FVector4f Vector4Value(Distribution->ChannelConstantsAndRanges[0], Distribution->ChannelConstantsAndRanges[1], Distribution->ChannelConstantsAndRanges[2], Distribution->ChannelConstantsAndRanges[3]);
+						InTargetClipboardPortableValue = FNiagaraClipboardPortableValue::CreateFromStructValue(*Vector4Struct, (uint8*)&Vector4Value);
+						return true;
+					}
+				}
+				}
+			}
+			else if (Distribution->IsCurve())
+			{
+				UScriptStruct* CurveCollectionStruct = FNiagaraClipboardCurveCollection::StaticStruct();
+				FNiagaraClipboardCurveCollection CurveCollectionValue;
+				CurveCollectionValue.Curves = Distribution->ChannelCurves;
+				InTargetClipboardPortableValue = FNiagaraClipboardPortableValue::CreateFromStructValue(*CurveCollectionStruct, (uint8*)&CurveCollectionValue);
+				return true;
+			}
+		}
+	}
+	return false;
+}
+
+bool FNiagaraDistributionPropertyEditorUtilities::TryUpdatePropertyFromClipboardPortableValue(const FNiagaraClipboardPortableValue& InSourceClipboardPortableValue, IPropertyHandle& InPropertyHandle) const
+{
+	FStructProperty* StructProperty = CastField<FStructProperty>(InPropertyHandle.GetProperty());
+	if (StructProperty != nullptr && StructProperty->Struct->IsChildOf(FNiagaraDistributionBase::StaticStruct()))
+	{
+		void* ValueData;
+		if (InPropertyHandle.GetValueData(ValueData) == FPropertyAccess::Success)
+		{
+			FNiagaraDistributionBase* Distribution = (FNiagaraDistributionBase*)ValueData;
+			TOptional<FNiagaraDistributionBase> ParsedDistribution;
+
+			float FloatValue;
+			if (LexTryParseString(FloatValue, *InSourceClipboardPortableValue.ValueString))
+			{
+				ParsedDistribution = FNiagaraDistributionBase();
+				ParsedDistribution->Mode = ENiagaraDistributionMode::UniformConstant;
+				ParsedDistribution->ChannelConstantsAndRanges.Add(FloatValue);
+			}
+
+			if (ParsedDistribution.IsSet() == false && Distribution->GetBaseNumberOfChannels() == 2)
+			{
+				UScriptStruct* Vector2Struct = TVariantStructure<FVector2f>::Get();
+				FVector2f Vector2Value;
+				if (InSourceClipboardPortableValue.TryUpdateStructValue(*Vector2Struct, (uint8*)&Vector2Value))
+				{
+					ParsedDistribution = FNiagaraDistributionBase();
+					ParsedDistribution->Mode = ENiagaraDistributionMode::NonUniformConstant;
+					ParsedDistribution->ChannelConstantsAndRanges = { Vector2Value.X, Vector2Value.Y };
+				}
+			}
+
+			if (ParsedDistribution.IsSet() == false && Distribution->GetBaseNumberOfChannels() == 3)
+			{
+				UScriptStruct* Vector3Struct = TVariantStructure<FVector3f>::Get();
+				FVector3f Vector3Value;
+				if (InSourceClipboardPortableValue.TryUpdateStructValue(*Vector3Struct, (uint8*)&Vector3Value))
+				{
+					ParsedDistribution = FNiagaraDistributionBase();
+					ParsedDistribution->Mode = ENiagaraDistributionMode::NonUniformConstant;
+					ParsedDistribution->ChannelConstantsAndRanges = { Vector3Value.X, Vector3Value.Y, Vector3Value.Z };
+				}
+			}
+
+			if (ParsedDistribution.IsSet() == false && Distribution->GetBaseNumberOfChannels() == 4 && Distribution->DisplayAsColor() == false)
+			{
+				UScriptStruct* Vector4Struct = TVariantStructure<FVector4f>::Get();
+				FVector4f Vector4Value;
+				if (InSourceClipboardPortableValue.TryUpdateStructValue(*Vector4Struct, (uint8*)&Vector4Value))
+				{
+					ParsedDistribution = FNiagaraDistributionBase();
+					ParsedDistribution->Mode = ENiagaraDistributionMode::NonUniformConstant;
+					ParsedDistribution->ChannelConstantsAndRanges = { Vector4Value.X, Vector4Value.Y, Vector4Value.Z, Vector4Value.W };
+				}
+			}
+
+			if (ParsedDistribution.IsSet() == false && Distribution->GetBaseNumberOfChannels() == 4 && Distribution->DisplayAsColor())
+			{
+				UScriptStruct* ColorStruct = TBaseStructure<FLinearColor>::Get();
+				FLinearColor ColorValue;
+				if (InSourceClipboardPortableValue.TryUpdateStructValue(*ColorStruct, (uint8*)&ColorValue))
+				{
+					ParsedDistribution = FNiagaraDistributionBase();
+					ParsedDistribution->Mode = ENiagaraDistributionMode::NonUniformConstant;
+					ParsedDistribution->ChannelConstantsAndRanges = { ColorValue.R, ColorValue.G, ColorValue.B, ColorValue.A };
+				}
+			}
+
+			if (ParsedDistribution.IsSet() == false && Distribution->AllowCurves())
+			{
+				UScriptStruct* CurveCollectionStruct = FNiagaraClipboardCurveCollection::StaticStruct();
+				FNiagaraClipboardCurveCollection CurveCollectionValue;
+				if (InSourceClipboardPortableValue.TryUpdateStructValue(*CurveCollectionStruct, (uint8*)&CurveCollectionValue))
+				{
+					if (CurveCollectionValue.Curves.Num() == 1 || CurveCollectionValue.Curves.Num() == Distribution->GetBaseNumberOfChannels())
+					{
+						ParsedDistribution = FNiagaraDistributionBase();
+						ParsedDistribution->Mode = CurveCollectionValue.Curves.Num() == 1 ? ENiagaraDistributionMode::UniformCurve : ENiagaraDistributionMode::NonUniformCurve;
+						ParsedDistribution->ChannelCurves = CurveCollectionValue.Curves;
+					}
+				}
+			}
+
+			if (ParsedDistribution.IsSet())
+			{
+				const FScopedTransaction Transaction(NSLOCTEXT("DistributionPropertyEditorUtilities", "SetDistributionTransaction", "Set distribution value."));
+
+				TArray<UObject*> OuterObjects;
+				InPropertyHandle.GetOuterObjects(OuterObjects);
+				for (UObject* OuterObject : OuterObjects)
+				{
+					OuterObject->Modify();
+				}
+
+				InPropertyHandle.NotifyPreChange();
+				Distribution->Mode = ParsedDistribution->Mode;
+				Distribution->ChannelConstantsAndRanges = ParsedDistribution->ChannelConstantsAndRanges;
+				Distribution->ChannelCurves = ParsedDistribution->ChannelCurves;
+				InPropertyHandle.NotifyPostChange(EPropertyChangeType::ValueSet);
+				return true;
+			}
+		}
+	}
+	return false; 
+}
\ No newline at end of file
diff --git a/Engine/Plugins/FX/Niagara/Source/NiagaraEditor/Private/TypeEditorUtilities/NiagaraDistributionPropertyEditorUtilities.h b/Engine/Plugins/FX/Niagara/Source/NiagaraEditor/Private/TypeEditorUtilities/NiagaraDistributionPropertyEditorUtilities.h
new file mode 100644
--- /dev/null
+++ b/Engine/Plugins/FX/Niagara/Source/NiagaraEditor/Private/TypeEditorUtilities/NiagaraDistributionPropertyEditorUtilities.h
@@ -0,0 +1,14 @@
+// Copyright Epic Games, Inc. All Rights Reserved.
+
+#include "INiagaraEditorTypeUtilities.h"
+
+struct FNiagaraClipboardPortableValue;
+class IPropertyHandle;
+
+class FNiagaraDistributionPropertyEditorUtilities : public FNiagaraEditorPropertyUtilities
+{
+public:
+	virtual bool SupportsClipboardPortableValues() const override;
+	virtual bool TryUpdateClipboardPortableValueFromProperty(const IPropertyHandle& InPropertyHandle, FNiagaraClipboardPortableValue& InTargetClipboardPortableValue) const override;
+	virtual bool TryUpdatePropertyFromClipboardPortableValue(const FNiagaraClipboardPortableValue& InSourceClipboardPortableValue, IPropertyHandle& InPropertyHandle) const override;
+};
\ No newline at end of file
diff --git a/Engine/Plugins/FX/Niagara/Source/NiagaraEditor/Private/TypeEditorUtilities/NiagaraEnumTypeEditorUtilities.cpp b/Engine/Plugins/FX/Niagara/Source/NiagaraEditor/Private/TypeEditorUtilities/NiagaraEnumTypeEditorUtilities.cpp
--- a/Engine/Plugins/FX/Niagara/Source/NiagaraEditor/Private/TypeEditorUtilities/NiagaraEnumTypeEditorUtilities.cpp
+++ b/Engine/Plugins/FX/Niagara/Source/NiagaraEditor/Private/TypeEditorUtilities/NiagaraEnumTypeEditorUtilities.cpp
@@ -2,7 +2,9 @@
 
 #include "NiagaraEnumTypeEditorUtilities.h"
 #include "SNiagaraParameterEditor.h"
+#include "NiagaraClipboard.h"
 #include "NiagaraTypes.h"
+#include "NiagaraVariant.h"
 #include "NiagaraEditorCommon.h"
 #include "NiagaraEditorStyle.h"
 #include "NiagaraVariableMetaData.h"
@@ -236,3 +238,34 @@ FText FNiagaraEditorEnumTypeUtilities::GetStackDisplayText(const FNiagaraVariabl
 {
 	return Variable.GetType().GetEnum()->GetDisplayNameTextByValue(Variable.GetValue<int32>());
 }
+
+bool FNiagaraEditorEnumTypeUtilities::TryUpdateClipboardPortableValueFromTypedValue(const FNiagaraTypeDefinition& InSourceType, const FNiagaraVariant& InSourceValue, FNiagaraClipboardPortableValue& InTargetClipboardPortableValue) const
+{
+	if (InSourceType.IsEnum() && InSourceValue.GetNumBytes() == FNiagaraTypeDefinition::GetIntDef().GetSize())
+	{
+		FNiagaraVariable Temp(InSourceType, NAME_None);
+		Temp.SetData(InSourceValue.GetBytes());
+		UEnum* ValueEnum = InSourceType.GetEnum();
+		InTargetClipboardPortableValue.ValueString = ValueEnum->GetNameStringByValue(Temp.GetValue<FNiagaraInt32>().Value);
+		return true;
+	}
+	return false;
+}
+
+bool FNiagaraEditorEnumTypeUtilities::TryUpdateTypedValueFromClipboardPortableValue(const FNiagaraClipboardPortableValue& InSourceClipboardPortableValue, const FNiagaraTypeDefinition& InTargetType, FNiagaraVariant& InTargetValue) const
+{
+	if (InTargetType.IsEnum())
+	{
+		int32 EnumValue = (int32)InTargetType.GetEnum()->GetValueByNameString(InSourceClipboardPortableValue.ValueString);
+		if (EnumValue != INDEX_NONE)
+		{
+			FNiagaraInt32 NiagaraIntValue;
+			NiagaraIntValue.Value = EnumValue;
+			FNiagaraVariable Temp(InTargetType, NAME_None);
+			Temp.SetValue<FNiagaraInt32>(NiagaraIntValue);
+			InTargetValue.SetBytes(Temp.GetData(), Temp.GetSizeInBytes());
+			return true;
+		}
+	}
+	return false;
+}
diff --git a/Engine/Plugins/FX/Niagara/Source/NiagaraEditor/Private/TypeEditorUtilities/NiagaraEnumTypeEditorUtilities.h b/Engine/Plugins/FX/Niagara/Source/NiagaraEditor/Private/TypeEditorUtilities/NiagaraEnumTypeEditorUtilities.h
--- a/Engine/Plugins/FX/Niagara/Source/NiagaraEditor/Private/TypeEditorUtilities/NiagaraEnumTypeEditorUtilities.h
+++ b/Engine/Plugins/FX/Niagara/Source/NiagaraEditor/Private/TypeEditorUtilities/NiagaraEnumTypeEditorUtilities.h
@@ -21,4 +21,7 @@ public:
 	virtual bool SetValueFromDisplayName(const FText& TextValue, FNiagaraVariable& Variable) const override;
 	virtual FText GetSearchTextFromValue(const FNiagaraVariable& AllocatedVariable) const override;
 	virtual FText GetStackDisplayText(const FNiagaraVariable& Variable) const override;
+	virtual bool SupportsClipboardPortableValues() const { return true; }
+	virtual bool TryUpdateClipboardPortableValueFromTypedValue(const FNiagaraTypeDefinition& InSourceType, const FNiagaraVariant& InSourceValue, FNiagaraClipboardPortableValue& InTargetClipboardPortableValue) const override;
+	virtual bool TryUpdateTypedValueFromClipboardPortableValue(const FNiagaraClipboardPortableValue& InSourceClipboardPortableValue, const FNiagaraTypeDefinition& InTargetType, FNiagaraVariant& InTargetValue) const override;
 };
\ No newline at end of file
diff --git a/Engine/Plugins/FX/Niagara/Source/NiagaraEditor/Private/TypeEditorUtilities/NiagaraFloatTypeEditorUtilities.cpp b/Engine/Plugins/FX/Niagara/Source/NiagaraEditor/Private/TypeEditorUtilities/NiagaraFloatTypeEditorUtilities.cpp
--- a/Engine/Plugins/FX/Niagara/Source/NiagaraEditor/Private/TypeEditorUtilities/NiagaraFloatTypeEditorUtilities.cpp
+++ b/Engine/Plugins/FX/Niagara/Source/NiagaraEditor/Private/TypeEditorUtilities/NiagaraFloatTypeEditorUtilities.cpp
@@ -7,7 +7,9 @@
 #include "NiagaraEditorSettings.h"
 #include "NiagaraEditorStyle.h"
 #include "SNiagaraParameterEditor.h"
+#include "NiagaraClipboard.h"
 #include "NiagaraTypes.h"
+#include "NiagaraVariant.h"
 #include "Math/UnitConversion.h"
 #include "Widgets/SNiagaraNumericDropDown.h"
 #include "Widgets/Input/SNumericEntryBox.h"
@@ -267,3 +269,32 @@ FText FNiagaraEditorFloatTypeUtilities::GetSearchTextFromValue(const FNiagaraVar
 {
 	return FText::FromString(GetPinDefaultStringFromValue(AllocatedVariable));
 }
+
+bool FNiagaraEditorFloatTypeUtilities::TryUpdateClipboardPortableValueFromTypedValue(const FNiagaraTypeDefinition& InSourceType, const FNiagaraVariant& InSourceValue, FNiagaraClipboardPortableValue& InTargetClipboardPortableValue) const
+{
+	if (InSourceType == FNiagaraTypeDefinition::GetFloatDef() && InSourceValue.GetNumBytes() == FNiagaraTypeDefinition::GetFloatDef().GetSize())
+	{
+		FNiagaraVariable Temp(InSourceType, NAME_None);
+		Temp.SetData(InSourceValue.GetBytes());
+		float FloatValue = Temp.GetValue<FNiagaraFloat>().Value;
+		InTargetClipboardPortableValue.ValueString = LexToString(FloatValue);
+		return true;
+	}
+	return false;
+}
+
+bool FNiagaraEditorFloatTypeUtilities::TryUpdateTypedValueFromClipboardPortableValue(const FNiagaraClipboardPortableValue& InSourceClipboardPortableValue, const FNiagaraTypeDefinition& InTargetType, FNiagaraVariant& InTargetValue) const
+{
+	float FloatValue;
+	if (InTargetType == FNiagaraTypeDefinition::GetFloatDef() &&
+		LexTryParseString(FloatValue, *InSourceClipboardPortableValue.ValueString))
+	{
+		FNiagaraFloat NiagaraFloatValue;
+		NiagaraFloatValue.Value = FloatValue;
+		FNiagaraVariable Temp(InTargetType, NAME_None);
+		Temp.SetValue<FNiagaraFloat>(NiagaraFloatValue);
+		InTargetValue.SetBytes(Temp.GetData(), Temp.GetSizeInBytes());
+		return true;
+	}
+	return false;
+}
diff --git a/Engine/Plugins/FX/Niagara/Source/NiagaraEditor/Private/TypeEditorUtilities/NiagaraFloatTypeEditorUtilities.h b/Engine/Plugins/FX/Niagara/Source/NiagaraEditor/Private/TypeEditorUtilities/NiagaraFloatTypeEditorUtilities.h
--- a/Engine/Plugins/FX/Niagara/Source/NiagaraEditor/Private/TypeEditorUtilities/NiagaraFloatTypeEditorUtilities.h
+++ b/Engine/Plugins/FX/Niagara/Source/NiagaraEditor/Private/TypeEditorUtilities/NiagaraFloatTypeEditorUtilities.h
@@ -19,4 +19,7 @@ public:
 	virtual bool SetValueFromPinDefaultString(const FString& StringValue, FNiagaraVariable& Variable) const override;
 	virtual FText GetSearchTextFromValue(const FNiagaraVariable& AllocatedVariable) const override;
 	virtual FText GetStackDisplayText(const FNiagaraVariable& Variable) const override { return FText::AsNumber(Variable.GetValue<float>()); };
+	virtual bool SupportsClipboardPortableValues() const { return true; }
+	virtual bool TryUpdateClipboardPortableValueFromTypedValue(const FNiagaraTypeDefinition& InSourceType, const FNiagaraVariant& InSourceValue, FNiagaraClipboardPortableValue& OutTargetClipboardPortableValue) const override;
+	virtual bool TryUpdateTypedValueFromClipboardPortableValue(const FNiagaraClipboardPortableValue& InSourceClipboardPortableValue, const FNiagaraTypeDefinition& InTargetType, FNiagaraVariant& InTargetValue) const override;
 };
\ No newline at end of file
diff --git a/Engine/Plugins/FX/Niagara/Source/NiagaraEditor/Private/TypeEditorUtilities/NiagaraIntegerTypeEditorUtilities.cpp b/Engine/Plugins/FX/Niagara/Source/NiagaraEditor/Private/TypeEditorUtilities/NiagaraIntegerTypeEditorUtilities.cpp
--- a/Engine/Plugins/FX/Niagara/Source/NiagaraEditor/Private/TypeEditorUtilities/NiagaraIntegerTypeEditorUtilities.cpp
+++ b/Engine/Plugins/FX/Niagara/Source/NiagaraEditor/Private/TypeEditorUtilities/NiagaraIntegerTypeEditorUtilities.cpp
@@ -6,7 +6,9 @@
 #include "GraphEditorSettings.h"
 #include "NiagaraEditorSettings.h"
 #include "SNiagaraParameterEditor.h"
+#include "NiagaraClipboard.h"
 #include "NiagaraTypes.h"
+#include "NiagaraVariant.h"
 #include "NiagaraEditorStyle.h"
 #include "Widgets/SNiagaraNumericDropDown.h"
 
@@ -262,3 +264,32 @@ FText FNiagaraEditorIntegerTypeUtilities::GetSearchTextFromValue(const FNiagaraV
 {
 	return FText::FromString(GetPinDefaultStringFromValue(AllocatedVariable));
 }
+
+bool FNiagaraEditorIntegerTypeUtilities::TryUpdateClipboardPortableValueFromTypedValue(const FNiagaraTypeDefinition& InSourceType, const FNiagaraVariant& InSourceValue, FNiagaraClipboardPortableValue& InTargetClipboardPortableValue) const
+{
+	if (InSourceType == FNiagaraTypeDefinition::GetIntDef() && InSourceValue.GetNumBytes() == FNiagaraTypeDefinition::GetIntDef().GetSize())
+	{
+		FNiagaraVariable Temp(InSourceType, NAME_None);
+		Temp.SetData(InSourceValue.GetBytes());
+		int32 IntValue = Temp.GetValue<FNiagaraInt32>().Value;
+		InTargetClipboardPortableValue.ValueString = LexToString(IntValue);
+		return true;
+	}
+	return false;
+}
+
+bool FNiagaraEditorIntegerTypeUtilities::TryUpdateTypedValueFromClipboardPortableValue(const FNiagaraClipboardPortableValue& InSourceClipboardPortableValue, const FNiagaraTypeDefinition& InTargetType, FNiagaraVariant& InTargetValue) const
+{
+	int32 IntValue;
+	if (InTargetType == FNiagaraTypeDefinition::GetIntDef() &&
+		LexTryParseString(IntValue, *InSourceClipboardPortableValue.ValueString))
+	{
+		FNiagaraInt32 NiagaraIntValue;
+		NiagaraIntValue.Value = IntValue;
+		FNiagaraVariable Temp(InTargetType, NAME_None);
+		Temp.SetValue<FNiagaraInt32>(NiagaraIntValue);
+		InTargetValue.SetBytes(Temp.GetData(), Temp.GetSizeInBytes());
+		return true;
+	}
+	return false;
+}
diff --git a/Engine/Plugins/FX/Niagara/Source/NiagaraEditor/Private/TypeEditorUtilities/NiagaraIntegerTypeEditorUtilities.h b/Engine/Plugins/FX/Niagara/Source/NiagaraEditor/Private/TypeEditorUtilities/NiagaraIntegerTypeEditorUtilities.h
--- a/Engine/Plugins/FX/Niagara/Source/NiagaraEditor/Private/TypeEditorUtilities/NiagaraIntegerTypeEditorUtilities.h
+++ b/Engine/Plugins/FX/Niagara/Source/NiagaraEditor/Private/TypeEditorUtilities/NiagaraIntegerTypeEditorUtilities.h
@@ -16,4 +16,7 @@ public:
 	virtual FString GetPinDefaultStringFromValue(const FNiagaraVariable& AllocatedVariable) const override;
 	virtual bool SetValueFromPinDefaultString(const FString& StringValue, FNiagaraVariable& Variable) const override;
 	virtual FText GetSearchTextFromValue(const FNiagaraVariable& AllocatedVariable) const override;
+	virtual bool SupportsClipboardPortableValues() const { return true; }
+	virtual bool TryUpdateClipboardPortableValueFromTypedValue(const FNiagaraTypeDefinition& InSourceType, const FNiagaraVariant& InSourceValue, FNiagaraClipboardPortableValue& InTargetClipboardPortableValue) const override;
+	virtual bool TryUpdateTypedValueFromClipboardPortableValue(const FNiagaraClipboardPortableValue& InSourceClipboardPortableValue, const FNiagaraTypeDefinition& InTargetType, FNiagaraVariant& InTargetValue) const override;
 };
\ No newline at end of file
diff --git a/Engine/Plugins/FX/Niagara/Source/NiagaraEditor/Private/ViewModels/Stack/NiagaraStackFunctionInput.cpp b/Engine/Plugins/FX/Niagara/Source/NiagaraEditor/Private/ViewModels/Stack/NiagaraStackFunctionInput.cpp
--- a/Engine/Plugins/FX/Niagara/Source/NiagaraEditor/Private/ViewModels/Stack/NiagaraStackFunctionInput.cpp
+++ b/Engine/Plugins/FX/Niagara/Source/NiagaraEditor/Private/ViewModels/Stack/NiagaraStackFunctionInput.cpp
@@ -313,12 +313,29 @@ void UNiagaraStackFunctionInput::Copy(UNiagaraClipboardContent* ClipboardContent
 	if (ClipboardInput != nullptr)
 	{
 		ClipboardContent->FunctionInputs.Add(ClipboardInput);
+		FNiagaraClipboardPortableValue InputPortableValue;
+		if (ClipboardInput->ValueMode == ENiagaraClipboardFunctionInputValueMode::Local)
+		{
+			FNiagaraVariant LocalValue;
+			LocalValue.SetBytes(ClipboardInput->Local.GetData(), ClipboardInput->Local.Num());
+			InputPortableValue = FNiagaraClipboardPortableValue::CreateFromTypedValue(InputType, LocalValue);
+		}
+		else if (ClipboardInput->ValueMode == ENiagaraClipboardFunctionInputValueMode::Data)
+		{
+			FNiagaraVariant DataValue;
+			DataValue.SetDataInterface(ClipboardInput->Data);
+			InputPortableValue = FNiagaraClipboardPortableValue::CreateFromTypedValue(InputType, DataValue);
+		}
+		if (InputPortableValue.IsValid())
+		{
+			ClipboardContent->PortableValues.Add(InputPortableValue);
+		}
 	}
 }
 
 bool UNiagaraStackFunctionInput::TestCanPasteWithMessage(const UNiagaraClipboardContent* ClipboardContent, FText& OutMessage) const
 {
-	if (ClipboardContent->FunctionInputs.Num() == 0 || GetIsEnabledAndOwnerIsEnabled() == false)
+	if ((ClipboardContent->FunctionInputs.Num() == 0 && ClipboardContent->PortableValues.Num() == 0) || GetIsEnabledAndOwnerIsEnabled() == false)
 	{
 		// Empty clipboard, or disabled don't allow paste, but be silent.
 		return false;
@@ -353,11 +370,28 @@ bool UNiagaraStackFunctionInput::TestCanPasteWithMessage(const UNiagaraClipboard
 		}
 		return false;
 	}
-	else
+	else if (ClipboardContent->FunctionInputs.Num() > 1)
 	{
 		OutMessage = LOCTEXT("CantPasteMultipleInputs", "Can't paste multiple inputs onto a single input.");
 		return false;
 	}
+	else if (ClipboardContent->PortableValues.Num() == 1)
+	{
+		if (ClipboardContent->PortableValues[0].IsValid())
+		{
+			if (ClipboardContent->PortableValues[0].CanUpdateTypedValue(InputType))
+			{
+				OutMessage = LOCTEXT("PasteValueMessage", "Paste the value from the clipboard here.");
+				return true;
+			}
+			else
+			{
+				OutMessage = LOCTEXT("UnsupportedPortableValueMessage", "Pasting the copied value to this input is unsupported.");
+				return false;
+			}
+		}
+	}
+	return false;
 }
 
 FText UNiagaraStackFunctionInput::GetPasteTransactionText(const UNiagaraClipboardContent* ClipboardContent) const
@@ -367,18 +401,45 @@ FText UNiagaraStackFunctionInput::GetPasteTransactionText(const UNiagaraClipboar
 
 void UNiagaraStackFunctionInput::Paste(const UNiagaraClipboardContent* ClipboardContent, FText& OutPasteWarning)
 {
-	if (ensureMsgf(ClipboardContent != nullptr && ClipboardContent->FunctionInputs.Num() == 1, TEXT("Clipboard must not be null, and must contain a single input.  Call TestCanPasteWithMessage to validate")))
+	if (ensureMsgf(ClipboardContent != nullptr && (ClipboardContent->FunctionInputs.Num() == 1 || ClipboardContent->PortableValues.Num() == 1),
+		TEXT("Clipboard must not be null, and must contain a single input or portable value.  Call TestCanPasteWithMessage to validate")))
 	{
-		if (const UNiagaraClipboardFunctionInput* ClipboardInput = ClipboardContent->FunctionInputs[0])
+		
+		if (ClipboardContent->FunctionInputs.Num() == 1)
 		{
-			if (FNiagaraEditorUtilities::AreTypesAssignable(ClipboardInput->InputType, InputType))
+			if (const UNiagaraClipboardFunctionInput* ClipboardInput = ClipboardContent->FunctionInputs[0])
 			{
-				SetValueFromClipboardFunctionInput(*ClipboardInput);	
+				if (FNiagaraEditorUtilities::AreTypesAssignable(ClipboardInput->InputType, InputType))
+				{
+					SetValueFromClipboardFunctionInput(*ClipboardInput);
+				}
+				else
+				{
+					SetClipboardContentViaConversionScript(*ClipboardInput);
+				}
 			}
-			else
+		}
+		else if (ClipboardContent->PortableValues.Num() == 1)
+		{
+			FNiagaraVariant InputValue;
+			if (ClipboardContent->PortableValues[0].TryUpdateTypedValue(InputType, InputValue))
 			{
-				SetClipboardContentViaConversionScript(*ClipboardInput);
-			}			
+				const UNiagaraClipboardFunctionInput* ClipboardInput = nullptr;
+				if (InputValue.GetMode() == ENiagaraVariantMode::Bytes)
+				{
+					TArray<uint8> ValueBytes;
+					ValueBytes.Append(InputValue.GetBytes(), InputValue.GetNumBytes());
+					ClipboardInput = UNiagaraClipboardFunctionInput::CreateLocalValue(GetTransientPackage(), NAME_None, InputType, TOptional<bool>(), ValueBytes);
+				}
+				else if (InputValue.GetMode() == ENiagaraVariantMode::DataInterface)
+				{
+					ClipboardInput = UNiagaraClipboardFunctionInput::CreateDataValue(GetTransientPackage(), NAME_None, InputType, TOptional<bool>(), InputValue.GetDataInterface());
+				}
+				if (ClipboardInput != nullptr)
+				{
+					SetValueFromClipboardFunctionInput(*ClipboardInput);
+				}
+			}
 		}
 	}
 }
diff --git a/Engine/Plugins/FX/Niagara/Source/NiagaraEditor/Private/ViewModels/Stack/NiagaraStackPropertyRow.cpp b/Engine/Plugins/FX/Niagara/Source/NiagaraEditor/Private/ViewModels/Stack/NiagaraStackPropertyRow.cpp
--- a/Engine/Plugins/FX/Niagara/Source/NiagaraEditor/Private/ViewModels/Stack/NiagaraStackPropertyRow.cpp
+++ b/Engine/Plugins/FX/Niagara/Source/NiagaraEditor/Private/ViewModels/Stack/NiagaraStackPropertyRow.cpp
@@ -1,6 +1,7 @@
 // Copyright Epic Games, Inc. All Rights Reserved.
 
 #include "ViewModels/Stack/NiagaraStackPropertyRow.h"
+#include "NiagaraClipboard.h"
 #include "NiagaraNode.h"
 
 #include "IDetailTreeNode.h"
@@ -11,6 +12,8 @@
 
 #include UE_INLINE_GENERATED_CPP_BY_NAME(NiagaraStackPropertyRow)
 
+#define LOCTEXT_NAMESPACE "NiagaraStackPropertyRow"
+
 void UNiagaraStackPropertyRow::Initialize(
 	FRequiredEntryData InRequiredEntryData,
 	TSharedRef<IDetailTreeNode> InDetailTreeNode,
@@ -94,6 +97,84 @@ bool UNiagaraStackPropertyRow::CanDrag() const
 	return PropertyHandle.IsValid() && PropertyHandle->GetParentHandle().IsValid() && PropertyHandle->GetParentHandle()->AsArray().IsValid();
 }
 
+bool UNiagaraStackPropertyRow::SupportsCopy() const
+{
+	TSharedPtr<IPropertyHandle> PropertyHandle = DetailTreeNode->CreatePropertyHandle();
+	return PropertyHandle.IsValid() && PropertyHandle->IsValidHandle();
+}
+
+bool UNiagaraStackPropertyRow::TestCanCopyWithMessage(FText& OutMessage) const
+{
+	TSharedPtr<IPropertyHandle> PropertyHandle = DetailTreeNode->CreatePropertyHandle();
+	if (PropertyHandle.IsValid() && PropertyHandle->IsValidHandle() && FNiagaraClipboardPortableValue::CreateFromPropertyHandle(*PropertyHandle.Get()).IsValid())
+	{
+		OutMessage = LOCTEXT("CopyMessage", "Copy the value of this property.");
+		return true;
+	}
+	else
+	{
+		OutMessage = LOCTEXT("CantCopyMessage", "This row does not support copying.");
+		return false;
+	}
+}
+
+void UNiagaraStackPropertyRow::Copy(UNiagaraClipboardContent* ClipboardContent) const
+{
+	TSharedPtr<IPropertyHandle> PropertyHandle = DetailTreeNode->CreatePropertyHandle();
+	if (PropertyHandle.IsValid() && PropertyHandle->IsValidHandle())
+	{
+		FNiagaraClipboardPortableValue PortableValue = FNiagaraClipboardPortableValue::CreateFromPropertyHandle(*PropertyHandle.Get());
+		if (PortableValue.IsValid())
+		{
+			ClipboardContent->PortableValues.Add(PortableValue);
+		}
+	}
+}
+
+bool UNiagaraStackPropertyRow::SupportsPaste() const
+{
+	TSharedPtr<IPropertyHandle> PropertyHandle = DetailTreeNode->CreatePropertyHandle();
+	return PropertyHandle.IsValid() && PropertyHandle->IsValidHandle();
+}
+
+bool UNiagaraStackPropertyRow::TestCanPasteWithMessage(const UNiagaraClipboardContent* ClipboardContent, FText& OutMessage) const
+{
+	if (ClipboardContent->PortableValues.Num() != 1 || ClipboardContent->PortableValues[0].IsValid() == false)
+	{
+		return false;
+	}
+
+	TSharedPtr<IPropertyHandle> PropertyHandle = DetailTreeNode->CreatePropertyHandle();
+	if (PropertyHandle.IsValid() == false || PropertyHandle->IsValidHandle() == false)
+	{
+		OutMessage = LOCTEXT("CantPasteMessage", "Can not paste the clipboard value to this row.");
+		return false;
+	}
+	else
+	{
+		OutMessage = LOCTEXT("PasteMessage", "Paste the value from the clipboard to this property.");
+		return true;
+	}
+}
+
+FText UNiagaraStackPropertyRow::GetPasteTransactionText(const UNiagaraClipboardContent* ClipboardContent) const
+{
+	return LOCTEXT("PasteValueToPropertyTransaction", "Paste value to property.");
+}
+
+void UNiagaraStackPropertyRow::Paste(const UNiagaraClipboardContent* ClipboardContent, FText& OutPasteWarning)
+{
+	TSharedPtr<IPropertyHandle> PropertyHandle = DetailTreeNode->CreatePropertyHandle();
+	if (PropertyHandle.IsValid() && PropertyHandle->IsValidHandle() &&
+		ClipboardContent->PortableValues.Num() == 1 && ClipboardContent->PortableValues[0].IsValid())
+	{
+		if (ClipboardContent->PortableValues[0].TryUpdatePropertyHandle(*PropertyHandle.Get()) == false)
+		{
+			OutPasteWarning = LOCTEXT("PasteFailWarning", "Failed to paste the value from the clipboard");
+		}
+	}
+}
+
 void UNiagaraStackPropertyRow::FinalizeInternal()
 {
 	Super::FinalizeInternal();
@@ -268,3 +349,5 @@ FNiagaraHierarchyIdentity UNiagaraStackPropertyRow::DetermineSummaryIdentity() c
 	Identity.Names.Add(DetailTreeNode->GetNodeName());
 	return Identity;
 }
+
+#undef LOCTEXT_NAMESPACE
diff --git a/Engine/Plugins/FX/Niagara/Source/NiagaraEditor/Private/ViewModels/Stack/NiagaraStackStatelessEmitterSimulateGroup.cpp b/Engine/Plugins/FX/Niagara/Source/NiagaraEditor/Private/ViewModels/Stack/NiagaraStackStatelessEmitterSimulateGroup.cpp
--- a/Engine/Plugins/FX/Niagara/Source/NiagaraEditor/Private/ViewModels/Stack/NiagaraStackStatelessEmitterSimulateGroup.cpp
+++ b/Engine/Plugins/FX/Niagara/Source/NiagaraEditor/Private/ViewModels/Stack/NiagaraStackStatelessEmitterSimulateGroup.cpp
@@ -82,6 +82,8 @@ public:
 			StatelessModule->Modify();
 			StatelessModule->SetIsModuleEnabled(true);
 			StatelessModule->PostEditChange();
+			StackEditorData->Modify();
+			StackEditorData->SetStatelessModuleShowWhenDisabled(UNiagaraStackStatelessModuleItem::GenerateStackEditorDataKey(StatelessModule), true);
 			OnItemAdded.ExecuteIfBound(StatelessModule);
 		}
 	}
@@ -287,6 +289,7 @@ void UNiagaraStackStatelessModuleItem::SetIsEnabledInternal(bool bInIsEnabled)
 		StatelessModule->Modify();
 		StatelessModule->SetIsModuleEnabled(bInIsEnabled);
 		StatelessModule->PostEditChange();
+		GetStackEditorData().Modify();
 		GetStackEditorData().SetStatelessModuleShowWhenDisabled(GetStackEditorDataKey(), true);
 		OnDataObjectModified().Broadcast({ StatelessModule }, ENiagaraDataObjectChange::Changed);
 		RefreshChildren();
diff --git a/Engine/Plugins/FX/Niagara/Source/NiagaraEditor/Private/Widgets/INiagaraDistributionAdapter.h b/Engine/Plugins/FX/Niagara/Source/NiagaraEditor/Private/Widgets/INiagaraDistributionAdapter.h
--- a/Engine/Plugins/FX/Niagara/Source/NiagaraEditor/Private/Widgets/INiagaraDistributionAdapter.h
+++ b/Engine/Plugins/FX/Niagara/Source/NiagaraEditor/Private/Widgets/INiagaraDistributionAdapter.h
@@ -67,6 +67,9 @@ public:
 	/** Sets the current mode for this distribution. */
 	virtual void SetDistributionMode(ENiagaraDistributionEditorMode InMode) = 0;
 
+	/** A delegate which is called whenever the distribution mode has changed. */
+	virtual FSimpleMulticastDelegate& OnDistributionEditorModeChanged() = 0;
+
 	/** 
 		Gets a constant or range value by channel index and value index. 
 		@param ChannelIndex the index of the channel to get.
diff --git a/Engine/Plugins/FX/Niagara/Source/NiagaraEditor/Private/Widgets/SNiagaraDistributionEditor.cpp b/Engine/Plugins/FX/Niagara/Source/NiagaraEditor/Private/Widgets/SNiagaraDistributionEditor.cpp
--- a/Engine/Plugins/FX/Niagara/Source/NiagaraEditor/Private/Widgets/SNiagaraDistributionEditor.cpp
+++ b/Engine/Plugins/FX/Niagara/Source/NiagaraEditor/Private/Widgets/SNiagaraDistributionEditor.cpp
@@ -23,15 +23,11 @@ class SNiagaraDistributionModeSelector : public SComboButton
 {
 public:
 	SLATE_BEGIN_ARGS(SNiagaraDistributionModeSelector) { }
-		SLATE_EVENT(FSimpleDelegate, OnDistributionModeChanged)
 	SLATE_END_ARGS();
 
 	void Construct(const FArguments& InArgs, TSharedRef<INiagaraDistributionAdapter> InDistributionAdapter)
 	{
-		
 		DistributionAdapter = InDistributionAdapter;
-		OnDistributionChangedDelegate = InArgs._OnDistributionModeChanged;
-
 		SComboButton::Construct(SComboButton::FArguments()
 			.ButtonStyle(FAppStyle::Get(), "HoverHintOnly")
 			.ContentPadding(FMargin(0))
@@ -104,12 +100,10 @@ private:
 	{
 		DistributionAdapter->SetDistributionMode(InSelectedMode);
 		UpdateCachedValues();
-		OnDistributionChangedDelegate.ExecuteIfBound();
 	}
 
 private:
 	TSharedPtr<INiagaraDistributionAdapter> DistributionAdapter;
-	FSimpleDelegate OnDistributionChangedDelegate;
 
 	mutable TOptional<ENiagaraDistributionEditorMode> ModeCache;
 	mutable const FSlateBrush* ModeIconCache = nullptr;
@@ -120,7 +114,6 @@ class SNiagaraDistributionBindingEditor : public SCompoundWidget
 {
 public:
 	SLATE_BEGIN_ARGS(SNiagaraDistributionBindingEditor) { }
-		//SLATE_EVENT(FSimpleDelegate, OnDistributionModeChanged)
 	SLATE_END_ARGS();
 
 	void Construct(const FArguments& InArgs, TSharedRef<INiagaraDistributionAdapter> InDistributionAdapter)
@@ -619,7 +612,7 @@ private:
 void SNiagaraDistributionEditor::Construct(const FArguments& InArgs, TSharedRef<INiagaraDistributionAdapter> InDistributionAdapter)
 {
 	DistributionAdapter = InDistributionAdapter;
-
+	DistributionAdapter->OnDistributionEditorModeChanged().AddSP(this, &SNiagaraDistributionEditor::OnDistributionModeChanged);
 	ChildSlot
 	[
 		SNew(SHorizontalBox)
@@ -628,7 +621,6 @@ void SNiagaraDistributionEditor::Construct(const FArguments& InArgs, TSharedRef<
 		.Padding(0, 0, 5, 0)
 		[
 			SNew(SNiagaraDistributionModeSelector, InDistributionAdapter)
-			.OnDistributionModeChanged(this, &SNiagaraDistributionEditor::OnDistributionModeChanged)
 		]
 		+ SHorizontalBox::Slot()
 		[
diff --git a/Engine/Plugins/FX/Niagara/Source/NiagaraEditor/Public/INiagaraEditorTypeUtilities.h b/Engine/Plugins/FX/Niagara/Source/NiagaraEditor/Public/INiagaraEditorTypeUtilities.h
--- a/Engine/Plugins/FX/Niagara/Source/NiagaraEditor/Public/INiagaraEditorTypeUtilities.h
+++ b/Engine/Plugins/FX/Niagara/Source/NiagaraEditor/Public/INiagaraEditorTypeUtilities.h
@@ -6,14 +6,19 @@
 #include "Internationalization/Text.h"
 #include "Math/UnitConversion.h"
 #include "NiagaraTypes.h"
+#include "NiagaraVariant.h"
 #include "NiagaraVariableMetaData.h"
 
 class FStructOnScope;
+class IPropertyHandle;
 class SNiagaraParameterEditor;
 class SWidget;
 struct FNiagaraTypeDefinition;
 struct FNiagaraVariable;
+struct FNiagaraVariant;
 struct FNiagaraInputParameterCustomization;
+struct FNiagaraClipboardPortableValue;
+
 
 class INiagaraEditorTypeUtilities
 {
@@ -47,6 +52,14 @@ public:
 	virtual FText GetSearchTextFromValue(const FNiagaraVariable& AllocatedVariable) const = 0;
 
 	virtual FText GetStackDisplayText(const FNiagaraVariable& Variable) const = 0;
+
+	virtual bool SupportsClipboardPortableValues() const = 0;
+
+	virtual bool TryUpdateClipboardPortableValueFromTypedValue(const FNiagaraTypeDefinition& InSourceType, const FNiagaraVariant& InSourceValue, FNiagaraClipboardPortableValue& InTargetClipboardPortableValue) const = 0;
+
+	virtual bool CanUpdateTypedValueFromClipboardPortableValue(const FNiagaraClipboardPortableValue& InSourceClipboardPortableValue, const FNiagaraTypeDefinition& InTargetType) const = 0;
+
+	virtual bool TryUpdateTypedValueFromClipboardPortableValue(const FNiagaraClipboardPortableValue& InSourceClipboardPortableValue, const FNiagaraTypeDefinition& InTargetType, FNiagaraVariant& InTargetValue) const = 0;
 };
 
 class FNiagaraEditorTypeUtilities : public INiagaraEditorTypeUtilities, public TSharedFromThis<FNiagaraEditorTypeUtilities, ESPMode::ThreadSafe>
@@ -72,4 +85,29 @@ public:
 		FString DefaultString = GetPinDefaultStringFromValue(Variable);
 		return FText::FromString(DefaultString.IsEmpty() ? "[?]" : DefaultString);
 	}
-};
\ No newline at end of file
+	virtual bool SupportsClipboardPortableValues() const override { return false; }
+	virtual bool TryUpdateClipboardPortableValueFromTypedValue(const FNiagaraTypeDefinition& InSourceType, const FNiagaraVariant& InSourceValue, FNiagaraClipboardPortableValue& InTargetClipboardPortableValue) const override { return false; }
+	virtual bool CanUpdateTypedValueFromClipboardPortableValue(const FNiagaraClipboardPortableValue& InSourceClipboardPortableValue, const FNiagaraTypeDefinition& InTargetType) const override 
+	{
+		FNiagaraVariant Unused;
+		return TryUpdateTypedValueFromClipboardPortableValue(InSourceClipboardPortableValue, InTargetType, Unused);
+	}
+	virtual bool TryUpdateTypedValueFromClipboardPortableValue(const FNiagaraClipboardPortableValue& InSourceClipboardPortableValue, const FNiagaraTypeDefinition& InTargetType, FNiagaraVariant& InTargetValue) const override { return false; }
+};
+
+class INiagaraEditorPropertyUtilities
+{
+public:
+	virtual bool SupportsClipboardPortableValues() const = 0;
+	virtual bool TryUpdateClipboardPortableValueFromProperty(const IPropertyHandle& InPropertyHandle, FNiagaraClipboardPortableValue& InTargetClipboardPortableValue) const = 0;
+	virtual bool TryUpdatePropertyFromClipboardPortableValue(const FNiagaraClipboardPortableValue& InSourceClipboardPortableValue, IPropertyHandle& InPropertyHandle) const = 0;
+};
+
+class FNiagaraEditorPropertyUtilities : public INiagaraEditorPropertyUtilities
+{
+public:
+	virtual bool SupportsClipboardPortableValues() const override { return false; }
+	virtual bool TryUpdateClipboardPortableValueFromProperty(const IPropertyHandle& InPropertyHandle, FNiagaraClipboardPortableValue& InTargetClipboardPortableValue) const override { return false; }
+	virtual bool TryUpdatePropertyFromClipboardPortableValue(const FNiagaraClipboardPortableValue& InSourceClipboardPortableValue, IPropertyHandle& InPropertyHandle) const override { return false; }
+};
+
diff --git a/Engine/Plugins/FX/Niagara/Source/NiagaraEditor/Public/NiagaraClipboard.h b/Engine/Plugins/FX/Niagara/Source/NiagaraEditor/Public/NiagaraClipboard.h
--- a/Engine/Plugins/FX/Niagara/Source/NiagaraEditor/Public/NiagaraClipboard.h
+++ b/Engine/Plugins/FX/Niagara/Source/NiagaraEditor/Public/NiagaraClipboard.h
@@ -2,6 +2,7 @@
 
 #pragma once
 
+#include "Curves/RichCurve.h"
 #include "NiagaraMessages.h"
 #include "NiagaraScriptVariable.h"
 #include "NiagaraTypes.h"
@@ -10,6 +11,7 @@
 #include "UObject/SoftObjectPtr.h"
 #include "NiagaraClipboard.generated.h"
 
+class IPropertyHandle;
 class UNiagaraDataInterface;
 class UNiagaraScript;
 class UNiagaraRendererProperties;
@@ -179,6 +181,64 @@ struct FNiagaraClipboardScriptVariable
 		return OriginalChangeId == OtherVariable.OriginalChangeId;
 	}
 };
+
+USTRUCT()
+struct FNiagaraClipboardCurveCollection
+{
+	GENERATED_BODY()
+
+	UPROPERTY()
+	TArray<FRichCurve> Curves;
+};
+
+USTRUCT()
+struct FNiagaraClipboardPortableValue
+{
+	GENERATED_BODY()
+
+	class FErrorPipe : public FOutputDevice
+	{
+	public:
+
+		int32 NumErrors;
+
+		FErrorPipe()
+			: FOutputDevice()
+			, NumErrors(0)
+		{
+		}
+
+		virtual void Serialize(const TCHAR* V, ELogVerbosity::Type Verbosity, const class FName& Category) override
+		{
+			NumErrors++;
+		}
+	};
+
+	FNiagaraClipboardPortableValue()
+		: ValueString(FString())
+	{
+	}
+
+	bool IsValid() const { return ValueString.IsEmpty() == false; } // ValueType != ENiagaraClipboardPortableValueType::None; }
+
+	void Reset()
+	{
+		*this = FNiagaraClipboardPortableValue();
+	}
+
+	static FNiagaraClipboardPortableValue NIAGARAEDITOR_API CreateFromStructValue(const UScriptStruct& TargetStruct, uint8* StructMemory);
+	static FNiagaraClipboardPortableValue NIAGARAEDITOR_API CreateFromTypedValue(const FNiagaraTypeDefinition& InType, const FNiagaraVariant& InValue);
+	static FNiagaraClipboardPortableValue NIAGARAEDITOR_API CreateFromPropertyHandle(const IPropertyHandle& InPropertyHandle);
+
+	NIAGARAEDITOR_API bool TryUpdateStructValue(const UScriptStruct& TargetStruct, uint8* StructMemory) const;
+	NIAGARAEDITOR_API bool CanUpdateTypedValue(const FNiagaraTypeDefinition& InTargetInputType) const;
+	NIAGARAEDITOR_API bool TryUpdateTypedValue(const FNiagaraTypeDefinition& InTargetInputType, FNiagaraVariant& InValue) const;
+	NIAGARAEDITOR_API bool TryUpdatePropertyHandle(IPropertyHandle& InTargetPropertyHandle) const;
+
+	UPROPERTY()
+	FString ValueString;
+};
+
 UCLASS(MinimalAPI)
 class UNiagaraClipboardContent : public UObject
 {
@@ -212,6 +272,9 @@ public:
 
 	UPROPERTY()
 	FNiagaraStackNoteData StackNote;
+
+	UPROPERTY()
+	TArray<FNiagaraClipboardPortableValue> PortableValues;
 };
 
 class FNiagaraClipboard
diff --git a/Engine/Plugins/FX/Niagara/Source/NiagaraEditor/Public/NiagaraEditorModule.h b/Engine/Plugins/FX/Niagara/Source/NiagaraEditor/Public/NiagaraEditorModule.h
--- a/Engine/Plugins/FX/Niagara/Source/NiagaraEditor/Public/NiagaraEditorModule.h
+++ b/Engine/Plugins/FX/Niagara/Source/NiagaraEditor/Public/NiagaraEditorModule.h
@@ -19,6 +19,7 @@ class FNiagaraRecentAndFavoritesManager;
 class IAssetTools;
 class IAssetTypeActions;
 class INiagaraEditorTypeUtilities;
+class INiagaraEditorPropertyUtilities;
 class UNiagaraDataInterface;
 class UNiagaraSettings;
 class USequencerSettings;
@@ -187,6 +188,9 @@ public:
 	/** Registers niagara editor type utilities for a specific type. */
 	void RegisterTypeUtilities(FNiagaraTypeDefinition Type, TSharedRef<INiagaraEditorTypeUtilities, ESPMode::ThreadSafe> EditorUtilities);
 
+	/** Registers niagara editor property utilities for a specific struct type. */
+	void RegisterPropertyUtilities(const UScriptStruct* InStruct, TSharedRef<INiagaraEditorPropertyUtilities, ESPMode::ThreadSafe> PropertyUtilities);
+
 	/** Register/unregister niagara editor settings. */
 	void RegisterSettings();
 	void UnregisterSettings();
@@ -194,6 +198,9 @@ public:
 	/** Gets Niagara editor type utilities for a specific type if there are any registered. */
 	TSharedPtr<INiagaraEditorTypeUtilities, ESPMode::ThreadSafe> NIAGARAEDITOR_API GetTypeUtilities(const FNiagaraTypeDefinition& Type);
 
+	/** Gets Niagara editor property utilities for a specific struct if there are any registered. */
+	TSharedPtr<INiagaraEditorPropertyUtilities, ESPMode::ThreadSafe> NIAGARAEDITOR_API GetPropertyUtilities(const UScriptStruct& Struct);
+
 	NIAGARAEDITOR_API void RegisterWidgetProvider(TSharedRef<INiagaraEditorWidgetProvider> InWidgetProvider);
 	NIAGARAEDITOR_API void UnregisterWidgetProvider(TSharedRef<INiagaraEditorWidgetProvider> InWidgetProvider);
 
@@ -404,6 +411,7 @@ private:
 	FCriticalSection TypeEditorsCS;
 	TMap<FNiagaraTypeDefinition, TSharedRef<INiagaraEditorTypeUtilities, ESPMode::ThreadSafe>> TypeToEditorUtilitiesMap;
 	TSharedPtr<INiagaraEditorTypeUtilities, ESPMode::ThreadSafe> EnumTypeUtilities;
+	TMap<const UScriptStruct*, TSharedRef<INiagaraEditorPropertyUtilities, ESPMode::ThreadSafe>> StructToPropertyUtilitiesMap;
 
 	FOnScriptApplied OnScriptAppliedDelegate;
 
diff --git a/Engine/Plugins/FX/Niagara/Source/NiagaraEditor/Public/NiagaraEditorUtilities.h b/Engine/Plugins/FX/Niagara/Source/NiagaraEditor/Public/NiagaraEditorUtilities.h
--- a/Engine/Plugins/FX/Niagara/Source/NiagaraEditor/Public/NiagaraEditorUtilities.h
+++ b/Engine/Plugins/FX/Niagara/Source/NiagaraEditor/Public/NiagaraEditorUtilities.h
@@ -113,7 +113,7 @@ namespace FNiagaraEditorUtilities
 	TSet<FName> GetSystemConstantNames();
 
 	/** Resets the variables value to default, either based on the struct, or if available through registered type utilities. */
-	void ResetVariableToDefaultValue(FNiagaraVariable& Variable);
+	void NIAGARAEDITOR_API ResetVariableToDefaultValue(FNiagaraVariable& Variable);
 
 	/** Fills DefaultData with the types default, either based on the struct, or if available through registered type utilities. */
 	void NIAGARAEDITOR_API GetTypeDefaultValue(const FNiagaraTypeDefinition& Type, TArray<uint8>& DefaultData);
diff --git a/Engine/Plugins/FX/Niagara/Source/NiagaraEditor/Public/NiagaraStackFunctionInputBinder.h b/Engine/Plugins/FX/Niagara/Source/NiagaraEditor/Public/NiagaraStackFunctionInputBinder.h
--- a/Engine/Plugins/FX/Niagara/Source/NiagaraEditor/Public/NiagaraStackFunctionInputBinder.h
+++ b/Engine/Plugins/FX/Niagara/Source/NiagaraEditor/Public/NiagaraStackFunctionInputBinder.h
@@ -17,7 +17,7 @@ class UEdGraphPin;
 class FNiagaraStackFunctionInputBinder
 {
 public:
-	FNiagaraStackFunctionInputBinder();
+	NIAGARAEDITOR_API FNiagaraStackFunctionInputBinder();
 
 	bool TryBind(
 		UNiagaraScript* InScript,
@@ -31,7 +31,7 @@ public:
 		bool bInIsRequired,
 		FText& OutErrorMessage);
 
-	bool TryBind(
+	NIAGARAEDITOR_API bool TryBind(
 		UNiagaraScript* InScript,
 		TArray<UNiagaraScript*> InDependentScripts,
 		FCompileConstantResolver InConstantResolver,
@@ -69,7 +69,7 @@ public:
 
 	TArray<uint8> GetData() const;
 
-	void SetData(const uint8* InValue, int32 InSize);
+	NIAGARAEDITOR_API void SetData(const uint8* InValue, int32 InSize);
 
 private:
 	DECLARE_DELEGATE_RetVal_OneParam(bool, FInputMatchesPredicate, FNiagaraVariable);
diff --git a/Engine/Plugins/FX/Niagara/Source/NiagaraEditor/Public/ViewModels/Stack/NiagaraStackPropertyRow.h b/Engine/Plugins/FX/Niagara/Source/NiagaraEditor/Public/ViewModels/Stack/NiagaraStackPropertyRow.h
--- a/Engine/Plugins/FX/Niagara/Source/NiagaraEditor/Public/ViewModels/Stack/NiagaraStackPropertyRow.h
+++ b/Engine/Plugins/FX/Niagara/Source/NiagaraEditor/Public/ViewModels/Stack/NiagaraStackPropertyRow.h
@@ -43,6 +43,21 @@ public:
 	NIAGARAEDITOR_API virtual bool CanDrag() const override;
 
 	void SetOwnerGuid(TOptional<FGuid> InGuid) { OwnerGuid = InGuid; }
+
+	virtual bool SupportsCopy() const override;
+
+	virtual bool TestCanCopyWithMessage(FText& OutMessage) const override;
+
+	virtual void Copy(UNiagaraClipboardContent* ClipboardContent) const override;
+
+	virtual bool SupportsPaste() const override;
+
+	virtual bool TestCanPasteWithMessage(const UNiagaraClipboardContent* ClipboardContent, FText& OutMessage) const override;
+
+	virtual FText GetPasteTransactionText(const UNiagaraClipboardContent* ClipboardContent) const override;
+
+	virtual void Paste(const UNiagaraClipboardContent* ClipboardContent, FText& OutPasteWarning) override;
+
 protected:
 	NIAGARAEDITOR_API virtual void FinalizeInternal() override;
 
diff --git a/Engine/Plugins/FX/Niagara/Source/NiagaraEditorWidgets/Private/SNiagaraStack.cpp b/Engine/Plugins/FX/Niagara/Source/NiagaraEditorWidgets/Private/SNiagaraStack.cpp
--- a/Engine/Plugins/FX/Niagara/Source/NiagaraEditorWidgets/Private/SNiagaraStack.cpp
+++ b/Engine/Plugins/FX/Niagara/Source/NiagaraEditorWidgets/Private/SNiagaraStack.cpp
@@ -1019,7 +1019,7 @@ SNiagaraStack::FRowWidgets SNiagaraStack::ConstructNameAndValueWidgetsForItem(UN
 		TAttribute<bool> IsEnabled;
 		IsEnabled.BindUObject(Item, &UNiagaraStackEntry::GetIsEnabledAndOwnerIsEnabled);
 		
-		Container->AddFillRowContextMenuHandler(FNiagaraStackPropertyRowUtilities::CreateOnFillRowContextMenu(PropertyRow->GetDetailTreeNode()->CreatePropertyHandle(), PropertyRowWidgets.Actions));
+		Container->AddFillRowContextMenuHandler(FNiagaraStackPropertyRowUtilities::CreateOnFillRowContextMenu(*PropertyRow, PropertyRowWidgets.Actions));
 
 		if (PropertyRowWidgets.WholeRowWidget.IsValid())
 		{
diff --git a/Engine/Plugins/FX/Niagara/Source/NiagaraEditorWidgets/Private/Stack/NiagaraStackPropertyRowUtilities.cpp b/Engine/Plugins/FX/Niagara/Source/NiagaraEditorWidgets/Private/Stack/NiagaraStackPropertyRowUtilities.cpp
--- a/Engine/Plugins/FX/Niagara/Source/NiagaraEditorWidgets/Private/Stack/NiagaraStackPropertyRowUtilities.cpp
+++ b/Engine/Plugins/FX/Niagara/Source/NiagaraEditorWidgets/Private/Stack/NiagaraStackPropertyRowUtilities.cpp
@@ -5,17 +5,31 @@
 #include "HAL/PlatformApplicationMisc.h"
 #include "IDetailTreeNode.h"
 #include "PropertyHandle.h"
+#include "ViewModels/Stack/NiagaraStackPropertyRow.h"
 
-SNiagaraStackTableRow::FOnFillRowContextMenu FNiagaraStackPropertyRowUtilities::CreateOnFillRowContextMenu(TSharedPtr<IPropertyHandle> PropertyHandle, const FNodeWidgetActions& GeneratedPropertyNodeWidgetActions)
+SNiagaraStackTableRow::FOnFillRowContextMenu FNiagaraStackPropertyRowUtilities::CreateOnFillRowContextMenu(UNiagaraStackPropertyRow& PropertyRow, const FNodeWidgetActions& GeneratedPropertyNodeWidgetActions)
 {
-	return SNiagaraStackTableRow::FOnFillRowContextMenu::CreateStatic(&OnFillPropertyRowContextMenu, PropertyHandle, GeneratedPropertyNodeWidgetActions);
+	return SNiagaraStackTableRow::FOnFillRowContextMenu::CreateStatic(&OnFillPropertyRowContextMenu, TWeakObjectPtr<UNiagaraStackPropertyRow>(&PropertyRow), GeneratedPropertyNodeWidgetActions);
 }
 
-void FNiagaraStackPropertyRowUtilities::OnFillPropertyRowContextMenu(FMenuBuilder& MenuBuilder, TSharedPtr<IPropertyHandle> PropertyHandle, FNodeWidgetActions PropertyNodeWidgetActions)
+void FNiagaraStackPropertyRowUtilities::OnFillPropertyRowContextMenu(FMenuBuilder& MenuBuilder, TWeakObjectPtr<UNiagaraStackPropertyRow> PropertyRowWeak, FNodeWidgetActions PropertyNodeWidgetActions)
 {
 	FUIAction CopyAction;
 	FUIAction PasteAction;
 
+	UNiagaraStackPropertyRow* PropertyRow = PropertyRowWeak.Get();
+	if (PropertyRow == nullptr)
+	{
+		return;
+	}
+
+	if (PropertyRow->SupportsCopy() && PropertyRow->SupportsPaste())
+	{
+		// If the row supports stack copy, don't add the actions from the property editor.
+		return;
+	}
+
+	TSharedPtr<IPropertyHandle> PropertyHandle = PropertyRow->GetDetailTreeNode()->CreatePropertyHandle();
 	if (PropertyNodeWidgetActions.CopyMenuAction.ExecuteAction.IsBound() && PropertyNodeWidgetActions.PasteMenuAction.ExecuteAction.IsBound())
 	{
 		CopyAction = PropertyNodeWidgetActions.CopyMenuAction;
diff --git a/Engine/Plugins/FX/Niagara/Source/NiagaraEditorWidgets/Private/Stack/NiagaraStackPropertyRowUtilities.h b/Engine/Plugins/FX/Niagara/Source/NiagaraEditorWidgets/Private/Stack/NiagaraStackPropertyRowUtilities.h
--- a/Engine/Plugins/FX/Niagara/Source/NiagaraEditorWidgets/Private/Stack/NiagaraStackPropertyRowUtilities.h
+++ b/Engine/Plugins/FX/Niagara/Source/NiagaraEditorWidgets/Private/Stack/NiagaraStackPropertyRowUtilities.h
@@ -6,14 +6,14 @@
 #include "Stack/SNiagaraStackTableRow.h"
 #include "Templates/SharedPointer.h"
 
-class IPropertyHandle;
 class FMenuBuilder;
+class UNiagaraStackPropertyRow;
 
 class FNiagaraStackPropertyRowUtilities
 {
 public:
-	static SNiagaraStackTableRow::FOnFillRowContextMenu CreateOnFillRowContextMenu(TSharedPtr<IPropertyHandle> PropertyHandle, const FNodeWidgetActions& GeneratedPropertyNodeWidgetActions);
+	static SNiagaraStackTableRow::FOnFillRowContextMenu CreateOnFillRowContextMenu(UNiagaraStackPropertyRow& PropertyRow, const FNodeWidgetActions& GeneratedPropertyNodeWidgetActions);
 
 private:
-	static void OnFillPropertyRowContextMenu(FMenuBuilder& MenuBuilder, TSharedPtr<IPropertyHandle> PropertyHandle, FNodeWidgetActions PropertyNodeWidgetActions);
+	static void OnFillPropertyRowContextMenu(FMenuBuilder& MenuBuilder, TWeakObjectPtr<UNiagaraStackPropertyRow> PropertyRowWeak, FNodeWidgetActions PropertyNodeWidgetActions);
 };
\ No newline at end of file
