Niagara - Fix missing pragma from 33264004 - Allow copy/paste between module inputs and properties of objects shown in the stack.

diff --git a/Engine/Plugins/FX/Niagara/Source/NiagaraEditor/Private/TypeEditorUtilities/NiagaraDistributionPropertyEditorUtilities.h b/Engine/Plugins/FX/Niagara/Source/NiagaraEditor/Private/TypeEditorUtilities/NiagaraDistributionPropertyEditorUtilities.h
--- a/Engine/Plugins/FX/Niagara/Source/NiagaraEditor/Private/TypeEditorUtilities/NiagaraDistributionPropertyEditorUtilities.h
+++ b/Engine/Plugins/FX/Niagara/Source/NiagaraEditor/Private/TypeEditorUtilities/NiagaraDistributionPropertyEditorUtilities.h
@@ -1,5 +1,7 @@
 // Copyright Epic Games, Inc. All Rights Reserved.
 
+#pragma once
+
 #include "INiagaraEditorTypeUtilities.h"
 
 struct FNiagaraClipboardPortableValue;
