Clean up dynamic string allocations when processing NPC set functions

diff --git a/Engine/Plugins/FX/Niagara/Source/Niagara/Classes/NiagaraParameterCollection.h b/Engine/Plugins/FX/Niagara/Source/Niagara/Classes/NiagaraParameterCollection.h
--- a/Engine/Plugins/FX/Niagara/Source/Niagara/Classes/NiagaraParameterCollection.h
+++ b/Engine/Plugins/FX/Niagara/Source/Niagara/Classes/NiagaraParameterCollection.h
@@ -163,7 +163,12 @@ public:
 	Takes the friendly name presented to the UI and converts to the real parameter name used under the hood.
 	Converts from "ParameterName" to "CollectionUniqueName_ParameterName".
 	*/
-	NIAGARA_API FString ParameterNameFromFriendlyName(const FString& FriendlyName)const;
+	UE_DEPRECATED(5.5, "Use ParameterNameFromFriendlyString(FString) or ConditionalAddFullNamespace(FName) instead")
+	NIAGARA_API FString ParameterNameFromFriendlyName(const FString& FriendlyName) const;
+
+	NIAGARA_API FName ConditionalAddFullNamespace(FName FriendlyParameterName) const;
+	NIAGARA_API FName ParameterNameFromFriendlyString(const FString& FriendlyName) const;
+
 	/**
 	Takes the real parameter name used under the hood and converts to the friendly name for use in the UI.
 	Converts from "CollectionUniqueName_ParameterName" to "ParameterName".
@@ -172,8 +177,14 @@ public:
 	NIAGARA_API FNiagaraVariable CollectionParameterFromFriendlyParameter(const FNiagaraVariable& FriendlyParameter)const;
 	NIAGARA_API FNiagaraVariable FriendlyParameterFromCollectionParameter(const FNiagaraVariable& CollectionParameter)const;
 
-	NIAGARA_API FString FriendlyNameFromParameterName(FString ParameterName)const;
-	NIAGARA_API FString GetFullNamespace()const;
+	UE_DEPRECATED(5.5, "Use FriendlyNameFromParameterName(FName) instead")
+	NIAGARA_API FString FriendlyNameFromParameterName(FString ParameterName) const;
+	NIAGARA_API FName FriendlyNameFromParameterName(FName ParameterName) const;
+
+	UE_DEPRECATED(5.5, "Use GetFullNamespaceName() instead")
+	NIAGARA_API FString GetFullNamespace() const;
+
+	FName GetFullNamespaceName() const { return FullNamespace; };
 
 	/** The compile Id is an indicator to any compiled scripts that reference this collection that contents may have changed and a recompile is recommended to be safe.*/
 	NIAGARA_API FNiagaraCompileHash GetCompileHash() const;
@@ -203,7 +214,7 @@ protected:
 	/** Namespace for this parameter collection. Is enforced to be unique across all parameter collections. */
 	UPROPERTY(EditAnywhere, Category = "Parameter Collection", AssetRegistrySearchable)
 	FName Namespace;
-	
+
 	UPROPERTY()
 	TArray<FNiagaraVariable> Parameters;
 
@@ -217,4 +228,9 @@ protected:
 	/** Used to track whenever something of note changes in this parameter collection that might invalidate a compilation downstream of a script/emitter/system.*/
 	UPROPERTY()
 	FGuid CompileId;
+
+	void BuildFullNamespace();
+
+	// transient variable holding the full namespace (NPC.<namespace>.)
+	FName FullNamespace;
 };
diff --git a/Engine/Plugins/FX/Niagara/Source/Niagara/Private/NiagaraParameterCollection.cpp b/Engine/Plugins/FX/Niagara/Source/Niagara/Private/NiagaraParameterCollection.cpp
--- a/Engine/Plugins/FX/Niagara/Source/Niagara/Private/NiagaraParameterCollection.cpp
+++ b/Engine/Plugins/FX/Niagara/Source/Niagara/Private/NiagaraParameterCollection.cpp
@@ -8,8 +8,7 @@
 #include "Materials/MaterialParameterCollection.h"
 #include "Materials/MaterialParameterCollectionInstance.h"
 #include "Misc/SecureHash.h"
-#if WITH_EDITORONLY_DATA
-#endif
+#include "NiagaraConstants.h"
 #include "AssetRegistry/AssetRegistryModule.h"
 
 #include UE_INLINE_GENERATED_CPP_BY_NAME(NiagaraParameterCollection)
@@ -138,8 +137,8 @@ void UNiagaraParameterCollectionInstance::RefreshSourceParameters(
 	{
 		if (UMaterialParameterCollectionInstance* SourceInstance = World->GetParameterCollectionInstance(SourceCollection))
 		{
-			TStringBuilder<256> VariableName;
-			VariableName << Collection->GetFullNamespace();
+			FNameBuilder VariableName;
+			VariableName << Collection->GetFullNamespaceName();
 			const int32 NamespaceLength = VariableName.Len();
 
 			if (ScalarParameters.Num())
@@ -280,59 +279,61 @@ void UNiagaraParameterCollectionInstance::PostEditChangeProperty(FPropertyChange
 //Blueprint Accessors
 bool UNiagaraParameterCollectionInstance::GetBoolParameter(const FString& InVariableName)
 {
-	return ParameterStorage.GetParameterValue<int32>(FNiagaraVariable(FNiagaraTypeDefinition::GetBoolDef(), *Collection->ParameterNameFromFriendlyName(InVariableName))) == FNiagaraBool::True;
+	return ParameterStorage.GetParameterValue<int32>(FNiagaraVariable(FNiagaraTypeDefinition::GetBoolDef(), Collection->ParameterNameFromFriendlyString(InVariableName))) == FNiagaraBool::True;
 }
 
 float UNiagaraParameterCollectionInstance::GetFloatParameter(const FString& InVariableName)
 {
-	return ParameterStorage.GetParameterValue<float>(FNiagaraVariable(FNiagaraTypeDefinition::GetFloatDef(), *Collection->ParameterNameFromFriendlyName(InVariableName)));
+	return ParameterStorage.GetParameterValue<float>(FNiagaraVariable(FNiagaraTypeDefinition::GetFloatDef(), Collection->ParameterNameFromFriendlyString(InVariableName)));
 }
 
 int32 UNiagaraParameterCollectionInstance::GetIntParameter(const FString& InVariableName)
 {
-	return ParameterStorage.GetParameterValue<int32>(FNiagaraVariable(FNiagaraTypeDefinition::GetIntStruct(), *Collection->ParameterNameFromFriendlyName(InVariableName)));
+	return ParameterStorage.GetParameterValue<int32>(FNiagaraVariable(FNiagaraTypeDefinition::GetIntStruct(), Collection->ParameterNameFromFriendlyString(InVariableName)));
 }
 
 FVector2D UNiagaraParameterCollectionInstance::GetVector2DParameter(const FString& InVariableName)
 {
-	return FVector2D(ParameterStorage.GetParameterValue<FVector2f>(FNiagaraVariable(FNiagaraTypeDefinition::GetVec2Def(), *Collection->ParameterNameFromFriendlyName(InVariableName))));
+	return FVector2D(ParameterStorage.GetParameterValue<FVector2f>(FNiagaraVariable(FNiagaraTypeDefinition::GetVec2Def(), Collection->ParameterNameFromFriendlyString(InVariableName))));
 }
 
 FVector UNiagaraParameterCollectionInstance::GetVectorParameter(const FString& InVariableName)
 {
-	return FVector(ParameterStorage.GetParameterValue<FVector3f>(FNiagaraVariable(FNiagaraTypeDefinition::GetVec3Def(), *Collection->ParameterNameFromFriendlyName(InVariableName))));
+	return FVector(ParameterStorage.GetParameterValue<FVector3f>(FNiagaraVariable(FNiagaraTypeDefinition::GetVec3Def(), Collection->ParameterNameFromFriendlyString(InVariableName))));
 }
 
 FVector4 UNiagaraParameterCollectionInstance::GetVector4Parameter(const FString& InVariableName)
 {
-	return FVector4(ParameterStorage.GetParameterValue<FVector4f>(FNiagaraVariable(FNiagaraTypeDefinition::GetVec4Def(), *Collection->ParameterNameFromFriendlyName(InVariableName))));
+	return FVector4(ParameterStorage.GetParameterValue<FVector4f>(FNiagaraVariable(FNiagaraTypeDefinition::GetVec4Def(), Collection->ParameterNameFromFriendlyString(InVariableName))));
 }
 
 
 FQuat UNiagaraParameterCollectionInstance::GetQuatParameter(const FString& InVariableName)
 {
-	return FQuat(ParameterStorage.GetParameterValue<FQuat4f>(FNiagaraVariable(FNiagaraTypeDefinition::GetQuatDef(), *Collection->ParameterNameFromFriendlyName(InVariableName))));
+	return FQuat(ParameterStorage.GetParameterValue<FQuat4f>(FNiagaraVariable(FNiagaraTypeDefinition::GetQuatDef(), Collection->ParameterNameFromFriendlyString(InVariableName))));
 }
 
 FLinearColor UNiagaraParameterCollectionInstance::GetColorParameter(const FString& InVariableName)
 {
-	return ParameterStorage.GetParameterValue<FLinearColor>(FNiagaraVariable(FNiagaraTypeDefinition::GetColorDef(), *Collection->ParameterNameFromFriendlyName(InVariableName)));
+	return ParameterStorage.GetParameterValue<FLinearColor>(FNiagaraVariable(FNiagaraTypeDefinition::GetColorDef(), Collection->ParameterNameFromFriendlyString(InVariableName)));
 }
 
+#define NPC_SUPPORT_FUNCTION_LOGGING (!UE_BUILD_SHIPPING && !UE_BUILD_TEST)
+
 template<typename T>
-static bool CheckConflictWithSourceMpc(FName ParameterName, FString FunctionCall, const T& Value, const UNiagaraParameterCollection* Collection)
+static bool CheckConflictWithSourceMpc(FName ParameterName, const TCHAR* FunctionCall, const T& Value, const UNiagaraParameterCollection* Collection)
 {
 	if (const UMaterialParameterCollection* SourceCollection = Collection ? Collection->GetSourceCollection() : nullptr)
 	{
 		if (SourceCollection->GetParameterId(ParameterName).IsValid())
 		{
-#if !UE_BUILD_SHIPPING
+#if NPC_SUPPORT_FUNCTION_LOGGING
 			static bool LogWrittenOnce = false;
 
 			if (!LogWrittenOnce)
 			{
 				UE_LOG(LogNiagara, Warning, TEXT("Skipping attempt to %s for parameter %s of %s because it is driven by MPC %s"),
-					*ParameterName.ToString(), *FunctionCall, *Collection->GetFullName(), *Collection->GetSourceCollection()->GetFullName());
+					*ParameterName.ToString(), FunctionCall, *Collection->GetFullName(), *Collection->GetSourceCollection()->GetFullName());
 
 				LogWrittenOnce = true;
 			}
@@ -345,10 +346,16 @@ static bool CheckConflictWithSourceMpc(FName ParameterName, FString FunctionCall
 	return false;
 }
 
+#if NPC_SUPPORT_FUNCTION_LOGGING
+	#define NPC_BUILD_FUNCTION_STRING() (StringCast<TCHAR>(__FUNCTION__).Get())
+#else
+	#define NPC_BUILD_FUNCTION_STRING() (nullptr)
+#endif
+
 void UNiagaraParameterCollectionInstance::SetBoolParameter(const FString& InVariableName, bool InValue)
 {
-	const FName ParameterName = *Collection->ParameterNameFromFriendlyName(InVariableName);
-	if (!CheckConflictWithSourceMpc(ParameterName, __FUNCTION__, InValue, Collection))
+	const FName ParameterName = Collection->ParameterNameFromFriendlyString(InVariableName);
+	if (!CheckConflictWithSourceMpc(ParameterName, NPC_BUILD_FUNCTION_STRING(), InValue, Collection))
 	{
 		ParameterStorage.SetParameterValue(InValue ? FNiagaraBool::True : FNiagaraBool::False, FNiagaraVariable(FNiagaraTypeDefinition::GetBoolDef(), ParameterName));
 	}
@@ -356,8 +363,8 @@ void UNiagaraParameterCollectionInstance::SetBoolParameter(const FString& InVari
 
 void UNiagaraParameterCollectionInstance::SetFloatParameter(const FString& InVariableName, float InValue)
 {
-	const FName ParameterName = *Collection->ParameterNameFromFriendlyName(InVariableName);
-	if (!CheckConflictWithSourceMpc(ParameterName, __FUNCTION__, InValue, Collection))
+	const FName ParameterName = Collection->ParameterNameFromFriendlyString(InVariableName);
+	if (!CheckConflictWithSourceMpc(ParameterName, NPC_BUILD_FUNCTION_STRING(), InValue, Collection))
 	{
 		ParameterStorage.SetParameterValue(InValue, FNiagaraVariable(FNiagaraTypeDefinition::GetFloatDef(), ParameterName));
 	}
@@ -365,8 +372,8 @@ void UNiagaraParameterCollectionInstance::SetFloatParameter(const FString& InVar
 
 void UNiagaraParameterCollectionInstance::SetIntParameter(const FString& InVariableName, int32 InValue)
 {
-	const FName ParameterName = *Collection->ParameterNameFromFriendlyName(InVariableName);
-	if (!CheckConflictWithSourceMpc(ParameterName, __FUNCTION__, InValue, Collection))
+	const FName ParameterName = Collection->ParameterNameFromFriendlyString(InVariableName);
+	if (!CheckConflictWithSourceMpc(ParameterName, NPC_BUILD_FUNCTION_STRING(), InValue, Collection))
 	{
 		ParameterStorage.SetParameterValue(InValue, FNiagaraVariable(FNiagaraTypeDefinition::GetIntDef(), ParameterName));
 	}
@@ -374,8 +381,8 @@ void UNiagaraParameterCollectionInstance::SetIntParameter(const FString& InVaria
 
 void UNiagaraParameterCollectionInstance::SetVector2DParameter(const FString& InVariableName, FVector2D InValue)
 {
-	const FName ParameterName = *Collection->ParameterNameFromFriendlyName(InVariableName);
-	if (!CheckConflictWithSourceMpc(ParameterName, __FUNCTION__, InValue, Collection))
+	const FName ParameterName = Collection->ParameterNameFromFriendlyString(InVariableName);
+	if (!CheckConflictWithSourceMpc(ParameterName, NPC_BUILD_FUNCTION_STRING(), InValue, Collection))
 	{
 		ParameterStorage.SetParameterValue(FVector2f(InValue), FNiagaraVariable(FNiagaraTypeDefinition::GetVec2Def(), ParameterName));
 	}
@@ -383,8 +390,8 @@ void UNiagaraParameterCollectionInstance::SetVector2DParameter(const FString& In
 
 void UNiagaraParameterCollectionInstance::SetVectorParameter(const FString& InVariableName, FVector InValue)
 {
-	const FName ParameterName = *Collection->ParameterNameFromFriendlyName(InVariableName);
-	if (!CheckConflictWithSourceMpc(ParameterName, __FUNCTION__, InValue, Collection))
+	const FName ParameterName = Collection->ParameterNameFromFriendlyString(InVariableName);
+	if (!CheckConflictWithSourceMpc(ParameterName, NPC_BUILD_FUNCTION_STRING(), InValue, Collection))
 	{
 		ParameterStorage.SetParameterValue(FVector3f(InValue), FNiagaraVariable(FNiagaraTypeDefinition::GetVec3Def(), ParameterName));
 	}
@@ -392,8 +399,8 @@ void UNiagaraParameterCollectionInstance::SetVectorParameter(const FString& InVa
 
 void UNiagaraParameterCollectionInstance::SetVector4Parameter(const FString& InVariableName, const FVector4& InValue)
 {
-	const FName ParameterName = *Collection->ParameterNameFromFriendlyName(InVariableName);
-	if (!CheckConflictWithSourceMpc(ParameterName, __FUNCTION__, InValue, Collection))
+	const FName ParameterName = Collection->ParameterNameFromFriendlyString(InVariableName);
+	if (!CheckConflictWithSourceMpc(ParameterName, NPC_BUILD_FUNCTION_STRING(), InValue, Collection))
 	{
 		ParameterStorage.SetParameterValue(FVector4f(InValue), FNiagaraVariable(FNiagaraTypeDefinition::GetVec4Def(), ParameterName));
 	}
@@ -401,8 +408,8 @@ void UNiagaraParameterCollectionInstance::SetVector4Parameter(const FString& InV
 
 void UNiagaraParameterCollectionInstance::SetColorParameter(const FString& InVariableName, FLinearColor InValue)
 {
-	const FName ParameterName = *Collection->ParameterNameFromFriendlyName(InVariableName);
-	if (!CheckConflictWithSourceMpc(ParameterName, __FUNCTION__, InValue, Collection))
+	const FName ParameterName = Collection->ParameterNameFromFriendlyString(InVariableName);
+	if (!CheckConflictWithSourceMpc(ParameterName, NPC_BUILD_FUNCTION_STRING(), InValue, Collection))
 	{
 		ParameterStorage.SetParameterValue(InValue, FNiagaraVariable(FNiagaraTypeDefinition::GetColorDef(), ParameterName));
 	}
@@ -410,8 +417,8 @@ void UNiagaraParameterCollectionInstance::SetColorParameter(const FString& InVar
 
 void UNiagaraParameterCollectionInstance::SetQuatParameter(const FString& InVariableName, const FQuat& InValue)
 {
-	const FName ParameterName = *Collection->ParameterNameFromFriendlyName(InVariableName);
-	if (!CheckConflictWithSourceMpc(ParameterName, __FUNCTION__, InValue, Collection))
+	const FName ParameterName = Collection->ParameterNameFromFriendlyString(InVariableName);
+	if (!CheckConflictWithSourceMpc(ParameterName, NPC_BUILD_FUNCTION_STRING(), InValue, Collection))
 	{
 		ParameterStorage.SetParameterValue(FQuat4f(InValue), FNiagaraVariable(FNiagaraTypeDefinition::GetQuatDef(), ParameterName));
 	}
@@ -423,6 +430,8 @@ UNiagaraParameterCollection::UNiagaraParameterCollection(const FObjectInitialize
 	: Super(ObjectInitializer)
 {
 	Namespace = *GetName();
+	BuildFullNamespace();
+
 	DefaultInstance = ObjectInitializer.CreateDefaultSubobject<UNiagaraParameterCollectionInstance>(this, TEXT("Default Instance"));
 	DefaultInstance->SetParent(this);
 }
@@ -498,9 +507,14 @@ void UNiagaraParameterCollection::RenameParameter(FNiagaraVariable& Parameter, F
 	DefaultInstance->RenameParameter(Parameter, NewName);
 }
 
-FString UNiagaraParameterCollection::GetFullNamespace()const
+void UNiagaraParameterCollection::BuildFullNamespace()
 {
-	return TEXT("NPC.") + Namespace.ToString() + TEXT(".");
+	FNameBuilder FullNamespaceBuilder;
+	FullNamespaceBuilder << PARAM_MAP_NPC_STR;
+	FullNamespaceBuilder << Namespace;
+	FullNamespaceBuilder << TEXT(".");
+
+	FullNamespace = FName(FullNamespaceBuilder);
 }
 
 FNiagaraCompileHash UNiagaraParameterCollection::GetCompileHash() const
@@ -524,27 +538,71 @@ void UNiagaraParameterCollection::RefreshCompileId()
 
 FNiagaraVariable UNiagaraParameterCollection::CollectionParameterFromFriendlyParameter(const FNiagaraVariable& FriendlyParameter)const
 {
-	return FNiagaraVariable(FriendlyParameter.GetType(), *ParameterNameFromFriendlyName(FriendlyParameter.GetName().ToString()));
+	return FNiagaraVariable(FriendlyParameter.GetType(), ConditionalAddFullNamespace(FriendlyParameter.GetName()));
 }
 
 FNiagaraVariable UNiagaraParameterCollection::FriendlyParameterFromCollectionParameter(const FNiagaraVariable& CollectionParameter)const
 {
-	return FNiagaraVariable(CollectionParameter.GetType(), *FriendlyNameFromParameterName(CollectionParameter.GetName().ToString()));
+	return FNiagaraVariable(CollectionParameter.GetType(), ConditionalAddFullNamespace(CollectionParameter.GetName()));
 }
 
-FString UNiagaraParameterCollection::FriendlyNameFromParameterName(FString ParameterName)const
+FString UNiagaraParameterCollection::GetFullNamespace() const
 {
-	ParameterName.RemoveFromStart(GetFullNamespace());
+	FNameBuilder FullNamespaceBuilder(FullNamespace);
+	return FullNamespaceBuilder.ToString();
+}
+
+// deprecated
+FString UNiagaraParameterCollection::FriendlyNameFromParameterName(FString ParameterString) const
+{
+	FNameBuilder ParameterNameBuilder(*ParameterString);
+	return FriendlyNameFromParameterName(FName(ParameterNameBuilder)).ToString();
+}
+
+FName UNiagaraParameterCollection::FriendlyNameFromParameterName(FName ParameterName) const
+{
+	FNameBuilder ParameterNameBuilder(ParameterName);
+	FStringView ParameterNameView = ParameterNameBuilder.ToView();
+
+	FNameBuilder FullNamespaceBuilder(FullNamespace);
+
+	if (ParameterNameView.StartsWith(FullNamespaceBuilder))
+	{
+		ParameterNameView.RemovePrefix(FullNamespaceBuilder.Len());
+		return FName(ParameterNameView);
+	}
+
 	return ParameterName;
 }
 
-FString UNiagaraParameterCollection::ParameterNameFromFriendlyName(const FString& FriendlyName)const
+FName UNiagaraParameterCollection::ConditionalAddFullNamespace(FName ParameterName) const
 {
-	if (FriendlyName.StartsWith(GetFullNamespace()))
+	FNameBuilder ParameterNameBuilder(ParameterName);
+	FNameBuilder FullNamespaceBuilder(FullNamespace);
+
+	if (ParameterNameBuilder.ToView().StartsWith(FullNamespaceBuilder))
 	{
-		return FriendlyName;
+		return ParameterName;
 	}
-	return FString::Printf(TEXT("%s%s"), *GetFullNamespace(), *FriendlyName);
+
+	FNameBuilder ResultName;
+	ResultName << FullNamespace;
+	ResultName << ParameterNameBuilder;
+
+	return FName(ResultName);
+}
+
+// deprecated
+FString UNiagaraParameterCollection::ParameterNameFromFriendlyName(const FString& FriendlyString) const
+{
+	return ParameterNameFromFriendlyString(FriendlyString).ToString();
+}
+
+FName UNiagaraParameterCollection::ParameterNameFromFriendlyString(const FString& FriendlyString) const
+{
+	FName FriendlyName(*FriendlyString);
+
+	return ConditionalAddFullNamespace(FriendlyName);
 }
 
 #if WITH_EDITORONLY_DATA
@@ -582,6 +640,7 @@ void UNiagaraParameterCollection::MakeNamespaceNameUnique()
 
 		UE_LOG(LogNiagara, Warning, TEXT("Parameter collection namespace conflict found. \"%s\" is already in use!"), *Namespace.ToString());
 		Namespace = UniqueName;
+		BuildFullNamespace();
 	}
 }
 
@@ -600,7 +659,7 @@ void UNiagaraParameterCollection::AddDefaultSourceParameters()
 
 		for (const FName& ScalarParameterName : ScalarParameterNames)
 		{
-			FNiagaraVariable ScalarParameter(ScalarDef, *ParameterNameFromFriendlyName(ScalarParameterName.ToString()));
+			FNiagaraVariable ScalarParameter(ScalarDef, ConditionalAddFullNamespace(ScalarParameterName));
 			ScalarParameter.SetValue(SourceMaterialCollection->GetScalarParameterByName(ScalarParameterName)->DefaultValue);
 
 			AddParameter(ScalarParameter);
@@ -608,7 +667,7 @@ void UNiagaraParameterCollection::AddDefaultSourceParameters()
 
 		for (const FName& VectorParameterName : VectorParameterNames)
 		{
-			FNiagaraVariable VectorParameter(ColorDef, *ParameterNameFromFriendlyName(VectorParameterName.ToString()));
+			FNiagaraVariable VectorParameter(ColorDef, ConditionalAddFullNamespace(VectorParameterName));
 			VectorParameter.SetValue(SourceMaterialCollection->GetVectorParameterByName(VectorParameterName)->DefaultValue);
 			AddParameter(VectorParameter);
 		}
diff --git a/Engine/Plugins/FX/Niagara/Source/NiagaraEditor/Private/NiagaraEditorModule.cpp b/Engine/Plugins/FX/Niagara/Source/NiagaraEditor/Private/NiagaraEditorModule.cpp
--- a/Engine/Plugins/FX/Niagara/Source/NiagaraEditor/Private/NiagaraEditorModule.cpp
+++ b/Engine/Plugins/FX/Niagara/Source/NiagaraEditor/Private/NiagaraEditorModule.cpp
@@ -2383,7 +2383,7 @@ UNiagaraParameterCollection* FNiagaraEditorModule::FindCollectionForVariable(con
 		{
 			if (UNiagaraParameterCollection* Collection = CollectionPtr.Get())
 			{
-				if (Prefix.StartsWith(Collection->GetFullNamespace()))
+				if (Prefix.StartsWith(Collection->GetFullNamespaceName().ToString()))
 				{
 					return Collection;
 				}
diff --git a/Engine/Plugins/FX/Niagara/Source/NiagaraEditor/Private/NiagaraParameterMapHistory.cpp b/Engine/Plugins/FX/Niagara/Source/NiagaraEditor/Private/NiagaraParameterMapHistory.cpp
--- a/Engine/Plugins/FX/Niagara/Source/NiagaraEditor/Private/NiagaraParameterMapHistory.cpp
+++ b/Engine/Plugins/FX/Niagara/Source/NiagaraEditor/Private/NiagaraParameterMapHistory.cpp
@@ -930,7 +930,9 @@ FNiagaraCompilationGraphBridge::FParameterCollection TNiagaraParameterMapHistory
 	FString VarName = InVar.GetName().ToString();
 	for (int32 i = 0; i < EncounteredParameterCollections.Collections.Num(); ++i)
 	{
-		if (VarName.StartsWith(EncounteredParameterCollections.CollectionNamespaces[i]))
+		FNameBuilder CollectionNameBuilder(EncounteredParameterCollections.CollectionNamespaces[i]);
+
+		if (VarName.StartsWith(CollectionNameBuilder))
 		{
 			bMissingParameter = !EncounteredParameterCollections.CollectionVariables[i].Contains(InVar);
 			return EncounteredParameterCollections.Collections[i];
@@ -2184,7 +2186,7 @@ bool TNiagaraParameterMapHistoryBuilder<FNiagaraCompilationGraphBridge>::AddPara
 	int32 Index = CollectionStore.Collections.AddUnique(Collection);
 	CollectionStore.CollectionNamespaces.SetNum(CollectionStore.Collections.Num());
 	CollectionStore.CollectionVariables.SetNum(CollectionStore.Collections.Num());
-	CollectionStore.CollectionNamespaces[Index] = Collection->GetFullNamespace();
+	CollectionStore.CollectionNamespaces[Index] = Collection->GetFullNamespaceName();
 	CollectionStore.CollectionVariables[Index] = Collection->GetParameters();
 
 	return true;
diff --git a/Engine/Plugins/FX/Niagara/Source/NiagaraEditor/Private/ViewModels/NiagaraCollectionParameterViewModel.cpp b/Engine/Plugins/FX/Niagara/Source/NiagaraEditor/Private/ViewModels/NiagaraCollectionParameterViewModel.cpp
--- a/Engine/Plugins/FX/Niagara/Source/NiagaraEditor/Private/ViewModels/NiagaraCollectionParameterViewModel.cpp
+++ b/Engine/Plugins/FX/Niagara/Source/NiagaraEditor/Private/ViewModels/NiagaraCollectionParameterViewModel.cpp
@@ -47,7 +47,7 @@ FName FNiagaraCollectionParameterViewModel::GetName() const
 {
 	if (CollectionInst && CollectionInst->GetParent())
 	{
-		return *CollectionInst->GetParent()->FriendlyNameFromParameterName(Parameter.GetName().ToString());
+		return CollectionInst->GetParent()->FriendlyNameFromParameterName(Parameter.GetName());
 	}
 	else
 	{
@@ -62,7 +62,7 @@ FText FNiagaraCollectionParameterViewModel::GetTypeDisplayName() const
 
 void FNiagaraCollectionParameterViewModel::NameTextCommitted(const FText& Name, ETextCommit::Type CommitInfo)
 {
-	FName NewName = *CollectionInst->GetParent()->ParameterNameFromFriendlyName(*Name.ToString());
+	FName NewName = CollectionInst->GetParent()->ParameterNameFromFriendlyString(Name.ToString());
 
 	FName OldName = Parameter.GetName();
 	if (!Parameter.GetName().IsEqual(NewName, ENameCase::CaseSensitive))
@@ -82,7 +82,7 @@ bool FNiagaraCollectionParameterViewModel::VerifyNodeNameTextChanged(const FText
 		return false;
 	}
 
-	FName NewName = *CollectionInst->GetParent()->ParameterNameFromFriendlyName(*NewText.ToString());
+	FName NewName = CollectionInst->GetParent()->ParameterNameFromFriendlyString(NewText.ToString());
 
 	for (FNiagaraVariable& Var : CollectionInst->GetParent()->GetParameters())
 	{
diff --git a/Engine/Plugins/FX/Niagara/Source/NiagaraEditor/Private/ViewModels/NiagaraParameterCollectionAssetViewModel.cpp b/Engine/Plugins/FX/Niagara/Source/NiagaraEditor/Private/ViewModels/NiagaraParameterCollectionAssetViewModel.cpp
--- a/Engine/Plugins/FX/Niagara/Source/NiagaraEditor/Private/ViewModels/NiagaraParameterCollectionAssetViewModel.cpp
+++ b/Engine/Plugins/FX/Niagara/Source/NiagaraEditor/Private/ViewModels/NiagaraParameterCollectionAssetViewModel.cpp
@@ -111,7 +111,7 @@ FName FNiagaraParameterCollectionAssetViewModel::GenerateNewName(FNiagaraTypeDef
 		Existing.Add(ParameterViewModel->GetName());
 	}
 
-	return *Collection->ParameterNameFromFriendlyName(FNiagaraUtilities::GetUniqueName(ProposedName, Existing).ToString());
+	return Collection->ConditionalAddFullNamespace(FNiagaraUtilities::GetUniqueName(ProposedName, Existing));
 }
 
 void FNiagaraParameterCollectionAssetViewModel::AddParameter(TSharedPtr<FNiagaraTypeDefinition> ParameterType)
@@ -228,7 +228,7 @@ void FNiagaraParameterCollectionAssetViewModel::DeleteParameters(TArray<FNiagara
 	TSet<FNiagaraVariableBase> ResolvedParametersToDelete;
 	for(const FNiagaraVariableBase& Parameter : ParametersToDelete)
 	{
-		FNiagaraVariableBase ToDelete(Parameter.GetType(), *Collection->ParameterNameFromFriendlyName(Parameter.GetName().ToString()));
+		FNiagaraVariableBase ToDelete(Parameter.GetType(), Collection->ConditionalAddFullNamespace(Parameter.GetName()));
 		ResolvedParametersToDelete.Add(ToDelete);
 	}
 
@@ -353,7 +353,7 @@ void FNiagaraParameterCollectionAssetViewModel::OnParameterNameChanged(FName Old
 	int32 Index = Collection->IndexOfParameter(ParameterVariable);
 	check(Index != INDEX_NONE);
 
-	FName ParamName = *Collection->ParameterNameFromFriendlyName(GetParameters()[Index]->GetName().ToString());
+	FName ParamName = Collection->ConditionalAddFullNamespace(GetParameters()[Index]->GetName());
 	Collection->GetParameters()[Index].SetName(ParamName);
 	Instance->RenameParameter(ParameterVariable, ParamName);
 	CollectionChanged(false);
diff --git a/Engine/Plugins/FX/Niagara/Source/NiagaraEditor/Public/NiagaraCompilationBridge.h b/Engine/Plugins/FX/Niagara/Source/NiagaraEditor/Public/NiagaraCompilationBridge.h
--- a/Engine/Plugins/FX/Niagara/Source/NiagaraEditor/Public/NiagaraCompilationBridge.h
+++ b/Engine/Plugins/FX/Niagara/Source/NiagaraEditor/Public/NiagaraCompilationBridge.h
@@ -119,7 +119,7 @@ struct FNiagaraCompilationGraphBridge
 		/** Cached off contents of used parameter collections, in case they change during threaded compilation. */
 		TArray<TArray<FNiagaraVariable>> CollectionVariables;
 		/** Cached off contents of used parameter collections, in case they change during threaded compilation. */
-		TArray<FString> CollectionNamespaces;
+		TArray<FName> CollectionNamespaces;
 	};
 
 	struct FAvailableParameterCollections
