Fix for Niagara HWRT causing an assert with shader table
#tests Ran GPU fountain with -dx12 in EngineTest after adding collision module with HWRT as the trace provider

diff --git a/Engine/Plugins/FX/Niagara/Source/NiagaraShader/Private/NiagaraAsyncGpuTraceProviderHwrt.cpp b/Engine/Plugins/FX/Niagara/Source/NiagaraShader/Private/NiagaraAsyncGpuTraceProviderHwrt.cpp
--- a/Engine/Plugins/FX/Niagara/Source/NiagaraShader/Private/NiagaraAsyncGpuTraceProviderHwrt.cpp
+++ b/Engine/Plugins/FX/Niagara/Source/NiagaraShader/Private/NiagaraAsyncGpuTraceProviderHwrt.cpp
@@ -257,6 +257,7 @@ static void BindNiagaraRayTracingMeshCommands(
 		Bindings,
 		bCopyDataToInlineStorage);
 	RHICmdList.SetRayTracingMissShader(RayTracingScene, 0, Pipeline, 0 /* ShaderIndexInPipeline */, 0, nullptr, 0);
+	RHICmdList.CommitRayTracingBindings(RayTracingScene);
 }
 
 ////////////////////////////////////////////////////////////////////////////////
