Fix for stripping graph network paths based on static branches with unconnected selector pins

diff --git a/Engine/Plugins/FX/Niagara/Source/NiagaraEditor/Private/NiagaraGraphDigest.cpp b/Engine/Plugins/FX/Niagara/Source/NiagaraEditor/Private/NiagaraGraphDigest.cpp
--- a/Engine/Plugins/FX/Niagara/Source/NiagaraEditor/Private/NiagaraGraphDigest.cpp
+++ b/Engine/Plugins/FX/Niagara/Source/NiagaraEditor/Private/NiagaraGraphDigest.cpp
@@ -4255,20 +4255,27 @@ TArray<const FNiagaraCompilationInputPin*> FNiagaraCompilationNodeStaticSwitch::
 
 	if (bSetByPin)
 	{
-		const FNiagaraVariable* Found = FNiagaraConstants::FindStaticSwitchConstant(SwitchConstant);
-		FNiagaraVariable Constant = Found ? *Found : FNiagaraVariable();
-		if (Found && Context.ConstantResolver.ResolveConstant(Constant))
+		auto EvaluateVariableSwitchValue = [this](const FNiagaraVariable& ConstantVariable, int32& OutSwitchValue) -> bool
 		{
 			if (SwitchType == ENiagaraStaticSwitchType::Bool)
 			{
-				SwitchValue = Constant.GetValue<bool>();
-				IsValueSet = true;
+				OutSwitchValue = ConstantVariable.GetValue<bool>();
+				return true;
 			}
 			else if (SwitchType == ENiagaraStaticSwitchType::Integer || SwitchType == ENiagaraStaticSwitchType::Enum)
 			{
-				SwitchValue = Constant.GetValue<int32>();
-				IsValueSet = true;
+				OutSwitchValue = ConstantVariable.GetValue<int32>();
+				return true;
 			}
+
+			return false;
+		};
+
+		const FNiagaraVariable* Found = FNiagaraConstants::FindStaticSwitchConstant(SwitchConstant);
+		FNiagaraVariable Constant = Found ? *Found : FNiagaraVariable();
+		if (Found && Context.ConstantResolver.ResolveConstant(Constant))
+		{
+			IsValueSet = EvaluateVariableSwitchValue(Constant, SwitchValue);
 		}
 
 		// if we're set by pin then we can go through the previously completed parameter map traversal history
@@ -4297,18 +4304,7 @@ TArray<const FNiagaraCompilationInputPin*> FNiagaraCompilationNodeStaticSwitch::
 					FNiagaraEditorUtilities::ResetVariableToDefaultValue(DefaultVariable);
 				}
 
-				if (DefaultVariable.GetType().IsSameBaseDefinition(FNiagaraTypeDefinition::GetBoolDef()))
-				{
-					SwitchValue = DefaultVariable.GetValue<bool>();
-				}
-				else if (DefaultVariable.GetType().IsSameBaseDefinition(FNiagaraTypeDefinition::GetIntDef()) || DefaultVariable.GetType().IsEnum())
-				{
-					SwitchValue = DefaultVariable.GetValue<int32>();
-				}
-				else
-				{
-					check(false); // panic
-				}
+				IsValueSet = EvaluateVariableSwitchValue(DefaultVariable, SwitchValue);
 			}
 		}
 	}
