Fix for defaults of MapGet nodes being evaluated when the output was stripped (because of being a part of an untaken static branch).  This could lead to 'Fail If Previously Not Set' errors because the Set itself was being stripped out.

diff --git a/Engine/Plugins/FX/Niagara/Source/NiagaraEditor/Private/NiagaraGraphDigest.cpp b/Engine/Plugins/FX/Niagara/Source/NiagaraEditor/Private/NiagaraGraphDigest.cpp
--- a/Engine/Plugins/FX/Niagara/Source/NiagaraEditor/Private/NiagaraGraphDigest.cpp
+++ b/Engine/Plugins/FX/Niagara/Source/NiagaraEditor/Private/NiagaraGraphDigest.cpp
@@ -1100,6 +1100,8 @@ void FNiagaraCompilationGraphInstanced::Refine(FNiagaraCompilationGraphInstanceC
 		}
 	}
 
+	StripUnconnectedPins(InstantiationContext);
+
 	// validate that now that we've refined the graph we have no more generic numerics and also ensure
 	// that there are no more static switches connected
 	ValidateRefinement();
@@ -1208,6 +1210,55 @@ void FNiagaraCompilationGraphInstanced::ResolveNumerics(FNiagaraCompilationGraph
 	}
 }
 
+void FNiagaraCompilationGraphInstanced::StripUnconnectedPins(FNiagaraCompilationGraphInstanceContext& Context)
+{
+	// through instantiation some pins may become disconnected (for example because of static switches being evaluated).  For some nodes,
+	// like MapGet, this can be problematic because the culled paths may lead to some pins being in a weird state.
+	
+	// If the output of a MapGet is unconnected, then it's default input pin can also be disconnected.
+	TArray<FNiagaraCompilationNodeParameterMapGet*, TInlineAllocator<32>> MapGetNodes;
+
+	for (TUniquePtr<FNiagaraCompilationNode>& Node : Nodes)
+	{
+		if (FNiagaraCompilationNodeParameterMapGet* MapGetNode = Node->AsType<FNiagaraCompilationNodeParameterMapGet>())
+		{
+			MapGetNodes.Add(MapGetNode);
+		}
+	}
+
+	bool bVisitMapGetNodes = !MapGetNodes.IsEmpty();
+
+	while (bVisitMapGetNodes)
+	{
+		bVisitMapGetNodes = false;
+
+		for (FNiagaraCompilationNodeParameterMapGet* MapGetNode : MapGetNodes)
+		{
+			const int32 OutputPinCount = MapGetNode->OutputPins.Num();
+			for (int32 OutputPinIt = 0; OutputPinIt < OutputPinCount; ++OutputPinIt)
+			{
+				const FNiagaraCompilationOutputPin& OutputPin = MapGetNode->OutputPins[OutputPinIt];
+				if (OutputPin.LinkedTo.IsEmpty())
+				{
+					if (MapGetNode->DefaultInputPinIndices.IsValidIndex(OutputPinIt))
+					{
+						const int32 InputPinIndex = MapGetNode->DefaultInputPinIndices[OutputPinIt];
+						if (MapGetNode->InputPins.IsValidIndex(InputPinIndex))
+						{
+							// if we do change something in the graph then revisit our map nodes to see if there
+							// may be more connections that can be removed
+							if (MapGetNode->InputPins[InputPinIndex].Disconnect())
+							{
+								bVisitMapGetNodes = true;
+							}
+						}
+					}
+				}
+			}
+		}
+	}
+}
+
 void FNiagaraCompilationGraphInstanced::InheritDebugState(FNiagaraCompilationGraphInstanceContext& Context, FNiagaraCompilationNodeFunctionCall& FunctionCallNode)
 {
 	FunctionCallNode.DebugState = FunctionCallNode.bInheritDebugState ? Context.ConstantResolver.GetDebugState() : ENiagaraFunctionDebugState::NoDebug;
@@ -1602,6 +1653,19 @@ const FNiagaraCompilationInputPin* FNiagaraCompilationInputPin::TraceBranchMap(c
 	return CurrentInputPin;
 }
 
+bool FNiagaraCompilationInputPin::Disconnect()
+{
+	if (LinkedTo)
+	{
+		FNiagaraCompilationOutputPin* MutableOutputPin = const_cast<FNiagaraCompilationOutputPin*>(LinkedTo);
+		LinkedTo = nullptr;
+
+		return MutableOutputPin->LinkedTo.RemoveSingle(this) > 0;
+	}
+
+	return false;
+}
+
 FNiagaraCompilationOutputPin::FNiagaraCompilationOutputPin(const UEdGraphPin* InPin)
 	: FNiagaraCompilationPin(InPin)
 {
diff --git a/Engine/Plugins/FX/Niagara/Source/NiagaraEditor/Private/NiagaraGraphDigest.h b/Engine/Plugins/FX/Niagara/Source/NiagaraEditor/Private/NiagaraGraphDigest.h
--- a/Engine/Plugins/FX/Niagara/Source/NiagaraEditor/Private/NiagaraGraphDigest.h
+++ b/Engine/Plugins/FX/Niagara/Source/NiagaraEditor/Private/NiagaraGraphDigest.h
@@ -88,6 +88,8 @@ public:
 	FNiagaraCompilationInputPin(const UEdGraphPin* InPin);
 	FNiagaraCompilationInputPin(const FNiagaraCompilationInputPin& SourceInputPin, const FNiagaraCompilationNode* InOwningNode, FNiagaraCompilationGraphDuplicateContext& Context);
 
+	bool Disconnect();
+
 	virtual FLinkedPinView GetLinkedPins() const override
 	{
 		return MakeArrayView(reinterpret_cast<const FNiagaraCompilationPin* const*>(&LinkedTo), 1);
@@ -372,7 +374,7 @@ protected:
 	void PatchGenericNumericsFromCaller(FNiagaraCompilationGraphInstanceContext& Context);
 	void InheritDebugState(FNiagaraCompilationGraphInstanceContext& Context, FNiagaraCompilationNodeFunctionCall& FunctionCallNode);
 	void PropagateDefaultValues(FNiagaraCompilationGraphInstanceContext& Context, FNiagaraCompilationNodeFunctionCall& FunctionCallNode);
-	void RemoveUnconnectedNodes();
+	void StripUnconnectedPins(FNiagaraCompilationGraphInstanceContext& Context);
 };
 
 class FNiagaraCompilationScript
