Fixes to NPC caching of FullNamespace to also handle the PostLoad case
#rnx

diff --git a/Engine/Plugins/FX/Niagara/Source/Niagara/Private/NiagaraParameterCollection.cpp b/Engine/Plugins/FX/Niagara/Source/Niagara/Private/NiagaraParameterCollection.cpp
--- a/Engine/Plugins/FX/Niagara/Source/Niagara/Private/NiagaraParameterCollection.cpp
+++ b/Engine/Plugins/FX/Niagara/Source/Niagara/Private/NiagaraParameterCollection.cpp
@@ -680,6 +680,9 @@ void UNiagaraParameterCollection::PostLoad()
 {
 	Super::PostLoad();
 
+	// after serialization we need to build the transient full namespace
+	BuildFullNamespace();
+
 	DefaultInstance->ConditionalPostLoad();
 
 	if (CompileId.IsValid() == false)
