Niagara Asset Browser Search Bar
- Tag Assets with 0 valid tags are filtered out
- By default we now display a tag asset with 1 tag in parent form. This is more consistent, and while it adds needless clutter, the user can still use the "display as flat list" option

diff --git a/Engine/Plugins/FX/Niagara/Source/NiagaraEditor/Private/Widgets/AssetBrowser/NiagaraMenuFilters.cpp b/Engine/Plugins/FX/Niagara/Source/NiagaraEditor/Private/Widgets/AssetBrowser/NiagaraMenuFilters.cpp
--- a/Engine/Plugins/FX/Niagara/Source/NiagaraEditor/Private/Widgets/AssetBrowser/NiagaraMenuFilters.cpp
+++ b/Engine/Plugins/FX/Niagara/Source/NiagaraEditor/Private/Widgets/AssetBrowser/NiagaraMenuFilters.cpp
@@ -72,4 +72,14 @@ bool FNiagaraAssetBrowserMainFilter::DoesAssetHaveAnyTagFromTagDefinitionsAsset(
 	return AssetTagDefinitionsAsset->DoesAssetDataContainAnyTag(AssetCandidate);
 }
 
+bool FNiagaraAssetBrowserMainFilter::DoesFilterMatchTextQuery(const FText& Text)
+{
+	if(GetDisplayName().ToString().Replace(TEXT(" "), TEXT("")).Contains(Text.ToString()))
+	{
+		return true;
+	}
+
+	return false;
+}
+
 #undef LOCTEXT_NAMESPACE
diff --git a/Engine/Plugins/FX/Niagara/Source/NiagaraEditor/Private/Widgets/AssetBrowser/SNiagaraAssetBrowser.cpp b/Engine/Plugins/FX/Niagara/Source/NiagaraEditor/Private/Widgets/AssetBrowser/SNiagaraAssetBrowser.cpp
--- a/Engine/Plugins/FX/Niagara/Source/NiagaraEditor/Private/Widgets/AssetBrowser/SNiagaraAssetBrowser.cpp
+++ b/Engine/Plugins/FX/Niagara/Source/NiagaraEditor/Private/Widgets/AssetBrowser/SNiagaraAssetBrowser.cpp
@@ -224,8 +224,28 @@ void SNiagaraAssetBrowser::PopulateFiltersSlot()
 		MainFilterSelector->SetSelection(*AllFilter);
 		MainFilterSelector->SetItemExpansion(*AllFilter, true);
 	}
+
+	TSharedRef<SSearchBox> SearchBox = SNew(SSearchBox)
+		.OnTextChanged(this, &SNiagaraAssetBrowser::OnFilterSearchTextChanged)
+		.OnTextCommitted(this, &SNiagaraAssetBrowser::OnFilterSearchTextCommitted)
+		.OnSearch(SSearchBox::FOnSearch::CreateSP(this, &SNiagaraAssetBrowser::OnSearchButtonClicked))
+		.SearchResultData(this, &SNiagaraAssetBrowser::GetSearchResultData)
+		.DelayChangeNotificationsWhileTyping(true);
+
+	TSharedRef<SVerticalBox> MainFilterBox = SNew(SVerticalBox)
+		+ SVerticalBox::Slot()
+		.AutoHeight()
+		.Padding(2.f)
+		[
+			SearchBox
+		]
+		+ SVerticalBox::Slot()
+		.Padding(1.f, 2.f)
+		[
+			MainFilterSelector.ToSharedRef()	
+		];
 	
-	FiltersSlot->AttachWidget(MainFilterSelector.ToSharedRef());
+	FiltersSlot->AttachWidget(MainFilterBox);
 }
 
 void SNiagaraAssetBrowser::PopulateAssetBrowserContentSlot()
@@ -351,8 +371,8 @@ TArray<TSharedRef<FNiagaraAssetBrowserMainFilter>> SNiagaraAssetBrowser::GetMain
 			// If there is an asset, we check if we want to display a parent-entry per asset first
 			else
 			{
-				// If we only want to display the tags or we only have 1 tag defined, we add the tags directly to a flat list
-				if(AssetTagDefinitionData.DefinitionsAsset->DisplayTagsAsFlatList() || AssetTagDefinitionData.AssetTagDefinitions.Num() == 1)
+				// If we only want to display the tags we add the tags directly to a flat list
+				if(AssetTagDefinitionData.DefinitionsAsset->DisplayTagsAsFlatList())
 				{
 					for(const FNiagaraAssetTagDefinition& AssetTagDefinition : AssetTagDefinitionData.AssetTagDefinitions)
 					{
@@ -365,7 +385,21 @@ TArray<TSharedRef<FNiagaraAssetBrowserMainFilter>> SNiagaraAssetBrowser::GetMain
 				// If not, we keep track of the asset here so we can construct a hierarchy of filters per asset
 				else
 				{
-					DisplayedAssetTagDefinitionAssets.Add(AssetTagDefinitionData.DefinitionsAsset);
+					// We only want to display the asset if at least one tag is valid
+					bool bAnyTagValid = false;
+					for(const FNiagaraAssetTagDefinition& AssetTagDefinition : AssetTagDefinitionData.AssetTagDefinitions)
+					{
+						if(IsAssetTagDefinitionValid(AssetTagDefinition))
+						{
+							bAnyTagValid = true;
+							break;
+						}
+					}
+
+					if(bAnyTagValid)
+					{
+						DisplayedAssetTagDefinitionAssets.Add(AssetTagDefinitionData.DefinitionsAsset);
+					}
 				}
 			}			
 		}
@@ -390,6 +424,7 @@ TArray<TSharedRef<FNiagaraAssetBrowserMainFilter>> SNiagaraAssetBrowser::GetMain
 				TArray<TSharedRef<FNiagaraAssetBrowserMainFilter>> TagChildFilters;
 				for(const FNiagaraAssetTagDefinition& AssetTagDefinition : AssetTagDefinitionsAsset->GetAssetTagDefinitions())
 				{
+					// Filter out tags that are invalid; one asset can contain valid and invalid tags for a given use-case
 					if(IsAssetTagDefinitionValid(AssetTagDefinition))
 					{
 						TSharedRef<FNiagaraAssetBrowserMainFilter> AssetTagFilter = MakeShared<FNiagaraAssetBrowserMainFilter>(FNiagaraAssetBrowserMainFilter::EFilterMode::NiagaraAssetTag);
@@ -413,11 +448,193 @@ TArray<TSharedRef<FNiagaraAssetBrowserMainFilter>> SNiagaraAssetBrowser::GetMain
 	return MainFilters;
 }
 
+TArray<TSharedRef<FNiagaraAssetBrowserMainFilter>> SNiagaraAssetBrowser::GetAllFilters() const
+{
+	TArray<TSharedRef<FNiagaraAssetBrowserMainFilter>> AllFilters = GetMainFilters();
+
+	TArray<TSharedRef<FNiagaraAssetBrowserMainFilter>> CurrentFilters = AllFilters;
+	while(CurrentFilters.IsEmpty() == false)
+	{
+		TSharedRef<FNiagaraAssetBrowserMainFilter> CurrentFilter = CurrentFilters[0];
+		CurrentFilters.RemoveAt(0);
+		AllFilters.AddUnique(CurrentFilter);
+		CurrentFilters.Append(CurrentFilter->ChildFilters);
+	}
+
+	return AllFilters;
+}
+
 void SNiagaraAssetBrowser::OnFilterChanged() const
 {
 	RefreshBackendFilter();
 }
 
+void SNiagaraAssetBrowser::ExpandMainFilterSearchResults()
+{
+	MainFilterSelector->ClearExpandedItems();
+
+	for(const FSearchItem& SearchResult : SourceSearchResults)
+	{
+		for(const TSharedRef<FNiagaraAssetBrowserMainFilter>& EntryInPath : SearchResult.Path)
+		{
+			MainFilterSelector->SetItemExpansion(EntryInPath, true);
+		}
+	}
+}
+
+void SNiagaraAssetBrowser::SelectNextMainFilterSearchResult()
+{
+	if(SourceSearchResults.IsEmpty())
+	{
+		return;
+	}
+	
+	if(!FocusedSearchResult.IsSet())
+	{
+		FocusedSearchResult = SourceSearchResults[0];
+	}
+	else
+	{
+		int32 CurrentSearchResultIndex = SourceSearchResults.Find(FocusedSearchResult.GetValue());
+		if(SourceSearchResults.IsValidIndex(CurrentSearchResultIndex+1))
+		{
+			FocusedSearchResult = SourceSearchResults[CurrentSearchResultIndex+1];
+		}
+		else
+		{
+			FocusedSearchResult = SourceSearchResults[0];
+		}
+	}
+
+	MainFilterSelector->ClearSelection();
+	MainFilterSelector->RequestScrollIntoView(FocusedSearchResult.GetValue().GetEntry().ToSharedRef());
+	MainFilterSelector->SetItemSelection(FocusedSearchResult.GetValue().GetEntry().ToSharedRef(), true);
+}
+
+void SNiagaraAssetBrowser::SelectPreviousMainFilterSearchResult()
+{
+	if(SourceSearchResults.IsEmpty())
+	{
+		return;
+	}
+	
+	if(!FocusedSearchResult.IsSet())
+	{
+		FocusedSearchResult = SourceSearchResults[0];
+	}
+	else
+	{
+		int32 CurrentSearchResultIndex = SourceSearchResults.Find(FocusedSearchResult.GetValue());
+		if(SourceSearchResults.IsValidIndex(CurrentSearchResultIndex-1))
+		{
+			FocusedSearchResult = SourceSearchResults[CurrentSearchResultIndex-1];
+		}
+		else
+		{
+			FocusedSearchResult = SourceSearchResults[SourceSearchResults.Num()-1];
+		}
+	}
+
+	MainFilterSelector->ClearSelection();
+	MainFilterSelector->RequestScrollIntoView(FocusedSearchResult.GetValue().GetEntry().ToSharedRef());
+	MainFilterSelector->SetItemSelection(FocusedSearchResult.GetValue().GetEntry().ToSharedRef(), true);
+}
+
+void SNiagaraAssetBrowser::GenerateSearchItems(TSharedRef<FNiagaraAssetBrowserMainFilter> Root, TArray<TSharedRef<FNiagaraAssetBrowserMainFilter>> ParentChain, TArray<FSearchItem>& OutSearchItems) const
+{
+	const TArray<TSharedRef<FNiagaraAssetBrowserMainFilter>> FilteredChildren = Root->ChildFilters;
+	ParentChain.Add(Root);
+	OutSearchItems.Add(FSearchItem{ParentChain});
+	for(TSharedRef<FNiagaraAssetBrowserMainFilter> Child : FilteredChildren)
+	{
+		GenerateSearchItems(Child, ParentChain, OutSearchItems);
+	}
+}
+
+void SNiagaraAssetBrowser::OnFilterSearchTextChanged(const FText& Text)
+{
+	SourceSearchResults.Empty();
+	FocusedSearchResult.Reset();
+	MainFilterSelector->ClearSelection();
+
+	if(!Text.IsEmpty())
+	{
+		FText NoWhitespaceText = FText::FromString(Text.ToString().Replace(TEXT(" "), TEXT("")));
+		TArray<TSharedRef<FNiagaraAssetBrowserMainFilter>> MainFilters = AssetBrowserMainFilters;
+		for(const TSharedRef<FNiagaraAssetBrowserMainFilter>& MainFilter : MainFilters)
+		{
+			TArray<FSearchItem> SearchItems;
+			GenerateSearchItems(MainFilter, {}, SearchItems);
+
+			for(const FSearchItem& SearchItem : SearchItems)
+			{
+				if(SearchItem.GetEntry()->DoesFilterMatchTextQuery(NoWhitespaceText))
+				{
+					SourceSearchResults.Add(SearchItem);
+				}
+			}
+		}
+
+		ExpandMainFilterSearchResults();
+		SelectNextMainFilterSearchResult();
+	}
+	else
+	{
+		MainFilterSelector->ClearExpandedItems();
+	}
+}
+
+void SNiagaraAssetBrowser::OnSearchButtonClicked(SSearchBox::SearchDirection SearchDirection)
+{
+	if(SearchDirection == SSearchBox::Next)
+	{
+		SelectNextMainFilterSearchResult();
+	}
+	else
+	{
+		SelectPreviousMainFilterSearchResult();
+	}
+}
+
+void SNiagaraAssetBrowser::OnFilterSearchTextCommitted(const FText& Text, ETextCommit::Type CommitType)
+{
+	bool bIsShiftDown = FSlateApplication::Get().GetModifierKeys().IsShiftDown();
+	if(CommitType == ETextCommit::OnEnter)
+	{
+		if(bIsShiftDown == false)
+		{
+			SelectNextMainFilterSearchResult();
+		}
+		else
+		{
+			SelectPreviousMainFilterSearchResult();
+		}
+	}
+}
+
+TOptional<SSearchBox::FSearchResultData> SNiagaraAssetBrowser::GetSearchResultData() const
+{
+	if(SourceSearchResults.Num() > 0)
+	{
+		SSearchBox::FSearchResultData SearchResultData;
+		SearchResultData.NumSearchResults = SourceSearchResults.Num();
+
+		if(FocusedSearchResult.IsSet())
+		{
+			// we add one just to make it look nicer as this is merely for cosmetic purposes
+			SearchResultData.CurrentSearchResultIndex = SourceSearchResults.Find(FocusedSearchResult.GetValue()) + 1;
+		}
+		else
+		{
+			SearchResultData.CurrentSearchResultIndex = INDEX_NONE;
+		}
+
+		return SearchResultData;
+	}
+
+	return TOptional<SSearchBox::FSearchResultData>();
+}
+
 void SNiagaraAssetBrowser::OnGetChildFiltersForFilter(TSharedRef<FNiagaraAssetBrowserMainFilter> NiagaraAssetBrowserMainFilter, TArray<TSharedRef<FNiagaraAssetBrowserMainFilter>>& OutChildren) const
 {
 	OutChildren.Append(NiagaraAssetBrowserMainFilter->ChildFilters);
@@ -617,19 +834,9 @@ void SNiagaraAssetBrowser::OnMainFilterSelected(TSharedPtr<FNiagaraAssetBrowserM
 void SNiagaraAssetBrowser::OnAssetTagActivated(const FNiagaraAssetTagDefinition& NiagaraAssetTagDefinition)
 {
 	// First we attempt to select it as a main filter
-	// TArray<FNiagaraAssetBrowserMainFilter> MainFilters = GetMainFilters();
-	TArray<TSharedRef<FNiagaraAssetBrowserMainFilter>> MainFilters = GetMainFilters();
-
-	TArray<TSharedRef<FNiagaraAssetBrowserMainFilter>> AllFilters = MainFilters;
-	while(AllFilters.IsEmpty() == false)
-	{
-		TSharedRef<FNiagaraAssetBrowserMainFilter> CurrentFilter = AllFilters[0];
-		AllFilters.RemoveAt(0);
-		MainFilters.AddUnique(CurrentFilter);
-		AllFilters.Append(CurrentFilter->ChildFilters);
-	}
+	TArray<TSharedRef<FNiagaraAssetBrowserMainFilter>> AllFilters = GetAllFilters();
 	
-	TSharedRef<FNiagaraAssetBrowserMainFilter>* FoundMainFilter = MainFilters.FindByPredicate([NiagaraAssetTagDefinition](const TSharedRef<FNiagaraAssetBrowserMainFilter>& MainFilterCandidate)
+	TSharedRef<FNiagaraAssetBrowserMainFilter>* FoundMainFilter = AllFilters.FindByPredicate([NiagaraAssetTagDefinition](const TSharedRef<FNiagaraAssetBrowserMainFilter>& MainFilterCandidate)
 	{
 		return MainFilterCandidate->AssetTagDefinition == NiagaraAssetTagDefinition;
 	});
diff --git a/Engine/Plugins/FX/Niagara/Source/NiagaraEditor/Public/Widgets/AssetBrowser/NiagaraMenuFilters.h b/Engine/Plugins/FX/Niagara/Source/NiagaraEditor/Public/Widgets/AssetBrowser/NiagaraMenuFilters.h
--- a/Engine/Plugins/FX/Niagara/Source/NiagaraEditor/Public/Widgets/AssetBrowser/NiagaraMenuFilters.h
+++ b/Engine/Plugins/FX/Niagara/Source/NiagaraEditor/Public/Widgets/AssetBrowser/NiagaraMenuFilters.h
@@ -80,6 +80,8 @@ struct FNiagaraAssetBrowserMainFilter
 	bool DoesAssetHaveTag(const FAssetData& AssetCandidate) const;
 	bool DoesAssetHaveAnyTagFromTagDefinitionsAsset(const FAssetData& AssetCandidate) const;
 
+	bool DoesFilterMatchTextQuery(const FText& Text);
+	
 	bool operator==(const FNiagaraAssetBrowserMainFilter& Other) const
 	{
 		if(FilterMode != Other.FilterMode)
diff --git a/Engine/Plugins/FX/Niagara/Source/NiagaraEditor/Public/Widgets/AssetBrowser/SNiagaraAssetBrowser.h b/Engine/Plugins/FX/Niagara/Source/NiagaraEditor/Public/Widgets/AssetBrowser/SNiagaraAssetBrowser.h
--- a/Engine/Plugins/FX/Niagara/Source/NiagaraEditor/Public/Widgets/AssetBrowser/SNiagaraAssetBrowser.h
+++ b/Engine/Plugins/FX/Niagara/Source/NiagaraEditor/Public/Widgets/AssetBrowser/SNiagaraAssetBrowser.h
@@ -56,17 +56,48 @@ private:
 	void PopulateAssetBrowserDetailsSlot();
 	
 	TArray<TSharedRef<FNiagaraAssetBrowserMainFilter>> GetMainFilters() const;
+	TArray<TSharedRef<FNiagaraAssetBrowserMainFilter>> GetAllFilters() const;
 
 private:
 	/** This has to be called whenever any kind of filtering changes and refreshes the actual asset view */
 	void OnFilterChanged() const;
+
+	struct FSearchItem
+	{
+		TArray<TSharedRef<FNiagaraAssetBrowserMainFilter>> Path;
+
+		TSharedPtr<FNiagaraAssetBrowserMainFilter> GetEntry() const
+		{
+			TSharedPtr<FNiagaraAssetBrowserMainFilter> Entry = nullptr;
+			
+			if(Path.Num() > 0)
+			{
+				Entry = Path[Path.Num() - 1];
+			}
+			
+			return Entry;
+		}
+
+		bool operator==(const FSearchItem& Item) const
+		{
+			return Path == Item.Path;
+		}
+	};
 	
+	void OnFilterSearchTextChanged(const FText& Text);
+	void OnSearchButtonClicked(SSearchBox::SearchDirection SearchDirection);
+	void OnFilterSearchTextCommitted(const FText& Text, ETextCommit::Type Arg);
+	TOptional<SSearchBox::FSearchResultData> GetSearchResultData() const;
+	void GenerateSearchItems(TSharedRef<FNiagaraAssetBrowserMainFilter> Root, TArray<TSharedRef<FNiagaraAssetBrowserMainFilter>> ParentChain, TArray<FSearchItem>& OutSearchItems) const;
+	void ExpandMainFilterSearchResults();
+	void SelectNextMainFilterSearchResult();
+	void SelectPreviousMainFilterSearchResult();
+
 	void OnGetChildFiltersForFilter(TSharedRef<FNiagaraAssetBrowserMainFilter> NiagaraAssetBrowserMainFilter, TArray<TSharedRef<FNiagaraAssetBrowserMainFilter>>& OutChildren) const;
 	bool OnCompareMainFiltersForEquality(const FNiagaraAssetBrowserMainFilter& NiagaraAssetBrowserMainFilter, const FNiagaraAssetBrowserMainFilter& NiagaraAssetBrowserMainFilter1) const;
 	
 	TSharedRef<ITableRow> GenerateWidgetRowForMainFilter(TSharedRef<FNiagaraAssetBrowserMainFilter> MainFilter, const TSharedRef<STableViewBase>& OwningTable) const;
-
-
+	
 	void OnAssetSelected(const FAssetData& AssetData);
 	void OnAssetsActivated(const TArray<FAssetData>& AssetData, EAssetTypeActivationMethod::Type) const;
 	TSharedPtr<SWidget> OnGetAssetContextMenu(const TArray<FAssetData>& AssetData) const;
@@ -104,6 +135,9 @@ private:
 
 	mutable TMap<FNiagaraAssetTagDefinition, TSharedRef<FFrontendFilter>> DropdownFilterCache;
 	TArray<TSharedRef<FNiagaraAssetBrowserMainFilter>> AssetBrowserMainFilters;
+
+	TArray<FSearchItem> SourceSearchResults;
+	TOptional<FSearchItem> FocusedSearchResult;
 private:
 	TSharedPtr<STreeView<TSharedRef<FNiagaraAssetBrowserMainFilter>>> MainFilterSelector;
 	TSharedPtr<SNiagaraAssetBrowserContent> AssetBrowserContent;
