- Make SetAsset and bResetExistingOverrideParameters consistent between editor and cooked
- Without this change you can set a user parameter call set asset and that parameter is guaranteed to transfer over in PIE but will not in cooked

#tests Made a test that reproduce the incorrect transfering, tested in PIE & cooked.  Ran EngineTests.

diff --git a/Engine/Plugins/FX/Niagara/Source/Niagara/Private/NiagaraComponent.cpp b/Engine/Plugins/FX/Niagara/Source/Niagara/Private/NiagaraComponent.cpp
--- a/Engine/Plugins/FX/Niagara/Source/Niagara/Private/NiagaraComponent.cpp
+++ b/Engine/Plugins/FX/Niagara/Source/Niagara/Private/NiagaraComponent.cpp
@@ -4094,6 +4094,13 @@ void UNiagaraComponent::SetAsset(UNiagaraSystem* InAsset, bool bResetExistingOve
 	Asset = InAsset;
 
 #if WITH_EDITOR
+	if (Asset != nullptr && bResetExistingOverrideParameters)
+	{
+		// To mirror the cooked path we reset the parameters is the asset is valid
+		OverrideParameters.Empty();
+		InstanceParameterOverrides.Empty();
+	}
+
 	SynchronizeWithSourceSystem();
 	if (Asset != nullptr)
 	{
