- Stateless CPU execution path support
- Modules have been converted (expect noise) to allow execution on the CPU
- Modules can specify the feature mask they support and console variables can be used to control the execution path
- More work to come to ensure random is deterministic between CPU <-> GPU which will allow for dynamic load balancing, i.e. run on CPU if particle count is < X
- Parallel execution of CPU simulation is also possible but not implemented yet

diff --git a/Engine/Plugins/FX/Niagara/Shaders/Private/Stateless/Modules/NiagaraStatelessModule_CalculateAccurateVelocity.ush b/Engine/Plugins/FX/Niagara/Shaders/Private/Stateless/Modules/NiagaraStatelessModule_CalculateAccurateVelocity.ush
--- a/Engine/Plugins/FX/Niagara/Shaders/Private/Stateless/Modules/NiagaraStatelessModule_CalculateAccurateVelocity.ush
+++ b/Engine/Plugins/FX/Niagara/Shaders/Private/Stateless/Modules/NiagaraStatelessModule_CalculateAccurateVelocity.ush
@@ -2,15 +2,7 @@
 
 #include "../NiagaraStatelessCommon.ush"
 
-struct FStatelessModule_CalculateAccurateVelocity
-{
-};
-
-void CalculateAccurateVelocity_Initialize(inout FStatelessModule_CalculateAccurateVelocity ModuleData, inout FStatelessParticle Particle)
-{  
-}
-
-void CalculateAccurateVelocity_Simulate(in FStatelessModule_CalculateAccurateVelocity ModuleData, inout FStatelessParticle Particle)
+void CalculateAccurateVelocity_Simulate(inout FStatelessParticle Particle)
 {  
 	Particle.Velocity			= Particle.Position - Particle.PreviousPosition;
 	Particle.PreviousVelocity	= Particle.Velocity;
diff --git a/Engine/Plugins/FX/Niagara/Shaders/Private/Stateless/Modules/NiagaraStatelessModule_CameraOffset.ush b/Engine/Plugins/FX/Niagara/Shaders/Private/Stateless/Modules/NiagaraStatelessModule_CameraOffset.ush
--- a/Engine/Plugins/FX/Niagara/Shaders/Private/Stateless/Modules/NiagaraStatelessModule_CameraOffset.ush
+++ b/Engine/Plugins/FX/Niagara/Shaders/Private/Stateless/Modules/NiagaraStatelessModule_CameraOffset.ush
@@ -4,21 +4,13 @@
 
 uint3 CameraOffset_Distribution;
 
-struct FStatelessModule_CameraOffset
+void CameraOffset_Simulate(inout FStatelessParticle Particle)
 {
-	FStatelessDistributionFloat	CameraOffsetDistribution;
-};
-
-void CameraOffset_Initialize(inout FStatelessModule_CameraOffset ModuleData, inout FStatelessParticle Particle)
-{
-	ModuleData.CameraOffsetDistribution.Init(Particle, CameraOffset_Distribution);
-}
-	
-void CameraOffset_Simulate(in FStatelessModule_CameraOffset ModuleData, inout FStatelessParticle Particle)
-{
-	if (ModuleData.CameraOffsetDistribution.IsValid())
+	FStatelessDistributionSamplerFloat OffsetSampler;
+	OffsetSampler.Init(Particle, CameraOffset_Distribution);
+	if (OffsetSampler.IsValid())
 	{
-		Particle.CameraOffset			= ModuleData.CameraOffsetDistribution.GetValue(Particle.NormalizedAge);
-		Particle.PreviousCameraOffset	= ModuleData.CameraOffsetDistribution.GetValue(Particle.PreviousNormalizedAge);
+		Particle.CameraOffset			= OffsetSampler.GetValue(Particle.NormalizedAge);
+		Particle.PreviousCameraOffset	= OffsetSampler.GetValue(Particle.PreviousNormalizedAge);
 	}
 }
diff --git a/Engine/Plugins/FX/Niagara/Shaders/Private/Stateless/Modules/NiagaraStatelessModule_DynamicMaterialParameters.ush b/Engine/Plugins/FX/Niagara/Shaders/Private/Stateless/Modules/NiagaraStatelessModule_DynamicMaterialParameters.ush
--- a/Engine/Plugins/FX/Niagara/Shaders/Private/Stateless/Modules/NiagaraStatelessModule_DynamicMaterialParameters.ush
+++ b/Engine/Plugins/FX/Niagara/Shaders/Private/Stateless/Modules/NiagaraStatelessModule_DynamicMaterialParameters.ush
@@ -7,34 +7,79 @@ uint3	DynamicMaterialParameters_Parameter0X;
 uint3	DynamicMaterialParameters_Parameter0Y;
 uint3	DynamicMaterialParameters_Parameter0Z;
 uint3	DynamicMaterialParameters_Parameter0W;
-
-struct FStatelessModule_DynamicMaterialParameters
+uint3	DynamicMaterialParameters_Parameter1X;
+uint3	DynamicMaterialParameters_Parameter1Y;
+uint3	DynamicMaterialParameters_Parameter1Z;
+uint3	DynamicMaterialParameters_Parameter1W;
+uint3	DynamicMaterialParameters_Parameter2X;
+uint3	DynamicMaterialParameters_Parameter2Y;
+uint3	DynamicMaterialParameters_Parameter2Z;
+uint3	DynamicMaterialParameters_Parameter2W;
+uint3	DynamicMaterialParameters_Parameter3X;
+uint3	DynamicMaterialParameters_Parameter3Y;
+uint3	DynamicMaterialParameters_Parameter3Z;
+uint3	DynamicMaterialParameters_Parameter3W;
+	
+void DynamicMaterialParameters_Simulate(inout FStatelessParticle Particle)
 {
-	FStatelessDistributionFloat	Parameter0XDistribution;
-	FStatelessDistributionFloat	Parameter0YDistribution;
-	FStatelessDistributionFloat	Parameter0ZDistribution;
-	FStatelessDistributionFloat	Parameter0WDistribution;
-};
+	Particle.DynamicMaterialParameter0 = 0.0f;
+	Particle.DynamicMaterialParameter1 = 0.0f;
+	Particle.DynamicMaterialParameter2 = 0.0f;
+	Particle.DynamicMaterialParameter3 = 0.0f;
 
-void DynamicMaterialParameters_Initialize(inout FStatelessModule_DynamicMaterialParameters ModuleData, inout FStatelessParticle Particle)
-{
-	if ((DynamicMaterialParameters_ChannelMask & 0xf) != 0)
+	if ((DynamicMaterialParameters_ChannelMask & 0x000f) != 0)
 	{
-		ModuleData.Parameter0XDistribution.Init(Particle, DynamicMaterialParameters_Parameter0X);
-		ModuleData.Parameter0YDistribution.Init(Particle, DynamicMaterialParameters_Parameter0Y);
-		ModuleData.Parameter0ZDistribution.Init(Particle, DynamicMaterialParameters_Parameter0Z);
-		ModuleData.Parameter0WDistribution.Init(Particle, DynamicMaterialParameters_Parameter0W);
+		FStatelessDistributionSamplerFloat ParameterSampler[4];
+		ParameterSampler[0].Init(Particle, DynamicMaterialParameters_Parameter0X);
+		ParameterSampler[1].Init(Particle, DynamicMaterialParameters_Parameter0Y);
+		ParameterSampler[2].Init(Particle, DynamicMaterialParameters_Parameter0Z);
+		ParameterSampler[3].Init(Particle, DynamicMaterialParameters_Parameter0W);
+		
+		Particle.DynamicMaterialParameter0.x = (DynamicMaterialParameters_ChannelMask & 0x0001) != 0 ? ParameterSampler[0].GetValue(Particle.NormalizedAge) : 0.0f;
+		Particle.DynamicMaterialParameter0.y = (DynamicMaterialParameters_ChannelMask & 0x0002) != 0 ? ParameterSampler[1].GetValue(Particle.NormalizedAge) : 0.0f;
+		Particle.DynamicMaterialParameter0.z = (DynamicMaterialParameters_ChannelMask & 0x0004) != 0 ? ParameterSampler[2].GetValue(Particle.NormalizedAge) : 0.0f;
+		Particle.DynamicMaterialParameter0.w = (DynamicMaterialParameters_ChannelMask & 0x0008) != 0 ? ParameterSampler[3].GetValue(Particle.NormalizedAge) : 0.0f;
 	}
-}
-	
-void DynamicMaterialParameters_Simulate(in FStatelessModule_DynamicMaterialParameters ModuleData, inout FStatelessParticle Particle)
-{
-	Particle.DynamicMaterialParameter0 = 0.0f;
-	if ((DynamicMaterialParameters_ChannelMask & 0xf) != 0)
+
+	if ((DynamicMaterialParameters_ChannelMask & 0x00f0) != 0)
+	{
+		FStatelessDistributionSamplerFloat ParameterSampler[4];
+		ParameterSampler[0].Init(Particle, DynamicMaterialParameters_Parameter1X);
+		ParameterSampler[1].Init(Particle, DynamicMaterialParameters_Parameter1Y);
+		ParameterSampler[2].Init(Particle, DynamicMaterialParameters_Parameter1Z);
+		ParameterSampler[3].Init(Particle, DynamicMaterialParameters_Parameter1W);
+		
+		Particle.DynamicMaterialParameter1.x = (DynamicMaterialParameters_ChannelMask & 0x0010) != 0 ? ParameterSampler[0].GetValue(Particle.NormalizedAge) : 0.0f;
+		Particle.DynamicMaterialParameter1.y = (DynamicMaterialParameters_ChannelMask & 0x0020) != 0 ? ParameterSampler[1].GetValue(Particle.NormalizedAge) : 0.0f;
+		Particle.DynamicMaterialParameter1.z = (DynamicMaterialParameters_ChannelMask & 0x0040) != 0 ? ParameterSampler[2].GetValue(Particle.NormalizedAge) : 0.0f;
+		Particle.DynamicMaterialParameter1.w = (DynamicMaterialParameters_ChannelMask & 0x0080) != 0 ? ParameterSampler[3].GetValue(Particle.NormalizedAge) : 0.0f;
+	}
+
+	if ((DynamicMaterialParameters_ChannelMask & 0x0f00) != 0)
+	{
+		FStatelessDistributionSamplerFloat ParameterSampler[4];
+		ParameterSampler[0].Init(Particle, DynamicMaterialParameters_Parameter2X);
+		ParameterSampler[1].Init(Particle, DynamicMaterialParameters_Parameter2Y);
+		ParameterSampler[2].Init(Particle, DynamicMaterialParameters_Parameter2Z);
+		ParameterSampler[3].Init(Particle, DynamicMaterialParameters_Parameter2W);
+		
+		Particle.DynamicMaterialParameter2.x = (DynamicMaterialParameters_ChannelMask & 0x0100) != 0 ? ParameterSampler[0].GetValue(Particle.NormalizedAge) : 0.0f;
+		Particle.DynamicMaterialParameter2.y = (DynamicMaterialParameters_ChannelMask & 0x0200) != 0 ? ParameterSampler[1].GetValue(Particle.NormalizedAge) : 0.0f;
+		Particle.DynamicMaterialParameter2.z = (DynamicMaterialParameters_ChannelMask & 0x0400) != 0 ? ParameterSampler[2].GetValue(Particle.NormalizedAge) : 0.0f;
+		Particle.DynamicMaterialParameter2.w = (DynamicMaterialParameters_ChannelMask & 0x0800) != 0 ? ParameterSampler[3].GetValue(Particle.NormalizedAge) : 0.0f;
+	}
+
+	if ((DynamicMaterialParameters_ChannelMask & 0xf000) != 0)
 	{
-		Particle.DynamicMaterialParameter0.x = (DynamicMaterialParameters_ChannelMask & 0x1) != 0 ? ModuleData.Parameter0XDistribution.GetValue(Particle.NormalizedAge) : 0.0f;
-		Particle.DynamicMaterialParameter0.y = (DynamicMaterialParameters_ChannelMask & 0x2) != 0 ? ModuleData.Parameter0YDistribution.GetValue(Particle.NormalizedAge) : 0.0f;
-		Particle.DynamicMaterialParameter0.z = (DynamicMaterialParameters_ChannelMask & 0x4) != 0 ? ModuleData.Parameter0ZDistribution.GetValue(Particle.NormalizedAge) : 0.0f;
-		Particle.DynamicMaterialParameter0.w = (DynamicMaterialParameters_ChannelMask & 0x8) != 0 ? ModuleData.Parameter0WDistribution.GetValue(Particle.NormalizedAge) : 0.0f;
+		FStatelessDistributionSamplerFloat ParameterSampler[4];
+		ParameterSampler[0].Init(Particle, DynamicMaterialParameters_Parameter3X);
+		ParameterSampler[1].Init(Particle, DynamicMaterialParameters_Parameter3Y);
+		ParameterSampler[2].Init(Particle, DynamicMaterialParameters_Parameter3Z);
+		ParameterSampler[3].Init(Particle, DynamicMaterialParameters_Parameter3W);
+		
+		Particle.DynamicMaterialParameter3.x = (DynamicMaterialParameters_ChannelMask & 0x1000) != 0 ? ParameterSampler[0].GetValue(Particle.NormalizedAge) : 0.0f;
+		Particle.DynamicMaterialParameter3.y = (DynamicMaterialParameters_ChannelMask & 0x2000) != 0 ? ParameterSampler[1].GetValue(Particle.NormalizedAge) : 0.0f;
+		Particle.DynamicMaterialParameter3.z = (DynamicMaterialParameters_ChannelMask & 0x4000) != 0 ? ParameterSampler[2].GetValue(Particle.NormalizedAge) : 0.0f;
+		Particle.DynamicMaterialParameter3.w = (DynamicMaterialParameters_ChannelMask & 0x8000) != 0 ? ParameterSampler[3].GetValue(Particle.NormalizedAge) : 0.0f;
 	}
 }
diff --git a/Engine/Plugins/FX/Niagara/Shaders/Private/Stateless/Modules/NiagaraStatelessModule_InitialMeshOrientation.ush b/Engine/Plugins/FX/Niagara/Shaders/Private/Stateless/Modules/NiagaraStatelessModule_InitialMeshOrientation.ush
--- a/Engine/Plugins/FX/Niagara/Shaders/Private/Stateless/Modules/NiagaraStatelessModule_InitialMeshOrientation.ush
+++ b/Engine/Plugins/FX/Niagara/Shaders/Private/Stateless/Modules/NiagaraStatelessModule_InitialMeshOrientation.ush
@@ -5,17 +5,9 @@
 float3	InitialMeshOrientation_Rotation;
 float3	InitialMeshOrientation_RandomRangeScale;
 
-struct FStatelessModule_InitialMeshOrientation
-{
-};
-
-void InitialMeshOrientation_Initialize(inout FStatelessModule_InitialMeshOrientation ModuleData, inout FStatelessParticle Particle)
+void InitialMeshOrientation_Simulate(inout FStatelessParticle Particle)
 {
 	const float3 InitialRotation = InitialMeshOrientation_Rotation + (RandomFloat3(Particle) * InitialMeshOrientation_RandomRangeScale);
 	Particle.MeshOrientation			= RotatorToQuat(InitialRotation);
 	Particle.PreviousMeshOrientation	= Particle.MeshOrientation;
 }
-	
-void InitialMeshOrientation_Simulate(in FStatelessModule_InitialMeshOrientation ModuleData, inout FStatelessParticle Particle)
-{
-}
diff --git a/Engine/Plugins/FX/Niagara/Shaders/Private/Stateless/Modules/NiagaraStatelessModule_InitializeParticle.ush b/Engine/Plugins/FX/Niagara/Shaders/Private/Stateless/Modules/NiagaraStatelessModule_InitializeParticle.ush
--- a/Engine/Plugins/FX/Niagara/Shaders/Private/Stateless/Modules/NiagaraStatelessModule_InitializeParticle.ush
+++ b/Engine/Plugins/FX/Niagara/Shaders/Private/Stateless/Modules/NiagaraStatelessModule_InitializeParticle.ush
@@ -20,16 +20,12 @@ float3	InitializeParticle_MeshScaleBias;
 float	InitializeParticle_RibbonWidthScale;
 float	InitializeParticle_RibbonWidthBias;
 	
-struct FStatelessModule_InitializeParticle
-{
-};
-
-void InitializeParticle_Initialize(inout FStatelessModule_InitializeParticle ModuleData, inout FStatelessParticle Particle)
+void InitializeParticle_Simulate(inout FStatelessParticle Particle)
 {  
 	const bool bUniformSpriteSize = (InitializeParticle_ModuleFlags & EInitializeParticleModuleFlag_UniformSpriteSize) != 0;
 	const bool bUniformMeshScale  = (InitializeParticle_ModuleFlags & EInitializeParticleModuleFlag_UniformMeshScale) != 0;
 
-	FStatelessDistributionFloat3 InitialPosition;
+	FStatelessDistributionSamplerFloat3 InitialPosition;
 	InitialPosition.Init(Particle, InitializeParticle_InitialPosition);
 		
 	Particle.UniqueID		= Particle.UniqueIndex;
@@ -46,7 +42,3 @@ void InitializeParticle_Initialize(inout FStatelessModule_InitializeParticle Mod
 	Particle.PreviousScale			= Particle.Scale;
 	Particle.PreviousSpriteRotation	= Particle.SpriteRotation;
 }
-
-void InitializeParticle_Simulate(in FStatelessModule_InitializeParticle ModuleData, inout FStatelessParticle Particle)
-{  
-}
diff --git a/Engine/Plugins/FX/Niagara/Shaders/Private/Stateless/Modules/NiagaraStatelessModule_MeshIndex.ush b/Engine/Plugins/FX/Niagara/Shaders/Private/Stateless/Modules/NiagaraStatelessModule_MeshIndex.ush
--- a/Engine/Plugins/FX/Niagara/Shaders/Private/Stateless/Modules/NiagaraStatelessModule_MeshIndex.ush
+++ b/Engine/Plugins/FX/Niagara/Shaders/Private/Stateless/Modules/NiagaraStatelessModule_MeshIndex.ush
@@ -6,15 +6,7 @@ int MeshIndex_Index;
 int MeshIndex_TableOffset;
 int MeshIndex_TableNumElements;	// -1 to avoid added instructions
 
-struct FStatelessModule_MeshIndex
-{
-};
-
-void MeshIndex_Initialize(inout FStatelessModule_MeshIndex ModuleData, inout FStatelessParticle Particle)
-{
-}
-	
-void MeshIndex_Simulate(in FStatelessModule_MeshIndex ModuleData, inout FStatelessParticle Particle)
+void MeshIndex_Simulate(inout FStatelessParticle Particle)
 {
 	const bool bIsParameter = (MeshIndex_Index & 0x80000000) != 0;
 	int MeshIndex = MeshIndex_Index & ~0x80000000u;
diff --git a/Engine/Plugins/FX/Niagara/Shaders/Private/Stateless/Modules/NiagaraStatelessModule_MeshRotationRate.ush b/Engine/Plugins/FX/Niagara/Shaders/Private/Stateless/Modules/NiagaraStatelessModule_MeshRotationRate.ush
--- a/Engine/Plugins/FX/Niagara/Shaders/Private/Stateless/Modules/NiagaraStatelessModule_MeshRotationRate.ush
+++ b/Engine/Plugins/FX/Niagara/Shaders/Private/Stateless/Modules/NiagaraStatelessModule_MeshRotationRate.ush
@@ -2,22 +2,19 @@
 
 #include "../NiagaraStatelessCommon.ush"
 
-float3 MeshRotationRate_Scale;
-float3 MeshRotationRate_Bias;
+int	MeshRotationRate_ModuleEnabled;
+float3	MeshRotationRate_Scale;
+float3	MeshRotationRate_Bias;
 
-struct FStatelessModule_MeshRotationRate
+void MeshRotationRate_Simulate(inout FStatelessParticle Particle)
 {
-	float3 RotationRate;
-};
+	if ( MeshRotationRate_ModuleEnabled == 0 )
+	{
+		return;
+	}
+	const float3 RotationRate = RandomScaleBiasFloat(Particle, MeshRotationRate_Scale, MeshRotationRate_Bias);
 
-void MeshRotationRate_Initialize(inout FStatelessModule_MeshRotationRate ModuleData, inout FStatelessParticle Particle)
-{
-	ModuleData.RotationRate = RandomScaleBiasFloat(Particle, MeshRotationRate_Scale, MeshRotationRate_Bias);
-}
-	
-void MeshRotationRate_Simulate(in FStatelessModule_MeshRotationRate ModuleData, inout FStatelessParticle Particle)
-{
-	Particle.MeshOrientation			= MultiplyQuat(Particle.MeshOrientation,			RotatorToQuat(Particle.Age * ModuleData.RotationRate));
-	Particle.PreviousMeshOrientation	= MultiplyQuat(Particle.PreviousMeshOrientation,	RotatorToQuat(Particle.PreviousAge * ModuleData.RotationRate));
+	Particle.MeshOrientation			= MultiplyQuat(Particle.MeshOrientation,			RotatorToQuat(Particle.Age * RotationRate));
+	Particle.PreviousMeshOrientation	= MultiplyQuat(Particle.PreviousMeshOrientation,	RotatorToQuat(Particle.PreviousAge * RotationRate));
 }
 
diff --git a/Engine/Plugins/FX/Niagara/Shaders/Private/Stateless/Modules/NiagaraStatelessModule_RotateAroundPoint.ush b/Engine/Plugins/FX/Niagara/Shaders/Private/Stateless/Modules/NiagaraStatelessModule_RotateAroundPoint.ush
--- a/Engine/Plugins/FX/Niagara/Shaders/Private/Stateless/Modules/NiagaraStatelessModule_RotateAroundPoint.ush
+++ b/Engine/Plugins/FX/Niagara/Shaders/Private/Stateless/Modules/NiagaraStatelessModule_RotateAroundPoint.ush
@@ -26,15 +26,13 @@ void RotateAroundPoint_IntegratePosition(in FStatelessModule_RotateAroundPoint M
 	Position += cos(ModuleData.InitialPhase + Age * ModuleData.Rate) * ModuleData.Radius * AxisY;
 }
 
-void RotateAroundPoint_Initialize(inout FStatelessModule_RotateAroundPoint ModuleData, inout FStatelessParticle Particle)
+void RotateAroundPoint_Simulate(inout FStatelessParticle Particle)
 {
+	FStatelessModule_RotateAroundPoint ModuleData;
 	ModuleData.Rate			= RandomScaleBiasFloat(Particle, RotateAroundPoint_RateScale, RotateAroundPoint_RateBias);
 	ModuleData.Radius		= RandomScaleBiasFloat(Particle, RotateAroundPoint_RadiusScale, RotateAroundPoint_RadiusBias);
 	ModuleData.InitialPhase	= RandomScaleBiasFloat(Particle, RotateAroundPoint_InitialPhaseScale, RotateAroundPoint_InitialPhaseBias);
-}
-	
-void RotateAroundPoint_Simulate(in FStatelessModule_RotateAroundPoint ModuleData, inout FStatelessParticle Particle)
-{
+		
 	RotateAroundPoint_IntegratePosition(ModuleData, Particle.NormalizedAge, Particle.Position);
 	RotateAroundPoint_IntegratePosition(ModuleData, Particle.PreviousNormalizedAge, Particle.PreviousPosition);
 }
diff --git a/Engine/Plugins/FX/Niagara/Shaders/Private/Stateless/Modules/NiagaraStatelessModule_ScaleColor.ush b/Engine/Plugins/FX/Niagara/Shaders/Private/Stateless/Modules/NiagaraStatelessModule_ScaleColor.ush
--- a/Engine/Plugins/FX/Niagara/Shaders/Private/Stateless/Modules/NiagaraStatelessModule_ScaleColor.ush
+++ b/Engine/Plugins/FX/Niagara/Shaders/Private/Stateless/Modules/NiagaraStatelessModule_ScaleColor.ush
@@ -4,20 +4,12 @@
 
 uint3	ScaleColor_Distribution;
 
-struct FStatelessModule_ScaleColor
+void ScaleColor_Simulate(inout FStatelessParticle Particle)
 {
-	FStatelessDistributionFloat4	ScaleDistribution;
-};
-
-void ScaleColor_Initialize(inout FStatelessModule_ScaleColor ModuleData, inout FStatelessParticle Particle)
-{
-	ModuleData.ScaleDistribution.Init(Particle, ScaleColor_Distribution);
-}
-	
-void ScaleColor_Simulate(in FStatelessModule_ScaleColor ModuleData, inout FStatelessParticle Particle)
-{
-	if (ModuleData.ScaleDistribution.IsValid())
+	FStatelessDistributionSamplerFloat4 ScaleDistribution;
+	ScaleDistribution.Init(Particle, ScaleColor_Distribution);
+	if (ScaleDistribution.IsValid())
 	{
-		Particle.Color *= ModuleData.ScaleDistribution.GetValue(Particle.NormalizedAge);
+		Particle.Color *= ScaleDistribution.GetValue(Particle.NormalizedAge);
 	}
 }
diff --git a/Engine/Plugins/FX/Niagara/Shaders/Private/Stateless/Modules/NiagaraStatelessModule_ScaleMeshSize.ush b/Engine/Plugins/FX/Niagara/Shaders/Private/Stateless/Modules/NiagaraStatelessModule_ScaleMeshSize.ush
--- a/Engine/Plugins/FX/Niagara/Shaders/Private/Stateless/Modules/NiagaraStatelessModule_ScaleMeshSize.ush
+++ b/Engine/Plugins/FX/Niagara/Shaders/Private/Stateless/Modules/NiagaraStatelessModule_ScaleMeshSize.ush
@@ -6,23 +6,16 @@ uint3	ScaleMeshSize_Distribution;
 float3	ScaleMeshSize_CurveScale;
 int		ScaleMeshSize_CurveScaleOffset;
 	
-struct FStatelessModule_ScaleMeshSize
+void ScaleMeshSize_Simulate(inout FStatelessParticle Particle)
 {
-	FStatelessDistributionFloat3	MeshSizeDistribution;
-};
+	FStatelessDistributionSamplerFloat3 MeshSizeDistribution;
+	MeshSizeDistribution.Init(Particle, ScaleMeshSize_Distribution);
 
-void ScaleMeshSize_Initialize(inout FStatelessModule_ScaleMeshSize ModuleData, inout FStatelessParticle Particle)
-{
-	ModuleData.MeshSizeDistribution.Init(Particle, ScaleMeshSize_Distribution);
-}
-	
-void ScaleMeshSize_Simulate(in FStatelessModule_ScaleMeshSize ModuleData, inout FStatelessParticle Particle)
-{
-	if (ModuleData.MeshSizeDistribution.IsValid())
+	if (MeshSizeDistribution.IsValid())
 	{
 		const float3 ScaleFactor = GetParameterBufferFloat3(ScaleMeshSize_CurveScaleOffset, 0, ScaleMeshSize_CurveScale);
 			
-		Particle.Scale			*= ModuleData.MeshSizeDistribution.GetValue(Particle.NormalizedAge) * ScaleFactor;
-		Particle.PreviousScale	*= ModuleData.MeshSizeDistribution.GetValue(Particle.PreviousNormalizedAge) * ScaleFactor;
+		Particle.Scale			*= MeshSizeDistribution.GetValue(Particle.NormalizedAge) * ScaleFactor;
+		Particle.PreviousScale	*= MeshSizeDistribution.GetValue(Particle.PreviousNormalizedAge) * ScaleFactor;
 	}
 }
diff --git a/Engine/Plugins/FX/Niagara/Shaders/Private/Stateless/Modules/NiagaraStatelessModule_ScaleMeshSizeBySpeed.ush b/Engine/Plugins/FX/Niagara/Shaders/Private/Stateless/Modules/NiagaraStatelessModule_ScaleMeshSizeBySpeed.ush
--- a/Engine/Plugins/FX/Niagara/Shaders/Private/Stateless/Modules/NiagaraStatelessModule_ScaleMeshSizeBySpeed.ush
+++ b/Engine/Plugins/FX/Niagara/Shaders/Private/Stateless/Modules/NiagaraStatelessModule_ScaleMeshSizeBySpeed.ush
@@ -5,15 +5,7 @@
 uint2	ScaleMeshSizeBySpeed_ScaleDistribution;
 float	ScaleMeshSizeBySpeed_VelocityNorm;
 
-struct FStatelessModule_ScaleMeshSizeBySpeed
-{
-};
-
-void ScaleMeshSizeBySpeed_Initialize(inout FStatelessModule_ScaleMeshSizeBySpeed ModuleData, inout FStatelessParticle Particle)
-{
-}
-	
-void ScaleMeshSizeBySpeed_Simulate(in FStatelessModule_ScaleMeshSizeBySpeed ModuleData, inout FStatelessParticle Particle)
+void ScaleMeshSizeBySpeed_Simulate(inout FStatelessParticle Particle)
 {
 	if (ScaleMeshSizeBySpeed_ScaleDistribution.y > 0)
 	{
diff --git a/Engine/Plugins/FX/Niagara/Shaders/Private/Stateless/Modules/NiagaraStatelessModule_ScaleSpriteSize.ush b/Engine/Plugins/FX/Niagara/Shaders/Private/Stateless/Modules/NiagaraStatelessModule_ScaleSpriteSize.ush
--- a/Engine/Plugins/FX/Niagara/Shaders/Private/Stateless/Modules/NiagaraStatelessModule_ScaleSpriteSize.ush
+++ b/Engine/Plugins/FX/Niagara/Shaders/Private/Stateless/Modules/NiagaraStatelessModule_ScaleSpriteSize.ush
@@ -4,25 +4,17 @@
 
 uint3	ScaleSpriteSize_Distribution;
 float2	ScaleSpriteSize_CurveScale;
-int		ScaleSpriteSize_CurveScaleOffset;
 
-struct FStatelessModule_ScaleSpriteSize
+void ScaleSpriteSize_Simulate(inout FStatelessParticle Particle)
 {
-	FStatelessDistributionFloat2	ScaleDistribution;
-};
-	
-void ScaleSpriteSize_Initialize(inout FStatelessModule_ScaleSpriteSize ModuleData, inout FStatelessParticle Particle)
-{
-	ModuleData.ScaleDistribution.Init(Particle, ScaleSpriteSize_Distribution);
-}
-	
-void ScaleSpriteSize_Simulate(in FStatelessModule_ScaleSpriteSize ModuleData, inout FStatelessParticle Particle)
-{
-	if (ModuleData.ScaleDistribution.IsValid())
+	FStatelessDistributionSamplerFloat2 ScaleDistribution;
+	ScaleDistribution.Init(Particle, ScaleSpriteSize_Distribution);
+
+	if (ScaleDistribution.IsValid())
 	{
-		const float2 ScaleFactor = GetParameterBufferFloat2(ScaleSpriteSize_CurveScaleOffset, 0, ScaleSpriteSize_CurveScale);
+		const float2 ScaleFactor = ScaleSpriteSize_CurveScale;
 
-		Particle.SpriteSize			*= ModuleData.ScaleDistribution.GetValue(Particle.NormalizedAge) * ScaleFactor;
-		Particle.PreviousSpriteSize	*= ModuleData.ScaleDistribution.GetValue(Particle.PreviousNormalizedAge) * ScaleFactor;
+		Particle.SpriteSize			*= ScaleDistribution.GetValue(Particle.NormalizedAge) * ScaleFactor;
+		Particle.PreviousSpriteSize	*= ScaleDistribution.GetValue(Particle.PreviousNormalizedAge) * ScaleFactor;
 	}
 }
diff --git a/Engine/Plugins/FX/Niagara/Shaders/Private/Stateless/Modules/NiagaraStatelessModule_ScaleSpriteSizeBySpeed.ush b/Engine/Plugins/FX/Niagara/Shaders/Private/Stateless/Modules/NiagaraStatelessModule_ScaleSpriteSizeBySpeed.ush
--- a/Engine/Plugins/FX/Niagara/Shaders/Private/Stateless/Modules/NiagaraStatelessModule_ScaleSpriteSizeBySpeed.ush
+++ b/Engine/Plugins/FX/Niagara/Shaders/Private/Stateless/Modules/NiagaraStatelessModule_ScaleSpriteSizeBySpeed.ush
@@ -5,15 +5,7 @@
 uint2	ScaleSpriteSizeBySpeed_ScaleDistribution;
 float	ScaleSpriteSizeBySpeed_VelocityNorm;
 
-struct FStatelessModule_ScaleSpriteSizeBySpeed
-{
-};
-
-void ScaleSpriteSizeBySpeed_Initialize(inout FStatelessModule_ScaleSpriteSizeBySpeed ModuleData, inout FStatelessParticle Particle)
-{
-}
-	
-void ScaleSpriteSizeBySpeed_Simulate(in FStatelessModule_ScaleSpriteSizeBySpeed ModuleData, inout FStatelessParticle Particle)
+void ScaleSpriteSizeBySpeed_Simulate(inout FStatelessParticle Particle)
 {
 	if (ScaleSpriteSizeBySpeed_ScaleDistribution.y > 0)
 	{
diff --git a/Engine/Plugins/FX/Niagara/Shaders/Private/Stateless/Modules/NiagaraStatelessModule_ShapeLocation.ush b/Engine/Plugins/FX/Niagara/Shaders/Private/Stateless/Modules/NiagaraStatelessModule_ShapeLocation.ush
--- a/Engine/Plugins/FX/Niagara/Shaders/Private/Stateless/Modules/NiagaraStatelessModule_ShapeLocation.ush
+++ b/Engine/Plugins/FX/Niagara/Shaders/Private/Stateless/Modules/NiagaraStatelessModule_ShapeLocation.ush
@@ -6,10 +6,6 @@ uint4	ShapeLocation_Mode;
 float4	ShapeLocation_Parameters0;
 float4	ShapeLocation_Parameters1;
 
-struct FStatelessModule_ShapeLocation
-{
-};
-
 float3 ShapeLocation_GetLocation(inout FStatelessParticle Particle)
 {
 	// ENSM_ShapePrimitive::Box | ENSM_ShapePrimitive::Plane
@@ -80,13 +76,9 @@ float3 ShapeLocation_GetLocation(inout FStatelessParticle Particle)
 	}
 }
 
-void ShapeLocation_Initialize(inout FStatelessModule_ShapeLocation ModuleData, inout FStatelessParticle Particle)
+void ShapeLocation_Simulate(inout FStatelessParticle Particle)
 {
 	float3 Location = ShapeLocation_GetLocation(Particle);
 	Particle.Position += Location;
 	Particle.PreviousPosition += Location;
 }
-	
-void ShapeLocation_Simulate(in FStatelessModule_ShapeLocation ModuleData, inout FStatelessParticle Particle)
-{
-}
diff --git a/Engine/Plugins/FX/Niagara/Shaders/Private/Stateless/Modules/NiagaraStatelessModule_SolveVelocitiesAndForces.ush b/Engine/Plugins/FX/Niagara/Shaders/Private/Stateless/Modules/NiagaraStatelessModule_SolveVelocitiesAndForces.ush
--- a/Engine/Plugins/FX/Niagara/Shaders/Private/Stateless/Modules/NiagaraStatelessModule_SolveVelocitiesAndForces.ush
+++ b/Engine/Plugins/FX/Niagara/Shaders/Private/Stateless/Modules/NiagaraStatelessModule_SolveVelocitiesAndForces.ush
@@ -55,8 +55,10 @@ void SolveVelocitiesAndForces_IntegratePosition(in FStatelessModule_SolveVelocit
 	Position += (Age - LambdaAge) * (ModuleData.Acceleration / LambdaDragMass);
 }
 	
-void SolveVelocitiesAndForces_Initialize(inout FStatelessModule_SolveVelocitiesAndForces ModuleData, inout FStatelessParticle Particle)
+void SolveVelocitiesAndForces_Simulate(inout FStatelessParticle Particle)
 {
+	// Initialize shared parameters
+	FStatelessModule_SolveVelocitiesAndForces ModuleData;
 	ModuleData.Mass = RandomScaleBiasFloat(Particle, SolveVelocitiesAndForces_MassScale, SolveVelocitiesAndForces_MassBias);
 	ModuleData.Drag = RandomScaleBiasFloat(Particle, SolveVelocitiesAndForces_DragScale, SolveVelocitiesAndForces_DragBias);
 	ModuleData.Velocity = RandomScaleBiasFloat(Particle, SolveVelocitiesAndForces_VelocityScale, SolveVelocitiesAndForces_VelocityBias);
@@ -86,8 +88,7 @@ void SolveVelocitiesAndForces_Initialize(inout FStatelessModule_SolveVelocitiesA
 	{
 		const float3 FallbackDir = RandomUnitFloat3(Particle);
 		const float3 Delta = Particle.Position - SolveVelocitiesAndForces_PointOrigin;
-		const float  DeltaL2 = length2(Delta);
-		const float3 Dir = DeltaL2 > UE_KINDA_SMALL_NUMBER ? Delta * rsqrt(DeltaL2) : FallbackDir;
+		const float3 Dir = SafeNormalize(Delta, FallbackDir);
 			
 		const float VelocityScale = RandomScaleBiasFloat(Particle, SolveVelocitiesAndForces_PointVelocityScale, SolveVelocitiesAndForces_PointVelocityBias);
 			
@@ -156,10 +157,8 @@ void SolveVelocitiesAndForces_Initialize(inout FStatelessModule_SolveVelocitiesA
 
 		ModuleData.Velocity += NoiseValue.xyz * SolveVelocitiesAndForces_NoiseAmplitude;
 	}
-}
 	
-void SolveVelocitiesAndForces_Simulate(in FStatelessModule_SolveVelocitiesAndForces ModuleData, inout FStatelessParticle Particle)
-{
+	// Simulate
 	SolveVelocitiesAndForces_IntegratePosition(ModuleData, Particle.Age, Particle.Position);
 	SolveVelocitiesAndForces_IntegratePosition(ModuleData, Particle.PreviousAge, Particle.PreviousPosition);
 		
diff --git a/Engine/Plugins/FX/Niagara/Shaders/Private/Stateless/Modules/NiagaraStatelessModule_SpriteFacingAndAlignment.ush b/Engine/Plugins/FX/Niagara/Shaders/Private/Stateless/Modules/NiagaraStatelessModule_SpriteFacingAndAlignment.ush
--- a/Engine/Plugins/FX/Niagara/Shaders/Private/Stateless/Modules/NiagaraStatelessModule_SpriteFacingAndAlignment.ush
+++ b/Engine/Plugins/FX/Niagara/Shaders/Private/Stateless/Modules/NiagaraStatelessModule_SpriteFacingAndAlignment.ush
@@ -5,11 +5,7 @@
 float3	SpriteFacingAndAlignment_SpriteFacing;
 float3	SpriteFacingAndAlignment_SpriteAlignment;
 
-struct FStatelessModule_SpriteFacingAndAlignment
-{
-};
-
-void SpriteFacingAndAlignment_Initialize(inout FStatelessModule_SpriteFacingAndAlignment ModuleData, inout FStatelessParticle Particle)
+void SpriteFacingAndAlignment_Simulate(inout FStatelessParticle Particle)
 {
 	Particle.SpriteFacing = SpriteFacingAndAlignment_SpriteFacing;
 	Particle.SpriteAlignment = SpriteFacingAndAlignment_SpriteAlignment;
@@ -17,7 +13,3 @@ void SpriteFacingAndAlignment_Initialize(inout FStatelessModule_SpriteFacingAndA
 	Particle.PreviousSpriteFacing = SpriteFacingAndAlignment_SpriteFacing;
 	Particle.PreviousSpriteAlignment = SpriteFacingAndAlignment_SpriteAlignment;
 }
-	
-void SpriteFacingAndAlignment_Simulate(in FStatelessModule_SpriteFacingAndAlignment ModuleData, inout FStatelessParticle Particle)
-{
-}
diff --git a/Engine/Plugins/FX/Niagara/Shaders/Private/Stateless/Modules/NiagaraStatelessModule_SpriteRotationRate.ush b/Engine/Plugins/FX/Niagara/Shaders/Private/Stateless/Modules/NiagaraStatelessModule_SpriteRotationRate.ush
--- a/Engine/Plugins/FX/Niagara/Shaders/Private/Stateless/Modules/NiagaraStatelessModule_SpriteRotationRate.ush
+++ b/Engine/Plugins/FX/Niagara/Shaders/Private/Stateless/Modules/NiagaraStatelessModule_SpriteRotationRate.ush
@@ -5,18 +5,10 @@
 float		SpriteRotationRate_Scale;
 float		SpriteRotationRate_Bias;
 
-struct FStatelessModule_SpriteRotationRate
+void SpriteRotationRate_Simulate(inout FStatelessParticle Particle)
 {
-	float RotationRate;
-};
+	const float RotationRate = RandomScaleBiasFloat(Particle, SpriteRotationRate_Scale, SpriteRotationRate_Bias);
 
-void SpriteRotationRate_Initialize(inout FStatelessModule_SpriteRotationRate ModuleData, inout FStatelessParticle Particle)
-{
-	ModuleData.RotationRate = RandomScaleBiasFloat(Particle, SpriteRotationRate_Scale, SpriteRotationRate_Bias);
-}
-	
-void SpriteRotationRate_Simulate(in FStatelessModule_SpriteRotationRate ModuleData, inout FStatelessParticle Particle)
-{
-	Particle.SpriteRotation += Particle.Age * ModuleData.RotationRate;
-	Particle.PreviousSpriteRotation += Particle.PreviousAge * ModuleData.RotationRate;
+	Particle.SpriteRotation += Particle.Age * RotationRate;
+	Particle.PreviousSpriteRotation += Particle.PreviousAge * RotationRate;
 }
diff --git a/Engine/Plugins/FX/Niagara/Shaders/Private/Stateless/Modules/NiagaraStatelessModule_SubUVAnimation.ush b/Engine/Plugins/FX/Niagara/Shaders/Private/Stateless/Modules/NiagaraStatelessModule_SubUVAnimation.ush
--- a/Engine/Plugins/FX/Niagara/Shaders/Private/Stateless/Modules/NiagaraStatelessModule_SubUVAnimation.ush
+++ b/Engine/Plugins/FX/Niagara/Shaders/Private/Stateless/Modules/NiagaraStatelessModule_SubUVAnimation.ush
@@ -11,21 +11,11 @@ float	SubUVAnimation_AnimFrameStart;
 float	SubUVAnimation_AnimFrameRange;
 float	SubUVAnimation_RateScale;
 
-struct FStatelessModule_SubUVAnimation
+void SubUVAnimation_Simulate(inout FStatelessParticle Particle)
 {
-	float InitialFrame;
-};
+	const uint SeedOffset = uint(Particle.Age * SubUVAnimation_InitialFrameRateChange);
+	float Frame = RandomSeedOffsetFloat(Particle, SeedOffset) * SubUVAnimation_InitialFrameScale + SubUVAnimation_InitialFrameBias;
 
-void SubUVAnimation_Initialize(inout FStatelessModule_SubUVAnimation ModuleData, inout FStatelessParticle Particle)
-{
-	const int Seed = int(floor(Particle.Age * SubUVAnimation_InitialFrameRateChange) * 0xffff) ^ Particle.RandomSeed.x ^ Particle.UniqueIndex;
-	ModuleData.InitialFrame = RandBBSfloat(float(Seed) / 0xffff) * SubUVAnimation_InitialFrameScale + SubUVAnimation_InitialFrameBias;
-	//ModuleData.InitialFrame = RandomScaleBiasFloat(Particle, SubUVAnimation_InitialFrameScale, SubUVAnimation_InitialFrameBias);
-}
-	
-void SubUVAnimation_Simulate(in FStatelessModule_SubUVAnimation ModuleData, inout FStatelessParticle Particle)
-{
-	float Frame = ModuleData.InitialFrame;
 	if ( SubUVAnimation_Mode == 0 )
 	{
 		const float Interp = Particle.Age * SubUVAnimation_RateScale;
diff --git a/Engine/Plugins/FX/Niagara/Shaders/Private/Stateless/NiagaraStatelessCommon.ush b/Engine/Plugins/FX/Niagara/Shaders/Private/Stateless/NiagaraStatelessCommon.ush
--- a/Engine/Plugins/FX/Niagara/Shaders/Private/Stateless/NiagaraStatelessCommon.ush
+++ b/Engine/Plugins/FX/Niagara/Shaders/Private/Stateless/NiagaraStatelessCommon.ush
@@ -122,6 +122,18 @@ float4 RandomScaleBiasFloat(inout FStatelessParticle Particle, float4 Scale, flo
 float2 RandomUnitFloat2(inout FStatelessParticle Particle) { return SafeNormalize(RandomFloat2(Particle) - 0.5f); }
 float3 RandomUnitFloat3(inout FStatelessParticle Particle) { return SafeNormalize(RandomFloat3(Particle) - 0.5f); }
 
+float4 RandomSeedOffsetFloat4(in FStatelessParticle Particle, uint SeedOffset)
+{
+	uint4 Seed = Particle.RandomSeed;
+	Seed.x += SeedOffset;
+		
+	uint4 V = Rand4DPCG32(Seed);
+	return float4((V >> 8) & 0x00ffffff) / 16777216.0; // 0x01000000 == 16777216
+}
+float  RandomSeedOffsetFloat(inout FStatelessParticle Particle, uint SeedOffset)  { return RandomSeedOffsetFloat4(Particle, SeedOffset).x; }
+float2 RandomSeedOffsetFloat2(inout FStatelessParticle Particle, uint SeedOffset) { return RandomSeedOffsetFloat4(Particle, SeedOffset).xy; }
+float3 RandomSeedOffsetFloat3(inout FStatelessParticle Particle, uint SeedOffset) { return RandomSeedOffsetFloat4(Particle, SeedOffset).xyz; }
+
 float  GetStaticFloat (int BaseOffset, int Element)	{ return Common_StaticFloatBuffer[BaseOffset + Element]; }
 float2 GetStaticFloat2(int BaseOffset, int Element)	{ return float2(Common_StaticFloatBuffer[BaseOffset + Element*2], Common_StaticFloatBuffer[BaseOffset + 1 + Element*2]); }
 float3 GetStaticFloat3(int BaseOffset, int Element)	{ return float3(Common_StaticFloatBuffer[BaseOffset + Element*3], Common_StaticFloatBuffer[BaseOffset + 1 + Element*3], Common_StaticFloatBuffer[BaseOffset + 2 + Element*3]); }
diff --git a/Engine/Plugins/FX/Niagara/Shaders/Private/Stateless/NiagaraStatelessDistribution.ush b/Engine/Plugins/FX/Niagara/Shaders/Private/Stateless/NiagaraStatelessDistribution.ush
--- a/Engine/Plugins/FX/Niagara/Shaders/Private/Stateless/NiagaraStatelessDistribution.ush
+++ b/Engine/Plugins/FX/Niagara/Shaders/Private/Stateless/NiagaraStatelessDistribution.ush
@@ -5,210 +5,49 @@
 #define StatelessDistributionFlag_Uniform	0x00000002
 #define StatelessDistributionFlag_Binding	0x00000004
 
-struct FStatelessDistributionFloat
-{
-	bool	bRandomDistribution;
-	bool	bUniformValue;
-	bool	bBindingValue;
-	uint	TableOffset;
-	float	TableLength;
-	float	RandomOffset;
-
-	void Init(inout FStatelessParticle Particle, uint3 Parameters)
-	{
-		bRandomDistribution = (Parameters.x & StatelessDistributionFlag_Random) != 0;
-		bUniformValue		= (Parameters.x & StatelessDistributionFlag_Uniform) != 0;
-		bBindingValue		= (Parameters.x & StatelessDistributionFlag_Binding) != 0;
-		TableOffset			= Parameters.y;
-		TableLength			= asfloat(Parameters.z);
-		RandomOffset		= RandomFloat(Particle);
-	}
-
-	bool IsValid()
-	{
-		return TableLength > 0.0f;
-	}
-	
-	float GetValue(float NormalizedAge)
-	{
-		if ( bBindingValue )
-		{
-			return GetParameterBufferFloat(TableOffset, 0);
-		}
-		else if ( bRandomDistribution )
-		{
-			const float	Value0 = GetStaticFloat(TableOffset, 0);
-			const float	Value1 = GetStaticFloat(TableOffset, 1);
-			return lerp(Value0, Value1, RandomOffset);
-		}
-		else
-		{
-			const float Offset = NormalizedAge * TableLength;
-			const float	Value0 = GetStaticFloat(TableOffset, floor(Offset));
-			const float	Value1 = GetStaticFloat(TableOffset, ceil(Offset));
-			return lerp(Value0, Value1, frac(Offset));
-		}
-	}
-};
-
-struct FStatelessDistributionFloat2
-{
-	bool	bRandomDistribution;
-	bool	bUniformValue;
-	bool	bBindingValue;
-	uint	TableOffset;
-	float	TableLength;
-	float2	RandomOffset;
-
-	void Init(inout FStatelessParticle Particle, uint3 Parameters)
-	{
-		bRandomDistribution = (Parameters.x & StatelessDistributionFlag_Random) != 0;
-		bUniformValue		= (Parameters.x & StatelessDistributionFlag_Uniform) != 0;
-		bBindingValue		= (Parameters.x & StatelessDistributionFlag_Binding) != 0;
-		TableOffset			= Parameters.y;
-		TableLength			= asfloat(Parameters.z);
-		RandomOffset		= 0;
-		if ( bRandomDistribution )
-		{
-			RandomOffset	= bUniformValue ? RandomFloat(Particle) : RandomFloat2(Particle);
-		}
-	}
-
-	bool IsValid()
-	{
-		return TableLength > 0.0f;
-	}
-	
-	float2 GetValue(float NormalizedAge)
-	{
-		if ( bBindingValue )
-		{
-			return GetParameterBufferFloat2(TableOffset, 0);
-		}
-		else if ( bRandomDistribution )
-		{
-			const float2 Value0 = GetStaticFloat2(TableOffset, 0);
-			const float2 Value1 = GetStaticFloat2(TableOffset, 1);
-			return float2(
-				lerp(Value0.x, Value1.x, RandomOffset.x),
-				lerp(Value0.y, Value1.y, RandomOffset.y)
-			);
-		}
-		else
-		{
-			const float  Offset = NormalizedAge * TableLength;
-			const float2 Value0 = GetStaticFloat2(TableOffset, floor(Offset));
-			const float2 Value1 = GetStaticFloat2(TableOffset, ceil(Offset));
-			return lerp(Value0, Value1, frac(Offset));
-		}
-	}
-};
-
-struct FStatelessDistributionFloat3
-{
-	bool	bRandomDistribution;
-	bool	bUniformValue;
-	bool	bBindingValue;
-	uint	TableOffset;
-	float	TableLength;
-	float3	RandomOffset;
-
-	void Init(inout FStatelessParticle Particle, uint3 Parameters)
-	{
-		bRandomDistribution = (Parameters.x & StatelessDistributionFlag_Random) != 0;
-		bUniformValue		= (Parameters.x & StatelessDistributionFlag_Uniform) != 0;
-		bBindingValue		= (Parameters.x & StatelessDistributionFlag_Binding) != 0;
-		TableOffset			= Parameters.y;
-		TableLength			= asfloat(Parameters.z);
-		RandomOffset		= 0;
-		if ( bRandomDistribution )
-		{
-			RandomOffset	= bUniformValue ? RandomFloat(Particle) : RandomFloat3(Particle);
-		}
-	}
-
-	bool IsValid()
-	{
-		return TableLength > 0.0f;
-	}
-	
-	float3 GetValue(float NormalizedAge)
-	{
-		if ( bBindingValue )
-		{
-			return GetParameterBufferFloat3(TableOffset, 0);
-		}
-		else if ( bRandomDistribution )
-		{
-			const float3 Value0 = GetStaticFloat3(TableOffset, 0);
-			const float3 Value1 = GetStaticFloat3(TableOffset, 1);
-			return float3(
-				lerp(Value0.x, Value1.x, RandomOffset.x),
-				lerp(Value0.y, Value1.y, RandomOffset.y),
-				lerp(Value0.z, Value1.z, RandomOffset.z)
-			);
-		}
-		else
-		{
-			const float  Offset = NormalizedAge * TableLength;
-			const float3 Value0 = GetStaticFloat3(TableOffset, floor(Offset));
-			const float3 Value1 = GetStaticFloat3(TableOffset, ceil(Offset));
-			return lerp(Value0, Value1, frac(Offset));
-		}
-	}
-};
-
-struct FStatelessDistributionFloat4
-{
-	bool	bRandomDistribution;
-	bool	bUniformValue;
-	bool	bBindingValue;
-	uint	TableOffset;
-	float	TableLength;
-	float4	RandomOffset;
-
-	void Init(inout FStatelessParticle Particle, uint3 Parameters)
-	{
-		bRandomDistribution = (Parameters.x & StatelessDistributionFlag_Random) != 0;
-		bUniformValue		= (Parameters.x & StatelessDistributionFlag_Uniform) != 0;
-		bBindingValue		= (Parameters.x & StatelessDistributionFlag_Binding) != 0;
-		TableOffset			= Parameters.y;
-		TableLength			= asfloat(Parameters.z);
-		RandomOffset		= 0;
-		if ( bRandomDistribution )
-		{
-			RandomOffset	= bUniformValue ? RandomFloat(Particle) : RandomFloat4(Particle);
-		}
-	}
-
-	bool IsValid()
-	{
-		return TableLength > 0.0f;
-	}
-	
-	float4 GetValue(float NormalizedAge)
-	{
-		if ( bBindingValue )
-		{
-			return GetParameterBufferFloat4(TableOffset, 0);
-		}
-		else if ( bRandomDistribution )
-		{
-			const float4 Value0 = GetStaticFloat4(TableOffset, 0);
-			const float4 Value1 = GetStaticFloat4(TableOffset, 1);
-			return float4(
-				lerp(Value0.x, Value1.x, RandomOffset.x),
-				lerp(Value0.y, Value1.y, RandomOffset.y),
-				lerp(Value0.z, Value1.z, RandomOffset.z),
-				lerp(Value0.w, Value1.w, RandomOffset.w)
-			);
-		}
-		else
-		{
-			const float  Offset = NormalizedAge * TableLength;
-			const float4 Value0 = GetStaticFloat4(TableOffset, floor(Offset));
-			const float4 Value1 = GetStaticFloat4(TableOffset, ceil(Offset));
-			return lerp(Value0, Value1, frac(Offset));
-		}
-	}
-};
+#define DefineStatelessDistributionSampler(HLSLTYPE, NAMETYPE) \
+	struct FStatelessDistributionSampler##NAMETYPE \
+	{ \
+		void Init(inout FStatelessParticle Particle, uint3 InParameters) \
+		{ \
+			Parameters = InParameters; \
+			RandomOffset = 0.0f; \
+			if ( (Parameters.x & StatelessDistributionFlag_Random) != 0) \
+			{ \
+				RandomOffset = RandomFloat(Particle); \
+			} \
+		} \
+		\
+		bool IsValid() { return Parameters.z > 0; } \
+		\
+		HLSLTYPE GetValue(float NormalizedAge) \
+		{ \
+			if ((Parameters.x & StatelessDistributionFlag_Binding) != 0) \
+			{ \
+				return GetParameterBuffer##NAMETYPE##(Parameters.y, 0); \
+			} \
+			else if ((Parameters.x & StatelessDistributionFlag_Random) != 0) \
+			{ \
+				const HLSLTYPE Value0 = GetStatic##NAMETYPE##(Parameters.y, 0); \
+				const HLSLTYPE Value1 = GetStatic##NAMETYPE##(Parameters.y, 1); \
+				return lerp(Value0, Value1, RandomOffset); \
+			} \
+			else \
+			{ \
+				const float Offset = NormalizedAge * float(Parameters.z); \
+				const HLSLTYPE Value0 = GetStatic##NAMETYPE##(Parameters.y, floor(Offset)); \
+				const HLSLTYPE Value1 = GetStatic##NAMETYPE##(Parameters.y, ceil(Offset)); \
+				return lerp(Value0, Value1, frac(Offset)); \
+			} \
+		} \
+		\
+		uint3	Parameters; \
+		float	RandomOffset; \
+	}
+
+DefineStatelessDistributionSampler(float,  Float);
+DefineStatelessDistributionSampler(float2, Float2);
+DefineStatelessDistributionSampler(float3, Float3);
+DefineStatelessDistributionSampler(float4, Float4);
+
+#undef DefineStatelessDistributionSampler
diff --git a/Engine/Plugins/FX/Niagara/Shaders/Private/Stateless/NiagaraStatelessSimulationDefault.usf b/Engine/Plugins/FX/Niagara/Shaders/Private/Stateless/NiagaraStatelessSimulationDefault.usf
--- a/Engine/Plugins/FX/Niagara/Shaders/Private/Stateless/NiagaraStatelessSimulationDefault.usf
+++ b/Engine/Plugins/FX/Niagara/Shaders/Private/Stateless/NiagaraStatelessSimulationDefault.usf
@@ -1,11 +1,15 @@
 // Copyright Epic Games, Inc. All Rights Reserved.
 
+// Defines potential output variables, usage in renderers will determine what we actually output
 #define PARTICLE_ATTRIBUTES_OUTPUT \
 	PARTICLE_ATTRIBUTE(int,		UniqueID) \
 	PARTICLE_ATTRIBUTE(float3,	Position) \
 	PARTICLE_ATTRIBUTE(float,	CameraOffset) \
 	PARTICLE_ATTRIBUTE(float4,	Color) \
 	PARTICLE_ATTRIBUTE(float4,	DynamicMaterialParameter0) \
+	PARTICLE_ATTRIBUTE(float4,	DynamicMaterialParameter1) \
+	PARTICLE_ATTRIBUTE(float4,	DynamicMaterialParameter2) \
+	PARTICLE_ATTRIBUTE(float4,	DynamicMaterialParameter3) \
 	PARTICLE_ATTRIBUTE(int,		MeshIndex) \
 	PARTICLE_ATTRIBUTE(float4,	MeshOrientation) \
 	PARTICLE_ATTRIBUTE(float,	RibbonWidth) \
@@ -27,24 +31,26 @@
 	PARTICLE_ATTRIBUTE(float,	PreviousSpriteRotation) \
 	PARTICLE_ATTRIBUTE(float3,	PreviousVelocity) \
 
+// Defines transient variables, i.e. ones we will not output
 #define PARTICLE_ATTRIBUTES_TRANSIENT \
 
+// Defines modules and order of execution
 #define PARTICLE_MODULES \
 	PARTICLE_MODULE(InitializeParticle) \
 	PARTICLE_MODULE(InitialMeshOrientation) \
 	PARTICLE_MODULE(ShapeLocation) \
+	PARTICLE_MODULE(SolveVelocitiesAndForces) \
 	PARTICLE_MODULE(CameraOffset) \
-	PARTICLE_MODULE(ScaleColor) \
-	PARTICLE_MODULE(ScaleSpriteSize) \
-	PARTICLE_MODULE(SpriteFacingAndAlignment) \
-	PARTICLE_MODULE(SpriteRotationRate) \
+	PARTICLE_MODULE(DynamicMaterialParameters) \
 	PARTICLE_MODULE(MeshIndex) \
 	PARTICLE_MODULE(MeshRotationRate) \
+	PARTICLE_MODULE(ScaleColor) \
+	PARTICLE_MODULE(ScaleSpriteSize) \
+	PARTICLE_MODULE(ScaleSpriteSizeBySpeed) \
 	PARTICLE_MODULE(ScaleMeshSize) \
-	PARTICLE_MODULE(SolveVelocitiesAndForces) \
 	PARTICLE_MODULE(ScaleMeshSizeBySpeed) \
-	PARTICLE_MODULE(ScaleSpriteSizeBySpeed) \
+	PARTICLE_MODULE(SpriteFacingAndAlignment) \
+	PARTICLE_MODULE(SpriteRotationRate) \
 	PARTICLE_MODULE(SubUVAnimation) \
-	PARTICLE_MODULE(DynamicMaterialParameters) \
 
 #include "NiagaraStatelessSimulationTemplate.ush"
diff --git a/Engine/Plugins/FX/Niagara/Shaders/Private/Stateless/NiagaraStatelessSimulationExample1.usf b/Engine/Plugins/FX/Niagara/Shaders/Private/Stateless/NiagaraStatelessSimulationExample1.usf
--- a/Engine/Plugins/FX/Niagara/Shaders/Private/Stateless/NiagaraStatelessSimulationExample1.usf
+++ b/Engine/Plugins/FX/Niagara/Shaders/Private/Stateless/NiagaraStatelessSimulationExample1.usf
@@ -1,5 +1,6 @@
 // Copyright Epic Games, Inc. All Rights Reserved.
 
+// Defines potential output variables, usage in renderers will determine what we actually output
 #define PARTICLE_ATTRIBUTES_OUTPUT \
 	PARTICLE_ATTRIBUTE(int,		UniqueID) \
 	PARTICLE_ATTRIBUTE(float3,	Position) \
@@ -16,9 +17,13 @@
 	PARTICLE_ATTRIBUTE(float,	PreviousSpriteRotation) \
 	PARTICLE_ATTRIBUTE(float3,	PreviousVelocity) \
 
+// Defines transient variables, i.e. ones we will not output
 #define PARTICLE_ATTRIBUTES_TRANSIENT \
 	PARTICLE_ATTRIBUTE(float,	CameraOffset) \
 	PARTICLE_ATTRIBUTE(float4,	DynamicMaterialParameter0) \
+	PARTICLE_ATTRIBUTE(float4,	DynamicMaterialParameter1) \
+	PARTICLE_ATTRIBUTE(float4,	DynamicMaterialParameter2) \
+	PARTICLE_ATTRIBUTE(float4,	DynamicMaterialParameter3) \
 	PARTICLE_ATTRIBUTE(int,		MeshIndex) \
 	PARTICLE_ATTRIBUTE(float4,	MeshOrientation) \
 	PARTICLE_ATTRIBUTE(float3,	SpriteFacing) \
@@ -29,6 +34,7 @@
 	PARTICLE_ATTRIBUTE(float3,	PreviousSpriteFacing) \
 	PARTICLE_ATTRIBUTE(float3,	PreviousSpriteAlignment) \
 
+// Defines modules and order of execution
 #define PARTICLE_MODULES \
 	PARTICLE_MODULE(InitializeParticle) \
 	PARTICLE_MODULE(ShapeLocation) \
diff --git a/Engine/Plugins/FX/Niagara/Shaders/Private/Stateless/NiagaraStatelessSimulationTemplate.ush b/Engine/Plugins/FX/Niagara/Shaders/Private/Stateless/NiagaraStatelessSimulationTemplate.ush
--- a/Engine/Plugins/FX/Niagara/Shaders/Private/Stateless/NiagaraStatelessSimulationTemplate.ush
+++ b/Engine/Plugins/FX/Niagara/Shaders/Private/Stateless/NiagaraStatelessSimulationTemplate.ush
@@ -30,13 +30,6 @@
 #include "Modules/NiagaraStatelessModule_SpriteRotationRate.ush"
 #include "Modules/NiagaraStatelessModule_SubUVAnimation.ush"
 
-struct FStatelessParticleSimuation
-{
-	#define PARTICLE_MODULE(NAME) FStatelessModule_##NAME	NAME##Module;
-		PARTICLE_MODULES
-	#undef PARTICLE_MODULE
-};
-
 /////////////////////////////////////////////////////////////////////////////////////////////////////////////////////////
 
 #define NIAGARA_MAX_GPU_SPAWN_INFOS 4
@@ -114,20 +107,8 @@ void StatelessMain(uint DispatchThreadId : SV_DispatchThreadID)
 	
 	Particle.bAlive = true;
 	
-	// Run Particle Initialization
-	FStatelessParticleSimuation ParticleSimulation;
-
-	#define PARTICLE_MODULE(NAME) ##NAME##_Initialize(ParticleSimulation.##NAME##Module, Particle);
-		PARTICLE_MODULES
-	#undef PARTICLE_MODULE
-
-	if ( !Particle.bAlive )
-	{
-		return;
-	}
-
 	// Run Particle Simulate
-	#define PARTICLE_MODULE(NAME) ##NAME##_Simulate(ParticleSimulation.##NAME##Module, Particle);
+	#define PARTICLE_MODULE(NAME) ##NAME##_Simulate(Particle);
 		PARTICLE_MODULES
 	#undef PARTICLE_MODULE
 
diff --git a/Engine/Plugins/FX/Niagara/Source/Niagara/Internal/Stateless/Modules/NiagaraStatelessModule_AccelerationForce.h b/Engine/Plugins/FX/Niagara/Source/Niagara/Internal/Stateless/Modules/NiagaraStatelessModule_AccelerationForce.h
--- a/Engine/Plugins/FX/Niagara/Source/Niagara/Internal/Stateless/Modules/NiagaraStatelessModule_AccelerationForce.h
+++ b/Engine/Plugins/FX/Niagara/Source/Niagara/Internal/Stateless/Modules/NiagaraStatelessModule_AccelerationForce.h
@@ -18,7 +18,7 @@ public:
 	UPROPERTY(EditAnywhere, Category = "Parameters", meta = (DisplayName = "Acceleration", DisableUniformDistribution, DisableBindingDistribution))
 	FNiagaraDistributionRangeVector3 AccelerationDistribution = FNiagaraDistributionRangeVector3(FVector3f::ZeroVector);
 
-	virtual void BuildEmitterData(FNiagaraStatelessEmitterDataBuildContext& BuildContext) const override
+	virtual void BuildEmitterData(const FNiagaraStatelessEmitterDataBuildContext& BuildContext) const override
 	{
 		if (!IsModuleEnabled())
 		{
diff --git a/Engine/Plugins/FX/Niagara/Source/Niagara/Internal/Stateless/Modules/NiagaraStatelessModule_AddVelocity.h b/Engine/Plugins/FX/Niagara/Source/Niagara/Internal/Stateless/Modules/NiagaraStatelessModule_AddVelocity.h
--- a/Engine/Plugins/FX/Niagara/Source/Niagara/Internal/Stateless/Modules/NiagaraStatelessModule_AddVelocity.h
+++ b/Engine/Plugins/FX/Niagara/Source/Niagara/Internal/Stateless/Modules/NiagaraStatelessModule_AddVelocity.h
@@ -50,7 +50,7 @@ public:
 	UPROPERTY(EditAnywhere, Category = "Parameters", meta = (EditCondition = "bSpeedFalloffFromConeAxisEnabled && VelocityType == ENSM_VelocityType::InCone", ClampMin = "0.0", ClampMax = "1.0", UIMin = "0.0", UIMax = "1.0"))
 	float SpeedFalloffFromConeAxis = 0.0f;
 
-	virtual void BuildEmitterData(FNiagaraStatelessEmitterDataBuildContext& BuildContext) const override;
+	virtual void BuildEmitterData(const FNiagaraStatelessEmitterDataBuildContext& BuildContext) const override;
 
 #if WITH_EDITOR
 	virtual bool CanDisableModule() const override { return true; }
diff --git a/Engine/Plugins/FX/Niagara/Source/Niagara/Internal/Stateless/Modules/NiagaraStatelessModule_CameraOffset.h b/Engine/Plugins/FX/Niagara/Source/Niagara/Internal/Stateless/Modules/NiagaraStatelessModule_CameraOffset.h
--- a/Engine/Plugins/FX/Niagara/Source/Niagara/Internal/Stateless/Modules/NiagaraStatelessModule_CameraOffset.h
+++ b/Engine/Plugins/FX/Niagara/Source/Niagara/Internal/Stateless/Modules/NiagaraStatelessModule_CameraOffset.h
@@ -5,6 +5,7 @@
 #include "Stateless/NiagaraStatelessModule.h"
 #include "Stateless/NiagaraStatelessEmitterDataBuildContext.h"
 #include "Stateless/NiagaraStatelessModuleShaderParameters.h"
+#include "Stateless/NiagaraStatelessParticleSimContext.h"
 
 #include "NiagaraStatelessModule_CameraOffset.generated.h"
 
@@ -16,7 +17,9 @@ public:
 
 	struct FModuleBuiltData
 	{
-		FUintVector3	DistributionParameters = FUintVector3::ZeroValue;
+		FUintVector3	DistributionParameters			= FUintVector3::ZeroValue;
+		int32			CameraVariableOffset			= INDEX_NONE;
+		int32			PreviousCameraVariableOffset	= INDEX_NONE;
 	};
 
 public:
@@ -25,10 +28,21 @@ public:
 	UPROPERTY(EditAnywhere, Category = "Parameters", meta = (DisplayName = "CameraOffset"))
 	FNiagaraDistributionFloat CameraOffsetDistribution = FNiagaraDistributionFloat(0.0f);
 
-	virtual void BuildEmitterData(FNiagaraStatelessEmitterDataBuildContext& BuildContext) const override
+	virtual void BuildEmitterData(const FNiagaraStatelessEmitterDataBuildContext& BuildContext) const override
 	{
-		FModuleBuiltData* BuiltData = BuildContext.AllocateBuiltData<FModuleBuiltData>();
-		BuiltData->DistributionParameters = BuildContext.AddDistribution(CameraOffsetDistribution, IsModuleEnabled());
+		const FNiagaraStatelessGlobals& StatelessGlobals = FNiagaraStatelessGlobals::Get();
+
+		FModuleBuiltData* BuiltData				= BuildContext.AllocateBuiltData<FModuleBuiltData>();
+		BuiltData->CameraVariableOffset			= BuildContext.FindParticleVariableIndex(StatelessGlobals.CameraOffsetVariable);
+		BuiltData->PreviousCameraVariableOffset	= BuildContext.FindParticleVariableIndex(StatelessGlobals.PreviousCameraOffsetVariable);
+
+		const bool bAttributesUsed = (BuiltData->CameraVariableOffset != INDEX_NONE || BuiltData->PreviousCameraVariableOffset != INDEX_NONE);
+		if (IsModuleEnabled() && bAttributesUsed)
+		{
+			BuiltData->DistributionParameters = BuildContext.AddDistribution(CameraOffsetDistribution);
+
+			BuildContext.AddParticleSimulationExecSimulate(&UNiagaraStatelessModule_CameraOffset::ParticleSimulate);
+		}
 	}
 
 	virtual void SetShaderParameters(const FNiagaraStatelessSetShaderParameterContext& SetShaderParameterContext) const override
@@ -39,6 +53,29 @@ public:
 		Parameters->CameraOffset_Distribution = ModuleBuiltData->DistributionParameters;
 	}
 
+	static void ParticleSimulate(const NiagaraStateless::FParticleSimulationContext& ParticleSimulationContext)
+	{
+		using namespace NiagaraStateless;
+
+		const FModuleBuiltData* ModuleBuiltData	= ParticleSimulationContext.ReadBuiltData<FModuleBuiltData>();
+		const float* NormalizedAgeData			= ParticleSimulationContext.GetParticleNormalizedAge();
+		const float* PreviousNormalizedAgeData	= ParticleSimulationContext.GetParticlePreviousNormalizedAge();
+
+		for (uint32 i = 0; i < ParticleSimulationContext.GetNumInstances(); ++i)
+		{
+			const FStatelessDistributionSampler<float> CameraOffsetSampler(ParticleSimulationContext, ModuleBuiltData->DistributionParameters, i);
+
+			const float NormalizedAge			= NormalizedAgeData[i];
+			const float PreviousNormalizedAge	= PreviousNormalizedAgeData[i];
+
+			const float CameraOffset			= CameraOffsetSampler.GetValue(ParticleSimulationContext, NormalizedAge);
+			const float PreviousCameraOffset	= CameraOffsetSampler.GetValue(ParticleSimulationContext, PreviousNormalizedAge);
+
+			ParticleSimulationContext.WriteParticleVariable(ModuleBuiltData->CameraVariableOffset, i, CameraOffset);
+			ParticleSimulationContext.WriteParticleVariable(ModuleBuiltData->PreviousCameraVariableOffset, i, PreviousCameraOffset);
+		}
+	}
+
 #if WITH_EDITOR
 	virtual bool CanDisableModule() const override { return true; }
 #endif
diff --git a/Engine/Plugins/FX/Niagara/Source/Niagara/Internal/Stateless/Modules/NiagaraStatelessModule_CurlNoiseForce.h b/Engine/Plugins/FX/Niagara/Source/Niagara/Internal/Stateless/Modules/NiagaraStatelessModule_CurlNoiseForce.h
--- a/Engine/Plugins/FX/Niagara/Source/Niagara/Internal/Stateless/Modules/NiagaraStatelessModule_CurlNoiseForce.h
+++ b/Engine/Plugins/FX/Niagara/Source/Niagara/Internal/Stateless/Modules/NiagaraStatelessModule_CurlNoiseForce.h
@@ -37,7 +37,10 @@ public:
 	UPROPERTY(EditAnywhere, Category = "Parameters", meta = (AllowedClasses = "/Script/Engine.VectorField,/Script/Engine.VolumeTexture,/Script/Engine.TextureRenderTargetVolume"))
 	TObjectPtr<UObject> NoiseTexture;
 
-	virtual void BuildEmitterData(FNiagaraStatelessEmitterDataBuildContext& BuildContext) const override;
+	//-TODO: Add support for GPU once we settle on a feature set for this module
+	virtual ENiagaraStatelessFeatureMask GetFeatureMask() const { return ENiagaraStatelessFeatureMask::ExecuteGPU; }
+
+	virtual void BuildEmitterData(const FNiagaraStatelessEmitterDataBuildContext& BuildContext) const override;
 
 #if WITH_EDITOR
 	virtual bool CanDisableModule() const override { return true; }
diff --git a/Engine/Plugins/FX/Niagara/Source/Niagara/Internal/Stateless/Modules/NiagaraStatelessModule_Drag.h b/Engine/Plugins/FX/Niagara/Source/Niagara/Internal/Stateless/Modules/NiagaraStatelessModule_Drag.h
--- a/Engine/Plugins/FX/Niagara/Source/Niagara/Internal/Stateless/Modules/NiagaraStatelessModule_Drag.h
+++ b/Engine/Plugins/FX/Niagara/Source/Niagara/Internal/Stateless/Modules/NiagaraStatelessModule_Drag.h
@@ -20,7 +20,7 @@ public:
 	UPROPERTY(EditAnywhere, Category = "Parameters", meta = (DisplayName = "Drag", DisableBindingDistribution))
 	FNiagaraDistributionRangeFloat DragDistribution = FNiagaraDistributionRangeFloat(DefaultDrag);
 
-	virtual void BuildEmitterData(FNiagaraStatelessEmitterDataBuildContext& BuildContext) const override
+	virtual void BuildEmitterData(const FNiagaraStatelessEmitterDataBuildContext& BuildContext) const override
 	{
 		if (!IsModuleEnabled())
 		{
diff --git a/Engine/Plugins/FX/Niagara/Source/Niagara/Internal/Stateless/Modules/NiagaraStatelessModule_DynamicMaterialParameters.h b/Engine/Plugins/FX/Niagara/Source/Niagara/Internal/Stateless/Modules/NiagaraStatelessModule_DynamicMaterialParameters.h
--- a/Engine/Plugins/FX/Niagara/Source/Niagara/Internal/Stateless/Modules/NiagaraStatelessModule_DynamicMaterialParameters.h
+++ b/Engine/Plugins/FX/Niagara/Source/Niagara/Internal/Stateless/Modules/NiagaraStatelessModule_DynamicMaterialParameters.h
@@ -5,15 +5,46 @@
 #include "Stateless/NiagaraStatelessModule.h"
 #include "Stateless/NiagaraStatelessEmitterDataBuildContext.h"
 #include "Stateless/NiagaraStatelessModuleShaderParameters.h"
+#include "Stateless/NiagaraStatelessParticleSimContext.h"
 
 #include "NiagaraStatelessModule_DynamicMaterialParameters.generated.h"
 
+USTRUCT()
+struct FNiagaraStatelessDynamicParameterSet
+{
+	GENERATED_BODY()
+		
+	UPROPERTY(EditAnywhere, Category = "Parameters", meta = (InlineEditConditionToggle))
+	uint32 bXChannelEnabled : 1 = true;
+
+	UPROPERTY(EditAnywhere, Category = "Parameters", meta = (InlineEditConditionToggle))
+	uint32 bYChannelEnabled : 1 = true;
+
+	UPROPERTY(EditAnywhere, Category = "Parameters", meta = (InlineEditConditionToggle))
+	uint32 bZChannelEnabled : 1 = true;
+
+	UPROPERTY(EditAnywhere, Category = "Parameters", meta = (InlineEditConditionToggle))
+	uint32 bWChannelEnabled : 1 = true;
+
+	UPROPERTY(EditAnywhere, Category = "Parameters", meta = (DisplayName = "X Channel", EditCondition = "bXChannelEnabled"))
+	FNiagaraDistributionFloat XChannelDistribution = FNiagaraDistributionFloat(FNiagaraStatelessGlobals::GetDefaultDynamicMaterialParametersValue().X);
+
+	UPROPERTY(EditAnywhere, Category = "Parameters", meta = (DisplayName = "Y Channel", EditCondition = "bYChannelEnabled"))
+	FNiagaraDistributionFloat YChannelDistribution = FNiagaraDistributionFloat(FNiagaraStatelessGlobals::GetDefaultDynamicMaterialParametersValue().Y);
+
+	UPROPERTY(EditAnywhere, Category = "Parameters", meta = (DisplayName = "Z Channel", EditCondition = "bZChannelEnabled"))
+	FNiagaraDistributionFloat ZChannelDistribution = FNiagaraDistributionFloat(FNiagaraStatelessGlobals::GetDefaultDynamicMaterialParametersValue().Z);
+
+	UPROPERTY(EditAnywhere, Category = "Parameters", meta = (DisplayName = "W Channel", EditCondition = "bWChannelEnabled"))
+	FNiagaraDistributionFloat WChannelDistribution = FNiagaraDistributionFloat(FNiagaraStatelessGlobals::GetDefaultDynamicMaterialParametersValue().W);
+};
+
 UCLASS(MinimalAPI, EditInlineNew, meta = (DisplayName = "Dynamic Material Parameters"))
 class UNiagaraStatelessModule_DynamicMaterialParameters : public UNiagaraStatelessModule
 {
 	GENERATED_BODY()
 
-	static constexpr int32 NumParameters			= 1;
+	static constexpr int32 NumParameters			= 4;
 	static constexpr int32 NumChannelPerParameter	= 4;
 
 	struct FModuleBuiltData
@@ -24,76 +55,101 @@ class UNiagaraStatelessModule_DynamicMaterialParameters : public UNiagaraStatele
 			{
 				v = FUintVector3::ZeroValue;
 			}
+
+			for (int32& ParameterVariableOffset : ParameterVariableOffsets)
+			{
+				ParameterVariableOffset = INDEX_NONE;
+			}
 		}
 
 		uint32			ChannelMask = 0;
 		FUintVector3	ParameterDistributions[NumParameters * NumChannelPerParameter];
+		int32			ParameterVariableOffsets[NumParameters];
 	};
 
 public:
 	using FParameters = NiagaraStateless::FDynamicMaterialParametersModule_ShaderParameters;
 
 	UPROPERTY(EditAnywhere, Category = "Parameters", meta = (InlineEditConditionToggle))
-	uint32 bParameter0XEnabled : 1 = false;
+	uint32 bParameter0Enabled : 1 = true;
 
 	UPROPERTY(EditAnywhere, Category = "Parameters", meta = (InlineEditConditionToggle))
-	uint32 bParameter0YEnabled : 1 = false;
+	uint32 bParameter1Enabled : 1 = false;
 
 	UPROPERTY(EditAnywhere, Category = "Parameters", meta = (InlineEditConditionToggle))
-	uint32 bParameter0ZEnabled : 1 = false;
+	uint32 bParameter2Enabled : 1 = false;
 
 	UPROPERTY(EditAnywhere, Category = "Parameters", meta = (InlineEditConditionToggle))
-	uint32 bParameter0WEnabled : 1 = false;
+	uint32 bParameter3Enabled : 1 = false;
 
-	UPROPERTY(EditAnywhere, Category = "Parameters", meta = (DisplayName = "Parameter0.X", EditCondition = "bParameter0XEnabled"))
-	FNiagaraDistributionFloat Parameter0XDistribution = FNiagaraDistributionFloat(FNiagaraStatelessGlobals::GetDefaultDynamicMaterialParameters0Value().X);
+	UPROPERTY(EditAnywhere, Category = "Parameters", meta = (EditCondition = "bParameter0Enabled"))
+	FNiagaraStatelessDynamicParameterSet Parameter0;
 
-	UPROPERTY(EditAnywhere, Category = "Parameters", meta = (DisplayName = "Parameter0.Y", EditCondition = "bParameter0YEnabled"))
-	FNiagaraDistributionFloat Parameter0YDistribution = FNiagaraDistributionFloat(FNiagaraStatelessGlobals::GetDefaultDynamicMaterialParameters0Value().Y);
+	UPROPERTY(EditAnywhere, Category = "Parameters", meta = (EditCondition = "bParameter1Enabled"))
+	FNiagaraStatelessDynamicParameterSet Parameter1;
 
-	UPROPERTY(EditAnywhere, Category = "Parameters", meta = (DisplayName = "Parameter0.Z", EditCondition = "bParameter0ZEnabled"))
-	FNiagaraDistributionFloat Parameter0ZDistribution = FNiagaraDistributionFloat(FNiagaraStatelessGlobals::GetDefaultDynamicMaterialParameters0Value().Z);
+	UPROPERTY(EditAnywhere, Category = "Parameters", meta = (EditCondition = "bParameter2Enabled"))
+	FNiagaraStatelessDynamicParameterSet Parameter2;
 
-	UPROPERTY(EditAnywhere, Category = "Parameters", meta = (DisplayName = "Parameter0.W", EditCondition = "bParameter0WEnabled"))
-	FNiagaraDistributionFloat Parameter0WDistribution = FNiagaraDistributionFloat(FNiagaraStatelessGlobals::GetDefaultDynamicMaterialParameters0Value().W);
+	UPROPERTY(EditAnywhere, Category = "Parameters", meta = (EditCondition = "bParameter3Enabled"))
+	FNiagaraStatelessDynamicParameterSet Parameter3;
+
+	const FNiagaraStatelessDynamicParameterSet& GetParameterSet(int32 ParameterIndex) const
+	{
+		switch (ParameterIndex)
+		{
+			case 0: return Parameter0;
+			case 1: return Parameter1;
+			case 2: return Parameter2;
+			case 3: return Parameter3;
+			default: checkNoEntry(); return Parameter0;
+		}
+	}
 
 	uint32 GetParameterChannelMask(int32 ParameterIndex) const
 	{
 		uint32 Mask = 0;
-		switch (ParameterIndex)
+		if ((ParameterIndex == 0 && bParameter0Enabled) ||
+			(ParameterIndex == 1 && bParameter1Enabled) ||
+			(ParameterIndex == 2 && bParameter2Enabled) ||
+			(ParameterIndex == 3 && bParameter3Enabled))
 		{
-			case 0:
-				Mask |= bParameter0XEnabled ? 1 << 0 : 0;
-				Mask |= bParameter0YEnabled ? 1 << 1 : 0;
-				Mask |= bParameter0ZEnabled ? 1 << 2 : 0;
-				Mask |= bParameter0WEnabled ? 1 << 3 : 0;
-				break;
+			const FNiagaraStatelessDynamicParameterSet& ParameterSet = GetParameterSet(ParameterIndex);
+			Mask |= ParameterSet.bXChannelEnabled ? 1 << 0 : 0;
+			Mask |= ParameterSet.bYChannelEnabled ? 1 << 1 : 0;
+			Mask |= ParameterSet.bZChannelEnabled ? 1 << 2 : 0;
+			Mask |= ParameterSet.bWChannelEnabled ? 1 << 3 : 0;
 		}
 		return Mask;
 	}
 
-	const FNiagaraDistributionFloat& GetParameterDistributions(int32 ParameterIndex, int32 ChannelIndex) const
+	const FNiagaraDistributionFloat& GetParameterChannelDistribution(int32 ParameterIndex, int32 ChannelIndex) const
 	{
-		switch (ParameterIndex)
+		const FNiagaraStatelessDynamicParameterSet& ParameterSet = GetParameterSet(ParameterIndex);
+		switch (ChannelIndex)
 		{
-			case 0:
-				switch (ChannelIndex)
-				{
-					case 0:		return Parameter0XDistribution;
-					case 1:		return Parameter0YDistribution;
-					case 2:		return Parameter0ZDistribution;
-					case 3:		return Parameter0WDistribution;
-					default:	checkNoEntry();	return Parameter0XDistribution;
-				}
-				break;
+			case 0: return ParameterSet.XChannelDistribution;
+			case 1: return ParameterSet.YChannelDistribution;
+			case 2: return ParameterSet.ZChannelDistribution;
+			case 3: return ParameterSet.WChannelDistribution;
+			default: checkNoEntry(); return ParameterSet.XChannelDistribution;
+		}
+	}
 
-			default:
-				checkNoEntry();
-				return Parameter0XDistribution;
+	const FNiagaraVariableBase& GetParameterVariable(int32 ParameterIndex) const
+	{
+		const FNiagaraStatelessGlobals& StatelessGlobals = FNiagaraStatelessGlobals::Get();
+		switch (ParameterIndex)
+		{
+			default: checkNoEntry(); return StatelessGlobals.DynamicMaterialParameters0Variable;
+			case 0: return StatelessGlobals.DynamicMaterialParameters0Variable;
+			case 1: return StatelessGlobals.DynamicMaterialParameters1Variable;
+			case 2: return StatelessGlobals.DynamicMaterialParameters2Variable;
+			case 3: return StatelessGlobals.DynamicMaterialParameters3Variable;
 		}
 	}
 
-	virtual void BuildEmitterData(FNiagaraStatelessEmitterDataBuildContext& BuildContext) const override
+	virtual void BuildEmitterData(const FNiagaraStatelessEmitterDataBuildContext& BuildContext) const override
 	{
 		FModuleBuiltData* BuiltData = BuildContext.AllocateBuiltData<FModuleBuiltData>();
 
@@ -106,12 +162,24 @@ public:
 				{
 					continue;
 				}
-				BuiltData->ChannelMask |= ParameterChannelMask << (iParameter * 4);
+
+				BuiltData->ParameterVariableOffsets[iParameter] = BuildContext.FindParticleVariableIndex(GetParameterVariable(iParameter));
+				if (BuiltData->ParameterVariableOffsets[iParameter] == INDEX_NONE)
+				{
+					continue;
+				}
+
+				BuiltData->ChannelMask |= ParameterChannelMask << (iParameter * NumChannelPerParameter);
 				for (int32 iChannel = 0; iChannel < NumChannelPerParameter; ++iChannel)
 				{
-					BuiltData->ParameterDistributions[iParameter * NumParameters + iChannel] = BuildContext.AddDistribution(GetParameterDistributions(iParameter, iChannel), true);
+					BuiltData->ParameterDistributions[(iParameter * NumChannelPerParameter) + iChannel] = BuildContext.AddDistribution(GetParameterChannelDistribution(iParameter, iChannel));
 				}
 			}
+
+			if (BuiltData->ChannelMask != 0)
+			{
+				BuildContext.AddParticleSimulationExecSimulate(&UNiagaraStatelessModule_DynamicMaterialParameters::ParticleSimulate);
+			}
 		}
 	}
 
@@ -125,6 +193,54 @@ public:
 		Parameters->DynamicMaterialParameters_Parameter0Y = ModuleBuiltData->ParameterDistributions[1];
 		Parameters->DynamicMaterialParameters_Parameter0Z = ModuleBuiltData->ParameterDistributions[2];
 		Parameters->DynamicMaterialParameters_Parameter0W = ModuleBuiltData->ParameterDistributions[3];
+		Parameters->DynamicMaterialParameters_Parameter1X = ModuleBuiltData->ParameterDistributions[4];
+		Parameters->DynamicMaterialParameters_Parameter1Y = ModuleBuiltData->ParameterDistributions[5];
+		Parameters->DynamicMaterialParameters_Parameter1Z = ModuleBuiltData->ParameterDistributions[6];
+		Parameters->DynamicMaterialParameters_Parameter1W = ModuleBuiltData->ParameterDistributions[7];
+		Parameters->DynamicMaterialParameters_Parameter2X = ModuleBuiltData->ParameterDistributions[8];
+		Parameters->DynamicMaterialParameters_Parameter2Y = ModuleBuiltData->ParameterDistributions[9];
+		Parameters->DynamicMaterialParameters_Parameter2Z = ModuleBuiltData->ParameterDistributions[10];
+		Parameters->DynamicMaterialParameters_Parameter2W = ModuleBuiltData->ParameterDistributions[11];
+		Parameters->DynamicMaterialParameters_Parameter3X = ModuleBuiltData->ParameterDistributions[12];
+		Parameters->DynamicMaterialParameters_Parameter3Y = ModuleBuiltData->ParameterDistributions[13];
+		Parameters->DynamicMaterialParameters_Parameter3Z = ModuleBuiltData->ParameterDistributions[14];
+		Parameters->DynamicMaterialParameters_Parameter3W = ModuleBuiltData->ParameterDistributions[15];
+	}
+
+	static void ParticleSimulate(const NiagaraStateless::FParticleSimulationContext& ParticleSimulationContext)
+	{
+		using namespace NiagaraStateless;
+
+		const FModuleBuiltData* ModuleBuiltData = ParticleSimulationContext.ReadBuiltData<FModuleBuiltData>();
+		const float* NormalizedAgeData = ParticleSimulationContext.GetParticleNormalizedAge();
+
+		for (uint32 i = 0; i < ParticleSimulationContext.GetNumInstances(); ++i)
+		{
+			for (int32 iParameter = 0; iParameter < NumParameters; ++iParameter)
+			{
+				const uint32 ChannelMask = (ModuleBuiltData->ChannelMask >> (iParameter * NumChannelPerParameter)) & 0xf;
+				if (ChannelMask == 0)
+				{
+					continue;
+				}
+
+				const int32 FirstChannel = iParameter * NumChannelPerParameter;
+				const FStatelessDistributionSampler<float> XChannelSampler(ParticleSimulationContext, ModuleBuiltData->ParameterDistributions[FirstChannel + 0], i);
+				const FStatelessDistributionSampler<float> YChannelSampler(ParticleSimulationContext, ModuleBuiltData->ParameterDistributions[FirstChannel + 1], i);
+				const FStatelessDistributionSampler<float> ZChannelSampler(ParticleSimulationContext, ModuleBuiltData->ParameterDistributions[FirstChannel + 2], i);
+				const FStatelessDistributionSampler<float> WChannelSampler(ParticleSimulationContext, ModuleBuiltData->ParameterDistributions[FirstChannel + 3], i);
+
+				const float NormalizedAge = NormalizedAgeData[i];
+				const FVector4f DynamicParameter(
+					((ChannelMask & 0x1) != 0) ? XChannelSampler.GetValue(ParticleSimulationContext, NormalizedAge) : 0.0f,
+					((ChannelMask & 0x2) != 0) ? YChannelSampler.GetValue(ParticleSimulationContext, NormalizedAge) : 0.0f,
+					((ChannelMask & 0x4) != 0) ? ZChannelSampler.GetValue(ParticleSimulationContext, NormalizedAge) : 0.0f,
+					((ChannelMask & 0x8) != 0) ? WChannelSampler.GetValue(ParticleSimulationContext, NormalizedAge) : 0.0f
+				);
+
+				ParticleSimulationContext.WriteParticleVariable(ModuleBuiltData->ParameterVariableOffsets[iParameter], i, DynamicParameter);
+			}
+		}
 	}
 
 #if WITH_EDITOR
@@ -134,10 +250,13 @@ public:
 	virtual void GetOutputVariables(TArray<FNiagaraVariableBase>& OutVariables) const override
 	{
 		const FNiagaraStatelessGlobals& StatelessGlobals = FNiagaraStatelessGlobals::Get();
-		//-TODO: Channel masking, etc
-		if (GetParameterChannelMask(0) != 0)
+
+		for (int32 i = 0; i < NumParameters; ++i)
 		{
-			OutVariables.AddUnique(StatelessGlobals.DynamicMaterialParameters0Variable);
+			if (GetParameterChannelMask(i) != 0)
+			{
+				OutVariables.AddUnique(GetParameterVariable(i));
+			}
 		}
 	}
 #endif
diff --git a/Engine/Plugins/FX/Niagara/Source/Niagara/Internal/Stateless/Modules/NiagaraStatelessModule_GravityForce.h b/Engine/Plugins/FX/Niagara/Source/Niagara/Internal/Stateless/Modules/NiagaraStatelessModule_GravityForce.h
--- a/Engine/Plugins/FX/Niagara/Source/Niagara/Internal/Stateless/Modules/NiagaraStatelessModule_GravityForce.h
+++ b/Engine/Plugins/FX/Niagara/Source/Niagara/Internal/Stateless/Modules/NiagaraStatelessModule_GravityForce.h
@@ -18,7 +18,7 @@ public:
 	UPROPERTY(EditAnywhere, Category = "Parameters", meta = (DisplayName = "Gravity", DisableUniformDistribution, DisableBindingDistribution))
 	FNiagaraDistributionRangeVector3 GravityDistribution = FNiagaraDistributionRangeVector3(GetDefaultValue());
 
-	virtual void BuildEmitterData(FNiagaraStatelessEmitterDataBuildContext& BuildContext) const override
+	virtual void BuildEmitterData(const FNiagaraStatelessEmitterDataBuildContext& BuildContext) const override
 	{
 		if (!IsModuleEnabled())
 		{
diff --git a/Engine/Plugins/FX/Niagara/Source/Niagara/Internal/Stateless/Modules/NiagaraStatelessModule_InitialMeshOrientation.h b/Engine/Plugins/FX/Niagara/Source/Niagara/Internal/Stateless/Modules/NiagaraStatelessModule_InitialMeshOrientation.h
--- a/Engine/Plugins/FX/Niagara/Source/Niagara/Internal/Stateless/Modules/NiagaraStatelessModule_InitialMeshOrientation.h
+++ b/Engine/Plugins/FX/Niagara/Source/Niagara/Internal/Stateless/Modules/NiagaraStatelessModule_InitialMeshOrientation.h
@@ -5,6 +5,7 @@
 #include "Stateless/NiagaraStatelessModule.h"
 #include "Stateless/NiagaraStatelessEmitterDataBuildContext.h"
 #include "Stateless/NiagaraStatelessModuleShaderParameters.h"
+#include "Stateless/NiagaraStatelessParticleSimContext.h"
 
 #include "NiagaraStatelessModule_InitialMeshOrientation.generated.h"
 
@@ -13,6 +14,14 @@ class UNiagaraStatelessModule_InitialMeshOrientation : public UNiagaraStatelessM
 {
 	GENERATED_BODY()
 
+	struct FModuleBuiltData
+	{
+		FVector3f	Rotation = FVector3f::ZeroVector;
+		FVector3f	RandomRotationRange = FVector3f::ZeroVector;
+		int32		MeshOrientationVariableOffset = INDEX_NONE;
+		int32		PreviousMeshOrientationVariableOffset = INDEX_NONE;
+	};
+
 public:
 	using FParameters = NiagaraStateless::FInitialMeshOrientationModule_ShaderParameters;
 
@@ -22,18 +31,47 @@ public:
 	UPROPERTY(EditAnywhere, Category = "Parameters")
 	FVector3f	RandomRotationRange = FVector3f(360.0f, 360.0f, 360.0f);
 
-	virtual void SetShaderParameters(const FNiagaraStatelessSetShaderParameterContext& SetShaderParameterContext) const override
+	virtual void BuildEmitterData(const FNiagaraStatelessEmitterDataBuildContext& BuildContext) const override
 	{
-		FParameters* Parameters = SetShaderParameterContext.GetParameterNestedStruct<FParameters>();
+		FModuleBuiltData* BuiltData = BuildContext.AllocateBuiltData<FModuleBuiltData>();
+
+		const FNiagaraStatelessGlobals& StatelessGlobals	= FNiagaraStatelessGlobals::Get();
+		BuiltData->MeshOrientationVariableOffset			= BuildContext.FindParticleVariableIndex(StatelessGlobals.MeshOrientationVariable);
+		BuiltData->PreviousMeshOrientationVariableOffset	= BuildContext.FindParticleVariableIndex(StatelessGlobals.PreviousMeshOrientationVariable);
+
 		if (IsModuleEnabled())
 		{
-			Parameters->InitialMeshOrientation_Rotation			= Rotation / 360.0f;
-			Parameters->InitialMeshOrientation_RandomRangeScale	= RandomRotationRange / 360.0f;
+			BuiltData->Rotation				= Rotation / 360.0f;
+			BuiltData->RandomRotationRange	= RandomRotationRange / 360.0f;
 		}
-		else
+
+		if (BuiltData->MeshOrientationVariableOffset != INDEX_NONE || BuiltData->PreviousMeshOrientationVariableOffset != INDEX_NONE)
+		{
+			BuildContext.AddParticleSimulationExecSimulate(&UNiagaraStatelessModule_InitialMeshOrientation::ParticleSimulate);
+		}
+	}
+
+	virtual void SetShaderParameters(const FNiagaraStatelessSetShaderParameterContext& SetShaderParameterContext) const override
+	{
+		FParameters* Parameters = SetShaderParameterContext.GetParameterNestedStruct<FParameters>();
+
+		const FModuleBuiltData* ModuleBuiltData = SetShaderParameterContext.ReadBuiltData<FModuleBuiltData>();
+		Parameters->InitialMeshOrientation_Rotation			= ModuleBuiltData->Rotation;
+		Parameters->InitialMeshOrientation_RandomRangeScale	= ModuleBuiltData->RandomRotationRange;
+	}
+
+	static void ParticleSimulate(const NiagaraStateless::FParticleSimulationContext& ParticleSimulationContext)
+	{
+		using namespace NiagaraStateless;
+
+		const FModuleBuiltData* ModuleBuiltData = ParticleSimulationContext.ReadBuiltData<FModuleBuiltData>();
+		for (uint32 i = 0; i < ParticleSimulationContext.GetNumInstances(); ++i)
 		{
-			Parameters->InitialMeshOrientation_Rotation			= FVector3f::ZeroVector;
-			Parameters->InitialMeshOrientation_RandomRangeScale	= FVector3f::ZeroVector;
+			const FVector3f	Rotation = ModuleBuiltData->Rotation + (ParticleSimulationContext.RandomFloat3(i) * ModuleBuiltData->RandomRotationRange);
+			const FQuat4f Quat = ParticleSimulationContext.RotatorToQuat(Rotation);
+
+			ParticleSimulationContext.WriteParticleVariable(ModuleBuiltData->MeshOrientationVariableOffset, i, Quat);
+			ParticleSimulationContext.WriteParticleVariable(ModuleBuiltData->PreviousMeshOrientationVariableOffset, i, Quat);
 		}
 	}
 
diff --git a/Engine/Plugins/FX/Niagara/Source/Niagara/Internal/Stateless/Modules/NiagaraStatelessModule_InitializeParticle.h b/Engine/Plugins/FX/Niagara/Source/Niagara/Internal/Stateless/Modules/NiagaraStatelessModule_InitializeParticle.h
--- a/Engine/Plugins/FX/Niagara/Source/Niagara/Internal/Stateless/Modules/NiagaraStatelessModule_InitializeParticle.h
+++ b/Engine/Plugins/FX/Niagara/Source/Niagara/Internal/Stateless/Modules/NiagaraStatelessModule_InitializeParticle.h
@@ -5,6 +5,7 @@
 #include "Stateless/NiagaraStatelessModule.h"
 #include "Stateless/NiagaraStatelessEmitterDataBuildContext.h"
 #include "Stateless/NiagaraStatelessModuleShaderParameters.h"
+#include "Stateless/NiagaraStatelessParticleSimContext.h"
 #include "Stateless/Modules/NiagaraStatelessModuleCommon.h"
 
 #include "NiagaraParameterBinding.h"
@@ -31,6 +32,19 @@ class UNiagaraStatelessModule_InitializeParticle : public UNiagaraStatelessModul
 		FNiagaraStatelessRangeFloat		SpriteRotationRange;
 		FNiagaraStatelessRangeVector3	MeshScaleRange;
 		FNiagaraStatelessRangeFloat		RibbonWidthRange;
+
+		int32							PositionVariableOffset = INDEX_NONE;
+		int32							ColorVariableOffset = INDEX_NONE;
+		int32							RibbonWidthVariableOffset = INDEX_NONE;
+		int32							SpriteSizeVariableOffset = INDEX_NONE;
+		int32							SpriteRotationVariableOffset = INDEX_NONE;
+		int32							ScaleVariableOffset = INDEX_NONE;
+
+		int32							PreviousPositionVariableOffset = INDEX_NONE;
+		int32							PreviousRibbonWidthVariableOffset = INDEX_NONE;
+		int32							PreviousSpriteSizeVariableOffset = INDEX_NONE;
+		int32							PreviousSpriteRotationVariableOffset = INDEX_NONE;
+		int32							PreviousScaleVariableOffset = INDEX_NONE;
 	};
 
 public:
@@ -63,8 +77,10 @@ public:
 	UPROPERTY(EditAnywhere, Category = "Parameters", meta = (DisableCurveDistribution))
 	FNiagaraDistributionVector3	InitialPosition = FNiagaraDistributionVector3(FVector3f::ZeroVector);
 
-	virtual void BuildEmitterData(FNiagaraStatelessEmitterDataBuildContext& BuildContext) const override
+	virtual void BuildEmitterData(const FNiagaraStatelessEmitterDataBuildContext& BuildContext) const override
 	{
+		BuildContext.AddParticleSimulationExecSimulate(&UNiagaraStatelessModule_InitializeParticle::ParticleSimulate);
+
 		FModuleBuiltData* BuiltData			= BuildContext.AllocateBuiltData<FModuleBuiltData>();
 		BuiltData->ModuleFlags				 = SpriteSizeDistribution.IsUniform() ? EInitializeParticleModuleFlag_UniformSpriteSize : 0;
 		BuiltData->ModuleFlags				|= MeshScaleDistribution.IsUniform() ? EInitializeParticleModuleFlag_UniformMeshScale : 0;
@@ -78,6 +94,19 @@ public:
 		BuiltData->MeshScaleRange			= BuildContext.ConvertDistributionToRange(MeshScaleDistribution, FNiagaraStatelessGlobals::GetDefaultScaleValue());
 		BuiltData->RibbonWidthRange			= BuildContext.ConvertDistributionToRange(RibbonWidthDistribution, FNiagaraStatelessGlobals::GetDefaultRibbonWidthValue());
 
+		const FNiagaraStatelessGlobals& StatelessGlobals = FNiagaraStatelessGlobals::Get();
+		BuiltData->PositionVariableOffset				= BuildContext.FindParticleVariableIndex(StatelessGlobals.PositionVariable);
+		BuiltData->ColorVariableOffset					= BuildContext.FindParticleVariableIndex(StatelessGlobals.ColorVariable);
+		BuiltData->RibbonWidthVariableOffset			= BuildContext.FindParticleVariableIndex(StatelessGlobals.RibbonWidthVariable);
+		BuiltData->SpriteSizeVariableOffset				= BuildContext.FindParticleVariableIndex(StatelessGlobals.SpriteSizeVariable);
+		BuiltData->SpriteRotationVariableOffset			= BuildContext.FindParticleVariableIndex(StatelessGlobals.SpriteRotationVariable);
+		BuiltData->ScaleVariableOffset					= BuildContext.FindParticleVariableIndex(StatelessGlobals.ScaleVariable);
+		BuiltData->PreviousPositionVariableOffset		= BuildContext.FindParticleVariableIndex(StatelessGlobals.PreviousPositionVariable);
+		BuiltData->PreviousRibbonWidthVariableOffset	= BuildContext.FindParticleVariableIndex(StatelessGlobals.PreviousRibbonWidthVariable);
+		BuiltData->PreviousSpriteSizeVariableOffset		= BuildContext.FindParticleVariableIndex(StatelessGlobals.PreviousSpriteSizeVariable);
+		BuiltData->PreviousSpriteRotationVariableOffset	= BuildContext.FindParticleVariableIndex(StatelessGlobals.PreviousSpriteRotationVariable);
+		BuiltData->PreviousScaleVariableOffset			= BuildContext.FindParticleVariableIndex(StatelessGlobals.PreviousScaleVariable);
+
 		NiagaraStateless::FPhysicsBuildData& PhysicsBuildData = BuildContext.GetTransientBuildData<NiagaraStateless::FPhysicsBuildData>();
 		PhysicsBuildData.MassRange			= MassDistribution.CalculateRange(FNiagaraStatelessGlobals::GetDefaultMassValue());
 	}
@@ -96,6 +125,38 @@ public:
 		SetShaderParameterContext.ConvertRangeToScaleBias(ModuleBuiltData->RibbonWidthRange,	Parameters->InitializeParticle_RibbonWidthScale, Parameters->InitializeParticle_RibbonWidthBias);
 	}
 
+	static void ParticleSimulate(const NiagaraStateless::FParticleSimulationContext& ParticleSimulationContext)
+	{
+		using namespace NiagaraStateless;
+
+		const FModuleBuiltData* ModuleBuiltData = ParticleSimulationContext.ReadBuiltData<FModuleBuiltData>();
+
+		for (uint32 i = 0; i < ParticleSimulationContext.GetNumInstances(); ++i)
+		{
+			const FStatelessDistributionSampler<FVector3f> PositionSampler(ParticleSimulationContext, ModuleBuiltData->InitialPosition, i);
+
+			const FVector3f		Position	= PositionSampler.GetValue(ParticleSimulationContext, 0.0f);
+			const FLinearColor	Color		= ParticleSimulationContext.RandomScaleBiasFloat(i, ModuleBuiltData->ColorRange);
+			const float			RibbonWidth	= ParticleSimulationContext.RandomScaleBiasFloat(i, ModuleBuiltData->RibbonWidthRange);
+			const FVector2f		SpriteSize	= ParticleSimulationContext.RandomScaleBiasFloat(i, ModuleBuiltData->SpriteSizeRange);
+			const float			SpriteRot	= ParticleSimulationContext.RandomScaleBiasFloat(i, ModuleBuiltData->SpriteRotationRange);
+			const FVector3f		Scale		= ParticleSimulationContext.RandomScaleBiasFloat(i, ModuleBuiltData->MeshScaleRange);
+
+			ParticleSimulationContext.WriteParticleVariable(ModuleBuiltData->PositionVariableOffset,				i, Position);
+			ParticleSimulationContext.WriteParticleVariable(ModuleBuiltData->ColorVariableOffset,					i, Color);
+			ParticleSimulationContext.WriteParticleVariable(ModuleBuiltData->RibbonWidthVariableOffset,				i, RibbonWidth);
+			ParticleSimulationContext.WriteParticleVariable(ModuleBuiltData->SpriteSizeVariableOffset,				i, SpriteSize);
+			ParticleSimulationContext.WriteParticleVariable(ModuleBuiltData->SpriteRotationVariableOffset,			i, SpriteRot);
+			ParticleSimulationContext.WriteParticleVariable(ModuleBuiltData->ScaleVariableOffset,					i, Scale);
+
+			ParticleSimulationContext.WriteParticleVariable(ModuleBuiltData->PreviousPositionVariableOffset,		i, Position);
+			ParticleSimulationContext.WriteParticleVariable(ModuleBuiltData->PreviousRibbonWidthVariableOffset,		i, RibbonWidth);
+			ParticleSimulationContext.WriteParticleVariable(ModuleBuiltData->PreviousSpriteSizeVariableOffset,		i, SpriteSize);
+			ParticleSimulationContext.WriteParticleVariable(ModuleBuiltData->PreviousSpriteRotationVariableOffset,	i, SpriteRot);
+			ParticleSimulationContext.WriteParticleVariable(ModuleBuiltData->PreviousScaleVariableOffset,			i, Scale);
+		}
+	}
+
 #if WITH_EDITORONLY_DATA
 	virtual void GetOutputVariables(TArray<FNiagaraVariableBase>& OutVariables) const override
 	{
diff --git a/Engine/Plugins/FX/Niagara/Source/Niagara/Internal/Stateless/Modules/NiagaraStatelessModule_MeshIndex.h b/Engine/Plugins/FX/Niagara/Source/Niagara/Internal/Stateless/Modules/NiagaraStatelessModule_MeshIndex.h
--- a/Engine/Plugins/FX/Niagara/Source/Niagara/Internal/Stateless/Modules/NiagaraStatelessModule_MeshIndex.h
+++ b/Engine/Plugins/FX/Niagara/Source/Niagara/Internal/Stateless/Modules/NiagaraStatelessModule_MeshIndex.h
@@ -23,7 +23,7 @@ public:
 	UPROPERTY(EditAnywhere, Category = "Parameters", meta = (EditConditionHides, EditCondition="NeedsMeshIndexWeights()"))
 	TArray<float> MeshIndexWeight;
 
-	virtual void BuildEmitterData(FNiagaraStatelessEmitterDataBuildContext& BuildContext) const override;
+	virtual void BuildEmitterData(const FNiagaraStatelessEmitterDataBuildContext& BuildContext) const override;
 	virtual void SetShaderParameters(const FNiagaraStatelessSetShaderParameterContext& SetShaderParameterContext) const override;
 
 	UFUNCTION()
diff --git a/Engine/Plugins/FX/Niagara/Source/Niagara/Internal/Stateless/Modules/NiagaraStatelessModule_MeshRotationRate.h b/Engine/Plugins/FX/Niagara/Source/Niagara/Internal/Stateless/Modules/NiagaraStatelessModule_MeshRotationRate.h
--- a/Engine/Plugins/FX/Niagara/Source/Niagara/Internal/Stateless/Modules/NiagaraStatelessModule_MeshRotationRate.h
+++ b/Engine/Plugins/FX/Niagara/Source/Niagara/Internal/Stateless/Modules/NiagaraStatelessModule_MeshRotationRate.h
@@ -5,6 +5,7 @@
 #include "Stateless/NiagaraStatelessModule.h"
 #include "Stateless/NiagaraStatelessEmitterDataBuildContext.h"
 #include "Stateless/NiagaraStatelessModuleShaderParameters.h"
+#include "Stateless/NiagaraStatelessParticleSimContext.h"
 
 #include "NiagaraStatelessModule_MeshRotationRate.generated.h"
 
@@ -15,7 +16,10 @@ class UNiagaraStatelessModule_MeshRotationRate : public UNiagaraStatelessModule
 
 	struct FModuleBuiltData
 	{
+		bool							ModuleEnabled = false;
 		FNiagaraStatelessRangeVector3	RotationRange;
+		int32							MeshOrientationVariableOffset = INDEX_NONE;
+		int32							PreviousMeshOrientationVariableOffset = INDEX_NONE;
 	};
 
 public:
@@ -24,22 +28,62 @@ public:
 	UPROPERTY(EditAnywhere, Category = "Parameters", meta = (DisplayName = "Rotation Rate"))
 	FNiagaraDistributionRangeVector3 RotationRateDistribution = FNiagaraDistributionRangeVector3(FVector3f::ZeroVector);
 
-	virtual void BuildEmitterData(FNiagaraStatelessEmitterDataBuildContext& BuildContext) const override
+	virtual void BuildEmitterData(const FNiagaraStatelessEmitterDataBuildContext& BuildContext) const override
 	{
 		FModuleBuiltData* BuiltData = BuildContext.AllocateBuiltData<FModuleBuiltData>();
-		BuiltData->RotationRange = BuildContext.ConvertDistributionToRange(RotationRateDistribution, FVector3f::ZeroVector, IsModuleEnabled());
-		BuiltData->RotationRange.Min *= 1.0f / 360.0f;
-		BuiltData->RotationRange.Max *= 1.0f / 360.0f;
+		if (!IsModuleEnabled())
+		{
+			return;
+		}
+
+		const FNiagaraStatelessGlobals& StatelessGlobals	= FNiagaraStatelessGlobals::Get();
+		BuiltData->MeshOrientationVariableOffset			= BuildContext.FindParticleVariableIndex(StatelessGlobals.MeshOrientationVariable);
+		BuiltData->PreviousMeshOrientationVariableOffset	= BuildContext.FindParticleVariableIndex(StatelessGlobals.PreviousMeshOrientationVariable);
+
+		if (BuiltData->MeshOrientationVariableOffset != INDEX_NONE || BuiltData->PreviousMeshOrientationVariableOffset != INDEX_NONE)
+		{
+			BuiltData->ModuleEnabled = true;
+			BuiltData->RotationRange = BuildContext.ConvertDistributionToRange(RotationRateDistribution, FVector3f::ZeroVector);
+			BuiltData->RotationRange.Min *= 1.0f / 360.0f;
+			BuiltData->RotationRange.Max *= 1.0f / 360.0f;
+
+			BuildContext.AddParticleSimulationExecSimulate(&UNiagaraStatelessModule_MeshRotationRate::ParticleSimulate);
+		}
 	}
 
 	virtual void SetShaderParameters(const FNiagaraStatelessSetShaderParameterContext& SetShaderParameterContext) const override
 	{
 		FParameters* Parameters = SetShaderParameterContext.GetParameterNestedStruct<FParameters>();
 		const FModuleBuiltData* ModuleBuiltData = SetShaderParameterContext.ReadBuiltData<FModuleBuiltData>();
-
+		Parameters->MeshRotationRate_ModuleEnabled = ModuleBuiltData->ModuleEnabled;
 		SetShaderParameterContext.ConvertRangeToScaleBias(ModuleBuiltData->RotationRange, Parameters->MeshRotationRate_Scale, Parameters->MeshRotationRate_Bias);
 	}
 
+	static void ParticleSimulate(const NiagaraStateless::FParticleSimulationContext& ParticleSimulationContext)
+	{
+		using namespace NiagaraStateless;
+
+		const FModuleBuiltData* ModuleBuiltData = ParticleSimulationContext.ReadBuiltData<FModuleBuiltData>();
+		const float* AgeData = ParticleSimulationContext.GetParticleAge();
+		const float* PreviousAgeData = ParticleSimulationContext.GetParticlePreviousAge();
+
+		for (uint32 i = 0; i < ParticleSimulationContext.GetNumInstances(); ++i)
+		{
+			const FVector3f RotationRate = ParticleSimulationContext.RandomScaleBiasFloat(i, ModuleBuiltData->RotationRange);
+			const float Age = AgeData[i];
+			const float PreviousAge = PreviousAgeData[i];
+
+			FQuat4f MeshOrientation = ParticleSimulationContext.ReadParticleVariable(ModuleBuiltData->MeshOrientationVariableOffset, i, FQuat4f::Identity);
+			FQuat4f PreviousMeshOrientation = ParticleSimulationContext.ReadParticleVariable(ModuleBuiltData->PreviousMeshOrientationVariableOffset, i, FQuat4f::Identity);
+
+			MeshOrientation *= ParticleSimulationContext.RotatorToQuat(RotationRate * Age);
+			PreviousMeshOrientation *= ParticleSimulationContext.RotatorToQuat(RotationRate * PreviousAge);
+
+			ParticleSimulationContext.WriteParticleVariable(ModuleBuiltData->MeshOrientationVariableOffset, i, MeshOrientation);
+			ParticleSimulationContext.WriteParticleVariable(ModuleBuiltData->PreviousMeshOrientationVariableOffset, i, PreviousMeshOrientation);
+		}
+	}
+
 #if WITH_EDITOR
 	virtual bool CanDisableModule() const override { return true; }
 #endif
diff --git a/Engine/Plugins/FX/Niagara/Source/Niagara/Internal/Stateless/Modules/NiagaraStatelessModule_RotateAroundPoint.h b/Engine/Plugins/FX/Niagara/Source/Niagara/Internal/Stateless/Modules/NiagaraStatelessModule_RotateAroundPoint.h
--- a/Engine/Plugins/FX/Niagara/Source/Niagara/Internal/Stateless/Modules/NiagaraStatelessModule_RotateAroundPoint.h
+++ b/Engine/Plugins/FX/Niagara/Source/Niagara/Internal/Stateless/Modules/NiagaraStatelessModule_RotateAroundPoint.h
@@ -30,6 +30,9 @@ public:
 	UPROPERTY(EditAnywhere, Category = "Parameters")
 	float InitialPhaseMax = 0.f;
 
+	//-TODO: Add support for GPU once we settle on a feature set for this module
+	virtual ENiagaraStatelessFeatureMask GetFeatureMask() const { return ENiagaraStatelessFeatureMask::ExecuteGPU; }
+
 	virtual void SetShaderParameters(const FNiagaraStatelessSetShaderParameterContext& SetShaderParameterContext) const override
 	{
 		FParameters* Parameters = SetShaderParameterContext.GetParameterNestedStruct<FParameters>();
diff --git a/Engine/Plugins/FX/Niagara/Source/Niagara/Internal/Stateless/Modules/NiagaraStatelessModule_ScaleColor.h b/Engine/Plugins/FX/Niagara/Source/Niagara/Internal/Stateless/Modules/NiagaraStatelessModule_ScaleColor.h
--- a/Engine/Plugins/FX/Niagara/Source/Niagara/Internal/Stateless/Modules/NiagaraStatelessModule_ScaleColor.h
+++ b/Engine/Plugins/FX/Niagara/Source/Niagara/Internal/Stateless/Modules/NiagaraStatelessModule_ScaleColor.h
@@ -5,6 +5,7 @@
 #include "Stateless/NiagaraStatelessModule.h"
 #include "Stateless/NiagaraStatelessEmitterDataBuildContext.h"
 #include "Stateless/NiagaraStatelessModuleShaderParameters.h"
+#include "Stateless/NiagaraStatelessParticleSimContext.h"
 
 #include "NiagaraStatelessModule_ScaleColor.generated.h"
 
@@ -16,6 +17,7 @@ class UNiagaraStatelessModule_ScaleColor : public UNiagaraStatelessModule
 	struct FModuleBuiltData
 	{
 		FUintVector3	DistributionParameters = FUintVector3::ZeroValue;
+		int32			ColorVariableOffset = INDEX_NONE;
 	};
 
 public:
@@ -24,10 +26,23 @@ public:
 	UPROPERTY(EditAnywhere, Category = "Parameters", meta = (DisplayName = "Scale"))
 	FNiagaraDistributionColor ScaleDistribution = FNiagaraDistributionColor(FLinearColor::White);
 
-	virtual void BuildEmitterData(FNiagaraStatelessEmitterDataBuildContext& BuildContext) const override
+	virtual void BuildEmitterData(const FNiagaraStatelessEmitterDataBuildContext& BuildContext) const override
 	{
 		FModuleBuiltData* BuiltData = BuildContext.AllocateBuiltData<FModuleBuiltData>();
-		BuiltData->DistributionParameters = BuildContext.AddDistribution(ScaleDistribution, IsModuleEnabled());
+		if (!IsModuleEnabled())
+		{
+			return;
+		}
+
+		const FNiagaraStatelessGlobals& StatelessGlobals = FNiagaraStatelessGlobals::Get();
+		BuiltData->ColorVariableOffset = BuildContext.FindParticleVariableIndex(StatelessGlobals.ColorVariable);
+
+		if (BuiltData->ColorVariableOffset != INDEX_NONE)
+		{
+			BuiltData->DistributionParameters = BuildContext.AddDistribution(ScaleDistribution);
+
+			BuildContext.AddParticleSimulationExecSimulate(&UNiagaraStatelessModule_ScaleColor::ParticleSimulate);
+		}
 	}
 
 	virtual void SetShaderParameters(const FNiagaraStatelessSetShaderParameterContext& SetShaderParameterContext) const override
@@ -38,6 +53,25 @@ public:
 		Parameters->ScaleColor_Distribution = ModuleBuiltData->DistributionParameters;
 	}
 
+	static void ParticleSimulate(const NiagaraStateless::FParticleSimulationContext& ParticleSimulationContext)
+	{
+		using namespace NiagaraStateless;
+
+		const FModuleBuiltData* ModuleBuiltData = ParticleSimulationContext.ReadBuiltData<FModuleBuiltData>();
+		const float* NormalizedAgeData = ParticleSimulationContext.GetParticleNormalizedAge();
+
+		for (uint32 i = 0; i < ParticleSimulationContext.GetNumInstances(); ++i)
+		{
+			FStatelessDistributionSampler<FLinearColor> ColorSampler(ParticleSimulationContext, ModuleBuiltData->DistributionParameters, i);
+
+			FLinearColor Color = ParticleSimulationContext.ReadParticleVariable(ModuleBuiltData->ColorVariableOffset, i, FLinearColor::White);
+
+			Color *= ColorSampler.GetValue(ParticleSimulationContext, NormalizedAgeData[i]);
+
+			ParticleSimulationContext.WriteParticleVariable(ModuleBuiltData->ColorVariableOffset, i, Color);
+		}
+	}
+
 #if WITH_EDITOR
 	virtual bool CanDisableModule() const override { return true; }
 #endif
diff --git a/Engine/Plugins/FX/Niagara/Source/Niagara/Internal/Stateless/Modules/NiagaraStatelessModule_ScaleMeshSize.h b/Engine/Plugins/FX/Niagara/Source/Niagara/Internal/Stateless/Modules/NiagaraStatelessModule_ScaleMeshSize.h
--- a/Engine/Plugins/FX/Niagara/Source/Niagara/Internal/Stateless/Modules/NiagaraStatelessModule_ScaleMeshSize.h
+++ b/Engine/Plugins/FX/Niagara/Source/Niagara/Internal/Stateless/Modules/NiagaraStatelessModule_ScaleMeshSize.h
@@ -5,6 +5,7 @@
 #include "Stateless/NiagaraStatelessModule.h"
 #include "Stateless/NiagaraStatelessEmitterDataBuildContext.h"
 #include "Stateless/NiagaraStatelessModuleShaderParameters.h"
+#include "Stateless/NiagaraStatelessParticleSimContext.h"
 
 #include "NiagaraParameterBinding.h"
 
@@ -20,6 +21,9 @@ class UNiagaraStatelessModule_ScaleMeshSize : public UNiagaraStatelessModule
 		FUintVector3	DistributionParameters = FUintVector3::ZeroValue;
 		FVector3f		CurveScale = FVector3f::OneVector;
 		int32			CurveScaleOffset = INDEX_NONE;
+
+		int32			ScaleVariableOffset = INDEX_NONE;
+		int32			PreviousScaleVariableOffset = INDEX_NONE;
 	};
 
 public:
@@ -45,15 +49,30 @@ public:
 	#endif
 	}
 
-	virtual void BuildEmitterData(FNiagaraStatelessEmitterDataBuildContext& BuildContext) const override
+	virtual void BuildEmitterData(const FNiagaraStatelessEmitterDataBuildContext& BuildContext) const override
 	{
 		FModuleBuiltData* BuiltData = BuildContext.AllocateBuiltData<FModuleBuiltData>();
-		BuiltData->DistributionParameters = BuildContext.AddDistribution(ScaleDistribution, IsModuleEnabled());
-		if (IsModuleEnabled() && UseScaleCurveRange())
+		if (!IsModuleEnabled())
+		{
+			return;
+		}
+
+		const FNiagaraStatelessGlobals& StatelessGlobals = FNiagaraStatelessGlobals::Get();
+		BuiltData->ScaleVariableOffset = BuildContext.FindParticleVariableIndex(StatelessGlobals.ScaleVariable);
+		BuiltData->PreviousScaleVariableOffset = BuildContext.FindParticleVariableIndex(StatelessGlobals.PreviousScaleVariable);
+
+		if (BuiltData->ScaleVariableOffset == INDEX_NONE && BuiltData->PreviousScaleVariableOffset == INDEX_NONE)
+		{
+			return;
+		}
+
+		BuiltData->DistributionParameters = BuildContext.AddDistribution(ScaleDistribution);
+		if (UseScaleCurveRange())
 		{
 			BuiltData->CurveScaleOffset = BuildContext.AddRendererBinding(ScaleCurveRange.ResolvedParameter);
 			BuiltData->CurveScale = ScaleCurveRange.GetDefaultValue<FVector3f>();
 		}
+		BuildContext.AddParticleSimulationExecSimulate(&UNiagaraStatelessModule_ScaleMeshSize::ParticleSimulate);
 	}
 
 	virtual void SetShaderParameters(const FNiagaraStatelessSetShaderParameterContext& SetShaderParameterContext) const override
@@ -66,6 +85,31 @@ public:
 		Parameters->ScaleMeshSize_CurveScaleOffset	= ModuleBuiltData->CurveScaleOffset;
 	}
 
+	static void ParticleSimulate(const NiagaraStateless::FParticleSimulationContext& ParticleSimulationContext)
+	{
+		using namespace NiagaraStateless;
+
+		const FModuleBuiltData* ModuleBuiltData = ParticleSimulationContext.ReadBuiltData<FModuleBuiltData>();
+		const float* NormalizedAgeData = ParticleSimulationContext.GetParticleNormalizedAge();
+		const float* PreviousNormalizedAgeData = ParticleSimulationContext.GetParticlePreviousNormalizedAge();
+
+		const FVector3f CurveScale = ParticleSimulationContext.GetParameterBufferFloat(ModuleBuiltData->CurveScaleOffset, ModuleBuiltData->CurveScale);
+
+		for (uint32 i = 0; i < ParticleSimulationContext.GetNumInstances(); ++i)
+		{
+			FStatelessDistributionSampler<FVector3f> ScaleSampler(ParticleSimulationContext, ModuleBuiltData->DistributionParameters, i);
+
+			FVector3f Scale = ParticleSimulationContext.ReadParticleVariable(ModuleBuiltData->ScaleVariableOffset, i, FVector3f::OneVector);
+			FVector3f PreviousScale = ParticleSimulationContext.ReadParticleVariable(ModuleBuiltData->PreviousScaleVariableOffset, i, FVector3f::OneVector);
+
+			Scale = Scale * ScaleSampler.GetValue(ParticleSimulationContext, NormalizedAgeData[i]) * CurveScale;
+			PreviousScale = PreviousScale * ScaleSampler.GetValue(ParticleSimulationContext, PreviousNormalizedAgeData[i]) * CurveScale;
+
+			ParticleSimulationContext.WriteParticleVariable(ModuleBuiltData->ScaleVariableOffset, i, Scale);
+			ParticleSimulationContext.WriteParticleVariable(ModuleBuiltData->PreviousScaleVariableOffset, i, PreviousScale);
+		}
+	}
+
 	UFUNCTION()
 	bool UseScaleCurveRange() const { return ScaleDistribution.IsCurve(); }
 
diff --git a/Engine/Plugins/FX/Niagara/Source/Niagara/Internal/Stateless/Modules/NiagaraStatelessModule_ScaleMeshSizeBySpeed.h b/Engine/Plugins/FX/Niagara/Source/Niagara/Internal/Stateless/Modules/NiagaraStatelessModule_ScaleMeshSizeBySpeed.h
--- a/Engine/Plugins/FX/Niagara/Source/Niagara/Internal/Stateless/Modules/NiagaraStatelessModule_ScaleMeshSizeBySpeed.h
+++ b/Engine/Plugins/FX/Niagara/Source/Niagara/Internal/Stateless/Modules/NiagaraStatelessModule_ScaleMeshSizeBySpeed.h
@@ -5,6 +5,7 @@
 #include "Stateless/NiagaraStatelessModule.h"
 #include "Stateless/NiagaraStatelessEmitterDataBuildContext.h"
 #include "Stateless/NiagaraStatelessModuleShaderParameters.h"
+#include "Stateless/NiagaraStatelessParticleSimContext.h"
 
 #include "NiagaraStatelessModule_ScaleMeshSizeBySpeed.generated.h"
 
@@ -17,6 +18,11 @@ class UNiagaraStatelessModule_ScaleMeshSizeBySpeed : public UNiagaraStatelessMod
 	{
 		float			VelocityNorm = 0.0f;
 		FUintVector2	ScaleDistribution = FUintVector2::ZeroValue;
+
+		int32			PositionVariableOffset = INDEX_NONE;
+		int32			PreviousPositionVariableOffset = INDEX_NONE;
+		int32			ScaleVariableVariableOffset = INDEX_NONE;
+		int32			PreviousScaleVariableVariableOffset = INDEX_NONE;
 	};
 
 public:
@@ -28,24 +34,38 @@ public:
 	UPROPERTY(EditAnywhere, Category = "Parameters", meta = (DisplayName = "Scale", DisableRangeDistribution, DisableBindingDistribution))
 	FNiagaraDistributionVector3 ScaleDistribution = FNiagaraDistributionVector3(1.0f);
 
-	virtual void BuildEmitterData(FNiagaraStatelessEmitterDataBuildContext& BuildContext) const override
+	virtual void BuildEmitterData(const FNiagaraStatelessEmitterDataBuildContext& BuildContext) const override
 	{
 		FModuleBuiltData* BuiltData = BuildContext.AllocateBuiltData<FModuleBuiltData>();
-		if (IsModuleEnabled())
+		if (!IsModuleEnabled())
+		{
+			return;
+		}
+
+		const FNiagaraStatelessGlobals& StatelessGlobals = FNiagaraStatelessGlobals::Get();
+		BuiltData->PositionVariableOffset				= BuildContext.FindParticleVariableIndex(StatelessGlobals.PositionVariable);
+		BuiltData->PreviousPositionVariableOffset		= BuildContext.FindParticleVariableIndex(StatelessGlobals.PreviousPositionVariable);
+		BuiltData->ScaleVariableVariableOffset			= BuildContext.FindParticleVariableIndex(StatelessGlobals.ScaleVariable);
+		BuiltData->PreviousScaleVariableVariableOffset	= BuildContext.FindParticleVariableIndex(StatelessGlobals.PreviousScaleVariable);
+
+		if (BuiltData->ScaleVariableVariableOffset == INDEX_NONE && BuiltData->PreviousScaleVariableVariableOffset == INDEX_NONE)
+		{
+			return;
+		}
+
+		BuiltData->VelocityNorm = VelocityThreshold > 0.0f ? 1.0f / (VelocityThreshold * VelocityThreshold) : 0.0f;
+		if (ScaleDistribution.IsCurve() && ScaleDistribution.Values.Num() > 1)
+		{
+			BuiltData->ScaleDistribution.X = BuildContext.AddStaticData(ScaleDistribution.Values);
+			BuiltData->ScaleDistribution.Y = ScaleDistribution.Values.Num() - 1;
+		}
+		else
 		{
-			BuiltData->VelocityNorm = VelocityThreshold > 0.0f ? 1.0f / (VelocityThreshold * VelocityThreshold) : 0.0f;
-			if (ScaleDistribution.IsCurve() && ScaleDistribution.Values.Num() > 1)
-			{
-				BuiltData->ScaleDistribution.X = BuildContext.AddStaticData(ScaleDistribution.Values);
-				BuiltData->ScaleDistribution.Y = ScaleDistribution.Values.Num() - 1;
-			}
-			else
-			{
-				const FVector3f Values[] = { FVector3f::One(), ScaleDistribution.Values.Num() > 0 ? ScaleDistribution.Values[0] : FVector3f::One() };
-				BuiltData->ScaleDistribution.X = BuildContext.AddStaticData(Values);
-				BuiltData->ScaleDistribution.Y = 1;
-			}
+			const FVector3f Values[] = { FVector3f::One(), ScaleDistribution.Values.Num() > 0 ? ScaleDistribution.Values[0] : FVector3f::One() };
+			BuiltData->ScaleDistribution.X = BuildContext.AddStaticData(Values);
+			BuiltData->ScaleDistribution.Y = 1;
 		}
+		BuildContext.AddParticleSimulationExecSimulate(&UNiagaraStatelessModule_ScaleMeshSizeBySpeed::ParticleSimulate);
 	}
 
 	virtual void SetShaderParameters(const FNiagaraStatelessSetShaderParameterContext& SetShaderParameterContext) const override
@@ -57,6 +77,31 @@ public:
 		Parameters->ScaleMeshSizeBySpeed_ScaleDistribution	= ModuleBuiltData->ScaleDistribution;
 	}
 
+	static void ParticleSimulate(const NiagaraStateless::FParticleSimulationContext& ParticleSimulationContext)
+	{
+		using namespace NiagaraStateless;
+
+		const FModuleBuiltData* ModuleBuiltData = ParticleSimulationContext.ReadBuiltData<FModuleBuiltData>();
+
+		for (uint32 i = 0; i < ParticleSimulationContext.GetNumInstances(); ++i)
+		{
+			const FVector3f Position			= ParticleSimulationContext.ReadParticleVariable(ModuleBuiltData->PositionVariableOffset, i, FVector3f::OneVector);
+			const FVector3f PreviousPosition	= ParticleSimulationContext.ReadParticleVariable(ModuleBuiltData->PreviousPositionVariableOffset, i, FVector3f::OneVector);
+			const float Speed					= (Position - PreviousPosition).SquaredLength() * ParticleSimulationContext.GetInvDeltaTime();
+			const float NormSpeed				= FMath::Clamp(Speed * ModuleBuiltData->VelocityNorm, 0.0f, 1.0f);
+
+			FVector3f Scale			= ParticleSimulationContext.ReadParticleVariable(ModuleBuiltData->ScaleVariableVariableOffset, i, FVector3f::OneVector);
+			FVector3f PreviousScale	= ParticleSimulationContext.ReadParticleVariable(ModuleBuiltData->PreviousScaleVariableVariableOffset, i, FVector3f::OneVector);
+
+			const FVector3f ScaleModifier = ParticleSimulationContext.LerpStaticFloat<FVector3f>(ModuleBuiltData->ScaleDistribution, NormSpeed);
+			Scale			*= ScaleModifier;
+			PreviousScale	*= ScaleModifier;
+
+			ParticleSimulationContext.WriteParticleVariable(ModuleBuiltData->ScaleVariableVariableOffset, i, Scale);
+			ParticleSimulationContext.WriteParticleVariable(ModuleBuiltData->PreviousScaleVariableVariableOffset, i, PreviousScale);
+		}
+	}
+
 #if WITH_EDITOR
 	virtual bool CanDisableModule() const override { return true; }
 #endif
diff --git a/Engine/Plugins/FX/Niagara/Source/Niagara/Internal/Stateless/Modules/NiagaraStatelessModule_ScaleSpriteSize.h b/Engine/Plugins/FX/Niagara/Source/Niagara/Internal/Stateless/Modules/NiagaraStatelessModule_ScaleSpriteSize.h
--- a/Engine/Plugins/FX/Niagara/Source/Niagara/Internal/Stateless/Modules/NiagaraStatelessModule_ScaleSpriteSize.h
+++ b/Engine/Plugins/FX/Niagara/Source/Niagara/Internal/Stateless/Modules/NiagaraStatelessModule_ScaleSpriteSize.h
@@ -6,6 +6,7 @@
 #include "Stateless/NiagaraStatelessModule.h"
 #include "Stateless/NiagaraStatelessEmitterDataBuildContext.h"
 #include "Stateless/NiagaraStatelessModuleShaderParameters.h"
+#include "Stateless/NiagaraStatelessParticleSimContext.h"
 
 #include "NiagaraStatelessModule_ScaleSpriteSize.generated.h"
 
@@ -19,6 +20,9 @@ class UNiagaraStatelessModule_ScaleSpriteSize : public UNiagaraStatelessModule
 		FUintVector3	DistributionParameters = FUintVector3::ZeroValue;
 		FVector2f		CurveScale = FVector2f::One();
 		int32			CurveScaleOffset = INDEX_NONE;
+
+		int32			SpriteSizeVariableOffset = INDEX_NONE;
+		int32			PreviousSpriteSizeVariableOffset = INDEX_NONE;
 	};
 
 public:
@@ -44,14 +48,25 @@ public:
 	#endif
 	}
 
-	virtual void BuildEmitterData(FNiagaraStatelessEmitterDataBuildContext& BuildContext) const override
+	virtual void BuildEmitterData(const FNiagaraStatelessEmitterDataBuildContext& BuildContext) const override
 	{
 		FModuleBuiltData* BuiltData = BuildContext.AllocateBuiltData<FModuleBuiltData>();
-		BuiltData->DistributionParameters = BuildContext.AddDistribution(ScaleDistribution, IsModuleEnabled());
-		if (IsModuleEnabled() && UseScaleCurveRange())
+
+		const FNiagaraStatelessGlobals& StatelessGlobals = FNiagaraStatelessGlobals::Get();
+		BuiltData->SpriteSizeVariableOffset			= BuildContext.FindParticleVariableIndex(StatelessGlobals.SpriteSizeVariable);
+		BuiltData->PreviousSpriteSizeVariableOffset	= BuildContext.FindParticleVariableIndex(StatelessGlobals.PreviousSpriteSizeVariable);
+
+		const bool bAttributesUsed = (BuiltData->SpriteSizeVariableOffset != INDEX_NONE || BuiltData->PreviousSpriteSizeVariableOffset != INDEX_NONE);
+		if (IsModuleEnabled() && bAttributesUsed)
 		{
-			BuiltData->CurveScaleOffset = BuildContext.AddRendererBinding(ScaleCurveRange.ResolvedParameter);
-			BuiltData->CurveScale = ScaleCurveRange.GetDefaultValue<FVector2f>();
+			BuiltData->DistributionParameters = BuildContext.AddDistribution(ScaleDistribution);
+			if (UseScaleCurveRange())
+			{
+				BuiltData->CurveScaleOffset = BuildContext.AddRendererBinding(ScaleCurveRange.ResolvedParameter);
+				BuiltData->CurveScale = ScaleCurveRange.GetDefaultValue<FVector2f>();
+			}
+
+			BuildContext.AddParticleSimulationExecSimulate(&UNiagaraStatelessModule_ScaleSpriteSize::ParticleSimulate);
 		}
 	}
 
@@ -61,8 +76,34 @@ public:
 
 		FParameters* Parameters						= SetShaderParameterContext.GetParameterNestedStruct<FParameters>();
 		Parameters->ScaleSpriteSize_Distribution	= ModuleBuiltData->DistributionParameters;
-		Parameters->ScaleSpriteSize_CurveScale		= ModuleBuiltData->CurveScale;
-		Parameters->ScaleSpriteSize_CurveScaleOffset= ModuleBuiltData->CurveScaleOffset;
+		SetShaderParameterContext.GetRendererParameterValue(Parameters->ScaleSpriteSize_CurveScale, ModuleBuiltData->CurveScaleOffset, ModuleBuiltData->CurveScale);
+	}
+
+	static void ParticleSimulate(const NiagaraStateless::FParticleSimulationContext& ParticleSimulationContext)
+	{
+		using namespace NiagaraStateless;
+
+		const FModuleBuiltData* ModuleBuiltData = ParticleSimulationContext.ReadBuiltData<FModuleBuiltData>();
+
+		const FVector2f ScaleFactor				= ParticleSimulationContext.GetParameterBufferFloat(ModuleBuiltData->CurveScaleOffset, ModuleBuiltData->CurveScale);
+		const float* NormalizedAgeData			= ParticleSimulationContext.GetParticleNormalizedAge();
+		const float* PreviousNormalizedAgeData	= ParticleSimulationContext.GetParticlePreviousNormalizedAge();
+
+		for (uint32 i = 0; i < ParticleSimulationContext.GetNumInstances(); ++i)
+		{
+			const FStatelessDistributionSampler<FVector2f> SpriteScaleSampler(ParticleSimulationContext, ModuleBuiltData->DistributionParameters, i);
+
+			const float NormalizedAge			= NormalizedAgeData[i];
+			const float PreviousNormalizedAge	= PreviousNormalizedAgeData[i];
+
+			const FVector2f SpriteSize			= ParticleSimulationContext.ReadParticleVariable(ModuleBuiltData->SpriteSizeVariableOffset, i, FVector2f::ZeroVector);
+			const FVector2f PreviousSpriteSize	= ParticleSimulationContext.ReadParticleVariable(ModuleBuiltData->PreviousSpriteSizeVariableOffset, i, FVector2f::ZeroVector);
+			const FVector2f Scale				= SpriteScaleSampler.GetValue(ParticleSimulationContext, NormalizedAge) * ScaleFactor;
+			const FVector2f PreviousScale		= SpriteScaleSampler.GetValue(ParticleSimulationContext, PreviousNormalizedAge) * ScaleFactor;
+
+			ParticleSimulationContext.WriteParticleVariable(ModuleBuiltData->SpriteSizeVariableOffset, i, SpriteSize * Scale);
+			ParticleSimulationContext.WriteParticleVariable(ModuleBuiltData->PreviousSpriteSizeVariableOffset, i, PreviousSpriteSize * PreviousScale);
+		}
 	}
 
 	UFUNCTION()
diff --git a/Engine/Plugins/FX/Niagara/Source/Niagara/Internal/Stateless/Modules/NiagaraStatelessModule_ScaleSpriteSizeBySpeed.h b/Engine/Plugins/FX/Niagara/Source/Niagara/Internal/Stateless/Modules/NiagaraStatelessModule_ScaleSpriteSizeBySpeed.h
--- a/Engine/Plugins/FX/Niagara/Source/Niagara/Internal/Stateless/Modules/NiagaraStatelessModule_ScaleSpriteSizeBySpeed.h
+++ b/Engine/Plugins/FX/Niagara/Source/Niagara/Internal/Stateless/Modules/NiagaraStatelessModule_ScaleSpriteSizeBySpeed.h
@@ -5,6 +5,7 @@
 #include "Stateless/NiagaraStatelessModule.h"
 #include "Stateless/NiagaraStatelessEmitterDataBuildContext.h"
 #include "Stateless/NiagaraStatelessModuleShaderParameters.h"
+#include "Stateless/NiagaraStatelessParticleSimContext.h"
 
 #include "NiagaraStatelessModule_ScaleSpriteSizeBySpeed.generated.h"
 
@@ -17,6 +18,11 @@ class UNiagaraStatelessModule_ScaleSpriteSizeBySpeed : public UNiagaraStatelessM
 	{
 		float			VelocityNorm = 0.0f;
 		FUintVector2	ScaleDistribution = FUintVector2::ZeroValue;
+
+		int32			PositionVariableOffset = INDEX_NONE;
+		int32			PreviousPositionVariableOffset = INDEX_NONE;
+		int32			SpriteSizeVariableOffset = INDEX_NONE;
+		int32			PreviousSpriteSizeVariableOffset = INDEX_NONE;
 	};
 
 public:
@@ -28,24 +34,38 @@ public:
 	UPROPERTY(EditAnywhere, Category = "Parameters", meta = (DisplayName = "Scale", DisableRangeDistribution, DisableBindingDistribution))
 	FNiagaraDistributionVector2 ScaleDistribution = FNiagaraDistributionVector2(1.0f);
 
-	virtual void BuildEmitterData(FNiagaraStatelessEmitterDataBuildContext& BuildContext) const override
+	virtual void BuildEmitterData(const FNiagaraStatelessEmitterDataBuildContext& BuildContext) const override
 	{
 		FModuleBuiltData* BuiltData = BuildContext.AllocateBuiltData<FModuleBuiltData>();
-		if (IsModuleEnabled())
+		if (!IsModuleEnabled())
+		{
+			return;
+		}
+
+		const FNiagaraStatelessGlobals& StatelessGlobals = FNiagaraStatelessGlobals::Get();
+		BuiltData->PositionVariableOffset			= BuildContext.FindParticleVariableIndex(StatelessGlobals.PositionVariable);
+		BuiltData->PreviousPositionVariableOffset	= BuildContext.FindParticleVariableIndex(StatelessGlobals.PreviousPositionVariable);
+		BuiltData->SpriteSizeVariableOffset			= BuildContext.FindParticleVariableIndex(StatelessGlobals.SpriteSizeVariable);
+		BuiltData->PreviousSpriteSizeVariableOffset	= BuildContext.FindParticleVariableIndex(StatelessGlobals.PreviousSpriteSizeVariable);
+
+		if (BuiltData->SpriteSizeVariableOffset == INDEX_NONE && BuiltData->PreviousSpriteSizeVariableOffset == INDEX_NONE)
+		{
+			return;
+		}
+
+		BuiltData->VelocityNorm			= VelocityThreshold > 0.0f ? 1.0f / (VelocityThreshold * VelocityThreshold) : 0.0f;
+		if (ScaleDistribution.IsCurve() && ScaleDistribution.Values.Num() > 1)
+		{
+			BuiltData->ScaleDistribution.X = BuildContext.AddStaticData(ScaleDistribution.Values);
+			BuiltData->ScaleDistribution.Y = ScaleDistribution.Values.Num() - 1;
+		}
+		else
 		{
-			BuiltData->VelocityNorm			= VelocityThreshold > 0.0f ? 1.0f / (VelocityThreshold * VelocityThreshold) : 0.0f;
-			if (ScaleDistribution.IsCurve() && ScaleDistribution.Values.Num() > 1)
-			{
-				BuiltData->ScaleDistribution.X = BuildContext.AddStaticData(ScaleDistribution.Values);
-				BuiltData->ScaleDistribution.Y = ScaleDistribution.Values.Num() - 1;
-			}
-			else
-			{
-				const FVector2f Values[] = { FVector2f::One(), ScaleDistribution.Values.Num() > 0 ? ScaleDistribution.Values[0] : FVector2f::One() };
-				BuiltData->ScaleDistribution.X = BuildContext.AddStaticData(Values);
-				BuiltData->ScaleDistribution.Y = 1;
-			}
+			const FVector2f Values[] = { FVector2f::One(), ScaleDistribution.Values.Num() > 0 ? ScaleDistribution.Values[0] : FVector2f::One() };
+			BuiltData->ScaleDistribution.X = BuildContext.AddStaticData(Values);
+			BuiltData->ScaleDistribution.Y = 1;
 		}
+		BuildContext.AddParticleSimulationExecSimulate(&UNiagaraStatelessModule_ScaleSpriteSizeBySpeed::ParticleSimulate);
 	}
 
 	virtual void SetShaderParameters(const FNiagaraStatelessSetShaderParameterContext& SetShaderParameterContext) const override
@@ -57,6 +77,30 @@ public:
 		Parameters->ScaleSpriteSizeBySpeed_ScaleDistribution	= ModuleBuiltData->ScaleDistribution;
 	}
 
+	static void ParticleSimulate(const NiagaraStateless::FParticleSimulationContext& ParticleSimulationContext)
+	{
+		using namespace NiagaraStateless;
+
+		const FModuleBuiltData* ModuleBuiltData = ParticleSimulationContext.ReadBuiltData<FModuleBuiltData>();
+
+		for (uint32 i = 0; i < ParticleSimulationContext.GetNumInstances(); ++i)
+		{
+			const FVector3f Position			= ParticleSimulationContext.ReadParticleVariable(ModuleBuiltData->PositionVariableOffset, i, FVector3f::OneVector);
+			const FVector3f PreviousPosition	= ParticleSimulationContext.ReadParticleVariable(ModuleBuiltData->PreviousPositionVariableOffset, i, FVector3f::OneVector);
+			const float Speed					= (Position - PreviousPosition).SquaredLength() * ParticleSimulationContext.GetInvDeltaTime();
+			const float NormSpeed				= FMath::Clamp(Speed * ModuleBuiltData->VelocityNorm, 0.0f, 1.0f);
+
+			FVector2f SpriteSize			= ParticleSimulationContext.ReadParticleVariable(ModuleBuiltData->SpriteSizeVariableOffset, i, FVector2f::One());
+			FVector2f PreviousSpriteSize	= ParticleSimulationContext.ReadParticleVariable(ModuleBuiltData->PreviousSpriteSizeVariableOffset, i, FVector2f::One());
+
+			const FVector2f ScaleModifier = ParticleSimulationContext.LerpStaticFloat<FVector2f>(ModuleBuiltData->ScaleDistribution, NormSpeed);
+			SpriteSize			*= ScaleModifier;
+			PreviousSpriteSize	*= ScaleModifier;
+
+			ParticleSimulationContext.WriteParticleVariable(ModuleBuiltData->SpriteSizeVariableOffset, i, SpriteSize);
+			ParticleSimulationContext.WriteParticleVariable(ModuleBuiltData->PreviousSpriteSizeVariableOffset, i, PreviousSpriteSize);
+		}
+	}
 #if WITH_EDITOR
 	virtual bool CanDisableModule() const override { return true; }
 #endif
diff --git a/Engine/Plugins/FX/Niagara/Source/Niagara/Internal/Stateless/Modules/NiagaraStatelessModule_ShapeLocation.h b/Engine/Plugins/FX/Niagara/Source/Niagara/Internal/Stateless/Modules/NiagaraStatelessModule_ShapeLocation.h
--- a/Engine/Plugins/FX/Niagara/Source/Niagara/Internal/Stateless/Modules/NiagaraStatelessModule_ShapeLocation.h
+++ b/Engine/Plugins/FX/Niagara/Source/Niagara/Internal/Stateless/Modules/NiagaraStatelessModule_ShapeLocation.h
@@ -4,6 +4,7 @@
 
 #include "Stateless/NiagaraStatelessModule.h"
 #include "Stateless/NiagaraStatelessModuleShaderParameters.h"
+#include "Stateless/NiagaraStatelessParticleSimContext.h"
 
 #include "NiagaraStatelessModule_ShapeLocation.generated.h"
 
@@ -59,6 +60,7 @@ public:
 	UPROPERTY(EditAnywhere, Category = "Parameters", meta=(EditConditionHides, EditCondition="ShapePrimitive == ENSM_ShapePrimitive::Sphere"))
 	float SphereMax = 100.0f;
 
+	virtual void BuildEmitterData(const FNiagaraStatelessEmitterDataBuildContext& BuildContext) const override;
 	virtual void SetShaderParameters(const FNiagaraStatelessSetShaderParameterContext& SetShaderParameterContext) const override;
 
 #if WITH_EDITOR
diff --git a/Engine/Plugins/FX/Niagara/Source/Niagara/Internal/Stateless/Modules/NiagaraStatelessModule_SolveVelocitiesAndForces.h b/Engine/Plugins/FX/Niagara/Source/Niagara/Internal/Stateless/Modules/NiagaraStatelessModule_SolveVelocitiesAndForces.h
--- a/Engine/Plugins/FX/Niagara/Source/Niagara/Internal/Stateless/Modules/NiagaraStatelessModule_SolveVelocitiesAndForces.h
+++ b/Engine/Plugins/FX/Niagara/Source/Niagara/Internal/Stateless/Modules/NiagaraStatelessModule_SolveVelocitiesAndForces.h
@@ -18,38 +18,43 @@ class UNiagaraStatelessModule_SolveVelocitiesAndForces : public UNiagaraStateles
 {
 	GENERATED_BODY()
 
+	struct FModuleBuiltData
+	{
+		NiagaraStateless::FPhysicsBuildData	PhysicsData;
+		int32								PositionVariableOffset = INDEX_NONE;
+		int32								VelocityVariableOffset = INDEX_NONE;
+		int32								PreviousPositionVariableOffset = INDEX_NONE;
+		int32								PreviousVelocityVariableOffset = INDEX_NONE;
+	};
+
 public:
-	using FModuleBuiltData = NiagaraStateless::FPhysicsBuildData;
 	using FParameters = NiagaraStateless::FSolveVelocitiesAndForcesModule_ShaderParameters;
 
-	virtual void BuildEmitterData(FNiagaraStatelessEmitterDataBuildContext& BuildContext) const override
+	virtual void BuildEmitterData(const FNiagaraStatelessEmitterDataBuildContext& BuildContext) const override
 	{
 		NiagaraStateless::FPhysicsBuildData& PhysicsBuildData = BuildContext.GetTransientBuildData<NiagaraStateless::FPhysicsBuildData>();
 
-		FModuleBuiltData* BuiltData		= BuildContext.AllocateBuiltData<FModuleBuiltData>();
-		BuiltData->MassRange			= PhysicsBuildData.MassRange;
-		BuiltData->DragRange.Min		= FMath::Max(PhysicsBuildData.DragRange.Min, 0.01f);
-		BuiltData->DragRange.Max		= FMath::Max(PhysicsBuildData.DragRange.Max, 0.01f);
-		BuiltData->VelocityRange		= PhysicsBuildData.VelocityRange;
-		BuiltData->WindRange			= PhysicsBuildData.WindRange;
-		BuiltData->AccelerationRange	= PhysicsBuildData.AccelerationRange;
-		BuiltData->bConeVelocity		= PhysicsBuildData.bConeVelocity;
-		BuiltData->ConeQuat				= PhysicsBuildData.ConeQuat;
-		BuiltData->ConeVelocityRange	= PhysicsBuildData.ConeVelocityRange;
-		BuiltData->ConeOuterAngle		= PhysicsBuildData.ConeOuterAngle;
-		BuiltData->ConeInnerAngle		= PhysicsBuildData.ConeInnerAngle;
-		BuiltData->ConeVelocityFalloff	= PhysicsBuildData.ConeVelocityFalloff;
-		BuiltData->bPointVelocity		= PhysicsBuildData.bPointVelocity;
-		BuiltData->PointVelocityRange	= PhysicsBuildData.PointVelocityRange;
-		BuiltData->PointOrigin			= PhysicsBuildData.PointOrigin;
-		BuiltData->bNoiseEnabled		= PhysicsBuildData.bNoiseEnabled;
-		BuiltData->NoiseAmplitude		= PhysicsBuildData.NoiseAmplitude;
-		BuiltData->NoiseFrequency		= PhysicsBuildData.NoiseFrequency;
-		BuiltData->NoiseTexture			= PhysicsBuildData.NoiseTexture;
-		BuiltData->NoiseMode			= PhysicsBuildData.NoiseMode;
-		BuiltData->NoiseLUTOffset		= PhysicsBuildData.NoiseLUTOffset;
-		BuiltData->NoiseLUTNumChannel	= PhysicsBuildData.NoiseLUTNumChannel;
-		BuiltData->NoiseLUTChannelWidth	= PhysicsBuildData.NoiseLUTChannelWidth;
+		FModuleBuiltData* BuiltData				= BuildContext.AllocateBuiltData<FModuleBuiltData>();
+		BuiltData->PhysicsData					= PhysicsBuildData;
+		BuiltData->PhysicsData.DragRange.Min	= FMath::Max(PhysicsBuildData.DragRange.Min, 0.01f);
+		BuiltData->PhysicsData.DragRange.Max	= FMath::Max(PhysicsBuildData.DragRange.Max, 0.01f);
+
+		const FNiagaraStatelessGlobals& StatelessGlobals = FNiagaraStatelessGlobals::Get();
+		BuiltData->PositionVariableOffset				= BuildContext.FindParticleVariableIndex(StatelessGlobals.PositionVariable);
+		BuiltData->VelocityVariableOffset				= BuildContext.FindParticleVariableIndex(StatelessGlobals.VelocityVariable);
+		BuiltData->PreviousPositionVariableOffset		= BuildContext.FindParticleVariableIndex(StatelessGlobals.PreviousPositionVariable);
+		BuiltData->PreviousVelocityVariableOffset		= BuildContext.FindParticleVariableIndex(StatelessGlobals.PreviousVelocityVariable);
+		
+		const bool bAttributesUsed =
+			BuiltData->PositionVariableOffset != INDEX_NONE ||
+			BuiltData->VelocityVariableOffset != INDEX_NONE ||
+			BuiltData->PreviousPositionVariableOffset != INDEX_NONE ||
+			BuiltData->PreviousVelocityVariableOffset != INDEX_NONE;
+
+		if (bAttributesUsed)
+		{
+			BuildContext.AddParticleSimulationExecSimulate(&UNiagaraStatelessModule_SolveVelocitiesAndForces::ParticleSimulate);
+		}
 	}
 
 	virtual void SetShaderParameters(const FNiagaraStatelessSetShaderParameterContext& SetShaderParameterContext) const override
@@ -57,40 +62,40 @@ public:
 		const FModuleBuiltData* ModuleBuiltData = SetShaderParameterContext.ReadBuiltData<FModuleBuiltData>();
 
 		FParameters* Parameters = SetShaderParameterContext.GetParameterNestedStruct<FParameters>();
-		Parameters->SolveVelocitiesAndForces_MassScale				= ModuleBuiltData->MassRange.GetScale();
-		Parameters->SolveVelocitiesAndForces_MassBias				= ModuleBuiltData->MassRange.Min;
-		Parameters->SolveVelocitiesAndForces_DragScale				= ModuleBuiltData->DragRange.GetScale();
-		Parameters->SolveVelocitiesAndForces_DragBias				= ModuleBuiltData->DragRange.Min;
-		Parameters->SolveVelocitiesAndForces_VelocityScale			= ModuleBuiltData->VelocityRange.GetScale();
-		Parameters->SolveVelocitiesAndForces_VelocityBias			= ModuleBuiltData->VelocityRange.Min;
-		Parameters->SolveVelocitiesAndForces_WindScale				= ModuleBuiltData->WindRange.GetScale();
-		Parameters->SolveVelocitiesAndForces_WindBias				= ModuleBuiltData->WindRange.Min;
-		Parameters->SolveVelocitiesAndForces_AccelerationScale		= ModuleBuiltData->AccelerationRange.GetScale();
-		Parameters->SolveVelocitiesAndForces_AccelerationBias		= ModuleBuiltData->AccelerationRange.Min;
-
-		Parameters->SolveVelocitiesAndForces_ConeVelocityEnabled	= ModuleBuiltData->bConeVelocity ? 1 : 0;
-		Parameters->SolveVelocitiesAndForces_ConeQuat				= ModuleBuiltData->ConeQuat;
-		Parameters->SolveVelocitiesAndForces_ConeVelocityScale		= ModuleBuiltData->ConeVelocityRange.GetScale();
-		Parameters->SolveVelocitiesAndForces_ConeVelocityBias		= ModuleBuiltData->ConeVelocityRange.Min;
-		Parameters->SolveVelocitiesAndForces_ConeAngleScale			= (ModuleBuiltData->ConeOuterAngle - ModuleBuiltData->ConeInnerAngle) * (UE_PI / 360.0f);
-		Parameters->SolveVelocitiesAndForces_ConeAngleBias			= ModuleBuiltData->ConeInnerAngle * (UE_PI / 360.0f);
-		Parameters->SolveVelocitiesAndForces_ConeVelocityFalloff	= ModuleBuiltData->ConeVelocityFalloff;
-
-		Parameters->SolveVelocitiesAndForces_PontVelocityEnabled	= ModuleBuiltData->bPointVelocity ? 1 : 0;
-		Parameters->SolveVelocitiesAndForces_PointVelocityScale		= ModuleBuiltData->PointVelocityRange.GetScale();
-		Parameters->SolveVelocitiesAndForces_PointVelocityBias		= ModuleBuiltData->PointVelocityRange.Min;
-		Parameters->SolveVelocitiesAndForces_PointOrigin			= ModuleBuiltData->PointOrigin;
-
-		Parameters->SolveVelocitiesAndForces_NoiseEnabled			= ModuleBuiltData->bNoiseEnabled ? 1 : 0;
-		Parameters->SolveVelocitiesAndForces_NoiseAmplitude			= ModuleBuiltData->NoiseAmplitude;
-		Parameters->SolveVelocitiesAndForces_NoiseFrequency			= FVector3f(ModuleBuiltData->NoiseFrequency, ModuleBuiltData->NoiseFrequency, ModuleBuiltData->NoiseFrequency);
+		Parameters->SolveVelocitiesAndForces_MassScale				= ModuleBuiltData->PhysicsData.MassRange.GetScale();
+		Parameters->SolveVelocitiesAndForces_MassBias				= ModuleBuiltData->PhysicsData.MassRange.Min;
+		Parameters->SolveVelocitiesAndForces_DragScale				= ModuleBuiltData->PhysicsData.DragRange.GetScale();
+		Parameters->SolveVelocitiesAndForces_DragBias				= ModuleBuiltData->PhysicsData.DragRange.Min;
+		Parameters->SolveVelocitiesAndForces_VelocityScale			= ModuleBuiltData->PhysicsData.VelocityRange.GetScale();
+		Parameters->SolveVelocitiesAndForces_VelocityBias			= ModuleBuiltData->PhysicsData.VelocityRange.Min;
+		Parameters->SolveVelocitiesAndForces_WindScale				= ModuleBuiltData->PhysicsData.WindRange.GetScale();
+		Parameters->SolveVelocitiesAndForces_WindBias				= ModuleBuiltData->PhysicsData.WindRange.Min;
+		Parameters->SolveVelocitiesAndForces_AccelerationScale		= ModuleBuiltData->PhysicsData.AccelerationRange.GetScale();
+		Parameters->SolveVelocitiesAndForces_AccelerationBias		= ModuleBuiltData->PhysicsData.AccelerationRange.Min;
+
+		Parameters->SolveVelocitiesAndForces_ConeVelocityEnabled	= ModuleBuiltData->PhysicsData.bConeVelocity ? 1 : 0;
+		Parameters->SolveVelocitiesAndForces_ConeQuat				= ModuleBuiltData->PhysicsData.ConeQuat;
+		Parameters->SolveVelocitiesAndForces_ConeVelocityScale		= ModuleBuiltData->PhysicsData.ConeVelocityRange.GetScale();
+		Parameters->SolveVelocitiesAndForces_ConeVelocityBias		= ModuleBuiltData->PhysicsData.ConeVelocityRange.Min;
+		Parameters->SolveVelocitiesAndForces_ConeAngleScale			= (ModuleBuiltData->PhysicsData.ConeOuterAngle - ModuleBuiltData->PhysicsData.ConeInnerAngle) * (UE_PI / 360.0f);
+		Parameters->SolveVelocitiesAndForces_ConeAngleBias			= ModuleBuiltData->PhysicsData.ConeInnerAngle * (UE_PI / 360.0f);
+		Parameters->SolveVelocitiesAndForces_ConeVelocityFalloff	= ModuleBuiltData->PhysicsData.ConeVelocityFalloff;
+
+		Parameters->SolveVelocitiesAndForces_PontVelocityEnabled	= ModuleBuiltData->PhysicsData.bPointVelocity ? 1 : 0;
+		Parameters->SolveVelocitiesAndForces_PointVelocityScale		= ModuleBuiltData->PhysicsData.PointVelocityRange.GetScale();
+		Parameters->SolveVelocitiesAndForces_PointVelocityBias		= ModuleBuiltData->PhysicsData.PointVelocityRange.Min;
+		Parameters->SolveVelocitiesAndForces_PointOrigin			= ModuleBuiltData->PhysicsData.PointOrigin;
+
+		Parameters->SolveVelocitiesAndForces_NoiseEnabled			= ModuleBuiltData->PhysicsData.bNoiseEnabled ? 1 : 0;
+		Parameters->SolveVelocitiesAndForces_NoiseAmplitude			= ModuleBuiltData->PhysicsData.NoiseAmplitude;
+		Parameters->SolveVelocitiesAndForces_NoiseFrequency			= FVector3f(ModuleBuiltData->PhysicsData.NoiseFrequency, ModuleBuiltData->PhysicsData.NoiseFrequency, ModuleBuiltData->PhysicsData.NoiseFrequency);
 		//SetShaderParameterContext.SetTextureResource(&Parameters->SolveVelocitiesAndForces_NoiseTexture, ModuleBuildData->NoiseTexture);
-		Parameters->SolveVelocitiesAndForces_NoiseMode				= ModuleBuiltData->NoiseMode;
-		Parameters->SolveVelocitiesAndForces_NoiseLUTOffset			= ModuleBuiltData->NoiseLUTOffset;
-		Parameters->SolveVelocitiesAndForces_NoiseLUTNumChannel		= ModuleBuiltData->NoiseLUTNumChannel;
-		Parameters->SolveVelocitiesAndForces_NoiseLUTChannelWidth	= ModuleBuiltData->NoiseLUTChannelWidth;
+		Parameters->SolveVelocitiesAndForces_NoiseMode				= ModuleBuiltData->PhysicsData.NoiseMode;
+		Parameters->SolveVelocitiesAndForces_NoiseLUTOffset			= ModuleBuiltData->PhysicsData.NoiseLUTOffset;
+		Parameters->SolveVelocitiesAndForces_NoiseLUTNumChannel		= ModuleBuiltData->PhysicsData.NoiseLUTNumChannel;
+		Parameters->SolveVelocitiesAndForces_NoiseLUTChannelWidth	= ModuleBuiltData->PhysicsData.NoiseLUTChannelWidth;
 
-		FVectorFieldTextureAccessor TextureAccessor(Cast<UVectorField>(ModuleBuiltData->NoiseTexture));
+		FVectorFieldTextureAccessor TextureAccessor(Cast<UVectorField>(ModuleBuiltData->PhysicsData.NoiseTexture));
 
 		ENQUEUE_RENDER_COMMAND(FNaughtyTest)(
 			[Parameters, TextureAccessor](FRHICommandListImmediate& RHICmdList)
@@ -108,25 +113,79 @@ public:
 				}
 			}
 		);
+	}
+
+	static FVector3f IntegratePosition(float Age, float Mass, float Drag, const FVector3f& Velocity, const FVector3f& Wind, const FVector3f& Acceleration)
+	{
+		const FVector3f IntVelocity = (Velocity - Wind) + (Wind * Age * Age);
+		const float LambdaDragMass = FMath::Max(Drag * (1.0f / Mass), 0.0001f);
+		const float LambdaAge = (1.0f - FMath::Exp(-(LambdaDragMass * Age))) / LambdaDragMass;
+		FVector3f Position = IntVelocity * LambdaAge;
+		Position += (Age - LambdaAge) * (Acceleration / LambdaDragMass);
+		return Position;
+	}
+
+	static void ParticleSimulate(const NiagaraStateless::FParticleSimulationContext& ParticleSimulationContext)
+	{
+		const FModuleBuiltData* ModuleBuiltData = ParticleSimulationContext.ReadBuiltData<FModuleBuiltData>();
+
+		const float* NormalizedAgeData			= ParticleSimulationContext.GetParticleNormalizedAge();
+		const float* PreviousNormalizedAgeData	= ParticleSimulationContext.GetParticlePreviousNormalizedAge();
+
+		for (uint32 i = 0; i < ParticleSimulationContext.GetNumInstances(); ++i)
+		{
+			const float Mass				= ParticleSimulationContext.RandomScaleBiasFloat(i, ModuleBuiltData->PhysicsData.MassRange);
+			const float Drag				= ParticleSimulationContext.RandomScaleBiasFloat(i, ModuleBuiltData->PhysicsData.DragRange);
+			FVector3f InitialVelocity		= ParticleSimulationContext.RandomScaleBiasFloat(i, ModuleBuiltData->PhysicsData.VelocityRange);
+			const FVector3f Wind			= ParticleSimulationContext.RandomScaleBiasFloat(i, ModuleBuiltData->PhysicsData.WindRange);
+			const FVector3f Acceleration	= ParticleSimulationContext.RandomScaleBiasFloat(i, ModuleBuiltData->PhysicsData.AccelerationRange);
+
+			FVector3f Position			= ParticleSimulationContext.ReadParticleVariable(ModuleBuiltData->PositionVariableOffset, i, FVector3f::ZeroVector);
+			FVector3f PreviousPosition	= ParticleSimulationContext.ReadParticleVariable(ModuleBuiltData->PreviousPositionVariableOffset, i, FVector3f::ZeroVector);
+
+			if (ModuleBuiltData->PhysicsData.bConeVelocity)
+			{
+				const float ConeAngle = ParticleSimulationContext.RandomScaleBiasFloat(i, FNiagaraStatelessRangeFloat(ModuleBuiltData->PhysicsData.ConeInnerAngle, ModuleBuiltData->PhysicsData.ConeOuterAngle)) * (UE_PI / 360.0f);
+				const float ConeRotation = ParticleSimulationContext.RandomFloat(i) * UE_TWO_PI;
+				const FVector2f scAng = FVector2f(FMath::Sin(ConeAngle), FMath::Cos(ConeAngle));
+				const FVector2f scRot = FVector2f(FMath::Sin(ConeRotation), FMath::Cos(ConeRotation));
+				const FVector3f Direction = FVector3f(scRot.X * scAng.X, scRot.Y * scAng.X, scAng.Y);
+
+				float VelocityScale = ParticleSimulationContext.RandomScaleBiasFloat(i, ModuleBuiltData->PhysicsData.ConeVelocityRange);
+				if (ModuleBuiltData->PhysicsData.bConeVelocity)
+				{
+					const float pf = FMath::Pow(FMath::Clamp(scAng.Y, 0.0f, 1.0f), ModuleBuiltData->PhysicsData.ConeVelocityFalloff * 10.0f);
+					VelocityScale *= FMath::Lerp(1.0f, pf, ModuleBuiltData->PhysicsData.ConeVelocityFalloff);
+				}
+
+				InitialVelocity += ModuleBuiltData->PhysicsData.ConeQuat.RotateVector(Direction) * VelocityScale;
+			}
+
+			if (ModuleBuiltData->PhysicsData.bPointVelocity)
+			{
+				const FVector3f FallbackDir	= ParticleSimulationContext.RandomUnitFloat3(i);
+				const FVector3f Delta		= Position - ModuleBuiltData->PhysicsData.PointOrigin;
+				const FVector3f Direction	= ParticleSimulationContext.SafeNormalize(Delta, FallbackDir);
+				const float		VelocityScale = ParticleSimulationContext.RandomScaleBiasFloat(i, ModuleBuiltData->PhysicsData.PointVelocityRange);
+
+				InitialVelocity += Direction * VelocityScale;
+			}
+
+			if (ModuleBuiltData->PhysicsData.bNoiseEnabled)
+			{
+				//-TODO:
+			}
+
+			Position += IntegratePosition(NormalizedAgeData[i], Mass, Drag, InitialVelocity, Wind, Acceleration);
+			PreviousPosition += IntegratePosition(PreviousNormalizedAgeData[i], Mass, Drag, InitialVelocity, Wind, Acceleration);
+
+			ParticleSimulationContext.WriteParticleVariable(ModuleBuiltData->PositionVariableOffset, i, Position);
+			ParticleSimulationContext.WriteParticleVariable(ModuleBuiltData->PreviousPositionVariableOffset, i, PreviousPosition);
 
-		//ModuleBuildData->NoiseTexture
-		//SetShaderParameterContext.AddSetResource(
-		//	[]()
-		//	{
-		//	}
-		//);
-
-		//FRHITexture**		TextureRef = &Parameters->SolveVelocitiesAndForces_NoiseTexture;
-		//FRHISamplerState**	SamplerRef = &Parameters->SolveVelocitiesAndForces_NoiseSampler;
-		//FTextureRHIRef
-		//if (ModuleBuiltData->NoiseTexture)
-		//{
-		//}
-		//SetShaderParameterContext.AddRenderThreadSet(
-		//	[]()
-		//	{
-		//	}
-		//);
+			const FVector3f Velocity = Position - PreviousPosition;
+			ParticleSimulationContext.WriteParticleVariable(ModuleBuiltData->VelocityVariableOffset, i, Velocity);
+			ParticleSimulationContext.WriteParticleVariable(ModuleBuiltData->PreviousVelocityVariableOffset, i, Velocity);
+		}
 	}
 
 #if WITH_EDITORONLY_DATA
diff --git a/Engine/Plugins/FX/Niagara/Source/Niagara/Internal/Stateless/Modules/NiagaraStatelessModule_SpriteFacingAndAlignment.h b/Engine/Plugins/FX/Niagara/Source/Niagara/Internal/Stateless/Modules/NiagaraStatelessModule_SpriteFacingAndAlignment.h
--- a/Engine/Plugins/FX/Niagara/Source/Niagara/Internal/Stateless/Modules/NiagaraStatelessModule_SpriteFacingAndAlignment.h
+++ b/Engine/Plugins/FX/Niagara/Source/Niagara/Internal/Stateless/Modules/NiagaraStatelessModule_SpriteFacingAndAlignment.h
@@ -5,6 +5,7 @@
 #include "Stateless/NiagaraStatelessModule.h"
 #include "Stateless/NiagaraStatelessEmitterDataBuildContext.h"
 #include "Stateless/NiagaraStatelessModuleShaderParameters.h"
+#include "Stateless/NiagaraStatelessParticleSimContext.h"
 
 #include "NiagaraStatelessModule_SpriteFacingAndAlignment.generated.h"
 
@@ -13,6 +14,17 @@ class UNiagaraStatelessModule_SpriteFacingAndAlignment : public UNiagaraStateles
 {
 	GENERATED_BODY()
 
+	struct FModuleBuiltData
+	{
+		FVector3f	SpriteFacing	= FVector3f::XAxisVector;
+		FVector3f	SpriteAlignment = FVector3f::YAxisVector;
+
+		int32		SpriteFacingVariableOffset = INDEX_NONE;
+		int32		PreviousSpriteFacingVariableOffset = INDEX_NONE;
+		int32		SpriteAlignmentVariable = INDEX_NONE;
+		int32		PreviousSpriteAlignmentVariable = INDEX_NONE;
+	};
+
 public:
 	using FParameters = NiagaraStateless::FSpriteFacingAndAlignmentModule_ShaderParameters;
 
@@ -28,11 +40,61 @@ public:
 	UPROPERTY(EditAnywhere, Category = "Parameters", meta = (EditCondition = "bSpriteAlignmentEnabled"))
 	FVector3f	SpriteAlignment = FVector3f::YAxisVector;
 
+	virtual void BuildEmitterData(const FNiagaraStatelessEmitterDataBuildContext& BuildContext) const override
+	{
+		FModuleBuiltData* BuiltData = BuildContext.AllocateBuiltData<FModuleBuiltData>();
+		if (!IsModuleEnabled())
+		{
+			return;
+		}
+
+		const FNiagaraStatelessGlobals& StatelessGlobals = FNiagaraStatelessGlobals::Get();
+		if (bSpriteFacingEnabled)
+		{
+			BuiltData->SpriteFacingVariableOffset			= BuildContext.FindParticleVariableIndex(StatelessGlobals.SpriteFacingVariable);
+			BuiltData->PreviousSpriteFacingVariableOffset	= BuildContext.FindParticleVariableIndex(StatelessGlobals.PreviousSpriteFacingVariable);
+		}
+		if (bSpriteAlignmentEnabled)
+		{
+			BuiltData->SpriteAlignmentVariable				= BuildContext.FindParticleVariableIndex(StatelessGlobals.SpriteAlignmentVariable);
+			BuiltData->PreviousSpriteAlignmentVariable		= BuildContext.FindParticleVariableIndex(StatelessGlobals.PreviousSpriteAlignmentVariable);
+		}
+
+		if ((BuiltData->SpriteFacingVariableOffset == INDEX_NONE) && (BuiltData->PreviousSpriteFacingVariableOffset == INDEX_NONE) &&
+			(BuiltData->SpriteAlignmentVariable == INDEX_NONE) && (BuiltData->PreviousSpriteAlignmentVariable == INDEX_NONE))
+		{
+			return;
+		}
+
+		BuiltData->SpriteFacing		= SpriteFacing;
+		BuiltData->SpriteAlignment	= SpriteAlignment;
+
+		BuildContext.AddParticleSimulationExecSimulate(&UNiagaraStatelessModule_SpriteFacingAndAlignment::ParticleSimulate);
+	}
+
 	virtual void SetShaderParameters(const FNiagaraStatelessSetShaderParameterContext& SetShaderParameterContext) const override
 	{
 		FParameters* Parameters = SetShaderParameterContext.GetParameterNestedStruct<FParameters>();
-		Parameters->SpriteFacingAndAlignment_SpriteFacing		= SpriteFacing;
-		Parameters->SpriteFacingAndAlignment_SpriteAlignment	= SpriteAlignment;
+		const FModuleBuiltData* ModuleBuiltData = SetShaderParameterContext.ReadBuiltData<FModuleBuiltData>();
+
+		Parameters->SpriteFacingAndAlignment_SpriteFacing		= ModuleBuiltData->SpriteFacing;
+		Parameters->SpriteFacingAndAlignment_SpriteAlignment	= ModuleBuiltData->SpriteAlignment;
+	}
+
+	static void ParticleSimulate(const NiagaraStateless::FParticleSimulationContext& ParticleSimulationContext)
+	{
+		using namespace NiagaraStateless;
+
+		const FModuleBuiltData* ModuleBuiltData = ParticleSimulationContext.ReadBuiltData<FModuleBuiltData>();
+
+		for (uint32 i = 0; i < ParticleSimulationContext.GetNumInstances(); ++i)
+		{
+			ParticleSimulationContext.WriteParticleVariable(ModuleBuiltData->SpriteFacingVariableOffset, i, ModuleBuiltData->SpriteFacing);
+			ParticleSimulationContext.WriteParticleVariable(ModuleBuiltData->PreviousSpriteFacingVariableOffset, i, ModuleBuiltData->SpriteFacing);
+
+			ParticleSimulationContext.WriteParticleVariable(ModuleBuiltData->SpriteAlignmentVariable, i, ModuleBuiltData->SpriteAlignment);
+			ParticleSimulationContext.WriteParticleVariable(ModuleBuiltData->PreviousSpriteAlignmentVariable, i, ModuleBuiltData->SpriteAlignment);
+		}
 	}
 
 #if WITH_EDITOR
diff --git a/Engine/Plugins/FX/Niagara/Source/Niagara/Internal/Stateless/Modules/NiagaraStatelessModule_SpriteRotationRate.h b/Engine/Plugins/FX/Niagara/Source/Niagara/Internal/Stateless/Modules/NiagaraStatelessModule_SpriteRotationRate.h
--- a/Engine/Plugins/FX/Niagara/Source/Niagara/Internal/Stateless/Modules/NiagaraStatelessModule_SpriteRotationRate.h
+++ b/Engine/Plugins/FX/Niagara/Source/Niagara/Internal/Stateless/Modules/NiagaraStatelessModule_SpriteRotationRate.h
@@ -5,6 +5,7 @@
 #include "Stateless/NiagaraStatelessModule.h"
 #include "Stateless/NiagaraStatelessEmitterDataBuildContext.h"
 #include "Stateless/NiagaraStatelessModuleShaderParameters.h"
+#include "Stateless/NiagaraStatelessParticleSimContext.h"
 
 #include "NiagaraStatelessModule_SpriteRotationRate.generated.h"
 
@@ -16,6 +17,8 @@ class UNiagaraStatelessModule_SpriteRotationRate : public UNiagaraStatelessModul
 	struct FModuleBuiltData
 	{
 		FNiagaraStatelessRangeFloat		RotationRange;
+		int32							SpriteRotationVariableOffset = INDEX_NONE;
+		int32							PreviousSpriteRotationVariableOffset = INDEX_NONE;
 	};
 
 public:
@@ -24,10 +27,26 @@ public:
 	UPROPERTY(EditAnywhere, Category = "Parameters", meta = (DisplayName = "Rotation Rate"))
 	FNiagaraDistributionRangeFloat RotationRateDistribution = FNiagaraDistributionRangeFloat(FNiagaraStatelessGlobals::GetDefaultSpriteRotationValue());
 
-	virtual void BuildEmitterData(FNiagaraStatelessEmitterDataBuildContext& BuildContext) const override
+	virtual void BuildEmitterData(const FNiagaraStatelessEmitterDataBuildContext& BuildContext) const override
 	{
 		FModuleBuiltData* BuiltData = BuildContext.AllocateBuiltData<FModuleBuiltData>();
-		BuiltData->RotationRange = BuildContext.ConvertDistributionToRange(RotationRateDistribution, 0.0f, IsModuleEnabled());
+		if (!IsModuleEnabled())
+		{
+			return;
+		}
+
+		const FNiagaraStatelessGlobals& StatelessGlobals = FNiagaraStatelessGlobals::Get();
+		BuiltData->SpriteRotationVariableOffset			= BuildContext.FindParticleVariableIndex(StatelessGlobals.SpriteRotationVariable);
+		BuiltData->PreviousSpriteRotationVariableOffset	= BuildContext.FindParticleVariableIndex(StatelessGlobals.PreviousSpriteRotationVariable);
+
+		if (BuiltData->SpriteRotationVariableOffset == INDEX_NONE && BuiltData->PreviousSpriteRotationVariableOffset == INDEX_NONE)
+		{
+			return;
+		}
+		
+		BuiltData->RotationRange = BuildContext.ConvertDistributionToRange(RotationRateDistribution, 0.0f);
+
+		BuildContext.AddParticleSimulationExecSimulate(&UNiagaraStatelessModule_SpriteRotationRate::ParticleSimulate);
 	}
 
 	virtual void SetShaderParameters(const FNiagaraStatelessSetShaderParameterContext& SetShaderParameterContext) const override
@@ -38,6 +57,29 @@ public:
 		SetShaderParameterContext.ConvertRangeToScaleBias(ModuleBuiltData->RotationRange, Parameters->SpriteRotationRate_Scale, Parameters->SpriteRotationRate_Bias);
 	}
 
+	static void ParticleSimulate(const NiagaraStateless::FParticleSimulationContext& ParticleSimulationContext)
+	{
+		using namespace NiagaraStateless;
+
+		const FModuleBuiltData* ModuleBuiltData = ParticleSimulationContext.ReadBuiltData<FModuleBuiltData>();
+		const float* AgeData = ParticleSimulationContext.GetParticleAge();
+		const float* PreviousAgeData = ParticleSimulationContext.GetParticlePreviousAge();
+
+		for (uint32 i = 0; i < ParticleSimulationContext.GetNumInstances(); ++i)
+		{
+			const float RotationRate = ParticleSimulationContext.RandomScaleBiasFloat(i, ModuleBuiltData->RotationRange);
+
+			float SpriteRotation = ParticleSimulationContext.ReadParticleVariable(ModuleBuiltData->SpriteRotationVariableOffset, i, 0.0f);
+			float PreviousSpriteRotation = ParticleSimulationContext.ReadParticleVariable(ModuleBuiltData->PreviousSpriteRotationVariableOffset, i, 0.0f);
+
+			SpriteRotation += AgeData[i] * RotationRate;
+			PreviousSpriteRotation += PreviousAgeData[i] * RotationRate;
+
+			ParticleSimulationContext.WriteParticleVariable(ModuleBuiltData->SpriteRotationVariableOffset, i, SpriteRotation);
+			ParticleSimulationContext.WriteParticleVariable(ModuleBuiltData->PreviousSpriteRotationVariableOffset, i, PreviousSpriteRotation);
+		}
+	}
+
 #if WITH_EDITOR
 	virtual bool CanDisableModule() const override { return true; }
 #endif
diff --git a/Engine/Plugins/FX/Niagara/Source/Niagara/Internal/Stateless/Modules/NiagaraStatelessModule_SubUVAnimation.h b/Engine/Plugins/FX/Niagara/Source/Niagara/Internal/Stateless/Modules/NiagaraStatelessModule_SubUVAnimation.h
--- a/Engine/Plugins/FX/Niagara/Source/Niagara/Internal/Stateless/Modules/NiagaraStatelessModule_SubUVAnimation.h
+++ b/Engine/Plugins/FX/Niagara/Source/Niagara/Internal/Stateless/Modules/NiagaraStatelessModule_SubUVAnimation.h
@@ -5,6 +5,7 @@
 #include "Stateless/NiagaraStatelessModule.h"
 #include "Stateless/NiagaraStatelessEmitterDataBuildContext.h"
 #include "Stateless/NiagaraStatelessModuleShaderParameters.h"
+#include "Stateless/NiagaraStatelessParticleSimContext.h"
 
 #include "NiagaraStatelessModule_SubUVAnimation.generated.h"
 
@@ -21,6 +22,20 @@ class UNiagaraStatelessModule_SubUVAnimation : public UNiagaraStatelessModule
 {
 	GENERATED_BODY()
 
+	struct FModuleBuiltData
+	{
+		int32	Mode					= 0;
+		float	NumFrames				= 0.0f;
+		float	InitialFrameScale		= 0.0f;
+		float	InitialFrameBias		= 0.0f;
+		float	InitialFrameRateChange	= 0.0f;
+		float	AnimFrameStart			= 0.0f;
+		float	AnimFrameRange			= 0.0f;
+		float	RateScale				= 0.0f;
+
+		int32	SubImageIndexVariableOffset = INDEX_NONE;
+	};
+
 public:
 	using FParameters = NiagaraStateless::FSubUVAnimationModule_ShaderParameters;
 
@@ -55,55 +70,103 @@ public:
 	UPROPERTY(EditAnywhere, Category = "Parameters", meta = (EditConditionHides, EditCondition = "AnimationMode == ENSMSubUVAnimation_Mode::Random"))
 	float RandomChangeInterval = 0.1f;
 
+	virtual void BuildEmitterData(const FNiagaraStatelessEmitterDataBuildContext& BuildContext) const override
+	{
+		FModuleBuiltData* BuiltData = BuildContext.AllocateBuiltData<FModuleBuiltData>();
+		if (!IsModuleEnabled())
+		{
+			return;
+		}
+
+		const FNiagaraStatelessGlobals& StatelessGlobals = FNiagaraStatelessGlobals::Get();
+		BuiltData->SubImageIndexVariableOffset = BuildContext.FindParticleVariableIndex(StatelessGlobals.SubImageIndexVariable);
+		if (BuiltData->SubImageIndexVariableOffset == INDEX_NONE)
+		{
+			return;
+		}
+
+		const float FrameRangeStart	= bStartFrameRangeOverride_Enabled ? FMath::Clamp(float(StartFrameRangeOverride) / float(NumFrames - 1), 0.0f, 1.0f) : 0.0f;
+		const float FrameRangeEnd	= bEndFrameRangeOverride_Enabled   ? FMath::Clamp(float(EndFrameRangeOverride)   / float(NumFrames - 1), 0.0f, 1.0f) : 1.0f;
+
+		BuiltData->Mode			= (int)AnimationMode;
+		BuiltData->NumFrames	= float(NumFrames);
+		switch (AnimationMode)
+		{
+			case ENSMSubUVAnimation_Mode::InfiniteLoop:
+				BuiltData->InitialFrameScale		= 0.0f;
+				BuiltData->InitialFrameBias			= 0.0f;
+				BuiltData->InitialFrameRateChange	= 0.0f;
+				BuiltData->AnimFrameStart			= FrameRangeStart;
+				BuiltData->AnimFrameRange			= FrameRangeEnd - FrameRangeStart;
+				BuiltData->RateScale				= LoopsPerSecond;
+				break;
+
+			case ENSMSubUVAnimation_Mode::Linear:
+				BuiltData->InitialFrameScale		= 0.0f;
+				BuiltData->InitialFrameBias			= 0.0f;
+				BuiltData->InitialFrameRateChange	= 0.0f;
+				BuiltData->AnimFrameStart			= FrameRangeStart;
+				BuiltData->AnimFrameRange			= FrameRangeEnd - FrameRangeStart;
+				BuiltData->RateScale				= 1.0f;
+				break;
+
+			case ENSMSubUVAnimation_Mode::Random:
+				BuiltData->InitialFrameScale		= FrameRangeEnd - FrameRangeStart;
+				BuiltData->InitialFrameBias			= FrameRangeStart;
+				BuiltData->InitialFrameRateChange	= RandomChangeInterval > 0.0f ? 1.0f / RandomChangeInterval : 0.0f;
+				BuiltData->AnimFrameStart			= 0.0f;
+				BuiltData->AnimFrameRange			= 0.0f;
+				BuiltData->RateScale				= 0.0f;
+				break;
+		}
+
+		BuildContext.AddParticleSimulationExecSimulate(&UNiagaraStatelessModule_SubUVAnimation::ParticleSimulate);
+	}
+
 	virtual void SetShaderParameters(const FNiagaraStatelessSetShaderParameterContext& SetShaderParameterContext) const override
 	{
 		FParameters* Parameters = SetShaderParameterContext.GetParameterNestedStruct<FParameters>();
-		if (IsModuleEnabled())
+		const FModuleBuiltData* ModuleBuiltData = SetShaderParameterContext.ReadBuiltData<FModuleBuiltData>();
+
+		Parameters->SubUVAnimation_Mode						= ModuleBuiltData->Mode;
+		Parameters->SubUVAnimation_NumFrames				= ModuleBuiltData->NumFrames;
+		Parameters->SubUVAnimation_InitialFrameScale		= ModuleBuiltData->InitialFrameScale;
+		Parameters->SubUVAnimation_InitialFrameBias			= ModuleBuiltData->InitialFrameBias;
+		Parameters->SubUVAnimation_InitialFrameRateChange	= ModuleBuiltData->InitialFrameRateChange;
+		Parameters->SubUVAnimation_AnimFrameStart			= ModuleBuiltData->AnimFrameStart;
+		Parameters->SubUVAnimation_AnimFrameRange			= ModuleBuiltData->AnimFrameRange;
+		Parameters->SubUVAnimation_RateScale				= ModuleBuiltData->RateScale;
+	}
+
+	static void ParticleSimulate(const NiagaraStateless::FParticleSimulationContext& ParticleSimulationContext)
+	{
+		using namespace NiagaraStateless;
+
+		const FModuleBuiltData* ModuleBuiltData = ParticleSimulationContext.ReadBuiltData<FModuleBuiltData>();
+		const float* ParticleAges = ParticleSimulationContext.GetParticleAge();
+		const float* ParticleNormalizedAges = ParticleSimulationContext.GetParticleNormalizedAge();
+
+		for (uint32 i = 0; i < ParticleSimulationContext.GetNumInstances(); ++i)
 		{
-			const float FrameRangeStart	= bStartFrameRangeOverride_Enabled ? FMath::Clamp(float(StartFrameRangeOverride) / float(NumFrames - 1), 0.0f, 1.0f) : 0.0f;
-			const float FrameRangeEnd	= bEndFrameRangeOverride_Enabled   ? FMath::Clamp(float(EndFrameRangeOverride)   / float(NumFrames - 1), 0.0f, 1.0f) : 1.0f;
+			const float Particle_Age = ParticleAges[i];
+			const float Particle_NormalizedAge = ParticleNormalizedAges[i];
 
-			Parameters->SubUVAnimation_Mode			= (int)AnimationMode;
-			Parameters->SubUVAnimation_NumFrames	= float(NumFrames);
-			switch (AnimationMode)
+			const uint32 SeedOffset = uint32(Particle_Age * ModuleBuiltData->InitialFrameRateChange);
+			float Frame = ParticleSimulationContext.RandomSeedOffsetFloat(i, SeedOffset) * ModuleBuiltData->InitialFrameScale + ModuleBuiltData->InitialFrameBias;
+
+			if (ModuleBuiltData->Mode == 0)
 			{
-				case ENSMSubUVAnimation_Mode::InfiniteLoop:
-					Parameters->SubUVAnimation_InitialFrameScale		= 0.0f;
-					Parameters->SubUVAnimation_InitialFrameBias			= 0.0f;
-					Parameters->SubUVAnimation_InitialFrameRateChange	= 0.0f;
-					Parameters->SubUVAnimation_AnimFrameStart			= FrameRangeStart;
-					Parameters->SubUVAnimation_AnimFrameRange			= FrameRangeEnd - FrameRangeStart;
-					Parameters->SubUVAnimation_RateScale				= LoopsPerSecond;
-					break;
-
-				case ENSMSubUVAnimation_Mode::Linear:
-					Parameters->SubUVAnimation_InitialFrameScale		= 0.0f;
-					Parameters->SubUVAnimation_InitialFrameBias			= 0.0f;
-					Parameters->SubUVAnimation_InitialFrameRateChange	= 0.0f;
-					Parameters->SubUVAnimation_AnimFrameStart			= FrameRangeStart;
-					Parameters->SubUVAnimation_AnimFrameRange			= FrameRangeEnd - FrameRangeStart;
-					Parameters->SubUVAnimation_RateScale				= 1.0f;
-					break;
-
-				case ENSMSubUVAnimation_Mode::Random:
-					Parameters->SubUVAnimation_InitialFrameScale		= FrameRangeEnd - FrameRangeStart;
-					Parameters->SubUVAnimation_InitialFrameBias			= FrameRangeStart;
-					Parameters->SubUVAnimation_InitialFrameRateChange	= RandomChangeInterval > 0.0f ? 1.0f / RandomChangeInterval : 0.0f;
-					Parameters->SubUVAnimation_AnimFrameStart			= 0.0f;
-					Parameters->SubUVAnimation_AnimFrameRange			= 0.0f;
-					Parameters->SubUVAnimation_RateScale				= 0.0f;
-					break;
+				const float Interp = Particle_Age * ModuleBuiltData->RateScale;
+				Frame = FMath::Fractional(Frame + ModuleBuiltData->AnimFrameStart + (Interp * ModuleBuiltData->AnimFrameRange));
 			}
-		}
-		else
-		{
-			Parameters->SubUVAnimation_Mode						= 0;
-			Parameters->SubUVAnimation_InitialFrameScale		= 0.0f;
-			Parameters->SubUVAnimation_InitialFrameBias			= 0.0f;
-			Parameters->SubUVAnimation_InitialFrameRateChange	= 0.0f;
-			Parameters->SubUVAnimation_AnimFrameStart			= 0.0f;
-			Parameters->SubUVAnimation_AnimFrameRange			= 0.0f;
-			Parameters->SubUVAnimation_RateScale				= 0.0f;
+			else if (ModuleBuiltData->Mode == 1)
+			{
+				const float Interp = Particle_NormalizedAge * ModuleBuiltData->RateScale;
+				Frame = FMath::Clamp(Frame + ModuleBuiltData->AnimFrameStart + (Interp * ModuleBuiltData->AnimFrameRange), 0.0f, 1.0f);
+			}
+			const float SubImageIndex = Frame * ModuleBuiltData->NumFrames;
+
+			ParticleSimulationContext.WriteParticleVariable(ModuleBuiltData->SubImageIndexVariableOffset, i, SubImageIndex);
 		}
 	}
 
diff --git a/Engine/Plugins/FX/Niagara/Source/Niagara/Internal/Stateless/NiagaraStatelessCommon.h b/Engine/Plugins/FX/Niagara/Source/Niagara/Internal/Stateless/NiagaraStatelessCommon.h
--- a/Engine/Plugins/FX/Niagara/Source/Niagara/Internal/Stateless/NiagaraStatelessCommon.h
+++ b/Engine/Plugins/FX/Niagara/Source/Niagara/Internal/Stateless/NiagaraStatelessCommon.h
@@ -4,12 +4,25 @@
 
 #include "NiagaraCommon.h"
 
+enum class ENiagaraStatelessFeatureMask : uint8
+{
+	None			= 0,
+
+	ExecuteGPU		= 1 << 0,					// We can execute on the GPU (Might be broken down into GPUCompute | GPUGraphics | GPUAsyncCompute in future but this will remain the master mask)
+	ExecuteCPU		= 1 << 1,					// We can execute on the CPU
+	ExecuteAll		= ExecuteGPU | ExecuteCPU,	// We can execute on both the CPU & GPU
+
+	All				= ExecuteAll,				// We support all features
+};
+ENUM_CLASS_FLAGS(ENiagaraStatelessFeatureMask);
+
 template<typename TType>
 struct FNiagaraStatelessRange
 {
 	using ValueType = TType;
 
 	FNiagaraStatelessRange() = default;
+	FNiagaraStatelessRange(const FNiagaraStatelessRange& Other) = default;
 	explicit FNiagaraStatelessRange(const ValueType& InMinMax) : Min(InMinMax), Max(InMinMax) {}
 	explicit FNiagaraStatelessRange(const ValueType& InMin, const ValueType& InMax) : Min(InMin), Max(InMax) {}
 
@@ -27,11 +40,22 @@ using FNiagaraStatelessRangeVector3	= FNiagaraStatelessRange<FVector3f>;
 using FNiagaraStatelessRangeVector4	= FNiagaraStatelessRange<FVector4f>;
 using FNiagaraStatelessRangeColor	= FNiagaraStatelessRange<FLinearColor>;
 
+enum class ENiagaraStatelessBuiltDistributionFlag
+{
+	Random  = 0x00000001,	// StatelessDistributionFlag_Random
+	Uniform = 0x00000002,	// StatelessDistributionFlag_Uniform
+	Binding = 0x00000004,	// StatelessDistributionFlag_Binding
+};
+ENUM_CLASS_FLAGS(ENiagaraStatelessBuiltDistributionFlag);
+
 struct FNiagaraStatelessGlobals
 {
 	FNiagaraVariableBase	CameraOffsetVariable;
 	FNiagaraVariableBase	ColorVariable;
 	FNiagaraVariableBase	DynamicMaterialParameters0Variable;
+	FNiagaraVariableBase	DynamicMaterialParameters1Variable;
+	FNiagaraVariableBase	DynamicMaterialParameters2Variable;
+	FNiagaraVariableBase	DynamicMaterialParameters3Variable;
 	FNiagaraVariableBase	MeshIndexVariable;
 	FNiagaraVariableBase	MeshOrientationVariable;
 	FNiagaraVariableBase	PositionVariable;
@@ -58,8 +82,10 @@ struct FNiagaraStatelessGlobals
 	FNiagaraVariableBase	PreviousSpriteRotationVariable;
 	FNiagaraVariableBase	PreviousVelocityVariable;
 
+	ENiagaraStatelessFeatureMask	FeatureMask = ENiagaraStatelessFeatureMask::All;
+
 	inline static FLinearColor	GetDefaultColorValue() { return FLinearColor::White; }
-	inline static FVector4f		GetDefaultDynamicMaterialParameters0Value() { return FVector4f::Zero(); }
+	inline static FVector4f		GetDefaultDynamicMaterialParametersValue() { return FVector4f::Zero(); }
 	inline static float			GetDefaultLifetimeValue() { return 1.0f; }
 	inline static float			GetDefaultMassValue() { return 1.0f; }
 	inline static FQuat4f		GetDefaultMeshOrientationValue() { return FQuat4f::Identity; }
@@ -74,4 +100,5 @@ struct FNiagaraStatelessGlobals
 namespace NiagaraStatelessCommon
 {
 	extern void Initialize();
+	extern void UpdateSettings();
 }
diff --git a/Engine/Plugins/FX/Niagara/Source/Niagara/Internal/Stateless/NiagaraStatelessEmitterData.h b/Engine/Plugins/FX/Niagara/Source/Niagara/Internal/Stateless/NiagaraStatelessEmitterData.h
--- a/Engine/Plugins/FX/Niagara/Source/Niagara/Internal/Stateless/NiagaraStatelessEmitterData.h
+++ b/Engine/Plugins/FX/Niagara/Source/Niagara/Internal/Stateless/NiagaraStatelessEmitterData.h
@@ -16,6 +16,7 @@ namespace NiagaraStateless
 {
 	class FSimulationShader;
 	class FSpawnInfoShaderParameters;
+	class FParticleSimulationExecData;
 }
 
 struct FNiagaraStatelessEmitterData
@@ -35,6 +36,7 @@ struct FNiagaraStatelessEmitterData
 
 	bool											bCanEverExecute = false;
 	bool											bDeterministic = false;
+	ENiagaraStatelessFeatureMask					FeatureMask = ENiagaraStatelessFeatureMask::None;
 	int32											RandomSeed = 0;
 	FNiagaraStatelessRangeFloat						LifetimeRange = FNiagaraStatelessRangeFloat(0.0f, 0.0f);
 	FBox											FixedBounds = FBox(ForceInit);
@@ -50,9 +52,11 @@ struct FNiagaraStatelessEmitterData
 	const UNiagaraStatelessEmitterTemplate*			EmitterTemplate = nullptr;	// Used to access shader information
 
 	TArray<uint8>									BuiltData;					// Built data, generally allocated by modules if any
-	TArray<float>									StaticFloatData;			// Transient data used in build process, do not access directly
+	TArray<float>									StaticFloatData;			// Used with CPU generation, must be valid if ParticleSimExecData is also valid
 	FReadBuffer										StaticFloatBuffer;
 
+	NiagaraStateless::FParticleSimulationExecData*	ParticleSimExecData = nullptr;	// CPU simulation execution data, when null we don't provide a CPU path
+
 	void InitRenderResources();
 
 	TShaderRef<NiagaraStateless::FSimulationShader>	GetShader() const;
diff --git a/Engine/Plugins/FX/Niagara/Source/Niagara/Internal/Stateless/NiagaraStatelessEmitterDataBuildContext.h b/Engine/Plugins/FX/Niagara/Source/Niagara/Internal/Stateless/NiagaraStatelessEmitterDataBuildContext.h
--- a/Engine/Plugins/FX/Niagara/Source/Niagara/Internal/Stateless/NiagaraStatelessEmitterDataBuildContext.h
+++ b/Engine/Plugins/FX/Niagara/Source/Niagara/Internal/Stateless/NiagaraStatelessEmitterDataBuildContext.h
@@ -5,9 +5,15 @@
 #include "NiagaraStatelessCommon.h"
 #include "NiagaraStatelessDistribution.h"
 
+struct FNiagaraDataSetCompiledData;
 struct FNiagaraParameterBinding;
 struct FNiagaraParameterBindingWithValue;
 struct FNiagaraParameterStore;
+namespace NiagaraStateless
+{
+	class FParticleSimulationContext;
+	class FParticleSimulationExecData;
+}
 
 class FNiagaraStatelessEmitterDataBuildContext
 {
@@ -18,32 +24,36 @@ class FNiagaraStatelessEmitterDataBuildContext
 public:
 	UE_NONCOPYABLE(FNiagaraStatelessEmitterDataBuildContext);
 
-	FNiagaraStatelessEmitterDataBuildContext(FNiagaraParameterStore& InRendererBindings, TArray<uint8>& InBuiltData, TArray<float>& InStaticFloatData)
-		: RendererBindings(InRendererBindings)
+	FNiagaraStatelessEmitterDataBuildContext(FNiagaraDataSetCompiledData& InParticleDataSet, FNiagaraParameterStore& InRendererBindings, TArray<uint8>& InBuiltData, TArray<float>& InStaticFloatData, NiagaraStateless::FParticleSimulationExecData* InParticleExecData)
+		: ParticleDataSet(InParticleDataSet)
+		, RendererBindings(InRendererBindings)
 		, BuiltData(InBuiltData)
 		, StaticFloatData(InStaticFloatData)
+		, ParticleExecData(InParticleExecData)
 	{
 	}
 
-	uint32 AddStaticData(TConstArrayView<float> FloatData);
-	uint32 AddStaticData(TConstArrayView<FVector2f> FloatData);
-	uint32 AddStaticData(TConstArrayView<FVector3f> FloatData);
-	uint32 AddStaticData(TConstArrayView<FVector4f> FloatData);
-	uint32 AddStaticData(TConstArrayView<FLinearColor> FloatData);
+	void PreModuleBuild();
+
+	uint32 AddStaticData(TConstArrayView<float> FloatData) const;
+	uint32 AddStaticData(TConstArrayView<FVector2f> FloatData) const;
+	uint32 AddStaticData(TConstArrayView<FVector3f> FloatData) const;
+	uint32 AddStaticData(TConstArrayView<FVector4f> FloatData) const;
+	uint32 AddStaticData(TConstArrayView<FLinearColor> FloatData) const;
 
 	template<typename T>
-	T* AllocateBuiltData()
+	T* AllocateBuiltData() const
 	{
 		static_assert(TIsTrivial<T>::Value, "Only trivial types can be used for built data");
 
-		const int32 Align = BuiltData.Num() % alignof(T);
-		const int32 Offset = BuiltData.AddZeroed(sizeof(T) + Align);
-		void* NewData = BuiltData.GetData() + Offset + Align;
+		const int32 Offset = Align(BuiltData.Num(), alignof(T));
+		BuiltData.AddZeroed(Offset + sizeof(T) - BuiltData.Num());
+		void* NewData = BuiltData.GetData() + Offset;
 		return new(NewData) T();
 	}
 
 	template<typename T>
-	T& GetTransientBuildData()
+	T& GetTransientBuildData() const
 	{
 		TUniquePtr<FTransientObject>& TransientObj = TransientBuildData.FindOrAdd(T::GetName());
 		if (TransientObj.IsValid() == false)
@@ -56,38 +66,40 @@ public:
 	// Adds a binding to the renderer parameter store
 	// This allows you to read the parameter data inside the simulation process
 	// The returned value is INDEX_NONE is the variables is index otherwise the offset in DWORDs
-	int32 AddRendererBinding(const FNiagaraVariableBase& Variable);
-	int32 AddRendererBinding(const FNiagaraParameterBinding& Binding);
-	int32 AddRendererBinding(const FNiagaraParameterBindingWithValue& Binding);	
+	int32 AddRendererBinding(const FNiagaraVariableBase& Variable) const;
+	int32 AddRendererBinding(const FNiagaraParameterBinding& Binding) const;
+	int32 AddRendererBinding(const FNiagaraParameterBindingWithValue& Binding) const;
 
 	// Adds an distribution into the LUT if enabled
 	template<typename TType>
-	FUintVector3 AddDistribution(ENiagaraDistributionMode Mode, TConstArrayView<TType> Values, bool bEnabled)
+	FUintVector3 AddDistribution(ENiagaraDistributionMode Mode, TConstArrayView<TType> Values, bool bEnabled) const
 	{
+		using namespace NiagaraStateless;
+
 		FUintVector3 Parameters = FUintVector3::ZeroValue;
 		if (bEnabled && Values.Num() > 0)
 		{
 			switch (Mode)
 			{
 				case ENiagaraDistributionMode::Binding:				checkNoEntry(); break;
-				case ENiagaraDistributionMode::UniformConstant:		Parameters.X = StatelessDistributionFlag_Random | StatelessDistributionFlag_Uniform; break;
-				case ENiagaraDistributionMode::NonUniformConstant:	Parameters.X = StatelessDistributionFlag_Random; break;
-				case ENiagaraDistributionMode::UniformRange:		Parameters.X = StatelessDistributionFlag_Random | StatelessDistributionFlag_Uniform; break;
-				case ENiagaraDistributionMode::NonUniformRange:		Parameters.X = StatelessDistributionFlag_Random; break;
-				case ENiagaraDistributionMode::UniformCurve:		Parameters.X = StatelessDistributionFlag_Uniform; break;
+				case ENiagaraDistributionMode::UniformConstant:		Parameters.X = uint32(ENiagaraStatelessBuiltDistributionFlag::Random | ENiagaraStatelessBuiltDistributionFlag::Uniform); break;
+				case ENiagaraDistributionMode::NonUniformConstant:	Parameters.X = uint32(ENiagaraStatelessBuiltDistributionFlag::Random); break;
+				case ENiagaraDistributionMode::UniformRange:		Parameters.X = uint32(ENiagaraStatelessBuiltDistributionFlag::Random | ENiagaraStatelessBuiltDistributionFlag::Uniform); break;
+				case ENiagaraDistributionMode::NonUniformRange:		Parameters.X = uint32(ENiagaraStatelessBuiltDistributionFlag::Random); break;
+				case ENiagaraDistributionMode::UniformCurve:		Parameters.X = uint32(ENiagaraStatelessBuiltDistributionFlag::Uniform); break;
 				case ENiagaraDistributionMode::NonUniformCurve:		Parameters.X = 0; break;
 				default:											checkNoEntry(); break;
 			}
 
 			Parameters.Y = AddStaticData(Values);
-			reinterpret_cast<float&>(Parameters.Z) = Values.Num() - 1;
+			Parameters.Z = Values.Num() - 1;
 		}
 		return Parameters;
 	}
 
 	// Adds a distribution into the LUT if enabled and returns the packed information to send to the shader
 	template<typename TDistribution>
-	FUintVector3 AddDistribution(const TDistribution& Distribution, bool bEnabled)
+	FUintVector3 AddDistribution(const TDistribution& Distribution, bool bEnabled = true) const
 	{
 		FUintVector3 Parameters = FUintVector3::ZeroValue;
 		if ( bEnabled )
@@ -99,7 +111,7 @@ public:
 				{
 					Parameters.X = StatelessDistributionFlag_Binding;
 					Parameters.Y = ParameterOffset;
-					reinterpret_cast<float&>(Parameters.Z) = 1.0f;
+					Parameters.Z = 1.0f;
 				}
 			}
 			else
@@ -111,38 +123,43 @@ public:
 	}
 
 	template<typename TRange, typename TDistribution, typename TDefaultValue>
-	TRange ConvertDistributionToRangeHelper(const TDistribution& Distribution, const TDefaultValue& DefaultValue, bool bEnabled)
+	TRange ConvertDistributionToRangeHelper(const TDistribution& Distribution, const TDefaultValue& DefaultValue) const
 	{
 		TRange Range(DefaultValue);
-		if (bEnabled)
+		if (Distribution.Mode == ENiagaraDistributionMode::Binding)
 		{
-			if (Distribution.Mode == ENiagaraDistributionMode::Binding)
-			{
-				Range.ParameterOffset = AddRendererBinding(Distribution.ParameterBinding);
-			}
-			else
-			{
-				Range = Distribution.CalculateRange(DefaultValue);
-			}
+			Range.ParameterOffset = AddRendererBinding(Distribution.ParameterBinding);
+		}
+		else
+		{
+			Range = Distribution.CalculateRange(DefaultValue);
 		}
 		return Range;
 	}
 
-	FNiagaraStatelessRangeFloat   ConvertDistributionToRange(const FNiagaraDistributionFloat& Distribution, float DefaultValue, bool bEnabled = true) { return ConvertDistributionToRangeHelper<FNiagaraStatelessRangeFloat>(Distribution, DefaultValue, bEnabled); }
-	FNiagaraStatelessRangeVector2 ConvertDistributionToRange(const FNiagaraDistributionVector2& Distribution, const FVector2f& DefaultValue, bool bEnabled = true) { return ConvertDistributionToRangeHelper<FNiagaraStatelessRangeVector2>(Distribution, DefaultValue, bEnabled); }
-	FNiagaraStatelessRangeVector3 ConvertDistributionToRange(const FNiagaraDistributionVector3& Distribution, const FVector3f& DefaultValue, bool bEnabled = true) { return ConvertDistributionToRangeHelper<FNiagaraStatelessRangeVector3>(Distribution, DefaultValue, bEnabled); }
-	FNiagaraStatelessRangeColor   ConvertDistributionToRange(const FNiagaraDistributionColor& Distribution, const FLinearColor& DefaultValue, bool bEnabled = true) { return ConvertDistributionToRangeHelper<FNiagaraStatelessRangeColor>(Distribution, DefaultValue, bEnabled); }
+	FNiagaraStatelessRangeFloat   ConvertDistributionToRange(const FNiagaraDistributionFloat& Distribution, float DefaultValue) const { return ConvertDistributionToRangeHelper<FNiagaraStatelessRangeFloat>(Distribution, DefaultValue); }
+	FNiagaraStatelessRangeVector2 ConvertDistributionToRange(const FNiagaraDistributionVector2& Distribution, const FVector2f& DefaultValue) const { return ConvertDistributionToRangeHelper<FNiagaraStatelessRangeVector2>(Distribution, DefaultValue); }
+	FNiagaraStatelessRangeVector3 ConvertDistributionToRange(const FNiagaraDistributionVector3& Distribution, const FVector3f& DefaultValue) const { return ConvertDistributionToRangeHelper<FNiagaraStatelessRangeVector3>(Distribution, DefaultValue); }
+	FNiagaraStatelessRangeColor   ConvertDistributionToRange(const FNiagaraDistributionColor& Distribution, const FLinearColor& DefaultValue) const { return ConvertDistributionToRangeHelper<FNiagaraStatelessRangeColor>(Distribution, DefaultValue); }
 
-	FNiagaraStatelessRangeFloat   ConvertDistributionToRange(const FNiagaraDistributionRangeFloat& Distribution, float DefaultValue, bool bEnabled = true) { return ConvertDistributionToRangeHelper<FNiagaraStatelessRangeFloat>(Distribution, DefaultValue, bEnabled); }
-	FNiagaraStatelessRangeVector2 ConvertDistributionToRange(const FNiagaraDistributionRangeVector2& Distribution, const FVector2f& DefaultValue, bool bEnabled = true) { return ConvertDistributionToRangeHelper<FNiagaraStatelessRangeVector2>(Distribution, DefaultValue, bEnabled); }
-	FNiagaraStatelessRangeVector3 ConvertDistributionToRange(const FNiagaraDistributionRangeVector3& Distribution, const FVector3f& DefaultValue, bool bEnabled = true) { return ConvertDistributionToRangeHelper<FNiagaraStatelessRangeVector3>(Distribution, DefaultValue, bEnabled); }
-	FNiagaraStatelessRangeColor   ConvertDistributionToRange(const FNiagaraDistributionRangeColor& Distribution, const FLinearColor& DefaultValue, bool bEnabled = true) { return ConvertDistributionToRangeHelper<FNiagaraStatelessRangeColor>(Distribution, DefaultValue, bEnabled); }
-	FNiagaraStatelessRangeInt     ConvertDistributionToRange(const FNiagaraDistributionRangeInt& Distribution, int32 DefaultValue, bool bEnabled = true) { return ConvertDistributionToRangeHelper<FNiagaraStatelessRangeInt>(Distribution, DefaultValue, bEnabled); }
+	FNiagaraStatelessRangeFloat   ConvertDistributionToRange(const FNiagaraDistributionRangeFloat& Distribution, float DefaultValue) const { return ConvertDistributionToRangeHelper<FNiagaraStatelessRangeFloat>(Distribution, DefaultValue); }
+	FNiagaraStatelessRangeVector2 ConvertDistributionToRange(const FNiagaraDistributionRangeVector2& Distribution, const FVector2f& DefaultValue) const { return ConvertDistributionToRangeHelper<FNiagaraStatelessRangeVector2>(Distribution, DefaultValue); }
+	FNiagaraStatelessRangeVector3 ConvertDistributionToRange(const FNiagaraDistributionRangeVector3& Distribution, const FVector3f& DefaultValue) const { return ConvertDistributionToRangeHelper<FNiagaraStatelessRangeVector3>(Distribution, DefaultValue); }
+	FNiagaraStatelessRangeColor   ConvertDistributionToRange(const FNiagaraDistributionRangeColor& Distribution, const FLinearColor& DefaultValue) const { return ConvertDistributionToRangeHelper<FNiagaraStatelessRangeColor>(Distribution, DefaultValue); }
+	FNiagaraStatelessRangeInt     ConvertDistributionToRange(const FNiagaraDistributionRangeInt& Distribution, int32 DefaultValue) const { return ConvertDistributionToRangeHelper<FNiagaraStatelessRangeInt>(Distribution, DefaultValue); }
+
+	void AddParticleSimulationExecSimulate(TFunction<void(const NiagaraStateless::FParticleSimulationContext&)> Func) const;
+
+	int32 FindParticleVariableIndex(const FNiagaraVariableBase& Variable) const;
 
 private:
-	FNiagaraParameterStore& RendererBindings;
-	TArray<uint8>&			BuiltData;
-	TArray<float>&			StaticFloatData;
+	FNiagaraDataSetCompiledData&					ParticleDataSet;
+	FNiagaraParameterStore&							RendererBindings;
+	TArray<uint8>&									BuiltData;
+	TArray<float>&									StaticFloatData;
+	NiagaraStateless::FParticleSimulationExecData*	ParticleExecData = nullptr;
+
+	int32											ModuleBuiltDataOffset = 0;
 
 	struct FTransientObject
 	{
@@ -160,5 +177,5 @@ private:
 		T TheObject;
 	};
 
-	TMap<FName, TUniquePtr<FTransientObject>>	TransientBuildData;
+	mutable TMap<FName, TUniquePtr<FTransientObject>>	TransientBuildData;
 };
diff --git a/Engine/Plugins/FX/Niagara/Source/Niagara/Internal/Stateless/NiagaraStatelessEmitterInstance.h b/Engine/Plugins/FX/Niagara/Source/Niagara/Internal/Stateless/NiagaraStatelessEmitterInstance.h
--- a/Engine/Plugins/FX/Niagara/Source/Niagara/Internal/Stateless/NiagaraStatelessEmitterInstance.h
+++ b/Engine/Plugins/FX/Niagara/Source/Niagara/Internal/Stateless/NiagaraStatelessEmitterInstance.h
@@ -22,6 +22,7 @@ namespace NiagaraStateless
 
 		int32													RandomSeed = 0;
 		float													Age = 0.0f;
+		float													DeltaTime = 0.0f;		//-OPT: We should be able to pull this from the view which is important for zero tick path
 		ENiagaraExecutionState									ExecutionState = ENiagaraExecutionState::Active;
 		TArray<FNiagaraStatelessRuntimeSpawnInfo>				SpawnInfos;
 		TUniquePtr<NiagaraStateless::FCommonShaderParameters>	ShaderParameters;
diff --git a/Engine/Plugins/FX/Niagara/Source/Niagara/Internal/Stateless/NiagaraStatelessEmitterTemplate.h b/Engine/Plugins/FX/Niagara/Source/Niagara/Internal/Stateless/NiagaraStatelessEmitterTemplate.h
--- a/Engine/Plugins/FX/Niagara/Source/Niagara/Internal/Stateless/NiagaraStatelessEmitterTemplate.h
+++ b/Engine/Plugins/FX/Niagara/Source/Niagara/Internal/Stateless/NiagaraStatelessEmitterTemplate.h
@@ -35,9 +35,6 @@ public:
 	virtual void SetShaderParameters(uint8* ShaderParametersBase, TConstArrayView<int32> ComponentOffsets) const { checkNoEntry(); }
 	//-TODO: Add a way to set them directly, we should know that the final struct is a series of ints in the order of the provided variables
 
-	static FNiagaraVariableBase MakeProvidedVariable(FNiagaraVariableBase Variable);
-	static FNiagaraVariableBase MakeProvidedPreviousVariable(FNiagaraVariableBase Variable);
-
 protected:
 #if WITH_EDITORONLY_DATA
 	TArray<TObjectPtr<UClass>>		Modules;
diff --git a/Engine/Plugins/FX/Niagara/Source/Niagara/Internal/Stateless/NiagaraStatelessModule.h b/Engine/Plugins/FX/Niagara/Source/Niagara/Internal/Stateless/NiagaraStatelessModule.h
--- a/Engine/Plugins/FX/Niagara/Source/Niagara/Internal/Stateless/NiagaraStatelessModule.h
+++ b/Engine/Plugins/FX/Niagara/Source/Niagara/Internal/Stateless/NiagaraStatelessModule.h
@@ -44,8 +44,7 @@ public:
 	template<typename T>
 	const T* ReadBuiltData() const
 	{
-		const int32 Align = BuiltDataOffset % alignof(T);
-		const int32 Offset = BuiltDataOffset + Align;
+		const int32 Offset = Align(BuiltDataOffset, alignof(T));
 		BuiltDataOffset = Offset + sizeof(T);
 		check(BuiltDataOffset <= BuiltData.Num());
 		return reinterpret_cast<const T*>(BuiltData.GetData() + Offset);
@@ -111,10 +110,11 @@ public:
 		static NIAGARA_API const FName bModuleEnabled;
 		static NIAGARA_API const FName bDebugDrawEnabled;
 	};
-
 #endif
 
-	virtual void BuildEmitterData(FNiagaraStatelessEmitterDataBuildContext& BuildContext) const {}
+	virtual ENiagaraStatelessFeatureMask GetFeatureMask() const { return ENiagaraStatelessFeatureMask::All; }
+
+	virtual void BuildEmitterData(const FNiagaraStatelessEmitterDataBuildContext& BuildContext) const {}
 	virtual void SetShaderParameters(const FNiagaraStatelessSetShaderParameterContext& SetShaderParameterContext) const { }
 
 #if WITH_EDITOR
diff --git a/Engine/Plugins/FX/Niagara/Source/Niagara/Internal/Stateless/NiagaraStatelessParticleSimContext.h b/Engine/Plugins/FX/Niagara/Source/Niagara/Internal/Stateless/NiagaraStatelessParticleSimContext.h
new file mode 100644
--- /dev/null
+++ b/Engine/Plugins/FX/Niagara/Source/Niagara/Internal/Stateless/NiagaraStatelessParticleSimContext.h
@@ -0,0 +1,248 @@
+// Copyright Epic Games, Inc. All Rights Reserved.
+
+#pragma once
+
+#include "NiagaraStatelessCommon.h"
+
+class FNiagaraDataBuffer;
+struct FNiagaraStatelessEmitterData;
+struct FNiagaraStatelessRuntimeSpawnInfo;
+class FRHICommandList;
+
+namespace NiagaraStateless
+{
+class FSpawnInfoShaderParameters;
+
+enum class EParticleComponent
+{
+	Alive,
+	Lifetime,
+	Age,
+	NormalizedAge,
+	PreviousAge,
+	PreviousNormalizedAge,
+	RandomSeed,
+	Num
+};
+
+class FParticleSimulationContext
+{
+public:
+	explicit FParticleSimulationContext(const FNiagaraStatelessEmitterData* EmitterData, TConstArrayView<uint8> DynamicBufferData);
+
+	void Simulate(int32 EmitterRandomSeed, float EmitterAge, float DeltaTime, TConstArrayView<FNiagaraStatelessRuntimeSpawnInfo> SpawnInfos, FNiagaraDataBuffer* DestinationData);
+	void SimulateGPU(FRHICommandList& RHICmdList, int32 EmitterRandomSeed, float EmitterAge, float DeltaTime, TConstArrayView<FNiagaraStatelessRuntimeSpawnInfo> SpawnInfos, FNiagaraDataBuffer* DestinationData);
+
+private:
+	void SimulateInternal(int32 EmitterRandomSeed, float EmitterAge, float DeltaTime, FSpawnInfoShaderParameters& SpawnParameters, uint32 ActiveParticles);
+
+public:
+	uint32 GetNumInstances() const { return NumInstances; }
+	uint32 GetParticleComponentOffset(int32 iComponent) const { return BufferStride * iComponent; };
+
+	float GetDeltaTime() const { return DeltaTime; }
+	float GetInvDeltaTime() const { return InvDeltaTime; }
+
+	// Get a pointer to required particle components
+	int32* GetParticleAlive() const { return static_cast<int32*>(RequiredComponents[int32(EParticleComponent::Alive)]); }
+	float* GetParticleLifetime() { return static_cast<float*>(RequiredComponents[int32(EParticleComponent::Lifetime)]); }
+	float* GetParticleAge() { return static_cast<float*>(RequiredComponents[int32(EParticleComponent::Age)]); }
+	float* GetParticleNormalizedAge() { return static_cast<float*>(RequiredComponents[int32(EParticleComponent::NormalizedAge)]); }
+	float* GetParticlePreviousAge() { return static_cast<float*>(RequiredComponents[int32(EParticleComponent::PreviousAge)]); }
+	float* GetParticlePreviousNormalizedAge() { return static_cast<float*>(RequiredComponents[int32(EParticleComponent::PreviousNormalizedAge)]); }
+	//FUintVector4* GetParticleRandomSeed() { return static_cast<FUintVector4*>(RequiredComponents[int32(EParticleComponent::RandomSeed)]); }
+
+	float* GetParticleLifetime() const { return static_cast<float*>(RequiredComponents[int32(EParticleComponent::Lifetime)]); }
+	float* GetParticleAge() const { return static_cast<float*>(RequiredComponents[int32(EParticleComponent::Age)]); }
+	float* GetParticleNormalizedAge() const { return static_cast<float*>(RequiredComponents[int32(EParticleComponent::NormalizedAge)]); }
+	float* GetParticlePreviousAge() const { return static_cast<float*>(RequiredComponents[int32(EParticleComponent::PreviousAge)]); }
+	float* GetParticlePreviousNormalizedAge() const { return static_cast<float*>(RequiredComponents[int32(EParticleComponent::PreviousNormalizedAge)]); }
+
+	// Get a pointer to an optional particle component (i.e. one not required in the data set)
+	float* GetParticleVariableFloat(int32 iVariable) const { return iVariable >= 0 ? reinterpret_cast<float*>(VariableComponents[iVariable]) : nullptr; }
+	int32* GetParticleVariableInt32(int32 iVariable) const { return iVariable >= 0 ? reinterpret_cast<int32*>(VariableComponents[iVariable]) : nullptr; }
+
+	template<typename TType>
+	const TType ReadParticleVariableComponent(int32 iVariable, int32 iInstance, int32 iComponent) const { return *reinterpret_cast<const TType*>(VariableComponents[iVariable] + (BufferStride * iComponent) + (iInstance * sizeof(float))); }
+	int32 ReadParticleVariable(int32 iVariable, uint32 iInstance, int32 DefaultValue) const { check(iInstance < NumInstances); return iVariable >= 0 ? ReadParticleVariableComponent<float>(iVariable, iInstance, 0) : DefaultValue; }
+	float ReadParticleVariable(int32 iVariable, uint32 iInstance, float DefaultValue) const { check(iInstance < NumInstances); return iVariable >= 0 ? ReadParticleVariableComponent<int32>(iVariable, iInstance, 0) : DefaultValue; }
+	FVector2f ReadParticleVariable(int32 iVariable, uint32 iInstance, const FVector2f& DefaultValue) const { check(iInstance < NumInstances); return iVariable >= 0 ? FVector2f(ReadParticleVariableComponent<float>(iVariable, iInstance, 0), ReadParticleVariableComponent<float>(iVariable, iInstance, 1)) : DefaultValue; }
+	FVector3f ReadParticleVariable(int32 iVariable, uint32 iInstance, const FVector3f& DefaultValue) const { check(iInstance < NumInstances); return iVariable >= 0 ? FVector3f(ReadParticleVariableComponent<float>(iVariable, iInstance, 0), ReadParticleVariableComponent<float>(iVariable, iInstance, 1), ReadParticleVariableComponent<float>(iVariable, iInstance, 2)) : DefaultValue; }
+	FVector4f ReadParticleVariable(int32 iVariable, uint32 iInstance, const FVector4f& DefaultValue) const { check(iInstance < NumInstances); return iVariable >= 0 ? FVector4f(ReadParticleVariableComponent<float>(iVariable, iInstance, 0), ReadParticleVariableComponent<float>(iVariable, iInstance, 1), ReadParticleVariableComponent<float>(iVariable, iInstance, 2), ReadParticleVariableComponent<float>(iVariable, iInstance, 3)) : DefaultValue; }
+	FQuat4f ReadParticleVariable(int32 iVariable, uint32 iInstance, const FQuat4f& DefaultValue) const { check(iInstance < NumInstances); return iVariable >= 0 ? FQuat4f(ReadParticleVariableComponent<float>(iVariable, iInstance, 0), ReadParticleVariableComponent<float>(iVariable, iInstance, 1), ReadParticleVariableComponent<float>(iVariable, iInstance, 2), ReadParticleVariableComponent<float>(iVariable, iInstance, 3)) : DefaultValue; }
+	FLinearColor ReadParticleVariable(int32 iVariable, uint32 iInstance, const FLinearColor& DefaultValue) const { check(iInstance < NumInstances); return iVariable >= 0 ? FLinearColor(ReadParticleVariableComponent<float>(iVariable, iInstance, 0), ReadParticleVariableComponent<float>(iVariable, iInstance, 1), ReadParticleVariableComponent<float>(iVariable, iInstance, 2), ReadParticleVariableComponent<float>(iVariable, iInstance, 3)) : DefaultValue; }
+
+	void WriteParticleVariableComponent(int32 iVariable, int32 iInstance, int32 iComponent, float Value) const { *reinterpret_cast<float*>(VariableComponents[iVariable] + (BufferStride * iComponent) + (iInstance * sizeof(float))) = Value; }
+	void WriteParticleVariableComponent(int32 iVariable, int32 iInstance, int32 iComponent, int32 Value) const { *reinterpret_cast<int32*>(VariableComponents[iVariable] + (BufferStride * iComponent) + (iInstance * sizeof(float))) = Value; }
+	void WriteParticleVariable(int32 iVariable, uint32 iInstance, int32 Value) const { check(iInstance < NumInstances); if (iVariable >= 0 ) { WriteParticleVariableComponent(iVariable, iInstance, 0, Value); } }
+	void WriteParticleVariable(int32 iVariable, uint32 iInstance, float Value) const { check(iInstance < NumInstances); if (iVariable >= 0 ) { WriteParticleVariableComponent(iVariable, iInstance, 0, Value); } }
+	void WriteParticleVariable(int32 iVariable, uint32 iInstance, const FVector2f& Value) const { check(iInstance < NumInstances); if (iVariable >= 0 ) { WriteParticleVariableComponent(iVariable, iInstance, 0, Value.X); WriteParticleVariableComponent(iVariable, iInstance, 1, Value.Y); } }
+	void WriteParticleVariable(int32 iVariable, uint32 iInstance, const FVector3f& Value) const { check(iInstance < NumInstances); if (iVariable >= 0 ) { WriteParticleVariableComponent(iVariable, iInstance, 0, Value.X); WriteParticleVariableComponent(iVariable, iInstance, 1, Value.Y); WriteParticleVariableComponent(iVariable, iInstance, 2, Value.Z); } }
+	void WriteParticleVariable(int32 iVariable, uint32 iInstance, const FVector4f& Value) const { check(iInstance < NumInstances); if (iVariable >= 0) { WriteParticleVariableComponent(iVariable, iInstance, 0, Value.X); WriteParticleVariableComponent(iVariable, iInstance, 1, Value.Y); WriteParticleVariableComponent(iVariable, iInstance, 2, Value.Z); WriteParticleVariableComponent(iVariable, iInstance, 3, Value.W); } }
+	void WriteParticleVariable(int32 iVariable, uint32 iInstance, const FQuat4f& Value) const { check(iInstance < NumInstances); if (iVariable >= 0) { WriteParticleVariableComponent(iVariable, iInstance, 0, Value.X); WriteParticleVariableComponent(iVariable, iInstance, 1, Value.Y); WriteParticleVariableComponent(iVariable, iInstance, 2, Value.Z); WriteParticleVariableComponent(iVariable, iInstance, 3, Value.W); } }
+	void WriteParticleVariable(int32 iVariable, uint32 iInstance, const FLinearColor& Value) const { check(iInstance < NumInstances); if (iVariable >= 0 ) { WriteParticleVariableComponent(iVariable, iInstance, 0, Value.R); WriteParticleVariableComponent(iVariable, iInstance, 1, Value.G); WriteParticleVariableComponent(iVariable, iInstance, 2, Value.B); WriteParticleVariableComponent(iVariable, iInstance, 3, Value.A); } }
+
+	template<typename TType>
+	TType GetParameterBufferValue(int32 Offset, int32 Element) const
+	{
+		TType Value;
+		FMemory::Memcpy(&Value, &DynamicBufferData[Offset * 4], sizeof(TType));
+		return Value;
+	}
+
+	int32 GetParameterBufferInt(int32 Offset, int32 Element) const { return GetParameterBufferValue<int32>(Offset, Element); }
+
+	template<typename TType>
+	TType GetParameterBufferFloat(int32 Offset, int32 Element) const { return GetParameterBufferValue<TType>(Offset, Element); }
+	template<typename TType>
+	TType GetParameterBufferFloat(int32 Offset, const TType& DefaultValue) const { return Offset >= 0 ? GetParameterBufferValue<TType>(Offset, 0) : DefaultValue; }
+
+	template<typename TType>
+	TType GetStaticFloat(uint32 Offset, uint32 Element) const;
+	template<> float GetStaticFloat<float>(uint32 Offset, uint32 Element) const { return StaticFloatData[Offset + Element]; }
+	template<> FVector2f GetStaticFloat<FVector2f>(uint32 Offset, uint32 Element) const { Offset += Element * 2; return FVector2f(StaticFloatData[Offset], StaticFloatData[Offset + 1]); }
+	template<> FVector3f GetStaticFloat<FVector3f>(uint32 Offset, uint32 Element) const { Offset += Element * 3; return FVector3f(StaticFloatData[Offset], StaticFloatData[Offset + 1], StaticFloatData[Offset + 2]); }
+	template<> FVector4f GetStaticFloat<FVector4f>(uint32 Offset, uint32 Element) const { Offset += Element * 4; return FVector4f(StaticFloatData[Offset], StaticFloatData[Offset + 1], StaticFloatData[Offset + 2], StaticFloatData[Offset + 3]); }
+	template<> FLinearColor GetStaticFloat<FLinearColor>(uint32 Offset, uint32 Element) const { Offset += Element * 4; return FLinearColor(StaticFloatData[Offset], StaticFloatData[Offset + 1], StaticFloatData[Offset + 2], StaticFloatData[Offset + 3]); }
+
+	template<typename TType>
+	TType LerpStaticFloat(FUintVector2 Parameters, float U) const
+	{
+		const float Offset = U * float(Parameters.Y);
+		const TType Value0 = GetStaticFloat<TType>(Parameters.X, FMath::FloorToInt(Offset));
+		const TType Value1 = GetStaticFloat<TType>(Parameters.X, FMath::CeilToInt(Offset));
+		return FMath::Lerp(Value0, Value1, FMath::Fractional(Offset));
+	}
+
+	uint32			RandomUInt(uint32 iInstance) const;
+	FUintVector2	RandomUInt2(uint32 iInstance) const;
+	FUintVector3	RandomUInt3(uint32 iInstance) const;
+	FUintVector4	RandomUInt4(uint32 iInstance) const;
+
+	float			RandomFloat(uint32 iInstance) const;
+	FVector2f		RandomFloat2(uint32 iInstance) const;
+	FVector3f		RandomFloat3(uint32 iInstance) const;
+	FVector4f		RandomFloat4(uint32 iInstance) const;
+
+	float			RandomScaleBiasFloat(uint32 iInstance, const float Scale, const float Bias) const { return Bias + (RandomFloat(iInstance) * Scale); }
+	FVector2f		RandomScaleBiasFloat(uint32 iInstance, const FVector2f& Scale, const FVector2f& Bias) const { return Bias + (RandomFloat(iInstance) * Scale); }
+	FVector3f		RandomScaleBiasFloat(uint32 iInstance, const FVector3f& Scale, const FVector3f& Bias) const { return Bias + (RandomFloat(iInstance) * Scale); }
+	FVector4f		RandomScaleBiasFloat(uint32 iInstance, const FVector4f& Scale, const FVector4f& Bias) const { return Bias + (RandomFloat(iInstance) * Scale); }
+
+	float			RandomScaleBiasFloat(uint32 iInstance, const FNiagaraStatelessRangeFloat& Range) const { return Range.Min + (RandomFloat(iInstance) * Range.GetScale()); }
+	FVector2f		RandomScaleBiasFloat(uint32 iInstance, const FNiagaraStatelessRangeVector2& Range) const { return Range.Min + (RandomFloat2(iInstance) * Range.GetScale()); }
+	FVector3f		RandomScaleBiasFloat(uint32 iInstance, const FNiagaraStatelessRangeVector3& Range) const { return Range.Min + (RandomFloat3(iInstance) * Range.GetScale()); }
+	FVector4f		RandomScaleBiasFloat(uint32 iInstance, const FNiagaraStatelessRangeVector4& Range) const { return Range.Min + (RandomFloat4(iInstance) * Range.GetScale()); }
+	FLinearColor	RandomScaleBiasFloat(uint32 iInstance, const FNiagaraStatelessRangeColor& Range) const { return Range.Min + (RandomFloat4(iInstance) * Range.GetScale()); }
+
+	float			RandomScaleBiasFloat(uint32 iInstance, const FNiagaraStatelessRangeFloat& Range, bool bUniform) const { return RandomScaleBiasFloat(iInstance, Range); }
+	FVector2f		RandomScaleBiasFloat(uint32 iInstance, const FNiagaraStatelessRangeVector2& Range, bool bUniform) const { const FVector2f Random = RandomScaleBiasFloat(iInstance, Range); return bUniform ? FVector2f(Random.X, Random.X) : Random; }
+	FVector3f		RandomScaleBiasFloat(uint32 iInstance, const FNiagaraStatelessRangeVector3& Range, bool bUniform) const { const FVector3f Random = RandomScaleBiasFloat(iInstance, Range); return bUniform ? FVector3f(Random.X, Random.X, Random.X) : Random; }
+	FVector4f		RandomScaleBiasFloat(uint32 iInstance, const FNiagaraStatelessRangeVector4& Range, bool bUniform) const { const FVector4f Random = RandomScaleBiasFloat(iInstance, Range); return bUniform ? FVector4f(Random.X, Random.X, Random.X, Random.X) : Random; }
+	FLinearColor	RandomScaleBiasFloat(uint32 iInstance, const FNiagaraStatelessRangeColor& Range, bool bUniform) const { const FLinearColor Random = RandomScaleBiasFloat(iInstance, Range); return bUniform ? FLinearColor(Random.R, Random.R, Random.R, Random.R) : Random; }
+
+	FVector2f		RandomUnitFloat2(uint32 iInstance) const { return SafeNormalize(RandomFloat2(iInstance) - 0.5f); }
+	FVector3f		RandomUnitFloat3(uint32 iInstance) const { return SafeNormalize(RandomFloat3(iInstance) - 0.5f); }
+
+	float			RandomSeedOffsetFloat(uint32 iInstance, uint32 SeedOffset) const;
+	FVector2f		RandomSeedOffsetFloat2(uint32 iInstance, uint32 SeedOffset) const;
+	FVector3f		RandomSeedOffsetFloat3(uint32 iInstance, uint32 SeedOffset) const;
+	FVector4f		RandomSeedOffsetFloat4(uint32 iInstance, uint32 SeedOffset) const;
+
+	FVector2f		SafeNormalize(const FVector2f& v, const FVector2f& Fallback) const { const float l2 = v.SquaredLength(); return l2 < UE_KINDA_SMALL_NUMBER ? Fallback : v * (1.0f / FMath::Sqrt(l2)); }
+	FVector3f		SafeNormalize(const FVector3f& v, const FVector3f& Fallback) const { const float l2 = v.SquaredLength(); return l2 < UE_KINDA_SMALL_NUMBER ? Fallback : v * (1.0f / FMath::Sqrt(l2)); }
+	FVector2f		SafeNormalize(const FVector2f& v) const { return SafeNormalize(v, FVector2f(1.0f, 0.0f)); }
+	FVector3f		SafeNormalize(const FVector3f& v) const { return SafeNormalize(v, FVector3f(1.0f, 0.0f, 0.0f)); }
+
+	template<typename T>
+	const T* ReadBuiltData() const
+	{
+		const int32 Offset = Align(BuiltDataOffset, alignof(T));
+		BuiltDataOffset = Offset + sizeof(T);
+		check(BuiltDataOffset <= BuiltData.Num());
+		return reinterpret_cast<const T*>(BuiltData.GetData() + Offset);
+	}
+
+	//-TODO: Reconcile why this differs from C++ version and match shader accordingly
+	static FQuat4f RotatorToQuat(FVector3f Rotator)
+	{
+		Rotator.X = FMath::Fractional(Rotator.X) * UE_PI;
+		Rotator.Y = FMath::Fractional(Rotator.Y) * UE_PI;
+		Rotator.Z = FMath::Fractional(Rotator.Z) * UE_PI;
+
+		float SP, CP, SY, CY, SR, CR;
+		FMath::SinCos(&SR, &CR, Rotator.X);
+		FMath::SinCos(&SP, &CP, Rotator.Y);
+		FMath::SinCos(&SY, &CY, Rotator.Z);
+
+		return FQuat4f(
+			 CR * SP * SY - SR * CP * CY,
+			-CR * SP * CY - SR * CP * SY,
+			 CR * CP * SY - SR * SP * CY,
+			 CR * CP * CY + SR * SP * SY
+		);
+	}
+
+private:
+	FUintVector4* GetParticleRandomSeed() const { return static_cast<FUintVector4*>(RequiredComponents[int32(EParticleComponent::RandomSeed)]); }
+
+private:
+	const FNiagaraStatelessEmitterData*		EmitterData = nullptr;
+	uint32									NumInstances = 0;
+	float									DeltaTime = 0.0f;
+	float									InvDeltaTime = 0.0f;
+
+	uint32									BufferStride = 0;
+	uint8*									BufferFloatData = nullptr;
+	uint8*									BufferInt32Data = nullptr;
+
+	void*									RequiredComponents[int32(EParticleComponent::Num)] = {};
+	TArray<uint8*, TInlineAllocator<16>>	VariableComponents;
+
+	const TConstArrayView<uint8>			BuiltData;
+	mutable int32							BuiltDataOffset = 0;
+	const TConstArrayView<float>			StaticFloatData;
+	const TConstArrayView<uint8>			DynamicBufferData;
+};
+
+template<typename TType>
+struct FStatelessDistributionSampler
+{
+	explicit FStatelessDistributionSampler(const FParticleSimulationContext& ParticleSimulationContext, const FUintVector3& InParameters, int32 iInstance)
+		: Parameters(InParameters)
+	{
+		if ((Parameters.X & uint32(ENiagaraStatelessBuiltDistributionFlag::Random)) != 0)
+		{
+			RandomOffset = ParticleSimulationContext.RandomFloat(iInstance);
+		}
+	}
+
+	bool IsValid() const { return Parameters.Z > 0; }
+
+	TType GetValue(const FParticleSimulationContext& ParticleSimulationContext, float NormalizedAge) const
+	{
+		if ((Parameters.X & uint32(ENiagaraStatelessBuiltDistributionFlag::Binding)) != 0)
+		{
+			return ParticleSimulationContext.GetParameterBufferFloat<TType>(Parameters.Y, 0);
+		}
+		else if ((Parameters.X & uint32(ENiagaraStatelessBuiltDistributionFlag::Random)) != 0)
+		{
+			//-OPT: Could move into constructor
+			const TType	Value0 = ParticleSimulationContext.GetStaticFloat<TType>(Parameters.Y, 0);
+			const TType	Value1 = ParticleSimulationContext.GetStaticFloat<TType>(Parameters.Y, 1);
+			return FMath::Lerp(Value0, Value1, RandomOffset);
+		}
+		else
+		{
+			const float Offset = NormalizedAge * float(Parameters.Z);
+			const TType	Value0 = ParticleSimulationContext.GetStaticFloat<TType>(Parameters.Y, FMath::FloorToInt(Offset));
+			const TType	Value1 = ParticleSimulationContext.GetStaticFloat<TType>(Parameters.Y, FMath::CeilToInt(Offset));
+			return FMath::Lerp(Value0, Value1, FMath::Fractional(Offset));
+		}
+	}
+
+	const FUintVector3	Parameters;
+	float				RandomOffset = 0.0f;
+};
+
+} //namespace NiagaraStateless
diff --git a/Engine/Plugins/FX/Niagara/Source/Niagara/Internal/Stateless/NiagaraStatelessParticleSimExecData.h b/Engine/Plugins/FX/Niagara/Source/Niagara/Internal/Stateless/NiagaraStatelessParticleSimExecData.h
new file mode 100644
--- /dev/null
+++ b/Engine/Plugins/FX/Niagara/Source/Niagara/Internal/Stateless/NiagaraStatelessParticleSimExecData.h
@@ -0,0 +1,45 @@
+// Copyright Epic Games, Inc. All Rights Reserved.
+
+#pragma once
+
+#include "NiagaraStatelessCommon.h"
+
+struct FNiagaraDataSetCompiledData;
+
+namespace NiagaraStateless
+{
+class FParticleSimulationContext;
+
+class FParticleSimulationExecData
+{
+public:
+	struct FCallback
+	{
+		FCallback() = default;
+		explicit FCallback(TFunction<void(const FParticleSimulationContext&)> InFunc, int32 InBuiltDataOffset)
+			: Function(MoveTemp(InFunc))
+			, BuiltDataOffset(InBuiltDataOffset)
+		{
+		}
+
+		TFunction<void(const FParticleSimulationContext&)>	Function;
+		int32												BuiltDataOffset = 0;
+	};
+
+	struct FVariableOffset
+	{
+		const bool IsFloat() const { return Type == 0; }
+		const bool IsInt32() const { return Type == 1; }
+		const uint32 GetOffset() const { return Offset; }
+
+		uint16	Type	: 1	 = 0;
+		uint16	Offset	: 15 = 0;
+	};
+
+	explicit FParticleSimulationExecData(const FNiagaraDataSetCompiledData& ParticleDataSetCompiledData);
+
+	TArray<FVariableOffset>	VariableComponentOffsets;		// Stored offsets per variable
+	TArray<FCallback>		SimulateFunctions;				// Series of functions to simulate particles
+};
+
+} //namespace NiagaraStateless
diff --git a/Engine/Plugins/FX/Niagara/Source/Niagara/Private/NiagaraModule.cpp b/Engine/Plugins/FX/Niagara/Source/Niagara/Private/NiagaraModule.cpp
--- a/Engine/Plugins/FX/Niagara/Source/Niagara/Private/NiagaraModule.cpp
+++ b/Engine/Plugins/FX/Niagara/Source/Niagara/Private/NiagaraModule.cpp
@@ -524,6 +524,7 @@ void INiagaraModule::OnPostEngineInit()
 	);
 
 	FNiagaraComponentSettings::UpdateSettings();
+	NiagaraStatelessCommon::UpdateSettings();
 }
 
 void INiagaraModule::OnPreExit()
@@ -596,6 +597,7 @@ void INiagaraModule::OnPostGarbageCollect()
 void INiagaraModule::OnBeginFrame()
 {
 	FNiagaraComponentSettings::UpdateSettings();
+	NiagaraStatelessCommon::UpdateSettings();
 	FNiagaraPlatformSet::RefreshScalability();
 }
 
diff --git a/Engine/Plugins/FX/Niagara/Source/Niagara/Private/Stateless/Modules/NiagaraStatelessModule_AddVelocity.cpp b/Engine/Plugins/FX/Niagara/Source/Niagara/Private/Stateless/Modules/NiagaraStatelessModule_AddVelocity.cpp
--- a/Engine/Plugins/FX/Niagara/Source/Niagara/Private/Stateless/Modules/NiagaraStatelessModule_AddVelocity.cpp
+++ b/Engine/Plugins/FX/Niagara/Source/Niagara/Private/Stateless/Modules/NiagaraStatelessModule_AddVelocity.cpp
@@ -3,7 +3,7 @@
 #include "Stateless/Modules/NiagaraStatelessModule_AddVelocity.h"
 #include "Stateless/NiagaraStatelessDrawDebugContext.h"
 
-void UNiagaraStatelessModule_AddVelocity::BuildEmitterData(FNiagaraStatelessEmitterDataBuildContext& BuildContext) const
+void UNiagaraStatelessModule_AddVelocity::BuildEmitterData(const FNiagaraStatelessEmitterDataBuildContext& BuildContext) const
 {
 	if (!IsModuleEnabled())
 	{
diff --git a/Engine/Plugins/FX/Niagara/Source/Niagara/Private/Stateless/Modules/NiagaraStatelessModule_CurlNoiseForce.cpp b/Engine/Plugins/FX/Niagara/Source/Niagara/Private/Stateless/Modules/NiagaraStatelessModule_CurlNoiseForce.cpp
--- a/Engine/Plugins/FX/Niagara/Source/Niagara/Private/Stateless/Modules/NiagaraStatelessModule_CurlNoiseForce.cpp
+++ b/Engine/Plugins/FX/Niagara/Source/Niagara/Private/Stateless/Modules/NiagaraStatelessModule_CurlNoiseForce.cpp
@@ -300,7 +300,7 @@ namespace NoiseHelper
 #undef MGradientScale 
 }
 
-void UNiagaraStatelessModule_CurlNoiseForce::BuildEmitterData(FNiagaraStatelessEmitterDataBuildContext& BuildContext) const
+void UNiagaraStatelessModule_CurlNoiseForce::BuildEmitterData(const FNiagaraStatelessEmitterDataBuildContext& BuildContext) const
 {
 	if (!IsModuleEnabled())
 	{
diff --git a/Engine/Plugins/FX/Niagara/Source/Niagara/Private/Stateless/Modules/NiagaraStatelessModule_MeshIndex.cpp b/Engine/Plugins/FX/Niagara/Source/Niagara/Private/Stateless/Modules/NiagaraStatelessModule_MeshIndex.cpp
--- a/Engine/Plugins/FX/Niagara/Source/Niagara/Private/Stateless/Modules/NiagaraStatelessModule_MeshIndex.cpp
+++ b/Engine/Plugins/FX/Niagara/Source/Niagara/Private/Stateless/Modules/NiagaraStatelessModule_MeshIndex.cpp
@@ -1,6 +1,7 @@
 // Copyright Epic Games, Inc. All Rights Reserved.
 
 #include "Stateless/Modules/NiagaraStatelessModule_MeshIndex.h"
+#include "Stateless/NiagaraStatelessParticleSimContext.h"
 
 #include "WeightedRandomSampler.h"
 
@@ -41,46 +42,95 @@ namespace NSMMeshIndexPrivate
 		int	Index				= 0;
 		int	TableOffset			= 0;
 		int	TableNumElements	= 0;
+
+		int MeshIndexOffset		= INDEX_NONE;
 	};
+
+	static void ParticleSimulate(const NiagaraStateless::FParticleSimulationContext& ParticleSimulationContext)
+	{
+		using namespace NiagaraStateless;
+
+		const FModuleBuiltData* ModuleBuiltData = ParticleSimulationContext.ReadBuiltData<FModuleBuiltData>();
+
+		const bool bIsParameter = (ModuleBuiltData->Index & 0x80000000) != 0;
+		int32 MeshIndex = ModuleBuiltData->Index & ~0x80000000;;
+		if (bIsParameter)
+		{
+			MeshIndex = ParticleSimulationContext.GetParameterBufferInt(MeshIndex, 0);
+		}
+
+		if (ModuleBuiltData->TableNumElements > 0)
+		{
+			for (uint32 i = 0; i < ParticleSimulationContext.GetNumInstances(); ++i)
+			{
+				const FVector2f Rand = ParticleSimulationContext.RandomFloat2(i);
+				MeshIndex = FMath::RoundToInt(Rand.X * float(ModuleBuiltData->TableNumElements));
+				const float Probability = ParticleSimulationContext.GetStaticFloat<float>(ModuleBuiltData->TableOffset, (MeshIndex * 2) + 0);
+				if (Rand.Y > Probability)
+				{
+					MeshIndex = ParticleSimulationContext.GetStaticFloat<float>(ModuleBuiltData->TableOffset, (MeshIndex * 2) + 1);
+				}
+				ParticleSimulationContext.WriteParticleVariable(ModuleBuiltData->MeshIndexOffset, i, MeshIndex);
+			}
+		}
+		else
+		{
+			for (uint32 i = 0; i < ParticleSimulationContext.GetNumInstances(); ++i)
+			{
+				ParticleSimulationContext.WriteParticleVariable(ModuleBuiltData->MeshIndexOffset, i, MeshIndex);
+			}
+		}
+	}
 }
 
-void UNiagaraStatelessModule_MeshIndex::BuildEmitterData(FNiagaraStatelessEmitterDataBuildContext& BuildContext) const
+void UNiagaraStatelessModule_MeshIndex::BuildEmitterData(const FNiagaraStatelessEmitterDataBuildContext& BuildContext) const
 {
 	using namespace NSMMeshIndexPrivate;
 
 	FModuleBuiltData* BuiltData = BuildContext.AllocateBuiltData<FModuleBuiltData>();
-	if (IsModuleEnabled())
+	if (!IsModuleEnabled())
 	{
-		const FNiagaraStatelessRangeInt MeshIndexRange = BuildContext.ConvertDistributionToRange(MeshIndex, 0, true);
-		if (MeshIndexRange.ParameterOffset != INDEX_NONE)
-		{
-			BuiltData->Index = MeshIndexRange.ParameterOffset | 0x80000000;
-		}
-		else
+		return;
+	}
+
+	const FNiagaraStatelessGlobals& StatelessGlobals = FNiagaraStatelessGlobals::Get();
+	BuiltData->MeshIndexOffset = BuildContext.FindParticleVariableIndex(StatelessGlobals.MeshIndexVariable);
+	if (BuiltData->MeshIndexOffset == INDEX_NONE)
+	{
+		return;
+	}
+
+	const FNiagaraStatelessRangeInt MeshIndexRange = BuildContext.ConvertDistributionToRange(MeshIndex, 0);
+	if (MeshIndexRange.ParameterOffset != INDEX_NONE)
+	{
+		BuiltData->Index = MeshIndexRange.ParameterOffset | 0x80000000;
+	}
+	else
+	{
+		BuiltData->Index = MeshIndexRange.Min;
+		if (MeshIndexRange.GetScale() > 0 && MeshIndexRange.GetScale() < 256)
 		{
-			BuiltData->Index = MeshIndexRange.Min;
-			if (MeshIndexRange.GetScale() > 0 && MeshIndexRange.GetScale() < 256)
+			FMeshIndexWeightedSampler Sampler(MeshIndexRange.GetScale() + 1, MeshIndexWeight);
+			Sampler.Initialize();
+
+			const int32 NumTableEntries = Sampler.GetNumEntries();
+			if (NumTableEntries > 1)
 			{
-				FMeshIndexWeightedSampler Sampler(MeshIndexRange.GetScale() + 1, MeshIndexWeight);
-				Sampler.Initialize();
+				BuiltData->TableNumElements = NumTableEntries - 1;
 
-				const int32 NumTableEntries = Sampler.GetNumEntries();
-				if (NumTableEntries > 1)
+				TArray<float, TInlineAllocator<16>> StaticData;
+				StaticData.AddUninitialized(NumTableEntries * 2);
+				for (int32 i = 0; i < NumTableEntries; ++i)
 				{
-					BuiltData->TableNumElements = NumTableEntries - 1;
-
-					TArray<float, TInlineAllocator<16>> StaticData;
-					StaticData.AddUninitialized(NumTableEntries * 2);
-					for (int32 i = 0; i < NumTableEntries; ++i)
-					{
-						StaticData[i * 2 + 0] = Sampler.GetProb()[i];
-						StaticData[i * 2 + 1] = float(MeshIndexRange.Min + Sampler.GetAlias()[i]);
-					}
-					BuiltData->TableOffset = BuildContext.AddStaticData(StaticData);
+					StaticData[i * 2 + 0] = Sampler.GetProb()[i];
+					StaticData[i * 2 + 1] = float(MeshIndexRange.Min + Sampler.GetAlias()[i]);
 				}
+				BuiltData->TableOffset = BuildContext.AddStaticData(StaticData);
 			}
 		}
 	}
+
+	BuildContext.AddParticleSimulationExecSimulate(&ParticleSimulate);
 }
 
 void UNiagaraStatelessModule_MeshIndex::SetShaderParameters(const FNiagaraStatelessSetShaderParameterContext& SetShaderParameterContext) const
diff --git a/Engine/Plugins/FX/Niagara/Source/Niagara/Private/Stateless/Modules/NiagaraStatelessModule_ShapeLocation.cpp b/Engine/Plugins/FX/Niagara/Source/Niagara/Private/Stateless/Modules/NiagaraStatelessModule_ShapeLocation.cpp
--- a/Engine/Plugins/FX/Niagara/Source/Niagara/Private/Stateless/Modules/NiagaraStatelessModule_ShapeLocation.cpp
+++ b/Engine/Plugins/FX/Niagara/Source/Niagara/Private/Stateless/Modules/NiagaraStatelessModule_ShapeLocation.cpp
@@ -1,42 +1,152 @@
 // Copyright Epic Games, Inc. All Rights Reserved.
 
 #include "Stateless/Modules/NiagaraStatelessModule_ShapeLocation.h"
+#include "Stateless/NiagaraStatelessEmitterDataBuildContext.h"
+#include "Stateless/NiagaraStatelessModuleShaderParameters.h"
+#include "Stateless/NiagaraStatelessParticleSimContext.h"
 #include "Stateless/NiagaraStatelessDrawDebugContext.h"
 
-void UNiagaraStatelessModule_ShapeLocation::SetShaderParameters(const FNiagaraStatelessSetShaderParameterContext& SetShaderParameterContext) const
+namespace NSMShapeLocationPrivate
 {
-	FParameters* Parameters = SetShaderParameterContext.GetParameterNestedStruct<FParameters>();
+	struct FModuleBuiltData
+	{
+		FUintVector4	Mode		= FUintVector4::ZeroValue;
+		FVector4f		Parameters0	= FVector4f::Zero();
+		FVector4f		Parameters1 = FVector4f::Zero();
+
+		int32			PositionVariableOffset = INDEX_NONE;
+		int32			PreviousPositionVariableOffset = INDEX_NONE;
+	};
+
+	FVector3f ShapeLocation_GetLocation(const NiagaraStateless::FParticleSimulationContext& ParticleSimulationContext, const FModuleBuiltData* ModuleBuiltData, uint32 iInstance)
+	{
+		// ENSM_ShapePrimitive::Box | ENSM_ShapePrimitive::Plane
+		if ( ModuleBuiltData->Mode.X == 0 )
+		{
+			const FVector3f BoxScale = FVector3f(ModuleBuiltData->Parameters0);
+			const FVector3f BoxBias = FVector3f(ModuleBuiltData->Parameters1);
+			const bool bOnSurface = ModuleBuiltData->Mode.Y == 1;
+			const float SurfaceScale = ModuleBuiltData->Parameters0.X;
+			const float SurfaceBias = ModuleBuiltData->Parameters1.W;
+	
+			const FVector3f P0 = ParticleSimulationContext.RandomFloat3(iInstance);
+			if (bOnSurface)
+			{
+				//-TODO: This isn't quite what we want
+				const FVector3f L0(FMath::RoundToFloat(P0.X), FMath::RoundToFloat(P0.Y), FMath::RoundToFloat(P0.Z));
+				const uint32 S = ParticleSimulationContext.RandomUInt(iInstance) % 3;
+				FVector3f Location;
+				Location.X = S != 0 ? P0.X : L0.X;
+				Location.Y = S != 1 ? P0.Y : L0.Y;
+				Location.Z = S != 2 ? P0.Z : L0.Z;
+				return (Location * BoxScale + BoxBias) + (P0 * SurfaceScale + SurfaceBias);
+			}
+			else
+			{
+				return P0 * BoxScale + BoxBias;
+			}
+		}
+	
+		// ENSM_ShapePrimitive::Cylinder:
+		if ( ModuleBuiltData->Mode.X == 1 )
+		{
+			const FVector4f Random = ParticleSimulationContext.RandomFloat4(iInstance);
+			const float HeightScale = ModuleBuiltData->Parameters0.X;
+			const float HeightBias = ModuleBuiltData->Parameters0.Y;
+			const float Radius = ModuleBuiltData->Parameters0.Z;
+	
+			const FVector2f UnitVec = ParticleSimulationContext.SafeNormalize(FVector2f(Random.X - 0.5f, Random.Y - 0.5f));
+			return FVector3f(
+				UnitVec.X * Radius * Random.Z,
+				UnitVec.Y * Radius * Random.Z,
+				Random.W * HeightScale + HeightBias
+			);
+		}
+	
+		// ENSM_ShapePrimitive::Ring:
+		if ( ModuleBuiltData->Mode.X == 2 )
+		{
+			const float RadiusScale = ModuleBuiltData->Parameters0.X;
+			const float RadiusBias = ModuleBuiltData->Parameters0.Y;
+			const float UDistributionScale = ModuleBuiltData->Parameters0.Z;
+			const float UDistributionBias = ModuleBuiltData->Parameters0.W;
+	
+			const float Radius = ParticleSimulationContext.RandomScaleBiasFloat(iInstance, RadiusScale, RadiusBias);
+			const float U = ParticleSimulationContext.RandomScaleBiasFloat(iInstance, UDistributionScale, UDistributionBias);
+	
+			return FVector3f(cos(U) * Radius, sin(U) * Radius, 0.0f);
+		}
+	
+		// ENSM_ShapePrimitive::Sphere:
+		{
+			const float SphereScale	= ModuleBuiltData->Parameters0.X;
+			const float SphereBias	= ModuleBuiltData->Parameters0.Y;
+	
+			const FVector3f Vector = ParticleSimulationContext.RandomUnitFloat3(iInstance);
+			return Vector * ParticleSimulationContext.RandomScaleBiasFloat(iInstance, SphereScale, SphereBias);
+		}
+	}
+
+	static void ParticleSimulate(const NiagaraStateless::FParticleSimulationContext& ParticleSimulationContext)
+	{
+		using namespace NiagaraStateless;
+
+		const FModuleBuiltData* ModuleBuiltData = ParticleSimulationContext.ReadBuiltData<FModuleBuiltData>();
+
+		for (uint32 i = 0; i < ParticleSimulationContext.GetNumInstances(); ++i)
+		{
+			const FVector3f ShapeLocation = ShapeLocation_GetLocation(ParticleSimulationContext, ModuleBuiltData, i);
+			const FVector3f Position = ParticleSimulationContext.ReadParticleVariable(ModuleBuiltData->PositionVariableOffset, i, FVector3f::ZeroVector);
+			const FVector3f PreviousPosition = ParticleSimulationContext.ReadParticleVariable(ModuleBuiltData->PreviousPositionVariableOffset, i, FVector3f::ZeroVector);
+
+			ParticleSimulationContext.WriteParticleVariable(ModuleBuiltData->PositionVariableOffset, i, Position + ShapeLocation);
+			ParticleSimulationContext.WriteParticleVariable(ModuleBuiltData->PreviousPositionVariableOffset, i, PreviousPosition + ShapeLocation);
+		}
+	}
+}
+
+void UNiagaraStatelessModule_ShapeLocation::BuildEmitterData(const FNiagaraStatelessEmitterDataBuildContext& BuildContext) const
+{
+	using namespace NSMShapeLocationPrivate;
+
+	FModuleBuiltData* BuiltData = BuildContext.AllocateBuiltData<FModuleBuiltData>();
 	if (!IsModuleEnabled())
 	{
-		Parameters->ShapeLocation_Mode.X		= 0;
-		Parameters->ShapeLocation_Parameters0	= FVector4f::Zero();
-		Parameters->ShapeLocation_Parameters1	= FVector4f::Zero();
 		return;
 	}
 
+	const FNiagaraStatelessGlobals& StatelessGlobals = FNiagaraStatelessGlobals::Get();
+	BuiltData->PositionVariableOffset			= BuildContext.FindParticleVariableIndex(StatelessGlobals.PositionVariable);
+	BuiltData->PreviousPositionVariableOffset	= BuildContext.FindParticleVariableIndex(StatelessGlobals.PreviousPositionVariable);
+	if (BuiltData->PositionVariableOffset == INDEX_NONE && BuiltData->PreviousPositionVariableOffset == INDEX_NONE)
+	{
+		return;
+	}
+
+	//-TODO: Build Parameters
 	switch (ShapePrimitive)
 	{
 		case ENSM_ShapePrimitive::Box:
 		{
-			Parameters->ShapeLocation_Mode.X		= 0;
-			Parameters->ShapeLocation_Mode.Y		= bBoxSurfaceOnly ? 1 : 0;
-			Parameters->ShapeLocation_Parameters0	= FVector4f(BoxSize, BoxSurfaceThicknessMax - BoxSurfaceThicknessMin);
-			Parameters->ShapeLocation_Parameters1	= FVector4f(BoxSize * -0.5f, BoxSurfaceThicknessMin);
+			BuiltData->Mode.X		= 0;
+			BuiltData->Mode.Y		= bBoxSurfaceOnly ? 1 : 0;
+			BuiltData->Parameters0	= FVector4f(BoxSize, BoxSurfaceThicknessMax - BoxSurfaceThicknessMin);
+			BuiltData->Parameters1	= FVector4f(BoxSize * -0.5f, BoxSurfaceThicknessMin);
 			break;
 		}
 		case ENSM_ShapePrimitive::Plane:
 		{
-			Parameters->ShapeLocation_Mode.X		= 0;
-			Parameters->ShapeLocation_Parameters0	= FVector4f(PlaneSize.X, PlaneSize.Y, 0.0f, 0.0f);
-			Parameters->ShapeLocation_Parameters1	= FVector4f(-PlaneSize.X * 0.5f, -PlaneSize.Y * 0.5f, 0.0f, 0.0f);
+			BuiltData->Mode.X		= 0;
+			BuiltData->Parameters0	= FVector4f(PlaneSize.X, PlaneSize.Y, 0.0f, 0.0f);
+			BuiltData->Parameters1	= FVector4f(-PlaneSize.X * 0.5f, -PlaneSize.Y * 0.5f, 0.0f, 0.0f);
 			break;
 		}
 		case ENSM_ShapePrimitive::Cylinder:
 		{
-			Parameters->ShapeLocation_Mode.X		= 1;
-			Parameters->ShapeLocation_Parameters0.X = CylinderHeight;
-			Parameters->ShapeLocation_Parameters0.Y = CylinderHeight * -CylinderHeightMidpoint;
-			Parameters->ShapeLocation_Parameters0.Z = CylinderRadius;
+			BuiltData->Mode.X		= 1;
+			BuiltData->Parameters0.X = CylinderHeight;
+			BuiltData->Parameters0.Y = CylinderHeight * -CylinderHeightMidpoint;
+			BuiltData->Parameters0.Z = CylinderRadius;
 			break;
 		}
 		case ENSM_ShapePrimitive::Ring:
@@ -44,28 +154,42 @@ void UNiagaraStatelessModule_ShapeLocation::SetShaderParameters(const FNiagaraSt
 			const float DC = FMath::Clamp(1.0f - DiscCoverage, 0.0f, 1.0f);
 			const float SDC = DC > 0.0f ? FMath::Sqrt(DC) : 0.0f;
 
-			Parameters->ShapeLocation_Mode.X		= 2;
-			Parameters->ShapeLocation_Parameters0.X = RingRadius * (1.0f - SDC);
-			Parameters->ShapeLocation_Parameters0.Y = RingRadius * SDC;
-			Parameters->ShapeLocation_Parameters0.Z = -UE_TWO_PI * (1.0f - RingUDistribution);
-			Parameters->ShapeLocation_Parameters0.W = 0.0f;
+			BuiltData->Mode.X		= 2;
+			BuiltData->Parameters0.X = RingRadius * (1.0f - SDC);
+			BuiltData->Parameters0.Y = RingRadius * SDC;
+			BuiltData->Parameters0.Z = -UE_TWO_PI * (1.0f - RingUDistribution);
+			BuiltData->Parameters0.W = 0.0f;
 			break;
 		}
 		case ENSM_ShapePrimitive::Sphere:
 		{
-			Parameters->ShapeLocation_Mode.X		= 3;
-			Parameters->ShapeLocation_Parameters0.X = SphereMax - SphereMin;
-			Parameters->ShapeLocation_Parameters0.Y = SphereMin;
+			BuiltData->Mode.X		= 3;
+			BuiltData->Parameters0.X = SphereMax - SphereMin;
+			BuiltData->Parameters0.Y = SphereMin;
 			break;
 		}
 		default:
 		{
-			Parameters->ShapeLocation_Mode.X = 0;
-			Parameters->ShapeLocation_Parameters0 = FVector4f::Zero();
-			Parameters->ShapeLocation_Parameters1 = FVector4f::Zero();
+			BuiltData->Mode.X = 0;
+			BuiltData->Parameters0 = FVector4f::Zero();
+			BuiltData->Parameters1 = FVector4f::Zero();
 			checkNoEntry();
 		}
 	}
+
+	BuildContext.AddParticleSimulationExecSimulate(&ParticleSimulate);
+}
+
+void UNiagaraStatelessModule_ShapeLocation::SetShaderParameters(const FNiagaraStatelessSetShaderParameterContext& SetShaderParameterContext) const
+{
+	using namespace NSMShapeLocationPrivate;
+
+	FParameters* Parameters = SetShaderParameterContext.GetParameterNestedStruct<FParameters>();
+	const FModuleBuiltData* ModuleBuiltData = SetShaderParameterContext.ReadBuiltData<FModuleBuiltData>();
+
+	Parameters->ShapeLocation_Mode			= ModuleBuiltData->Mode;
+	Parameters->ShapeLocation_Parameters0	= ModuleBuiltData->Parameters0;
+	Parameters->ShapeLocation_Parameters1	= ModuleBuiltData->Parameters1;
 }
 
 #if WITH_EDITOR
diff --git a/Engine/Plugins/FX/Niagara/Source/Niagara/Private/Stateless/NiagaraStatelessCommon.cpp b/Engine/Plugins/FX/Niagara/Source/Niagara/Private/Stateless/NiagaraStatelessCommon.cpp
--- a/Engine/Plugins/FX/Niagara/Source/Niagara/Private/Stateless/NiagaraStatelessCommon.cpp
+++ b/Engine/Plugins/FX/Niagara/Source/Niagara/Private/Stateless/NiagaraStatelessCommon.cpp
@@ -3,16 +3,52 @@
 #include "Stateless/NiagaraStatelessCommon.h"
 
 #include "Stateless/NiagaraStatelessEmitterTemplate.h"
+#include "NiagaraSystem.h"
+
+#include "UObject/UObjectIterator.h"
 
 namespace NiagaraStatelessCommon
 {
-	static FNiagaraStatelessGlobals GGlobals;
+	static FNiagaraStatelessGlobals					GGlobals;
+	static TOptional<ENiagaraStatelessFeatureMask>	GUpdatedFeatureMask;
+
+	static void SetUpdateFeatureMask(ENiagaraStatelessFeatureMask Mask, bool bEnabled)
+	{
+		ENiagaraStatelessFeatureMask NewMask = GUpdatedFeatureMask.Get(GGlobals.FeatureMask);
+		if (bEnabled)
+		{
+			EnumAddFlags(NewMask, Mask);
+		}
+		else
+		{
+			EnumRemoveFlags(NewMask, Mask);
+		}
+		GUpdatedFeatureMask = NewMask;
+	}
+
+	#define DEFINE_FEATURE_MASK(ENUM, DESC) \
+		static bool bFeatureMask_##ENUM = true; \
+		static FAutoConsoleVariableRef CVarFeatureMask_##ENUM( \
+			TEXT("fx.NiagaraStateless.Feature."#ENUM), \
+			bFeatureMask_##ENUM, \
+			TEXT(DESC), \
+			FConsoleVariableDelegate::CreateLambda([](IConsoleVariable*) {SetUpdateFeatureMask(ENiagaraStatelessFeatureMask::ENUM, bFeatureMask_##ENUM);}), \
+			ECVF_Scalability | ECVF_Default \
+		)
+
+		DEFINE_FEATURE_MASK(ExecuteGPU, "When enabled simulations are allowed to execute on the GPU");
+		DEFINE_FEATURE_MASK(ExecuteCPU, "When enabled simulations are allowed to execute on the CPU");
+
+	#undef DEFINE_FEATURE_MASK
 
 	void Initialize()
 	{
 		GGlobals.CameraOffsetVariable				= FNiagaraVariableBase(FNiagaraTypeDefinition::GetFloatDef(), TEXT("CameraOffset"));
 		GGlobals.ColorVariable						= FNiagaraVariableBase(FNiagaraTypeDefinition::GetColorDef(), TEXT("Color"));
 		GGlobals.DynamicMaterialParameters0Variable = FNiagaraVariableBase(FNiagaraTypeDefinition::GetVec4Def(), TEXT("DynamicMaterialParameter"));
+		GGlobals.DynamicMaterialParameters1Variable = FNiagaraVariableBase(FNiagaraTypeDefinition::GetVec4Def(), TEXT("DynamicMaterialParameter1"));
+		GGlobals.DynamicMaterialParameters2Variable = FNiagaraVariableBase(FNiagaraTypeDefinition::GetVec4Def(), TEXT("DynamicMaterialParameter2"));
+		GGlobals.DynamicMaterialParameters3Variable = FNiagaraVariableBase(FNiagaraTypeDefinition::GetVec4Def(), TEXT("DynamicMaterialParameter3"));
 		GGlobals.MeshIndexVariable					= FNiagaraVariableBase(FNiagaraTypeDefinition::GetIntDef(), TEXT("MeshIndex"));
 		GGlobals.MeshOrientationVariable			= FNiagaraVariableBase(FNiagaraTypeDefinition::GetQuatDef(), TEXT("MeshOrientation"));
 		GGlobals.PositionVariable					= FNiagaraVariableBase(FNiagaraTypeDefinition::GetPositionDef(), TEXT("Position"));
@@ -41,6 +77,28 @@ namespace NiagaraStatelessCommon
 
 		UNiagaraStatelessEmitterTemplate::InitCDOPropertiesAfterModuleStartup();
 	}
+
+	void UpdateSettings()
+	{
+		if (!GUpdatedFeatureMask.IsSet())
+		{
+			return;
+		}
+
+		if (GUpdatedFeatureMask.GetValue() != GGlobals.FeatureMask)
+		{
+			GGlobals.FeatureMask = GUpdatedFeatureMask.GetValue();
+
+			for (TObjectIterator<UNiagaraSystem> It; It; ++It)
+			{
+				if (UNiagaraSystem* System = *It)
+				{
+					System->UpdateScalability();
+				}
+			}
+		}
+		GUpdatedFeatureMask.Reset();
+	}
 } //NiagaraStatelessCommon
 
 const FNiagaraStatelessGlobals& FNiagaraStatelessGlobals::Get()
diff --git a/Engine/Plugins/FX/Niagara/Source/Niagara/Private/Stateless/NiagaraStatelessComputeManager.cpp b/Engine/Plugins/FX/Niagara/Source/Niagara/Private/Stateless/NiagaraStatelessComputeManager.cpp
--- a/Engine/Plugins/FX/Niagara/Source/Niagara/Private/Stateless/NiagaraStatelessComputeManager.cpp
+++ b/Engine/Plugins/FX/Niagara/Source/Niagara/Private/Stateless/NiagaraStatelessComputeManager.cpp
@@ -4,6 +4,7 @@
 #include "Stateless/NiagaraStatelessEmitterData.h"
 #include "Stateless/NiagaraStatelessEmitterInstance.h"
 #include "Stateless/NiagaraStatelessSimulationShader.h"
+#include "Stateless/NiagaraStatelessParticleSimContext.h"
 
 #include "NiagaraGpuComputeDispatchInterface.h"
 #include "NiagaraGPUInstanceCountManager.h"
@@ -15,7 +16,7 @@
 #include "RenderGraphUtils.h"
 #include "SceneView.h"
 
-bool GNiagaraStatelessComputeManager_UseCache = 1;
+bool GNiagaraStatelessComputeManager_UseCache = true;
 FAutoConsoleVariableRef CVarNiagaraStatelessComputeManager_UseCache(
 	TEXT("fx.NiagaraStateless.ComputeManager.UseCache"),
 	GNiagaraStatelessComputeManager_UseCache,
@@ -23,6 +24,14 @@ FAutoConsoleVariableRef CVarNiagaraStatelessComputeManager_UseCache(
 	ECVF_Default
 );
 
+int32 GNiagaraStatelessComputeManager_CPUThreshold = 0;
+FAutoConsoleVariableRef CVarNiagaraStatelessComputeManager_CPUThreshold(
+	TEXT("fx.NiagaraStateless.ComputeManager.CPUThreshold"),
+	GNiagaraStatelessComputeManager_CPUThreshold,
+	TEXT("When lower than this particle count prefer to use the CPU over dispatching a compute shader."),
+	ECVF_Default
+);
+
 FNiagaraStatelessComputeManager::FNiagaraStatelessComputeManager(FNiagaraGpuComputeDispatchInterface* InOwnerInterface)
 	: FNiagaraGpuComputeDataManager(InOwnerInterface)
 {
@@ -36,6 +45,8 @@ FNiagaraStatelessComputeManager::~FNiagaraStatelessComputeManager()
 
 FNiagaraDataBuffer* FNiagaraStatelessComputeManager::GetDataBuffer(uintptr_t EmitterKey, const NiagaraStateless::FEmitterInstance_RT* EmitterInstance)
 {
+	using namespace NiagaraStateless;
+
 	if (TUniquePtr<FStatelessDataCache>* ExistingData = UsedData.Find(EmitterKey))
 	{
 		return (*ExistingData)->DataBuffer;
@@ -97,13 +108,40 @@ FNiagaraDataBuffer* FNiagaraStatelessComputeManager::GetDataBuffer(uintptr_t Emi
 	CacheData->EmitterInstance = EmitterInstance;
 	CacheData->ActiveParticles = ActiveParticles;
 
-	CacheData->DataBuffer->AllocateGPU(FRHICommandListExecutor::GetImmediateCommandList(), CacheData->ActiveParticles, ComputeInterface->GetFeatureLevel(), TEXT("StatelessSimBuffer"));
-	CacheData->DataBuffer->SetNumInstances(CacheData->ActiveParticles);
+	FRHICommandListImmediate& RHICmdList = FRHICommandListExecutor::GetImmediateCommandList();
+	CacheData->DataBuffer->AllocateGPU(RHICmdList, CacheData->ActiveParticles, ComputeInterface->GetFeatureLevel(), TEXT("StatelessSimBuffer"));
 
-	FNiagaraGPUInstanceCountManager& CountManager = ComputeInterface->GetGPUInstanceCounterManager();
-	const uint32 CountOffset = CountManager.AcquireOrAllocateEntry(FRHICommandListExecutor::GetImmediateCommandList());
-	CacheData->DataBuffer->SetGPUInstanceCountBufferOffset(CountOffset);
-	CountsToRelease.Add(CountOffset);
+	const bool bAllowGPUExec = EnumHasAnyFlags(EmitterData->FeatureMask, ENiagaraStatelessFeatureMask::ExecuteGPU);
+	const bool bUseCPUExec = EnumHasAnyFlags(EmitterData->FeatureMask, ENiagaraStatelessFeatureMask::ExecuteCPU) && (!bAllowGPUExec || (ActiveParticles <= uint32(GNiagaraStatelessComputeManager_CPUThreshold)));
+	if (bUseCPUExec)
+	{
+		FParticleSimulationContext ParticleSimulation(EmitterData, EmitterInstance->BindingBufferData.Get(TArray<uint8>()));
+		ParticleSimulation.SimulateGPU(RHICmdList, EmitterInstance->RandomSeed, EmitterInstance->Age, EmitterInstance->DeltaTime, EmitterInstance->SpawnInfos, CacheData->DataBuffer);
+		if (ParticleSimulation.GetNumInstances() == 0)
+		{
+			FreeData.Emplace(CacheData);
+			return nullptr;
+		}
+
+		//CacheData->DataBuffer->Dump(0, CacheData->DataBuffer->GetNumInstances(), TEXT("Stateless"));
+	}
+	else
+	{
+		if (!ensure(bAllowGPUExec))
+		{
+			FreeData.Emplace(CacheData);
+			return nullptr;
+		}
+
+		CacheData->DataBuffer->SetNumInstances(CacheData->ActiveParticles);
+
+		FNiagaraGPUInstanceCountManager& CountManager = ComputeInterface->GetGPUInstanceCounterManager();
+		const uint32 CountOffset = CountManager.AcquireOrAllocateEntry(RHICmdList);
+		CacheData->DataBuffer->SetGPUInstanceCountBufferOffset(CountOffset);
+		CountsToRelease.Add(CountOffset);
+
+		GPUDataToGenerate.Add(CacheData);
+	}
 
 	UsedData.Emplace(EmitterKey, CacheData);
 	return CacheData->DataBuffer;
@@ -112,7 +150,7 @@ FNiagaraDataBuffer* FNiagaraStatelessComputeManager::GetDataBuffer(uintptr_t Emi
 void FNiagaraStatelessComputeManager::OnPostPreRender(FRDGBuilder& GraphBuilder)
 {
 	// Anything to process?
-	if (UsedData.Num() == 0)
+	if (GPUDataToGenerate.Num() == 0)
 	{
 		return;
 	}
@@ -123,12 +161,11 @@ void FNiagaraStatelessComputeManager::OnPostPreRender(FRDGBuilder& GraphBuilder)
 	AddPass(
 		GraphBuilder,
 		RDG_EVENT_NAME("FNiagaraStatelessComputeManager::OnPostPreRender"),
-		[this](FRHICommandListImmediate& RHICmdList)
+		[DataToGenerate=MoveTemp(GPUDataToGenerate), ComputeInterface=GetOwnerInterface()](FRHICommandListImmediate& RHICmdList)
 		{
-			const int32 NumJobs = UsedData.Num();
+			const int32 NumJobs = DataToGenerate.Num();
 
 			// Get Count Buffer
-			FNiagaraGpuComputeDispatchInterface* ComputeInterface = GetOwnerInterface();
 			FNiagaraGPUInstanceCountManager& CountManager = ComputeInterface->GetGPUInstanceCounterManager();
 			FRHIUnorderedAccessView* CountBufferUAV = CountManager.GetInstanceCountBuffer().UAV;
 
@@ -142,10 +179,8 @@ void FNiagaraStatelessComputeManager::OnPostPreRender(FRDGBuilder& GraphBuilder)
 				TransitionsBefore.Emplace(CountManager.GetInstanceCountBuffer().Buffer, FNiagaraGPUInstanceCountManager::kCountBufferDefaultState, ERHIAccess::UAVCompute);
 				TransitionsAfter.Emplace(CountManager.GetInstanceCountBuffer().Buffer, ERHIAccess::UAVCompute, FNiagaraGPUInstanceCountManager::kCountBufferDefaultState);
 
-				for (auto it = UsedData.CreateIterator(); it; ++it)
+				for (FStatelessDataCache* CacheData : DataToGenerate)
 				{
-					FStatelessDataCache* CacheData = it.Value().Get();
-
 					const FRWBuffer& FloatBuffer = CacheData->DataBuffer->GetGPUBufferFloat();
 					if (FloatBuffer.NumBytes > 0)
 					{
@@ -161,17 +196,15 @@ void FNiagaraStatelessComputeManager::OnPostPreRender(FRDGBuilder& GraphBuilder)
 				}
 			}
 
-			FNiagaraGpuComputeDispatchInterface* ComputeDispatchInterface = GetOwnerInterface();
-			FNiagaraEmptyUAVPoolScopedAccess UAVPoolAccessScope(ComputeDispatchInterface->GetEmptyUAVPool());
-			FRHIUnorderedAccessView* EmptyIntBufferUAV = ComputeDispatchInterface->GetEmptyUAVFromPool(RHICmdList, PF_R32_SINT, ENiagaraEmptyUAVType::Buffer);
+			FNiagaraEmptyUAVPoolScopedAccess UAVPoolAccessScope(ComputeInterface->GetEmptyUAVPool());
+			FRHIUnorderedAccessView* EmptyIntBufferUAV = ComputeInterface->GetEmptyUAVFromPool(RHICmdList, PF_R32_SINT, ENiagaraEmptyUAVType::Buffer);
 
 			// Execute Simulations
 			RHICmdList.Transition(TransitionsBefore);
 
 			RHICmdList.BeginUAVOverlap(CountBufferUAV);
-			for (auto it = UsedData.CreateIterator(); it; ++it)
+			for (FStatelessDataCache* CacheData : DataToGenerate)
 			{
-				FStatelessDataCache* CacheData = it.Value().Get();
 				const NiagaraStateless::FEmitterInstance_RT* EmitterInstance = CacheData->EmitterInstance;
 				const FNiagaraStatelessEmitterData* EmitterData = CacheData->EmitterInstance->EmitterData.Get();
 
@@ -189,8 +222,8 @@ void FNiagaraStatelessComputeManager::OnPostPreRender(FRDGBuilder& GraphBuilder)
 				// Update parameters for this compute invocation
 				NiagaraStateless::FCommonShaderParameters* ShaderParameters = EmitterInstance->ShaderParameters.Get();
 				ShaderParameters->Common_SimulationTime			= EmitterInstance->Age;
-				ShaderParameters->Common_SimulationDeltaTime	= 1.0f / 60.0f;		//-TODO: Pull from view information, needs reworking of how we link to the dispatch
-				ShaderParameters->Common_SimulationInvDeltaTime	= 60.0f;
+				ShaderParameters->Common_SimulationDeltaTime	= EmitterInstance->DeltaTime;
+				ShaderParameters->Common_SimulationInvDeltaTime	= EmitterInstance->DeltaTime > 0.0f ? (1.0f / EmitterInstance->DeltaTime) : 0.0f;
 				ShaderParameters->Common_OutputBufferStride		= CacheData->DataBuffer->GetFloatStride() / sizeof(float);
 				ShaderParameters->Common_GPUCountBufferOffset	= CacheData->DataBuffer->GetGPUInstanceCountBufferOffset();
 				ShaderParameters->Common_FloatOutputBuffer		= CacheData->DataBuffer->GetGPUBufferFloat().UAV;
@@ -215,6 +248,7 @@ void FNiagaraStatelessComputeManager::OnPostPreRender(FRDGBuilder& GraphBuilder)
 			RHICmdList.Transition(TransitionsAfter);
 		}
 	);
+	GPUDataToGenerate.Empty();
 }
 
 void FNiagaraStatelessComputeManager::OnPostPostRender(FRDGBuilder& GraphBuilder)
@@ -233,17 +267,20 @@ void FNiagaraStatelessComputeManager::OnPostPostRender(FRDGBuilder& GraphBuilder
 		RDG_EVENT_NAME("FNiagaraStatelessComputeManager::OnPostPostRender"),
 		[this](FRHICommandListImmediate& RHICmdList)
 		{
-			FreeData.Empty();
+			FreeData.Empty(UsedData.Num());
 			for (auto it=UsedData.CreateIterator(); it; ++it)
 			{
 				FreeData.Emplace(it.Value().Release());
 			}
 			UsedData.Empty();
 
-			FNiagaraGpuComputeDispatchInterface* ComputeInterface = GetOwnerInterface();
-			FNiagaraGPUInstanceCountManager& CountManager = ComputeInterface->GetGPUInstanceCounterManager();
-			CountManager.FreeEntryArray(CountsToRelease);
-			CountsToRelease.Reset();
+			if (CountsToRelease.Num() > 0)
+			{
+				FNiagaraGpuComputeDispatchInterface* ComputeInterface = GetOwnerInterface();
+				FNiagaraGPUInstanceCountManager& CountManager = ComputeInterface->GetGPUInstanceCounterManager();
+				CountManager.FreeEntryArray(CountsToRelease);
+				CountsToRelease.Reset();
+			}
 		}
 	);
 }
diff --git a/Engine/Plugins/FX/Niagara/Source/Niagara/Private/Stateless/NiagaraStatelessComputeManager.h b/Engine/Plugins/FX/Niagara/Source/Niagara/Private/Stateless/NiagaraStatelessComputeManager.h
--- a/Engine/Plugins/FX/Niagara/Source/Niagara/Private/Stateless/NiagaraStatelessComputeManager.h
+++ b/Engine/Plugins/FX/Niagara/Source/Niagara/Private/Stateless/NiagaraStatelessComputeManager.h
@@ -50,4 +50,6 @@ private:
 	TMap<uintptr_t, TUniquePtr<FStatelessDataCache>>	UsedData;
 	TArray<TUniquePtr<FStatelessDataCache>>				FreeData;
 	TArray<uint32>										CountsToRelease;
+
+	TArray<FStatelessDataCache*>						GPUDataToGenerate;
 };
diff --git a/Engine/Plugins/FX/Niagara/Source/Niagara/Private/Stateless/NiagaraStatelessEmitter.cpp b/Engine/Plugins/FX/Niagara/Source/Niagara/Private/Stateless/NiagaraStatelessEmitter.cpp
--- a/Engine/Plugins/FX/Niagara/Source/Niagara/Private/Stateless/NiagaraStatelessEmitter.cpp
+++ b/Engine/Plugins/FX/Niagara/Source/Niagara/Private/Stateless/NiagaraStatelessEmitter.cpp
@@ -5,6 +5,7 @@
 #include "Stateless/NiagaraStatelessEmitterDataBuildContext.h"
 #include "Stateless/NiagaraStatelessModule.h"
 #include "Stateless/NiagaraStatelessDrawDebugContext.h"
+#include "Stateless/NiagaraStatelessParticleSimExecData.h"
 #include "Stateless/Modules/NiagaraStatelessModule_InitializeParticle.h"
 #include "NiagaraConstants.h"
 #include "NiagaraModule.h"
@@ -258,18 +259,46 @@ void UNiagaraStatelessEmitter::CacheFromCompiledData()
 	StatelessEmitterData->bCanEverExecute = NiagaraStatelessInternal::IsValid(*StatelessEmitterData);
 	StatelessEmitterData->bCanEverExecute &= Platforms.IsActive();
 
+	// Determine our supported feature set mask
+	// If we can not execute on any enabled units then the emitter is considered disabled
+	if (StatelessEmitterData->bCanEverExecute)
+	{
+		StatelessEmitterData->FeatureMask = FNiagaraStatelessGlobals::Get().FeatureMask;
+		for (const UNiagaraStatelessModule* Module : Modules)
+		{
+			if (Module->IsModuleEnabled())
+			{
+				StatelessEmitterData->FeatureMask &= Module->GetFeatureMask();
+			}
+		}
+
+		if (StatelessEmitterData->FeatureMask == ENiagaraStatelessFeatureMask::None)
+		{
+			StatelessEmitterData->bCanEverExecute = false;
+			UE_LOG(LogNiagara, Verbose, TEXT("Stateless Emitter (%s) can not execute on any available path and will be disabled."), *GetFullName());
+		}
+	}
+
 	// Build buffers that are shared across all instances
 	//-OPT: We should be able to build and serialize this data as part of the UNiagaraStatelessEmitter, potentially all of this data even since it's immutable and does not change at runtime
 	if (StatelessEmitterData->bCanEverExecute)
 	{
+		if (EnumHasAnyFlags(StatelessEmitterData->FeatureMask, ENiagaraStatelessFeatureMask::ExecuteCPU))
+		{
+			StatelessEmitterData->ParticleSimExecData = new NiagaraStateless::FParticleSimulationExecData(StatelessEmitterData->ParticleDataSetCompiledData);
+		}
+
 		FNiagaraStatelessEmitterDataBuildContext EmitterBuildContext(
+			StatelessEmitterData->ParticleDataSetCompiledData,
 			StatelessEmitterData->RendererBindings,
 			StatelessEmitterData->BuiltData,
-			StatelessEmitterData->StaticFloatData
+			StatelessEmitterData->StaticFloatData,
+			StatelessEmitterData->ParticleSimExecData
 		);
 
 		for (const UNiagaraStatelessModule* Module : Modules)
 		{
+			EmitterBuildContext.PreModuleBuild();
 			Module->BuildEmitterData(EmitterBuildContext);
 		}
 		StatelessEmitterData->bModulesHaveRendererBindings = StatelessEmitterData->RendererBindings.Num() > 0;
diff --git a/Engine/Plugins/FX/Niagara/Source/Niagara/Private/Stateless/NiagaraStatelessEmitterData.cpp b/Engine/Plugins/FX/Niagara/Source/Niagara/Private/Stateless/NiagaraStatelessEmitterData.cpp
--- a/Engine/Plugins/FX/Niagara/Source/Niagara/Private/Stateless/NiagaraStatelessEmitterData.cpp
+++ b/Engine/Plugins/FX/Niagara/Source/Niagara/Private/Stateless/NiagaraStatelessEmitterData.cpp
@@ -2,6 +2,7 @@
 
 #include "Stateless/NiagaraStatelessEmitterData.h"
 #include "Stateless/NiagaraStatelessEmitterTemplate.h"
+#include "Stateless/NiagaraStatelessParticleSimExecData.h"
 #include "Stateless/NiagaraStatelessSimulationShader.h"
 
 void FNiagaraStatelessEmitterData::FDeleter::operator()(FNiagaraStatelessEmitterData* EmitterData) const
@@ -26,6 +27,12 @@ FNiagaraStatelessEmitterData::~FNiagaraStatelessEmitterData()
 {
 	check(IsInRenderingThread());
 	StaticFloatBuffer.Release();
+
+	if (ParticleSimExecData)
+	{
+		delete ParticleSimExecData;
+		ParticleSimExecData = nullptr;
+	}
 }
 
 void FNiagaraStatelessEmitterData::InitRenderResources()
diff --git a/Engine/Plugins/FX/Niagara/Source/Niagara/Private/Stateless/NiagaraStatelessEmitterDataBuildContext.cpp b/Engine/Plugins/FX/Niagara/Source/Niagara/Private/Stateless/NiagaraStatelessEmitterDataBuildContext.cpp
--- a/Engine/Plugins/FX/Niagara/Source/Niagara/Private/Stateless/NiagaraStatelessEmitterDataBuildContext.cpp
+++ b/Engine/Plugins/FX/Niagara/Source/Niagara/Private/Stateless/NiagaraStatelessEmitterDataBuildContext.cpp
@@ -1,11 +1,19 @@
 // Copyright Epic Games, Inc. All Rights Reserved.
 
 #include "Stateless/NiagaraStatelessEmitterDataBuildContext.h"
+#include "Stateless/NiagaraStatelessParticleSimExecData.h"
 
+#include "NiagaraDataSetCompiledData.h"
 #include "NiagaraParameterBinding.h"
 #include "NiagaraParameterStore.h"
+#include "NiagaraParameterStore.h"
+
+void FNiagaraStatelessEmitterDataBuildContext::PreModuleBuild()
+{
+	ModuleBuiltDataOffset = BuiltData.Num();
+}
 
-uint32 FNiagaraStatelessEmitterDataBuildContext::AddStaticData(TConstArrayView<float> FloatData)
+uint32 FNiagaraStatelessEmitterDataBuildContext::AddStaticData(TConstArrayView<float> FloatData) const
 {
 	for (int32 i=0; i <= StaticFloatData.Num() - FloatData.Num(); ++i)
 	{
@@ -21,27 +29,27 @@ uint32 FNiagaraStatelessEmitterDataBuildContext::AddStaticData(TConstArrayView<f
 	return OutIndex;
 }
 
-uint32 FNiagaraStatelessEmitterDataBuildContext::AddStaticData(TConstArrayView<FVector2f> FloatData)
+uint32 FNiagaraStatelessEmitterDataBuildContext::AddStaticData(TConstArrayView<FVector2f> FloatData) const
 {
 	return AddStaticData(MakeArrayView(reinterpret_cast<const float*>(FloatData.GetData()), FloatData.Num() * 2));
 }
 
-uint32 FNiagaraStatelessEmitterDataBuildContext::AddStaticData(TConstArrayView<FVector3f> FloatData)
+uint32 FNiagaraStatelessEmitterDataBuildContext::AddStaticData(TConstArrayView<FVector3f> FloatData) const
 {
 	return AddStaticData(MakeArrayView(reinterpret_cast<const float*>(FloatData.GetData()), FloatData.Num() * 3));
 }
 
-uint32 FNiagaraStatelessEmitterDataBuildContext::AddStaticData(TConstArrayView<FVector4f> FloatData)
+uint32 FNiagaraStatelessEmitterDataBuildContext::AddStaticData(TConstArrayView<FVector4f> FloatData) const
 {
 	return AddStaticData(MakeArrayView(reinterpret_cast<const float*>(FloatData.GetData()), FloatData.Num() * 4));
 }
 
-uint32 FNiagaraStatelessEmitterDataBuildContext::AddStaticData(TConstArrayView<FLinearColor> FloatData)
+uint32 FNiagaraStatelessEmitterDataBuildContext::AddStaticData(TConstArrayView<FLinearColor> FloatData) const
 {
 	return AddStaticData(MakeArrayView(reinterpret_cast<const float*>(FloatData.GetData()), FloatData.Num() * 4));
 }
 
-int32 FNiagaraStatelessEmitterDataBuildContext::AddRendererBinding(const FNiagaraVariableBase& Variable)
+int32 FNiagaraStatelessEmitterDataBuildContext::AddRendererBinding(const FNiagaraVariableBase& Variable) const
 {
 	int32 DataOffset = INDEX_NONE;
 	if (Variable.IsValid())
@@ -54,7 +62,7 @@ int32 FNiagaraStatelessEmitterDataBuildContext::AddRendererBinding(const FNiagar
 	return DataOffset;
 }
 
-int32 FNiagaraStatelessEmitterDataBuildContext::AddRendererBinding(const FNiagaraParameterBinding& Binding)
+int32 FNiagaraStatelessEmitterDataBuildContext::AddRendererBinding(const FNiagaraParameterBinding& Binding) const
 {
 	int32 DataOffset = INDEX_NONE;
 	if (Binding.ResolvedParameter.IsValid())
@@ -67,7 +75,7 @@ int32 FNiagaraStatelessEmitterDataBuildContext::AddRendererBinding(const FNiagar
 	return DataOffset;
 }
 
-int32 FNiagaraStatelessEmitterDataBuildContext::AddRendererBinding(const FNiagaraParameterBindingWithValue& Binding)
+int32 FNiagaraStatelessEmitterDataBuildContext::AddRendererBinding(const FNiagaraParameterBindingWithValue& Binding) const
 {
 	int32 DataOffset = INDEX_NONE;
 	if (Binding.ResolvedParameter.IsValid())
@@ -88,3 +96,18 @@ int32 FNiagaraStatelessEmitterDataBuildContext::AddRendererBinding(const FNiagar
 
 	return DataOffset;
 }
+
+void FNiagaraStatelessEmitterDataBuildContext::AddParticleSimulationExecSimulate(TFunction<void(const NiagaraStateless::FParticleSimulationContext&)> Func) const
+{
+	if (!ParticleExecData)
+	{
+		return;
+	}
+
+	ParticleExecData->SimulateFunctions.Emplace(MoveTemp(Func), ModuleBuiltDataOffset);
+}
+
+int32 FNiagaraStatelessEmitterDataBuildContext::FindParticleVariableIndex(const FNiagaraVariableBase& Variable) const
+{
+	return ParticleDataSet.Variables.IndexOfByKey(Variable);
+}
diff --git a/Engine/Plugins/FX/Niagara/Source/Niagara/Private/Stateless/NiagaraStatelessEmitterInstance.cpp b/Engine/Plugins/FX/Niagara/Source/Niagara/Private/Stateless/NiagaraStatelessEmitterInstance.cpp
--- a/Engine/Plugins/FX/Niagara/Source/Niagara/Private/Stateless/NiagaraStatelessEmitterInstance.cpp
+++ b/Engine/Plugins/FX/Niagara/Source/Niagara/Private/Stateless/NiagaraStatelessEmitterInstance.cpp
@@ -82,6 +82,7 @@ void FNiagaraStatelessEmitterInstance::Init(int32 InEmitterIndex)
 	RenderThreadDataPtr->EmitterData = EmitterData;
 	RenderThreadDataPtr->RandomSeed = RandomSeed;
 	RenderThreadDataPtr->Age = 0.0f;
+	RenderThreadDataPtr->DeltaTime = 0.0f;
 	RenderThreadDataPtr->ExecutionState = ENiagaraExecutionState::Active;
 	RenderThreadDataPtr->ShaderParameters.Reset(WeakStatelessEmitter->AllocateShaderParameters(RendererBindings));
 	RenderThreadDataPtr->ShaderParameters->Common_RandomSeed = RandomSeed;
@@ -137,6 +138,7 @@ void FNiagaraStatelessEmitterInstance::ResetSimulation(bool bKillExisting)
 			[RenderThreadData](FRHICommandListImmediate& RHICmdList)
 			{
 				RenderThreadData->Age = 0.0f;
+				RenderThreadData->DeltaTime = 0.0f;
 				RenderThreadData->ExecutionState = ENiagaraExecutionState::Active;
 			}
 		);
@@ -351,6 +353,7 @@ void FNiagaraStatelessEmitterInstance::SendRenderData()
 	ENQUEUE_RENDER_COMMAND(UpdateStatelessAge)(
 		[RenderThreadData, EmitterData=MoveTemp(DataForRenderThread)](FRHICommandListImmediate& RHICmdList) mutable
 		{
+			RenderThreadData->DeltaTime			= FMath::Max(EmitterData.Age - RenderThreadData->Age, 0.0f);
 			RenderThreadData->Age				= EmitterData.Age;
 			RenderThreadData->ExecutionState	= EmitterData.ExecutionState;
 
diff --git a/Engine/Plugins/FX/Niagara/Source/Niagara/Private/Stateless/NiagaraStatelessEmitterTemplate.cpp b/Engine/Plugins/FX/Niagara/Source/Niagara/Private/Stateless/NiagaraStatelessEmitterTemplate.cpp
--- a/Engine/Plugins/FX/Niagara/Source/Niagara/Private/Stateless/NiagaraStatelessEmitterTemplate.cpp
+++ b/Engine/Plugins/FX/Niagara/Source/Niagara/Private/Stateless/NiagaraStatelessEmitterTemplate.cpp
@@ -72,27 +72,30 @@ void UNiagaraStatelessEmitterDefault::InitModulesAndAttributes()
 #if WITH_EDITORONLY_DATA
 	Modules =
 	{
+		// Initializer Modules
 		UNiagaraStatelessModule_InitializeParticle::StaticClass(),
 		UNiagaraStatelessModule_InitialMeshOrientation::StaticClass(),
 		UNiagaraStatelessModule_ShapeLocation::StaticClass(),
-		UNiagaraStatelessModule_CameraOffset::StaticClass(),
-		UNiagaraStatelessModule_ScaleColor::StaticClass(),
-		UNiagaraStatelessModule_ScaleSpriteSize::StaticClass(),
-		UNiagaraStatelessModule_ScaleSpriteSizeBySpeed::StaticClass(),
-		UNiagaraStatelessModule_ScaleMeshSize::StaticClass(),
-		UNiagaraStatelessModule_ScaleMeshSizeBySpeed::StaticClass(),
-		UNiagaraStatelessModule_MeshIndex::StaticClass(),
-		UNiagaraStatelessModule_MeshRotationRate::StaticClass(),
+		// Before Solve
 		UNiagaraStatelessModule_AddVelocity::StaticClass(),
 		UNiagaraStatelessModule_AccelerationForce::StaticClass(),
 		UNiagaraStatelessModule_CurlNoiseForce::StaticClass(),
 		UNiagaraStatelessModule_Drag::StaticClass(),
 		UNiagaraStatelessModule_GravityForce::StaticClass(),
 		UNiagaraStatelessModule_SolveVelocitiesAndForces::StaticClass(),
+		// Post Solve
+		UNiagaraStatelessModule_CameraOffset::StaticClass(),
+		UNiagaraStatelessModule_DynamicMaterialParameters::StaticClass(),
+		UNiagaraStatelessModule_MeshIndex::StaticClass(),
+		UNiagaraStatelessModule_MeshRotationRate::StaticClass(),
+		UNiagaraStatelessModule_ScaleColor::StaticClass(),
+		UNiagaraStatelessModule_ScaleSpriteSize::StaticClass(),
+		UNiagaraStatelessModule_ScaleSpriteSizeBySpeed::StaticClass(),
+		UNiagaraStatelessModule_ScaleMeshSize::StaticClass(),
+		UNiagaraStatelessModule_ScaleMeshSizeBySpeed::StaticClass(),
 		UNiagaraStatelessModule_SpriteFacingAndAlignment::StaticClass(),
 		UNiagaraStatelessModule_SpriteRotationRate::StaticClass(),
 		UNiagaraStatelessModule_SubUVAnimation::StaticClass(),
-		UNiagaraStatelessModule_DynamicMaterialParameters::StaticClass(),
 	};
 
 	const FNiagaraStatelessGlobals& StatelessGlobals = FNiagaraStatelessGlobals::Get();
@@ -103,6 +106,9 @@ void UNiagaraStatelessEmitterDefault::InitModulesAndAttributes()
 		StatelessGlobals.CameraOffsetVariable,
 		StatelessGlobals.ColorVariable,
 		StatelessGlobals.DynamicMaterialParameters0Variable,
+		StatelessGlobals.DynamicMaterialParameters1Variable,
+		StatelessGlobals.DynamicMaterialParameters2Variable,
+		StatelessGlobals.DynamicMaterialParameters3Variable,
 		StatelessGlobals.MeshIndexVariable,
 		StatelessGlobals.MeshOrientationVariable,
 		StatelessGlobals.RibbonWidthVariable,
@@ -150,6 +156,9 @@ void UNiagaraStatelessEmitterDefault::SetShaderParameters(uint8* ShaderParameter
 	ShaderParameters->Permutation_CameraOffsetComponent				= ComponentOffsets[iComponent++];
 	ShaderParameters->Permutation_ColorComponent					= ComponentOffsets[iComponent++];
 	ShaderParameters->Permutation_DynamicMaterialParameter0Component= ComponentOffsets[iComponent++];
+	ShaderParameters->Permutation_DynamicMaterialParameter1Component= ComponentOffsets[iComponent++];
+	ShaderParameters->Permutation_DynamicMaterialParameter2Component= ComponentOffsets[iComponent++];
+	ShaderParameters->Permutation_DynamicMaterialParameter3Component= ComponentOffsets[iComponent++];
 	ShaderParameters->Permutation_MeshIndexComponent				= ComponentOffsets[iComponent++];
 	ShaderParameters->Permutation_MeshOrientationComponent			= ComponentOffsets[iComponent++];
 	ShaderParameters->Permutation_RibbonWidthComponent				= ComponentOffsets[iComponent++];
diff --git a/Engine/Plugins/FX/Niagara/Source/Niagara/Private/Stateless/NiagaraStatelessParticleSimContext.cpp b/Engine/Plugins/FX/Niagara/Source/Niagara/Private/Stateless/NiagaraStatelessParticleSimContext.cpp
new file mode 100644
--- /dev/null
+++ b/Engine/Plugins/FX/Niagara/Source/Niagara/Private/Stateless/NiagaraStatelessParticleSimContext.cpp
@@ -0,0 +1,352 @@
+// Copyright Epic Games, Inc. All Rights Reserved.
+
+#include "Stateless/NiagaraStatelessParticleSimContext.h"
+
+#include "Stateless/NiagaraStatelessEmitterData.h"
+#include "Stateless/NiagaraStatelessEmitterInstance.h"
+#include "Stateless/NiagaraStatelessParticleSimExecData.h"
+
+namespace NiagaraStateless
+{
+	constexpr uint32 ParticleComponentOffsets[int32(EParticleComponent::Num)] =
+	{
+		0,	// uint32		- Alive
+		1,	// float		- Lifetime
+		2,	// float		- Age
+		3,	// float		- NormalizedAge
+		4,	// float		- PreviousAge
+		5,	// float		- PreviousNormalizedAge
+		6,	// FUintVector4	- RandomSeed
+	};
+	constexpr uint32 ParticleComponentTotalSize = 10;
+
+	FUintVector4 Rand4DPCG32(FUintVector4 v)
+	{
+		// Linear congruential step.
+		v.X = v.X * 1664525u + 1013904223u;
+		v.Y = v.Y * 1664525u + 1013904223u;
+		v.Z = v.Z * 1664525u + 1013904223u;
+		v.W = v.W * 1664525u + 1013904223u;
+
+		// shuffle
+		v.X += v.Y * v.W;
+		v.Y += v.Z * v.X;
+		v.Z += v.X * v.Y;
+		v.W += v.Y * v.Z;
+
+		// xoring high bits into low makes all 32 bits pretty good
+		v.X ^= (v.X >> 16u);
+		v.Y ^= (v.Y >> 16u);
+		v.Z ^= (v.Z >> 16u);
+		v.W ^= (v.W >> 16u);
+
+		// final shuffle
+		v.X += v.Y * v.W;
+		v.Y += v.Z * v.X;
+		v.Z += v.X * v.Y;
+		v.W += v.Y * v.Z;
+
+		return v;
+	}
+
+	FParticleSimulationContext::FParticleSimulationContext(const FNiagaraStatelessEmitterData* InEmitterData, TConstArrayView<uint8> InDynamicBufferData)
+		: EmitterData(InEmitterData)
+		, BuiltData(EmitterData->BuiltData)
+		, StaticFloatData(EmitterData->StaticFloatData)
+		, DynamicBufferData(InDynamicBufferData)
+	{
+		check(EmitterData->ParticleSimExecData);
+	}
+
+	void FParticleSimulationContext::Simulate(int32 EmitterRandomSeed, float EmitterAge, float InDeltaTime, TConstArrayView<FNiagaraStatelessRuntimeSpawnInfo> SpawnInfos, FNiagaraDataBuffer* DestinationData)
+	{
+		NumInstances = 0;
+
+		FSpawnInfoShaderParameters SpawnParameters;
+		const uint32 ActiveParticles = EmitterData->CalculateActiveParticles(EmitterRandomSeed, SpawnInfos, EmitterAge, &SpawnParameters);
+		if (ActiveParticles > 0)
+		{
+			// Setup data buffer pointers
+			DestinationData->Allocate(ActiveParticles);
+			BufferStride	= DestinationData->GetFloatStride();
+			BufferFloatData	= DestinationData->GetComponentPtrFloat(0);
+			BufferInt32Data	= DestinationData->GetComponentPtrInt32(0);
+
+			// Run Simulation
+			SimulateInternal(EmitterRandomSeed, EmitterAge, InDeltaTime, SpawnParameters, ActiveParticles);
+		}
+
+		// Set instance count
+		DestinationData->SetNumInstances(NumInstances);
+	}
+
+	void FParticleSimulationContext::SimulateGPU(FRHICommandList& RHICmdList, int32 EmitterRandomSeed, float EmitterAge, float InDeltaTime, TConstArrayView<FNiagaraStatelessRuntimeSpawnInfo> SpawnInfos, FNiagaraDataBuffer* DestinationData)
+	{
+		NumInstances = 0;
+
+		FSpawnInfoShaderParameters SpawnParameters;
+		const uint32 ActiveParticles = EmitterData->CalculateActiveParticles(EmitterRandomSeed, SpawnInfos, EmitterAge, &SpawnParameters);
+		if (ActiveParticles > 0)
+		{
+			check(DestinationData->GetNumInstancesAllocated() <= ActiveParticles);
+
+			FRWBuffer& FloatBuffer = DestinationData->GetGPUBufferFloat();
+			FRWBuffer& Int32Buffer = DestinationData->GetGPUBufferInt();
+
+			// Setup data buffer pointers
+			BufferStride	= DestinationData->GetFloatStride();
+			BufferFloatData	= FloatBuffer.NumBytes > 0 ? reinterpret_cast<uint8*>(RHICmdList.LockBuffer(FloatBuffer.Buffer, 0, FloatBuffer.NumBytes, RLM_WriteOnly)) : nullptr;
+			BufferInt32Data	= Int32Buffer.NumBytes > 0 ? reinterpret_cast<uint8*>(RHICmdList.LockBuffer(Int32Buffer.Buffer, 0, Int32Buffer.NumBytes, RLM_WriteOnly)) : nullptr;
+
+			// Run Simulation
+			SimulateInternal(EmitterRandomSeed, EmitterAge, InDeltaTime, SpawnParameters, ActiveParticles);
+
+			// Unlock buffers
+			if (BufferFloatData)
+			{
+				RHICmdList.UnlockBuffer(FloatBuffer.Buffer);
+			}
+			if (BufferInt32Data)
+			{
+				RHICmdList.UnlockBuffer(Int32Buffer.Buffer);
+			}
+		}
+
+		// Set instance count
+		DestinationData->SetNumInstances(NumInstances);
+	}
+
+	void FParticleSimulationContext::SimulateInternal(int32 EmitterRandomSeed, float EmitterAge, float InDeltaTime, FSpawnInfoShaderParameters& SpawnParameters, uint32 ActiveParticles)
+	{
+		// Setup simulation	
+		NumInstances = 0;
+		DeltaTime = InDeltaTime;
+		InvDeltaTime = DeltaTime > 0.0f ? 1.0f / DeltaTime : 0.0f;
+
+		// Setup Required Components
+		{
+			static TArray<uint8>		RequiredStorage;
+			static uint32				RequiredStorageStride;
+			if (RequiredStorageStride < int32(BufferStride))
+			{
+				RequiredStorageStride = BufferStride;
+				RequiredStorage.SetNumUninitialized(ParticleComponentTotalSize * BufferStride);
+			}
+
+			for (int32 i = 0; i < int32(EParticleComponent::Num); ++i)
+			{
+				RequiredComponents[i] = RequiredStorage.GetData() + (ParticleComponentOffsets[i] * BufferStride);
+			}
+		}
+
+		// Setup optional components
+		const FParticleSimulationExecData* ExecData = EmitterData->ParticleSimExecData;
+		VariableComponents.AddDefaulted(ExecData->VariableComponentOffsets.Num());
+		for (int32 i = 0; i < ExecData->VariableComponentOffsets.Num(); ++i)
+		{
+			const uint32 VariableOffset = ExecData->VariableComponentOffsets[i].GetOffset();
+			if (ExecData->VariableComponentOffsets[i].IsFloat())
+			{
+				VariableComponents[i] = BufferFloatData + (VariableOffset * BufferStride);
+			}
+			else //if (ExecData->VariableComponentOffsets.IsInt32())
+			{
+				VariableComponents[i] = BufferInt32Data + (VariableOffset * BufferStride);
+			}
+		}
+
+		// Setup particles
+		for (uint32 iParticle = 0; iParticle < ActiveParticles; ++iParticle)
+		{
+			uint32 Particle_UniqueIndex = 0;
+			float Particle_Age = -1.0f;
+			{
+				uint32 SpawnInfoIndex = iParticle;
+				for (int32 GpuSpawnIndex = 0; GpuSpawnIndex < NiagaraStateless::MaxGpuSpawnInfos; ++GpuSpawnIndex)
+				{
+					const uint32 SpawnInfoNumActive = GET_SCALAR_ARRAY_ELEMENT(SpawnParameters.SpawnInfo_NumActive, GpuSpawnIndex);
+					const uint32 SpawnInfoParticleOffset = GET_SCALAR_ARRAY_ELEMENT(SpawnParameters.SpawnInfo_ParticleOffset, GpuSpawnIndex);
+					const uint32 SpawnInfoUniqueOffset = GET_SCALAR_ARRAY_ELEMENT(SpawnParameters.SpawnInfo_UniqueOffset, GpuSpawnIndex);
+					const float  SpawnInfoTime = GET_SCALAR_ARRAY_ELEMENT(SpawnParameters.SpawnInfo_Time, GpuSpawnIndex);
+					const float  SpawnInfoRate = GET_SCALAR_ARRAY_ELEMENT(SpawnParameters.SpawnInfo_Rate, GpuSpawnIndex);
+					if (SpawnInfoIndex < SpawnInfoNumActive)
+					{
+						const uint32 SpawnParticleIndex = SpawnInfoIndex + SpawnInfoParticleOffset;
+						Particle_UniqueIndex = SpawnInfoIndex + SpawnInfoUniqueOffset + SpawnInfoParticleOffset;
+						Particle_Age = EmitterAge - (SpawnInfoTime + float(SpawnParticleIndex) * SpawnInfoRate);
+						break;
+					}
+
+					SpawnInfoIndex -= SpawnInfoNumActive;
+				}
+				if (Particle_Age < 0.0f)
+				{
+					continue;
+				}
+			}
+
+			FUintVector4 RandomSeed;
+			RandomSeed.X = 7123u;
+			RandomSeed.Y = EmitterRandomSeed;
+			RandomSeed.Z = Particle_UniqueIndex * 3581u;
+			RandomSeed.W = 0;
+
+			// Write random seed here so we can use our random functions
+			GetParticleRandomSeed()[NumInstances] = RandomSeed;
+
+			float Particle_Lifetime = RandomScaleBiasFloat(NumInstances, EmitterData->LifetimeRange);
+			if (Particle_Lifetime <= 0.0f || Particle_Age >= Particle_Lifetime)
+			{
+				return;
+			}
+
+			const float PreviousAge = FMath::Max(Particle_Age - DeltaTime, 0.0f);
+
+			// Initialize particle information
+			GetParticleAlive()[NumInstances] = 1;
+			GetParticleLifetime()[NumInstances] = Particle_Lifetime;
+			GetParticleAge()[NumInstances] = Particle_Age;
+			GetParticleNormalizedAge()[NumInstances] = Particle_Age / Particle_Lifetime;
+			GetParticlePreviousAge()[NumInstances] = PreviousAge;
+			GetParticlePreviousNormalizedAge()[NumInstances] = PreviousAge / Particle_Lifetime;
+
+			++NumInstances;
+		}
+
+		if (NumInstances == 0)
+		{
+			return;
+		}
+
+		// Execute the simulation
+		for (const auto& Callback : ExecData->SimulateFunctions)
+		{
+			BuiltDataOffset = Callback.BuiltDataOffset;
+			Callback.Function(*this);
+		}
+	}
+
+	uint32 FParticleSimulationContext::RandomUInt(uint32 iInstance) const
+	{
+		FUintVector4& RandomSeed = GetParticleRandomSeed()[iInstance];
+		++RandomSeed.X;
+		RandomSeed.W = RandomSeed.X ^ RandomSeed.Y;
+
+		FUintVector4 RandomValue = Rand4DPCG32(RandomSeed);
+		return RandomValue.X;
+	}
+
+	FUintVector2 FParticleSimulationContext::RandomUInt2(uint32 iInstance) const
+	{
+		FUintVector4& RandomSeed = GetParticleRandomSeed()[iInstance];
+		++RandomSeed.X;
+		RandomSeed.W = RandomSeed.X ^ RandomSeed.Y;
+
+		FUintVector4 RandomValue = Rand4DPCG32(RandomSeed);
+		return FUintVector2(RandomValue.X, RandomValue.Y);
+	}
+
+	FUintVector3 FParticleSimulationContext::RandomUInt3(uint32 iInstance) const
+	{
+		FUintVector4& RandomSeed = GetParticleRandomSeed()[iInstance];
+		++RandomSeed.X;
+		RandomSeed.W = RandomSeed.X ^ RandomSeed.Y;
+
+		FUintVector4 RandomValue = Rand4DPCG32(RandomSeed);
+		return FUintVector3(RandomValue.X, RandomValue.Y, RandomValue.Z);
+	}
+
+	FUintVector4 FParticleSimulationContext::RandomUInt4(uint32 iInstance) const
+	{
+		FUintVector4& RandomSeed = GetParticleRandomSeed()[iInstance];
+		++RandomSeed.X;
+		RandomSeed.W = RandomSeed.X ^ RandomSeed.Y;
+
+		FUintVector4 RandomValue = Rand4DPCG32(RandomSeed);
+		return RandomValue;
+	}
+
+	float FParticleSimulationContext::RandomFloat(uint32 iInstance) const
+	{
+		const uint32 V = RandomUInt(iInstance);
+		return float((V >> 8) & 0x00ffffff) / 16777216.0f;
+	}
+
+	FVector2f FParticleSimulationContext::RandomFloat2(uint32 iInstance) const
+	{
+		const FUintVector2 V = RandomUInt2(iInstance);
+		return FVector2f(
+			float((V.X >> 8) & 0x00ffffff) / 16777216.0f,
+			float((V.Y >> 8) & 0x00ffffff) / 16777216.0f
+		);
+	}
+
+	FVector3f FParticleSimulationContext::RandomFloat3(uint32 iInstance) const
+	{
+		const FUintVector3 V = RandomUInt3(iInstance);
+		return FVector3f(
+			float((V.X >> 8) & 0x00ffffff) / 16777216.0f,
+			float((V.Y >> 8) & 0x00ffffff) / 16777216.0f,
+			float((V.Z >> 8) & 0x00ffffff) / 16777216.0f
+		);
+	}
+
+	FVector4f FParticleSimulationContext::RandomFloat4(uint32 iInstance) const
+	{
+		const FUintVector4 V = RandomUInt4(iInstance);
+		return FVector4f(
+			float((V.X >> 8) & 0x00ffffff) / 16777216.0f,
+			float((V.Y >> 8) & 0x00ffffff) / 16777216.0f,
+			float((V.Z >> 8) & 0x00ffffff) / 16777216.0f,
+			float((V.W >> 8) & 0x00ffffff) / 16777216.0f
+		);
+	}
+
+	float FParticleSimulationContext::RandomSeedOffsetFloat(uint32 iInstance, uint32 SeedOffset) const
+	{
+		FUintVector4 RandomSeed = GetParticleRandomSeed()[iInstance];
+		RandomSeed.X += SeedOffset;
+
+		FUintVector4 RandomValue = Rand4DPCG32(RandomSeed);
+		return float((RandomValue.X >> 8) & 0x00ffffff) / 16777216.0f;
+	}
+
+	FVector2f FParticleSimulationContext::RandomSeedOffsetFloat2(uint32 iInstance, uint32 SeedOffset) const
+	{
+		FUintVector4 RandomSeed = GetParticleRandomSeed()[iInstance];
+		RandomSeed.X += SeedOffset;
+
+		FUintVector4 RandomValue = Rand4DPCG32(RandomSeed);
+		return FVector2f(
+			float((RandomValue.X >> 8) & 0x00ffffff) / 16777216.0f,
+			float((RandomValue.Y >> 8) & 0x00ffffff) / 16777216.0f
+		);
+	}
+
+	FVector3f FParticleSimulationContext::RandomSeedOffsetFloat3(uint32 iInstance, uint32 SeedOffset) const
+	{
+		FUintVector4 RandomSeed = GetParticleRandomSeed()[iInstance];
+		RandomSeed.X += SeedOffset;
+
+		FUintVector4 RandomValue = Rand4DPCG32(RandomSeed);
+		return FVector3f(
+			float((RandomValue.X >> 8) & 0x00ffffff) / 16777216.0f,
+			float((RandomValue.Y >> 8) & 0x00ffffff) / 16777216.0f,
+			float((RandomValue.Z >> 8) & 0x00ffffff) / 16777216.0f
+		);
+	}
+
+	FVector4f FParticleSimulationContext::RandomSeedOffsetFloat4(uint32 iInstance, uint32 SeedOffset) const
+	{
+		FUintVector4 RandomSeed = GetParticleRandomSeed()[iInstance];
+		RandomSeed.X += SeedOffset;
+
+		FUintVector4 RandomValue = Rand4DPCG32(RandomSeed);
+		return FVector4f(
+			float((RandomValue.X >> 8) & 0x00ffffff) / 16777216.0f,
+			float((RandomValue.Y >> 8) & 0x00ffffff) / 16777216.0f,
+			float((RandomValue.Z >> 8) & 0x00ffffff) / 16777216.0f,
+			float((RandomValue.W >> 8) & 0x00ffffff) / 16777216.0f
+		);
+	}
+} //namespace NiagaraStateless
diff --git a/Engine/Plugins/FX/Niagara/Source/Niagara/Private/Stateless/NiagaraStatelessParticleSimExecData.cpp b/Engine/Plugins/FX/Niagara/Source/Niagara/Private/Stateless/NiagaraStatelessParticleSimExecData.cpp
new file mode 100644
--- /dev/null
+++ b/Engine/Plugins/FX/Niagara/Source/Niagara/Private/Stateless/NiagaraStatelessParticleSimExecData.cpp
@@ -0,0 +1,37 @@
+// Copyright Epic Games, Inc. All Rights Reserved.
+
+#include "Stateless/NiagaraStatelessParticleSimExecData.h"
+
+#include "NiagaraDataSetCompiledData.h"
+
+namespace NiagaraStateless
+{
+	FParticleSimulationExecData::FParticleSimulationExecData(const FNiagaraDataSetCompiledData& ParticleDataSetCompiledData)
+	{
+		const int32 NumVariables = FMath::Min(ParticleDataSetCompiledData.Variables.Num(), ParticleDataSetCompiledData.VariableLayouts.Num());
+		VariableComponentOffsets.AddDefaulted(NumVariables);
+
+		for (int32 i = 0; i < NumVariables; ++i)
+		{
+			const FNiagaraVariableLayoutInfo& VariableLayout = ParticleDataSetCompiledData.VariableLayouts[i];
+			if (VariableLayout.GetNumFloatComponents() > 0)
+			{
+				check(VariableLayout.GetNumInt32Components() == 0 && VariableLayout.GetNumHalfComponents() == 0);
+				VariableComponentOffsets[i].Type	= 0;
+				VariableComponentOffsets[i].Offset	= VariableLayout.GetFloatComponentStart();
+			}
+			else if (VariableLayout.GetNumInt32Components() > 0)
+			{
+				check(VariableLayout.GetNumFloatComponents() == 0 && VariableLayout.GetNumHalfComponents() == 0);
+				VariableComponentOffsets[i].Type	= 1;
+				VariableComponentOffsets[i].Offset	= VariableLayout.GetInt32ComponentStart();
+			}
+			else
+			{
+				// We don't support half components
+				checkNoEntry();
+			}
+		}
+	}
+
+} //namespace NiagaraStateless
diff --git a/Engine/Plugins/FX/Niagara/Source/NiagaraShader/Internal/Stateless/NiagaraStatelessEmitterShaders.h b/Engine/Plugins/FX/Niagara/Source/NiagaraShader/Internal/Stateless/NiagaraStatelessEmitterShaders.h
--- a/Engine/Plugins/FX/Niagara/Source/NiagaraShader/Internal/Stateless/NiagaraStatelessEmitterShaders.h
+++ b/Engine/Plugins/FX/Niagara/Source/NiagaraShader/Internal/Stateless/NiagaraStatelessEmitterShaders.h
@@ -24,28 +24,34 @@ namespace NiagaraStateless
 
 		BEGIN_SHADER_PARAMETER_STRUCT(FParameters, NIAGARASHADER_API)
 			SHADER_PARAMETER_STRUCT_INCLUDE(FCommonShaderParameters, CommonParameters)
+			// Initialize Modules
 			ADD_STATELESS_MODULE(FInitializeParticleModule_ShaderParameters)
 			ADD_STATELESS_MODULE(FInitialMeshOrientationModule_ShaderParameters)
 			ADD_STATELESS_MODULE(FShapeLocationModule_ShaderParameters)
+			// Pre Solve Modules
+			ADD_STATELESS_MODULE(FSolveVelocitiesAndForcesModule_ShaderParameters)
+			// Post Solve Modules
 			ADD_STATELESS_MODULE(FCameraOffsetModule_ShaderParameters)
+			ADD_STATELESS_MODULE(FDynamicMaterialParametersModule_ShaderParameters)
+			ADD_STATELESS_MODULE(FMeshIndexModule_ShaderParameters)
+			ADD_STATELESS_MODULE(FMeshRotationRateModule_ShaderParameters)
 			ADD_STATELESS_MODULE(FScaleColorModule_ShaderParameters)
 			ADD_STATELESS_MODULE(FScaleSpriteSizeModule_ShaderParameters)
 			ADD_STATELESS_MODULE(FScaleSpriteSizeBySpeedModule_ShaderParameters)
 			ADD_STATELESS_MODULE(FScaleMeshSizeModule_ShaderParameters)
 			ADD_STATELESS_MODULE(FScaleMeshSizeBySpeedModule_ShaderParameters)
-			ADD_STATELESS_MODULE(FMeshIndexModule_ShaderParameters)
-			ADD_STATELESS_MODULE(FMeshRotationRateModule_ShaderParameters)
-			ADD_STATELESS_MODULE(FSolveVelocitiesAndForcesModule_ShaderParameters)
 			ADD_STATELESS_MODULE(FSpriteFacingAndAlignmentModule_ShaderParameters)
 			ADD_STATELESS_MODULE(FSpriteRotationRateModule_ShaderParameters)
 			ADD_STATELESS_MODULE(FSubUVAnimationModule_ShaderParameters)
-			ADD_STATELESS_MODULE(FDynamicMaterialParametersModule_ShaderParameters)
 
 			SHADER_PARAMETER(int, Permutation_UniqueIDComponent)
 			SHADER_PARAMETER(int, Permutation_PositionComponent)
 			SHADER_PARAMETER(int, Permutation_CameraOffsetComponent)
 			SHADER_PARAMETER(int, Permutation_ColorComponent)
 			SHADER_PARAMETER(int, Permutation_DynamicMaterialParameter0Component)
+			SHADER_PARAMETER(int, Permutation_DynamicMaterialParameter1Component)
+			SHADER_PARAMETER(int, Permutation_DynamicMaterialParameter2Component)
+			SHADER_PARAMETER(int, Permutation_DynamicMaterialParameter3Component)
 			SHADER_PARAMETER(int, Permutation_MeshIndexComponent)
 			SHADER_PARAMETER(int, Permutation_MeshOrientationComponent)
 			SHADER_PARAMETER(int, Permutation_RibbonWidthComponent)
diff --git a/Engine/Plugins/FX/Niagara/Source/NiagaraShader/Internal/Stateless/NiagaraStatelessModuleShaderParameters.h b/Engine/Plugins/FX/Niagara/Source/NiagaraShader/Internal/Stateless/NiagaraStatelessModuleShaderParameters.h
--- a/Engine/Plugins/FX/Niagara/Source/NiagaraShader/Internal/Stateless/NiagaraStatelessModuleShaderParameters.h
+++ b/Engine/Plugins/FX/Niagara/Source/NiagaraShader/Internal/Stateless/NiagaraStatelessModuleShaderParameters.h
@@ -50,6 +50,18 @@ namespace NiagaraStateless
 		SHADER_PARAMETER(FUintVector3,	DynamicMaterialParameters_Parameter0Y)
 		SHADER_PARAMETER(FUintVector3,	DynamicMaterialParameters_Parameter0Z)
 		SHADER_PARAMETER(FUintVector3,	DynamicMaterialParameters_Parameter0W)
+		SHADER_PARAMETER(FUintVector3,	DynamicMaterialParameters_Parameter1X)
+		SHADER_PARAMETER(FUintVector3,	DynamicMaterialParameters_Parameter1Y)
+		SHADER_PARAMETER(FUintVector3,	DynamicMaterialParameters_Parameter1Z)
+		SHADER_PARAMETER(FUintVector3,	DynamicMaterialParameters_Parameter1W)
+		SHADER_PARAMETER(FUintVector3,	DynamicMaterialParameters_Parameter2X)
+		SHADER_PARAMETER(FUintVector3,	DynamicMaterialParameters_Parameter2Y)
+		SHADER_PARAMETER(FUintVector3,	DynamicMaterialParameters_Parameter2Z)
+		SHADER_PARAMETER(FUintVector3,	DynamicMaterialParameters_Parameter2W)
+		SHADER_PARAMETER(FUintVector3,	DynamicMaterialParameters_Parameter3X)
+		SHADER_PARAMETER(FUintVector3,	DynamicMaterialParameters_Parameter3Y)
+		SHADER_PARAMETER(FUintVector3,	DynamicMaterialParameters_Parameter3Z)
+		SHADER_PARAMETER(FUintVector3,	DynamicMaterialParameters_Parameter3W)
 	END_SHADER_PARAMETER_STRUCT()
 
 	BEGIN_SHADER_PARAMETER_STRUCT(FCameraOffsetModule_ShaderParameters, NIAGARASHADER_API)
@@ -78,6 +90,7 @@ namespace NiagaraStateless
 	END_SHADER_PARAMETER_STRUCT()
 
 	BEGIN_SHADER_PARAMETER_STRUCT(FMeshRotationRateModule_ShaderParameters, NIAGARASHADER_API)
+		SHADER_PARAMETER(int32,			MeshRotationRate_ModuleEnabled)
 		SHADER_PARAMETER(FVector3f,		MeshRotationRate_Scale)
 		SHADER_PARAMETER(FVector3f,		MeshRotationRate_Bias)
 	END_SHADER_PARAMETER_STRUCT()
@@ -85,7 +98,6 @@ namespace NiagaraStateless
 	BEGIN_SHADER_PARAMETER_STRUCT(FScaleSpriteSizeModule_ShaderParameters, NIAGARASHADER_API)
 		SHADER_PARAMETER(FUintVector3,	ScaleSpriteSize_Distribution)
 		SHADER_PARAMETER(FVector2f,		ScaleSpriteSize_CurveScale)
-		SHADER_PARAMETER(int32,			ScaleSpriteSize_CurveScaleOffset)
 	END_SHADER_PARAMETER_STRUCT()
 
 	BEGIN_SHADER_PARAMETER_STRUCT(FScaleSpriteSizeBySpeedModule_ShaderParameters, NIAGARASHADER_API)
