Fix for double brackets with UNROLL macro in NDC shader

diff --git a/Engine/Plugins/FX/Niagara/Shaders/Private/DataChannel/NiagaraDataInterfaceDataChannelTemplateCommon.ush b/Engine/Plugins/FX/Niagara/Shaders/Private/DataChannel/NiagaraDataInterfaceDataChannelTemplateCommon.ush
--- a/Engine/Plugins/FX/Niagara/Shaders/Private/DataChannel/NiagaraDataInterfaceDataChannelTemplateCommon.ush
+++ b/Engine/Plugins/FX/Niagara/Shaders/Private/DataChannel/NiagaraDataInterfaceDataChannelTemplateCommon.ush
@@ -82,7 +82,7 @@ void GetNDCSpawnData_{ParameterName}(NiagaraEmitterID InEmitterID, int InExecInd
 	uint MaxBucketExecIndex = 0;
 	uint BucketEntryStart = NDCSPawnData_NDCIndexBucketsStart;
 
-	[UNROLL]
+	UNROLL
 	for(int BucketIdx = 0; BucketIdx < 16; ++BucketIdx)
 	{
 		uint BucketSize = (1U << 15U) >> BucketIdx;
@@ -111,7 +111,7 @@ void GetNDCSpawnData_{ParameterName}(NiagaraEmitterID InEmitterID, int InExecInd
 		OutNDCSpawnIndex = 0;
 		MaxBucketExecIndex = 0;
 		uint Count = OutNDCSpawnCount;
-		[UNROLL]
+		UNROLL
 		for(int BucketIdx = 0; BucketIdx < 16; ++BucketIdx)
 		{
 			uint BucketSize = (1U << 15U) >> BucketIdx;
