Moving from name to guid for child resolution during traversal
#tests Cooked old asset that had multiple emitters with the same UniqueName

diff --git a/Engine/Plugins/FX/Niagara/Source/NiagaraEditor/Private/NiagaraCompilationPrivate.h b/Engine/Plugins/FX/Niagara/Source/NiagaraEditor/Private/NiagaraCompilationPrivate.h
--- a/Engine/Plugins/FX/Niagara/Source/NiagaraEditor/Private/NiagaraCompilationPrivate.h
+++ b/Engine/Plugins/FX/Niagara/Source/NiagaraEditor/Private/NiagaraCompilationPrivate.h
@@ -46,8 +46,8 @@ public:
 
 	ENiagaraFunctionDebugState GetDebugState() const;
 
-	void AddNamedChildResolver(FName ScopeName, const FNiagaraFixedConstantResolver& ChildResolver);
-	const FNiagaraFixedConstantResolver* FindChildResolver(FName ScopeName) const;
+	void AddChildResolver(const FGuid& ChildId, const FNiagaraFixedConstantResolver& ChildResolver);
+	const FNiagaraFixedConstantResolver* FindChildResolver(const FGuid& ChildId) const;
 
 private:
 	void InitConstants();
@@ -70,8 +70,8 @@ private:
 
 	TArray<FNiagaraVariable, TFixedAllocator<(uint8)EResolvedConstant::Count>> ResolvedConstants;
 
-	using FNamedResolverPair = TTuple<FName, FNiagaraFixedConstantResolver>;
-	TArray<FNamedResolverPair> ChildResolversByName;
+	using FNamedResolverPair = TTuple<FGuid, FNiagaraFixedConstantResolver>;
+	TArray<FNamedResolverPair> ChildResolvers;
 };
 
 struct FNiagaraSimulationStageInfo
diff --git a/Engine/Plugins/FX/Niagara/Source/NiagaraEditor/Private/NiagaraCompilationTasks.cpp b/Engine/Plugins/FX/Niagara/Source/NiagaraEditor/Private/NiagaraCompilationTasks.cpp
--- a/Engine/Plugins/FX/Niagara/Source/NiagaraEditor/Private/NiagaraCompilationTasks.cpp
+++ b/Engine/Plugins/FX/Niagara/Source/NiagaraEditor/Private/NiagaraCompilationTasks.cpp
@@ -1368,6 +1368,7 @@ void FNiagaraSystemCompilationTask::DigestSystemInfo()
 			if (const UNiagaraEmitter* Emitter = HandleInstance.Emitter)
 			{
 				EmitterInfo.UniqueEmitterName = HandleInstance.Emitter->GetUniqueEmitterName();
+				EmitterInfo.EmitterHandleId = Handle.GetId();
 				EmitterInfo.UniqueInstanceName = Handle.GetUniqueInstanceName();
 				EmitterInfo.Enabled = Handle.GetIsEnabled();
 				EmitterInfo.ConstantResolver = FNiagaraFixedConstantResolver(FCompileConstantResolver(HandleInstance, ENiagaraScriptUsage::EmitterSpawnScript));
@@ -1376,7 +1377,7 @@ void FNiagaraSystemCompilationTask::DigestSystemInfo()
 				EmitterInfo.SourceGraph = DigestDatabase.CreateGraphDigest(EmitterGraph, ChangeIdBuilder);
 
 				// be sure to incorporate our constant resolver into the top level SystemInfo.ConstantResolver
-				SystemInfo.ConstantResolver.AddNamedChildResolver(*EmitterInfo.UniqueEmitterName, EmitterInfo.ConstantResolver);
+				SystemInfo.ConstantResolver.AddChildResolver(EmitterInfo.EmitterHandleId, EmitterInfo.ConstantResolver);
 			}
 
 			{
diff --git a/Engine/Plugins/FX/Niagara/Source/NiagaraEditor/Private/NiagaraCompilationTasks.h b/Engine/Plugins/FX/Niagara/Source/NiagaraEditor/Private/NiagaraCompilationTasks.h
--- a/Engine/Plugins/FX/Niagara/Source/NiagaraEditor/Private/NiagaraCompilationTasks.h
+++ b/Engine/Plugins/FX/Niagara/Source/NiagaraEditor/Private/NiagaraCompilationTasks.h
@@ -71,6 +71,7 @@ struct FNiagaraSystemCompilationTask
 
 	struct FEmitterInfo
 	{
+		FGuid EmitterHandleId;
 		FString UniqueEmitterName;
 		FString UniqueInstanceName;
 		FNiagaraDigestedGraphPtr SourceGraph;
diff --git a/Engine/Plugins/FX/Niagara/Source/NiagaraEditor/Private/NiagaraGraphDigest.cpp b/Engine/Plugins/FX/Niagara/Source/NiagaraEditor/Private/NiagaraGraphDigest.cpp
--- a/Engine/Plugins/FX/Niagara/Source/NiagaraEditor/Private/NiagaraGraphDigest.cpp
+++ b/Engine/Plugins/FX/Niagara/Source/NiagaraEditor/Private/NiagaraGraphDigest.cpp
@@ -2092,7 +2092,6 @@ FNiagaraCompilationNodeEmitter::FNiagaraCompilationNodeEmitter(const UNiagaraNod
 	EmitterName = InNode->GetName();
 	EmitterPathName = InNode->GetPathName();
 	EmitterHandleIdString = EmitterHandleID.ToString(EGuidFormats::Digits);
-	EmitterUniqueFName = *EmitterUniqueName;
 }
 
 
@@ -2104,7 +2103,6 @@ FNiagaraCompilationNodeEmitter::FNiagaraCompilationNodeEmitter(const FNiagaraCom
 	, EmitterName(InNode.EmitterName)
 	, EmitterPathName(InNode.EmitterPathName)
 	, EmitterHandleIdString(InNode.EmitterHandleIdString)
-	, EmitterUniqueFName(InNode.EmitterUniqueFName)
 	, Usage(InNode.Usage)
 {
 	// we need to replace the CalledGraph here with the graph that has been instantiated already
@@ -2140,7 +2138,7 @@ void FNiagaraCompilationNodeEmitter::BuildParameterMapHistory(FParameterMapHisto
 		return;
 	}
 
-	const FNiagaraFixedConstantResolver* ChildConstantResolver = Builder.ConstantResolver->FindChildResolver(EmitterUniqueFName);
+	const FNiagaraFixedConstantResolver* ChildConstantResolver = Builder.ConstantResolver->FindChildResolver(EmitterHandleID);
 	if (!ChildConstantResolver)
 	{
 		// if no child resolver was found for the specified emitter, that means that the emitter is likely not enabled and so we can proceed without
diff --git a/Engine/Plugins/FX/Niagara/Source/NiagaraEditor/Private/NiagaraGraphDigest.h b/Engine/Plugins/FX/Niagara/Source/NiagaraEditor/Private/NiagaraGraphDigest.h
--- a/Engine/Plugins/FX/Niagara/Source/NiagaraEditor/Private/NiagaraGraphDigest.h
+++ b/Engine/Plugins/FX/Niagara/Source/NiagaraEditor/Private/NiagaraGraphDigest.h
@@ -401,7 +401,6 @@ public:
 	FString EmitterName;
 	FString EmitterPathName;
 	FString EmitterHandleIdString;
-	FName EmitterUniqueFName;
 	FNiagaraCompilationGraph::FSharedPtr CalledGraph;
 	ENiagaraScriptUsage Usage;
 };
diff --git a/Engine/Plugins/FX/Niagara/Source/NiagaraEditor/Private/NiagaraTraversalStateContext.cpp b/Engine/Plugins/FX/Niagara/Source/NiagaraEditor/Private/NiagaraTraversalStateContext.cpp
--- a/Engine/Plugins/FX/Niagara/Source/NiagaraEditor/Private/NiagaraTraversalStateContext.cpp
+++ b/Engine/Plugins/FX/Niagara/Source/NiagaraEditor/Private/NiagaraTraversalStateContext.cpp
@@ -333,19 +333,19 @@ ENiagaraFunctionDebugState FNiagaraFixedConstantResolver::GetDebugState() const
 	return (ENiagaraFunctionDebugState)EnumValue.Value;
 }
 
-void FNiagaraFixedConstantResolver::AddNamedChildResolver(FName ScopeName, const FNiagaraFixedConstantResolver& ChildResolver)
+void FNiagaraFixedConstantResolver::AddChildResolver(const FGuid& ChildId, const FNiagaraFixedConstantResolver& ChildResolver)
 {
-	if (ensure(FindChildResolver(ScopeName) == nullptr))
+	if (ensure(FindChildResolver(ChildId) == nullptr))
 	{
-		ChildResolversByName.Emplace(ScopeName, ChildResolver);
+		ChildResolvers.Emplace(ChildId, ChildResolver);
 	}
 }
 
-const FNiagaraFixedConstantResolver* FNiagaraFixedConstantResolver::FindChildResolver(FName ScopeName) const
+const FNiagaraFixedConstantResolver* FNiagaraFixedConstantResolver::FindChildResolver(const FGuid& ChildId) const
 {
-	const FNamedResolverPair* ChildResolver = ChildResolversByName.FindByPredicate([ScopeName](const FNamedResolverPair& NamedPair) -> bool
+	const FNamedResolverPair* ChildResolver = ChildResolvers.FindByPredicate([ChildId](const FNamedResolverPair& NamedPair) -> bool
 	{
-			return NamedPair.Key == ScopeName;
+			return NamedPair.Key == ChildId;
 	});
 
 	return ChildResolver ? &ChildResolver->Value : nullptr;
