Make sure that the translator collapses attributes that are going to be overlapping in script execution (same name and equivalent type)
#tests Translated all assets and compared the results.  A number of systems were found to no longer suffer from overlapping registers

diff --git a/Engine/Plugins/FX/Niagara/Source/NiagaraEditor/Private/NiagaraHlslTranslator.cpp b/Engine/Plugins/FX/Niagara/Source/NiagaraEditor/Private/NiagaraHlslTranslator.cpp
--- a/Engine/Plugins/FX/Niagara/Source/NiagaraEditor/Private/NiagaraHlslTranslator.cpp
+++ b/Engine/Plugins/FX/Niagara/Source/NiagaraEditor/Private/NiagaraHlslTranslator.cpp
@@ -1172,6 +1172,53 @@ static void ConvertFloatToHalf(const FNiagaraCompileOptions& InCompileOptions, T
 	}
 }
 
+// Collapses the set of attributes based on if they would overlap (i.e. attributes share the same name and either have
+// the same type or an equivalent type, like position and vector, (for now)).
+static void CollapseAttributes(const FNiagaraCompileOptions& InCompileOptions, TArray<FNiagaraVariable>& SortedAttributes)
+{
+	const FNiagaraTypeDefinition& Vec3Type = FNiagaraTypeDefinition::GetVec3Def();
+	const FNiagaraTypeDefinition& PosType = FNiagaraTypeDefinition::GetPositionDef();
+
+	auto CanCollapseAttributeTypes = [&](const FNiagaraTypeDefinition& Lhs, const FNiagaraTypeDefinition& Rhs) -> bool
+	{
+		if (Lhs == Rhs)
+		{
+			return true;
+		}
+
+		// for now we're only going to consider the implications of Vec3 vs Position types
+		if ((Lhs == Vec3Type && Rhs == PosType)
+			|| (Lhs == PosType && Rhs == Vec3Type))
+		{
+			return true;
+		}
+
+		return false;
+	};
+
+	for (int32 AttrIt = 0; AttrIt < SortedAttributes.Num(); ++AttrIt)
+	{
+		const FNiagaraVariable& CurrentVariable = SortedAttributes[AttrIt];
+		const FNiagaraTypeDefinition& CurrentVariableType = CurrentVariable.GetType();
+
+		// look at all attributes sharing the name and remove all that have the same or equivalent type
+		int32 NextAttrIt = AttrIt + 1;
+		while (SortedAttributes.IsValidIndex(NextAttrIt) && SortedAttributes[NextAttrIt].GetName() == CurrentVariable.GetName())
+		{
+			const FNiagaraVariable& NextVariable = SortedAttributes[NextAttrIt];
+
+			if (CanCollapseAttributeTypes(CurrentVariableType, NextVariable.GetType()))
+			{
+				SortedAttributes.RemoveAt(NextAttrIt, EAllowShrinking::No);
+			}
+			else
+			{
+				++NextAttrIt;
+			}
+		}
+	}
+}
+
 template<typename GraphBridge>
 FNiagaraTranslateResults TNiagaraHlslTranslator<GraphBridge>::Translate(const FNiagaraCompileOptions& InCompileOptions, const FHlslNiagaraTranslatorOptions& InTranslateOptions)
 {
@@ -2036,6 +2083,7 @@ FNiagaraTranslateResults TNiagaraHlslTranslator<GraphBridge>::Translate(const FN
 		});
 
 		ConvertFloatToHalf(CompileOptions, BasicAttributes);
+		CollapseAttributes(CompileOptions, BasicAttributes);
 
 		DataSetVariables[InstanceReadVarsIndex] = BasicAttributes;
 		DataSetVariables[InstanceWriteVarsIndex] = BasicAttributes;
diff --git a/Engine/Source/Runtime/Core/Private/UObject/FortniteMainBranchObjectVersion.cpp b/Engine/Source/Runtime/Core/Private/UObject/FortniteMainBranchObjectVersion.cpp
--- a/Engine/Source/Runtime/Core/Private/UObject/FortniteMainBranchObjectVersion.cpp
+++ b/Engine/Source/Runtime/Core/Private/UObject/FortniteMainBranchObjectVersion.cpp
@@ -11,7 +11,7 @@ TMap<FGuid, FGuid> FFortniteMainBranchObjectVersion::GetSystemGuids()
 	SystemGuids.Add(DevGuids.MATERIALSHADERMAP_DERIVEDDATA_VER, FGuid("BB56D95CD4FD48E38C4CD7869B44B987"));
 	SystemGuids.Add(DevGuids.NANITE_DERIVEDDATA_VER, FGuid("D28C003F825F4C178F698938FB0E47F7"));
 	SystemGuids.Add(DevGuids.NIAGARASHADERMAP_DERIVEDDATA_VER, FGuid("C417A50271B3427D98785388266E3B66"));
-	SystemGuids.Add(DevGuids.Niagara_LatestScriptCompileVersion, FGuid("E267F600682A4E94885E8AEC5CE1DC91"));
+	SystemGuids.Add(DevGuids.Niagara_LatestScriptCompileVersion, FGuid("E7F9817BB9A14D74996523C3575A62EE"));
 	SystemGuids.Add(DevGuids.SkeletalMeshDerivedDataVersion, FGuid("B9AEAA3EE7AC4FCFB3860E8F3E06DD75"));
 	SystemGuids.Add(DevGuids.STATICMESH_DERIVEDDATA_VER, FGuid("543843B103794E0E9BA4BE80FB602F79"));
 	SystemGuids.Add(DevGuids.MaterialTranslationDDCVersion, FGuid("7A75688F394C4F27837B35000234B087"));
