Improvements to the NiagaraDumpByteCode commandlet
-takes advantage of the AsyncTasks compilation mode and batches up systems so that they can be prepared in parallel rather than one at a time
-adds the -hlsl option to include dumping the hlsl source for GPUComputeScript
#tests Ran commandlet over all the content

diff --git a/Engine/Plugins/FX/Niagara/Source/Niagara/Classes/NiagaraSystem.h b/Engine/Plugins/FX/Niagara/Source/Niagara/Classes/NiagaraSystem.h
--- a/Engine/Plugins/FX/Niagara/Source/Niagara/Classes/NiagaraSystem.h
+++ b/Engine/Plugins/FX/Niagara/Source/Niagara/Classes/NiagaraSystem.h
@@ -421,6 +421,9 @@ public:
 	/** Invalidates any active compilation requests which will ignore their results. */
 	NIAGARA_API void InvalidateActiveCompiles();
 
+	/** Returns true if there are any active compilations. */
+	NIAGARA_API bool HasActiveCompilations() const;
+
 	/** Delegate called when the system's dependencies have all been compiled.*/
 	NIAGARA_API FOnSystemCompiled& OnSystemCompiled();
 
diff --git a/Engine/Plugins/FX/Niagara/Source/Niagara/Private/NiagaraSystem.cpp b/Engine/Plugins/FX/Niagara/Source/Niagara/Private/NiagaraSystem.cpp
--- a/Engine/Plugins/FX/Niagara/Source/Niagara/Private/NiagaraSystem.cpp
+++ b/Engine/Plugins/FX/Niagara/Source/Niagara/Private/NiagaraSystem.cpp
@@ -2850,6 +2850,11 @@ void UNiagaraSystem::InvalidateActiveCompiles()
 	}
 }
 
+bool UNiagaraSystem::HasActiveCompilations() const
+{
+	return !ActiveCompilations.IsEmpty();
+}
+
 bool UNiagaraSystem::PollForCompilationComplete(bool bFlushRequestCompile)
 {
 	if (bNeedsRequestCompile && bFlushRequestCompile)
diff --git a/Engine/Plugins/FX/Niagara/Source/NiagaraEditor/Private/Commandlets/NiagaraDumpBytecodeCommandlet.cpp b/Engine/Plugins/FX/Niagara/Source/NiagaraEditor/Private/Commandlets/NiagaraDumpBytecodeCommandlet.cpp
--- a/Engine/Plugins/FX/Niagara/Source/NiagaraEditor/Private/Commandlets/NiagaraDumpBytecodeCommandlet.cpp
+++ b/Engine/Plugins/FX/Niagara/Source/NiagaraEditor/Private/Commandlets/NiagaraDumpBytecodeCommandlet.cpp
@@ -2,6 +2,7 @@
 
 #include "Commandlets/NiagaraDumpBytecodeCommandlet.h"
 
+#include "AssetCompilingManager.h"
 #include "AssetRegistry/ARFilter.h"
 #include "AssetRegistry/AssetData.h"
 #include "AssetRegistry/AssetRegistryModule.h"
@@ -9,6 +10,7 @@
 #include "CollectionManagerTypes.h"
 #include "HAL/FileManager.h"
 #include "ICollectionManager.h"
+#include "Misc/FileHelper.h"
 #include "Misc/OutputDeviceArchiveWrapper.h"
 #include "Misc/PackageName.h"
 #include "Misc/Paths.h"
@@ -61,6 +63,11 @@ int32 UNiagaraDumpByteCodeCommandlet::Main(const FString& Params)
 		PackagePaths.Add(FName(TEXT("/Game")));
 	}
 
+	if (Switches.Contains("IncludeDev"))
+	{
+		IncludeDeveloperFolder = true;
+	}
+
 	if (Switches.Contains("COOKED"))
 	{
 		ForceBakedRapidIteration = true;
@@ -77,103 +84,113 @@ int32 UNiagaraDumpByteCodeCommandlet::Main(const FString& Params)
 		ForceAttributeTrimming = true;
 	}
 
+	if (Switches.Contains("HLSL"))
+	{
+		DumpTranslatedHlsl = true;
+	}
+
 	ProcessNiagaraScripts();
 
 	return 0;
 }
 
-void UNiagaraDumpByteCodeCommandlet::ProcessNiagaraScripts()
+void UNiagaraDumpByteCodeCommandlet::ProcessBatch(TArray<FAssetData>& BatchAssets)
 {
-	FAssetRegistryModule& AssetRegistryModule = FModuleManager::LoadModuleChecked<FAssetRegistryModule>(TEXT("AssetRegistry"));
-	IAssetRegistry& AssetRegistry = AssetRegistryModule.Get();
-	AssetRegistry.SearchAllAssets(true);
-
-	FARFilter Filter;
-	Filter.PackagePaths = PackagePaths;
-	Filter.bRecursivePaths = true;
+	const UEnum* UsageEnum = StaticEnum<ENiagaraScriptUsage>();
+	const int32 BatchSize = BatchAssets.Num();
 
-	Filter.ClassPaths.Add(UNiagaraSystem::StaticClass()->GetClassPathName());
-	if (!FilterCollection.IsEmpty())
-	{
-		FCollectionManagerModule& CollectionManagerModule = FCollectionManagerModule::GetModule();
-		CollectionManagerModule.Get().GetObjectsInCollection(FName(*FilterCollection), ECollectionShareType::CST_All, Filter.SoftObjectPaths, ECollectionRecursionFlags::SelfAndChildren);
-	}
+	TArray<TObjectPtr<UNiagaraSystem>> LoadedSystems;
+	TArray<TObjectPtr<UNiagaraSystem>> PendingSystems;
+	TArray<TObjectPtr<UNiagaraSystem>> CompiledSystems;
+	TArray<TObjectPtr<UNiagaraSystem>> CompletedSystems;
 
-	TArray<FAssetData> AssetList;
-	AssetRegistry.GetAssets(Filter, AssetList);
-
-	const static UEnum* UsageEnum = StaticEnum<ENiagaraScriptUsage>();
+	LoadedSystems.Reserve(BatchSize);
+	PendingSystems.Reserve(BatchSize);
+	CompiledSystems.Reserve(BatchSize);
+	CompletedSystems.Reserve(BatchSize);
 
-	const double StartProcessNiagaraSystemsTime = FPlatformTime::Seconds();
-
-	//  Iterate over all scripts
-	const FString DevelopersFolder = FPackageName::FilenameToLongPackageName(FPaths::GameDevelopersDir().LeftChop(1));
-	FString LastPackageName = TEXT("");
-	UPackage* CurrentPackage = nullptr;
-	for (const FAssetData& AssetIt : AssetList)
+	// FullyLoad all the packages
+	for (const FAssetData& AssetIt : BatchAssets)
 	{
 		const FString SystemName = AssetIt.GetObjectPathString();
-		const FString PackageName = AssetIt.PackageName.ToString();
+		FNameBuilder PackageNameBuilder(AssetIt.PackageName);
 
-		if (PackageName.StartsWith(DevelopersFolder))
+		if (UPackage* Package = ::LoadPackage(nullptr, *PackageNameBuilder, LOAD_None))
 		{
-			// Skip developer folders
+			FNameBuilder ShortSystemNameBuilder(AssetIt.AssetName);
+
+			Package->FullyLoad();
+
+			if (UNiagaraSystem* NiagaraSystem = FindObject<UNiagaraSystem>(Package, *ShortSystemNameBuilder))
+			{
+				LoadedSystems.Emplace(NiagaraSystem);
+			}
+		}
+		else
+		{
+			UE_LOG(LogNiagaraDumpBytecodeCommandlet, Warning, TEXT("Failed to load package %s processing %s"), *PackageNameBuilder, *SystemName);
 			continue;
 		}
+	}
 
-		if (PackageName != LastPackageName)
+	bool bProcessing = true;
+
+	auto AdvanceSystems = [](TArray<TObjectPtr<UNiagaraSystem>>& InSystems, TArray<TObjectPtr<UNiagaraSystem>>& OutSystems, TFunction<bool(UNiagaraSystem*)> Op) -> void
+	{
+		while (!InSystems.IsEmpty())
 		{
-			UPackage* Package = ::LoadPackage(nullptr, *PackageName, LOAD_None);
-			if (Package != nullptr)
-			{
-				LastPackageName = PackageName;
-				Package->FullyLoad();
-				CurrentPackage = Package;
-			}
-			else
+			for (int32 InIt = 0; InIt < InSystems.Num(); ++InIt)
 			{
-				UE_LOG(LogNiagaraDumpBytecodeCommandlet, Warning, TEXT("Failed to load package %s processing %s"), *PackageName, *SystemName);
-				CurrentPackage = nullptr;
+				UNiagaraSystem* InSystem = InSystems[InIt];
+				if (Op(InSystem))
+				{
+					InSystems.RemoveAt(InIt, EAllowShrinking::No);
+					OutSystems.Add(InSystem);
+				}
 			}
-		}
 
-		const FString ShorterSystemName = AssetIt.AssetName.ToString();
-		UNiagaraSystem* NiagaraSystem = FindObject<UNiagaraSystem>(CurrentPackage, *ShorterSystemName);
-		if (NiagaraSystem == nullptr)
-		{
-			UE_LOG(LogNiagaraDumpBytecodeCommandlet, Warning, TEXT("Failed to load Niagara system %s"), *SystemName);
-			continue;
+			FAssetCompilingManager::Get().ProcessAsyncTasks(true);
 		}
+	};
 
-		NiagaraSystem->WaitForCompilationComplete();
+	AdvanceSystems(LoadedSystems, PendingSystems, [](UNiagaraSystem* System) -> bool
+	{
+		return !System->HasActiveCompilations() || System->PollForCompilationComplete();
+	});
 
+	AdvanceSystems(PendingSystems, CompiledSystems, [this](UNiagaraSystem* System) -> bool
+	{
 		if (ForceBakedRapidIteration || ForceAttributeTrimming)
 		{
 			if (ForceBakedRapidIteration)
 			{
-				NiagaraSystem->SetBakeOutRapidIterationOnCook(true);
+				System->SetBakeOutRapidIterationOnCook(true);
 			}
 			if (ForceAttributeTrimming)
 			{
-				NiagaraSystem->SetTrimAttributesOnCook(true);
+				System->SetTrimAttributesOnCook(true);
 			}
 
-			NiagaraSystem->RequestCompile(true);
-			NiagaraSystem->WaitForCompilationComplete(true);
+			System->RequestCompile(true);
 		}
 
-		if (!NiagaraSystem->IsValid())
-		{
-			UE_LOG(LogNiagaraDumpBytecodeCommandlet, Warning, TEXT("Loaded system was Invalid! %s"), *SystemName);
-		}
+		return true;
+	});
+
+	AdvanceSystems(CompiledSystems, CompletedSystems, [](UNiagaraSystem* System) -> bool
+	{
+		return System->PollForCompilationComplete();
+	});
 
+	for (UNiagaraSystem* NiagaraSystem : CompletedSystems)
+	{
 		const FString SystemPathName = NiagaraSystem->GetPathName();
 		const FString HashedPathName = FString::Printf(TEXT("%08x"), GetTypeHash(SystemPathName));
 
 		IFileManager::Get().MakeDirectory(*(AuditOutputFolder / HashedPathName));
-		DumpByteCode(NiagaraSystem->GetSystemSpawnScript(), SystemPathName, HashedPathName, TEXT("SystemSpawnScript.txt"));
-		DumpByteCode(NiagaraSystem->GetSystemUpdateScript(), SystemPathName, HashedPathName, TEXT("SystemUpdateScript.txt"));
-		
+		DumpByteCode(NiagaraSystem->GetSystemSpawnScript(), SystemPathName, HashedPathName, TEXT("SystemSpawnScript"));
+		DumpByteCode(NiagaraSystem->GetSystemUpdateScript(), SystemPathName, HashedPathName, TEXT("SystemUpdateScript"));
+
 		for (const auto& EmitterHandle : NiagaraSystem->GetEmitterHandles())
 		{
 			if (!EmitterHandle.GetIsEnabled())
@@ -194,13 +211,64 @@ void UNiagaraDumpByteCodeCommandlet::ProcessNiagaraScripts()
 
 					for (const auto* EmitterScript : EmitterScripts)
 					{
-						DumpByteCode(EmitterScript, SystemPathName, HashedPathName, EmitterName / UsageEnum->GetNameStringByValue(static_cast<int64>(EmitterScript->GetUsage())) + TEXT(".txt"));
+						DumpByteCode(EmitterScript, SystemPathName, HashedPathName, EmitterName / UsageEnum->GetNameStringByValue(static_cast<int64>(EmitterScript->GetUsage())));
 					}
 				}
 			}
 		}
 	}
 
+	::CollectGarbage(RF_NoFlags);
+}
+
+void UNiagaraDumpByteCodeCommandlet::ProcessNiagaraScripts()
+{
+	FAssetRegistryModule& AssetRegistryModule = FModuleManager::LoadModuleChecked<FAssetRegistryModule>(TEXT("AssetRegistry"));
+	IAssetRegistry& AssetRegistry = AssetRegistryModule.Get();
+	AssetRegistry.SearchAllAssets(true);
+
+	FARFilter Filter;
+	Filter.PackagePaths = PackagePaths;
+	Filter.bRecursivePaths = true;
+
+	Filter.ClassPaths.Add(UNiagaraSystem::StaticClass()->GetClassPathName());
+	if (!FilterCollection.IsEmpty())
+	{
+		FCollectionManagerModule& CollectionManagerModule = FCollectionManagerModule::GetModule();
+		CollectionManagerModule.Get().GetObjectsInCollection(FName(*FilterCollection), ECollectionShareType::CST_All, Filter.SoftObjectPaths, ECollectionRecursionFlags::SelfAndChildren);
+	}
+
+	TArray<FAssetData> AssetList;
+	AssetRegistry.GetAssets(Filter, AssetList);
+
+	const double StartProcessNiagaraSystemsTime = FPlatformTime::Seconds();
+	constexpr uint32 BatchPackageCount = 256;
+
+	TArray<FAssetData> BatchAssets;
+
+	//  Iterate over all scripts
+	const FString DevelopersFolder = FPackageName::FilenameToLongPackageName(FPaths::GameDevelopersDir().LeftChop(1));
+	for (const FAssetData& AssetIt : AssetList)
+	{
+		if (!IncludeDeveloperFolder && AssetIt.PackageName.ToString().StartsWith(DevelopersFolder))
+		{
+			// Skip developer folders
+			continue;
+		}
+
+		if (BatchAssets.Num() == BatchPackageCount)
+		{
+			ProcessBatch(BatchAssets);
+		}
+
+		BatchAssets.Add(AssetIt);
+	}
+
+	if (!BatchAssets.IsEmpty())
+	{
+		ProcessBatch(BatchAssets);
+	}
+
 	// sort the data alphabetically (based on MetaData.FullName
 	ScriptMetaData.StableSort([](const FScriptMetaData& Lhs, const FScriptMetaData& Rhs)
 	{
@@ -278,11 +346,9 @@ void UNiagaraDumpByteCodeCommandlet::DumpByteCode(const UNiagaraScript* Script,
 		MetaData.AttributeCount += Var.GetType().GetSize() / 4;
 	}
 
-	//const static UEnum* VmOpEnum = StaticEnum<EVectorVMOp>();
-
 	if (Script)
 	{
-		const FString FullFilePath = AuditOutputFolder / HashName / FilePath;
+		const FString FullFilePath = AuditOutputFolder / HashName / FilePath + TEXT(".vm");
 
 		TUniquePtr<FArchive> FileArchive(IFileManager::Get().CreateDebugFileWriter(*FullFilePath));
 		if (!FileArchive)
@@ -325,4 +391,11 @@ void UNiagaraDumpByteCodeCommandlet::DumpByteCode(const UNiagaraScript* Script,
 			OutputStream->Log(CurrentLine);
 		}
 	}
+
+	if (Script && DumpTranslatedHlsl)
+	{
+		const FString FullFilePath = AuditOutputFolder / HashName / FilePath + TEXT(".usf");
+
+		FFileHelper::SaveStringToFile(Script->GetVMExecutableData().LastHlslTranslation, *FullFilePath);
+	}
 }
diff --git a/Engine/Plugins/FX/Niagara/Source/NiagaraEditor/Public/Commandlets/NiagaraDumpBytecodeCommandlet.h b/Engine/Plugins/FX/Niagara/Source/NiagaraEditor/Public/Commandlets/NiagaraDumpBytecodeCommandlet.h
--- a/Engine/Plugins/FX/Niagara/Source/NiagaraEditor/Public/Commandlets/NiagaraDumpBytecodeCommandlet.h
+++ b/Engine/Plugins/FX/Niagara/Source/NiagaraEditor/Public/Commandlets/NiagaraDumpBytecodeCommandlet.h
@@ -31,8 +31,15 @@ class UNiagaraDumpByteCodeCommandlet : public UCommandlet
 	/** Whether to force attributes to be trimmed to their minimal set */
 	bool ForceAttributeTrimming = false;
 
+	/** Whether the translated hlsl should also be dumped */
+	bool DumpTranslatedHlsl = false;
+
+	/** */
+	bool IncludeDeveloperFolder = false;
+
 private:
 	void ProcessNiagaraScripts();
+	void ProcessBatch(TArray<FAssetData>& BatchAssets);
 	void DumpByteCode(const UNiagaraScript* Script, const FString& PathName, const FString& HashName, const FString& FilePath);
 
 	struct FScriptMetaData
