- Add CSV tracking for core FX systems
- Add missing finalize tasks from per system tracking

#tests Various replay runs to gather stats and chart them

diff --git a/Engine/Plugins/FX/Niagara/Source/Niagara/Private/NiagaraComponent.cpp b/Engine/Plugins/FX/Niagara/Source/Niagara/Private/NiagaraComponent.cpp
--- a/Engine/Plugins/FX/Niagara/Source/Niagara/Private/NiagaraComponent.cpp
+++ b/Engine/Plugins/FX/Niagara/Source/Niagara/Private/NiagaraComponent.cpp
@@ -2201,6 +2201,7 @@ void UNiagaraComponent::SendRenderDynamicData_Concurrent()
 {
 	LLM_SCOPE(ELLMTag::Niagara);
 	CSV_SCOPED_TIMING_STAT_EXCLUSIVE(Effects);
+	CSV_SCOPED_TIMING_STAT(Particles, CoreSystems_NiagaraSendRenderDynamicData);
 	SCOPE_CYCLE_COUNTER(STAT_NiagaraComponentSendRenderData);
 	PARTICLE_PERF_STAT_CYCLES_GT(FParticlePerfStatsContext(GetWorld(), GetAsset(), this), EndOfFrame);
 
diff --git a/Engine/Plugins/FX/Niagara/Source/Niagara/Private/NiagaraDataChannelManager.cpp b/Engine/Plugins/FX/Niagara/Source/Niagara/Private/NiagaraDataChannelManager.cpp
--- a/Engine/Plugins/FX/Niagara/Source/Niagara/Private/NiagaraDataChannelManager.cpp
+++ b/Engine/Plugins/FX/Niagara/Source/Niagara/Private/NiagaraDataChannelManager.cpp
@@ -66,6 +66,8 @@ void FNiagaraDataChannelManager::BeginFrame(float DeltaSeconds)
 	{
 		check(IsInGameThread());
 		SCOPE_CYCLE_COUNTER(STAT_DataChannelManager_BeginFrame);
+		CSV_SCOPED_TIMING_STAT(Particles, CoreSystems_NiagaraDataChannelManager);
+
 		//Tick all DataChannel channel handlers.
 		for (auto It = Channels.CreateIterator(); It; ++It)
 		{
@@ -87,6 +89,8 @@ void FNiagaraDataChannelManager::EndFrame(float DeltaSeconds)
 	{
 		check(IsInGameThread());
 		SCOPE_CYCLE_COUNTER(STAT_DataChannelManager_EndFrame);
+		CSV_SCOPED_TIMING_STAT(Particles, CoreSystems_NiagaraDataChannelManager);
+
 		//Tick all DataChannel channel handlers.
 		for (auto& ChannelPair : Channels)
 		{
@@ -101,6 +105,8 @@ void FNiagaraDataChannelManager::Tick(float DeltaSeconds, ETickingGroup TickGrou
 	{
 		check(IsInGameThread());
 		SCOPE_CYCLE_COUNTER(STAT_DataChannelManager_Tick);
+		CSV_SCOPED_TIMING_STAT(Particles, CoreSystems_NiagaraDataChannelManager);
+
 		//Tick all DataChannel channel handlers.
 		for (auto& ChannelPair : Channels)
 		{
diff --git a/Engine/Plugins/FX/Niagara/Source/Niagara/Private/NiagaraSystemSimulation.cpp b/Engine/Plugins/FX/Niagara/Source/Niagara/Private/NiagaraSystemSimulation.cpp
--- a/Engine/Plugins/FX/Niagara/Source/Niagara/Private/NiagaraSystemSimulation.cpp
+++ b/Engine/Plugins/FX/Niagara/Source/Niagara/Private/NiagaraSystemSimulation.cpp
@@ -432,6 +432,9 @@ struct FNiagaraSystemInstanceFinalizeTask
 	{
 #if DO_CHECK
 		DebugCounter = 0;
+#endif
+#if WITH_PARTICLE_PERF_STATS
+		SystemPerfStatsContext = FParticlePerfStatsContext(InSystemSimulation->GetWorld(), InSystemSimulation->GetSystem());
 #endif
 		for ( int32 i=0; i < Batch.Num(); ++i )
 		{
@@ -449,6 +452,9 @@ struct FNiagaraSystemInstanceFinalizeTask
 
 	void DoTask(ENamedThreads::Type CurrentThread, const FGraphEventRef& MyCompletionGraphEvent)
 	{
+#if WITH_PARTICLE_PERF_STATS
+		PARTICLE_PERF_STAT_CYCLES_GT(SystemPerfStatsContext, Finalize);
+#endif
 		CSV_SCOPED_TIMING_STAT_EXCLUSIVE(Effects);
 		check(CurrentThread == ENamedThreads::GameThread);
 
@@ -510,6 +516,9 @@ struct FNiagaraSystemInstanceFinalizeTask
 #if DO_CHECK
 	std::atomic<int> DebugCounter;
 #endif
+#if WITH_PARTICLE_PERF_STATS
+	FParticlePerfStatsContext SystemPerfStatsContext;
+#endif
 };
 
 //////////////////////////////////////////////////////////////////////////
@@ -1115,7 +1124,6 @@ void FNiagaraSystemSimulation::Tick_GameThread_Internal(float DeltaSeconds, cons
 	LLM_SCOPE(ELLMTag::Niagara);
 	FScopeCycleCounterUObject AdditionalScope(GetSystem(), GET_STATID(STAT_NiagaraOverview_GT_CNC));
 
-
 #if STATS
 	FScopeCycleCounter SystemStatCounter(System->GetStatID(true, false));
 #endif
@@ -1324,7 +1332,6 @@ void FNiagaraSystemSimulation::UpdateTickGroups_GameThread()
 	check(IsInGameThread());
 	check(!bIsSolo);
 
-	SCOPE_CYCLE_COUNTER(STAT_NiagaraSystemSim_SpawnNewGT);
 	SCOPE_CYCLE_COUNTER(STAT_NiagaraOverview_GT);
 	CSV_SCOPED_TIMING_STAT_EXCLUSIVE(Effects);
 	LLM_SCOPE(ELLMTag::Niagara);
@@ -1391,6 +1398,7 @@ void FNiagaraSystemSimulation::Spawn_GameThread(float DeltaSeconds, bool bPostAc
 	ConcurrentTickGraphEvent = nullptr;
 	AllWorkCompleteGraphEvent = nullptr;
 
+	PARTICLE_PERF_STAT_CYCLES_GT(FParticlePerfStatsContext(World, System), TickGameThread);
 	SCOPE_CYCLE_COUNTER(STAT_NiagaraSystemSim_SpawnNewGT);
 	SCOPE_CYCLE_COUNTER(STAT_NiagaraOverview_GT);
 	CSV_SCOPED_TIMING_STAT_EXCLUSIVE(Effects);
diff --git a/Engine/Plugins/FX/Niagara/Source/Niagara/Private/NiagaraWorldManager.cpp b/Engine/Plugins/FX/Niagara/Source/Niagara/Private/NiagaraWorldManager.cpp
--- a/Engine/Plugins/FX/Niagara/Source/Niagara/Private/NiagaraWorldManager.cpp
+++ b/Engine/Plugins/FX/Niagara/Source/Niagara/Private/NiagaraWorldManager.cpp
@@ -572,6 +572,7 @@ UNiagaraCullProxyComponent* FNiagaraWorldManager::GetCullProxy(UNiagaraComponent
 void FNiagaraWorldManager::TickParameterCollections()
 {
 	SCOPE_CYCLE_COUNTER(STAT_NiagaraWorldManTickParamCollections);
+	CSV_SCOPED_TIMING_STAT(Particles, CoreSystems_NiagaraTickParameterCollections);
 
 	//-TODO: Do we need to do this per tick group?
 	for (TPair<UNiagaraParameterCollection*, UNiagaraParameterCollectionInstance*> CollectionInstPair : ParameterCollections)
@@ -1031,6 +1032,7 @@ void FNiagaraWorldManager::TickStart(float DeltaSeconds)
 #if WITH_NIAGARA_LEAK_DETECTOR
 	if (ComponentLeakDetector.IsValid() && HasActiveWorld())
 	{
+		CSV_SCOPED_TIMING_STAT(Particles, CoreSystems_NiagaraLeakDetector);
 		ComponentLeakDetector->Tick(World);
 	}
 #endif
@@ -1070,12 +1072,15 @@ void FNiagaraWorldManager::PostActorTick(float DeltaSeconds)
 		}
 
 		// Update tick groups
-		for (int32 i = 0; i < SimulationsWithPostActorWork.Num(); ++i)
 		{
-			FNiagaraSystemSimulationRef& Simulation = SimulationsWithPostActorWork[i];
-			if (Simulation->IsValid())
+			CSV_SCOPED_TIMING_STAT(Particles, CoreSystems_NiagaraPostActorUpdateTickGroups);
+			for (int32 i = 0; i < SimulationsWithPostActorWork.Num(); ++i)
 			{
-				Simulation->UpdateTickGroups_GameThread();
+				FNiagaraSystemSimulationRef& Simulation = SimulationsWithPostActorWork[i];
+				if (Simulation->IsValid())
+				{
+					Simulation->UpdateTickGroups_GameThread();
+				}
 			}
 		}
 
@@ -1114,6 +1119,7 @@ void FNiagaraWorldManager::PostActorTick(float DeltaSeconds)
 	if (NiagaraDebugHud != nullptr)
 	{
 		QUICK_SCOPE_CYCLE_COUNTER(STAT_NiagaraDebugHUD);
+		CSV_SCOPED_TIMING_STAT(Particles, CoreSystems_NiagaraDebugHUD);
 		NiagaraDebugHud->GatherSystemInfo();
 	}
 #endif
@@ -1345,42 +1351,45 @@ void FNiagaraWorldManager::Tick(ETickingGroup TickGroup, float DeltaSeconds, ELe
 			return;
 		}
 
-		bool bUseWorldCachedViews = !World->GetPlayerControllerIterator();
-#if WITH_EDITOR
-		if (GCurrentLevelEditingViewportClient && (GCurrentLevelEditingViewportClient->GetWorld() == World))
 		{
-			bUseWorldCachedViews = true;
-		}
+			CSV_SCOPED_TIMING_STAT(Particles, CoreSystems_NiagaraCacheViewInformation);
+			bool bUseWorldCachedViews = !World->GetPlayerControllerIterator();
+#if WITH_EDITOR
+			if (GCurrentLevelEditingViewportClient && (GCurrentLevelEditingViewportClient->GetWorld() == World))
+			{
+				bUseWorldCachedViews = true;
+			}
 #endif
-		// Cache player view info for all system instances to access
-		//-TODO: Do we need to do this per tick group?
-		if (bUseWorldCachedViews)
-		{
-			for (int32 i = 0; i < World->CachedViewInfoRenderedLastFrame.Num(); ++i)
+			// Cache player view info for all system instances to access
+			//-TODO: Do we need to do this per tick group?
+			if (bUseWorldCachedViews)
 			{
-				FWorldCachedViewInfo& WorldViewInfo = World->CachedViewInfoRenderedLastFrame[i];
+				for (int32 i = 0; i < World->CachedViewInfoRenderedLastFrame.Num(); ++i)
+				{
+					FWorldCachedViewInfo& WorldViewInfo = World->CachedViewInfoRenderedLastFrame[i];
 
-				FNiagaraCachedViewInfo& ViewInfo = CachedViewInfo.AddDefaulted_GetRef();
-				ViewInfo.Init(WorldViewInfo);
+					FNiagaraCachedViewInfo& ViewInfo = CachedViewInfo.AddDefaulted_GetRef();
+					ViewInfo.Init(WorldViewInfo);
+				}
 			}
-		}
-		else
-		{
-			for (FConstPlayerControllerIterator Iterator = World->GetPlayerControllerIterator(); Iterator; ++Iterator)
+			else
 			{
-				APlayerController* PlayerController = Iterator->Get();
-				if (PlayerController && PlayerController->IsLocalPlayerController())
+				for (FConstPlayerControllerIterator Iterator = World->GetPlayerControllerIterator(); Iterator; ++Iterator)
 				{
-					FNiagaraCachedViewInfo& ViewInfo = CachedViewInfo.AddDefaulted_GetRef();
-					
-					const bool bIsValid = PrepareCachedViewInfo(PlayerController, ViewInfo);
-					if (!bIsValid)
+					APlayerController* PlayerController = Iterator->Get();
+					if (PlayerController && PlayerController->IsLocalPlayerController())
 					{
-						CachedViewInfo.RemoveAt(CachedViewInfo.Num() - 1);
+						FNiagaraCachedViewInfo& ViewInfo = CachedViewInfo.AddDefaulted_GetRef();
+					
+						const bool bIsValid = PrepareCachedViewInfo(PlayerController, ViewInfo);
+						if (!bIsValid)
+						{
+							CachedViewInfo.RemoveAt(CachedViewInfo.Num() - 1);
+						}
 					}
 				}
-			}
-		}		
+			}		
+		}
 
 		UpdateScalabilityManagers(DeltaSeconds, false);
 
@@ -1394,9 +1403,12 @@ void FNiagaraWorldManager::Tick(ETickingGroup TickGroup, float DeltaSeconds, ELe
 	}
 
 	// Tick generated data
-	for (auto& GeneratedData : DIGeneratedData)
 	{
-		GeneratedData.Value->Tick(TickGroup, DeltaSeconds);
+		CSV_SCOPED_TIMING_STAT(Particles, CoreSystems_NiagaraDIGeneratedData);
+		for (auto& GeneratedData : DIGeneratedData)
+		{
+			GeneratedData.Value->Tick(TickGroup, DeltaSeconds);
+		}
 	}
 
 	// Now tick all system instances. 
@@ -1545,6 +1557,7 @@ FVector::FReal FNiagaraWorldManager::GetLODDistance(FVector Location)const
 void FNiagaraWorldManager::UpdateScalabilityManagers(float DeltaSeconds, bool bNewSpawnsOnly)
 {
 	SCOPE_CYCLE_COUNTER(STAT_UpdateScalabilityManagers);
+	CSV_SCOPED_TIMING_STAT(Particles, CoreSystems_NiagaraUpdateScalabilityManagers);
 
 	for (auto& Pair : ScalabilityManagers)
 	{
diff --git a/Engine/Source/Runtime/Engine/Private/Particles/ParticleComponents.cpp b/Engine/Source/Runtime/Engine/Private/Particles/ParticleComponents.cpp
--- a/Engine/Source/Runtime/Engine/Private/Particles/ParticleComponents.cpp
+++ b/Engine/Source/Runtime/Engine/Private/Particles/ParticleComponents.cpp
@@ -4247,6 +4247,7 @@ void UParticleSystemComponent::SendRenderDynamicData_Concurrent()
 	SCOPE_CYCLE_COUNTER(STAT_ParticleSystemComponent_SendRenderDynamicData_Concurrent);
 	SCOPE_CYCLE_COUNTER(STAT_ParticlesOverview_GT_CNC);
 	CSV_SCOPED_TIMING_STAT_EXCLUSIVE(Effects);
+	CSV_SCOPED_TIMING_STAT(Particles, CoreSystems_CascadeSendRenderDynamicData);
 	PARTICLE_PERF_STAT_CYCLES_GT(FParticlePerfStatsContext(GetWorld(), Template, this), EndOfFrame);
 
 	ForceAsyncWorkCompletion(ENSURE_AND_STALL, false, true);
diff --git a/Engine/Source/Runtime/Engine/Private/Particles/ParticleSystemManager.cpp b/Engine/Source/Runtime/Engine/Private/Particles/ParticleSystemManager.cpp
--- a/Engine/Source/Runtime/Engine/Private/Particles/ParticleSystemManager.cpp
+++ b/Engine/Source/Runtime/Engine/Private/Particles/ParticleSystemManager.cpp
@@ -705,7 +705,8 @@ void FParticleSystemWorldManager::Tick(ETickingGroup TickGroup, float DeltaTime,
 
 	SCOPE_CYCLE_COUNTER(STAT_PSCMan_Tick);
 	CSV_SCOPED_TIMING_STAT_EXCLUSIVE(Effects);
-	
+	CSV_SCOPED_TIMING_STAT(Particles, CoreSystems_CascadeSystemWorldManager);
+
 	//Do some book keeping in the first tick group, PrePhysics.
 	int32 BuildListStart = ManagedPSCs.Num();
 	if (TickGroup == TG_PrePhysics)
