- Add Position distribution
- Use in stateless initialize particle module to match stateful (i.e. only bind to positions not vectors)

#rnx

diff --git a/Engine/Plugins/FX/Niagara/Source/Niagara/Internal/Stateless/Modules/NiagaraStatelessModule_InitializeParticle.h b/Engine/Plugins/FX/Niagara/Source/Niagara/Internal/Stateless/Modules/NiagaraStatelessModule_InitializeParticle.h
--- a/Engine/Plugins/FX/Niagara/Source/Niagara/Internal/Stateless/Modules/NiagaraStatelessModule_InitializeParticle.h
+++ b/Engine/Plugins/FX/Niagara/Source/Niagara/Internal/Stateless/Modules/NiagaraStatelessModule_InitializeParticle.h
@@ -75,7 +75,7 @@ public:
 	FNiagaraDistributionRangeFloat RibbonWidthDistribution = FNiagaraDistributionRangeFloat(FNiagaraStatelessGlobals::GetDefaultRibbonWidthValue());
 
 	UPROPERTY(EditAnywhere, Category = "Parameters", meta = (DisableCurveDistribution))
-	FNiagaraDistributionVector3	InitialPosition = FNiagaraDistributionVector3(FVector3f::ZeroVector);
+	FNiagaraDistributionPosition InitialPositionDistribution = FNiagaraDistributionPosition(FVector3f::ZeroVector);
 
 	virtual void BuildEmitterData(const FNiagaraStatelessEmitterDataBuildContext& BuildContext) const override
 	{
@@ -85,7 +85,7 @@ public:
 		BuiltData->ModuleFlags				 = SpriteSizeDistribution.IsUniform() ? EInitializeParticleModuleFlag_UniformSpriteSize : 0;
 		BuiltData->ModuleFlags				|= MeshScaleDistribution.IsUniform() ? EInitializeParticleModuleFlag_UniformMeshScale : 0;
 
-		BuiltData->InitialPosition			= BuildContext.AddDistribution(InitialPosition, true);
+		BuiltData->InitialPosition			= BuildContext.AddDistribution(InitialPositionDistribution);
 		BuiltData->LifetimeRange			= LifetimeDistribution.CalculateRange(FNiagaraStatelessGlobals::GetDefaultLifetimeValue());
 		BuiltData->ColorRange				= BuildContext.ConvertDistributionToRange(ColorDistribution, FNiagaraStatelessGlobals::GetDefaultColorValue());
 		BuiltData->MassRange				= BuildContext.ConvertDistributionToRange(MassDistribution, FNiagaraStatelessGlobals::GetDefaultMassValue());
diff --git a/Engine/Plugins/FX/Niagara/Source/Niagara/Internal/Stateless/NiagaraStatelessDistribution.h b/Engine/Plugins/FX/Niagara/Source/Niagara/Internal/Stateless/NiagaraStatelessDistribution.h
--- a/Engine/Plugins/FX/Niagara/Source/Niagara/Internal/Stateless/NiagaraStatelessDistribution.h
+++ b/Engine/Plugins/FX/Niagara/Source/Niagara/Internal/Stateless/NiagaraStatelessDistribution.h
@@ -277,6 +277,20 @@ struct FNiagaraDistributionVector3 : public FNiagaraDistributionBase
 #endif
 };
 
+USTRUCT()
+struct FNiagaraDistributionPosition : public FNiagaraDistributionVector3
+{
+	GENERATED_BODY()
+
+	FNiagaraDistributionPosition() = default;
+	explicit FNiagaraDistributionPosition(const float ConstantValue) { InitConstant(ConstantValue); }
+	explicit FNiagaraDistributionPosition(const FVector3f& ConstantValue) { InitConstant(ConstantValue); }
+
+#if WITH_EDITORONLY_DATA
+	virtual FNiagaraTypeDefinition GetBindingTypeDef() const { return FNiagaraTypeDefinition::GetPositionDef(); }
+#endif
+};
+
 USTRUCT()
 struct FNiagaraDistributionColor : public FNiagaraDistributionBase
 {
diff --git a/Engine/Plugins/FX/Niagara/Source/NiagaraEditor/Private/Stateless/NiagaraDistributionPropertyCustomization.cpp b/Engine/Plugins/FX/Niagara/Source/NiagaraEditor/Private/Stateless/NiagaraDistributionPropertyCustomization.cpp
--- a/Engine/Plugins/FX/Niagara/Source/NiagaraEditor/Private/Stateless/NiagaraDistributionPropertyCustomization.cpp
+++ b/Engine/Plugins/FX/Niagara/Source/NiagaraEditor/Private/Stateless/NiagaraDistributionPropertyCustomization.cpp
@@ -904,6 +904,11 @@ TSharedRef<IPropertyTypeCustomization> FNiagaraDistributionPropertyCustomization
 	return MakeShareable<FNiagaraDistributionPropertyCustomization>(new FNiagaraDistributionPropertyCustomization(Vector3DistributionPropertyHandleToDistributionAdapter));
 }
 
+TSharedRef<IPropertyTypeCustomization> FNiagaraDistributionPropertyCustomization::MakePositionInstance()
+{
+	return MakeVector3Instance();
+}
+
 TSharedRef<IPropertyTypeCustomization> FNiagaraDistributionPropertyCustomization::MakeColorInstance()
 {
 	FPropertyHandleToDistributionAdapter ColorDistributionPropertyHandleToDistributionAdapter = FPropertyHandleToDistributionAdapter::CreateLambda([](TSharedRef<IPropertyHandle> ColorDistributionPropertyHandle)
diff --git a/Engine/Plugins/FX/Niagara/Source/NiagaraEditor/Private/Stateless/NiagaraDistributionPropertyCustomization.h b/Engine/Plugins/FX/Niagara/Source/NiagaraEditor/Private/Stateless/NiagaraDistributionPropertyCustomization.h
--- a/Engine/Plugins/FX/Niagara/Source/NiagaraEditor/Private/Stateless/NiagaraDistributionPropertyCustomization.h
+++ b/Engine/Plugins/FX/Niagara/Source/NiagaraEditor/Private/Stateless/NiagaraDistributionPropertyCustomization.h
@@ -13,6 +13,7 @@ public:
 	static TSharedRef<IPropertyTypeCustomization> MakeFloatInstance();
 	static TSharedRef<IPropertyTypeCustomization> MakeVector2Instance();
 	static TSharedRef<IPropertyTypeCustomization> MakeVector3Instance();
+	static TSharedRef<IPropertyTypeCustomization> MakePositionInstance();
 	static TSharedRef<IPropertyTypeCustomization> MakeColorInstance();
 
 	virtual void CustomizeHeader(TSharedRef<IPropertyHandle> PropertyHandle, FDetailWidgetRow& HeaderRow, IPropertyTypeCustomizationUtils& CustomizationUtils) override;
diff --git a/Engine/Plugins/FX/Niagara/Source/NiagaraEditor/Private/ViewModels/Stack/NiagaraStackStatelessEmitterGroup.cpp b/Engine/Plugins/FX/Niagara/Source/NiagaraEditor/Private/ViewModels/Stack/NiagaraStackStatelessEmitterGroup.cpp
--- a/Engine/Plugins/FX/Niagara/Source/NiagaraEditor/Private/ViewModels/Stack/NiagaraStackStatelessEmitterGroup.cpp
+++ b/Engine/Plugins/FX/Niagara/Source/NiagaraEditor/Private/ViewModels/Stack/NiagaraStackStatelessEmitterGroup.cpp
@@ -118,6 +118,7 @@ void UNiagaraStackStatelessEmitterObjectItem::RefreshChildrenInternal(const TArr
 			StatelessEmitterStackObject->RegisterInstancedCustomPropertyTypeLayout(FNiagaraDistributionFloat::StaticStruct()->GetFName(), FOnGetPropertyTypeCustomizationInstance::CreateStatic(&FNiagaraDistributionPropertyCustomization::MakeFloatInstance));
 			StatelessEmitterStackObject->RegisterInstancedCustomPropertyTypeLayout(FNiagaraDistributionVector2::StaticStruct()->GetFName(), FOnGetPropertyTypeCustomizationInstance::CreateStatic(&FNiagaraDistributionPropertyCustomization::MakeVector2Instance));
 			StatelessEmitterStackObject->RegisterInstancedCustomPropertyTypeLayout(FNiagaraDistributionVector3::StaticStruct()->GetFName(), FOnGetPropertyTypeCustomizationInstance::CreateStatic(&FNiagaraDistributionPropertyCustomization::MakeVector3Instance));
+			StatelessEmitterStackObject->RegisterInstancedCustomPropertyTypeLayout(FNiagaraDistributionPosition::StaticStruct()->GetFName(), FOnGetPropertyTypeCustomizationInstance::CreateStatic(&FNiagaraDistributionPropertyCustomization::MakePositionInstance));
 			StatelessEmitterStackObject->RegisterInstancedCustomPropertyTypeLayout(FNiagaraDistributionColor::StaticStruct()->GetFName(), FOnGetPropertyTypeCustomizationInstance::CreateStatic(&FNiagaraDistributionPropertyCustomization::MakeColorInstance));
 			StatelessEmitterStackObject->RegisterInstancedCustomPropertyTypeLayout(FNiagaraDistributionRangeFloat::StaticStruct()->GetFName(), FOnGetPropertyTypeCustomizationInstance::CreateStatic(&FNiagaraDistributionPropertyCustomization::MakeFloatInstance));
 			StatelessEmitterStackObject->RegisterInstancedCustomPropertyTypeLayout(FNiagaraDistributionRangeVector2::StaticStruct()->GetFName(), FOnGetPropertyTypeCustomizationInstance::CreateStatic(&FNiagaraDistributionPropertyCustomization::MakeVector2Instance));
diff --git a/Engine/Plugins/FX/Niagara/Source/NiagaraEditor/Private/ViewModels/Stack/NiagaraStackStatelessEmitterSimulateGroup.cpp b/Engine/Plugins/FX/Niagara/Source/NiagaraEditor/Private/ViewModels/Stack/NiagaraStackStatelessEmitterSimulateGroup.cpp
--- a/Engine/Plugins/FX/Niagara/Source/NiagaraEditor/Private/ViewModels/Stack/NiagaraStackStatelessEmitterSimulateGroup.cpp
+++ b/Engine/Plugins/FX/Niagara/Source/NiagaraEditor/Private/ViewModels/Stack/NiagaraStackStatelessEmitterSimulateGroup.cpp
@@ -231,6 +231,7 @@ void UNiagaraStackStatelessModuleItem::RefreshChildrenInternal(const TArray<UNia
 			ModuleObject->RegisterInstancedCustomPropertyTypeLayout(FNiagaraDistributionFloat::StaticStruct()->GetFName(), FOnGetPropertyTypeCustomizationInstance::CreateStatic(&FNiagaraDistributionPropertyCustomization::MakeFloatInstance));
 			ModuleObject->RegisterInstancedCustomPropertyTypeLayout(FNiagaraDistributionVector2::StaticStruct()->GetFName(), FOnGetPropertyTypeCustomizationInstance::CreateStatic(&FNiagaraDistributionPropertyCustomization::MakeVector2Instance));
 			ModuleObject->RegisterInstancedCustomPropertyTypeLayout(FNiagaraDistributionVector3::StaticStruct()->GetFName(), FOnGetPropertyTypeCustomizationInstance::CreateStatic(&FNiagaraDistributionPropertyCustomization::MakeVector3Instance));
+			ModuleObject->RegisterInstancedCustomPropertyTypeLayout(FNiagaraDistributionPosition::StaticStruct()->GetFName(), FOnGetPropertyTypeCustomizationInstance::CreateStatic(&FNiagaraDistributionPropertyCustomization::MakePositionInstance));
 			ModuleObject->RegisterInstancedCustomPropertyTypeLayout(FNiagaraDistributionColor::StaticStruct()->GetFName(), FOnGetPropertyTypeCustomizationInstance::CreateStatic(&FNiagaraDistributionPropertyCustomization::MakeColorInstance));
 			ModuleObject->RegisterInstancedCustomPropertyTypeLayout(FNiagaraDistributionRangeFloat::StaticStruct()->GetFName(), FOnGetPropertyTypeCustomizationInstance::CreateStatic(&FNiagaraDistributionPropertyCustomization::MakeFloatInstance));
 			ModuleObject->RegisterInstancedCustomPropertyTypeLayout(FNiagaraDistributionRangeVector2::StaticStruct()->GetFName(), FOnGetPropertyTypeCustomizationInstance::CreateStatic(&FNiagaraDistributionPropertyCustomization::MakeVector2Instance));
