PR #11844: Update NiagaraSystem.h - Lost NIAGARA_API in UNiagaraSystem::PreSave declaration

diff --git a/Engine/Plugins/FX/Niagara/Source/Niagara/Classes/NiagaraSystem.h b/Engine/Plugins/FX/Niagara/Source/Niagara/Classes/NiagaraSystem.h
--- a/Engine/Plugins/FX/Niagara/Source/Niagara/Classes/NiagaraSystem.h
+++ b/Engine/Plugins/FX/Niagara/Source/Niagara/Classes/NiagaraSystem.h
@@ -238,7 +238,7 @@ public:
 #endif
 	NIAGARA_API virtual void BeginDestroy() override;
 	NIAGARA_API virtual bool IsReadyForFinishDestroy() override;
-	virtual void PreSave(FObjectPreSaveContext ObjectSaveContext) override;
+	NIAGARA_API virtual void PreSave(FObjectPreSaveContext ObjectSaveContext) override;
 #if WITH_EDITOR
 	NIAGARA_API virtual void PostRename(UObject* OldOuter, const FName OldName) override;
 	NIAGARA_API virtual void PreEditChange(FProperty* PropertyThatWillChange)override;
