- Implement uniform / non uniform when sampling a stateless distribution
- This worked with range's and not full distributions before this change

#tests Modify initial position distribution to use the various options, before non-uniform range was uniform always

diff --git a/Engine/Plugins/FX/Niagara/Shaders/Private/Stateless/NiagaraStatelessCommon.ush b/Engine/Plugins/FX/Niagara/Shaders/Private/Stateless/NiagaraStatelessCommon.ush
--- a/Engine/Plugins/FX/Niagara/Shaders/Private/Stateless/NiagaraStatelessCommon.ush
+++ b/Engine/Plugins/FX/Niagara/Shaders/Private/Stateless/NiagaraStatelessCommon.ush
@@ -176,6 +176,11 @@ float3 LerpStaticFloat3(uint2 Parameters, float U)
 	const float3 Value1 = GetStaticFloat3(Parameters.x, ceil(Offset));
 	return lerp(Value0, Value1, frac(Offset));
 }
+
+float  StatelessLerp(float  lhs, float  rhs, float  u, bool bUniform) { return lerp(lhs, rhs, u); }
+float2 StatelessLerp(float2 lhs, float2 rhs, float2 u, bool bUniform) { return bUniform ? lerp(lhs.x, rhs.x, u.x).xx   : lerp(lhs, rhs, u); }
+float3 StatelessLerp(float3 lhs, float3 rhs, float3 u, bool bUniform) { return bUniform ? lerp(lhs.x, rhs.x, u.x).xxx  : lerp(lhs, rhs, u); }
+float4 StatelessLerp(float4 lhs, float4 rhs, float4 u, bool bUniform) { return bUniform ? lerp(lhs.x, rhs.x, u.x).xxxx : lerp(lhs, rhs, u); }
 	
 uint GetComponentOffset(uint OutputIndex, int ComponentOffset)
 {
diff --git a/Engine/Plugins/FX/Niagara/Shaders/Private/Stateless/NiagaraStatelessDistribution.ush b/Engine/Plugins/FX/Niagara/Shaders/Private/Stateless/NiagaraStatelessDistribution.ush
--- a/Engine/Plugins/FX/Niagara/Shaders/Private/Stateless/NiagaraStatelessDistribution.ush
+++ b/Engine/Plugins/FX/Niagara/Shaders/Private/Stateless/NiagaraStatelessDistribution.ush
@@ -14,7 +14,7 @@
 			RandomOffset = 0.0f; \
 			if ( (Parameters.x & StatelessDistributionFlag_Random) != 0) \
 			{ \
-				RandomOffset = RandomFloat(Particle); \
+				RandomOffset = Random##NAMETYPE##(Particle); \
 			} \
 		} \
 		\
@@ -30,7 +30,7 @@
 			{ \
 				const HLSLTYPE Value0 = GetStatic##NAMETYPE##(Parameters.y, 0); \
 				const HLSLTYPE Value1 = GetStatic##NAMETYPE##(Parameters.y, 1); \
-				return lerp(Value0, Value1, RandomOffset); \
+				return StatelessLerp(Value0, Value1, RandomOffset, (Parameters.x & StatelessDistributionFlag_Uniform) != 0); \
 			} \
 			else \
 			{ \
@@ -41,8 +41,8 @@
 			} \
 		} \
 		\
-		uint3	Parameters; \
-		float	RandomOffset; \
+		uint3		Parameters; \
+		HLSLTYPE	RandomOffset; \
 	}
 
 DefineStatelessDistributionSampler(float,  Float);
diff --git a/Engine/Plugins/FX/Niagara/Source/Niagara/Internal/Stateless/NiagaraStatelessEmitterDataBuildContext.h b/Engine/Plugins/FX/Niagara/Source/Niagara/Internal/Stateless/NiagaraStatelessEmitterDataBuildContext.h
--- a/Engine/Plugins/FX/Niagara/Source/Niagara/Internal/Stateless/NiagaraStatelessEmitterDataBuildContext.h
+++ b/Engine/Plugins/FX/Niagara/Source/Niagara/Internal/Stateless/NiagaraStatelessEmitterDataBuildContext.h
@@ -17,10 +17,6 @@ namespace NiagaraStateless
 
 class FNiagaraStatelessEmitterDataBuildContext
 {
-	static constexpr uint32 StatelessDistributionFlag_Random	= 0x00000001;
-	static constexpr uint32 StatelessDistributionFlag_Uniform	= 0x00000002;
-	static constexpr uint32 StatelessDistributionFlag_Binding	= 0x00000004;
-
 public:
 	UE_NONCOPYABLE(FNiagaraStatelessEmitterDataBuildContext);
 
@@ -72,12 +68,12 @@ public:
 
 	// Adds an distribution into the LUT if enabled
 	template<typename TType>
-	FUintVector3 AddDistribution(ENiagaraDistributionMode Mode, TConstArrayView<TType> Values, bool bEnabled) const
+	FUintVector3 AddDistribution(ENiagaraDistributionMode Mode, TConstArrayView<TType> Values) const
 	{
 		using namespace NiagaraStateless;
 
 		FUintVector3 Parameters = FUintVector3::ZeroValue;
-		if (bEnabled && Values.Num() > 0)
+		if (Values.Num() > 0)
 		{
 			switch (Mode)
 			{
@@ -99,26 +95,23 @@ public:
 
 	// Adds a distribution into the LUT if enabled and returns the packed information to send to the shader
 	template<typename TDistribution>
-	FUintVector3 AddDistribution(const TDistribution& Distribution, bool bEnabled = true) const
+	FUintVector3 AddDistribution(const TDistribution& Distribution) const
 	{
 		FUintVector3 Parameters = FUintVector3::ZeroValue;
-		if ( bEnabled )
+		if (Distribution.Mode == ENiagaraDistributionMode::Binding)
 		{
-			if (Distribution.Mode == ENiagaraDistributionMode::Binding)
-			{
-				const int32 ParameterOffset = AddRendererBinding(Distribution.ParameterBinding);
-				if (ParameterOffset >= 0)
-				{
-					Parameters.X = StatelessDistributionFlag_Binding;
-					Parameters.Y = ParameterOffset;
-					Parameters.Z = 1.0f;
-				}
-			}
-			else
+			const int32 ParameterOffset = AddRendererBinding(Distribution.ParameterBinding);
+			if (ParameterOffset >= 0)
 			{
-				Parameters = AddDistribution(Distribution.Mode, MakeArrayView(Distribution.Values), bEnabled);
+				Parameters.X = uint32(ENiagaraStatelessBuiltDistributionFlag::Binding);
+				Parameters.Y = ParameterOffset;
+				Parameters.Z = 1.0f;
 			}
 		}
+		else
+		{
+			Parameters = AddDistribution(Distribution.Mode, MakeArrayView(Distribution.Values));
+		}
 		return Parameters;
 	}
 
diff --git a/Engine/Plugins/FX/Niagara/Source/Niagara/Internal/Stateless/NiagaraStatelessParticleSimContext.h b/Engine/Plugins/FX/Niagara/Source/Niagara/Internal/Stateless/NiagaraStatelessParticleSimContext.h
--- a/Engine/Plugins/FX/Niagara/Source/Niagara/Internal/Stateless/NiagaraStatelessParticleSimContext.h
+++ b/Engine/Plugins/FX/Niagara/Source/Niagara/Internal/Stateless/NiagaraStatelessParticleSimContext.h
@@ -114,6 +114,19 @@ public:
 		return FMath::Lerp(Value0, Value1, FMath::Fractional(Offset));
 	}
 
+	float			Lerp(const float& Lhs, const float& Rhs, const float& U, bool bUniform) const { return FMath::Lerp(Lhs, Rhs, U); }
+	FVector2f		Lerp(const FVector2f& Lhs, const FVector2f& Rhs, const FVector2f& U, bool bUniform) const { const FVector2f V = FMath::Lerp(Lhs, Rhs, U); return bUniform ? FVector2f(V.X, V.X) : V; }
+	FVector3f		Lerp(const FVector3f& Lhs, const FVector3f& Rhs, const FVector3f& U, bool bUniform) const { const FVector3f V = FMath::Lerp(Lhs, Rhs, U); return bUniform ? FVector3f(V.X, V.X, V.X) : V; }
+	FVector4f		Lerp(const FVector4f& Lhs, const FVector4f& Rhs, const FVector4f& U, bool bUniform) const { const FVector4f V = FMath::Lerp(Lhs, Rhs, U); return bUniform ? FVector4f(V.X, V.X, V.X, V.X) : V; }
+
+	template<typename TType>
+	TType					TRandomFloat(uint32 iInstance) const;
+	template<> float		TRandomFloat<float>(uint32 iInstance) const { return RandomFloat(iInstance); }
+	template<> FVector2f	TRandomFloat<FVector2f>(uint32 iInstance) const { return RandomFloat2(iInstance); }
+	template<> FVector3f	TRandomFloat<FVector3f>(uint32 iInstance) const { return RandomFloat3(iInstance); }
+	template<> FVector4f	TRandomFloat<FVector4f>(uint32 iInstance) const { return RandomFloat4(iInstance); }
+	template<> FLinearColor	TRandomFloat<FLinearColor>(uint32 iInstance) const { return FLinearColor(RandomFloat4(iInstance)); }
+
 	uint32			RandomUInt(uint32 iInstance) const;
 	FUintVector2	RandomUInt2(uint32 iInstance) const;
 	FUintVector3	RandomUInt3(uint32 iInstance) const;
@@ -213,7 +226,7 @@ struct FStatelessDistributionSampler
 	{
 		if ((Parameters.X & uint32(ENiagaraStatelessBuiltDistributionFlag::Random)) != 0)
 		{
-			RandomOffset = ParticleSimulationContext.RandomFloat(iInstance);
+			RandomOffset = ParticleSimulationContext.TRandomFloat<TType>(iInstance);
 		}
 	}
 
@@ -230,7 +243,7 @@ struct FStatelessDistributionSampler
 			//-OPT: Could move into constructor
 			const TType	Value0 = ParticleSimulationContext.GetStaticFloat<TType>(Parameters.Y, 0);
 			const TType	Value1 = ParticleSimulationContext.GetStaticFloat<TType>(Parameters.Y, 1);
-			return FMath::Lerp(Value0, Value1, RandomOffset);
+			return ParticleSimulationContext.Lerp(Value0, Value1, RandomOffset, (Parameters.X & uint32(ENiagaraStatelessBuiltDistributionFlag::Uniform)) != 0);
 		}
 		else
 		{
@@ -242,7 +255,7 @@ struct FStatelessDistributionSampler
 	}
 
 	const FUintVector3	Parameters;
-	float				RandomOffset = 0.0f;
+	TType				RandomOffset = {};
 };
 
 } //namespace NiagaraStateless
