- Remove usage of DFFast* as it is causing NaNs in the VS

diff --git a/Engine/Plugins/FX/Niagara/Shaders/Private/NiagaraMeshParticleUtils.ush b/Engine/Plugins/FX/Niagara/Shaders/Private/NiagaraMeshParticleUtils.ush
--- a/Engine/Plugins/FX/Niagara/Shaders/Private/NiagaraMeshParticleUtils.ush
+++ b/Engine/Plugins/FX/Niagara/Shaders/Private/NiagaraMeshParticleUtils.ush
@@ -329,8 +329,7 @@ float3 NiagaraCalculateCameraOffset(FLWCVector3 ParticleSimPosition, NiagaraMesh
 	const FLWCVector3 WorldPosition = NiagaraSimToWorldPos(ParticleSimPosition, Params);
 	const float Offset = NiagaraSafeGetFloat(Params.CameraOffsetDataOffset, Params.ParticleIndex, Params.DefaultCameraOffset);
 
-	//-Note: LWCNormalizeTile should not be required but a precision issue is causing PLAY-32984, this is a temporary workaround
-	const FLWCVector3 CameraVec = LWCNormalizeTile(LWCSubtract(Params.CameraOrigin, WorldPosition));
+	const FLWCVector3 CameraVec = LWCSubtract(Params.CameraOrigin, WorldPosition);
 	const float3 CameraDir = NiagaraWorldToSimVec(LWCNormalize(CameraVec), Params);
 		
 	return CameraDir * Offset;
diff --git a/Engine/Plugins/FX/Niagara/Shaders/Private/NiagaraMeshVertexFactory.ush b/Engine/Plugins/FX/Niagara/Shaders/Private/NiagaraMeshVertexFactory.ush
--- a/Engine/Plugins/FX/Niagara/Shaders/Private/NiagaraMeshVertexFactory.ush
+++ b/Engine/Plugins/FX/Niagara/Shaders/Private/NiagaraMeshVertexFactory.ush
@@ -536,9 +536,9 @@ FVertexFactoryIntermediates GetVertexFactoryIntermediates(FVertexFactoryInput In
 
 	// Compute transforms
 #if (VF_USE_PRIMITIVE_SCENE_DATA == 1)
-	Intermediates.ParticleToWorld = DFFastToTileOffset(Intermediates.SceneData.InstanceData.LocalToWorld);
-	Intermediates.WorldToParticle = DFFastToTileOffset(Intermediates.SceneData.InstanceData.WorldToLocal);
-	Intermediates.PrevParticleToWorld = DFFastToTileOffset(Intermediates.SceneData.InstanceData.PrevLocalToWorld);
+	Intermediates.ParticleToWorld = DFToTileOffset(Intermediates.SceneData.InstanceData.LocalToWorld);
+	Intermediates.WorldToParticle = DFToTileOffset(Intermediates.SceneData.InstanceData.WorldToLocal);
+	Intermediates.PrevParticleToWorld = DFToTileOffset(Intermediates.SceneData.InstanceData.PrevLocalToWorld);
 
 	const float3 InstInvScale = Intermediates.SceneData.InstanceData.InvNonUniformScale;
 	Intermediates.ParticleToWorldNoScale = LWCToFloat3x3(Intermediates.ParticleToWorld);
@@ -584,9 +584,9 @@ FVertexFactoryIntermediates GetVertexFactoryIntermediates(FVertexFactoryInput In
 	Params.PrevCameraOrigin 			= ResolvedView.TileOffset.PrevWorldCameraOrigin;
 	Params.PrevCameraForwardDir			= ResolvedView.PrevViewToTranslatedWorld[2].xyz;
 	Params.PrevCameraUpDir				= ResolvedView.PrevViewToTranslatedWorld[1].xyz;
-	Params.PrimitiveLocalToWorld		= DFFastToTileOffset(Intermediates.SceneData.Primitive.LocalToWorld);
-	Params.PrimitiveWorldToLocal 		= DFFastToTileOffset(Intermediates.SceneData.Primitive.WorldToLocal);
-	Params.PrimitivePrevLocalToWorld 	= DFFastToTileOffset(Intermediates.SceneData.Primitive.PreviousLocalToWorld);
+	Params.PrimitiveLocalToWorld		= DFToTileOffset(Intermediates.SceneData.Primitive.LocalToWorld);
+	Params.PrimitiveWorldToLocal 		= DFToTileOffset(Intermediates.SceneData.Primitive.WorldToLocal);
+	Params.PrimitivePrevLocalToWorld 	= DFToTileOffset(Intermediates.SceneData.Primitive.PreviousLocalToWorld);
 	Params.PrimitiveInvNonUniformScale	= Intermediates.SceneData.Primitive.InvNonUniformScale;
 
 	NiagaraMeshParticleTransforms Transforms = NiagaraCalculateMeshParticleTransforms(Params);
