Fixed niagara asset browser ensure around having just 1 tag in an asset tag asset

diff --git a/Engine/Plugins/FX/Niagara/Source/NiagaraEditor/Private/Widgets/AssetBrowser/SNiagaraAssetBrowser.cpp b/Engine/Plugins/FX/Niagara/Source/NiagaraEditor/Private/Widgets/AssetBrowser/SNiagaraAssetBrowser.cpp
--- a/Engine/Plugins/FX/Niagara/Source/NiagaraEditor/Private/Widgets/AssetBrowser/SNiagaraAssetBrowser.cpp
+++ b/Engine/Plugins/FX/Niagara/Source/NiagaraEditor/Private/Widgets/AssetBrowser/SNiagaraAssetBrowser.cpp
@@ -414,27 +414,24 @@ TArray<TSharedRef<FNiagaraAssetBrowserMainFilter>> SNiagaraAssetBrowser::GetMain
 
 		for(const UNiagaraAssetTagDefinitions* AssetTagDefinitionsAsset : DisplayedAssetTagDefinitionAssets)
 		{
-			// This code should only execute for assets with > 1 tag. If there is only 1 tag, it should have been automatically added to the flat list instead
-			if(ensure(AssetTagDefinitionsAsset->GetAssetTagDefinitions().Num() > 1))
-			{
-				TSharedRef<FNiagaraAssetBrowserMainFilter> AssetTagDefinitionsAssetsFilter = MakeShared<FNiagaraAssetBrowserMainFilter>(FNiagaraAssetBrowserMainFilter::EFilterMode::NiagaraAssetTagDefinitionsAsset);
-				AssetTagDefinitionsAssetsFilter->AssetTagDefinitionsAsset = AssetTagDefinitionsAsset;
-				TagFilters.Add(AssetTagDefinitionsAssetsFilter);
+			// All assets here have been confirmed to have at least one valid asset tag so we can safely add them to the filters
+			TSharedRef<FNiagaraAssetBrowserMainFilter> AssetTagDefinitionsAssetsFilter = MakeShared<FNiagaraAssetBrowserMainFilter>(FNiagaraAssetBrowserMainFilter::EFilterMode::NiagaraAssetTagDefinitionsAsset);
+			AssetTagDefinitionsAssetsFilter->AssetTagDefinitionsAsset = AssetTagDefinitionsAsset;
+			TagFilters.Add(AssetTagDefinitionsAssetsFilter);
 
-				TArray<TSharedRef<FNiagaraAssetBrowserMainFilter>> TagChildFilters;
-				for(const FNiagaraAssetTagDefinition& AssetTagDefinition : AssetTagDefinitionsAsset->GetAssetTagDefinitions())
+			TArray<TSharedRef<FNiagaraAssetBrowserMainFilter>> TagChildFilters;
+			for(const FNiagaraAssetTagDefinition& AssetTagDefinition : AssetTagDefinitionsAsset->GetAssetTagDefinitions())
+			{
+				// Filter out tags that are invalid; one asset can contain valid and invalid tags for a given use-case
+				if(IsAssetTagDefinitionValid(AssetTagDefinition))
 				{
-					// Filter out tags that are invalid; one asset can contain valid and invalid tags for a given use-case
-					if(IsAssetTagDefinitionValid(AssetTagDefinition))
-					{
-						TSharedRef<FNiagaraAssetBrowserMainFilter> AssetTagFilter = MakeShared<FNiagaraAssetBrowserMainFilter>(FNiagaraAssetBrowserMainFilter::EFilterMode::NiagaraAssetTag);
-						AssetTagFilter->AssetTagDefinition = AssetTagDefinition;
-						TagChildFilters.Add(AssetTagFilter);
-					}
+					TSharedRef<FNiagaraAssetBrowserMainFilter> AssetTagFilter = MakeShared<FNiagaraAssetBrowserMainFilter>(FNiagaraAssetBrowserMainFilter::EFilterMode::NiagaraAssetTag);
+					AssetTagFilter->AssetTagDefinition = AssetTagDefinition;
+					TagChildFilters.Add(AssetTagFilter);
 				}
+			}
 
-				AssetTagDefinitionsAssetsFilter->ChildFilters = TagChildFilters;
-			}			
+			AssetTagDefinitionsAssetsFilter->ChildFilters = TagChildFilters;
 		}
 	}
 	
