- Fix crashing from reimporting a static mesh while in use on a Niagara mesh renderer
- We have to add a handler to OnPreMeshBuild to avoid reading data that will no longer be valid, this is the result of changes to the SM build process
- Handle the case where we are inside EOF updates and the static mesh is still building, for the moment we assume that we will get a callback from the SM system when the build is complete and we can restart the system.  This is a temporary solution and should be replaced with something that will suppress rendering instead.

#tests Reimport a SM various times while in / out of PIE

diff --git a/Engine/Plugins/FX/Niagara/Source/Niagara/Private/NiagaraMeshRendererProperties.cpp b/Engine/Plugins/FX/Niagara/Source/Niagara/Private/NiagaraMeshRendererProperties.cpp
--- a/Engine/Plugins/FX/Niagara/Source/Niagara/Private/NiagaraMeshRendererProperties.cpp
+++ b/Engine/Plugins/FX/Niagara/Source/Niagara/Private/NiagaraMeshRendererProperties.cpp
@@ -12,6 +12,7 @@
 #include "NiagaraSystemImpl.h"
 
 #include "Engine/StaticMesh.h"
+#include "Engine/World.h"
 #include "Materials/MaterialRenderProxy.h"
 #include "Materials/MaterialInstanceConstant.h"
 #include "Modules/ModuleManager.h"
@@ -200,9 +201,25 @@ namespace NiagaraMeshRendererPropertiesInternal
 				}
 			}
 		}
-		if (MeshProperties.Mesh && MeshProperties.Mesh->GetRenderData())
+		if (MeshProperties.Mesh)
 		{
-			OutStaticMesh = MeshProperties.Mesh;
+			bool bInvalidForAsynCompiling = false;
+			#if WITH_EDITOR
+				// During EOF updates we can not test GetRenderData as we will cause a wait / reregister to occur which is invalid while in EOF updates
+				// The assumption is that when we do hit this situation we are waiting on a static mesh build from reimport, etc, so we can just skip until
+				// we get the post build callback, ideally we would have a better way to handle this by not rendering while the build is in progress but
+				// that requires quite a large rework of how we handle reading from the static mesh data, especially around updating MICs in PostLoad.
+				if (MeshProperties.Mesh->IsCompiling())
+				{
+					FNiagaraSystemInstance* SystemInstance = EmitterInstance ? EmitterInstance->GetParentSystemInstance() : nullptr;
+					UWorld* World = SystemInstance ? SystemInstance->GetWorld() : nullptr;
+					bInvalidForAsynCompiling = World && World->bPostTickComponentUpdate;
+				}
+			#endif
+			if (!bInvalidForAsynCompiling && MeshProperties.Mesh->GetRenderData())
+			{
+				OutStaticMesh = MeshProperties.Mesh;
+			}
 		}
 	}
 }
@@ -915,6 +932,7 @@ void UNiagaraMeshRendererProperties::PostLoad()
 			if (GIsEditor)
 			{
 				MeshProperties.Mesh->GetOnMeshChanged().AddUObject(this, &UNiagaraMeshRendererProperties::OnMeshChanged);
+				MeshProperties.Mesh->OnPreMeshBuild().AddUObject(this, &UNiagaraMeshRendererProperties::OnMeshPostBuild);
 				MeshProperties.Mesh->OnPostMeshBuild().AddUObject(this, &UNiagaraMeshRendererProperties::OnMeshPostBuild);
 			}
 #endif
@@ -1097,6 +1115,7 @@ void UNiagaraMeshRendererProperties::BeginDestroy()
 			if (MeshProperties.Mesh)
 			{
 				MeshProperties.Mesh->GetOnMeshChanged().RemoveAll(this);
+				MeshProperties.Mesh->OnPreMeshBuild().RemoveAll(this);
 				MeshProperties.Mesh->OnPostMeshBuild().RemoveAll(this);
 			}
 		}
@@ -1115,6 +1134,7 @@ void UNiagaraMeshRendererProperties::PreEditChange(class FProperty* PropertyThat
 			if (MeshProperties.Mesh)
 			{
 				MeshProperties.Mesh->GetOnMeshChanged().RemoveAll(this);
+				MeshProperties.Mesh->OnPreMeshBuild().RemoveAll(this);
 				MeshProperties.Mesh->OnPostMeshBuild().RemoveAll(this);
 			}
 		}
@@ -1139,6 +1159,7 @@ void UNiagaraMeshRendererProperties::PostEditChangeProperty(FPropertyChangedEven
 			if (MeshProperties.Mesh)
 			{
 				MeshProperties.Mesh->GetOnMeshChanged().RemoveAll(this);
+				MeshProperties.Mesh->OnPreMeshBuild().RemoveAll(this);
 				MeshProperties.Mesh->OnPostMeshBuild().RemoveAll(this);
 			}
 		}
@@ -1185,6 +1206,7 @@ void UNiagaraMeshRendererProperties::PostEditChangeProperty(FPropertyChangedEven
 			if (MeshProperties.Mesh)
 			{
 				MeshProperties.Mesh->GetOnMeshChanged().AddUObject(this, &UNiagaraMeshRendererProperties::OnMeshChanged);
+				MeshProperties.Mesh->OnPreMeshBuild().AddUObject(this, &UNiagaraMeshRendererProperties::OnMeshPostBuild);
 				MeshProperties.Mesh->OnPostMeshBuild().AddUObject(this, &UNiagaraMeshRendererProperties::OnMeshPostBuild);
 			}
 		}
