- Add optional GT Only column to perf display
- Handle system name length better to avoid overlapping the display

diff --git a/Engine/Plugins/FX/Niagara/Source/Niagara/Private/NiagaraDebugHud.cpp b/Engine/Plugins/FX/Niagara/Source/Niagara/Private/NiagaraDebugHud.cpp
--- a/Engine/Plugins/FX/Niagara/Source/Niagara/Private/NiagaraDebugHud.cpp
+++ b/Engine/Plugins/FX/Niagara/Source/Niagara/Private/NiagaraDebugHud.cpp
@@ -831,6 +831,8 @@ FNiagaraDebugHud::FNiagaraDebugHud(UWorld* World)
 
 	WeakWorld = World;
 
+	LongestSystemName = TEXT("NS_SomeBigLongNiagaraSystemName");
+
 	if ( !GDebugDrawHandle.IsValid() )
 	{
 		GDebugDrawHandle = UDebugDrawService::Register(TEXT("Particles"), FDebugDrawDelegate::CreateStatic(&FNiagaraDebugHud::DebugDrawCallback));
@@ -1003,6 +1005,7 @@ void FNiagaraDebugHud::GatherSystemInfo()
 	}
 
 #if WITH_PARTICLE_PERF_STATS
+	//-TODO: Enable based one perf warnings
 	bool bUpdateStats = false;
 	if (Settings.bOverviewEnabled && (Settings.OverviewMode == ENiagaraDebugHUDOverviewMode::Performance || Settings.OverviewMode == ENiagaraDebugHUDOverviewMode::PerformanceGraph))
 	{
@@ -1070,6 +1073,13 @@ void FNiagaraDebugHud::GatherSystemInfo()
 		if (SystemDebugInfo.SystemName.IsEmpty())
 		{
 			SystemDebugInfo.SystemName = GetNameSafe(FXComponent->GetFXSystemAsset());
+
+			static const FString SystemNamePostFix = TEXT(" (Fast Path) ");
+			if (SystemDebugInfo.SystemName.Len() + SystemNamePostFix.Len() > LongestSystemName.Len())
+			{
+				LongestSystemName = SystemDebugInfo.SystemName;
+				LongestSystemName.Append(SystemNamePostFix);
+			}
 		}
 	#if WITH_EDITORONLY_DATA
 		SystemDebugInfo.bCompileForEdit = NiagaraComponent ? NiagaraComponent->GetAsset()->GetCompileForEdit() : false;
@@ -1659,7 +1669,7 @@ void FNiagaraDebugHud::DrawOverview(class FNiagaraWorldManager* WorldManager, FC
 
 	if (Settings.bOverviewEnabled)
 	{
-		OverviewColumns.Emplace(TEXT(""), TEXT(""), TEXT("System Name"), ColumnOffset, Font, TEXT("NS_SomeBigLongNiagaraSystemName"),
+		OverviewColumns.Emplace(TEXT(""), TEXT(""), TEXT("System Name"), ColumnOffset, Font, *LongestSystemName,
 			[&DetailColor, &DetailHighlightColor, &fAdvanceHeight](FCanvas* Canvas, UFont* Font, float X, float Y, FOverviewColumn& Col, const FSystemDebugInfo& SystemInfo)
 			{
 				FLinearColor RowBGColor = SystemInfo.UniqueColor;
@@ -1847,6 +1857,19 @@ void FNiagaraDebugHud::DrawOverview(class FNiagaraWorldManager* WorldManager, FC
 					Canvas->DrawShadowedString(X, Y, *FormatPerfValue(SystemInfo.PerfStats ? SystemInfo.PerfStats->Avg.Time_GT : 0.0), Font, RowColor);
 				});
 
+			if (Settings.bShowPerfColumGameThreadOnly)
+			{
+				OverviewColumns.Emplace(TEXT("Game Thread Only Avg:"), FormatPerfValue(GlobalPerfStats.Avg.Time_GTOnly), GlobalDataStringSize, FormatPerfString(TEXT("GT Only Avg")), ColumnOffset, Font, SystemStringSize,
+					[&DetailColor, &DetailHighlightColor, &fAdvanceHeight](FCanvas* Canvas, UFont* Font, float X, float Y, FOverviewColumn& Col, const FSystemDebugInfo& SystemInfo)
+					{
+						FLinearColor RowBGColor = SystemInfo.UniqueColor;
+						RowBGColor.A = Settings.SystemColorTableOpacity;
+						Canvas->DrawTile(X, Y, Col.MaxWidth, fAdvanceHeight, 0, 0, 0, 0, RowBGColor);
+						const FLinearColor RowColor = SystemInfo.bShowInWorld ? DetailHighlightColor : DetailColor;
+						Canvas->DrawShadowedString(X, Y, *FormatPerfValue(SystemInfo.PerfStats ? SystemInfo.PerfStats->Avg.Time_GTOnly : 0.0), Font, RowColor);
+					});
+			}
+
 			OverviewColumns.Emplace(TEXT("Game Thread Max:"), FormatPerfValue(GlobalPerfStats.Max.Time_GT), GlobalDataStringSize, FormatPerfString(TEXT("GT Max")), ColumnOffset, Font, SystemStringSize,
 				[&DetailColor, &DetailHighlightColor, &fAdvanceHeight](FCanvas* Canvas, UFont* Font, float X, float Y, FOverviewColumn& Col, const FSystemDebugInfo& SystemInfo)
 				{
@@ -3462,16 +3485,19 @@ bool FNiagaraDebugHUDStatsListener::Tick()
 
 		if (FAccumulatedParticlePerfStats* Stats = GetStats(System))
 		{
-			double SysAvg;
-			double SysMax;
+			double SysAvg = 0.0;
+			double SysAvgGTOnly = 0.0;
+			double SysMax = 0.0;
 			if (Settings.PerfSampleMode == ENiagaraDebugHUDPerfSampleMode::FrameTotal)
 			{
 				SysAvg = Stats->GetGameThreadStats().GetPerFrameAvg();
+				SysAvgGTOnly = Stats->GetGameThreadStats().GetPerFrameAvg_GTOnly();
 				SysMax = Stats->GetGameThreadStats().GetPerFrameMax();
 			}
 			else
 			{
 				SysAvg = Stats->GetGameThreadStats().GetPerInstanceAvg();
+				SysAvgGTOnly = Stats->GetGameThreadStats().GetPerInstanceAvg_GTOnly();
 				SysMax = Stats->GetGameThreadStats().GetPerInstanceMax();
 			}
 
@@ -3480,6 +3506,7 @@ bool FNiagaraDebugHUDStatsListener::Tick()
 			if (bPushStats)
 			{
 				HUDStats->Avg.Time_GT = SysAvg;
+				HUDStats->Avg.Time_GTOnly = SysAvgGTOnly;
 				HUDStats->Max.Time_GT = SysMax;
 			}
 
diff --git a/Engine/Plugins/FX/Niagara/Source/Niagara/Private/NiagaraDebugHud.h b/Engine/Plugins/FX/Niagara/Source/Niagara/Private/NiagaraDebugHud.h
--- a/Engine/Plugins/FX/Niagara/Source/Niagara/Private/NiagaraDebugHud.h
+++ b/Engine/Plugins/FX/Niagara/Source/Niagara/Private/NiagaraDebugHud.h
@@ -23,6 +23,7 @@ class FNiagaraDebugHud;
 struct FNiagaraDebugHudFrameStat
 {
 	double Time_GT = 0.0;
+	double Time_GTOnly = 0.0;
 	double Time_RT = 0.0;
 	double Time_GPU = 0.0;
 };
@@ -266,6 +267,7 @@ private:
 
 	int32 GlobalTotalPlayerSystems = 0;
 
+	FString LongestSystemName;
 	TMap<FName, FSystemDebugInfo>	PerSystemDebugInfo;
 
 	TArray<TWeakObjectPtr<class UFXSystemComponent>>	InWorldComponents;
diff --git a/Engine/Plugins/FX/Niagara/Source/Niagara/Public/NiagaraDebuggerCommon.h b/Engine/Plugins/FX/Niagara/Source/Niagara/Public/NiagaraDebuggerCommon.h
--- a/Engine/Plugins/FX/Niagara/Source/Niagara/Public/NiagaraDebuggerCommon.h
+++ b/Engine/Plugins/FX/Niagara/Source/Niagara/Public/NiagaraDebuggerCommon.h
@@ -679,6 +679,9 @@ struct FNiagaraDebugHUDSettingsData
 	UPROPERTY(EditAnywhere, Category = "Perf Overview")
 	ENiagaraDebugHUDPerfUnits PerfUnits = ENiagaraDebugHUDPerfUnits::Microseconds;
 
+	UPROPERTY(EditAnywhere, Category = "Perf Overview", meta = (EditConditionHides, EditCondition = "OverviewMode == ENiagaraDebugHUDOverviewMode::Performance"))
+	bool bShowPerfColumGameThreadOnly = false;
+
 	/** Time range of the Y Axis of the perf graph */
 	UPROPERTY(EditAnywhere, Category = "Perf Overview")
 	ENiagaraDebugHUDPerfGraphMode PerfGraphMode = ENiagaraDebugHUDPerfGraphMode::GameThread;
@@ -893,6 +896,7 @@ enum class ENiagaraDebugMessageType : uint8
 	Warning,
 	Error
 };
+
 struct FNiagaraDebugMessage
 {
 	ENiagaraDebugMessageType Type;
diff --git a/Engine/Source/Runtime/Engine/Classes/Particles/ParticlePerfStats.h b/Engine/Source/Runtime/Engine/Classes/Particles/ParticlePerfStats.h
--- a/Engine/Source/Runtime/Engine/Classes/Particles/ParticlePerfStats.h
+++ b/Engine/Source/Runtime/Engine/Classes/Particles/ParticlePerfStats.h
@@ -109,9 +109,12 @@ struct FParticlePerfStats_GT
 		ActivationCycles = 0;
 		WaitCycles = 0;
 	}
-	FORCEINLINE uint64 GetTotalCycles_GTOnly()const { return TickGameThreadCycles + FinalizeCycles + ActivationCycles + WaitCycles; }
-	FORCEINLINE uint64 GetTotalCycles()const { return GetTotalCycles_GTOnly() + TickConcurrentCycles + EndOfFrameCycles; }
-	FORCEINLINE uint64 GetPerInstanceAvgCycles()const { return NumInstances > 0 ? GetTotalCycles() / NumInstances : 0; }
+
+	FORCEINLINE uint64 GetTotalCycles_GTOnly() const { return TickGameThreadCycles + FinalizeCycles + ActivationCycles + WaitCycles; }
+	FORCEINLINE uint64 GetPerInstanceAvgCycles_GTOnly() const { return NumInstances > 0 ? GetTotalCycles_GTOnly() / NumInstances : 0; }
+
+	FORCEINLINE uint64 GetTotalCycles() const { return GetTotalCycles_GTOnly() + TickConcurrentCycles + EndOfFrameCycles; }
+	FORCEINLINE uint64 GetPerInstanceAvgCycles() const { return NumInstances > 0 ? GetTotalCycles() / NumInstances : 0; }
 };
 
 /** Stats gathered on the render thread. */
diff --git a/Engine/Source/Runtime/Engine/Classes/Particles/ParticlePerfStatsManager.h b/Engine/Source/Runtime/Engine/Classes/Particles/ParticlePerfStatsManager.h
--- a/Engine/Source/Runtime/Engine/Classes/Particles/ParticlePerfStatsManager.h
+++ b/Engine/Source/Runtime/Engine/Classes/Particles/ParticlePerfStatsManager.h
@@ -29,14 +29,17 @@ struct FAccumulatedParticlePerfStats_GT
 
 	/** Returns the total cycles used by all GameThread stats. */
 	FORCEINLINE uint64 GetTotalCycles() const { return AccumulatedStats.GetTotalCycles(); }
+	FORCEINLINE uint64 GetTotalCycles_GTOnly() const { return AccumulatedStats.GetTotalCycles_GTOnly(); }
 
 	/** Returns the average cycles per frame by all GameThread stats. */
 	FORCEINLINE uint64 GetPerFrameAvgCycles() const { return NumFrames > 0 ? AccumulatedStats.GetTotalCycles() / NumFrames : 0; }
+	FORCEINLINE uint64 GetPerFrameAvgCycles_GTOnly() const { return NumFrames > 0 ? AccumulatedStats.GetTotalCycles_GTOnly() / NumFrames : 0; }
 	/** Returns the max cycles per frame by all GameThread stats. */
 	FORCEINLINE uint64 GetPerFrameMaxCycles(int32 Index = 0) const { return MaxPerFrameTotalCycles[Index]; }
 
 	/** Returns the average time in µs per frame by all GameThread stats. */
 	FORCEINLINE float GetPerFrameAvg() const { return float(FPlatformTime::ToMilliseconds64(GetPerFrameAvgCycles()) * 1000.0); }
+	FORCEINLINE float GetPerFrameAvg_GTOnly() const { return float(FPlatformTime::ToMilliseconds64(GetPerFrameAvgCycles_GTOnly()) * 1000.0); }
 	/** Returns the max time in µs per frame by all GameThread stats. */
 	FORCEINLINE float GetPerFrameMax(int32 Index = 0) const { return float(FPlatformTime::ToMilliseconds64(GetPerFrameMaxCycles(Index)) * 1000.0); }
 
@@ -47,6 +50,7 @@ struct FAccumulatedParticlePerfStats_GT
 
 	/** Returns the average time in µs per instance by all GameThread stats. */
 	FORCEINLINE float GetPerInstanceAvg() const { return float(FPlatformTime::ToMilliseconds64(GetPerInstanceAvgCycles()) * 1000.0); }
+	FORCEINLINE float GetPerInstanceAvg_GTOnly() const { return float(FPlatformTime::ToMilliseconds64(AccumulatedStats.GetPerInstanceAvgCycles_GTOnly()) * 1000.0); }
 	/** Returns the max time in µs per instance by all GameThread stats. */
 	FORCEINLINE float GetPerInstanceMax(int32 Index = 0) const { return float(FPlatformTime::ToMilliseconds64(GetPerInstanceMaxCycles(Index)) * 1000.0); }
 };
