- When the debugger settings are updated, ensure the search customizations are also updated

#rnx

diff --git a/Engine/Plugins/FX/Niagara/Source/NiagaraEditor/Private/Customizations/NiagaraDebugHUDCustomization.cpp b/Engine/Plugins/FX/Niagara/Source/NiagaraEditor/Private/Customizations/NiagaraDebugHUDCustomization.cpp
--- a/Engine/Plugins/FX/Niagara/Source/NiagaraEditor/Private/Customizations/NiagaraDebugHUDCustomization.cpp
+++ b/Engine/Plugins/FX/Niagara/Source/NiagaraEditor/Private/Customizations/NiagaraDebugHUDCustomization.cpp
@@ -40,6 +40,7 @@ public:
 	SLATE_BEGIN_ARGS(SDebuggerSuggestionTextBox) {}
 		SLATE_ARGUMENT(TSharedPtr<IPropertyHandle>, PropertyHandle)
 		SLATE_ARGUMENT(UClass*, ObjectClass)
+		SLATE_ARGUMENT(TWeakObjectPtr<UNiagaraDebugHUDSettings>, WeakHUDSettings)
 	SLATE_END_ARGS()
 
 	void Construct(const FArguments& InArgs)
@@ -53,6 +54,12 @@ public:
 
 		PropertyHandle = InArgs._PropertyHandle;
 		ObjectClass = InArgs._ObjectClass;
+		WeakHUDSettings = InArgs._WeakHUDSettings;
+
+		if (UNiagaraDebugHUDSettings* HUDSettings = WeakHUDSettings.Get())
+		{
+			HUDSettings->OnChangedDelegate.AddSP(this, &SDebuggerSuggestionTextBox::OnHUDSettingsChanged);
+		}
 
 		FText CurrentValue;
 		PropertyHandle->GetValueAsFormattedText(CurrentValue);
@@ -101,6 +108,13 @@ public:
 		}
 	}
 
+	void OnHUDSettingsChanged()
+	{
+		FText CurrentValue;
+		PropertyHandle->GetValueAsFormattedText(CurrentValue);
+		SetText(CurrentValue);
+	}
+
 	void OnSimpleClientInfoChanged(const FNiagaraSimpleClientInfo& ClientInfo)
 	{
 		bWaitingUpdate = false;
@@ -121,10 +135,11 @@ public:
 		}
 	}
 
-	TSharedPtr<FNiagaraDebugger>	Debugger;
-	TSharedPtr<IPropertyHandle>		PropertyHandle;
-	UClass*							ObjectClass = nullptr;
-	bool							bWaitingUpdate = false;
+	TSharedPtr<FNiagaraDebugger>				Debugger;
+	TSharedPtr<IPropertyHandle>					PropertyHandle;
+	TWeakObjectPtr<UNiagaraDebugHUDSettings>	WeakHUDSettings;
+	UClass*										ObjectClass = nullptr;
+	bool										bWaitingUpdate = false;
 };
 
 } // namespace
@@ -265,6 +280,7 @@ void FNiagaraDebugHUDSettingsDetailsCustomization::MakeCustomAssetSearch(IDetail
 			SNew(NiagaraDebugHUDSettingsDetailsCustomizationInternal::SDebuggerSuggestionTextBox)
 			.PropertyHandle(PropertyHandle)
 			.ObjectClass(ObjRefClass)
+			.WeakHUDSettings(WeakSettings)
 		]
 	];
 }
