Small optimization to the ticking of NPC.  Cache the MPC instance, if one is referenced, rather than looking for it on the world each tick
#tests EngineTest and validated that NPC are still being picked up in the world

diff --git a/Engine/Plugins/FX/Niagara/Source/Niagara/Classes/NiagaraParameterCollection.h b/Engine/Plugins/FX/Niagara/Source/Niagara/Classes/NiagaraParameterCollection.h
--- a/Engine/Plugins/FX/Niagara/Source/Niagara/Classes/NiagaraParameterCollection.h
+++ b/Engine/Plugins/FX/Niagara/Source/Niagara/Classes/NiagaraParameterCollection.h
@@ -8,6 +8,7 @@
 #include "NiagaraParameterCollection.generated.h"
 
 class UMaterialParameterCollection;
+class UMaterialParameterCollectionInstance;
 class UNiagaraParameterCollection;
 
 UCLASS(MinimalAPI)
@@ -62,6 +63,9 @@ private:
 	UPROPERTY()
 	FNiagaraParameterStore ParameterStorage;
 
+	UPROPERTY(Transient)
+	TObjectPtr<UMaterialParameterCollectionInstance> SourceMaterialCollectionInstance;
+
 	FRWLock DirtyParameterLock;
 	TArray<TPair<FName, float>> DirtyScalarParameters;
 	TArray<TPair<FName, FLinearColor>> DirtyVectorParameters;
diff --git a/Engine/Plugins/FX/Niagara/Source/Niagara/Private/NiagaraParameterCollection.cpp b/Engine/Plugins/FX/Niagara/Source/Niagara/Private/NiagaraParameterCollection.cpp
--- a/Engine/Plugins/FX/Niagara/Source/Niagara/Private/NiagaraParameterCollection.cpp
+++ b/Engine/Plugins/FX/Niagara/Source/Niagara/Private/NiagaraParameterCollection.cpp
@@ -89,15 +89,17 @@ void UNiagaraParameterCollectionInstance::Bind(UWorld* World)
 {
 	if (const UMaterialParameterCollection* SourceCollection = Collection ? Collection->GetSourceCollection() : nullptr)
 	{
-		if (UMaterialParameterCollectionInstance* SourceInstance = World->GetParameterCollectionInstance(SourceCollection))
+		SourceMaterialCollectionInstance = World->GetParameterCollectionInstance(SourceCollection);
+
+		if (SourceMaterialCollectionInstance)
 		{
-			SourceInstance->OnScalarParameterUpdated().AddLambda([this](UMaterialParameterCollectionInstance::ScalarParameterUpdate DirtyParameter)
+			SourceMaterialCollectionInstance->OnScalarParameterUpdated().AddLambda([this](UMaterialParameterCollectionInstance::ScalarParameterUpdate DirtyParameter)
 			{
 				FRWScopeLock WriteLock(DirtyParameterLock, SLT_Write);
 				DirtyScalarParameters.Emplace(DirtyParameter);
 			});
 
-			SourceInstance->OnVectorParameterUpdated().AddLambda([this](UMaterialParameterCollectionInstance::VectorParameterUpdate DirtyParameter)
+			SourceMaterialCollectionInstance->OnVectorParameterUpdated().AddLambda([this](UMaterialParameterCollectionInstance::VectorParameterUpdate DirtyParameter)
 			{
 				FRWScopeLock WriteLock(DirtyParameterLock, SLT_Write);
 				DirtyVectorParameters.Emplace(DirtyParameter);
@@ -110,7 +112,7 @@ void UNiagaraParameterCollectionInstance::Bind(UWorld* World)
 			for (int32 ScalarIt = 0; ScalarIt < ScalarParameterCount; ++ScalarIt)
 			{
 				ScalarParameters[ScalarIt].Key = SourceCollection->ScalarParameters[ScalarIt].ParameterName;
-				SourceInstance->GetScalarParameterValue(SourceCollection->ScalarParameters[ScalarIt], ScalarParameters[ScalarIt].Value);
+				SourceMaterialCollectionInstance->GetScalarParameterValue(SourceCollection->ScalarParameters[ScalarIt], ScalarParameters[ScalarIt].Value);
 			}
 
 			TArray<TPair<FName, FLinearColor>> VectorParameters;
@@ -119,7 +121,7 @@ void UNiagaraParameterCollectionInstance::Bind(UWorld* World)
 			for (int32 VectorIt = 0; VectorIt < VectorParameterCount; ++VectorIt)
 			{
 				VectorParameters[VectorIt].Key = SourceCollection->VectorParameters[VectorIt].ParameterName;
-				SourceInstance->GetVectorParameterValue(SourceCollection->VectorParameters[VectorIt], VectorParameters[VectorIt].Value);
+				SourceMaterialCollectionInstance->GetVectorParameterValue(SourceCollection->VectorParameters[VectorIt], VectorParameters[VectorIt].Value);
 			}
 
 			RefreshSourceParameters(World, ScalarParameters, VectorParameters);
@@ -135,7 +137,13 @@ void UNiagaraParameterCollectionInstance::RefreshSourceParameters(
 	// if the NPC uses any MPC as sources, the make those bindings now
 	if (const UMaterialParameterCollection* SourceCollection = Collection ? Collection->GetSourceCollection() : nullptr)
 	{
-		if (UMaterialParameterCollectionInstance* SourceInstance = World->GetParameterCollectionInstance(SourceCollection))
+		// find the appropriate Instance
+		if (!SourceMaterialCollectionInstance || SourceMaterialCollectionInstance->GetCollection() != SourceCollection)
+		{
+			SourceMaterialCollectionInstance = World->GetParameterCollectionInstance(SourceCollection);
+		}
+
+		if (SourceMaterialCollectionInstance)
 		{
 			FNameBuilder VariableName;
 			VariableName << Collection->GetFullNamespaceName();
