Improved debug hud ux

#tests Tested the various overview modes and toolbar buttons, watching parameters from a system and filtering systems in the scene. Nothing too involved, as this change didn't add or remove functionality.

diff --git a/Engine/Plugins/FX/Niagara/Source/Niagara/Private/NiagaraDebugHud.cpp b/Engine/Plugins/FX/Niagara/Source/Niagara/Private/NiagaraDebugHud.cpp
--- a/Engine/Plugins/FX/Niagara/Source/Niagara/Private/NiagaraDebugHud.cpp
+++ b/Engine/Plugins/FX/Niagara/Source/Niagara/Private/NiagaraDebugHud.cpp
@@ -131,6 +131,7 @@ namespace NiagaraDebugLocal
 				#endif
 			}
 		),
+		MakeTuple(TEXT("ValidationEnabled="), TEXT("Enable or disable validation features in general"), [](FString Arg) {Settings.bValidationEnabled = FCString::Atoi(*Arg) != 0; }),
 		MakeTuple(TEXT("ValidateSystemSimulationDataBuffers="), TEXT("Enable or disable validation on system data buffers"), [](FString Arg) {Settings.bValidateSystemSimulationDataBuffers = FCString::Atoi(*Arg) != 0; }),
 		MakeTuple(TEXT("ValidateParticleDataBuffers="), TEXT("Enable or disable validation on particle data buffers"), [](FString Arg) {Settings.bValidateParticleDataBuffers = FCString::Atoi(*Arg) != 0; }),
 		MakeTuple(TEXT("ValidationLogErrors="), TEXT("When enabled validation errors will be logged as we as on screen"), [](FString Arg) {Settings.bValidationLogErrors = FCString::Atoi(*Arg) != 0; }),
@@ -157,7 +158,6 @@ namespace NiagaraDebugLocal
 		MakeTuple(TEXT("SystemEmitterVerbosity="), TEXT("Set the in world system emitter debug verbosity"), [](FString Arg) {Settings.SystemEmitterVerbosity = FMath::Clamp(ENiagaraDebugHudVerbosity(FCString::Atoi(*Arg)), ENiagaraDebugHudVerbosity::None, ENiagaraDebugHudVerbosity::Verbose); }),
 		MakeTuple(TEXT("DataInterfaceVerbosity="), TEXT("Set the in world system data interface debug verbosity"), [](FString Arg) {Settings.DataInterfaceVerbosity = FMath::Clamp(ENiagaraDebugHudVerbosity(FCString::Atoi(*Arg)), ENiagaraDebugHudVerbosity::None, ENiagaraDebugHudVerbosity::Verbose); }),
 		MakeTuple(TEXT("SystemVariables="), TEXT("Set the system variables to display"), [](FString Arg) {Settings.SystemVariables.Empty(); FNiagaraDebugHUDVariable::InitFromString(Arg, Settings.SystemVariables); GCachedSystemVariables.Empty(); }),
-		MakeTuple(TEXT("ShowSystemVariables="), TEXT("Set system variables visibility"), [](FString Arg) {Settings.bShowSystemVariables = FCString::Atoi(*Arg) != 0; GCachedSystemVariables.Empty(); }),
 
 		// Particle commands
 		MakeTuple(TEXT("EnableGpuParticleReadback="), TEXT("Enables GPU readback support for particle attributes"), [](FString Arg) {Settings.bEnableGpuParticleReadback = FCString::Atoi(*Arg) != 0;}),
@@ -489,7 +489,7 @@ namespace NiagaraDebugLocal
 			CachedVariables->CompiledDelegate = NiagaraSystem->OnSystemCompiled().AddLambda([](UNiagaraSystem* NiagaraSystem) { GCachedSystemVariables.Remove(NiagaraSystem); });
 #endif
 
-			if (Settings.bShowSystemVariables && Settings.SystemVariables.Num() > 0)
+			if (Settings.SystemVariables.Num() > 0)
 			{
 				FindSystemVariablesByWildcard(NiagaraSystem, Settings.SystemVariables, CachedVariables);
 				
@@ -520,7 +520,7 @@ namespace NiagaraDebugLocal
 						continue;
 					}
 
-					if (Settings.bEmitterFilterEnabled && !EmitterHandle.GetUniqueInstanceName().MatchesWildcard(Settings.EmitterFilter))
+					if (Settings.bEmitterFilterEnabled && !Settings.EmitterFilter.IsEmpty() && !EmitterHandle.GetUniqueInstanceName().MatchesWildcard(Settings.EmitterFilter))
 					{
 						continue;
 					}
@@ -998,7 +998,7 @@ void FNiagaraDebugHud::GatherSystemInfo()
 	// If the overview is not enabled and we don't have any filters we can skip everything below as nothing will be displayed
 	if (!Settings.bOverviewEnabled)
 	{
-		if ( !Settings.bActorFilterEnabled && !Settings.bComponentFilterEnabled && !Settings.bSystemFilterEnabled )
+		if ( !Settings.bActorFilterEnabled && !Settings.bComponentFilterEnabled )
 		{
 			return;
 		}
@@ -1085,8 +1085,8 @@ void FNiagaraDebugHud::GatherSystemInfo()
 		SystemDebugInfo.bCompileForEdit = NiagaraComponent ? NiagaraComponent->GetAsset()->GetCompileForEdit() : false;
 	#endif
 		SystemDebugInfo.bSystemStateFastPath = NiagaraComponent ? NiagaraComponent->GetAsset()->SystemStateFastPathEnabled() : false;
-		SystemDebugInfo.bShowInWorld = Settings.bSystemFilterEnabled && SystemDebugInfo.SystemName.MatchesWildcard(Settings.SystemFilter);
-		SystemDebugInfo.bPassesSystemFilter = !Settings.bSystemFilterEnabled || SystemDebugInfo.SystemName.MatchesWildcard(Settings.SystemFilter);
+		SystemDebugInfo.bShowInWorld = !Settings.bSystemFilterEnabled || Settings.SystemFilter.IsEmpty() || SystemDebugInfo.SystemName.MatchesWildcard(Settings.SystemFilter);
+		SystemDebugInfo.bPassesSystemFilter = SystemDebugInfo.bShowInWorld;
 
 		const bool bCanShowInWorld = 
 			SystemDebugInfo.bShowInWorld &&
@@ -1098,14 +1098,14 @@ void FNiagaraDebugHud::GatherSystemInfo()
 			bool bIsMatch = true;
 
 			// Filter by actor
-			if ( Settings.bActorFilterEnabled )
+			if ( Settings.bActorFilterEnabled && !Settings.ActorFilter.IsEmpty())
 			{
 				AActor* Actor = FXComponent->GetOwner();
 				bIsMatch &= (Actor != nullptr) && Actor->GetActorNameOrLabel().MatchesWildcard(Settings.ActorFilter);
 			}
 
 			// Filter by component
-			if ( bIsMatch && Settings.bComponentFilterEnabled )
+			if ( bIsMatch && Settings.bComponentFilterEnabled && !Settings.ComponentFilter.IsEmpty())
 			{
 				bIsMatch &= FXComponent->GetName().MatchesWildcard(Settings.ComponentFilter);
 			}
@@ -2163,7 +2163,7 @@ void FNiagaraDebugHud::DrawOverview(class FNiagaraWorldManager* WorldManager, FC
 						SysInfo.PerfStats->History.GetHistoryFrames_GPU(Frames);
 					}
 
-					if (Settings.bEnableSmoothing)
+					if (Settings.SmoothingWidth > 0)
 					{
 						TArray<double> Smoothed;
 						Smoothed.Reserve(Frames.Num());
@@ -2347,7 +2347,7 @@ void FNiagaraDebugHud::DrawGpuComputeOverriew(class FNiagaraWorldManager* WorldM
 			}
 		}
 
-		const bool bShowDetailed = Settings.bSystemFilterEnabled && OwnerSystem->GetName().MatchesWildcard(Settings.SystemFilter);
+		const bool bShowDetailed = !Settings.bSystemFilterEnabled || Settings.SystemFilter.IsEmpty() || OwnerSystem->GetName().MatchesWildcard(Settings.SystemFilter);
 		SystemIt.Value().bShowDetailed = bShowDetailed;
 #if WITH_EDITORONLY_DATA
 		SystemIt.Value().bCompileForEdit = OwnerSystem->GetCompileForEdit();
@@ -2578,7 +2578,7 @@ void FNiagaraDebugHud::DrawGlobalBudgetInfo(class FNiagaraWorldManager* WorldMan
 		DrawCanvas->DrawTile(TextLocation.X - 1.0f, TextLocation.Y - 1.0f, GuessWidth + 1.0f, 2.0f + (float(NumLines) * fAdvanceHeight), 0.0f, 0.0f, 0.0f, 0.0f, BackgroundColor);
 		DrawCanvas->DrawShadowedString(TextLocation.X, TextLocation.Y, TEXT("Global Budget Info"), Font, HeadingColor);
 		TextLocation.Y += fAdvanceHeight;
-		DrawCanvas->DrawShadowedString(TextLocation.X, TextLocation.Y, TEXT("Global budget tracking is disabled."), Font, DetailHighlightColor);
+		DrawCanvas->DrawShadowedString(TextLocation.X, TextLocation.Y, TEXT("Global budget tracking is disabled. Can be enabled by setting cvars fx.Budget.Enabled and fx.Budget.EnabledInEditor to true."), Font, DetailHighlightColor);
 		TextLocation.Y += fAdvanceHeight;
 	}
 }
@@ -2587,7 +2587,7 @@ void FNiagaraDebugHud::DrawValidation(class FNiagaraWorldManager* WorldManager,
 {
 	using namespace NiagaraDebugLocal;
 
-	if (!Settings.bValidateSystemSimulationDataBuffers && !Settings.bValidateParticleDataBuffers)
+	if (!Settings.bValidationEnabled || (!Settings.bValidateSystemSimulationDataBuffers && !Settings.bValidateParticleDataBuffers))
 	{
 		return;
 	}
@@ -2755,7 +2755,7 @@ void FNiagaraDebugHud::DrawValidation(class FNiagaraWorldManager* WorldManager,
 
 			DrawCanvas->DrawShadowedString(TextLocation.X, TextLocation.Y, ErrorString.ToString(), Font, Settings.MessageErrorTextColor);
 
-			if (Settings.bValidationLogErrors)
+			if (Settings.bValidationEnabled && Settings.bValidationLogErrors)
 			{
 				UE_LOG(LogNiagara, Warning, TEXT("Validation Errors - %s"), ErrorString.ToString());
 			}
diff --git a/Engine/Plugins/FX/Niagara/Source/Niagara/Public/NiagaraDebuggerCommon.h b/Engine/Plugins/FX/Niagara/Source/Niagara/Public/NiagaraDebuggerCommon.h
--- a/Engine/Plugins/FX/Niagara/Source/Niagara/Public/NiagaraDebuggerCommon.h
+++ b/Engine/Plugins/FX/Niagara/Source/Niagara/Public/NiagaraDebuggerCommon.h
@@ -456,80 +456,83 @@
 #endif
 
 	/** Primary control for all HUD features. */
-	UPROPERTY(EditAnywhere, Category = "Debug General", meta = (DisplayName = "Debug HUD Enabled"))
+	UPROPERTY()
 	bool bHudEnabled = true;
 
 	/** Primary control for HUD rendering. */
-	UPROPERTY(EditAnywhere, Category = "Debug General", meta = (DisplayName = "Debug HUD Rendering Enabled"))
+	UPROPERTY()
 	bool bHudRenderingEnabled = true;
 
-	/**
-	When enabled all Niagara systems that pass the filter will have the simulation data buffers validation.
-	i.e. we will look for NaN or other invalidate data  inside it
-	Note: This will have an impact on performance.
-	*/
-	UPROPERTY(EditAnywhere, Category = "Debug Validation")
-	bool bValidateSystemSimulationDataBuffers = false;
+	/** Primary control for validation settings. */
+	UPROPERTY(VisibleAnywhere, Category = "Debug Overview")
+	bool bValidationEnabled = false;
 
-	/**
-	When enabled all Niagara systems that pass the filter will have the particle data buffers validation.
-	i.e. we will look for NaN or other invalidate data  inside it
-	Note: This will have an impact on performance.
-	*/
-	UPROPERTY(EditAnywhere, Category = "Debug Validation")
-	bool bValidateParticleDataBuffers = false;
-
-	/**
-	When enabled all validation errors will be sent to the log as warnings.
-	This can be useful to try and narrow down the exact source of an invalid value in the data buffers as often
-	they will cascade from the first frame where one is generated into other attributes in the subsequent frames.
-	*/
-	UPROPERTY(EditAnywhere, Category = "Debug Validation")
-	bool bValidationLogErrors = false;
-
-	/**
-	When > 0 this is the maximum number of attributes we will display that contain a NaN,
-	there could be more but the display will be truncated to this amount.  This is to reduce generating long strings.
-	*/
-	UPROPERTY(EditAnywhere, Category = "Debug Validation")
-	int32 ValidationAttributeDisplayTruncate = 3;
-
 	/** When enabled the overview display will be enabled. */
-	UPROPERTY(EditAnywhere, Category = "Debug Overview", meta = (DisplayName = "Debug Overview Enabled"))
+	UPROPERTY(VisibleAnywhere, Category = "Debug Overview")
 	bool bOverviewEnabled = false;
 
-	/** When enabled the overview display will include cascade FX. */
-	UPROPERTY(EditAnywhere, Category = "Debug Overview")
-	bool bIncludeCascade = true;
+	/** When enabled will show the system bounds for all filtered systems. */
+	UPROPERTY(VisibleAnywhere, Category = "Debug Overview")
+	bool bSystemShowBounds = false;
 
-	UPROPERTY(EditAnywhere, Category = "Debug Overview", meta = (DisplayName = "Debug Overview Mode"))
+	UPROPERTY(EditAnywhere, Category = "Debug Overview", meta = (DisplayName = "Debug Overview Mode", EditCondition = "bOverviewEnabled"))
 	ENiagaraDebugHUDOverviewMode OverviewMode = ENiagaraDebugHUDOverviewMode::Overview;
 
-	UPROPERTY(EditAnywhere, Category = "Debug Overview", meta = (DisplayName = "Debug Overview Sort Mode"))
+	// How to sort the displayed list of systems
+	UPROPERTY(EditAnywhere, Category = "Debug Overview", meta = (DisplayName = "Debug Overview Sort Mode", EditCondition = "bOverviewEnabled"))
 	ENiagaraDebugHUDDOverviewSort OverviewSortMode = ENiagaraDebugHUDDOverviewSort::Name;
 
-	/** Overview display font to use. */
-	UPROPERTY(EditAnywhere, Category = "Debug Overview", meta = (DisplayName = "Debug Overview Font", EditCondition = "bOverviewEnabled"))
-	ENiagaraDebugHudFont OverviewFont = ENiagaraDebugHudFont::Normal;
+	/** When enabled the overview display will include cascade FX. */
+	UPROPERTY(EditAnywhere, Category = "Debug Overview", AdvancedDisplay, meta = (EditCondition = "bOverviewEnabled"))
+	bool bIncludeCascade = true;
 
-	/** Overview display location. */
-	UPROPERTY(EditAnywhere, Category = "Debug Overview", meta = (DisplayName = "Debug Overview Text Location", EditCondition = "bOverviewEnabled"))
-	FVector2D OverviewLocation = FVector2D(30.0f, 150.0f);
-
-	/** Overview display font to use. */
 	UPROPERTY(EditAnywhere, Category = "Debug Overview", meta = (EditCondition = "bOverviewEnabled && OverviewMode == ENiagaraDebugHUDOverviewMode::Overview"))
 	bool bShowRegisteredComponents = false;
 	
 	/** When enabled the overview will only show the filter system information. */
-	UPROPERTY(EditAnywhere, Category = "Debug Overview", meta = (EditCondition = "bOverviewEnabled"))
+	UPROPERTY(EditAnywhere, Category = "Debug Overview", meta = (DisplayName = "Show only filtered system info", EditCondition = "bOverviewEnabled"))
 	bool bOverviewShowFilteredSystemOnly = false;
 
+	UPROPERTY(EditAnywhere, Category = "Debug Overview", AdvancedDisplay)
+	bool bShowGlobalBudgetInfo = false;
+	
 	/**
+	If enabled, then only systems matching the filters will be considered by the debug hud
+	*/
+	UPROPERTY(EditAnywhere, Category = "Debug Filter", meta = (DisplayName = "Filter displayed systems"))
+	bool bSystemFilterEnabled = false;
+
+	/**
+	Wildcard filter for the systems to show more detailed information about.
+	For example,. "NS_*" would match all systems starting with NS_.
+	*/
+	UPROPERTY(EditAnywhere, Category = "Debug Filter", meta = (DisplayName = "System Name Filter", EditCondition = "bSystemFilterEnabled"))
+	FString SystemFilter;
+	
+
+	UPROPERTY(EditAnywhere, Category = "Debug Filter", meta = (InlineEditConditionToggle))
+	bool bEmitterFilterEnabled = false;
+
+	/**
+	Wildcard filter used to match emitters when generating particle attribute view.
+	For example,. "Fluid*" would match all emitters starting with Fluid and only particle attributes for those would be visible.
+	Requires a valid SystemFilter to be set.
+	*/
+	UPROPERTY(EditAnywhere, Category = "Debug Filter", meta = (DisplayName = "Emitter Name Filter",EditCondition = "bEmitterFilterEnabled"))
+	FString EmitterFilter;
+	
+
+	UPROPERTY(EditAnywhere, Category = "Debug Filter", meta = (InlineEditConditionToggle))
+	bool bActorFilterEnabled = false;
+
+	/**
 	Wildcard filter which is compared against the Components Actor name to narrow down the detailed information.
 	For example, "*Water*" would match all actors that contain the string "water".
+	Requires a valid SystemFilter to be set.
 	*/
-	UPROPERTY(EditAnywhere, Category = "Debug Filter", meta = (EditCondition = "bActorFilterEnabled"))
+	UPROPERTY(EditAnywhere, Category = "Debug Filter", meta = (DisplayName = "Actor Name Filter", EditCondition = "bActorFilterEnabled"))
 	FString ActorFilter;
+	
 
 	UPROPERTY(EditAnywhere, Category = "Debug Filter", meta = (InlineEditConditionToggle))
 	bool bComponentFilterEnabled = false;
@@ -537,36 +540,46 @@
 	/**
 	Wildcard filter for the components to show more detailed information about.
 	For example, "*MyComp*" would match all components that contain MyComp.
+	Requires a valid SystemFilter to be set.
 	*/
-	UPROPERTY(EditAnywhere, Category = "Debug Filter", meta = (EditCondition = "bComponentFilterEnabled"))
+	UPROPERTY(EditAnywhere, Category = "Debug Filter", meta = (DisplayName = "Component Name Filter", EditCondition = "bComponentFilterEnabled"))
 	FString ComponentFilter;
 
-	UPROPERTY(EditAnywhere, Category = "Debug Filter", meta = (InlineEditConditionToggle))
-	bool bSystemFilterEnabled = false;
+	
+	/**
+	When enabled all Niagara systems that pass the filter will have the simulation data buffers validation.
+	i.e. we will look for NaN or other invalidate data  inside it
+	Note: This will have an impact on performance.
+	*/
+	UPROPERTY(EditAnywhere, Category = "Debug Validation", meta=(EditCondition = bValidationEnabled, EditConditionHides))
+	bool bValidateSystemSimulationDataBuffers = true;
 
 	/**
-	Wildcard filter for the systems to show more detailed information about.
-	For example,. "NS_*" would match all systems starting with NS_.
+	When enabled all Niagara systems that pass the filter will have the particle data buffers validation.
+	i.e. we will look for NaN or other invalidate data  inside it
+	Note: This will have an impact on performance.
 	*/
-	UPROPERTY(EditAnywhere, Category = "Debug Filter", meta = (EditCondition = "bSystemFilterEnabled"))
-	FString SystemFilter;
-
-	UPROPERTY(EditAnywhere, Category = "Debug Filter", meta = (InlineEditConditionToggle))
-	bool bEmitterFilterEnabled = false;
+	UPROPERTY(EditAnywhere, Category = "Debug Validation", meta=(EditCondition = bValidationEnabled, EditConditionHides))
+	bool bValidateParticleDataBuffers = true;
 
 	/**
-	Wildcard filter used to match emitters when generating particle attribute view.
-	For example,. "Fluid*" would match all emtiters starting with Fluid and only particle attributes for those would be visible.
+	When enabled all validation errors will be sent to the log as warnings.
+	This can be useful to try and narrow down the exact source of an invalid value in the data buffers as often
+	they will cascade from the first frame where one is generated into other attributes in the subsequent frames.
 	*/
-	UPROPERTY(EditAnywhere, Category = "Debug Filter", meta = (EditCondition = "bEmitterFilterEnabled"))
-	FString EmitterFilter;
+	UPROPERTY(EditAnywhere, Category = "Debug Validation", meta=(EditCondition = bValidationEnabled, EditConditionHides))
+	bool bValidationLogErrors = false;
 
-	UPROPERTY(EditAnywhere, Category = "Debug Filter", meta = (InlineEditConditionToggle))
-	bool bActorFilterEnabled = false;
+	/**
+	When > 0 this is the maximum number of attributes we will display that contain a NaN,
+	there could be more but the display will be truncated to this amount.  This is to reduce generating long strings.
+	*/
+	UPROPERTY(EditAnywhere, Category = "Debug Validation", meta=(EditCondition = bValidationEnabled, EditConditionHides))
+	int32 ValidationAttributeDisplayTruncate = 3;
 
 	/** When enabled system debug information will be displayed in world. */
 	UPROPERTY(EditAnywhere, Category = "Debug System")
-	ENiagaraDebugHudVerbosity SystemDebugVerbosity = ENiagaraDebugHudVerbosity::Basic;
+	ENiagaraDebugHudVerbosity SystemDebugVerbosity = ENiagaraDebugHudVerbosity::None;
 
 	/** When enabled we show information about emitter / particle counts. */
 	UPROPERTY(EditAnywhere, Category = "Debug System", meta = (EditCondition = "SystemDebugVerbosity != ENiagaraDebugHudVerbosity::None"))
@@ -576,36 +589,27 @@
 	UPROPERTY(EditAnywhere, Category = "Debug System", meta = (EditCondition = "SystemDebugVerbosity != ENiagaraDebugHudVerbosity::None"))
 	ENiagaraDebugHudVerbosity DataInterfaceVerbosity = ENiagaraDebugHudVerbosity::None;
 
-	/** When enabled will show the system bounds for all filtered systems. */
-	UPROPERTY(EditAnywhere, Category = "Debug System")
-	bool bSystemShowBounds = false;
-
-	/** When bounds display is enabled allows you to draw a solid box if alpha is > 0. */
-	UPROPERTY(EditAnywhere, Category = "Debug System", meta = (EditCondition = "bSystemShowBounds"))
-	float SystemBoundsSolidBoxAlpha = 0.0f;
-
-	/** When disabled in world rendering will show systems deactivated by scalability. */
-	UPROPERTY(EditAnywhere, Category = "Debug System", meta = (EditCondition = "SystemDebugVerbosity != ENiagaraDebugHudVerbosity::None"))
-	bool bSystemShowActiveOnlyInWorld = true;
-
-	/** Should we display the system attributes. */
-	UPROPERTY(EditAnywhere, Category = "Debug System", meta = (EditCondition = "SystemDebugVerbosity != ENiagaraDebugHudVerbosity::None", DisplayName="Show System Attributes"))
-	bool bShowSystemVariables = true;
-
 	/**
 	List of attributes to show about the system, each entry uses wildcard matching.
 	For example, "System.*" would match all system attributes.
 	*/
-	UPROPERTY(EditAnywhere, Category = "Debug System", meta = (EditCondition = "SystemDebugVerbosity != ENiagaraDebugHudVerbosity::None && bShowSystemVariables", DisplayName="System Attributes"))
+	UPROPERTY(EditAnywhere, Category = "Debug System", meta = (EditCondition = "SystemDebugVerbosity != ENiagaraDebugHudVerbosity::None", DisplayName="System Attributes"))
 	TArray<FNiagaraDebugHUDVariable> SystemVariables;
 
-	/** Sets display text options for system information. */
-	UPROPERTY(EditAnywhere, Category = "Debug System", meta=(EditCondition="SystemDebugVerbosity != ENiagaraDebugHudVerbosity::None"))
-	FNiagaraDebugHudTextOptions SystemTextOptions;
+	/** When disabled in world rendering will show systems deactivated by scalability. */
+	UPROPERTY(EditAnywhere, Category = "Debug System", meta = (EditCondition = "SystemDebugVerbosity != ENiagaraDebugHudVerbosity::None"))
+	bool bSystemShowActiveOnlyInWorld = true;
 
 	/** When enabled will show particle attributes from the list. */
 	UPROPERTY(EditAnywhere, Category = "Debug Particles", meta = (DisplayName="Show Particle Attributes"))
-	bool bShowParticleVariables = true;
+	bool bShowParticleVariables = false;
 
 	/**
+	List of attributes to show per particle, each entry uses wildcard matching.
+	For example, "*Position" would match all attributes that end in Position.
+	*/
+	UPROPERTY(EditAnywhere, Category = "Debug Particles", meta = (EditCondition = "bShowParticleVariables", DisplayName="Particle Attributes"))
+	TArray<FNiagaraDebugHUDVariable> ParticlesVariables;
+
+	/**
 	When enabled GPU particle data will be copied from the GPU to the CPU.
@@ -623,145 +627,151 @@
 	bool bShowParticleIndex = false;
 
 	/**
-	List of attributes to show per particle, each entry uses wildcard matching.
-	For example, "*Position" would match all attributes that end in Position.
-	*/
-	UPROPERTY(EditAnywhere, Category = "Debug Particles", meta = (EditCondition = "bShowParticleVariables", DisplayName="Particle Attributes"))
-	TArray<FNiagaraDebugHUDVariable> ParticlesVariables;
-
-	/** Sets display text options for particle information. */
-	UPROPERTY(EditAnywhere, Category = "Debug Particles", meta = (EditCondition = "bShowParticleVariables"))
-	FNiagaraDebugHudTextOptions ParticleTextOptions;
-
-	/**
 	When enabled particle attributes will display with the system information
-	rather than in world at the particle location.
+	rather than in world at the particle location. Only enabled if system is set up to display information as well.
 	*/
-	UPROPERTY(EditAnywhere, Category = "Debug Particles", meta = (EditCondition = "bShowParticleVariables", DisplayName="Show Particles Attributes With System"))
+	UPROPERTY(EditAnywhere, Category = "Debug Particles", meta = (EditCondition = "bShowParticleVariables && SystemDebugVerbosity != ENiagaraDebugHudVerbosity::None", DisplayName="Display With System Attributes"))
 	bool bShowParticlesVariablesWithSystem = false;
 
-	UPROPERTY(EditAnywhere, Category = "Debug Particles", meta = (EditCondition = "bShowParticleVariables", DisplayName = "Show Particle Attributes Vertical"))
+	UPROPERTY(EditAnywhere, Category = "Debug Particles", meta = (EditCondition = "bShowParticleVariables", DisplayName = "Stack Attributes Vertical"))
 	bool bShowParticleVariablesVertical = false;
 
 	UPROPERTY(EditAnywhere, Category = "Debug Particles", meta = (EditCondition = "bShowParticleVariables"))
 	bool bUseMaxParticlesToDisplay = true;
 
+	/**
+	When enabled, the maximum number of particles to show information about.
+	When disabled all particles will show attributes, this can result in poor performance & potential OOM on some platforms.
+	*/
+	UPROPERTY(EditAnywhere, Category = "Debug Particles", meta = (EditCondition = "bUseMaxParticlesToDisplay && bShowParticleVariables", UIMin="1", ClampMin="1"))
+	int32 MaxParticlesToDisplay = 32;
+
 	/** When enabled we use the clip planes to narrow down which particles to display */
-	UPROPERTY(EditAnywhere, Category = "Debug Particles", meta = (EditCondition = "bShowParticleVariables"))
+	UPROPERTY(EditAnywhere, AdvancedDisplay, Category = "Debug Particles", meta = (EditCondition = "bShowParticleVariables"))
 	bool bUseParticleDisplayClip = false;
 
 	/** Clipping planes used to display particle attributes. */
-	UPROPERTY(EditAnywhere, Category = "Debug Particles", meta = (EditCondition = "bShowParticleVariables", UIMin = "0", ClampMin = "0"))
+	UPROPERTY(EditAnywhere, AdvancedDisplay, Category = "Debug Particles", meta = (EditCondition = "bShowParticleVariables", UIMin = "0", ClampMin = "0"))
 	FVector2D ParticleDisplayClip = FVector2D(0.0f, 10000.0f);
 
 	/** When enabled we use a radius from the display center to avoid showing too many particle attributes. */
-	UPROPERTY(EditAnywhere, Category = "Debug Particles", meta = (EditCondition = "bShowParticleVariables"))
+	UPROPERTY(EditAnywhere, AdvancedDisplay, Category = "Debug Particles", meta = (EditCondition = "bShowParticleVariables"))
 	bool bUseParticleDisplayCenterRadius = false;
 
 	/** Radius from screen center where 0 is center to 1.0 is edge to avoid display too many particle attributes. */
-	UPROPERTY(EditAnywhere, Category = "Debug Particles", meta = (EditCondition = "bShowParticleVariables", UIMin = "0", ClampMin = "0"))
+	UPROPERTY(EditAnywhere, AdvancedDisplay, Category = "Debug Particles", meta = (EditCondition = "bShowParticleVariables", UIMin = "0", ClampMin = "0"))
 	float ParticleDisplayCenterRadius = 1.0f;
 
-	/**
-	When enabled, the maximum number of particles to show information about.
-	When disabled all particles will show attributes, this can result in poor performance & potential OOM on some platforms.
-	*/
-	UPROPERTY(EditAnywhere, Category = "Debug Particles", meta = (EditCondition = "bUseMaxParticlesToDisplay && bShowParticleVariables", UIMin="1", ClampMin="1"))
-	int32 MaxParticlesToDisplay = 32;
-
 	/** How many frames to capture in between updates to the max and average perf report values. */
-	UPROPERTY(EditAnywhere, Category = "Perf Overview")
+	UPROPERTY(EditAnywhere, Category = "Performance")
 	int32 PerfReportFrames = 60;
 
-	UPROPERTY(EditAnywhere, Category = "Perf Overview")
+	UPROPERTY(EditAnywhere, Category = "Performance")
 	ENiagaraDebugHUDPerfSampleMode PerfSampleMode = ENiagaraDebugHUDPerfSampleMode::FrameTotal;
 
-	UPROPERTY(EditAnywhere, Category = "Perf Overview")
+	UPROPERTY(EditAnywhere, Category = "Performance")
 	ENiagaraDebugHUDPerfUnits PerfUnits = ENiagaraDebugHUDPerfUnits::Microseconds;
 
-	UPROPERTY(EditAnywhere, Category = "Perf Overview", meta = (EditConditionHides, EditCondition = "OverviewMode == ENiagaraDebugHUDOverviewMode::Performance"))
+	// Adds a column that displays only the GT cost to the perf overview. Only enabled when performance overview mode is active.
+	UPROPERTY(EditAnywhere, Category = "Performance", meta = (DisplayName = "Show separate GT perf column", EditCondition = "OverviewMode == ENiagaraDebugHUDOverviewMode::Performance"))
 	bool bShowPerfColumGameThreadOnly = false;
 
 	/** Time range of the Y Axis of the perf graph */
-	UPROPERTY(EditAnywhere, Category = "Perf Overview")
+	UPROPERTY(EditAnywhere, Category = "Performance Graph", meta = (EditCondition = "OverviewMode == ENiagaraDebugHUDOverviewMode::PerformanceGraph", EditConditionHides))
 	ENiagaraDebugHUDPerfGraphMode PerfGraphMode = ENiagaraDebugHUDPerfGraphMode::GameThread;
 
 	/** How many frames of history to display in the perf graphs. */
-	UPROPERTY(EditAnywhere, Category = "Perf Overview")
+	UPROPERTY(EditAnywhere, Category = "Performance Graph", meta = (EditCondition = "OverviewMode == ENiagaraDebugHUDOverviewMode::PerformanceGraph", EditConditionHides))
 	int32 PerfHistoryFrames = 600;
 
 	/** Use the specified user range when enabled, otherwise we will auto detect the range to use. */
-	UPROPERTY(EditAnywhere, Category = "Perf Overview", meta = (InlineEditConditionToggle))
+	UPROPERTY(EditAnywhere, Category = "Performance Graph", meta = (DisplayName = "Fixed Y Axis Range", EditCondition = "OverviewMode == ENiagaraDebugHUDOverviewMode::PerformanceGraph", EditConditionHides))
 	bool bUsePerfGraphTimeRange = false;
 
 	/** Time range of the Y Axis of the perf graph */
-	UPROPERTY(EditAnywhere, Category = "Perf Overview", meta = (EditCondition = "bUsePerfGraphTimeRange"))
+	UPROPERTY(EditAnywhere, Category = "Performance Graph", meta = (DisplayName = "Y Axis Time Range", EditCondition = "OverviewMode == ENiagaraDebugHUDOverviewMode::PerformanceGraph", EditConditionHides))
 	float PerfGraphTimeRange = 500.0f;
 
 	/** Pixel size of the perf graph. */
-	UPROPERTY(EditAnywhere, Category = "Perf Overview")
+	UPROPERTY(EditAnywhere, Category = "Performance Graph", meta = (EditCondition = "OverviewMode == ENiagaraDebugHUDOverviewMode::PerformanceGraph", EditConditionHides))
 	FVector2D PerfGraphSize = FVector2D(500,500);
 
-	UPROPERTY(EditAnywhere, Category = "Perf Overview")
+	UPROPERTY(EditAnywhere, Category = "Performance Graph", meta = (EditCondition = "OverviewMode == ENiagaraDebugHUDOverviewMode::PerformanceGraph", EditConditionHides))
 	FLinearColor PerfGraphAxisColor = FLinearColor::White;
 
-	// True if perf graphs should be smoothed.
-	UPROPERTY(EditAnywhere, Category = "Perf Overview", meta = (InlineEditConditionToggle))
-	bool bEnableSmoothing = true;
-
 	//Number of samples to use either size of a value when smoothing perf graphs.
-	UPROPERTY(EditAnywhere, Category = "Perf Overview", meta = (EditCondition = "bEnableSmoothing"))
+	UPROPERTY(EditAnywhere, Category = "Performance Graph", meta = (EditCondition = "OverviewMode == ENiagaraDebugHUDOverviewMode::PerformanceGraph", EditConditionHides))
 	int32 SmoothingWidth = 4;
 
+	/** Overview display font to use. */
+	UPROPERTY(EditAnywhere, Category = "Style", meta = (DisplayName = "Debug Overview Font", EditCondition = "bOverviewEnabled"))
+	ENiagaraDebugHudFont OverviewFont = ENiagaraDebugHudFont::Normal;
+
+	/** Overview display location. */
+	UPROPERTY(EditAnywhere, Category = "Style", meta = (DisplayName = "Debug Overview Text Location", EditCondition = "bOverviewEnabled"))
+	FVector2D OverviewLocation = FVector2D(30.0f, 150.0f);
+
+	/** Sets display text options for system information. */
+	UPROPERTY(EditAnywhere, Category = "Style", meta=(EditCondition="SystemDebugVerbosity != ENiagaraDebugHudVerbosity::None"))
+	FNiagaraDebugHudTextOptions SystemTextOptions;
+
+	/** Sets display text options for particle information. */
+	UPROPERTY(EditAnywhere, Category = "Style", meta = (EditCondition = "bShowParticleVariables"))
+	FNiagaraDebugHudTextOptions ParticleTextOptions;
+
+	/** When bounds display is enabled allows you to draw a solid box if alpha is > 0. */
+	UPROPERTY(EditAnywhere, Category = "Style", meta = (EditCondition = "bSystemShowBounds"))
+	float SystemBoundsSolidBoxAlpha = 0.0f;
+
 	// Default background color used generally for panels
-	UPROPERTY(EditAnywhere, Category = "Colors")
+	UPROPERTY(EditAnywhere, Category = "Style")
 	FLinearColor DefaultBackgroundColor = FLinearColor(0.0f, 0.0f, 0.0f, 0.5f);
 
 	// Overview heading text color
-	UPROPERTY(EditAnywhere, Category = "Colors")
+	UPROPERTY(EditAnywhere, Category = "Style")
 	FLinearColor OverviewHeadingColor = FLinearColor::Green;
 
 	// Overview detail text color
-	UPROPERTY(EditAnywhere, Category = "Colors")
+	UPROPERTY(EditAnywhere, Category = "Style")
 	FLinearColor OverviewDetailColor = FLinearColor::White;
 
 	// Overview detail highlight text color
-	UPROPERTY(EditAnywhere, Category = "Colors")
+	UPROPERTY(EditAnywhere, Category = "Style")
 	FLinearColor OverviewDetailHighlightColor = FLinearColor::Yellow;
 
 	// In world text if an error is detected
-	UPROPERTY(EditAnywhere, Category = "Colors")
+	UPROPERTY(EditAnywhere, Category = "Style")
 	FLinearColor InWorldErrorTextColor = FLinearColor(1.0f, 0.4f, 0.3f, 1.0f);
 	// In world text color
-	UPROPERTY(EditAnywhere, Category = "Colors")
+	UPROPERTY(EditAnywhere, Category = "Style")
 	FLinearColor InWorldTextColor = FLinearColor::White;
 
 	// Message display text color
-	UPROPERTY(EditAnywhere, Category = "Colors")
+	UPROPERTY(EditAnywhere, Category = "Style")
 	FLinearColor MessageInfoTextColor = FLinearColor::White;
 
 	// Message display warning text color
-	UPROPERTY(EditAnywhere, Category = "Colors")
+	UPROPERTY(EditAnywhere, Category = "Style")
 	FLinearColor MessageWarningTextColor = FLinearColor(0.9f, 0.7f, 0.0f, 1.0f);
 
 	// Message display error text color
-	UPROPERTY(EditAnywhere, Category = "Colors")
+	UPROPERTY(EditAnywhere, Category = "Style")
 	FLinearColor MessageErrorTextColor = FLinearColor(1.0f, 0.4f, 0.3f, 1.0f);
 	
 	/** Opacity of the system color background tile in overview table rows. */
-	UPROPERTY(EditAnywhere, Category = "Colors")
+	UPROPERTY(EditAnywhere, Category = "Style")
 	float SystemColorTableOpacity = 0.2f;
 
 	/** Additional seed value for random system colors. Useful if current colors of systems are too similar. */
-	UPROPERTY(EditAnywhere, Category = "Colors")
+	UPROPERTY(EditAnywhere, Category = "Style")
 	uint32 SystemColorSeed = 0;
 
 	/** Minimum HSV values for the random colors generated for each System. */
-	UPROPERTY(EditAnywhere, Category = "Colors")
+	UPROPERTY(EditAnywhere, Category = "Style")
 	FVector SystemColorHSVMin = FVector(0, 200, 200);
 
 	/** Maximum HSV values for the random colors generated for each System. */
-	UPROPERTY(EditAnywhere, Category = "Colors")
+	UPROPERTY(EditAnywhere, Category = "Style")
 	FVector SystemColorHSVMax = FVector(255, 255, 255);
 
 	UPROPERTY()
@@ -778,9 +788,6 @@
 
 	UPROPERTY()
 	float LoopTime = 1.0f;
-
-	UPROPERTY(EditAnywhere, Category = "Performance")
-	bool bShowGlobalBudgetInfo = false;
 };
 
 /** Message passed from debugger to client when it needs updated simple client info. */
diff --git a/Engine/Plugins/FX/Niagara/Source/NiagaraEditor/Private/Customizations/NiagaraDebugHUDCustomization.cpp b/Engine/Plugins/FX/Niagara/Source/NiagaraEditor/Private/Customizations/NiagaraDebugHUDCustomization.cpp
--- a/Engine/Plugins/FX/Niagara/Source/NiagaraEditor/Private/Customizations/NiagaraDebugHUDCustomization.cpp
+++ b/Engine/Plugins/FX/Niagara/Source/NiagaraEditor/Private/Customizations/NiagaraDebugHUDCustomization.cpp
@@ -9,21 +9,18 @@
 #include "DetailLayoutBuilder.h"
 #include "DetailWidgetRow.h"
 #include "Framework/Application/SlateApplication.h"
-#include "IDetailChildrenBuilder.h"
 //Widgets
 #include "Widgets/Input/SCheckBox.h"
 #include "Widgets/Input/SEditableTextBox.h"
 #include "Widgets/SBoxPanel.h"
 #include "Widgets/SNiagaraDebugger.h"
+#include "Widgets/Input/SComboButton.h"
 #include "Widgets/Views/STreeView.h"
 #include "Widgets/Input/SMenuAnchor.h"
-#include "Widgets/Input/SSearchBox.h"
-#include "Widgets/Input/SSuggestionTextBox.h"
-#include "Widgets/Layout/SScrollBorder.h"
-#include "Widgets/Input/SMultiLineEditableTextBox.h"
 ///Niagara
 #include "NiagaraEditorModule.h"
 #include "NiagaraComponent.h"
+#include "NiagaraEditorStyle.h"
 
 #if WITH_NIAGARA_DEBUGGER
 
@@ -34,7 +31,7 @@
 namespace NiagaraDebugHUDSettingsDetailsCustomizationInternal
 {
 
-class SDebuggerSuggestionTextBox : public SSuggestionTextBox
+class SDebuggerSuggestionTextBox : public SComboButton
 {
 public:
 	SLATE_BEGIN_ARGS(SDebuggerSuggestionTextBox) {}
@@ -63,14 +60,22 @@ public:
 
 		FText CurrentValue;
 		PropertyHandle->GetValueAsFormattedText(CurrentValue);
-
-		SSuggestionTextBox::Construct(
-			SSuggestionTextBox::FArguments()
-			.ForegroundColor(FSlateColor::UseForeground())
-			.Font(IDetailLayoutBuilder::GetDetailFont())
+		 Textbox = SNew(SEditableTextBox)
+			.Padding(0)
+			.Style(FNiagaraEditorStyle::Get(), "NiagaraEditor.Debugger.SuggestionDropdownInput")
+			.MinDesiredWidth(20)
+			.RevertTextOnEscape(true)
+			.SelectAllTextWhenFocused(true)
 			.Text(CurrentValue)
-			.OnTextCommitted(this, &SDebuggerSuggestionTextBox::OnDebuggerTextCommitted)
-			.OnShowingSuggestions(this, &SDebuggerSuggestionTextBox::GetDebuggerSuggestions)
+			.OnTextCommitted(this, &SDebuggerSuggestionTextBox::OnDebuggerTextCommitted);
+		
+		SComboButton::Construct(
+			SComboButton::FArguments()
+			.OnGetMenuContent(this, &SDebuggerSuggestionTextBox::GetDebuggerSuggestions)
+			.ButtonContent()
+			[
+				Textbox.ToSharedRef()
+			]
 		);
 	}
 
@@ -82,37 +87,56 @@ public:
 		}
 	}
 
-	void GetDebuggerSuggestions(const FString& CurrText, TArray<FString>& OutSuggestions)
+	void SelectDropdownValue(FString NewValue)
 	{
+		if (Textbox)
+		{
+			Textbox->SetText(FText::FromString(NewValue));
+		}
+	}
+
+	TSharedRef<SWidget> GetDebuggerSuggestions()
+	{
+		FMenuBuilder MenuBuilder( true, NULL );
+		
 		if (!Debugger.IsValid() || !ObjectClass)
 		{
-			return;
+			return MenuBuilder.MakeWidget();
 		}
 
 		const FNiagaraSimpleClientInfo& SimpleClientInfo = Debugger->GetSimpleClientInfo();
+		TArray<FString> Options;
 		if (ObjectClass == UNiagaraSystem::StaticClass())
 		{
-			OutSuggestions.Append(SimpleClientInfo.Systems);
+			Options = SimpleClientInfo.Systems;
 		}
 		else if (ObjectClass == UNiagaraEmitter::StaticClass())
 		{
-			OutSuggestions.Append(SimpleClientInfo.Emitters);
+			Options = SimpleClientInfo.Emitters;
 		}
 		else if (ObjectClass == AActor::StaticClass())
 		{
-			OutSuggestions.Append(SimpleClientInfo.Actors);
+			Options = SimpleClientInfo.Actors;
 		}
 		else if (ObjectClass == UNiagaraComponent::StaticClass())
 		{
-			OutSuggestions.Append(SimpleClientInfo.Components);
+			Options = SimpleClientInfo.Components;
+		}
+		Options.Sort();
+
+		for (const FString& Option : Options)
+		{
+			FUIAction MenuAction(FExecuteAction::CreateSP(this, &SDebuggerSuggestionTextBox::SelectDropdownValue, Option));
+			MenuBuilder.AddMenuEntry(FText::FromString(Option), FText(), FSlateIcon(), MenuAction);
 		}
+		return MenuBuilder.MakeWidget();
 	}
 
 	void OnHUDSettingsChanged()
 	{
 		FText CurrentValue;
 		PropertyHandle->GetValueAsFormattedText(CurrentValue);
-		SetText(CurrentValue);
+		Textbox->SetText(CurrentValue);
 	}
 
 	void OnSimpleClientInfoChanged(const FNiagaraSimpleClientInfo& ClientInfo)
@@ -122,7 +146,7 @@ public:
 
 	virtual void Tick(const FGeometry& AllottedGeometry, const double InCurrentTime, const float InDeltaTime) override
 	{
-		SSuggestionTextBox::Tick(AllottedGeometry, InCurrentTime, InDeltaTime);
+		SComboButton::Tick(AllottedGeometry, InCurrentTime, InDeltaTime);
 
 		if (Debugger)
 		{
@@ -135,11 +159,12 @@ public:
 		}
 	}
 
-	TSharedPtr<FNiagaraDebugger>				Debugger;
-	TSharedPtr<IPropertyHandle>					PropertyHandle;
 	TWeakObjectPtr<UNiagaraDebugHUDSettings>	WeakHUDSettings;
-	UClass*										ObjectClass = nullptr;
-	bool										bWaitingUpdate = false;
+	TSharedPtr<FNiagaraDebugger>	Debugger;
+	TSharedPtr<IPropertyHandle>		PropertyHandle;
+	TSharedPtr<SEditableTextBox>	Textbox;
+	UClass*							ObjectClass = nullptr;
+	bool							bWaitingUpdate = false;
 };
 
 } // namespace
@@ -218,39 +243,36 @@ FNiagaraDebugHUDSettingsDetailsCustomization::FNiagaraDebugHUDSettingsDetailsCus
 
 void FNiagaraDebugHUDSettingsDetailsCustomization::CustomizeDetails(IDetailLayoutBuilder& DetailBuilder)
 {
-	// Customize General
-	{
-		IDetailCategoryBuilder& GeneralCategory = DetailBuilder.EditCategory("Debug General");
-	}
-
-	// Customize Overview
 	{
 		IDetailCategoryBuilder& OverviewCategory = DetailBuilder.EditCategory("Debug Overview");
+		DetailBuilder.HideProperty(DetailBuilder.GetProperty(GET_MEMBER_NAME_CHECKED(FNiagaraDebugHUDSettingsData, bValidationEnabled), FNiagaraDebugHUDSettingsData::StaticStruct()));
+		DetailBuilder.HideProperty(DetailBuilder.GetProperty(GET_MEMBER_NAME_CHECKED(FNiagaraDebugHUDSettingsData, bOverviewEnabled), FNiagaraDebugHUDSettingsData::StaticStruct()));
+		DetailBuilder.HideProperty(DetailBuilder.GetProperty(GET_MEMBER_NAME_CHECKED(FNiagaraDebugHUDSettingsData, bSystemShowBounds), FNiagaraDebugHUDSettingsData::StaticStruct()));
 	}
-
+	
 	// Customize Filters
 	{
 		IDetailCategoryBuilder& FilterCategory = DetailBuilder.EditCategory("Debug Filter");
-		MakeCustomAssetSearch(DetailBuilder, FilterCategory, DetailBuilder.GetProperty(GET_MEMBER_NAME_CHECKED(FNiagaraDebugHUDSettingsData, SystemFilter), FNiagaraDebugHUDSettingsData::StaticStruct()), UNiagaraSystem::StaticClass(), [&]() -> bool& { return WeakSettings.Get()->Data.bSystemFilterEnabled; });
-		MakeCustomAssetSearch(DetailBuilder, FilterCategory, DetailBuilder.GetProperty(GET_MEMBER_NAME_CHECKED(FNiagaraDebugHUDSettingsData, EmitterFilter), FNiagaraDebugHUDSettingsData::StaticStruct()), UNiagaraEmitter::StaticClass(), [&]() -> bool& { return WeakSettings.Get()->Data.bEmitterFilterEnabled; });
-		MakeCustomAssetSearch(DetailBuilder, FilterCategory, DetailBuilder.GetProperty(GET_MEMBER_NAME_CHECKED(FNiagaraDebugHUDSettingsData, ActorFilter), FNiagaraDebugHUDSettingsData::StaticStruct()), AActor::StaticClass(), [&]() -> bool& { return WeakSettings.Get()->Data.bActorFilterEnabled; });
-		MakeCustomAssetSearch(DetailBuilder, FilterCategory, DetailBuilder.GetProperty(GET_MEMBER_NAME_CHECKED(FNiagaraDebugHUDSettingsData, ComponentFilter), FNiagaraDebugHUDSettingsData::StaticStruct()), UNiagaraComponent::StaticClass(), [&]() -> bool& { return WeakSettings.Get()->Data.bComponentFilterEnabled; });
+		
+		FilterCategory.AddProperty(DetailBuilder.GetProperty(GET_MEMBER_NAME_CHECKED(FNiagaraDebugHUDSettingsData, bSystemFilterEnabled), FNiagaraDebugHUDSettingsData::StaticStruct()));
+		MakeCustomAssetSearch(DetailBuilder, FilterCategory, DetailBuilder.GetProperty(GET_MEMBER_NAME_CHECKED(FNiagaraDebugHUDSettingsData, SystemFilter), FNiagaraDebugHUDSettingsData::StaticStruct()), UNiagaraSystem::StaticClass(), [&]() -> bool& { return WeakSettings.Get()->Data.bSystemFilterEnabled; }, [&]() -> bool { return true; }, false);
+		MakeCustomAssetSearch(DetailBuilder, FilterCategory, DetailBuilder.GetProperty(GET_MEMBER_NAME_CHECKED(FNiagaraDebugHUDSettingsData, EmitterFilter), FNiagaraDebugHUDSettingsData::StaticStruct()), UNiagaraEmitter::StaticClass(), [&]() -> bool& { return WeakSettings.Get()->Data.bEmitterFilterEnabled; }, [&]() -> bool { return WeakSettings.Get()->Data.bSystemFilterEnabled; });
+		MakeCustomAssetSearch(DetailBuilder, FilterCategory, DetailBuilder.GetProperty(GET_MEMBER_NAME_CHECKED(FNiagaraDebugHUDSettingsData, ActorFilter), FNiagaraDebugHUDSettingsData::StaticStruct()), AActor::StaticClass(), [&]() -> bool& { return WeakSettings.Get()->Data.bActorFilterEnabled; }, [&]() -> bool { return WeakSettings.Get()->Data.bSystemFilterEnabled; });
+		MakeCustomAssetSearch(DetailBuilder, FilterCategory, DetailBuilder.GetProperty(GET_MEMBER_NAME_CHECKED(FNiagaraDebugHUDSettingsData, ComponentFilter), FNiagaraDebugHUDSettingsData::StaticStruct()), UNiagaraComponent::StaticClass(), [&]() -> bool& { return WeakSettings.Get()->Data.bComponentFilterEnabled; }, [&]() -> bool { return WeakSettings.Get()->Data.bSystemFilterEnabled; });
 	}
 }
 
-void FNiagaraDebugHUDSettingsDetailsCustomization::MakeCustomAssetSearch(IDetailLayoutBuilder& DetailBuilder, IDetailCategoryBuilder& DetailCategory, TSharedRef<IPropertyHandle> PropertyHandle, UClass* ObjRefClass, TFunction<bool&()> GetEditBool)
+void FNiagaraDebugHUDSettingsDetailsCustomization::MakeCustomAssetSearch(IDetailLayoutBuilder& DetailBuilder, IDetailCategoryBuilder& DetailCategory, TSharedRef<IPropertyHandle> PropertyHandle, UClass* ObjRefClass, TFunction<bool&()> GetEditBool, TFunction<bool ()> GetVisibleBool, bool bShowInlineCheckbox)
 {
 	if ( !PropertyHandle->IsValidHandle() || (ObjRefClass == nullptr) )
 	{
 		return;
 	}
-
-	FText CurrentValue;
-	PropertyHandle->GetValueAsFormattedText(CurrentValue);
-
+	
 	DetailBuilder.HideProperty(PropertyHandle);
-
+	
 	DetailCategory.AddCustomRow(PropertyHandle->GetPropertyDisplayName())
+	.Visibility(TAttribute<EVisibility>::CreateLambda([GetVisibleBool]() -> EVisibility { return GetVisibleBool() ? EVisibility::Visible : EVisibility::Collapsed; }))
 	.NameContent()
 	[
 		SNew(SHorizontalBox)
@@ -258,6 +280,7 @@ void FNiagaraDebugHUDSettingsDetailsCustomization::MakeCustomAssetSearch(IDetail
 		.AutoWidth()
 		[
 			SNew(SCheckBox)
+			.Visibility(bShowInlineCheckbox ? EVisibility::Visible : EVisibility::Collapsed)
 			.IsChecked_Lambda([GetEditBool]() -> ECheckBoxState { return GetEditBool() ? ECheckBoxState::Checked : ECheckBoxState::Unchecked; } )
 			.OnCheckStateChanged_Lambda([this, GetEditBool](ECheckBoxState NewState) { GetEditBool() = NewState == ECheckBoxState::Checked; WeakSettings.Get()->NotifyPropertyChanged(); })
 		]
diff --git a/Engine/Plugins/FX/Niagara/Source/NiagaraEditor/Private/Customizations/NiagaraDebugHUDCustomization.h b/Engine/Plugins/FX/Niagara/Source/NiagaraEditor/Private/Customizations/NiagaraDebugHUDCustomization.h
--- a/Engine/Plugins/FX/Niagara/Source/NiagaraEditor/Private/Customizations/NiagaraDebugHUDCustomization.h
+++ b/Engine/Plugins/FX/Niagara/Source/NiagaraEditor/Private/Customizations/NiagaraDebugHUDCustomization.h
@@ -63,7 +63,7 @@ public:
 	// End IDetailCustomization interface
 
 private:
-	void MakeCustomAssetSearch(IDetailLayoutBuilder& DetailBuilder, IDetailCategoryBuilder& DetailCategory, TSharedRef<IPropertyHandle> PropertyHandle, UClass* ObjRefClass, TFunction<bool& ()> GetEditBool);
+	void MakeCustomAssetSearch(IDetailLayoutBuilder& DetailBuilder, IDetailCategoryBuilder& DetailCategory, TSharedRef<IPropertyHandle> PropertyHandle, UClass* ObjRefClass, TFunction<bool& ()> GetEditBool, TFunction<bool ()> GetVisibleBool, bool bShowInlineCheckbox = true);
 
 private:
 	TWeakObjectPtr<UNiagaraDebugHUDSettings> WeakSettings = nullptr;
diff --git a/Engine/Plugins/FX/Niagara/Source/NiagaraEditor/Private/NiagaraEditorStyle.cpp b/Engine/Plugins/FX/Niagara/Source/NiagaraEditor/Private/NiagaraEditorStyle.cpp
--- a/Engine/Plugins/FX/Niagara/Source/NiagaraEditor/Private/NiagaraEditorStyle.cpp
+++ b/Engine/Plugins/FX/Niagara/Source/NiagaraEditor/Private/NiagaraEditorStyle.cpp
@@ -697,6 +697,18 @@ void FNiagaraEditorStyle::InitDebuggerStyle()
 		.SetNormalPadding(FMargin(0, 0, 0, 1))
 		.SetPressedPadding(FMargin(0, 1, 0, 0));
 	Set("NiagaraEditor.Debugger.Outliner.Toolbar", OutlinerToolBarButton);
+
+	FTextBlockStyle NormalText = FCoreStyle::Get().GetWidgetStyle<FTextBlockStyle>("NormalText");
+	FEditableTextBoxStyle DefaultTextBoxStyle = FAppStyle::GetWidgetStyle<FEditableTextBoxStyle>("NormalEditableTextBox");
+	const FEditableTextBoxStyle EditableTextBoxStyle = DefaultTextBoxStyle
+			.SetTextStyle(NormalText.SetFont(FAppStyle::Get().GetFontStyle("PropertyWindow.NormalFont")))
+			.SetPadding(0)
+			.SetBackgroundImageNormal( FSlateNoResource() )
+			.SetBackgroundImageHovered( FSlateNoResource() )
+			.SetBackgroundImageFocused( FSlateNoResource() )
+			.SetBackgroundImageReadOnly( FSlateNoResource() );
+	
+	Set("NiagaraEditor.Debugger.SuggestionDropdownInput", EditableTextBoxStyle);
 }
 
 void FNiagaraEditorStyle::InitBakerStyle()
diff --git a/Engine/Plugins/FX/Niagara/Source/NiagaraEditor/Private/Widgets/SNiagaraDebugger.cpp b/Engine/Plugins/FX/Niagara/Source/NiagaraEditor/Private/Widgets/SNiagaraDebugger.cpp
--- a/Engine/Plugins/FX/Niagara/Source/NiagaraEditor/Private/Widgets/SNiagaraDebugger.cpp
+++ b/Engine/Plugins/FX/Niagara/Source/NiagaraEditor/Private/Widgets/SNiagaraDebugger.cpp
@@ -21,8 +21,9 @@
 #include "Widgets/Docking/SDockTab.h"
 #include "Widgets/Images/SImage.h"
 #include "Widgets/Input/SButton.h"
-#include "Widgets/Input/SNumericDropDown.h"
+#include "Widgets/Input/SComboButton.h"
 #include "Widgets/Input/SNumericEntryBox.h"
+#include "Widgets/Input/SEditableTextBox.h"
 #include "Widgets/Layout/SScrollBox.h"
 #include "Widgets/Layout/SUniformGridPanel.h"
 #include "Widgets/SBoxPanel.h"
@@ -75,6 +76,7 @@ namespace NiagaraDebugHudTab
 		StructureViewArgs.bShowClasses = true;
 		StructureViewArgs.bShowInterfaces = true;
 
+		// we set the struct data after we create the details view, otherwise the custom details layout won't be picked up
 		TSharedRef<IStructureDetailsView> StructureDetailsView = PropertyModule.CreateStructureDetailView(DetailsViewArgs, StructureViewArgs, nullptr);
 		StructureDetailsView->GetDetailsView()->SetGenericLayoutDetailsDelegate(FOnGetDetailCustomizationInstance::CreateStatic(&FNiagaraDebugHUDSettingsDetailsCustomization::MakeInstance, DebugHudSettings));
 
@@ -1047,7 +1049,6 @@ void SNiagaraDebugger::InvokeDebugger(UNiagaraSystem* InSystem, TArray<FNiagaraE
 
 			if (bHasSystemAttributes)
 			{
-				HudSettings->Data.bShowSystemVariables = true;
 				for (const FString& Var : SystemAttribNames)
 				{
 					FNiagaraDebugHUDVariable NewVar;
@@ -1141,17 +1142,50 @@ TSharedRef<SWidget> SNiagaraDebugger::MakeToolbar()
 
 	FSlimHorizontalToolBarBuilder ToolbarBuilder(MakeShareable(new FUICommandList), FMultiBoxCustomization::None);
 	UNiagaraDebugHUDSettings* Settings = GetMutableDefault<UNiagaraDebugHUDSettings>();
+	FCanExecuteAction HudEnabledAction = FCanExecuteAction::CreateLambda([=]() { return Settings->Data.bHudEnabled; });
 	
 	ToolbarBuilder.BeginStyleOverride("CalloutToolbar");
 	ToolbarBuilder.BeginSection("Main");
-	// Refresh button
+
+	// Enabled modes Button
+	{
+		ToolbarBuilder.AddToolBarButton(
+			FUIAction(
+				FExecuteAction::CreateLambda([=]() {Settings->Data.bHudEnabled = !Settings->Data.bHudEnabled; Settings->NotifyPropertyChanged(); }),
+				FCanExecuteAction(),
+				FIsActionChecked::CreateLambda([=]() { return Settings->Data.bHudEnabled; })
+			),
+			NAME_None,
+			LOCTEXT("HudEnableLabel", "HUD"),
+			LOCTEXT("HudEnableTooltip", "Enables or disables the debug hud completely."),
+			FSlateIcon(FAppStyle::Get().GetStyleSetName(), "Debug"),
+			EUserInterfaceActionType::ToggleButton
+		);
+
+		// Additional options
+		ToolbarBuilder.AddComboButton(
+			FUIAction(),
+			FOnGetContent::CreateSP(this, &SNiagaraDebugger::MakeModeOptionsMenu),
+			FText(),
+			LOCTEXT("ModeOptionsTooltip", "Additional options to show or hide debug features."),
+			FSlateIcon(FAppStyle::GetAppStyleSetName(), "MaterialEditor.ToggleMaterialStats"),
+			true
+		);
+	}
+
+	// Show bounds button
 	{
 		ToolbarBuilder.AddToolBarButton(
-			FUIAction(FExecuteAction::CreateLambda([Owner=Debugger]() { Owner->UpdateDebugHUDSettings(); })),
+			FUIAction(
+				FExecuteAction::CreateLambda([=]() {Settings->Data.bSystemShowBounds = !Settings->Data.bSystemShowBounds; Settings->NotifyPropertyChanged(); }),
+				HudEnabledAction,
+				FIsActionChecked::CreateLambda([=]() { return Settings->Data.bSystemShowBounds; })
+			),
 			NAME_None,
-			LOCTEXT("Refresh", "Refresh"),
-			LOCTEXT("RefreshTooltip", "Refesh the settings on the target device.  Used if we get out of sync."),
-			FSlateIcon(FAppStyle::GetAppStyleSetName(), "Icons.Refresh")
+			LOCTEXT("BoundsLabel", "Bounds"),
+			LOCTEXT("BoundsTooltip", "Show system bounding boxes"),
+			FSlateIcon(FAppStyle::Get().GetStyleSetName(), "AssetEditor.ToggleShowBounds"),
+			EUserInterfaceActionType::ToggleButton
 		);
 	}
 
@@ -1164,7 +1198,7 @@ TSharedRef<SWidget> SNiagaraDebugger::MakeToolbar()
 			ToolbarBuilder.AddToolBarButton(
 				FUIAction(
 					FExecuteAction::CreateLambda([=]() {Settings->Data.PlaybackMode = ENiagaraDebugPlaybackMode::Play; Settings->NotifyPropertyChanged(); }),
-					FCanExecuteAction(),
+					HudEnabledAction,
 					FIsActionChecked::CreateLambda([=]() { return Settings->Data.PlaybackMode == ENiagaraDebugPlaybackMode::Play; })
 				),
 				NAME_None,
@@ -1179,7 +1213,7 @@ TSharedRef<SWidget> SNiagaraDebugger::MakeToolbar()
 			ToolbarBuilder.AddToolBarButton(
 				FUIAction(
 					FExecuteAction::CreateLambda([=]() {Settings->Data.PlaybackMode = ENiagaraDebugPlaybackMode::Paused; Settings->NotifyPropertyChanged(); }),
-					FCanExecuteAction(),
+					HudEnabledAction,
 					FIsActionChecked::CreateLambda([=]() { return Settings->Data.PlaybackMode == ENiagaraDebugPlaybackMode::Paused; })
 				),
 				NAME_None,
@@ -1202,7 +1236,7 @@ TSharedRef<SWidget> SNiagaraDebugger::MakeToolbar()
 						FNiagaraSystemUpdateContext UpdateContext;
 						UpdateContext.AddAll(false);
 					}),
-					FCanExecuteAction(),
+					HudEnabledAction,
 					FIsActionChecked::CreateLambda([=]() { return Settings->Data.PlaybackMode == ENiagaraDebugPlaybackMode::Loop; })
 				),
 				NAME_None,
@@ -1217,7 +1251,7 @@ TSharedRef<SWidget> SNiagaraDebugger::MakeToolbar()
 			ToolbarBuilder.AddToolBarButton(
 				FUIAction(
 					FExecuteAction::CreateLambda([=]() {Settings->Data.PlaybackMode = ENiagaraDebugPlaybackMode::Step; Settings->NotifyPropertyChanged(); Settings->Data.PlaybackMode = ENiagaraDebugPlaybackMode::Paused; }),
-					FCanExecuteAction(),
+					HudEnabledAction,
 					FIsActionChecked::CreateLambda([=]() { return Settings->Data.PlaybackMode == ENiagaraDebugPlaybackMode::Step; })
 				),
 				NAME_None,
@@ -1232,7 +1266,7 @@ TSharedRef<SWidget> SNiagaraDebugger::MakeToolbar()
 			ToolbarBuilder.AddToolBarButton(
 				FUIAction(
 					FExecuteAction::CreateLambda([=]() {Settings->Data.bPlaybackRateEnabled = !Settings->Data.bPlaybackRateEnabled; Settings->NotifyPropertyChanged(); }),
-					FCanExecuteAction(),
+					HudEnabledAction,
 					FIsActionChecked::CreateLambda([=]() { return Settings->Data.bPlaybackRateEnabled; })
 				),
 				NAME_None,
@@ -1241,16 +1275,30 @@ TSharedRef<SWidget> SNiagaraDebugger::MakeToolbar()
 				FSlateIcon(FNiagaraEditorStyle::Get().GetStyleSetName(), "NiagaraEditor.Debugger.SpeedIcon"),
 				EUserInterfaceActionType::ToggleButton
 			);
+
+			// Additional options
+			ToolbarBuilder.AddComboButton(
+				FUIAction(),
+				FOnGetContent::CreateSP(this, &SNiagaraDebugger::MakePlaybackOptionsMenu),
+				FText(),
+				LOCTEXT("PlaybackOptionsTooltip", "Additional options to control playback."),
+				FSlateIcon(FAppStyle::GetAppStyleSetName(), "MaterialEditor.ToggleMaterialStats"),
+				true
+			);
+		}
+
+		ToolbarBuilder.AddSeparator();
+	
+		// Refresh button
+		{
+			ToolbarBuilder.AddToolBarButton(
+				FUIAction(FExecuteAction::CreateLambda([Owner=Debugger]() { Owner->UpdateDebugHUDSettings(); })),
+				NAME_None,
+				LOCTEXT("Refresh", "Refresh"),
+				LOCTEXT("RefreshTooltip", "Refesh the settings on the target device.  Used if we get out of sync."),
+				FSlateIcon(FAppStyle::GetAppStyleSetName(), "Icons.Refresh")
+			);
 		}
-		// Additional options
-		ToolbarBuilder.AddComboButton(
-			FUIAction(),
-			FOnGetContent::CreateSP(this, &SNiagaraDebugger::MakePlaybackOptionsMenu),
-			FText(),
-			LOCTEXT("PlaybackOptionsTooltip", "Additional options to control playback."),
-			FSlateIcon(FAppStyle::GetAppStyleSetName(), "MaterialEditor.ToggleMaterialStats"),
-			true
-		);
 	}
 
 	ToolbarBuilder.EndSection();
@@ -1362,6 +1410,49 @@ TSharedRef<SWidget> SNiagaraDebugger::MakePlaybackOptionsMenu()
 	return MenuBuilder.MakeWidget();
 }
 
+TSharedRef<SWidget> SNiagaraDebugger::MakeModeOptionsMenu()
+{
+	FMenuBuilder MenuBuilder(/*bInShouldCloseWindowAfterMenuSelection=*/true, nullptr);
+	UNiagaraDebugHUDSettings* Settings = GetMutableDefault<UNiagaraDebugHUDSettings>();
+
+	{
+		MenuBuilder.AddMenuEntry(
+			LOCTEXT("EnableUIRendering", "Viewport UI"),
+			LOCTEXT("EnableUIRenderingTooltip", "When disabled, nothing will be drawn in the viewport. This is useful when you only want to use things like single frame stepping or data validation."),
+			FSlateIcon(),
+			FUIAction(
+				FExecuteAction::CreateLambda([=]() {Settings->Data.bHudRenderingEnabled = !Settings->Data.bHudRenderingEnabled; Settings->NotifyPropertyChanged(); }),
+				FCanExecuteAction(),
+				FIsActionChecked::CreateLambda([=]() { return Settings->Data.bHudRenderingEnabled; })
+			),
+			NAME_None, EUserInterfaceActionType::ToggleButton);
+
+		MenuBuilder.AddMenuEntry(
+			LOCTEXT("EnableOverviewUI", "Show Overview"),
+			LOCTEXT("EnableOverviewUITooltip", "When enabled, the debug overview of all active and filtered systems is drawn to the viewport."),
+			FSlateIcon(),
+			FUIAction(
+				FExecuteAction::CreateLambda([=]() {Settings->Data.bOverviewEnabled = !Settings->Data.bOverviewEnabled; Settings->NotifyPropertyChanged(); }),
+				FCanExecuteAction::CreateLambda([=]() { return Settings->Data.bHudRenderingEnabled; }),
+				FIsActionChecked::CreateLambda([=]() { return Settings->Data.bOverviewEnabled; })
+			),
+			NAME_None, EUserInterfaceActionType::ToggleButton);
+
+		MenuBuilder.AddMenuEntry(
+			LOCTEXT("EnableValidation", "Data Validation"),
+			LOCTEXT("EnableValidationTooltip", "Will allow you to validate system and particle data buffers when enabled."),
+			FSlateIcon(),
+			FUIAction(
+				FExecuteAction::CreateLambda([=]() {Settings->Data.bValidationEnabled = !Settings->Data.bValidationEnabled; Settings->NotifyPropertyChanged(); }),
+				FCanExecuteAction(),
+				FIsActionChecked::CreateLambda([=]() { return Settings->Data.bValidationEnabled; })
+			),
+			NAME_None, EUserInterfaceActionType::ToggleButton);
+	}
+
+	return MenuBuilder.MakeWidget();
+}
+
 #undef LOCTEXT_NAMESPACE
 
 #endif // WITH_NIAGARA_DEBUGGER
diff --git a/Engine/Plugins/FX/Niagara/Source/NiagaraEditor/Private/Widgets/SNiagaraDebugger.h b/Engine/Plugins/FX/Niagara/Source/NiagaraEditor/Private/Widgets/SNiagaraDebugger.h
--- a/Engine/Plugins/FX/Niagara/Source/NiagaraEditor/Private/Widgets/SNiagaraDebugger.h
+++ b/Engine/Plugins/FX/Niagara/Source/NiagaraEditor/Private/Widgets/SNiagaraDebugger.h
@@ -44,6 +44,7 @@ public:
 private:
 	TSharedRef<SWidget> MakeToolbar();
 	TSharedRef<SWidget> MakePlaybackOptionsMenu();
+	TSharedRef<SWidget> MakeModeOptionsMenu();
 
 protected:
 	TSharedPtr<FTabManager>		TabManager;
