Replace adhoc traversal of graph nodes with UNiagaraGraph::BuildTraversal() when deleting modules.  Avoids potential combinatorial explosion leading to a hang.
#tests Deleting a module with lots of dynamic inputs in the stack overview and measure the difference in time

diff --git a/Engine/Plugins/FX/Niagara/Source/NiagaraEditor/Private/ViewModels/Stack/NiagaraStackGraphUtilities.cpp b/Engine/Plugins/FX/Niagara/Source/NiagaraEditor/Private/ViewModels/Stack/NiagaraStackGraphUtilities.cpp
--- a/Engine/Plugins/FX/Niagara/Source/NiagaraEditor/Private/ViewModels/Stack/NiagaraStackGraphUtilities.cpp
+++ b/Engine/Plugins/FX/Niagara/Source/NiagaraEditor/Private/ViewModels/Stack/NiagaraStackGraphUtilities.cpp
@@ -2447,44 +2447,25 @@ bool FNiagaraStackGraphUtilities::RemoveModuleFromStack(UNiagaraScript& OwningSc
 
 	// Traverse all of the nodes in the group to find the nodes to remove.
 	FNiagaraStackGraphUtilities::FStackNodeGroup ModuleGroup = StackNodeGroups[ModuleStackIndex];
-	TArray<UNiagaraNode*> NodesToRemove;
-	TArray<UNiagaraNode*> NodesToCheck;
-	FPinCollectorArray InputPins;
-	NodesToCheck.Add(ModuleGroup.EndNode);
-	while (NodesToCheck.Num() > 0)
+	
+	// Remove the nodes in the group from the graph.
+	if (UNiagaraGraph* Graph = ModuleNode.GetNiagaraGraph())
 	{
-		UNiagaraNode* NodeToRemove = NodesToCheck[0];
-		NodesToCheck.RemoveAt(0);
-		NodesToRemove.AddUnique(NodeToRemove);
+		TArray<UNiagaraNode*> NodesToRemove;
+		Graph->BuildTraversal(NodesToRemove, ModuleGroup.EndNode);
 
-		InputPins.Reset();
-		NodeToRemove->GetInputPins(InputPins);
-		for (UEdGraphPin* InputPin : InputPins)
+		for (UNiagaraNode* NodeToRemove : NodesToRemove)
 		{
-			if (InputPin->LinkedTo.Num() == 1)
+			NodeToRemove->Modify();
+			Graph->RemoveNode(NodeToRemove);
+			UNiagaraNodeInput* InputNode = Cast<UNiagaraNodeInput>(NodeToRemove);
+			if (InputNode != nullptr)
 			{
-				UNiagaraNode* LinkedNode = Cast<UNiagaraNode>(InputPin->LinkedTo[0]->GetOwningNode());
-				if (LinkedNode != nullptr)
-				{
-					NodesToCheck.Add(LinkedNode);
-				}
+				OutRemovedInputNodes.Add(InputNode);
 			}
 		}
 	}
 
-	// Remove the nodes in the group from the graph.
-	UNiagaraGraph* Graph = ModuleNode.GetNiagaraGraph();
-	for (UNiagaraNode* NodeToRemove : NodesToRemove)
-	{
-		NodeToRemove->Modify();
-		Graph->RemoveNode(NodeToRemove);
-		UNiagaraNodeInput* InputNode = Cast<UNiagaraNodeInput>(NodeToRemove);
-		if (InputNode != nullptr)
-		{
-			OutRemovedInputNodes.Add(InputNode);
-		}
-	}
-
 	return true;
 }
 
