- Disable stateless debug draw when the module is disabled

#rnx

diff --git a/Engine/Plugins/FX/Niagara/Source/Niagara/Internal/Stateless/NiagaraStatelessModule.h b/Engine/Plugins/FX/Niagara/Source/Niagara/Internal/Stateless/NiagaraStatelessModule.h
--- a/Engine/Plugins/FX/Niagara/Source/Niagara/Internal/Stateless/NiagaraStatelessModule.h
+++ b/Engine/Plugins/FX/Niagara/Source/Niagara/Internal/Stateless/NiagaraStatelessModule.h
@@ -103,7 +103,7 @@ protected:
 public:
 	bool IsModuleEnabled() const { return bModuleEnabled; }
 #if WITH_EDITOR
-	bool IsDebugDrawEnabled() const { return bDebugDrawEnabled; }
+	bool IsDebugDrawEnabled() const { return bModuleEnabled && bDebugDrawEnabled; }
 
 	struct PrivateMemberNames
 	{
