Apply Scalability Instance counts more rigidly.
Do not allow new instances to activate until we are below the instance limit.

//TODO: Should we gate this behind a CVar or an FXType setting?
//TODO: Need to handle per emitter instance counts more rigidly also. This is harder as this stuff is all in script.
#tests win64 cooked, editor, pie

diff --git a/Engine/Plugins/FX/Niagara/Source/Niagara/Classes/NiagaraSystem.h b/Engine/Plugins/FX/Niagara/Source/Niagara/Classes/NiagaraSystem.h
--- a/Engine/Plugins/FX/Niagara/Source/Niagara/Classes/NiagaraSystem.h
+++ b/Engine/Plugins/FX/Niagara/Source/Niagara/Classes/NiagaraSystem.h
@@ -17,6 +17,7 @@
 #include "Particles/ParticleSystem.h"
 #include "UObject/Object.h"
 #include "UObject/ObjectMacros.h"
+#include "Particles/FXBudget.h"
 
 #if UE_ENABLE_INCLUDE_ORDER_DEPRECATED_IN_5_2
 #include "NiagaraBakerSettings.h"
@@ -743,6 +744,8 @@ public:
 	FORCEINLINE const FNiagaraSystemScalabilitySettings& GetScalabilitySettings()const { return CurrentScalabilitySettings; }
 	NIAGARA_API const FNiagaraSystemScalabilityOverride& GetCurrentOverrideSettings() const;
 	FORCEINLINE bool NeedsSortedSignificanceCull()const{ return bNeedsSortedSignificanceCull; }
+
+	FORCEINLINE void GetMaxInstanceCounts(int32& OutSystemInstanceMax, int32& OutFXTypeInstanceMax, bool bBudgetAdjusted)const;
 	
 	NIAGARA_API FNiagaraPlatformSet& GetScalabilityPlatformSet() { return Platforms; }
 	NIAGARA_API const FNiagaraPlatformSet& GetScalabilityPlatformSet() const { return Platforms; }
@@ -1083,3 +1086,26 @@ FORCEINLINE void UNiagaraSystem::UnregisterActiveInstance()
 {
 	--ActiveInstances;
 }
+
+
+FORCEINLINE void UNiagaraSystem::GetMaxInstanceCounts(int32& OutSystemInstanceMax, int32& OutFXTypeInstanceMax, bool bBudgetAdjusted)const 
+{
+	OutSystemInstanceMax = CurrentScalabilitySettings.MaxSystemInstances;
+	OutFXTypeInstanceMax = CurrentScalabilitySettings.MaxInstances;
+
+	if (bBudgetAdjusted && (CurrentScalabilitySettings.BudgetScaling.bScaleMaxInstanceCountByGlobalBudgetUse || CurrentScalabilitySettings.BudgetScaling.bScaleSystemInstanceCountByGlobalBudgetUse))
+	{
+		float Usage = FFXBudget::GetWorstAdjustedUsage();
+
+		if (CurrentScalabilitySettings.bCullMaxInstanceCount && CurrentScalabilitySettings.BudgetScaling.bScaleMaxInstanceCountByGlobalBudgetUse)
+		{
+			const float Scale = CurrentScalabilitySettings.BudgetScaling.MaxInstanceCountScaleByGlobalBudgetUse.Evaluate(Usage);
+			OutFXTypeInstanceMax = int32(float(OutFXTypeInstanceMax) * Scale);
+		}
+		if (CurrentScalabilitySettings.bCullPerSystemMaxInstanceCount && CurrentScalabilitySettings.BudgetScaling.bScaleSystemInstanceCountByGlobalBudgetUse)
+		{
+			const float Scale = CurrentScalabilitySettings.BudgetScaling.MaxSystemInstanceCountScaleByGlobalBudgetUse.Evaluate(Usage);
+			OutSystemInstanceMax = int32(float(OutSystemInstanceMax) * Scale);
+		}
+	}
+}
\ No newline at end of file
diff --git a/Engine/Plugins/FX/Niagara/Source/Niagara/Private/NiagaraScalabilityManager.cpp b/Engine/Plugins/FX/Niagara/Source/Niagara/Private/NiagaraScalabilityManager.cpp
--- a/Engine/Plugins/FX/Niagara/Source/Niagara/Private/NiagaraScalabilityManager.cpp
+++ b/Engine/Plugins/FX/Niagara/Source/Niagara/Private/NiagaraScalabilityManager.cpp
@@ -20,6 +20,9 @@ static FAutoConsoleVariableRef CVarScalabilityManParallelThreshold(TEXT("fx.Scal
 static int32 GScalabilityMaxUpdatesPerFrame = 50;
 static FAutoConsoleVariableRef CVarScalabilityMaxUpdatesPerFrame(TEXT("fx.ScalabilityMaxUpdatesPerFrame"), GScalabilityMaxUpdatesPerFrame, TEXT("Number of instances that can be processed per frame when updating scalability state. -1 for all of them. \n"), ECVF_Default);
 
+static int32 GApplyInstanceCountsRigidly = 1;
+static FAutoConsoleVariableRef CVarApplyInstanceCountsRigidly(TEXT("fx.Niagara.Scalability.ApplyInstanceCountsRigidly"), GApplyInstanceCountsRigidly, TEXT("If true we'll apply instance counts more rigidly. Un-culled Systems wil not be able to activate until overall systems fall below the instance count. \n"), ECVF_Default);
+
 static float GetScalabilityUpdatePeriod(ENiagaraScalabilityUpdateFrequency Frequency)
 {
 	switch (Frequency)
@@ -337,6 +340,8 @@ void FNiagaraScalabilityManager::ProcessSignificance(FNiagaraWorldManager* World
 	}
 }
 
+extern int32 GEnableNiagaraGlobalBudgetCulling;
+
 bool FNiagaraScalabilityManager::ApplyScalabilityState(int32 ComponentIndex, ENiagaraCullReaction CullReaction, bool bNewOnly)
 {
 	FNiagaraScalabilityState& CompState = State[ComponentIndex];
@@ -370,8 +375,6 @@ bool FNiagaraScalabilityManager::ApplyScalabilityState(int32 ComponentIndex, ENi
 			}
 		}
 
-		CompState.Apply();
-
 #if WITH_NIAGARA_DEBUGGER
 		//Tell the debugger about our scalability state.
 		//Unfortunately cannot have the debugger just read the manager state data as components are removed.
@@ -381,6 +384,8 @@ bool FNiagaraScalabilityManager::ApplyScalabilityState(int32 ComponentIndex, ENi
 
 		if (CompState.bCulled)
 		{
+			CompState.Apply();
+
 			switch (CullReaction)
 			{
 			case ENiagaraCullReaction::Deactivate:					Component->DeactivateInternal(false); bContinueIteration = false; break;//We don't increment CompIdx here as this call will remove an entry from ManagedObjects;
@@ -392,17 +397,33 @@ bool FNiagaraScalabilityManager::ApplyScalabilityState(int32 ComponentIndex, ENi
 		}
 		else
 		{
-			if (CullReaction == ENiagaraCullReaction::Deactivate || CullReaction == ENiagaraCullReaction::DeactivateImmediate)
-			{
-				UE_LOG(LogNiagara, Error, TEXT("Niagara Component is incorrectly still registered with the scalability manager. %d - %s "), (int32)CullReaction, *Component->GetAsset()->GetFullName());
-			}
-			else if(CullReaction == ENiagaraCullReaction::PauseResume)
-			{
-				Component->SetPausedInternal(false, true);
-			}
-			else
+			UNiagaraSystem* System = Component->GetAsset();
+			check(System);
+			
+			int32 SystemInstanceMax = 0;
+			int32 EffectTypeInstanceMax = 0;
+
+			bool bBudgetCullEnabled = GEnableNiagaraGlobalBudgetCulling && FFXBudget::Enabled() && INiagaraModule::UseGlobalFXBudget();
+			Component->GetAsset()->GetMaxInstanceCounts(SystemInstanceMax, EffectTypeInstanceMax, bBudgetCullEnabled);
+
+			//We do not allow new activations if it would blow our instance count limits.
+			//Though we do allow currently running but inactive components to (!IsComplete()) to reactivate
+			bool bAllow = GApplyInstanceCountsRigidly == 0 || (Component->IsComplete() == false || (System->GetActiveInstancesCount() < SystemInstanceMax && EffectType->NumInstances < EffectTypeInstanceMax));
+			if(bAllow)
 			{
-				Component->ActivateInternal(false, true);
+				CompState.Apply();
+				if (CullReaction == ENiagaraCullReaction::Deactivate || CullReaction == ENiagaraCullReaction::DeactivateImmediate)
+				{
+					UE_LOG(LogNiagara, Error, TEXT("Niagara Component is incorrectly still registered with the scalability manager. %d - %s "), (int32)CullReaction, *Component->GetAsset()->GetFullName());
+				}
+				else if (CullReaction == ENiagaraCullReaction::PauseResume)
+				{
+					Component->SetPausedInternal(false, true);
+				}
+				else
+				{
+					Component->ActivateInternal(false, true);
+				}
 			}
 		}
 
@@ -441,7 +462,8 @@ void FNiagaraScalabilityManager::UpdateInternal(FNiagaraWorldManager* WorldMan,
 
 	if (Context.bProcessAllComponents || !Context.ComponentRequiresUpdate.Contains(true))
 	{
-		if (Context.bRequiresGlobalSignificancePass && EffectType->SignificanceHandler)
+		bool bProcessSignificance = Context.bRequiresGlobalSignificancePass && EffectType->SignificanceHandler;
+		if (bProcessSignificance)
 		{
 			ProcessSignificance(WorldMan, EffectType->SignificanceHandler, Context);
 		}
@@ -450,13 +472,35 @@ void FNiagaraScalabilityManager::UpdateInternal(FNiagaraWorldManager* WorldMan,
 		{
 			const ENiagaraCullReaction CullReaction = EffectType->CullReaction;
 
-			int32 CompIdx = 0;
 			//As we'll be activating and deactivating here, this must be done on the game thread.
-			while (CompIdx < ManagedComponents.Num())
+			
+			//Do one pass to apply in significance order so that any re-activations are done on the most significant items first.
+			int32 CompIdx = 0;
+			if(bProcessSignificance && GApplyInstanceCountsRigidly != 0)
+			{
+				while (CompIdx < Context.SignificanceIndices.Num())
+				{
+					int32 SortedIdx = Context.SignificanceIndices[CompIdx++];
+					if(!ApplyScalabilityState(SortedIdx, CullReaction, Context.bNewOnly))
+					{
+						//If we deactivate a component then we'll have to bail on this pass as the sorted index list is no longer valid. 
+						//By this point we should have processed all the activations we want in order anyway.
+						break;
+					}
+				}
+			}
+
+			//We have to do another pass if we've not processed everything yet. 
+			//Might want to move to having a separate index list for dirty components
+			if(CompIdx < ManagedComponents.Num())
 			{
-				if (ApplyScalabilityState(CompIdx, CullReaction, Context.bNewOnly))
+				CompIdx = 0;
+				while (CompIdx < ManagedComponents.Num())
 				{
-					++CompIdx;
+					if (ApplyScalabilityState(CompIdx, CullReaction, Context.bNewOnly))
+					{
+						++CompIdx;
+					}
 				}
 			}
 
diff --git a/Engine/Plugins/FX/Niagara/Source/Niagara/Private/NiagaraWorldManager.cpp b/Engine/Plugins/FX/Niagara/Source/Niagara/Private/NiagaraWorldManager.cpp
--- a/Engine/Plugins/FX/Niagara/Source/Niagara/Private/NiagaraWorldManager.cpp
+++ b/Engine/Plugins/FX/Niagara/Source/Niagara/Private/NiagaraWorldManager.cpp
@@ -137,7 +137,7 @@ static FAutoConsoleVariableRef CVarEnableNiagaraInstanceCountCulling(
 	ECVF_Default
 );
 
-static int GEnableNiagaraGlobalBudgetCulling = 1;
+int GEnableNiagaraGlobalBudgetCulling = 1;
 static FAutoConsoleVariableRef CVarEnableNiagaraGlobalBudgetCulling(
 	TEXT("fx.Niagara.Scalability.GlobalBudgetCulling"),
 	GEnableNiagaraGlobalBudgetCulling,
@@ -1753,8 +1753,12 @@ void FNiagaraWorldManager::SortedSignificanceCull(UNiagaraEffectType* EffectType
 	
 	if(GetScalabilityCullingMode() == ENiagaraScalabilityCullingMode::Enabled && GEnableNiagaraInstanceCountCulling)
 	{
-		int32 SystemInstanceMax = ScalabilitySettings.MaxSystemInstances;
-		int32 EffectTypeInstanceMax = ScalabilitySettings.MaxInstances;
+		UNiagaraSystem* System = Component->GetAsset();
+		check(System);
+
+		int32 SystemInstanceMax = 0;
+		int32 EffectTypeInstanceMax = 0;
+		System->GetMaxInstanceCounts(SystemInstanceMax, EffectTypeInstanceMax, false);
 
 		bCull = ScalabilitySettings.bCullMaxInstanceCount && EffectTypeInstCount >= EffectTypeInstanceMax;
 		bCull |= ScalabilitySettings.bCullPerSystemMaxInstanceCount && SystemInstCount >= SystemInstanceMax;
@@ -1767,22 +1771,13 @@ void FNiagaraWorldManager::SortedSignificanceCull(UNiagaraEffectType* EffectType
 			OutState.bCulledByGlobalBudget = false;
 #endif
 		}
-		else if (bBudgetCullEnabled && ScalabilitySettings.BudgetScaling.bCullByGlobalBudget)
-	 	{
-			float Usage = FFXBudget::GetWorstAdjustedUsage();
+		else if (bBudgetCullEnabled)
+		{
+			//Grab Adjusted instance counts.
+			System->GetMaxInstanceCounts(SystemInstanceMax, EffectTypeInstanceMax, bBudgetCullEnabled);
 
-			if (ScalabilitySettings.bCullMaxInstanceCount && ScalabilitySettings.BudgetScaling.bScaleMaxInstanceCountByGlobalBudgetUse)
-			{
-				const float Scale = ScalabilitySettings.BudgetScaling.MaxInstanceCountScaleByGlobalBudgetUse.Evaluate(Usage);
-				EffectTypeInstanceMax = int32(float(EffectTypeInstanceMax) * Scale);
-				bCull = EffectTypeInstCount >= EffectTypeInstanceMax;
-			}
-			if (ScalabilitySettings.bCullPerSystemMaxInstanceCount && ScalabilitySettings.BudgetScaling.bScaleSystemInstanceCountByGlobalBudgetUse)
-			{
-				const float Scale = ScalabilitySettings.BudgetScaling.MaxSystemInstanceCountScaleByGlobalBudgetUse.Evaluate(Usage);
-				SystemInstanceMax = int32(float(SystemInstanceMax) * Scale);
-				bCull |= SystemInstCount >= SystemInstanceMax;
-			}
+			bCull = EffectTypeInstCount >= EffectTypeInstanceMax;
+			bCull |= SystemInstCount >= SystemInstanceMax;
 
 #if DEBUG_SCALABILITY_STATE
 			OutState.bCulledByGlobalBudget |= bCull;
@@ -1793,6 +1788,8 @@ void FNiagaraWorldManager::SortedSignificanceCull(UNiagaraEffectType* EffectType
 	OutState.bCulled |= bCull;
 
 	//Only increment the instance counts if this is not culled. Including other causes of culling.
+	//This considers only culled state and will only affect which components we *want* to turn on.
+	//There is another check inside the apply that will limit instance counts to ensure they never exceed the max that are actually running (or winding down)
 	if(OutState.bCulled == false)
 	{
 		++EffectTypeInstCount;
