- Add dispatch index wrapping for higher particle counts

#rnx

diff --git a/Engine/Plugins/FX/Niagara/Shaders/Private/Stateless/NiagaraStatelessSimulationTemplate.ush b/Engine/Plugins/FX/Niagara/Shaders/Private/Stateless/NiagaraStatelessSimulationTemplate.ush
--- a/Engine/Plugins/FX/Niagara/Shaders/Private/Stateless/NiagaraStatelessSimulationTemplate.ush
+++ b/Engine/Plugins/FX/Niagara/Shaders/Private/Stateless/NiagaraStatelessSimulationTemplate.ush
@@ -1,6 +1,7 @@
 // Copyright Epic Games, Inc. All Rights Reserved.
 
 #include "NiagaraStatelessCommon.ush"
+#include "/Engine/Private/ComputeShaderUtils.ush"
 //#include "WaveOpUtil.ush"
 
 /////////////////////////////////////////////////////////////////////////////////////////////////////////////////////////
@@ -54,15 +55,18 @@ DECLARE_SCALAR_ARRAY(float,	SpawnInfo_Rate, NIAGARA_MAX_GPU_SPAWN_INFOS);
 /////////////////////////////////////////////////////////////////////////////////////////////////////////////////////////
 
 [numthreads(THREADGROUP_SIZE, 1, 1)]
-void StatelessMain(uint DispatchThreadId : SV_DispatchThreadID)
+void StatelessMain(uint3 GroupId : SV_GroupID, int GroupThreadIndex : SV_GroupIndex)
 {
+	// Convert thread ID into linear thread
+	uint LinearThreadId = GetUnWrappedDispatchThreadId(GroupId, GroupThreadIndex, THREADGROUP_SIZE);
+
 	// Initialize our particle data
 	FStatelessParticle Particle;
 	Particle.Age = -1.0f;
 
 	// Find which spawn info we belong to and fill out details
 	{
-		uint SpawnInfoIndex = DispatchThreadId.x;
+		uint SpawnInfoIndex = LinearThreadId;
 		for ( int i=0; i < NIAGARA_MAX_GPU_SPAWN_INFOS; ++i )
 		{
 			const uint SpawnInfoNumActive		= GET_SCALAR_ARRAY_ELEMENT(SpawnInfo_NumActive, i);
diff --git a/Engine/Plugins/FX/Niagara/Source/Niagara/Private/Stateless/NiagaraStatelessComputeManager.cpp b/Engine/Plugins/FX/Niagara/Source/Niagara/Private/Stateless/NiagaraStatelessComputeManager.cpp
--- a/Engine/Plugins/FX/Niagara/Source/Niagara/Private/Stateless/NiagaraStatelessComputeManager.cpp
+++ b/Engine/Plugins/FX/Niagara/Source/Niagara/Private/Stateless/NiagaraStatelessComputeManager.cpp
@@ -238,10 +238,8 @@ void FNiagaraStatelessComputeManager::OnPostPreRender(FRDGBuilder& GraphBuilder)
 				FRHIComputeShader* ComputeShaderRHI = ComputeShader.GetComputeShader();
 				const uint32 NumThreadGroups = FMath::DivideAndRoundUp<uint32>(CacheData->ActiveParticles, NiagaraStateless::FSimulationShader::ThreadGroupSize);
 
-				SetComputePipelineState(RHICmdList, ComputeShaderRHI);
-				SetShaderParameters(RHICmdList, ComputeShader, ComputeShaderRHI, EmitterData->GetShaderParametersMetadata(), *ShaderParameters);
-				RHICmdList.DispatchComputeShader(NumThreadGroups, 1, 1);
-				UnsetShaderUAVs(RHICmdList, ComputeShader, ComputeShaderRHI);
+				const FIntVector NumWrappedThreadGroups = FComputeShaderUtils::GetGroupCountWrapped(NumThreadGroups);
+				FComputeShaderUtils::Dispatch(RHICmdList, ComputeShader, EmitterData->GetShaderParametersMetadata(), *ShaderParameters, NumWrappedThreadGroups);
 			}
 			RHICmdList.EndUAVOverlap(CountBufferUAV);
 
