- Scalability view fixes for stateless emitters

#rnx
#tests Created several stateless emitters with varying scalability and tested in preview & cooked

diff --git a/Engine/Plugins/FX/Niagara/Source/Niagara/Classes/NiagaraEmitterHandle.h b/Engine/Plugins/FX/Niagara/Source/Niagara/Classes/NiagaraEmitterHandle.h
--- a/Engine/Plugins/FX/Niagara/Source/Niagara/Classes/NiagaraEmitterHandle.h
+++ b/Engine/Plugins/FX/Niagara/Source/Niagara/Classes/NiagaraEmitterHandle.h
@@ -68,6 +68,9 @@ public:
 	/** Sets whether or not this emitter is enabled within the System.  Disabled emitters aren't simulated. Returns whether or not the enabled state changed. */
 	NIAGARA_API bool SetIsEnabled(bool bInIsEnabled, UNiagaraSystem& InOwnerSystem, bool bRecompileIfChanged);
 
+	/** Test to see if the emitter is allowed by scalability or not */
+	NIAGARA_API bool IsAllowedByScalability() const;
+
 #if WITH_EDITORONLY_DATA
 	bool IsIsolated() const {	return bIsolated; }
 	void SetIsolated(bool bInIsolated) { bIsolated = bInIsolated; }
diff --git a/Engine/Plugins/FX/Niagara/Source/Niagara/Internal/Stateless/NiagaraStatelessEmitter.h b/Engine/Plugins/FX/Niagara/Source/Niagara/Internal/Stateless/NiagaraStatelessEmitter.h
--- a/Engine/Plugins/FX/Niagara/Source/Niagara/Internal/Stateless/NiagaraStatelessEmitter.h
+++ b/Engine/Plugins/FX/Niagara/Source/Niagara/Internal/Stateless/NiagaraStatelessEmitter.h
@@ -59,7 +59,19 @@ public:
 	{
 		for (UNiagaraRendererProperties* Renderer : RendererProperties)
 		{
-			if (Renderer && Renderer->GetIsEnabled() && Renderer->IsSimTargetSupported(ENiagaraSimTarget::GPUComputeSim))
+			if (Renderer && Renderer->GetIsEnabled())
+			{
+				Func(Renderer);
+			}
+		}
+	}
+
+	template<typename TAction>
+	void ForEachRenderer(TAction Func) const
+	{
+		for (UNiagaraRendererProperties* Renderer : RendererProperties)
+		{
+			if (Renderer)
 			{
 				Func(Renderer);
 			}
@@ -72,6 +84,8 @@ public:
 	FNiagaraStatelessEmitterDataPtr GetEmitterData() const { return StatelessEmitterData; }
 	NiagaraStateless::FCommonShaderParameters* AllocateShaderParameters(const FNiagaraParameterStore& RendererBindings) const;
 
+	NIAGARA_API bool IsAllowedByScalability() const;
+
 #if WITH_EDITOR
 	NIAGARA_API void SetEmitterTemplateClass(UClass* TemplateClass);
 
@@ -95,6 +109,8 @@ public:
 
 	NIAGARA_API const TArray<TObjectPtr<UNiagaraStatelessModule>>& GetModules() const { return Modules; }
 
+	NIAGARA_API FNiagaraPlatformSet& GetPlatformSet() { return Platforms; }
+
 	UNiagaraStatelessEmitter* CreateAsDuplicate(FName InDuplicateName, UNiagaraSystem& InDuplicateOwnerSystem) const;
 
 	NIAGARA_API void DrawModuleDebug(UWorld* World, const FTransform& LocalToWorld) const;
diff --git a/Engine/Plugins/FX/Niagara/Source/Niagara/Private/NiagaraEmitterHandle.cpp b/Engine/Plugins/FX/Niagara/Source/Niagara/Private/NiagaraEmitterHandle.cpp
--- a/Engine/Plugins/FX/Niagara/Source/Niagara/Private/NiagaraEmitterHandle.cpp
+++ b/Engine/Plugins/FX/Niagara/Source/Niagara/Private/NiagaraEmitterHandle.cpp
@@ -168,6 +168,19 @@ bool FNiagaraEmitterHandle::SetIsEnabled(bool bInIsEnabled, UNiagaraSystem& InOw
 	return false;
 }
 
+bool FNiagaraEmitterHandle::IsAllowedByScalability() const
+{
+	if (EmitterMode == ENiagaraEmitterMode::Standard)
+	{
+		FVersionedNiagaraEmitterData* EmitterData = VersionedInstance.GetEmitterData();
+		return EmitterData ? EmitterData->IsAllowedByScalability() : false;
+	}
+	else
+	{
+		return StatelessEmitter ? StatelessEmitter->IsAllowedByScalability() : false;
+	}
+}
+
 #if WITH_EDITORONLY_DATA
 void FNiagaraEmitterHandle::SetEmitterMode(UNiagaraSystem& InOwningSystem, ENiagaraEmitterMode InEmitterMode)
 {
diff --git a/Engine/Plugins/FX/Niagara/Source/Niagara/Private/Stateless/NiagaraStatelessEmitter.cpp b/Engine/Plugins/FX/Niagara/Source/Niagara/Private/Stateless/NiagaraStatelessEmitter.cpp
--- a/Engine/Plugins/FX/Niagara/Source/Niagara/Private/Stateless/NiagaraStatelessEmitter.cpp
+++ b/Engine/Plugins/FX/Niagara/Source/Niagara/Private/Stateless/NiagaraStatelessEmitter.cpp
@@ -257,7 +257,7 @@ void UNiagaraStatelessEmitter::CacheFromCompiledData()
 	}
 
 	StatelessEmitterData->bCanEverExecute = NiagaraStatelessInternal::IsValid(*StatelessEmitterData);
-	StatelessEmitterData->bCanEverExecute &= Platforms.IsActive();
+	StatelessEmitterData->bCanEverExecute &= IsAllowedByScalability();
 
 	// Determine our supported feature set mask
 	// If we can not execute on any enabled units then the emitter is considered disabled
@@ -457,6 +457,11 @@ NiagaraStateless::FCommonShaderParameters* UNiagaraStatelessEmitter::AllocateSha
 	return CommonParameters;
 }
 
+bool UNiagaraStatelessEmitter::IsAllowedByScalability() const
+{
+	return Platforms.IsActive();
+}
+
 #if WITH_EDITOR
 void UNiagaraStatelessEmitter::SetEmitterTemplateClass(UClass* TemplateClass)
 {
diff --git a/Engine/Plugins/FX/Niagara/Source/NiagaraEditor/Private/Sequencer/NiagaraSequence/MovieSceneNiagaraEmitterTrack.cpp b/Engine/Plugins/FX/Niagara/Source/NiagaraEditor/Private/Sequencer/NiagaraSequence/MovieSceneNiagaraEmitterTrack.cpp
--- a/Engine/Plugins/FX/Niagara/Source/NiagaraEditor/Private/Sequencer/NiagaraSequence/MovieSceneNiagaraEmitterTrack.cpp
+++ b/Engine/Plugins/FX/Niagara/Source/NiagaraEditor/Private/Sequencer/NiagaraSequence/MovieSceneNiagaraEmitterTrack.cpp
@@ -181,14 +181,15 @@ void UMovieSceneNiagaraEmitterTrack::Tick(float DeltaTime)
 {
 	if(!HasAnyFlags(RF_ClassDefaultObject) && bScalabilityModeActive)
 	{
-		FVersionedNiagaraEmitterData* EmitterData = nullptr;
-		// the emitter instance might no longer be valid if we deleted the emitter but the track still lives until garbage collection
-		if(EmitterHandleViewModel.IsValid() && EmitterHandleViewModel.Pin()->GetEmitterHandle() != nullptr)
+		bool bIsAllowedByScalability = false;
+
+		if (TSharedPtr<FNiagaraEmitterHandleViewModel> EmitterViewModel = EmitterHandleViewModel.Pin())
 		{
-			EmitterData = EmitterHandleViewModel.Pin()->GetEmitterHandle()->GetEmitterData();
+			FNiagaraEmitterHandle* EmitterHandle = EmitterViewModel->GetEmitterHandle();
+			bIsAllowedByScalability = EmitterHandle ? EmitterHandle->IsAllowedByScalability() : false;
 		}
-		
-		if(EmitterData && !EmitterData->IsAllowedByScalability())
+
+		if (bIsAllowedByScalability)
 		{
 			SetColorTint(FNiagaraEditorStyle::Get().GetColor("NiagaraEditor.SystemOverview.ExcludedFromScalability").ToFColor(true));
 		}
diff --git a/Engine/Plugins/FX/Niagara/Source/NiagaraEditor/Private/Toolkits/SystemToolkitModes/NiagaraSystemToolkitMode_Scalability.cpp b/Engine/Plugins/FX/Niagara/Source/NiagaraEditor/Private/Toolkits/SystemToolkitModes/NiagaraSystemToolkitMode_Scalability.cpp
--- a/Engine/Plugins/FX/Niagara/Source/NiagaraEditor/Private/Toolkits/SystemToolkitModes/NiagaraSystemToolkitMode_Scalability.cpp
+++ b/Engine/Plugins/FX/Niagara/Source/NiagaraEditor/Private/Toolkits/SystemToolkitModes/NiagaraSystemToolkitMode_Scalability.cpp
@@ -12,6 +12,7 @@
 #include "ViewModels/NiagaraEmitterHandleViewModel.h"
 #include "Widgets/Scalability/SNiagaraScalabilityContext.h"
 #include "Framework/MultiBox/MultiBoxBuilder.h"
+#include "Stateless/NiagaraStatelessEmitter.h"
 
 #define LOCTEXT_NAMESPACE "NiagaraSystemToolkitMode_Scalability"
 
@@ -88,8 +89,16 @@ void FNiagaraSystemToolkitMode_Scalability::UnbindEmitterUpdates()
 
 	for(const TSharedRef<FNiagaraEmitterHandleViewModel>& EmitterHandleViewModel : SystemViewModel->GetEmitterHandleViewModels())
 	{
-		EmitterHandleViewModel->GetEmitterHandle()->GetInstance().Emitter->OnRenderersChanged().RemoveAll(this);
-		EmitterHandleViewModel->GetEmitterHandle()->GetInstance().Emitter->OnPropertiesChanged().RemoveAll(this);
+		FNiagaraEmitterHandle* EmitterHandle = EmitterHandleViewModel->GetEmitterHandle();
+		if (EmitterHandle->GetEmitterMode() == ENiagaraEmitterMode::Standard)
+		{
+			EmitterHandle->GetInstance().Emitter->OnRenderersChanged().RemoveAll(this);
+			EmitterHandle->GetInstance().Emitter->OnPropertiesChanged().RemoveAll(this);
+		}
+		else if (EmitterHandle->GetEmitterMode() == ENiagaraEmitterMode::Stateless)
+		{
+			EmitterHandle->GetStatelessEmitter()->OnRenderersChanged().RemoveAll(this);
+		}
 	}
 }
 
@@ -100,24 +109,44 @@ void FNiagaraSystemToolkitMode_Scalability::BindEmitterPreviewOverrides()
 	
 	for(const TSharedRef<FNiagaraEmitterHandleViewModel>& EmitterHandleViewModel : SystemViewModel->GetEmitterHandleViewModels())
 	{
-		FVersionedNiagaraEmitterData* EmitterData = EmitterHandleViewModel->GetEmitterHandle()->GetEmitterData();
-		// we bind the quality level and device profile of the emitter itself to use our scalability preview
-		EmitterData->Platforms.OnOverrideQualityLevelDelegate.BindUObject(ScalabilityViewModel, &UNiagaraSystemScalabilityViewModel::GetPreviewQualityLevel);
-		EmitterData->Platforms.OnOverrideActiveDeviceProfileDelegate.BindUObject(ScalabilityViewModel, &UNiagaraSystemScalabilityViewModel::GetPreviewDeviceProfile);
-
-		// we do the same for the overrides
-		for(FNiagaraEmitterScalabilityOverride& EmitterScalabilityOverride : EmitterData->ScalabilityOverrides.Overrides)
+		FNiagaraEmitterHandle* EmitterHandle = EmitterHandleViewModel->GetEmitterHandle();
+		if (EmitterHandle->GetEmitterMode() == ENiagaraEmitterMode::Standard)
 		{
-			EmitterScalabilityOverride.Platforms.OnOverrideQualityLevelDelegate.BindUObject(ScalabilityViewModel, &UNiagaraSystemScalabilityViewModel::GetPreviewQualityLevel);
-			EmitterScalabilityOverride.Platforms.OnOverrideActiveDeviceProfileDelegate.BindUObject(ScalabilityViewModel, &UNiagaraSystemScalabilityViewModel::GetPreviewDeviceProfile);
-		}
+			FVersionedNiagaraEmitterData* EmitterData = EmitterHandle->GetEmitterData();
+
+			// we bind the quality level and device profile of the emitter itself to use our scalability preview
+			EmitterData->Platforms.OnOverrideQualityLevelDelegate.BindUObject(ScalabilityViewModel, &UNiagaraSystemScalabilityViewModel::GetPreviewQualityLevel);
+			EmitterData->Platforms.OnOverrideActiveDeviceProfileDelegate.BindUObject(ScalabilityViewModel, &UNiagaraSystemScalabilityViewModel::GetPreviewDeviceProfile);
+
+			// we do the same for the overrides
+			for (FNiagaraEmitterScalabilityOverride& EmitterScalabilityOverride : EmitterData->ScalabilityOverrides.Overrides)
+			{
+				EmitterScalabilityOverride.Platforms.OnOverrideQualityLevelDelegate.BindUObject(ScalabilityViewModel, &UNiagaraSystemScalabilityViewModel::GetPreviewQualityLevel);
+				EmitterScalabilityOverride.Platforms.OnOverrideActiveDeviceProfileDelegate.BindUObject(ScalabilityViewModel, &UNiagaraSystemScalabilityViewModel::GetPreviewDeviceProfile);
+			}
 
-		// and also for the renderers
-		EmitterData->ForEachRenderer([=](UNiagaraRendererProperties* RendererProperties)
+			// and also for the renderers
+			EmitterData->ForEachRenderer([=](UNiagaraRendererProperties* RendererProperties)
+			{
+				RendererProperties->Platforms.OnOverrideQualityLevelDelegate.BindUObject(ScalabilityViewModel, &UNiagaraSystemScalabilityViewModel::GetPreviewQualityLevel);
+				RendererProperties->Platforms.OnOverrideActiveDeviceProfileDelegate.BindUObject(ScalabilityViewModel, &UNiagaraSystemScalabilityViewModel::GetPreviewDeviceProfile);
+			});
+		}
+		else if (EmitterHandle->GetEmitterMode() == ENiagaraEmitterMode::Stateless)
 		{
-			RendererProperties->Platforms.OnOverrideQualityLevelDelegate.BindUObject(ScalabilityViewModel, &UNiagaraSystemScalabilityViewModel::GetPreviewQualityLevel);
-			RendererProperties->Platforms.OnOverrideActiveDeviceProfileDelegate.BindUObject(ScalabilityViewModel, &UNiagaraSystemScalabilityViewModel::GetPreviewDeviceProfile);
-		});
+			UNiagaraStatelessEmitter* StatelessEmitter = EmitterHandle->GetStatelessEmitter();
+
+			// we bind the quality level and device profile of the emitter itself to use our scalability preview
+			StatelessEmitter->GetPlatformSet().OnOverrideQualityLevelDelegate.BindUObject(ScalabilityViewModel, &UNiagaraSystemScalabilityViewModel::GetPreviewQualityLevel);
+			StatelessEmitter->GetPlatformSet().OnOverrideActiveDeviceProfileDelegate.BindUObject(ScalabilityViewModel, &UNiagaraSystemScalabilityViewModel::GetPreviewDeviceProfile);
+
+			// and also for the renderers
+			StatelessEmitter->ForEachRenderer([=](UNiagaraRendererProperties* RendererProperties)
+			{
+				RendererProperties->Platforms.OnOverrideQualityLevelDelegate.BindUObject(ScalabilityViewModel, &UNiagaraSystemScalabilityViewModel::GetPreviewQualityLevel);
+				RendererProperties->Platforms.OnOverrideActiveDeviceProfileDelegate.BindUObject(ScalabilityViewModel, &UNiagaraSystemScalabilityViewModel::GetPreviewDeviceProfile);
+			});
+		}
 	}
 }
 
@@ -141,21 +170,38 @@ void FNiagaraSystemToolkitMode_Scalability::UnbindEmitterPreviewOverrides()
 	// we undo all the bindings we created in the Bind function here 
 	for(const TSharedRef<FNiagaraEmitterHandleViewModel>& EmitterHandleViewModel : SystemViewModel->GetEmitterHandleViewModels())
 	{
-		FVersionedNiagaraEmitterData* EmitterData = EmitterHandleViewModel->GetEmitterHandle()->GetEmitterData();
-		EmitterData->Platforms.OnOverrideQualityLevelDelegate.Unbind();
-		EmitterData->Platforms.OnOverrideActiveDeviceProfileDelegate.Unbind();
-		
-		for(FNiagaraEmitterScalabilityOverride& EmitterScalabilityOverride : EmitterData->ScalabilityOverrides.Overrides)
+		FNiagaraEmitterHandle* EmitterHandle = EmitterHandleViewModel->GetEmitterHandle();
+		if (EmitterHandle->GetEmitterMode() == ENiagaraEmitterMode::Standard)
 		{
-			EmitterScalabilityOverride.Platforms.OnOverrideQualityLevelDelegate.Unbind();
-			EmitterScalabilityOverride.Platforms.OnOverrideActiveDeviceProfileDelegate.Unbind();
-		}
+			FVersionedNiagaraEmitterData* EmitterData = EmitterHandleViewModel->GetEmitterHandle()->GetEmitterData();
+			EmitterData->Platforms.OnOverrideQualityLevelDelegate.Unbind();
+			EmitterData->Platforms.OnOverrideActiveDeviceProfileDelegate.Unbind();
+		
+			for(FNiagaraEmitterScalabilityOverride& EmitterScalabilityOverride : EmitterData->ScalabilityOverrides.Overrides)
+			{
+				EmitterScalabilityOverride.Platforms.OnOverrideQualityLevelDelegate.Unbind();
+				EmitterScalabilityOverride.Platforms.OnOverrideActiveDeviceProfileDelegate.Unbind();
+			}
 
-		EmitterData->ForEachRenderer([](UNiagaraRendererProperties* RendererProperties)
+			EmitterData->ForEachRenderer([](UNiagaraRendererProperties* RendererProperties)
+			{
+				RendererProperties->Platforms.OnOverrideQualityLevelDelegate.Unbind();
+				RendererProperties->Platforms.OnOverrideActiveDeviceProfileDelegate.Unbind();
+			});
+		}
+		else if (EmitterHandle->GetEmitterMode() == ENiagaraEmitterMode::Stateless)
 		{
-			RendererProperties->Platforms.OnOverrideQualityLevelDelegate.Unbind();
-			RendererProperties->Platforms.OnOverrideActiveDeviceProfileDelegate.Unbind();
-		});
+			UNiagaraStatelessEmitter* StatelessEmitter = EmitterHandle->GetStatelessEmitter();
+
+			StatelessEmitter->GetPlatformSet().OnOverrideQualityLevelDelegate.Unbind();
+			StatelessEmitter->GetPlatformSet().OnOverrideActiveDeviceProfileDelegate.Unbind();
+
+			StatelessEmitter->ForEachRenderer([](UNiagaraRendererProperties* RendererProperties)
+			{
+				RendererProperties->Platforms.OnOverrideQualityLevelDelegate.Unbind();
+				RendererProperties->Platforms.OnOverrideActiveDeviceProfileDelegate.Unbind();
+			});
+		}
 	}
 }
 
@@ -179,9 +225,17 @@ void FNiagaraSystemToolkitMode_Scalability::BindEmitterUpdates()
 	// if our emitter changes properties (= platforms) or a renderer changes (= platforms), we refresh the bindings to make sure all new values get properly overridden
 	TSharedPtr<FNiagaraSystemViewModel> SystemViewModel = SystemToolkit.Pin()->GetSystemViewModel();
 	for(const TSharedRef<FNiagaraEmitterHandleViewModel>& EmitterHandleViewModel : SystemViewModel->GetEmitterHandleViewModels())
-	{		
-		EmitterHandleViewModel->GetEmitterHandle()->GetInstance().Emitter->OnRenderersChanged().AddSP(this, &FNiagaraSystemToolkitMode_Scalability::BindEmitterPreviewOverrides);
-		EmitterHandleViewModel->GetEmitterHandle()->GetInstance().Emitter->OnPropertiesChanged().AddSP(this, &FNiagaraSystemToolkitMode_Scalability::BindEmitterPreviewOverrides);
+	{
+		FNiagaraEmitterHandle* EmitterHandle = EmitterHandleViewModel->GetEmitterHandle();
+		if (EmitterHandle->GetEmitterMode() == ENiagaraEmitterMode::Standard)
+		{
+			EmitterHandle->GetInstance().Emitter->OnRenderersChanged().AddSP(this, &FNiagaraSystemToolkitMode_Scalability::BindEmitterPreviewOverrides);
+			EmitterHandle->GetInstance().Emitter->OnPropertiesChanged().AddSP(this, &FNiagaraSystemToolkitMode_Scalability::BindEmitterPreviewOverrides);
+		}
+		else if ( EmitterHandle->GetEmitterMode() == ENiagaraEmitterMode::Stateless )
+		{
+			EmitterHandle->GetStatelessEmitter()->OnRenderersChanged().AddSP(this, &FNiagaraSystemToolkitMode_Scalability::BindEmitterPreviewOverrides);
+		}
 	}
 
 	BindEmitterPreviewOverrides();
diff --git a/Engine/Plugins/FX/Niagara/Source/NiagaraEditor/Private/Widgets/Scalability/SNiagaraScalabilityContext.cpp b/Engine/Plugins/FX/Niagara/Source/NiagaraEditor/Private/Widgets/Scalability/SNiagaraScalabilityContext.cpp
--- a/Engine/Plugins/FX/Niagara/Source/NiagaraEditor/Private/Widgets/Scalability/SNiagaraScalabilityContext.cpp
+++ b/Engine/Plugins/FX/Niagara/Source/NiagaraEditor/Private/Widgets/Scalability/SNiagaraScalabilityContext.cpp
@@ -7,6 +7,7 @@
 #include "NiagaraMeshRendererProperties.h"
 #include "NiagaraObjectSelection.h"
 #include "Modules/ModuleManager.h"
+#include "Stateless/NiagaraStatelessEmitter.h"
 #include "ViewModels/NiagaraOverviewGraphViewModel.h"
 #include "ViewModels/NiagaraSystemSelectionViewModel.h"
 #include "ViewModels/NiagaraSystemViewModel.h"
@@ -16,6 +17,7 @@
 #include "ViewModels/NiagaraEmitterViewModel.h"
 #include "ViewModels/Stack/NiagaraStackEmitterPropertiesGroup.h"
 #include "ViewModels/Stack/NiagaraStackRoot.h"
+#include "ViewModels/Stack/NiagaraStackStatelessEmitterGroup.h"
 #include "ViewModels/Stack/NiagaraStackSystemSettingsGroup.h"
 #include "Widgets/Layout/SWidgetSwitcher.h"
 
@@ -40,7 +42,7 @@ void SNiagaraScalabilityContext::Construct(const FArguments& InArgs, UNiagaraSys
 	// without this, a few properties from the mesh renderer properties make it into the details panel even though the delegate returns false for them.
 	DetailsView->RegisterInstancedCustomPropertyLayout(UNiagaraMeshRendererProperties::StaticClass(), DetailsView->GetGenericLayoutDetailsDelegate());
 	DetailsView->RegisterInstancedCustomPropertyLayout(UNiagaraEmitter::StaticClass(), FOnGetDetailCustomizationInstance::CreateStatic(&FNiagaraEmitterScalabilityDetails::MakeInstance));
-	
+
 	DetailsView->SetIsPropertyVisibleDelegate(FIsPropertyVisible::CreateSP(this, &SNiagaraScalabilityContext::FilterScalabilityProperties));
 	
 	UpdateScalabilityContent();
@@ -95,6 +97,10 @@ void SNiagaraScalabilityContext::SetObject(UObject* Object)
 	{
 		DetailsView->SetObject(Emitter);
 	}
+	else if (UNiagaraStatelessEmitter* StatelessEmitter = Cast<UNiagaraStatelessEmitter>(Object))
+	{
+		DetailsView->SetObject(StatelessEmitter);
+	}
 	else if(UNiagaraSystem* System = Cast<UNiagaraSystem>(Object))
 	{
 		DetailsView->SetObject(System);
@@ -128,6 +134,11 @@ void SNiagaraScalabilityContext::UpdateScalabilityContent()
 			UNiagaraEmitter* Emitter = Cast<UNiagaraStackEmitterPropertiesGroup>(StackEntry)->GetEmitterViewModel()->GetEmitter().Emitter;
 			NewSelection = Emitter;
 		}
+		else if (StackEntry->IsA(UNiagaraStackStatelessEmitterGroup::StaticClass()))
+		{
+			UNiagaraStatelessEmitter* Emitter = Cast<UNiagaraStackStatelessEmitterGroup>(StackEntry)->GetStatelessEmitter();
+			NewSelection = Emitter;
+		}
 		else if(StackEntry->IsA(UNiagaraStackSystemPropertiesGroup::StaticClass()))
 		{
 			UNiagaraSystem& System = Cast<UNiagaraStackSystemPropertiesGroup>(StackEntry)->GetSystemViewModel()->GetSystem();
@@ -142,7 +153,14 @@ void SNiagaraScalabilityContext::UpdateScalabilityContent()
 		{
 			if(FNiagaraEmitterHandle* EmitterHandle = OverviewNode->TryGetEmitterHandle())
 			{
-				NewSelection = EmitterHandle->GetInstance().Emitter;
+				if (EmitterHandle->GetEmitterMode() == ENiagaraEmitterMode::Standard)
+				{
+					NewSelection = EmitterHandle->GetInstance().Emitter;
+				}
+				else
+				{
+					NewSelection = EmitterHandle->GetStatelessEmitter();
+				}
 			}
 			else if(UNiagaraSystem* System = OverviewNode->GetOwningSystem())
 			{
diff --git a/Engine/Plugins/FX/Niagara/Source/NiagaraEditorWidgets/Private/SNiagaraOverviewStackNode.cpp b/Engine/Plugins/FX/Niagara/Source/NiagaraEditorWidgets/Private/SNiagaraOverviewStackNode.cpp
--- a/Engine/Plugins/FX/Niagara/Source/NiagaraEditorWidgets/Private/SNiagaraOverviewStackNode.cpp
+++ b/Engine/Plugins/FX/Niagara/Source/NiagaraEditorWidgets/Private/SNiagaraOverviewStackNode.cpp
@@ -25,6 +25,7 @@
 #include "Stack/SNiagaraSummaryViewToggle.h"
 #include "Stack/SNiagaraSimTargetToggle.h"
 #include "Stack/SNiagaraStackIssueIcon.h"
+#include "Stateless/NiagaraStatelessEmitter.h"
 #include "Styling/AppStyle.h"
 #include "Styling/StyleColors.h"
 #include "ViewModels/NiagaraEmitterHandleViewModel.h"
@@ -867,12 +868,29 @@ EVisibility SNiagaraOverviewStackNode::GetScalabilityIndicatorVisibility() const
 
 	if(EmitterHandleViewModel.IsValid())
 	{
-		if(FVersionedNiagaraEmitterData* EmitterData = EmitterHandleViewModel->GetEmitterHandle()->GetInstance().GetEmitterData())
+		FNiagaraEmitterHandle* EmitterHandle = EmitterHandleViewModel->GetEmitterHandle();
+		const FNiagaraPlatformSet* EmitterPlatformSet = nullptr;
+		const FNiagaraEmitterScalabilityOverrides* EmitterScalabilityOverrides = nullptr;
+
+		if (EmitterHandle->GetEmitterMode() == ENiagaraEmitterMode::Standard)
 		{
-			bool bIsQualityLevelMaskSetup = EmitterData->Platforms.QualityLevelMask != INDEX_NONE;
-			bool bIsScalabilitySetup = EmitterData->ScalabilityOverrides.Overrides.Num() != 0 || (bIsQualityLevelMaskSetup && EmitterData->Platforms.QualityLevelMask != FNiagaraPlatformSet::GetFullQualityLevelMask(GetDefault<UNiagaraSettings>()->QualityLevels.Num())); 
-			return bIsScalabilitySetup ? EVisibility::Visible : DisplayMode == EDisplayMode::Summary ? EVisibility::Hidden : EVisibility::Collapsed;
+			if (FVersionedNiagaraEmitterData* EmitterData = EmitterHandle->GetInstance().GetEmitterData())
+			{
+				EmitterPlatformSet = &EmitterData->Platforms;
+				EmitterScalabilityOverrides = &EmitterData->ScalabilityOverrides;
+			}
 		}
+		else
+		{
+			if (UNiagaraStatelessEmitter* StatelessEmitter = EmitterHandle->GetStatelessEmitter())
+			{
+				EmitterPlatformSet = &StatelessEmitter->GetPlatformSet();
+			}
+		}
+
+		const bool bIsQualityLevelMaskSetup = EmitterPlatformSet && EmitterPlatformSet->QualityLevelMask != INDEX_NONE && EmitterPlatformSet->QualityLevelMask != FNiagaraPlatformSet::GetFullQualityLevelMask(GetDefault<UNiagaraSettings>()->QualityLevels.Num());
+		const bool bIsScalabilitySetup = EmitterScalabilityOverrides && EmitterScalabilityOverrides->Overrides.Num() != 0;
+		return (bIsQualityLevelMaskSetup || bIsScalabilitySetup) ? EVisibility::Visible : DisplayMode == EDisplayMode::Summary ? EVisibility::Hidden : EVisibility::Collapsed;
 	}
 
 	if(UNiagaraSystem* System = OverviewStackNode->GetOwningSystem())
@@ -1330,7 +1348,7 @@ EVisibility SNiagaraOverviewStackNode::ShowExcludedOverlay() const
 
 		if (const TSharedPtr<FNiagaraEmitterHandleViewModel> EmitterHandleViewModel = EmitterHandleViewModelWeak.Pin())
 		{
-			if (!EmitterHandleViewModel->GetEmitterHandle()->GetEmitterData()->IsAllowedByScalability())
+			if (!EmitterHandleViewModel->GetEmitterHandle()->IsAllowedByScalability())
 			{
 				return EVisibility::HitTestInvisible;
 			}
