- Enable new VM by default and strip the original byte code

#rnx

diff --git a/Engine/Plugins/FX/Niagara/Source/Niagara/Public/NiagaraSettings.h b/Engine/Plugins/FX/Niagara/Source/Niagara/Public/NiagaraSettings.h
--- a/Engine/Plugins/FX/Niagara/Source/Niagara/Public/NiagaraSettings.h
+++ b/Engine/Plugins/FX/Niagara/Source/Niagara/Public/NiagaraSettings.h
@@ -162,7 +162,7 @@ class UNiagaraSettings : public UDeveloperSettings
 	 the cvars fx.NiagaraScript.StripByteCodeOnLoad and fx.ForceExecVMPath.
 	*/
 	UPROPERTY(config, EditAnywhere, Category = Niagara, meta = (DisplayName = "Enable building data for Optimized VM"))
-	bool bExperimentalVMEnabled = false;
+	bool bExperimentalVMEnabled = true;
 
 	/** If set to true, quaternion attributes will be interpolated via slerp instead of lerp in interpolated spawn scripts. */
 	UPROPERTY(config, EditAnywhere, AdvancedDisplay, Category = Niagara)
@@ -196,7 +196,7 @@ class UNiagaraSettings : public UDeveloperSettings
 	 Controls how byte code will be stripped when loading assets that have multiple sets of bytecode (i.e. optimized).
 	 */
 	UPROPERTY(config, EditAnywhere, Category = Niagara, meta = (DisplayName = "Option for how to strip bytecode"))
-	ENiagaraStripScriptByteCodeOption ByteCodeStripOption = ENiagaraStripScriptByteCodeOption::Default;
+	ENiagaraStripScriptByteCodeOption ByteCodeStripOption = ENiagaraStripScriptByteCodeOption::Strip_Original;
 
 #if WITH_EDITORONLY_DATA
 	UPROPERTY(config, EditAnywhere, Category = Niagara, meta = (DisplayName = "Option for how to compile Niagara scripts"))
