Fixed issue when attempting to generate per signature hlsl for NDC DI funcitons that do not need it.

diff --git a/Engine/Plugins/FX/Niagara/Source/Niagara/Private/DataInterface/NiagaraDataInterfaceDataChannelRead.cpp b/Engine/Plugins/FX/Niagara/Source/Niagara/Private/DataInterface/NiagaraDataInterfaceDataChannelRead.cpp
--- a/Engine/Plugins/FX/Niagara/Source/Niagara/Private/DataInterface/NiagaraDataInterfaceDataChannelRead.cpp
+++ b/Engine/Plugins/FX/Niagara/Source/Niagara/Private/DataInterface/NiagaraDataInterfaceDataChannelRead.cpp
@@ -1742,15 +1742,17 @@ void UNiagaraDataInterfaceDataChannelRead::GetParameterDefinitionHLSL(FNiagaraDa
 		const FNiagaraDataInterfaceGeneratedFunction& Func = HlslGenContext.ParameterInfo.GeneratedFunctions[FuncIdx];
 		const FNiagaraFunctionSignature& Signature = HlslGenContext.Signatures[FuncIdx];
 
+		const TCHAR* FunctionTemplateFilename = NDIDataChannelReadLocal::GetFunctionTemplate(Signature.Name);
+
 		//Init/Reset our per function hlsl template args.
 		FunctionSymbol.StringValue = HlslGenContext.GetFunctionSignatureSymbol(Signature);
 		FunctionParameters.StringValue.Reset();//Reset function parameters ready to rebuild when we iterate over the parameters.
 		PerFunctionParameterShaderCode.StringValue.Reset(); //Reset function parameters ready to rebuild when we iterate over the parameters.
 
-		//Generate the function hlsl if we've not done so already for this signature.
+		//Generate the function hlsl if we've not done so already for this signature. We can also skip for functions that have no template file as they don't need special per signature hlsl generation.
 		uint32 FuncHash = GetSignatureHash(Signature);
 		FString& FunctionHlsl = FunctionHlslMap.FindOrAdd(FuncHash);
-		if (FunctionHlsl.IsEmpty())
+		if (FunctionHlsl.IsEmpty() && FunctionTemplateFilename != nullptr)
 		{
 			//Function that will recurse down a parameter's type and generate the appropriate IO code for all of it's members.
 			TFunction<void(bool, UScriptStruct*, FString&)> GenerateRWParameterCode = [&](bool bRead, UScriptStruct* Struct, FString& OutCode)
@@ -1873,7 +1875,7 @@ void UNiagaraDataInterfaceDataChannelRead::GetParameterDefinitionHLSL(FNiagaraDa
 			if (FunctionTemplateFile.IsEmpty())
 			{
 				//Load it if we have not previously.
-				LoadShaderSourceFile(NDIDataChannelReadLocal::GetFunctionTemplate(Signature.Name), EShaderPlatform::SP_PCD3D_SM5, &FunctionTemplateFile, nullptr);
+				LoadShaderSourceFile(FunctionTemplateFilename, EShaderPlatform::SP_PCD3D_SM5, &FunctionTemplateFile, nullptr);
 			}
 
 			//Finally generate the final code for this function and add it to the final hlsl.
