- Validation Changes
- Banned Module rule can specify CPU or GPU and Severity
- Banned Renderers can specify Severity
- Banned Data Interfaces now uses Severity

#rnx
#tests Tests out rules on some content that would break them

diff --git a/Engine/Plugins/FX/Niagara/Source/NiagaraEditor/Private/NiagaraValidationRules.cpp b/Engine/Plugins/FX/Niagara/Source/NiagaraEditor/Private/NiagaraValidationRules.cpp
--- a/Engine/Plugins/FX/Niagara/Source/NiagaraEditor/Private/NiagaraValidationRules.cpp
+++ b/Engine/Plugins/FX/Niagara/Source/NiagaraEditor/Private/NiagaraValidationRules.cpp
@@ -619,7 +619,7 @@ void UNiagaraValidationRule_BannedRenderers::CheckValidity(const FNiagaraValidat
 					{
 						FNiagaraValidationResult& Result = Results.AddDefaulted_GetRef();
 						
-						Result.Severity = ENiagaraValidationSeverity::Warning;
+						Result.Severity = Severity;
 						Result.SummaryText = LOCTEXT("BannedRenderSummary", "Banned renderers used.");
 						Result.Description = LOCTEXT("BannedRenderDescription", "Please ensure only allowed renderers are used for each platform according to the validation rules in the System's Effect Type.");
 						Result.SourceObject = StackItem;
@@ -654,53 +654,62 @@ void UNiagaraValidationRule_BannedModules::CheckValidity(const FNiagaraValidatio
 
 	for (UNiagaraStackModuleItem* Item : StackModuleItems)
 	{
-		if (Item && Item->GetIsEnabled())
+		if (!Item || !Item->GetIsEnabled())
 		{
-			UNiagaraNodeFunctionCall& FuncCall = Item->GetModuleNode();
+			continue;
+		}
 
-			for (UNiagaraScript* BannedModule : BannedModules)
+		UNiagaraNodeFunctionCall& FuncCall = Item->GetModuleNode();
+		for (UNiagaraScript* BannedModule : BannedModules)
+		{
+			if (BannedModule != FuncCall.FunctionScript)
 			{
-				if (BannedModule == FuncCall.FunctionScript)
-				{
-					FVersionedNiagaraEmitterData* EmitterData = Item->GetEmitterViewModel().IsValid() ? Item->GetEmitterViewModel()->GetEmitter().GetEmitterData() : nullptr;
-
-					bool bApplyBan = true;
-					if (EmitterData)
-					{
-						//If we're on an emitter, this emitter may be culled on the platforms the rule applies to.
-						const TArray<FNiagaraPlatformSetConflictInfo> Conflicts = NiagaraValidation::GatherPlatformSetConflicts(&Platforms, &EmitterData->Platforms);
-						bApplyBan = Conflicts.Num() > 0;
-					}
+				continue;
+			}
 
-					if (!bApplyBan)
-					{
-						continue;
-					}
+			FVersionedNiagaraEmitterData* EmitterData = Item->GetEmitterViewModel().IsValid() ? Item->GetEmitterViewModel()->GetEmitter().GetEmitterData() : nullptr;
+			if ( EmitterData )
+			{
+				bool bApplyBan =
+					(bBanOnCpu && EmitterData->SimTarget == ENiagaraSimTarget::CPUSim) ||
+					(bBanOnGpu && EmitterData->SimTarget == ENiagaraSimTarget::GPUComputeSim);
+
+				//If we're on an emitter, this emitter may be culled on the platforms the rule applies to.
+				const TArray<FNiagaraPlatformSetConflictInfo> Conflicts = NiagaraValidation::GatherPlatformSetConflicts(&Platforms, &EmitterData->Platforms);
+				bApplyBan &= Conflicts.Num() > 0;
+				if (!bApplyBan)
+				{
+					continue;
+				}
+			}
+			else if (!bBanOnCpu)
+			{
+				// System & Emitter scripts only run on the CPU
+				continue;
+			}
 
-					const FTextFormat Format(LOCTEXT("BannedModuleFormat", "Module {0} is banned on some currently enabled platforms"));
-					const FText WarningMessage = FText::Format(Format, FText::FromString(FuncCall.FunctionScript->GetName()));
+			const FTextFormat Format(LOCTEXT("BannedModuleFormat", "Module {0} is banned on some currently enabled platforms"));
+			const FText WarningMessage = FText::Format(Format, FText::FromString(FuncCall.FunctionScript->GetName()));
 
-					FNiagaraValidationResult& Result = Results.AddDefaulted_GetRef();
-					Result.Severity = ENiagaraValidationSeverity::Warning;
-					Result.SummaryText = WarningMessage;
-					Result.Description = LOCTEXT("BanndeModulesDescription", "Check this module against the Effect Type's Banned Modules validators");
-					Result.SourceObject = Item;
+			FNiagaraValidationResult& Result = Results.AddDefaulted_GetRef();
+			Result.Severity = Severity;
+			Result.SummaryText = WarningMessage;
+			Result.Description = LOCTEXT("BanndeModulesDescription", "Check this module against the Effect Type's Banned Modules validators");
+			Result.SourceObject = Item;
 
-					NiagaraValidation::AddGoToFXTypeLink(Result, System.GetEffectType());
+			NiagaraValidation::AddGoToFXTypeLink(Result, System.GetEffectType());
 
-					//Add autofix to disable the module
-					FNiagaraValidationFix& DisableModuleFix = Result.Fixes.AddDefaulted_GetRef();
-					DisableModuleFix.Description = LOCTEXT("DisableBannedModuleFix", "Disable Banned Module");					
-					TWeakObjectPtr<UNiagaraStackModuleItem> WeakModuleItem = Item;
-					DisableModuleFix.FixDelegate = FNiagaraValidationFixDelegate::CreateLambda([WeakModuleItem]()
-					{
-						if (UNiagaraStackModuleItem* ModuleItem = WeakModuleItem.Get())
-						{
-							ModuleItem->SetEnabled(false);
-						}
-					});
+			//Add autofix to disable the module
+			FNiagaraValidationFix& DisableModuleFix = Result.Fixes.AddDefaulted_GetRef();
+			DisableModuleFix.Description = LOCTEXT("DisableBannedModuleFix", "Disable Banned Module");					
+			TWeakObjectPtr<UNiagaraStackModuleItem> WeakModuleItem = Item;
+			DisableModuleFix.FixDelegate = FNiagaraValidationFixDelegate::CreateLambda([WeakModuleItem]()
+			{
+				if (UNiagaraStackModuleItem* ModuleItem = WeakModuleItem.Get())
+				{
+					ModuleItem->SetEnabled(false);
 				}
-			}
+			});
 		}
 	}
 }
@@ -709,6 +718,9 @@ void UNiagaraValidationRule_BannedDataInterfaces::CheckValidity(const FNiagaraVa
 {
 	UNiagaraSystem* NiagaraSystem = &Context.ViewModel->GetSystem();
 
+	FNiagaraDataInterfaceUtilities::FDataInterfaceSearchOptions SearchOptions;
+	SearchOptions.bIncludeInternal = true;
+
 	FNiagaraDataInterfaceUtilities::ForEachDataInterface(
 		NiagaraSystem,
 		[&](const FNiagaraDataInterfaceUtilities::FDataInterfaceUsageContext& UsageContext) -> bool
@@ -761,7 +773,7 @@ void UNiagaraValidationRule_BannedDataInterfaces::CheckValidity(const FNiagaraVa
 				if (bBanOnCpu == true)
 				{
 					Results.Emplace(
-						ENiagaraValidationSeverity::Warning,
+						Severity,
 						FText::Format(WarningFormat, FText::FromName(UsageContext.Variable.GetName())),
 						FText::Format(SystemDescFormat, FText::FromName(UsageContext.Variable.GetName()), FText::FromName(DIClass->GetFName())),
 						NiagaraValidation::GetStackEntry<UNiagaraStackSystemPropertiesItem>(Context.ViewModel->GetSystemStackViewModel())
@@ -770,7 +782,8 @@ void UNiagaraValidationRule_BannedDataInterfaces::CheckValidity(const FNiagaraVa
 			}
 
 			return true;
-		}
+		},
+		SearchOptions
 	);
 }
 
diff --git a/Engine/Plugins/FX/Niagara/Source/NiagaraEditor/Public/NiagaraValidationRules.h b/Engine/Plugins/FX/Niagara/Source/NiagaraEditor/Public/NiagaraValidationRules.h
--- a/Engine/Plugins/FX/Niagara/Source/NiagaraEditor/Public/NiagaraValidationRules.h
+++ b/Engine/Plugins/FX/Niagara/Source/NiagaraEditor/Public/NiagaraValidationRules.h
@@ -111,6 +111,9 @@ public:
 	UPROPERTY(EditAnywhere, Category = Validation)
 	FNiagaraPlatformSet Platforms;
 
+	UPROPERTY(EditAnywhere, Category = Validation)
+	ENiagaraValidationSeverity Severity = ENiagaraValidationSeverity::Warning;
+
 	UPROPERTY(EditAnywhere, Category = Validation)
 	TArray<TSubclassOf<UNiagaraRendererProperties>> BannedRenderers;
 
@@ -129,6 +132,15 @@ public:
 	UPROPERTY(EditAnywhere, Category=Validation)
 	FNiagaraPlatformSet Platforms;
 
+	UPROPERTY(EditAnywhere, Category = Validation)
+	bool bBanOnGpu = true;
+
+	UPROPERTY(EditAnywhere, Category = Validation)
+	bool bBanOnCpu = true;
+
+	UPROPERTY(EditAnywhere, Category = Validation)
+	ENiagaraValidationSeverity Severity = ENiagaraValidationSeverity::Warning;
+
 	UPROPERTY(EditAnywhere, Category = Validation)
 	TArray<TObjectPtr<UNiagaraScript>> BannedModules;
 
