- Add NPC binding support to stateless

#rnx

diff --git a/Engine/Plugins/FX/Niagara/Source/Niagara/Internal/Stateless/NiagaraStatelessEmitter.h b/Engine/Plugins/FX/Niagara/Source/Niagara/Internal/Stateless/NiagaraStatelessEmitter.h
--- a/Engine/Plugins/FX/Niagara/Source/Niagara/Internal/Stateless/NiagaraStatelessEmitter.h
+++ b/Engine/Plugins/FX/Niagara/Source/Niagara/Internal/Stateless/NiagaraStatelessEmitter.h
@@ -13,6 +13,7 @@
 
 //-TODO:Stateless: Merge this into UNiagaraEmitterBase perhaps?
 
+class UNiagaraParameterCollection;
 struct FNiagaraStatelessEmitterData;
 class UNiagaraStatelessModule;
 class UNiagaraRendererProperties;
@@ -43,9 +44,12 @@ public:
 
 #if WITH_EDITOR
 	void OnEmitterTemplateChanged();
+	void OnCacheParameterCollectionReferences();
 #endif
 	NIAGARA_API const UNiagaraStatelessEmitterTemplate* GetEmitterTemplate() const;
 
+	bool UsesCollection(const UNiagaraParameterCollection* Collection) const;
+
 	const TArray<UNiagaraRendererProperties*>& GetRenderers() { return RendererProperties; }
 	const TArray<UNiagaraRendererProperties*>& GetRenderers() const { return RendererProperties; }
 
@@ -155,6 +159,9 @@ protected:
 	UPROPERTY()
 	TArray<int32> ComponentOffsets;
 
+	UPROPERTY()
+	TArray<TObjectPtr<UNiagaraParameterCollection>> CachedParameterCollectionReferences;
+
 #if WITH_EDITORONLY_DATA
 	FSimpleMulticastDelegate OnRenderersChangedDelegate;
 #endif
diff --git a/Engine/Plugins/FX/Niagara/Source/Niagara/Internal/Stateless/NiagaraStatelessEmitterData.h b/Engine/Plugins/FX/Niagara/Source/Niagara/Internal/Stateless/NiagaraStatelessEmitterData.h
--- a/Engine/Plugins/FX/Niagara/Source/Niagara/Internal/Stateless/NiagaraStatelessEmitterData.h
+++ b/Engine/Plugins/FX/Niagara/Source/Niagara/Internal/Stateless/NiagaraStatelessEmitterData.h
@@ -10,6 +10,7 @@
 
 #include "Shader.h"
 
+class UNiagaraParameterCollection;
 class UNiagaraRendererProperties;
 class UNiagaraStatelessEmitterTemplate;
 namespace NiagaraStateless
@@ -45,6 +46,7 @@ struct FNiagaraStatelessEmitterData
 	TArray<FNiagaraStatelessSpawnInfo>				SpawnInfos;
 
 	TArray<TObjectPtr<UNiagaraRendererProperties>>	RendererProperties;
+	TArray<TObjectPtr<UNiagaraParameterCollection>>	BoundParameterCollections;
 
 	bool											bModulesHaveRendererBindings = false;
 	FNiagaraParameterStore							RendererBindings;			// Contains all bindings for modules & renderers
diff --git a/Engine/Plugins/FX/Niagara/Source/Niagara/Private/NiagaraSystem.cpp b/Engine/Plugins/FX/Niagara/Source/Niagara/Private/NiagaraSystem.cpp
--- a/Engine/Plugins/FX/Niagara/Source/Niagara/Private/NiagaraSystem.cpp
+++ b/Engine/Plugins/FX/Niagara/Source/Niagara/Private/NiagaraSystem.cpp
@@ -488,10 +488,21 @@ bool UNiagaraSystem::UsesCollection(const UNiagaraParameterCollection* Collectio
 
 	for (const FNiagaraEmitterHandle& EmitterHandle : GetEmitterHandles())
 	{
-		FVersionedNiagaraEmitterData* EmitterData = EmitterHandle.GetEmitterData();
-		if (EmitterData && EmitterData->UsesCollection(Collection))
+		if (EmitterHandle.GetEmitterMode() == ENiagaraEmitterMode::Standard)
 		{
-			return true;
+			FVersionedNiagaraEmitterData* EmitterData = EmitterHandle.GetEmitterData();
+			if (EmitterData && EmitterData->UsesCollection(Collection))
+			{
+				return true;
+			}
+		}
+		else
+		{
+			UNiagaraStatelessEmitter* Emitter = EmitterHandle.GetStatelessEmitter();
+			if (Emitter && Emitter->UsesCollection(Collection))
+			{
+				return true;
+			}
 		}
 	}
 
diff --git a/Engine/Plugins/FX/Niagara/Source/Niagara/Private/Stateless/NiagaraStatelessEmitter.cpp b/Engine/Plugins/FX/Niagara/Source/Niagara/Private/Stateless/NiagaraStatelessEmitter.cpp
--- a/Engine/Plugins/FX/Niagara/Source/Niagara/Private/Stateless/NiagaraStatelessEmitter.cpp
+++ b/Engine/Plugins/FX/Niagara/Source/Niagara/Private/Stateless/NiagaraStatelessEmitter.cpp
@@ -15,6 +15,7 @@
 
 #include "Algo/Copy.h"
 #include "Interfaces/ITargetPlatform.h"
+#include "UObject/UObjectIterator.h"
 
 namespace NiagaraStatelessInternal
 {
@@ -99,6 +100,7 @@ void UNiagaraStatelessEmitter::PostLoad()
 #if WITH_EDITOR
 	// Ensure our module list is up to date
 	OnEmitterTemplateChanged();
+	//OnCacheParameterCollectionReferences();
 #endif
 }
 
@@ -157,6 +159,7 @@ void UNiagaraStatelessEmitter::PostEditChangeProperty(FPropertyChangedEvent& Pro
 		//if (PropertyChangedEvent.GetPropertyName() == GET_MEMBER_NAME_CHECKED(UNiagaraStatelessEmitter, EmitterTemplateClass))
 		{
 			OnEmitterTemplateChanged();
+			OnCacheParameterCollectionReferences();
 		}
 	}
 }
@@ -198,8 +201,58 @@ void UNiagaraStatelessEmitter::OnEmitterTemplateChanged()
 	}
 	Modules = MoveTemp(NewModules);
 }
+
+void UNiagaraStatelessEmitter::OnCacheParameterCollectionReferences()
+{
+	TArray<TObjectPtr<UNiagaraParameterCollection>> OriginalReferences;
+	Swap(OriginalReferences, CachedParameterCollectionReferences);
+
+	for (const UNiagaraStatelessModule* Module : Modules)
+	{
+		for (TFieldIterator<FStructProperty> PropIt(Module->GetClass()); PropIt; ++PropIt)
+		{
+			FStructProperty* StructProp = *PropIt;
+			if (!StructProp || !StructProp->Struct || !StructProp->Struct->IsChildOf(FNiagaraDistributionBase::StaticStruct()))
+			{
+				continue;
+			}
+
+			const FNiagaraDistributionBase* DistributionBase = StructProp->ContainerPtrToValuePtr<const FNiagaraDistributionBase>(Module);
+			if (!DistributionBase->IsBinding())
+			{
+				continue;
+			}
+
+
+			if (!DistributionBase->ParameterBinding.IsInNameSpace(FNiagaraConstants::ParameterCollectionNamespaceString))
+			{
+				continue;
+			}
+
+			for (TObjectIterator<UNiagaraParameterCollection> NPCIt; NPCIt; ++NPCIt)
+			{
+				UNiagaraParameterCollection* NPC = *NPCIt;
+				if (NPC->GetParameters().Contains(DistributionBase->ParameterBinding))
+				{
+					CachedParameterCollectionReferences.AddUnique(NPC);
+					break;
+				}
+			}
+		}
+	}
+	
+	if ( OriginalReferences != CachedParameterCollectionReferences )
+	{
+		Modify();
+	}
+}
 #endif //WITH_EDITOR
 
+bool UNiagaraStatelessEmitter::UsesCollection(const UNiagaraParameterCollection* Collection) const
+{
+	return CachedParameterCollectionReferences.Contains(Collection);
+}
+
 const UNiagaraStatelessEmitterTemplate* UNiagaraStatelessEmitter::GetEmitterTemplate() const
 {
 	return EmitterTemplateClass ? EmitterTemplateClass->GetDefaultObject<UNiagaraStatelessEmitterTemplate>() : nullptr;
@@ -309,6 +362,21 @@ void UNiagaraStatelessEmitter::CacheFromCompiledData()
 		}
 		StatelessEmitterData->InitRenderResources();
 
+		// Gather the parameter collections we need to bind to
+	#if WITH_EDITOR
+		// UNiagaraStatelessEmitter does not get a notification when bindings change.
+		//-OPT: Shift this to be when distributions change mode in / out of binding + binding changed
+		OnCacheParameterCollectionReferences();
+	#endif
+		StatelessEmitterData->BoundParameterCollections.Reset(CachedParameterCollectionReferences.Num());
+		for (UNiagaraParameterCollection* ParameterCollection : CachedParameterCollectionReferences)
+		{
+			if (ParameterCollection)
+			{
+				StatelessEmitterData->BoundParameterCollections.Add(ParameterCollection);
+			}
+		}
+
 		// Prepare renderer bindings this avoid having to do this per instance spawned
 		// Note: Order is important here we detect if we need to update shader parameters on binding changes above by looking to see if we had any renderer bindings so this must be below
 		for (UNiagaraRendererProperties* Renderer : RendererProperties )
diff --git a/Engine/Plugins/FX/Niagara/Source/Niagara/Private/Stateless/NiagaraStatelessEmitterInstance.cpp b/Engine/Plugins/FX/Niagara/Source/Niagara/Private/Stateless/NiagaraStatelessEmitterInstance.cpp
--- a/Engine/Plugins/FX/Niagara/Source/Niagara/Private/Stateless/NiagaraStatelessEmitterInstance.cpp
+++ b/Engine/Plugins/FX/Niagara/Source/Niagara/Private/Stateless/NiagaraStatelessEmitterInstance.cpp
@@ -189,6 +189,14 @@ void FNiagaraStatelessEmitterInstance::BindParameters(bool bExternalOnly)
 		if (ParentSystemInstance)
 		{
 			ParentSystemInstance->BindToParameterStore(RendererBindings);
+
+			for (UNiagaraParameterCollection* ParameterCollection : EmitterData->BoundParameterCollections)
+			{
+				if (UNiagaraParameterCollectionInstance* ParameterCollectionInstance = ParentSystemInstance->GetParameterCollectionInstance(ParameterCollection))
+				{
+					ParameterCollectionInstance->GetParameterStore().Bind(&RendererBindings);
+				}
+			}
 		}
 	}
 }
@@ -200,6 +208,14 @@ void FNiagaraStatelessEmitterInstance::UnbindParameters(bool bExternalOnly)
 		if (ParentSystemInstance)
 		{
 			ParentSystemInstance->UnbindFromParameterStore(RendererBindings);
+
+			for (UNiagaraParameterCollection* ParameterCollection : EmitterData->BoundParameterCollections)
+			{
+				if (UNiagaraParameterCollectionInstance* ParameterCollectionInstance = ParentSystemInstance->GetParameterCollectionInstance(ParameterCollection))
+				{
+					ParameterCollectionInstance->GetParameterStore().Unbind(&RendererBindings);
+				}
+			}
 		}
 	}
 }
diff --git a/Engine/Plugins/FX/Niagara/Source/NiagaraEditor/Private/Stateless/NiagaraDistributionPropertyCustomization.cpp b/Engine/Plugins/FX/Niagara/Source/NiagaraEditor/Private/Stateless/NiagaraDistributionPropertyCustomization.cpp
--- a/Engine/Plugins/FX/Niagara/Source/NiagaraEditor/Private/Stateless/NiagaraDistributionPropertyCustomization.cpp
+++ b/Engine/Plugins/FX/Niagara/Source/NiagaraEditor/Private/Stateless/NiagaraDistributionPropertyCustomization.cpp
@@ -15,6 +15,7 @@
 #include "Widgets/SCompoundWidget.h"
 #include "Widgets/SNiagaraDistributionEditor.h"
 
+#include "NiagaraEditorUtilities.h"
 #include "NiagaraNodeParameterMapBase.h"
 #include "NiagaraParameterMapHistory.h"
 #include "NiagaraScriptSource.h"
@@ -181,6 +182,8 @@ public:
 		{
 			switch (InSourceMode)
 			{
+			case ENiagaraDistributionMode::Binding:
+				return ENiagaraDistributionEditorMode::Binding;
 			case ENiagaraDistributionMode::NonUniformConstant:
 				return ENiagaraDistributionEditorMode::ColorConstant;
 			case ENiagaraDistributionMode::NonUniformRange:
@@ -492,6 +495,19 @@ public:
 					AvailableBindings.Add(Variable);
 				}
 			}
+
+			TArray<UNiagaraParameterCollection*> AvailableParameterCollections;
+			FNiagaraEditorUtilities::GetAvailableParameterCollections(AvailableParameterCollections);
+			for (UNiagaraParameterCollection* NPCollection : AvailableParameterCollections)
+			{
+				for (const FNiagaraVariable& NPCVariable : NPCollection->GetParameters())
+				{
+					if (NPCVariable.GetType() == AllowedTypeDef)
+					{
+						AvailableBindings.Add(NPCVariable);
+					}
+				}
+			}
 		}
 		return AvailableBindings;
 	}
