Temporarily disable new scalability feature while some bugs are investigated.

diff --git a/Engine/Plugins/FX/Niagara/Source/Niagara/Private/NiagaraScalabilityManager.cpp b/Engine/Plugins/FX/Niagara/Source/Niagara/Private/NiagaraScalabilityManager.cpp
--- a/Engine/Plugins/FX/Niagara/Source/Niagara/Private/NiagaraScalabilityManager.cpp
+++ b/Engine/Plugins/FX/Niagara/Source/Niagara/Private/NiagaraScalabilityManager.cpp
@@ -20,7 +20,7 @@ static FAutoConsoleVariableRef CVarScalabilityManParallelThreshold(TEXT("fx.Scal
 static int32 GScalabilityMaxUpdatesPerFrame = 50;
 static FAutoConsoleVariableRef CVarScalabilityMaxUpdatesPerFrame(TEXT("fx.ScalabilityMaxUpdatesPerFrame"), GScalabilityMaxUpdatesPerFrame, TEXT("Number of instances that can be processed per frame when updating scalability state. -1 for all of them. \n"), ECVF_Default);
 
-static int32 GApplyInstanceCountsRigidly = 1;
+static int32 GApplyInstanceCountsRigidly = 0;
 static FAutoConsoleVariableRef CVarApplyInstanceCountsRigidly(TEXT("fx.Niagara.Scalability.ApplyInstanceCountsRigidly"), GApplyInstanceCountsRigidly, TEXT("If true we'll apply instance counts more rigidly. Un-culled Systems wil not be able to activate until overall systems fall below the instance count. \n"), ECVF_Default);
 
 static float GetScalabilityUpdatePeriod(ENiagaraScalabilityUpdateFrequency Frequency)
