Added an export function for Niagara sim caches to save their data to disk. There is no importer, so this is a one-way trip to get the data stored within into an independent format like json.

#tests Right-clicked on an existing cache from a 2d fluid sim, exported each frame to disk with hdr images

diff --git a/Engine/Plugins/FX/Niagara/Source/Niagara/Classes/NiagaraDataInterfaceRenderTarget2D.h b/Engine/Plugins/FX/Niagara/Source/Niagara/Classes/NiagaraDataInterfaceRenderTarget2D.h
--- a/Engine/Plugins/FX/Niagara/Source/Niagara/Classes/NiagaraDataInterfaceRenderTarget2D.h
+++ b/Engine/Plugins/FX/Niagara/Source/Niagara/Classes/NiagaraDataInterfaceRenderTarget2D.h
@@ -124,6 +124,7 @@ public:
 	NIAGARA_API virtual bool SimCacheWriteFrame(UObject* StorageObject, int FrameIndex, FNiagaraSystemInstance* SystemInstance, const void* OptionalPerInstanceData, FNiagaraSimCacheFeedbackContext& FeedbackContext) const override;
 	NIAGARA_API virtual bool SimCacheEndWrite(UObject* StorageObject) const override;
 	NIAGARA_API virtual bool SimCacheReadFrame(UObject* StorageObject, int FrameA, int FrameB, float Interp, FNiagaraSystemInstance* SystemInstance, void* OptionalPerInstanceData) override;
+	NIAGARA_API virtual TSharedPtr<FJsonObject> SimCacheToJson(const UObject* StorageObject, int FrameIndex, TOptional<FString> TargetFolder, TOptional<FString> FilenamePrefix) const override;
 	//~ UNiagaraDataInterface interface END
 
 	NIAGARA_API void VMGetSize(FVectorVMExternalFunctionContext& Context);
diff --git a/Engine/Plugins/FX/Niagara/Source/Niagara/Classes/NiagaraSimCache.h b/Engine/Plugins/FX/Niagara/Source/Niagara/Classes/NiagaraSimCache.h
--- a/Engine/Plugins/FX/Niagara/Source/Niagara/Classes/NiagaraSimCache.h
+++ b/Engine/Plugins/FX/Niagara/Source/Niagara/Classes/NiagaraSimCache.h
@@ -373,6 +373,7 @@ UCLASS(Experimental, BlueprintType, MinimalAPI)
 class UNiagaraSimCache : public UObject
 {
 	friend struct FNiagaraSimCacheAttributeReaderHelper;
+	friend struct FNiagaraSimCacheExport;
 	friend struct FNiagaraSimCacheCompare;
 	friend struct FNiagaraSimCacheHelper;
 	friend struct FNiagaraSimCacheGpuResource;
diff --git a/Engine/Plugins/FX/Niagara/Source/Niagara/Classes/NiagaraSimCacheCustomStorageInterface.h b/Engine/Plugins/FX/Niagara/Source/Niagara/Classes/NiagaraSimCacheCustomStorageInterface.h
--- a/Engine/Plugins/FX/Niagara/Source/Niagara/Classes/NiagaraSimCacheCustomStorageInterface.h
+++ b/Engine/Plugins/FX/Niagara/Source/Niagara/Classes/NiagaraSimCacheCustomStorageInterface.h
@@ -4,13 +4,14 @@
 
 #include "NiagaraCommon.h"
 #include "UObject/Interface.h"
+#include "Dom/JsonObject.h"
 #include "NiagaraSimCacheCustomStorageInterface.generated.h"
 
 struct FNiagaraSimCacheFeedbackContext;
 class FNiagaraSystemInstance;
 
 // Interface for UObjects to implement renderable mesh
-UINTERFACE()
+UINTERFACE(MinimalAPI)
 class UNiagaraSimCacheCustomStorageInterface : public UInterface
 {
 	GENERATED_BODY()
@@ -70,4 +71,12 @@ public:
 	this into 'MyEmitter.Particles.MyAttribute' by checking the UsageContext is a UNiagaraEmitter and then creating the variable from the unique name.
 	*/
 	virtual TArray<FNiagaraVariableBase> GetSimCacheRendererAttributes(UObject* UsageContext) const { return TArray<FNiagaraVariableBase>(); }
+
+	/**
+	 This converts the content of the storage object to a json representation. If another interchange format (e.g. an image format) is better, then the json this method produces should link to the secondary files.
+	@param TargetFolder (optional) the folder where to save auxiliary data from this frame. Might not be set if external files are not supported (e.g. when called over network).
+	@param FilenamePrefix (optional) unique name for this data interface - can either be used directly as filename or add extensions like .png or even to create a folder and put in multiple files related to this data interface
+	@return a json string representing this data interface object
+	 */
+	virtual TSharedPtr<FJsonObject> SimCacheToJson(const UObject* StorageObject, int FrameIndex, TOptional<FString> TargetFolder, TOptional<FString> FilenamePrefix) const { return TSharedPtr<FJsonObject>(); }
 };
diff --git a/Engine/Plugins/FX/Niagara/Source/Niagara/Classes/NiagaraSimCacheJson.h b/Engine/Plugins/FX/Niagara/Source/Niagara/Classes/NiagaraSimCacheJson.h
--- a/Engine/Plugins/FX/Niagara/Source/Niagara/Classes/NiagaraSimCacheJson.h
+++ b/Engine/Plugins/FX/Niagara/Source/Niagara/Classes/NiagaraSimCacheJson.h
@@ -9,8 +9,18 @@ class FJsonObject;
 /** Experimental Json interop for Niagara Sim Caches. */
 struct FNiagaraSimCacheJson
 {
-	NIAGARA_API static void DumpToFile(const UNiagaraSimCache& SimCache, const FString& FullPath);
+	enum class EExportType
+	{
+		SingleJsonFile,
+		SeparateEachFrame
+	};
+	
+	NIAGARA_API static bool DumpToFile(const UNiagaraSimCache& SimCache, const FString& FullPath, EExportType ExportType = EExportType::SingleJsonFile);
 
 	NIAGARA_API static TSharedPtr<FJsonObject> ToJson(const UNiagaraSimCache& SimCache);
+	NIAGARA_API static TSharedPtr<FJsonObject> SystemDataToJson(const UNiagaraSimCache& SimCache);
 	NIAGARA_API static TSharedPtr<FJsonObject> EmitterFrameToJson(const UNiagaraSimCache& SimCache, int EmitterIndex, int FrameIndex);
+
+private:
+	static bool DumpFramesToFolder(const UNiagaraSimCache& SimCache, const FString& TargetFolder);
 };
diff --git a/Engine/Plugins/FX/Niagara/Source/Niagara/Niagara.Build.cs b/Engine/Plugins/FX/Niagara/Source/Niagara/Niagara.Build.cs
--- a/Engine/Plugins/FX/Niagara/Source/Niagara/Niagara.Build.cs
+++ b/Engine/Plugins/FX/Niagara/Source/Niagara/Niagara.Build.cs
@@ -14,6 +14,7 @@ public class Niagara : ModuleRules
                 "ApplicationCore",
                 "AudioPlatformConfiguration",
                 "Core",
+                "ImageCore",
                 "DeveloperSettings",
                 "Engine",
                 "Json",
diff --git a/Engine/Plugins/FX/Niagara/Source/Niagara/Private/NiagaraDataInterfaceRenderTarget2D.cpp b/Engine/Plugins/FX/Niagara/Source/Niagara/Private/NiagaraDataInterfaceRenderTarget2D.cpp
--- a/Engine/Plugins/FX/Niagara/Source/Niagara/Private/NiagaraDataInterfaceRenderTarget2D.cpp
+++ b/Engine/Plugins/FX/Niagara/Source/Niagara/Private/NiagaraDataInterfaceRenderTarget2D.cpp
@@ -4,6 +4,8 @@
 #include "NiagaraCompileHashVisitor.h"
 #include "TextureResource.h"
 #include "CanvasTypes.h"
+#include "ImageCoreUtils.h"
+#include "ImageUtils.h"
 #include "Engine/TextureRenderTarget2D.h"
 
 #include "NDIRenderTargetVolumeSimCacheData.h"
@@ -763,6 +765,90 @@ NIAGARA_API bool UNiagaraDataInterfaceRenderTarget2D::SimCacheReadFrame(UObject*
 	return true;
 }
 
+namespace NiagaraRenderTarget2DImageHelper
+{
+	void ConvertPixel(const uint8* SourceData, EPixelFormat Format, FFloat16Color& OutColor)
+	{
+		switch (Format)
+		{
+		case PF_FloatRGB:
+			OutColor.R = *reinterpret_cast<const float*>(SourceData);
+			OutColor.G = *reinterpret_cast<const float*>(SourceData + 4);
+			OutColor.B = *reinterpret_cast<const float*>(SourceData + 8);
+		case PF_FloatRGBA:
+			OutColor.R = *reinterpret_cast<const float*>(SourceData);
+			OutColor.G = *reinterpret_cast<const float*>(SourceData + 4);
+			OutColor.B = *reinterpret_cast<const float*>(SourceData + 8);
+			OutColor.A = *reinterpret_cast<const float*>(SourceData + 12);
+		case PF_G16R16F:
+		case PF_G16R16F_FILTER:
+			FMemory::Memcpy(&OutColor.R, SourceData, 4);
+		default: // do nothing;
+			break;
+		}
+	}
+}
+
+TSharedPtr<FJsonObject> UNiagaraDataInterfaceRenderTarget2D::SimCacheToJson(const UObject* StorageObject, int FrameIndex, TOptional<FString> TargetFolder, TOptional<FString> FilenamePrefix) const
+{
+	const UNDIRenderTargetVolumeSimCacheData* SimCacheData = CastChecked<UNDIRenderTargetVolumeSimCacheData>(StorageObject);
+	if (SimCacheData->Frames.IsValidIndex(FrameIndex))
+	{
+		FNDIRenderTargetVolumeSimCacheFrame CacheFrame = SimCacheData->Frames[FrameIndex];
+		TSharedPtr<FJsonObject> JsonCacheObject = MakeShared<FJsonObject>();
+		JsonCacheObject->SetStringField(TEXT("CompressionType"), SimCacheData->CompressionType.ToString());
+		JsonCacheObject->SetNumberField(TEXT("PixelFormat"), static_cast<int32>(CacheFrame.Format));
+		JsonCacheObject->SetNumberField(TEXT("CompressedSize"), CacheFrame.CompressedSize);
+		JsonCacheObject->SetNumberField(TEXT("UncompressedSize"), CacheFrame.UncompressedSize);
+		JsonCacheObject->SetNumberField(TEXT("SizeX"), CacheFrame.Size.X);
+		JsonCacheObject->SetNumberField(TEXT("SizeY"), CacheFrame.Size.Y);
+		
+		const uint8* PixelData = CacheFrame.GetPixelData();
+		if (!PixelData || !TargetFolder.IsSet() || !FilenamePrefix.IsSet())
+		{
+			return JsonCacheObject;
+		}
+
+		ERawImageFormat::Type ImageFormat = FImageCoreUtils::GetRawImageFormatForPixelFormat(CacheFrame.Format);
+		if (ImageFormat != ERawImageFormat::RGBA16F)
+		{
+			UE_LOG(LogNiagara, Error, TEXT("Unable to save render target to file with current pixel format"));
+			return JsonCacheObject;
+		}
+
+		uint32 PixelCount = CacheFrame.Size.X * CacheFrame.Size.Y;
+		uint32 BlockBytes = GPixelFormats[CacheFrame.Format].BlockBytes;
+		TArray<uint8> Decompressed;
+		if (CacheFrame.CompressedSize > 0)
+		{
+			Decompressed.AddUninitialized(BlockBytes * PixelCount);
+			if (!FCompression::UncompressMemory(SimCacheData->CompressionType, Decompressed.GetData(), Decompressed.Num(), PixelData, CacheFrame.CompressedSize))
+			{
+				UE_LOG(LogNiagara, Error, TEXT("Error decompressing render target data"));
+				return JsonCacheObject;
+			}
+		}
+		const uint8* SrcData = Decompressed.Num() > 0 ? Decompressed.GetData() : PixelData;
+
+		// we need to convert the raw pixel data into a format that the image utils understand 
+		TArray<FFloat16Color> ImagePixelData;
+		ImagePixelData.AddZeroed(PixelCount);
+		for (uint32 Index = 0; Index < PixelCount; Index++)
+		{
+			uint32 Offset = Index * BlockBytes;
+			NiagaraRenderTarget2DImageHelper::ConvertPixel(SrcData + Offset, CacheFrame.Format, ImagePixelData[Index]);
+		}
+		
+		FImageView ImageView(ImagePixelData.GetData(), CacheFrame.Size.X, CacheFrame.Size.Y, ImageFormat);
+		if (FImageUtils::SaveImageByExtension(*FPaths::Combine(TargetFolder.GetValue(), FilenamePrefix.GetValue() + ".exr"), ImageView))
+		{
+			JsonCacheObject->SetStringField(TEXT("TextureData"), FilenamePrefix.GetValue() + ".exr");
+			return JsonCacheObject;
+		}		
+	}
+	return TSharedPtr<FJsonObject>();
+}
+
 void UNiagaraDataInterfaceRenderTarget2D::VMSetSize(FVectorVMExternalFunctionContext& Context)
 {
 	// This should only be called from a system or emitter script due to a need for only setting up initially.
diff --git a/Engine/Plugins/FX/Niagara/Source/Niagara/Private/NiagaraSimCacheJson.cpp b/Engine/Plugins/FX/Niagara/Source/Niagara/Private/NiagaraSimCacheJson.cpp
--- a/Engine/Plugins/FX/Niagara/Source/Niagara/Private/NiagaraSimCacheJson.cpp
+++ b/Engine/Plugins/FX/Niagara/Source/Niagara/Private/NiagaraSimCacheJson.cpp
@@ -2,25 +2,38 @@
 
 #include "NiagaraSimCacheJson.h"
 
+#include "NiagaraSimCacheCustomStorageInterface.h"
 #include "Dom/JsonObject.h"
 #include "Dom/JsonValue.h"
 #include "Misc/FileHelper.h"
+#include "Misc/ScopedSlowTask.h"
 #include "Policies/PrettyJsonPrintPolicy.h"
 #include "Serialization/JsonSerializer.h"
 #include "Serialization/JsonWriter.h"
 
-void FNiagaraSimCacheJson::DumpToFile(const UNiagaraSimCache& SimCache, const FString& FullPath)
+DEFINE_LOG_CATEGORY_STATIC(LogNiagaraSimCacheJson, Log, All);
+
+bool FNiagaraSimCacheJson::DumpToFile(const UNiagaraSimCache& SimCache, const FString& FullPath, EExportType ExportType)
 {
-	TSharedPtr<FJsonObject> JsonObject = ToJson(SimCache);
-	if (JsonObject.IsValid())
+	if (ExportType == EExportType::SingleJsonFile)
 	{
-		UE_LOG(LogNiagara, Warning, TEXT("Writing file to %s"), *FullPath);
+		TSharedPtr<FJsonObject> JsonObject = ToJson(SimCache);
+		if (JsonObject.IsValid())
+		{
+			UE_LOG(LogNiagara, Warning, TEXT("Writing file to %s"), *FullPath);
 
-		FString OutputString;
-		TSharedRef<TJsonWriter<TCHAR, TPrettyJsonPrintPolicy<TCHAR>>> Writer = TJsonWriterFactory<TCHAR, TPrettyJsonPrintPolicy<TCHAR>>::Create(&OutputString);
-		FJsonSerializer::Serialize(JsonObject.ToSharedRef(), Writer);
-		FFileHelper::SaveStringToFile(OutputString, *FullPath);
+			FString OutputString;
+			TSharedRef<TJsonWriter<>> Writer = TJsonWriterFactory<>::Create(&OutputString);
+			FJsonSerializer::Serialize(JsonObject.ToSharedRef(), Writer);
+			return FFileHelper::SaveStringToFile(OutputString, *FullPath);
+		}
+		return false;
 	}
+	else if (ExportType == EExportType::SeparateEachFrame)
+	{
+		return DumpFramesToFolder(SimCache, FullPath);
+	}
+	return false;
 }
 
 TSharedPtr<FJsonObject> FNiagaraSimCacheJson::ToJson(const UNiagaraSimCache& SimCache)
@@ -30,13 +43,7 @@ TSharedPtr<FJsonObject> FNiagaraSimCacheJson::ToJson(const UNiagaraSimCache& Sim
 		return nullptr;
 	}
 
-	TSharedPtr<FJsonObject> JsonCacheObject = MakeShared<FJsonObject>();
-	JsonCacheObject->SetStringField(TEXT("SystemAsset"), SimCache.GetSystemAsset().ToString());
-	JsonCacheObject->SetStringField(TEXT("CacheGuid"), SimCache.GetCacheGuid().ToString());
-	JsonCacheObject->SetNumberField(TEXT("StartSeconds"), SimCache.GetStartSeconds());
-	JsonCacheObject->SetNumberField(TEXT("DurationSeconds"), SimCache.GetDurationSeconds());
-	JsonCacheObject->SetNumberField(TEXT("NumFrames"), SimCache.GetNumFrames());
-	JsonCacheObject->SetNumberField(TEXT("NumEmitters"), SimCache.GetNumEmitters());
+	TSharedPtr<FJsonObject> JsonCacheObject = SystemDataToJson(SimCache);
 
 	// Write System Instance
 	{
@@ -64,6 +71,18 @@ TSharedPtr<FJsonObject> FNiagaraSimCacheJson::ToJson(const UNiagaraSimCache& Sim
 	return JsonCacheObject;
 }
 
+TSharedPtr<FJsonObject> FNiagaraSimCacheJson::SystemDataToJson(const UNiagaraSimCache& SimCache)
+{
+	TSharedPtr<FJsonObject> JsonCacheObject = MakeShared<FJsonObject>();
+	JsonCacheObject->SetStringField(TEXT("SystemAsset"), SimCache.GetSystemAsset().ToString());
+	JsonCacheObject->SetStringField(TEXT("CacheGuid"), SimCache.GetCacheGuid().ToString());
+	JsonCacheObject->SetNumberField(TEXT("StartSeconds"), SimCache.GetStartSeconds());
+	JsonCacheObject->SetNumberField(TEXT("DurationSeconds"), SimCache.GetDurationSeconds());
+	JsonCacheObject->SetNumberField(TEXT("NumFrames"), SimCache.GetNumFrames());
+	JsonCacheObject->SetNumberField(TEXT("NumEmitters"), SimCache.GetNumEmitters());
+	return JsonCacheObject;
+}
+
 TSharedPtr<FJsonObject> FNiagaraSimCacheJson::EmitterFrameToJson(const UNiagaraSimCache& SimCache, int EmitterIndex, int FrameIndex)
 {
 	const int NumInstances = SimCache.GetEmitterNumInstances(EmitterIndex, FrameIndex);
@@ -92,30 +111,85 @@ TSharedPtr<FJsonObject> FNiagaraSimCacheJson::EmitterFrameToJson(const UNiagaraS
 			TArray<int32>		Ints;
 			SimCache.ReadAttribute(Floats, Halfs, Ints, Attribute.GetName(), EmitterName, FrameIndex);
 
+			int32 TypeSize = Attribute.GetType().GetSize();
+
+			// add float components
 			if (Floats.Num() > 0)
 			{
 				TArray<TSharedPtr<FJsonValue>> Values;
-				for ( float v : Floats )
+				if (TypeSize > sizeof(float))
+				{
+					int32 ComponentNum = TypeSize / sizeof(float);
+					for (int i = 0; i < Floats.Num(); i += ComponentNum)
+					{
+						TArray<TSharedPtr<FJsonValue>> SubValues;
+						for (int Component = i; Floats.IsValidIndex(Component) && Component < i + ComponentNum; Component++)
+						{
+							SubValues.Emplace(MakeShared<FJsonValueNumber>(Floats[Component]));
+						}
+						Values.Emplace(MakeShared<FJsonValueArray>(SubValues));
+					}
+				}
+				else
 				{
-					Values.Emplace(MakeShared<FJsonValueNumber>(v));
+					for ( float v : Floats )
+					{
+						Values.Emplace(MakeShared<FJsonValueNumber>(v));
+					}
 				}
 				AttributeObject->SetArrayField(TEXT("Floats"), Values);
 			}
+
+			// add half components (converted to float)
 			if (Halfs.Num() > 0)
 			{
 				TArray<TSharedPtr<FJsonValue>> Values;
-				for (FFloat16 v : Halfs)
+				if (TypeSize > sizeof(uint16))
+				{
+					int32 ComponentNum = TypeSize / sizeof(uint16);
+					for (int i = 0; i < Halfs.Num(); i += ComponentNum)
+					{
+						TArray<TSharedPtr<FJsonValue>> SubValues;
+						for (int Component = i; Halfs.IsValidIndex(Component) && Component < i + ComponentNum; Component++)
+						{
+							SubValues.Emplace(MakeShared<FJsonValueNumber>(Halfs[Component].GetFloat()));
+						}
+						Values.Emplace(MakeShared<FJsonValueArray>(SubValues));
+					}
+				}
+				else
 				{
-					Values.Emplace(MakeShared<FJsonValueNumber>(v));
+					for (FFloat16 v : Halfs)
+					{
+						Values.Emplace(MakeShared<FJsonValueNumber>(v.GetFloat()));
+					}
 				}
 				AttributeObject->SetArrayField(TEXT("Halfs"), Values);
 			}
+
+			// add integer components
 			if (Ints.Num() > 0)
 			{
 				TArray<TSharedPtr<FJsonValue>> Values;
-				for (int32 v : Ints)
+				if (TypeSize > sizeof(int32))
 				{
-					Values.Emplace(MakeShared<FJsonValueNumber>(v));
+					int32 ComponentNum = TypeSize / sizeof(int32);
+					for (int i = 0; i < Ints.Num(); i += ComponentNum)
+					{
+						TArray<TSharedPtr<FJsonValue>> SubValues;
+						for (int Component = i; Ints.IsValidIndex(Component) && Component < i + ComponentNum; Component++)
+						{
+							SubValues.Emplace(MakeShared<FJsonValueNumber>(Ints[Component]));
+						}
+						Values.Emplace(MakeShared<FJsonValueArray>(SubValues));
+					}
+				}
+				else
+				{
+					for (int32 v : Ints)
+					{
+						Values.Emplace(MakeShared<FJsonValueNumber>(v));
+					}
 				}
 				AttributeObject->SetArrayField(TEXT("Ints"), Values);
 			}
@@ -124,3 +198,109 @@ TSharedPtr<FJsonObject> FNiagaraSimCacheJson::EmitterFrameToJson(const UNiagaraS
 	}
 	return EmitterObject;
 }
+
+bool FNiagaraSimCacheJson::DumpFramesToFolder(const UNiagaraSimCache& SimCache, const FString& TargetFolder)
+{
+	IFileManager& FileManager = IFileManager::Get();
+	if (!FileManager.DirectoryExists(*TargetFolder))
+	{
+		UE_LOG(LogNiagaraSimCacheJson, Warning, TEXT("Unable to save to folder %s"), *TargetFolder);
+		return false;
+	}
+
+	int32 EmitterCount = SimCache.GetNumEmitters();
+	FScopedSlowTask SlowTask(SimCache.GetNumFrames(), NSLOCTEXT("SimCacheExport", "SlowTaskLabel","Exporting frames..."));
+	SlowTask.MakeDialog(true);
+
+	{
+		// basic system data for the whole cache
+		TSharedPtr<FJsonObject> SystemJson = SystemDataToJson(SimCache);
+		if (!SystemJson.IsValid())
+		{
+			UE_LOG(LogNiagaraSimCacheJson, Warning, TEXT("Unable to export system data for cache %s"), *SimCache.GetName());
+			return false;
+		}
+		FString SystemFile = FPaths::Combine(TargetFolder, FPaths::MakeValidFileName(SimCache.GetName() + ".json", '_'));
+		FString OutputString;
+		TSharedRef<TJsonWriter<>> Writer = TJsonWriterFactory<>::Create(&OutputString);
+		FJsonSerializer::Serialize(SystemJson.ToSharedRef(), Writer);
+		FFileHelper::SaveStringToFile(OutputString, *SystemFile);
+	}
+	
+	for (int FrameIndex = 0; FrameIndex < SimCache.GetNumFrames(); FrameIndex++)
+	{
+		SlowTask.EnterProgressFrame();
+		if (SlowTask.ShouldCancel())
+		{
+			UE_LOG(LogNiagaraSimCacheJson, Warning, TEXT("Export cancelled by user"));
+			return false;
+		}
+		
+		FString FrameFolder = FPaths::Combine(TargetFolder, TEXT("Frame_") + FString::FormatAsNumber(FrameIndex));
+		if (!FileManager.MakeDirectory(*FrameFolder))
+		{
+			return false;
+		}
+		
+		{
+			// system data for the frame
+			TSharedPtr<FJsonObject> SystemJson = EmitterFrameToJson(SimCache, INDEX_NONE, FrameIndex);
+			if (!SystemJson.IsValid())
+			{
+				UE_LOG(LogNiagaraSimCacheJson, Warning, TEXT("Unable to export system data for cache %s"), *SimCache.GetName());
+				return false;
+			}
+			FString SystemFile = FPaths::Combine(FrameFolder, "_SystemAttributes_.json");
+
+			FString OutputString;
+			TSharedRef<TJsonWriter<>> Writer = TJsonWriterFactory<>::Create(&OutputString);
+			FJsonSerializer::Serialize(SystemJson.ToSharedRef(), Writer);
+			FFileHelper::SaveStringToFile(OutputString, *SystemFile);
+			SlowTask.TickProgress();
+		}
+
+		for (int EmitterIndex = 0; EmitterIndex < EmitterCount; EmitterIndex++)
+		{
+			// emitter data	
+			TSharedPtr<FJsonObject> EmitterJson = EmitterFrameToJson(SimCache, EmitterIndex, FrameIndex);
+			if (!EmitterJson.IsValid())
+			{
+				UE_LOG(LogNiagaraSimCacheJson, Warning, TEXT("Unable to export data for emitter %s"), *SimCache.GetEmitterName(EmitterIndex).ToString());
+				return false;
+			}
+			FString EmitterFile = FPaths::Combine(FrameFolder, FPaths::MakeValidFileName(SimCache.GetEmitterName(EmitterIndex).ToString() + ".json", '_'));
+
+			FString OutputString;
+			TSharedRef<TJsonWriter<>> Writer = TJsonWriterFactory<>::Create(&OutputString);
+			FJsonSerializer::Serialize(EmitterJson.ToSharedRef(), Writer);
+			FFileHelper::SaveStringToFile(OutputString, *EmitterFile);
+			SlowTask.TickProgress();
+		}
+
+		for (const FNiagaraVariableBase& DataInterface : SimCache.GetStoredDataInterfaces())
+		{
+			// data interfaces
+			if (UObject* StorageObject = SimCache.GetDataInterfaceStorageObject(DataInterface))
+			{
+				UClass* Class = DataInterface.GetType().GetClass();
+				INiagaraSimCacheCustomStorageInterface* DataInterfaceCDO = Class ? Cast<INiagaraSimCacheCustomStorageInterface>(Class->GetDefaultObject()) : nullptr;
+				if (DataInterfaceCDO)
+				{
+					FString FilenamePrefix = FPaths::MakeValidFileName(DataInterface.GetName().ToString(), '_');
+					FilenamePrefix.ReplaceCharInline('.', '-');
+					TSharedPtr<FJsonObject> Json = DataInterfaceCDO->SimCacheToJson(StorageObject, FrameIndex, FrameFolder, FilenamePrefix);
+					if (Json.IsValid())
+					{
+						FString OutputString;
+						TSharedRef<TJsonWriter<>> Writer = TJsonWriterFactory<>::Create(&OutputString);
+						FJsonSerializer::Serialize(Json.ToSharedRef(), Writer);
+						FFileHelper::SaveStringToFile(OutputString, *FPaths::Combine(FrameFolder, FilenamePrefix + ".json"));
+					}
+				}
+			}
+			SlowTask.TickProgress();
+		}
+	}
+
+	return true;
+}
diff --git a/Engine/Plugins/FX/Niagara/Source/NiagaraEditor/NiagaraEditor.Build.cs b/Engine/Plugins/FX/Niagara/Source/NiagaraEditor/NiagaraEditor.Build.cs
--- a/Engine/Plugins/FX/Niagara/Source/NiagaraEditor/NiagaraEditor.Build.cs
+++ b/Engine/Plugins/FX/Niagara/Source/NiagaraEditor/NiagaraEditor.Build.cs
@@ -17,6 +17,7 @@ public class NiagaraEditor : ModuleRules
 				"ApplicationCore",
                 "InputCore",
 				"RenderCore",
+				"Json",
 				"Slate", 
 				"SlateCore",
 				"SlateNullRenderer",
diff --git a/Engine/Plugins/FX/Niagara/Source/NiagaraEditor/Private/AssetDefinitions/AssetDefinition_NiagaraSimCache.cpp b/Engine/Plugins/FX/Niagara/Source/NiagaraEditor/Private/AssetDefinitions/AssetDefinition_NiagaraSimCache.cpp
--- a/Engine/Plugins/FX/Niagara/Source/NiagaraEditor/Private/AssetDefinitions/AssetDefinition_NiagaraSimCache.cpp
+++ b/Engine/Plugins/FX/Niagara/Source/NiagaraEditor/Private/AssetDefinitions/AssetDefinition_NiagaraSimCache.cpp
@@ -2,8 +2,19 @@
 
 #include "AssetDefinition_NiagaraSimCache.h"
 
+#include "ContentBrowserMenuContexts.h"
+#include "DesktopPlatformModule.h"
+#include "EditorDirectories.h"
 #include "NiagaraEditorStyle.h"
+#include "NiagaraSimCacheJson.h"
+#include "Framework/Application/SlateApplication.h"
+#include "Framework/Notifications/NotificationManager.h"
 #include "Toolkits/NiagaraSimCacheToolkit.h"
+#include "Widgets/Notifications/SNotificationList.h"
+
+#define LOCTEXT_NAMESPACE "UAssetDefinition_NiagaraSimCache"
+
+DEFINE_LOG_CATEGORY_STATIC(LogNiagaraSimCache, Log, All);
 
 FLinearColor UAssetDefinition_NiagaraSimCache::GetAssetColor() const
 {
@@ -20,3 +31,94 @@ EAssetCommandResult UAssetDefinition_NiagaraSimCache::OpenAssets(const FAssetOpe
 
 	return EAssetCommandResult::Handled;
 }
+
+// Menu Extensions
+//--------------------------------------------------------------------
+
+namespace MenuExtension_NiagaraSimCache
+{	
+	void ExportToDisk(const FToolMenuContext& InContext)
+	{
+		FString ExportFolder;
+		
+		if (IDesktopPlatform* DesktopPlatform = FDesktopPlatformModule::Get())
+		{
+			const bool bFolderPicked = DesktopPlatform->OpenDirectoryDialog(
+				FSlateApplication::Get().FindBestParentWindowHandleForDialogs(nullptr),
+				LOCTEXT("ExportSimCache", "Pick SimCache Export Folder" ).ToString(),
+				*(FEditorDirectories::Get().GetLastDirectory(ELastDirectory::GENERIC_EXPORT)),
+				ExportFolder);
+
+			if (!bFolderPicked)
+			{
+				return;
+			}
+			FEditorDirectories::Get().SetLastDirectory(ELastDirectory::GENERIC_EXPORT, ExportFolder); 
+		}
+		
+		bool bWarning = false;
+		const UContentBrowserAssetContextMenuContext* CBContext = UContentBrowserAssetContextMenuContext::FindContextWithAssets(InContext);
+		IFileManager& FileManager = IFileManager::Get();
+		for (UNiagaraSimCache* Cache : CBContext->LoadSelectedObjects<UNiagaraSimCache>())
+		{
+			if (Cache == nullptr)
+			{
+				continue;
+			}
+
+			FString CacheRootFolder = FPaths::Combine(ExportFolder, FPaths::MakeValidFileName(Cache->GetName(), '_'));
+			if (FileManager.DirectoryExists(*CacheRootFolder))
+			{
+				if (!FileManager.DeleteDirectory(*CacheRootFolder, true, true))
+				{
+					UE_LOG(LogNiagaraSimCache, Warning, TEXT("Unable to delete existing folder %s"), *CacheRootFolder);
+					bWarning = true;
+					continue;
+				}
+			}
+			if (!FileManager.MakeDirectory(*CacheRootFolder, true))
+			{
+				UE_LOG(LogNiagaraSimCache, Warning, TEXT("Unable to create folder %s"), *CacheRootFolder);
+				bWarning = true;
+				continue;
+			}
+			if (!FNiagaraSimCacheJson::DumpToFile(*Cache, CacheRootFolder, FNiagaraSimCacheJson::EExportType::SeparateEachFrame))
+			{
+				bWarning = true;
+			}
+		}
+		
+		FNotificationInfo Info(LOCTEXT("ExportToDisk_DoneInfo", "Export completed."));
+		Info.ExpireDuration = 4.0f;
+		if (bWarning)
+		{
+			Info.Text = LOCTEXT("ExportData_DoneWarn", "Export completed with warnings.\nPlease check the log.");
+			Info.Image = FCoreStyle::Get().GetBrush(TEXT("MessageLog.Warning"));
+		}
+		FSlateNotificationManager::Get().AddNotification(Info);
+	}
+	
+	static FDelayedAutoRegisterHelper DelayedAutoRegister(EDelayedRegisterRunPhase::EndOfEngineInit, []{ 
+		UToolMenus::RegisterStartupCallback(FSimpleMulticastDelegate::FDelegate::CreateLambda([]()
+		{
+			FToolMenuOwnerScoped OwnerScoped("Niagara SimCache");
+			UToolMenu* Menu = UE::ContentBrowser::ExtendToolMenu_AssetContextMenu(UNiagaraSimCache::StaticClass());
+		
+			FToolMenuSection& Section = Menu->FindOrAddSection("GetAssetActions");
+			Section.AddDynamicEntry(NAME_None, FNewToolMenuSectionDelegate::CreateLambda([](FToolMenuSection& InSection)
+			{
+				{
+					const TAttribute<FText> Label = LOCTEXT("ExportToDisk", "Export To Disk");
+					const TAttribute<FText> ToolTip = LOCTEXT("ExportToDiskTooltip", "Exports the raw data for each frame to disk. Note that data from data interfaces is only exported if they implement support for it.");
+					const FSlateIcon Icon = FSlateIcon();
+
+					FToolUIAction UIAction;
+					UIAction.ExecuteAction = FToolMenuExecuteAction::CreateStatic(&ExportToDisk);
+					InSection.AddMenuEntry("ExportToDisk", Label, ToolTip, Icon, UIAction);
+				}
+			}));
+		}));
+	});
+}
+
+#undef LOCTEXT_NAMESPACE
