Disabled incorrect static analysis warning

diff --git a/Engine/Plugins/FX/Niagara/Source/Niagara/Private/NiagaraDataInterfaceRenderTarget2D.cpp b/Engine/Plugins/FX/Niagara/Source/Niagara/Private/NiagaraDataInterfaceRenderTarget2D.cpp
--- a/Engine/Plugins/FX/Niagara/Source/Niagara/Private/NiagaraDataInterfaceRenderTarget2D.cpp
+++ b/Engine/Plugins/FX/Niagara/Source/Niagara/Private/NiagaraDataInterfaceRenderTarget2D.cpp
@@ -782,7 +782,7 @@ namespace NiagaraRenderTarget2DImageHelper
 			OutColor.A = *reinterpret_cast<const float*>(SourceData + 12);
 		case PF_G16R16F:
 		case PF_G16R16F_FILTER:
-			FMemory::Memcpy(&OutColor.R, SourceData, 4);
+			FMemory::Memcpy(&OutColor.R, SourceData, 4); //-V512
 		default: // do nothing;
 			break;
 		}
