Allow individual module inputs to opt out of the "immediate update" mode when the user is typing. This is mainly for inputs like grid size that cause a big lag spike when the whole system needs to adjust with each keystroke.

#tests Tested by changing the option in the "widget customization" part of a module input's details, then changing the value for that input in the stack

diff --git a/Engine/Plugins/FX/Niagara/Source/Niagara/Public/NiagaraVariableMetaData.h b/Engine/Plugins/FX/Niagara/Source/Niagara/Public/NiagaraVariableMetaData.h
--- a/Engine/Plugins/FX/Niagara/Source/Niagara/Public/NiagaraVariableMetaData.h
+++ b/Engine/Plugins/FX/Niagara/Source/Niagara/Public/NiagaraVariableMetaData.h
@@ -159,7 +159,8 @@ struct FNiagaraInputParameterCustomization
 	UPROPERTY(EditAnywhere, Category="Customization", meta=(EditCondition="WidgetType == ENiagaraInputWidgetType::SegmentedButtons", EditConditionHides))
 	TArray<FWidgetSegmentValueOverride> SegmentValueOverrides;
 
-	UPROPERTY()
+	// If set to true then values will not be updated while typing, only when committing the input field. This is useful for values that cause big downstream changes and computational overhead, like grid size or spawned particles. 
+	UPROPERTY(EditAnywhere, Category="Customization")
 	bool bBroadcastValueChangesOnCommitOnly = false;
 	
 	/** If true then the input is also displayed and editable as a 3d widget in the viewport (vector and transform types only). */
