- Call PreEditChange when modifying a baker texture
- This seems to be a new requirement and fixes a check() that you would hit without it

#rnx
#tests Use the baker with 2D / 3D output

diff --git a/Engine/Plugins/FX/Niagara/Source/NiagaraEditor/Private/NiagaraBakerRendererOutputTexture2D.cpp b/Engine/Plugins/FX/Niagara/Source/NiagaraEditor/Private/NiagaraBakerRendererOutputTexture2D.cpp
--- a/Engine/Plugins/FX/Niagara/Source/NiagaraEditor/Private/NiagaraBakerRendererOutputTexture2D.cpp
+++ b/Engine/Plugins/FX/Niagara/Source/NiagaraEditor/Private/NiagaraBakerRendererOutputTexture2D.cpp
@@ -176,14 +176,13 @@ void FNiagaraBakerRendererOutputTexture2D::BakeFrame(FNiagaraBakerFeedbackContex
 		{
 			const bool bIsPoT = FMath::IsPowerOfTwo(BakerOutput->FrameSize.X) && FMath::IsPowerOfTwo(BakerOutput->FrameSize.Y);
 
+			OutputTexture->PreEditChange(nullptr);
 			OutputTexture->Source.Init(BakerOutput->FrameSize.X, BakerOutput->FrameSize.Y, 1, 1, TSF_RGBA16F, (const uint8*)(FrameResults.GetData()));
 			OutputTexture->PowerOfTwoMode = ETexturePowerOfTwoSetting::None;
 			OutputTexture->MipGenSettings = bIsPoT ? TextureMipGenSettings::TMGS_FromTextureGroup : TextureMipGenSettings::TMGS_NoMipmaps;
 			OutputTexture->AddressX = BakerOutput->bSetTextureAddressX ? BakerOutput->TextureAddressX : OutputTexture->AddressX;
 			OutputTexture->AddressY = BakerOutput->bSetTextureAddressY ? BakerOutput->TextureAddressY : OutputTexture->AddressY;
-			OutputTexture->UpdateResource();
 			OutputTexture->PostEditChange();
-			OutputTexture->MarkPackageDirty();
 		}
 	}
 
@@ -206,6 +205,7 @@ void FNiagaraBakerRendererOutputTexture2D::EndBake(FNiagaraBakerFeedbackContext&
 		{
 			const bool bIsPoT = FMath::IsPowerOfTwo(BakerOutput->AtlasTextureSize.X) && FMath::IsPowerOfTwo(BakerOutput->AtlasTextureSize.Y);
 
+			OutputTexture->PreEditChange(nullptr);
 			OutputTexture->Source.Init(BakerOutput->AtlasTextureSize.X, BakerOutput->AtlasTextureSize.Y, 1, 1, TSF_RGBA16F, (const uint8*)(BakeAtlasTextureData.GetData()));
 			OutputTexture->PowerOfTwoMode = ETexturePowerOfTwoSetting::None;
 			OutputTexture->MipGenSettings = bIsPoT ? TextureMipGenSettings::TMGS_FromTextureGroup : TextureMipGenSettings::TMGS_NoMipmaps;
diff --git a/Engine/Plugins/FX/Niagara/Source/NiagaraEditor/Private/NiagaraBakerRendererOutputVolumeTexture.cpp b/Engine/Plugins/FX/Niagara/Source/NiagaraEditor/Private/NiagaraBakerRendererOutputVolumeTexture.cpp
--- a/Engine/Plugins/FX/Niagara/Source/NiagaraEditor/Private/NiagaraBakerRendererOutputVolumeTexture.cpp
+++ b/Engine/Plugins/FX/Niagara/Source/NiagaraEditor/Private/NiagaraBakerRendererOutputVolumeTexture.cpp
@@ -476,12 +476,11 @@ void FNiagaraBakerRendererOutputVolumeTexture::BakeFrame(FNiagaraBakerFeedbackCo
 			const FString AssetFullName	= OutputVolumeTexture->GetAssetPath(OutputVolumeTexture->FramesAssetPathFormat, FrameIndex);
 			if (UVolumeTexture* OutputTexture = UNiagaraBakerOutput::GetOrCreateAsset<UVolumeTexture, UVolumeTextureFactory>(AssetFullName))
 			{
+				OutputTexture->PreEditChange(nullptr);
 				OutputTexture->Source.Init(TextureSize.X, TextureSize.Y, TextureSize.Z, 1, TSF_RGBA16F, (const uint8*)(TextureData.GetData()));
 				OutputTexture->PowerOfTwoMode = ETexturePowerOfTwoSetting::None;
 				OutputTexture->MipGenSettings = TextureMipGenSettings::TMGS_NoMipmaps;
-				OutputTexture->UpdateResource();
 				OutputTexture->PostEditChange();
-				OutputTexture->MarkPackageDirty();
 			}
 		}
 
@@ -527,6 +526,7 @@ void FNiagaraBakerRendererOutputVolumeTexture::EndBake(FNiagaraBakerFeedbackCont
 		const FString AssetFullName = OutputVolumeTexture->GetAssetPath(OutputVolumeTexture->AtlasAssetPathFormat, 0);
 		if (UVolumeTexture* OutputTexture = UNiagaraBakerOutput::GetOrCreateAsset<UVolumeTexture, UVolumeTextureFactory>(AssetFullName))
 		{
+			OutputTexture->PreEditChange(nullptr);
 			OutputTexture->Source.Init(BakeAtlasTextureSize.X, BakeAtlasTextureSize.Y, BakeAtlasTextureSize.Z, 1, TSF_RGBA16F, (const uint8*)(BakeAtlasTextureData.GetData()));
 			OutputTexture->PowerOfTwoMode = ETexturePowerOfTwoSetting::None;
 			OutputTexture->MipGenSettings = TextureMipGenSettings::TMGS_NoMipmaps;
