Fix for user exposed curves in cooked builds not getting their LUT range copied over properly.

diff --git a/Engine/Plugins/FX/Niagara/Source/Niagara/Private/NiagaraDataInterfaceCurveBase.cpp b/Engine/Plugins/FX/Niagara/Source/Niagara/Private/NiagaraDataInterfaceCurveBase.cpp
--- a/Engine/Plugins/FX/Niagara/Source/Niagara/Private/NiagaraDataInterfaceCurveBase.cpp
+++ b/Engine/Plugins/FX/Niagara/Source/Niagara/Private/NiagaraDataInterfaceCurveBase.cpp
@@ -246,6 +246,9 @@ bool UNiagaraDataInterfaceCurveBase::CopyToInternal(UNiagaraDataInterface* Desti
 	UNiagaraDataInterfaceCurveBase* DestinationTyped = CastChecked<UNiagaraDataInterfaceCurveBase>(Destination);
 	DestinationTyped->bUseLUT = bUseLUT;
 	DestinationTyped->ShaderLUT = ShaderLUT;
+	DestinationTyped->LUTMinTime = LUTMinTime;
+	DestinationTyped->LUTMaxTime = LUTMaxTime;
+	DestinationTyped->LUTInvTimeRange = LUTInvTimeRange;
 	DestinationTyped->LUTNumSamplesMinusOne = LUTNumSamplesMinusOne;
 #if WITH_EDITORONLY_DATA
 	DestinationTyped->CurveAsset = CurveAsset;
