- Remove FName creation from Debug Sim Name
- The FName is now the system name, and the full debug name is the System:Emitter:Version information
- The FName is only really used in CSV profiling which in turn doesn't display the emitter name when charted

#rnx

diff --git a/Engine/Plugins/FX/Niagara/Source/Niagara/Classes/NiagaraComputeExecutionContext.h b/Engine/Plugins/FX/Niagara/Source/Niagara/Classes/NiagaraComputeExecutionContext.h
--- a/Engine/Plugins/FX/Niagara/Source/Niagara/Classes/NiagaraComputeExecutionContext.h
+++ b/Engine/Plugins/FX/Niagara/Source/Niagara/Classes/NiagaraComputeExecutionContext.h
@@ -112,15 +112,15 @@ struct FNiagaraComputeExecutionContext : public INiagaraComputeDataBufferInterfa
 #if WITH_NIAGARA_DEBUG_EMITTER_NAME
 	FName GetDebugSimFName() const { return DebugSimFName; }
 	const TCHAR* GetDebugSimName() const { return *DebugSimName; }
-	void SetDebugSimName(const TCHAR* InDebugSimName)
+	void SetDebugSimName(const TCHAR* InDebugSimName, FName InDebugSimFName)
 	{
 		DebugSimName = InDebugSimName;
-		DebugSimFName = FName(DebugSimName);
+		DebugSimFName = InDebugSimFName;
 	}
 #else
 	FName GetDebugSimFName() const { return NAME_None; }
 	const TCHAR* GetDebugSimName() const { return TEXT(""); }
-	void SetDebugSimName(const TCHAR*) { }
+	void SetDebugSimName(const TCHAR*, FName) { }
 #endif
 
 #if STATS
diff --git a/Engine/Plugins/FX/Niagara/Source/Niagara/Classes/NiagaraDataSet.h b/Engine/Plugins/FX/Niagara/Source/Niagara/Classes/NiagaraDataSet.h
--- a/Engine/Plugins/FX/Niagara/Source/Niagara/Classes/NiagaraDataSet.h
+++ b/Engine/Plugins/FX/Niagara/Source/Niagara/Classes/NiagaraDataSet.h
@@ -325,7 +325,7 @@ public:
 		return *DestinationData;
 	}
 
-	NIAGARA_API void AllocateGPUFreeIDs(uint32 InNumInstances, FRHICommandList& RHICmdList, ERHIFeatureLevel::Type FeatureLevel, FName DebugSimName);
+	NIAGARA_API void AllocateGPUFreeIDs(uint32 InNumInstances, FRHICommandList& RHICmdList, ERHIFeatureLevel::Type FeatureLevel, const TCHAR* DebugSimName);
 
 	void SetMaxInstanceCount(uint32 InMaxInstanceCount) { MaxInstanceCount = InMaxInstanceCount; }
 	void SetMaxAllocationCount(uint32 InMaxAllocationCount) { MaxAllocationCount = InMaxAllocationCount; }
diff --git a/Engine/Plugins/FX/Niagara/Source/Niagara/Private/NiagaraDataSet.cpp b/Engine/Plugins/FX/Niagara/Source/Niagara/Private/NiagaraDataSet.cpp
--- a/Engine/Plugins/FX/Niagara/Source/Niagara/Private/NiagaraDataSet.cpp
+++ b/Engine/Plugins/FX/Niagara/Source/Niagara/Private/NiagaraDataSet.cpp
@@ -445,7 +445,7 @@ void FNiagaraDataSet::Dump(int32 StartIndex, int32 NumInstances, const FString&
 	}
 }
 
-void FNiagaraDataSet::AllocateGPUFreeIDs(uint32 InNumInstances, FRHICommandList& RHICmdList, ERHIFeatureLevel::Type FeatureLevel, FName DebugSimName)
+void FNiagaraDataSet::AllocateGPUFreeIDs(uint32 InNumInstances, FRHICommandList& RHICmdList, ERHIFeatureLevel::Type FeatureLevel, const TCHAR* DebugSimName)
 {
 	checkSlow(GetSimTarget() == ENiagaraSimTarget::GPUComputeSim && RequiresPersistentIDs());
 
@@ -465,17 +465,14 @@ void FNiagaraDataSet::AllocateGPUFreeIDs(uint32 InNumInstances, FRHICommandList&
 	BufferSizeBytes -= GPUFreeIDs.NumBytes;
 #endif
 
-	SCOPED_DRAW_EVENTF(RHICmdList, NiagaraGPUComputeInitFreeIDs, TEXT("Init Free IDs - %s"), DebugSimName);
+	SCOPED_DRAW_EVENT(RHICmdList, NiagaraGPUComputeInitFreeIDs);
 
-	TCHAR DebugBufferName[128];
-	{
-		TCHAR const Str[] = TEXT("NiagaraFreeIDList_");
-		FCString::Strncpy(DebugBufferName, Str, UE_ARRAY_COUNT(DebugBufferName));
-		DebugSimName.ToString(DebugBufferName + UE_ARRAY_COUNT(Str), UE_ARRAY_COUNT(DebugBufferName) - UE_ARRAY_COUNT(Str));
-	}
+	FNameBuilder DebugBufferName;
+	DebugBufferName.Append(TEXT("NiagaraFreeIDList_"));
+	DebugBufferName.Append(DebugSimName);
 
 	FRWBuffer NewFreeIDsBuffer;
-	NewFreeIDsBuffer.Initialize(RHICmdList, DebugBufferName, sizeof(int32), NumIDsToAlloc, EPixelFormat::PF_R32_SINT, BUF_Static);
+	NewFreeIDsBuffer.Initialize(RHICmdList, DebugBufferName.ToString(), sizeof(int32), NumIDsToAlloc, EPixelFormat::PF_R32_SINT, BUF_Static);
 
 	FRHIShaderResourceView* ExistingBuffer = GPUNumAllocatedIDs > 0 ? GPUFreeIDs.SRV.GetReference() : FNiagaraRenderer::GetDummyIntBuffer();
 
diff --git a/Engine/Plugins/FX/Niagara/Source/Niagara/Private/NiagaraEmitterInstanceImpl.cpp b/Engine/Plugins/FX/Niagara/Source/Niagara/Private/NiagaraEmitterInstanceImpl.cpp
--- a/Engine/Plugins/FX/Niagara/Source/Niagara/Private/NiagaraEmitterInstanceImpl.cpp
+++ b/Engine/Plugins/FX/Niagara/Source/Niagara/Private/NiagaraEmitterInstanceImpl.cpp
@@ -362,7 +362,7 @@ void FNiagaraEmitterInstanceImpl::Init(int32 InEmitterIndex)
 		{
 			GPUExecContext = new FNiagaraComputeExecutionContext();
 			GPUExecContext->InitParams(EmitterData->GetGPUComputeScript(), EmitterData->GetSimStageExcecutionData(), EmitterData->SimTarget);
-			GPUExecContext->SetDebugSimName(EmitterData->GetDebugSimName());
+			GPUExecContext->SetDebugSimName(EmitterData->GetDebugSimName(), ParentSystemInstance ? ParentSystemInstance->GetSystem()->GetFName() : NAME_None);
 			GPUExecContext->ProfilingComponentPtr = ParentSystemInstance ? ParentSystemInstance->GetAttachComponent() : nullptr;
 			GPUExecContext->ProfilingEmitterPtr = VersionedEmitter.ToWeakPtr();
 			GPUExecContext->MainDataSet = ParticleDataSet;
diff --git a/Engine/Plugins/FX/Niagara/Source/Niagara/Private/NiagaraGpuComputeDispatch.cpp b/Engine/Plugins/FX/Niagara/Source/Niagara/Private/NiagaraGpuComputeDispatch.cpp
--- a/Engine/Plugins/FX/Niagara/Source/Niagara/Private/NiagaraGpuComputeDispatch.cpp
+++ b/Engine/Plugins/FX/Niagara/Source/Niagara/Private/NiagaraGpuComputeDispatch.cpp
@@ -790,7 +790,7 @@ void FNiagaraGpuComputeDispatch::PrepareTicksForProxy(FRHICommandListImmediate&
 					if (bRaisedWarningThisFrame == false)
 					{
 						bRaisedWarningThisFrame = true;
-						GEngine->AddOnScreenDebugMessage(uint64(this), 1.f, FColor::White, *FString::Printf(TEXT("GPU Simulation(%s) will not show in preview mode, as we may sample from wrong SceneTextures buffer."), *ComputeContext->GetDebugSimFName().ToString()));
+						GEngine->AddOnScreenDebugMessage(uint64(this), 1.f, FColor::White, *FString::Printf(TEXT("GPU Simulation(%s) will not show in preview mode, as we may sample from wrong SceneTextures buffer."), *ComputeContext->GetDebugSimName()));
 					}
 					continue;
 				}
@@ -962,7 +962,7 @@ void FNiagaraGpuComputeDispatch::PrepareTicksForProxy(FRHICommandListImmediate&
 		//-OPT: We should allocate all GPU free IDs together since they require a transition
 		if (ComputeContext->MainDataSet->RequiresPersistentIDs())
 		{
-			ComputeContext->MainDataSet->AllocateGPUFreeIDs(ComputeContext->CurrentMaxAllocateInstances_RT + 1, RHICmdList, FeatureLevel, ComputeContext->GetDebugSimFName());
+			ComputeContext->MainDataSet->AllocateGPUFreeIDs(ComputeContext->CurrentMaxAllocateInstances_RT + 1, RHICmdList, FeatureLevel, ComputeContext->GetDebugSimName());
 		}
 
 		// Allocate space for the buffers we need to perform ticking.  In cases of multiple ticks or multiple write stages we need 3 buffers (current rendered and two simulation buffers).
