- Dynamic Parameter Channel Masking Support for Stateless & Fix Per Script Usage with Stateful
- Stateful emitters will output a mask per script type, this fixes an issue where static variables might differ between scripts causing compilation issues, we now concatenate all the script masks together.  While this doesn't provide the most narrow scope for the masks (i.e. event dynamic params would impact all scripts) it is an improvement over how it used to work before.  Ideally we would have a binding for the channel mask, but that's not possible due to VF's not binding to the int buffer, and if they did we would need a sampler slot which would cause issues on mobile.

#virtualized

diff --git a/Engine/Plugins/FX/Niagara/Source/Niagara/Internal/Stateless/Modules/NiagaraStatelessModule_DynamicMaterialParameters.h b/Engine/Plugins/FX/Niagara/Source/Niagara/Internal/Stateless/Modules/NiagaraStatelessModule_DynamicMaterialParameters.h
--- a/Engine/Plugins/FX/Niagara/Source/Niagara/Internal/Stateless/Modules/NiagaraStatelessModule_DynamicMaterialParameters.h
+++ b/Engine/Plugins/FX/Niagara/Source/Niagara/Internal/Stateless/Modules/NiagaraStatelessModule_DynamicMaterialParameters.h
@@ -149,6 +149,17 @@ public:
 		}
 	}
 
+
+	int32 GetRendererChannelMask() const
+	{
+		uint32 ChannelMask = 0;
+		for (int32 iParameter = 0; iParameter < NumParameters; ++iParameter)
+		{
+			ChannelMask |= GetParameterChannelMask(iParameter) << (iParameter * NumChannelPerParameter);
+		}
+		return ChannelMask;
+	}
+
 	virtual void BuildEmitterData(const FNiagaraStatelessEmitterDataBuildContext& BuildContext) const override
 	{
 		FModuleBuiltData* BuiltData = BuildContext.AllocateBuiltData<FModuleBuiltData>();
diff --git a/Engine/Plugins/FX/Niagara/Source/Niagara/Internal/Stateless/NiagaraStatelessEmitter.h b/Engine/Plugins/FX/Niagara/Source/Niagara/Internal/Stateless/NiagaraStatelessEmitter.h
--- a/Engine/Plugins/FX/Niagara/Source/Niagara/Internal/Stateless/NiagaraStatelessEmitter.h
+++ b/Engine/Plugins/FX/Niagara/Source/Niagara/Internal/Stateless/NiagaraStatelessEmitter.h
@@ -113,6 +113,10 @@ public:
 
 	NIAGARA_API const TArray<TObjectPtr<UNiagaraStatelessModule>>& GetModules() const { return Modules; }
 
+	template<typename TType>
+	const TType* GetModule() const { return (TType*)GetModule(TType::StaticClass()); }
+	NIAGARA_API UNiagaraStatelessModule* GetModule(UClass* Class) const;
+
 	NIAGARA_API FNiagaraPlatformSet& GetPlatformSet() { return Platforms; }
 
 	UNiagaraStatelessEmitter* CreateAsDuplicate(FName InDuplicateName, UNiagaraSystem& InDuplicateOwnerSystem) const;
diff --git a/Engine/Plugins/FX/Niagara/Source/Niagara/Private/NiagaraMeshRendererProperties.cpp b/Engine/Plugins/FX/Niagara/Source/Niagara/Private/NiagaraMeshRendererProperties.cpp
--- a/Engine/Plugins/FX/Niagara/Source/Niagara/Private/NiagaraMeshRendererProperties.cpp
+++ b/Engine/Plugins/FX/Niagara/Source/Niagara/Private/NiagaraMeshRendererProperties.cpp
@@ -613,11 +613,12 @@ void UNiagaraMeshRendererProperties::CacheFromCompiledData(const FNiagaraDataSet
 #if WITH_EDITORONLY_DATA
 	// Build dynamic parameter mask
 	// Serialize in cooked builds
-	const FVersionedNiagaraEmitterData* EmitterData = GetEmitterData();
-	MaterialParamValidMask  = bDynamicParam0Valid ? GetDynamicParameterChannelMask(EmitterData, DynamicMaterialBinding.GetName(), 0xf) << 0 : 0;
-	MaterialParamValidMask |= bDynamicParam1Valid ? GetDynamicParameterChannelMask(EmitterData, DynamicMaterial1Binding.GetName(), 0xf) << 4 : 0;
-	MaterialParamValidMask |= bDynamicParam2Valid ? GetDynamicParameterChannelMask(EmitterData, DynamicMaterial2Binding.GetName(), 0xf) << 8 : 0;
-	MaterialParamValidMask |= bDynamicParam3Valid ? GetDynamicParameterChannelMask(EmitterData, DynamicMaterial3Binding.GetName(), 0xf) << 12 : 0;
+	MaterialParamValidMask = GetDynamicParameterCombinedChannelMask(
+		bDynamicParam0Valid ? DynamicMaterialBinding.GetName() : NAME_None,
+		bDynamicParam1Valid ? DynamicMaterial1Binding.GetName() : NAME_None,
+		bDynamicParam2Valid ? DynamicMaterial2Binding.GetName() : NAME_None,
+		bDynamicParam3Valid ? DynamicMaterial3Binding.GetName() : NAME_None
+	);
 
 	// Gather LOD information per mesh
 	UNiagaraSystem* OwnerSystem = GetTypedOuter<UNiagaraSystem>();
diff --git a/Engine/Plugins/FX/Niagara/Source/Niagara/Private/NiagaraRendererProperties.cpp b/Engine/Plugins/FX/Niagara/Source/Niagara/Private/NiagaraRendererProperties.cpp
--- a/Engine/Plugins/FX/Niagara/Source/Niagara/Private/NiagaraRendererProperties.cpp
+++ b/Engine/Plugins/FX/Niagara/Source/Niagara/Private/NiagaraRendererProperties.cpp
@@ -12,6 +12,7 @@
 #include "NiagaraSystemImpl.h"
 
 #include "Stateless/NiagaraStatelessEmitter.h"
+#include "Stateless/Modules/NiagaraStatelessModule_DynamicMaterialParameters.h"
 
 #include "Interfaces/ITargetPlatform.h"
 #include "Materials/MaterialInterface.h"
@@ -642,32 +643,77 @@ void UNiagaraRendererProperties::UpdateMaterialParametersMIC(const FNiagaraRende
 
 int32 UNiagaraRendererProperties::GetDynamicParameterChannelMask(const FVersionedNiagaraEmitterData* EmitterData, FName BindingName, int32 DefaultChannelMask) const
 {
-	if (EmitterData == nullptr)
+	if (EmitterData == nullptr || BindingName.IsNone())
 	{
 		return DefaultChannelMask;
 	}
 
 	TOptional<int32> ChannelMask;
+
+	// We store the mask per script type to avoid static variable name collisions so we need to search by Particles.*.DynamicParameterChannelMask
+	FNameBuilder BindingNameSearch;
+	BindingName.ToString(BindingNameSearch);
+	int32 NamespaceLocation = INDEX_NONE;
+	if (!BindingNameSearch.ToView().FindChar('.', NamespaceLocation))
+	{
+		return DefaultChannelMask;
+	}
+	BindingNameSearch.Append(TEXT("ChannelMask"));
+
+	FStringView BindingNamePrefix = BindingNameSearch.ToView().Mid(0, NamespaceLocation + 1);
+	FStringView BindingNamePostfix = BindingNameSearch.ToView().Mid(NamespaceLocation, BindingNameSearch.ToView().Len() - NamespaceLocation);
+
 	EmitterData->ForEachScript(
-		[&ChannelMask, &BindingName](const UNiagaraScript* NiagaraScript)
+		[&ChannelMask, &BindingNamePrefix, &BindingNamePostfix](const UNiagaraScript* NiagaraScript)
 		{
-			const FNiagaraVMExecutableData& VMExecData = NiagaraScript->GetVMExecutableData();
+			const FNiagaraTypeDefinition VariableTypeDef = FNiagaraTypeDefinition::GetIntDef().ToStaticDef();
 
-			FNameBuilder NameBuilder;
-			BindingName.ToString(NameBuilder);
-			NameBuilder.Append(TEXT("ChannelMask"));
-
-			const FNiagaraVariableBase ChannelMaskVariable(FNiagaraTypeDefinition::GetIntDef().ToStaticDef(), FName(NameBuilder));
-			TOptional<int32> ChannelMaskValue = NiagaraScript->GetCompiledStaticVariableValue<int32>(ChannelMaskVariable);
-			if (ChannelMaskValue.IsSet())
+			const FNiagaraVMExecutableData& VMExecData = NiagaraScript->GetVMExecutableData();
+			for (const FNiagaraVariable& StaticVariable : VMExecData.StaticVariablesWritten)
 			{
-				ChannelMask = ChannelMask.Get(0) | ChannelMaskValue.GetValue();
+				if (StaticVariable.GetType() != VariableTypeDef || !StaticVariable.IsDataAllocated())
+				{
+					continue;
+				}
+
+				FNameBuilder VariableName;
+				StaticVariable.GetName().ToString(VariableName);
+				if (VariableName.ToView().StartsWith(BindingNamePrefix) && VariableName.ToView().EndsWith(BindingNamePostfix))
+				{
+					ChannelMask = ChannelMask.Get(0) | StaticVariable.GetValue<int32>();
+				}
 			}
 		}
 	);
 	return ChannelMask.Get(DefaultChannelMask);
 }
 
+int32 UNiagaraRendererProperties::GetDynamicParameterCombinedChannelMask(FName Parameter0Name, FName Parameter1Name, FName Parameter2Name, FName Parameter3Name) const
+{
+	const FVersionedNiagaraEmitterData* EmitterData = GetEmitterData();
+	int32 CombinedChannelMask = 0;
+	if (EmitterData == nullptr)
+	{
+		// This is a bit clunky but we no relationship to do this in a more agnostic way at the moment
+		// We could pass down the owner emitter handle to CacheFromCompiledData.
+		if (const UNiagaraStatelessEmitter* StatelessEmitter = GetTypedOuter<UNiagaraStatelessEmitter>())
+		{
+			if (const UNiagaraStatelessModule_DynamicMaterialParameters* DynamicParameterModule = StatelessEmitter->GetModule<UNiagaraStatelessModule_DynamicMaterialParameters>())
+			{
+				CombinedChannelMask = DynamicParameterModule->GetRendererChannelMask();
+			}
+		}
+	}
+	else
+	{
+		CombinedChannelMask |= GetDynamicParameterChannelMask(EmitterData, Parameter0Name, 0xf) << 0;
+		CombinedChannelMask |= GetDynamicParameterChannelMask(EmitterData, Parameter1Name, 0xf) << 4;
+		CombinedChannelMask |= GetDynamicParameterChannelMask(EmitterData, Parameter2Name, 0xf) << 8;
+		CombinedChannelMask |= GetDynamicParameterChannelMask(EmitterData, Parameter3Name, 0xf) << 12;
+	}
+	return CombinedChannelMask;
+}
+
 FNiagaraVariable UNiagaraRendererProperties::GetBoundAttribute(const FNiagaraVariableAttributeBinding* Binding) const
 {
 	if (Binding->GetParamMapBindableVariable().IsValid())
diff --git a/Engine/Plugins/FX/Niagara/Source/Niagara/Private/NiagaraRibbonRendererProperties.cpp b/Engine/Plugins/FX/Niagara/Source/Niagara/Private/NiagaraRibbonRendererProperties.cpp
--- a/Engine/Plugins/FX/Niagara/Source/Niagara/Private/NiagaraRibbonRendererProperties.cpp
+++ b/Engine/Plugins/FX/Niagara/Source/Niagara/Private/NiagaraRibbonRendererProperties.cpp
@@ -516,11 +516,12 @@ void UNiagaraRibbonRendererProperties::CacheFromCompiledData(const FNiagaraDataS
 #if WITH_EDITORONLY_DATA
 	// Build dynamic parameter mask
 	// Serialized in cooked builds
-	const FVersionedNiagaraEmitterData* EmitterData = GetEmitterData();
-	MaterialParamValidMask  = bDynamicParam0Valid ? GetDynamicParameterChannelMask(EmitterData, DynamicMaterialBinding.GetName(),  0xf) <<  0 : 0;
-	MaterialParamValidMask |= bDynamicParam1Valid ? GetDynamicParameterChannelMask(EmitterData, DynamicMaterial1Binding.GetName(), 0xf) <<  4 : 0;
-	MaterialParamValidMask |= bDynamicParam2Valid ? GetDynamicParameterChannelMask(EmitterData, DynamicMaterial2Binding.GetName(), 0xf) <<  8 : 0;
-	MaterialParamValidMask |= bDynamicParam3Valid ? GetDynamicParameterChannelMask(EmitterData, DynamicMaterial3Binding.GetName(), 0xf) << 12 : 0;
+	MaterialParamValidMask = GetDynamicParameterCombinedChannelMask(
+		bDynamicParam0Valid ? DynamicMaterialBinding.GetName() : NAME_None,
+		bDynamicParam1Valid ? DynamicMaterial1Binding.GetName() : NAME_None,
+		bDynamicParam2Valid ? DynamicMaterial2Binding.GetName() : NAME_None,
+		bDynamicParam3Valid ? DynamicMaterial3Binding.GetName() : NAME_None
+	);
 #endif
 }
 
diff --git a/Engine/Plugins/FX/Niagara/Source/Niagara/Private/NiagaraSpriteRendererProperties.cpp b/Engine/Plugins/FX/Niagara/Source/Niagara/Private/NiagaraSpriteRendererProperties.cpp
--- a/Engine/Plugins/FX/Niagara/Source/Niagara/Private/NiagaraSpriteRendererProperties.cpp
+++ b/Engine/Plugins/FX/Niagara/Source/Niagara/Private/NiagaraSpriteRendererProperties.cpp
@@ -395,11 +395,12 @@ void UNiagaraSpriteRendererProperties::CacheFromCompiledData(const FNiagaraDataS
 #if WITH_EDITORONLY_DATA
 	// Build dynamic parameter mask
 	// Serialize in cooked builds
-	const FVersionedNiagaraEmitterData* EmitterData = GetEmitterData();
-	MaterialParamValidMask  = bDynamicParam0Valid ? GetDynamicParameterChannelMask(EmitterData, DynamicMaterialBinding.GetName(),  0xf) <<  0 : 0;
-	MaterialParamValidMask |= bDynamicParam1Valid ? GetDynamicParameterChannelMask(EmitterData, DynamicMaterial1Binding.GetName(), 0xf) <<  4 : 0;
-	MaterialParamValidMask |= bDynamicParam2Valid ? GetDynamicParameterChannelMask(EmitterData, DynamicMaterial2Binding.GetName(), 0xf) <<  8 : 0;
-	MaterialParamValidMask |= bDynamicParam3Valid ? GetDynamicParameterChannelMask(EmitterData, DynamicMaterial3Binding.GetName(), 0xf) << 12 : 0;
+	MaterialParamValidMask = GetDynamicParameterCombinedChannelMask(
+		bDynamicParam0Valid ? DynamicMaterialBinding.GetName() : NAME_None,
+		bDynamicParam1Valid ? DynamicMaterial1Binding.GetName() : NAME_None,
+		bDynamicParam2Valid ? DynamicMaterial2Binding.GetName() : NAME_None,
+		bDynamicParam3Valid ? DynamicMaterial3Binding.GetName() : NAME_None
+	);
 #endif
 }
 
diff --git a/Engine/Plugins/FX/Niagara/Source/Niagara/Private/Stateless/NiagaraStatelessEmitter.cpp b/Engine/Plugins/FX/Niagara/Source/Niagara/Private/Stateless/NiagaraStatelessEmitter.cpp
--- a/Engine/Plugins/FX/Niagara/Source/Niagara/Private/Stateless/NiagaraStatelessEmitter.cpp
+++ b/Engine/Plugins/FX/Niagara/Source/Niagara/Private/Stateless/NiagaraStatelessEmitter.cpp
@@ -653,6 +653,18 @@ FNiagaraStatelessSpawnInfo* UNiagaraStatelessEmitter::GetSpawnInfoByIndex(int32
 	return SpawnInfo;
 }
 
+UNiagaraStatelessModule* UNiagaraStatelessEmitter::GetModule(UClass* Class) const
+{
+	for (UNiagaraStatelessModule* Module : Modules)
+	{
+		if (Module && Module->GetClass() == Class)
+		{
+			return Module;
+		}
+	}
+	return nullptr;
+}
+
 UNiagaraStatelessEmitter* UNiagaraStatelessEmitter::CreateAsDuplicate(FName InDuplicateName, UNiagaraSystem& InDuplicateOwnerSystem) const
 {
 	UNiagaraStatelessEmitter* NewEmitter = Cast<UNiagaraStatelessEmitter>(StaticDuplicateObject(this, &InDuplicateOwnerSystem));
diff --git a/Engine/Plugins/FX/Niagara/Source/Niagara/Public/NiagaraRendererProperties.h b/Engine/Plugins/FX/Niagara/Source/Niagara/Public/NiagaraRendererProperties.h
--- a/Engine/Plugins/FX/Niagara/Source/Niagara/Public/NiagaraRendererProperties.h
+++ b/Engine/Plugins/FX/Niagara/Source/Niagara/Public/NiagaraRendererProperties.h
@@ -469,6 +469,7 @@ protected:
 	NIAGARA_API void UpdateMaterialParametersMIC(const FNiagaraRendererMaterialParameters& MaterialParameters, TArrayView<UMaterialInterface*> Materials, TArray<TObjectPtr<UMaterialInstanceConstant>>& InOutMICs);
 
 	NIAGARA_API int32 GetDynamicParameterChannelMask(const FVersionedNiagaraEmitterData* EmitterData, FName BindingName, int32 DefaultChannelMask) const;
+	NIAGARA_API int32 GetDynamicParameterCombinedChannelMask(FName Parameter0Name, FName Parameter1Name, FName Parameter2Name, FName Parameter3Name) const;
 #endif
 
 #if WITH_EDITOR
