Use FNAME_Find when searching for function names

diff --git a/Engine/Plugins/FX/Niagara/Source/Niagara/Private/NiagaraComponentRendererProperties.cpp b/Engine/Plugins/FX/Niagara/Source/Niagara/Private/NiagaraComponentRendererProperties.cpp
--- a/Engine/Plugins/FX/Niagara/Source/Niagara/Private/NiagaraComponentRendererProperties.cpp
+++ b/Engine/Plugins/FX/Niagara/Source/Niagara/Private/NiagaraComponentRendererProperties.cpp
@@ -389,11 +389,14 @@ void UNiagaraComponentRendererProperties::UpdateSetterFunctions()
 
 			for (const FString& Prefix : SetterPrefixes)
 			{
-				FName SetterFunctionName = FName(Prefix + PropertyName);
-				SetterFunction = TemplateComponent->FindFunction(SetterFunctionName);
-				if (SetterFunction)
+				FName SetterFunctionName = FName(Prefix + PropertyName, FNAME_Find);
+				if (!SetterFunctionName.IsNone())
 				{
-					break;
+					SetterFunction = TemplateComponent->FindFunction(SetterFunctionName);
+					if (SetterFunction)
+					{
+						break;
+					}
 				}
 			}
 		}
