Fix issue with new rigid instance count culling that would apply culling incorrectly in some cases.
Leaving the feature disabled for 30.10 but will re-enable in 30.20

diff --git a/Engine/Plugins/FX/Niagara/Source/Niagara/Private/NiagaraScalabilityManager.cpp b/Engine/Plugins/FX/Niagara/Source/Niagara/Private/NiagaraScalabilityManager.cpp
--- a/Engine/Plugins/FX/Niagara/Source/Niagara/Private/NiagaraScalabilityManager.cpp
+++ b/Engine/Plugins/FX/Niagara/Source/Niagara/Private/NiagaraScalabilityManager.cpp
@@ -408,7 +408,12 @@ bool FNiagaraScalabilityManager::ApplyScalabilityState(int32 ComponentIndex, ENi
 
 			//We do not allow new activations if it would blow our instance count limits.
 			//Though we do allow currently running but inactive components to (!IsComplete()) to reactivate
-			bool bAllow = GApplyInstanceCountsRigidly == 0 || (Component->IsComplete() == false || (System->GetActiveInstancesCount() < SystemInstanceMax && EffectType->NumInstances < EffectTypeInstanceMax));
+			const FNiagaraSystemScalabilitySettings& ScalabilitySettings = System->GetScalabilitySettings();
+			bool bAllow = GApplyInstanceCountsRigidly == 0;
+			bAllow = bAllow || Component->IsComplete() == false;
+			bAllow = bAllow || (!ScalabilitySettings.bCullMaxInstanceCount || EffectType->NumInstances < EffectTypeInstanceMax);
+			bAllow = bAllow || (!ScalabilitySettings.bCullPerSystemMaxInstanceCount || System->GetActiveInstancesCount() < SystemInstanceMax);
+
 			if(bAllow)
 			{
 				CompState.Apply();
diff --git a/Engine/Plugins/FX/Niagara/Source/Niagara/Private/NiagaraWorldManager.cpp b/Engine/Plugins/FX/Niagara/Source/Niagara/Private/NiagaraWorldManager.cpp
--- a/Engine/Plugins/FX/Niagara/Source/Niagara/Private/NiagaraWorldManager.cpp
+++ b/Engine/Plugins/FX/Niagara/Source/Niagara/Private/NiagaraWorldManager.cpp
@@ -1776,8 +1776,8 @@ void FNiagaraWorldManager::SortedSignificanceCull(UNiagaraEffectType* EffectType
 			//Grab Adjusted instance counts.
 			System->GetMaxInstanceCounts(SystemInstanceMax, EffectTypeInstanceMax, bBudgetCullEnabled);
 
-			bCull = EffectTypeInstCount >= EffectTypeInstanceMax;
-			bCull |= SystemInstCount >= SystemInstanceMax;
+			bCull = ScalabilitySettings.bCullMaxInstanceCount && EffectTypeInstCount >= EffectTypeInstanceMax;
+			bCull |= ScalabilitySettings.bCullPerSystemMaxInstanceCount && SystemInstCount >= SystemInstanceMax;
 
 #if DEBUG_SCALABILITY_STATE
 			OutState.bCulledByGlobalBudget |= bCull;
