Lazy populate Niagara perf stats FNames

diff --git a/Engine/Plugins/FX/Niagara/Source/Niagara/Classes/NiagaraEffectType.h b/Engine/Plugins/FX/Niagara/Source/Niagara/Classes/NiagaraEffectType.h
--- a/Engine/Plugins/FX/Niagara/Source/Niagara/Classes/NiagaraEffectType.h
+++ b/Engine/Plugins/FX/Niagara/Source/Niagara/Classes/NiagaraEffectType.h
@@ -400,7 +400,6 @@ class UNiagaraEffectType : public UObject
 	NIAGARA_API virtual void BeginDestroy()override;
 	NIAGARA_API virtual bool IsReadyForFinishDestroy()override;
 	NIAGARA_API virtual void Serialize(FArchive& Ar)override;
-	NIAGARA_API virtual void PostInitProperties()override;
 	NIAGARA_API virtual void PostLoad()override;
 #if WITH_EDITORONLY_DATA
 	static NIAGARA_API void DeclareConstructClasses(TArray<FTopLevelAssetPath>& OutConstructClasses, const UClass* SpecificSubclass);
@@ -509,6 +508,8 @@ public:
 #endif
 
 #if WITH_PER_FXTYPE_PARTICLE_PERF_STATS
+	void PopulatePerfStatNames();
+
 	//Cached CSV Stat names for this system.
 	FName CSVStat_Count = NAME_None;
 	FName CSVStat_Total = NAME_None;
diff --git a/Engine/Plugins/FX/Niagara/Source/Niagara/Private/NiagaraEffectType.cpp b/Engine/Plugins/FX/Niagara/Source/Niagara/Private/NiagaraEffectType.cpp
--- a/Engine/Plugins/FX/Niagara/Source/Niagara/Private/NiagaraEffectType.cpp
+++ b/Engine/Plugins/FX/Niagara/Source/Niagara/Private/NiagaraEffectType.cpp
@@ -44,18 +44,6 @@ void UNiagaraEffectType::Serialize(FArchive& Ar)
 	Ar.UsingCustomVersion(FNiagaraCustomVersion::GUID);
 }
 
-void UNiagaraEffectType::PostInitProperties()
-{
-	Super::PostInitProperties();
-#if WITH_PER_FXTYPE_PARTICLE_PERF_STATS
-	CSVStat_Total = *FString::Printf(TEXT("FXType_Total/%s"), *GetFName().ToString());
-	CSVStat_GTOnly = *FString::Printf(TEXT("FXType_GTOnly/%s"), *GetFName().ToString());
-	CSVStat_RT = *FString::Printf(TEXT("FXType_RT/%s"), *GetFName().ToString());
-	CSVStat_GPU = *FString::Printf(TEXT("FXType_GPU/%s"), *GetFName().ToString());
-	CSVStat_Count = *FString::Printf(TEXT("FXType_Count/%s"), *GetFName().ToString());
-#endif
-}
-
 void UNiagaraEffectType::PostLoad()
 {
 	Super::PostLoad();
@@ -135,6 +123,22 @@ const FNiagaraEmitterScalabilitySettings& UNiagaraEffectType::GetActiveEmitterSc
 	return Dummy;
 }
 
+#if WITH_PER_FXTYPE_PARTICLE_PERF_STATS
+void UNiagaraEffectType::PopulatePerfStatNames()
+{
+	// Lazy populate to keep unneeded FNames down
+	if (CSVStat_Total == NAME_None)
+	{
+		const FString FNameString = GetFName().ToString();
+		CSVStat_Total = *FString::Printf(TEXT("FXType_Total/%s"), *FNameString);
+		CSVStat_GTOnly = *FString::Printf(TEXT("FXType_GTOnly/%s"), *FNameString);
+		CSVStat_RT = *FString::Printf(TEXT("FXType_RT/%s"), *FNameString);
+		CSVStat_GPU = *FString::Printf(TEXT("FXType_GPU/%s"), *FNameString);
+		CSVStat_Count = *FString::Printf(TEXT("FXType_Count/%s"), *FNameString);
+	}
+}
+#endif
+
 #if WITH_EDITOR
 
 void UNiagaraEffectType::PostEditChangeProperty(struct FPropertyChangedEvent& PropertyChangedEvent)
@@ -456,6 +460,8 @@ bool FParticlePerfStatsListener_EffectType::Tick()
 
 				if (Stats.GetGameThreadStats().NumInstances > 0)
 				{
+					FXType->PopulatePerfStatNames();
+
 					float TotalTime = FPlatformTime::ToMilliseconds64(Stats.GetGameThreadStats().GetTotalCycles()) * 1000.0f;
 					float GTTime = FPlatformTime::ToMilliseconds64(Stats.GetGameThreadStats().GetTotalCycles_GTOnly()) * 1000.0f;
 					//float AvgTime = FPlatformTime::ToMilliseconds64(Stats->GetGameThreadStats().GetPerInstanceAvgCycles()) * 1000.0f;
@@ -510,6 +516,8 @@ void FParticlePerfStatsListener_EffectType::TickRT()
 
 				if (Stats.GetRenderThreadStats().NumInstances > 0)
 				{
+					FXType->PopulatePerfStatNames();
+
 					const float RTTime = FPlatformTime::ToMilliseconds64(Stats.GetRenderThreadStats().GetTotalCycles()) * 1000.0f;
 					//const float RTAvgTime = FPlatformTime::ToMilliseconds64(Stats->GetRenderThreadStats().GetPerInstanceAvgCycles()) * 1000.0f;
 					CSVProfiler->RecordCustomStat(FXType->CSVStat_RT, CSV_CATEGORY_INDEX(Particles), RTTime, ECsvCustomStatOp::Set);
