- Don't set the channel mask if the parameter name is none

#rnx

diff --git a/Engine/Plugins/FX/Niagara/Source/Niagara/Private/NiagaraRendererProperties.cpp b/Engine/Plugins/FX/Niagara/Source/Niagara/Private/NiagaraRendererProperties.cpp
--- a/Engine/Plugins/FX/Niagara/Source/Niagara/Private/NiagaraRendererProperties.cpp
+++ b/Engine/Plugins/FX/Niagara/Source/Niagara/Private/NiagaraRendererProperties.cpp
@@ -645,7 +645,7 @@ int32 UNiagaraRendererProperties::GetDynamicParameterChannelMask(const FVersione
 {
 	if (EmitterData == nullptr || BindingName.IsNone())
 	{
-		return DefaultChannelMask;
+		return 0;
 	}
 
 	TOptional<int32> ChannelMask;
