Fixed issue that would re-order NDC data on CPU leading to CPU/GPU mis-match

diff --git a/Engine/Plugins/FX/Niagara/Source/Niagara/Private/NiagaraDataChannel.cpp b/Engine/Plugins/FX/Niagara/Source/Niagara/Private/NiagaraDataChannel.cpp
--- a/Engine/Plugins/FX/Niagara/Source/Niagara/Private/NiagaraDataChannel.cpp
+++ b/Engine/Plugins/FX/Niagara/Source/Niagara/Private/NiagaraDataChannel.cpp
@@ -744,10 +744,10 @@ int32 FNiagaraDataChannelData::ConsumePublishRequests(UNiagaraDataChannelHandler
 #if WITH_NIAGARA_DEBUGGER
 		FNiagaraDataChannelDebugUtilities::LogWrite(PublishRequest, Owner->GetDataChannel(), TickGroup);
 #endif
-		
-		It.RemoveCurrentSwap();
 	}
 
+	PublishRequests.Reset();
+
 	CPUSimData->EndSimulate();
 
 #if WITH_NIAGARA_DEBUGGER
@@ -764,7 +764,6 @@ int32 FNiagaraDataChannelData::ConsumePublishRequests(UNiagaraDataChannelHandler
 	}
 #endif
 
-	PublishRequests.Reset();
 	return RequestCount;
 }
 
