- Fix static bool material parameter bindings

#rnx

diff --git a/Engine/Plugins/FX/Niagara/Source/NiagaraEditor/Private/Customizations/NiagaraTypeCustomizations.cpp b/Engine/Plugins/FX/Niagara/Source/NiagaraEditor/Private/Customizations/NiagaraTypeCustomizations.cpp
--- a/Engine/Plugins/FX/Niagara/Source/NiagaraEditor/Private/Customizations/NiagaraTypeCustomizations.cpp
+++ b/Engine/Plugins/FX/Niagara/Source/NiagaraEditor/Private/Customizations/NiagaraTypeCustomizations.cpp
@@ -2462,16 +2462,18 @@ TSharedRef<SWidget> FNiagaraRendererMaterialStaticBoolParameterCustomization::On
 				}
 			};
 
-		if (UNiagaraSystem* NiagaraSystem = RenderProperties->GetTypedOuter<UNiagaraSystem>())
-		{
-			AddStaticVariables(NiagaraSystem->GetSystemSpawnScript());
-			AddStaticVariables(NiagaraSystem->GetSystemUpdateScript());
-		}
+		// We must add the emitter alias context before the system scripts as they will contain the emitter spawn & update script
 		if (FVersionedNiagaraEmitterData* EmitterData = NiagaraEmitter.GetEmitterData())
 		{
 			AliasContext.ChangeEmitterNameToEmitter(NiagaraEmitter.Emitter->GetUniqueEmitterName());
 			EmitterData->ForEachScript(AddStaticVariables);
 		}
+
+		if (UNiagaraSystem* NiagaraSystem = RenderProperties->GetTypedOuter<UNiagaraSystem>())
+		{
+			AddStaticVariables(NiagaraSystem->GetSystemSpawnScript());
+			AddStaticVariables(NiagaraSystem->GetSystemUpdateScript());
+		}
 	}
 
 	return MenuBuilder.MakeWidget();
