- Remove invalid DisplayName from lock desired age seek delta get / set functions

#rnx

diff --git a/Engine/Plugins/FX/Niagara/Source/Niagara/Public/NiagaraComponent.h b/Engine/Plugins/FX/Niagara/Source/Niagara/Public/NiagaraComponent.h
--- a/Engine/Plugins/FX/Niagara/Source/Niagara/Public/NiagaraComponent.h
+++ b/Engine/Plugins/FX/Niagara/Source/Niagara/Public/NiagaraComponent.h
@@ -359,13 +359,13 @@ public:
 	/** Gets whether or not the delta time used to tick the system instance when using desired age is locked to the seek delta.  When true, the system instance
 	will only be ticked when the desired age has changed by more than the seek delta.  When false the system instance will be ticked by the change in desired
 	age when not seeking. */
-	UFUNCTION(BlueprintCallable, Category = Niagara, meta = (DisplayName = "Get whether or not to lock the desired age delta time to the seek delta."))
+	UFUNCTION(BlueprintCallable, Category = Niagara)
 	NIAGARA_API bool GetLockDesiredAgeDeltaTimeToSeekDelta() const;
 
 	/** Sets whether or not the delta time used to tick the system instance when using desired age is locked to the seek delta.  When true, the system instance
 	will only be ticked when the desired age has changed by more than the seek delta.  When false the system instance will be ticked by the change in desired
 	age when not seeking. */
-	UFUNCTION(BlueprintCallable, Category = Niagara, meta = (DisplayName = "Set whether or not to lock the desired age delta time to the seek delta."))
+	UFUNCTION(BlueprintCallable, Category = Niagara)
 	NIAGARA_API void SetLockDesiredAgeDeltaTimeToSeekDelta(bool bLock);
 
 	/**
