Fix for niagara system copying crashing in the texture compiler.
#rnx

diff --git a/Engine/Plugins/FX/Niagara/Source/NiagaraEditor/Private/NiagaraThumbnailRenderer.cpp b/Engine/Plugins/FX/Niagara/Source/NiagaraEditor/Private/NiagaraThumbnailRenderer.cpp
--- a/Engine/Plugins/FX/Niagara/Source/NiagaraEditor/Private/NiagaraThumbnailRenderer.cpp
+++ b/Engine/Plugins/FX/Niagara/Source/NiagaraEditor/Private/NiagaraThumbnailRenderer.cpp
@@ -5,7 +5,6 @@
 #include "Engine/Texture2D.h"
 #include "NiagaraEmitter.h"
 #include "NiagaraSystem.h"
-#include "UObject/Package.h"
 
 #include UE_INLINE_GENERATED_CPP_BY_NAME(NiagaraThumbnailRenderer)
 
@@ -29,12 +28,8 @@ UTexture2D* UNiagaraEmitterThumbnailRenderer::GetThumbnailTextureFromObject(UObj
 	UNiagaraEmitter* Emitter = Cast<UNiagaraEmitter>(Object);
 	if (Emitter && Emitter->ThumbnailImage)
 	{
-		// Avoid calling UpdateResource on cooked texture as doing so will destroy the texture's data
-		if (!Emitter->GetPackage()->HasAnyPackageFlags(PKG_Cooked | PKG_FilterEditorOnly))
-		{
-			Emitter->ThumbnailImage->FinishCachePlatformData();
-			Emitter->ThumbnailImage->UpdateResource();
-		}
+		// we have to have completed build in order to get draw/get size
+		Emitter->ThumbnailImage->BlockOnAnyAsyncBuild();
 		return Emitter->ThumbnailImage;
 	}
 	return nullptr;
@@ -45,12 +40,8 @@ UTexture2D* UNiagaraSystemThumbnailRenderer::GetThumbnailTextureFromObject(UObje
 	UNiagaraSystem* System = Cast<UNiagaraSystem>(Object);
 	if (System && System->ThumbnailImage)
 	{
-		// Avoid calling UpdateResource on cooked texture as doing so will destroy the texture's data
-		if (!System->GetPackage()->HasAnyPackageFlags(PKG_Cooked | PKG_FilterEditorOnly))
-		{
-			System->ThumbnailImage->FinishCachePlatformData();
-			System->ThumbnailImage->UpdateResource();
-		}
+		// we have to have completed build in order to get draw/get size
+		System->ThumbnailImage->BlockOnAnyAsyncBuild();
 		return System->ThumbnailImage;
 	}
 	return nullptr;
diff --git a/Engine/Source/Runtime/Engine/Classes/Engine/Texture.h b/Engine/Source/Runtime/Engine/Classes/Engine/Texture.h
--- a/Engine/Source/Runtime/Engine/Classes/Engine/Texture.h
+++ b/Engine/Source/Runtime/Engine/Classes/Engine/Texture.h
@@ -1819,7 +1819,12 @@ public:
 	ENGINE_API bool IsAsyncCacheComplete() const;
 
 	/**
-	 * Blocks on async cache tasks and prepares platform data for use.
+	 * Blocks on async cache tasks and prepares platform data for use. This should only be called by the texture compilation manager.
+	 * If you call this manually when the texture has been registed in the texture compilation manager, it doesn't get cleared there
+	 * and you can end up crashing in FTextureCompilationManager::AddTextures.
+	 * 
+	 * If you need the texture resource after you've made modifications, you should wrap your changes in PreEditChange/PostEditChange,
+	 * then call BlockOnAnyAsyncBuild on your texture.
 	 */
 	ENGINE_API void FinishCachePlatformData();
 
diff --git a/Engine/Source/Runtime/Engine/Private/TextureCompiler.cpp b/Engine/Source/Runtime/Engine/Private/TextureCompiler.cpp
--- a/Engine/Source/Runtime/Engine/Private/TextureCompiler.cpp
+++ b/Engine/Source/Runtime/Engine/Private/TextureCompiler.cpp
@@ -255,6 +255,15 @@ void FTextureCompilingManager::AddTextures(TArrayView<UTexture* const> InTexture
 	// to strcmp on the name of the repro texture and just log the full key suffix. Then you should immediately see the changed
 	// keys right before the crash and you can backsolve what value changed. Once you have that, you can set a data breakpoint on
 	// the property and see who is poking it.
+	//
+	// **
+	//
+	// One thing to be aware of is this can be caused by a system manually calling FinishCachePlatformData + UpdateResource
+	// instead of calling BlockOnAnyAsyncBuild. This causes the async task to become null,
+	// which prevents any IsCompiling / BlockOnAnyAsyncBuild from detecting it, even though it's still pending a PostCompilation
+	// in here. As a result you can edit the DDC key any time between the FinishCachePlatformData and the subsequent CreateResource
+	// call and get this crash. If you have a repro, best bet is to try and get a breakpoint on FinishCachePlatformData for the texture
+	// in question - only the compilation manager shoulid be calling that for editor resources.
 	if (bIsRoutingPostCompilation)
 	{
 		// This has been updated to Fatal because it potentially modifies RegisteredTextureBuckets below which is iterated upon
