- Niagara Audit update to output compressed systems & provide full path for dynamic bounds emitters

#rnx

diff --git a/Engine/Plugins/FX/Niagara/Source/NiagaraEditor/Private/Commandlets/NiagaraSystemAuditCommandlet.cpp b/Engine/Plugins/FX/Niagara/Source/NiagaraEditor/Private/Commandlets/NiagaraSystemAuditCommandlet.cpp
--- a/Engine/Plugins/FX/Niagara/Source/NiagaraEditor/Private/Commandlets/NiagaraSystemAuditCommandlet.cpp
+++ b/Engine/Plugins/FX/Niagara/Source/NiagaraEditor/Private/Commandlets/NiagaraSystemAuditCommandlet.cpp
@@ -315,6 +315,8 @@ bool UNiagaraSystemAuditCommandlet::ProcessNiagaraSystems()
 		FString EmittersWithSimulationStages;
 		bool bHasLights = false;
 		bool bHasEvents = false;
+		bool bHasGPUEmitters = false;
+		bool bHasCPUEmitters = false;
 
 		for (const FNiagaraEmitterHandle& EmitterHandle : NiagaraSystem->GetEmitterHandles())
 		{
@@ -329,6 +331,9 @@ bool UNiagaraSystemAuditCommandlet::ProcessNiagaraSystems()
 				continue;
 			}
 
+			bHasGPUEmitters |= EmitterData->SimTarget == ENiagaraSimTarget::GPUComputeSim;
+			bHasCPUEmitters |= EmitterData->SimTarget != ENiagaraSimTarget::GPUComputeSim;
+
 			if (EmitterData->SimTarget == ENiagaraSimTarget::GPUComputeSim)
 			{
 				// Optionally disable GPU emitters
@@ -441,7 +446,15 @@ bool UNiagaraSystemAuditCommandlet::ProcessNiagaraSystems()
 
 			if ( (EmitterData->CalculateBoundsMode == ENiagaraEmitterCalculateBoundMode::Dynamic) && !NiagaraSystem->bFixedBounds )
 			{
-				EmittersWithDynamicBounds.Append(EmitterData->GetDebugSimName());
+				const FString EmitterName = EmitterHandle.GetUniqueInstanceName();
+				if (EmittersWithDynamicBounds.Find(EmitterName) == INDEX_NONE)
+				{
+					if (EmittersWithDynamicBounds.Len() != 0)
+					{
+						EmittersWithDynamicBounds.Append(TEXT(" "));
+					}
+					EmittersWithDynamicBounds.Append(EmitterName);
+				}
 			}
 		}
 
@@ -480,9 +493,26 @@ bool UNiagaraSystemAuditCommandlet::ProcessNiagaraSystems()
 			NiagaraSystemsWithSimulationStages.Add(FString::Printf(TEXT("%s,%s"), *NiagaraSystem->GetPathName(), *EmittersWithSimulationStages));
 		}
 
+		if (NiagaraSystem->ShouldCompressAttributes())
+		{
+			TStringBuilder<256> StringBuilder;
+			StringBuilder.Append(NiagaraSystem->GetPathName());
+			StringBuilder.Append(",");
+			if (bHasGPUEmitters)
+			{
+				StringBuilder.Append("GPU ");
+			}
+			if (bHasCPUEmitters)
+			{
+				StringBuilder.Append("CPU ");
+			}
+
+			NiagaraSystemsWithCompression.Add(StringBuilder.ToString());
+		}
+
 		if ( !EmittersWithDynamicBounds.IsEmpty() )
 		{
-			NiagaraSystemsWithDynamicBounds.Add(EmittersWithDynamicBounds);
+			NiagaraSystemsWithDynamicBounds.Add(FString::Printf(TEXT("%s,%s"), *NiagaraSystem->GetPathName(), *EmittersWithDynamicBounds));
 		}
 
 		if (SystemUserDataInterfaces.Num() > 0)
@@ -554,6 +584,7 @@ void UNiagaraSystemAuditCommandlet::DumpResults()
 		DumpSimpleSet(NiagaraSystemsWithGPUEmitters, TEXT("NiagaraSystemsWithGPUEmitters"), HeaderString.ToString());
 	}
 	DumpSimpleSet(NiagaraSystemsWithSimulationStages, TEXT("NiagaraSystemsWithSimulationStages"), TEXT("System,Emitters"));
+	DumpSimpleSet(NiagaraSystemsWithCompression, TEXT("NiagaraSystemsWithCompression"), TEXT("System,ExecutionTypes"));
 
 	if (NiagaraDataInterfaceUsage.Num() > 0)
 	{
diff --git a/Engine/Plugins/FX/Niagara/Source/NiagaraEditor/Public/Commandlets/NiagaraSystemAuditCommandlet.h b/Engine/Plugins/FX/Niagara/Source/NiagaraEditor/Public/Commandlets/NiagaraSystemAuditCommandlet.h
--- a/Engine/Plugins/FX/Niagara/Source/NiagaraEditor/Public/Commandlets/NiagaraSystemAuditCommandlet.h
+++ b/Engine/Plugins/FX/Niagara/Source/NiagaraEditor/Public/Commandlets/NiagaraSystemAuditCommandlet.h
@@ -28,6 +28,8 @@ class UNiagaraSystemAuditCommandlet : public UCommandlet
 	TSet<FString> NiagaraSystemsWithDataInterfaceToFind;
 	/** All Niagara Systems & Emitters with simulation stages enabled*/
 	TSet<FString> NiagaraSystemsWithSimulationStages;
+	/** All Niagara systems with compression enabled */
+	TSet<FString> NiagaraSystemsWithCompression;
 
 	/** All Validation Issues */
 	TUniquePtr<FArchive> NiagaraValidationIssues;
