- Fix dangling reference when opening a script editor

#rnx

diff --git a/Engine/Plugins/FX/Niagara/Source/NiagaraEditor/Private/Toolkits/NiagaraScriptToolkit.cpp b/Engine/Plugins/FX/Niagara/Source/NiagaraEditor/Private/Toolkits/NiagaraScriptToolkit.cpp
--- a/Engine/Plugins/FX/Niagara/Source/NiagaraEditor/Private/Toolkits/NiagaraScriptToolkit.cpp
+++ b/Engine/Plugins/FX/Niagara/Source/NiagaraEditor/Private/Toolkits/NiagaraScriptToolkit.cpp
@@ -673,9 +673,8 @@ void FNiagaraScriptToolkit::ExtendToolbar()
 				LOCTEXT("NiagaraShowModuleVersionsTooltip", "Manage different versions of this module script."),
 					FSlateIcon(FAppStyle::GetAppStyleSetName(), "Versions")));
 
-				
 				FUIAction DropdownAction;
-				DropdownAction.IsActionVisibleDelegate = FIsActionButtonVisible::CreateLambda([NiagaraScriptToolkit]() { return NiagaraScriptToolkit->EditedNiagaraScript.Script->GetAllAvailableVersions().Num() > 1; });
+				DropdownAction.IsActionVisibleDelegate = FIsActionButtonVisible::CreateLambda([WeakNiagaraScriptToolkit = NiagaraScriptToolkit.ToWeakPtr()]() { return WeakNiagaraScriptToolkit.IsValid() ? WeakNiagaraScriptToolkit.Pin()->EditedNiagaraScript.Script->GetAllAvailableVersions().Num() > 1 : false; });
 
 				InSection.AddEntry(FToolMenuEntry::InitComboButton("VersionOptions",
 					DropdownAction,
