- Data Interface reading for Niagara Sim Caches
- Add ReadDataInterfaceAs method which will return the storage object (if the types match)
- In the future we will extend the DI interface to allow converting the data to the provided type, i.e. densify the SVT into a volume for a given frame
- For simple types (i.e. array of floats) we need to consider a different API
- Change GetDataInterfaceStorageObject to return a const UObject as we wouldn't want the data to be modified externally

diff --git a/Engine/Plugins/FX/Niagara/Source/Niagara/Classes/NiagaraDataInterfaceArray.h b/Engine/Plugins/FX/Niagara/Source/Niagara/Classes/NiagaraDataInterfaceArray.h
--- a/Engine/Plugins/FX/Niagara/Source/Niagara/Classes/NiagaraDataInterfaceArray.h
+++ b/Engine/Plugins/FX/Niagara/Source/Niagara/Classes/NiagaraDataInterfaceArray.h
@@ -124,7 +124,7 @@ public:
 	virtual bool SimCacheWriteFrame(UObject* StorageObject, int FrameIndex, FNiagaraSystemInstance* SystemInstance, const void* OptionalPerInstanceData, FNiagaraSimCacheFeedbackContext& FeedbackContext) const override;
 	//virtual bool SimCacheEndWrite(UObject* StorageObject) const override;
 	virtual bool SimCacheReadFrame(UObject* StorageObject, int FrameA, int FrameB, float Interp, FNiagaraSystemInstance* SystemInstance, void* OptionalPerInstanceData) override;
-	virtual bool SimCacheCompareFrame(UObject* LhsStorageObject, UObject* RhsStorageObject, int FrameIndex, TOptional<float> Tolerance, FString& OutErrors) const override;
+	virtual bool SimCacheCompareFrame(const UObject* LhsStorageObject, const UObject* RhsStorageObject, int FrameIndex, TOptional<float> Tolerance, FString& OutErrors) const override;
 	//~ INiagaraSimCacheCustomStorageInterface interface END
 	NIAGARA_API FString SimCacheVisualizerRead(const UNDIArraySimCacheData* CacheData, const FNDIArraySimCacheDataFrame& FrameData, int Element) const;
 
diff --git a/Engine/Plugins/FX/Niagara/Source/Niagara/Classes/NiagaraSimCache.h b/Engine/Plugins/FX/Niagara/Source/Niagara/Classes/NiagaraSimCache.h
--- a/Engine/Plugins/FX/Niagara/Source/Niagara/Classes/NiagaraSimCache.h
+++ b/Engine/Plugins/FX/Niagara/Source/Niagara/Classes/NiagaraSimCache.h
@@ -460,7 +460,7 @@ public:
 	NIAGARA_API TArray<FNiagaraVariableBase> GetStoredDataInterfaces() const;
 	
 	/** Returns the actual data we have captured in the SimCache for the given data interface. */
-	NIAGARA_API UObject* GetDataInterfaceStorageObject(const FNiagaraVariableBase& DataInterface) const;
+	NIAGARA_API const UObject* GetDataInterfaceStorageObject(const FNiagaraVariableBase& DataInterface) const;
 
 	/**
 	Get number of active instances for the emitter at the given frame.
@@ -566,7 +566,17 @@ public:
 	*/
 	UFUNCTION(BlueprintCallable, Category = NiagaraSimCache)
 	NIAGARA_API void ReadQuatAttributeWithRebase(TArray<FQuat>& OutValues, FQuat Quat, FName AttributeName = FName("MeshOrientation"), FName EmitterName = NAME_None, int FrameIndex = 0) const;
-	
+
+	/**
+	Reads data interface data from the cache as the requested type.
+	This method will return nullptr if the attribute does not exists or the requests type is not supported by the storage type.
+	*/
+	UFUNCTION(BlueprintCallable, Category = NiagaraSimCache)
+	NIAGARA_API UObject* ReadDataInterfaceAs(UClass* RequestedType, FName AttributeName, int FrameIndex = 0) const;
+
+	template<typename T>
+	T* ReadDataInterfaceAs(FName AttributeName, int FrameIndex = 0) const { return (T*)ReadDataInterfaceAs(T::StaticClass(), AttributeName, FrameIndex); }
+
 private:
 	UPROPERTY(VisibleAnywhere, Category=SimCache)
 	FGuid CacheGuid;
diff --git a/Engine/Plugins/FX/Niagara/Source/Niagara/Classes/NiagaraSimCacheCustomStorageInterface.h b/Engine/Plugins/FX/Niagara/Source/Niagara/Classes/NiagaraSimCacheCustomStorageInterface.h
--- a/Engine/Plugins/FX/Niagara/Source/Niagara/Classes/NiagaraSimCacheCustomStorageInterface.h
+++ b/Engine/Plugins/FX/Niagara/Source/Niagara/Classes/NiagaraSimCacheCustomStorageInterface.h
@@ -62,7 +62,7 @@ public:
 	Called to compare a frame between two separate simulation cache storages. Mainly useful for unit testing.
 	This will be called on the CDO object since we do not have the actual data interface.
 	*/
-	virtual bool SimCacheCompareFrame(UObject* LhsStorageObject, UObject* RhsStorageObject, int FrameIndex, TOptional<float> Tolerance, FString& OutErrors) const { OutErrors = TEXT("Compare not implemented"); return false; }
+	virtual bool SimCacheCompareFrame(const UObject* LhsStorageObject, const UObject* RhsStorageObject, int FrameIndex, TOptional<float> Tolerance, FString& OutErrors) const { OutErrors = TEXT("Compare not implemented"); return false; }
 
 	/**
 	This function allows you to preserve a list of attributes when building a renderer only cache.
diff --git a/Engine/Plugins/FX/Niagara/Source/Niagara/Internal/DataInterface/NiagaraDataInterfaceDataChannelWrite.h b/Engine/Plugins/FX/Niagara/Source/Niagara/Internal/DataInterface/NiagaraDataInterfaceDataChannelWrite.h
--- a/Engine/Plugins/FX/Niagara/Source/Niagara/Internal/DataInterface/NiagaraDataInterfaceDataChannelWrite.h
+++ b/Engine/Plugins/FX/Niagara/Source/Niagara/Internal/DataInterface/NiagaraDataInterfaceDataChannelWrite.h
@@ -186,7 +186,7 @@ public:
 	virtual bool SimCacheWriteFrame(UObject* StorageObject, int FrameIndex, FNiagaraSystemInstance* SystemInstance, const void* OptionalPerInstanceData, FNiagaraSimCacheFeedbackContext& FeedbackContext) const override;
 	virtual bool SimCacheReadFrame(UObject* StorageObject, int FrameA, int FrameB, float Interp, FNiagaraSystemInstance* SystemInstance, void* OptionalPerInstanceData) override;
 	virtual void SimCachePostReadFrame(void* OptionalPerInstanceData, FNiagaraSystemInstance* SystemInstance) override;
-	virtual bool SimCacheCompareFrame(UObject* LhsStorageObject, UObject* RhsStorageObject, int FrameIndex, TOptional<float> Tolerance, FString& OutErrors) const override;
+	virtual bool SimCacheCompareFrame(const UObject* LhsStorageObject, const UObject* RhsStorageObject, int FrameIndex, TOptional<float> Tolerance, FString& OutErrors) const override;
 protected:
 #if WITH_EDITORONLY_DATA
 	NIAGARA_API virtual void GetFunctionsInternal(TArray<FNiagaraFunctionSignature>& OutFunctions) const override;
diff --git a/Engine/Plugins/FX/Niagara/Source/Niagara/Private/DataInterface/NiagaraDataInterfaceDataChannelWrite.cpp b/Engine/Plugins/FX/Niagara/Source/Niagara/Private/DataInterface/NiagaraDataInterfaceDataChannelWrite.cpp
--- a/Engine/Plugins/FX/Niagara/Source/Niagara/Private/DataInterface/NiagaraDataInterfaceDataChannelWrite.cpp
+++ b/Engine/Plugins/FX/Niagara/Source/Niagara/Private/DataInterface/NiagaraDataInterfaceDataChannelWrite.cpp
@@ -691,10 +691,10 @@ void UNiagaraDataInterfaceDataChannelWrite::SimCachePostReadFrame(void* Optional
 	// send data to data channel
 }
 
-bool UNiagaraDataInterfaceDataChannelWrite::SimCacheCompareFrame(UObject* LhsStorageObject, UObject* RhsStorageObject, int FrameIndex, TOptional<float> Tolerance, FString& OutErrors) const
+bool UNiagaraDataInterfaceDataChannelWrite::SimCacheCompareFrame(const UObject* LhsStorageObject, const UObject* RhsStorageObject, int FrameIndex, TOptional<float> Tolerance, FString& OutErrors) const
 {
-	UNDIDataChannelWriteSimCacheData* Storage1 = Cast<UNDIDataChannelWriteSimCacheData>(LhsStorageObject);
-	UNDIDataChannelWriteSimCacheData* Storage2 = Cast<UNDIDataChannelWriteSimCacheData>(RhsStorageObject);
+	const UNDIDataChannelWriteSimCacheData* Storage1 = Cast<const UNDIDataChannelWriteSimCacheData>(LhsStorageObject);
+	const UNDIDataChannelWriteSimCacheData* Storage2 = Cast<const UNDIDataChannelWriteSimCacheData>(RhsStorageObject);
 
 	if (Storage1 == nullptr && Storage2 == nullptr)
 	{
@@ -719,8 +719,8 @@ bool UNiagaraDataInterfaceDataChannelWrite::SimCacheCompareFrame(UObject* LhsSto
 	bool bEqual = true;
 	for (int i = 0; i < Storage1->FrameData.Num(); i++)
 	{
-		FNDIDataChannelWriteSimCacheFrame& Frame1 = Storage1->FrameData[i];
-		FNDIDataChannelWriteSimCacheFrame& Frame2 = Storage2->FrameData[i];
+		const FNDIDataChannelWriteSimCacheFrame& Frame1 = Storage1->FrameData[i];
+		const FNDIDataChannelWriteSimCacheFrame& Frame2 = Storage2->FrameData[i];
 
 		if (Frame1.NumElements != Frame2.NumElements)
 		{
@@ -736,8 +736,8 @@ bool UNiagaraDataInterfaceDataChannelWrite::SimCacheCompareFrame(UObject* LhsSto
 		{
 			for (int k = 0; k < Frame1.VariableData.Num(); k++)
 			{
-				FNDIDataChannelWriteSimCacheFrameBuffer& Buffer1 = Frame1.VariableData[k];
-				FNDIDataChannelWriteSimCacheFrameBuffer& Buffer2 = Frame2.VariableData[k];
+				const FNDIDataChannelWriteSimCacheFrameBuffer& Buffer1 = Frame1.VariableData[k];
+				const FNDIDataChannelWriteSimCacheFrameBuffer& Buffer2 = Frame2.VariableData[k];
 				if (Buffer1.SourceVar != Buffer2.SourceVar || Buffer1.Data != Buffer2.Data)
 				{
 					OutErrors += FString::Format(TEXT("Frame {0}: different buffers in data channel store for source var {1}\n"), {i, *Buffer1.SourceVar.GetName().ToString()});
diff --git a/Engine/Plugins/FX/Niagara/Source/Niagara/Private/DataInterface/NiagaraDataInterfaceSimpleCounter.cpp b/Engine/Plugins/FX/Niagara/Source/Niagara/Private/DataInterface/NiagaraDataInterfaceSimpleCounter.cpp
--- a/Engine/Plugins/FX/Niagara/Source/Niagara/Private/DataInterface/NiagaraDataInterfaceSimpleCounter.cpp
+++ b/Engine/Plugins/FX/Niagara/Source/Niagara/Private/DataInterface/NiagaraDataInterfaceSimpleCounter.cpp
@@ -567,10 +567,10 @@ bool UNiagaraDataInterfaceSimpleCounter::SimCacheReadFrame(UObject* StorageObjec
 	return true;
 }
 
-bool UNiagaraDataInterfaceSimpleCounter::SimCacheCompareFrame(UObject* LhsStorageObject, UObject* RhsStorageObject, int FrameIndex, TOptional<float> Tolerance, FString& OutErrors) const
+bool UNiagaraDataInterfaceSimpleCounter::SimCacheCompareFrame(const UObject* LhsStorageObject, const UObject* RhsStorageObject, int FrameIndex, TOptional<float> Tolerance, FString& OutErrors) const
 {
-	UNDISimpleCounterSimCacheData* LhsCacheData = CastChecked<UNDISimpleCounterSimCacheData>(LhsStorageObject);
-	UNDISimpleCounterSimCacheData* RhsCacheData = CastChecked<UNDISimpleCounterSimCacheData>(RhsStorageObject);
+	const UNDISimpleCounterSimCacheData* LhsCacheData = CastChecked<const UNDISimpleCounterSimCacheData>(LhsStorageObject);
+	const UNDISimpleCounterSimCacheData* RhsCacheData = CastChecked<const UNDISimpleCounterSimCacheData>(RhsStorageObject);
 
 	const int32 ValueOffset = FrameIndex * 2;
 	const int32 ExpectedValues = ValueOffset + 2;
diff --git a/Engine/Plugins/FX/Niagara/Source/Niagara/Private/DataInterface/NiagaraDataInterfaceSimpleCounter.h b/Engine/Plugins/FX/Niagara/Source/Niagara/Private/DataInterface/NiagaraDataInterfaceSimpleCounter.h
--- a/Engine/Plugins/FX/Niagara/Source/Niagara/Private/DataInterface/NiagaraDataInterfaceSimpleCounter.h
+++ b/Engine/Plugins/FX/Niagara/Source/Niagara/Private/DataInterface/NiagaraDataInterfaceSimpleCounter.h
@@ -63,7 +63,7 @@ public:
 	virtual UObject* SimCacheBeginWrite(UObject* SimCache, FNiagaraSystemInstance* NiagaraSystemInstance, const void* OptionalPerInstanceData, FNiagaraSimCacheFeedbackContext& FeedbackContext) const override;
 	virtual bool SimCacheWriteFrame(UObject* StorageObject, int FrameIndex, FNiagaraSystemInstance* SystemInstance, const void* OptionalPerInstanceData, FNiagaraSimCacheFeedbackContext& FeedbackContext) const override;
 	virtual bool SimCacheReadFrame(UObject* StorageObject, int FrameA, int FrameB, float Interp, FNiagaraSystemInstance* SystemInstance, void* OptionalPerInstanceData) override;
-	virtual bool SimCacheCompareFrame(UObject* LhsStorageObject, UObject* RhsStorageObject, int FrameIndex, TOptional<float> Tolerance, FString& OutErrors) const override;
+	virtual bool SimCacheCompareFrame(const UObject* LhsStorageObject, const UObject* RhsStorageObject, int FrameIndex, TOptional<float> Tolerance, FString& OutErrors) const override;
 	//~ UNiagaraDataInterface interface END
 
 	void UpdateDIProxy();
diff --git a/Engine/Plugins/FX/Niagara/Source/Niagara/Private/NiagaraDataInterfaceArray.cpp b/Engine/Plugins/FX/Niagara/Source/Niagara/Private/NiagaraDataInterfaceArray.cpp
--- a/Engine/Plugins/FX/Niagara/Source/Niagara/Private/NiagaraDataInterfaceArray.cpp
+++ b/Engine/Plugins/FX/Niagara/Source/Niagara/Private/NiagaraDataInterfaceArray.cpp
@@ -118,10 +118,10 @@ bool UNiagaraDataInterfaceArray::SimCacheReadFrame(UObject* StorageObject, int F
 	return ArrayProxy->SimCacheReadFrame(CacheData, FrameA, SystemInstance);
 }
 
-bool UNiagaraDataInterfaceArray::SimCacheCompareFrame(UObject* LhsStorageObject, UObject* RhsStorageObject, int FrameIndex, TOptional<float> InTolerance, FString& OutErrors) const
+bool UNiagaraDataInterfaceArray::SimCacheCompareFrame(const UObject* LhsStorageObject, const UObject* RhsStorageObject, int FrameIndex, TOptional<float> InTolerance, FString& OutErrors) const
 {
-	UNDIArraySimCacheData* LhsCacheData = CastChecked<UNDIArraySimCacheData>(LhsStorageObject);
-	UNDIArraySimCacheData* RhsCacheData = CastChecked<UNDIArraySimCacheData>(RhsStorageObject);
+	const UNDIArraySimCacheData* LhsCacheData = CastChecked<const UNDIArraySimCacheData>(LhsStorageObject);
+	const UNDIArraySimCacheData* RhsCacheData = CastChecked<const UNDIArraySimCacheData>(RhsStorageObject);
 
 	if (!LhsCacheData->CpuFrameData.IsValidIndex(FrameIndex) || !RhsCacheData->CpuFrameData.IsValidIndex(FrameIndex) ||
 		!LhsCacheData->GpuFrameData.IsValidIndex(FrameIndex) || !RhsCacheData->GpuFrameData.IsValidIndex(FrameIndex) )
diff --git a/Engine/Plugins/FX/Niagara/Source/Niagara/Private/NiagaraSimCache.cpp b/Engine/Plugins/FX/Niagara/Source/Niagara/Private/NiagaraSimCache.cpp
--- a/Engine/Plugins/FX/Niagara/Source/Niagara/Private/NiagaraSimCache.cpp
+++ b/Engine/Plugins/FX/Niagara/Source/Niagara/Private/NiagaraSimCache.cpp
@@ -1124,7 +1124,7 @@ TArray<FNiagaraVariableBase> UNiagaraSimCache::GetStoredDataInterfaces() const
 	return DataInterfaces;
 }
 
-UObject* UNiagaraSimCache::GetDataInterfaceStorageObject(const FNiagaraVariableBase& DataInterface) const
+const UObject* UNiagaraSimCache::GetDataInterfaceStorageObject(const FNiagaraVariableBase& DataInterface) const
 {
 	if (const TObjectPtr<UObject>* StoredObject = DataInterfaceStorage.Find(DataInterface))
 	{
@@ -1370,3 +1370,24 @@ void UNiagaraSimCache::ReadQuatAttributeWithRebase(TArray<FQuat>& OutValues, FQu
 		}
 	}
 }
+
+UObject* UNiagaraSimCache::ReadDataInterfaceAs(UClass* RequestedType, FName AttributeName, int FrameIndex) const
+{
+	for (auto DataInterfaceIt=DataInterfaceStorage.CreateConstIterator(); DataInterfaceIt; ++DataInterfaceIt)
+	{
+		const FNiagaraVariableBase& Variable = DataInterfaceIt.Key();
+		if (Variable.GetName() != AttributeName)
+		{
+			continue;
+		}
+
+		UObject* StorageObject = DataInterfaceIt.Value();
+		if (StorageObject->IsA(RequestedType))
+		{
+			return StorageObject;
+		}
+		//-TODO: Allow data interface to convert to the requested type
+		return nullptr;
+	}
+	return nullptr;
+}
diff --git a/Engine/Plugins/FX/Niagara/Source/Niagara/Private/NiagaraSimCacheCompare.cpp b/Engine/Plugins/FX/Niagara/Source/Niagara/Private/NiagaraSimCacheCompare.cpp
--- a/Engine/Plugins/FX/Niagara/Source/Niagara/Private/NiagaraSimCacheCompare.cpp
+++ b/Engine/Plugins/FX/Niagara/Source/Niagara/Private/NiagaraSimCacheCompare.cpp
@@ -279,8 +279,8 @@ bool FNiagaraSimCacheCompare::Compare(const UNiagaraSimCache& LhsCache, const UN
 					Tolerance.Emplace(*UserTolerance);
 				}
 
-				UObject* LhsStorageObject = LhsCache.GetDataInterfaceStorageObject(DIVariable);
-				UObject* RhsStorageObject = RhsCache.GetDataInterfaceStorageObject(DIVariable);
+				const UObject* LhsStorageObject = LhsCache.GetDataInterfaceStorageObject(DIVariable);
+				const UObject* RhsStorageObject = RhsCache.GetDataInterfaceStorageObject(DIVariable);
 				if (LhsStorageObject && RhsStorageObject)
 				{
 					for (int FrameIndex=0; FrameIndex < LhsCache.GetNumFrames(); FrameIndex++)
diff --git a/Engine/Plugins/FX/Niagara/Source/NiagaraEditor/Private/Customizations/NiagaraDataInterfaceSimCacheVisualizer.cpp b/Engine/Plugins/FX/Niagara/Source/NiagaraEditor/Private/Customizations/NiagaraDataInterfaceSimCacheVisualizer.cpp
--- a/Engine/Plugins/FX/Niagara/Source/NiagaraEditor/Private/Customizations/NiagaraDataInterfaceSimCacheVisualizer.cpp
+++ b/Engine/Plugins/FX/Niagara/Source/NiagaraEditor/Private/Customizations/NiagaraDataInterfaceSimCacheVisualizer.cpp
@@ -5,7 +5,7 @@
 #include "ViewModels/NiagaraSimCacheViewModel.h"
 
 
-TSharedPtr<SWidget> INiagaraDataInterfaceSimCacheVisualizer::CreateWidgetFor(UObject*, TSharedPtr<FNiagaraSimCacheViewModel>)
+TSharedPtr<SWidget> INiagaraDataInterfaceSimCacheVisualizer::CreateWidgetFor(const UObject*, TSharedPtr<FNiagaraSimCacheViewModel>)
 {
 	return SNullWidget::NullWidget;
 }
diff --git a/Engine/Plugins/FX/Niagara/Source/NiagaraEditor/Private/Customizations/SimCache/FNiagaraDataChannelSimCacheVisualizer.cpp b/Engine/Plugins/FX/Niagara/Source/NiagaraEditor/Private/Customizations/SimCache/FNiagaraDataChannelSimCacheVisualizer.cpp
--- a/Engine/Plugins/FX/Niagara/Source/NiagaraEditor/Private/Customizations/SimCache/FNiagaraDataChannelSimCacheVisualizer.cpp
+++ b/Engine/Plugins/FX/Niagara/Source/NiagaraEditor/Private/Customizations/SimCache/FNiagaraDataChannelSimCacheVisualizer.cpp
@@ -11,7 +11,7 @@
 
 namespace DataChannelVisualizer
 {
-	const FNDIDataChannelWriteSimCacheFrame& GetFrame(TSharedPtr<FNiagaraSimCacheViewModel> SimCacheViewModel, UNDIDataChannelWriteSimCacheData* DataChannelData)
+	const FNDIDataChannelWriteSimCacheFrame& GetFrame(TSharedPtr<FNiagaraSimCacheViewModel> SimCacheViewModel, const UNDIDataChannelWriteSimCacheData* DataChannelData)
 	{
 		int32 FrameIndex = FMath::Clamp(SimCacheViewModel->GetFrameIndex(), 0, DataChannelData->FrameData.Num() - 1);
 		return DataChannelData->FrameData[FrameIndex];
@@ -91,7 +91,7 @@ public:
 
 	virtual ~SNiagaraDataChannelCacheView() override;
 
-	void Construct(const FArguments& InArgs, TSharedPtr<FNiagaraSimCacheViewModel> InViewModel, UNDIDataChannelWriteSimCacheData* InData);
+	void Construct(const FArguments& InArgs, TSharedPtr<FNiagaraSimCacheViewModel> InViewModel, const UNDIDataChannelWriteSimCacheData* InData);
 
 	TSharedRef<ITableRow> MakeRowWidget(const TSharedPtr<int32> RowIndexPtr, const TSharedRef<STableViewBase>& OwnerTable) const;
 
@@ -104,7 +104,7 @@ private:
 	void UpdateRows(bool Refresh = true);
 	
 	TSharedPtr<FNiagaraSimCacheViewModel> SimCacheViewModel;
-	TStrongObjectPtr<UNDIDataChannelWriteSimCacheData> DataChannelData;
+	TStrongObjectPtr<const UNDIDataChannelWriteSimCacheData> DataChannelData;
 
 	TSharedPtr<SListView<TSharedPtr<int32>>> ListViewWidget;
 	TArray<TSharedPtr<int32>> RowItems;	
@@ -118,7 +118,7 @@ public:
 SLATE_BEGIN_ARGS(SDataChannelVisualizerRowWidget) {}
 	SLATE_ARGUMENT(TSharedPtr<int32>, RowIndexPtr)
 	SLATE_ARGUMENT(TSharedPtr<FNiagaraSimCacheViewModel>, SimCacheViewModel)
-	SLATE_ARGUMENT(UNDIDataChannelWriteSimCacheData*, Data)
+	SLATE_ARGUMENT(const UNDIDataChannelWriteSimCacheData*, Data)
 SLATE_END_ARGS()
 
 	void Construct(const FArguments& InArgs, const TSharedRef<STableViewBase>& InOwnerTableView, const TMap<FName, int32>& NameMapping);
@@ -126,7 +126,7 @@ SLATE_END_ARGS()
 
 	TSharedPtr<int32> RowIndexPtr;
 	TSharedPtr<FNiagaraSimCacheViewModel> SimCacheViewModel;
-	TStrongObjectPtr<UNDIDataChannelWriteSimCacheData> DataChannelData;
+	TStrongObjectPtr<const UNDIDataChannelWriteSimCacheData> DataChannelData;
 	TMap<FName, int32> NameToIndexMapping;
 };
 
@@ -135,7 +135,7 @@ SNiagaraDataChannelCacheView::~SNiagaraDataChannelCacheView()
 	SimCacheViewModel->OnViewDataChanged().RemoveAll(this);
 }
 
-void SNiagaraDataChannelCacheView::Construct(const FArguments&, TSharedPtr<FNiagaraSimCacheViewModel> InViewModel, UNDIDataChannelWriteSimCacheData* InData)
+void SNiagaraDataChannelCacheView::Construct(const FArguments&, TSharedPtr<FNiagaraSimCacheViewModel> InViewModel, const UNDIDataChannelWriteSimCacheData* InData)
 {
 	SimCacheViewModel = InViewModel;
 	DataChannelData.Reset(InData);
@@ -300,9 +300,9 @@ TSharedRef<SWidget> SDataChannelVisualizerRowWidget::GenerateWidgetForColumn(con
 		.Text(DisplayText);
 }
 
-TSharedPtr<SWidget> FNiagaraDataChannelSimCacheVisualizer::CreateWidgetFor(UObject* CachedData, TSharedPtr<FNiagaraSimCacheViewModel> ViewModel)
+TSharedPtr<SWidget> FNiagaraDataChannelSimCacheVisualizer::CreateWidgetFor(const UObject* CachedData, TSharedPtr<FNiagaraSimCacheViewModel> ViewModel)
 {
-	if (UNDIDataChannelWriteSimCacheData* DataChannelData = Cast<UNDIDataChannelWriteSimCacheData>(CachedData))
+	if (const UNDIDataChannelWriteSimCacheData* DataChannelData = Cast<const UNDIDataChannelWriteSimCacheData>(CachedData))
 	{
 		return SNew(SNiagaraDataChannelCacheView, ViewModel, DataChannelData);
 	}
diff --git a/Engine/Plugins/FX/Niagara/Source/NiagaraEditor/Private/Customizations/SimCache/FNiagaraDataChannelSimCacheVisualizer.h b/Engine/Plugins/FX/Niagara/Source/NiagaraEditor/Private/Customizations/SimCache/FNiagaraDataChannelSimCacheVisualizer.h
--- a/Engine/Plugins/FX/Niagara/Source/NiagaraEditor/Private/Customizations/SimCache/FNiagaraDataChannelSimCacheVisualizer.h
+++ b/Engine/Plugins/FX/Niagara/Source/NiagaraEditor/Private/Customizations/SimCache/FNiagaraDataChannelSimCacheVisualizer.h
@@ -12,5 +12,5 @@ class FNiagaraDataChannelSimCacheVisualizer : public INiagaraDataInterfaceSimCac
 public:
 	virtual ~FNiagaraDataChannelSimCacheVisualizer() override = default;
 	
-	virtual TSharedPtr<SWidget> CreateWidgetFor(UObject* CachedData, TSharedPtr<FNiagaraSimCacheViewModel> ViewModel) override;
+	virtual TSharedPtr<SWidget> CreateWidgetFor(const UObject* CachedData, TSharedPtr<FNiagaraSimCacheViewModel> ViewModel) override;
 };
diff --git a/Engine/Plugins/FX/Niagara/Source/NiagaraEditor/Private/Customizations/SimCache/NiagaraArraySimCacheVisualizer.cpp b/Engine/Plugins/FX/Niagara/Source/NiagaraEditor/Private/Customizations/SimCache/NiagaraArraySimCacheVisualizer.cpp
--- a/Engine/Plugins/FX/Niagara/Source/NiagaraEditor/Private/Customizations/SimCache/NiagaraArraySimCacheVisualizer.cpp
+++ b/Engine/Plugins/FX/Niagara/Source/NiagaraEditor/Private/Customizations/SimCache/NiagaraArraySimCacheVisualizer.cpp
@@ -26,7 +26,7 @@ class SVisualizerRowWidget : public SMultiColumnTableRow<TSharedPtr<int32>>
 public:
 	SLATE_BEGIN_ARGS(SVisualizerRowWidget) {}
 		SLATE_ARGUMENT(TSharedPtr<FNiagaraSimCacheViewModel>,	ViewModel)
-		SLATE_ARGUMENT(UNDIArraySimCacheData*,					CacheData)
+		SLATE_ARGUMENT(const UNDIArraySimCacheData*,			CacheData)
 		SLATE_ARGUMENT(TSharedPtr<int32>,						RowIndexPtr)
 		SLATE_ARGUMENT(UNiagaraDataInterfaceArray*,				DataInterface)
 	SLATE_END_ARGS()
@@ -80,10 +80,10 @@ public:
 		return SNullWidget::NullWidget;
 	}
 
-	TSharedPtr<FNiagaraSimCacheViewModel>	ViewModel;
-	TStrongObjectPtr<UNDIArraySimCacheData>	CacheData;
-	TSharedPtr<int32>						RowIndexPtr;
-	UNiagaraDataInterfaceArray*				DataInterface = nullptr;
+	TSharedPtr<FNiagaraSimCacheViewModel>			ViewModel;
+	TStrongObjectPtr<const UNDIArraySimCacheData>	CacheData;
+	TSharedPtr<int32>								RowIndexPtr;
+	UNiagaraDataInterfaceArray*						DataInterface = nullptr;
 };
 
 class SSimCacheView : public SCompoundWidget
@@ -98,7 +98,7 @@ public:
 		ViewModel->OnViewDataChanged().RemoveAll(this);
 	}
 
-	void Construct(const FArguments& InArgs, TSharedPtr<FNiagaraSimCacheViewModel> InViewModel, UNDIArraySimCacheData* InCacheData, UNiagaraDataInterfaceArray* InDataInterface)
+	void Construct(const FArguments& InArgs, TSharedPtr<FNiagaraSimCacheViewModel> InViewModel, const UNDIArraySimCacheData* InCacheData, UNiagaraDataInterfaceArray* InDataInterface)
 	{
 		ViewModel = InViewModel;
 		CacheData.Reset(InCacheData);
@@ -259,24 +259,24 @@ public:
 	}
 
 private:
-	TSharedPtr<FNiagaraSimCacheViewModel>		ViewModel;
-	TStrongObjectPtr<UNDIArraySimCacheData>		CacheData;
+	TSharedPtr<FNiagaraSimCacheViewModel>			ViewModel;
+	TStrongObjectPtr<const UNDIArraySimCacheData>	CacheData;
 
-	TSharedPtr<SHeaderRow>						HeaderRowWidget;
-	TSharedPtr<SListView<TSharedPtr<int32>>>	ListViewWidget;
-	TArray<TSharedPtr<int32>>					RowItems;
-	UNiagaraDataInterfaceArray*					DataInterface = nullptr;
+	TSharedPtr<SHeaderRow>							HeaderRowWidget;
+	TSharedPtr<SListView<TSharedPtr<int32>>>		ListViewWidget;
+	TArray<TSharedPtr<int32>>						RowItems;
+	UNiagaraDataInterfaceArray*						DataInterface = nullptr;
 };
 
 } // NDIArraySimCacheVisualizer
 
 ////////////////////////////////////////////////////////////////////////////////////////////////////////////////////
 
-TSharedPtr<SWidget> FNiagaraArraySimCacheVisualizer::CreateWidgetFor(UObject* InCachedData, TSharedPtr<FNiagaraSimCacheViewModel> ViewModel)
+TSharedPtr<SWidget> FNiagaraArraySimCacheVisualizer::CreateWidgetFor(const UObject* InCachedData, TSharedPtr<FNiagaraSimCacheViewModel> ViewModel)
 {
 	using namespace NDIArraySimCacheVisualizer;
 
-	if (UNDIArraySimCacheData* CachedData = Cast<UNDIArraySimCacheData>(InCachedData))
+	if (const UNDIArraySimCacheData* CachedData = Cast<const UNDIArraySimCacheData>(InCachedData))
 	{
 		check(ArrayDIClass);
 		UNiagaraDataInterfaceArray* DataInterface = CastChecked<UNiagaraDataInterfaceArray>(ArrayDIClass->GetDefaultObject());		
diff --git a/Engine/Plugins/FX/Niagara/Source/NiagaraEditor/Private/Customizations/SimCache/NiagaraArraySimCacheVisualizer.h b/Engine/Plugins/FX/Niagara/Source/NiagaraEditor/Private/Customizations/SimCache/NiagaraArraySimCacheVisualizer.h
--- a/Engine/Plugins/FX/Niagara/Source/NiagaraEditor/Private/Customizations/SimCache/NiagaraArraySimCacheVisualizer.h
+++ b/Engine/Plugins/FX/Niagara/Source/NiagaraEditor/Private/Customizations/SimCache/NiagaraArraySimCacheVisualizer.h
@@ -19,7 +19,7 @@ public:
 	}
 	virtual ~FNiagaraArraySimCacheVisualizer() override = default;
 	
-	virtual TSharedPtr<SWidget> CreateWidgetFor(UObject* CachedData, TSharedPtr<FNiagaraSimCacheViewModel> ViewModel) override;
+	virtual TSharedPtr<SWidget> CreateWidgetFor(const UObject* CachedData, TSharedPtr<FNiagaraSimCacheViewModel> ViewModel) override;
 
 protected:
 	UClass* ArrayDIClass = nullptr;
diff --git a/Engine/Plugins/FX/Niagara/Source/NiagaraEditor/Private/Customizations/SimCache/NiagaraMemoryBufferSimCacheVisualizer.cpp b/Engine/Plugins/FX/Niagara/Source/NiagaraEditor/Private/Customizations/SimCache/NiagaraMemoryBufferSimCacheVisualizer.cpp
--- a/Engine/Plugins/FX/Niagara/Source/NiagaraEditor/Private/Customizations/SimCache/NiagaraMemoryBufferSimCacheVisualizer.cpp
+++ b/Engine/Plugins/FX/Niagara/Source/NiagaraEditor/Private/Customizations/SimCache/NiagaraMemoryBufferSimCacheVisualizer.cpp
@@ -20,7 +20,7 @@ namespace NDIMemoryBufferSimCacheVisualizerLocal
 static const FName NAME_ElementOffset = FName("_ElementOffset");
 
 template<typename TType>
-TConstArrayView<TType> GetCurrentData(const TSharedPtr<FNiagaraSimCacheViewModel>& ViewModel, UNDIMemoryBufferSimCacheData* CacheData, ENiagaraSimTarget SimTarget)
+TConstArrayView<TType> GetCurrentData(const TSharedPtr<FNiagaraSimCacheViewModel>& ViewModel, const UNDIMemoryBufferSimCacheData* CacheData, ENiagaraSimTarget SimTarget)
 {
 	const int32 FrameIndex = ViewModel->GetFrameIndex();
 	if (CacheData->FrameData.IsValidIndex(FrameIndex))
@@ -54,7 +54,7 @@ class SVisualizerRowWidget : public SMultiColumnTableRow<TSharedPtr<int32>>
 public:
 	SLATE_BEGIN_ARGS(SVisualizerRowWidget) {}
 		SLATE_ARGUMENT(TSharedPtr<FNiagaraSimCacheViewModel>,	ViewModel)
-		SLATE_ARGUMENT(UNDIMemoryBufferSimCacheData*,			CacheData)
+		SLATE_ARGUMENT(const UNDIMemoryBufferSimCacheData*,		CacheData)
 		SLATE_ARGUMENT(TSharedPtr<int32>,						RowIndexPtr)
 		SLATE_ATTRIBUTE(ENiagaraSimTarget,						DisplaySimTarget)
 		SLATE_ATTRIBUTE(ENDIMemoryBufferViewType,				DisplayAsType)
@@ -149,11 +149,11 @@ public:
 		return SNullWidget::NullWidget;
 	}
 
-	TSharedPtr<FNiagaraSimCacheViewModel>			ViewModel;
-	TStrongObjectPtr<UNDIMemoryBufferSimCacheData>	CacheData;
-	TAttribute<ENiagaraSimTarget>					DisplaySimTarget = ENiagaraSimTarget::CPUSim;
-	TAttribute<ENDIMemoryBufferViewType>			DisplayAsType = ENDIMemoryBufferViewType::Integer;
-	TSharedPtr<int32>								RowIndexPtr;
+	TSharedPtr<FNiagaraSimCacheViewModel>					ViewModel;
+	TStrongObjectPtr<const UNDIMemoryBufferSimCacheData>	CacheData;
+	TAttribute<ENiagaraSimTarget>							DisplaySimTarget = ENiagaraSimTarget::CPUSim;
+	TAttribute<ENDIMemoryBufferViewType>					DisplayAsType = ENDIMemoryBufferViewType::Integer;
+	TSharedPtr<int32>										RowIndexPtr;
 };
 
 class SSimCacheView : public SCompoundWidget
@@ -168,7 +168,7 @@ public:
 		ViewModel->OnViewDataChanged().RemoveAll(this);
 	}
 
-	void Construct(const FArguments& InArgs, TSharedPtr<FNiagaraSimCacheViewModel> InViewModel, UNDIMemoryBufferSimCacheData* InCacheData)
+	void Construct(const FArguments& InArgs, TSharedPtr<FNiagaraSimCacheViewModel> InViewModel, const UNDIMemoryBufferSimCacheData* InCacheData)
 	{
 		ViewModel = InViewModel;
 		CacheData.Reset(InCacheData);
@@ -404,8 +404,8 @@ public:
 	}
 
 private:
-	TSharedPtr<FNiagaraSimCacheViewModel>			ViewModel;
-	TStrongObjectPtr<UNDIMemoryBufferSimCacheData>	CacheData;
+	TSharedPtr<FNiagaraSimCacheViewModel>					ViewModel;
+	TStrongObjectPtr<const UNDIMemoryBufferSimCacheData>	CacheData;
 
 	int32											DisplayColumns = 1;
 	ENiagaraSimTarget								DisplaySimTarget = ENiagaraSimTarget::CPUSim;
@@ -420,11 +420,11 @@ private:
 
 ////////////////////////////////////////////////////////////////////////////////////////////////////////////////////
 
-TSharedPtr<SWidget> FNiagaraMemoryBufferSimCacheVisualizer::CreateWidgetFor(UObject* InCachedData, TSharedPtr<FNiagaraSimCacheViewModel> ViewModel)
+TSharedPtr<SWidget> FNiagaraMemoryBufferSimCacheVisualizer::CreateWidgetFor(const UObject* InCachedData, TSharedPtr<FNiagaraSimCacheViewModel> ViewModel)
 {
 	using namespace NDIMemoryBufferSimCacheVisualizer;
 
-	if (UNDIMemoryBufferSimCacheData* CachedData = Cast<UNDIMemoryBufferSimCacheData>(InCachedData))
+	if (const UNDIMemoryBufferSimCacheData* CachedData = Cast<const UNDIMemoryBufferSimCacheData>(InCachedData))
 	{
 		return SNew(SSimCacheView, ViewModel, CachedData);
 	}
diff --git a/Engine/Plugins/FX/Niagara/Source/NiagaraEditor/Private/Customizations/SimCache/NiagaraMemoryBufferSimCacheVisualizer.h b/Engine/Plugins/FX/Niagara/Source/NiagaraEditor/Private/Customizations/SimCache/NiagaraMemoryBufferSimCacheVisualizer.h
--- a/Engine/Plugins/FX/Niagara/Source/NiagaraEditor/Private/Customizations/SimCache/NiagaraMemoryBufferSimCacheVisualizer.h
+++ b/Engine/Plugins/FX/Niagara/Source/NiagaraEditor/Private/Customizations/SimCache/NiagaraMemoryBufferSimCacheVisualizer.h
@@ -40,5 +40,5 @@ class FNiagaraMemoryBufferSimCacheVisualizer : public INiagaraDataInterfaceSimCa
 public:
 	virtual ~FNiagaraMemoryBufferSimCacheVisualizer() override = default;
 	
-	virtual TSharedPtr<SWidget> CreateWidgetFor(UObject* CachedData, TSharedPtr<FNiagaraSimCacheViewModel> ViewModel) override;
+	virtual TSharedPtr<SWidget> CreateWidgetFor(const UObject* CachedData, TSharedPtr<FNiagaraSimCacheViewModel> ViewModel) override;
 };
diff --git a/Engine/Plugins/FX/Niagara/Source/NiagaraEditor/Private/Customizations/SimCache/NiagaraRenderTargetVolumeSimCacheVisualizer.cpp b/Engine/Plugins/FX/Niagara/Source/NiagaraEditor/Private/Customizations/SimCache/NiagaraRenderTargetVolumeSimCacheVisualizer.cpp
--- a/Engine/Plugins/FX/Niagara/Source/NiagaraEditor/Private/Customizations/SimCache/NiagaraRenderTargetVolumeSimCacheVisualizer.cpp
+++ b/Engine/Plugins/FX/Niagara/Source/NiagaraEditor/Private/Customizations/SimCache/NiagaraRenderTargetVolumeSimCacheVisualizer.cpp
@@ -35,7 +35,7 @@ public:
 	SLATE_BEGIN_ARGS(SNiagaraVolumeTextureViewport) {}
 	SLATE_END_ARGS()
 
-	void Construct(const FArguments& InArgs, TSharedPtr<FNiagaraSimCacheViewModel> InViewModel, UAnimatedSparseVolumeTexture* InVolumeTexture);
+	void Construct(const FArguments& InArgs, TSharedPtr<FNiagaraSimCacheViewModel> InViewModel, const UAnimatedSparseVolumeTexture* InVolumeTexture);
 	virtual ~SNiagaraVolumeTextureViewport() override;
 
 	void UpdatePreviewComponent(bool bReset = true);
@@ -68,12 +68,12 @@ private:
 
 	TSharedPtr<FNiagaraSimCacheViewModel> ViewModel;
 	TStrongObjectPtr<UPostProcessComponent> PostProcessComponent;
-	TStrongObjectPtr<UAnimatedSparseVolumeTexture> VolumeTexture;
+	TStrongObjectPtr<const UAnimatedSparseVolumeTexture> VolumeTexture;
 };
 
-TSharedPtr<SWidget> FNiagaraRenderTargetVolumeSimCacheVisualizer::CreateWidgetFor(UObject* CachedData, TSharedPtr<FNiagaraSimCacheViewModel> ViewModel)
+TSharedPtr<SWidget> FNiagaraRenderTargetVolumeSimCacheVisualizer::CreateWidgetFor(const UObject* CachedData, TSharedPtr<FNiagaraSimCacheViewModel> ViewModel)
 {
-	if (UAnimatedSparseVolumeTexture* VolumeTexture = Cast<UAnimatedSparseVolumeTexture>(CachedData))
+	if (const UAnimatedSparseVolumeTexture* VolumeTexture = Cast<const UAnimatedSparseVolumeTexture>(CachedData))
 	{
 		static UEnum* MaskEnum = StaticEnum<ENiagaraRenderTargetVolumeVisualizerMask>();
 		TSharedPtr<SNiagaraVolumeTextureViewport> Viewport = SNew(SNiagaraVolumeTextureViewport, ViewModel, VolumeTexture);
@@ -243,7 +243,7 @@ void FNiagaraVolumeTextureViewportClient::SetIsSimulateInEditorViewport(bool bIn
 
 //////////////////////////////////////////////////////////////////////////
 
-void SNiagaraVolumeTextureViewport::Construct(const FArguments&, TSharedPtr<FNiagaraSimCacheViewModel> InViewModel, UAnimatedSparseVolumeTexture* InVolumeTexture)
+void SNiagaraVolumeTextureViewport::Construct(const FArguments&, TSharedPtr<FNiagaraSimCacheViewModel> InViewModel, const UAnimatedSparseVolumeTexture* InVolumeTexture)
 {
 	ViewModel = InViewModel;
 	VolumeTexture.Reset(InVolumeTexture);
@@ -261,7 +261,7 @@ void SNiagaraVolumeTextureViewport::Construct(const FArguments&, TSharedPtr<FNia
 	DuplicateMat->SetStaticComponentMaskParameterValueEditorOnly("Temperature Mask", false, true, false, false, ExprGuid);
 	DuplicateMat->SetStaticSwitchParameterValueEditorOnly("Temperature (Attributes B)", false, SwitchGuid);
 	DuplicateMat->SetScalarParameterValueEditorOnly("Density Scale", DensityScale);
-	DuplicateMat->SetSparseVolumeTextureParameterValueEditorOnly("SparseVolumeTexture", InVolumeTexture);
+	DuplicateMat->SetSparseVolumeTextureParameterValueEditorOnly("SparseVolumeTexture", const_cast<UAnimatedSparseVolumeTexture*>(InVolumeTexture));
 
 	PreviewComponent->OverrideMaterials.Add(DuplicateMat);
 	PreviewComponent->bIssueBlockingRequests = false;
diff --git a/Engine/Plugins/FX/Niagara/Source/NiagaraEditor/Private/Customizations/SimCache/NiagaraRenderTargetVolumeSimCacheVisualizer.h b/Engine/Plugins/FX/Niagara/Source/NiagaraEditor/Private/Customizations/SimCache/NiagaraRenderTargetVolumeSimCacheVisualizer.h
--- a/Engine/Plugins/FX/Niagara/Source/NiagaraEditor/Private/Customizations/SimCache/NiagaraRenderTargetVolumeSimCacheVisualizer.h
+++ b/Engine/Plugins/FX/Niagara/Source/NiagaraEditor/Private/Customizations/SimCache/NiagaraRenderTargetVolumeSimCacheVisualizer.h
@@ -18,5 +18,5 @@ class FNiagaraRenderTargetVolumeSimCacheVisualizer : public INiagaraDataInterfac
 public:
 	virtual ~FNiagaraRenderTargetVolumeSimCacheVisualizer() override = default;
 	
-	virtual TSharedPtr<SWidget> CreateWidgetFor(UObject* CachedData, TSharedPtr<FNiagaraSimCacheViewModel> ViewModel) override;
+	virtual TSharedPtr<SWidget> CreateWidgetFor(const UObject* CachedData, TSharedPtr<FNiagaraSimCacheViewModel> ViewModel) override;
 };
diff --git a/Engine/Plugins/FX/Niagara/Source/NiagaraEditor/Private/ViewModels/NiagaraSimCacheViewModel.cpp b/Engine/Plugins/FX/Niagara/Source/NiagaraEditor/Private/ViewModels/NiagaraSimCacheViewModel.cpp
--- a/Engine/Plugins/FX/Niagara/Source/NiagaraEditor/Private/ViewModels/NiagaraSimCacheViewModel.cpp
+++ b/Engine/Plugins/FX/Niagara/Source/NiagaraEditor/Private/ViewModels/NiagaraSimCacheViewModel.cpp
@@ -156,7 +156,7 @@ void FNiagaraSimCacheViewModel::SetFrameIndex(const int32 InFrameIndex)
 	OnViewDataChangedDelegate.Broadcast(false);
 }
 
-UObject* FNiagaraSimCacheViewModel::GetActiveDataInterfaceStorage() const
+const UObject* FNiagaraSimCacheViewModel::GetActiveDataInterfaceStorage() const
 {
 	if (SimCache)
 	{
diff --git a/Engine/Plugins/FX/Niagara/Source/NiagaraEditor/Private/Widgets/SNiagaraSimCacheView.cpp b/Engine/Plugins/FX/Niagara/Source/NiagaraEditor/Private/Widgets/SNiagaraSimCacheView.cpp
--- a/Engine/Plugins/FX/Niagara/Source/NiagaraEditor/Private/Widgets/SNiagaraSimCacheView.cpp
+++ b/Engine/Plugins/FX/Niagara/Source/NiagaraEditor/Private/Widgets/SNiagaraSimCacheView.cpp
@@ -219,7 +219,7 @@ void SNiagaraSimCacheView::UpdateDIWidget()
 		FNiagaraEditorModule& NiagaraEditorModule = FModuleManager::GetModuleChecked<FNiagaraEditorModule>("NiagaraEditor");
 		for (TSharedRef<INiagaraDataInterfaceSimCacheVisualizer> Visualizer : NiagaraEditorModule.FindDataInterfaceCacheVisualizer(SimCacheViewModel->GetActiveDataInterface().GetType().GetClass()))
 		{
-			if (UObject* DataObject = SimCacheViewModel->GetActiveDataInterfaceStorage())
+			if (const UObject* DataObject = SimCacheViewModel->GetActiveDataInterfaceStorage())
 			{
 				TSharedPtr<SWidget> DIVisualizerWidget = Visualizer->CreateWidgetFor(DataObject, SimCacheViewModel);
 				DIVisualizerWidgets.Add(DIVisualizerWidget);
diff --git a/Engine/Plugins/FX/Niagara/Source/NiagaraEditor/Public/Customizations/NiagaraDataInterfaceSimCacheVisualizer.h b/Engine/Plugins/FX/Niagara/Source/NiagaraEditor/Public/Customizations/NiagaraDataInterfaceSimCacheVisualizer.h
--- a/Engine/Plugins/FX/Niagara/Source/NiagaraEditor/Public/Customizations/NiagaraDataInterfaceSimCacheVisualizer.h
+++ b/Engine/Plugins/FX/Niagara/Source/NiagaraEditor/Public/Customizations/NiagaraDataInterfaceSimCacheVisualizer.h
@@ -18,5 +18,5 @@ class NIAGARAEDITOR_API INiagaraDataInterfaceSimCacheVisualizer
 public:
 	virtual ~INiagaraDataInterfaceSimCacheVisualizer() = default;
 	
-	virtual TSharedPtr<SWidget> CreateWidgetFor(UObject* CachedData, TSharedPtr<FNiagaraSimCacheViewModel> ViewModel);
+	virtual TSharedPtr<SWidget> CreateWidgetFor(const UObject* CachedData, TSharedPtr<FNiagaraSimCacheViewModel> ViewModel);
 };
diff --git a/Engine/Plugins/FX/Niagara/Source/NiagaraEditor/Public/ViewModels/NiagaraSimCacheViewModel.h b/Engine/Plugins/FX/Niagara/Source/NiagaraEditor/Public/ViewModels/NiagaraSimCacheViewModel.h
--- a/Engine/Plugins/FX/Niagara/Source/NiagaraEditor/Public/ViewModels/NiagaraSimCacheViewModel.h
+++ b/Engine/Plugins/FX/Niagara/Source/NiagaraEditor/Public/ViewModels/NiagaraSimCacheViewModel.h
@@ -44,7 +44,7 @@ public:
 	NIAGARAEDITOR_API void SetFrameIndex(const int32 InFrameIndex);
 
 	NIAGARAEDITOR_API FNiagaraVariableBase GetActiveDataInterface() const { return ActiveDataInterface; };
-	NIAGARAEDITOR_API UObject* GetActiveDataInterfaceStorage() const;
+	NIAGARAEDITOR_API const UObject* GetActiveDataInterfaceStorage() const;
 
 	NIAGARAEDITOR_API FOnViewDataChanged& OnViewDataChanged();
 	NIAGARAEDITOR_API FOnSimCacheChanged& OnSimCacheChanged();
