- Fix for compile error due to API change

#rnx

diff --git a/Engine/Plugins/FX/Niagara/Source/Niagara/Private/NiagaraSimCacheJson.cpp b/Engine/Plugins/FX/Niagara/Source/Niagara/Private/NiagaraSimCacheJson.cpp
--- a/Engine/Plugins/FX/Niagara/Source/Niagara/Private/NiagaraSimCacheJson.cpp
+++ b/Engine/Plugins/FX/Niagara/Source/Niagara/Private/NiagaraSimCacheJson.cpp
@@ -280,7 +280,7 @@ bool FNiagaraSimCacheJson::DumpFramesToFolder(const UNiagaraSimCache& SimCache,
 		for (const FNiagaraVariableBase& DataInterface : SimCache.GetStoredDataInterfaces())
 		{
 			// data interfaces
-			if (UObject* StorageObject = SimCache.GetDataInterfaceStorageObject(DataInterface))
+			if (const UObject* StorageObject = SimCache.GetDataInterfaceStorageObject(DataInterface))
 			{
 				UClass* Class = DataInterface.GetType().GetClass();
 				INiagaraSimCacheCustomStorageInterface* DataInterfaceCDO = Class ? Cast<INiagaraSimCacheCustomStorageInterface>(Class->GetDefaultObject()) : nullptr;
