Fix for some shader warnings with the NiagaraUvMappingUtils.ush
#tests Recompiled assets that were generating warnings

diff --git a/Engine/Plugins/FX/Niagara/Shaders/Private/Experimental/NiagaraUvMappingUtils.ush b/Engine/Plugins/FX/Niagara/Shaders/Private/Experimental/NiagaraUvMappingUtils.ush
--- a/Engine/Plugins/FX/Niagara/Shaders/Private/Experimental/NiagaraUvMappingUtils.ush
+++ b/Engine/Plugins/FX/Niagara/Shaders/Private/Experimental/NiagaraUvMappingUtils.ush
@@ -86,6 +86,8 @@ bool UvMapping_QuadTreeNodeContains(
 	out int OutTriangleIndex,
 	out float3 OutBaryCoord)
 {
+	OutTriangleIndex = -1;
+
 	if (NodeBufferOffset == -1)
 	{
 		return false;
@@ -193,6 +195,8 @@ bool UvMapping_QuadTreeNodeContainsAabb(
 	out int OutTriangleIndex,
 	out float3 OutBaryCoord)
 {
+	OutTriangleIndex = -1;
+
 	if (NodeBufferOffset == -1)
 	{
 		return false;
