Enable Niagara's AsyncTask compilation mode for EngineTest.  Required updating two cached tests because the new compilation mode provides more opportunities to trim attributes.  Also adjusts the code for comparing cache results to limit false positive warnings.

#virtualized

diff --git a/Engine/Plugins/FX/Niagara/Source/Niagara/Private/NiagaraSimCacheCompare.cpp b/Engine/Plugins/FX/Niagara/Source/Niagara/Private/NiagaraSimCacheCompare.cpp
--- a/Engine/Plugins/FX/Niagara/Source/Niagara/Private/NiagaraSimCacheCompare.cpp
+++ b/Engine/Plugins/FX/Niagara/Source/Niagara/Private/NiagaraSimCacheCompare.cpp
@@ -352,7 +352,12 @@ bool FNiagaraSimCacheCompare::CompareEmitter(const UNiagaraSimCache& LhsCache, c
 				bool IsValid = true;
 				for (const FNiagaraSimCacheVariable& Var : LhsVariables)
 				{
-					if (!RhsVariables.Contains(Var))
+					auto ComparePredicate = [&Var](const FNiagaraSimCacheVariable& Rhs) -> bool
+					{
+						return Var.Variable == Rhs.Variable;
+					};
+
+					if (!RhsVariables.ContainsByPredicate(ComparePredicate))
 					{
 						AddError(OutDifferences, FString::Printf(TEXT("%s attribute %s does not exist in both caches."), *EmitterNameString, *Var.Variable.GetName().ToString()));
 						IsValid = false;
