- Initialize Tri out parameter for GPU RandomSectionTriangle

#rnx

diff --git a/Engine/Plugins/FX/Niagara/Shaders/Private/NiagaraDataInterfaceStaticMeshTemplate.ush b/Engine/Plugins/FX/Niagara/Shaders/Private/NiagaraDataInterfaceStaticMeshTemplate.ush
--- a/Engine/Plugins/FX/Niagara/Shaders/Private/NiagaraDataInterfaceStaticMeshTemplate.ush
+++ b/Engine/Plugins/FX/Niagara/Shaders/Private/NiagaraDataInterfaceStaticMeshTemplate.ush
@@ -590,6 +590,7 @@ void RandomSectionTriangle_{ParameterName}(NiagaraRandInfo RandInfo, int Section
 	int SectionTriangleCount;
 	GetSectionTriangleCount_{ParameterName}(Section, SectionTriangleCount);
 
+	Tri = 0;
 	if (SectionTriangleCount > 0)
 	{
 		int SectionTriangleOffset = {ParameterName}_SectionInfos[Section].x;
