- Only reinitialize the instance when releasing the cache if it is currently active

#rnx

diff --git a/Engine/Plugins/FX/Niagara/Source/Niagara/Private/NiagaraComponent.cpp b/Engine/Plugins/FX/Niagara/Source/Niagara/Private/NiagaraComponent.cpp
--- a/Engine/Plugins/FX/Niagara/Source/Niagara/Private/NiagaraComponent.cpp
+++ b/Engine/Plugins/FX/Niagara/Source/Niagara/Private/NiagaraComponent.cpp
@@ -3992,7 +3992,7 @@ void UNiagaraComponent::SetPreviewLODDistance(bool bInEnablePreviewLODDistance,
 
 void UNiagaraComponent::SetSimCache(UNiagaraSimCache* InSimCache, bool bResetSystem)
 {
-	const bool bForceReset = SimCache ? SimCache->GetAttributeCaptureMode() != ENiagaraSimCacheAttributeCaptureMode::All : false;
+	const bool bForceReset = IsActive() && SimCache ? SimCache->GetAttributeCaptureMode() != ENiagaraSimCacheAttributeCaptureMode::All : false;
 
 	SimCache = InSimCache;
 	UpdateInstanceSoloMode();
