- Fix for Niagara VT sampling using sRGB vs UNORM
- This results in the YCoCg sampling incorrectly applying sRGB conversion

diff --git a/Engine/Plugins/FX/Niagara/Shaders/Private/NiagaraDataInterfaceVirtualTexture.ush b/Engine/Plugins/FX/Niagara/Shaders/Private/NiagaraDataInterfaceVirtualTexture.ush
--- a/Engine/Plugins/FX/Niagara/Shaders/Private/NiagaraDataInterfaceVirtualTexture.ush
+++ b/Engine/Plugins/FX/Niagara/Shaders/Private/NiagaraDataInterfaceVirtualTexture.ush
@@ -10,7 +10,7 @@
 
 // Mirror of ERuntimeVirtualTextureMaterialType
 #define ERuntimeVirtualTextureMaterialType_BaseColor								0
-#define ERuntimeVirtualTextureMaterialType_BaseColor_Normal_DEPRECATED				1
+#define ERuntimeVirtualTextureMaterialType_Mask4									1
 #define ERuntimeVirtualTextureMaterialType_BaseColor_Normal_Roughness				2
 #define ERuntimeVirtualTextureMaterialType_BaseColor_Normal_Specular				3
 #define ERuntimeVirtualTextureMaterialType_BaseColor_Normal_Specular_YCoCg			4
diff --git a/Engine/Plugins/FX/Niagara/Shaders/Private/NiagaraDataInterfaceVirtualTextureTemplate.ush b/Engine/Plugins/FX/Niagara/Shaders/Private/NiagaraDataInterfaceVirtualTextureTemplate.ush
--- a/Engine/Plugins/FX/Niagara/Shaders/Private/NiagaraDataInterfaceVirtualTextureTemplate.ush
+++ b/Engine/Plugins/FX/Niagara/Shaders/Private/NiagaraDataInterfaceVirtualTextureTemplate.ush
@@ -41,7 +41,7 @@ void GetAttributesValid_{ParameterName}(out bool BaseColorValid, out bool Specul
 			NormalValid = true;
 			break;
 
-		case ERuntimeVirtualTextureMaterialType_BaseColor_Normal_DEPRECATED:
+		case ERuntimeVirtualTextureMaterialType_Mask4:
 		case ERuntimeVirtualTextureMaterialType_BaseColor_Normal_Specular:
 			BaseColorValid = true;
 			SpecularValid = true;
@@ -106,22 +106,22 @@ void SampleRVT_{ParameterName}(in float3 WorldPosition, out bool bInsideVolume,
 	{
 		case ERuntimeVirtualTextureMaterialType_BaseColor:
 		{
-			BaseColor = LayerSample[0].xyz;
+			BaseColor = LayerSample[0].rgb;
 			break;
 		}
 
 		case ERuntimeVirtualTextureMaterialType_BaseColor_Normal_Roughness:
 		{
-			BaseColor = VirtualTextureUnpackBaseColorSRGB(LayerSample[0]);
+			BaseColor = LayerSample[0].rgb;
 			Roughness = LayerSample[1].y;
 			Normal = VirtualTextureUnpackNormalBGR565(LayerSample[1]);
 			break;
 		}
 
-		case ERuntimeVirtualTextureMaterialType_BaseColor_Normal_DEPRECATED:
+		case ERuntimeVirtualTextureMaterialType_Mask4:
 		case ERuntimeVirtualTextureMaterialType_BaseColor_Normal_Specular:
 		{
-			BaseColor = LayerSample[0].xyz;
+			BaseColor = LayerSample[0].rgb;
 			Specular = LayerSample[1].x;
 			Roughness = LayerSample[1].y;
 			Normal = VirtualTextureUnpackNormalBC3BC3(LayerSample[0], LayerSample[1]);
diff --git a/Engine/Plugins/FX/Niagara/Source/Niagara/Private/DataInterface/NiagaraDataInterfaceVirtualTexture.cpp b/Engine/Plugins/FX/Niagara/Source/Niagara/Private/DataInterface/NiagaraDataInterfaceVirtualTexture.cpp
--- a/Engine/Plugins/FX/Niagara/Source/Niagara/Private/DataInterface/NiagaraDataInterfaceVirtualTexture.cpp
+++ b/Engine/Plugins/FX/Niagara/Source/Niagara/Private/DataInterface/NiagaraDataInterfaceVirtualTexture.cpp
@@ -19,15 +19,15 @@
 namespace NDIVirtualTextureLocal
 {
 	BEGIN_SHADER_PARAMETER_STRUCT(FShaderParameters, )
-		SHADER_PARAMETER_TEXTURE(Texture2D,			VirtualTexture0)
+		SHADER_PARAMETER_SRV(Texture2D,				VirtualTexture0)
 		SHADER_PARAMETER_TEXTURE(Texture2D<uint4>,	VirtualTexture0PageTable)
 		SHADER_PARAMETER(FUintVector4,				VirtualTexture0TextureUniforms)
 
-		SHADER_PARAMETER_TEXTURE(Texture2D,			VirtualTexture1)
+		SHADER_PARAMETER_SRV(Texture2D,				VirtualTexture1)
 		SHADER_PARAMETER_TEXTURE(Texture2D<uint4>,	VirtualTexture1PageTable)
 		SHADER_PARAMETER(FUintVector4,				VirtualTexture1TextureUniforms)
 
-		SHADER_PARAMETER_TEXTURE(Texture2D,			VirtualTexture2)
+		SHADER_PARAMETER_SRV(Texture2D,				VirtualTexture2)
 		SHADER_PARAMETER_TEXTURE(Texture2D<uint4>,	VirtualTexture2PageTable)
 		SHADER_PARAMETER(FUintVector4,				VirtualTexture2TextureUniforms)
 
@@ -75,12 +75,13 @@ namespace NDIVirtualTextureLocal
 		{
 			bool IsValid() const
 			{
-				return Texture.IsValid() && PageTable.IsValid();
+				return Texture.IsValid() && TextureSRV.IsValid() && PageTable.IsValid();
 			}
 
-			FTextureRHIRef	Texture;
-			FTextureRHIRef	PageTable;
-			FUintVector4	TextureUniforms = FUintVector4(0, 0, 0, 0);
+			FTextureRHIRef				Texture;
+			FRHIShaderResourceViewRef	TextureSRV;
+			FTextureRHIRef				PageTable;
+			FUintVector4				TextureUniforms = FUintVector4(0, 0, 0, 0);
 		};
 
 		FLayerTexture		LayerTexture[MaxRVTLayers];
@@ -280,6 +281,11 @@ bool UNiagaraDataInterfaceVirtualTexture::PerInstanceTick(void* PerInstanceData,
 					InstanceData.UVUniforms[1] = RT_Texture->GetUniformParameter(ERuntimeVirtualTextureShaderUniform_WorldToUVTransform1);
 					InstanceData.UVUniforms[2] = RT_Texture->GetUniformParameter(ERuntimeVirtualTextureShaderUniform_WorldToUVTransform2);
 
+					const bool bBaseColorSRGB = 
+						InstanceData.MaterialType == ERuntimeVirtualTextureMaterialType::BaseColor ||
+						InstanceData.MaterialType == ERuntimeVirtualTextureMaterialType::BaseColor_Normal_Roughness ||
+						InstanceData.MaterialType == ERuntimeVirtualTextureMaterialType::BaseColor_Normal_Specular;
+
 					const FVector4 WorldHeightUnpack = RT_Texture->GetUniformParameter(ERuntimeVirtualTextureShaderUniform_WorldHeightUnpack);
 					InstanceData.WorldHeightUnpack = FVector2f(float(WorldHeightUnpack.X), float(WorldHeightUnpack.Y));		//LWC Precision Loss
 
@@ -287,6 +293,7 @@ bool UNiagaraDataInterfaceVirtualTexture::PerInstanceTick(void* PerInstanceData,
 					{
 						FNDInstanceData_RenderThread::FLayerTexture& LayerTexture = InstanceData.LayerTexture[iLayer];
 						LayerTexture.Texture = VirtualTexture->GetPhysicalTexture(iLayer);
+						LayerTexture.TextureSRV = VirtualTexture->GetPhysicalTextureSRV(iLayer, iLayer == 0 && bBaseColorSRGB);
 						LayerTexture.PageTable = VirtualTexture->GetPageTableTexture(0);//-TODO:iLayer);
 
 						VirtualTexture->GetPackedUniform(&LayerTexture.TextureUniforms, iLayer);
@@ -389,21 +396,21 @@ void UNiagaraDataInterfaceVirtualTexture::SetShaderParameters(const FNiagaraData
 
 		if ( bLayerValid[0] == true)
 		{
-			ShaderParameters->VirtualTexture0					= InstanceData->LayerTexture[0].Texture;
+			ShaderParameters->VirtualTexture0					= InstanceData->LayerTexture[0].TextureSRV;
 			ShaderParameters->VirtualTexture0PageTable			= InstanceData->LayerTexture[0].PageTable;
 			ShaderParameters->VirtualTexture0TextureUniforms	= InstanceData->LayerTexture[0].TextureUniforms;
 		}
 
 		if ( bLayerValid[1] == true)
 		{
-			ShaderParameters->VirtualTexture1					= InstanceData->LayerTexture[1].Texture;
+			ShaderParameters->VirtualTexture1					= InstanceData->LayerTexture[1].TextureSRV;
 			ShaderParameters->VirtualTexture1PageTable			= InstanceData->LayerTexture[1].PageTable;
 			ShaderParameters->VirtualTexture1TextureUniforms	= InstanceData->LayerTexture[1].TextureUniforms;
 		}
 
 		if ( bLayerValid[2] == true)
 		{
-			ShaderParameters->VirtualTexture2					= InstanceData->LayerTexture[2].Texture;
+			ShaderParameters->VirtualTexture2					= InstanceData->LayerTexture[2].TextureSRV;
 			ShaderParameters->VirtualTexture2PageTable			= InstanceData->LayerTexture[2].PageTable;
 			ShaderParameters->VirtualTexture2TextureUniforms	= InstanceData->LayerTexture[2].TextureUniforms;
 		}
@@ -417,7 +424,7 @@ void UNiagaraDataInterfaceVirtualTexture::SetShaderParameters(const FNiagaraData
 		ShaderParameters->MaterialType			= int(InstanceData->MaterialType);
 		ShaderParameters->PageTableUniforms[0]	= InstanceData->PageTableUniforms[0];
 		ShaderParameters->PageTableUniforms[1]	= InstanceData->PageTableUniforms[1];
-		ShaderParameters->UVUniforms[0]			= FVector4f(UVOriginRebased);
+		ShaderParameters->UVUniforms[0]			= FVector4f(UVOriginRebased, InstanceData->UVUniforms[0].W);
 		ShaderParameters->UVUniforms[1]			= FVector4f(InstanceData->UVUniforms[1]);
 		ShaderParameters->UVUniforms[2]			= FVector4f(InstanceData->UVUniforms[2]);
 		ShaderParameters->WorldHeightUnpack		= InstanceData->WorldHeightUnpack;
@@ -439,19 +446,19 @@ void UNiagaraDataInterfaceVirtualTexture::SetShaderParameters(const FNiagaraData
 	
 	if ( bLayerValid[0] == false )
 	{
-		ShaderParameters->VirtualTexture0					= GBlackTexture->GetTextureRHI();
+		ShaderParameters->VirtualTexture0					= GBlackTextureWithSRV->ShaderResourceViewRHI;
 		ShaderParameters->VirtualTexture0PageTable			= GBlackUintTexture->TextureRHI;
 		ShaderParameters->VirtualTexture0TextureUniforms	= FUintVector4(0, 0, 0, 0);
 	}
 	if ( bLayerValid[1] == false )
 	{
-		ShaderParameters->VirtualTexture1					= GBlackTexture->GetTextureRHI();
+		ShaderParameters->VirtualTexture1					= GBlackTextureWithSRV->ShaderResourceViewRHI;
 		ShaderParameters->VirtualTexture1PageTable			= GBlackUintTexture->TextureRHI;
 		ShaderParameters->VirtualTexture1TextureUniforms	= FUintVector4(0, 0, 0, 0);
 	}
 	if ( bLayerValid[2] == false )
 	{
-		ShaderParameters->VirtualTexture2					= GBlackTexture->GetTextureRHI();
+		ShaderParameters->VirtualTexture2					= GBlackTextureWithSRV->ShaderResourceViewRHI;
 		ShaderParameters->VirtualTexture2PageTable			= GBlackUintTexture->TextureRHI;
 		ShaderParameters->VirtualTexture2TextureUniforms	= FUintVector4(0, 0, 0, 0);
 	}
