- Stateless random changes, to ensure consistency between CPU & GPU paths
- Modules has a consistent random seed offset and calls to random must specify an offset
- This allows for consistent random behavior when adding / removing a module
- We can also optimize the CPU path more by permuting functions into parts that are enabled / disabled

diff --git a/Engine/Plugins/FX/Niagara/Shaders/Private/Stateless/Modules/NiagaraStatelessModule_CameraOffset.ush b/Engine/Plugins/FX/Niagara/Shaders/Private/Stateless/Modules/NiagaraStatelessModule_CameraOffset.ush
--- a/Engine/Plugins/FX/Niagara/Shaders/Private/Stateless/Modules/NiagaraStatelessModule_CameraOffset.ush
+++ b/Engine/Plugins/FX/Niagara/Shaders/Private/Stateless/Modules/NiagaraStatelessModule_CameraOffset.ush
@@ -7,7 +7,7 @@ uint3 CameraOffset_Distribution;
 void CameraOffset_Simulate(inout FStatelessParticle Particle)
 {
 	FStatelessDistributionSamplerFloat OffsetSampler;
-	OffsetSampler.Init(Particle, CameraOffset_Distribution);
+	OffsetSampler.Init(0, CameraOffset_Distribution);
 	if (OffsetSampler.IsValid())
 	{
 		Particle.CameraOffset			= OffsetSampler.GetValue(Particle.NormalizedAge);
diff --git a/Engine/Plugins/FX/Niagara/Shaders/Private/Stateless/Modules/NiagaraStatelessModule_DynamicMaterialParameters.ush b/Engine/Plugins/FX/Niagara/Shaders/Private/Stateless/Modules/NiagaraStatelessModule_DynamicMaterialParameters.ush
--- a/Engine/Plugins/FX/Niagara/Shaders/Private/Stateless/Modules/NiagaraStatelessModule_DynamicMaterialParameters.ush
+++ b/Engine/Plugins/FX/Niagara/Shaders/Private/Stateless/Modules/NiagaraStatelessModule_DynamicMaterialParameters.ush
@@ -30,10 +30,10 @@ void DynamicMaterialParameters_Simulate(inout FStatelessParticle Particle)
 	if ((DynamicMaterialParameters_ChannelMask & 0x000f) != 0)
 	{
 		FStatelessDistributionSamplerFloat ParameterSampler[4];
-		ParameterSampler[0].Init(Particle, DynamicMaterialParameters_Parameter0X);
-		ParameterSampler[1].Init(Particle, DynamicMaterialParameters_Parameter0Y);
-		ParameterSampler[2].Init(Particle, DynamicMaterialParameters_Parameter0Z);
-		ParameterSampler[3].Init(Particle, DynamicMaterialParameters_Parameter0W);
+		ParameterSampler[0].Init(0, DynamicMaterialParameters_Parameter0X);
+		ParameterSampler[1].Init(1, DynamicMaterialParameters_Parameter0Y);
+		ParameterSampler[2].Init(2, DynamicMaterialParameters_Parameter0Z);
+		ParameterSampler[3].Init(3, DynamicMaterialParameters_Parameter0W);
 		
 		Particle.DynamicMaterialParameter0.x = (DynamicMaterialParameters_ChannelMask & 0x0001) != 0 ? ParameterSampler[0].GetValue(Particle.NormalizedAge) : 0.0f;
 		Particle.DynamicMaterialParameter0.y = (DynamicMaterialParameters_ChannelMask & 0x0002) != 0 ? ParameterSampler[1].GetValue(Particle.NormalizedAge) : 0.0f;
@@ -44,10 +44,10 @@ void DynamicMaterialParameters_Simulate(inout FStatelessParticle Particle)
 	if ((DynamicMaterialParameters_ChannelMask & 0x00f0) != 0)
 	{
 		FStatelessDistributionSamplerFloat ParameterSampler[4];
-		ParameterSampler[0].Init(Particle, DynamicMaterialParameters_Parameter1X);
-		ParameterSampler[1].Init(Particle, DynamicMaterialParameters_Parameter1Y);
-		ParameterSampler[2].Init(Particle, DynamicMaterialParameters_Parameter1Z);
-		ParameterSampler[3].Init(Particle, DynamicMaterialParameters_Parameter1W);
+		ParameterSampler[0].Init(4, DynamicMaterialParameters_Parameter1X);
+		ParameterSampler[1].Init(5, DynamicMaterialParameters_Parameter1Y);
+		ParameterSampler[2].Init(6, DynamicMaterialParameters_Parameter1Z);
+		ParameterSampler[3].Init(7, DynamicMaterialParameters_Parameter1W);
 		
 		Particle.DynamicMaterialParameter1.x = (DynamicMaterialParameters_ChannelMask & 0x0010) != 0 ? ParameterSampler[0].GetValue(Particle.NormalizedAge) : 0.0f;
 		Particle.DynamicMaterialParameter1.y = (DynamicMaterialParameters_ChannelMask & 0x0020) != 0 ? ParameterSampler[1].GetValue(Particle.NormalizedAge) : 0.0f;
@@ -58,10 +58,10 @@ void DynamicMaterialParameters_Simulate(inout FStatelessParticle Particle)
 	if ((DynamicMaterialParameters_ChannelMask & 0x0f00) != 0)
 	{
 		FStatelessDistributionSamplerFloat ParameterSampler[4];
-		ParameterSampler[0].Init(Particle, DynamicMaterialParameters_Parameter2X);
-		ParameterSampler[1].Init(Particle, DynamicMaterialParameters_Parameter2Y);
-		ParameterSampler[2].Init(Particle, DynamicMaterialParameters_Parameter2Z);
-		ParameterSampler[3].Init(Particle, DynamicMaterialParameters_Parameter2W);
+		ParameterSampler[0].Init(8, DynamicMaterialParameters_Parameter2X);
+		ParameterSampler[1].Init(9, DynamicMaterialParameters_Parameter2Y);
+		ParameterSampler[2].Init(10, DynamicMaterialParameters_Parameter2Z);
+		ParameterSampler[3].Init(11, DynamicMaterialParameters_Parameter2W);
 		
 		Particle.DynamicMaterialParameter2.x = (DynamicMaterialParameters_ChannelMask & 0x0100) != 0 ? ParameterSampler[0].GetValue(Particle.NormalizedAge) : 0.0f;
 		Particle.DynamicMaterialParameter2.y = (DynamicMaterialParameters_ChannelMask & 0x0200) != 0 ? ParameterSampler[1].GetValue(Particle.NormalizedAge) : 0.0f;
@@ -72,10 +72,10 @@ void DynamicMaterialParameters_Simulate(inout FStatelessParticle Particle)
 	if ((DynamicMaterialParameters_ChannelMask & 0xf000) != 0)
 	{
 		FStatelessDistributionSamplerFloat ParameterSampler[4];
-		ParameterSampler[0].Init(Particle, DynamicMaterialParameters_Parameter3X);
-		ParameterSampler[1].Init(Particle, DynamicMaterialParameters_Parameter3Y);
-		ParameterSampler[2].Init(Particle, DynamicMaterialParameters_Parameter3Z);
-		ParameterSampler[3].Init(Particle, DynamicMaterialParameters_Parameter3W);
+		ParameterSampler[0].Init(12, DynamicMaterialParameters_Parameter3X);
+		ParameterSampler[1].Init(13, DynamicMaterialParameters_Parameter3Y);
+		ParameterSampler[2].Init(14, DynamicMaterialParameters_Parameter3Z);
+		ParameterSampler[3].Init(15, DynamicMaterialParameters_Parameter3W);
 		
 		Particle.DynamicMaterialParameter3.x = (DynamicMaterialParameters_ChannelMask & 0x1000) != 0 ? ParameterSampler[0].GetValue(Particle.NormalizedAge) : 0.0f;
 		Particle.DynamicMaterialParameter3.y = (DynamicMaterialParameters_ChannelMask & 0x2000) != 0 ? ParameterSampler[1].GetValue(Particle.NormalizedAge) : 0.0f;
diff --git a/Engine/Plugins/FX/Niagara/Shaders/Private/Stateless/Modules/NiagaraStatelessModule_InitialMeshOrientation.ush b/Engine/Plugins/FX/Niagara/Shaders/Private/Stateless/Modules/NiagaraStatelessModule_InitialMeshOrientation.ush
--- a/Engine/Plugins/FX/Niagara/Shaders/Private/Stateless/Modules/NiagaraStatelessModule_InitialMeshOrientation.ush
+++ b/Engine/Plugins/FX/Niagara/Shaders/Private/Stateless/Modules/NiagaraStatelessModule_InitialMeshOrientation.ush
@@ -7,7 +7,7 @@ float3	InitialMeshOrientation_RandomRangeScale;
 
 void InitialMeshOrientation_Simulate(inout FStatelessParticle Particle)
 {
-	const float3 InitialRotation = InitialMeshOrientation_Rotation + (RandomFloat3(Particle) * InitialMeshOrientation_RandomRangeScale);
+	const float3 InitialRotation = InitialMeshOrientation_Rotation + (RandomFloat3(0) * InitialMeshOrientation_RandomRangeScale);
 	Particle.MeshOrientation			= RotatorToQuat(InitialRotation);
 	Particle.PreviousMeshOrientation	= Particle.MeshOrientation;
 }
diff --git a/Engine/Plugins/FX/Niagara/Shaders/Private/Stateless/Modules/NiagaraStatelessModule_InitializeParticle.ush b/Engine/Plugins/FX/Niagara/Shaders/Private/Stateless/Modules/NiagaraStatelessModule_InitializeParticle.ush
--- a/Engine/Plugins/FX/Niagara/Shaders/Private/Stateless/Modules/NiagaraStatelessModule_InitializeParticle.ush
+++ b/Engine/Plugins/FX/Niagara/Shaders/Private/Stateless/Modules/NiagaraStatelessModule_InitializeParticle.ush
@@ -26,15 +26,15 @@ void InitializeParticle_Simulate(inout FStatelessParticle Particle)
 	const bool bUniformMeshScale  = (InitializeParticle_ModuleFlags & EInitializeParticleModuleFlag_UniformMeshScale) != 0;
 
 	FStatelessDistributionSamplerFloat3 InitialPosition;
-	InitialPosition.Init(Particle, InitializeParticle_InitialPosition);
+	InitialPosition.Init(0, InitializeParticle_InitialPosition);
 		
 	Particle.UniqueID		= Particle.UniqueIndex;
 	Particle.Position		= InitialPosition.GetValue(0.0f);
-	Particle.Color			= RandomScaleBiasFloat(Particle, InitializeParticle_ColorScale, InitializeParticle_ColorBias);
-	Particle.RibbonWidth	= RandomScaleBiasFloat(Particle, InitializeParticle_RibbonWidthScale, InitializeParticle_RibbonWidthBias);
-	Particle.SpriteSize		= RandomScaleBiasFloat(Particle, InitializeParticle_SpriteSizeScale, InitializeParticle_SpriteSizeBias, bUniformSpriteSize);
-	Particle.SpriteRotation	= RandomScaleBiasFloat(Particle, InitializeParticle_SpriteRotationScale, InitializeParticle_SpriteRotationBias);
-	Particle.Scale			= RandomScaleBiasFloat(Particle, InitializeParticle_MeshScaleScale, InitializeParticle_MeshScaleBias, bUniformMeshScale);
+	Particle.Color			= RandomScaleBiasFloat(1, InitializeParticle_ColorScale, InitializeParticle_ColorBias);
+	Particle.RibbonWidth	= RandomScaleBiasFloat(2, InitializeParticle_RibbonWidthScale, InitializeParticle_RibbonWidthBias);
+	Particle.SpriteSize		= RandomScaleBiasFloat(3, InitializeParticle_SpriteSizeScale, InitializeParticle_SpriteSizeBias, bUniformSpriteSize);
+	Particle.SpriteRotation	= RandomScaleBiasFloat(4, InitializeParticle_SpriteRotationScale, InitializeParticle_SpriteRotationBias);
+	Particle.Scale			= RandomScaleBiasFloat(5, InitializeParticle_MeshScaleScale, InitializeParticle_MeshScaleBias, bUniformMeshScale);
 
 	Particle.PreviousPosition		= Particle.Position;
 	Particle.PreviousRibbonWidth	= Particle.RibbonWidth;
diff --git a/Engine/Plugins/FX/Niagara/Shaders/Private/Stateless/Modules/NiagaraStatelessModule_MeshIndex.ush b/Engine/Plugins/FX/Niagara/Shaders/Private/Stateless/Modules/NiagaraStatelessModule_MeshIndex.ush
--- a/Engine/Plugins/FX/Niagara/Shaders/Private/Stateless/Modules/NiagaraStatelessModule_MeshIndex.ush
+++ b/Engine/Plugins/FX/Niagara/Shaders/Private/Stateless/Modules/NiagaraStatelessModule_MeshIndex.ush
@@ -19,7 +19,7 @@ void MeshIndex_Simulate(inout FStatelessParticle Particle)
 	{
 		if ( MeshIndex_TableNumElements > 0 )
 		{
-			const float2 Rand = RandomFloat2(Particle);
+			const float2 Rand = RandomFloat2(0);
 			MeshIndex = round(Rand.x * float(MeshIndex_TableNumElements));
 			const float Probability = GetStaticFloat(MeshIndex_TableOffset, (MeshIndex * 2) + 0);
 			if (Rand.y > Probability)
diff --git a/Engine/Plugins/FX/Niagara/Shaders/Private/Stateless/Modules/NiagaraStatelessModule_MeshRotationRate.ush b/Engine/Plugins/FX/Niagara/Shaders/Private/Stateless/Modules/NiagaraStatelessModule_MeshRotationRate.ush
--- a/Engine/Plugins/FX/Niagara/Shaders/Private/Stateless/Modules/NiagaraStatelessModule_MeshRotationRate.ush
+++ b/Engine/Plugins/FX/Niagara/Shaders/Private/Stateless/Modules/NiagaraStatelessModule_MeshRotationRate.ush
@@ -12,7 +12,7 @@ void MeshRotationRate_Simulate(inout FStatelessParticle Particle)
 	{
 		return;
 	}
-	const float3 RotationRate = RandomScaleBiasFloat(Particle, MeshRotationRate_Scale, MeshRotationRate_Bias);
+	const float3 RotationRate = RandomScaleBiasFloat(0, MeshRotationRate_Scale, MeshRotationRate_Bias);
 
 	Particle.MeshOrientation			= MultiplyQuat(Particle.MeshOrientation,			RotatorToQuat(Particle.Age * RotationRate));
 	Particle.PreviousMeshOrientation	= MultiplyQuat(Particle.PreviousMeshOrientation,	RotatorToQuat(Particle.PreviousAge * RotationRate));
diff --git a/Engine/Plugins/FX/Niagara/Shaders/Private/Stateless/Modules/NiagaraStatelessModule_RotateAroundPoint.ush b/Engine/Plugins/FX/Niagara/Shaders/Private/Stateless/Modules/NiagaraStatelessModule_RotateAroundPoint.ush
--- a/Engine/Plugins/FX/Niagara/Shaders/Private/Stateless/Modules/NiagaraStatelessModule_RotateAroundPoint.ush
+++ b/Engine/Plugins/FX/Niagara/Shaders/Private/Stateless/Modules/NiagaraStatelessModule_RotateAroundPoint.ush
@@ -29,9 +29,9 @@ void RotateAroundPoint_IntegratePosition(in FStatelessModule_RotateAroundPoint M
 void RotateAroundPoint_Simulate(inout FStatelessParticle Particle)
 {
 	FStatelessModule_RotateAroundPoint ModuleData;
-	ModuleData.Rate			= RandomScaleBiasFloat(Particle, RotateAroundPoint_RateScale, RotateAroundPoint_RateBias);
-	ModuleData.Radius		= RandomScaleBiasFloat(Particle, RotateAroundPoint_RadiusScale, RotateAroundPoint_RadiusBias);
-	ModuleData.InitialPhase	= RandomScaleBiasFloat(Particle, RotateAroundPoint_InitialPhaseScale, RotateAroundPoint_InitialPhaseBias);
+	ModuleData.Rate			= RandomScaleBiasFloat(0, RotateAroundPoint_RateScale, RotateAroundPoint_RateBias);
+	ModuleData.Radius		= RandomScaleBiasFloat(1, RotateAroundPoint_RadiusScale, RotateAroundPoint_RadiusBias);
+	ModuleData.InitialPhase	= RandomScaleBiasFloat(2, RotateAroundPoint_InitialPhaseScale, RotateAroundPoint_InitialPhaseBias);
 		
 	RotateAroundPoint_IntegratePosition(ModuleData, Particle.NormalizedAge, Particle.Position);
 	RotateAroundPoint_IntegratePosition(ModuleData, Particle.PreviousNormalizedAge, Particle.PreviousPosition);
diff --git a/Engine/Plugins/FX/Niagara/Shaders/Private/Stateless/Modules/NiagaraStatelessModule_ScaleColor.ush b/Engine/Plugins/FX/Niagara/Shaders/Private/Stateless/Modules/NiagaraStatelessModule_ScaleColor.ush
--- a/Engine/Plugins/FX/Niagara/Shaders/Private/Stateless/Modules/NiagaraStatelessModule_ScaleColor.ush
+++ b/Engine/Plugins/FX/Niagara/Shaders/Private/Stateless/Modules/NiagaraStatelessModule_ScaleColor.ush
@@ -7,7 +7,7 @@ uint3	ScaleColor_Distribution;
 void ScaleColor_Simulate(inout FStatelessParticle Particle)
 {
 	FStatelessDistributionSamplerFloat4 ScaleDistribution;
-	ScaleDistribution.Init(Particle, ScaleColor_Distribution);
+	ScaleDistribution.Init(0, ScaleColor_Distribution);
 	if (ScaleDistribution.IsValid())
 	{
 		Particle.Color *= ScaleDistribution.GetValue(Particle.NormalizedAge);
diff --git a/Engine/Plugins/FX/Niagara/Shaders/Private/Stateless/Modules/NiagaraStatelessModule_ScaleMeshSize.ush b/Engine/Plugins/FX/Niagara/Shaders/Private/Stateless/Modules/NiagaraStatelessModule_ScaleMeshSize.ush
--- a/Engine/Plugins/FX/Niagara/Shaders/Private/Stateless/Modules/NiagaraStatelessModule_ScaleMeshSize.ush
+++ b/Engine/Plugins/FX/Niagara/Shaders/Private/Stateless/Modules/NiagaraStatelessModule_ScaleMeshSize.ush
@@ -9,7 +9,7 @@ int		ScaleMeshSize_CurveScaleOffset;
 void ScaleMeshSize_Simulate(inout FStatelessParticle Particle)
 {
 	FStatelessDistributionSamplerFloat3 MeshSizeDistribution;
-	MeshSizeDistribution.Init(Particle, ScaleMeshSize_Distribution);
+	MeshSizeDistribution.Init(0, ScaleMeshSize_Distribution);
 
 	if (MeshSizeDistribution.IsValid())
 	{
diff --git a/Engine/Plugins/FX/Niagara/Shaders/Private/Stateless/Modules/NiagaraStatelessModule_ScaleSpriteSize.ush b/Engine/Plugins/FX/Niagara/Shaders/Private/Stateless/Modules/NiagaraStatelessModule_ScaleSpriteSize.ush
--- a/Engine/Plugins/FX/Niagara/Shaders/Private/Stateless/Modules/NiagaraStatelessModule_ScaleSpriteSize.ush
+++ b/Engine/Plugins/FX/Niagara/Shaders/Private/Stateless/Modules/NiagaraStatelessModule_ScaleSpriteSize.ush
@@ -8,7 +8,7 @@ float2	ScaleSpriteSize_CurveScale;
 void ScaleSpriteSize_Simulate(inout FStatelessParticle Particle)
 {
 	FStatelessDistributionSamplerFloat2 ScaleDistribution;
-	ScaleDistribution.Init(Particle, ScaleSpriteSize_Distribution);
+	ScaleDistribution.Init(0, ScaleSpriteSize_Distribution);
 
 	if (ScaleDistribution.IsValid())
 	{
diff --git a/Engine/Plugins/FX/Niagara/Shaders/Private/Stateless/Modules/NiagaraStatelessModule_ShapeLocation.ush b/Engine/Plugins/FX/Niagara/Shaders/Private/Stateless/Modules/NiagaraStatelessModule_ShapeLocation.ush
--- a/Engine/Plugins/FX/Niagara/Shaders/Private/Stateless/Modules/NiagaraStatelessModule_ShapeLocation.ush
+++ b/Engine/Plugins/FX/Niagara/Shaders/Private/Stateless/Modules/NiagaraStatelessModule_ShapeLocation.ush
@@ -18,12 +18,12 @@ float3 ShapeLocation_GetLocation(inout FStatelessParticle Particle)
 		const float SurfaceScale = ShapeLocation_Parameters0.w;
 		const float SurfaceBias = ShapeLocation_Parameters1.w;
 
-		const float3 P0 = RandomFloat3(Particle);
+		const float3 P0 = RandomFloat3(0);
 		if (bOnSurface)
 		{
 			//-TODO: This isn't quite what we want
 			const float3 L0 = round(P0);
-			const uint S = RandomUInt(Particle) % 3;
+			const uint S = RandomUInt(1) % 3;
 			float3 Location;
 			Location.x = S != 0 ? P0.x : L0.x;
 			Location.y = S != 1 ? P0.y : L0.y;
@@ -40,7 +40,7 @@ float3 ShapeLocation_GetLocation(inout FStatelessParticle Particle)
 	BRANCH
 	if ( ShapeLocation_Mode.x == 1 )
 	{
-		const float4 Random = RandomFloat4(Particle);
+		const float4 Random = RandomFloat4(0);
 		const float HeightScale = ShapeLocation_Parameters0.x;
 		const float HeightBias = ShapeLocation_Parameters0.y;
 		const float Radius = ShapeLocation_Parameters0.z;
@@ -60,8 +60,8 @@ float3 ShapeLocation_GetLocation(inout FStatelessParticle Particle)
 		const float UDistributionScale = ShapeLocation_Parameters0.z;
 		const float UDistributionBias = ShapeLocation_Parameters0.w;
 
-		const float Radius = RandomScaleBiasFloat(Particle, RadiusScale, RadiusBias);
-		const float U = RandomScaleBiasFloat(Particle, UDistributionScale, UDistributionBias);
+		const float Radius = RandomScaleBiasFloat(0, RadiusScale, RadiusBias);
+		const float U = RandomScaleBiasFloat(1, UDistributionScale, UDistributionBias);
 
 		return float3(cos(U) * Radius, sin(U) * Radius, 0.0f);
 	}
@@ -71,8 +71,8 @@ float3 ShapeLocation_GetLocation(inout FStatelessParticle Particle)
 		const float SphereScale	= ShapeLocation_Parameters0.x;
 		const float SphereBias	= ShapeLocation_Parameters0.y;
 
-		const float3 Vector = RandomUnitFloat3(Particle);
-		return Vector * RandomScaleBiasFloat(Particle, SphereScale, SphereBias);
+		const float3 Vector = RandomUnitFloat3(0);
+		return Vector * RandomScaleBiasFloat(1, SphereScale, SphereBias);
 	}
 }
 
diff --git a/Engine/Plugins/FX/Niagara/Shaders/Private/Stateless/Modules/NiagaraStatelessModule_SolveVelocitiesAndForces.ush b/Engine/Plugins/FX/Niagara/Shaders/Private/Stateless/Modules/NiagaraStatelessModule_SolveVelocitiesAndForces.ush
--- a/Engine/Plugins/FX/Niagara/Shaders/Private/Stateless/Modules/NiagaraStatelessModule_SolveVelocitiesAndForces.ush
+++ b/Engine/Plugins/FX/Niagara/Shaders/Private/Stateless/Modules/NiagaraStatelessModule_SolveVelocitiesAndForces.ush
@@ -59,22 +59,22 @@ void SolveVelocitiesAndForces_Simulate(inout FStatelessParticle Particle)
 {
 	// Initialize shared parameters
 	FStatelessModule_SolveVelocitiesAndForces ModuleData;
-	ModuleData.Mass = RandomScaleBiasFloat(Particle, SolveVelocitiesAndForces_MassScale, SolveVelocitiesAndForces_MassBias);
-	ModuleData.Drag = RandomScaleBiasFloat(Particle, SolveVelocitiesAndForces_DragScale, SolveVelocitiesAndForces_DragBias);
-	ModuleData.Velocity = RandomScaleBiasFloat(Particle, SolveVelocitiesAndForces_VelocityScale, SolveVelocitiesAndForces_VelocityBias);
-	ModuleData.Wind = RandomScaleBiasFloat(Particle, SolveVelocitiesAndForces_WindScale, SolveVelocitiesAndForces_WindBias);
-	ModuleData.Acceleration = RandomScaleBiasFloat(Particle, SolveVelocitiesAndForces_AccelerationScale, SolveVelocitiesAndForces_AccelerationBias);
+	ModuleData.Mass = RandomScaleBiasFloat(0, SolveVelocitiesAndForces_MassScale, SolveVelocitiesAndForces_MassBias);
+	ModuleData.Drag = RandomScaleBiasFloat(1, SolveVelocitiesAndForces_DragScale, SolveVelocitiesAndForces_DragBias);
+	ModuleData.Velocity = RandomScaleBiasFloat(2, SolveVelocitiesAndForces_VelocityScale, SolveVelocitiesAndForces_VelocityBias);
+	ModuleData.Wind = RandomScaleBiasFloat(3, SolveVelocitiesAndForces_WindScale, SolveVelocitiesAndForces_WindBias);
+	ModuleData.Acceleration = RandomScaleBiasFloat(4, SolveVelocitiesAndForces_AccelerationScale, SolveVelocitiesAndForces_AccelerationBias);
 
 	if ( SolveVelocitiesAndForces_ConeVelocityEnabled != 0 )
 	{
-		const float ConeAngle = RandomScaleBiasFloat(Particle, SolveVelocitiesAndForces_ConeAngleScale, SolveVelocitiesAndForces_ConeAngleBias);
-		const float ConeRotation = RandomFloat(Particle) * UE_TWO_PI;
+		const float ConeAngle = RandomScaleBiasFloat(5, SolveVelocitiesAndForces_ConeAngleScale, SolveVelocitiesAndForces_ConeAngleBias);
+		const float ConeRotation = RandomFloat(6) * UE_TWO_PI;
 			
 		float2 scAng = SinCos(ConeAngle);
 		float2 scRot = SinCos(ConeRotation);
 		const float3 Direction = float3(scRot.x * scAng.x, scRot.y * scAng.x, scAng.y);
 		
-		float VelocityScale = RandomScaleBiasFloat(Particle, SolveVelocitiesAndForces_ConeVelocityScale, SolveVelocitiesAndForces_ConeVelocityBias);
+		float VelocityScale = RandomScaleBiasFloat(7, SolveVelocitiesAndForces_ConeVelocityScale, SolveVelocitiesAndForces_ConeVelocityBias);
 		if ( SolveVelocitiesAndForces_ConeVelocityFalloff > 0.0f )
 		{
 			const float pf = pow(saturate(scAng.y), SolveVelocitiesAndForces_ConeVelocityFalloff * 10.0f);
@@ -86,18 +86,18 @@ void SolveVelocitiesAndForces_Simulate(inout FStatelessParticle Particle)
 
 	if ( SolveVelocitiesAndForces_PontVelocityEnabled != 0 )
 	{
-		const float3 FallbackDir = RandomUnitFloat3(Particle);
+		const float3 FallbackDir = RandomUnitFloat3(8);
 		const float3 Delta = Particle.Position - SolveVelocitiesAndForces_PointOrigin;
 		const float3 Dir = SafeNormalize(Delta, FallbackDir);
 			
-		const float VelocityScale = RandomScaleBiasFloat(Particle, SolveVelocitiesAndForces_PointVelocityScale, SolveVelocitiesAndForces_PointVelocityBias);
+		const float VelocityScale = RandomScaleBiasFloat(9, SolveVelocitiesAndForces_PointVelocityScale, SolveVelocitiesAndForces_PointVelocityBias);
 			
 		ModuleData.Velocity += Dir * VelocityScale;
 	}
 
 	if ( SolveVelocitiesAndForces_NoiseEnabled  != 0 )
 	{
-		const float3 NoiseOffset = RandomFloat3(Particle);
+		const float3 NoiseOffset = RandomFloat3(10);
 		float3 NoiseValue = 0;
 		if ( SolveVelocitiesAndForces_NoiseMode == 0 )
 		{
@@ -114,8 +114,8 @@ void SolveVelocitiesAndForces_Simulate(inout FStatelessParticle Particle)
 		}
 		else if ( SolveVelocitiesAndForces_NoiseMode == 2 )
 		{
-			const uint3 ChannelIndex = RandomUInt3(Particle) % SolveVelocitiesAndForces_NoiseLUTNumChannel;
-			const float3 ChannelLerp = (Particle.Age * SolveVelocitiesAndForces_NoiseFrequency) + RandomFloat(Particle);
+			const uint3 ChannelIndex = RandomUInt3(11) % SolveVelocitiesAndForces_NoiseLUTNumChannel;
+			const float3 ChannelLerp = (Particle.Age * SolveVelocitiesAndForces_NoiseFrequency) + RandomFloat(12);
 			const uint3 ChannelLUTOffsetA = (ChannelIndex * SolveVelocitiesAndForces_NoiseLUTChannelWidth) + min(floor(ChannelLerp) + 0, SolveVelocitiesAndForces_NoiseLUTChannelWidth - 1);
 			const uint3 ChannelLUTOffsetB = (ChannelIndex * SolveVelocitiesAndForces_NoiseLUTChannelWidth) + min(floor(ChannelLerp) + 1, SolveVelocitiesAndForces_NoiseLUTChannelWidth - 1);
 				
@@ -133,8 +133,8 @@ void SolveVelocitiesAndForces_Simulate(inout FStatelessParticle Particle)
 		}
 		else if ( SolveVelocitiesAndForces_NoiseMode == 3 )
 		{
-			const uint3 ChannelIndex = RandomUInt3(Particle) % SolveVelocitiesAndForces_NoiseLUTNumChannel;
-			const float ChannelLerp = (Particle.Age * SolveVelocitiesAndForces_NoiseFrequency.x) + RandomFloat(Particle);
+			const uint3 ChannelIndex = RandomUInt3(13) % SolveVelocitiesAndForces_NoiseLUTNumChannel;
+			const float ChannelLerp = (Particle.Age * SolveVelocitiesAndForces_NoiseFrequency.x) + RandomFloat(14);
 			const uint3 ChannelLUTOffsetA = (ChannelIndex * SolveVelocitiesAndForces_NoiseLUTChannelWidth) + min(floor(ChannelLerp) + 0, SolveVelocitiesAndForces_NoiseLUTChannelWidth - 1);
 			const uint3 ChannelLUTOffsetB = (ChannelIndex * SolveVelocitiesAndForces_NoiseLUTChannelWidth) + min(floor(ChannelLerp) + 1, SolveVelocitiesAndForces_NoiseLUTChannelWidth - 1);
 			const uint3 ChannelLUTOffsetC = (ChannelIndex * SolveVelocitiesAndForces_NoiseLUTChannelWidth) + min(floor(ChannelLerp) + 2, SolveVelocitiesAndForces_NoiseLUTChannelWidth - 1);
diff --git a/Engine/Plugins/FX/Niagara/Shaders/Private/Stateless/Modules/NiagaraStatelessModule_SpriteRotationRate.ush b/Engine/Plugins/FX/Niagara/Shaders/Private/Stateless/Modules/NiagaraStatelessModule_SpriteRotationRate.ush
--- a/Engine/Plugins/FX/Niagara/Shaders/Private/Stateless/Modules/NiagaraStatelessModule_SpriteRotationRate.ush
+++ b/Engine/Plugins/FX/Niagara/Shaders/Private/Stateless/Modules/NiagaraStatelessModule_SpriteRotationRate.ush
@@ -7,7 +7,7 @@ float		SpriteRotationRate_Bias;
 
 void SpriteRotationRate_Simulate(inout FStatelessParticle Particle)
 {
-	const float RotationRate = RandomScaleBiasFloat(Particle, SpriteRotationRate_Scale, SpriteRotationRate_Bias);
+	const float RotationRate = RandomScaleBiasFloat(0, SpriteRotationRate_Scale, SpriteRotationRate_Bias);
 
 	Particle.SpriteRotation += Particle.Age * RotationRate;
 	Particle.PreviousSpriteRotation += Particle.PreviousAge * RotationRate;
diff --git a/Engine/Plugins/FX/Niagara/Shaders/Private/Stateless/Modules/NiagaraStatelessModule_SubUVAnimation.ush b/Engine/Plugins/FX/Niagara/Shaders/Private/Stateless/Modules/NiagaraStatelessModule_SubUVAnimation.ush
--- a/Engine/Plugins/FX/Niagara/Shaders/Private/Stateless/Modules/NiagaraStatelessModule_SubUVAnimation.ush
+++ b/Engine/Plugins/FX/Niagara/Shaders/Private/Stateless/Modules/NiagaraStatelessModule_SubUVAnimation.ush
@@ -14,7 +14,7 @@ float	SubUVAnimation_RateScale;
 void SubUVAnimation_Simulate(inout FStatelessParticle Particle)
 {
 	const uint SeedOffset = uint(Particle.Age * SubUVAnimation_InitialFrameRateChange);
-	float Frame = RandomSeedOffsetFloat(Particle, SeedOffset) * SubUVAnimation_InitialFrameScale + SubUVAnimation_InitialFrameBias;
+	float Frame = RandomFloat(SeedOffset) * SubUVAnimation_InitialFrameScale + SubUVAnimation_InitialFrameBias;
 
 	if ( SubUVAnimation_Mode == 0 )
 	{
diff --git a/Engine/Plugins/FX/Niagara/Shaders/Private/Stateless/NiagaraStatelessCommon.ush b/Engine/Plugins/FX/Niagara/Shaders/Private/Stateless/NiagaraStatelessCommon.ush
--- a/Engine/Plugins/FX/Niagara/Shaders/Private/Stateless/NiagaraStatelessCommon.ush
+++ b/Engine/Plugins/FX/Niagara/Shaders/Private/Stateless/NiagaraStatelessCommon.ush
@@ -66,7 +66,6 @@ struct FStatelessParticle
 {
 	bool bAlive;					// If the particle should remain alive or not (only evaluated as specific points)
 	uint UniqueIndex;				// Unique particle index, sequential based on previous spawn info accumulation
-	uint4 RandomSeed;				// Random seed for this particle
 	float Lifetime;					// Overall lifetime for the particle
 	float Age;						// Current age for the particle
 	float NormalizedAge;			// Current normalized age for the particle
@@ -84,55 +83,45 @@ struct FStatelessParticle
 	#undef PARTICLE_ATTRIBUTE
 	/////////////////////////////////////////////////////////////////////////////////////////////////////////////////////////
 };
-	
-uint4 RandomUInt4(inout FStatelessParticle Particle)
+
+static uint4 GRandomSeedInternal;
+
+uint4 RandomUInt4(uint RandomSeedOffset)
 {
-	++Particle.RandomSeed.x;
-	Particle.RandomSeed.w = Particle.RandomSeed.x ^ Particle.RandomSeed.y;
-	return Rand4DPCG32(Particle.RandomSeed);
+	uint4 RandomSeed = GRandomSeedInternal;
+	RandomSeed.x += RandomSeedOffset;
+	return Rand4DPCG32(RandomSeed);
 }
-uint   RandomUInt(inout FStatelessParticle Particle) { return RandomUInt4(Particle).x; }
-uint2  RandomUInt2(inout FStatelessParticle Particle) { return RandomUInt4(Particle).xy; }
-uint3  RandomUInt3(inout FStatelessParticle Particle) { return RandomUInt4(Particle).xyz; }
+uint   RandomUInt(uint RandomSeedOffset) { return RandomUInt4(RandomSeedOffset).x; }
+uint2  RandomUInt2(uint RandomSeedOffset) { return RandomUInt4(RandomSeedOffset).xy; }
+uint3  RandomUInt3(uint RandomSeedOffset) { return RandomUInt4(RandomSeedOffset).xyz; }
 
-float4 RandomFloat4(inout FStatelessParticle Particle)
+float4 RandomFloat4(uint RandomSeedOffset)
 {
-	uint4 v = RandomUInt4(Particle);
+	uint4 v = RandomUInt4(RandomSeedOffset);
 	return float4((v >> 8) & 0x00ffffff) / 16777216.0; // 0x01000000 == 16777216
 }
-float  RandomFloat(inout FStatelessParticle Particle)  { return RandomFloat4(Particle).x; }
-float2 RandomFloat2(inout FStatelessParticle Particle) { return RandomFloat4(Particle).xy; }
-float3 RandomFloat3(inout FStatelessParticle Particle) { return RandomFloat4(Particle).xyz; }
+float  RandomFloat(uint RandomSeedOffset)  { return RandomFloat4(RandomSeedOffset).x; }
+float2 RandomFloat2(uint RandomSeedOffset) { return RandomFloat4(RandomSeedOffset).xy; }
+float3 RandomFloat3(uint RandomSeedOffset) { return RandomFloat4(RandomSeedOffset).xyz; }
 	
-float  RandomRangeFloat(inout FStatelessParticle Particle, float  Min, float  Max) { return (RandomFloat(Particle)  * (Max - Min)) + Min; }
-float2 RandomRangeFloat(inout FStatelessParticle Particle, float2 Min, float2 Max) { return (RandomFloat2(Particle) * (Max - Min)) + Min; }
-float3 RandomRangeFloat(inout FStatelessParticle Particle, float3 Min, float3 Max) { return (RandomFloat3(Particle) * (Max - Min)) + Min; }
-float4 RandomRangeFloat(inout FStatelessParticle Particle, float4 Min, float4 Max) { return (RandomFloat4(Particle) * (Max - Min)) + Min; }
-
-float  RandomScaleBiasFloat(inout FStatelessParticle Particle, float  Scale, float  Bias) { return (RandomFloat(Particle)  * Scale) + Bias; }
-float2 RandomScaleBiasFloat(inout FStatelessParticle Particle, float2 Scale, float2 Bias) { return (RandomFloat2(Particle) * Scale) + Bias; }
-float3 RandomScaleBiasFloat(inout FStatelessParticle Particle, float3 Scale, float3 Bias) { return (RandomFloat3(Particle) * Scale) + Bias; }
-float4 RandomScaleBiasFloat(inout FStatelessParticle Particle, float4 Scale, float4 Bias) { return (RandomFloat4(Particle) * Scale) + Bias; }
-
-float  RandomScaleBiasFloat(inout FStatelessParticle Particle, float  Scale, float  Bias, bool bUniform) { return (RandomFloat(Particle)  * Scale) + Bias; }
-float2 RandomScaleBiasFloat(inout FStatelessParticle Particle, float2 Scale, float2 Bias, bool bUniform) { float2 Random = (RandomFloat2(Particle) * Scale) + Bias; return bUniform ? Random.xx : Random; }
-float3 RandomScaleBiasFloat(inout FStatelessParticle Particle, float3 Scale, float3 Bias, bool bUniform) { float3 Random = (RandomFloat3(Particle) * Scale) + Bias; return bUniform ? Random.xxx : Random; }
-float4 RandomScaleBiasFloat(inout FStatelessParticle Particle, float4 Scale, float4 Bias, bool bUniform) { float4 Random = (RandomFloat4(Particle) * Scale) + Bias; return bUniform ? Random.xxxx : Random; }
-
-float2 RandomUnitFloat2(inout FStatelessParticle Particle) { return SafeNormalize(RandomFloat2(Particle) - 0.5f); }
-float3 RandomUnitFloat3(inout FStatelessParticle Particle) { return SafeNormalize(RandomFloat3(Particle) - 0.5f); }
-
-float4 RandomSeedOffsetFloat4(in FStatelessParticle Particle, uint SeedOffset)
-{
-	uint4 Seed = Particle.RandomSeed;
-	Seed.x += SeedOffset;
-		
-	uint4 V = Rand4DPCG32(Seed);
-	return float4((V >> 8) & 0x00ffffff) / 16777216.0; // 0x01000000 == 16777216
-}
-float  RandomSeedOffsetFloat(inout FStatelessParticle Particle, uint SeedOffset)  { return RandomSeedOffsetFloat4(Particle, SeedOffset).x; }
-float2 RandomSeedOffsetFloat2(inout FStatelessParticle Particle, uint SeedOffset) { return RandomSeedOffsetFloat4(Particle, SeedOffset).xy; }
-float3 RandomSeedOffsetFloat3(inout FStatelessParticle Particle, uint SeedOffset) { return RandomSeedOffsetFloat4(Particle, SeedOffset).xyz; }
+float  RandomRangeFloat(uint RandomSeedOffset, float  Min, float  Max) { return (RandomFloat(RandomSeedOffset)  * (Max - Min)) + Min; }
+float2 RandomRangeFloat(uint RandomSeedOffset, float2 Min, float2 Max) { return (RandomFloat2(RandomSeedOffset) * (Max - Min)) + Min; }
+float3 RandomRangeFloat(uint RandomSeedOffset, float3 Min, float3 Max) { return (RandomFloat3(RandomSeedOffset) * (Max - Min)) + Min; }
+float4 RandomRangeFloat(uint RandomSeedOffset, float4 Min, float4 Max) { return (RandomFloat4(RandomSeedOffset) * (Max - Min)) + Min; }
+
+float  RandomScaleBiasFloat(uint RandomSeedOffset, float  Scale, float  Bias) { return (RandomFloat(RandomSeedOffset)  * Scale) + Bias; }
+float2 RandomScaleBiasFloat(uint RandomSeedOffset, float2 Scale, float2 Bias) { return (RandomFloat2(RandomSeedOffset) * Scale) + Bias; }
+float3 RandomScaleBiasFloat(uint RandomSeedOffset, float3 Scale, float3 Bias) { return (RandomFloat3(RandomSeedOffset) * Scale) + Bias; }
+float4 RandomScaleBiasFloat(uint RandomSeedOffset, float4 Scale, float4 Bias) { return (RandomFloat4(RandomSeedOffset) * Scale) + Bias; }
+
+float  RandomScaleBiasFloat(uint RandomSeedOffset, float  Scale, float  Bias, bool bUniform) { return (RandomFloat(RandomSeedOffset)  * Scale) + Bias; }
+float2 RandomScaleBiasFloat(uint RandomSeedOffset, float2 Scale, float2 Bias, bool bUniform) { float2 Random = (RandomFloat2(RandomSeedOffset) * Scale) + Bias; return bUniform ? Random.xx : Random; }
+float3 RandomScaleBiasFloat(uint RandomSeedOffset, float3 Scale, float3 Bias, bool bUniform) { float3 Random = (RandomFloat3(RandomSeedOffset) * Scale) + Bias; return bUniform ? Random.xxx : Random; }
+float4 RandomScaleBiasFloat(uint RandomSeedOffset, float4 Scale, float4 Bias, bool bUniform) { float4 Random = (RandomFloat4(RandomSeedOffset) * Scale) + Bias; return bUniform ? Random.xxxx : Random; }
+
+float2 RandomUnitFloat2(uint RandomSeedOffset) { return SafeNormalize(RandomFloat2(RandomSeedOffset) - 0.5f); }
+float3 RandomUnitFloat3(uint RandomSeedOffset) { return SafeNormalize(RandomFloat3(RandomSeedOffset) - 0.5f); }
 
 float  GetStaticFloat (int BaseOffset, int Element)	{ return Common_StaticFloatBuffer[BaseOffset + Element]; }
 float2 GetStaticFloat2(int BaseOffset, int Element)	{ return float2(Common_StaticFloatBuffer[BaseOffset + Element*2], Common_StaticFloatBuffer[BaseOffset + 1 + Element*2]); }
diff --git a/Engine/Plugins/FX/Niagara/Shaders/Private/Stateless/NiagaraStatelessDistribution.ush b/Engine/Plugins/FX/Niagara/Shaders/Private/Stateless/NiagaraStatelessDistribution.ush
--- a/Engine/Plugins/FX/Niagara/Shaders/Private/Stateless/NiagaraStatelessDistribution.ush
+++ b/Engine/Plugins/FX/Niagara/Shaders/Private/Stateless/NiagaraStatelessDistribution.ush
@@ -8,13 +8,13 @@
 #define DefineStatelessDistributionSampler(HLSLTYPE, NAMETYPE) \
 	struct FStatelessDistributionSampler##NAMETYPE \
 	{ \
-		void Init(inout FStatelessParticle Particle, uint3 InParameters) \
+		void Init(uint RandomSeedOffset, uint3 InParameters) \
 		{ \
 			Parameters = InParameters; \
 			RandomOffset = 0.0f; \
 			if ( (Parameters.x & StatelessDistributionFlag_Random) != 0) \
 			{ \
-				RandomOffset = Random##NAMETYPE##(Particle); \
+				RandomOffset = Random##NAMETYPE##(RandomSeedOffset); \
 			} \
 		} \
 		\
diff --git a/Engine/Plugins/FX/Niagara/Shaders/Private/Stateless/NiagaraStatelessSimulationDefault.usf b/Engine/Plugins/FX/Niagara/Shaders/Private/Stateless/NiagaraStatelessSimulationDefault.usf
--- a/Engine/Plugins/FX/Niagara/Shaders/Private/Stateless/NiagaraStatelessSimulationDefault.usf
+++ b/Engine/Plugins/FX/Niagara/Shaders/Private/Stateless/NiagaraStatelessSimulationDefault.usf
@@ -36,10 +36,18 @@
 
 // Defines modules and order of execution
 #define PARTICLE_MODULES \
+	/* Initializer Modules */ \
 	PARTICLE_MODULE(InitializeParticle) \
 	PARTICLE_MODULE(InitialMeshOrientation) \
 	PARTICLE_MODULE(ShapeLocation) \
+	/* Before Solve */ \
+	PARTICLE_MODULE(Null) /* AddVelocity */ \
+	PARTICLE_MODULE(Null) /* AccelerationForce */ \
+	PARTICLE_MODULE(Null) /* CurlNoiseForce */ \
+	PARTICLE_MODULE(Null) /* Drag */ \
+	PARTICLE_MODULE(Null) /* GravityForce */ \
 	PARTICLE_MODULE(SolveVelocitiesAndForces) \
+	/* Post Solve */ \
 	PARTICLE_MODULE(CameraOffset) \
 	PARTICLE_MODULE(DynamicMaterialParameters) \
 	PARTICLE_MODULE(MeshIndex) \
diff --git a/Engine/Plugins/FX/Niagara/Shaders/Private/Stateless/NiagaraStatelessSimulationTemplate.ush b/Engine/Plugins/FX/Niagara/Shaders/Private/Stateless/NiagaraStatelessSimulationTemplate.ush
--- a/Engine/Plugins/FX/Niagara/Shaders/Private/Stateless/NiagaraStatelessSimulationTemplate.ush
+++ b/Engine/Plugins/FX/Niagara/Shaders/Private/Stateless/NiagaraStatelessSimulationTemplate.ush
@@ -54,6 +54,12 @@ DECLARE_SCALAR_ARRAY(float,	SpawnInfo_Rate, NIAGARA_MAX_GPU_SPAWN_INFOS);
 
 /////////////////////////////////////////////////////////////////////////////////////////////////////////////////////////
 
+void Null_Simulate(inout FStatelessParticle Particle)
+{
+}
+
+/////////////////////////////////////////////////////////////////////////////////////////////////////////////////////////
+
 [numthreads(THREADGROUP_SIZE, 1, 1)]
 void StatelessMain(uint3 GroupId : SV_GroupID, int GroupThreadIndex : SV_GroupIndex)
 {
@@ -91,12 +97,9 @@ void StatelessMain(uint3 GroupId : SV_GroupID, int GroupThreadIndex : SV_GroupIn
 	}
 	
 	// Initialize variables / determine lifetime / etc
-	Particle.RandomSeed.x = 7123u;
-	Particle.RandomSeed.y = Common_RandomSeed;
-	Particle.RandomSeed.z = Particle.UniqueIndex * 3581u;
-	Particle.RandomSeed.w = 0;
+	GRandomSeedInternal = uint4(7123u, Common_RandomSeed, Particle.UniqueIndex * 3581u, 6405u);
 
-	Particle.Lifetime = RandomScaleBiasFloat(Particle, Common_LifetimeScaleBias.x, Common_LifetimeScaleBias.y);
+	Particle.Lifetime = RandomScaleBiasFloat(0, Common_LifetimeScaleBias.x, Common_LifetimeScaleBias.y);
 	if ( Particle.Lifetime <= 0.0f || Particle.Age >= Particle.Lifetime )
 	{
 		return;
@@ -112,8 +115,12 @@ void StatelessMain(uint3 GroupId : SV_GroupID, int GroupThreadIndex : SV_GroupIn
 	Particle.bAlive = true;
 	
 	// Run Particle Simulate
-	#define PARTICLE_MODULE(NAME) ##NAME##_Simulate(Particle);
+	#define PARTICLE_MODULE(NAME) \
+		GRandomSeedInternal.w += 6405u;\
+		##NAME##_Simulate(Particle);
+
 		PARTICLE_MODULES
+
 	#undef PARTICLE_MODULE
 
 	if ( !Particle.bAlive )
diff --git a/Engine/Plugins/FX/Niagara/Source/Niagara/Internal/Stateless/Modules/NiagaraStatelessModule_CameraOffset.h b/Engine/Plugins/FX/Niagara/Source/Niagara/Internal/Stateless/Modules/NiagaraStatelessModule_CameraOffset.h
--- a/Engine/Plugins/FX/Niagara/Source/Niagara/Internal/Stateless/Modules/NiagaraStatelessModule_CameraOffset.h
+++ b/Engine/Plugins/FX/Niagara/Source/Niagara/Internal/Stateless/Modules/NiagaraStatelessModule_CameraOffset.h
@@ -63,7 +63,7 @@ public:
 
 		for (uint32 i = 0; i < ParticleSimulationContext.GetNumInstances(); ++i)
 		{
-			const FStatelessDistributionSampler<float> CameraOffsetSampler(ParticleSimulationContext, ModuleBuiltData->DistributionParameters, i);
+			const FStatelessDistributionSampler<float> CameraOffsetSampler(ParticleSimulationContext, ModuleBuiltData->DistributionParameters, i, 0);
 
 			const float NormalizedAge			= NormalizedAgeData[i];
 			const float PreviousNormalizedAge	= PreviousNormalizedAgeData[i];
diff --git a/Engine/Plugins/FX/Niagara/Source/Niagara/Internal/Stateless/Modules/NiagaraStatelessModule_DynamicMaterialParameters.h b/Engine/Plugins/FX/Niagara/Source/Niagara/Internal/Stateless/Modules/NiagaraStatelessModule_DynamicMaterialParameters.h
--- a/Engine/Plugins/FX/Niagara/Source/Niagara/Internal/Stateless/Modules/NiagaraStatelessModule_DynamicMaterialParameters.h
+++ b/Engine/Plugins/FX/Niagara/Source/Niagara/Internal/Stateless/Modules/NiagaraStatelessModule_DynamicMaterialParameters.h
@@ -236,10 +236,10 @@ public:
 				}
 
 				const int32 FirstChannel = iParameter * NumChannelPerParameter;
-				const FStatelessDistributionSampler<float> XChannelSampler(ParticleSimulationContext, ModuleBuiltData->ParameterDistributions[FirstChannel + 0], i);
-				const FStatelessDistributionSampler<float> YChannelSampler(ParticleSimulationContext, ModuleBuiltData->ParameterDistributions[FirstChannel + 1], i);
-				const FStatelessDistributionSampler<float> ZChannelSampler(ParticleSimulationContext, ModuleBuiltData->ParameterDistributions[FirstChannel + 2], i);
-				const FStatelessDistributionSampler<float> WChannelSampler(ParticleSimulationContext, ModuleBuiltData->ParameterDistributions[FirstChannel + 3], i);
+				const FStatelessDistributionSampler<float> XChannelSampler(ParticleSimulationContext, ModuleBuiltData->ParameterDistributions[FirstChannel + 0], i, FirstChannel + 0);
+				const FStatelessDistributionSampler<float> YChannelSampler(ParticleSimulationContext, ModuleBuiltData->ParameterDistributions[FirstChannel + 1], i, FirstChannel + 1);
+				const FStatelessDistributionSampler<float> ZChannelSampler(ParticleSimulationContext, ModuleBuiltData->ParameterDistributions[FirstChannel + 2], i, FirstChannel + 2);
+				const FStatelessDistributionSampler<float> WChannelSampler(ParticleSimulationContext, ModuleBuiltData->ParameterDistributions[FirstChannel + 3], i, FirstChannel + 3);
 
 				const float NormalizedAge = NormalizedAgeData[i];
 				const FVector4f DynamicParameter(
diff --git a/Engine/Plugins/FX/Niagara/Source/Niagara/Internal/Stateless/Modules/NiagaraStatelessModule_InitialMeshOrientation.h b/Engine/Plugins/FX/Niagara/Source/Niagara/Internal/Stateless/Modules/NiagaraStatelessModule_InitialMeshOrientation.h
--- a/Engine/Plugins/FX/Niagara/Source/Niagara/Internal/Stateless/Modules/NiagaraStatelessModule_InitialMeshOrientation.h
+++ b/Engine/Plugins/FX/Niagara/Source/Niagara/Internal/Stateless/Modules/NiagaraStatelessModule_InitialMeshOrientation.h
@@ -67,7 +67,7 @@ public:
 		const FModuleBuiltData* ModuleBuiltData = ParticleSimulationContext.ReadBuiltData<FModuleBuiltData>();
 		for (uint32 i = 0; i < ParticleSimulationContext.GetNumInstances(); ++i)
 		{
-			const FVector3f	Rotation = ModuleBuiltData->Rotation + (ParticleSimulationContext.RandomFloat3(i) * ModuleBuiltData->RandomRotationRange);
+			const FVector3f	Rotation = ModuleBuiltData->Rotation + (ParticleSimulationContext.RandomFloat3(i, 0) * ModuleBuiltData->RandomRotationRange);
 			const FQuat4f Quat = ParticleSimulationContext.RotatorToQuat(Rotation);
 
 			ParticleSimulationContext.WriteParticleVariable(ModuleBuiltData->MeshOrientationVariableOffset, i, Quat);
diff --git a/Engine/Plugins/FX/Niagara/Source/Niagara/Internal/Stateless/Modules/NiagaraStatelessModule_InitializeParticle.h b/Engine/Plugins/FX/Niagara/Source/Niagara/Internal/Stateless/Modules/NiagaraStatelessModule_InitializeParticle.h
--- a/Engine/Plugins/FX/Niagara/Source/Niagara/Internal/Stateless/Modules/NiagaraStatelessModule_InitializeParticle.h
+++ b/Engine/Plugins/FX/Niagara/Source/Niagara/Internal/Stateless/Modules/NiagaraStatelessModule_InitializeParticle.h
@@ -133,14 +133,14 @@ public:
 
 		for (uint32 i = 0; i < ParticleSimulationContext.GetNumInstances(); ++i)
 		{
-			const FStatelessDistributionSampler<FVector3f> PositionSampler(ParticleSimulationContext, ModuleBuiltData->InitialPosition, i);
+			const FStatelessDistributionSampler<FVector3f> PositionSampler(ParticleSimulationContext, ModuleBuiltData->InitialPosition, i, 0);
 
 			const FVector3f		Position	= PositionSampler.GetValue(ParticleSimulationContext, 0.0f);
-			const FLinearColor	Color		= ParticleSimulationContext.RandomScaleBiasFloat(i, ModuleBuiltData->ColorRange);
-			const float			RibbonWidth	= ParticleSimulationContext.RandomScaleBiasFloat(i, ModuleBuiltData->RibbonWidthRange);
-			const FVector2f		SpriteSize	= ParticleSimulationContext.RandomScaleBiasFloat(i, ModuleBuiltData->SpriteSizeRange);
-			const float			SpriteRot	= ParticleSimulationContext.RandomScaleBiasFloat(i, ModuleBuiltData->SpriteRotationRange);
-			const FVector3f		Scale		= ParticleSimulationContext.RandomScaleBiasFloat(i, ModuleBuiltData->MeshScaleRange);
+			const FLinearColor	Color		= ParticleSimulationContext.RandomScaleBiasFloat(i, 1 ,ModuleBuiltData->ColorRange);
+			const float			RibbonWidth	= ParticleSimulationContext.RandomScaleBiasFloat(i, 2, ModuleBuiltData->RibbonWidthRange);
+			const FVector2f		SpriteSize	= ParticleSimulationContext.RandomScaleBiasFloat(i, 3, ModuleBuiltData->SpriteSizeRange);
+			const float			SpriteRot	= ParticleSimulationContext.RandomScaleBiasFloat(i, 4, ModuleBuiltData->SpriteRotationRange);
+			const FVector3f		Scale		= ParticleSimulationContext.RandomScaleBiasFloat(i, 5, ModuleBuiltData->MeshScaleRange);
 
 			ParticleSimulationContext.WriteParticleVariable(ModuleBuiltData->PositionVariableOffset,				i, Position);
 			ParticleSimulationContext.WriteParticleVariable(ModuleBuiltData->ColorVariableOffset,					i, Color);
diff --git a/Engine/Plugins/FX/Niagara/Source/Niagara/Internal/Stateless/Modules/NiagaraStatelessModule_MeshRotationRate.h b/Engine/Plugins/FX/Niagara/Source/Niagara/Internal/Stateless/Modules/NiagaraStatelessModule_MeshRotationRate.h
--- a/Engine/Plugins/FX/Niagara/Source/Niagara/Internal/Stateless/Modules/NiagaraStatelessModule_MeshRotationRate.h
+++ b/Engine/Plugins/FX/Niagara/Source/Niagara/Internal/Stateless/Modules/NiagaraStatelessModule_MeshRotationRate.h
@@ -69,7 +69,7 @@ public:
 
 		for (uint32 i = 0; i < ParticleSimulationContext.GetNumInstances(); ++i)
 		{
-			const FVector3f RotationRate = ParticleSimulationContext.RandomScaleBiasFloat(i, ModuleBuiltData->RotationRange);
+			const FVector3f RotationRate = ParticleSimulationContext.RandomScaleBiasFloat(i, 0, ModuleBuiltData->RotationRange);
 			const float Age = AgeData[i];
 			const float PreviousAge = PreviousAgeData[i];
 
diff --git a/Engine/Plugins/FX/Niagara/Source/Niagara/Internal/Stateless/Modules/NiagaraStatelessModule_ScaleColor.h b/Engine/Plugins/FX/Niagara/Source/Niagara/Internal/Stateless/Modules/NiagaraStatelessModule_ScaleColor.h
--- a/Engine/Plugins/FX/Niagara/Source/Niagara/Internal/Stateless/Modules/NiagaraStatelessModule_ScaleColor.h
+++ b/Engine/Plugins/FX/Niagara/Source/Niagara/Internal/Stateless/Modules/NiagaraStatelessModule_ScaleColor.h
@@ -62,7 +62,7 @@ public:
 
 		for (uint32 i = 0; i < ParticleSimulationContext.GetNumInstances(); ++i)
 		{
-			FStatelessDistributionSampler<FLinearColor> ColorSampler(ParticleSimulationContext, ModuleBuiltData->DistributionParameters, i);
+			FStatelessDistributionSampler<FLinearColor> ColorSampler(ParticleSimulationContext, ModuleBuiltData->DistributionParameters, i, 0);
 
 			FLinearColor Color = ParticleSimulationContext.ReadParticleVariable(ModuleBuiltData->ColorVariableOffset, i, FLinearColor::White);
 
diff --git a/Engine/Plugins/FX/Niagara/Source/Niagara/Internal/Stateless/Modules/NiagaraStatelessModule_ScaleMeshSize.h b/Engine/Plugins/FX/Niagara/Source/Niagara/Internal/Stateless/Modules/NiagaraStatelessModule_ScaleMeshSize.h
--- a/Engine/Plugins/FX/Niagara/Source/Niagara/Internal/Stateless/Modules/NiagaraStatelessModule_ScaleMeshSize.h
+++ b/Engine/Plugins/FX/Niagara/Source/Niagara/Internal/Stateless/Modules/NiagaraStatelessModule_ScaleMeshSize.h
@@ -97,7 +97,7 @@ public:
 
 		for (uint32 i = 0; i < ParticleSimulationContext.GetNumInstances(); ++i)
 		{
-			FStatelessDistributionSampler<FVector3f> ScaleSampler(ParticleSimulationContext, ModuleBuiltData->DistributionParameters, i);
+			FStatelessDistributionSampler<FVector3f> ScaleSampler(ParticleSimulationContext, ModuleBuiltData->DistributionParameters, i, 0);
 
 			FVector3f Scale = ParticleSimulationContext.ReadParticleVariable(ModuleBuiltData->ScaleVariableOffset, i, FVector3f::OneVector);
 			FVector3f PreviousScale = ParticleSimulationContext.ReadParticleVariable(ModuleBuiltData->PreviousScaleVariableOffset, i, FVector3f::OneVector);
diff --git a/Engine/Plugins/FX/Niagara/Source/Niagara/Internal/Stateless/Modules/NiagaraStatelessModule_ScaleSpriteSize.h b/Engine/Plugins/FX/Niagara/Source/Niagara/Internal/Stateless/Modules/NiagaraStatelessModule_ScaleSpriteSize.h
--- a/Engine/Plugins/FX/Niagara/Source/Niagara/Internal/Stateless/Modules/NiagaraStatelessModule_ScaleSpriteSize.h
+++ b/Engine/Plugins/FX/Niagara/Source/Niagara/Internal/Stateless/Modules/NiagaraStatelessModule_ScaleSpriteSize.h
@@ -91,7 +91,7 @@ public:
 
 		for (uint32 i = 0; i < ParticleSimulationContext.GetNumInstances(); ++i)
 		{
-			const FStatelessDistributionSampler<FVector2f> SpriteScaleSampler(ParticleSimulationContext, ModuleBuiltData->DistributionParameters, i);
+			const FStatelessDistributionSampler<FVector2f> SpriteScaleSampler(ParticleSimulationContext, ModuleBuiltData->DistributionParameters, i, 0);
 
 			const float NormalizedAge			= NormalizedAgeData[i];
 			const float PreviousNormalizedAge	= PreviousNormalizedAgeData[i];
diff --git a/Engine/Plugins/FX/Niagara/Source/Niagara/Internal/Stateless/Modules/NiagaraStatelessModule_SolveVelocitiesAndForces.h b/Engine/Plugins/FX/Niagara/Source/Niagara/Internal/Stateless/Modules/NiagaraStatelessModule_SolveVelocitiesAndForces.h
--- a/Engine/Plugins/FX/Niagara/Source/Niagara/Internal/Stateless/Modules/NiagaraStatelessModule_SolveVelocitiesAndForces.h
+++ b/Engine/Plugins/FX/Niagara/Source/Niagara/Internal/Stateless/Modules/NiagaraStatelessModule_SolveVelocitiesAndForces.h
@@ -135,24 +135,24 @@ public:
 
 		for (uint32 i = 0; i < ParticleSimulationContext.GetNumInstances(); ++i)
 		{
-			const float Mass				= ParticleSimulationContext.RandomScaleBiasFloat(i, ModuleBuiltData->PhysicsData.MassRange);
-			const float Drag				= ParticleSimulationContext.RandomScaleBiasFloat(i, ModuleBuiltData->PhysicsData.DragRange);
-			FVector3f InitialVelocity		= ParticleSimulationContext.RandomScaleBiasFloat(i, ModuleBuiltData->PhysicsData.VelocityRange);
-			const FVector3f Wind			= ParticleSimulationContext.RandomScaleBiasFloat(i, ModuleBuiltData->PhysicsData.WindRange);
-			const FVector3f Acceleration	= ParticleSimulationContext.RandomScaleBiasFloat(i, ModuleBuiltData->PhysicsData.AccelerationRange);
+			const float Mass				= ParticleSimulationContext.RandomScaleBiasFloat(i, 0, ModuleBuiltData->PhysicsData.MassRange);
+			const float Drag				= ParticleSimulationContext.RandomScaleBiasFloat(i, 1, ModuleBuiltData->PhysicsData.DragRange);
+			FVector3f InitialVelocity		= ParticleSimulationContext.RandomScaleBiasFloat(i, 2, ModuleBuiltData->PhysicsData.VelocityRange);
+			const FVector3f Wind			= ParticleSimulationContext.RandomScaleBiasFloat(i, 3, ModuleBuiltData->PhysicsData.WindRange);
+			const FVector3f Acceleration	= ParticleSimulationContext.RandomScaleBiasFloat(i, 4, ModuleBuiltData->PhysicsData.AccelerationRange);
 
 			FVector3f Position			= ParticleSimulationContext.ReadParticleVariable(ModuleBuiltData->PositionVariableOffset, i, FVector3f::ZeroVector);
 			FVector3f PreviousPosition	= ParticleSimulationContext.ReadParticleVariable(ModuleBuiltData->PreviousPositionVariableOffset, i, FVector3f::ZeroVector);
 
 			if (ModuleBuiltData->PhysicsData.bConeVelocity)
 			{
-				const float ConeAngle = ParticleSimulationContext.RandomScaleBiasFloat(i, FNiagaraStatelessRangeFloat(ModuleBuiltData->PhysicsData.ConeInnerAngle, ModuleBuiltData->PhysicsData.ConeOuterAngle)) * (UE_PI / 360.0f);
-				const float ConeRotation = ParticleSimulationContext.RandomFloat(i) * UE_TWO_PI;
+				const float ConeAngle = ParticleSimulationContext.RandomScaleBiasFloat(i, 5, FNiagaraStatelessRangeFloat(ModuleBuiltData->PhysicsData.ConeInnerAngle, ModuleBuiltData->PhysicsData.ConeOuterAngle)) * (UE_PI / 360.0f);
+				const float ConeRotation = ParticleSimulationContext.RandomFloat(i, 6) * UE_TWO_PI;
 				const FVector2f scAng = FVector2f(FMath::Sin(ConeAngle), FMath::Cos(ConeAngle));
 				const FVector2f scRot = FVector2f(FMath::Sin(ConeRotation), FMath::Cos(ConeRotation));
 				const FVector3f Direction = FVector3f(scRot.X * scAng.X, scRot.Y * scAng.X, scAng.Y);
 
-				float VelocityScale = ParticleSimulationContext.RandomScaleBiasFloat(i, ModuleBuiltData->PhysicsData.ConeVelocityRange);
+				float VelocityScale = ParticleSimulationContext.RandomScaleBiasFloat(i, 7, ModuleBuiltData->PhysicsData.ConeVelocityRange);
 				if (ModuleBuiltData->PhysicsData.bConeVelocity)
 				{
 					const float pf = FMath::Pow(FMath::Clamp(scAng.Y, 0.0f, 1.0f), ModuleBuiltData->PhysicsData.ConeVelocityFalloff * 10.0f);
@@ -164,10 +164,10 @@ public:
 
 			if (ModuleBuiltData->PhysicsData.bPointVelocity)
 			{
-				const FVector3f FallbackDir	= ParticleSimulationContext.RandomUnitFloat3(i);
+				const FVector3f FallbackDir	= ParticleSimulationContext.RandomUnitFloat3(i, 8);
 				const FVector3f Delta		= Position - ModuleBuiltData->PhysicsData.PointOrigin;
 				const FVector3f Direction	= ParticleSimulationContext.SafeNormalize(Delta, FallbackDir);
-				const float		VelocityScale = ParticleSimulationContext.RandomScaleBiasFloat(i, ModuleBuiltData->PhysicsData.PointVelocityRange);
+				const float		VelocityScale = ParticleSimulationContext.RandomScaleBiasFloat(i, 9, ModuleBuiltData->PhysicsData.PointVelocityRange);
 
 				InitialVelocity += Direction * VelocityScale;
 			}
diff --git a/Engine/Plugins/FX/Niagara/Source/Niagara/Internal/Stateless/Modules/NiagaraStatelessModule_SpriteRotationRate.h b/Engine/Plugins/FX/Niagara/Source/Niagara/Internal/Stateless/Modules/NiagaraStatelessModule_SpriteRotationRate.h
--- a/Engine/Plugins/FX/Niagara/Source/Niagara/Internal/Stateless/Modules/NiagaraStatelessModule_SpriteRotationRate.h
+++ b/Engine/Plugins/FX/Niagara/Source/Niagara/Internal/Stateless/Modules/NiagaraStatelessModule_SpriteRotationRate.h
@@ -67,7 +67,7 @@ public:
 
 		for (uint32 i = 0; i < ParticleSimulationContext.GetNumInstances(); ++i)
 		{
-			const float RotationRate = ParticleSimulationContext.RandomScaleBiasFloat(i, ModuleBuiltData->RotationRange);
+			const float RotationRate = ParticleSimulationContext.RandomScaleBiasFloat(i, 0, ModuleBuiltData->RotationRange);
 
 			float SpriteRotation = ParticleSimulationContext.ReadParticleVariable(ModuleBuiltData->SpriteRotationVariableOffset, i, 0.0f);
 			float PreviousSpriteRotation = ParticleSimulationContext.ReadParticleVariable(ModuleBuiltData->PreviousSpriteRotationVariableOffset, i, 0.0f);
diff --git a/Engine/Plugins/FX/Niagara/Source/Niagara/Internal/Stateless/Modules/NiagaraStatelessModule_SubUVAnimation.h b/Engine/Plugins/FX/Niagara/Source/Niagara/Internal/Stateless/Modules/NiagaraStatelessModule_SubUVAnimation.h
--- a/Engine/Plugins/FX/Niagara/Source/Niagara/Internal/Stateless/Modules/NiagaraStatelessModule_SubUVAnimation.h
+++ b/Engine/Plugins/FX/Niagara/Source/Niagara/Internal/Stateless/Modules/NiagaraStatelessModule_SubUVAnimation.h
@@ -152,7 +152,7 @@ public:
 			const float Particle_NormalizedAge = ParticleNormalizedAges[i];
 
 			const uint32 SeedOffset = uint32(Particle_Age * ModuleBuiltData->InitialFrameRateChange);
-			float Frame = ParticleSimulationContext.RandomSeedOffsetFloat(i, SeedOffset) * ModuleBuiltData->InitialFrameScale + ModuleBuiltData->InitialFrameBias;
+			float Frame = ParticleSimulationContext.RandomFloat(i, SeedOffset) * ModuleBuiltData->InitialFrameScale + ModuleBuiltData->InitialFrameBias;
 
 			if (ModuleBuiltData->Mode == 0)
 			{
diff --git a/Engine/Plugins/FX/Niagara/Source/Niagara/Internal/Stateless/NiagaraStatelessCommon.h b/Engine/Plugins/FX/Niagara/Source/Niagara/Internal/Stateless/NiagaraStatelessCommon.h
--- a/Engine/Plugins/FX/Niagara/Source/Niagara/Internal/Stateless/NiagaraStatelessCommon.h
+++ b/Engine/Plugins/FX/Niagara/Source/Niagara/Internal/Stateless/NiagaraStatelessCommon.h
@@ -4,15 +4,17 @@
 
 #include "NiagaraCommon.h"
 
+UENUM(meta = (Bitflags, UseEnumValuesAsMaskValuesInEditor = "true"))
 enum class ENiagaraStatelessFeatureMask : uint8
 {
-	None			= 0,
-
-	ExecuteGPU		= 1 << 0,					// We can execute on the GPU (Might be broken down into GPUCompute | GPUGraphics | GPUAsyncCompute in future but this will remain the master mask)
-	ExecuteCPU		= 1 << 1,					// We can execute on the CPU
-	ExecuteAll		= ExecuteGPU | ExecuteCPU,	// We can execute on both the CPU & GPU
-
-	All				= ExecuteAll,				// We support all features
+	// We can execute on the GPU (Might be broken down into GPUCompute | GPUGraphics | GPUAsyncCompute in future but this will remain the master mask)
+	ExecuteGPU		= 1 << 0,
+	// We can execute on the CPU
+	ExecuteCPU		= 1 << 1,
+
+	None			= 0							UMETA(Hidden),
+	ExecuteAll		= ExecuteGPU | ExecuteCPU	UMETA(Hidden),
+	All				= ExecuteAll				UMETA(Hidden),
 };
 ENUM_CLASS_FLAGS(ENiagaraStatelessFeatureMask);
 
diff --git a/Engine/Plugins/FX/Niagara/Source/Niagara/Internal/Stateless/NiagaraStatelessEmitter.h b/Engine/Plugins/FX/Niagara/Source/Niagara/Internal/Stateless/NiagaraStatelessEmitter.h
--- a/Engine/Plugins/FX/Niagara/Source/Niagara/Internal/Stateless/NiagaraStatelessEmitter.h
+++ b/Engine/Plugins/FX/Niagara/Source/Niagara/Internal/Stateless/NiagaraStatelessEmitter.h
@@ -136,6 +136,9 @@ protected:
 	UPROPERTY(EditAnywhere, Category = "General")
 	bool bDeterministic = false;
 
+	UPROPERTY(EditAnywhere, Category = "General", AdvancedDisplay, meta = (Bitmask, BitMaskEnum = "/Script/Niagara.ENiagaraStatelessFeatureMask"))
+	uint32 AllowedFeatureMask = uint32(ENiagaraStatelessFeatureMask::All);
+
 	UPROPERTY(EditAnywhere, Category = "General")
 	int32 RandomSeed = 0;
 
diff --git a/Engine/Plugins/FX/Niagara/Source/Niagara/Internal/Stateless/NiagaraStatelessEmitterDataBuildContext.h b/Engine/Plugins/FX/Niagara/Source/Niagara/Internal/Stateless/NiagaraStatelessEmitterDataBuildContext.h
--- a/Engine/Plugins/FX/Niagara/Source/Niagara/Internal/Stateless/NiagaraStatelessEmitterDataBuildContext.h
+++ b/Engine/Plugins/FX/Niagara/Source/Niagara/Internal/Stateless/NiagaraStatelessEmitterDataBuildContext.h
@@ -153,6 +153,7 @@ private:
 	NiagaraStateless::FParticleSimulationExecData*	ParticleExecData = nullptr;
 
 	int32											ModuleBuiltDataOffset = 0;
+	int32											RandomSeedOffest = 1;
 
 	struct FTransientObject
 	{
diff --git a/Engine/Plugins/FX/Niagara/Source/Niagara/Internal/Stateless/NiagaraStatelessParticleSimContext.h b/Engine/Plugins/FX/Niagara/Source/Niagara/Internal/Stateless/NiagaraStatelessParticleSimContext.h
--- a/Engine/Plugins/FX/Niagara/Source/Niagara/Internal/Stateless/NiagaraStatelessParticleSimContext.h
+++ b/Engine/Plugins/FX/Niagara/Source/Niagara/Internal/Stateless/NiagaraStatelessParticleSimContext.h
@@ -21,7 +21,7 @@ enum class EParticleComponent
 	NormalizedAge,
 	PreviousAge,
 	PreviousNormalizedAge,
-	RandomSeed,
+	UniqueIndex,
 	Num
 };
 
@@ -50,13 +50,14 @@ public:
 	float* GetParticleNormalizedAge() { return static_cast<float*>(RequiredComponents[int32(EParticleComponent::NormalizedAge)]); }
 	float* GetParticlePreviousAge() { return static_cast<float*>(RequiredComponents[int32(EParticleComponent::PreviousAge)]); }
 	float* GetParticlePreviousNormalizedAge() { return static_cast<float*>(RequiredComponents[int32(EParticleComponent::PreviousNormalizedAge)]); }
-	//FUintVector4* GetParticleRandomSeed() { return static_cast<FUintVector4*>(RequiredComponents[int32(EParticleComponent::RandomSeed)]); }
+	int32* GetParticleUniqueIndex() { return static_cast<int32*>(RequiredComponents[int32(EParticleComponent::UniqueIndex)]); }
 
-	float* GetParticleLifetime() const { return static_cast<float*>(RequiredComponents[int32(EParticleComponent::Lifetime)]); }
-	float* GetParticleAge() const { return static_cast<float*>(RequiredComponents[int32(EParticleComponent::Age)]); }
-	float* GetParticleNormalizedAge() const { return static_cast<float*>(RequiredComponents[int32(EParticleComponent::NormalizedAge)]); }
-	float* GetParticlePreviousAge() const { return static_cast<float*>(RequiredComponents[int32(EParticleComponent::PreviousAge)]); }
-	float* GetParticlePreviousNormalizedAge() const { return static_cast<float*>(RequiredComponents[int32(EParticleComponent::PreviousNormalizedAge)]); }
+	const float* GetParticleLifetime() const { return static_cast<float*>(RequiredComponents[int32(EParticleComponent::Lifetime)]); }
+	const float* GetParticleAge() const { return static_cast<float*>(RequiredComponents[int32(EParticleComponent::Age)]); }
+	const float* GetParticleNormalizedAge() const { return static_cast<float*>(RequiredComponents[int32(EParticleComponent::NormalizedAge)]); }
+	const float* GetParticlePreviousAge() const { return static_cast<float*>(RequiredComponents[int32(EParticleComponent::PreviousAge)]); }
+	const float* GetParticlePreviousNormalizedAge() const { return static_cast<float*>(RequiredComponents[int32(EParticleComponent::PreviousNormalizedAge)]); }
+	const int32* GetParticleUniqueIndex() const { return static_cast<int32*>(RequiredComponents[int32(EParticleComponent::UniqueIndex)]); }
 
 	// Get a pointer to an optional particle component (i.e. one not required in the data set)
 	float* GetParticleVariableFloat(int32 iVariable) const { return iVariable >= 0 ? reinterpret_cast<float*>(VariableComponents[iVariable]) : nullptr; }
@@ -120,47 +121,42 @@ public:
 	FVector4f		Lerp(const FVector4f& Lhs, const FVector4f& Rhs, const FVector4f& U, bool bUniform) const { const FVector4f V = FMath::Lerp(Lhs, Rhs, U); return bUniform ? FVector4f(V.X, V.X, V.X, V.X) : V; }
 
 	template<typename TType>
-	TType					TRandomFloat(uint32 iInstance) const;
-	template<> float		TRandomFloat<float>(uint32 iInstance) const { return RandomFloat(iInstance); }
-	template<> FVector2f	TRandomFloat<FVector2f>(uint32 iInstance) const { return RandomFloat2(iInstance); }
-	template<> FVector3f	TRandomFloat<FVector3f>(uint32 iInstance) const { return RandomFloat3(iInstance); }
-	template<> FVector4f	TRandomFloat<FVector4f>(uint32 iInstance) const { return RandomFloat4(iInstance); }
-	template<> FLinearColor	TRandomFloat<FLinearColor>(uint32 iInstance) const { return FLinearColor(RandomFloat4(iInstance)); }
-
-	uint32			RandomUInt(uint32 iInstance) const;
-	FUintVector2	RandomUInt2(uint32 iInstance) const;
-	FUintVector3	RandomUInt3(uint32 iInstance) const;
-	FUintVector4	RandomUInt4(uint32 iInstance) const;
-
-	float			RandomFloat(uint32 iInstance) const;
-	FVector2f		RandomFloat2(uint32 iInstance) const;
-	FVector3f		RandomFloat3(uint32 iInstance) const;
-	FVector4f		RandomFloat4(uint32 iInstance) const;
-
-	float			RandomScaleBiasFloat(uint32 iInstance, const float Scale, const float Bias) const { return Bias + (RandomFloat(iInstance) * Scale); }
-	FVector2f		RandomScaleBiasFloat(uint32 iInstance, const FVector2f& Scale, const FVector2f& Bias) const { return Bias + (RandomFloat(iInstance) * Scale); }
-	FVector3f		RandomScaleBiasFloat(uint32 iInstance, const FVector3f& Scale, const FVector3f& Bias) const { return Bias + (RandomFloat(iInstance) * Scale); }
-	FVector4f		RandomScaleBiasFloat(uint32 iInstance, const FVector4f& Scale, const FVector4f& Bias) const { return Bias + (RandomFloat(iInstance) * Scale); }
-
-	float			RandomScaleBiasFloat(uint32 iInstance, const FNiagaraStatelessRangeFloat& Range) const { return Range.Min + (RandomFloat(iInstance) * Range.GetScale()); }
-	FVector2f		RandomScaleBiasFloat(uint32 iInstance, const FNiagaraStatelessRangeVector2& Range) const { return Range.Min + (RandomFloat2(iInstance) * Range.GetScale()); }
-	FVector3f		RandomScaleBiasFloat(uint32 iInstance, const FNiagaraStatelessRangeVector3& Range) const { return Range.Min + (RandomFloat3(iInstance) * Range.GetScale()); }
-	FVector4f		RandomScaleBiasFloat(uint32 iInstance, const FNiagaraStatelessRangeVector4& Range) const { return Range.Min + (RandomFloat4(iInstance) * Range.GetScale()); }
-	FLinearColor	RandomScaleBiasFloat(uint32 iInstance, const FNiagaraStatelessRangeColor& Range) const { return Range.Min + (RandomFloat4(iInstance) * Range.GetScale()); }
-
-	float			RandomScaleBiasFloat(uint32 iInstance, const FNiagaraStatelessRangeFloat& Range, bool bUniform) const { return RandomScaleBiasFloat(iInstance, Range); }
-	FVector2f		RandomScaleBiasFloat(uint32 iInstance, const FNiagaraStatelessRangeVector2& Range, bool bUniform) const { const FVector2f Random = RandomScaleBiasFloat(iInstance, Range); return bUniform ? FVector2f(Random.X, Random.X) : Random; }
-	FVector3f		RandomScaleBiasFloat(uint32 iInstance, const FNiagaraStatelessRangeVector3& Range, bool bUniform) const { const FVector3f Random = RandomScaleBiasFloat(iInstance, Range); return bUniform ? FVector3f(Random.X, Random.X, Random.X) : Random; }
-	FVector4f		RandomScaleBiasFloat(uint32 iInstance, const FNiagaraStatelessRangeVector4& Range, bool bUniform) const { const FVector4f Random = RandomScaleBiasFloat(iInstance, Range); return bUniform ? FVector4f(Random.X, Random.X, Random.X, Random.X) : Random; }
-	FLinearColor	RandomScaleBiasFloat(uint32 iInstance, const FNiagaraStatelessRangeColor& Range, bool bUniform) const { const FLinearColor Random = RandomScaleBiasFloat(iInstance, Range); return bUniform ? FLinearColor(Random.R, Random.R, Random.R, Random.R) : Random; }
-
-	FVector2f		RandomUnitFloat2(uint32 iInstance) const { return SafeNormalize(RandomFloat2(iInstance) - 0.5f); }
-	FVector3f		RandomUnitFloat3(uint32 iInstance) const { return SafeNormalize(RandomFloat3(iInstance) - 0.5f); }
-
-	float			RandomSeedOffsetFloat(uint32 iInstance, uint32 SeedOffset) const;
-	FVector2f		RandomSeedOffsetFloat2(uint32 iInstance, uint32 SeedOffset) const;
-	FVector3f		RandomSeedOffsetFloat3(uint32 iInstance, uint32 SeedOffset) const;
-	FVector4f		RandomSeedOffsetFloat4(uint32 iInstance, uint32 SeedOffset) const;
+	TType					TRandomFloat(uint32 iInstance, uint32 RandomSeedOffset) const;
+	template<> float		TRandomFloat<float>(uint32 iInstance, uint32 RandomSeedOffset) const { return RandomFloat(iInstance, RandomSeedOffset); }
+	template<> FVector2f	TRandomFloat<FVector2f>(uint32 iInstance, uint32 RandomSeedOffset) const { return RandomFloat2(iInstance, RandomSeedOffset); }
+	template<> FVector3f	TRandomFloat<FVector3f>(uint32 iInstance, uint32 RandomSeedOffset) const { return RandomFloat3(iInstance, RandomSeedOffset); }
+	template<> FVector4f	TRandomFloat<FVector4f>(uint32 iInstance, uint32 RandomSeedOffset) const { return RandomFloat4(iInstance, RandomSeedOffset); }
+	template<> FLinearColor	TRandomFloat<FLinearColor>(uint32 iInstance, uint32 RandomSeedOffset) const { return FLinearColor(RandomFloat4(iInstance, RandomSeedOffset)); }
+
+	uint32			RandomUInt(uint32 iInstance, uint32 RandomSeedOffset) const;
+	FUintVector2	RandomUInt2(uint32 iInstance, uint32 RandomSeedOffset) const;
+	FUintVector3	RandomUInt3(uint32 iInstance, uint32 RandomSeedOffset) const;
+	FUintVector4	RandomUInt4(uint32 iInstance, uint32 RandomSeedOffset) const;
+
+	float			RandomFloat(uint32 iInstance, uint32 RandomSeedOffset) const;
+	FVector2f		RandomFloat2(uint32 iInstance, uint32 RandomSeedOffset) const;
+	FVector3f		RandomFloat3(uint32 iInstance, uint32 RandomSeedOffset) const;
+	FVector4f		RandomFloat4(uint32 iInstance, uint32 RandomSeedOffset) const;
+
+	float			RandomScaleBiasFloat(uint32 iInstance, uint32 RandomSeedOffset, const float Scale, const float Bias) const { return Bias + (RandomFloat(iInstance, RandomSeedOffset) * Scale); }
+	FVector2f		RandomScaleBiasFloat(uint32 iInstance, uint32 RandomSeedOffset, const FVector2f& Scale, const FVector2f& Bias) const { return Bias + (RandomFloat(iInstance, RandomSeedOffset) * Scale); }
+	FVector3f		RandomScaleBiasFloat(uint32 iInstance, uint32 RandomSeedOffset, const FVector3f& Scale, const FVector3f& Bias) const { return Bias + (RandomFloat(iInstance, RandomSeedOffset) * Scale); }
+	FVector4f		RandomScaleBiasFloat(uint32 iInstance, uint32 RandomSeedOffset, const FVector4f& Scale, const FVector4f& Bias) const { return Bias + (RandomFloat(iInstance, RandomSeedOffset) * Scale); }
+
+	float			RandomScaleBiasFloat(uint32 iInstance, uint32 RandomSeedOffset, const FNiagaraStatelessRangeFloat& Range) const { return Range.Min + (RandomFloat(iInstance, RandomSeedOffset) * Range.GetScale()); }
+	FVector2f		RandomScaleBiasFloat(uint32 iInstance, uint32 RandomSeedOffset, const FNiagaraStatelessRangeVector2& Range) const { return Range.Min + (RandomFloat2(iInstance, RandomSeedOffset) * Range.GetScale()); }
+	FVector3f		RandomScaleBiasFloat(uint32 iInstance, uint32 RandomSeedOffset, const FNiagaraStatelessRangeVector3& Range) const { return Range.Min + (RandomFloat3(iInstance, RandomSeedOffset) * Range.GetScale()); }
+	FVector4f		RandomScaleBiasFloat(uint32 iInstance, uint32 RandomSeedOffset, const FNiagaraStatelessRangeVector4& Range) const { return Range.Min + (RandomFloat4(iInstance, RandomSeedOffset) * Range.GetScale()); }
+	FLinearColor	RandomScaleBiasFloat(uint32 iInstance, uint32 RandomSeedOffset, const FNiagaraStatelessRangeColor& Range) const { return Range.Min + (RandomFloat4(iInstance, RandomSeedOffset) * Range.GetScale()); }
+
+	float			RandomScaleBiasFloat(uint32 iInstance, uint32 RandomSeedOffset, const FNiagaraStatelessRangeFloat& Range, bool bUniform) const { return RandomScaleBiasFloat(iInstance, RandomSeedOffset, Range); }
+	FVector2f		RandomScaleBiasFloat(uint32 iInstance, uint32 RandomSeedOffset, const FNiagaraStatelessRangeVector2& Range, bool bUniform) const { const FVector2f Random = RandomScaleBiasFloat(iInstance, RandomSeedOffset, Range); return bUniform ? FVector2f(Random.X, Random.X) : Random; }
+	FVector3f		RandomScaleBiasFloat(uint32 iInstance, uint32 RandomSeedOffset, const FNiagaraStatelessRangeVector3& Range, bool bUniform) const { const FVector3f Random = RandomScaleBiasFloat(iInstance, RandomSeedOffset, Range); return bUniform ? FVector3f(Random.X, Random.X, Random.X) : Random; }
+	FVector4f		RandomScaleBiasFloat(uint32 iInstance, uint32 RandomSeedOffset, const FNiagaraStatelessRangeVector4& Range, bool bUniform) const { const FVector4f Random = RandomScaleBiasFloat(iInstance, RandomSeedOffset, Range); return bUniform ? FVector4f(Random.X, Random.X, Random.X, Random.X) : Random; }
+	FLinearColor	RandomScaleBiasFloat(uint32 iInstance, uint32 RandomSeedOffset, const FNiagaraStatelessRangeColor& Range, bool bUniform) const { const FLinearColor Random = RandomScaleBiasFloat(iInstance, RandomSeedOffset, Range); return bUniform ? FLinearColor(Random.R, Random.R, Random.R, Random.R) : Random; }
+
+	FVector2f		RandomUnitFloat2(uint32 iInstance, uint32 RandomSeedOffset) const { return SafeNormalize(RandomFloat2(iInstance, RandomSeedOffset) - 0.5f); }
+	FVector3f		RandomUnitFloat3(uint32 iInstance, uint32 RandomSeedOffset) const { return SafeNormalize(RandomFloat3(iInstance, RandomSeedOffset) - 0.5f); }
 
 	FVector2f		SafeNormalize(const FVector2f& v, const FVector2f& Fallback) const { const float l2 = v.SquaredLength(); return l2 < UE_KINDA_SMALL_NUMBER ? Fallback : v * (1.0f / FMath::Sqrt(l2)); }
 	FVector3f		SafeNormalize(const FVector3f& v, const FVector3f& Fallback) const { const float l2 = v.SquaredLength(); return l2 < UE_KINDA_SMALL_NUMBER ? Fallback : v * (1.0f / FMath::Sqrt(l2)); }
@@ -196,15 +192,15 @@ public:
 		);
 	}
 
-private:
-	FUintVector4* GetParticleRandomSeed() const { return static_cast<FUintVector4*>(RequiredComponents[int32(EParticleComponent::RandomSeed)]); }
-
 private:
 	const FNiagaraStatelessEmitterData*		EmitterData = nullptr;
 	uint32									NumInstances = 0;
 	float									DeltaTime = 0.0f;
 	float									InvDeltaTime = 0.0f;
 
+	uint32									EmitterRandomSeed = 0;
+	uint32									ModuleRandomSeed = 0;
+
 	uint32									BufferStride = 0;
 	uint8*									BufferFloatData = nullptr;
 	uint8*									BufferInt32Data = nullptr;
@@ -221,12 +217,12 @@ private:
 template<typename TType>
 struct FStatelessDistributionSampler
 {
-	explicit FStatelessDistributionSampler(const FParticleSimulationContext& ParticleSimulationContext, const FUintVector3& InParameters, int32 iInstance)
+	explicit FStatelessDistributionSampler(const FParticleSimulationContext& ParticleSimulationContext, const FUintVector3& InParameters, int32 iInstance, uint32 RandomSeedOffset)
 		: Parameters(InParameters)
 	{
 		if ((Parameters.X & uint32(ENiagaraStatelessBuiltDistributionFlag::Random)) != 0)
 		{
-			RandomOffset = ParticleSimulationContext.TRandomFloat<TType>(iInstance);
+			RandomOffset = ParticleSimulationContext.TRandomFloat<TType>(iInstance, RandomSeedOffset);
 		}
 	}
 
diff --git a/Engine/Plugins/FX/Niagara/Source/Niagara/Internal/Stateless/NiagaraStatelessParticleSimExecData.h b/Engine/Plugins/FX/Niagara/Source/Niagara/Internal/Stateless/NiagaraStatelessParticleSimExecData.h
--- a/Engine/Plugins/FX/Niagara/Source/Niagara/Internal/Stateless/NiagaraStatelessParticleSimExecData.h
+++ b/Engine/Plugins/FX/Niagara/Source/Niagara/Internal/Stateless/NiagaraStatelessParticleSimExecData.h
@@ -16,14 +16,16 @@ public:
 	struct FCallback
 	{
 		FCallback() = default;
-		explicit FCallback(TFunction<void(const FParticleSimulationContext&)> InFunc, int32 InBuiltDataOffset)
+		explicit FCallback(TFunction<void(const FParticleSimulationContext&)> InFunc, int32 InBuiltDataOffset, uint32 InRandomSeedOffset)
 			: Function(MoveTemp(InFunc))
 			, BuiltDataOffset(InBuiltDataOffset)
+			, RandomSeedOffset(InRandomSeedOffset)
 		{
 		}
 
 		TFunction<void(const FParticleSimulationContext&)>	Function;
-		int32												BuiltDataOffset = 0;
+		int32	BuiltDataOffset = 0;
+		uint32	RandomSeedOffset = 0;
 	};
 
 	struct FVariableOffset
diff --git a/Engine/Plugins/FX/Niagara/Source/Niagara/Private/Stateless/Modules/NiagaraStatelessModule_MeshIndex.cpp b/Engine/Plugins/FX/Niagara/Source/Niagara/Private/Stateless/Modules/NiagaraStatelessModule_MeshIndex.cpp
--- a/Engine/Plugins/FX/Niagara/Source/Niagara/Private/Stateless/Modules/NiagaraStatelessModule_MeshIndex.cpp
+++ b/Engine/Plugins/FX/Niagara/Source/Niagara/Private/Stateless/Modules/NiagaraStatelessModule_MeshIndex.cpp
@@ -63,7 +63,7 @@ namespace NSMMeshIndexPrivate
 		{
 			for (uint32 i = 0; i < ParticleSimulationContext.GetNumInstances(); ++i)
 			{
-				const FVector2f Rand = ParticleSimulationContext.RandomFloat2(i);
+				const FVector2f Rand = ParticleSimulationContext.RandomFloat2(i, 0);
 				MeshIndex = FMath::RoundToInt(Rand.X * float(ModuleBuiltData->TableNumElements));
 				const float Probability = ParticleSimulationContext.GetStaticFloat<float>(ModuleBuiltData->TableOffset, (MeshIndex * 2) + 0);
 				if (Rand.Y > Probability)
diff --git a/Engine/Plugins/FX/Niagara/Source/Niagara/Private/Stateless/Modules/NiagaraStatelessModule_ShapeLocation.cpp b/Engine/Plugins/FX/Niagara/Source/Niagara/Private/Stateless/Modules/NiagaraStatelessModule_ShapeLocation.cpp
--- a/Engine/Plugins/FX/Niagara/Source/Niagara/Private/Stateless/Modules/NiagaraStatelessModule_ShapeLocation.cpp
+++ b/Engine/Plugins/FX/Niagara/Source/Niagara/Private/Stateless/Modules/NiagaraStatelessModule_ShapeLocation.cpp
@@ -26,15 +26,15 @@ namespace NSMShapeLocationPrivate
 			const FVector3f BoxScale = FVector3f(ModuleBuiltData->Parameters0);
 			const FVector3f BoxBias = FVector3f(ModuleBuiltData->Parameters1);
 			const bool bOnSurface = ModuleBuiltData->Mode.Y == 1;
-			const float SurfaceScale = ModuleBuiltData->Parameters0.X;
+			const float SurfaceScale = ModuleBuiltData->Parameters0.W;
 			const float SurfaceBias = ModuleBuiltData->Parameters1.W;
 	
-			const FVector3f P0 = ParticleSimulationContext.RandomFloat3(iInstance);
+			const FVector3f P0 = ParticleSimulationContext.RandomFloat3(iInstance, 0);
 			if (bOnSurface)
 			{
 				//-TODO: This isn't quite what we want
 				const FVector3f L0(FMath::RoundToFloat(P0.X), FMath::RoundToFloat(P0.Y), FMath::RoundToFloat(P0.Z));
-				const uint32 S = ParticleSimulationContext.RandomUInt(iInstance) % 3;
+				const uint32 S = ParticleSimulationContext.RandomUInt(iInstance, 1) % 3;
 				FVector3f Location;
 				Location.X = S != 0 ? P0.X : L0.X;
 				Location.Y = S != 1 ? P0.Y : L0.Y;
@@ -50,7 +50,7 @@ namespace NSMShapeLocationPrivate
 		// ENSM_ShapePrimitive::Cylinder:
 		if ( ModuleBuiltData->Mode.X == 1 )
 		{
-			const FVector4f Random = ParticleSimulationContext.RandomFloat4(iInstance);
+			const FVector4f Random = ParticleSimulationContext.RandomFloat4(iInstance, 0);
 			const float HeightScale = ModuleBuiltData->Parameters0.X;
 			const float HeightBias = ModuleBuiltData->Parameters0.Y;
 			const float Radius = ModuleBuiltData->Parameters0.Z;
@@ -71,8 +71,8 @@ namespace NSMShapeLocationPrivate
 			const float UDistributionScale = ModuleBuiltData->Parameters0.Z;
 			const float UDistributionBias = ModuleBuiltData->Parameters0.W;
 	
-			const float Radius = ParticleSimulationContext.RandomScaleBiasFloat(iInstance, RadiusScale, RadiusBias);
-			const float U = ParticleSimulationContext.RandomScaleBiasFloat(iInstance, UDistributionScale, UDistributionBias);
+			const float Radius = ParticleSimulationContext.RandomScaleBiasFloat(iInstance, 0, RadiusScale, RadiusBias);
+			const float U = ParticleSimulationContext.RandomScaleBiasFloat(iInstance, 1, UDistributionScale, UDistributionBias);
 	
 			return FVector3f(cos(U) * Radius, sin(U) * Radius, 0.0f);
 		}
@@ -82,8 +82,8 @@ namespace NSMShapeLocationPrivate
 			const float SphereScale	= ModuleBuiltData->Parameters0.X;
 			const float SphereBias	= ModuleBuiltData->Parameters0.Y;
 	
-			const FVector3f Vector = ParticleSimulationContext.RandomUnitFloat3(iInstance);
-			return Vector * ParticleSimulationContext.RandomScaleBiasFloat(iInstance, SphereScale, SphereBias);
+			const FVector3f Vector = ParticleSimulationContext.RandomUnitFloat3(iInstance, 0);
+			return Vector * ParticleSimulationContext.RandomScaleBiasFloat(iInstance, 1, SphereScale, SphereBias);
 		}
 	}
 
diff --git a/Engine/Plugins/FX/Niagara/Source/Niagara/Private/Stateless/NiagaraStatelessComputeManager.cpp b/Engine/Plugins/FX/Niagara/Source/Niagara/Private/Stateless/NiagaraStatelessComputeManager.cpp
--- a/Engine/Plugins/FX/Niagara/Source/Niagara/Private/Stateless/NiagaraStatelessComputeManager.cpp
+++ b/Engine/Plugins/FX/Niagara/Source/Niagara/Private/Stateless/NiagaraStatelessComputeManager.cpp
@@ -163,6 +163,8 @@ void FNiagaraStatelessComputeManager::OnPostPreRender(FRDGBuilder& GraphBuilder)
 		RDG_EVENT_NAME("FNiagaraStatelessComputeManager::OnPostPreRender"),
 		[DataToGenerate=MoveTemp(GPUDataToGenerate), ComputeInterface=GetOwnerInterface()](FRHICommandListImmediate& RHICmdList)
 		{
+			SCOPED_DRAW_EVENT(RHICmdList, FNiagaraStatelessComputeManager_OnPostPreRender);
+	
 			const int32 NumJobs = DataToGenerate.Num();
 
 			// Get Count Buffer
diff --git a/Engine/Plugins/FX/Niagara/Source/Niagara/Private/Stateless/NiagaraStatelessEmitter.cpp b/Engine/Plugins/FX/Niagara/Source/Niagara/Private/Stateless/NiagaraStatelessEmitter.cpp
--- a/Engine/Plugins/FX/Niagara/Source/Niagara/Private/Stateless/NiagaraStatelessEmitter.cpp
+++ b/Engine/Plugins/FX/Niagara/Source/Niagara/Private/Stateless/NiagaraStatelessEmitter.cpp
@@ -317,6 +317,8 @@ void UNiagaraStatelessEmitter::CacheFromCompiledData()
 	if (StatelessEmitterData->bCanEverExecute)
 	{
 		StatelessEmitterData->FeatureMask = FNiagaraStatelessGlobals::Get().FeatureMask;
+		StatelessEmitterData->FeatureMask &= ENiagaraStatelessFeatureMask(AllowedFeatureMask);
+
 		for (const UNiagaraStatelessModule* Module : Modules)
 		{
 			if (Module->IsModuleEnabled())
diff --git a/Engine/Plugins/FX/Niagara/Source/Niagara/Private/Stateless/NiagaraStatelessEmitterDataBuildContext.cpp b/Engine/Plugins/FX/Niagara/Source/Niagara/Private/Stateless/NiagaraStatelessEmitterDataBuildContext.cpp
--- a/Engine/Plugins/FX/Niagara/Source/Niagara/Private/Stateless/NiagaraStatelessEmitterDataBuildContext.cpp
+++ b/Engine/Plugins/FX/Niagara/Source/Niagara/Private/Stateless/NiagaraStatelessEmitterDataBuildContext.cpp
@@ -11,6 +11,7 @@
 void FNiagaraStatelessEmitterDataBuildContext::PreModuleBuild()
 {
 	ModuleBuiltDataOffset = BuiltData.Num();
+	++RandomSeedOffest;
 }
 
 uint32 FNiagaraStatelessEmitterDataBuildContext::AddStaticData(TConstArrayView<float> FloatData) const
@@ -104,7 +105,7 @@ void FNiagaraStatelessEmitterDataBuildContext::AddParticleSimulationExecSimulate
 		return;
 	}
 
-	ParticleExecData->SimulateFunctions.Emplace(MoveTemp(Func), ModuleBuiltDataOffset);
+	ParticleExecData->SimulateFunctions.Emplace(MoveTemp(Func), ModuleBuiltDataOffset, RandomSeedOffest);
 }
 
 int32 FNiagaraStatelessEmitterDataBuildContext::FindParticleVariableIndex(const FNiagaraVariableBase& Variable) const
diff --git a/Engine/Plugins/FX/Niagara/Source/Niagara/Private/Stateless/NiagaraStatelessParticleSimContext.cpp b/Engine/Plugins/FX/Niagara/Source/Niagara/Private/Stateless/NiagaraStatelessParticleSimContext.cpp
--- a/Engine/Plugins/FX/Niagara/Source/Niagara/Private/Stateless/NiagaraStatelessParticleSimContext.cpp
+++ b/Engine/Plugins/FX/Niagara/Source/Niagara/Private/Stateless/NiagaraStatelessParticleSimContext.cpp
@@ -16,9 +16,9 @@ namespace NiagaraStateless
 		3,	// float		- NormalizedAge
 		4,	// float		- PreviousAge
 		5,	// float		- PreviousNormalizedAge
-		6,	// FUintVector4	- RandomSeed
+		6,	// int32		- UniqueIndex
 	};
-	constexpr uint32 ParticleComponentTotalSize = 10;
+	constexpr uint32 ParticleComponentTotalSize = 7;
 
 	FUintVector4 Rand4DPCG32(FUintVector4 v)
 	{
@@ -58,12 +58,12 @@ namespace NiagaraStateless
 		check(EmitterData->ParticleSimExecData);
 	}
 
-	void FParticleSimulationContext::Simulate(int32 EmitterRandomSeed, float EmitterAge, float InDeltaTime, TConstArrayView<FNiagaraStatelessRuntimeSpawnInfo> SpawnInfos, FNiagaraDataBuffer* DestinationData)
+	void FParticleSimulationContext::Simulate(int32 InEmitterRandomSeed, float EmitterAge, float InDeltaTime, TConstArrayView<FNiagaraStatelessRuntimeSpawnInfo> SpawnInfos, FNiagaraDataBuffer* DestinationData)
 	{
 		NumInstances = 0;
 
 		FSpawnInfoShaderParameters SpawnParameters;
-		const uint32 ActiveParticles = EmitterData->CalculateActiveParticles(EmitterRandomSeed, SpawnInfos, EmitterAge, &SpawnParameters);
+		const uint32 ActiveParticles = EmitterData->CalculateActiveParticles(InEmitterRandomSeed, SpawnInfos, EmitterAge, &SpawnParameters);
 		if (ActiveParticles > 0)
 		{
 			// Setup data buffer pointers
@@ -73,19 +73,19 @@ namespace NiagaraStateless
 			BufferInt32Data	= DestinationData->GetComponentPtrInt32(0);
 
 			// Run Simulation
-			SimulateInternal(EmitterRandomSeed, EmitterAge, InDeltaTime, SpawnParameters, ActiveParticles);
+			SimulateInternal(InEmitterRandomSeed, EmitterAge, InDeltaTime, SpawnParameters, ActiveParticles);
 		}
 
 		// Set instance count
 		DestinationData->SetNumInstances(NumInstances);
 	}
 
-	void FParticleSimulationContext::SimulateGPU(FRHICommandList& RHICmdList, int32 EmitterRandomSeed, float EmitterAge, float InDeltaTime, TConstArrayView<FNiagaraStatelessRuntimeSpawnInfo> SpawnInfos, FNiagaraDataBuffer* DestinationData)
+	void FParticleSimulationContext::SimulateGPU(FRHICommandList& RHICmdList, int32 InEmitterRandomSeed, float EmitterAge, float InDeltaTime, TConstArrayView<FNiagaraStatelessRuntimeSpawnInfo> SpawnInfos, FNiagaraDataBuffer* DestinationData)
 	{
 		NumInstances = 0;
 
 		FSpawnInfoShaderParameters SpawnParameters;
-		const uint32 ActiveParticles = EmitterData->CalculateActiveParticles(EmitterRandomSeed, SpawnInfos, EmitterAge, &SpawnParameters);
+		const uint32 ActiveParticles = EmitterData->CalculateActiveParticles(InEmitterRandomSeed, SpawnInfos, EmitterAge, &SpawnParameters);
 		if (ActiveParticles > 0)
 		{
 			check(DestinationData->GetNumInstancesAllocated() <= ActiveParticles);
@@ -99,7 +99,7 @@ namespace NiagaraStateless
 			BufferInt32Data	= Int32Buffer.NumBytes > 0 ? reinterpret_cast<uint8*>(RHICmdList.LockBuffer(Int32Buffer.Buffer, 0, Int32Buffer.NumBytes, RLM_WriteOnly)) : nullptr;
 
 			// Run Simulation
-			SimulateInternal(EmitterRandomSeed, EmitterAge, InDeltaTime, SpawnParameters, ActiveParticles);
+			SimulateInternal(InEmitterRandomSeed, EmitterAge, InDeltaTime, SpawnParameters, ActiveParticles);
 
 			// Unlock buffers
 			if (BufferFloatData)
@@ -116,12 +116,14 @@ namespace NiagaraStateless
 		DestinationData->SetNumInstances(NumInstances);
 	}
 
-	void FParticleSimulationContext::SimulateInternal(int32 EmitterRandomSeed, float EmitterAge, float InDeltaTime, FSpawnInfoShaderParameters& SpawnParameters, uint32 ActiveParticles)
+	void FParticleSimulationContext::SimulateInternal(int32 InEmitterRandomSeed, float EmitterAge, float InDeltaTime, FSpawnInfoShaderParameters& SpawnParameters, uint32 ActiveParticles)
 	{
 		// Setup simulation	
 		NumInstances = 0;
 		DeltaTime = InDeltaTime;
 		InvDeltaTime = DeltaTime > 0.0f ? 1.0f / DeltaTime : 0.0f;
+		EmitterRandomSeed = InEmitterRandomSeed;
+		ModuleRandomSeed = 6405u;
 
 		// Setup Required Components
 		{
@@ -185,16 +187,10 @@ namespace NiagaraStateless
 				}
 			}
 
-			FUintVector4 RandomSeed;
-			RandomSeed.X = 7123u;
-			RandomSeed.Y = EmitterRandomSeed;
-			RandomSeed.Z = Particle_UniqueIndex * 3581u;
-			RandomSeed.W = 0;
+			// Write unique index as it's needed for the random operation
+			GetParticleUniqueIndex()[NumInstances] = Particle_UniqueIndex;
 
-			// Write random seed here so we can use our random functions
-			GetParticleRandomSeed()[NumInstances] = RandomSeed;
-
-			float Particle_Lifetime = RandomScaleBiasFloat(NumInstances, EmitterData->LifetimeRange);
+			float Particle_Lifetime = RandomScaleBiasFloat(NumInstances, 0, EmitterData->LifetimeRange);
 			if (Particle_Lifetime <= 0.0f || Particle_Age >= Particle_Lifetime)
 			{
 				return;
@@ -222,68 +218,61 @@ namespace NiagaraStateless
 		for (const auto& Callback : ExecData->SimulateFunctions)
 		{
 			BuiltDataOffset = Callback.BuiltDataOffset;
+			ModuleRandomSeed = Callback.RandomSeedOffset * 6405u;
 			Callback.Function(*this);
 		}
 	}
 
-	uint32 FParticleSimulationContext::RandomUInt(uint32 iInstance) const
+	uint32 FParticleSimulationContext::RandomUInt(uint32 iInstance, uint32 RandomSeedOffset) const
 	{
-		FUintVector4& RandomSeed = GetParticleRandomSeed()[iInstance];
-		++RandomSeed.X;
-		RandomSeed.W = RandomSeed.X ^ RandomSeed.Y;
-
-		FUintVector4 RandomValue = Rand4DPCG32(RandomSeed);
+		const uint32 UniqueIndex = uint32(GetParticleUniqueIndex()[iInstance]);
+		const FUintVector4 RandomSeed(7123u + RandomSeedOffset, EmitterRandomSeed, UniqueIndex * 3581u, ModuleRandomSeed);
+		const FUintVector4 RandomValue = Rand4DPCG32(RandomSeed);
 		return RandomValue.X;
 	}
 
-	FUintVector2 FParticleSimulationContext::RandomUInt2(uint32 iInstance) const
+	FUintVector2 FParticleSimulationContext::RandomUInt2(uint32 iInstance, uint32 RandomSeedOffset) const
 	{
-		FUintVector4& RandomSeed = GetParticleRandomSeed()[iInstance];
-		++RandomSeed.X;
-		RandomSeed.W = RandomSeed.X ^ RandomSeed.Y;
-
-		FUintVector4 RandomValue = Rand4DPCG32(RandomSeed);
+		const uint32 UniqueIndex = uint32(GetParticleUniqueIndex()[iInstance]);
+		const FUintVector4 RandomSeed(7123u + RandomSeedOffset, EmitterRandomSeed, UniqueIndex * 3581u, ModuleRandomSeed);
+		const FUintVector4 RandomValue = Rand4DPCG32(RandomSeed);
 		return FUintVector2(RandomValue.X, RandomValue.Y);
 	}
 
-	FUintVector3 FParticleSimulationContext::RandomUInt3(uint32 iInstance) const
+	FUintVector3 FParticleSimulationContext::RandomUInt3(uint32 iInstance, uint32 RandomSeedOffset) const
 	{
-		FUintVector4& RandomSeed = GetParticleRandomSeed()[iInstance];
-		++RandomSeed.X;
-		RandomSeed.W = RandomSeed.X ^ RandomSeed.Y;
-
-		FUintVector4 RandomValue = Rand4DPCG32(RandomSeed);
+		const uint32 UniqueIndex = uint32(GetParticleUniqueIndex()[iInstance]);
+		const FUintVector4 RandomSeed(7123u + RandomSeedOffset, EmitterRandomSeed, UniqueIndex * 3581u, ModuleRandomSeed);
+		const FUintVector4 RandomValue = Rand4DPCG32(RandomSeed);
 		return FUintVector3(RandomValue.X, RandomValue.Y, RandomValue.Z);
 	}
 
-	FUintVector4 FParticleSimulationContext::RandomUInt4(uint32 iInstance) const
+	FUintVector4 FParticleSimulationContext::RandomUInt4(uint32 iInstance, uint32 RandomSeedOffset) const
 	{
-		FUintVector4& RandomSeed = GetParticleRandomSeed()[iInstance];
-		++RandomSeed.X;
-		RandomSeed.W = RandomSeed.X ^ RandomSeed.Y;
-
-		FUintVector4 RandomValue = Rand4DPCG32(RandomSeed);
+		const uint32 UniqueIndex = uint32(GetParticleUniqueIndex()[iInstance]);
+		const FUintVector4 RandomSeed(7123u + RandomSeedOffset, EmitterRandomSeed, UniqueIndex * 3581u, ModuleRandomSeed);
+		const FUintVector4 RandomValue = Rand4DPCG32(RandomSeed);
 		return RandomValue;
 	}
 
-	float FParticleSimulationContext::RandomFloat(uint32 iInstance) const
+	float FParticleSimulationContext::RandomFloat(uint32 iInstance, uint32 RandomSeedOffset) const
 	{
-		const uint32 V = RandomUInt(iInstance);
+		const uint32 V = RandomUInt(iInstance, RandomSeedOffset);
 		return float((V >> 8) & 0x00ffffff) / 16777216.0f;
 	}
 
-	FVector2f FParticleSimulationContext::RandomFloat2(uint32 iInstance) const
+	FVector2f FParticleSimulationContext::RandomFloat2(uint32 iInstance, uint32 RandomSeedOffset) const
 	{
-		const FUintVector2 V = RandomUInt2(iInstance);
+		const FUintVector2 V = RandomUInt2(iInstance, RandomSeedOffset);
 		return FVector2f(
 			float((V.X >> 8) & 0x00ffffff) / 16777216.0f,
 			float((V.Y >> 8) & 0x00ffffff) / 16777216.0f
 		);
 	}
 
-	FVector3f FParticleSimulationContext::RandomFloat3(uint32 iInstance) const
+	FVector3f FParticleSimulationContext::RandomFloat3(uint32 iInstance, uint32 RandomSeedOffset) const
 	{
-		const FUintVector3 V = RandomUInt3(iInstance);
+		const FUintVector3 V = RandomUInt3(iInstance, RandomSeedOffset);
 		return FVector3f(
 			float((V.X >> 8) & 0x00ffffff) / 16777216.0f,
 			float((V.Y >> 8) & 0x00ffffff) / 16777216.0f,
@@ -291,9 +280,9 @@ namespace NiagaraStateless
 		);
 	}
 
-	FVector4f FParticleSimulationContext::RandomFloat4(uint32 iInstance) const
+	FVector4f FParticleSimulationContext::RandomFloat4(uint32 iInstance, uint32 RandomSeedOffset) const
 	{
-		const FUintVector4 V = RandomUInt4(iInstance);
+		const FUintVector4 V = RandomUInt4(iInstance, RandomSeedOffset);
 		return FVector4f(
 			float((V.X >> 8) & 0x00ffffff) / 16777216.0f,
 			float((V.Y >> 8) & 0x00ffffff) / 16777216.0f,
@@ -301,52 +290,4 @@ namespace NiagaraStateless
 			float((V.W >> 8) & 0x00ffffff) / 16777216.0f
 		);
 	}
-
-	float FParticleSimulationContext::RandomSeedOffsetFloat(uint32 iInstance, uint32 SeedOffset) const
-	{
-		FUintVector4 RandomSeed = GetParticleRandomSeed()[iInstance];
-		RandomSeed.X += SeedOffset;
-
-		FUintVector4 RandomValue = Rand4DPCG32(RandomSeed);
-		return float((RandomValue.X >> 8) & 0x00ffffff) / 16777216.0f;
-	}
-
-	FVector2f FParticleSimulationContext::RandomSeedOffsetFloat2(uint32 iInstance, uint32 SeedOffset) const
-	{
-		FUintVector4 RandomSeed = GetParticleRandomSeed()[iInstance];
-		RandomSeed.X += SeedOffset;
-
-		FUintVector4 RandomValue = Rand4DPCG32(RandomSeed);
-		return FVector2f(
-			float((RandomValue.X >> 8) & 0x00ffffff) / 16777216.0f,
-			float((RandomValue.Y >> 8) & 0x00ffffff) / 16777216.0f
-		);
-	}
-
-	FVector3f FParticleSimulationContext::RandomSeedOffsetFloat3(uint32 iInstance, uint32 SeedOffset) const
-	{
-		FUintVector4 RandomSeed = GetParticleRandomSeed()[iInstance];
-		RandomSeed.X += SeedOffset;
-
-		FUintVector4 RandomValue = Rand4DPCG32(RandomSeed);
-		return FVector3f(
-			float((RandomValue.X >> 8) & 0x00ffffff) / 16777216.0f,
-			float((RandomValue.Y >> 8) & 0x00ffffff) / 16777216.0f,
-			float((RandomValue.Z >> 8) & 0x00ffffff) / 16777216.0f
-		);
-	}
-
-	FVector4f FParticleSimulationContext::RandomSeedOffsetFloat4(uint32 iInstance, uint32 SeedOffset) const
-	{
-		FUintVector4 RandomSeed = GetParticleRandomSeed()[iInstance];
-		RandomSeed.X += SeedOffset;
-
-		FUintVector4 RandomValue = Rand4DPCG32(RandomSeed);
-		return FVector4f(
-			float((RandomValue.X >> 8) & 0x00ffffff) / 16777216.0f,
-			float((RandomValue.Y >> 8) & 0x00ffffff) / 16777216.0f,
-			float((RandomValue.Z >> 8) & 0x00ffffff) / 16777216.0f,
-			float((RandomValue.W >> 8) & 0x00ffffff) / 16777216.0f
-		);
-	}
 } //namespace NiagaraStateless
