Split off the helper function SampleClosestPointMeshDistanceFIeld() into it's own header file so that multiple RigidBody DI can be used in the same script
#tests Made a GPU system with two Rigid Body DI and call GetClosestPointMeshDistanceField on each

diff --git a/Engine/Plugins/FX/Niagara/Shaders/Private/NiagaraDataInterfaceRigidMeshCollisionQuery.ush b/Engine/Plugins/FX/Niagara/Shaders/Private/NiagaraDataInterfaceRigidMeshCollisionQuery.ush
--- a/Engine/Plugins/FX/Niagara/Shaders/Private/NiagaraDataInterfaceRigidMeshCollisionQuery.ush
+++ b/Engine/Plugins/FX/Niagara/Shaders/Private/NiagaraDataInterfaceRigidMeshCollisionQuery.ush
@@ -4,6 +4,8 @@
 NiagaraDataInterfaceRigidMeshCollisionQuery.ush 
 =============================================================================*/
 
+#include "/Plugin/FX/Niagara/Private/NiagaraDistanceFieldQueries.ush"
+
 /* -----------------------------------------------------------------
  * static mesh constants and context
  * -----------------------------------------------------------------
@@ -498,57 +500,6 @@ void GetClosestDistance_{ParameterName}(in float3 WorldPosition, in float TimeFr
 	}
 }
 
-bool SampleClosestPointMeshDistanceField(int DFIndex, const FDFVector3 LWCWorldPosition, float MaxDistance, out float3 OutClosestPosition, out float3 OutClosestNormal)
-{
-	FDFObjectData DFObjectData = LoadDFObjectData(DFIndex);
-
-	// DistanceToNearestSurfaceForObject()
-	const float3 VolumePosition = DFMultiplyDemote(LWCWorldPosition, DFObjectData.WorldToVolume);
-	const float3 ClampedVolumePosition = clamp(VolumePosition, -DFObjectData.VolumePositionExtent, DFObjectData.VolumePositionExtent);
-
-	const float3 ToBox = (abs(ClampedVolumePosition) - DFObjectData.VolumePositionExtent) * DFObjectData.VolumeToWorldScale;
-	const float DistanceToBoxOutside = length(max(0.0f, ToBox));
-	const float DistanceToBoxInside = min(0.0f, max3(ToBox.x, ToBox.y, ToBox.z));
-	const float DistanceToBox = DistanceToBoxOutside + DistanceToBoxInside;
-
-	float DistanceToOccluder = MaxDistance;
-
-	BRANCH
-	if (DistanceToBox < MaxDistance)
-	{
-		uint NumMips = LoadDFAssetData(DFObjectData.AssetIndex, 0).NumMips;
-		FDFAssetData DFAssetData = LoadDFAssetData(DFObjectData.AssetIndex, NumMips - 1);
-
-		DistanceToOccluder = DistanceToMeshSurfaceStandalone(ClampedVolumePosition, DFObjectData) * DFObjectData.VolumeScale;
-
-		// Don't allow negative SDF distance to stick out of the bounding box
-		DistanceToOccluder = max(DistanceToOccluder, 0.0f);
-
-		float3 VolumeGradient = CalculateMeshSDFGradient(ClampedVolumePosition, DFAssetData);
-		float VolumeGradientLength = length(VolumeGradient);
-		float3 VolumeNormal = VolumeGradientLength > 0.0f ? VolumeGradient / VolumeGradientLength : 0;
-
-		// Transform by transposed inverse to handle non-uniform scaling
-		float4x4 WorldToVolume = DFHackToFloat(DFObjectData.WorldToVolume);
-		float3 WorldGradient = mul(VolumeNormal, transpose((float3x3)WorldToVolume));
-		float WorldGradientLength = length(WorldGradient);
-		const bool bNormalIsValid = WorldGradientLength > 0.0f;
-			
-		// output 			
-		if (bNormalIsValid)
-		{
-			OutClosestNormal = WorldGradient / WorldGradientLength;
-
-			const float3 WorldClampedPosition = DFDemote(DFMultiply(ClampedVolumePosition, DFObjectData.VolumeToWorld));
-			OutClosestPosition = WorldClampedPosition - OutClosestNormal * DistanceToOccluder;
-
-			return true;
-		}
-	}
-
-	return false;
-}
-
 void EvaluateWorldPositionOnElement_{ParameterName}(int ElementIndex, float3 WorldPosition, float DeltaTime, float TimeFraction, out float3 OutPosition, out float3 OutVelocity)
 {
 	const float3 LocalPosition = mul(GetInverseTransform_{ParameterName}(ElementIndex), float4(WorldPosition,1.0)).xyz;			
diff --git a/Engine/Plugins/FX/Niagara/Shaders/Private/NiagaraDistanceFieldQueries.ush b/Engine/Plugins/FX/Niagara/Shaders/Private/NiagaraDistanceFieldQueries.ush
new file mode 100644
--- /dev/null
+++ b/Engine/Plugins/FX/Niagara/Shaders/Private/NiagaraDistanceFieldQueries.ush
@@ -0,0 +1,58 @@
+// Copyright Epic Games, Inc. All Rights Reserved.
+
+/*=============================================================================
+NiagaraDistanceFieldQueries.ush 
+=============================================================================*/
+
+#pragma once
+
+bool SampleClosestPointMeshDistanceField(int DFIndex, const FDFVector3 LWCWorldPosition, float MaxDistance, out float3 OutClosestPosition, out float3 OutClosestNormal)
+{
+	FDFObjectData DFObjectData = LoadDFObjectData(DFIndex);
+
+	// DistanceToNearestSurfaceForObject()
+	const float3 VolumePosition = DFMultiplyDemote(LWCWorldPosition, DFObjectData.WorldToVolume);
+	const float3 ClampedVolumePosition = clamp(VolumePosition, -DFObjectData.VolumePositionExtent, DFObjectData.VolumePositionExtent);
+
+	const float3 ToBox = (abs(ClampedVolumePosition) - DFObjectData.VolumePositionExtent) * DFObjectData.VolumeToWorldScale;
+	const float DistanceToBoxOutside = length(max(0.0f, ToBox));
+	const float DistanceToBoxInside = min(0.0f, max3(ToBox.x, ToBox.y, ToBox.z));
+	const float DistanceToBox = DistanceToBoxOutside + DistanceToBoxInside;
+
+	float DistanceToOccluder = MaxDistance;
+
+	BRANCH
+	if (DistanceToBox < MaxDistance)
+	{
+		uint NumMips = LoadDFAssetData(DFObjectData.AssetIndex, 0).NumMips;
+		FDFAssetData DFAssetData = LoadDFAssetData(DFObjectData.AssetIndex, NumMips - 1);
+
+		DistanceToOccluder = DistanceToMeshSurfaceStandalone(ClampedVolumePosition, DFObjectData) * DFObjectData.VolumeScale;
+
+		// Don't allow negative SDF distance to stick out of the bounding box
+		DistanceToOccluder = max(DistanceToOccluder, 0.0f);
+
+		float3 VolumeGradient = CalculateMeshSDFGradient(ClampedVolumePosition, DFAssetData);
+		float VolumeGradientLength = length(VolumeGradient);
+		float3 VolumeNormal = VolumeGradientLength > 0.0f ? VolumeGradient / VolumeGradientLength : 0;
+
+		// Transform by transposed inverse to handle non-uniform scaling
+		float4x4 WorldToVolume = DFHackToFloat(DFObjectData.WorldToVolume);
+		float3 WorldGradient = mul(VolumeNormal, transpose((float3x3)WorldToVolume));
+		float WorldGradientLength = length(WorldGradient);
+		const bool bNormalIsValid = WorldGradientLength > 0.0f;
+			
+		// output 			
+		if (bNormalIsValid)
+		{
+			OutClosestNormal = WorldGradient / WorldGradientLength;
+
+			const float3 WorldClampedPosition = DFDemote(DFMultiply(ClampedVolumePosition, DFObjectData.VolumeToWorld));
+			OutClosestPosition = WorldClampedPosition - OutClosestNormal * DistanceToOccluder;
+
+			return true;
+		}
+	}
+
+	return false;
+}
