Fixup handling of DebugState static select nodes with the AsyncTasks compilation mode
#tests Use a scratch pad module to select color based on debug state

diff --git a/Engine/Plugins/FX/Niagara/Source/NiagaraEditor/Private/NiagaraGraphDigest.cpp b/Engine/Plugins/FX/Niagara/Source/NiagaraEditor/Private/NiagaraGraphDigest.cpp
--- a/Engine/Plugins/FX/Niagara/Source/NiagaraEditor/Private/NiagaraGraphDigest.cpp
+++ b/Engine/Plugins/FX/Niagara/Source/NiagaraEditor/Private/NiagaraGraphDigest.cpp
@@ -79,11 +79,12 @@ struct FNiagaraCompilationGraphDuplicateContext
 
 struct FNiagaraCompilationGraphInstanceContext
 {
-	FNiagaraCompilationGraphInstanceContext(const FNiagaraFixedConstantResolver& InConstantResolver, const FNiagaraPrecompileData* InPrecompileData)
+	FNiagaraCompilationGraphInstanceContext(const FNiagaraFixedConstantResolver& InConstantResolver, const FNiagaraCompilationGraph* InCompilationGraph, const FNiagaraPrecompileData* InPrecompileData)
 		: ConstantResolver(InConstantResolver)
 		, PrecompileData(InPrecompileData)
 	{
-
+		bDisableDebugSwitches = InPrecompileData ? InPrecompileData->bDisableDebugSwitches : false;
+		TraversalContext.BeginContext(InCompilationGraph, ConstantResolver);
 	}
 	FNiagaraCompilationGraphInstanceContext() = delete;
 
@@ -93,6 +94,7 @@ struct FNiagaraCompilationGraphInstanceContext
 	FGraphTraversalHandle GraphTraversalHandle;
 	TArray<const FNiagaraCompilationNodeFunctionCall*> FunctionStack;
 	bool bForceNumericResolution = false;
+	bool bDisableDebugSwitches = false;
 
 	void EnterFunction(const FNiagaraCompilationNodeFunctionCall* InCallingNode)
 	{
@@ -1261,7 +1263,7 @@ void FNiagaraCompilationGraphInstanced::StripUnconnectedPins(FNiagaraCompilation
 
 void FNiagaraCompilationGraphInstanced::InheritDebugState(FNiagaraCompilationGraphInstanceContext& Context, FNiagaraCompilationNodeFunctionCall& FunctionCallNode)
 {
-	FunctionCallNode.DebugState = FunctionCallNode.bInheritDebugState ? Context.ConstantResolver.GetDebugState() : ENiagaraFunctionDebugState::NoDebug;
+	FunctionCallNode.DebugState = FunctionCallNode.bInheritDebugState ? Context.ConstantResolver.GetDebugState() : FunctionCallNode.DebugState;
 }
 
 void FNiagaraCompilationGraphInstanced::PropagateDefaultValues(FNiagaraCompilationGraphInstanceContext& Context, FNiagaraCompilationNodeFunctionCall& FunctionCallNode)
@@ -1311,7 +1313,7 @@ TSharedPtr<FNiagaraCompilationGraphInstanced, ESPMode::ThreadSafe> FNiagaraCompi
 
 	// initialize the traversal context with the data that was pulled from the parameter map history done during the
 	// precompile
-	FNiagaraCompilationGraphInstanceContext InstantiationContext(ConstantResolver, PrecompileData);
+	FNiagaraCompilationGraphInstanceContext InstantiationContext(ConstantResolver, InstantiatedGraph.Get(), PrecompileData);
 
 	if (ensure(InstantiatedGraph))
 	{
@@ -2982,7 +2984,7 @@ void FNiagaraCompilationNodeFunctionCall::MultiFindParameterMapDefaultValues(ENi
 	{
 		FNiagaraCompilationBranchMap Branches;
 
-		FNiagaraCompilationGraphInstanceContext DummyContext(ConstantResolver, nullptr);
+		FNiagaraCompilationGraphInstanceContext DummyContext(ConstantResolver, nullptr, nullptr);
 		DummyContext.EnterFunction(this);
 		CalledGraph->EvaluateStaticBranches(DummyContext, Branches);
 	
@@ -4068,6 +4070,7 @@ FNiagaraCompilationNodeStaticSwitch::FNiagaraCompilationNodeStaticSwitch(const U
 {
 	bSetByCompiler = InNode->IsSetByCompiler();
 	bSetByPin = InNode->IsSetByPin();
+	bDebugStateSwitch = InNode->IsDebugSwitch();
 	SwitchType = InNode->SwitchTypeData.SwitchType;
 	SwitchBranchCount = InNode->GetOptionValues().Num();
 	InputParameterName = InNode->InputParameterName;
@@ -4307,7 +4310,7 @@ const FNiagaraCompilationOutputPin* FNiagaraCompilationNodeStaticSwitch::TraceOu
 
 TArray<const FNiagaraCompilationInputPin*> FNiagaraCompilationNodeStaticSwitch::EvaluateBranches(FNiagaraCompilationGraphInstanceContext& Context, FNiagaraCompilationBranchMap& Branches) const
 {
-	if (bSetByCompiler)
+	if (bSetByCompiler && !bDebugStateSwitch)
 	{
 		return FNiagaraCompilationNode::EvaluateBranches(Context, Branches);
 	}
@@ -4315,7 +4318,18 @@ TArray<const FNiagaraCompilationInputPin*> FNiagaraCompilationNodeStaticSwitch::
 	int32 SwitchValue = INDEX_NONE;
 	bool IsValueSet = false;
 
-	if (bSetByPin)
+	if (bDebugStateSwitch)
+	{
+		ENiagaraFunctionDebugState FunctionDebugState = ENiagaraFunctionDebugState::NoDebug;
+		if (!Context.bDisableDebugSwitches)
+		{
+			Context.TraversalContext.GetCurrentDebugState(FunctionDebugState);
+		}
+
+		IsValueSet = true;
+		SwitchValue = static_cast<int32>(FunctionDebugState);
+	}
+	else if (bSetByPin)
 	{
 		auto EvaluateVariableSwitchValue = [this](const FNiagaraVariable& ConstantVariable, int32& OutSwitchValue) -> bool
 		{
diff --git a/Engine/Plugins/FX/Niagara/Source/NiagaraEditor/Private/NiagaraGraphDigest.h b/Engine/Plugins/FX/Niagara/Source/NiagaraEditor/Private/NiagaraGraphDigest.h
--- a/Engine/Plugins/FX/Niagara/Source/NiagaraEditor/Private/NiagaraGraphDigest.h
+++ b/Engine/Plugins/FX/Niagara/Source/NiagaraEditor/Private/NiagaraGraphDigest.h
@@ -712,6 +712,7 @@ public:
 
 	bool bSetByCompiler;
 	bool bSetByPin;
+	bool bDebugStateSwitch;
 	ENiagaraStaticSwitchType SwitchType;
 	FNiagaraTypeDefinition InputType;
 	int32 SwitchBranchCount;
diff --git a/Engine/Plugins/FX/Niagara/Source/NiagaraEditor/Private/NiagaraTraversalStateContext.cpp b/Engine/Plugins/FX/Niagara/Source/NiagaraEditor/Private/NiagaraTraversalStateContext.cpp
--- a/Engine/Plugins/FX/Niagara/Source/NiagaraEditor/Private/NiagaraTraversalStateContext.cpp
+++ b/Engine/Plugins/FX/Niagara/Source/NiagaraEditor/Private/NiagaraTraversalStateContext.cpp
@@ -44,7 +44,7 @@ void FNiagaraTraversalStateContext::PushGraphInternal(const FNiagaraCompilationN
 					}
 				}
 			}
-			else if (!StaticSwitchNode->bSetByCompiler)
+			else if (!StaticSwitchNode->bSetByCompiler && CallingNode)
 			{
 				for (const FNiagaraCompilationInputPin& InputPin : CallingNode->InputPins)
 				{
@@ -79,35 +79,38 @@ void FNiagaraTraversalStateContext::PushGraphInternal(const FNiagaraCompilationN
 
 			// based on the original code bInheritDebugState drives whether we use the system value vs NoDebug.  Note
 			// that the serialized Debugstate is never used.
-			ENiagaraFunctionDebugState CachedDebugState = InnerFunctionNode->bInheritDebugState
+			ENiagaraFunctionDebugState CachedDebugState = (CallingNode && InnerFunctionNode->bInheritDebugState)
 				? ConstantResolver.GetDebugState()
-				: ENiagaraFunctionDebugState::NoDebug;
+				: InnerFunctionNode->DebugState;
 
 			if (ensure(!FunctionDebugStateMap.Contains(InnerFunctionNodeHash)))
 			{
 				FunctionDebugStateMap.Add(InnerFunctionNodeHash, CachedDebugState);
 			}
 
-			for (const FNiagaraCompilationNodeFunctionCall::FTaggedVariable& TaggedVariable : InnerFunctionNode->PropagatedStaticSwitchParameters)
+			if (CallingNode)
 			{
-				const FNiagaraCompilationInputPin* ValuePin = InnerFunctionNode->InputPins.FindByPredicate([&TaggedVariable](const FNiagaraCompilationInputPin& InputPin) -> bool
+				for (const FNiagaraCompilationNodeFunctionCall::FTaggedVariable& TaggedVariable : InnerFunctionNode->PropagatedStaticSwitchParameters)
+				{
+					const FNiagaraCompilationInputPin* ValuePin = InnerFunctionNode->InputPins.FindByPredicate([&TaggedVariable](const FNiagaraCompilationInputPin& InputPin) -> bool
 					{
 						return InputPin.PinName == TaggedVariable.Key.GetName();
 					});
 
-				if (ValuePin)
-				{
-					const FNiagaraCompilationInputPin* CallerInputPin = CallingNode->InputPins.FindByPredicate([&TaggedVariable](const FNiagaraCompilationInputPin& InputPin) -> bool
+					if (ValuePin)
+					{
+						const FNiagaraCompilationInputPin* CallerInputPin = CallingNode->InputPins.FindByPredicate([&TaggedVariable](const FNiagaraCompilationInputPin& InputPin) -> bool
 						{
 							return InputPin.PinName == TaggedVariable.Value;
 						});
 
-					if (CallerInputPin)
-					{
-						const uint32 FunctionPinNodeHash = HashCombine(InnerFunctionNodeHash, GetTypeHash(ValuePin->PinName));
-						if (ensure(!FunctionDefaultValueMap.Contains(FunctionPinNodeHash)))
+						if (CallerInputPin)
 						{
-							FunctionDefaultValueMap.Add(FunctionPinNodeHash, CallerInputPin->DefaultValue);
+							const uint32 FunctionPinNodeHash = HashCombine(InnerFunctionNodeHash, GetTypeHash(ValuePin->PinName));
+							if (ensure(!FunctionDefaultValueMap.Contains(FunctionPinNodeHash)))
+							{
+								FunctionDefaultValueMap.Add(FunctionPinNodeHash, CallerInputPin->DefaultValue);
+							}
 						}
 					}
 				}
@@ -116,6 +119,19 @@ void FNiagaraTraversalStateContext::PushGraphInternal(const FNiagaraCompilationN
 	}
 }
 
+void FNiagaraTraversalStateContext::BeginContext(const FNiagaraCompilationGraph* ParentGraph, const FNiagaraFixedConstantResolver& ConstantResolver)
+{
+	if (ParentGraph)
+	{
+		FNiagaraTraversalStackEntry& StackTop = TraversalStack.AddDefaulted_GetRef();
+#if WITH_NIAGARA_TRAVERSAL_FRIENDLY_NAME
+		StackTop.FriendlyName = TEXT("Root");
+#endif
+
+		PushGraphInternal(nullptr, ParentGraph, ConstantResolver);
+	}
+}
+
 void FNiagaraTraversalStateContext::PushFunction(const FNiagaraCompilationNodeFunctionCall* FunctionCall, const FNiagaraFixedConstantResolver& ConstantResolver)
 {
 	TOptional<uint32> CurrentStackHash;
@@ -125,11 +141,13 @@ void FNiagaraTraversalStateContext::PushFunction(const FNiagaraCompilationNodeFu
 	}
 
 	FNiagaraTraversalStackEntry& StackTop = TraversalStack.AddDefaulted_GetRef();
-	StackTop.FriendlyName = FString::Printf(TEXT("FuncionName - %s | FullName - %s | FullTitle - %s | NodeType - %d"),
+#if WITH_NIAGARA_TRAVERSAL_FRIENDLY_NAME
+	StackTop.FriendlyName = FString::Printf(TEXT("FunctionName - %s | FullName - %s | FullTitle - %s | NodeType - %d"),
 		*FunctionCall->FunctionName,
 		*FunctionCall->FullName,
 		*FunctionCall->FullTitle,
 		(int32)FunctionCall->NodeType);
+#endif
 
 	StackTop.NodeGuid = FunctionCall->NodeGuid;
 	const uint32 NodeGuidHash = GetTypeHash(StackTop.NodeGuid);
@@ -152,11 +170,13 @@ void FNiagaraTraversalStateContext::PushEmitter(const FNiagaraCompilationNodeEmi
 	}
 
 	FNiagaraTraversalStackEntry& StackTop = TraversalStack.AddDefaulted_GetRef();
+#if WITH_NIAGARA_TRAVERSAL_FRIENDLY_NAME
 	StackTop.FriendlyName = FString::Printf(TEXT("EmitterName - %s | FullName - %s | FullTitle - %s | NodeType - %d"),
 		*Emitter->EmitterUniqueName,
 		*Emitter->FullName,
 		*Emitter->FullTitle,
 		(int32)Emitter->NodeType);
+#endif
 
 	StackTop.NodeGuid = Emitter->NodeGuid;
 	const uint32 NodeGuidHash = GetTypeHash(StackTop.NodeGuid);
@@ -219,6 +239,19 @@ bool FNiagaraTraversalStateContext::GetFunctionDebugState(const FGuid& NodeGuid,
 	return false;
 }
 
+bool FNiagaraTraversalStateContext::GetCurrentDebugState(ENiagaraFunctionDebugState& DebugState) const
+{
+	if (!TraversalStack.IsEmpty())
+	{
+		if (const ENiagaraFunctionDebugState* ValuePtr = FunctionDebugStateMap.Find(TraversalStack.Top().FullStackHash))
+		{
+			DebugState = *ValuePtr;
+			return true;
+		}
+	}
+	return false;
+}
+
 FNiagaraFixedConstantResolver::FNiagaraFixedConstantResolver()
 {
 	InitConstants();
diff --git a/Engine/Plugins/FX/Niagara/Source/NiagaraEditor/Private/NiagaraTraversalStateContext.h b/Engine/Plugins/FX/Niagara/Source/NiagaraEditor/Private/NiagaraTraversalStateContext.h
--- a/Engine/Plugins/FX/Niagara/Source/NiagaraEditor/Private/NiagaraTraversalStateContext.h
+++ b/Engine/Plugins/FX/Niagara/Source/NiagaraEditor/Private/NiagaraTraversalStateContext.h
@@ -10,15 +10,21 @@ class FNiagaraCompilationNodeEmitter;
 class FNiagaraCompilationNodeFunctionCall;
 class FNiagaraFixedConstantResolver;
 
+#define WITH_NIAGARA_TRAVERSAL_FRIENDLY_NAME (UE_BUILD_DEBUG)
+
 struct FNiagaraTraversalStackEntry
 {
 	FGuid NodeGuid;
 	uint32 FullStackHash;
+#if WITH_NIAGARA_TRAVERSAL_FRIENDLY_NAME
 	FString FriendlyName;
+#endif
 };
 
 struct FNiagaraTraversalStateContext
 {
+	void BeginContext(const FNiagaraCompilationGraph* ParentGraph, const FNiagaraFixedConstantResolver& ConstantResolver);
+
 	void PushFunction(const FNiagaraCompilationNodeFunctionCall* FunctionCall, const FNiagaraFixedConstantResolver& ConstantResolver);
 	void PopFunction(const FNiagaraCompilationNodeFunctionCall* FunctionCall);
 
@@ -29,6 +35,8 @@ struct FNiagaraTraversalStateContext
 	bool GetFunctionDefaultValue(const FGuid& NodeGuid, FName PinName, FString& FunctionDefaultValue) const;
 	bool GetFunctionDebugState(const FGuid& NodeGuid, ENiagaraFunctionDebugState& DebugState) const;
 
+	bool GetCurrentDebugState(ENiagaraFunctionDebugState& DebugState) const;
+
 	TArray<FNiagaraTraversalStackEntry> TraversalStack;
 
 	TMap<uint32, int32> StaticSwitchValueMap;
