Reflect parameter removal to AffectedScripts when stack function inputs are reset to their base

diff --git a/Engine/Plugins/FX/Niagara/Source/NiagaraEditor/Private/ViewModels/Stack/NiagaraStackFunctionInput.cpp b/Engine/Plugins/FX/Niagara/Source/NiagaraEditor/Private/ViewModels/Stack/NiagaraStackFunctionInput.cpp
--- a/Engine/Plugins/FX/Niagara/Source/NiagaraEditor/Private/ViewModels/Stack/NiagaraStackFunctionInput.cpp
+++ b/Engine/Plugins/FX/Niagara/Source/NiagaraEditor/Private/ViewModels/Stack/NiagaraStackFunctionInput.cpp
@@ -2353,8 +2353,22 @@ void UNiagaraStackFunctionInput::ResetToBase()
 	{
 		TSharedRef<FNiagaraScriptMergeManager> MergeManager = FNiagaraScriptMergeManager::Get();
 
-		FVersionedNiagaraEmitter BaseEmitter = GetEmitterViewModel()->GetEmitter().GetEmitterData()->GetParent();
+		UNiagaraScript* OwningScript = nullptr;
+		FVersionedNiagaraEmitterData* EmitterData = GetEmitterViewModel()->GetEmitter().GetEmitterData();
 		UNiagaraNodeOutput* OutputNode = FNiagaraStackGraphUtilities::GetEmitterOutputNodeForStackNode(*OwningFunctionCallNode.Get());
+		if (ensure(OutputNode))
+		{
+			OwningScript = EmitterData->GetScript(OutputNode->GetUsage(), OutputNode->GetUsageId());
+		}
+
+		if (!OwningScript)
+		{
+			return;
+		}
+
+		FVersionedNiagaraEmitter BaseEmitter = EmitterData->GetParent();
+		TArray<FNiagaraVariable> OriginalRapidIterationParameters;
+		OwningScript->RapidIterationParameters.GetParameters(OriginalRapidIterationParameters);
 
 		FScopedTransaction ScopedTransaction(LOCTEXT("ResetInputToBaseTransaction", "Reset this input to match the parent emitter."));
 		FNiagaraScriptMergeManager::FApplyDiffResults Results = MergeManager->ResetModuleInputToBase(
@@ -2370,11 +2384,26 @@ void UNiagaraStackFunctionInput::ResetToBase()
 			// If resetting to the base succeeded, an unknown number of rapid iteration parameters may have been added.  To fix
 			// this copy all of the owning scripts rapid iteration parameters to all other affected scripts.
 			// TODO: Either the merge should take care of this directly, or at least provide more information about what changed.
-			UNiagaraScript* OwningScript = GetEmitterViewModel()->GetEmitter().GetEmitterData()->GetScript(OutputNode->GetUsage(), OutputNode->GetUsageId());
 			TArray<FNiagaraVariable> OwningScriptRapidIterationParameters;
 			OwningScript->RapidIterationParameters.GetParameters(OwningScriptRapidIterationParameters);
-			if (OwningScriptRapidIterationParameters.Num() > 0)
+
+			// we also need to check if we've removed any RI parameters, and if so pass that forward to the affected scripts as well
+			TArray<FNiagaraVariable> RemovedVariables = OriginalRapidIterationParameters;
+			RemovedVariables.SetNum(Algo::RemoveIf(RemovedVariables, [&OwningScriptRapidIterationParameters](const FNiagaraVariable& Var) -> bool
 			{
+				return OwningScriptRapidIterationParameters.Contains(Var);
+			}));
+			
+			if (!OwningScriptRapidIterationParameters.IsEmpty() || !RemovedVariables.IsEmpty())
+			{
+				auto ContainsStaticVariable = [](const FNiagaraVariable& Variable) -> bool
+				{
+					return Variable.GetType().IsStatic();
+				};
+
+				const bool bChangeRequiresRecompile = OwningScriptRapidIterationParameters.ContainsByPredicate(ContainsStaticVariable)
+					|| RemovedVariables.ContainsByPredicate(ContainsStaticVariable);
+
 				for (TWeakObjectPtr<UNiagaraScript> AffectedScript : AffectedScripts)
 				{
 					if (AffectedScript.Get() != OwningScript)
@@ -2386,6 +2415,21 @@ void UNiagaraStackFunctionInput::ResetToBase()
 							AffectedScript->RapidIterationParameters.SetParameterData(
 								OwningScript->RapidIterationParameters.GetParameterData(OwningScriptRapidIterationParameter), OwningScriptRapidIterationParameter, bAddParameterIfMissing);
 						}
+
+						for (FNiagaraVariable& RemovedVariable : RemovedVariables)
+						{
+							AffectedScript->RapidIterationParameters.RemoveParameter(RemovedVariable);
+						}
+					}
+
+					// we mark the graphs associated with the scripts as dirty to also ensure that they 
+					if (bChangeRequiresRecompile)
+					{
+						UNiagaraScriptSource* Source = Cast<UNiagaraScriptSource>(AffectedScript->GetLatestSource());
+						if (Source && Source->NodeGraph)
+						{
+							Source->NodeGraph->NotifyGraphNeedsRecompile();
+						}
 					}
 				}
 			}
