Fix for crash in Niagara attempting to access unsynchronized skeletal mesh buffers when flushing pending ticks.
 - The Niagara tick flush needs to sync the SkeletalMesh render command pipe to access pending ticks. It also needs to flush the GPU skin cache since there might be pending dispatches queued up.

diff --git a/Engine/Plugins/FX/Niagara/Source/Niagara/Private/NiagaraGpuComputeDispatch.cpp b/Engine/Plugins/FX/Niagara/Source/Niagara/Private/NiagaraGpuComputeDispatch.cpp
--- a/Engine/Plugins/FX/Niagara/Source/Niagara/Private/NiagaraGpuComputeDispatch.cpp
+++ b/Engine/Plugins/FX/Niagara/Source/Niagara/Private/NiagaraGpuComputeDispatch.cpp
@@ -34,6 +34,8 @@
 #include "SceneRenderTargetParameters.h"
 #include "TextureResource.h"
 #include "FXRenderingUtils.h"
+#include "Rendering/RenderCommandPipes.h"
+#include "GPUSkinCache.h"
 
 DECLARE_CYCLE_STAT(TEXT("GPU Dispatch Setup [RT]"), STAT_NiagaraGPUDispatchSetup_RT, STATGROUP_Niagara);
 DECLARE_CYCLE_STAT(TEXT("GPU Emitter Dispatch [RT]"), STAT_NiagaraGPUSimTick_RT, STATGROUP_Niagara);
@@ -297,6 +299,7 @@ void FNiagaraGpuComputeDispatch::RemoveGpuComputeProxy(FNiagaraSystemGpuComputeP
 void FNiagaraGpuComputeDispatch::Tick(UWorld* World, float DeltaTime)
 {
 	check(IsInGameThread());
+	UE::RenderCommandPipe::FSyncScope Scope(MakeArrayView({ &UE::RenderCommandPipe::SkeletalMesh }));
 	ENQUEUE_RENDER_COMMAND(NiagaraPumpBatcher)(
 		[RT_NiagaraBatcher=this](FRHICommandListImmediate& RHICmdList)
 		{
@@ -309,6 +312,7 @@ void FNiagaraGpuComputeDispatch::Tick(UWorld* World, float DeltaTime)
 void FNiagaraGpuComputeDispatch::FlushPendingTicks_GameThread()
 {
 	check(IsInGameThread());
+	UE::RenderCommandPipe::FSyncScope Scope(MakeArrayView({ &UE::RenderCommandPipe::SkeletalMesh }));
 	ENQUEUE_RENDER_COMMAND(NiagaraFlushPendingTicks)(
 		[RT_NiagaraBatcher=this](FRHICommandListImmediate& RHICmdList)
 		{
@@ -321,6 +325,7 @@ void FNiagaraGpuComputeDispatch::FlushPendingTicks_GameThread()
 void FNiagaraGpuComputeDispatch::FlushAndWait_GameThread()
 {
 	check(IsInGameThread());
+	UE::RenderCommandPipe::FSyncScope Scope(MakeArrayView({ &UE::RenderCommandPipe::SkeletalMesh }));
 	ENQUEUE_RENDER_COMMAND(NiagaraFlushPendingTicks)(
 		[RT_NiagaraBatcher=this](FRHICommandListImmediate& RHICmdList)
 		{
@@ -401,6 +406,12 @@ void FNiagaraGpuComputeDispatch::ProcessPendingTicksFlush(FRHICommandListImmedia
 			FDeferredUpdateResource::UpdateResources(RHICmdList);
 			FMaterialRenderProxy::UpdateDeferredCachedUniformExpressions();
 
+			// Flush all batched up skin cache dispatches
+			if (FGPUSkinCache* GPUSkinCache = GetSceneInterface()->GetGPUSkinCache())
+			{
+				GPUSkinCache->DoDispatch(RHICmdList);
+			}
+
 			// Make a temporary ViewInfo
 			//-TODO: We could gather some more information here perhaps?
 			FSceneViewFamilyContext ViewFamily(
