Remove dead code & files

diff --git a/Engine/Plugins/FX/Niagara/Source/NiagaraEditor/Private/Customizations/NiagaraMetaDataCustomNodeBuilder.cpp b/Engine/Plugins/FX/Niagara/Source/NiagaraEditor/Private/Customizations/NiagaraMetaDataCustomNodeBuilder.cpp
deleted file mode 100644
diff --git a/Engine/Plugins/FX/Niagara/Source/NiagaraEditor/Private/Customizations/NiagaraMetaDataCustomNodeBuilder.h b/Engine/Plugins/FX/Niagara/Source/NiagaraEditor/Private/Customizations/NiagaraMetaDataCustomNodeBuilder.h
deleted file mode 100644
diff --git a/Engine/Plugins/FX/Niagara/Source/NiagaraEditor/Private/Customizations/NiagaraScriptDetails.cpp b/Engine/Plugins/FX/Niagara/Source/NiagaraEditor/Private/Customizations/NiagaraScriptDetails.cpp
--- a/Engine/Plugins/FX/Niagara/Source/NiagaraEditor/Private/Customizations/NiagaraScriptDetails.cpp
+++ b/Engine/Plugins/FX/Niagara/Source/NiagaraEditor/Private/Customizations/NiagaraScriptDetails.cpp
@@ -4,32 +4,21 @@
 
 #include "DetailCategoryBuilder.h"
 #include "DetailLayoutBuilder.h"
-#include "DetailWidgetRow.h"
-#include "IDetailChildrenBuilder.h"
-#include "INiagaraEditorTypeUtilities.h"
-#include "NiagaraEditorModule.h"
 #include "NiagaraEditorStyle.h"
-#include "NiagaraMetaDataCustomNodeBuilder.h"
-#include "NiagaraParameterCollectionCustomNodeBuilder.h"
-#include "ViewModels/NiagaraParameterCollectionViewModel.h"
 #include "ViewModels/NiagaraScriptInputCollectionViewModel.h"
 #include "ViewModels/NiagaraScriptOutputCollectionViewModel.h"
-#include "ViewModels/NiagaraParameterViewModel.h"
+#include "NiagaraParameterCollectionCustomNodeBuilder.h"
+#include "ViewModels/NiagaraParameterCollectionViewModel.h"
 #include "NiagaraScript.h"
 #include "ViewModels/NiagaraScriptGraphViewModel.h"
-#include "PropertyCustomizationHelpers.h"
 #include "ScopedTransaction.h"
 #include "Framework/MultiBox/MultiBoxBuilder.h"
 #include "ViewModels/NiagaraScriptViewModel.h"
-#include "Widgets/SItemSelector.h"
 #include "Widgets/SOverlay.h"
-#include "Widgets/Colors/SColorBlock.h"
-#include "Widgets/Colors/SColorPicker.h"
 #include "Widgets/Images/SImage.h"
 #include "Widgets/Input/SButton.h"
 #include "Widgets/Input/SComboButton.h"
 #include "Widgets/Layout/SBox.h"
-#include "Widgets/Text/SInlineEditableTextBlock.h"
 
 #define LOCTEXT_NAMESPACE "NiagaraScriptDetails"
 
@@ -238,15 +227,6 @@ void FNiagaraScriptDetails::CustomizeDetails(IDetailLayoutBuilder& DetailBuilder
 	}
 }
 
-FReply FNiagaraScriptDetails::OnRefreshMetadata()
-{
-	if (MetaDataBuilder.IsValid())
-	{
-		MetaDataBuilder->Rebuild();
-	}
-    return FReply::Handled();
-}
-
 float ButtonWidth = 50;
 float DisplayColorHeight = 16;
 float DisplayColorWidth = 24;
diff --git a/Engine/Plugins/FX/Niagara/Source/NiagaraEditor/Private/Customizations/NiagaraScriptDetails.h b/Engine/Plugins/FX/Niagara/Source/NiagaraEditor/Private/Customizations/NiagaraScriptDetails.h
--- a/Engine/Plugins/FX/Niagara/Source/NiagaraEditor/Private/Customizations/NiagaraScriptDetails.h
+++ b/Engine/Plugins/FX/Niagara/Source/NiagaraEditor/Private/Customizations/NiagaraScriptDetails.h
@@ -24,9 +24,6 @@ public:
 
 	virtual void CustomizeDetails(IDetailLayoutBuilder& DetailBuilder);
 
-	FReply OnRefreshMetadata();
-
 private:
 	TSharedPtr<FNiagaraScriptViewModel> ScriptViewModel;
-	TSharedPtr<FNiagaraMetaDataCustomNodeBuilder> MetaDataBuilder;
 };
