Fixup NiagaraScripts when considering iterative cooking.  AppendToClassSchema is used to mark global settings/configurations that wouold impact the cook results outisde of object dependencies.
#tests iterative cook of assets between two CL

diff --git a/Engine/Plugins/FX/Niagara/Source/Niagara/Classes/NiagaraScript.h b/Engine/Plugins/FX/Niagara/Source/Niagara/Classes/NiagaraScript.h
--- a/Engine/Plugins/FX/Niagara/Source/Niagara/Classes/NiagaraScript.h
+++ b/Engine/Plugins/FX/Niagara/Source/Niagara/Classes/NiagaraScript.h
@@ -1157,8 +1157,17 @@ public:
 	void SaveShaderStableKeys(const class ITargetPlatform* TP);
 
 	TArray<FName> FindShaderFormatsForCooking(const ITargetPlatform* TargetPlatform) const;
-
 #endif // WITH_EDITOR
+
+#if WITH_EDITORONLY_DATA
+	/**
+	 * Append config values or settings that can change how instances of the class are cooked, including especially
+	 * values that determine how version upgraded are conducted. Can also append a unique guid when necessary to
+	 * invalidate previous results because serialization changed and no custom version was updated.
+	 */
+	NIAGARA_API static void AppendToClassSchema(FAppendToClassSchemaContext& Context);
+#endif
+
 	FNiagaraShaderScript* GetRenderThreadScript()
 	{
 		return ScriptResource.Get();
diff --git a/Engine/Plugins/FX/Niagara/Source/Niagara/Private/NiagaraScript.cpp b/Engine/Plugins/FX/Niagara/Source/Niagara/Private/NiagaraScript.cpp
--- a/Engine/Plugins/FX/Niagara/Source/Niagara/Private/NiagaraScript.cpp
+++ b/Engine/Plugins/FX/Niagara/Source/Niagara/Private/NiagaraScript.cpp
@@ -4292,6 +4292,38 @@ void UNiagaraScript::ResolveParameterCollectionReferences()
 
 #endif
 
+#if WITH_EDITORONLY_DATA
+void UNiagaraScript::AppendToClassSchema(FAppendToClassSchemaContext& Context)
+{
+	Super::AppendToClassSchema(Context);
+
+	// Used by iterative cooking.  This will provide additional context for if things have changed such that a cook will
+	// be required.  This is focused on global settings rather than the usual dependencies between objects.
+
+	// the compiler version
+	const FGuid CompilerVersion = FNiagaraCustomVersion::GetLatestScriptCompileVersion();
+	Context.Update(&CompilerVersion, sizeof(FGuid));
+
+	// the DI source hash
+	if (GNiagaraCompileHashAllDataInterfaces)
+	{
+		Context.Update(NiagaraScriptInternal::CompileHashAllDataInterfaces().GetData(), FNiagaraCompileHash::HashSize);
+	}
+
+	// settings/configurations
+	Context.Update(&GNiagaraTranslatorFailIfNotSetSeverity, sizeof(GNiagaraTranslatorFailIfNotSetSeverity));
+
+	const UNiagaraSettings* NiagaraSettings = GetDefault<UNiagaraSettings>();
+	checkf(NiagaraSettings, TEXT("Failed to acquire default UNiagaraSettings during AppendToClassSchema"));
+
+	if (NiagaraSettings)
+	{
+		Context.Update(&NiagaraSettings->InvalidNamespaceWriteSeverity, sizeof(NiagaraSettings->InvalidNamespaceWriteSeverity));
+	}
+}
+#endif
+
+
 TArray<TObjectPtr<UNiagaraParameterCollection>>& UNiagaraScript::GetCachedParameterCollectionReferences()
 {
 #if WITH_EDITORONLY_DATA
