- Split RDG UAV pool into per format types

#rnx

diff --git a/Engine/Plugins/FX/Niagara/Source/Niagara/Private/NiagaraEmptyUAVPool.cpp b/Engine/Plugins/FX/Niagara/Source/Niagara/Private/NiagaraEmptyUAVPool.cpp
--- a/Engine/Plugins/FX/Niagara/Source/Niagara/Private/NiagaraEmptyUAVPool.cpp
+++ b/Engine/Plugins/FX/Niagara/Source/Niagara/Private/NiagaraEmptyUAVPool.cpp
@@ -47,12 +47,8 @@ FNiagaraEmptyRDGUAVPoolScopedAccess::~FNiagaraEmptyRDGUAVPoolScopedAccess()
 void FNiagaraEmptyUAVPool::Tick()
 {
 	check(IsInRenderingThread());
-	BufferRDGUAVPool.UAVs.Empty();
-	Texture2DRDGUAVPool.UAVs.Empty();
-	Texture2DArrayRDGUAVPool.UAVs.Empty();
-	Texture3DRDGUAVPool.UAVs.Empty();
-	TextureCubeRDGUAVPool.UAVs.Empty();
-	TextureCubeArrayRDGUAVPool.UAVs.Empty();
+	BufferRDGUAVPool.Empty();
+	TextureRDGUAVPool.Empty();
 }
 
 FRHIUnorderedAccessView* FNiagaraEmptyUAVPool::GetEmptyUAVFromPool(FRHICommandList& RHICmdList, EPixelFormat Format, ENiagaraEmptyUAVType Type)
@@ -174,7 +170,7 @@ FRDGBufferUAVRef FNiagaraEmptyUAVPool::GetEmptyRDGUAVFromPool(FRDGBuilder& Graph
 	checkf(RDGUAVAccessCounter != 0, TEXT("Accessing Niagara's RDG UAV Pool while not within a scope, this could result in a memory leak!"));
 	check(UE::PixelFormat::HasCapabilities(Format, EPixelFormatCapabilities::TypedUAVStore));
 
-	FBufferRDGUAVPool& Pool = BufferRDGUAVPool;
+	FBufferRDGUAVPool& Pool = BufferRDGUAVPool.FindOrAdd(Format);
 	if (Pool.NextFreeIndex == Pool.UAVs.Num())
 	{
 		const uint32 BytesPerElement = GPixelFormats[Format].BlockBytes;
@@ -198,22 +194,13 @@ FRDGTextureUAVRef FNiagaraEmptyUAVPool::GetEmptyRDGUAVFromPool(FRDGBuilder& Grap
 	checkf(RDGUAVAccessCounter != 0, TEXT("Accessing Niagara's RDG UAV Pool while not within a scope, this could result in a memory leak!"));
 	check(UE::PixelFormat::HasCapabilities(Format, EPixelFormatCapabilities::TypedUAVStore));
 
-	FTextureRDGUAVPool* Pool = nullptr;
-	switch (TextureDimension)
-	{
-		case ETextureDimension::Texture2D:			Pool = &Texture2DRDGUAVPool; break;
-		case ETextureDimension::Texture2DArray:		Pool = &Texture2DArrayRDGUAVPool; break;
-		case ETextureDimension::Texture3D:			Pool = &Texture3DRDGUAVPool; break;
-		case ETextureDimension::TextureCube:		Pool = &TextureCubeRDGUAVPool; break;
-		case ETextureDimension::TextureCubeArray:	Pool = &TextureCubeArrayRDGUAVPool; break;
-		default:									checkNoEntry();	return nullptr;
-	}
-	if (Pool->NextFreeIndex == Pool->UAVs.Num())
+	FTextureRDGUAVPool& Pool = TextureRDGUAVPool.FindOrAdd(TPairInitializer(Format, TextureDimension));
+	if (Pool.NextFreeIndex == Pool.UAVs.Num())
 	{
 		switch (TextureDimension)
 		{
 			case ETextureDimension::Texture2D:
-				Pool->UAVs.Add(
+				Pool.UAVs.Add(
 					GraphBuilder.CreateUAV(
 						GraphBuilder.CreateTexture(
 							FRDGTextureDesc::Create2D(FIntPoint(1, 1), Format, FClearValueBinding::Black, ETextureCreateFlags::ShaderResource | ETextureCreateFlags::UAV),
@@ -225,7 +212,7 @@ FRDGTextureUAVRef FNiagaraEmptyUAVPool::GetEmptyRDGUAVFromPool(FRDGBuilder& Grap
 				break;
 
 			case ETextureDimension::Texture2DArray:
-				Pool->UAVs.Add(
+				Pool.UAVs.Add(
 					GraphBuilder.CreateUAV(
 						GraphBuilder.CreateTexture(
 							FRDGTextureDesc::Create2DArray(FIntPoint(1, 1), Format, FClearValueBinding::Black, ETextureCreateFlags::ShaderResource | ETextureCreateFlags::UAV, 1),
@@ -237,7 +224,7 @@ FRDGTextureUAVRef FNiagaraEmptyUAVPool::GetEmptyRDGUAVFromPool(FRDGBuilder& Grap
 				break;
 
 			case ETextureDimension::Texture3D:
-				Pool->UAVs.Add(
+				Pool.UAVs.Add(
 					GraphBuilder.CreateUAV(
 						GraphBuilder.CreateTexture(
 							FRDGTextureDesc::Create3D(FIntVector(1, 1, 1), Format, FClearValueBinding::Black, ETextureCreateFlags::ShaderResource | ETextureCreateFlags::UAV),
@@ -249,7 +236,7 @@ FRDGTextureUAVRef FNiagaraEmptyUAVPool::GetEmptyRDGUAVFromPool(FRDGBuilder& Grap
 				break;
 
 			case ETextureDimension::TextureCube:
-				Pool->UAVs.Add(
+				Pool.UAVs.Add(
 					GraphBuilder.CreateUAV(
 						GraphBuilder.CreateTexture(
 							FRDGTextureDesc::CreateCube(1, Format, FClearValueBinding::Black, ETextureCreateFlags::ShaderResource | ETextureCreateFlags::UAV),
@@ -261,7 +248,7 @@ FRDGTextureUAVRef FNiagaraEmptyUAVPool::GetEmptyRDGUAVFromPool(FRDGBuilder& Grap
 				break;
 
 			case ETextureDimension::TextureCubeArray:
-				Pool->UAVs.Add(
+				Pool.UAVs.Add(
 					GraphBuilder.CreateUAV(
 						GraphBuilder.CreateTexture(
 							FRDGTextureDesc::CreateCubeArray(1, Format, FClearValueBinding::Black, ETextureCreateFlags::ShaderResource | ETextureCreateFlags::UAV, 1),
@@ -277,17 +264,20 @@ FRDGTextureUAVRef FNiagaraEmptyUAVPool::GetEmptyRDGUAVFromPool(FRDGBuilder& Grap
 				return nullptr;
 		}
 	}
-	return Pool->UAVs[Pool->NextFreeIndex++];
+	return Pool.UAVs[Pool.NextFreeIndex++];
 }
 
 void FNiagaraEmptyUAVPool::ResetEmptyRDGUAVPools()
 {
-	BufferRDGUAVPool.NextFreeIndex = 0;
-	Texture2DRDGUAVPool.NextFreeIndex = 0;
-	Texture2DArrayRDGUAVPool.NextFreeIndex = 0;
-	Texture3DRDGUAVPool.NextFreeIndex = 0;
-	TextureCubeRDGUAVPool.NextFreeIndex = 0;
-	TextureCubeArrayRDGUAVPool.NextFreeIndex = 0;
+	for ( auto it=BufferRDGUAVPool.CreateIterator(); it; ++it )
+	{
+		it.Value().NextFreeIndex = 0;
+	}
+
+	for ( auto it=TextureRDGUAVPool.CreateIterator(); it; ++it )
+	{
+		it.Value().NextFreeIndex = 0;
+	}
 }
 
 FNiagaraEmptyUAVPool::FEmptyUAV::~FEmptyUAV()
diff --git a/Engine/Plugins/FX/Niagara/Source/Niagara/Public/NiagaraEmptyUAVPool.h b/Engine/Plugins/FX/Niagara/Source/Niagara/Public/NiagaraEmptyUAVPool.h
--- a/Engine/Plugins/FX/Niagara/Source/Niagara/Public/NiagaraEmptyUAVPool.h
+++ b/Engine/Plugins/FX/Niagara/Source/Niagara/Public/NiagaraEmptyUAVPool.h
@@ -111,10 +111,6 @@ protected:
 	};
 
 	uint32 RDGUAVAccessCounter = 0;
-	FBufferRDGUAVPool BufferRDGUAVPool;
-	FTextureRDGUAVPool Texture2DRDGUAVPool;
-	FTextureRDGUAVPool Texture2DArrayRDGUAVPool;
-	FTextureRDGUAVPool Texture3DRDGUAVPool;
-	FTextureRDGUAVPool TextureCubeRDGUAVPool;
-	FTextureRDGUAVPool TextureCubeArrayRDGUAVPool;
+	TMap<EPixelFormat, FBufferRDGUAVPool> BufferRDGUAVPool;
+	TMap<TPair<EPixelFormat, ETextureDimension>, FTextureRDGUAVPool> TextureRDGUAVPool;
 };
