- Stateless add shared headers between CPU & GPU
- Added a definition header which contains random seed generation to share code
- Added a built distribution header to share how we pack / unpack data, allows us to more easily extended and adjust how data is stored for distributions
- Modified distribution to take a time range so we can change between lookups (i.e. not always normalized age)
- Additional optimization step to LUT generation (controlled with fx.NiagaraStateless.Distribution.ReduceLUTTimeRange) to reduce the time ranges if sampled keys are identical, this improves resolution of the LUT where it counts

#rnx

diff --git a/Engine/Plugins/FX/Niagara/Shaders/Private/Stateless/Modules/NiagaraStatelessModule_CameraOffset.ush b/Engine/Plugins/FX/Niagara/Shaders/Private/Stateless/Modules/NiagaraStatelessModule_CameraOffset.ush
--- a/Engine/Plugins/FX/Niagara/Shaders/Private/Stateless/Modules/NiagaraStatelessModule_CameraOffset.ush
+++ b/Engine/Plugins/FX/Niagara/Shaders/Private/Stateless/Modules/NiagaraStatelessModule_CameraOffset.ush
@@ -2,7 +2,7 @@
 
 #include "../NiagaraStatelessCommon.ush"
 
-uint3 CameraOffset_Distribution;
+FNiagaraStatelessBuiltDistributionType CameraOffset_Distribution;
 
 void CameraOffset_Simulate(inout FStatelessParticle Particle)
 {
diff --git a/Engine/Plugins/FX/Niagara/Shaders/Private/Stateless/Modules/NiagaraStatelessModule_DynamicMaterialParameters.ush b/Engine/Plugins/FX/Niagara/Shaders/Private/Stateless/Modules/NiagaraStatelessModule_DynamicMaterialParameters.ush
--- a/Engine/Plugins/FX/Niagara/Shaders/Private/Stateless/Modules/NiagaraStatelessModule_DynamicMaterialParameters.ush
+++ b/Engine/Plugins/FX/Niagara/Shaders/Private/Stateless/Modules/NiagaraStatelessModule_DynamicMaterialParameters.ush
@@ -3,22 +3,22 @@
 #include "../NiagaraStatelessCommon.ush"
 
 uint	DynamicMaterialParameters_ChannelMask;
-uint3	DynamicMaterialParameters_Parameter0X;
-uint3	DynamicMaterialParameters_Parameter0Y;
-uint3	DynamicMaterialParameters_Parameter0Z;
-uint3	DynamicMaterialParameters_Parameter0W;
-uint3	DynamicMaterialParameters_Parameter1X;
-uint3	DynamicMaterialParameters_Parameter1Y;
-uint3	DynamicMaterialParameters_Parameter1Z;
-uint3	DynamicMaterialParameters_Parameter1W;
-uint3	DynamicMaterialParameters_Parameter2X;
-uint3	DynamicMaterialParameters_Parameter2Y;
-uint3	DynamicMaterialParameters_Parameter2Z;
-uint3	DynamicMaterialParameters_Parameter2W;
-uint3	DynamicMaterialParameters_Parameter3X;
-uint3	DynamicMaterialParameters_Parameter3Y;
-uint3	DynamicMaterialParameters_Parameter3Z;
-uint3	DynamicMaterialParameters_Parameter3W;
+FNiagaraStatelessBuiltDistributionType	DynamicMaterialParameters_Parameter0X;
+FNiagaraStatelessBuiltDistributionType	DynamicMaterialParameters_Parameter0Y;
+FNiagaraStatelessBuiltDistributionType	DynamicMaterialParameters_Parameter0Z;
+FNiagaraStatelessBuiltDistributionType	DynamicMaterialParameters_Parameter0W;
+FNiagaraStatelessBuiltDistributionType	DynamicMaterialParameters_Parameter1X;
+FNiagaraStatelessBuiltDistributionType	DynamicMaterialParameters_Parameter1Y;
+FNiagaraStatelessBuiltDistributionType	DynamicMaterialParameters_Parameter1Z;
+FNiagaraStatelessBuiltDistributionType	DynamicMaterialParameters_Parameter1W;
+FNiagaraStatelessBuiltDistributionType	DynamicMaterialParameters_Parameter2X;
+FNiagaraStatelessBuiltDistributionType	DynamicMaterialParameters_Parameter2Y;
+FNiagaraStatelessBuiltDistributionType	DynamicMaterialParameters_Parameter2Z;
+FNiagaraStatelessBuiltDistributionType	DynamicMaterialParameters_Parameter2W;
+FNiagaraStatelessBuiltDistributionType	DynamicMaterialParameters_Parameter3X;
+FNiagaraStatelessBuiltDistributionType	DynamicMaterialParameters_Parameter3Y;
+FNiagaraStatelessBuiltDistributionType	DynamicMaterialParameters_Parameter3Z;
+FNiagaraStatelessBuiltDistributionType	DynamicMaterialParameters_Parameter3W;
 	
 void DynamicMaterialParameters_Simulate(inout FStatelessParticle Particle)
 {
diff --git a/Engine/Plugins/FX/Niagara/Shaders/Private/Stateless/Modules/NiagaraStatelessModule_InitializeParticle.ush b/Engine/Plugins/FX/Niagara/Shaders/Private/Stateless/Modules/NiagaraStatelessModule_InitializeParticle.ush
--- a/Engine/Plugins/FX/Niagara/Shaders/Private/Stateless/Modules/NiagaraStatelessModule_InitializeParticle.ush
+++ b/Engine/Plugins/FX/Niagara/Shaders/Private/Stateless/Modules/NiagaraStatelessModule_InitializeParticle.ush
@@ -6,7 +6,7 @@ static const uint EInitializeParticleModuleFlag_UniformSpriteSize	= 1 << 0;
 static const uint EInitializeParticleModuleFlag_UniformMeshScale	= 1 << 1;
 
 uint	InitializeParticle_ModuleFlags;
-uint3	InitializeParticle_InitialPosition;
+FNiagaraStatelessBuiltDistributionType	InitializeParticle_InitialPosition;
 float4	InitializeParticle_ColorScale;
 float4	InitializeParticle_ColorBias;
 float2	InitializeParticle_SpriteSizeScale;
diff --git a/Engine/Plugins/FX/Niagara/Shaders/Private/Stateless/Modules/NiagaraStatelessModule_ScaleColor.ush b/Engine/Plugins/FX/Niagara/Shaders/Private/Stateless/Modules/NiagaraStatelessModule_ScaleColor.ush
--- a/Engine/Plugins/FX/Niagara/Shaders/Private/Stateless/Modules/NiagaraStatelessModule_ScaleColor.ush
+++ b/Engine/Plugins/FX/Niagara/Shaders/Private/Stateless/Modules/NiagaraStatelessModule_ScaleColor.ush
@@ -2,7 +2,7 @@
 
 #include "../NiagaraStatelessCommon.ush"
 
-uint3	ScaleColor_Distribution;
+FNiagaraStatelessBuiltDistributionType	ScaleColor_Distribution;
 
 void ScaleColor_Simulate(inout FStatelessParticle Particle)
 {
diff --git a/Engine/Plugins/FX/Niagara/Shaders/Private/Stateless/Modules/NiagaraStatelessModule_ScaleMeshSize.ush b/Engine/Plugins/FX/Niagara/Shaders/Private/Stateless/Modules/NiagaraStatelessModule_ScaleMeshSize.ush
--- a/Engine/Plugins/FX/Niagara/Shaders/Private/Stateless/Modules/NiagaraStatelessModule_ScaleMeshSize.ush
+++ b/Engine/Plugins/FX/Niagara/Shaders/Private/Stateless/Modules/NiagaraStatelessModule_ScaleMeshSize.ush
@@ -2,7 +2,7 @@
 
 #include "../NiagaraStatelessCommon.ush"
 
-uint3	ScaleMeshSize_Distribution;
+FNiagaraStatelessBuiltDistributionType	ScaleMeshSize_Distribution;
 float3	ScaleMeshSize_CurveScale;
 int		ScaleMeshSize_CurveScaleOffset;
 	
diff --git a/Engine/Plugins/FX/Niagara/Shaders/Private/Stateless/Modules/NiagaraStatelessModule_ScaleSpriteSize.ush b/Engine/Plugins/FX/Niagara/Shaders/Private/Stateless/Modules/NiagaraStatelessModule_ScaleSpriteSize.ush
--- a/Engine/Plugins/FX/Niagara/Shaders/Private/Stateless/Modules/NiagaraStatelessModule_ScaleSpriteSize.ush
+++ b/Engine/Plugins/FX/Niagara/Shaders/Private/Stateless/Modules/NiagaraStatelessModule_ScaleSpriteSize.ush
@@ -2,7 +2,7 @@
 
 #include "../NiagaraStatelessCommon.ush"
 
-uint3	ScaleSpriteSize_Distribution;
+FNiagaraStatelessBuiltDistributionType	ScaleSpriteSize_Distribution;
 float2	ScaleSpriteSize_CurveScale;
 
 void ScaleSpriteSize_Simulate(inout FStatelessParticle Particle)
diff --git a/Engine/Plugins/FX/Niagara/Shaders/Private/Stateless/NiagaraStatelessCommon.ush b/Engine/Plugins/FX/Niagara/Shaders/Private/Stateless/NiagaraStatelessCommon.ush
--- a/Engine/Plugins/FX/Niagara/Shaders/Private/Stateless/NiagaraStatelessCommon.ush
+++ b/Engine/Plugins/FX/Niagara/Shaders/Private/Stateless/NiagaraStatelessCommon.ush
@@ -4,6 +4,8 @@
 #include "/Engine/Private/Common.ush"
 #include "/Engine/Private/Random.ush"
 
+#include "/Plugin/FX/Niagara/Shared/NiagaraStatelessDefinitions.h"
+#include "/Plugin/FX/Niagara/Shared/NiagaraStatelessBuiltDistribution.h"
 #include "/Plugin/FX/Niagara/Private/NiagaraQuaternionUtils.ush"
 
 ////////////////////////////////////////////////////////////////////////////////////
@@ -88,8 +90,7 @@ static uint4 GRandomSeedInternal;
 
 uint4 RandomUInt4(uint RandomSeedOffset)
 {
-	uint4 RandomSeed = GRandomSeedInternal;
-	RandomSeed.x += RandomSeedOffset;
+	uint4 RandomSeed = FNiagaraStatelessDefinitions::OffsetRandomSeedForCall(GRandomSeedInternal, RandomSeedOffset);
 	return Rand4DPCG32(RandomSeed);
 }
 uint   RandomUInt(uint RandomSeedOffset) { return RandomUInt4(RandomSeedOffset).x; }
diff --git a/Engine/Plugins/FX/Niagara/Shaders/Private/Stateless/NiagaraStatelessDistribution.ush b/Engine/Plugins/FX/Niagara/Shaders/Private/Stateless/NiagaraStatelessDistribution.ush
--- a/Engine/Plugins/FX/Niagara/Shaders/Private/Stateless/NiagaraStatelessDistribution.ush
+++ b/Engine/Plugins/FX/Niagara/Shaders/Private/Stateless/NiagaraStatelessDistribution.ush
@@ -1,48 +1,45 @@
 // Copyright Epic Games, Inc. All Rights Reserved.
 #pragma once
 
-#define StatelessDistributionFlag_Random	0x00000001
-#define StatelessDistributionFlag_Uniform	0x00000002
-#define StatelessDistributionFlag_Binding	0x00000004
-
 #define DefineStatelessDistributionSampler(HLSLTYPE, NAMETYPE) \
 	struct FStatelessDistributionSampler##NAMETYPE \
 	{ \
-		void Init(uint RandomSeedOffset, uint3 InParameters) \
+		void Init(uint RandomSeedOffset, FNiagaraStatelessBuiltDistributionType InBuildDistribution) \
 		{ \
-			Parameters = InParameters; \
+			BuildDistribution = InBuildDistribution; \
 			RandomOffset = 0.0f; \
-			if ( (Parameters.x & StatelessDistributionFlag_Random) != 0) \
+			if ( FNiagaraStatelessBuiltDistribution::IsRandom(BuildDistribution) != 0) \
 			{ \
 				RandomOffset = Random##NAMETYPE##(RandomSeedOffset); \
 			} \
 		} \
 		\
-		bool IsValid() { return Parameters.z > 0; } \
+		bool IsValid() { return FNiagaraStatelessBuiltDistribution::IsValid(BuildDistribution); } \
 		\
-		HLSLTYPE GetValue(float NormalizedAge) \
+		HLSLTYPE GetValue(float Time) \
 		{ \
-			if ((Parameters.x & StatelessDistributionFlag_Binding) != 0) \
+			const uint DataOffset = FNiagaraStatelessBuiltDistribution::GetDataOffset(BuildDistribution); \
+			if (FNiagaraStatelessBuiltDistribution::IsBinding(BuildDistribution)) \
 			{ \
-				return GetParameterBuffer##NAMETYPE##(Parameters.y, 0); \
+				return GetParameterBuffer##NAMETYPE##(DataOffset, 0); \
 			} \
-			else if ((Parameters.x & StatelessDistributionFlag_Random) != 0) \
+			else if (FNiagaraStatelessBuiltDistribution::IsRandom(BuildDistribution)) \
 			{ \
-				const HLSLTYPE Value0 = GetStatic##NAMETYPE##(Parameters.y, 0); \
-				const HLSLTYPE Value1 = GetStatic##NAMETYPE##(Parameters.y, 1); \
-				return StatelessLerp(Value0, Value1, RandomOffset, (Parameters.x & StatelessDistributionFlag_Uniform) != 0); \
+				const HLSLTYPE Value0 = GetStatic##NAMETYPE##(DataOffset, 0); \
+				const HLSLTYPE Value1 = GetStatic##NAMETYPE##(DataOffset, 1); \
+				return StatelessLerp(Value0, Value1, RandomOffset, FNiagaraStatelessBuiltDistribution::IsUniform(BuildDistribution)); \
 			} \
 			else \
 			{ \
-				const float Offset = NormalizedAge * float(Parameters.z); \
-				const HLSLTYPE Value0 = GetStatic##NAMETYPE##(Parameters.y, floor(Offset)); \
-				const HLSLTYPE Value1 = GetStatic##NAMETYPE##(Parameters.y, ceil(Offset)); \
+				const float Offset = FNiagaraStatelessBuiltDistribution::ConvertTimeToLookup(BuildDistribution, Time); \
+				const HLSLTYPE Value0 = GetStatic##NAMETYPE##(DataOffset, floor(Offset)); \
+				const HLSLTYPE Value1 = GetStatic##NAMETYPE##(DataOffset, ceil(Offset)); \
 				return lerp(Value0, Value1, frac(Offset)); \
 			} \
 		} \
 		\
-		uint3		Parameters; \
-		HLSLTYPE	RandomOffset; \
+		FNiagaraStatelessBuiltDistributionType		BuildDistribution; \
+		HLSLTYPE									RandomOffset; \
 	}
 
 DefineStatelessDistributionSampler(float,  Float);
diff --git a/Engine/Plugins/FX/Niagara/Shaders/Private/Stateless/NiagaraStatelessSimulationTemplate.ush b/Engine/Plugins/FX/Niagara/Shaders/Private/Stateless/NiagaraStatelessSimulationTemplate.ush
--- a/Engine/Plugins/FX/Niagara/Shaders/Private/Stateless/NiagaraStatelessSimulationTemplate.ush
+++ b/Engine/Plugins/FX/Niagara/Shaders/Private/Stateless/NiagaraStatelessSimulationTemplate.ush
@@ -97,7 +97,8 @@ void StatelessMain(uint3 GroupId : SV_GroupID, int GroupThreadIndex : SV_GroupIn
 	}
 	
 	// Initialize variables / determine lifetime / etc
-	GRandomSeedInternal = uint4(7123u, Common_RandomSeed, Particle.UniqueIndex * 3581u, 6405u);
+	uint ModuleSeedOffset = 0;
+	GRandomSeedInternal = FNiagaraStatelessDefinitions::MakeRandomSeed(Common_RandomSeed, Particle.UniqueIndex, ModuleSeedOffset++, 0);
 
 	Particle.Lifetime = RandomScaleBiasFloat(0, Common_LifetimeScaleBias.x, Common_LifetimeScaleBias.y);
 	if ( Particle.Lifetime <= 0.0f || Particle.Age >= Particle.Lifetime )
@@ -116,7 +117,7 @@ void StatelessMain(uint3 GroupId : SV_GroupID, int GroupThreadIndex : SV_GroupIn
 	
 	// Run Particle Simulate
 	#define PARTICLE_MODULE(NAME) \
-		GRandomSeedInternal.w += 6405u;\
+		GRandomSeedInternal = FNiagaraStatelessDefinitions::MakeRandomSeed(Common_RandomSeed, Particle.UniqueIndex, ModuleSeedOffset++, 0); \
 		##NAME##_Simulate(Particle);
 
 		PARTICLE_MODULES
diff --git a/Engine/Plugins/FX/Niagara/Shaders/Shared/NiagaraStatelessBuiltDistribution.h b/Engine/Plugins/FX/Niagara/Shaders/Shared/NiagaraStatelessBuiltDistribution.h
new file mode 100644
--- /dev/null
+++ b/Engine/Plugins/FX/Niagara/Shaders/Shared/NiagaraStatelessBuiltDistribution.h
@@ -0,0 +1,81 @@
+// Copyright Epic Games, Inc. All Rights Reserved.
+#pragma once
+
+#ifndef GPU_SIMULATION
+	#define uint				uint32
+	typedef FUintVector3		FNiagaraStatelessBuiltDistributionType;
+	typedef const FUintVector3&	FNiagaraStatelessBuiltDistributionTypeIn;
+#else
+	//typedef uint3			FNiagaraStatelessBuiltDistributionType;
+	#define FNiagaraStatelessBuiltDistributionType		uint3
+	#define FNiagaraStatelessBuiltDistributionTypeIn	uint3
+#endif
+
+#define ENiagaraStatelessBuiltDistributionFlag_Binding			0x10000000u
+#define ENiagaraStatelessBuiltDistributionFlag_Random			0x20000000u
+#define ENiagaraStatelessBuiltDistributionFlag_Uniform			0x40000000u
+#define ENiagaraStatelessBuiltDistributionFlag_TableLengthBits	9u				// 512 entries in LUT
+#define ENiagaraStatelessBuiltDistributionFlag_TableLengthShift	20u
+#define ENiagaraStatelessBuiltDistributionFlag_TableLengthMask	((1u << ENiagaraStatelessBuiltDistributionFlag_TableLengthBits) - 1u)
+#define ENiagaraStatelessBuiltDistributionFlag_DataOffsetBits	20u
+#define ENiagaraStatelessBuiltDistributionFlag_DataOffsetShift	0u
+#define ENiagaraStatelessBuiltDistributionFlag_DataOffsetMask	((1u << ENiagaraStatelessBuiltDistributionFlag_DataOffsetBits) - 1u)
+
+struct FNiagaraStatelessBuiltDistribution
+{
+	static bool IsValid(FNiagaraStatelessBuiltDistributionTypeIn BuiltData) { return BuiltData[2] != 0; }
+	static bool IsBinding(FNiagaraStatelessBuiltDistributionTypeIn BuiltData) { return (BuiltData[0] & ENiagaraStatelessBuiltDistributionFlag_Binding) != 0; }
+	static bool IsRandom(FNiagaraStatelessBuiltDistributionTypeIn BuiltData) { return (BuiltData[0] & ENiagaraStatelessBuiltDistributionFlag_Random) != 0; }
+	static bool IsUniform(FNiagaraStatelessBuiltDistributionTypeIn BuiltData) { return (BuiltData[0] & ENiagaraStatelessBuiltDistributionFlag_Uniform) != 0; }
+
+	static uint GetDataOffset(FNiagaraStatelessBuiltDistributionTypeIn BuiltData) { return (BuiltData[0] >> ENiagaraStatelessBuiltDistributionFlag_DataOffsetShift) & ENiagaraStatelessBuiltDistributionFlag_DataOffsetMask; }
+	static uint GetTableLength(FNiagaraStatelessBuiltDistributionTypeIn BuiltData) { return (BuiltData[0] >> ENiagaraStatelessBuiltDistributionFlag_TableLengthShift) & ENiagaraStatelessBuiltDistributionFlag_TableLengthMask; }
+
+#ifdef GPU_SIMULATION
+	static float ConvertTimeToLookup(FNiagaraStatelessBuiltDistributionTypeIn BuiltData, float Time)
+	{
+		const float TimeBias	= asfloat(BuiltData[1]);
+		const float TimeScale	= asfloat(BuiltData[2]);
+		const float TableLength	= float(GetTableLength(BuiltData));
+		const float Offset		= clamp((Time - TimeBias) * TimeScale, 0.0f, TableLength);
+		return Offset;
+	}
+#else
+	static float ConvertTimeToLookup(FNiagaraStatelessBuiltDistributionTypeIn BuiltData, float Time)
+	{
+		const float TimeBias	= reinterpret_cast<const float&>(BuiltData[1]);
+		const float TimeScale	= reinterpret_cast<const float&>(BuiltData[2]);
+		const float TableLength	= float(GetTableLength(BuiltData));
+		const float Offset		= FMath::Clamp((Time - TimeBias) * TimeScale, 0.0f, TableLength);
+		return Offset;
+	}
+	
+	static FNiagaraStatelessBuiltDistributionType GetDefault() { return FNiagaraStatelessBuiltDistributionType::ZeroValue; }
+
+	static void SetIsBinding(FNiagaraStatelessBuiltDistributionType& BuiltData) { BuiltData[0] |= ENiagaraStatelessBuiltDistributionFlag_Binding; }
+	static void SetIsRandom(FNiagaraStatelessBuiltDistributionType& BuiltData) { BuiltData[0] |= ENiagaraStatelessBuiltDistributionFlag_Random; }
+	static void SetIsUniform(FNiagaraStatelessBuiltDistributionType& BuiltData) { BuiltData[0] |= ENiagaraStatelessBuiltDistributionFlag_Uniform; }
+
+	static void SetLookupParameters(FNiagaraStatelessBuiltDistributionType& BuiltData, uint DataOffset, uint TableLength, const FVector2f& TimeRange)
+	{
+		const uint TableLengthMinusOne = TableLength - 1;
+
+		check(DataOffset <= ENiagaraStatelessBuiltDistributionFlag_DataOffsetMask);
+		check(TableLengthMinusOne <= ENiagaraStatelessBuiltDistributionFlag_TableLengthMask);
+
+		BuiltData[0] &= ~(ENiagaraStatelessBuiltDistributionFlag_DataOffsetMask << ENiagaraStatelessBuiltDistributionFlag_DataOffsetShift);
+		BuiltData[0] |= DataOffset << ENiagaraStatelessBuiltDistributionFlag_DataOffsetShift;
+
+		BuiltData[0] &= ~(ENiagaraStatelessBuiltDistributionFlag_TableLengthMask << ENiagaraStatelessBuiltDistributionFlag_TableLengthShift);
+		BuiltData[0] |= TableLengthMinusOne << ENiagaraStatelessBuiltDistributionFlag_TableLengthShift;
+
+		const float TimeDuration = TimeRange.Y - TimeRange.X;
+		reinterpret_cast<float&>(BuiltData[1]) = TimeRange.X;
+		reinterpret_cast<float&>(BuiltData[2]) = TableLengthMinusOne > 0 && TimeDuration > 0.0f ? float(TableLengthMinusOne) / TimeDuration : 0.0f;
+	}
+#endif
+};
+
+#ifndef GPU_SIMULATION
+	#undef uint
+#endif
diff --git a/Engine/Plugins/FX/Niagara/Shaders/Shared/NiagaraStatelessDefinitions.h b/Engine/Plugins/FX/Niagara/Shaders/Shared/NiagaraStatelessDefinitions.h
new file mode 100644
--- /dev/null
+++ b/Engine/Plugins/FX/Niagara/Shaders/Shared/NiagaraStatelessDefinitions.h
@@ -0,0 +1,19 @@
+// Copyright Epic Games, Inc. All Rights Reserved.
+#pragma once
+
+#ifndef GPU_SIMULATION
+	#define uint	uint32
+	#define uint4	FUintVector4
+#endif
+
+// Core random functionality
+struct FNiagaraStatelessDefinitions
+{
+	static uint4 MakeRandomSeed(uint EmitterSeed, uint UniqueIndex, uint ModuleOffset, uint CallOffset) { return uint4(7123u + CallOffset, EmitterSeed, 3581u + UniqueIndex, 6405u + ModuleOffset); }
+	static uint4 OffsetRandomSeedForCall(uint4 Seed, uint CallOffset) { Seed[0] += CallOffset; return Seed; }
+};
+
+#ifndef GPU_SIMULATION
+	#undef uint
+	#undef uint4
+#endif
diff --git a/Engine/Plugins/FX/Niagara/Source/Niagara/Internal/Stateless/NiagaraStatelessCommon.h b/Engine/Plugins/FX/Niagara/Source/Niagara/Internal/Stateless/NiagaraStatelessCommon.h
--- a/Engine/Plugins/FX/Niagara/Source/Niagara/Internal/Stateless/NiagaraStatelessCommon.h
+++ b/Engine/Plugins/FX/Niagara/Source/Niagara/Internal/Stateless/NiagaraStatelessCommon.h
@@ -3,6 +3,7 @@
 #pragma once
 
 #include "NiagaraCommon.h"
+#include "NiagaraStatelessDefinitions.h"
 
 UENUM(meta = (Bitflags, UseEnumValuesAsMaskValuesInEditor = "true"))
 enum class ENiagaraStatelessFeatureMask : uint8
@@ -42,14 +43,6 @@ using FNiagaraStatelessRangeVector3	= FNiagaraStatelessRange<FVector3f>;
 using FNiagaraStatelessRangeVector4	= FNiagaraStatelessRange<FVector4f>;
 using FNiagaraStatelessRangeColor	= FNiagaraStatelessRange<FLinearColor>;
 
-enum class ENiagaraStatelessBuiltDistributionFlag
-{
-	Random  = 0x00000001,	// StatelessDistributionFlag_Random
-	Uniform = 0x00000002,	// StatelessDistributionFlag_Uniform
-	Binding = 0x00000004,	// StatelessDistributionFlag_Binding
-};
-ENUM_CLASS_FLAGS(ENiagaraStatelessBuiltDistributionFlag);
-
 struct FNiagaraStatelessGlobals
 {
 	FNiagaraVariableBase	CameraOffsetVariable;
diff --git a/Engine/Plugins/FX/Niagara/Source/Niagara/Internal/Stateless/NiagaraStatelessDistribution.h b/Engine/Plugins/FX/Niagara/Source/Niagara/Internal/Stateless/NiagaraStatelessDistribution.h
--- a/Engine/Plugins/FX/Niagara/Source/Niagara/Internal/Stateless/NiagaraStatelessDistribution.h
+++ b/Engine/Plugins/FX/Niagara/Source/Niagara/Internal/Stateless/NiagaraStatelessDistribution.h
@@ -47,7 +47,7 @@ struct FNiagaraDistributionBase
 	TArray<FRichCurve> ChannelCurves;
 
 	UPROPERTY(EditAnywhere, Category = "Parameters")
-	int32 MaxLutSampleCount = 64;
+	int32 MaxLutSampleCount = 128;
 
 	bool operator==(const FNiagaraDistributionBase& Other) const
 	{
@@ -217,6 +217,9 @@ struct FNiagaraDistributionFloat : public FNiagaraDistributionBase
 	UPROPERTY(EditAnywhere, Category = "Parameters")
 	TArray<float> Values;
 
+	UPROPERTY(EditAnywhere, Category = "Parameters")
+	FVector2f ValuesTimeRange = FVector2f(0.0f, 1.0f);
+
 	NIAGARA_API void InitConstant(float Value);
 	NIAGARA_API FNiagaraStatelessRangeFloat CalculateRange(const float Default = 0.0f) const;
 
@@ -243,6 +246,9 @@ struct FNiagaraDistributionVector2 : public FNiagaraDistributionBase
 	UPROPERTY(EditAnywhere, Category = "Parameters")
 	TArray<FVector2f> Values;
 
+	UPROPERTY(EditAnywhere, Category = "Parameters")
+	FVector2f ValuesTimeRange = FVector2f(0.0f, 1.0f);
+
 	NIAGARA_API void InitConstant(const float Value);
 	NIAGARA_API void InitConstant(const FVector2f& Value);
 	NIAGARA_API FNiagaraStatelessRangeVector2 CalculateRange(const FVector2f& Default = FVector2f::ZeroVector) const;
@@ -266,6 +272,9 @@ struct FNiagaraDistributionVector3 : public FNiagaraDistributionBase
 	UPROPERTY(EditAnywhere, Category = "Parameters")
 	TArray<FVector3f> Values;
 
+	UPROPERTY(EditAnywhere, Category = "Parameters")
+	FVector2f ValuesTimeRange = FVector2f(0.0f, 1.0f);
+
 	NIAGARA_API void InitConstant(const float Value);
 	NIAGARA_API void InitConstant(const FVector3f& Value);
 	NIAGARA_API FNiagaraStatelessRangeVector3 CalculateRange(const FVector3f& Default = FVector3f::ZeroVector) const;
@@ -302,6 +311,9 @@ struct FNiagaraDistributionColor : public FNiagaraDistributionBase
 	UPROPERTY(EditAnywhere, Category = "Parameters")
 	TArray<FLinearColor> Values;
 
+	UPROPERTY(EditAnywhere, Category = "Parameters")
+	FVector2f ValuesTimeRange = FVector2f(0.0f, 1.0f);
+
 	NIAGARA_API void InitConstant(const FLinearColor& Value);
 	NIAGARA_API FNiagaraStatelessRangeColor CalculateRange(const FLinearColor& Default = FLinearColor::White) const;
 
diff --git a/Engine/Plugins/FX/Niagara/Source/Niagara/Internal/Stateless/NiagaraStatelessEmitterDataBuildContext.h b/Engine/Plugins/FX/Niagara/Source/Niagara/Internal/Stateless/NiagaraStatelessEmitterDataBuildContext.h
--- a/Engine/Plugins/FX/Niagara/Source/Niagara/Internal/Stateless/NiagaraStatelessEmitterDataBuildContext.h
+++ b/Engine/Plugins/FX/Niagara/Source/Niagara/Internal/Stateless/NiagaraStatelessEmitterDataBuildContext.h
@@ -4,6 +4,7 @@
 
 #include "NiagaraStatelessCommon.h"
 #include "NiagaraStatelessDistribution.h"
+#include "NiagaraStatelessBuiltDistribution.h"
 
 struct FNiagaraDataSetCompiledData;
 struct FNiagaraParameterBinding;
@@ -68,51 +69,49 @@ public:
 
 	// Adds an distribution into the LUT if enabled
 	template<typename TType>
-	FUintVector3 AddDistribution(ENiagaraDistributionMode Mode, TConstArrayView<TType> Values) const
+	FNiagaraStatelessBuiltDistributionType AddDistribution(ENiagaraDistributionMode Mode, TConstArrayView<TType> Values, const FVector2f& TimeRange) const
 	{
 		using namespace NiagaraStateless;
 
-		FUintVector3 Parameters = FUintVector3::ZeroValue;
+		FNiagaraStatelessBuiltDistributionType BuiltDistribution = FNiagaraStatelessBuiltDistribution::GetDefault();
 		if (Values.Num() > 0)
 		{
 			switch (Mode)
 			{
 				case ENiagaraDistributionMode::Binding:				checkNoEntry(); break;
-				case ENiagaraDistributionMode::UniformConstant:		Parameters.X = uint32(ENiagaraStatelessBuiltDistributionFlag::Random | ENiagaraStatelessBuiltDistributionFlag::Uniform); break;
-				case ENiagaraDistributionMode::NonUniformConstant:	Parameters.X = uint32(ENiagaraStatelessBuiltDistributionFlag::Random); break;
-				case ENiagaraDistributionMode::UniformRange:		Parameters.X = uint32(ENiagaraStatelessBuiltDistributionFlag::Random | ENiagaraStatelessBuiltDistributionFlag::Uniform); break;
-				case ENiagaraDistributionMode::NonUniformRange:		Parameters.X = uint32(ENiagaraStatelessBuiltDistributionFlag::Random); break;
-				case ENiagaraDistributionMode::UniformCurve:		Parameters.X = uint32(ENiagaraStatelessBuiltDistributionFlag::Uniform); break;
-				case ENiagaraDistributionMode::NonUniformCurve:		Parameters.X = 0; break;
+				case ENiagaraDistributionMode::UniformConstant:		FNiagaraStatelessBuiltDistribution::SetIsRandom(BuiltDistribution); FNiagaraStatelessBuiltDistribution::SetIsUniform(BuiltDistribution); break;
+				case ENiagaraDistributionMode::NonUniformConstant:	FNiagaraStatelessBuiltDistribution::SetIsRandom(BuiltDistribution); break;
+				case ENiagaraDistributionMode::UniformRange:		FNiagaraStatelessBuiltDistribution::SetIsRandom(BuiltDistribution); FNiagaraStatelessBuiltDistribution::SetIsUniform(BuiltDistribution); break;
+				case ENiagaraDistributionMode::NonUniformRange:		FNiagaraStatelessBuiltDistribution::SetIsRandom(BuiltDistribution); break;
+				case ENiagaraDistributionMode::UniformCurve:		FNiagaraStatelessBuiltDistribution::SetIsUniform(BuiltDistribution); break;
+				case ENiagaraDistributionMode::NonUniformCurve:		break;
 				default:											checkNoEntry(); break;
 			}
 
-			Parameters.Y = AddStaticData(Values);
-			Parameters.Z = Values.Num() - 1;
+			FNiagaraStatelessBuiltDistribution::SetLookupParameters(BuiltDistribution, AddStaticData(Values), Values.Num(), TimeRange);
 		}
-		return Parameters;
+		return BuiltDistribution;
 	}
 
 	// Adds a distribution into the LUT if enabled and returns the packed information to send to the shader
 	template<typename TDistribution>
-	FUintVector3 AddDistribution(const TDistribution& Distribution) const
+	FNiagaraStatelessBuiltDistributionType AddDistribution(const TDistribution& Distribution) const
 	{
-		FUintVector3 Parameters = FUintVector3::ZeroValue;
+		FNiagaraStatelessBuiltDistributionType BuiltDistribution = FNiagaraStatelessBuiltDistribution::GetDefault();
 		if (Distribution.Mode == ENiagaraDistributionMode::Binding)
 		{
 			const int32 ParameterOffset = AddRendererBinding(Distribution.ParameterBinding);
 			if (ParameterOffset >= 0)
 			{
-				Parameters.X = uint32(ENiagaraStatelessBuiltDistributionFlag::Binding);
-				Parameters.Y = ParameterOffset;
-				Parameters.Z = 1.0f;
+				FNiagaraStatelessBuiltDistribution::SetIsBinding(BuiltDistribution);
+				FNiagaraStatelessBuiltDistribution::SetLookupParameters(BuiltDistribution, ParameterOffset, 1, FVector2f(0.0f, 1.0f));
 			}
 		}
 		else
 		{
-			Parameters = AddDistribution(Distribution.Mode, MakeArrayView(Distribution.Values));
+			BuiltDistribution = AddDistribution(Distribution.Mode, MakeArrayView(Distribution.Values), Distribution.ValuesTimeRange);
 		}
-		return Parameters;
+		return BuiltDistribution;
 	}
 
 	template<typename TRange, typename TDistribution, typename TDefaultValue>
@@ -153,7 +152,7 @@ private:
 	NiagaraStateless::FParticleSimulationExecData*	ParticleExecData = nullptr;
 
 	int32											ModuleBuiltDataOffset = 0;
-	int32											RandomSeedOffest = 1;
+	int32											RandomSeedOffest = 0;
 
 	struct FTransientObject
 	{
diff --git a/Engine/Plugins/FX/Niagara/Source/Niagara/Internal/Stateless/NiagaraStatelessParticleSimContext.h b/Engine/Plugins/FX/Niagara/Source/Niagara/Internal/Stateless/NiagaraStatelessParticleSimContext.h
--- a/Engine/Plugins/FX/Niagara/Source/Niagara/Internal/Stateless/NiagaraStatelessParticleSimContext.h
+++ b/Engine/Plugins/FX/Niagara/Source/Niagara/Internal/Stateless/NiagaraStatelessParticleSimContext.h
@@ -3,6 +3,7 @@
 #pragma once
 
 #include "NiagaraStatelessCommon.h"
+#include "NiagaraStatelessBuiltDistribution.h"
 
 class FNiagaraDataBuffer;
 struct FNiagaraStatelessEmitterData;
@@ -217,41 +218,42 @@ private:
 template<typename TType>
 struct FStatelessDistributionSampler
 {
-	explicit FStatelessDistributionSampler(const FParticleSimulationContext& ParticleSimulationContext, const FUintVector3& InParameters, int32 iInstance, uint32 RandomSeedOffset)
-		: Parameters(InParameters)
+	explicit FStatelessDistributionSampler(const FParticleSimulationContext& ParticleSimulationContext, const FNiagaraStatelessBuiltDistributionType& InBuiltDistribution, int32 iInstance, uint32 RandomSeedOffset)
+		: BuiltDistribution(InBuiltDistribution)
 	{
-		if ((Parameters.X & uint32(ENiagaraStatelessBuiltDistributionFlag::Random)) != 0)
+		if (FNiagaraStatelessBuiltDistribution::IsRandom(BuiltDistribution))
 		{
 			RandomOffset = ParticleSimulationContext.TRandomFloat<TType>(iInstance, RandomSeedOffset);
 		}
 	}
 
-	bool IsValid() const { return Parameters.Z > 0; }
+	bool IsValid() const { return FNiagaraStatelessBuiltDistribution::IsValid(BuiltDistribution); }
 
-	TType GetValue(const FParticleSimulationContext& ParticleSimulationContext, float NormalizedAge) const
+	TType GetValue(const FParticleSimulationContext& ParticleSimulationContext, float Time) const
 	{
-		if ((Parameters.X & uint32(ENiagaraStatelessBuiltDistributionFlag::Binding)) != 0)
+		const uint32 DataOffset = FNiagaraStatelessBuiltDistribution::GetDataOffset(BuiltDistribution);
+		if (FNiagaraStatelessBuiltDistribution::IsBinding(BuiltDistribution))
 		{
-			return ParticleSimulationContext.GetParameterBufferFloat<TType>(Parameters.Y, 0);
+			return ParticleSimulationContext.GetParameterBufferFloat<TType>(DataOffset, 0);
 		}
-		else if ((Parameters.X & uint32(ENiagaraStatelessBuiltDistributionFlag::Random)) != 0)
+		else if (FNiagaraStatelessBuiltDistribution::IsRandom(BuiltDistribution))
 		{
 			//-OPT: Could move into constructor
-			const TType	Value0 = ParticleSimulationContext.GetStaticFloat<TType>(Parameters.Y, 0);
-			const TType	Value1 = ParticleSimulationContext.GetStaticFloat<TType>(Parameters.Y, 1);
-			return ParticleSimulationContext.Lerp(Value0, Value1, RandomOffset, (Parameters.X & uint32(ENiagaraStatelessBuiltDistributionFlag::Uniform)) != 0);
+			const TType	Value0 = ParticleSimulationContext.GetStaticFloat<TType>(DataOffset, 0);
+			const TType	Value1 = ParticleSimulationContext.GetStaticFloat<TType>(DataOffset, 1);
+			return ParticleSimulationContext.Lerp(Value0, Value1, RandomOffset, FNiagaraStatelessBuiltDistribution::IsUniform(BuiltDistribution));
 		}
 		else
 		{
-			const float Offset = NormalizedAge * float(Parameters.Z);
-			const TType	Value0 = ParticleSimulationContext.GetStaticFloat<TType>(Parameters.Y, FMath::FloorToInt(Offset));
-			const TType	Value1 = ParticleSimulationContext.GetStaticFloat<TType>(Parameters.Y, FMath::CeilToInt(Offset));
+			const float Offset = FNiagaraStatelessBuiltDistribution::ConvertTimeToLookup(BuiltDistribution, Time);
+			const TType	Value0 = ParticleSimulationContext.GetStaticFloat<TType>(DataOffset, FMath::FloorToInt(Offset));
+			const TType	Value1 = ParticleSimulationContext.GetStaticFloat<TType>(DataOffset, FMath::CeilToInt(Offset));
 			return FMath::Lerp(Value0, Value1, FMath::Fractional(Offset));
 		}
 	}
 
-	const FUintVector3	Parameters;
-	TType				RandomOffset = {};
+	const FNiagaraStatelessBuiltDistributionType	BuiltDistribution;
+	TType											RandomOffset = {};
 };
 
 } //namespace NiagaraStateless
diff --git a/Engine/Plugins/FX/Niagara/Source/Niagara/Private/Stateless/NiagaraStatelessDistribution.cpp b/Engine/Plugins/FX/Niagara/Source/Niagara/Private/Stateless/NiagaraStatelessDistribution.cpp
--- a/Engine/Plugins/FX/Niagara/Source/Niagara/Private/Stateless/NiagaraStatelessDistribution.cpp
+++ b/Engine/Plugins/FX/Niagara/Source/Niagara/Private/Stateless/NiagaraStatelessDistribution.cpp
@@ -14,15 +14,26 @@ namespace NiagaraStatelessDistributionPrivate
 		ECVF_Default
 	);
 
-	TArray<float> CurvesToLUT(TArrayView<const FRichCurve> Curves, int32 NumSamples)
+	bool GReduceLUTTimeRange = true;
+	FAutoConsoleVariableRef CVarNiagaraStatelessDistributionReduceLUTTimeRange(
+		TEXT("fx.NiagaraStateless.Distribution.ReduceLUTTimeRange"),
+		GReduceLUTTimeRange,
+		TEXT("When LUT optimization is enabled we will try to reduce the start / end time if we get duplicate keys, this can improve resolution in the area that matters."),
+		ECVF_Default
+	);
+
+	TArray<float> CurvesToLUT(TArrayView<const FRichCurve> Curves, int32 NumSamples, const FVector2f& TimeRange)
 	{
 		const int32 NumChannels = Curves.Num();
 
+		const float Duration = TimeRange.Y - TimeRange.X;
+		const float SampleToDuration = Duration > 0.0f ? Duration / float(NumSamples - 1) : 0.0f;
+
 		TArray<float> LUT;
 		LUT.SetNumUninitialized(NumChannels * NumSamples);
 		for (int32 iSample=0; iSample < NumSamples; ++iSample)
 		{
-			const float Time = float(iSample) / float(NumSamples - 1);
+			const float Time = (float(iSample) * SampleToDuration) + TimeRange.X;
 			for (int32 iChannel=0; iChannel < NumChannels; ++iChannel)
 			{
 				LUT[(iSample * NumChannels) + iChannel] = Curves[iChannel].Eval(Time);
@@ -31,6 +42,69 @@ namespace NiagaraStatelessDistributionPrivate
 		return LUT;
 	}
 
+	FVector2f CurvesFindTimeRange(TArrayView<const FRichCurve> Curves, int32 NumSamples)
+	{
+		const int32 NumChannels = Curves.Num();
+		check(NumChannels > 0);
+
+		// Find time range based on keys
+		FVector2f TimeRange = FVector2f(0.0f, 1.0f);
+		Curves[0].GetTimeRange(TimeRange.X, TimeRange.Y);
+		for (int32 iChannel = 1; iChannel < NumChannels; ++iChannel)
+		{
+			FVector2f ChannelTimeRange = TimeRange;
+			Curves[iChannel].GetTimeRange(ChannelTimeRange.X, ChannelTimeRange.Y);
+			TimeRange.X = FMath::Min(TimeRange.X, ChannelTimeRange.X);
+			TimeRange.Y = FMath::Max(TimeRange.Y, ChannelTimeRange.Y);
+		}
+
+		// Reduce start / end times if the sampled results match
+		if (GReduceLUTTimeRange)
+		{
+			TArray<float> LUT = CurvesToLUT(Curves, NumSamples, TimeRange);
+
+			auto DoKeysMatch =
+				[&LUT, &NumChannels](int iCurrKey, int iNextKey)
+				{
+					for (int32 iChannel = 0; iChannel < NumChannels; ++iChannel)
+					{
+						if (!FMath::IsNearlyEqual(LUT[(iCurrKey * NumChannels) + iChannel], LUT[(iNextKey * NumChannels) + iChannel]))
+						{
+							return false;
+						}
+					}
+					return true;
+				};
+
+			int32 iFirstKey = 0;
+			while (iFirstKey < NumSamples - 1 && DoKeysMatch(iFirstKey, iFirstKey + 1))
+			{
+				++iFirstKey;
+			}
+
+			int32 iLastKey = NumSamples - 1;
+			while (iLastKey > 0 && DoKeysMatch(iLastKey, iLastKey - 1))
+			{
+				--iLastKey;
+			}
+
+			if (iFirstKey >= iLastKey)
+			{
+				TimeRange.Y = TimeRange.X;
+			}
+			else
+			{
+				const float StartTime = TimeRange.X;
+				const float Duration = TimeRange.Y - TimeRange.X;
+				const float SampleToDuration = Duration > 0.0f ? Duration / float(NumSamples - 1) : 0.0f;
+				TimeRange.X = StartTime + (float(iFirstKey) * SampleToDuration);
+				TimeRange.Y = StartTime + (float(iLastKey) * SampleToDuration);
+			}
+		}
+
+		return TimeRange;
+	}
+
 	bool AreLUTsAlmostEqual(TArrayView<float> Lhs, TArrayView<float> Rhs, int32 NumChannels, float ErrorThreshold = 0.01f)
 	{
 		const int32 LhsNumSamples = Lhs.Num() / NumChannels;
@@ -64,14 +138,21 @@ namespace NiagaraStatelessDistributionPrivate
 		return true;
 	}
 
-	TArray<float> CurvesToOptimizedLUT(TArrayView<const FRichCurve> Curves, int32 MaxLutSampleCount)
+	TArray<float> CurvesToOptimizedLUT(TArrayView<const FRichCurve> Curves, int32 MaxLutSampleCount, FVector2f& OutTimeRange)
 	{
-		TArray<float> LUT = CurvesToLUT(Curves, MaxLutSampleCount);
+		OutTimeRange = CurvesFindTimeRange(Curves, MaxLutSampleCount);
+		if (FMath::IsNearlyEqual(OutTimeRange.X, OutTimeRange.Y))
+		{
+			TArray<float> LUT = CurvesToLUT(Curves, 2, OutTimeRange);
+			return LUT;
+		}
+
+		TArray<float> LUT = CurvesToLUT(Curves, MaxLutSampleCount, OutTimeRange);
 		if (GOptimizeLUTs)
 		{
 			for (int32 iSamples = 2; iSamples < MaxLutSampleCount; ++iSamples)
 			{
-				TArray<float> NewLUT = CurvesToLUT(Curves, iSamples);
+				TArray<float> NewLUT = CurvesToLUT(Curves, iSamples, OutTimeRange);
 				if (AreLUTsAlmostEqual(NewLUT, LUT, Curves.Num()))
 				{
 					return NewLUT;
@@ -82,7 +163,7 @@ namespace NiagaraStatelessDistributionPrivate
 	}
 
 	template<typename FValueContainerSetNum, typename FValueAndChannelAccessor>
-	void UpdateDistributionValues(ENiagaraDistributionMode InMode, const TArray<float>& InChannelConstantsAndRanges, const TArray<FRichCurve>& InChannelCurves, int32 InChannelCount, FValueContainerSetNum InContainerNum, FValueAndChannelAccessor InValueAndChannelAccessor, int32 MaxLutSampleCount)
+	void UpdateDistributionValues(ENiagaraDistributionMode InMode, const TArray<float>& InChannelConstantsAndRanges, const TArray<FRichCurve>& InChannelCurves, int32 InChannelCount, FVector2f& OutTimeRange, FValueContainerSetNum InContainerNum, FValueAndChannelAccessor InValueAndChannelAccessor, int32 MaxLutSampleCount)
 	{
 		switch(InMode)
 		{
@@ -91,6 +172,7 @@ namespace NiagaraStatelessDistributionPrivate
 				{
 					// Note we set two values to simplify the GPU code
 					InContainerNum(2);
+					OutTimeRange = FVector2f(0.0f, 1.0f);
 					for (int32 ChannelIndex = 0; ChannelIndex < InChannelCount; ChannelIndex++)
 					{
 						InValueAndChannelAccessor(0, ChannelIndex) = InChannelConstantsAndRanges[0];
@@ -103,6 +185,7 @@ namespace NiagaraStatelessDistributionPrivate
 				{
 					// Note we set two values to simplify the GPU code
 					InContainerNum(2);
+					OutTimeRange = FVector2f(0.0f, 1.0f);
 					for (int32 ChannelIndex = 0; ChannelIndex < InChannelCount; ChannelIndex++)
 					{
 						InValueAndChannelAccessor(0, ChannelIndex) = InChannelConstantsAndRanges[ChannelIndex];
@@ -114,6 +197,7 @@ namespace NiagaraStatelessDistributionPrivate
 				if (InChannelConstantsAndRanges.Num() >= 2)
 				{
 					InContainerNum(2);
+					OutTimeRange = FVector2f(0.0f, 1.0f);
 					for (int32 ChannelIndex = 0; ChannelIndex < InChannelCount; ChannelIndex++)
 					{
 						InValueAndChannelAccessor(0, ChannelIndex) = InChannelConstantsAndRanges[0];
@@ -125,6 +209,7 @@ namespace NiagaraStatelessDistributionPrivate
 				if (InChannelConstantsAndRanges.Num() >= 2 * InChannelCount)
 				{
 					InContainerNum(2);
+					OutTimeRange = FVector2f(0.0f, 1.0f);
 					for (int32 ChannelIndex = 0; ChannelIndex < InChannelCount; ChannelIndex++)
 					{
 						InValueAndChannelAccessor(0, ChannelIndex) = InChannelConstantsAndRanges[ChannelIndex];
@@ -139,8 +224,7 @@ namespace NiagaraStatelessDistributionPrivate
 				if (InChannelCurves.Num() >= ExpectedChannels)
 				{
 					MaxLutSampleCount = FMath::Max(MaxLutSampleCount, 2);
-					//const TArray<float> LUT = CurvesToLUT(MakeArrayView(InChannelCurves.GetData(), ExpectedChannels), MaxLutSampleCount);
-					const TArray<float> LUT = CurvesToOptimizedLUT(MakeArrayView(InChannelCurves.GetData(), ExpectedChannels), MaxLutSampleCount);
+					const TArray<float> LUT = CurvesToOptimizedLUT(MakeArrayView(InChannelCurves.GetData(), ExpectedChannels), MaxLutSampleCount, OutTimeRange);
 					const int32 NumSamples = LUT.Num() / ExpectedChannels;
 
 					InContainerNum(NumSamples);
@@ -386,11 +470,13 @@ void FNiagaraDistributionBase::PostEditChangeProperty(UObject* OwnerObject, FPro
 
 void FNiagaraDistributionRangeFloat::UpdateValuesFromDistribution()
 {
+	FVector2f ValuesTimeRange;
 	NiagaraStatelessDistributionPrivate::UpdateDistributionValues(
 		Mode,
 		ChannelConstantsAndRanges,
 		ChannelCurves,
 		1,
+		ValuesTimeRange,
 		[this](int32 Num) {},
 		[this](int32 ValueIndex, int32 ChannelIndex) -> float& { return ValueIndex == 0 ? Min : Max; },
 		MaxLutSampleCount
@@ -399,11 +485,13 @@ void FNiagaraDistributionRangeFloat::UpdateValuesFromDistribution()
 
 void FNiagaraDistributionRangeVector2::UpdateValuesFromDistribution()
 {
+	FVector2f ValuesTimeRange;
 	NiagaraStatelessDistributionPrivate::UpdateDistributionValues(
 		Mode,
 		ChannelConstantsAndRanges,
 		ChannelCurves,
 		2,
+		ValuesTimeRange,
 		[this](int32 Num) {},
 		[this](int32 ValueIndex, int32 ChannelIndex) -> float& { return ValueIndex == 0 ? Min[ChannelIndex] : Max[ChannelIndex]; },
 		MaxLutSampleCount
@@ -412,11 +500,13 @@ void FNiagaraDistributionRangeVector2::UpdateValuesFromDistribution()
 
 void FNiagaraDistributionRangeVector3::UpdateValuesFromDistribution()
 {
+	FVector2f ValuesTimeRange;
 	NiagaraStatelessDistributionPrivate::UpdateDistributionValues(
 		Mode,
 		ChannelConstantsAndRanges,
 		ChannelCurves,
 		3,
+		ValuesTimeRange,
 		[this](int32 Num) {},
 		[this](int32 ValueIndex, int32 ChannelIndex) -> float& { return ValueIndex == 0 ? Min[ChannelIndex] : Max[ChannelIndex]; },
 		MaxLutSampleCount
@@ -425,11 +515,13 @@ void FNiagaraDistributionRangeVector3::UpdateValuesFromDistribution()
 
 void FNiagaraDistributionRangeColor::UpdateValuesFromDistribution()
 {
+	FVector2f ValuesTimeRange;
 	NiagaraStatelessDistributionPrivate::UpdateDistributionValues(
 		Mode,
 		ChannelConstantsAndRanges,
 		ChannelCurves,
 		4,
+		ValuesTimeRange,
 		[this](int32 Num) {},
 		[this](int32 ValueIndex, int32 ChannelIndex) -> float& { return ValueIndex == 0 ? Min.Component(ChannelIndex) : Max.Component(ChannelIndex); },
 		MaxLutSampleCount
@@ -443,6 +535,7 @@ void FNiagaraDistributionFloat::UpdateValuesFromDistribution()
 		ChannelConstantsAndRanges, 
 		ChannelCurves,
 		1,
+		ValuesTimeRange,
 		[this](int32 NumValues) { Values.SetNum(NumValues); },
 		[this](int32 ValueIndex, int32 ChannelIndex) -> float& { return Values[ValueIndex]; },
 		MaxLutSampleCount
@@ -456,6 +549,7 @@ void FNiagaraDistributionVector2::UpdateValuesFromDistribution()
 		ChannelConstantsAndRanges,
 		ChannelCurves,
 		2,
+		ValuesTimeRange,
 		[this](int32 NumValues) { Values.SetNum(NumValues); },
 		[this](int32 ValueIndex, int32 ChannelIndex) -> float& { return Values[ValueIndex][ChannelIndex]; },
 		MaxLutSampleCount
@@ -469,6 +563,7 @@ void FNiagaraDistributionVector3::UpdateValuesFromDistribution()
 		ChannelConstantsAndRanges,
 		ChannelCurves,
 		3,
+		ValuesTimeRange,
 		[this](int32 NumValues) { Values.SetNum(NumValues); },
 		[this](int32 ValueIndex, int32 ChannelIndex) -> float& { return Values[ValueIndex][ChannelIndex]; },
 		MaxLutSampleCount
@@ -482,6 +577,7 @@ void FNiagaraDistributionColor::UpdateValuesFromDistribution()
 		ChannelConstantsAndRanges,
 		ChannelCurves,
 		4,
+		ValuesTimeRange,
 		[this](int32 NumValues) { Values.SetNum(NumValues); },
 		[this](int32 ValueIndex, int32 ChannelIndex) -> float& { return Values[ValueIndex].Component(ChannelIndex); },
 		MaxLutSampleCount
diff --git a/Engine/Plugins/FX/Niagara/Source/Niagara/Private/Stateless/NiagaraStatelessParticleSimContext.cpp b/Engine/Plugins/FX/Niagara/Source/Niagara/Private/Stateless/NiagaraStatelessParticleSimContext.cpp
--- a/Engine/Plugins/FX/Niagara/Source/Niagara/Private/Stateless/NiagaraStatelessParticleSimContext.cpp
+++ b/Engine/Plugins/FX/Niagara/Source/Niagara/Private/Stateless/NiagaraStatelessParticleSimContext.cpp
@@ -123,7 +123,7 @@ namespace NiagaraStateless
 		DeltaTime = InDeltaTime;
 		InvDeltaTime = DeltaTime > 0.0f ? 1.0f / DeltaTime : 0.0f;
 		EmitterRandomSeed = InEmitterRandomSeed;
-		ModuleRandomSeed = 6405u;
+		ModuleRandomSeed = 0;
 
 		// Setup Required Components
 		{
@@ -218,7 +218,7 @@ namespace NiagaraStateless
 		for (const auto& Callback : ExecData->SimulateFunctions)
 		{
 			BuiltDataOffset = Callback.BuiltDataOffset;
-			ModuleRandomSeed = Callback.RandomSeedOffset * 6405u;
+			ModuleRandomSeed = Callback.RandomSeedOffset;
 			Callback.Function(*this);
 		}
 	}
@@ -226,7 +226,7 @@ namespace NiagaraStateless
 	uint32 FParticleSimulationContext::RandomUInt(uint32 iInstance, uint32 RandomSeedOffset) const
 	{
 		const uint32 UniqueIndex = uint32(GetParticleUniqueIndex()[iInstance]);
-		const FUintVector4 RandomSeed(7123u + RandomSeedOffset, EmitterRandomSeed, UniqueIndex * 3581u, ModuleRandomSeed);
+		const FUintVector4 RandomSeed = FNiagaraStatelessDefinitions::MakeRandomSeed(EmitterRandomSeed, UniqueIndex, ModuleRandomSeed, RandomSeedOffset);
 		const FUintVector4 RandomValue = Rand4DPCG32(RandomSeed);
 		return RandomValue.X;
 	}
@@ -234,7 +234,7 @@ namespace NiagaraStateless
 	FUintVector2 FParticleSimulationContext::RandomUInt2(uint32 iInstance, uint32 RandomSeedOffset) const
 	{
 		const uint32 UniqueIndex = uint32(GetParticleUniqueIndex()[iInstance]);
-		const FUintVector4 RandomSeed(7123u + RandomSeedOffset, EmitterRandomSeed, UniqueIndex * 3581u, ModuleRandomSeed);
+		const FUintVector4 RandomSeed = FNiagaraStatelessDefinitions::MakeRandomSeed(EmitterRandomSeed, UniqueIndex, ModuleRandomSeed, RandomSeedOffset);
 		const FUintVector4 RandomValue = Rand4DPCG32(RandomSeed);
 		return FUintVector2(RandomValue.X, RandomValue.Y);
 	}
@@ -242,7 +242,7 @@ namespace NiagaraStateless
 	FUintVector3 FParticleSimulationContext::RandomUInt3(uint32 iInstance, uint32 RandomSeedOffset) const
 	{
 		const uint32 UniqueIndex = uint32(GetParticleUniqueIndex()[iInstance]);
-		const FUintVector4 RandomSeed(7123u + RandomSeedOffset, EmitterRandomSeed, UniqueIndex * 3581u, ModuleRandomSeed);
+		const FUintVector4 RandomSeed = FNiagaraStatelessDefinitions::MakeRandomSeed(EmitterRandomSeed, UniqueIndex, ModuleRandomSeed, RandomSeedOffset);
 		const FUintVector4 RandomValue = Rand4DPCG32(RandomSeed);
 		return FUintVector3(RandomValue.X, RandomValue.Y, RandomValue.Z);
 	}
@@ -250,7 +250,7 @@ namespace NiagaraStateless
 	FUintVector4 FParticleSimulationContext::RandomUInt4(uint32 iInstance, uint32 RandomSeedOffset) const
 	{
 		const uint32 UniqueIndex = uint32(GetParticleUniqueIndex()[iInstance]);
-		const FUintVector4 RandomSeed(7123u + RandomSeedOffset, EmitterRandomSeed, UniqueIndex * 3581u, ModuleRandomSeed);
+		const FUintVector4 RandomSeed = FNiagaraStatelessDefinitions::MakeRandomSeed(EmitterRandomSeed, UniqueIndex, ModuleRandomSeed, RandomSeedOffset);
 		const FUintVector4 RandomValue = Rand4DPCG32(RandomSeed);
 		return RandomValue;
 	}
diff --git a/Engine/Plugins/FX/Niagara/Source/NiagaraShader/Internal/Stateless/NiagaraStatelessModuleShaderParameters.h b/Engine/Plugins/FX/Niagara/Source/NiagaraShader/Internal/Stateless/NiagaraStatelessModuleShaderParameters.h
--- a/Engine/Plugins/FX/Niagara/Source/NiagaraShader/Internal/Stateless/NiagaraStatelessModuleShaderParameters.h
+++ b/Engine/Plugins/FX/Niagara/Source/NiagaraShader/Internal/Stateless/NiagaraStatelessModuleShaderParameters.h
@@ -4,12 +4,13 @@
 
 #include "RenderGraphFwd.h"
 #include "ShaderParameterStruct.h"
+#include "NiagaraStatelessBuiltDistribution.h"
 
 namespace NiagaraStateless
 {
 	BEGIN_SHADER_PARAMETER_STRUCT(FInitializeParticleModule_ShaderParameters, NIAGARASHADER_API)
-		SHADER_PARAMETER(uint32,		InitializeParticle_ModuleFlags)
-		SHADER_PARAMETER(FUintVector3,	InitializeParticle_InitialPosition)
+		SHADER_PARAMETER(uint32,									InitializeParticle_ModuleFlags)
+		SHADER_PARAMETER(FNiagaraStatelessBuiltDistributionType,	InitializeParticle_InitialPosition)
 
 		SHADER_PARAMETER(FLinearColor,	InitializeParticle_ColorScale)				// Unset / Direct Set / Random Range(Link RGBA / Link RGB|Link A / Random Channels) / Random Hue|Saturation|Value
 		SHADER_PARAMETER(FLinearColor,	InitializeParticle_ColorBias)
@@ -45,37 +46,37 @@ namespace NiagaraStateless
 	END_SHADER_PARAMETER_STRUCT()
 
 	BEGIN_SHADER_PARAMETER_STRUCT(FDynamicMaterialParametersModule_ShaderParameters, NIAGARASHADER_API)
-		SHADER_PARAMETER(uint32,		DynamicMaterialParameters_ChannelMask)
-		SHADER_PARAMETER(FUintVector3,	DynamicMaterialParameters_Parameter0X)
-		SHADER_PARAMETER(FUintVector3,	DynamicMaterialParameters_Parameter0Y)
-		SHADER_PARAMETER(FUintVector3,	DynamicMaterialParameters_Parameter0Z)
-		SHADER_PARAMETER(FUintVector3,	DynamicMaterialParameters_Parameter0W)
-		SHADER_PARAMETER(FUintVector3,	DynamicMaterialParameters_Parameter1X)
-		SHADER_PARAMETER(FUintVector3,	DynamicMaterialParameters_Parameter1Y)
-		SHADER_PARAMETER(FUintVector3,	DynamicMaterialParameters_Parameter1Z)
-		SHADER_PARAMETER(FUintVector3,	DynamicMaterialParameters_Parameter1W)
-		SHADER_PARAMETER(FUintVector3,	DynamicMaterialParameters_Parameter2X)
-		SHADER_PARAMETER(FUintVector3,	DynamicMaterialParameters_Parameter2Y)
-		SHADER_PARAMETER(FUintVector3,	DynamicMaterialParameters_Parameter2Z)
-		SHADER_PARAMETER(FUintVector3,	DynamicMaterialParameters_Parameter2W)
-		SHADER_PARAMETER(FUintVector3,	DynamicMaterialParameters_Parameter3X)
-		SHADER_PARAMETER(FUintVector3,	DynamicMaterialParameters_Parameter3Y)
-		SHADER_PARAMETER(FUintVector3,	DynamicMaterialParameters_Parameter3Z)
-		SHADER_PARAMETER(FUintVector3,	DynamicMaterialParameters_Parameter3W)
+		SHADER_PARAMETER(uint32,									DynamicMaterialParameters_ChannelMask)
+		SHADER_PARAMETER(FNiagaraStatelessBuiltDistributionType,	DynamicMaterialParameters_Parameter0X)
+		SHADER_PARAMETER(FNiagaraStatelessBuiltDistributionType,	DynamicMaterialParameters_Parameter0Y)
+		SHADER_PARAMETER(FNiagaraStatelessBuiltDistributionType,	DynamicMaterialParameters_Parameter0Z)
+		SHADER_PARAMETER(FNiagaraStatelessBuiltDistributionType,	DynamicMaterialParameters_Parameter0W)
+		SHADER_PARAMETER(FNiagaraStatelessBuiltDistributionType,	DynamicMaterialParameters_Parameter1X)
+		SHADER_PARAMETER(FNiagaraStatelessBuiltDistributionType,	DynamicMaterialParameters_Parameter1Y)
+		SHADER_PARAMETER(FNiagaraStatelessBuiltDistributionType,	DynamicMaterialParameters_Parameter1Z)
+		SHADER_PARAMETER(FNiagaraStatelessBuiltDistributionType,	DynamicMaterialParameters_Parameter1W)
+		SHADER_PARAMETER(FNiagaraStatelessBuiltDistributionType,	DynamicMaterialParameters_Parameter2X)
+		SHADER_PARAMETER(FNiagaraStatelessBuiltDistributionType,	DynamicMaterialParameters_Parameter2Y)
+		SHADER_PARAMETER(FNiagaraStatelessBuiltDistributionType,	DynamicMaterialParameters_Parameter2Z)
+		SHADER_PARAMETER(FNiagaraStatelessBuiltDistributionType,	DynamicMaterialParameters_Parameter2W)
+		SHADER_PARAMETER(FNiagaraStatelessBuiltDistributionType,	DynamicMaterialParameters_Parameter3X)
+		SHADER_PARAMETER(FNiagaraStatelessBuiltDistributionType,	DynamicMaterialParameters_Parameter3Y)
+		SHADER_PARAMETER(FNiagaraStatelessBuiltDistributionType,	DynamicMaterialParameters_Parameter3Z)
+		SHADER_PARAMETER(FNiagaraStatelessBuiltDistributionType,	DynamicMaterialParameters_Parameter3W)
 	END_SHADER_PARAMETER_STRUCT()
 
 	BEGIN_SHADER_PARAMETER_STRUCT(FCameraOffsetModule_ShaderParameters, NIAGARASHADER_API)
-		SHADER_PARAMETER(FUintVector3,	CameraOffset_Distribution)
+		SHADER_PARAMETER(FNiagaraStatelessBuiltDistributionType,	CameraOffset_Distribution)
 	END_SHADER_PARAMETER_STRUCT()
 
 	BEGIN_SHADER_PARAMETER_STRUCT(FScaleColorModule_ShaderParameters, NIAGARASHADER_API)
-		SHADER_PARAMETER(FUintVector3,	ScaleColor_Distribution)
+		SHADER_PARAMETER(FNiagaraStatelessBuiltDistributionType,	ScaleColor_Distribution)
 	END_SHADER_PARAMETER_STRUCT()
 
 	BEGIN_SHADER_PARAMETER_STRUCT(FScaleMeshSizeModule_ShaderParameters, NIAGARASHADER_API)
-		SHADER_PARAMETER(FUintVector3,	ScaleMeshSize_Distribution)
-		SHADER_PARAMETER(FVector3f,		ScaleMeshSize_CurveScale)
-		SHADER_PARAMETER(int32,			ScaleMeshSize_CurveScaleOffset)
+		SHADER_PARAMETER(FNiagaraStatelessBuiltDistributionType,	ScaleMeshSize_Distribution)
+		SHADER_PARAMETER(FVector3f,									ScaleMeshSize_CurveScale)
+		SHADER_PARAMETER(int32,										ScaleMeshSize_CurveScaleOffset)
 	END_SHADER_PARAMETER_STRUCT()
 
 	BEGIN_SHADER_PARAMETER_STRUCT(FScaleMeshSizeBySpeedModule_ShaderParameters, NIAGARASHADER_API)
@@ -96,8 +97,8 @@ namespace NiagaraStateless
 	END_SHADER_PARAMETER_STRUCT()
 
 	BEGIN_SHADER_PARAMETER_STRUCT(FScaleSpriteSizeModule_ShaderParameters, NIAGARASHADER_API)
-		SHADER_PARAMETER(FUintVector3,	ScaleSpriteSize_Distribution)
-		SHADER_PARAMETER(FVector2f,		ScaleSpriteSize_CurveScale)
+		SHADER_PARAMETER(FNiagaraStatelessBuiltDistributionType,	ScaleSpriteSize_Distribution)
+		SHADER_PARAMETER(FVector2f,									ScaleSpriteSize_CurveScale)
 	END_SHADER_PARAMETER_STRUCT()
 
 	BEGIN_SHADER_PARAMETER_STRUCT(FScaleSpriteSizeBySpeedModule_ShaderParameters, NIAGARASHADER_API)
diff --git a/Engine/Plugins/FX/Niagara/Source/NiagaraShader/Internal/Stateless/NiagaraStatelessSimulationShader.h b/Engine/Plugins/FX/Niagara/Source/NiagaraShader/Internal/Stateless/NiagaraStatelessSimulationShader.h
--- a/Engine/Plugins/FX/Niagara/Source/NiagaraShader/Internal/Stateless/NiagaraStatelessSimulationShader.h
+++ b/Engine/Plugins/FX/Niagara/Source/NiagaraShader/Internal/Stateless/NiagaraStatelessSimulationShader.h
@@ -52,6 +52,7 @@ namespace NiagaraStateless
 			FGlobalShader::ModifyCompilationEnvironment(Parameters, OutEnvironment);
 			OutEnvironment.SetDefine(TEXT("THREADGROUP_SIZE"), ThreadGroupSize);
 			OutEnvironment.SetDefine(TEXT("NIAGARA_MAX_GPU_SPAWN_INFOS"), NiagaraStateless::MaxGpuSpawnInfos);
+			OutEnvironment.SetDefine(TEXT("GPU_SIMULATION"), 1);
 		}
 	};
 }
diff --git a/Engine/Plugins/FX/Niagara/Source/NiagaraShader/NiagaraShader.Build.cs b/Engine/Plugins/FX/Niagara/Source/NiagaraShader/NiagaraShader.Build.cs
--- a/Engine/Plugins/FX/Niagara/Source/NiagaraShader/NiagaraShader.Build.cs
+++ b/Engine/Plugins/FX/Niagara/Source/NiagaraShader/NiagaraShader.Build.cs
@@ -7,7 +7,9 @@ public class NiagaraShader : ModuleRules
 {
     public NiagaraShader(ReadOnlyTargetRules Target) : base(Target)
     {
-        PrivateDependencyModuleNames.AddRange(
+		PublicIncludePaths.Add(Path.Combine(ModuleDirectory, "../../Shaders/Shared"));
+
+		PrivateDependencyModuleNames.AddRange(
             new string[] {
                 "CoreUObject",
                 "Engine",
diff --git a/Engine/Plugins/FX/Niagara/Source/NiagaraShader/Private/NiagaraShaderModule.cpp b/Engine/Plugins/FX/Niagara/Source/NiagaraShader/Private/NiagaraShaderModule.cpp
--- a/Engine/Plugins/FX/Niagara/Source/NiagaraShader/Private/NiagaraShaderModule.cpp
+++ b/Engine/Plugins/FX/Niagara/Source/NiagaraShader/Private/NiagaraShaderModule.cpp
@@ -17,6 +17,7 @@ void INiagaraShaderModule::StartupModule()
 	// Maps virtual shader source directory /Plugin/FX/Niagara to the plugin's actual Shaders directory.
 	FString PluginShaderDir = FPaths::Combine(IPluginManager::Get().FindPlugin(TEXT("Niagara"))->GetBaseDir(), TEXT("Shaders"));
 	AddShaderSourceDirectoryMapping(TEXT("/Plugin/FX/Niagara"), PluginShaderDir);
+	AddShaderSourceSharedVirtualDirectory(TEXT("/Plugin/FX/Niagara/Shared/"));
 }
 
 FDelegateHandle INiagaraShaderModule::SetOnProcessShaderCompilationQueue(FOnProcessQueue InOnProcessQueue)
