expose max encoded distance with rigid mesh collision di
fix extra sampling of mesh distance field with accurate variant
fix connections for rigid mesh DI integration with fluids
fix loss of sign for accurate mode SDF collisions

#virtualized

diff --git a/Engine/Plugins/FX/Niagara/Shaders/Private/NiagaraDataInterfaceRigidMeshCollisionQuery.ush b/Engine/Plugins/FX/Niagara/Shaders/Private/NiagaraDataInterfaceRigidMeshCollisionQuery.ush
--- a/Engine/Plugins/FX/Niagara/Shaders/Private/NiagaraDataInterfaceRigidMeshCollisionQuery.ush
+++ b/Engine/Plugins/FX/Niagara/Shaders/Private/NiagaraDataInterfaceRigidMeshCollisionQuery.ush
@@ -512,7 +512,7 @@ void EvaluateWorldPositionOnElement_{ParameterName}(int ElementIndex, float3 Wor
 }
 
 // Given a world space position (WorldPosition) compute the static mesh closest point (position,normal,velocity)
-void GetClosestPointMeshDistanceField_{ParameterName}(in float3 WorldPosition, in float DeltaTime, in float TimeFraction, in float MaxDistance, out float OutClosestDistance, out float3 OutClosestPosition, out float3 OutClosestNormal, out float3 OutClosestVelocity, out bool NormalIsValid)
+void GetClosestPointMeshDistanceField_{ParameterName}(in float3 WorldPosition, in float DeltaTime, in float TimeFraction, in float MaxDistance, out float OutClosestDistance, out float3 OutClosestPosition, out float3 OutClosestNormal, out float3 OutClosestVelocity, out bool NormalIsValid, out float MaxEncodedDistance)
 {
 	const FDFVector3 LWCWorldPosition = DFFromTileOffset_Hack(MakeLWCVector3({ParameterName}_SystemLWCTile, WorldPosition));
 
@@ -524,6 +524,8 @@ void GetClosestPointMeshDistanceField_{ParameterName}(in float3 WorldPosition, i
 	OutClosestPosition = float3(0,0,0);
 	OutClosestVelocity = float3(0,0,0);
 
+	MaxEncodedDistance = 0;
+
 	NormalIsValid = false;
 
 	float OutClosestDistanceTmp;
@@ -537,19 +539,22 @@ void GetClosestPointMeshDistanceField_{ParameterName}(in float3 WorldPosition, i
 		{
 			MaxDistance = max(0.001f, MaxDistance);
 
+			OutClosestDistance = GetDistanceToMeshDistanceField(DFIndex, LWCWorldPosition, MaxDistance);
+			
 			float3 SampledWorldPosition;
-			if (SampleClosestPointMeshDistanceField(DFIndex, LWCWorldPosition, MaxDistance, SampledWorldPosition, OutClosestNormal))
+			ComputeClosestPointMeshDistanceField(DFIndex, LWCWorldPosition, OutClosestDistance, SampledWorldPosition, OutClosestNormal, MaxEncodedDistance, NormalIsValid);
+
+			if (NormalIsValid)
 			{
 				EvaluateWorldPositionOnElement_{ParameterName}(ElementIndex, SampledWorldPosition, DeltaTime, TimeFraction, OutClosestPosition, OutClosestVelocity);
-				OutClosestDistance = length(WorldPosition - SampledWorldPosition);
-				NormalIsValid = true;
+				OutClosestDistance = length(WorldPosition - SampledWorldPosition) * sign(OutClosestDistanceTmp);				
 			}
 		}
 	}
 }
 
 // Given a world space position (WorldPosition) compute the static mesh closest point (position,normal,velocity) by querying all overlapping mesh distance fields
-void GetClosestPointMeshDistanceFieldAccurate_{ParameterName}(in float3 WorldPosition, in float DeltaTime, in float TimeFraction, in float MaxDistance, out float OutClosestDistance, out float3 OutClosestPosition, out float3 OutClosestNormal, out float3 OutClosestVelocity, out bool NormalIsValid)
+void GetClosestPointMeshDistanceFieldAccurate_{ParameterName}(in float3 WorldPosition, in float DeltaTime, in float TimeFraction, in float MaxDistance, out float OutClosestDistance, out float3 OutClosestPosition, out float3 OutClosestNormal, out float3 OutClosestVelocity, out bool NormalIsValid, out float MaxEncodedDistance)
 {
 	const FDFVector3 LWCWorldPosition = DFFromTileOffset_Hack(MakeLWCVector3({ParameterName}_SystemLWCTile, WorldPosition));
 
@@ -561,6 +566,8 @@ void GetClosestPointMeshDistanceFieldAccurate_{ParameterName}(in float3 WorldPos
 	OutClosestPosition = float3(0,0,0);
 	OutClosestVelocity = float3(0,0,0);
 	
+	MaxEncodedDistance = 0;
+
 	NormalIsValid = false;
 
 	for (int ElementIndex = 0; ElementIndex < {ParameterName}_ElementOffsets[NUM_ELEMENTS_INDEX]; ++ElementIndex)
@@ -570,16 +577,19 @@ void GetClosestPointMeshDistanceFieldAccurate_{ParameterName}(in float3 WorldPos
 		// #todo(dmp): this does a dynamic branch based on intersecting the bbox.  Maybe we can factor that out due to the broadphase here?
 		// #todo(dmp): Check the distance from the physics asset for a speedup?
 		MaxDistance = max(.001, MaxDistance);
-		float OutClosestDistanceTmp = DistanceToNearestSurfaceForObject(DFIndex, LWCWorldPosition, MaxDistance);
+		float OutClosestDistanceTmp = GetDistanceToMeshDistanceField(DFIndex, LWCWorldPosition, MaxDistance);
 
 		if (OutClosestDistanceTmp < OutClosestDistance)
 		{
+			OutClosestDistance = OutClosestDistanceTmp;
+
 			float3 SampledWorldPosition;
-			if (SampleClosestPointMeshDistanceField(DFIndex, LWCWorldPosition, OutClosestDistanceTmp, SampledWorldPosition, OutClosestNormal))
+			ComputeClosestPointMeshDistanceField(DFIndex, LWCWorldPosition, OutClosestDistance, SampledWorldPosition, OutClosestNormal, MaxEncodedDistance, NormalIsValid);
+			
+			if (NormalIsValid)
 			{
 				EvaluateWorldPositionOnElement_{ParameterName}(ElementIndex, SampledWorldPosition, DeltaTime, TimeFraction, OutClosestPosition, OutClosestVelocity);
-				OutClosestDistance = length(WorldPosition - SampledWorldPosition);
-				NormalIsValid = true;
+				OutClosestDistance = length(WorldPosition - SampledWorldPosition) * sign(OutClosestDistanceTmp);				
 			}
 		}
 	}	
@@ -624,6 +634,24 @@ void GetClosestPointMeshDistanceFieldNoNormal_{ParameterName}(in float3 WorldPos
 	}
 }
 
+void GetMaxEncodedDistanceMeshDistanceField_{ParameterName}(in int ElementIndex, out float MaxEncodedDistance)
+{
+	MaxEncodedDistance = 0;
+	if (ElementIndex != -1)
+	{
+		uint DFIndex = {ParameterName}_DFIndexBuffer[ElementIndex];
+
+		if (DFIndex >= 0 && DFIndex < NumSceneObjects)
+		{
+			FDFObjectData DFObjectData = LoadDFObjectData(DFIndex);
+
+			uint NumMips = LoadDFAssetData(DFObjectData.AssetIndex, 0).NumMips;
+			FDFAssetData DFAssetData = LoadDFAssetData(DFObjectData.AssetIndex, NumMips - 1);
+			MaxEncodedDistance = (DFAssetData.DistanceFieldToVolumeScaleBias.x + DFAssetData.DistanceFieldToVolumeScaleBias.y) * DFObjectData.VolumeScale;
+		}
+	}
+}
+
 #undef BOX_INDEX
 #undef SPHERE_INDEX
 #undef CAPSULE_INDEX
\ No newline at end of file
diff --git a/Engine/Plugins/FX/Niagara/Shaders/Private/NiagaraDistanceFieldQueries.ush b/Engine/Plugins/FX/Niagara/Shaders/Private/NiagaraDistanceFieldQueries.ush
--- a/Engine/Plugins/FX/Niagara/Shaders/Private/NiagaraDistanceFieldQueries.ush
+++ b/Engine/Plugins/FX/Niagara/Shaders/Private/NiagaraDistanceFieldQueries.ush
@@ -6,7 +6,7 @@ NiagaraDistanceFieldQueries.ush
 
 #pragma once
 
-bool SampleClosestPointMeshDistanceField(int DFIndex, const FDFVector3 LWCWorldPosition, float MaxDistance, out float3 OutClosestPosition, out float3 OutClosestNormal)
+float GetDistanceToMeshDistanceField(int DFIndex, const FDFVector3 LWCWorldPosition, const float MaxDistance)
 {
 	FDFObjectData DFObjectData = LoadDFObjectData(DFIndex);
 
@@ -28,31 +28,41 @@ bool SampleClosestPointMeshDistanceField(int DFIndex, const FDFVector3 LWCWorldP
 		FDFAssetData DFAssetData = LoadDFAssetData(DFObjectData.AssetIndex, NumMips - 1);
 
 		DistanceToOccluder = DistanceToMeshSurfaceStandalone(ClampedVolumePosition, DFObjectData) * DFObjectData.VolumeScale;
+	}
 
-		// Don't allow negative SDF distance to stick out of the bounding box
-		DistanceToOccluder = max(DistanceToOccluder, 0.0f);
+	return DistanceToOccluder;
+}
 
-		float3 VolumeGradient = CalculateMeshSDFGradient(ClampedVolumePosition, DFAssetData);
-		float VolumeGradientLength = length(VolumeGradient);
-		float3 VolumeNormal = VolumeGradientLength > 0.0f ? VolumeGradient / VolumeGradientLength : 0;
+void ComputeClosestPointMeshDistanceField(int DFIndex, const FDFVector3 LWCWorldPosition, const float GridSampledDistance, out float3 OutClosestPosition, out float3 OutClosestNormal, out float MaxEncodedDistance, out bool bNormalIsValid)
+{
+	FDFObjectData DFObjectData = LoadDFObjectData(DFIndex);
+	
+	const float3 VolumePosition = DFMultiplyDemote(LWCWorldPosition, DFObjectData.WorldToVolume);
+	const float3 ClampedVolumePosition = clamp(VolumePosition, -DFObjectData.VolumePositionExtent, DFObjectData.VolumePositionExtent);
+	
+	uint NumMips = LoadDFAssetData(DFObjectData.AssetIndex, 0).NumMips;
+	FDFAssetData DFAssetData = LoadDFAssetData(DFObjectData.AssetIndex, NumMips - 1);
 
-		// Transform by transposed inverse to handle non-uniform scaling
-		float4x4 WorldToVolume = DFHackToFloat(DFObjectData.WorldToVolume);
-		float3 WorldGradient = mul(VolumeNormal, transpose((float3x3)WorldToVolume));
-		float WorldGradientLength = length(WorldGradient);
-		const bool bNormalIsValid = WorldGradientLength > 0.0f;
-			
-		// output 			
-		if (bNormalIsValid)
-		{
-			OutClosestNormal = WorldGradient / WorldGradientLength;
+	MaxEncodedDistance = (DFAssetData.DistanceFieldToVolumeScaleBias.x + DFAssetData.DistanceFieldToVolumeScaleBias.y) * DFObjectData.VolumeScale;
+		
+	float3 VolumeGradient = CalculateMeshSDFGradient(ClampedVolumePosition, DFAssetData);
+	float VolumeGradientLength = length(VolumeGradient);
+	float3 VolumeNormal = VolumeGradientLength > 0.0f ? VolumeGradient / VolumeGradientLength : 0;
 
-			const float3 WorldClampedPosition = DFDemote(DFMultiply(ClampedVolumePosition, DFObjectData.VolumeToWorld));
-			OutClosestPosition = WorldClampedPosition - OutClosestNormal * DistanceToOccluder;
+	// Transform by transposed inverse to handle non-uniform scaling
+	float4x4 WorldToVolume = DFHackToFloat(DFObjectData.WorldToVolume);
+	float3 WorldGradient = mul(VolumeNormal, transpose((float3x3)WorldToVolume));
+	float WorldGradientLength = length(WorldGradient);
 
-			return true;
-		}
-	}
+	// @note: we should include checking if we are within the range of max encoded distance, but this breaks backwards compatibility
+	bNormalIsValid = WorldGradientLength > 0.0f && (MaxEncodedDistance - GridSampledDistance > 1e-8);
+			
+	// output 			
+	if (bNormalIsValid)
+	{
+		OutClosestNormal = WorldGradient / WorldGradientLength;
 
-	return false;
+		const float3 WorldClampedPosition = DFDemote(DFMultiply(ClampedVolumePosition, DFObjectData.VolumeToWorld));
+		OutClosestPosition = WorldClampedPosition - OutClosestNormal * GridSampledDistance;
+	}
 }
diff --git a/Engine/Plugins/FX/Niagara/Source/Niagara/Private/NiagaraDataInterfaceRigidMeshCollisionQuery.cpp b/Engine/Plugins/FX/Niagara/Source/Niagara/Private/NiagaraDataInterfaceRigidMeshCollisionQuery.cpp
--- a/Engine/Plugins/FX/Niagara/Source/Niagara/Private/NiagaraDataInterfaceRigidMeshCollisionQuery.cpp
+++ b/Engine/Plugins/FX/Niagara/Source/Niagara/Private/NiagaraDataInterfaceRigidMeshCollisionQuery.cpp
@@ -56,6 +56,7 @@ struct FNiagaraRigidMeshCollisionDIFunctionVersion
 		LargeWorldCoordinates = 1,
 		SetMaxDistance = 2,
 		FindActorRotation = 3,
+		MaxEncodedDistance = 4,
 
 		VersionPlusOne,
 		LatestVersion = VersionPlusOne - 1
@@ -105,6 +106,7 @@ static const FName GetClosestDistanceName(TEXT("GetClosestDistance"));
 static const FName GetClosestPointMeshDistanceFieldName(TEXT("GetClosestPointMeshDistanceField"));
 static const FName GetClosestPointMeshDistanceFieldAccurateName(TEXT("GetClosestPointMeshDistanceFieldAccurate"));
 static const FName GetClosestPointMeshDistanceFieldNoNormalName(TEXT("GetClosestPointMeshDistanceFieldNoNormal"));
+static const FName GetMaxEncodedDistanceMeshDistanceFieldName(TEXT("GetMaxEncodedDistanceMeshDistanceField"));
 
 static const FText OverlapOriginDescription = IF_WITH_EDITORONLY_DATA(
 	LOCTEXT("RigidBodyOverlapOriginDescription", "The center point, in world space, where the overlap trace will be performed."),
@@ -1669,6 +1671,7 @@ void UNiagaraDataInterfaceRigidMeshCollisionQuery::GetFunctionsInternal(TArray<F
 		Sig.Outputs.Add(FNiagaraVariable(FNiagaraTypeDefinition::GetVec3Def(), TEXT("Closest Normal")));
 		Sig.Outputs.Add(FNiagaraVariable(FNiagaraTypeDefinition::GetVec3Def(), TEXT("Closest Velocity")));
 		Sig.Outputs.Add(FNiagaraVariable(FNiagaraTypeDefinition::GetBoolDef(), TEXT("Normal Is Valid")));		
+		Sig.Outputs.Add(FNiagaraVariable(FNiagaraTypeDefinition::GetFloatDef(), TEXT("Max Encoded Distance")));
 		OutFunctions.Add(Sig);
 	}
 
@@ -1690,6 +1693,7 @@ void UNiagaraDataInterfaceRigidMeshCollisionQuery::GetFunctionsInternal(TArray<F
 		Sig.Outputs.Add(FNiagaraVariable(FNiagaraTypeDefinition::GetVec3Def(), TEXT("Closest Normal")));
 		Sig.Outputs.Add(FNiagaraVariable(FNiagaraTypeDefinition::GetVec3Def(), TEXT("Closest Velocity")));
 		Sig.Outputs.Add(FNiagaraVariable(FNiagaraTypeDefinition::GetBoolDef(), TEXT("Normal Is Valid")));
+		Sig.Outputs.Add(FNiagaraVariable(FNiagaraTypeDefinition::GetFloatDef(), TEXT("Max Encoded Distance")));
 		OutFunctions.Add(Sig);
 	}
 
@@ -1712,6 +1716,22 @@ void UNiagaraDataInterfaceRigidMeshCollisionQuery::GetFunctionsInternal(TArray<F
 
 		OutFunctions.Add(Sig);
 	}
+
+	{
+		FNiagaraFunctionSignature Sig;
+		Sig.Name = GetMaxEncodedDistanceMeshDistanceFieldName;
+		Sig.SetDescription(LOCTEXT("GetClosestPointMeshDistanceFieldNNDescription", "Returns the maximum distance stored in the SDF according to the bandwidth it was created with"));
+		Sig.SetFunctionVersion(FNiagaraRigidMeshCollisionDIFunctionVersion::LatestVersion);
+		Sig.bSupportsGPU = true;
+		Sig.bSupportsCPU = false;
+		Sig.bMemberFunction = true;
+		Sig.Inputs.Add(FNiagaraVariable(FNiagaraTypeDefinition(GetClass()), TEXT("Collision DI")));
+		Sig.Inputs.Add(FNiagaraVariable(FNiagaraTypeDefinition::GetIntDef(), TEXT("Element Index")));		
+		Sig.Outputs.Add(FNiagaraVariable(FNiagaraTypeDefinition::GetFloatDef(), TEXT("Max Distance")));
+
+		OutFunctions.Add(Sig);
+	}
+
 }
 #endif
 
@@ -1761,7 +1781,8 @@ bool UNiagaraDataInterfaceRigidMeshCollisionQuery::GetFunctionHLSL(const FNiagar
 		(FunctionInfo.DefinitionName == GetClosestDistanceName) ||
 		(FunctionInfo.DefinitionName == GetClosestPointMeshDistanceFieldName) ||
 		(FunctionInfo.DefinitionName == GetClosestPointMeshDistanceFieldAccurateName) ||
-		(FunctionInfo.DefinitionName == GetClosestPointMeshDistanceFieldNoNormalName) )
+		(FunctionInfo.DefinitionName == GetClosestPointMeshDistanceFieldNoNormalName) ||
+		(FunctionInfo.DefinitionName == GetMaxEncodedDistanceMeshDistanceFieldName) )
 	{
 		return true;
 	}
@@ -1852,6 +1873,15 @@ bool UNiagaraDataInterfaceRigidMeshCollisionQuery::UpgradeFunctionCall(FNiagaraF
 		}
 	}
 
+	if (FunctionSignature.FunctionVersion < FNiagaraRigidMeshCollisionDIFunctionVersion::MaxEncodedDistance)
+	{
+		if (FunctionSignature.Name == GetClosestPointMeshDistanceFieldName || FunctionSignature.Name == GetClosestPointMeshDistanceFieldAccurateName)
+		{
+			FunctionSignature.Outputs.Add(FNiagaraVariable(FNiagaraTypeDefinition::GetFloatDef(), TEXT("Max Encoded Distance")));
+			bChanged = true;
+		}
+	}
+
 	FunctionSignature.FunctionVersion = FNiagaraRigidMeshCollisionDIFunctionVersion::LatestVersion;
 
 	return bChanged;
