improve tooltips for rigid mesh collision DI functions

diff --git a/Engine/Plugins/FX/Niagara/Shaders/Private/NiagaraDistanceFieldQueries.ush b/Engine/Plugins/FX/Niagara/Shaders/Private/NiagaraDistanceFieldQueries.ush
--- a/Engine/Plugins/FX/Niagara/Shaders/Private/NiagaraDistanceFieldQueries.ush
+++ b/Engine/Plugins/FX/Niagara/Shaders/Private/NiagaraDistanceFieldQueries.ush
@@ -55,7 +55,7 @@ void ComputeClosestPointMeshDistanceField(int DFIndex, const FDFVector3 LWCWorld
 	float WorldGradientLength = length(WorldGradient);
 
 	// @note: we should include checking if we are within the range of max encoded distance, but this breaks backwards compatibility
-	bNormalIsValid = WorldGradientLength > 0.0f && (MaxEncodedDistance - GridSampledDistance > 1e-8);
+	bNormalIsValid = WorldGradientLength > 0.0f && (MaxEncodedDistance - abs(GridSampledDistance) > 1e-8);
 			
 	// output 			
 	if (bNormalIsValid)
diff --git a/Engine/Plugins/FX/Niagara/Source/Niagara/Private/NiagaraDataInterfaceRigidMeshCollisionQuery.cpp b/Engine/Plugins/FX/Niagara/Source/Niagara/Private/NiagaraDataInterfaceRigidMeshCollisionQuery.cpp
--- a/Engine/Plugins/FX/Niagara/Source/Niagara/Private/NiagaraDataInterfaceRigidMeshCollisionQuery.cpp
+++ b/Engine/Plugins/FX/Niagara/Source/Niagara/Private/NiagaraDataInterfaceRigidMeshCollisionQuery.cpp
@@ -1605,7 +1605,7 @@ void UNiagaraDataInterfaceRigidMeshCollisionQuery::GetFunctionsInternal(TArray<F
 	{
 		FNiagaraFunctionSignature Sig;
 		Sig.Name = GetElementPointMeshDistanceFieldNoNormalName;
-		Sig.SetDescription(LOCTEXT("GetClosestElementPointMeshDistanceFieldNoNormalDescription", "Given a world space position and an element index, computes the static mesh's closest point. Also returns normal and velocity for that point."));
+		Sig.SetDescription(LOCTEXT("GetClosestElementPointMeshDistanceFieldNoNormalDescription", "Given a world space position and an element index, computes the static mesh's closest distance and velocity."));
 		Sig.SetFunctionVersion(FNiagaraRigidMeshCollisionDIFunctionVersion::LatestVersion);
 		Sig.bSupportsGPU = true;
 		Sig.bSupportsCPU = false;
@@ -1700,7 +1700,7 @@ void UNiagaraDataInterfaceRigidMeshCollisionQuery::GetFunctionsInternal(TArray<F
 	{
 		FNiagaraFunctionSignature Sig;
 		Sig.Name = GetClosestPointMeshDistanceFieldNoNormalName;
-		Sig.SetDescription(LOCTEXT("GetClosestPointMeshDistanceFieldNNDescription", "Given a world space position, computes the distance to the closest point for the static mesh, using the mesh's distance field.\nSkips the normal calculation and is more performant than it's counterpart with normal."));
+		Sig.SetDescription(LOCTEXT("GetClosestPointMeshDistanceFieldNNDescription", "Given a world space position, computes the accurate closest distance and velocity for static meshes.  The closest position is approximate and computed by proximity to the physics assets and not the mesh itself."));
 		Sig.SetFunctionVersion(FNiagaraRigidMeshCollisionDIFunctionVersion::LatestVersion);
 		Sig.bSupportsGPU = true;
 		Sig.bSupportsCPU = false;
