- Add copy to clipboard to the sim cache viewer
- This adds functionality back that was missing from the old attribute spreadsheet

diff --git a/Engine/Plugins/FX/Niagara/Source/NiagaraEditor/Private/Toolkits/NiagaraSimCacheToolkit.cpp b/Engine/Plugins/FX/Niagara/Source/NiagaraEditor/Private/Toolkits/NiagaraSimCacheToolkit.cpp
--- a/Engine/Plugins/FX/Niagara/Source/NiagaraEditor/Private/Toolkits/NiagaraSimCacheToolkit.cpp
+++ b/Engine/Plugins/FX/Niagara/Source/NiagaraEditor/Private/Toolkits/NiagaraSimCacheToolkit.cpp
@@ -261,8 +261,24 @@ void FNiagaraSimCacheToolkit::ExtendToolbar()
 						FSlateIcon(FAppStyle::GetAppStyleSetName(), "SystemWideCommands.FindInContentBrowser.Small")
 					);
 				}
-				ToolbarBuilder.EndSection();
 			}
+			ToolbarBuilder.EndSection();
+			ToolbarBuilder.BeginSection("Export");
+			{
+				TSharedRef<FNiagaraSimCacheViewModel> SimCacheViewModelRef = Toolkit->SimCacheViewModel.ToSharedRef();
+
+				ToolbarBuilder.AddToolBarButton(
+					FUIAction(
+						FExecuteAction::CreateSP(SimCacheViewModelRef, &FNiagaraSimCacheViewModel::CopyActiveToClipboard),
+						FCanExecuteAction::CreateSP(SimCacheViewModelRef, &FNiagaraSimCacheViewModel::CanCopyActiveToClipboard)
+					),
+					NAME_None,
+					LOCTEXT("CopyToCSV", "Copy to CSV"),
+					LOCTEXT("CopyToCSVTooltip", "Copies the currently selected view to the clipboard in CSV format."),
+					FSlateIcon(FAppStyle::GetAppStyleSetName(), "GenericCommands.Copy")
+				);
+			}
+			ToolbarBuilder.EndSection();
 		}
 	};
 	
diff --git a/Engine/Plugins/FX/Niagara/Source/NiagaraEditor/Private/ViewModels/NiagaraSimCacheViewModel.cpp b/Engine/Plugins/FX/Niagara/Source/NiagaraEditor/Private/ViewModels/NiagaraSimCacheViewModel.cpp
--- a/Engine/Plugins/FX/Niagara/Source/NiagaraEditor/Private/ViewModels/NiagaraSimCacheViewModel.cpp
+++ b/Engine/Plugins/FX/Niagara/Source/NiagaraEditor/Private/ViewModels/NiagaraSimCacheViewModel.cpp
@@ -5,6 +5,7 @@
 #include "NiagaraSimCache.h"
 #include "NiagaraEditorCommon.h"
 #include "NiagaraDebuggerCommon.h"
+#include "HAL/PlatformApplicationMisc.h"
 #include "Serialization/MemoryReader.h"
 #include "Widgets/SNiagaraSimCacheTreeView.h"
 #include "UObject/Package.h"
@@ -71,8 +72,6 @@ void FNiagaraSimCacheViewModel::SetupPreviewComponentAndInstance()
 		PreviewComponent->SetRelativeLocation(FVector::ZeroVector);
 		PreviewComponent->SetDesiredAge(SimCache->GetStartSeconds());
 	}
-
-	
 }
 
 TConstArrayView<FNiagaraSimCacheViewModel::FComponentInfo> FNiagaraSimCacheViewModel::GetComponentInfos(
@@ -176,17 +175,17 @@ void FNiagaraSimCacheViewModel::SetEmitterIndex(const int32 InEmitterIndex, FNia
 	OnViewDataChangedDelegate.Broadcast(true);
 }
 
-bool FNiagaraSimCacheViewModel::IsCacheValid()
+bool FNiagaraSimCacheViewModel::IsCacheValid() const
 {
 	return SimCache ? SimCache->IsCacheValid() : false;
 }
 
-int32 FNiagaraSimCacheViewModel::GetNumEmitterLayouts()
+int32 FNiagaraSimCacheViewModel::GetNumEmitterLayouts() const
 {
 	return SimCache ? SimCache->GetNumEmitters() : 0;
 }
 
-FName FNiagaraSimCacheViewModel::GetEmitterLayoutName(const int32 Index)
+FName FNiagaraSimCacheViewModel::GetEmitterLayoutName(const int32 Index) const
 {
 	return SimCache ? SimCache->GetEmitterName(Index) : NAME_None;
 }
@@ -334,10 +333,9 @@ void FNiagaraSimCacheViewModel::BuildTreeItemChildren(TSharedPtr<FNiagaraSimCach
 		int32 BufferIndex = TreeItem->GetBufferIndex();
 		
 		SimCache->ForEachEmitterAttribute(BufferIndex,
-	[&](const FNiagaraSimCacheVariable& Variable)
+			[&](const FNiagaraSimCacheVariable& Variable)
 			{
 				FNiagaraTypeDefinition TypeDef = Variable.Variable.GetType();
-			
 
 				TSharedRef<FNiagaraSimCacheComponentTreeItem> CurrentItem = MakeShared<FNiagaraSimCacheComponentTreeItem>(OwningTreeView);
 				
@@ -503,6 +501,42 @@ TArray<TSharedRef<FNiagaraSimCacheOverviewItem>>* FNiagaraSimCacheViewModel::Get
 	return &BufferEntries;
 }
 
+bool FNiagaraSimCacheViewModel::CanCopyActiveToClipboard() const
+{
+	return IsCacheValid() && !ActiveDataInterface.IsValid();
+}
+
+void FNiagaraSimCacheViewModel::CopyActiveToClipboard() const
+{
+	if (!CanCopyActiveToClipboard())
+	{
+		return;
+	}
+
+	FString ClipboardString;
+	TConstArrayView<FComponentInfo> ComponentInfos = GetCurrentComponentInfos();
+
+	ClipboardString.Append(TEXT("Instance"));
+	for ( int iComponent=0; iComponent < ComponentInfos.Num(); ++iComponent)
+	{
+		ClipboardString.AppendChar(TEXT(','));
+		ComponentInfos[iComponent].Name.AppendString(ClipboardString);
+	}
+
+	for (int32 iInstance=0; iInstance < NumInstances; ++iInstance)
+	{
+		ClipboardString.AppendChar(TEXT('\n'));
+		ClipboardString.AppendInt(iInstance);
+
+		for (int iComponent = 0; iComponent < ComponentInfos.Num(); ++iComponent)
+		{
+			ClipboardString.AppendChar(TEXT(','));
+			ClipboardString.Append(GetComponentText(ComponentInfos[iComponent].Name, iInstance).ToString());
+		}
+	}
+	FPlatformApplicationMisc::ClipboardCopy(*ClipboardString);
+}
+
 void FNiagaraSimCacheViewModel::AddReferencedObjects(FReferenceCollector& Collector)
 {
 	if(SimCache != nullptr)
diff --git a/Engine/Plugins/FX/Niagara/Source/NiagaraEditor/Private/Widgets/SNiagaraDebugCaptureView.cpp b/Engine/Plugins/FX/Niagara/Source/NiagaraEditor/Private/Widgets/SNiagaraDebugCaptureView.cpp
--- a/Engine/Plugins/FX/Niagara/Source/NiagaraEditor/Private/Widgets/SNiagaraDebugCaptureView.cpp
+++ b/Engine/Plugins/FX/Niagara/Source/NiagaraEditor/Private/Widgets/SNiagaraDebugCaptureView.cpp
@@ -279,9 +279,24 @@ void SNiagaraDebugCaptureView::Construct(const FArguments& InArgs, const TShared
 				]
 		);
 	}
+	DebugCaptureToolbarBuilder.EndSection();
 
+	DebugCaptureToolbarBuilder.BeginSection("Export");
+	{
+		TSharedRef<FNiagaraSimCacheViewModel> SimCacheViewModelRef = SimCacheViewModel.ToSharedRef();
+		DebugCaptureToolbarBuilder.AddToolBarButton(
+			FUIAction(
+				FExecuteAction::CreateSP(SimCacheViewModelRef, &FNiagaraSimCacheViewModel::CopyActiveToClipboard),
+				FCanExecuteAction::CreateSP(SimCacheViewModelRef, &FNiagaraSimCacheViewModel::CanCopyActiveToClipboard)
+			),
+			NAME_None,
+			LOCTEXT("CopyToCSV", "Copy to CSV"),
+			LOCTEXT("CopyToCSVTooltip", "Copies the currently selected view to the clipboard in CSV format."),
+			FSlateIcon(FAppStyle::GetAppStyleSetName(), "GenericCommands.Copy")
+		);
+	}
 	DebugCaptureToolbarBuilder.EndSection();
-	
+
 	ChildSlot
 	[
 		DebugCaptureToolbarBuilder.MakeWidget()
@@ -410,8 +425,6 @@ void SNiagaraDebugCaptureView::OnCaptureComplete(UNiagaraSimCache* CapturedSimCa
 	default:
 		break;
 	}
-
-	
 }
 
 #undef LOCTEXT_NAMESPACE
diff --git a/Engine/Plugins/FX/Niagara/Source/NiagaraEditor/Public/ViewModels/NiagaraSimCacheViewModel.h b/Engine/Plugins/FX/Niagara/Source/NiagaraEditor/Public/ViewModels/NiagaraSimCacheViewModel.h
--- a/Engine/Plugins/FX/Niagara/Source/NiagaraEditor/Public/ViewModels/NiagaraSimCacheViewModel.h
+++ b/Engine/Plugins/FX/Niagara/Source/NiagaraEditor/Public/ViewModels/NiagaraSimCacheViewModel.h
@@ -60,15 +60,18 @@ public:
 	void SetEmitterIndex(int32 InEmitterIndex, FNiagaraVariableBase ActiveDataInterface = FNiagaraVariableBase());
 
 	FText GetComponentText(FName ComponentName, int32 InstanceIndex) const;
-	bool IsCacheValid();
-	int32 GetNumEmitterLayouts();
-	FName GetEmitterLayoutName(int32 Index);
+	bool IsCacheValid() const;
+	int32 GetNumEmitterLayouts() const;
+	FName GetEmitterLayoutName(int32 Index) const;
 
 	// Construct entries for the tree view.
 	void BuildEntries(TWeakPtr<SNiagaraSimCacheTreeView> OwningTreeView);
 	TArray<TSharedRef<FNiagaraSimCacheTreeItem>>* GetCurrentRootEntries();
 	TArray<TSharedRef<FNiagaraSimCacheOverviewItem>>* GetBufferEntries();
 
+	bool CanCopyActiveToClipboard() const;
+	void CopyActiveToClipboard() const;
+
 	UNiagaraComponent* GetPreviewComponent() { return PreviewComponent; }
 
 	//~ FGCObject interface
