fix loctext

diff --git a/Engine/Plugins/FX/Niagara/Source/Niagara/Private/NiagaraDataInterfaceRigidMeshCollisionQuery.cpp b/Engine/Plugins/FX/Niagara/Source/Niagara/Private/NiagaraDataInterfaceRigidMeshCollisionQuery.cpp
--- a/Engine/Plugins/FX/Niagara/Source/Niagara/Private/NiagaraDataInterfaceRigidMeshCollisionQuery.cpp
+++ b/Engine/Plugins/FX/Niagara/Source/Niagara/Private/NiagaraDataInterfaceRigidMeshCollisionQuery.cpp
@@ -1720,7 +1720,7 @@ void UNiagaraDataInterfaceRigidMeshCollisionQuery::GetFunctionsInternal(TArray<F
 	{
 		FNiagaraFunctionSignature Sig;
 		Sig.Name = GetMaxEncodedDistanceMeshDistanceFieldName;
-		Sig.SetDescription(LOCTEXT("GetClosestPointMeshDistanceFieldNNDescription", "Returns the maximum distance stored in the SDF according to the bandwidth it was created with"));
+		Sig.SetDescription(LOCTEXT("GetMaxEncodedDistanceMeshDistanceFieldDescription", "Returns the maximum distance stored in the SDF according to the bandwidth it was created with"));
 		Sig.SetFunctionVersion(FNiagaraRigidMeshCollisionDIFunctionVersion::LatestVersion);
 		Sig.bSupportsGPU = true;
 		Sig.bSupportsCPU = false;
