- Fix for stateless distribution bindings

#tests Tested binding sprite size to a user parameter
#rnx

diff --git a/Engine/Plugins/FX/Niagara/Shaders/Shared/NiagaraStatelessBuiltDistribution.h b/Engine/Plugins/FX/Niagara/Shaders/Shared/NiagaraStatelessBuiltDistribution.h
--- a/Engine/Plugins/FX/Niagara/Shaders/Shared/NiagaraStatelessBuiltDistribution.h
+++ b/Engine/Plugins/FX/Niagara/Shaders/Shared/NiagaraStatelessBuiltDistribution.h
@@ -11,9 +11,9 @@
 	#define FNiagaraStatelessBuiltDistributionTypeIn	uint3
 #endif
 
-#define ENiagaraStatelessBuiltDistributionFlag_Binding			0x10000000u
-#define ENiagaraStatelessBuiltDistributionFlag_Random			0x20000000u
-#define ENiagaraStatelessBuiltDistributionFlag_Uniform			0x40000000u
+#define ENiagaraStatelessBuiltDistributionFlag_Binding			0x20000000u
+#define ENiagaraStatelessBuiltDistributionFlag_Random			0x40000000u
+#define ENiagaraStatelessBuiltDistributionFlag_Uniform			0x80000000u
 #define ENiagaraStatelessBuiltDistributionFlag_TableLengthBits	9u				// 512 entries in LUT
 #define ENiagaraStatelessBuiltDistributionFlag_TableLengthShift	20u
 #define ENiagaraStatelessBuiltDistributionFlag_TableLengthMask	((1u << ENiagaraStatelessBuiltDistributionFlag_TableLengthBits) - 1u)
@@ -56,6 +56,19 @@ struct FNiagaraStatelessBuiltDistribution
 	static void SetIsRandom(FNiagaraStatelessBuiltDistributionType& BuiltData) { BuiltData[0] |= ENiagaraStatelessBuiltDistributionFlag_Random; }
 	static void SetIsUniform(FNiagaraStatelessBuiltDistributionType& BuiltData) { BuiltData[0] |= ENiagaraStatelessBuiltDistributionFlag_Uniform; }
 
+	static void SetLookupParameters(FNiagaraStatelessBuiltDistributionType& BuiltData, uint DataOffset)
+	{
+		check(IsBinding(BuiltData));
+		check(DataOffset <= ENiagaraStatelessBuiltDistributionFlag_DataOffsetMask);
+
+		BuiltData[0] &= ~(ENiagaraStatelessBuiltDistributionFlag_DataOffsetMask << ENiagaraStatelessBuiltDistributionFlag_DataOffsetShift);
+		BuiltData[0] &= ~(ENiagaraStatelessBuiltDistributionFlag_TableLengthMask << ENiagaraStatelessBuiltDistributionFlag_TableLengthShift);
+		BuiltData[0] |= DataOffset << ENiagaraStatelessBuiltDistributionFlag_DataOffsetShift;
+
+		reinterpret_cast<float&>(BuiltData[1]) = 0.0f;
+		reinterpret_cast<float&>(BuiltData[2]) = 1.0f;
+	}
+
 	static void SetLookupParameters(FNiagaraStatelessBuiltDistributionType& BuiltData, uint DataOffset, uint TableLength, const FVector2f& TimeRange)
 	{
 		const uint TableLengthMinusOne = TableLength - 1;
diff --git a/Engine/Plugins/FX/Niagara/Source/Niagara/Internal/Stateless/NiagaraStatelessEmitterDataBuildContext.h b/Engine/Plugins/FX/Niagara/Source/Niagara/Internal/Stateless/NiagaraStatelessEmitterDataBuildContext.h
--- a/Engine/Plugins/FX/Niagara/Source/Niagara/Internal/Stateless/NiagaraStatelessEmitterDataBuildContext.h
+++ b/Engine/Plugins/FX/Niagara/Source/Niagara/Internal/Stateless/NiagaraStatelessEmitterDataBuildContext.h
@@ -104,7 +104,7 @@ public:
 			if (ParameterOffset >= 0)
 			{
 				FNiagaraStatelessBuiltDistribution::SetIsBinding(BuiltDistribution);
-				FNiagaraStatelessBuiltDistribution::SetLookupParameters(BuiltDistribution, ParameterOffset, 1, FVector2f(0.0f, 1.0f));
+				FNiagaraStatelessBuiltDistribution::SetLookupParameters(BuiltDistribution, ParameterOffset);
 			}
 		}
 		else
