[TaskGraph/Niagara]
 - Replace busywait logic

#rnx

diff --git a/Engine/Plugins/FX/Niagara/Source/Niagara/Private/NiagaraSystemInstance.cpp b/Engine/Plugins/FX/Niagara/Source/Niagara/Private/NiagaraSystemInstance.cpp
--- a/Engine/Plugins/FX/Niagara/Source/Niagara/Private/NiagaraSystemInstance.cpp
+++ b/Engine/Plugins/FX/Niagara/Source/Niagara/Private/NiagaraSystemInstance.cpp
@@ -2362,18 +2362,10 @@ void FNiagaraSystemInstance::WaitForConcurrentTickDoNotFinalize(bool bEnsureComp
 		extern int32 GNiagaraSystemSimulationTaskStallTimeout;
 		if (GNiagaraSystemSimulationTaskStallTimeout > 0)
 		{
-			const double EndTimeoutSeconds = FPlatformTime::Seconds() + (double(GNiagaraSystemSimulationTaskStallTimeout) / 1000.0);
-			LowLevelTasks::BusyWaitUntil(
-				[this, EndTimeoutSeconds]()
-				{
-					if (FPlatformTime::Seconds() > EndTimeoutSeconds)
-					{
-						DumpStalledInfo();
-						return true;
-					}
-					return ConcurrentTickGraphEvent->IsComplete();
-				}
-			);
+			if (WaitForAnyTaskCompleted({ ConcurrentTickGraphEvent }, FTimespan::FromMicroseconds(GNiagaraSystemSimulationTaskStallTimeout)) == INDEX_NONE)
+			{
+				DumpStalledInfo();
+			}
 		}
 		else
 		{
diff --git a/Engine/Plugins/FX/Niagara/Source/Niagara/Private/NiagaraSystemSimulation.cpp b/Engine/Plugins/FX/Niagara/Source/Niagara/Private/NiagaraSystemSimulation.cpp
--- a/Engine/Plugins/FX/Niagara/Source/Niagara/Private/NiagaraSystemSimulation.cpp
+++ b/Engine/Plugins/FX/Niagara/Source/Niagara/Private/NiagaraSystemSimulation.cpp
@@ -1622,18 +1622,10 @@ void FNiagaraSystemSimulation::WaitForConcurrentTickComplete(bool bEnsureComplet
 
 		if (GNiagaraSystemSimulationTaskStallTimeout > 0)
 		{
-			const double EndTimeoutSeconds = FPlatformTime::Seconds() + (double(GNiagaraSystemSimulationTaskStallTimeout) / 1000.0);
-			LowLevelTasks::BusyWaitUntil(
-				[this, EndTimeoutSeconds]()
-				{
-					if (FPlatformTime::Seconds() > EndTimeoutSeconds)
-					{
-						DumpStalledInfo();
-						return true;
-					}
-					return ConcurrentTickGraphEvent->IsComplete();
-				}
-			);
+			if (WaitForAnyTaskCompleted({ ConcurrentTickGraphEvent }, FTimespan::FromMicroseconds(GNiagaraSystemSimulationTaskStallTimeout)) == INDEX_NONE)
+			{
+				DumpStalledInfo();
+			}
 		}
 		else
 		{
diff --git a/Engine/Plugins/FX/Niagara/Source/NiagaraEditor/Private/NiagaraCompilationTasks.cpp b/Engine/Plugins/FX/Niagara/Source/NiagaraEditor/Private/NiagaraCompilationTasks.cpp
--- a/Engine/Plugins/FX/Niagara/Source/NiagaraEditor/Private/NiagaraCompilationTasks.cpp
+++ b/Engine/Plugins/FX/Niagara/Source/NiagaraEditor/Private/NiagaraCompilationTasks.cpp
@@ -1720,9 +1720,9 @@ void FNiagaraSystemCompilationTask::WaitTillCompileCompletion()
 	const FTimespan WaitTimeout = FTimespan::FromMilliseconds(50.0);
 	while (!CompileCompletionEvent.IsCompleted())
 	{
-		// if the busy wait doesn't complete the task then we need to make sure to poke the
+		// if the wait doesn't complete the task then we need to make sure to poke the
 		// compilation manager since there we may be waiting on it
-		if (!CompileCompletionEvent.BusyWait(WaitTimeout))
+		if (!CompileCompletionEvent.Wait(WaitTimeout))
 		{
 			FAssetCompilingManager::Get().ProcessAsyncTasks();
 		}
