- Data interface for reading sockets
- Can read from USceneComponent / USkeletalMesh / UStaticMesh

#virtualized

diff --git a/Engine/Plugins/FX/Niagara/Shaders/Private/NiagaraDataInterfaceSocketReaderTemplate.ush b/Engine/Plugins/FX/Niagara/Shaders/Private/NiagaraDataInterfaceSocketReaderTemplate.ush
new file mode 100644
--- /dev/null
+++ b/Engine/Plugins/FX/Niagara/Shaders/Private/NiagaraDataInterfaceSocketReaderTemplate.ush
@@ -0,0 +1,118 @@
+// Copyright Epic Games, Inc. All Rights Reserved.
+
+#include "/Plugin/FX/Niagara/Private/NiagaraTransformUtils.ush"
+
+uint				{ParameterName}_IsDataValid;
+float				{ParameterName}_InvDeltaSeconds;
+int					{ParameterName}_NumSockets;
+int					{ParameterName}_NumFilteredSockets;
+int					{ParameterName}_NumUnfilteredSockets;
+float3				{ParameterName}_ComponentToTranslatedWorld_Translation;
+float4				{ParameterName}_ComponentToTranslatedWorld_Rotation;
+float3				{ParameterName}_ComponentToTranslatedWorld_Scale;
+float3				{ParameterName}_PreviousComponentToTranslatedWorld_Translation;
+float4				{ParameterName}_PreviousComponentToTranslatedWorld_Rotation;
+float3				{ParameterName}_PreviousComponentToTranslatedWorld_Scale;
+uint				{ParameterName}_SocketTransformOffset;
+uint				{ParameterName}_PreviousSocketTransformOffset;
+ByteAddressBuffer	{ParameterName}_SocketData;
+
+////////////////////////////////////////////////////
+// Note that when reading socket / socket index indirection we allocate +1 entries, this is to avoid branching when OOB
+
+int SanitizeSocketIndex_{ParameterName}(int SocketIndex)
+{
+	return SocketIndex >= 0 && SocketIndex < {ParameterName}_NumSockets ? SocketIndex : {ParameterName}_NumSockets;
+}
+
+int SanitizeFilteredSocketIndex_{ParameterName}(int SocketIndex)
+{
+	const int FilteredSockedIndex = SocketIndex >= 0 && SocketIndex < {ParameterName}_NumFilteredSockets ? (SockedIndex + 1) : 0;
+	return {ParameterName}_SocketData.Load(SocketIndex * 4);
+}
+
+int SanitizeUnfilteredSocketIndex_{ParameterName}(int SocketIndex)
+{
+	const int UnfilteredSockedIndex = SocketIndex >= 0 && SocketIndex < {ParameterName}_NumUnfilteredSockets ? (SockedIndex + 1 + {ParameterName}_NumFilteredSockets) : 0;
+	return {ParameterName}_SocketData.Load(SocketIndex * 4);
+}
+
+void GetSocketData_{ParameterName}(int SocketIndex, float Interp, out float3 OutPosition, out float4 OutRotation, out float3 OutScale, out float3 OutVelocity)
+{
+	uint CurrTransformOffset = {ParameterName}_SocketTransformOffset;
+	uint PrevTransformOffset = {ParameterName}_PreviousSocketTransformOffset;
+	CurrTransformOffset += SocketIndex * 10 * 4;
+	PrevTransformOffset += SocketIndex * 10 * 4;
+
+	FNiagaraTransform ComponentToTranslatedWorld = MakeTransform({ParameterName}_ComponentToTranslatedWorld_Translation, {ParameterName}_ComponentToTranslatedWorld_Rotation, {ParameterName}_ComponentToTranslatedWorld_Scale);
+	FNiagaraTransform PreviousComponentToTranslatedWorld = MakeTransform({ParameterName}_PreviousComponentToTranslatedWorld_Translation, {ParameterName}_PreviousComponentToTranslatedWorld_Rotation, {ParameterName}_PreviousComponentToTranslatedWorld_Scale);
+
+	FNiagaraTransform CurrTransform = LoadTransform({ParameterName}_SocketData, CurrTransformOffset);
+	FNiagaraTransform PrevTransform = LoadTransform({ParameterName}_SocketData, PrevTransformOffset);
+
+	CurrTransform = MultiplyTransform(CurrTransform, ComponentToTranslatedWorld);
+	PrevTransform = MultiplyTransform(PrevTransform, PreviousComponentToTranslatedWorld);
+
+	OutPosition	= lerp(PrevTransform.Translation, CurrTransform.Translation, Interp);
+	OutRotation	= SlerpQuat(PrevTransform.Rotation, CurrTransform.Rotation, Interp);
+	OutScale	= lerp(PrevTransform.Scale, CurrTransform.Scale, Interp);
+	OutVelocity	= (CurrTransform.Translation - PrevTransform.Translation) * {ParameterName}_InvDeltaSeconds;
+}
+
+////////////////////////////////////////////////////
+
+void IsValid_{ParameterName}(out bool bValid)
+{
+	bValid = {ParameterName}_IsDataValid != 0;
+}
+
+void GetSocketCount_{ParameterName}(out int Count)
+{
+	Count = {ParameterName}_NumSockets;
+}
+
+void GetFilteredSocketCount_{ParameterName}(out int Count)
+{
+	Count = {ParameterName}_NumFilteredSockets;
+}
+
+void GetUnfilteredSocketCount_{ParameterName}(out int Count)
+{
+	Count = {ParameterName}_NumUnfilteredSockets;
+}
+
+void GetSocketTransform_{ParameterName}(int SocketIndex, out float3 OutPosition, out float4 OutRotation, out float3 OutScale, out float3 OutVelocity)
+{
+	SocketIndex = SanitizeSocketIndex_{ParameterName}(SocketIndex);
+	GetSocketData_{ParameterName}(SocketIndex, 1.0f, OutPosition, OutRotation, OutScale, OutVelocity);
+}
+
+void GetFilteredSocketTransform_{ParameterName}(int SocketIndex, out float3 OutPosition, out float4 OutRotation, out float3 OutScale, out float3 OutVelocity)
+{
+	SocketIndex = SanitizeFilteredSocketIndex_{ParameterName}(SocketIndex);
+	GetSocketData_{ParameterName}(SocketIndex, 1.0f, OutPosition, OutRotation, OutScale, OutVelocity);
+}
+
+void GetUnfilteredSocketTransform_{ParameterName}(int SocketIndex, out float3 OutPosition, out float4 OutRotation, out float3 OutScale, out float3 OutVelocity)
+{
+	SocketIndex = SanitizeUnfilteredSocketIndex_{ParameterName}(SocketIndex);
+	GetSocketData_{ParameterName}(SocketIndex, 1.0f, OutPosition, OutRotation, OutScale, OutVelocity);
+}
+
+void GetSocketTransformInterpolated_{ParameterName}(int SocketIndex, float Interpolation, out float3 OutPosition, out float4 OutRotation, out float3 OutScale, out float3 OutVelocity)
+{
+	SocketIndex = SanitizeSocketIndex_{ParameterName}(SocketIndex);
+	GetSocketData_{ParameterName}(SocketIndex, Interpolation, OutPosition, OutRotation, OutScale, OutVelocity);
+}
+
+void GetFilteredSocketTransformInterpolated_{ParameterName}(int SocketIndex, float Interpolation, out float3 OutPosition, out float4 OutRotation, out float3 OutScale, out float3 OutVelocity)
+{
+	SocketIndex = SanitizeFilteredSocketIndex_{ParameterName}(SocketIndex);
+	GetSocketData_{ParameterName}(SocketIndex, Interpolation, OutPosition, OutRotation, OutScale, OutVelocity);
+}
+
+void GetUnfilteredSocketTransformInterpolated_{ParameterName}(int SocketIndex, float Interpolation, out float3 OutPosition, out float4 OutRotation, out float3 OutScale, out float3 OutVelocity)
+{
+	SocketIndex = SanitizeUnfilteredSocketIndex_{ParameterName}(SocketIndex);
+	GetSocketData_{ParameterName}(SocketIndex, Interpolation, OutPosition, OutRotation, OutScale, OutVelocity);
+}
diff --git a/Engine/Plugins/FX/Niagara/Shaders/Private/NiagaraTransformUtils.ush b/Engine/Plugins/FX/Niagara/Shaders/Private/NiagaraTransformUtils.ush
new file mode 100644
--- /dev/null
+++ b/Engine/Plugins/FX/Niagara/Shaders/Private/NiagaraTransformUtils.ush
@@ -0,0 +1,39 @@
+// Copyright Epic Games, Inc. All Rights Reserved.
+
+#pragma once
+
+#include "/Plugin/FX/Niagara/Private/NiagaraQuaternionUtils.ush"
+
+struct FNiagaraTransform
+{
+	float3	Translation;
+	float4	Rotation;
+	float3	Scale;
+};
+
+FNiagaraTransform LoadTransform(ByteAddressBuffer Buffer, uint Offset)
+{
+	FNiagaraTransform R;
+	R.Translation	= asfloat(Buffer.Load3(Offset + 0));
+	R.Rotation		= asfloat(Buffer.Load4(Offset + 12));
+	R.Scale			= asfloat(Buffer.Load3(Offset + 20));
+	return R;
+}
+
+FNiagaraTransform MakeTransform(float3 Translation, float4 Rotation, float3 Scale)
+{
+	FNiagaraTransform R;
+	R.Translation	= Translation;
+	R.Rotation		= Rotation;
+	R.Scale			= Scale;
+	return R;
+}
+
+FNiagaraTransform MultiplyTransform(in FNiagaraTransform A, in FNiagaraTransform B)
+{
+	FNiagaraTransform R;
+	R.Translation	= RotateVectorByQuat(A.Translation * B.Scale, B.Rotation) + B.Translation;
+	R.Rotation		= MultiplyQuat(A.Rotation, B.Rotation);
+	R.Scale			= A.Scale * B.Scale;
+	return R;
+}
diff --git a/Engine/Plugins/FX/Niagara/Source/Niagara/Internal/DataInterface/NiagaraDataInterfaceSocketReader.h b/Engine/Plugins/FX/Niagara/Source/Niagara/Internal/DataInterface/NiagaraDataInterfaceSocketReader.h
new file mode 100644
--- /dev/null
+++ b/Engine/Plugins/FX/Niagara/Source/Niagara/Internal/DataInterface/NiagaraDataInterfaceSocketReader.h
@@ -0,0 +1,131 @@
+// Copyright Epic Games, Inc. All Rights Reserved.
+#pragma once
+
+#include "NiagaraCommon.h"
+#include "NiagaraShared.h"
+#include "NiagaraDataInterface.h"
+#include "NiagaraParameterStore.h"
+#include "NiagaraDataInterfaceSocketReader.generated.h"
+
+class AActor;
+
+UENUM()
+enum class ENDISocketReaderSourceMode : uint8
+{
+	// Reads in the order of Parameter Binding -> Attached Parent -> Source
+	Default,
+	// Read from the parameter binding only
+	ParameterBindingOnly,
+	// Read from the attached parent only
+	// This will traverse the attachment hierarchy
+	AttachedParentOnly,
+	// Read from the source only
+	// This will read the Source Actor first then Source Asset
+	SourceOnly,
+};
+
+/**
+Data interface for reading sockets from various sources.
+This can be from a live component in the scene or from a static / skeletal mesh asset.
+*/
+UCLASS(EditInlineNew, Category = "Actor", CollapseCategories, meta = (DisplayName = "Socket Reader Interface"), MinimalAPI)
+class UNiagaraDataInterfaceSocketReader : public UNiagaraDataInterface
+{
+	GENERATED_UCLASS_BODY()
+
+	BEGIN_SHADER_PARAMETER_STRUCT(FShaderParameters, )
+		SHADER_PARAMETER(uint32,							IsDataValid)
+		SHADER_PARAMETER(float,								InvDeltaSeconds)
+		SHADER_PARAMETER(int32,								NumSockets)
+		SHADER_PARAMETER(int32,								NumFilteredSockets)
+		SHADER_PARAMETER(int32,								NumUnfilteredSockets)
+		SHADER_PARAMETER(FVector3f,							ComponentToTranslatedWorld_Translation)
+		SHADER_PARAMETER(FQuat4f,							ComponentToTranslatedWorld_Rotation)
+		SHADER_PARAMETER(FVector3f,							ComponentToTranslatedWorld_Scale)
+		SHADER_PARAMETER(FVector3f,							PreviousComponentToTranslatedWorld_Translation)
+		SHADER_PARAMETER(FQuat4f,							PreviousComponentToTranslatedWorld_Rotation)
+		SHADER_PARAMETER(FVector3f,							PreviousComponentToTranslatedWorld_Scale)
+		SHADER_PARAMETER(uint32,							SocketTransformOffset)
+		SHADER_PARAMETER(uint32,							PreviousSocketTransformOffset)
+		SHADER_PARAMETER_RDG_BUFFER_SRV(ByteAddressBuffer,	SocketData)
+	END_SHADER_PARAMETER_STRUCT();
+
+public:
+	/** Controls how we find the object we want to read sockets from. */
+	UPROPERTY(EditAnywhere, Category = "SocketReader")
+	ENDISocketReaderSourceMode SourceMode = ENDISocketReaderSourceMode::Default;
+
+	/** List of filtered sockets. */
+	UPROPERTY(EditAnywhere, Category = "SocketReader")
+	TArray<FName> FilteredSockets;
+
+#if WITH_EDITORONLY_DATA
+	/** When previewing in the editor this is the asset to use to gather the socket information. */
+	UPROPERTY(EditAnywhere, Category = "SocketReader", meta = (AllowedClasses = "/Script/Engine.StaticMesh,/Script/Engine.SkeletalMesh"))
+	TSoftObjectPtr<UObject> EditorPreviewAsset;
+#endif
+
+	/** Source actor to read sockets from. */
+	UPROPERTY(EditAnywhere, Category = "SocketReader", meta = (EditConditionHides, EditCondition = "SourceMode == ENDISocketReaderSourceMode::Default || SourceMode == ENDISocketReaderSourceMode::SourceOnly"))
+	TLazyObjectPtr<AActor> SourceActor;
+
+	/** Source object asset to read sockets from, the transsform for these would be in relation to the Niagara system. */
+	UPROPERTY(EditAnywhere, Category = "SocketReader", meta = (EditConditionHides, EditCondition = "SourceMode == ENDISocketReaderSourceMode::Default || SourceMode == ENDISocketReaderSourceMode::SourceOnly", AllowedClasses= "/Script/Engine.StaticMesh,/Script/Engine.SkeletalMesh"))
+	TObjectPtr<UObject> SourceAsset;
+
+	/**
+	Source object parameter binding.
+	Note: Source Mode impacts the order of operations.
+	*/
+	UPROPERTY(EditAnywhere, Category = "SocketReader", meta = (EditConditionHides, EditCondition = "SourceMode == ENDISocketReaderSourceMode::Default || SourceMode == ENDISocketReaderSourceMode::ParameterBindingOnly"))
+	FNiagaraUserParameterBinding ObjectParameterBinding;
+
+	/**
+	When enabled we will read update the sockets transforms each frame.
+	This is not required in all cases as the sockets might not be able to move.
+	*/
+	UPROPERTY(EditAnywhere, Category = "Performance")
+	bool bUpdateSocketsPerFrame = true;
+
+	/** When this option is disabled, we use the previous frame's data for the skeletal mesh and can often issue the simulation early. This greatly
+	reduces overhead and allows the game thread to run faster, but comes at a tradeoff if the dependencies might leave gaps or other visual artifacts.*/
+	UPROPERTY(EditAnywhere, Category = "Performance")
+	bool bRequireCurrentFrameData = true;
+	
+	//UObject Interface
+	NIAGARA_API virtual void PostInitProperties() override;
+	//UObject Interface End
+
+	//UNiagaraDataInterface Interface
+	virtual bool CanExecuteOnTarget(ENiagaraSimTarget Target) const override { return true; }
+
+	NIAGARA_API virtual void GetVMExternalFunction(const FVMExternalFunctionBindingInfo& BindingInfo, void* InstanceData, FVMExternalFunction &OutFunc) override;
+#if WITH_EDITORONLY_DATA
+	NIAGARA_API virtual bool AppendCompileHash(FNiagaraCompileHashVisitor* InVisitor) const override;
+	NIAGARA_API virtual void GetParameterDefinitionHLSL(const FNiagaraDataInterfaceGPUParamInfo& ParamInfo, FString& OutHLSL) override;
+	NIAGARA_API virtual bool GetFunctionHLSL(const FNiagaraDataInterfaceGPUParamInfo& ParamInfo, const FNiagaraDataInterfaceGeneratedFunction& FunctionInfo, int FunctionInstanceIndex, FString& OutHLSL) override;
+#endif
+	NIAGARA_API virtual void BuildShaderParameters(FNiagaraShaderParametersBuilder& ShaderParametersBuilder) const override;
+	NIAGARA_API virtual void SetShaderParameters(const FNiagaraDataInterfaceSetShaderParametersContext& Context) const override;
+
+	NIAGARA_API virtual bool InitPerInstanceData(void* PerInstanceData, FNiagaraSystemInstance* SystemInstance) override;
+	NIAGARA_API virtual void DestroyPerInstanceData(void* PerInstanceData, FNiagaraSystemInstance* SystemInstance) override;
+	NIAGARA_API virtual int32 PerInstanceDataSize() const override;
+	NIAGARA_API virtual bool PerInstanceTick(void* PerInstanceData, FNiagaraSystemInstance* SystemInstance, float DeltaSeconds) override;
+
+	NIAGARA_API virtual void ProvidePerInstanceDataForRenderThread(void* DataForRenderThread, void* PerInstanceData, const FNiagaraSystemInstanceID& SystemInstance) override;
+
+	virtual bool HasTickGroupPrereqs() const override { return true; }
+	NIAGARA_API virtual ETickingGroup CalculateTickGroup(const void* PerInstanceData) const override;
+
+	NIAGARA_API virtual bool Equals(const UNiagaraDataInterface* Other) const override;
+	NIAGARA_API virtual bool CopyToInternal(UNiagaraDataInterface* Destination) const override;
+
+	virtual bool HasPreSimulateTick() const override { return true; }
+	//UNiagaraDataInterface Interface
+
+protected:
+#if WITH_EDITORONLY_DATA
+	virtual void GetFunctionsInternal(TArray<FNiagaraFunctionSignature>& OutFunctions) const override;
+#endif
+};
diff --git a/Engine/Plugins/FX/Niagara/Source/Niagara/Private/DataInterface/NiagaraDataInterfaceSocketReader.cpp b/Engine/Plugins/FX/Niagara/Source/Niagara/Private/DataInterface/NiagaraDataInterfaceSocketReader.cpp
new file mode 100644
--- /dev/null
+++ b/Engine/Plugins/FX/Niagara/Source/Niagara/Private/DataInterface/NiagaraDataInterfaceSocketReader.cpp
@@ -0,0 +1,835 @@
+// Copyright Epic Games, Inc. All Rights Reserved.
+
+#include "DataInterface/NiagaraDataInterfaceSocketReader.h"
+#include "Components/PrimitiveComponent.h"
+#include "GameFramework/Pawn.h"
+#include "NiagaraCompileHashVisitor.h"
+#include "NiagaraComponent.h"
+#include "NiagaraTypes.h"
+#include "NiagaraShaderParametersBuilder.h"
+#include "NiagaraSystemInstance.h"
+#include "NiagaraWorldManager.h"
+
+#include "Components/SceneComponent.h"
+#include "Engine/SkeletalMesh.h"
+#include "Engine/SkeletalMeshSocket.h"
+#include "Engine/StaticMesh.h"
+#include "Engine/StaticMeshSocket.h"
+#include "ShaderCompilerCore.h"
+
+#include "RenderGraphBuilder.h"
+#include "RenderGraphUtils.h"
+#include "UnifiedBuffer.h"
+
+#include UE_INLINE_GENERATED_CPP_BY_NAME(NiagaraDataInterfaceSocketReader)
+
+#define LOCTEXT_NAMESPACE "NiagaraDataInterfaceSocketReader"
+
+namespace NDISocketReaderLocal
+{
+	static const TCHAR* TemplateShaderFile = TEXT("/Plugin/FX/Niagara/Private/NiagaraDataInterfaceSocketReaderTemplate.ush");
+
+	static const FName	IsValidName("IsValid");
+
+	static const FName	GetSocketCountName("GetSocketCount");
+	static const FName	GetFilteredSocketCountName("GetFilteredSocketCount");
+	static const FName	GetUnfilteredSocketCountName("GetUnfilteredSocketCount");
+
+	static const FName	GetSocketTransformName("GetSocketTransform");
+	static const FName	GetFilteredSocketTransformName("GetFilteredSocketTransform");
+	static const FName	GetUnfilteredSocketTransformName("GetUnfilteredSocketTransform");
+
+	static const FName	GetSocketTransformInterpolatedName("GetSocketTransformInterpolated");
+	static const FName	GetFilteredSocketTransformInterpolatedName("GetFilteredSocketTransformInterpolated");
+	static const FName	GetUnfilteredSocketTransformInterpolatedName("GetUnfilteredSocketTransformInterpolated");
+
+	struct FInstanceData_GameThread
+	{
+		FNiagaraParameterDirectBinding<UObject*>	UserParamBinding;
+		bool										bIsDataValid = false;
+		bool										bNeedsSocketRecache = true;
+		bool										bNeedsRenderUpdate = true;
+
+		float										DeltaSeconds = 0.0f;
+		float										InvDeltaSeconds = 0.0f;
+
+		FTransform									ComponentToWorld;
+		FTransform									PreviousComponentToWorld;
+		FTransform3f								ComponentToTranslatedWorld;
+		FTransform3f								PreviousComponentToTranslatedWorld;
+
+		int32										NumSockets = 0;
+		int32										NumFilteredSockets = 0;
+		int32										NumUnfilteredSockets = 0;
+		TArray<int32>								SocketFilterUnfilteredIndex;
+		TArray<FName>								SocketNames;
+		TArray<FTransform3f>						SocketTransforms;
+		TArray<FTransform3f>						PreviousSocketTransforms;
+
+		TWeakObjectPtr<UObject>						ResolvedObject;
+	};
+
+	struct FInstanceData_SharedData
+	{
+		bool							bIsDataValid = false;
+		float							InvDeltaSeconds = 0.0f;
+		int32							NumSockets = 0;
+		int32							NumFilteredSockets = 0;
+		int32							NumUnfilteredSockets = 0;
+		FTransform3f					ComponentToTranslatedWorld = FTransform3f::Identity;
+		FTransform3f					PreviousComponentToTranslatedWorld = FTransform3f::Identity;
+		uint32							SocketTransformOffset = 0;
+		uint32							PreviousSocketTransformOffset = 0;
+	};
+
+	struct FGameToRenderInstanceData : public FInstanceData_SharedData
+	{
+		TArray<uint8>					DataToUpload;
+	};
+
+	struct FInstanceData_RenderThread : public FInstanceData_SharedData
+	{
+		TArray<uint8>					DataToUpload;
+		TRefCountPtr<FRDGPooledBuffer>	PooledBuffer;
+	};
+
+	struct FNDIProxy : public FNiagaraDataInterfaceProxy
+	{
+		virtual int32 PerInstanceDataPassedToRenderThreadSize() const override { return sizeof(FGameToRenderInstanceData); }
+
+		static void CopyTransforms(uint8* OutBuffer, TConstArrayView<FTransform3f> Transforms)
+		{
+			float* OutFloats = reinterpret_cast<float*>(OutBuffer);
+			for ( const FTransform3f& Transform : Transforms )
+			{
+				const FVector3f Translation = Transform.GetTranslation();
+				const FQuat4f Rotation = Transform.GetRotation();
+				const FVector3f Scale = Transform.GetScale3D();
+
+				OutFloats[0] = Translation.X;
+				OutFloats[1] = Translation.Y;
+				OutFloats[2] = Translation.Z;
+				OutFloats[3] = Rotation.X;
+				OutFloats[4] = Rotation.Y;
+				OutFloats[5] = Rotation.Z;
+				OutFloats[6] = Rotation.W;
+				OutFloats[7] = Scale.X;
+				OutFloats[8] = Scale.Y;
+				OutFloats[9] = Scale.Z;
+				OutFloats += 10;
+			}
+		}
+
+		static void ProvidePerInstanceDataForRenderThread(void* InDataForRenderThread, void* PerInstanceData, const FNiagaraSystemInstanceID& SystemInstance)
+		{
+			FInstanceData_GameThread* InstanceData = static_cast<FInstanceData_GameThread*>(PerInstanceData);
+			FGameToRenderInstanceData* DataForRenderThread = new(InDataForRenderThread) FGameToRenderInstanceData();
+
+			DataForRenderThread->bIsDataValid						= InstanceData->bIsDataValid;
+			DataForRenderThread->InvDeltaSeconds					= InstanceData->InvDeltaSeconds;
+			DataForRenderThread->NumSockets							= InstanceData->NumSockets;
+			DataForRenderThread->NumFilteredSockets					= InstanceData->NumFilteredSockets;
+			DataForRenderThread->NumUnfilteredSockets				= InstanceData->NumUnfilteredSockets;
+			DataForRenderThread->ComponentToTranslatedWorld			= InstanceData->ComponentToTranslatedWorld;
+			DataForRenderThread->PreviousComponentToTranslatedWorld	= InstanceData->PreviousComponentToTranslatedWorld;
+
+			if (InstanceData->bNeedsRenderUpdate)
+			{
+				InstanceData->bNeedsRenderUpdate = false;
+
+				const int32 TransformGpuSize = 10 * 4;
+				const int32 SocketTableSize = InstanceData->SocketFilterUnfilteredIndex.Num() * InstanceData->SocketFilterUnfilteredIndex.GetTypeSize();
+				const int32 SocketTransformSize = InstanceData->SocketTransforms.Num() * TransformGpuSize;
+				const int32 BufferSize = SocketTableSize + (SocketTransformSize * 2);
+				DataForRenderThread->SocketTransformOffset			= SocketTableSize;
+				DataForRenderThread->PreviousSocketTransformOffset	= SocketTableSize + SocketTransformSize;
+
+				DataForRenderThread->DataToUpload.SetNumUninitialized(BufferSize);
+				FMemory::Memcpy(DataForRenderThread->DataToUpload.GetData(), InstanceData->SocketFilterUnfilteredIndex.GetData(), SocketTableSize);
+				CopyTransforms(DataForRenderThread->DataToUpload.GetData() + DataForRenderThread->SocketTransformOffset, InstanceData->SocketTransforms);
+				CopyTransforms(DataForRenderThread->DataToUpload.GetData() + DataForRenderThread->PreviousSocketTransformOffset, InstanceData->PreviousSocketTransforms);
+			}
+		}
+
+		virtual void ConsumePerInstanceDataFromGameThread(void* PerInstanceData, const FNiagaraSystemInstanceID& InstanceID) override
+		{
+			FGameToRenderInstanceData* InstanceDataFromGT = static_cast<FGameToRenderInstanceData*>(PerInstanceData);
+			FInstanceData_RenderThread& InstanceData = SystemInstancesToInstanceData_RT.FindOrAdd(InstanceID);
+
+			InstanceData.bIsDataValid						= InstanceDataFromGT->bIsDataValid;
+			InstanceData.InvDeltaSeconds					= InstanceDataFromGT->InvDeltaSeconds;
+			InstanceData.NumSockets							= InstanceDataFromGT->NumSockets;
+			InstanceData.NumFilteredSockets					= InstanceDataFromGT->NumFilteredSockets;
+			InstanceData.NumUnfilteredSockets				= InstanceDataFromGT->NumUnfilteredSockets;
+			InstanceData.ComponentToTranslatedWorld			= InstanceDataFromGT->ComponentToTranslatedWorld;
+			InstanceData.PreviousComponentToTranslatedWorld	= InstanceDataFromGT->PreviousComponentToTranslatedWorld;
+			InstanceData.SocketTransformOffset				= InstanceDataFromGT->SocketTransformOffset;
+			InstanceData.PreviousSocketTransformOffset		= InstanceDataFromGT->PreviousSocketTransformOffset;
+
+			if (InstanceDataFromGT->DataToUpload.Num() > 0)
+			{
+				// If we got new data then swap in for any existing data
+				// We don't clear InstanceData.DataToUpload as that is consumed when we need the GPU buffer
+				InstanceData.DataToUpload					= MoveTemp(InstanceDataFromGT->DataToUpload);
+			}
+
+			InstanceDataFromGT->~FGameToRenderInstanceData();
+		}
+
+		virtual void PreStage(const FNDIGpuComputePreStageContext& Context)
+		{
+			const FNiagaraSystemInstanceID SystemInstanceID = Context.GetSystemInstanceID();
+			FInstanceData_RenderThread& InstanceData = SystemInstancesToInstanceData_RT.FindChecked(SystemInstanceID);
+			if (InstanceData.DataToUpload.Num() > 0)
+			{
+				FRDGBuilder& GraphBuilder = Context.GetGraphBuilder();
+
+				const void* InitialData = InstanceData.DataToUpload.GetData();
+				const uint64 InitialDataSize = InstanceData.DataToUpload.Num() * InstanceData.DataToUpload.GetTypeSize();
+				const uint64 BufferSize = Align(InitialDataSize, 16);
+
+				FRDGBufferDesc BufferDesc = FRDGBufferDesc::CreateByteAddressDesc(BufferSize);
+				ResizeBufferIfNeeded(GraphBuilder, InstanceData.PooledBuffer, BufferDesc, TEXT("NiagaraSocketReader"));
+
+				GraphBuilder.QueueBufferUpload(
+					GraphBuilder.RegisterExternalBuffer(InstanceData.PooledBuffer),
+					InitialData,
+					InitialDataSize,
+					[Data = MoveTemp(InstanceData.DataToUpload)](const void*) {}
+				);
+
+				InstanceData.DataToUpload.Empty();
+			}
+		}
+
+		TMap<FNiagaraSystemInstanceID, FInstanceData_RenderThread> SystemInstancesToInstanceData_RT;
+	};
+
+	//////////////////////////////////////////////////////////////////////////
+
+	void ResolveObject(const UNiagaraDataInterfaceSocketReader* SocketDI, FInstanceData_GameThread* InstanceData, FNiagaraSystemInstance* SystemInstance)
+	{
+		UObject* ResolvedObject = nullptr;
+
+		const bool bTryParameterBinding = SocketDI->SourceMode == ENDISocketReaderSourceMode::Default || SocketDI->SourceMode == ENDISocketReaderSourceMode::ParameterBindingOnly;
+		const bool bTryAttachedParent	= SocketDI->SourceMode == ENDISocketReaderSourceMode::Default || SocketDI->SourceMode == ENDISocketReaderSourceMode::AttachedParentOnly;
+		const bool bTrySource			= SocketDI->SourceMode == ENDISocketReaderSourceMode::Default || SocketDI->SourceMode == ENDISocketReaderSourceMode::SourceOnly;
+
+		// Resolve Parameter Binding?
+		if (bTryParameterBinding)
+		{
+			ResolvedObject = InstanceData->UserParamBinding.GetValue();
+			if (ResolvedObject)
+			{
+				if (AActor* AsActor = Cast<AActor>(ResolvedObject))
+				{
+					ResolvedObject = AsActor->GetRootComponent();
+				}
+			}
+		}
+
+		// Resolve Attached Parent?
+		if (bTryAttachedParent && !ResolvedObject)
+		{
+			USceneComponent* Component = SystemInstance->GetAttachComponent();
+			if (Component->IsA<UNiagaraComponent>())
+			{
+				ResolvedObject = Component->GetAttachParent();
+				//-TODO: Do we need to add tags or look until we find sockets?
+			}
+		}
+
+		// Resolve Source
+		if (bTrySource && !ResolvedObject)
+		{
+			ResolvedObject = SocketDI->SourceActor.Get();
+			if (ResolvedObject)
+			{
+				if (AActor* AsActor = Cast<AActor>(ResolvedObject))
+				{
+					ResolvedObject = AsActor->GetRootComponent();
+				}
+			}
+			if ( !ResolvedObject )
+			{
+				ResolvedObject = SocketDI->SourceAsset;
+			}
+		}
+
+	#if WITH_EDITORONLY_DATA
+		// Fallback for editor preview
+		if (!ResolvedObject && !SystemInstance->GetWorld()->IsGameWorld())
+		{
+			ResolvedObject = SocketDI->EditorPreviewAsset.LoadSynchronous();
+		}
+	#endif
+
+		// Store the resolved object
+		InstanceData->bIsDataValid			= ResolvedObject != nullptr;
+		InstanceData->bNeedsSocketRecache  |= ResolvedObject != InstanceData->ResolvedObject;
+		InstanceData->ResolvedObject		= ResolvedObject;
+	}
+
+	void UpdateSocketCache(const UNiagaraDataInterfaceSocketReader* SocketDI, FInstanceData_GameThread* InstanceData, FNiagaraSystemInstance* SystemInstance)
+	{
+		const bool bResetPreviousTransform = InstanceData->bNeedsSocketRecache;
+		const bool bNeedsSocketUpdate = SocketDI->bUpdateSocketsPerFrame || InstanceData->bNeedsSocketRecache;
+
+		// Recache the list of sockets
+		if (InstanceData->bNeedsSocketRecache)
+		{
+			InstanceData->bNeedsSocketRecache = false;
+			InstanceData->bNeedsRenderUpdate = true;
+
+			InstanceData->ComponentToTranslatedWorld = FTransform3f::Identity;
+			InstanceData->PreviousComponentToTranslatedWorld = FTransform3f::Identity;
+
+			InstanceData->NumSockets = 0;
+			InstanceData->NumFilteredSockets = 0;
+			InstanceData->NumUnfilteredSockets = 0;
+
+			InstanceData->SocketFilterUnfilteredIndex.Empty();
+			InstanceData->SocketNames.Empty();
+
+			if (USceneComponent* SceneComponent = Cast<USceneComponent>(InstanceData->ResolvedObject))
+			{
+				InstanceData->SocketNames = SceneComponent->GetAllSocketNames();
+			}
+			else if (UStaticMesh* StaticMesh = Cast<UStaticMesh>(InstanceData->ResolvedObject))
+			{
+				const int32 NumSockets = StaticMesh->Sockets.Num();
+				InstanceData->SocketNames.Reserve(StaticMesh->Sockets.Num());
+				for ( UStaticMeshSocket* Socket : StaticMesh->Sockets )
+				{
+					check(Socket)
+					InstanceData->SocketNames.Add(Socket->SocketName);
+				}
+			}
+			else if (USkeletalMesh* SkeletalMesh = Cast<USkeletalMesh>(InstanceData->ResolvedObject))
+			{
+				const FReferenceSkeleton& RefSkeleton = SkeletalMesh->GetRefSkeleton();
+				const TArray<FMeshBoneInfo>& RefBoneInfo = RefSkeleton.GetRefBoneInfo();
+				const int32 NumBones = RefBoneInfo.Num();
+				const int32 NumSockets = SkeletalMesh->NumSockets();
+
+				InstanceData->SocketNames.Reserve(NumBones + NumSockets);
+				for (int i = 0; i < NumBones; ++i)
+				{
+					InstanceData->SocketNames.Add(RefBoneInfo[i].Name);
+				}
+
+				for (int i = 0; i < NumSockets; ++i)
+				{
+					USkeletalMeshSocket* Socket = SkeletalMesh->GetSocketByIndex(i);
+					check(Socket);
+					InstanceData->SocketNames.Add(Socket->SocketName);
+				}
+			}
+
+			InstanceData->NumSockets = InstanceData->SocketNames.Num();
+			if (InstanceData->NumSockets > 0 )
+			{
+				InstanceData->NumFilteredSockets = SocketDI->FilteredSockets.Num();
+
+				TBitArray<TInlineAllocator<2>> IsFilteredSocket;
+				IsFilteredSocket.Add(false, InstanceData->NumSockets);
+
+				InstanceData->SocketFilterUnfilteredIndex.Reserve(FMath::Max(InstanceData->NumSockets, InstanceData->NumFilteredSockets) + 1);
+				InstanceData->SocketFilterUnfilteredIndex.Add(InstanceData->NumSockets);
+
+				for (FName SocketName : SocketDI->FilteredSockets)
+				{
+					const int32 SocketIndex = InstanceData->SocketNames.IndexOfByKey(SocketName);
+					if (InstanceData->SocketNames.IsValidIndex(SocketIndex))
+					{
+						InstanceData->SocketFilterUnfilteredIndex.Add(SocketIndex);
+						IsFilteredSocket[SocketIndex] = true;
+					}
+					else
+					{
+						InstanceData->SocketFilterUnfilteredIndex.Add(InstanceData->NumSockets);
+					}
+				}
+
+				for (int32 i = 0; i < InstanceData->NumSockets; ++i)
+				{
+					if (!IsFilteredSocket[i])
+					{
+						InstanceData->SocketFilterUnfilteredIndex.Add(i);
+						++InstanceData->NumUnfilteredSockets;
+					}
+				}
+			}
+			else
+			{
+				InstanceData->SocketFilterUnfilteredIndex.Add(InstanceData->NumSockets);
+			}
+			InstanceData->SocketTransforms.SetNum(InstanceData->NumSockets + 1);
+			InstanceData->PreviousSocketTransforms.SetNum(InstanceData->NumSockets + 1);
+		}
+
+		// Update the socket data
+		if (bNeedsSocketUpdate && InstanceData->NumSockets > 0)
+		{
+			// If we are not resetting copy the current transform to previous
+			if (!bResetPreviousTransform)
+			{
+				InstanceData->PreviousSocketTransforms = InstanceData->SocketTransforms;
+			}
+
+			if (USceneComponent* SceneComponent = Cast<USceneComponent>(InstanceData->ResolvedObject))
+			{
+				InstanceData->bNeedsRenderUpdate = true;
+				for (int32 i = 0; i < InstanceData->NumSockets; ++i)
+				{
+					const FTransform SocketTransform			= SceneComponent->GetSocketTransform(InstanceData->SocketNames[i], ERelativeTransformSpace::RTS_Component);
+					InstanceData->SocketTransforms[i]			= FTransform3f(SocketTransform);
+				}
+			}
+			else if (UStaticMesh* StaticMesh = Cast<UStaticMesh>(InstanceData->ResolvedObject))
+			{
+				InstanceData->bNeedsRenderUpdate = true;
+
+				for (int32 i = 0; i < InstanceData->NumSockets; ++i)
+				{
+					UStaticMeshSocket* Socket = StaticMesh->Sockets[i];
+					InstanceData->SocketTransforms[i].SetTranslation(FVector3f(Socket->RelativeLocation));
+					InstanceData->SocketTransforms[i].SetRotation(FQuat4f(Socket->RelativeRotation.Quaternion()));
+					InstanceData->SocketTransforms[i].SetScale3D(FVector3f(Socket->RelativeScale));
+				}
+			}
+			else if (USkeletalMesh* SkeletalMesh = Cast<USkeletalMesh>(InstanceData->ResolvedObject))
+			{
+				InstanceData->bNeedsRenderUpdate = true;
+
+				const FReferenceSkeleton& RefSkeleton = SkeletalMesh->GetRefSkeleton();
+				const TArray<FMeshBoneInfo>& RefBoneInfo = RefSkeleton.GetRefBoneInfo();
+				const TArray<FTransform>& RefBonePose = RefSkeleton.GetRefBonePose();
+				const int32 NumBones = RefBoneInfo.Num();
+				const int32 NumSockets = SkeletalMesh->NumSockets();
+
+				for (int i = 0; i < NumBones; ++i)
+				{
+					FTransform3f SocketTransform = FTransform3f(RefBonePose[i]);
+					const int32 ParentIndex = RefBoneInfo[i].ParentIndex;
+					if (ParentIndex >= 0 && ParentIndex < i)
+					{
+						SocketTransform = SocketTransform * InstanceData->SocketTransforms[ParentIndex];
+					}
+					InstanceData->SocketTransforms[i] = FTransform3f(SocketTransform);
+					InstanceData->PreviousSocketTransforms[i] = InstanceData->SocketTransforms[i];
+				}
+
+				for (int i = 0; i < NumSockets; ++i)
+				{
+					USkeletalMeshSocket* Socket = SkeletalMesh->GetSocketByIndex(i);
+					check(Socket);
+
+					FTransform3f SocketTransform(FTransform(Socket->RelativeRotation.Quaternion(), Socket->RelativeLocation, Socket->RelativeScale));
+					const int32 ParentIndex = RefSkeleton.FindBoneIndex(Socket->BoneName);
+					if (ParentIndex >= 0)
+					{
+						SocketTransform = SocketTransform * InstanceData->SocketTransforms[ParentIndex];
+					}
+					InstanceData->SocketTransforms[i + NumBones] = FTransform3f(SocketTransform);
+				}
+			}
+
+			// Transforms are being reset so copy current -> previous
+			if (bResetPreviousTransform)
+			{
+				InstanceData->PreviousSocketTransforms = InstanceData->SocketTransforms;
+			}
+		}
+
+		USceneComponent* SceneComponent = Cast<USceneComponent>(InstanceData->ResolvedObject);
+		if ( !SceneComponent )
+		{
+			SceneComponent = SystemInstance->GetAttachComponent();
+		}
+
+		if ( SceneComponent )
+		{
+			if (bResetPreviousTransform)
+			{
+				InstanceData->ComponentToWorld			= SceneComponent->GetComponentToWorld();
+				InstanceData->PreviousComponentToWorld	= InstanceData->ComponentToWorld;
+			}
+			else
+			{
+				InstanceData->PreviousComponentToWorld	= InstanceData->ComponentToWorld;
+				InstanceData->ComponentToWorld			= SceneComponent->GetComponentToWorld();
+			}
+			const FNiagaraLWCConverter LWCConverter = SystemInstance->GetLWCConverter();
+			InstanceData->ComponentToTranslatedWorld = LWCConverter.ConvertWorldToSimulationTransform(InstanceData->ComponentToWorld);
+			InstanceData->PreviousComponentToTranslatedWorld = LWCConverter.ConvertWorldToSimulationTransform(InstanceData->PreviousComponentToWorld);
+		}
+	}
+
+	//////////////////////////////////////////////////////////////////////////
+
+	enum class ESocketReadType
+	{
+		Default,
+		Filtered,
+		Unfiltered,
+	};
+
+	void VMIsValid(FVectorVMExternalFunctionContext& Context)
+	{
+		VectorVM::FUserPtrHandler<FInstanceData_GameThread> InstanceData(Context);
+		FNDIOutputParam<bool>	OutIsValid(Context);
+
+		for (int32 i = 0; i < Context.GetNumInstances(); ++i)
+		{
+			OutIsValid.SetAndAdvance(InstanceData->bIsDataValid);
+		}
+	}
+
+	template<ESocketReadType ReadType>
+	void VMGetSocketCount(FVectorVMExternalFunctionContext& Context)
+	{
+		VectorVM::FUserPtrHandler<FInstanceData_GameThread> InstanceData(Context);
+		FNDIOutputParam<int32>	OutCount(Context);
+
+		int32 SocketCount = 0;
+		switch (ReadType)
+		{
+			case ESocketReadType::Filtered:		SocketCount = InstanceData->NumFilteredSockets; break;
+			case ESocketReadType::Unfiltered:	SocketCount = InstanceData->NumUnfilteredSockets; break;
+			default:							SocketCount = InstanceData->NumSockets; break;
+		}
+
+		for (int32 i=0; i < Context.GetNumInstances(); ++i)
+		{
+			OutCount.SetAndAdvance(SocketCount);
+		}
+	}
+
+	template<ESocketReadType ReadType, bool bInterpolated>
+	void VMGetSocketTransform(FVectorVMExternalFunctionContext& Context)
+	{
+		VectorVM::FUserPtrHandler<FInstanceData_GameThread> InstanceData(Context);
+		FNDIInputParam<int32>		InSocketIndex(Context);
+		FNDIInputParam<float>		InInterpolation;
+		if constexpr ( bInterpolated )
+		{
+			InInterpolation.Init(Context);
+		}
+
+		FNDIOutputParam<FVector3f>	OutPosition(Context);
+		FNDIOutputParam<FQuat4f>	OutRotation(Context);
+		FNDIOutputParam<FVector3f>	OutScale(Context);
+		FNDIOutputParam<FVector3f>	OutVelocity(Context);
+
+		for (int32 i=0; i < Context.GetNumInstances(); ++i)
+		{
+			int32 SocketIndex = InSocketIndex.GetAndAdvance();
+			const float Interpolation = bInterpolated ? InInterpolation.GetAndAdvance() : 1.0f;
+
+			switch (ReadType)
+			{
+				case ESocketReadType::Filtered:
+				{
+					const int32 FilteredSockedIndex = SocketIndex >= 0 && SocketIndex < InstanceData->NumFilteredSockets ? (SocketIndex + 1) : 0;
+					SocketIndex =  InstanceData->SocketFilterUnfilteredIndex[FilteredSockedIndex];
+					break;
+				}
+
+				case ESocketReadType::Unfiltered:
+				{
+					const int32 UnfilteredSockedIndex = SocketIndex >= 0 && SocketIndex < InstanceData->NumFilteredSockets ? (SocketIndex + 1 + InstanceData->NumFilteredSockets) : 0;
+					SocketIndex = InstanceData->SocketFilterUnfilteredIndex[UnfilteredSockedIndex];
+					break;
+				}
+
+				default:
+					SocketIndex = SocketIndex >= 0 && SocketIndex < InstanceData->NumSockets ? SocketIndex : InstanceData->NumSockets;
+					break;
+			}
+
+			const FTransform3f PreviousSocketTransform = InstanceData->PreviousSocketTransforms[SocketIndex] * InstanceData->PreviousComponentToTranslatedWorld;
+			const FTransform3f SocketTransform = InstanceData->SocketTransforms[SocketIndex] * InstanceData->ComponentToTranslatedWorld;
+
+			FVector3f Position	= SocketTransform.GetTranslation();
+			FQuat4f Rotation	= SocketTransform.GetRotation();
+			FVector3f Scale		= SocketTransform.GetScale3D();
+			FVector3f Velocity	= (Position - PreviousSocketTransform.GetTranslation()) * InstanceData->InvDeltaSeconds;
+
+			if constexpr (bInterpolated)
+			{
+				Position	= FMath::Lerp(PreviousSocketTransform.GetTranslation(), Position, Interpolation);
+				Rotation	= FQuat4f::Slerp(PreviousSocketTransform.GetRotation(), Rotation, Interpolation);
+				Scale		= FMath::Lerp(PreviousSocketTransform.GetScale3D(), Scale, Interpolation);
+			}
+
+			OutPosition.SetAndAdvance(Position);
+			OutRotation.SetAndAdvance(Rotation);
+			OutScale.SetAndAdvance(Scale);
+			OutVelocity.SetAndAdvance(Velocity);
+		}
+	}
+
+	static const TMap<FName, FVMExternalFunction> VMFunctionBindings =
+	{
+		{IsValidName,									FVMExternalFunction::CreateStatic(VMIsValid)},
+		{GetSocketCountName,							FVMExternalFunction::CreateStatic(VMGetSocketCount<ESocketReadType::Default>)},
+		{GetFilteredSocketCountName,					FVMExternalFunction::CreateStatic(VMGetSocketCount<ESocketReadType::Filtered>)},
+		{GetUnfilteredSocketCountName,					FVMExternalFunction::CreateStatic(VMGetSocketCount<ESocketReadType::Unfiltered>)},
+		{GetSocketTransformName,						FVMExternalFunction::CreateStatic(VMGetSocketTransform<ESocketReadType::Default, false>)},
+		{GetFilteredSocketTransformName,				FVMExternalFunction::CreateStatic(VMGetSocketTransform<ESocketReadType::Filtered, false>)},
+		{GetUnfilteredSocketTransformName,				FVMExternalFunction::CreateStatic(VMGetSocketTransform<ESocketReadType::Unfiltered, false>)},
+		{GetSocketTransformInterpolatedName,			FVMExternalFunction::CreateStatic(VMGetSocketTransform<ESocketReadType::Default, true>)},
+		{GetFilteredSocketTransformInterpolatedName,	FVMExternalFunction::CreateStatic(VMGetSocketTransform<ESocketReadType::Filtered, true>)},
+		{GetUnfilteredSocketTransformInterpolatedName,	FVMExternalFunction::CreateStatic(VMGetSocketTransform<ESocketReadType::Unfiltered, true>)},
+	};
+}
+
+//////////////////////////////////////////////////////////////////////////
+// Data Interface
+UNiagaraDataInterfaceSocketReader::UNiagaraDataInterfaceSocketReader(FObjectInitializer const& ObjectInitializer)
+	: Super(ObjectInitializer)
+{
+	Proxy.Reset(new NDISocketReaderLocal::FNDIProxy());
+
+	FNiagaraTypeDefinition Def(UObject::StaticClass());
+	ObjectParameterBinding.Parameter.SetType(Def);
+}
+
+void UNiagaraDataInterfaceSocketReader::PostInitProperties()
+{
+	Super::PostInitProperties();
+
+	if (HasAnyFlags(RF_ClassDefaultObject))
+	{
+		ENiagaraTypeRegistryFlags Flags = ENiagaraTypeRegistryFlags::AllowAnyVariable | ENiagaraTypeRegistryFlags::AllowParameter;
+		FNiagaraTypeRegistry::Register(FNiagaraTypeDefinition(GetClass()), Flags);
+	}
+}
+
+#if WITH_EDITORONLY_DATA
+void UNiagaraDataInterfaceSocketReader::GetFunctionsInternal(TArray<FNiagaraFunctionSignature>& OutFunctions) const
+{
+	using namespace NDISocketReaderLocal;
+	FNiagaraFunctionSignature ImmutableSig;
+	ImmutableSig.Inputs.Emplace(FNiagaraTypeDefinition(GetClass()), TEXT("SocketReader"));
+	ImmutableSig.bMemberFunction = true;
+	ImmutableSig.bRequiresContext = false;
+	ImmutableSig.bSupportsGPU = true;
+	{
+		FNiagaraFunctionSignature& FunctionSignature = OutFunctions.Add_GetRef(ImmutableSig);
+		FunctionSignature.Name = IsValidName;
+		FunctionSignature.Outputs.Emplace(FNiagaraTypeDefinition::GetBoolDef(), TEXT("IsValid"));
+	}
+	{
+		FNiagaraFunctionSignature FunctionSignature = ImmutableSig;
+		FunctionSignature.Outputs.Emplace(FNiagaraTypeDefinition::GetIntDef(), TEXT("Count"));
+
+		OutFunctions.Add_GetRef(FunctionSignature).Name = GetSocketCountName;
+		OutFunctions.Add_GetRef(FunctionSignature).Name = GetFilteredSocketCountName;
+		OutFunctions.Add_GetRef(FunctionSignature).Name = GetUnfilteredSocketCountName;
+	}
+	{
+		FNiagaraFunctionSignature FunctionSignature = ImmutableSig;
+		FunctionSignature.Inputs.Emplace(FNiagaraTypeDefinition::GetIntDef(), TEXT("SocketIndex"));
+		FunctionSignature.Outputs.Emplace(FNiagaraTypeDefinition::GetPositionDef(), TEXT("Position"));
+		FunctionSignature.Outputs.Emplace(FNiagaraTypeDefinition::GetQuatDef(), TEXT("Rotation"));
+		FunctionSignature.Outputs.Emplace(FNiagaraTypeDefinition::GetVec3Def(), TEXT("Scale"));
+		FunctionSignature.Outputs.Emplace(FNiagaraTypeDefinition::GetVec3Def(), TEXT("Velocity"));
+
+		OutFunctions.Add_GetRef(FunctionSignature).Name = GetSocketTransformName;
+		OutFunctions.Add_GetRef(FunctionSignature).Name = GetFilteredSocketTransformName;
+		OutFunctions.Add_GetRef(FunctionSignature).Name = GetUnfilteredSocketTransformName;
+
+		FunctionSignature.Inputs.Emplace(FNiagaraTypeDefinition::GetFloatDef(), TEXT("Interpolation"));
+		OutFunctions.Add_GetRef(FunctionSignature).Name = GetSocketTransformInterpolatedName;
+		OutFunctions.Add_GetRef(FunctionSignature).Name = GetFilteredSocketTransformInterpolatedName;
+		OutFunctions.Add_GetRef(FunctionSignature).Name = GetUnfilteredSocketTransformInterpolatedName;
+	}
+}
+#endif
+
+void UNiagaraDataInterfaceSocketReader::GetVMExternalFunction(const FVMExternalFunctionBindingInfo& BindingInfo, void* InstanceData, FVMExternalFunction& OutFunc)
+{
+	using namespace NDISocketReaderLocal;
+	OutFunc = VMFunctionBindings.FindRef(BindingInfo.Name);
+}
+
+#if WITH_EDITORONLY_DATA
+bool UNiagaraDataInterfaceSocketReader::AppendCompileHash(FNiagaraCompileHashVisitor* InVisitor) const
+{
+	bool bSuccess = Super::AppendCompileHash(InVisitor);
+	bSuccess &= InVisitor->UpdateShaderFile(NDISocketReaderLocal::TemplateShaderFile);
+	bSuccess &= InVisitor->UpdateShaderParameters<FShaderParameters>();
+	return bSuccess;
+}
+
+void UNiagaraDataInterfaceSocketReader::GetParameterDefinitionHLSL(const FNiagaraDataInterfaceGPUParamInfo& ParamInfo, FString& OutHLSL)
+{
+	TMap<FString, FStringFormatArg> TemplateArgs =
+	{
+		{TEXT("ParameterName"),	ParamInfo.DataInterfaceHLSLSymbol},
+	};
+	AppendTemplateHLSL(OutHLSL, NDISocketReaderLocal::TemplateShaderFile, TemplateArgs);
+}
+
+bool UNiagaraDataInterfaceSocketReader::GetFunctionHLSL(const FNiagaraDataInterfaceGPUParamInfo& ParamInfo, const FNiagaraDataInterfaceGeneratedFunction& FunctionInfo, int FunctionInstanceIndex, FString& OutHLSL)
+{
+	using namespace NDISocketReaderLocal;
+	//return (FunctionInfo.DefinitionName == GetMatrixName) || (FunctionInfo.DefinitionName == GetTransformName) || (FunctionInfo.DefinitionName == GetVelocityName);
+	return true;
+}
+#endif
+
+void UNiagaraDataInterfaceSocketReader::BuildShaderParameters(FNiagaraShaderParametersBuilder& ShaderParametersBuilder) const
+{
+	ShaderParametersBuilder.AddNestedStruct<FShaderParameters>();
+}
+
+void UNiagaraDataInterfaceSocketReader::SetShaderParameters(const FNiagaraDataInterfaceSetShaderParametersContext& Context) const
+{
+	using namespace NDISocketReaderLocal;
+
+	FNDIProxy& DIProxy = Context.GetProxy<FNDIProxy>();
+	FInstanceData_RenderThread& InstanceData = DIProxy.SystemInstancesToInstanceData_RT.FindChecked(Context.GetSystemInstanceID());
+
+	FRDGBuilder& GraphBuilder = Context.GetGraphBuilder();
+	FRDGBufferRef RDGBuffer = GraphBuilder.RegisterExternalBuffer(InstanceData.PooledBuffer);
+	FRDGBufferSRVRef RDGBufferSRV = GraphBuilder.CreateSRV(RDGBuffer);
+
+	FShaderParameters* ShaderParameters = Context.GetParameterNestedStruct<FShaderParameters>();
+	ShaderParameters->IsDataValid										= InstanceData.bIsDataValid ? 1 : 0;
+	ShaderParameters->InvDeltaSeconds									= InstanceData.InvDeltaSeconds;
+	ShaderParameters->NumSockets										= InstanceData.NumSockets;
+	ShaderParameters->NumFilteredSockets								= InstanceData.NumFilteredSockets;
+	ShaderParameters->NumUnfilteredSockets								= InstanceData.NumUnfilteredSockets;
+	ShaderParameters->ComponentToTranslatedWorld_Translation			= InstanceData.ComponentToTranslatedWorld.GetTranslation();
+	ShaderParameters->ComponentToTranslatedWorld_Rotation				= InstanceData.ComponentToTranslatedWorld.GetRotation();
+	ShaderParameters->ComponentToTranslatedWorld_Scale					= InstanceData.ComponentToTranslatedWorld.GetScale3D();
+	ShaderParameters->PreviousComponentToTranslatedWorld_Translation	= InstanceData.PreviousComponentToTranslatedWorld.GetTranslation();
+	ShaderParameters->PreviousComponentToTranslatedWorld_Rotation		= InstanceData.PreviousComponentToTranslatedWorld.GetRotation();
+	ShaderParameters->PreviousComponentToTranslatedWorld_Scale			= InstanceData.PreviousComponentToTranslatedWorld.GetScale3D();
+	ShaderParameters->SocketTransformOffset								= InstanceData.SocketTransformOffset;
+	ShaderParameters->PreviousSocketTransformOffset						= InstanceData.PreviousSocketTransformOffset;
+	ShaderParameters->SocketData										= RDGBufferSRV;
+}
+
+bool UNiagaraDataInterfaceSocketReader::InitPerInstanceData(void* PerInstanceData, FNiagaraSystemInstance* SystemInstance)
+{
+	using namespace NDISocketReaderLocal;
+
+	FInstanceData_GameThread* InstanceData = new (PerInstanceData) FInstanceData_GameThread;
+	InstanceData->UserParamBinding.Init(SystemInstance->GetInstanceParameters(), ObjectParameterBinding.Parameter);
+	ResolveObject(this, InstanceData, SystemInstance);
+
+	return true;
+}
+
+void UNiagaraDataInterfaceSocketReader::DestroyPerInstanceData(void* PerInstanceData, FNiagaraSystemInstance* SystemInstance)
+{
+	using namespace NDISocketReaderLocal;
+
+	FInstanceData_GameThread* InstanceData = (FInstanceData_GameThread*)PerInstanceData;
+	InstanceData->~FInstanceData_GameThread();
+
+	ENQUEUE_RENDER_COMMAND(RemoveProxy)
+	(
+		[RT_Proxy=GetProxyAs<FNDIProxy>(), InstanceID=SystemInstance->GetId()](FRHICommandListImmediate& CmdList)
+		{
+			RT_Proxy->SystemInstancesToInstanceData_RT.Remove(InstanceID);
+		}
+	);
+}
+
+int32 UNiagaraDataInterfaceSocketReader::PerInstanceDataSize() const
+{
+	return sizeof(NDISocketReaderLocal::FInstanceData_GameThread);
+}
+
+bool UNiagaraDataInterfaceSocketReader::PerInstanceTick(void* PerInstanceData, FNiagaraSystemInstance* SystemInstance, float DeltaSeconds)
+{
+	using namespace NDISocketReaderLocal;
+
+	FInstanceData_GameThread* InstanceData = static_cast<FInstanceData_GameThread*>(PerInstanceData);
+	InstanceData->DeltaSeconds = DeltaSeconds;
+	InstanceData->InvDeltaSeconds = DeltaSeconds > 0.0f ? 1.0f / DeltaSeconds : 0.0f;
+	ResolveObject(this, InstanceData, SystemInstance);
+	UpdateSocketCache(this, InstanceData, SystemInstance);
+
+	return false;
+}
+
+void UNiagaraDataInterfaceSocketReader::ProvidePerInstanceDataForRenderThread(void* DataForRenderThread, void* PerInstanceData, const FNiagaraSystemInstanceID& SystemInstance)
+{
+	using namespace NDISocketReaderLocal;
+	FNDIProxy::ProvidePerInstanceDataForRenderThread(DataForRenderThread, PerInstanceData, SystemInstance);
+}
+
+ETickingGroup UNiagaraDataInterfaceSocketReader::CalculateTickGroup(const void* PerInstanceData) const
+{
+	using namespace NDISocketReaderLocal;
+	if ( bRequireCurrentFrameData )
+	{
+		const FInstanceData_GameThread* InstanceData = static_cast<const FInstanceData_GameThread*>(PerInstanceData);
+		const USceneComponent* SceneComponent = InstanceData ? Cast<USceneComponent>(InstanceData->ResolvedObject.Get()) : nullptr;
+		if (SceneComponent)
+		{
+			ETickingGroup FinalTickGroup = FMath::Max(SceneComponent->PrimaryComponentTick.TickGroup, SceneComponent->PrimaryComponentTick.EndTickGroup);
+	//		//-TODO: Do we need to do this?
+	//		//if ( USkeletalMeshComponent* SkelMeshComponent = Cast<USkeletalMeshComponent>(ActorComponent) )
+	//		//{
+	//		//	if (SkelMeshComponent->bBlendPhysics)
+	//		//	{
+	//		//		FinalTickGroup = FMath::Max(FinalTickGroup, TG_EndPhysics);
+	//		//	}
+	//		//}
+			FinalTickGroup = FMath::Clamp(ETickingGroup(FinalTickGroup + 1), NiagaraFirstTickGroup, NiagaraLastTickGroup);
+			return FinalTickGroup;
+		}
+	}
+	return NiagaraFirstTickGroup;
+}
+
+bool UNiagaraDataInterfaceSocketReader::Equals(const UNiagaraDataInterface* Other) const
+{
+	if (!Super::Equals(Other))
+	{
+		return false;
+	}
+
+	const UNiagaraDataInterfaceSocketReader* OtherTyped = CastChecked<const UNiagaraDataInterfaceSocketReader>(Other);
+	return OtherTyped->SourceMode == SourceMode
+		&& OtherTyped->FilteredSockets == FilteredSockets
+	#if WITH_EDITORONLY_DATA
+		&& OtherTyped->EditorPreviewAsset == EditorPreviewAsset
+	#endif
+		&& OtherTyped->SourceActor == SourceActor
+		&& OtherTyped->SourceAsset == SourceAsset
+		&& OtherTyped->ObjectParameterBinding == ObjectParameterBinding
+		&& OtherTyped->bUpdateSocketsPerFrame == bUpdateSocketsPerFrame
+		&& OtherTyped->bRequireCurrentFrameData == bRequireCurrentFrameData;
+}
+
+bool UNiagaraDataInterfaceSocketReader::CopyToInternal(UNiagaraDataInterface* Destination) const
+{
+	if (!Super::CopyToInternal(Destination))
+	{
+		return false;
+	}
+
+	UNiagaraDataInterfaceSocketReader* OtherTyped = CastChecked<UNiagaraDataInterfaceSocketReader>(Destination);
+	OtherTyped->SourceMode = SourceMode;
+	OtherTyped->FilteredSockets = FilteredSockets;
+#if WITH_EDITORONLY_DATA
+	OtherTyped->EditorPreviewAsset = EditorPreviewAsset;
+#endif
+	OtherTyped->SourceActor = SourceActor;
+	OtherTyped->SourceAsset = SourceAsset;
+	OtherTyped->ObjectParameterBinding = ObjectParameterBinding;
+	OtherTyped->bUpdateSocketsPerFrame = bUpdateSocketsPerFrame;
+	OtherTyped->bRequireCurrentFrameData = bRequireCurrentFrameData;
+	return true;
+}
+
+#undef LOCTEXT_NAMESPACE
+
diff --git a/Engine/Plugins/FX/Niagara/Source/Niagara/Private/NiagaraTypes.cpp b/Engine/Plugins/FX/Niagara/Source/Niagara/Private/NiagaraTypes.cpp
--- a/Engine/Plugins/FX/Niagara/Source/Niagara/Private/NiagaraTypes.cpp
+++ b/Engine/Plugins/FX/Niagara/Source/Niagara/Private/NiagaraTypes.cpp
@@ -161,6 +161,24 @@ FMatrix FNiagaraLWCConverter::ConvertSimulationToWorldMatrix(const FMatrix& Matr
 	return Matrix.ConcatTranslation(SystemWorldPos);
 }
 
+FTransform3f FNiagaraLWCConverter::ConvertWorldToSimulationTransform(const FTransform& Transform) const
+{
+	FTransform3f NewTransform;
+	NewTransform.SetTranslation(FVector3f(Transform.GetTranslation() - SystemWorldPos));
+	NewTransform.SetRotation(FQuat4f(Transform.GetRotation()));
+	NewTransform.SetScale3D(FVector3f(Transform.GetScale3D()));
+	return NewTransform;
+}
+
+FTransform FNiagaraLWCConverter::ConvertSimulationToWorldTransform(const FTransform3f& Transform) const
+{
+	FTransform NewTransform;
+	NewTransform.SetTranslation(FVector(Transform.GetTranslation()) + SystemWorldPos);
+	NewTransform.SetRotation(FQuat(Transform.GetRotation()));
+	NewTransform.SetScale3D(FVector(Transform.GetScale3D()));
+	return NewTransform;
+}
+
 FNiagaraStructConversionStep::FNiagaraStructConversionStep()
 {
 }
diff --git a/Engine/Plugins/FX/Niagara/Source/Niagara/Public/NiagaraTypes.h b/Engine/Plugins/FX/Niagara/Source/Niagara/Public/NiagaraTypes.h
--- a/Engine/Plugins/FX/Niagara/Source/Niagara/Public/NiagaraTypes.h
+++ b/Engine/Plugins/FX/Niagara/Source/Niagara/Public/NiagaraTypes.h
@@ -245,6 +245,9 @@ struct FNiagaraLWCConverter
 	[[nodiscard]] NIAGARA_API FMatrix ConvertWorldToSimulationMatrix(const FMatrix& Matrix) const;
 	[[nodiscard]] NIAGARA_API FMatrix ConvertSimulationToWorldMatrix(const FMatrix& Matrix) const;
 
+	[[nodiscard]] NIAGARA_API FTransform3f ConvertWorldToSimulationTransform(const FTransform& Transform) const;
+	[[nodiscard]] NIAGARA_API FTransform ConvertSimulationToWorldTransform(const FTransform3f& Transform) const;
+
 private:
 	FVector SystemWorldPos;
 };
