Fix FNiagaraRendererMeshes::GetDynamicRayTracingInstances(...) not setting NumTransforms when SourceMode == ENiagaraRendererSourceDataMode::Emitter.

diff --git a/Engine/Plugins/FX/Niagara/Source/Niagara/Private/NiagaraRendererMeshes.cpp b/Engine/Plugins/FX/Niagara/Source/Niagara/Private/NiagaraRendererMeshes.cpp
--- a/Engine/Plugins/FX/Niagara/Source/Niagara/Private/NiagaraRendererMeshes.cpp
+++ b/Engine/Plugins/FX/Niagara/Source/Niagara/Private/NiagaraRendererMeshes.cpp
@@ -1854,6 +1854,7 @@ void FNiagaraRendererMeshes::GetDynamicRayTracingInstances(FRayTracingMaterialGa
 			}
 
 			RayTracingInstance.InstanceTransforms.Add(InstanceTransform);
+			RayTracingInstance.NumTransforms = 1;
 		}
 		else
 		{
