Added Deprecated & Hidden tags & View Options in the Niagara Asset Browser
Hidden tag will stop an asset from showing up, but will validate if used regardless
Deprecated tag will stop an asset from showing up, and will not pass validation
Both tags have additional view options to make them show up in the NAB

diff --git a/Engine/Plugins/FX/Niagara/Source/Niagara/Classes/NiagaraEmitter.h b/Engine/Plugins/FX/Niagara/Source/Niagara/Classes/NiagaraEmitter.h
--- a/Engine/Plugins/FX/Niagara/Source/Niagara/Classes/NiagaraEmitter.h
+++ b/Engine/Plugins/FX/Niagara/Source/Niagara/Classes/NiagaraEmitter.h
@@ -261,7 +261,7 @@ struct FVersionedNiagaraEmitterData
 	FFilePath ScriptAsset;
 #endif //WITH_EDITORONLY_DATA
 
-	/* If this emitter is no longer meant to be used, this option should be set.*/
+	/* If this emitter version is no longer meant to be used, this option should be set.*/
 	UPROPERTY()
 	bool bDeprecated = false;
 
diff --git a/Engine/Plugins/FX/Niagara/Source/Niagara/Private/NiagaraModule.cpp b/Engine/Plugins/FX/Niagara/Source/Niagara/Private/NiagaraModule.cpp
--- a/Engine/Plugins/FX/Niagara/Source/Niagara/Private/NiagaraModule.cpp
+++ b/Engine/Plugins/FX/Niagara/Source/Niagara/Private/NiagaraModule.cpp
@@ -83,6 +83,22 @@ const FNiagaraAssetTagDefinition INiagaraModule::LearningContentTagDefinition =
 	, FLinearColor::Green
 	, FGuid(0xFCF21AFA, 0x50764BFA, 0xB9B2E618, 0xF2A0CD6F)};
 
+const FNiagaraAssetTagDefinition INiagaraModule::HiddenAssetTagDefinition = {
+	LOCTEXT("HiddenDisplayName", "Hidden in Asset Browser")
+	, (int32) (ENiagaraAssetLibraryAssetTypes::Emitters) | (int32) (ENiagaraAssetLibraryAssetTypes::Systems) | (int32) (ENiagaraAssetLibraryAssetTypes::Scripts)
+	, FText::GetEmpty()
+	, ENiagaraAssetTagDefinitionImportance::Internal
+	, FLinearColor::Black
+	, FGuid(0xC7B2F2A1, 0x71762AFB, 0x51BDE232, 0x8CA4BBA3)};
+
+const FNiagaraAssetTagDefinition INiagaraModule::DeprecatedTagDefinition = {
+	LOCTEXT("DeprecatedDisplayName", "Deprecated")
+	, (int32) (ENiagaraAssetLibraryAssetTypes::Emitters) | (int32) (ENiagaraAssetLibraryAssetTypes::Systems) | (int32) (ENiagaraAssetLibraryAssetTypes::Scripts)
+	, FText::GetEmpty()
+	, ENiagaraAssetTagDefinitionImportance::Internal
+	, FLinearColor::Black
+	, FGuid(0xF1F2F2FB, 0x50C64AF9, 0xB6BCE631, 0x32C4CC61)};
+
 std::atomic<bool> INiagaraModule::bDataChannelRefreshRequested = false;
 
 int32 GEnableVerboseNiagaraChangeIdLogging = 0;
@@ -1900,6 +1916,8 @@ void INiagaraModule::RegisterInternalAssetTagDefinitions()
 	// All internal tags have to be manually assigned a unique and stable guid so that names can be changed without affecting assigned assets.
 	InternalAssetTagDefinitions.Add(&TemplateTagDefinition);
 	InternalAssetTagDefinitions.Add(&LearningContentTagDefinition);
+	InternalAssetTagDefinitions.Add(&HiddenAssetTagDefinition);
+	InternalAssetTagDefinitions.Add(&DeprecatedTagDefinition);
 }
 
 #if NIAGARA_PERF_BASELINES
diff --git a/Engine/Plugins/FX/Niagara/Source/Niagara/Public/NiagaraAssetTagDefinitions.h b/Engine/Plugins/FX/Niagara/Source/Niagara/Public/NiagaraAssetTagDefinitions.h
--- a/Engine/Plugins/FX/Niagara/Source/Niagara/Public/NiagaraAssetTagDefinitions.h
+++ b/Engine/Plugins/FX/Niagara/Source/Niagara/Public/NiagaraAssetTagDefinitions.h
@@ -19,7 +19,8 @@ UENUM()
 enum class ENiagaraAssetTagDefinitionImportance : uint8
 {
 	Primary UMETA(ToolTip="This Asset Tag Definition is considered important and will get displayed in the primary filter UI in the Niagara Asset Browsers."),
-	Secondary UMETA(ToolTip="This Asset Tag Definition is considered less important and will only get displayed in the additional drop down filter UI in the Niagara Asset Browsers.")
+	Secondary UMETA(ToolTip="This Asset Tag Definition is considered less important and will only get displayed in the additional drop down filter UI in the Niagara Asset Browsers."),
+	Internal UMETA(ToolTip="This Asset Tag Definition is for internal use only, and will not be displayed in the primary tag view nor the secondary drop down filter UI.")
 	//Both UMETA(ToolTip="This Asset Tag Definition will be displayed in both primary UI & the additional filter drop downs.")
 };
 
diff --git a/Engine/Plugins/FX/Niagara/Source/Niagara/Public/NiagaraModule.h b/Engine/Plugins/FX/Niagara/Source/Niagara/Public/NiagaraModule.h
--- a/Engine/Plugins/FX/Niagara/Source/Niagara/Public/NiagaraModule.h
+++ b/Engine/Plugins/FX/Niagara/Source/Niagara/Public/NiagaraModule.h
@@ -292,6 +292,9 @@ public:
 
 	static const FNiagaraAssetTagDefinition TemplateTagDefinition;
 	static const FNiagaraAssetTagDefinition LearningContentTagDefinition;
+	static NIAGARA_API const FNiagaraAssetTagDefinition HiddenAssetTagDefinition;
+	static NIAGARA_API const FNiagaraAssetTagDefinition DeprecatedTagDefinition;
+
 	TArray<const FNiagaraAssetTagDefinition*> InternalAssetTagDefinitions;
 private:
 	static NIAGARA_API FNiagaraVariable Engine_WorldDeltaTime;
diff --git a/Engine/Plugins/FX/Niagara/Source/NiagaraEditor/Private/NiagaraEditorMenuHelpers.cpp b/Engine/Plugins/FX/Niagara/Source/NiagaraEditor/Private/NiagaraEditorMenuHelpers.cpp
--- a/Engine/Plugins/FX/Niagara/Source/NiagaraEditor/Private/NiagaraEditorMenuHelpers.cpp
+++ b/Engine/Plugins/FX/Niagara/Source/NiagaraEditor/Private/NiagaraEditorMenuHelpers.cpp
@@ -5,6 +5,7 @@
 #include "ContentBrowserMenuContexts.h"
 #include "NiagaraEditorUtilities.h"
 #include "UObject/AssetRegistryTagsContext.h"
+#include "Widgets/AssetBrowser/NiagaraAssetBrowserConfig.h"
 #include "Widgets/AssetBrowser/NiagaraMenuFilters.h"
 #include "ToolMenus.h"
 
@@ -16,449 +17,546 @@ void NiagaraEditorMenuHelpers::RegisterToolMenus()
 
 	UToolMenus::Get()->RegisterStartupCallback(FSimpleMulticastDelegate::FDelegate::CreateLambda([]()
 	{
-		FToolMenuOwnerScoped OwnerScoped(UE_MODULE_NAME);
-		
-		const FName ContentBrowserNiagaraTagMenuName = TEXT("NiagaraEditorModule.ContentBrowserNiagaraTags");
-		UToolMenu* AllTagsMenu = UToolMenus::Get()->RegisterMenu(ContentBrowserNiagaraTagMenuName);
-		
-		AllTagsMenu->AddDynamicSection(NAME_None, FNewToolMenuDelegate::CreateLambda([](UToolMenu* InMenu)
-		{
-			// We read from and write to the context
-			UNiagaraTagsContentBrowserFilterContext* Context = InMenu->FindContext<UNiagaraTagsContentBrowserFilterContext>();
+		RegisterNiagaraAssetTagStandaloneMenu();
+	}));
+}
 
-			if(Context == nullptr)
-			{
-				return;
-			}
+void NiagaraEditorMenuHelpers::RegisterMenuExtensions()
+{
+	UToolMenus::Get()->RegisterStartupCallback(FSimpleMulticastDelegate::FDelegate::CreateLambda([]()
+	{
+		RegisterAssetBrowserViewOptionsProfile();
+		RegisterNiagaraAssetTagMenusForAssets();
+	}));
+}
 
-			TArray<FStructuredAssetTagDefinitionLookupData> AssetTagDefinitionData = GetStructuredSortedAssetTagDefinitions();
+void NiagaraEditorMenuHelpers::RegisterNiagaraAssetTagStandaloneMenu()
+{
+	using namespace FNiagaraEditorUtilities::AssetBrowser;
+	
+	FToolMenuOwnerScoped OwnerScoped(UE_MODULE_NAME);
+				
+	const FName ContentBrowserNiagaraTagMenuName = TEXT("NiagaraEditorModule.ContentBrowserNiagaraTags");
+	UToolMenu* AllTagsMenu = UToolMenus::Get()->RegisterMenu(ContentBrowserNiagaraTagMenuName);
+	
+	AllTagsMenu->AddDynamicSection(NAME_None, FNewToolMenuDelegate::CreateLambda([](UToolMenu* InMenu)
+	{
+		// We read from and write to the context
+		UNiagaraTagsContentBrowserFilterContext* Context = InMenu->FindContext<UNiagaraTagsContentBrowserFilterContext>();
 
-			auto GetCheckState = [](const FNiagaraAssetTagDefinition& AssetTagDefinition, const FToolMenuContext& Context) -> ECheckBoxState
-			{
-				return Context.FindContext<UNiagaraTagsContentBrowserFilterContext>()->FilterData->ContainsActiveTagGuid(AssetTagDefinition.TagGuid) ? ECheckBoxState::Checked : ECheckBoxState::Unchecked;
-			};
+		if(Context == nullptr)
+		{
+			return;
+		}
 
-			auto AddTag = [](const FNiagaraAssetTagDefinition& AssetTagDefinition, const FToolMenuContext& Context)
-			{
-				Context.FindContext<UNiagaraTagsContentBrowserFilterContext>()->FilterData->AddTagGuid(AssetTagDefinition.TagGuid);
-			};
+		TArray<FStructuredAssetTagDefinitionLookupData> AssetTagDefinitionData = GetStructuredSortedAssetTagDefinitions();
 
-			auto RemoveTag = [](const FNiagaraAssetTagDefinition& AssetTagDefinition, const FToolMenuContext& Context)
-			{
-				Context.FindContext<UNiagaraTagsContentBrowserFilterContext>()->FilterData->RemoveTagGuid(AssetTagDefinition.TagGuid);
-			};
+		auto GetCheckState = [](const FNiagaraAssetTagDefinition& AssetTagDefinition, const FToolMenuContext& Context) -> ECheckBoxState
+		{
+			return Context.FindContext<UNiagaraTagsContentBrowserFilterContext>()->FilterData->ContainsActiveTagGuid(AssetTagDefinition.TagGuid) ? ECheckBoxState::Checked : ECheckBoxState::Unchecked;
+		};
 
-			// We create 3 menu sections: core Niagara tag filters, Project specified filters, and filters coming from other sources such as Engine Plugins (other than Niagara)
-			const FText NiagaraSectionName = GetAssetTagSectionNameFromSource(EAssetTagSectionSource::NiagaraInternal);
-			const FText ProjectSectionName = GetAssetTagSectionNameFromSource(EAssetTagSectionSource::Project);
-			const FText OtherSectionName = GetAssetTagSectionNameFromSource(EAssetTagSectionSource::Other);
+		auto AddTag = [](const FNiagaraAssetTagDefinition& AssetTagDefinition, const FToolMenuContext& Context)
+		{
+			Context.FindContext<UNiagaraTagsContentBrowserFilterContext>()->FilterData->AddTagGuid(AssetTagDefinition.TagGuid);
+		};
 
-			{
-				FToolMenuSection& NiagaraSection = InMenu->AddSection(FName(NiagaraSectionName.ToString()), NiagaraSectionName);
-				NiagaraSection.InsertPosition = FToolMenuInsert(NAME_None, EToolMenuInsertType::First);
-			}
+		auto RemoveTag = [](const FNiagaraAssetTagDefinition& AssetTagDefinition, const FToolMenuContext& Context)
+		{
+			Context.FindContext<UNiagaraTagsContentBrowserFilterContext>()->FilterData->RemoveTagGuid(AssetTagDefinition.TagGuid);
+		};
 
-			{
-				FToolMenuSection& ProjectSection = InMenu->AddSection(FName(ProjectSectionName.ToString()), ProjectSectionName);
-				ProjectSection.InsertPosition = FToolMenuInsert(FName(NiagaraSectionName.ToString()), EToolMenuInsertType::After);
-			}
+		// We create 3 menu sections: core Niagara tag filters, Project specified filters, and filters coming from other sources such as Engine Plugins (other than Niagara)
+		const FText NiagaraSectionName = GetAssetTagSectionNameFromSource(EAssetTagSectionSource::NiagaraInternal);
+		const FText ProjectSectionName = GetAssetTagSectionNameFromSource(EAssetTagSectionSource::Project);
+		const FText OtherSectionName = GetAssetTagSectionNameFromSource(EAssetTagSectionSource::Other);
 
-			{
-				FToolMenuSection& Section = InMenu->AddSection(FName(OtherSectionName.ToString()), OtherSectionName);
-				Section.InsertPosition = FToolMenuInsert(FName(ProjectSectionName.ToString()), EToolMenuInsertType::After);
-			}
+		{
+			FToolMenuSection& NiagaraSection = InMenu->AddSection(FName(NiagaraSectionName.ToString()), NiagaraSectionName);
+			NiagaraSection.InsertPosition = FToolMenuInsert(NAME_None, EToolMenuInsertType::First);
+		}
 
-			for(const FStructuredAssetTagDefinitionLookupData& AssetTagDefinitionLookupData : AssetTagDefinitionData)
+		{
+			FToolMenuSection& ProjectSection = InMenu->AddSection(FName(ProjectSectionName.ToString()), ProjectSectionName);
+			ProjectSection.InsertPosition = FToolMenuInsert(FName(NiagaraSectionName.ToString()), EToolMenuInsertType::After);
+		}
+
+		{
+			FToolMenuSection& Section = InMenu->AddSection(FName(OtherSectionName.ToString()), OtherSectionName);
+			Section.InsertPosition = FToolMenuInsert(FName(ProjectSectionName.ToString()), EToolMenuInsertType::After);
+		}
+
+		for(const FStructuredAssetTagDefinitionLookupData& AssetTagDefinitionLookupData : AssetTagDefinitionData)
+		{
+			FName SectionName = FName(GetAssetTagSectionNameFromSource(AssetTagDefinitionLookupData.Source).ToString());
+			
+			for(const FNiagaraAssetTagDefinition& AssetTagDefinition : AssetTagDefinitionLookupData.AssetTagDefinitions)
 			{
-				FName SectionName = FName(GetAssetTagSectionNameFromSource(AssetTagDefinitionLookupData.Source).ToString());
+				FToolMenuSection* TagSection = InMenu->FindSection(SectionName);
 				
-				for(const FNiagaraAssetTagDefinition& AssetTagDefinition : AssetTagDefinitionLookupData.AssetTagDefinitions)
+				FText AssetTagLabel = AssetTagDefinition.AssetTag;
+				FText AssetTagDescription = AssetTagDefinition.Description;
+				
+				FToolUIAction Action;
+				Action.ExecuteAction = FToolMenuExecuteAction::CreateLambda([AssetTagDefinition, &GetCheckState, &AddTag, &RemoveTag](const FToolMenuContext& Context)
 				{
-					FToolMenuSection* TagSection = InMenu->FindSection(SectionName);
-					
-					FText AssetTagLabel = AssetTagDefinition.AssetTag;
-					FText AssetTagDescription = AssetTagDefinition.Description;
-					
-					FToolUIAction Action;
-					Action.ExecuteAction = FToolMenuExecuteAction::CreateLambda([AssetTagDefinition, &GetCheckState, &AddTag, &RemoveTag](const FToolMenuContext& Context)
-					{
-						UNiagaraTagsContentBrowserFilterContext* NiagaraTagsContentBrowserFilterContext = Context.FindContext<UNiagaraTagsContentBrowserFilterContext>();
+					UNiagaraTagsContentBrowserFilterContext* NiagaraTagsContentBrowserFilterContext = Context.FindContext<UNiagaraTagsContentBrowserFilterContext>();
 
-						if(NiagaraTagsContentBrowserFilterContext == nullptr)
-						{
-							return;
-						}
-					
-						ECheckBoxState CurrentState = GetCheckState(AssetTagDefinition, Context);
+					if(NiagaraTagsContentBrowserFilterContext == nullptr)
+					{
+						return;
+					}
+				
+					ECheckBoxState CurrentState = GetCheckState(AssetTagDefinition, Context);
 
-						if(CurrentState == ECheckBoxState::Checked)
-						{
-							RemoveTag(AssetTagDefinition, Context);
-						}
-						else if(CurrentState == ECheckBoxState::Unchecked)
-						{
-							AddTag(AssetTagDefinition, Context);
-						}
-					});
-					Action.GetActionCheckState = FToolMenuGetActionCheckState::CreateLambda([AssetTagDefinition, &GetCheckState](const FToolMenuContext& Context)
+					if(CurrentState == ECheckBoxState::Checked)
+					{
+						RemoveTag(AssetTagDefinition, Context);
+					}
+					else if(CurrentState == ECheckBoxState::Unchecked)
 					{
-						UNiagaraTagsContentBrowserFilterContext* NiagaraTagsContentBrowserFilterContext = Context.FindContext<UNiagaraTagsContentBrowserFilterContext>();
+						AddTag(AssetTagDefinition, Context);
+					}
+				});
+				Action.GetActionCheckState = FToolMenuGetActionCheckState::CreateLambda([AssetTagDefinition, &GetCheckState](const FToolMenuContext& Context)
+				{
+					UNiagaraTagsContentBrowserFilterContext* NiagaraTagsContentBrowserFilterContext = Context.FindContext<UNiagaraTagsContentBrowserFilterContext>();
 
-						if(NiagaraTagsContentBrowserFilterContext == nullptr)
-						{
-							return ECheckBoxState::Unchecked;
-						}
-						
-						return GetCheckState(AssetTagDefinition, Context);
-					});
+					if(NiagaraTagsContentBrowserFilterContext == nullptr)
+					{
+						return ECheckBoxState::Unchecked;
+					}
 					
-					FToolMenuEntry Entry = FToolMenuEntry::InitMenuEntry(
-						FName(AssetTagLabel.ToString()),
-						AssetTagLabel,
-						AssetTagDescription,
-						FSlateIcon(),
-						FToolUIActionChoice(Action), EUserInterfaceActionType::ToggleButton);
-					TagSection->AddEntry(Entry);
-				}
+					return GetCheckState(AssetTagDefinition, Context);
+				});
+				
+				FToolMenuEntry Entry = FToolMenuEntry::InitMenuEntry(
+					FName(AssetTagLabel.ToString()),
+					AssetTagLabel,
+					AssetTagDescription,
+					FSlateIcon(),
+					FToolUIActionChoice(Action), EUserInterfaceActionType::ToggleButton);
+				TagSection->AddEntry(Entry);
 			}
-		}));
+		}
 	}));
 }
 
-void NiagaraEditorMenuHelpers::RegisterMenuExtensions()
+void NiagaraEditorMenuHelpers::RegisterNiagaraAssetTagMenusForAssets()
 {
-	RegisterAssetBrowserViewOptionsProfile();
-	
 	using namespace FNiagaraEditorUtilities::AssetBrowser;
-	
-	UToolMenus::Get()->RegisterStartupCallback(FSimpleMulticastDelegate::FDelegate::CreateLambda([]()
-	{
-		FToolMenuOwnerScoped OwnerScoped(UE_MODULE_NAME);
 
-		// First we generate a standalone menu that will populate with all available tags dynamically
-		const FName StandaloneTagManagementMenuName = TEXT("NiagaraEditorModule.ManageAssetTags");
-		UToolMenu* TagMenu = UToolMenus::Get()->RegisterMenu(StandaloneTagManagementMenuName);
-		TagMenu->AddDynamicSection(NAME_None, FNewToolMenuDelegate::CreateLambda([](UToolMenu* InMenu)
-		{
-			TArray<FStructuredAssetTagDefinitionLookupData> AssetTagDefinitionData = GetStructuredSortedAssetTagDefinitions();
+	FToolMenuOwnerScoped OwnerScoped(UE_MODULE_NAME);
 
-			UContentBrowserAssetContextMenuContext* AssetContextMenuContext = InMenu->FindContext<UContentBrowserAssetContextMenuContext>();
+	// First we generate a standalone menu that will populate with all available tags dynamically
+	const FName StandaloneTagManagementMenuName = TEXT("NiagaraEditorModule.ManageAssetTags");
+	UToolMenu* TagMenu = UToolMenus::Get()->RegisterMenu(StandaloneTagManagementMenuName);
+	TagMenu->AddDynamicSection(NAME_None, FNewToolMenuDelegate::CreateLambda([](UToolMenu* InMenu)
+	{
+		TArray<FStructuredAssetTagDefinitionLookupData> AssetTagDefinitionData = GetStructuredSortedAssetTagDefinitions();
 
-			if(AssetContextMenuContext == nullptr)
-			{
-				return;
-			}
+		UContentBrowserAssetContextMenuContext* AssetContextMenuContext = InMenu->FindContext<UContentBrowserAssetContextMenuContext>();
+
+		if(AssetContextMenuContext == nullptr)
+		{
+			return;
+		}
 
-			auto GetCheckState = [](const FNiagaraAssetTagDefinition& AssetTagDefinition, TArray<FAssetData> SelectedAssets) -> ECheckBoxState
+		auto GetCheckState = [](const FNiagaraAssetTagDefinition& AssetTagDefinition, TArray<FAssetData> SelectedAssets) -> ECheckBoxState
+		{
+			bool bAnyContainsTag = false;
+			bool bAnyDoesNotContainTag = false;
+			for(const FAssetData& SelectedAsset : SelectedAssets)
 			{
-				bool bAnyContainsTag = false;
-				bool bAnyDoesNotContainTag = false;
-				for(const FAssetData& SelectedAsset : SelectedAssets)
+				if(SelectedAsset.GetClass() == UNiagaraEmitter::StaticClass())
 				{
-					if(SelectedAsset.GetClass() == UNiagaraEmitter::StaticClass())
+					if(SelectedAsset.IsAssetLoaded())
 					{
-						if(SelectedAsset.IsAssetLoaded())
+						UNiagaraEmitter* Emitter = Cast<UNiagaraEmitter>(SelectedAsset.GetAsset());
+						if(Emitter->AssetTags.Contains(AssetTagDefinition))
 						{
-							UNiagaraEmitter* Emitter = Cast<UNiagaraEmitter>(SelectedAsset.GetAsset());
-							if(Emitter->AssetTags.Contains(AssetTagDefinition))
-							{
-								bAnyContainsTag = true;
-							}
-							else
-							{
-								bAnyDoesNotContainTag = true;
-							}
+							bAnyContainsTag = true;
 						}
 						else
 						{
-							if(AssetTagDefinition.DoesAssetDataContainTag(SelectedAsset))
-							{
-								bAnyContainsTag = true;
-							}
-							else
-							{
-								bAnyDoesNotContainTag = true;
-							}
+							bAnyDoesNotContainTag = true;
 						}
 					}
-					if(SelectedAsset.GetClass() == UNiagaraSystem::StaticClass())
+					else
 					{
-						if(SelectedAsset.IsAssetLoaded())
+						if(AssetTagDefinition.DoesAssetDataContainTag(SelectedAsset))
 						{
-							UNiagaraSystem* System = Cast<UNiagaraSystem>(SelectedAsset.GetAsset());
-							if(System->AssetTags.Contains(AssetTagDefinition))
-							{
-								bAnyContainsTag = true;
-							}
-							else
-							{
-								bAnyDoesNotContainTag = true;
-							}
+							bAnyContainsTag = true;
 						}
-
 						else
 						{
-							if(AssetTagDefinition.DoesAssetDataContainTag(SelectedAsset))
-							{
-								bAnyContainsTag = true;
-							}
-							else
-							{
-								bAnyDoesNotContainTag = true;
-							}
-						}						
+							bAnyDoesNotContainTag = true;
+						}
 					}
-					if(SelectedAsset.GetClass() == UNiagaraScript::StaticClass())
+				}
+				if(SelectedAsset.GetClass() == UNiagaraSystem::StaticClass())
+				{
+					if(SelectedAsset.IsAssetLoaded())
 					{
-						if(SelectedAsset.IsAssetLoaded())
+						UNiagaraSystem* System = Cast<UNiagaraSystem>(SelectedAsset.GetAsset());
+						if(System->AssetTags.Contains(AssetTagDefinition))
 						{
-							UNiagaraScript* Script = Cast<UNiagaraScript>(SelectedAsset.GetAsset());
-							if(Script->GetLatestScriptData()->AssetTagDefinitionReferences.Contains(AssetTagDefinition))
-							{
-								bAnyContainsTag = true;
-							}
-							else
-							{
-								bAnyDoesNotContainTag = true;
-							}
+							bAnyContainsTag = true;
 						}
 						else
 						{
-							if(AssetTagDefinition.DoesAssetDataContainTag(SelectedAsset))
-							{
-								bAnyContainsTag = true;
-							}
-							else
-							{
-								bAnyDoesNotContainTag = true;
-							}
+							bAnyDoesNotContainTag = true;
 						}
 					}
 
-					if(bAnyContainsTag && bAnyDoesNotContainTag)
-					{
-						return ECheckBoxState::Undetermined;
-					}
-				}
-
-				if(bAnyContainsTag && bAnyDoesNotContainTag == false)
-				{
-					return ECheckBoxState::Checked;
-				}
-
-				return ECheckBoxState::Unchecked;
-			};
-
-			auto AddTag = [](const FNiagaraAssetTagDefinition& AssetTagDefinition, TArray<FAssetData> SelectedAssets)
-			{
-				FAssetRegistryModule& AssetRegistryModule = FModuleManager::LoadModuleChecked<FAssetRegistryModule>("AssetRegistry");
-
-				for(const FAssetData& SelectedAsset : SelectedAssets)
-				{
-					if(SelectedAsset.GetClass() == UNiagaraEmitter::StaticClass())
+					else
 					{
-						UNiagaraEmitter* Emitter = Cast<UNiagaraEmitter>(SelectedAsset.GetAsset());
-						if(Emitter->AssetTags.Contains(AssetTagDefinition))
+						if(AssetTagDefinition.DoesAssetDataContainTag(SelectedAsset))
 						{
-							continue;
+							bAnyContainsTag = true;
 						}
 						else
 						{
-							Emitter->Modify();
-							Emitter->AssetTags.Add(AssetTagDefinition);
+							bAnyDoesNotContainTag = true;
 						}
-					}
-					if(SelectedAsset.GetClass() == UNiagaraSystem::StaticClass())
+					}						
+				}
+				if(SelectedAsset.GetClass() == UNiagaraScript::StaticClass())
+				{
+					if(SelectedAsset.IsAssetLoaded())
 					{
-						UNiagaraSystem* System = Cast<UNiagaraSystem>(SelectedAsset.GetAsset());
-						if(System->AssetTags.Contains(AssetTagDefinition))
+						UNiagaraScript* Script = Cast<UNiagaraScript>(SelectedAsset.GetAsset());
+						if(Script->GetLatestScriptData()->AssetTagDefinitionReferences.Contains(AssetTagDefinition))
 						{
-							continue;
+							bAnyContainsTag = true;
 						}
 						else
 						{
-							System->Modify();
-							System->AssetTags.Add(AssetTagDefinition);
+							bAnyDoesNotContainTag = true;
 						}
 					}
-					if(SelectedAsset.GetClass() == UNiagaraScript::StaticClass())
+					else
 					{
-						UNiagaraScript* Script = Cast<UNiagaraScript>(SelectedAsset.GetAsset());
-						if(Script->GetLatestScriptData()->AssetTagDefinitionReferences.Contains(AssetTagDefinition))
+						if(AssetTagDefinition.DoesAssetDataContainTag(SelectedAsset))
 						{
-							continue;
+							bAnyContainsTag = true;
 						}
 						else
 						{
-							Script->Modify();
-							Script->GetLatestScriptData()->AssetTagDefinitionReferences.Add(AssetTagDefinition);
+							bAnyDoesNotContainTag = true;
 						}
 					}
+				}
 
-					AssetRegistryModule.Get().AssetUpdateTags(SelectedAsset.GetAsset(), EAssetRegistryTagsCaller::Fast);
+				if(bAnyContainsTag && bAnyDoesNotContainTag)
+				{
+					return ECheckBoxState::Undetermined;
 				}
-			};
-			
-			auto RemoveTag = [](const FNiagaraAssetTagDefinition& AssetTagDefinition, TArray<FAssetData> SelectedAssets)
+			}
+
+			if(bAnyContainsTag && bAnyDoesNotContainTag == false)
 			{
-				FAssetRegistryModule& AssetRegistryModule = FModuleManager::LoadModuleChecked<FAssetRegistryModule>("AssetRegistry");
+				return ECheckBoxState::Checked;
+			}
 
-				for(const FAssetData& SelectedAsset : SelectedAssets)
+			return ECheckBoxState::Unchecked;
+		};
+
+		auto AddTag = [](const FNiagaraAssetTagDefinition& AssetTagDefinition, TArray<FAssetData> SelectedAssets)
+		{
+			FAssetRegistryModule& AssetRegistryModule = FModuleManager::LoadModuleChecked<FAssetRegistryModule>("AssetRegistry");
+
+			for(const FAssetData& SelectedAsset : SelectedAssets)
+			{
+				if(SelectedAsset.GetClass() == UNiagaraEmitter::StaticClass())
 				{
-					if(SelectedAsset.GetClass() == UNiagaraEmitter::StaticClass())
+					UNiagaraEmitter* Emitter = Cast<UNiagaraEmitter>(SelectedAsset.GetAsset());
+					if(Emitter->AssetTags.Contains(AssetTagDefinition))
 					{
-						UNiagaraEmitter* Emitter = Cast<UNiagaraEmitter>(SelectedAsset.GetAsset());
-						if(Emitter->AssetTags.Contains(AssetTagDefinition) == false)
-						{
-							continue;
-						}
-						else
-						{
-							Emitter->Modify();
-							Emitter->AssetTags.Remove(AssetTagDefinition);
-						}
+						continue;
 					}
-					if(SelectedAsset.GetClass() == UNiagaraSystem::StaticClass())
+					else
 					{
-						UNiagaraSystem* System = Cast<UNiagaraSystem>(SelectedAsset.GetAsset());
-						if(System->AssetTags.Contains(AssetTagDefinition) == false)
-						{
-							continue;
-						}
-						else
-						{
-							System->Modify();
-							System->AssetTags.Remove(AssetTagDefinition);
-						}
+						Emitter->Modify();
+						Emitter->AssetTags.Add(AssetTagDefinition);
 					}
-					if(SelectedAsset.GetClass() == UNiagaraScript::StaticClass())
+				}
+				if(SelectedAsset.GetClass() == UNiagaraSystem::StaticClass())
+				{
+					UNiagaraSystem* System = Cast<UNiagaraSystem>(SelectedAsset.GetAsset());
+					if(System->AssetTags.Contains(AssetTagDefinition))
 					{
-						UNiagaraScript* Script = Cast<UNiagaraScript>(SelectedAsset.GetAsset());
-						if(Script->GetLatestScriptData()->AssetTagDefinitionReferences.Contains(AssetTagDefinition) == false)
-						{
-							continue;
-						}
-						else
-						{
-							Script->Modify();
-							Script->GetLatestScriptData()->AssetTagDefinitionReferences.Remove(AssetTagDefinition);
-						}
+						continue;
+					}
+					else
+					{
+						System->Modify();
+						System->AssetTags.Add(AssetTagDefinition);
+					}
+				}
+				if(SelectedAsset.GetClass() == UNiagaraScript::StaticClass())
+				{
+					UNiagaraScript* Script = Cast<UNiagaraScript>(SelectedAsset.GetAsset());
+					if(Script->GetLatestScriptData()->AssetTagDefinitionReferences.Contains(AssetTagDefinition))
+					{
+						continue;
+					}
+					else
+					{
+						Script->Modify();
+						Script->GetLatestScriptData()->AssetTagDefinitionReferences.Add(AssetTagDefinition);
 					}
-
-					AssetRegistryModule.Get().AssetUpdateTags(SelectedAsset.GetAsset(), EAssetRegistryTagsCaller::Fast);
 				}
-			};
-
-			// We create 3 menu sections: core Niagara tag filters, Project specified filters, and filters coming from other sources such as Engine Plugins (other than Niagara)
-			const FText NiagaraSectionName = GetAssetTagSectionNameFromSource(EAssetTagSectionSource::NiagaraInternal);
-			const FText ProjectSectionName = GetAssetTagSectionNameFromSource(EAssetTagSectionSource::Project);
-			const FText OtherSectionName = GetAssetTagSectionNameFromSource(EAssetTagSectionSource::Other);
 
-			{
-				FToolMenuSection& NiagaraSection = InMenu->AddSection(FName(NiagaraSectionName.ToString()), NiagaraSectionName);
-				NiagaraSection.InsertPosition = FToolMenuInsert(NAME_None, EToolMenuInsertType::First);
+				AssetRegistryModule.Get().AssetUpdateTags(SelectedAsset.GetAsset(), EAssetRegistryTagsCaller::Fast);
 			}
+		};
+		
+		auto RemoveTag = [](const FNiagaraAssetTagDefinition& AssetTagDefinition, TArray<FAssetData> SelectedAssets)
+		{
+			FAssetRegistryModule& AssetRegistryModule = FModuleManager::LoadModuleChecked<FAssetRegistryModule>("AssetRegistry");
 
+			for(const FAssetData& SelectedAsset : SelectedAssets)
 			{
-				FToolMenuSection& ProjectSection = InMenu->AddSection(FName(ProjectSectionName.ToString()), ProjectSectionName);
-				ProjectSection.InsertPosition = FToolMenuInsert(FName(NiagaraSectionName.ToString()), EToolMenuInsertType::After);
-			}
+				if(SelectedAsset.GetClass() == UNiagaraEmitter::StaticClass())
+				{
+					UNiagaraEmitter* Emitter = Cast<UNiagaraEmitter>(SelectedAsset.GetAsset());
+					if(Emitter->AssetTags.Contains(AssetTagDefinition) == false)
+					{
+						continue;
+					}
+					else
+					{
+						Emitter->Modify();
+						Emitter->AssetTags.Remove(AssetTagDefinition);
+					}
+				}
+				if(SelectedAsset.GetClass() == UNiagaraSystem::StaticClass())
+				{
+					UNiagaraSystem* System = Cast<UNiagaraSystem>(SelectedAsset.GetAsset());
+					if(System->AssetTags.Contains(AssetTagDefinition) == false)
+					{
+						continue;
+					}
+					else
+					{
+						System->Modify();
+						System->AssetTags.Remove(AssetTagDefinition);
+					}
+				}
+				if(SelectedAsset.GetClass() == UNiagaraScript::StaticClass())
+				{
+					UNiagaraScript* Script = Cast<UNiagaraScript>(SelectedAsset.GetAsset());
+					if(Script->GetLatestScriptData()->AssetTagDefinitionReferences.Contains(AssetTagDefinition) == false)
+					{
+						continue;
+					}
+					else
+					{
+						Script->Modify();
+						Script->GetLatestScriptData()->AssetTagDefinitionReferences.Remove(AssetTagDefinition);
+					}
+				}
 
-			{
-				FToolMenuSection& Section = InMenu->AddSection(FName(OtherSectionName.ToString()), OtherSectionName);
-				Section.InsertPosition = FToolMenuInsert(FName(ProjectSectionName.ToString()), EToolMenuInsertType::After);
+				AssetRegistryModule.Get().AssetUpdateTags(SelectedAsset.GetAsset(), EAssetRegistryTagsCaller::Fast);
 			}
+		};
+
+		// We create 3 menu sections: core Niagara tag filters, Project specified filters, and filters coming from other sources such as Engine Plugins (other than Niagara)
+		const FText NiagaraSectionName = GetAssetTagSectionNameFromSource(EAssetTagSectionSource::NiagaraInternal);
+		const FText ProjectSectionName = GetAssetTagSectionNameFromSource(EAssetTagSectionSource::Project);
+		const FText OtherSectionName = GetAssetTagSectionNameFromSource(EAssetTagSectionSource::Other);
+
+		{
+			FToolMenuSection& NiagaraSection = InMenu->AddSection(FName(NiagaraSectionName.ToString()), NiagaraSectionName);
+			NiagaraSection.InsertPosition = FToolMenuInsert(NAME_None, EToolMenuInsertType::First);
+		}
 
-			for(const FStructuredAssetTagDefinitionLookupData& AssetTagDefinitionLookupData : AssetTagDefinitionData)
+		{
+			FToolMenuSection& ProjectSection = InMenu->AddSection(FName(ProjectSectionName.ToString()), ProjectSectionName);
+			ProjectSection.InsertPosition = FToolMenuInsert(FName(NiagaraSectionName.ToString()), EToolMenuInsertType::After);
+		}
+
+		{
+			FToolMenuSection& Section = InMenu->AddSection(FName(OtherSectionName.ToString()), OtherSectionName);
+			Section.InsertPosition = FToolMenuInsert(FName(ProjectSectionName.ToString()), EToolMenuInsertType::After);
+		}
+
+		for(const FStructuredAssetTagDefinitionLookupData& AssetTagDefinitionLookupData : AssetTagDefinitionData)
+		{
+			FName SectionName = FName(GetAssetTagSectionNameFromSource(AssetTagDefinitionLookupData.Source).ToString());
+			
+			for(const FNiagaraAssetTagDefinition& AssetTagDefinition : AssetTagDefinitionLookupData.AssetTagDefinitions)
 			{
-				FName SectionName = FName(GetAssetTagSectionNameFromSource(AssetTagDefinitionLookupData.Source).ToString());
-				
-				for(const FNiagaraAssetTagDefinition& AssetTagDefinition : AssetTagDefinitionLookupData.AssetTagDefinitions)
+				if(AssetTagDefinition.GetSupportedClasses().Contains(AssetContextMenuContext->CommonClass))
 				{
-					if(AssetTagDefinition.GetSupportedClasses().Contains(AssetContextMenuContext->CommonClass))
+					FToolMenuSection* TagSection = InMenu->FindSection(SectionName);
+					
+					FText AssetTagLabel = AssetTagDefinition.AssetTag;
+					FText AssetTagDescription = AssetTagDefinition.Description;
+					
+					FToolUIAction Action;
+					Action.ExecuteAction = FToolMenuExecuteAction::CreateLambda([AssetTagDefinition, &GetCheckState, &AddTag, &RemoveTag](const FToolMenuContext& Context)
 					{
-						FToolMenuSection* TagSection = InMenu->FindSection(SectionName);
-						
-						FText AssetTagLabel = AssetTagDefinition.AssetTag;
-						FText AssetTagDescription = AssetTagDefinition.Description;
-						
-						FToolUIAction Action;
-						Action.ExecuteAction = FToolMenuExecuteAction::CreateLambda([AssetTagDefinition, &GetCheckState, &AddTag, &RemoveTag](const FToolMenuContext& Context)
+						UContentBrowserAssetContextMenuContext* AssetContextMenuContext = Context.FindContext<UContentBrowserAssetContextMenuContext>();
+
+						if(AssetContextMenuContext == nullptr)
 						{
-							UContentBrowserAssetContextMenuContext* AssetContextMenuContext = Context.FindContext<UContentBrowserAssetContextMenuContext>();
+							return;
+						}
+					
+						ECheckBoxState CurrentState = GetCheckState(AssetTagDefinition, AssetContextMenuContext->SelectedAssets);
 
-							if(AssetContextMenuContext == nullptr)
-							{
-								return;
-							}
-						
-							ECheckBoxState CurrentState = GetCheckState(AssetTagDefinition, AssetContextMenuContext->SelectedAssets);
-
-							if(CurrentState == ECheckBoxState::Checked)
-							{
-								RemoveTag(AssetTagDefinition, AssetContextMenuContext->SelectedAssets);
-							}
-							else if(CurrentState == ECheckBoxState::Unchecked)
-							{
-								AddTag(AssetTagDefinition, AssetContextMenuContext->SelectedAssets);
-							}
-							// if there is a mix of states, we are removing the tags
-							else if(CurrentState == ECheckBoxState::Undetermined)
-							{
-								RemoveTag(AssetTagDefinition, AssetContextMenuContext->SelectedAssets);
-							}
-						});
-						Action.GetActionCheckState = FToolMenuGetActionCheckState::CreateLambda([AssetTagDefinition, &GetCheckState](const FToolMenuContext& Context)
+						if(CurrentState == ECheckBoxState::Checked)
 						{
-							UContentBrowserAssetContextMenuContext* AssetContextMenuContext = Context.FindContext<UContentBrowserAssetContextMenuContext>();
-
-							if(AssetContextMenuContext == nullptr)
-							{
-								return ECheckBoxState::Unchecked;
-							}
-							
-							return GetCheckState(AssetTagDefinition, AssetContextMenuContext->SelectedAssets);
-						});
+							RemoveTag(AssetTagDefinition, AssetContextMenuContext->SelectedAssets);
+						}
+						else if(CurrentState == ECheckBoxState::Unchecked)
+						{
+							AddTag(AssetTagDefinition, AssetContextMenuContext->SelectedAssets);
+						}
+						// if there is a mix of states, we are removing the tags
+						else if(CurrentState == ECheckBoxState::Undetermined)
+						{
+							RemoveTag(AssetTagDefinition, AssetContextMenuContext->SelectedAssets);
+						}
+					});
+					Action.GetActionCheckState = FToolMenuGetActionCheckState::CreateLambda([AssetTagDefinition, &GetCheckState](const FToolMenuContext& Context)
+					{
+						UContentBrowserAssetContextMenuContext* AssetContextMenuContext = Context.FindContext<UContentBrowserAssetContextMenuContext>();
+
+						if(AssetContextMenuContext == nullptr)
+						{
+							return ECheckBoxState::Unchecked;
+						}
 						
-						FToolMenuEntry Entry = FToolMenuEntry::InitMenuEntry(
-							FName(AssetTagLabel.ToString()),
-							AssetTagLabel,
-							AssetTagDescription,
-							FSlateIcon(),
-							FToolUIActionChoice(Action), EUserInterfaceActionType::ToggleButton);
-						TagSection->AddEntry(Entry);
-					}
+						return GetCheckState(AssetTagDefinition, AssetContextMenuContext->SelectedAssets);
+					});
+					
+					FToolMenuEntry Entry = FToolMenuEntry::InitMenuEntry(
+						FName(AssetTagLabel.ToString()),
+						AssetTagLabel,
+						AssetTagDescription,
+						FSlateIcon(),
+						FToolUIActionChoice(Action), EUserInterfaceActionType::ToggleButton);
+					TagSection->AddEntry(Entry);
 				}
 			}
-		}));
-
-		for(UClass* ExtendedClass : {UNiagaraEmitter::StaticClass(), UNiagaraSystem::StaticClass(), UNiagaraScript::StaticClass()})
-		{
-			// Then we extend the asset context menu by adding an empty sub menu
-			const FName SubMenuName = TEXT("ManageTags");
-			UToolMenu* AssetContextMenu = UE::ContentBrowser::ExtendToolMenu_AssetContextMenu(ExtendedClass);
-			AssetContextMenu->AddSubMenu(OwnerScoped.GetOwner(), "GetAssetActions", SubMenuName, LOCTEXT("ManageTabsSubMenuLabel", "Manage Tags"));
-			
-			// Now that we have added a submenu with a specific name, we need the full asset context menu path including the same sub menu name
-			const FName AssetContextMenuNameWithSubMenu = UToolMenus::JoinMenuPaths(AssetContextMenu->GetMenuName(), SubMenuName);
-
-			// Then, we register the standalone menu as the parent of the submenu. This will 'link' the empty submenu with our standalone menu.
-			UToolMenus::Get()->RegisterMenu(AssetContextMenuNameWithSubMenu, StandaloneTagManagementMenuName);
 		}
 	}));
+
+	for(UClass* ExtendedClass : {UNiagaraEmitter::StaticClass(), UNiagaraSystem::StaticClass(), UNiagaraScript::StaticClass()})
+	{
+		// Then we extend the asset context menu by adding an empty sub menu
+		const FName SubMenuName = TEXT("ManageTags");
+		UToolMenu* AssetContextMenu = UE::ContentBrowser::ExtendToolMenu_AssetContextMenu(ExtendedClass);
+		AssetContextMenu->AddSubMenu(OwnerScoped.GetOwner(), "GetAssetActions", SubMenuName, LOCTEXT("ManageTabsSubMenuLabel", "Manage Tags"));
+		
+		// Now that we have added a submenu with a specific name, we need the full asset context menu path including the same sub menu name
+		const FName AssetContextMenuNameWithSubMenu = UToolMenus::JoinMenuPaths(AssetContextMenu->GetMenuName(), SubMenuName);
+
+		// Then, we register the standalone menu as the parent of the submenu. This will 'link' the empty submenu with our standalone menu.
+		UToolMenus::Get()->RegisterMenu(AssetContextMenuNameWithSubMenu, StandaloneTagManagementMenuName);
+	}
 }
 
 void NiagaraEditorMenuHelpers::RegisterAssetBrowserViewOptionsProfile()
 {
 	UToolMenus::Get()->RegisterStartupCallback(FSimpleMulticastDelegate::FDelegate::CreateLambda([]()
 	{
+		const FName MenuName("ContentBrowser.AssetViewOptions");
 		const FName ProfileName("NiagaraAssetBrowser");
 
-		FToolMenuProfile* NiagaraAssetBrowserProfile = UToolMenus::Get()->AddRuntimeMenuProfile(FName("ContentBrowser.AssetViewOptions"), ProfileName);
+		FToolMenuProfile* NiagaraAssetBrowserProfile = UToolMenus::Get()->AddRuntimeMenuProfile(MenuName, ProfileName);
 		if(FCustomizedToolMenuSection* Section = NiagaraAssetBrowserProfile->AddSection("View"))
 		{
 			Section->Visibility = ECustomizedToolMenuVisibility::Hidden;
 		}
+
+		UToolMenu* AssetPickerAssetViewOptionsMenu = UToolMenus::Get()->ExtendMenu(MenuName);
+		AssetPickerAssetViewOptionsMenu->AddDynamicSection("NiagaraDynamicSection", FNewSectionConstructChoice(FNewToolMenuDelegate::CreateLambda([](UToolMenu* ToolMenu)
+		{
+			// We only want to modify the menu if we are in a NiagaraAssetBrowser profile
+			if(UToolMenuProfileContext* ToolMenuProfileContext = ToolMenu->FindContext<UToolMenuProfileContext>())
+			{
+				if(ToolMenuProfileContext->ActiveProfiles.Contains("NiagaraAssetBrowser") == false)
+				{
+					return;
+				}
+			}
+			else
+			{
+				return;
+			}
+
+			FToolMenuSection& NiagaraSection = ToolMenu->AddSection("Niagara", FText::AsCultureInvariant("Niagara"));
+			
+			// Hidden
+			{
+				auto GetShowHiddenAssetsCheckState = []() -> ECheckBoxState
+			   {
+				   return UNiagaraAssetBrowserConfig::Get()->bShowHiddenAssets ? ECheckBoxState::Checked : ECheckBoxState::Unchecked;
+			   };
+			
+			   auto OnShowHiddenAssetsToggled = []() -> void
+			   {
+				   UNiagaraAssetBrowserConfig::Get()->bShowHiddenAssets = !UNiagaraAssetBrowserConfig::Get()->bShowHiddenAssets;
+				   UNiagaraAssetBrowserConfig::Get()->PostEditChange();
+			   };
+			
+			   FToolUIAction Action;
+			   Action.ExecuteAction = FToolMenuExecuteAction::CreateLambda([&GetShowHiddenAssetsCheckState, &OnShowHiddenAssetsToggled](const FToolMenuContext& Context)
+			   {
+				   OnShowHiddenAssetsToggled();
+			   });
+			   Action.GetActionCheckState = FToolMenuGetActionCheckState::CreateLambda([&GetShowHiddenAssetsCheckState](const FToolMenuContext& Context)
+			   {
+				   return GetShowHiddenAssetsCheckState();
+			   });
+
+			   FToolMenuEntry Entry = FToolMenuEntry::InitMenuEntry(
+				   FName("NiagaraShowHidden"),
+				   LOCTEXT("ShowHiddenLabel", "Show Hidden"),
+				   LOCTEXT("ShowHiddenTooltip", "Show assets that were tagged as Hidden"),
+				   FSlateIcon(),
+				   FToolUIActionChoice(Action), EUserInterfaceActionType::ToggleButton);
+				
+			   NiagaraSection.AddEntry(Entry);
+			}
+			
+			// Deprecation
+			{
+				auto GetShowDeprecatedAssetsCheckState = []() -> ECheckBoxState
+			   {
+				   return UNiagaraAssetBrowserConfig::Get()->bShowDeprecatedAssets ? ECheckBoxState::Checked : ECheckBoxState::Unchecked;
+			   };
+			
+			   auto OnShowDeprecatedAssetsToggled = []() -> void
+			   {
+				   UNiagaraAssetBrowserConfig::Get()->bShowDeprecatedAssets = !UNiagaraAssetBrowserConfig::Get()->bShowDeprecatedAssets;
+				   UNiagaraAssetBrowserConfig::Get()->PostEditChange();
+			   };
+			
+			   FToolUIAction Action;
+			   Action.ExecuteAction = FToolMenuExecuteAction::CreateLambda([&GetShowDeprecatedAssetsCheckState, &OnShowDeprecatedAssetsToggled](const FToolMenuContext& Context)
+			   {
+				   OnShowDeprecatedAssetsToggled();
+			   });
+			   Action.GetActionCheckState = FToolMenuGetActionCheckState::CreateLambda([&GetShowDeprecatedAssetsCheckState](const FToolMenuContext& Context)
+			   {
+				   return GetShowDeprecatedAssetsCheckState();
+			   });
+
+			   FToolMenuEntry Entry = FToolMenuEntry::InitMenuEntry(
+				   FName("NiagaraShowDeprecated"),
+				   LOCTEXT("ShowDeprecatedLabel", "Show Deprecated"),
+				   LOCTEXT("ShowDeprecatedTooltip", "Show assets that were tagged as Deprecated"),
+				   FSlateIcon(),
+				   FToolUIActionChoice(Action), EUserInterfaceActionType::ToggleButton);
+				
+				NiagaraSection.AddEntry(Entry);
+			}
+		})));
 	}));
 }
 
diff --git a/Engine/Plugins/FX/Niagara/Source/NiagaraEditor/Private/NiagaraEditorMenuHelpers.h b/Engine/Plugins/FX/Niagara/Source/NiagaraEditor/Private/NiagaraEditorMenuHelpers.h
--- a/Engine/Plugins/FX/Niagara/Source/NiagaraEditor/Private/NiagaraEditorMenuHelpers.h
+++ b/Engine/Plugins/FX/Niagara/Source/NiagaraEditor/Private/NiagaraEditorMenuHelpers.h
@@ -7,5 +7,10 @@ namespace NiagaraEditorMenuHelpers
 	void RegisterToolMenus();
 	void RegisterMenuExtensions();
 
+	// Registers a runtime profile for the AssetViewOptions in the NiagaraAssetBrowser to hide entries we don't need to display
 	void RegisterAssetBrowserViewOptionsProfile();
+	// The standalone menu requires to have a UNiagaraTagsContentBrowserFilterContext context object to save and write data into
+	void RegisterNiagaraAssetTagStandaloneMenu();
+	// This does a mix of extending and registering menus for Niagara Emitters, Systems and Scripts
+	void RegisterNiagaraAssetTagMenusForAssets();
 }
\ No newline at end of file
diff --git a/Engine/Plugins/FX/Niagara/Source/NiagaraEditor/Private/NiagaraValidationRules.cpp b/Engine/Plugins/FX/Niagara/Source/NiagaraEditor/Private/NiagaraValidationRules.cpp
--- a/Engine/Plugins/FX/Niagara/Source/NiagaraEditor/Private/NiagaraValidationRules.cpp
+++ b/Engine/Plugins/FX/Niagara/Source/NiagaraEditor/Private/NiagaraValidationRules.cpp
@@ -1034,6 +1034,26 @@ void UNiagaraValidationRule_HasEffectType::CheckValidity(const FNiagaraValidatio
 	}
 }
 
+void UNiagaraValidationRule_CheckDeprecatedEmitters::CheckValidity(const FNiagaraValidationContext& Context, TArray<FNiagaraValidationResult>& OutResults) const
+{
+	TArray<TSharedRef<FNiagaraEmitterHandleViewModel>> EmitterHandleViewModels = Context.ViewModel->GetEmitterHandleViewModels();
+	for (TSharedRef<FNiagaraEmitterHandleViewModel> EmitterHandleViewModel : EmitterHandleViewModels)
+	{
+		if(EmitterHandleViewModel->GetEmitterHandle()->GetInstance().Emitter->AssetTags.ContainsByPredicate([](const FNiagaraAssetTagDefinitionReference& AssetTagReferenceCandidate)
+		{
+			return INiagaraModule::Get().DeprecatedTagDefinition.TagGuid == AssetTagReferenceCandidate.GetTagDefinitionReferenceGuid();
+		}))
+		{
+			UNiagaraStackEmitterPropertiesItem* EmitterProperties = NiagaraValidation::GetStackEntry<UNiagaraStackEmitterPropertiesItem>(EmitterHandleViewModel.Get().GetEmitterStackViewModel());
+
+			FNiagaraValidationResult Result(Severity, FText::Format(
+				LOCTEXT("DeprecatedEmitterUsedTitle", "Emitter '{0}' is deprecated."), FText::FromName(EmitterHandleViewModel->GetEmitterHandle()->GetName())),
+				LOCTEXT("DeprecatedEmitterUsedDescription", "The emitter is deprecated and should not be used. Consider replacing it."), EmitterProperties);
+			OutResults.Add(Result);
+		}
+	}
+}
+
 void UNiagaraValidationRule_LWC::CheckValidity(const FNiagaraValidationContext& Context, TArray<FNiagaraValidationResult>& Results)  const
 {
 	const UNiagaraSettings* Settings = GetDefault<UNiagaraSettings>();
diff --git a/Engine/Plugins/FX/Niagara/Source/NiagaraEditor/Private/Widgets/AssetBrowser/SNiagaraAssetBrowser.cpp b/Engine/Plugins/FX/Niagara/Source/NiagaraEditor/Private/Widgets/AssetBrowser/SNiagaraAssetBrowser.cpp
--- a/Engine/Plugins/FX/Niagara/Source/NiagaraEditor/Private/Widgets/AssetBrowser/SNiagaraAssetBrowser.cpp
+++ b/Engine/Plugins/FX/Niagara/Source/NiagaraEditor/Private/Widgets/AssetBrowser/SNiagaraAssetBrowser.cpp
@@ -9,6 +9,7 @@
 #include "NiagaraEmitter.h"
 #include "NiagaraRecentAndFavoritesManager.h"
 #include "NiagaraSystem.h"
+#include "SAssetView.h"
 #include "SlateOptMacros.h"
 #include "AssetRegistry/AssetRegistryModule.h"
 #include "Toolkits/GlobalEditorCommonCommands.h"
@@ -87,14 +88,17 @@ void SNiagaraAssetBrowser::Construct(const FArguments& InArgs)
 
 	bSuppressSaveAndLoad = false;
 	LoadSettings();
-
+	
 	InitContextMenu();
+
+	UNiagaraAssetBrowserConfig::Get()->OnPropertyChanged().AddSP(this, &SNiagaraAssetBrowser::OnAssetBrowserConfigPropertyChanged);
 }
 
 SNiagaraAssetBrowser::~SNiagaraAssetBrowser()
 {
 	PreviewViewport.Reset();
 	SaveSettings();
+	UNiagaraAssetBrowserConfig::Get()->OnPropertyChanged().RemoveAll(this);
 }
 
 TArray<UClass*> SNiagaraAssetBrowser::GetDisplayedAssetTypes() const
@@ -169,6 +173,16 @@ bool SNiagaraAssetBrowser::ShouldFilterAsset(const FAssetData& AssetData) const
 	{
 		return true;
 	}
+
+	if(UNiagaraAssetBrowserConfig::Get()->bShowHiddenAssets == false && INiagaraModule::Get().HiddenAssetTagDefinition.DoesAssetDataContainTag(AssetData))
+	{
+		return true;
+	}
+	
+	if(UNiagaraAssetBrowserConfig::Get()->bShowDeprecatedAssets == false && INiagaraModule::Get().DeprecatedTagDefinition.DoesAssetDataContainTag(AssetData))
+	{
+		return true;
+	}
 	
 	// TODO (ME) This currently implies only one main filter/folder can be active at a time. is this wanted?
 	for(const TSharedRef<FNiagaraAssetBrowserMainFilter>& MainFilter : MainFilterSelector->GetSelectedItems())
@@ -311,7 +325,7 @@ void SNiagaraAssetBrowser::PopulateAssetBrowserDetailsSlot()
 TArray<TSharedRef<FNiagaraAssetBrowserMainFilter>> SNiagaraAssetBrowser::GetMainFilters() const
 {
 	using namespace FNiagaraEditorUtilities::AssetBrowser;
-
+	
 	TArray<TSharedRef<FNiagaraAssetBrowserMainFilter>> MainFilters;
 
 	// Recent
@@ -361,7 +375,7 @@ TArray<TSharedRef<FNiagaraAssetBrowserMainFilter>> SNiagaraAssetBrowser::GetMain
 			if(AssetTagDefinitionData.DefinitionsAsset == nullptr)
 			{
 				for(const FNiagaraAssetTagDefinition& AssetTagDefinition : AssetTagDefinitionData.AssetTagDefinitions)
-				{
+				{					
 					if(IsAssetTagDefinitionValid(AssetTagDefinition))
 					{
 						DisplayedFlatAssetTagDefinitionsList.Add(AssetTagDefinition);
@@ -961,6 +975,11 @@ FText SNiagaraAssetBrowser::OnGetShouldDisplayViewportTooltip() const
 	return LOCTEXT("ShouldDisplayViewportTooltip", "If activated, displays Niagara Systems live in a viewport instead of a thumbnail.\nThis will compile the Niagara System if necessary and might slow down performance.");
 }
 
+void SNiagaraAssetBrowser::OnAssetBrowserConfigPropertyChanged(const FPropertyChangedEvent& PropertyChangedEvent)
+{
+	OnFilterChanged();
+}
+
 void SNiagaraAssetBrowser::SaveSettings() const
 {
 	if(SaveSettingsName.IsSet() && bSuppressSaveAndLoad == false)
diff --git a/Engine/Plugins/FX/Niagara/Source/NiagaraEditor/Public/NiagaraValidationRules.h b/Engine/Plugins/FX/Niagara/Source/NiagaraEditor/Public/NiagaraValidationRules.h
--- a/Engine/Plugins/FX/Niagara/Source/NiagaraEditor/Public/NiagaraValidationRules.h
+++ b/Engine/Plugins/FX/Niagara/Source/NiagaraEditor/Public/NiagaraValidationRules.h
@@ -252,6 +252,19 @@ public:
 	virtual void CheckValidity(const FNiagaraValidationContext& Context, TArray<FNiagaraValidationResult>& OutResults) const override;
 };
 
+/** This validation rule will check if a system uses emitters that are tagged as Deprecated using the Niagara Asset Tags.
+ *  This is distinct from a Niagara Emitter version that is marked as deprecated, but might have a new, non-deprecated version. */
+UCLASS(Category = "Validation", DisplayName = "Check for Deprecated Emitters")
+class UNiagaraValidationRule_CheckDeprecatedEmitters : public UNiagaraValidationRule
+{
+	GENERATED_BODY()
+public:
+	UPROPERTY(EditAnywhere, Category = Validation)
+	ENiagaraValidationSeverity Severity = ENiagaraValidationSeverity::Warning;
+	
+	virtual void CheckValidity(const FNiagaraValidationContext& Context, TArray<FNiagaraValidationResult>& OutResults) const override;
+};
+
 /** This validation rule checks for various common issue with Large World Coordinates like mixing vector and position types. */
 UCLASS(Category = "Validation", DisplayName = "Large World Coordinates")
 class UNiagaraValidationRule_LWC : public UNiagaraValidationRule
diff --git a/Engine/Plugins/FX/Niagara/Source/NiagaraEditor/Public/Widgets/AssetBrowser/NiagaraAssetBrowserConfig.cpp b/Engine/Plugins/FX/Niagara/Source/NiagaraEditor/Public/Widgets/AssetBrowser/NiagaraAssetBrowserConfig.cpp
--- a/Engine/Plugins/FX/Niagara/Source/NiagaraEditor/Public/Widgets/AssetBrowser/NiagaraAssetBrowserConfig.cpp
+++ b/Engine/Plugins/FX/Niagara/Source/NiagaraEditor/Public/Widgets/AssetBrowser/NiagaraAssetBrowserConfig.cpp
@@ -12,3 +12,11 @@ void UNiagaraAssetBrowserConfig::Initialize()
 		Instance->AddToRoot();
 	}
 }
+
+void UNiagaraAssetBrowserConfig::PostEditChangeProperty(FPropertyChangedEvent& PropertyChangedEvent)
+{
+	Super::PostEditChangeProperty(PropertyChangedEvent);
+
+	OnPropertyChangedDelegate.Broadcast(PropertyChangedEvent);
+	SaveEditorConfig();
+}
diff --git a/Engine/Plugins/FX/Niagara/Source/NiagaraEditor/Public/Widgets/AssetBrowser/NiagaraAssetBrowserConfig.h b/Engine/Plugins/FX/Niagara/Source/NiagaraEditor/Public/Widgets/AssetBrowser/NiagaraAssetBrowserConfig.h
--- a/Engine/Plugins/FX/Niagara/Source/NiagaraEditor/Public/Widgets/AssetBrowser/NiagaraAssetBrowserConfig.h
+++ b/Engine/Plugins/FX/Niagara/Source/NiagaraEditor/Public/Widgets/AssetBrowser/NiagaraAssetBrowserConfig.h
@@ -22,14 +22,25 @@ UCLASS(EditorConfig="NiagaraAssetBrowser")
 class NIAGARAEDITOR_API UNiagaraAssetBrowserConfig : public UEditorConfigBase
 {
 	GENERATED_BODY()
+public:
+	DECLARE_MULTICAST_DELEGATE_OneParam(FOnPropertyChanged, const FPropertyChangedEvent&);
 	
-public:		
 	static void Initialize();
 	static UNiagaraAssetBrowserConfig* Get() { return Instance; }
+	FOnPropertyChanged& OnPropertyChanged() { return OnPropertyChangedDelegate; }
 	
 	UPROPERTY(meta=(EditorConfig))
 	TMap<FName, FNiagaraAssetBrowserConfiguration> MainFilterSelection;
-	
+
+	UPROPERTY(meta=(EditorConfig))
+    bool bShowHiddenAssets = false;
+    	
+	UPROPERTY(meta=(EditorConfig))
+	bool bShowDeprecatedAssets = false;
+
+	virtual void PostEditChangeProperty(FPropertyChangedEvent& PropertyChangedEvent) override;	
 private:
 	static TObjectPtr<UNiagaraAssetBrowserConfig> Instance;
+
+	FOnPropertyChanged OnPropertyChangedDelegate;
 };
diff --git a/Engine/Plugins/FX/Niagara/Source/NiagaraEditor/Public/Widgets/AssetBrowser/SNiagaraAssetBrowser.h b/Engine/Plugins/FX/Niagara/Source/NiagaraEditor/Public/Widgets/AssetBrowser/SNiagaraAssetBrowser.h
--- a/Engine/Plugins/FX/Niagara/Source/NiagaraEditor/Public/Widgets/AssetBrowser/SNiagaraAssetBrowser.h
+++ b/Engine/Plugins/FX/Niagara/Source/NiagaraEditor/Public/Widgets/AssetBrowser/SNiagaraAssetBrowser.h
@@ -121,6 +121,8 @@ private:
 	void OnShouldDisplayViewportChanged(ECheckBoxState CheckBoxState);
 	FText OnGetShouldDisplayViewportTooltip() const;
 
+	void OnAssetBrowserConfigPropertyChanged(const FPropertyChangedEvent& PropertyChangedEvent);
+	
 	void SaveSettings() const;
 	void LoadSettings();
 private:
diff --git a/Engine/Plugins/FX/Niagara/Source/NiagaraEditorWidgets/Private/SNiagaraOverviewGraphTitleBar.cpp b/Engine/Plugins/FX/Niagara/Source/NiagaraEditorWidgets/Private/SNiagaraOverviewGraphTitleBar.cpp
--- a/Engine/Plugins/FX/Niagara/Source/NiagaraEditorWidgets/Private/SNiagaraOverviewGraphTitleBar.cpp
+++ b/Engine/Plugins/FX/Niagara/Source/NiagaraEditorWidgets/Private/SNiagaraOverviewGraphTitleBar.cpp
@@ -159,7 +159,12 @@ FText SNiagaraOverviewGraphTitleBar::GetSystemSubheaderText() const
 			{
 				if (FVersionedNiagaraEmitterData* ParentData = EmitterData->GetParent().GetEmitterData())
 				{
-					if (ParentData->bDeprecated)
+					bool bIsDeprecatedAsset = EmitterData->GetParent().Emitter->AssetTags.ContainsByPredicate([](const FNiagaraAssetTagDefinitionReference& AssetTagDefinitionReferenceCandidate)
+					{
+						return AssetTagDefinitionReferenceCandidate.GetTagDefinitionReferenceGuid() == INiagaraModule::Get().DeprecatedTagDefinition.TagGuid;
+					});
+					
+					if (bIsDeprecatedAsset || ParentData->bDeprecated)
 					{
 						DeprecatedEmitterNames.Add(FText::FromName(EmitterModel->GetName()));
 					}
@@ -254,12 +259,12 @@ bool SNiagaraOverviewGraphTitleBar::IsUsingDeprecatedEmitter() const
 		{
 			if (FVersionedNiagaraEmitterData* EmitterData = EmitterModel->GetEmitterHandle()->GetEmitterData())
 			{
-				if (FVersionedNiagaraEmitterData* ParentData = EmitterData->GetParent().GetEmitterData())
+				if(FVersionedNiagaraEmitterData* ParentData = EmitterData->GetParent().GetEmitterData())
 				{
-					if (ParentData->bDeprecated)
+					return ParentData->bDeprecated || EmitterData->GetParent().Emitter->AssetTags.ContainsByPredicate([](const FNiagaraAssetTagDefinitionReference& AssetTagDefinitionReferenceCandidate)
 					{
-						return true;
-					}
+						return AssetTagDefinitionReferenceCandidate.GetTagDefinitionReferenceGuid() == INiagaraModule::Get().DeprecatedTagDefinition.TagGuid;
+					});
 				}
 			}
 		}
