- Add cache tests for Socket DI

#rnx

#virtualized

diff --git a/Engine/Plugins/FX/Niagara/Shaders/Private/NiagaraDataInterfaceSocketReaderTemplate.ush b/Engine/Plugins/FX/Niagara/Shaders/Private/NiagaraDataInterfaceSocketReaderTemplate.ush
--- a/Engine/Plugins/FX/Niagara/Shaders/Private/NiagaraDataInterfaceSocketReaderTemplate.ush
+++ b/Engine/Plugins/FX/Niagara/Shaders/Private/NiagaraDataInterfaceSocketReaderTemplate.ush
@@ -27,13 +27,13 @@ int SanitizeSocketIndex_{ParameterName}(int SocketIndex)
 
 int SanitizeFilteredSocketIndex_{ParameterName}(int SocketIndex)
 {
-	const int FilteredSockedIndex = SocketIndex >= 0 && SocketIndex < {ParameterName}_NumFilteredSockets ? (SockedIndex + 1) : 0;
+	const int FilteredSocketIndex = SocketIndex >= 0 && SocketIndex < {ParameterName}_NumFilteredSockets ? (SocketIndex + 1) : 0;
 	return {ParameterName}_SocketData.Load(SocketIndex * 4);
 }
 
 int SanitizeUnfilteredSocketIndex_{ParameterName}(int SocketIndex)
 {
-	const int UnfilteredSockedIndex = SocketIndex >= 0 && SocketIndex < {ParameterName}_NumUnfilteredSockets ? (SockedIndex + 1 + {ParameterName}_NumFilteredSockets) : 0;
+	const int UnfilteredSocketIndex = SocketIndex >= 0 && SocketIndex < {ParameterName}_NumUnfilteredSockets ? (SocketIndex + 1 + {ParameterName}_NumFilteredSockets) : 0;
 	return {ParameterName}_SocketData.Load(SocketIndex * 4);
 }
 
