Revisit implicit dependency tracking with niagara assets for iterative cooking.  Core responsibility is now on the NiagaraSystem which calls into the NiagaraScript & NiagaraShader to get the necessary data hashed.  Also includes the per-asset dependency tracking to handle per-platform shader platform details.
#tests Use incremental cooking and validation (-cookincremental -iterativevalidate) to validate a shader compiler change results in a dirtied package

diff --git a/Engine/Plugins/FX/Niagara/Source/Niagara/Classes/NiagaraScript.h b/Engine/Plugins/FX/Niagara/Source/Niagara/Classes/NiagaraScript.h
--- a/Engine/Plugins/FX/Niagara/Source/Niagara/Classes/NiagaraScript.h
+++ b/Engine/Plugins/FX/Niagara/Source/Niagara/Classes/NiagaraScript.h
@@ -22,6 +22,7 @@
 
 #include "NiagaraScript.generated.h"
 
+struct FAppendToClassSchemaContext;
 class UNiagaraDataInterface;
 class FNiagaraCompileRequestDataBase;
 class FNiagaraCompileRequestDuplicateDataBase;
@@ -1161,11 +1162,9 @@ public:
 
 #if WITH_EDITORONLY_DATA
 	/**
-	 * Append config values or settings that can change how instances of the class are cooked, including especially
-	 * values that determine how version upgraded are conducted. Can also append a unique guid when necessary to
-	 * invalidate previous results because serialization changed and no custom version was updated.
+	 * Called from UNiagaraSystem::AppendToClassSchema() to collect the Script specific details for iterative cooking
 	 */
-	NIAGARA_API static void AppendToClassSchema(FAppendToClassSchemaContext& Context);
+	static void BuildClassSchema(FAppendToClassSchemaContext& Context);
 #endif
 
 	FNiagaraShaderScript* GetRenderThreadScript()
diff --git a/Engine/Plugins/FX/Niagara/Source/Niagara/Classes/NiagaraSystem.h b/Engine/Plugins/FX/Niagara/Source/Niagara/Classes/NiagaraSystem.h
--- a/Engine/Plugins/FX/Niagara/Source/Niagara/Classes/NiagaraSystem.h
+++ b/Engine/Plugins/FX/Niagara/Source/Niagara/Classes/NiagaraSystem.h
@@ -276,6 +276,13 @@ public:
 	//~ End INiagaraParameterDefinitionsSubscriber interface
 
 	NIAGARA_API virtual bool ChangeEmitterVersion(const FVersionedNiagaraEmitter& Emitter, const FGuid& NewVersion);
+
+	/**
+	 * Append config values or settings that can change how instances of the class are cooked, including especially
+	 * values that determine how version upgraded are conducted. Can also append a unique guid when necessary to
+	 * invalidate previous results because serialization changed and no custom version was updated.
+	 */
+	NIAGARA_API static void AppendToClassSchema(FAppendToClassSchemaContext& Context);
 #endif 
 
 	/** Gets an array of the emitter handles. */
diff --git a/Engine/Plugins/FX/Niagara/Source/Niagara/Private/NiagaraScript.cpp b/Engine/Plugins/FX/Niagara/Source/Niagara/Private/NiagaraScript.cpp
--- a/Engine/Plugins/FX/Niagara/Source/Niagara/Private/NiagaraScript.cpp
+++ b/Engine/Plugins/FX/Niagara/Source/Niagara/Private/NiagaraScript.cpp
@@ -4293,10 +4293,8 @@ void UNiagaraScript::ResolveParameterCollectionReferences()
 #endif
 
 #if WITH_EDITORONLY_DATA
-void UNiagaraScript::AppendToClassSchema(FAppendToClassSchemaContext& Context)
+void UNiagaraScript::BuildClassSchema(FAppendToClassSchemaContext& Context)
 {
-	Super::AppendToClassSchema(Context);
-
 	// Used by iterative cooking.  This will provide additional context for if things have changed such that a cook will
 	// be required.  This is focused on global settings rather than the usual dependencies between objects.
 
@@ -4320,6 +4318,8 @@ void UNiagaraScript::AppendToClassSchema(FAppendToClassSchemaContext& Context)
 	{
 		Context.Update(&NiagaraSettings->InvalidNamespaceWriteSeverity, sizeof(NiagaraSettings->InvalidNamespaceWriteSeverity));
 	}
+
+	FNiagaraShader::BuildClassSchema(Context);
 }
 #endif
 
diff --git a/Engine/Plugins/FX/Niagara/Source/Niagara/Private/NiagaraSystem.cpp b/Engine/Plugins/FX/Niagara/Source/Niagara/Private/NiagaraSystem.cpp
--- a/Engine/Plugins/FX/Niagara/Source/Niagara/Private/NiagaraSystem.cpp
+++ b/Engine/Plugins/FX/Niagara/Source/Niagara/Private/NiagaraSystem.cpp
@@ -37,12 +37,18 @@
 
 #include "Algo/RemoveIf.h"
 #include "Algo/StableSort.h"
+#include "Cooker/CookDependency.h"
 #include "HAL/LowLevelMemTracker.h"
+#include "Interfaces/ITargetPlatform.h"
+#include "Interfaces/ITargetPlatformManagerModule.h"
 #include "Misc/ScopedSlowTask.h"
 #include "Misc/ScopeExit.h"
 #include "Modules/ModuleManager.h"
 #include "ProfilingDebugging/CookStats.h"
 #include "ProfilingDebugging/ScopedTimers.h"
+#include "Serialization/CompactBinary.h"
+#include "Serialization/CompactBinaryWriter.h"
+#include "Serialization/ShaderKeyGenerator.h"
 #include "UObject/AssetRegistryTagsContext.h"
 #include "UObject/ObjectSaveContext.h"
 #include "UObject/LinkerLoad.h"
@@ -56,9 +62,6 @@
 
 #define LOCTEXT_NAMESPACE "NiagaraSystem"
 
-#if WITH_EDITOR
-#endif
-
 DECLARE_CYCLE_STAT(TEXT("Niagara - System - CompileScript"), STAT_Niagara_System_CompileScript, STATGROUP_Niagara);
 DECLARE_CYCLE_STAT(TEXT("Niagara - System - CompileScript_ResetAfter"), STAT_Niagara_System_CompileScriptResetAfter, STATGROUP_Niagara);
 
@@ -89,8 +92,24 @@ namespace NiagaraSystemPrivate
 	static const FName NAME_HasGPUEmitter("HasGPUEmitter");
 	static const FName NAME_FixedBoundsSize("FixedBoundsSize");
 	static const FName NAME_NumEmitters("NumEmitters");
+
+#if WITH_EDITOR
+	constexpr int32 CookDependenciesArgsVersion = 1;
+	void HashDependenciesForCook(FCbFieldViewIterator Args, UE::Cook::FCookDependencyContext& Context);
+	void RegisterCookDependencies(FObjectPreSaveContext ObjectSaveContext, bool bIncludeGpuTarget);
+
+	struct FCookDependency
+	{
+		TArray<FName> ShaderFormatNames;
+
+		bool TryLoad(FCbFieldViewIterator& ArgField);
+		void Save(FCbWriter& Writer);
+	};
+#endif
 }
 
+UE_COOK_DEPENDENCY_FUNCTION(HashNiagaraSystemDependenciesForCook, NiagaraSystemPrivate::HashDependenciesForCook);
+
 //Disable for now until we can spend more time on a good method of applying the data gathered.
 int32 GEnableNiagaraRuntimeCycleCounts = 0;
 static FAutoConsoleVariableRef CVarEnableNiagaraRuntimeCycleCounts(TEXT("fx.EnableNiagaraRuntimeCycleCounts"), GEnableNiagaraRuntimeCycleCounts, TEXT("Toggle for runtime cylce counts tracking Niagara's frame time. \n"), ECVF_ReadOnly);
@@ -292,6 +311,14 @@ void UNiagaraSystem::PreSave(FObjectPreSaveContext ObjectSaveContext)
 #if WITH_EDITORONLY_DATA
 	WaitForCompilationComplete();
 #endif
+
+#if WITH_EDITOR
+	if (ObjectSaveContext.IsCooking())
+	{
+		UpdateHasGPUEmitters();
+		NiagaraSystemPrivate::RegisterCookDependencies(ObjectSaveContext, bHasAnyGPUEmitters);
+	}
+#endif
 }
 
 #if WITH_EDITOR
@@ -4167,5 +4194,125 @@ const FNiagaraDataSetCompiledData& FNiagaraEmitterCompiledData::GetGPUCaptureDat
 }
 #endif
 
+
+
+//////////////////////////////////////////////////////////////////////////
+// This section is related to iterative cooking and defining the implicit dependencies that 
+// the NiagaraScript might have.  We don't have to worry about package to package dependencies
+// but are more focused on what external changes might require us to recook the script.
+// There are two paths we take to handle these changes:
+// AppendToClassSchema - evaluation of anything static that isn't tied to the target platform.
+// Registered cook dependencies - saves data into the oplog of individual packages during cook
+//	which can be evaluated on subsequent cooks to see if anything has changed
+// 
+
+#if WITH_EDITORONLY_DATA
+void UNiagaraSystem::AppendToClassSchema(FAppendToClassSchemaContext& Context)
+{
+	Super::AppendToClassSchema(Context);
+
+	// Used by iterative cooking.  This will provide additional context for if things have changed such that a cook will
+	// be required.  This is focused on global settings rather than the usual dependencies between objects.
+
+	UNiagaraScript::BuildClassSchema(Context);
+}
+
+#endif
+
+#if WITH_EDITOR
+
+namespace NiagaraSystemPrivate
+{
+
+bool FCookDependency::TryLoad(FCbFieldViewIterator& ArgField)
+{
+	const int32 ShaderFormatCount = (ArgField++).AsInt32();
+	ShaderFormatNames.Reset(ShaderFormatCount);
+	for (int32 i = 0; i < ShaderFormatCount; ++i)
+	{
+		ShaderFormatNames.Emplace((ArgField++).AsString());
+	}
+
+	return true;
+}
+
+void FCookDependency::Save(FCbWriter& Writer)
+{
+	const int32 ShaderFormatCount = ShaderFormatNames.Num();
+	Writer << ShaderFormatCount;
+	for (int32 i = 0; i < ShaderFormatCount; ++i)
+	{
+		Writer << ShaderFormatNames[i].ToString();
+	}
+}
+
+void HashDependenciesForCook(FCbFieldViewIterator Args, UE::Cook::FCookDependencyContext& Context)
+{
+	FCbFieldViewIterator ArgField(Args);
+
+	FCookDependency CookDependency;
+	const int32 ArgsVersion = (ArgField++).AsInt32();
+	bool bDependenciesValid = false;
+	if (ArgsVersion == CookDependenciesArgsVersion)
+	{
+		if (CookDependency.TryLoad(ArgField))
+		{
+			bDependenciesValid = true;
+		}
+	}
+
+	if (!bDependenciesValid)
+	{
+		Context.LogError(FString::Printf(TEXT("Unsupported arguments version %d."), ArgsVersion));
+		return;
+	}
+
+	if (!CookDependency.ShaderFormatNames.IsEmpty())
+	{
+		FShaderKeyGenerator KeyGen([&Context](const void* Data, uint64 Size) { Context.Update(Data, Size); });
+
+		ITargetPlatformManagerModule& TargetPlatformManager = GetTargetPlatformManagerRef();
+		for (const FName& ShaderFormat : CookDependency.ShaderFormatNames)
+		{
+			const uint32 ShaderFormatVersion = TargetPlatformManager.ShaderFormatVersion(ShaderFormat);
+			Context.Update(&ShaderFormatVersion, sizeof(ShaderFormatVersion));
+
+			const EShaderPlatform ShaderPlatform = ShaderFormatToLegacyShaderPlatform(ShaderFormat);
+			if (ShaderPlatform != SP_NumPlatforms)
+			{
+				ShaderMapAppendKey(ShaderPlatform, KeyGen);
+			}
+		}
+	}
+}
+
+void RegisterCookDependencies(FObjectPreSaveContext ObjectSaveContext, bool bIncludeGpuTarget)
+{
+	const ITargetPlatform* TargetPlatform = ObjectSaveContext.GetTargetPlatform();
+
+	FCookDependency CookDependency;
+	if (bIncludeGpuTarget)
+	{
+		TargetPlatform->GetAllTargetedShaderFormats(CookDependency.ShaderFormatNames);
+	}
+
+	CookDependency.ShaderFormatNames.Add(TEXT("VVM_1_0"));
+
+	Algo::Sort(CookDependency.ShaderFormatNames, FNameLexicalLess());
+
+	FCbWriter Writer;
+	Writer << CookDependenciesArgsVersion;
+
+	CookDependency.Save(Writer);
+
+	ObjectSaveContext.AddCookBuildDependency(
+		UE::Cook::FCookDependency::Function(
+			UE_COOK_DEPENDENCY_FUNCTION_CALL(HashNiagaraSystemDependenciesForCook), Writer.Save()));
+}
+
+} // NiagaraSystemPrivate
+
+#endif // WITH_EDITOR
+
 #undef LOCTEXT_NAMESPACE // NiagaraSystem
 
diff --git a/Engine/Plugins/FX/Niagara/Source/NiagaraShader/Private/NiagaraShader.cpp b/Engine/Plugins/FX/Niagara/Source/NiagaraShader/Private/NiagaraShader.cpp
--- a/Engine/Plugins/FX/Niagara/Source/NiagaraShader/Private/NiagaraShader.cpp
+++ b/Engine/Plugins/FX/Niagara/Source/NiagaraShader/Private/NiagaraShader.cpp
@@ -23,6 +23,7 @@
 #include "UObject/UObjectGlobals.h"
 #include "NiagaraShaderModule.h"
 #include "NiagaraCustomVersion.h"
+#include "Serialization/ShaderKeyGenerator.h"
 #include "UObject/UObjectThreadContext.h"
 
 IMPLEMENT_SHADER_TYPE(, FNiagaraShader, TEXT("/Plugin/FX/Niagara/Private/NiagaraEmitterInstanceShader.usf"),TEXT("SimulateMain"), SF_Compute)
@@ -1200,6 +1201,41 @@ FNiagaraShader::FNiagaraShader(const FNiagaraShaderType::CompiledShaderInitializ
 	}
 }
 
+#if WITH_EDITORONLY_DATA
+void FNiagaraShader::BuildClassSchema(FAppendToClassSchemaContext& Context)
+{
+	// Used by iterative cooking.  This will provide additional context for if things have changed such that a cook will
+	// be required.  This is focused on global settings rather than the usual dependencies between objects.
+
+	FShaderKeyGenerator KeyGen([&Context](const void* Data, uint64 Size) { Context.Update(Data, Size); });
+
+	FPlatformTypeLayoutParameters DefaultLayoutParams;
+	DefaultLayoutParams.Append(KeyGen);
+
+	const FSHAHash ContentLayoutHash = GetShaderTypeLayoutHash(StaticGetTypeLayoutDesc<FNiagaraShaderMapContent>(), DefaultLayoutParams);
+	Context.Update(&ContentLayoutHash, sizeof(ContentLayoutHash));
+
+	const FSHAHash ShaderLayoutHash = GetShaderTypeLayoutHash(StaticGetTypeLayoutDesc<FNiagaraShader>(), DefaultLayoutParams);
+	Context.Update(&ShaderLayoutHash, sizeof(ShaderLayoutHash));
+
+	const uint32 ShaderParamStructHash = TShaderParameterStructTypeInfo<FNiagaraShader::FParameters>::GetStructMetadata()->GetLayoutHash();
+	Context.Update(&ShaderParamStructHash, sizeof(ShaderParamStructHash));
+
+	// Add in any referenced HLSL files.
+	const FSHAHash InstanceShaderHash = GetShaderFileHash((TEXT("/Plugin/FX/Niagara/Private/NiagaraEmitterInstanceShader.usf")), EShaderPlatform::SP_PCD3D_SM5);
+	Context.Update(&InstanceShaderHash, sizeof(InstanceShaderHash));
+
+	const FSHAHash NiagaraShaderVersionHash = GetShaderFileHash((TEXT("/Plugin/FX/Niagara/Private/NiagaraShaderVersion.ush")), EShaderPlatform::SP_PCD3D_SM5);
+	Context.Update(&NiagaraShaderVersionHash, sizeof(NiagaraShaderVersionHash));
+
+	const FSHAHash ShaderVersionHash = GetShaderFileHash((TEXT("/Engine/Public/ShaderVersion.ush")), EShaderPlatform::SP_PCD3D_SM5);
+	Context.Update(&ShaderVersionHash, sizeof(ShaderVersionHash));
+
+	Context.Update(&GNiagaraSkipVectorVMBackendOptimizations, sizeof(GNiagaraSkipVectorVMBackendOptimizations));
+}
+#endif
+
+
 //////////////////////////////////////////////////////////////////////////
 bool FNiagaraDataInterfaceGeneratedFunction::Serialize(FArchive& Ar)
 {
diff --git a/Engine/Plugins/FX/Niagara/Source/NiagaraShader/Public/NiagaraShader.h b/Engine/Plugins/FX/Niagara/Source/NiagaraShader/Public/NiagaraShader.h
--- a/Engine/Plugins/FX/Niagara/Source/NiagaraShader/Public/NiagaraShader.h
+++ b/Engine/Plugins/FX/Niagara/Source/NiagaraShader/Public/NiagaraShader.h
@@ -24,6 +24,7 @@
 #endif
 
 class UClass;
+struct FAppendToClassSchemaContext;
 
 template<typename TBufferStruct> class TUniformBufferRef;
 
@@ -236,6 +237,13 @@ public:
 		return MakeArrayView(DataInterfaceParameters);
 	}
 
+#if WITH_EDITORONLY_DATA
+	/**
+	 * Called from UNiagaraSystem::AppendToClassSchema() to collect the Script specific details for iterative cooking
+	 */
+	NIAGARASHADER_API static void BuildClassSchema(FAppendToClassSchemaContext& Context);
+#endif
+
 	LAYOUT_FIELD(bool, bNeedsViewUniformBuffer);
 	LAYOUT_FIELD(uint16, MiscUsageBitMask);
 	LAYOUT_ARRAY(FShaderUniformBufferParameter, ExternalConstantBufferParam, 2);
