- Customize socket reader socket selection
- When a preview mesh is selected the list of sockets will be gathered from that

#rnx

diff --git a/Engine/Plugins/FX/Niagara/Source/Niagara/Internal/DataInterface/NiagaraDataInterfaceSocketReader.h b/Engine/Plugins/FX/Niagara/Source/Niagara/Internal/DataInterface/NiagaraDataInterfaceSocketReader.h
--- a/Engine/Plugins/FX/Niagara/Source/Niagara/Internal/DataInterface/NiagaraDataInterfaceSocketReader.h
+++ b/Engine/Plugins/FX/Niagara/Source/Niagara/Internal/DataInterface/NiagaraDataInterfaceSocketReader.h
@@ -124,6 +124,11 @@ public:
 	virtual bool HasPreSimulateTick() const override { return true; }
 	//UNiagaraDataInterface Interface
 
+#if WITH_EDITORONLY_DATA
+	// Returns a list of sockets if the editor preview asset is valid
+	NIAGARA_API TArray<FName> GetEditorSocketNames() const;
+#endif
+
 protected:
 #if WITH_EDITORONLY_DATA
 	virtual void GetFunctionsInternal(TArray<FNiagaraFunctionSignature>& OutFunctions) const override;
diff --git a/Engine/Plugins/FX/Niagara/Source/Niagara/Private/DataInterface/NiagaraDataInterfaceSocketReader.cpp b/Engine/Plugins/FX/Niagara/Source/Niagara/Private/DataInterface/NiagaraDataInterfaceSocketReader.cpp
--- a/Engine/Plugins/FX/Niagara/Source/Niagara/Private/DataInterface/NiagaraDataInterfaceSocketReader.cpp
+++ b/Engine/Plugins/FX/Niagara/Source/Niagara/Private/DataInterface/NiagaraDataInterfaceSocketReader.cpp
@@ -270,6 +270,44 @@ namespace NDISocketReaderLocal
 		InstanceData->ResolvedObject		= ResolvedObject;
 	}
 
+	void GetSocketNames(UObject* ResolvedObject, TArray<FName>& OutSocketNames)
+	{
+		if (USceneComponent* SceneComponent = Cast<USceneComponent>(ResolvedObject))
+		{
+			OutSocketNames = SceneComponent->GetAllSocketNames();
+		}
+		else if (UStaticMesh* StaticMesh = Cast<UStaticMesh>(ResolvedObject))
+		{
+			const int32 NumSockets = StaticMesh->Sockets.Num();
+			OutSocketNames.Reserve(StaticMesh->Sockets.Num());
+			for (UStaticMeshSocket* Socket : StaticMesh->Sockets)
+			{
+				check(Socket);
+				OutSocketNames.Add(Socket->SocketName);
+			}
+		}
+		else if (USkeletalMesh* SkeletalMesh = Cast<USkeletalMesh>(ResolvedObject))
+		{
+			const FReferenceSkeleton& RefSkeleton = SkeletalMesh->GetRefSkeleton();
+			const TArray<FMeshBoneInfo>& RefBoneInfo = RefSkeleton.GetRefBoneInfo();
+			const int32 NumBones = RefBoneInfo.Num();
+			const int32 NumSockets = SkeletalMesh->NumSockets();
+
+			OutSocketNames.Reserve(NumBones + NumSockets);
+			for (int i = 0; i < NumBones; ++i)
+			{
+				OutSocketNames.Add(RefBoneInfo[i].Name);
+			}
+
+			for (int i = 0; i < NumSockets; ++i)
+			{
+				USkeletalMeshSocket* Socket = SkeletalMesh->GetSocketByIndex(i);
+				check(Socket);
+				OutSocketNames.Add(Socket->SocketName);
+			}
+		}
+	}
+
 	void UpdateSocketCache(const UNiagaraDataInterfaceSocketReader* SocketDI, FInstanceData_GameThread* InstanceData, FNiagaraSystemInstance* SystemInstance)
 	{
 		const bool bResetPreviousTransform = InstanceData->bNeedsSocketRecache;
@@ -291,40 +329,7 @@ namespace NDISocketReaderLocal
 			InstanceData->SocketFilterUnfilteredIndex.Empty();
 			InstanceData->SocketNames.Empty();
 
-			if (USceneComponent* SceneComponent = Cast<USceneComponent>(InstanceData->ResolvedObject))
-			{
-				InstanceData->SocketNames = SceneComponent->GetAllSocketNames();
-			}
-			else if (UStaticMesh* StaticMesh = Cast<UStaticMesh>(InstanceData->ResolvedObject))
-			{
-				const int32 NumSockets = StaticMesh->Sockets.Num();
-				InstanceData->SocketNames.Reserve(StaticMesh->Sockets.Num());
-				for ( UStaticMeshSocket* Socket : StaticMesh->Sockets )
-				{
-					check(Socket)
-					InstanceData->SocketNames.Add(Socket->SocketName);
-				}
-			}
-			else if (USkeletalMesh* SkeletalMesh = Cast<USkeletalMesh>(InstanceData->ResolvedObject))
-			{
-				const FReferenceSkeleton& RefSkeleton = SkeletalMesh->GetRefSkeleton();
-				const TArray<FMeshBoneInfo>& RefBoneInfo = RefSkeleton.GetRefBoneInfo();
-				const int32 NumBones = RefBoneInfo.Num();
-				const int32 NumSockets = SkeletalMesh->NumSockets();
-
-				InstanceData->SocketNames.Reserve(NumBones + NumSockets);
-				for (int i = 0; i < NumBones; ++i)
-				{
-					InstanceData->SocketNames.Add(RefBoneInfo[i].Name);
-				}
-
-				for (int i = 0; i < NumSockets; ++i)
-				{
-					USkeletalMeshSocket* Socket = SkeletalMesh->GetSocketByIndex(i);
-					check(Socket);
-					InstanceData->SocketNames.Add(Socket->SocketName);
-				}
-			}
+			GetSocketNames(InstanceData->ResolvedObject.Get(), InstanceData->SocketNames);
 
 			InstanceData->NumSockets = InstanceData->SocketNames.Num();
 			if (InstanceData->NumSockets > 0 )
@@ -831,5 +836,17 @@ bool UNiagaraDataInterfaceSocketReader::CopyToInternal(UNiagaraDataInterface* De
 	return true;
 }
 
+#if WITH_EDITORONLY_DATA
+TArray<FName> UNiagaraDataInterfaceSocketReader::GetEditorSocketNames() const
+{
+	TArray<FName> SocketNames;
+	if (UObject* EditorAsset = EditorPreviewAsset.LoadSynchronous())
+	{
+		NDISocketReaderLocal::GetSocketNames(EditorAsset, SocketNames);
+	}
+	return SocketNames;
+}
+#endif
+
 #undef LOCTEXT_NAMESPACE
 
diff --git a/Engine/Plugins/FX/Niagara/Source/NiagaraEditorWidgets/Private/DetailCustomizations/NiagaraDataInterfaceSocketReaderDetails.cpp b/Engine/Plugins/FX/Niagara/Source/NiagaraEditorWidgets/Private/DetailCustomizations/NiagaraDataInterfaceSocketReaderDetails.cpp
new file mode 100644
--- /dev/null
+++ b/Engine/Plugins/FX/Niagara/Source/NiagaraEditorWidgets/Private/DetailCustomizations/NiagaraDataInterfaceSocketReaderDetails.cpp
@@ -0,0 +1,84 @@
+// Copyright Epic Games, Inc. All Rights Reserved.
+
+#include "NiagaraDataInterfaceSocketReaderDetails.h"
+#include "DataInterface/NiagaraDataInterfaceSocketReader.h" 
+#include "NiagaraDetailSourcedArrayBuilder.h"
+
+#include "Algo/Transform.h"
+#include "DetailCategoryBuilder.h"
+#include "DetailLayoutBuilder.h"
+
+#define LOCTEXT_NAMESPACE "NiagaraDataInterfaceSocketReaderDetails"
+
+FNiagaraDataInterfaceSocketReaderDetails::~FNiagaraDataInterfaceSocketReaderDetails()
+{
+	if (UNiagaraDataInterfaceSocketReader* DataInterface = WeakDataInterface.Get())
+	{
+		DataInterface->OnChanged().RemoveAll(this);
+	}
+}
+
+void FNiagaraDataInterfaceSocketReaderDetails::CustomizeDetails(IDetailLayoutBuilder& DetailBuilder)
+{
+	TArray<TWeakObjectPtr<UObject>> SelectedObjects;
+	DetailBuilder.GetObjectsBeingCustomized(SelectedObjects);
+	if (SelectedObjects.Num() != 1 || SelectedObjects[0]->IsA<UNiagaraDataInterfaceSocketReader>() == false)
+	{
+		return;
+	}
+
+	// Get the data interface
+	UNiagaraDataInterfaceSocketReader* DataInterface = CastChecked<UNiagaraDataInterfaceSocketReader>(SelectedObjects[0].Get());
+	WeakDataInterface = DataInterface;
+
+	DataInterface->OnChanged().RemoveAll(this);
+	DataInterface->OnChanged().AddSP(this, &FNiagaraDataInterfaceSocketReaderDetails::OnDataChanged);
+
+	// Customize
+	{
+		IDetailCategoryBuilder& DetailCategory = DetailBuilder.EditCategory("SocketReader");
+
+		TArray<TSharedRef<IPropertyHandle>> Properties;
+		DetailCategory.GetDefaultProperties(Properties, true, true);
+
+		TSharedPtr<IPropertyHandle> FilteredSocketsProperty = DetailBuilder.GetProperty(GET_MEMBER_NAME_CHECKED(UNiagaraDataInterfaceSocketReader, FilteredSockets));
+		for (TSharedPtr<IPropertyHandle> Property : Properties)
+		{
+			FProperty* PropertyPtr = Property->GetProperty();
+			if (PropertyPtr == FilteredSocketsProperty->GetProperty())
+			{
+				SocketArrayBuilder = MakeShared<FNiagaraDetailSourcedArrayBuilder>(Property.ToSharedRef(), GetSocketNames());
+				DetailCategory.AddCustomBuilder(SocketArrayBuilder.ToSharedRef());
+			}
+			else
+			{
+				DetailCategory.AddProperty(Property);
+			}
+		}
+	}
+}
+
+ TSharedRef<IDetailCustomization> FNiagaraDataInterfaceSocketReaderDetails::MakeInstance()
+ {
+	 return MakeShared<FNiagaraDataInterfaceSocketReaderDetails>();
+ }
+
+void FNiagaraDataInterfaceSocketReaderDetails::OnDataChanged()
+{
+	if (SocketArrayBuilder.IsValid())
+	{
+		SocketArrayBuilder->SetSourceArray(GetSocketNames());
+	}
+ }
+
+TArray<TSharedPtr<FName>> FNiagaraDataInterfaceSocketReaderDetails::GetSocketNames() const
+{
+	TArray<TSharedPtr<FName>> SocketNames;
+	if (UNiagaraDataInterfaceSocketReader* DataInterface = WeakDataInterface.Get())
+	{
+		Algo::Transform(DataInterface->GetEditorSocketNames(), SocketNames, [](const FName& SocketName) { return MakeShared<FName>(SocketName); });
+	}
+	return SocketNames;
+}
+
+#undef LOCTEXT_NAMESPACE
diff --git a/Engine/Plugins/FX/Niagara/Source/NiagaraEditorWidgets/Private/DetailCustomizations/NiagaraDataInterfaceSocketReaderDetails.h b/Engine/Plugins/FX/Niagara/Source/NiagaraEditorWidgets/Private/DetailCustomizations/NiagaraDataInterfaceSocketReaderDetails.h
new file mode 100644
--- /dev/null
+++ b/Engine/Plugins/FX/Niagara/Source/NiagaraEditorWidgets/Private/DetailCustomizations/NiagaraDataInterfaceSocketReaderDetails.h
@@ -0,0 +1,26 @@
+// Copyright Epic Games, Inc. All Rights Reserved.
+
+#pragma once
+
+#include "NiagaraDataInterfaceDetails.h"
+
+class IDetailLayoutBuilder;
+class FNiagaraDetailSourcedArrayBuilder;
+class UNiagaraDataInterfaceSocketReader;
+
+class FNiagaraDataInterfaceSocketReaderDetails : public IDetailCustomization//FNiagaraDataInterfaceDetailsBase
+{
+public:
+	~FNiagaraDataInterfaceSocketReaderDetails();
+	virtual void CustomizeDetails(IDetailLayoutBuilder& DetailBuilder) override;
+	static TSharedRef<IDetailCustomization> MakeInstance();
+
+private:
+	void OnDataChanged();
+
+	TArray<TSharedPtr<FName>> GetSocketNames() const;
+
+private:
+	TWeakObjectPtr<UNiagaraDataInterfaceSocketReader>	WeakDataInterface;
+	TSharedPtr<FNiagaraDetailSourcedArrayBuilder>		SocketArrayBuilder;
+};
diff --git a/Engine/Plugins/FX/Niagara/Source/NiagaraEditorWidgets/Private/DetailCustomizations/NiagaraDetailSourcedArrayBuilder.cpp b/Engine/Plugins/FX/Niagara/Source/NiagaraEditorWidgets/Private/DetailCustomizations/NiagaraDetailSourcedArrayBuilder.cpp
--- a/Engine/Plugins/FX/Niagara/Source/NiagaraEditorWidgets/Private/DetailCustomizations/NiagaraDetailSourcedArrayBuilder.cpp
+++ b/Engine/Plugins/FX/Niagara/Source/NiagaraEditorWidgets/Private/DetailCustomizations/NiagaraDetailSourcedArrayBuilder.cpp
@@ -94,7 +94,7 @@ void FNiagaraDetailSourcedArrayBuilder::OnGenerateEntry(TSharedRef<IPropertyHand
 	}
 }
 
-void FNiagaraDetailSourcedArrayBuilder::SetSourceArray(TArray<TSharedPtr<FName>>& InOptionsSource)
+void FNiagaraDetailSourcedArrayBuilder::SetSourceArray(const TArray<TSharedPtr<FName>>& InOptionsSource)
 {
 	OptionsSourceList = InOptionsSource;
 	RefreshChildren();
diff --git a/Engine/Plugins/FX/Niagara/Source/NiagaraEditorWidgets/Private/DetailCustomizations/NiagaraDetailSourcedArrayBuilder.h b/Engine/Plugins/FX/Niagara/Source/NiagaraEditorWidgets/Private/DetailCustomizations/NiagaraDetailSourcedArrayBuilder.h
--- a/Engine/Plugins/FX/Niagara/Source/NiagaraEditorWidgets/Private/DetailCustomizations/NiagaraDetailSourcedArrayBuilder.h
+++ b/Engine/Plugins/FX/Niagara/Source/NiagaraEditorWidgets/Private/DetailCustomizations/NiagaraDetailSourcedArrayBuilder.h
@@ -21,7 +21,7 @@ public:
 	virtual void GenerateHeaderRowContent(FDetailWidgetRow& NodeRow) override;
 	virtual void GenerateChildContent(IDetailChildrenBuilder& ChildrenBuilder) override;
 	
-	void SetSourceArray(TArray<TSharedPtr<FName>>& InOptionsSource);
+	void SetSourceArray(const TArray<TSharedPtr<FName>>& InOptionsSource);
 
 private: 
 	void OnGenerateEntry(TSharedRef<IPropertyHandle> ElementProperty, int32 ElementIndex, IDetailChildrenBuilder& ChildrenBuilder);
diff --git a/Engine/Plugins/FX/Niagara/Source/NiagaraEditorWidgets/Private/NiagaraEditorWidgetsModule.cpp b/Engine/Plugins/FX/Niagara/Source/NiagaraEditorWidgets/Private/NiagaraEditorWidgetsModule.cpp
--- a/Engine/Plugins/FX/Niagara/Source/NiagaraEditorWidgets/Private/NiagaraEditorWidgetsModule.cpp
+++ b/Engine/Plugins/FX/Niagara/Source/NiagaraEditorWidgets/Private/NiagaraEditorWidgetsModule.cpp
@@ -9,6 +9,7 @@
 #include "DetailCustomizations/NiagaraDataInterfaceGrid3DCollectionDetails.h"
 #include "DetailCustomizations/NiagaraDataInterfaceSkeletalMeshDetails.h"
 #include "DetailCustomizations/NiagaraDataInterfaceStaticMeshDetails.h"
+#include "DetailCustomizations/NiagaraDataInterfaceSocketReaderDetails.h"
 #include "DetailCustomizations/NiagaraDataInterfaceMeshRendererInfoDetails.h"
 #include "DetailCustomizations/NiagaraDataInterfaceSpriteRendererInfoDetails.h"
 #include "DetailCustomizations/NiagaraDecalRendererDetails.h"
@@ -133,6 +134,7 @@ void FNiagaraEditorWidgetsModule::StartupModule()
 	PropertyModule.RegisterCustomClassLayout("NiagaraDataInterfaceColorCurve", FOnGetDetailCustomizationInstance::CreateStatic(&FNiagaraDataInterfaceColorCurveDetails::MakeInstance));
 	PropertyModule.RegisterCustomClassLayout("NiagaraDataInterfaceSkeletalMesh", FOnGetDetailCustomizationInstance::CreateStatic(&FNiagaraDataInterfaceSkeletalMeshDetails::MakeInstance));
 	PropertyModule.RegisterCustomClassLayout("NiagaraDataInterfaceStaticMesh", FOnGetDetailCustomizationInstance::CreateStatic(&FNiagaraDataInterfaceStaticMeshDetails::MakeInstance));
+	PropertyModule.RegisterCustomClassLayout("NiagaraDataInterfaceSocketReader", FOnGetDetailCustomizationInstance::CreateStatic(&FNiagaraDataInterfaceSocketReaderDetails::MakeInstance));
 	PropertyModule.RegisterCustomClassLayout("NiagaraDataInterfaceGrid2DCollection", FOnGetDetailCustomizationInstance::CreateStatic(&FNiagaraDataInterfaceGrid2DCollectionDetails::MakeInstance));
 	PropertyModule.RegisterCustomClassLayout("NiagaraDataInterfaceGrid3DCollection", FOnGetDetailCustomizationInstance::CreateStatic(&FNiagaraDataInterfaceGrid3DCollectionDetails::MakeInstance));
 	PropertyModule.RegisterCustomClassLayout("NiagaraDataInterfaceMeshRendererInfo", FOnGetDetailCustomizationInstance::CreateStatic(&FNiagaraDataInterfaceMeshRendererInfoDetails::MakeInstance));
