Couple of fixes for the AsyncTasks compilation path
-zero iniitalize traversal entry so that we don't get garbage hash collisions
-make sure we issue a call to the target manager so that it collects the shader formats (since it's not threadsafe) on the gamethread
#tests Compile large number of assets (like a cook of a large project)

diff --git a/Engine/Plugins/FX/Niagara/Source/NiagaraEditor/Private/NiagaraSystemCompilingManager.cpp b/Engine/Plugins/FX/Niagara/Source/NiagaraEditor/Private/NiagaraSystemCompilingManager.cpp
--- a/Engine/Plugins/FX/Niagara/Source/NiagaraEditor/Private/NiagaraSystemCompilingManager.cpp
+++ b/Engine/Plugins/FX/Niagara/Source/NiagaraEditor/Private/NiagaraSystemCompilingManager.cpp
@@ -3,13 +3,13 @@
 #include "NiagaraSystemCompilingManager.h"
 
 #include "DataDrivenShaderPlatformInfo.h"
+#include "Interfaces/ITargetPlatformManagerModule.h"
 #include "NiagaraCompilationTasks.h"
 #include "NiagaraEditorModule.h"
 #include "NiagaraShaderType.h"
 #include "ProfilingDebugging/CookStats.h"
 #include "UObject/UObjectIterator.h"
 #include "UObject/UObjectThreadContext.h"
-
 #include "Misc/ScopeRWLock.h"
 
 #define LOCTEXT_NAMESPACE "NiagaraCompilationManager"
@@ -94,6 +94,13 @@ FNiagaraShaderMapId BuildShaderMapId(FNiagaraShaderType* ShaderType, const ITarg
 	return ShaderMapId;
 }
 
+void EnsureTargetPlatformsLoaded(const FNiagaraSystemCompilationTask& CompilationTask)
+{
+	// in order to work around the fact that the target platform API is not thread safe we need to make sure that the target platform and the
+	// various shader formats it handles has been initialized
+	GetTargetPlatformManagerRef().ShaderFormatVersion(TEXT("VVM_1_0"));
+}
+
 };
 
 FNiagaraSystemCompilingManager& FNiagaraSystemCompilingManager::Get()
@@ -424,6 +431,8 @@ FNiagaraCompilationTaskHandle FNiagaraSystemCompilingManager::AddSystem(UNiagara
 			CompilationTask->AddScript(ScriptToCompile.EmitterIndex, ScriptToCompile.Script, ScriptToCompile.CompileId, ScriptToCompile.bRequiresCompilation, ScriptToCompile.ShaderRequests);
 		}
 
+		NiagaraSystemCompilingManagerImpl::EnsureTargetPlatformsLoaded(*CompilationTask);
+
 		CompilationTask->QueueStartTime = FPlatformTime::Seconds();
 	}
 
diff --git a/Engine/Plugins/FX/Niagara/Source/NiagaraEditor/Private/NiagaraTraversalStateContext.h b/Engine/Plugins/FX/Niagara/Source/NiagaraEditor/Private/NiagaraTraversalStateContext.h
--- a/Engine/Plugins/FX/Niagara/Source/NiagaraEditor/Private/NiagaraTraversalStateContext.h
+++ b/Engine/Plugins/FX/Niagara/Source/NiagaraEditor/Private/NiagaraTraversalStateContext.h
@@ -15,7 +15,7 @@ class FNiagaraFixedConstantResolver;
 struct FNiagaraTraversalStackEntry
 {
 	FGuid NodeGuid;
-	uint32 FullStackHash;
+	uint32 FullStackHash = 0;
 #if WITH_NIAGARA_TRAVERSAL_FRIENDLY_NAME
 	FString FriendlyName;
 #endif
