Speculative Mac compile fix

#rnx

diff --git a/Engine/Plugins/FX/Niagara/Source/Niagara/Private/NiagaraDebugHud.h b/Engine/Plugins/FX/Niagara/Source/Niagara/Private/NiagaraDebugHud.h
--- a/Engine/Plugins/FX/Niagara/Source/Niagara/Private/NiagaraDebugHud.h
+++ b/Engine/Plugins/FX/Niagara/Source/Niagara/Private/NiagaraDebugHud.h
@@ -162,6 +162,8 @@ class FNiagaraDebugHud
 	template<typename TCounterType>
 	struct FSmoothedCounter
 	{
+		FSmoothedCounter() {}
+
 		uint64			LastFrameSeen = 0;
 		bool			SmoothedSetOnce = false;
 		TCounterType	SmoothedFrameMax = TCounterType();
