New spawing functions for NDC.

SpawnDirect - Read a value from the NDC directly to control spawn count
ScaleSpawnCount - Apply a scale to previously set spawn counts based on some value in the NDC.

Both use a new way of encoding a variable inside a DI function specifier.
With some new UI/UX work layered on top that allows for easy and discoverable selection of NDC variables based on the context.

Misc:
- Streamlined method for specifying defaults for DI function inputs.
- Fixed issue with some node context menus appearing on the wrong DI nodes.
- Any DI function can now customize it's function specifier widget.
- Minor tweaks to data set readers to allow for direct reading of FNiagaraDataBuffers

diff --git a/Engine/Plugins/FX/Niagara/Source/Niagara/Classes/NiagaraDataSet.h b/Engine/Plugins/FX/Niagara/Source/Niagara/Classes/NiagaraDataSet.h
--- a/Engine/Plugins/FX/Niagara/Source/Niagara/Classes/NiagaraDataSet.h
+++ b/Engine/Plugins/FX/Niagara/Source/Niagara/Classes/NiagaraDataSet.h
@@ -515,6 +515,11 @@ namespace NiagaraDataSetPrivate
 	{
 		return DataSet.GetCompiledData();
 	}
+	
+	inline const FNiagaraDataSetCompiledData& GetCompiledData(const FNiagaraDataBuffer* DataBuffer)
+	{
+		return DataBuffer->GetOwner()->GetCompiledData();
+	}
 
 	inline FNiagaraDataBuffer* GetCurrentData(const FNiagaraDataSet& DataSet)
 	{
diff --git a/Engine/Plugins/FX/Niagara/Source/Niagara/Classes/NiagaraDataSetAccessor.h b/Engine/Plugins/FX/Niagara/Source/Niagara/Classes/NiagaraDataSetAccessor.h
--- a/Engine/Plugins/FX/Niagara/Source/Niagara/Classes/NiagaraDataSetAccessor.h
+++ b/Engine/Plugins/FX/Niagara/Source/Niagara/Classes/NiagaraDataSetAccessor.h
@@ -14,6 +14,7 @@ struct FNiagaraDataSetAccessorTypeInfo;
 namespace NiagaraDataSetPrivate
 {
 	inline const FNiagaraDataSetCompiledData& GetCompiledData(const FNiagaraDataSet& DataSet);
+	inline const FNiagaraDataSetCompiledData& GetCompiledData(const FNiagaraDataBuffer* DataBuffer);
 	inline FNiagaraDataBuffer* GetCurrentData(const FNiagaraDataSet& DataSet);
 	inline FNiagaraDataBuffer* GetDestinationData(const FNiagaraDataSet& DataSet);
 	inline uint8* GetComponentPtrFloat(FNiagaraDataBuffer* DataBuffer, uint32 ComponentIdx);
@@ -31,6 +32,7 @@ struct FNiagaraDataSetAccessor : public FNiagaraDataSetAccessorTypeInfo<TType>::
 	static_assert(!TIsUECoreVariant<TType, double>::Value, "Double core variant. Must be float type!");
 	FORCEINLINE FNiagaraDataSetAccessor<TType>() {}
 	FORCEINLINE explicit FNiagaraDataSetAccessor<TType>(const FNiagaraDataSet& DataSet, const FName VariableName) { FNiagaraDataSetAccessorTypeInfo<TType>::TAccessorBaseClass::Init(NiagaraDataSetPrivate::GetCompiledData(DataSet), VariableName); }
+	FORCEINLINE explicit FNiagaraDataSetAccessor<TType>(const FNiagaraDataBuffer* DataBuffer, const FName VariableName) { FNiagaraDataSetAccessorTypeInfo<TType>::TAccessorBaseClass::Init(NiagaraDataSetPrivate::GetCompiledData(DataBuffer), VariableName); }
 	FORCEINLINE explicit FNiagaraDataSetAccessor<TType>(const FNiagaraDataSetCompiledData& DataSetCompiledData, const FName VariableName) { FNiagaraDataSetAccessorTypeInfo<TType>::TAccessorBaseClass::Init(DataSetCompiledData, VariableName); }
 };
 
@@ -45,7 +47,7 @@ struct FNiagaraDataSetReaderFloat
 
 	FORCEINLINE FNiagaraDataSetReaderFloat() {}
 
-	explicit FNiagaraDataSetReaderFloat(FNiagaraDataBuffer* DataBuffer, bool bInIsFloat, int32 ComponentIndex)
+	explicit FNiagaraDataSetReaderFloat(const FNiagaraDataBuffer* DataBuffer, bool bInIsFloat, int32 ComponentIndex)
 		: bIsFloat(bInIsFloat)
 	{
 		using namespace NiagaraDataSetPrivate;
@@ -188,6 +190,7 @@ struct FNiagaraDataSetAccessorFloat
 	static_assert(!TIsUECoreVariant<TType, double>::Value, "Double core variant. Must be float type!");
 	FORCEINLINE FNiagaraDataSetAccessorFloat() {}
 	FORCEINLINE explicit FNiagaraDataSetAccessorFloat(const FNiagaraDataSet& DataSet, const FName VariableName) { Init(NiagaraDataSetPrivate::GetCompiledData(DataSet), VariableName); }
+	FORCEINLINE explicit FNiagaraDataSetAccessorFloat(const FNiagaraDataBuffer* DataBuffer, const FName VariableName) { Init(NiagaraDataSetPrivate::GetCompiledData(DataBuffer), VariableName); }
 	FORCEINLINE explicit FNiagaraDataSetAccessorFloat(const FNiagaraDataSetCompiledData& DataSetCompiledData, const FName VariableName) { Init(DataSetCompiledData, VariableName); }
 	FORCEINLINE explicit FNiagaraDataSetAccessorFloat(const FNiagaraDataSetCompiledData* DataSetCompiledData, const FName VariableName) { Init(DataSetCompiledData, VariableName); }
 	FORCEINLINE void Init(const FNiagaraDataSet& DataSet, const FName VariableName) { Init(NiagaraDataSetPrivate::GetCompiledData(DataSet), VariableName); }
@@ -242,9 +245,11 @@ struct FNiagaraDataSetAccessorFloat
 	FORCEINLINE bool IsValid() const { return ComponentIndex != INDEX_NONE; }
 
 	FORCEINLINE static FNiagaraDataSetReaderFloat<TType> CreateReader(const FNiagaraDataSet& DataSet, const FName VariableName) { return FNiagaraDataSetAccessorFloat<TType>(DataSet, VariableName).GetReader(DataSet); }
+	FORCEINLINE static FNiagaraDataSetReaderFloat<TType> CreateReader(const FNiagaraDataBuffer* DataBuffer, const FName VariableName) { return FNiagaraDataSetAccessorFloat<TType>(DataBuffer, VariableName).GetReader(DataBuffer); }
 	//FORCEINLINE static FNiagaraDataSetWriterFloat<TType> CreateWriter(const FNiagaraDataSet& DataSet, const FName VariableName) { return FNiagaraDataSetAccessorFloat<TType>(DataSet, VariableName).GetWriter(DataSet); }
 
 	FORCEINLINE FNiagaraDataSetReaderFloat<TType> GetReader(const FNiagaraDataSet& DataSet) const { return FNiagaraDataSetReaderFloat<TType>(NiagaraDataSetPrivate::GetCurrentData(DataSet), bIsFloat, ComponentIndex); }
+	FORCEINLINE FNiagaraDataSetReaderFloat<TType> GetReader(const FNiagaraDataBuffer* Data) const { return FNiagaraDataSetReaderFloat<TType>(Data, bIsFloat, ComponentIndex); }
 	//FORCEINLINE FNiagaraDataSetWriterFloat<TType> GetWriter(const FNiagaraDataSet& DataSet) const { return FNiagaraDataSetWriterFloat<TType>(NiagaraDataSetPrivate::GetDestinationData(DataSet), bIsFloat, ComponentIndex); }
 
 private:
@@ -351,7 +356,7 @@ struct FNiagaraDataSetReaderInt32
 {
 	FORCEINLINE FNiagaraDataSetReaderInt32() { }
 
-	explicit FNiagaraDataSetReaderInt32(FNiagaraDataBuffer* DataBuffer, int32 ComponentIndex)
+	explicit FNiagaraDataSetReaderInt32(const FNiagaraDataBuffer* DataBuffer, int32 ComponentIndex)
 	{
 		using namespace NiagaraDataSetPrivate;
 		if (DataBuffer != nullptr && ComponentIndex != INDEX_NONE)
@@ -455,6 +460,7 @@ struct FNiagaraDataSetAccessorInt32
 {
 	FORCEINLINE FNiagaraDataSetAccessorInt32() {}
 	FORCEINLINE explicit FNiagaraDataSetAccessorInt32(const FNiagaraDataSet& DataSet, const FName VariableName) { Init(NiagaraDataSetPrivate::GetCompiledData(DataSet), VariableName); }
+	FORCEINLINE explicit FNiagaraDataSetAccessorInt32(const FNiagaraDataBuffer* DataBuffer, const FName VariableName) { Init(NiagaraDataSetPrivate::GetCompiledData(DataBuffer), VariableName); }
 	FORCEINLINE explicit FNiagaraDataSetAccessorInt32(const FNiagaraDataSetCompiledData& DataSetCompiledData, const FName VariableName) { Init(DataSetCompiledData, VariableName); }
 	FORCEINLINE explicit FNiagaraDataSetAccessorInt32(const FNiagaraDataSetCompiledData* DataSetCompiledData, const FName VariableName) { Init(DataSetCompiledData, VariableName); }
 	FORCEINLINE void Init(const FNiagaraDataSet& DataSet, const FName VariableName) { Init(NiagaraDataSetPrivate::GetCompiledData(DataSet), VariableName); }
@@ -485,9 +491,11 @@ struct FNiagaraDataSetAccessorInt32
 	FORCEINLINE bool IsValid() const { return ComponentIndex != INDEX_NONE; }
 
 	FORCEINLINE static FNiagaraDataSetReaderInt32<TType> CreateReader(const FNiagaraDataSet& DataSet, const FName VariableName) { return FNiagaraDataSetAccessorInt32<TType>(DataSet, VariableName).GetReader(DataSet); }
+	FORCEINLINE static FNiagaraDataSetReaderInt32<TType> CreateReader(const FNiagaraDataBuffer* Data, const FName VariableName) { return FNiagaraDataSetAccessorInt32<TType>(Data, VariableName).GetReader(Data); }
 	FORCEINLINE static FNiagaraDataSetWriterInt32<TType> CreateWriter(const FNiagaraDataSet& DataSet, const FName VariableName) { return FNiagaraDataSetAccessorInt32<TType>(DataSet, VariableName).GetWriter(DataSet); }
 
 	FORCEINLINE FNiagaraDataSetReaderInt32<TType> GetReader(const FNiagaraDataSet& DataSet) const { return FNiagaraDataSetReaderInt32<TType>(NiagaraDataSetPrivate::GetCurrentData(DataSet), ComponentIndex); }
+	FORCEINLINE FNiagaraDataSetReaderInt32<TType> GetReader(const FNiagaraDataBuffer* Data) const { return FNiagaraDataSetReaderInt32<TType>(Data, ComponentIndex); }
 	FORCEINLINE FNiagaraDataSetWriterInt32<TType> GetWriter(const FNiagaraDataSet& DataSet) const { return FNiagaraDataSetWriterInt32<TType>(NiagaraDataSetPrivate::GetDestinationData(DataSet), ComponentIndex); }
 
 private:
diff --git a/Engine/Plugins/FX/Niagara/Source/Niagara/Internal/DataInterface/NiagaraDataInterfaceDataChannelRead.h b/Engine/Plugins/FX/Niagara/Source/Niagara/Internal/DataInterface/NiagaraDataInterfaceDataChannelRead.h
--- a/Engine/Plugins/FX/Niagara/Source/Niagara/Internal/DataInterface/NiagaraDataInterfaceDataChannelRead.h
+++ b/Engine/Plugins/FX/Niagara/Source/Niagara/Internal/DataInterface/NiagaraDataInterfaceDataChannelRead.h
@@ -52,6 +52,19 @@ enum class ENDIDataChannelSpawnMode
 	Max UMETA(Hidden),
 };
 
+/** Mode controlling the behavior of the ScaleSpawnCount function for Niagara's Data Channel Read Data Interface.*/
+UENUM()
+enum class ENDIDataChannelSpawnScaleMode
+{
+	/** This mode will override any previously set scale values. */
+	Override,
+
+	/** This mode will combine with previously set scales. e.g. Two calls that scale by 0.5 will result in a final spawn scale of 0.25. */
+	Scale,
+
+	Max UMETA(Hidden),
+};
+
 UCLASS(Experimental, EditInlineNew, Category = "Data Channels", CollapseCategories, meta = (DisplayName = "Data Channel Reader"), MinimalAPI)
 class UNiagaraDataInterfaceDataChannelRead : public UNiagaraDataInterfaceRWBase
 {
@@ -161,6 +174,12 @@ public:
 	//Emitter only functions.
 	NIAGARA_API void SpawnConditional(FVectorVMExternalFunctionContext& Context, int32 FuncIndex);
 
+	template<typename T>
+	NIAGARA_API void SpawnDirect(FVectorVMExternalFunctionContext& Context, FName NDCVarName);
+	
+	template<typename T>
+	NIAGARA_API void ScaleSpawnCount(FVectorVMExternalFunctionContext& Context, FName NDCVarName);
+
 	FNDIDataChannelCompiledData& GetCompiledData() { return CompiledData; }
 
 protected:
@@ -186,6 +205,21 @@ struct FNDIDataChannelRead_EmitterSpawnData
 	}
 };
 
+struct FNDIDataChannelRead_EmitterSpawnInfo
+{
+private:
+	uint32 Count = 0;
+	float Scale = 1.0f;
+
+public:
+	uint32 Get()const { return Count * Scale; }
+	
+	void SetCount(uint32 NewCount) { Count = NewCount; }
+	void Append(uint32 NewCount) { Count += NewCount; }
+	void SetScale(float NewScale) { Scale = NewScale; }
+	void ApplyScale(float NewScale) { Scale *= NewScale; }
+};
+
 struct FNDIDataChannelRead_EmitterInstanceData
 {
 	//Spawn data buffers needed for accessing the correct NDCIndex and spawn data from a particle during CPU or GPU execution.
@@ -193,7 +227,7 @@ struct FNDIDataChannelRead_EmitterInstanceData
 	FNDIDataChannelRead_EmitterSpawnData NDCSpawnData;
 
 	//Spawn Counts for each entry in the NDC.
-	TArray<int32> NDCSpawnCounts;
+	TArray<FNDIDataChannelRead_EmitterSpawnInfo> NDCSpawnCounts;
 
 	void Reset()
 	{
diff --git a/Engine/Plugins/FX/Niagara/Source/Niagara/Private/DataInterface/NiagaraDataInterfaceDataChannelRead.cpp b/Engine/Plugins/FX/Niagara/Source/Niagara/Private/DataInterface/NiagaraDataInterfaceDataChannelRead.cpp
--- a/Engine/Plugins/FX/Niagara/Source/Niagara/Private/DataInterface/NiagaraDataInterfaceDataChannelRead.cpp
+++ b/Engine/Plugins/FX/Niagara/Source/Niagara/Private/DataInterface/NiagaraDataInterfaceDataChannelRead.cpp
@@ -58,6 +58,29 @@ static FAutoConsoleVariableRef CVarNDCReadForcePrevFrame(
 	ECVF_Default
 );
 
+const static FName VarNameKey(TEXT("VarName"));
+const static FName VarTypeKey(TEXT("VarType"));
+
+//Decode a function specifier value back into a Niagara Variable.
+FNiagaraVariableBase DecodeVariableFromSpecifiers(const FVMExternalFunctionBindingInfo& BindingInfo)
+{
+	const FVMFunctionSpecifier* NameSpec = BindingInfo.FindSpecifier(VarNameKey);
+	const FVMFunctionSpecifier* TypeSpec = BindingInfo.FindSpecifier(VarTypeKey);
+	if(NameSpec && TypeSpec)
+	{
+		FStringOutputDevice ErrorOut;
+		const UScriptStruct* TypeStruct = FNiagaraTypeDefinition::StaticStruct();
+		FNiagaraTypeDefinition TypeDef;
+		if (TypeStruct->ImportText(*TypeSpec->Value.ToString(), &TypeDef, nullptr, PPF_None, &ErrorOut, TypeStruct->GetName(), true))
+		{
+			return FNiagaraVariableBase(TypeDef, NameSpec->Value);
+		}
+
+		UE_LOG(LogNiagara, Error, TEXT("%s"), *ErrorOut);		
+	}
+	return FNiagaraVariableBase();
+}
+
 namespace NDIDataChannelReadLocal
 {
 	static const TCHAR* CommonShaderFile = TEXT("/Plugin/FX/Niagara/Private/DataChannel/NiagaraDataInterfaceDataChannelCommon.ush");
@@ -100,7 +123,7 @@ namespace NDIDataChannelReadLocal
 			Sig.ModuleUsageBitmask = ENiagaraScriptUsageMask::Particle;
 			Sig.AddInput(FNiagaraVariable(FNiagaraTypeDefinition(UNiagaraDataInterfaceDataChannelRead::StaticClass()), TEXT("DataChannel interface")));
 			Sig.AddInput(FNiagaraVariable(FNiagaraTypeDefinition(FNiagaraEmitterID::StaticStruct()), TEXT("Emitter ID")), LOCTEXT("EmitterIDDesc", "ID of the emitter we'd like to spawn into. This can be obtained from Engine.Emitter.ID."));
-			Sig.AddInput(FNiagaraVariable(FNiagaraTypeDefinition::GetIntDef(), TEXT("Spawned Particle Exec Index")), LOCTEXT("GetNDCSpawnData_InExecIndexDesc","The execution index of the spawned particle."));
+			Sig.AddInputWithDefault(FNiagaraVariable(FNiagaraTypeDefinition::GetIntDef(), TEXT("Spawned Particle Exec Index")), 0, LOCTEXT("GetNDCSpawnData_InExecIndexDesc","The execution index of the spawned particle."));
 			Sig.AddOutput(FNiagaraVariable(FNiagaraTypeDefinition::GetIntDef(), TEXT("NDC Index")), LOCTEXT("GetNDCSpawnData_OutNDCIndexDesc","Index of the NDC item that spawned this particle."));
 			Sig.AddOutput(FNiagaraVariable(FNiagaraTypeDefinition::GetIntDef(), TEXT("NDC Spawn Index")), LOCTEXT("GetNDCSpawnData_OutNDCSpawnIndexDesc","The index of this particle in relation to all the particle spawned by the same NDC item."));
 			Sig.AddOutput(FNiagaraVariable(FNiagaraTypeDefinition::GetIntDef(), TEXT("NDC Spawn Count")), LOCTEXT("GetNDCSpawnData_OutNDCSpawnCountDesc","The number of particles spawned by the same NDC item."));
@@ -121,7 +144,7 @@ namespace NDIDataChannelReadLocal
 			Sig.bMemberFunction = true;
 			Sig.bExperimental = true;
 			Sig.AddInput(FNiagaraVariable(FNiagaraTypeDefinition(UNiagaraDataInterfaceDataChannelRead::StaticClass()), TEXT("DataChannel interface")));
-			Sig.AddInput(FNiagaraVariable(FNiagaraTypeDefinition::GetIntDef(), TEXT("Index")), LOCTEXT("ConsumeIndexInputDesc", "The index to read."));
+			Sig.AddInputWithDefault(FNiagaraVariable(FNiagaraTypeDefinition::GetIntDef(), TEXT("Index")), 0, LOCTEXT("ConsumeIndexInputDesc", "The index to read."));
 			Sig.AddOutput(FNiagaraVariable(FNiagaraTypeDefinition::GetBoolDef(), TEXT("Success")), LOCTEXT("ConsumeSuccessOutputDesc", "True if all reads succeeded."));
 			Sig.RequiredOutputs = Sig.Outputs.Num();//The user defines what we read in the graph.			
 		}
@@ -141,7 +164,7 @@ namespace NDIDataChannelReadLocal
 			Sig.bMemberFunction = true;
 			Sig.bExperimental = true;
 			Sig.AddInput(FNiagaraVariable(FNiagaraTypeDefinition(UNiagaraDataInterfaceDataChannelRead::StaticClass()), TEXT("DataChannel interface")));
-			Sig.AddInput(FNiagaraVariable(FNiagaraTypeDefinition::GetBoolDef(), TEXT("Consume")), LOCTEXT("ConsumeInputDesc", "True if this instance (particle/emitter etc) should consume data from the data channel in this call."));
+			Sig.AddInputWithDefault(FNiagaraVariable(FNiagaraTypeDefinition::GetBoolDef(), TEXT("Consume")), FNiagaraBool(true), LOCTEXT("ConsumeInputDesc", "True if this instance (particle/emitter etc) should consume data from the data channel in this call."));
 			Sig.AddOutput(FNiagaraVariable(FNiagaraTypeDefinition::GetBoolDef(), TEXT("Success")), LOCTEXT("ConsumeSuccessOutputDesc", "True if all reads succeeded."));
 			Sig.AddOutput(FNiagaraVariable(FNiagaraTypeDefinition::GetIntDef(), TEXT("Index")), LOCTEXT("ConsumeIndexOutputDesc", "The index we actually read from. If reading failed this can be -1. This allows subsequent reads of the data channel at this index."));
 			Sig.RequiredOutputs = Sig.Outputs.Num();//The user defines what we read in the graph.			
@@ -160,9 +183,6 @@ namespace NDIDataChannelReadLocal
 		static FNiagaraFunctionSignature Sig;
 		if(!Sig.IsValid())
 		{
-			FNiagaraVariable EnabledVar(FNiagaraTypeDefinition::GetBoolDef(), TEXT("Enable"));
-			EnabledVar.SetValue(FNiagaraBool(true));
-
 			Sig.Name = TEXT("SpawnConditional");
 #if WITH_EDITORONLY_DATA
 			NIAGARA_ADD_FUNCTION_SOURCE_INFO(Sig)
@@ -172,7 +192,6 @@ namespace NDIDataChannelReadLocal
 		For compound data types that contain multiple component floats or ints, comparissons are done on a per component basis.\n\
 		For example if you add a Vector condition parameter it will be compared against each component of the corresponding Vector in the Data Channel.\n\
 		Result = (Param.X == ChannelValue.X) && (Param.Y == ChannelValue.Y) && (Param.Z == ChannelValue.Z)");
-
 			Sig.FunctionVersion = (uint32)FunctionVersion_SpawnConditional::EmitterIDParameter;
 #endif
 			Sig.bMemberFunction = true;
@@ -180,16 +199,72 @@ namespace NDIDataChannelReadLocal
 			Sig.bRequiresExecPin = true;
 			Sig.ModuleUsageBitmask = ENiagaraScriptUsageMask::Emitter | ENiagaraScriptUsageMask::System;
 			Sig.AddInput(FNiagaraVariable(FNiagaraTypeDefinition(UNiagaraDataInterfaceDataChannelRead::StaticClass()), TEXT("DataChannel interface")));
-			Sig.AddInput(EnabledVar, LOCTEXT("SpawnEnableInputDesc", "Enable or disable this function call. If false, this call with have no effetcs."));
+			Sig.AddInputWithDefault(FNiagaraVariable(FNiagaraTypeDefinition::GetBoolDef(), TEXT("Enable")), FNiagaraBool(true), LOCTEXT("SpawnEnableInputDesc", "Enable or disable this function call. If false, this call with have no effetcs."));
 			Sig.AddInput(FNiagaraVariable(FNiagaraTypeDefinition(FNiagaraEmitterID::StaticStruct()), TEXT("Emitter ID")), LOCTEXT("EmitterIDDesc", "ID of the emitter we'd like to spawn into. This can be obtained from Engine.Emitter.ID."));
-			Sig.AddInput(FNiagaraVariable(StaticEnum<ENDIDataChannelSpawnMode>(), TEXT("Mode")), LOCTEXT("SpawnCondModeInputDesc", "A mode switch that controls how this funciton will behave."));
+			Sig.AddInput(FNiagaraVariable(StaticEnum<ENDIDataChannelSpawnMode>(), TEXT("Mode")), LOCTEXT("SpawnCondModeInputDesc", "A mode switch that controls how this funciton will behave and interact with other calls to spawn functions."));
 			Sig.AddInput(FNiagaraVariable(StaticEnum<ENiagaraConditionalOperator>(), TEXT("Operator")), LOCTEXT("SpawnCondOpInputDesc", "The comparison operator to use when comparing values in the data channel to conditional parameters."));
-			Sig.AddInput(FNiagaraVariable(FNiagaraTypeDefinition::GetIntDef(), TEXT("Min Spawn Count")), LOCTEXT("MinSpawnCountInputDesc", "Minimum number of particles to spawn for each element in the data channel."));
-			Sig.AddInput(FNiagaraVariable(FNiagaraTypeDefinition::GetIntDef(), TEXT("Max Spawn Count")), LOCTEXT("MaxSpawnCountInputDesc", "Maximum number of particles to spawn for each element in the data channel."));
+			Sig.AddInputWithDefault(FNiagaraVariable(FNiagaraTypeDefinition::GetIntDef(), TEXT("Min Spawn Count")), 1, LOCTEXT("MinSpawnCountInputDesc", "Minimum number of particles to spawn for each element in the data channel."));
+			Sig.AddInputWithDefault(FNiagaraVariable(FNiagaraTypeDefinition::GetIntDef(), TEXT("Max Spawn Count")), 1, LOCTEXT("MaxSpawnCountInputDesc", "Maximum number of particles to spawn for each element in the data channel."));
 			Sig.RequiredInputs = Sig.Inputs.Num();
 		}
 		return Sig;
 	}
+
+	const FNiagaraFunctionSignature& GetFunctionSig_SpawnDirect()
+	{
+		static FNiagaraFunctionSignature Sig;
+		if (!Sig.IsValid())
+		{
+			Sig.Name = TEXT("SpawnDirect");
+#if WITH_EDITORONLY_DATA
+			Sig.Description = LOCTEXT("SpawnDirectFunctionDescription", "Spawns particles into a given emitter for each entry in the bound NDC. Spawn count is determined directly from a value in the NDC. Additional per NDC random scale and a clamp is available.");
+			NIAGARA_ADD_FUNCTION_SOURCE_INFO(Sig)
+#endif
+			Sig.FunctionSpecifiers.Add(VarNameKey);
+			Sig.FunctionSpecifiers.Add(VarTypeKey);
+			Sig.bMemberFunction = true;
+			Sig.bExperimental = true;
+			Sig.bRequiresExecPin = true;
+			Sig.ModuleUsageBitmask = ENiagaraScriptUsageMask::Emitter | ENiagaraScriptUsageMask::System;
+			Sig.AddInput(FNiagaraVariable(FNiagaraTypeDefinition(UNiagaraDataInterfaceDataChannelRead::StaticClass()), TEXT("DataChannel interface")));
+			Sig.AddInputWithDefault(FNiagaraVariable(FNiagaraTypeDefinition::GetBoolDef(), TEXT("Enable")), FNiagaraBool(true), LOCTEXT("SpawnEnableInputDesc", "Enable or disable this function call. If false, this call with have no effetcs."));
+			Sig.AddInput(FNiagaraVariable(FNiagaraTypeDefinition(FNiagaraEmitterID::StaticStruct()), TEXT("Emitter ID")), LOCTEXT("EmitterIDDesc", "ID of the emitter we'd like to spawn into. This can be obtained from Engine.Emitter.ID."));
+			Sig.AddInput(FNiagaraVariable(StaticEnum<ENDIDataChannelSpawnMode>(), TEXT("Mode")), LOCTEXT("SpawnCondModeInputDesc", "A mode switch that controls how this funciton will behave and interact with other calls to spawn functions."));			
+			Sig.AddInputWithDefault(FNiagaraVariable(FNiagaraTypeDefinition::GetFloatDef(), TEXT("RandomScaleMin")), 1.0f, LOCTEXT("SpawnDirectRandomScaleMinInputDesc", "Minimum value for an additional random scale applied to each NDC spawn count."));
+			Sig.AddInputWithDefault(FNiagaraVariable(FNiagaraTypeDefinition::GetFloatDef(), TEXT("RandomScaleMax")), 1.0f, LOCTEXT("SpawnDirectRandomScaleMaxInputDesc", "Maximum value for an additional random scale applied to each NDC spawn count."));
+			Sig.AddInputWithDefault(FNiagaraVariable(FNiagaraTypeDefinition::GetIntDef(), TEXT("ClampMin")), 0, LOCTEXT("SpawnDirectClampMinInputDesc", "Minimum Spawn Count to use after random scale is applied."));
+			Sig.AddInputWithDefault(FNiagaraVariable(FNiagaraTypeDefinition::GetIntDef(), TEXT("ClampMax")), 1, LOCTEXT("SpawnDirectClampMaxInputDesc", "Minimum Spawn Count to use after random scale is applied. If < 0, No max clamp is applied."));
+		}
+		return Sig;
+	}
+
+	const FNiagaraFunctionSignature& GetFunctionSig_ScaleSpawnCount()
+	{
+		static FNiagaraFunctionSignature Sig;
+		if (!Sig.IsValid())
+		{
+			Sig.Name = TEXT("ScaleSpawnCount");
+#if WITH_EDITORONLY_DATA
+			Sig.Description = LOCTEXT("ScaleSpawnCountFunctionDescription", "Applies a scaling value for each NDC spawn based on some element in the NDC data. Optional additional random scale and clamp operations to the value read from each NDC entry.");
+			NIAGARA_ADD_FUNCTION_SOURCE_INFO(Sig)
+#endif
+			Sig.FunctionSpecifiers.Add(VarNameKey);
+			Sig.FunctionSpecifiers.Add(VarTypeKey);
+			Sig.bMemberFunction = true;
+			Sig.bExperimental = true;
+			Sig.bRequiresExecPin = true;
+			Sig.ModuleUsageBitmask = ENiagaraScriptUsageMask::Emitter | ENiagaraScriptUsageMask::System;
+			Sig.AddInput(FNiagaraVariable(FNiagaraTypeDefinition(UNiagaraDataInterfaceDataChannelRead::StaticClass()), TEXT("DataChannel interface")));
+			Sig.AddInputWithDefault(FNiagaraVariable(FNiagaraTypeDefinition::GetBoolDef(), TEXT("Enable")), FNiagaraBool(true), LOCTEXT("SpawnEnableInputDesc", "Enable or disable this function call. If false, this call with have no effetcs."));
+			Sig.AddInput(FNiagaraVariable(FNiagaraTypeDefinition(FNiagaraEmitterID::StaticStruct()), TEXT("Emitter ID")), LOCTEXT("EmitterIDDesc", "ID of the emitter we'd like to spawn into. This can be obtained from Engine.Emitter.ID."));
+			Sig.AddInput(FNiagaraVariable(StaticEnum<ENDIDataChannelSpawnScaleMode>(), TEXT("Mode")), LOCTEXT("SpawnScaleModeInputDesc", "Control whether to override or combine this scale with previously set scales when calling this function multiple times."));
+			Sig.AddInputWithDefault(FNiagaraVariable(FNiagaraTypeDefinition::GetFloatDef(), TEXT("RandomScaleMin")), 1.0f, LOCTEXT("ScaleSpawnCountRandomScaleMinInputDesc", "Minimum value for a random additional scale applied to each NDC spawn."));
+			Sig.AddInputWithDefault(FNiagaraVariable(FNiagaraTypeDefinition::GetFloatDef(), TEXT("RandomScaleMax")), 1.0f, LOCTEXT("ScaleSpawnCountRandomScaleMaxInputDesc", "Maximum value for a random additional scale applied to each NDC spawn."));
+			Sig.AddInputWithDefault(FNiagaraVariable(FNiagaraTypeDefinition::GetFloatDef(), TEXT("ClampMin")), 0.0f, LOCTEXT("ScaleSpawnCountClampMinInputDesc", "A minimum value for the scale value. If < 0 then no minimum is applied to the final scale."));
+			Sig.AddInputWithDefault(FNiagaraVariable(FNiagaraTypeDefinition::GetFloatDef(), TEXT("ClampMax")), 1.0f, LOCTEXT("ScaleSpawnCountClampMaxInputDesc", "A maximum value for the scale value. If < 0 then no maximum is applied to the final scale."));
+		}
+		return Sig;
+	}
 	// Function definitions END
 	//////////////////////////////////////////////////////////////////////////
 
@@ -614,96 +689,104 @@ void UNiagaraDataInterfaceDataChannelRead::PostStageTick(FNDICpuPostStageContext
 	for (auto& EmitterInstanceDataPair : InstanceData->EmitterInstanceData)
 	{
 		FNDIDataChannelRead_EmitterInstanceData& EmitterInstData = EmitterInstanceDataPair.Value;
+		FNiagaraEmitterInstance* TargetEmitter = EmitterInstanceDataPair.Key;
+		if(TargetEmitter == nullptr)
+		{
+			EmitterInstData.Reset();
+			continue;
+		}
+
+		TArray<uint32> PerNDCSpawnCounts;
+		PerNDCSpawnCounts.Reserve(EmitterInstData.NDCSpawnCounts.Num());
+		for(FNDIDataChannelRead_EmitterSpawnInfo& Info : EmitterInstData.NDCSpawnCounts)
+		{
+			PerNDCSpawnCounts.Add(Info.Get());
+		}
 
-		TArray<int32> PerNDCSpawnCounts = MoveTemp(EmitterInstData.NDCSpawnCounts);
-	
 		EmitterInstData.Reset();
 
-		if(FNiagaraEmitterInstance* TargetEmitter = EmitterInstanceDataPair.Key)
+		//-TODO:Stateless:
+		if (FNiagaraEmitterInstanceImpl* StatefulEmitter = TargetEmitter->AsStateful())
 		{
-			//-TODO:Stateless:
-			if (FNiagaraEmitterInstanceImpl* StatefulEmitter = TargetEmitter->AsStateful())
+			if (bOverrideSpawnGroupToDataChannelIndex)
 			{
-				if (bOverrideSpawnGroupToDataChannelIndex)
+				//If we're overriding the spawn group then we must submit one SpawnInfo per NDC entry.
+				FNiagaraSpawnInfo NewSpawnInfo(0, 0.0f, 0.0f, 0);
+				for (int32 i = 0; i < PerNDCSpawnCounts.Num(); ++i)
 				{
-					//If we're overriding the spawn group then we must submit one SpawnInfo per NDC entry.
-					FNiagaraSpawnInfo NewSpawnInfo(0,0.0f,0.0f,0);
-					for (int32 i = 0; i < PerNDCSpawnCounts.Num(); ++i)
+					uint32 SpawnCount = PerNDCSpawnCounts[i];
+					if (SpawnCount > 0)
 					{
-						int32 SpawnCount = PerNDCSpawnCounts[i];
-						if (SpawnCount > 0)
-						{
-							NewSpawnInfo.Count = SpawnCount;
-							NewSpawnInfo.SpawnGroup = i;
-							StatefulEmitter->GetSpawnInfo().Emplace(NewSpawnInfo);
-						}
+						NewSpawnInfo.Count = SpawnCount;
+						NewSpawnInfo.SpawnGroup = i;
+						StatefulEmitter->GetSpawnInfo().Emplace(NewSpawnInfo);
 					}
 				}
-				else 
+			}
+			else
+			{
+				//No need for indirection table but we're not overriding the spawn group either so still push a single combined spawn info.
+				FNiagaraSpawnInfo NewSpawnInfo(0, 0.0f, 0.0f, 0);
+				for (int32 i = 0; i < PerNDCSpawnCounts.Num(); ++i)
 				{
-					//No need for indirection table but we're not overriding the spawn group either so still push a single combined spawn info.
-					FNiagaraSpawnInfo NewSpawnInfo(0, 0.0f, 0.0f, 0);
-					for (int32 i = 0; i < PerNDCSpawnCounts.Num(); ++i)
-					{
-						NewSpawnInfo.Count += PerNDCSpawnCounts[i];
-					}
-					StatefulEmitter->GetSpawnInfo().Emplace(NewSpawnInfo);
+					NewSpawnInfo.Count += PerNDCSpawnCounts[i];
 				}
+				StatefulEmitter->GetSpawnInfo().Emplace(NewSpawnInfo);
+			}
 
-				if(CompiledData.NeedSpawnDataTable())
-				{
-					//Build an indirection table that allows us to map from ExecIndex back to the NDCIndex that generated it.
-					//The indirection table is arranged in power of two buckets.
-					//An NDC that spawns say 37 particles would add an entry to the 32, 4 and 1 buckets.
-					//This allows us to spawn any number of particles from each NDC and only have a max of 16 indirection table entries.
-					//Vs the naive per particle approach of 1 entry per particle.
-					//Buckets are processed in descending size order.
+			if (CompiledData.NeedSpawnDataTable())
+			{
+				//Build an indirection table that allows us to map from ExecIndex back to the NDCIndex that generated it.
+				//The indirection table is arranged in power of two buckets.
+				//An NDC that spawns say 37 particles would add an entry to the 32, 4 and 1 buckets.
+				//This allows us to spawn any number of particles from each NDC and only have a max of 16 indirection table entries.
+				//Vs the naive per particle approach of 1 entry per particle.
+				//Buckets are processed in descending size order.
 
-					//TODO: It should be possible to write this from the GPU too as long as we allocate fixed size buckets.
-					int32* SpawnDataBuckets = EmitterInstData.NDCSpawnData.NDCSpawnDataBuckets;					
+				//TODO: It should be possible to write this from the GPU too as long as we allocate fixed size buckets.
+				int32* SpawnDataBuckets = EmitterInstData.NDCSpawnData.NDCSpawnDataBuckets;
 
-					TArray<int32>& NDCSpawnData = EmitterInstData.NDCSpawnData.NDCSpawnData;
+				TArray<int32>& NDCSpawnData = EmitterInstData.NDCSpawnData.NDCSpawnData;
 
-					//Start of the buffer is the per NDC spawn counts.
-					uint32 TotalNDCSpawnDataSize = PerNDCSpawnCounts.Num();
+				//Start of the buffer is the per NDC spawn counts.
+				uint32 TotalNDCSpawnDataSize = PerNDCSpawnCounts.Num();
 
-					for (int32 i = 0; i < PerNDCSpawnCounts.Num(); ++i)
-					{						
-						uint32 Count = PerNDCSpawnCounts[i];
+				for (int32 i = 0; i < PerNDCSpawnCounts.Num(); ++i)
+				{
+					uint32 Count = PerNDCSpawnCounts[i];
 
-						//First section is the per NDC counts.
-						NDCSpawnData.Add(Count);
+					//First section is the per NDC counts.
+					NDCSpawnData.Add(Count);
 
-						for (uint32 Bucket = 0; Bucket < 16; ++Bucket)
-						{
-							uint32 BucketSize = (1<<15) >> Bucket;
-							uint32 Mask = (0xFFFF >> (Bucket + 1));
-							uint32 CountMasked = Count & ~Mask;
-							Count &= Mask;
-							uint32 NumBucketEntries = CountMasked / BucketSize;
-							SpawnDataBuckets[Bucket] += NumBucketEntries;
-							TotalNDCSpawnDataSize += NumBucketEntries;
-						}
+					for (uint32 Bucket = 0; Bucket < 16; ++Bucket)
+					{
+						uint32 BucketSize = (1 << 15) >> Bucket;
+						uint32 Mask = (0xFFFF >> (Bucket + 1));
+						uint32 CountMasked = Count & ~Mask;
+						Count &= Mask;
+						uint32 NumBucketEntries = CountMasked / BucketSize;
+						SpawnDataBuckets[Bucket] += NumBucketEntries;
+						TotalNDCSpawnDataSize += NumBucketEntries;
 					}
+				}
 
-					//Second part is the counts decomposed into power of two buckets that allows us to map ExecIndex at runtime to an NDCIndex entry in this table.
-					for (int32 Bucket = 0; Bucket < 16; ++Bucket)
+				//Second part is the counts decomposed into power of two buckets that allows us to map ExecIndex at runtime to an NDCIndex entry in this table.
+				for (int32 Bucket = 0; Bucket < 16; ++Bucket)
+				{
+					uint32 BucketSize = (1 << 15) >> Bucket;
+					uint32 StartSize = NDCSpawnData.Num();
+					for (int32 i = 0; i < PerNDCSpawnCounts.Num(); ++i)
 					{
-						uint32 BucketSize = (1 << 15) >> Bucket;
-						uint32 StartSize = NDCSpawnData.Num();
-						for (int32 i = 0; i < PerNDCSpawnCounts.Num(); ++i)
+						uint32& Count = PerNDCSpawnCounts[i];
+						while (Count >= BucketSize)
 						{
-							int32& Count = PerNDCSpawnCounts[i];
-							while ((uint32)Count >= BucketSize)
-							{
-								Count -= BucketSize;
-								NDCSpawnData.Add(i);
-							}
+							Count -= BucketSize;
+							NDCSpawnData.Add(i);
 						}
-
-						uint32 EndSize = NDCSpawnData.Num();
-						check(EndSize - StartSize == SpawnDataBuckets[Bucket]);
 					}
+
+					uint32 EndSize = NDCSpawnData.Num();
+					check(EndSize - StartSize == SpawnDataBuckets[Bucket]);
 				}
 			}
 		}
@@ -821,7 +904,10 @@ void UNiagaraDataInterfaceDataChannelRead::ProvidePerInstanceDataForRenderThread
 				TargetNDCSpawnData.Add(EmitterInstData.NDCSpawnData.NDCSpawnDataBuckets[i]);
 			}
 			//Next the per NDC Spawn Counts
-			TargetNDCSpawnData.Append(EmitterInstData.NDCSpawnCounts);
+			for(int32 i=0; i < EmitterInstData.NDCSpawnCounts.Num(); ++i)
+			{
+				TargetNDCSpawnData.Add(EmitterInstData.NDCSpawnCounts[i].Get());
+			}
 			//Finally the exec index to NDC index mapping table
 			TargetNDCSpawnData.Append(EmitterInstData.NDCSpawnData.NDCSpawnData);
 		}
@@ -1101,6 +1187,8 @@ void UNiagaraDataInterfaceDataChannelRead::GetFunctionsInternal(TArray<FNiagaraF
 	OutFunctions.Add(NDIDataChannelReadLocal::GetFunctionSig_Read());
 	OutFunctions.Add(NDIDataChannelReadLocal::GetFunctionSig_Consume());
 	OutFunctions.Add(NDIDataChannelReadLocal::GetFunctionSig_SpawnConditional());
+	OutFunctions.Add(NDIDataChannelReadLocal::GetFunctionSig_SpawnDirect());
+	OutFunctions.Add(NDIDataChannelReadLocal::GetFunctionSig_ScaleSpawnCount());
 }
 #endif
 
@@ -1114,6 +1202,68 @@ void UNiagaraDataInterfaceDataChannelRead::GetVMExternalFunction(const FVMExtern
 	{
 		OutFunc = FVMExternalFunction::CreateLambda([this](FVectorVMExternalFunctionContext& Context) { this->GetNDCSpawnData(Context); });
 	}
+	else if (BindingInfo.Name == NDIDataChannelReadLocal::GetFunctionSig_SpawnDirect().Name)
+	{
+		//Find the appropriate variable from the function binding and determine it's type to bind to the correct impl for ScaleSpawnCount and provide the variable index.
+		FNiagaraVariableBase FuncSpecVariable = DecodeVariableFromSpecifiers(BindingInfo);
+
+		//TODO: We are currently storing these as FNames and encoding/decoding the variable.
+		//This is slow, clunky, brittle and generally bad. Ideally imo we could change function specifiers to be FInstancedStructs and allow and DI to provide and handle their own types.
+		const FNiagaraDataChannelVariable* NDCVar = Channel->Get()->GetVariables().FindByPredicate(
+			[&](const FNiagaraDataChannelVariable& Var)
+			{
+				return Var.GetName() == FuncSpecVariable.GetName() && Var.GetType() == FuncSpecVariable.GetType();
+			});
+
+		if (NDCVar)
+		{
+			FName NDCVarName = NDCVar->GetName();
+			FNiagaraTypeDefinition NDCVarType = NDCVar->GetType();
+			if (NDCVarType == FNiagaraTypeDefinition::GetIntDef()) OutFunc = FVMExternalFunction::CreateLambda([this, NDCVarName](FVectorVMExternalFunctionContext& Context) { this->SpawnDirect<int32>(Context, NDCVarName); });
+			else
+			{
+				UE_LOG(LogTemp, Display, TEXT("Failed to bind Data Interface function as this is not a valid variable type for SpawnDirect.\nDI: %s.\nReceived Name: %s\nNDC: %s"), *GetPathNameSafe(this), *BindingInfo.Name.ToString(), *GetPathNameSafe(Channel));
+			}
+		}
+		else
+		{
+			UE_LOG(LogTemp, Display, TEXT("Failed to bind Data Interface function due to missing NDC variable for %s.\nReceived Name: %s\nNDC: %s"), *GetPathNameSafe(this), *BindingInfo.Name.ToString(), *GetPathNameSafe(Channel));
+		}
+	}
+	else if (BindingInfo.Name == NDIDataChannelReadLocal::GetFunctionSig_ScaleSpawnCount().Name)
+	{
+		//Find the appropriate variable from the function binding and determine it's type to bind to the correct impl for ScaleSpawnCount and provide the variable index.
+			//Can we do this decode on the PostCompile? Also would be good to validate/warn on missing or invalid types.
+		FNiagaraVariableBase FuncSpecVariable = DecodeVariableFromSpecifiers(BindingInfo);
+
+		//TODO: We are currently storing these as FNames and encoding/decoding the variable.
+		//This is slow, clunky, brittle and generally bad. Ideally imo we could change function specifiers to be FInstancedStructs and allow and DI to provide and handle their own types.
+		const FNiagaraDataChannelVariable* NDCVar = Channel->Get()->GetVariables().FindByPredicate(
+			[&](const FNiagaraDataChannelVariable& Var)
+			{
+				return Var.GetName() == FuncSpecVariable.GetName() && Var.GetType() == FuncSpecVariable.GetType();
+			});
+
+		if (NDCVar)
+		{
+			FName NDCVarName = NDCVar->GetName();
+			FNiagaraTypeDefinition NDCVarType = NDCVar->GetType();
+			if (NDCVarType == FNiagaraTypeDefinition::GetIntDef()) OutFunc = FVMExternalFunction::CreateLambda([this, NDCVarName](FVectorVMExternalFunctionContext& Context) { this->ScaleSpawnCount<int32>(Context, NDCVarName); });
+			else if (NDCVarType == FNiagaraTypeHelper::GetDoubleDef()) OutFunc = FVMExternalFunction::CreateLambda([this, NDCVarName](FVectorVMExternalFunctionContext& Context) { this->ScaleSpawnCount<double>(Context, NDCVarName); });
+			else if (NDCVarType == FNiagaraTypeHelper::GetVector2DDef()) OutFunc = FVMExternalFunction::CreateLambda([this, NDCVarName](FVectorVMExternalFunctionContext& Context) { this->ScaleSpawnCount<FVector2D>(Context, NDCVarName); });
+			else if (NDCVarType == FNiagaraTypeHelper::GetVectorDef()) OutFunc = FVMExternalFunction::CreateLambda([this, NDCVarName](FVectorVMExternalFunctionContext& Context) { this->ScaleSpawnCount<FVector>(Context, NDCVarName); });
+			else if (NDCVarType == FNiagaraTypeHelper::GetVector4Def()) OutFunc = FVMExternalFunction::CreateLambda([this, NDCVarName](FVectorVMExternalFunctionContext& Context) { this->ScaleSpawnCount<FVector4>(Context, NDCVarName); });
+			else if (NDCVarType == FNiagaraTypeDefinition::GetPositionDef()) OutFunc = FVMExternalFunction::CreateLambda([this, NDCVarName](FVectorVMExternalFunctionContext& Context) { this->ScaleSpawnCount<FNiagaraPosition>(Context, NDCVarName); });
+			else
+			{
+				UE_LOG(LogTemp, Display, TEXT("Failed to bind Data Interface function as this is not a valid variable type for ScaleSpawnCount.\nDI: %s.\nReceived Name: %s\nNDC: %s"), *GetPathNameSafe(this), *BindingInfo.Name.ToString(), *GetPathNameSafe(Channel));
+			}
+		}
+		else
+		{
+			UE_LOG(LogTemp, Display, TEXT("Failed to bind Data Interface function due to missing NDC variable for %s.\nReceived Name: %s\nNDC: %s"), *GetPathNameSafe(this), *BindingInfo.Name.ToString(), *GetPathNameSafe(Channel));
+		}
+	}
 	else
 	{
 		int32 FuncIndex = CompiledData.FindFunctionInfoIndex(BindingInfo.Name, BindingInfo.VariadicInputs, BindingInfo.VariadicOutputs);
@@ -1536,13 +1686,13 @@ void UNiagaraDataInterfaceDataChannelRead::SpawnConditional(FVectorVMExternalFun
 			return;
 		}
 
-		int32 SpawnMin = InSpawnMin.GetAndAdvance();
-		int32 SpawnMax = InSpawnMax.GetAndAdvance();
+		int32 SpawnMin = FMath::Max(0, InSpawnMin.GetAndAdvance());
+		int32 SpawnMax = FMath::Max(0, InSpawnMax.GetAndAdvance());
 
 		//Each data channel element has an additional spawn entry which accumulates across all spawning calls and can be nulled independently by a suppression call.
 		FNDIDataChannelRead_EmitterInstanceData& EmitterInstData = InstData->EmitterInstanceData.FindOrAdd(EmitterInst);
-		TArray<int32>& EmitterConditionalSpawns = EmitterInstData.NDCSpawnCounts;
-		EmitterConditionalSpawns.SetNumZeroed(NumDataChannelInstances);
+		TArray<FNDIDataChannelRead_EmitterSpawnInfo>& EmitterConditionalSpawns = EmitterInstData.NDCSpawnCounts;
+		EmitterConditionalSpawns.SetNum(NumDataChannelInstances);
 
 		for (int32 DataChannelIdx = 0; DataChannelIdx < NumDataChannelInstances; ++DataChannelIdx)
 		{
@@ -1570,17 +1720,208 @@ void UNiagaraDataInterfaceDataChannelRead::SpawnConditional(FVectorVMExternalFun
 				int32 Count = RandHelper.RandRange(DataChannelIdx, SpawnMin, SpawnMax);
 				if (Mode == ENDIDataChannelSpawnMode::Accumulate)
 				{
-					EmitterConditionalSpawns[DataChannelIdx] += Count;
+					EmitterConditionalSpawns[DataChannelIdx].Append(Count);
 				}
 				else if (Mode == ENDIDataChannelSpawnMode::Override)
 				{
-					EmitterConditionalSpawns[DataChannelIdx] = Count;
+					EmitterConditionalSpawns[DataChannelIdx].SetCount(Count);
 				}
 			}
 		}
 	}
 }
 
+template<typename T> 
+float NDCValueSize(const T& Value){ return 1.0f; }
+
+template<> float NDCValueSize<int32>(const int32& Value){ return FMath::Abs(Value); }
+template<> float NDCValueSize<float>(const float& Value) { return FMath::Abs(Value); }
+template<> float NDCValueSize<FVector2f>(const FVector2f& Value) { return Value.Size(); }
+template<> float NDCValueSize<FVector3f>(const FVector3f& Value) { return Value.Size(); }
+template<> float NDCValueSize<FVector4f>(const FVector4f& Value) { return Value.Size(); }
+
+template<typename T>
+T NDCValueDefault() { return 1.0f; }
+template<> int32 NDCValueDefault<int32>() { return 1; }
+template<> float NDCValueDefault<float>() { return 1.0f; }
+template<> FVector2f NDCValueDefault<FVector2f>() { return FVector2f(1.0f); }
+template<> FVector3f NDCValueDefault<FVector3f>() { return FVector3f(1.0f); }
+template<> FVector4f NDCValueDefault<FVector4f>() { return FVector4f(1.0f); }
+
+template <typename T> FNiagaraTypeDefinition GetTypeDef() { return FNiagaraTypeDefinition(); }
+template <typename T> FNiagaraTypeDefinition GetSimTypeDef() { return FNiagaraTypeDefinition(); }
+
+template<> FNiagaraTypeDefinition GetTypeDef<int32>() { return FNiagaraTypeDefinition::GetIntDef(); }
+template<> FNiagaraTypeDefinition GetTypeDef<double>() { return FNiagaraTypeHelper::GetDoubleDef(); }
+template<> FNiagaraTypeDefinition GetTypeDef<FVector2D>() { return FNiagaraTypeHelper::GetVector2DDef(); }
+template<> FNiagaraTypeDefinition GetTypeDef<FVector>() { return FNiagaraTypeHelper::GetVectorDef(); }
+template<> FNiagaraTypeDefinition GetTypeDef<FVector4>() { return FNiagaraTypeHelper::GetVector4Def(); }
+template<> FNiagaraTypeDefinition GetTypeDef<FNiagaraPosition>() { return FNiagaraTypeDefinition::GetPositionDef(); }
+
+template<> FNiagaraTypeDefinition GetSimTypeDef<int32>() { return FNiagaraTypeDefinition::GetIntDef(); }
+template<> FNiagaraTypeDefinition GetSimTypeDef<double>() { return FNiagaraTypeDefinition::GetFloatDef(); }
+template<> FNiagaraTypeDefinition GetSimTypeDef<FVector2D>() { return FNiagaraTypeDefinition::GetVec2Def(); }
+template<> FNiagaraTypeDefinition GetSimTypeDef<FVector>() { return FNiagaraTypeDefinition::GetVec3Def(); }
+template<> FNiagaraTypeDefinition GetSimTypeDef<FVector4>() { return FNiagaraTypeDefinition::GetVec4Def(); }
+template<> FNiagaraTypeDefinition GetSimTypeDef<FNiagaraPosition>() { return FNiagaraTypeDefinition::GetVec3Def(); }
+
+template<typename TLWCType>
+struct TLWCToSimType {};
+
+template<> struct TLWCToSimType<int32> { typedef int32 TSimType; };
+template<> struct TLWCToSimType<double> { typedef float TSimType;  };
+template<> struct TLWCToSimType<FVector2D> { typedef FVector2f TSimType; };
+template<> struct TLWCToSimType<FVector> { typedef FVector3f TSimType; };
+template<> struct TLWCToSimType<FVector4> { typedef FVector4f TSimType; };
+template<> struct TLWCToSimType<FNiagaraPosition> { typedef FVector3f TSimType; };
+
+template<typename T>
+void UNiagaraDataInterfaceDataChannelRead::SpawnDirect(FVectorVMExternalFunctionContext& Context, FName NDCVarName)
+{
+	//This should only be called from emitter scripts and since it has per instance data then we process them individually.
+	check(Context.GetNumInstances() == 1);
+
+	using TSimType = TLWCToSimType<T>::TSimType;
+	VectorVM::FUserPtrHandler<FNDIDataChannelReadInstanceData> InstData(Context);
+
+	FNDIInputParam<FNiagaraBool> InEnabled(Context);
+	FNDIInputParam<FNiagaraEmitterID> InEmitterID(Context);
+
+	FNDIInputParam<int32> InMode(Context);
+
+	FNDIInputParam<float> InRandMinScale(Context);
+	FNDIInputParam<float> InRandMaxScale(Context);
+	FNDIInputParam<int32> InClampMin(Context);
+	FNDIInputParam<int32> InClampMax(Context);
+
+	FNiagaraSystemInstance* SystemInstance = InstData->Owner;
+	check(SystemInstance);
+
+	const FNiagaraEmitterID EmitterID = InEmitterID.GetAndAdvance();
+	const int32 NumEmitters = SystemInstance->GetEmitters().Num();
+	FNiagaraEmitterInstance* EmitterInst = SystemInstance->GetEmitterByID(EmitterID);
+
+	FNiagaraDataChannelData* DataChannelData = InstData->DataChannelData.Get();
+
+	const bool bReadPrevFrame = bReadCurrentFrame == false || GNDCReadForcePrevFrame;
+	FNiagaraDataBuffer* Data = DataChannelData ? DataChannelData->GetCPUData(bReadPrevFrame) : nullptr;
+
+	const bool bApply = INiagaraModule::DataChannelsEnabled() && Data && Data->GetNumInstances() > 0 && EmitterInst && EmitterInst->IsActive() && InEnabled.GetAndAdvance();
+	if (bApply)
+	{
+		FNDIRandomHelperFromStream RandHelper(Context);
+
+		int32 NumDataChannelInstances = Data->GetNumInstances();
+
+		ENDIDataChannelSpawnMode Mode = (ENDIDataChannelSpawnMode)InMode.GetAndAdvance();
+
+		float RandMinScale = InRandMinScale.GetAndAdvance();
+		float RandMaxScale = InRandMaxScale.GetAndAdvance();
+		int32 ClampMin = FMath::Max(0, InClampMin.GetAndAdvance());
+		int32 ClampMax = InClampMax.GetAndAdvance();
+		ClampMax = ClampMax < 0 ? INT_MAX : ClampMax;
+
+		//Each data channel element has an additional spawn entry which accumulates across all spawning calls and can be nulled independently by a suppression call.
+		FNDIDataChannelRead_EmitterInstanceData& EmitterInstData = InstData->EmitterInstanceData.FindOrAdd(EmitterInst);
+		TArray<FNDIDataChannelRead_EmitterSpawnInfo>& EmitterConditionalSpawns = EmitterInstData.NDCSpawnCounts;
+		EmitterConditionalSpawns.SetNum(NumDataChannelInstances);
+
+		const auto ValueData = FNiagaraDataSetAccessor<TSimType>::CreateReader(Data, NDCVarName);
+
+		for (int32 DataChannelIndex = 0; DataChannelIndex < NumDataChannelInstances; DataChannelIndex++)
+		{
+			TSimType NDCValue = ValueData.GetSafe(DataChannelIndex, NDCValueDefault<TSimType>());
+			double VarSize = NDCValueSize(NDCValue);
+			float Scale = RandHelper.RandRange(DataChannelIndex, RandMinScale, RandMaxScale);
+			int32 ScaledCount = VarSize * Scale;
+			int32 FinalCount = FMath::Clamp(ScaledCount, ClampMin, ClampMax);
+
+			if (Mode == ENDIDataChannelSpawnMode::Accumulate)
+			{
+				EmitterConditionalSpawns[DataChannelIndex].Append(FinalCount);
+			}
+			else if (Mode == ENDIDataChannelSpawnMode::Override)
+			{
+				EmitterConditionalSpawns[DataChannelIndex].SetCount(FinalCount);
+			}
+		}
+	}
+}
+
+template<typename T>
+void UNiagaraDataInterfaceDataChannelRead::ScaleSpawnCount(FVectorVMExternalFunctionContext& Context, FName NDCVarName)
+{
+	//This should only be called from emitter scripts and since it has per instance data then we process them individually.
+	check(Context.GetNumInstances() == 1);
+
+	using TSimType = TLWCToSimType<T>::TSimType;
+	VectorVM::FUserPtrHandler<FNDIDataChannelReadInstanceData> InstData(Context);
+
+	FNDIInputParam<FNiagaraBool> InEnabled(Context);
+	FNDIInputParam<FNiagaraEmitterID> InEmitterID(Context);
+	FNDIInputParam<int32> InMode(Context);
+
+	FNDIInputParam<float> InRandMinScale(Context);
+	FNDIInputParam<float> InRandMaxScale(Context);
+	FNDIInputParam<float> InClampMin(Context);
+	FNDIInputParam<float> InClampMax(Context);
+	
+	FNiagaraSystemInstance* SystemInstance = InstData->Owner;
+	check(SystemInstance);
+
+	const FNiagaraEmitterID EmitterID = InEmitterID.GetAndAdvance();
+	const int32 NumEmitters = SystemInstance->GetEmitters().Num();
+	FNiagaraEmitterInstance* EmitterInst = SystemInstance->GetEmitterByID(EmitterID);
+
+	FNiagaraDataChannelData* DataChannelData = InstData->DataChannelData.Get();
+
+	const bool bReadPrevFrame = bReadCurrentFrame == false || GNDCReadForcePrevFrame;
+	FNiagaraDataBuffer* Data = DataChannelData ? DataChannelData->GetCPUData(bReadPrevFrame) : nullptr;
+
+	const bool bApplyScale = INiagaraModule::DataChannelsEnabled() && Data && Data->GetNumInstances() > 0 && EmitterInst && EmitterInst->IsActive() && InEnabled.GetAndAdvance();
+	if (bApplyScale)
+	{
+		int32 Mode = InMode.GetAndAdvance();
+		bool bOverrideScale = Mode == (int32)ENDIDataChannelSpawnScaleMode::Override;
+
+		FNDIRandomHelperFromStream RandHelper(Context);
+
+		int32 NumDataChannelInstances = Data->GetNumInstances();
+
+		float RandMinScale = InRandMinScale.GetAndAdvance();
+		float RandMaxScale = InRandMaxScale.GetAndAdvance();
+		float ClampMin = FMath::Max(0.0f, InClampMin.GetAndAdvance());
+		float ClampMax = InClampMax.GetAndAdvance();
+		ClampMax = ClampMax < 0.0f ? FLT_MAX : ClampMax;
+
+		//Each data channel element has an additional spawn entry which accumulates across all spawning calls and can be nulled independently by a suppression call.
+		FNDIDataChannelRead_EmitterInstanceData& EmitterInstData = InstData->EmitterInstanceData.FindOrAdd(EmitterInst);
+		TArray<FNDIDataChannelRead_EmitterSpawnInfo>& EmitterConditionalSpawns = EmitterInstData.NDCSpawnCounts;
+		EmitterConditionalSpawns.SetNumZeroed(NumDataChannelInstances);
+
+		const auto ValueData = FNiagaraDataSetAccessor<TSimType>::CreateReader(Data, NDCVarName);
+
+		for(int32 DataChannelIndex = 0; DataChannelIndex < NumDataChannelInstances; DataChannelIndex++)
+		{
+			TSimType NDCValue = ValueData.GetSafe(DataChannelIndex, NDCValueDefault<TSimType>());
+			float VarSize = NDCValueSize(NDCValue);
+			float Scale = RandHelper.RandRange(DataChannelIndex, RandMinScale, RandMaxScale);			
+			float FinalScale = VarSize * Scale;
+			FinalScale = FMath::Clamp(FinalScale, ClampMin, ClampMax);
+
+			//TODO: Either change this to a float/double or add a separate scale value applied at the end so that multiple scales will combine correctly.
+			if(bOverrideScale)
+			{
+				EmitterConditionalSpawns[DataChannelIndex].SetScale(FinalScale);
+			}
+			else
+			{
+				EmitterConditionalSpawns[DataChannelIndex].ApplyScale(FinalScale);
+			}
+		}
+	}
+}
+
 #if WITH_EDITORONLY_DATA
 bool UNiagaraDataInterfaceDataChannelRead::AppendCompileHash(FNiagaraCompileHashVisitor* InVisitor) const
 {
@@ -1981,7 +2322,7 @@ void UNiagaraDataInterfaceDataChannelRead::SetShaderParameters(const FNiagaraDat
 		InstParameters->HalfStride = 0;
 		
 		FRDGBufferRef DummyBuffer = GSystemTextures.GetDefaultBuffer(Context.GetGraphBuilder(), 4, 0u);
-		InstParameters->NDCSpawnDataBuffer = Context.GetGraphBuilder().CreateSRV(DummyBuffer);		
+		InstParameters->NDCSpawnDataBuffer = Context.GetGraphBuilder().CreateSRV(DummyBuffer, PF_R32_SINT);		
 	}
 }
 
@@ -1989,7 +2330,7 @@ void FNiagaraDataInterfaceProxy_DataChannelRead::PreStage(const FNDIGpuComputePr
 {
 	FNiagaraDataInterfaceProxy_DataChannelRead::FInstanceData* InstanceData = SystemInstancesToProxyData_RT.Find(Context.GetSystemInstanceID());
 
-	if(InstanceData->NDCSpawnDataBuffer == nullptr)
+	if(InstanceData && InstanceData->NDCSpawnDataBuffer == nullptr)
 	{
 		InstanceData->NDCSpawnDataBuffer = CreateUploadBuffer<int32>(
 			Context.GetGraphBuilder(),
@@ -2002,8 +2343,10 @@ void FNiagaraDataInterfaceProxy_DataChannelRead::PostSimulate(const FNDIGpuCompu
 {
 	if (Context.IsFinalPostSimulate())
 	{
-		FNiagaraDataInterfaceProxy_DataChannelRead::FInstanceData* InstanceData = SystemInstancesToProxyData_RT.Find(Context.GetSystemInstanceID());
-		InstanceData->NDCSpawnDataBuffer = nullptr;
+		if(FNiagaraDataInterfaceProxy_DataChannelRead::FInstanceData* InstanceData = SystemInstancesToProxyData_RT.Find(Context.GetSystemInstanceID()))
+		{
+			InstanceData->NDCSpawnDataBuffer = nullptr;
+		}
 	}
 }
 
diff --git a/Engine/Plugins/FX/Niagara/Source/Niagara/Public/NiagaraCommon.h b/Engine/Plugins/FX/Niagara/Source/Niagara/Public/NiagaraCommon.h
--- a/Engine/Plugins/FX/Niagara/Source/Niagara/Public/NiagaraCommon.h
+++ b/Engine/Plugins/FX/Niagara/Source/Niagara/Public/NiagaraCommon.h
@@ -621,6 +621,18 @@ struct FNiagaraFunctionSignature
 		}
 	#endif
 	}
+	
+	template<typename T>
+	void AddInputWithDefault(FNiagaraVariable InputVar, const T& Default, FText Tooltip = FText())
+	{
+		Inputs.Add_GetRef(InputVar).SetData((uint8*)&Default);
+#if WITH_EDITORONLY_DATA
+		if (!Tooltip.IsEmpty())
+		{
+			InputDescriptions.Add(InputVar, Tooltip);
+		}
+#endif
+	}
 
 	void AddOutput(FNiagaraVariable OutputVar, const FText& Tooltip = FText())
 	{
diff --git a/Engine/Plugins/FX/Niagara/Source/NiagaraEditor/Private/NiagaraActions.cpp b/Engine/Plugins/FX/Niagara/Source/NiagaraEditor/Private/NiagaraActions.cpp
--- a/Engine/Plugins/FX/Niagara/Source/NiagaraEditor/Private/NiagaraActions.cpp
+++ b/Engine/Plugins/FX/Niagara/Source/NiagaraEditor/Private/NiagaraActions.cpp
@@ -20,6 +20,7 @@
 #include "ViewModels/NiagaraParameterPanelViewModel.h"
 #include "Misc/MessageDialog.h"
 #include "Widgets/SNiagaraParameterName.h"
+#include "Widgets/SNiagaraParameterMenu.h"
 #include "Widgets/SToolTip.h"
 
 #include "NiagaraDataChannel.h"
@@ -725,6 +726,23 @@ void INiagaraDataInterfaceNodeActionProvider::CollectAddPinActions(UClass* DICla
 		}
 	}
 }
+
+TSharedPtr<SWidget> INiagaraDataInterfaceNodeActionProvider::GetCustomFunctionSpecifierWidget(UClass* DIClass, UNiagaraNodeFunctionCall* FunctionCallNode)
+{
+	if (DIClass && FunctionCallNode)
+	{
+		while (DIClass && DIClass != UObject::StaticClass())
+		{
+			if (TUniquePtr<INiagaraDataInterfaceNodeActionProvider>* Provider = RegisteredActionProviders.Find(DIClass->GetFName()))
+			{
+				return (*Provider)->GetCustomFunctionSpecifierWidgetImpl(FunctionCallNode);
+			}
+			DIClass = DIClass->GetSuperClass();
+		}
+	}
+	return nullptr;
+}
+
 ////////////////////////
 
 namespace NiagaraActionsLocal
@@ -740,7 +758,7 @@ void FNiagaraDataInterfaceNodeActionProvider_DataChannelWrite::GetNodeContextMen
 	using namespace NiagaraActionsLocal;
 
 	//For all functions except "Num", add a context menu to initialized to a specific data channel.
-	if (Signature.Name == TEXT("Num"))
+	if (Signature.Name == TEXT("Num") || Signature.Name == TEXT("SpawnConditional"))
 	{
 		return;
 	}
@@ -799,7 +817,7 @@ void FNiagaraDataInterfaceNodeActionProvider_DataChannelRead::GetNodeContextMenu
 	using namespace NiagaraActionsLocal;
 
 	//For all functions except "Num", add a context menu to initialized to a specific data channel.
-	if (Signature.Name == TEXT("Num"))
+	if (Signature.Name != TEXT("Read") && Signature.Name != TEXT("Consume"))
 	{
 		return;
 	}
@@ -829,8 +847,8 @@ INiagaraDataInterfaceNodeActionProvider::FInlineMenuDisplayOptions FNiagaraDataI
 	{
 		return {};
 	}
-	
-	if(FunctionCall->Signature.Name == TEXT("Num"))
+
+	if (FunctionCall->Signature.Name != TEXT("Read") && FunctionCall->Signature.Name != TEXT("Consume"))
 	{
 		return {};
 	}
@@ -945,4 +963,26 @@ void FNiagaraDataInterfaceNodeActionProvider_DataChannelRead::AddDataChannelInit
 	UNiagaraDataChannel::ForEachDataChannel(InitForDataChannelSection);
 }
 
+TSharedPtr<SWidget> FNiagaraDataInterfaceNodeActionProvider_DataChannelRead::GetCustomFunctionSpecifierWidgetImpl(UNiagaraNodeFunctionCall* FunctionCallNode) const
+{
+	TArray<FNiagaraTypeDefinition> AllowedTypes;
+	if(FunctionCallNode->Signature.Name == TEXT("SpawnDirect"))
+	{
+		AllowedTypes.Add(FNiagaraTypeDefinition::GetIntDef());
+	}
+	else if(FunctionCallNode->Signature.Name == TEXT("ScaleSpawnCount"))
+	{
+		AllowedTypes.Add(FNiagaraTypeDefinition::GetIntDef());
+		AllowedTypes.Add(FNiagaraTypeHelper::GetDoubleDef());
+		AllowedTypes.Add(FNiagaraTypeHelper::GetVector2DDef());
+		AllowedTypes.Add(FNiagaraTypeHelper::GetVectorDef());
+		AllowedTypes.Add(FNiagaraTypeHelper::GetVector4Def());
+		AllowedTypes.Add(FNiagaraTypeDefinition::GetPositionDef());
+	}
+
+	return SNew(SNiagaraFunctionSpecifierNDCVariablesSelector)
+		.WeakNodeToModify(FunctionCallNode)
+		.AllowedTypes(MoveTemp(AllowedTypes));
+}
+
 #undef LOCTEXT_NAMESPACE
diff --git a/Engine/Plugins/FX/Niagara/Source/NiagaraEditor/Private/NiagaraNodeFunctionCall.cpp b/Engine/Plugins/FX/Niagara/Source/NiagaraEditor/Private/NiagaraNodeFunctionCall.cpp
--- a/Engine/Plugins/FX/Niagara/Source/NiagaraEditor/Private/NiagaraNodeFunctionCall.cpp
+++ b/Engine/Plugins/FX/Niagara/Source/NiagaraEditor/Private/NiagaraNodeFunctionCall.cpp
@@ -268,6 +268,28 @@ void UNiagaraNodeFunctionCall::RemoveAllDynamicPins()
 	}
 }
 
+UClass* UNiagaraNodeFunctionCall::GetDIClass()const
+{
+	if(Signature.IsValid())
+	{
+		if (Signature.Inputs.Num() > 0)
+		{
+			if (Signature.Inputs[0].GetType().IsDataInterface() && GetValidateDataInterfaces())
+			{
+				return Signature.Inputs[0].GetType().GetClass();
+			}
+		}
+	}
+	return nullptr;	
+}
+
+void UNiagaraNodeFunctionCall::SetFunctionSpecifier(FName Key, FName Value)
+{
+	Modify();
+	FunctionSpecifiers.FindOrAdd(Key) = Value;
+	MarkNodeRequiresSynchronization(__FUNCTION__, true);
+}
+
 UEdGraphPin* UNiagaraNodeFunctionCall::AddStaticSwitchInputPin(FNiagaraVariable Input)
 {
 	UNiagaraGraph* Graph = GetCalledGraph();
@@ -799,27 +821,24 @@ void UNiagaraNodeFunctionCall::Compile(FTranslator* Translator, TArray<int32>& O
 	}
 	else if (MutableThis->Signature.IsValid())
 	{
-		if (MutableThis->Signature.Inputs.Num() > 0)
+		UClass* DIClass = GetDIClass();
+		if (DIClass)
 		{
-			if (MutableThis->Signature.Inputs[0].GetType().IsDataInterface() && GetValidateDataInterfaces())
+			if (UNiagaraDataInterface* DataInterfaceCDO = Cast<UNiagaraDataInterface>(DIClass->GetDefaultObject()))
 			{
-				UClass* DIClass = MutableThis->Signature.Inputs[0].GetType().GetClass();
-				if (UNiagaraDataInterface* DataInterfaceCDO = Cast<UNiagaraDataInterface>(DIClass->GetDefaultObject()))
-				{
-					TArray<FText> ValidationErrors;
-					DataInterfaceCDO->ValidateFunction(Signature, ValidationErrors);
+				TArray<FText> ValidationErrors;
+				DataInterfaceCDO->ValidateFunction(Signature, ValidationErrors);
 
-					bError = ValidationErrors.Num() > 0;
+				bError = ValidationErrors.Num() > 0;
 
-					for (FText& ValidationError : ValidationErrors)
-					{
-						Translator->Error(ValidationError, this, nullptr);
-					}
+				for (FText& ValidationError : ValidationErrors)
+				{
+					Translator->Error(ValidationError, this, nullptr);
+				}
 
-					if (bError)
-					{
-						return;
-					}
+				if (bError)
+				{
+					return;
 				}
 			}
 		}
diff --git a/Engine/Plugins/FX/Niagara/Source/NiagaraEditor/Private/Widgets/SNiagaraGraphNodeFunctionCallWithSpecifiers.cpp b/Engine/Plugins/FX/Niagara/Source/NiagaraEditor/Private/Widgets/SNiagaraGraphNodeFunctionCallWithSpecifiers.cpp
--- a/Engine/Plugins/FX/Niagara/Source/NiagaraEditor/Private/Widgets/SNiagaraGraphNodeFunctionCallWithSpecifiers.cpp
+++ b/Engine/Plugins/FX/Niagara/Source/NiagaraEditor/Private/Widgets/SNiagaraGraphNodeFunctionCallWithSpecifiers.cpp
@@ -20,15 +20,28 @@ void SNiagaraGraphNodeFunctionCallWithSpecifiers::Construct(const FArguments& In
 
 TSharedRef<SWidget> SNiagaraGraphNodeFunctionCallWithSpecifiers::CreateNodeContentArea()
 {
-	TSharedPtr<SVerticalBox> FunctionSpecifierWidget = SNew(SVerticalBox);
-	for (TTuple<FName, FName>& Entry : *FunctionSpecifiers)
+	TSharedPtr<SWidget> FunctionSpecifierWidget = nullptr;
+
+	UNiagaraNodeFunctionCall* FuncNode = CastChecked<UNiagaraNodeFunctionCall>(GraphNode);
+	if(FuncNode)
 	{
-		FunctionSpecifierWidget->AddSlot()
-			.VAlign(VAlign_Center)
-			[
-				SNew(SNiagaraFunctionSpecifier, Entry.Key, Entry.Value, *FunctionSpecifiers)
+		FunctionSpecifierWidget = INiagaraDataInterfaceNodeActionProvider::GetCustomFunctionSpecifierWidget(FuncNode->GetDIClass(), FuncNode);
+	}
+	
+	//Fall back to default if there is no override for this DI function.
+	if(FunctionSpecifierWidget == nullptr)
+	{
+		TSharedPtr<SVerticalBox> DefaultFunctionSpecifierWidget = SNew(SVerticalBox);
+		for (TTuple<FName, FName>& Entry : *FunctionSpecifiers)
+		{
+			DefaultFunctionSpecifierWidget->AddSlot()
+				.VAlign(VAlign_Center)
+				[
+					SNew(SNiagaraFunctionSpecifier, Entry.Key, Entry.Value, *FunctionSpecifiers)
 					.OnValueNameChanged(this, &SNiagaraGraphNodeFunctionCallWithSpecifiers::OnValueNameChanged)
-			];
+				];
+		}
+		FunctionSpecifierWidget = DefaultFunctionSpecifierWidget;
 	}
 
 	TSharedRef<SWidget> ContentAreaWidget = SGraphNode::CreateNodeContentArea();
diff --git a/Engine/Plugins/FX/Niagara/Source/NiagaraEditor/Public/NiagaraActions.h b/Engine/Plugins/FX/Niagara/Source/NiagaraEditor/Public/NiagaraActions.h
--- a/Engine/Plugins/FX/Niagara/Source/NiagaraEditor/Public/NiagaraActions.h
+++ b/Engine/Plugins/FX/Niagara/Source/NiagaraEditor/Public/NiagaraActions.h
@@ -417,6 +417,8 @@ public:
 	/** Allows DIs to add actions for add pins on the DI function call nodes. */
 	virtual void CollectAddPinActionsImpl(FNiagaraMenuActionCollector& Collector, UEdGraphPin* AddPin) const {}
 
+	virtual TSharedPtr<SWidget> GetCustomFunctionSpecifierWidgetImpl(UNiagaraNodeFunctionCall* FunctionCallNode) const { return nullptr; }
+
 	template<typename DIClass, typename ActionProviderClass>
 	static void Register();
 
@@ -437,6 +439,8 @@ public:
 
 	static NIAGARAEDITOR_API void CollectAddPinActions(UClass* DIClass, FNiagaraMenuActionCollector& Collector, UEdGraphPin* AddPin);
 
+	static NIAGARAEDITOR_API TSharedPtr<SWidget> GetCustomFunctionSpecifierWidget(UClass* DIClass, UNiagaraNodeFunctionCall* FunctionCallNode);
+
 	/** All currently registered action providers. */
 	static NIAGARAEDITOR_API TMap<FName, TUniquePtr<INiagaraDataInterfaceNodeActionProvider>> RegisteredActionProviders;
 };
@@ -491,6 +495,7 @@ public:
 	virtual void GetInlineNodeContextMenuActionsImpl(UToolMenu* ToolMenu) const override;
 	virtual FInlineMenuDisplayOptions GetInlineMenuDisplayOptionsImpl(UClass* DIClass, UEdGraphNode* Source) const override;
 	virtual void CollectAddPinActionsImpl(FNiagaraMenuActionCollector& Collector, UEdGraphPin* AddPin)const override;
+	virtual TSharedPtr<SWidget> GetCustomFunctionSpecifierWidgetImpl(UNiagaraNodeFunctionCall* FunctionCallNode) const override;
 
 private:
 	static void AddDataChannelInitActions(UToolMenu* ToolMenu);
diff --git a/Engine/Plugins/FX/Niagara/Source/NiagaraEditor/Public/NiagaraNodeFunctionCall.h b/Engine/Plugins/FX/Niagara/Source/NiagaraEditor/Public/NiagaraNodeFunctionCall.h
--- a/Engine/Plugins/FX/Niagara/Source/NiagaraEditor/Public/NiagaraNodeFunctionCall.h
+++ b/Engine/Plugins/FX/Niagara/Source/NiagaraEditor/Public/NiagaraNodeFunctionCall.h
@@ -192,6 +192,12 @@ public:
 	virtual bool GetValidateDataInterfaces() const { return true; };
 
 	void RemoveAllDynamicPins();
+
+	//If this is a DI function call, return the UClass of the DI. returns null for non DI function calls.
+	UClass* GetDIClass()const;
+
+	void SetFunctionSpecifier(FName Key, FName Value);
+
 protected:
 	UEdGraphPin* AddStaticSwitchInputPin(FNiagaraVariable Input);
 
diff --git a/Engine/Plugins/FX/Niagara/Source/NiagaraEditor/Public/Widgets/SNiagaraParameterMenu.cpp b/Engine/Plugins/FX/Niagara/Source/NiagaraEditor/Public/Widgets/SNiagaraParameterMenu.cpp
--- a/Engine/Plugins/FX/Niagara/Source/NiagaraEditor/Public/Widgets/SNiagaraParameterMenu.cpp
+++ b/Engine/Plugins/FX/Niagara/Source/NiagaraEditor/Public/Widgets/SNiagaraParameterMenu.cpp
@@ -27,6 +27,15 @@
 #include "ViewModels/TNiagaraViewModelManager.h"
 #include "Widgets/SNiagaraActionMenuExpander.h"
 
+#include "Widgets/Input/SComboButton.h"
+#include "Widgets/Input/SEditableTextBox.h"
+#include "Widgets/Layout/SScaleBox.h"
+#include "NiagaraEditorStyle.h"
+#include "Widgets/SNiagaraPinTypeSelector.h"
+
+
+#include "NiagaraDataChannel.h"
+
 #define LOCTEXT_NAMESPACE "SNiagaraParameterMenu"
 namespace UE::Private
 {
@@ -930,4 +939,225 @@ void SNiagaraChangePinTypeMenu::CollectAllActions(FGraphActionListBuilderBase& O
 	Collector.AddAllActionsTo(OutAllActions);
 }
 
+///////////////////////////////////////////////////////////////////////////////
+/// Change Pin Type Menu													///
+///////////////////////////////////////////////////////////////////////////////
+
+void SNiagaraFunctionSpecifierNDCVariablesMenu::Construct(const FArguments& InArgs)
+{
+	checkf(InArgs._Owner != nullptr, TEXT("Tried to construct NDC Variable menu without an owner."));
+	Owner = InArgs._Owner;
+	AllowedTypes = InArgs._AllowedTypes;
+
+	SNiagaraParameterMenu::FArguments SuperArgs;
+	SuperArgs._AutoExpandMenu = InArgs._AutoExpandMenu;
+	SNiagaraParameterMenu::Construct(SuperArgs);
+}
+
+void SNiagaraFunctionSpecifierNDCVariablesMenu::OnSelectionChanged(FNiagaraVariableBase Var)
+{
+	check(Owner);
+	Owner->OnMenuSelectionChanged(Var);
+}
+
+void SNiagaraFunctionSpecifierNDCVariablesMenu::CollectAllActions(FGraphActionListBuilderBase& OutAllActions)
+{
+	FNiagaraMenuActionCollector Collector;
+
+	UNiagaraDataChannel::ForEachDataChannel([&](UNiagaraDataChannel* DataChannel)
+	{
+		const UNiagaraDataChannelAsset* NDCAsset = DataChannel->GetAsset();
+		if(NDCAsset)
+		{
+			for (const FNiagaraDataChannelVariable& NDCVar : DataChannel->GetVariables())
+			{
+				FText Category = FText::FromString(NDCAsset->GetName());
+				const FText DisplayName = FText::FromName(NDCVar.GetName());
+				const FText Tooltip = FText::Format(LOCTEXT("SNiagaraFunctionSpecifierNDCVariablesMenu_TooltipFmt", "Variable {0} from Data Channel {1}"), DisplayName, Category);
+
+				if(AllowedTypes.Num() > 0 && !AllowedTypes.Contains(NDCVar.GetType()))
+				{
+					continue;//This type is not on the allowed list.
+				}
+
+				TSharedPtr<FNiagaraMenuAction> Action(new FNiagaraMenuAction(
+					Category, DisplayName, Tooltip, 0, FText::GetEmpty(),
+					FNiagaraMenuAction::FOnExecuteStackAction::CreateRaw(this, &SNiagaraFunctionSpecifierNDCVariablesMenu::OnSelectionChanged, FNiagaraVariableBase(NDCVar))));
+
+				//TODO: The base niagara action widget will store this parameter and display this regardless of a change in the NDC variable.
+				//Need to subclass this and have it refresh it's parameter on construction from the NDCVar that matches it's guid.
+				Action->SetParameterVariable(NDCVar);
+				Collector.AddAction(Action, 0);
+			}
+		}
+	});
+
+	Collector.AddAllActionsTo(OutAllActions);
+}
+
+
+//////////////////////////////////////////////////////////////////////////
+
+const static FName VarNameSpecifierKey(TEXT("VarName"));
+const static FName VarTypeSpecifierKey(TEXT("VarType"));
+
+//Encode a Niagara Variable as a function specifier value.
+void EncodeVariableAsSpecifiers(const FNiagaraVariableBase& Var, UNiagaraNodeFunctionCall* TargetNode)
+{
+	if(TargetNode)
+	{
+		FString TypeStr;
+		UScriptStruct* TypeStruct = FNiagaraTypeDefinition::StaticStruct();
+		FNiagaraTypeDefinition TypeDef = Var.GetType();
+		TypeStruct->ExportText(TypeStr, &TypeDef, nullptr, nullptr, PPF_None, nullptr);
+
+		TargetNode->SetFunctionSpecifier(VarNameSpecifierKey, Var.GetName());
+		TargetNode->SetFunctionSpecifier(VarTypeSpecifierKey, *TypeStr);
+	}
+}
+
+FNiagaraVariableBase DecodeVariableAsSpecifiers(UNiagaraNodeFunctionCall* Node)
+{
+	if(Node)
+	{
+		const FName* FuncSpecVarName = Node->FunctionSpecifiers.Find(VarNameSpecifierKey);
+		const FName* FuncSpecVarType = Node->FunctionSpecifiers.Find(VarTypeSpecifierKey);
+		if (FuncSpecVarName && FuncSpecVarType)
+		{
+			FStringOutputDevice ErrorOut;
+			UScriptStruct* TypeStruct = FNiagaraTypeDefinition::StaticStruct();
+			FNiagaraTypeDefinition TypeDef;
+			if (TypeStruct->ImportText(*FuncSpecVarType->ToString(), &TypeDef, nullptr, PPF_None, &ErrorOut, TypeStruct->GetName(), true))
+			{
+				return FNiagaraVariableBase(TypeDef, *FuncSpecVarName);
+			}
+		}
+	}
+	return FNiagaraVariableBase();
+}
+
+
+void SNiagaraFunctionSpecifierNDCVariablesSelector::Construct(const FArguments& InArgs)
+{
+	checkf(InArgs._WeakNodeToModify != nullptr, TEXT("Tried to construct change pin type menu without valid pin ptr!"));
+	this->WeakNodeToModify = InArgs._WeakNodeToModify;
+	this->AllowedTypes = InArgs._AllowedTypes;
+
+	GenerateButtonWidget();
+
+	SelectorButton = SNew(SComboButton)
+		.ContentPadding(3)
+		.MenuPlacement(MenuPlacement_BelowAnchor)
+		.HasDownArrow(true)
+		.ButtonStyle(FAppStyle::Get(), "Button")
+		.ToolTipText(GetTooltipText())
+		.OnGetMenuContent(this, &SNiagaraFunctionSpecifierNDCVariablesSelector::GetMenuContent)
+		.ButtonContent()
+		[
+			ButtonContent.ToSharedRef()
+		];
+
+	ChildSlot
+		.VAlign(VAlign_Fill)
+		.HAlign(HAlign_Fill)
+		[
+			SNew(SHorizontalBox)
+			+ SHorizontalBox::Slot()
+			.HAlign(HAlign_Fill)
+			.AutoWidth()
+			.Padding(5)
+			[
+				SNew(STextBlock)
+				.Text(FText::FromName(TEXT("NDC Variable")))
+			]
+			+ SHorizontalBox::Slot()
+			.HAlign(HAlign_Fill)
+			.AutoWidth()
+			.Padding(5)
+			[
+				SelectorButton.ToSharedRef()
+			]
+		];
+}
+
+void SNiagaraFunctionSpecifierNDCVariablesSelector::OnMenuSelectionChanged(const FNiagaraVariableBase& Var)
+{
+	if (UNiagaraNodeFunctionCall* NodeToModify = WeakNodeToModify.Get())
+	{
+		EncodeVariableAsSpecifiers(Var, NodeToModify);
+		GenerateButtonWidget();
+	}
+}
+
+void SNiagaraFunctionSpecifierNDCVariablesSelector::GenerateButtonWidget()
+{
+	if (UNiagaraNodeFunctionCall* NodeToModify = WeakNodeToModify.Get())
+	{
+		FNiagaraVariableBase FuncSpecVariable = DecodeVariableAsSpecifiers(NodeToModify);
+		FText			   IconToolTip = FuncSpecVariable.GetType().GetNameText();
+		FSlateBrush const* IconBrush = FuncSpecVariable.GetType().IsStatic() ? FNiagaraEditorStyle::Get().GetBrush(TEXT("NiagaraEditor.StaticIcon")) : FAppStyle::GetBrush(TEXT("Kismet.AllClasses.VariableIcon"));
+		const FLinearColor TypeColor = UEdGraphSchema_Niagara::GetTypeColor(FuncSpecVariable.GetType());
+		FSlateColor        IconColor = FSlateColor(TypeColor);
+		FString			   IconDocLink, IconDocExcerpt;
+		FSlateBrush const* SecondaryIconBrush = FAppStyle::GetBrush(TEXT("NoBrush"));
+		FSlateColor        SecondaryIconColor = IconColor;
+		TSharedRef<SWidget> IconWidget = SNew(SScaleBox)
+			[
+				SNew(SNiagaraIconWidget)
+				.IconToolTip(IconToolTip)
+			.IconBrush(IconBrush)
+			.IconColor(IconColor)
+			.DocLink(IconDocLink)
+			.DocExcerpt(IconDocExcerpt)
+			.SecondaryIconBrush(SecondaryIconBrush)
+			.SecondaryIconColor(SecondaryIconColor)
+			];
+
+		ButtonContent = SNew(SBox)
+			.MinDesiredWidth(150.0f)
+			[
+				SNew(SHorizontalBox)
+				+ SHorizontalBox::Slot()
+			.HAlign(HAlign_Fill)
+			.AutoWidth()
+			.Padding(0)
+			[
+				IconWidget
+			]
+		+ SHorizontalBox::Slot()
+			.HAlign(HAlign_Fill)
+			.AutoWidth()
+			.Padding(0)
+			[
+				SNew(STextBlock).Text(FText::FromName(FuncSpecVariable.GetName()))
+			]
+			];
+
+		return;
+	}
+
+	ButtonContent = SNew(STextBlock).Text(LOCTEXT("NDCVarButtonError","Error"));
+}
+
+TSharedRef<SWidget> SNiagaraFunctionSpecifierNDCVariablesSelector::GetMenuContent()
+{
+	TSharedRef<SNiagaraFunctionSpecifierNDCVariablesMenu> MenuWidget = SNew(SNiagaraFunctionSpecifierNDCVariablesMenu)
+			.Owner(StaticCastSharedRef<SNiagaraFunctionSpecifierNDCVariablesSelector>(AsShared()).ToSharedPtr())
+			.AutoExpandMenu(false)
+			.AllowedTypes(AllowedTypes);
+
+	SelectorButton->SetMenuContentWidgetToFocus(MenuWidget->GetSearchBox());
+	return SNew(SBorder)
+		.BorderImage(FAppStyle::GetBrush("Menu.Background"))
+		.Padding(1)
+		[
+			MenuWidget
+		];
+}
+
+FText SNiagaraFunctionSpecifierNDCVariablesSelector::GetTooltipText() const
+{
+	return LOCTEXT("FunctionSpecifierNDCVariableSelectorTooltip", "Select an NDC Variable to use.");
+}
+
 #undef LOCTEXT_NAMESPACE /*"SNiagaraParameterMenu"*/
diff --git a/Engine/Plugins/FX/Niagara/Source/NiagaraEditor/Public/Widgets/SNiagaraParameterMenu.h b/Engine/Plugins/FX/Niagara/Source/NiagaraEditor/Public/Widgets/SNiagaraParameterMenu.h
--- a/Engine/Plugins/FX/Niagara/Source/NiagaraEditor/Public/Widgets/SNiagaraParameterMenu.h
+++ b/Engine/Plugins/FX/Niagara/Source/NiagaraEditor/Public/Widgets/SNiagaraParameterMenu.h
@@ -22,6 +22,7 @@ class UNiagaraNodeAssignment;
 class UNiagaraNodeWithDynamicPins;
 class UNiagaraParameterDefinitions;
 class UNiagaraScriptVariable;
+class SComboButton;
 
 
 class SNiagaraParameterMenu : public SCompoundWidget
@@ -209,3 +210,58 @@ protected:
 private:
 	UEdGraphPin* PinToModify;
 };
+
+class SNiagaraFunctionSpecifierNDCVariablesSelector;
+class SNiagaraFunctionSpecifierNDCVariablesMenu : public SNiagaraParameterMenu
+{
+public:
+
+	SLATE_BEGIN_ARGS(SNiagaraFunctionSpecifierNDCVariablesMenu)
+		: _AutoExpandMenu(false)
+	{}
+		//~ Begin Required Args
+		SLATE_ARGUMENT(TSharedPtr<SNiagaraFunctionSpecifierNDCVariablesSelector>, Owner)
+		SLATE_ARGUMENT(TArray<FNiagaraTypeDefinition>, AllowedTypes)
+		//~ End Required Args
+		SLATE_ARGUMENT(bool, AutoExpandMenu)
+	SLATE_END_ARGS();
+
+	void Construct(const FArguments& InArgs);
+
+	void OnSelectionChanged(FNiagaraVariableBase Var);
+
+protected:
+	virtual void CollectAllActions(FGraphActionListBuilderBase& OutAllActions) override;
+
+	TSharedPtr<SNiagaraFunctionSpecifierNDCVariablesSelector> Owner = nullptr;
+	TArray<FNiagaraTypeDefinition> AllowedTypes;
+};
+
+class SNiagaraFunctionSpecifierNDCVariablesSelector : public SCompoundWidget
+{
+public:
+	SLATE_BEGIN_ARGS(SNiagaraFunctionSpecifierNDCVariablesSelector)
+		: _WeakNodeToModify(nullptr)		
+	{}
+	SLATE_ARGUMENT(TWeakObjectPtr<UNiagaraNodeFunctionCall>, WeakNodeToModify)
+	SLATE_ARGUMENT(TArray<FNiagaraTypeDefinition>, AllowedTypes)
+	//~ End Required Args
+	SLATE_END_ARGS()
+public:
+	NIAGARAEDITOR_API void Construct(const FArguments& InArgs);
+
+	void OnMenuSelectionChanged(const FNiagaraVariableBase& NewSelection);
+
+protected:
+	NIAGARAEDITOR_API virtual TSharedRef<SWidget>	GetMenuContent();
+	NIAGARAEDITOR_API FText GetTooltipText() const;
+	NIAGARAEDITOR_API FText GetSelectionText() const;
+
+	NIAGARAEDITOR_API void GenerateButtonWidget();
+
+private:
+	TWeakObjectPtr<UNiagaraNodeFunctionCall> WeakNodeToModify;
+	TArray<FNiagaraTypeDefinition> AllowedTypes;
+	TSharedPtr<SComboButton> SelectorButton;
+	TSharedPtr<SWidget> ButtonContent;
+};
\ No newline at end of file
