Mac fix

diff --git a/Engine/Plugins/FX/Niagara/Source/Niagara/Private/DataInterface/NiagaraDataInterfaceDataChannelRead.cpp b/Engine/Plugins/FX/Niagara/Source/Niagara/Private/DataInterface/NiagaraDataInterfaceDataChannelRead.cpp
--- a/Engine/Plugins/FX/Niagara/Source/Niagara/Private/DataInterface/NiagaraDataInterfaceDataChannelRead.cpp
+++ b/Engine/Plugins/FX/Niagara/Source/Niagara/Private/DataInterface/NiagaraDataInterfaceDataChannelRead.cpp
@@ -1781,7 +1781,7 @@ void UNiagaraDataInterfaceDataChannelRead::SpawnDirect(FVectorVMExternalFunction
 	//This should only be called from emitter scripts and since it has per instance data then we process them individually.
 	check(Context.GetNumInstances() == 1);
 
-	using TSimType = TLWCToSimType<T>::TSimType;
+	using TSimType = typename TLWCToSimType<T>::TSimType;
 	VectorVM::FUserPtrHandler<FNDIDataChannelReadInstanceData> InstData(Context);
 
 	FNDIInputParam<FNiagaraBool> InEnabled(Context);
@@ -1854,7 +1854,7 @@ void UNiagaraDataInterfaceDataChannelRead::ScaleSpawnCount(FVectorVMExternalFunc
 	//This should only be called from emitter scripts and since it has per instance data then we process them individually.
 	check(Context.GetNumInstances() == 1);
 
-	using TSimType = TLWCToSimType<T>::TSimType;
+	using TSimType = typename TLWCToSimType<T>::TSimType;
 	VectorVM::FUserPtrHandler<FNDIDataChannelReadInstanceData> InstData(Context);
 
 	FNDIInputParam<FNiagaraBool> InEnabled(Context);
