Fix for parameter naming with NiagaraTransformUtils
#tests compile asset for sm5 with geometry collection or socket interface DI

diff --git a/Engine/Plugins/FX/Niagara/Shaders/Private/NiagaraTransformUtils.ush b/Engine/Plugins/FX/Niagara/Shaders/Private/NiagaraTransformUtils.ush
--- a/Engine/Plugins/FX/Niagara/Shaders/Private/NiagaraTransformUtils.ush
+++ b/Engine/Plugins/FX/Niagara/Shaders/Private/NiagaraTransformUtils.ush
@@ -11,12 +11,12 @@ struct FNiagaraTransform
 	float3	Scale;
 };
 
-FNiagaraTransform LoadTransform(ByteAddressBuffer Buffer, uint Offset)
+FNiagaraTransform LoadTransform(ByteAddressBuffer RawBufferData, uint Offset)
 {
 	FNiagaraTransform R;
-	R.Translation	= asfloat(Buffer.Load3(Offset + 0));
-	R.Rotation		= asfloat(Buffer.Load4(Offset + 12));
-	R.Scale			= asfloat(Buffer.Load3(Offset + 28));
+	R.Translation	= asfloat(RawBufferData.Load3(Offset + 0));
+	R.Rotation		= asfloat(RawBufferData.Load4(Offset + 12));
+	R.Scale			= asfloat(RawBufferData.Load3(Offset + 28));
 	return R;
 }
 
