Incorporate the script path into the ddc key for the niagara gpu scripts to prevent non-deterministic DDC results
#tests cook and ddc-verify assets

diff --git a/Engine/Plugins/FX/Niagara/Source/NiagaraEditor/Private/NiagaraCompilationTasks.cpp b/Engine/Plugins/FX/Niagara/Source/NiagaraEditor/Private/NiagaraCompilationTasks.cpp
--- a/Engine/Plugins/FX/Niagara/Source/NiagaraEditor/Private/NiagaraCompilationTasks.cpp
+++ b/Engine/Plugins/FX/Niagara/Source/NiagaraEditor/Private/NiagaraCompilationTasks.cpp
@@ -55,6 +55,7 @@ namespace NiagaraSystemCookStats
 
 namespace NiagaraCompilationTasksImpl
 {
+	static const FGuid UE_NIAGARA_ASYNC_TASK_COMPILER_VER = FGuid(0x1D7ABD2B, 0xE882465E, 0xA93B3CA6, 0x7753CD29);
 	static UE::DerivedData::FCacheBucket NiagaraDDCBucket("NiagaraScript");
 
 	void GetUsagesToDuplicate(ENiagaraScriptUsage TargetUsage, TArray<ENiagaraScriptUsage>& DuplicateUsages)
@@ -134,14 +135,12 @@ namespace NiagaraCompilationTasksImpl
 
 	static UE::DerivedData::FCacheKey BuildNiagaraDDCCacheKey(const FNiagaraVMExecutableDataId& CompileId, const FString& ScriptPath)
 	{
-		enum { UE_NIAGARASCRIPT_DERIVEDDATA_VER = 2 };
-
 		FString KeyString;
 		KeyString.Reserve(1024);
 
-		KeyString.Appendf(TEXT("%i_%i"),
-			(int32)UE_NIAGARASCRIPT_DERIVEDDATA_VER, GNiagaraSkipVectorVMBackendOptimizations);
-
+		NiagaraCompilationTasksImpl::UE_NIAGARA_ASYNC_TASK_COMPILER_VER.AppendString(KeyString);
+		KeyString.AppendChar(TCHAR('_'));
+		KeyString.Appendf(TEXT("%i"), GNiagaraSkipVectorVMBackendOptimizations);
 		KeyString.AppendChar(TCHAR('_'));
 		KeyString.Append(ScriptPath);
 		KeyString.AppendChar(TCHAR('_'));
@@ -151,7 +150,7 @@ namespace NiagaraCompilationTasksImpl
 		return { NiagaraDDCBucket, FIoHash::HashBuffer(MakeMemoryView(FTCHARToUTF8(KeyString))) };
 	}
 
-	static UE::DerivedData::FCacheKey BuildNiagaraComputeDDCCacheKey(const FNiagaraShaderMapId& ShaderMapId, EShaderPlatform ShaderPlatform)
+	static UE::DerivedData::FCacheKey BuildNiagaraComputeDDCCacheKey(const FNiagaraShaderMapId& ShaderMapId, EShaderPlatform ShaderPlatform, const FString& ScriptPath)
 	{
 		static const FString NIAGARASHADERMAP_DERIVEDDATA_VER = FDevSystemGuids::GetSystemGuid(FDevSystemGuids::Get().NIAGARASHADERMAP_DERIVEDDATA_VER).ToString(EGuidFormats::DigitsWithHyphens);
 
@@ -165,6 +164,8 @@ namespace NiagaraCompilationTasksImpl
 		ShaderMapAppendKeyString(ShaderPlatform, KeyString);
 		ShaderMapId.AppendKeyString(KeyString);
 		KeyString.AppendChar(TCHAR('_'));
+		KeyString.Append(ScriptPath);
+		KeyString.AppendChar(TCHAR('_'));
 		KeyString.Append(NIAGARASHADERMAP_DERIVEDDATA_VER);
 
 		return { NiagaraDDCBucket, FIoHash::HashBuffer(MakeMemoryView(FTCHARToUTF8(KeyString))) };
@@ -275,7 +276,7 @@ namespace NiagaraCompilationTasksImpl
 		TArray<uint8> BinaryData;
 		FMemoryWriter Ar(BinaryData, true);
 		ShaderMap->Serialize(Ar);
-		
+
 		if (!BinaryData.IsEmpty() && !Ar.IsError())
 		{
 			return MakeSharedBufferFromArray(MoveTemp(BinaryData));
@@ -1132,7 +1133,8 @@ void FNiagaraSystemCompilationTask::FDispatchAndProcessDataCacheGetRequests::Lau
 		FNiagaraSystemCompilationTask::FCompileComputeShaderTaskInfo& ShaderCompileTask = SystemCompileTask->CompileComputeShaderTasks[ShaderCompileTaskIt];
 		if (ShaderCompileTask.IsOutstanding())
 		{
-			const UE::DerivedData::FCacheKey CacheKey = NiagaraCompilationTasksImpl::BuildNiagaraComputeDDCCacheKey(ShaderCompileTask.ShaderMapId, ShaderCompileTask.ShaderPlatform);
+			FNiagaraSystemCompilationTask::FCompileTaskInfo& CompileTask = SystemCompileTask->CompileTasks[ShaderCompileTask.ParentCompileTaskIndex];
+			const UE::DerivedData::FCacheKey CacheKey = NiagaraCompilationTasksImpl::BuildNiagaraComputeDDCCacheKey(ShaderCompileTask.ShaderMapId, ShaderCompileTask.ShaderPlatform, CompileTask.AssetPath);
 
 			if (CacheKey != ShaderCompileTask.DDCTaskInfo.DataCacheGetKey)
 			{
@@ -1140,8 +1142,6 @@ void FNiagaraSystemCompilationTask::FDispatchAndProcessDataCacheGetRequests::Lau
 				Index.CompileTaskIndex = ShaderCompileTaskIt;
 				Index.bShaderCompileTask = true;
 
-				FNiagaraSystemCompilationTask::FCompileTaskInfo& CompileTask = SystemCompileTask->CompileTasks[ShaderCompileTask.ParentCompileTaskIndex];
-
 				FCacheGetValueRequest& ShaderGetValueRequest = GetRequests.AddDefaulted_GetRef();
 				ShaderGetValueRequest.Name = CompileTask.AssetPath;
 				ShaderGetValueRequest.Key = CacheKey;
