Fixed a crash when checking the paused state on a recently disabled Niagara component

#tests Added two engine tests to check for regressions on this issue

diff --git a/Engine/Plugins/FX/Niagara/Source/Niagara/Private/NiagaraComponent.cpp b/Engine/Plugins/FX/Niagara/Source/Niagara/Private/NiagaraComponent.cpp
--- a/Engine/Plugins/FX/Niagara/Source/Niagara/Private/NiagaraComponent.cpp
+++ b/Engine/Plugins/FX/Niagara/Source/Niagara/Private/NiagaraComponent.cpp
@@ -1068,7 +1068,9 @@ bool UNiagaraComponent::IsPaused()const
 {
 	//check the system instance is actually in the right state. 
 	//In cases where we don't have a system instance or we're culled by scalability, the instance can have a differnt internal state to the user's desired state.
-	check(SystemInstanceController.IsValid() == false || bIsCulledByScalability || SystemInstanceController->IsPaused() == bDesiredPauseState);
+	bool bIsInstanceValid = SystemInstanceController.IsValid();
+	bool bIsInstanceRunning = bIsInstanceValid && SystemInstanceController->GetRequestedExecutionState() == ENiagaraExecutionState::Active && SystemInstanceController->GetActualExecutionState() == ENiagaraExecutionState::Active;
+	ensure(bIsInstanceRunning == false || bIsCulledByScalability || SystemInstanceController->IsPaused() == bDesiredPauseState);
 	return bDesiredPauseState;
 }
 
