Add ForceCompile option to Niagara Scripts

PR #8101

diff --git a/Engine/Plugins/FX/Niagara/Source/Niagara/Private/NiagaraScript.cpp b/Engine/Plugins/FX/Niagara/Source/Niagara/Private/NiagaraScript.cpp
--- a/Engine/Plugins/FX/Niagara/Source/Niagara/Private/NiagaraScript.cpp
+++ b/Engine/Plugins/FX/Niagara/Source/Niagara/Private/NiagaraScript.cpp
@@ -3378,7 +3378,7 @@ void UNiagaraScript::RequestCompile(const FGuid& ScriptVersion, bool bForceCompi
 		const FString ScriptPathName = GetPathName();
 
 		// check the ddc first
-		if (GetDerivedDataCacheRef().GetSynchronous(*GetNiagaraDDCKeyString(ScriptVersion, ScriptPathName), OutData, ScriptPathName))
+		if (!bForceCompile && GetDerivedDataCacheRef().GetSynchronous(*GetNiagaraDDCKeyString(ScriptVersion, ScriptPathName), OutData, ScriptPathName))
 		{
 			FNiagaraVMExecutableData ExeData;
 			if (BinaryToExecData(this, OutData, ExeData))
diff --git a/Engine/Plugins/FX/Niagara/Source/NiagaraEditor/Private/Toolkits/NiagaraScriptToolkit.cpp b/Engine/Plugins/FX/Niagara/Source/NiagaraEditor/Private/Toolkits/NiagaraScriptToolkit.cpp
--- a/Engine/Plugins/FX/Niagara/Source/NiagaraEditor/Private/Toolkits/NiagaraScriptToolkit.cpp
+++ b/Engine/Plugins/FX/Niagara/Source/NiagaraEditor/Private/Toolkits/NiagaraScriptToolkit.cpp
@@ -602,7 +602,7 @@ void FNiagaraScriptToolkit::SetupCommands()
 		FCanExecuteAction::CreateSP(this, &FNiagaraScriptToolkit::OnApplyEnabled));
 	GetToolkitCommands()->MapAction(
 		FNiagaraEditorCommands::Get().Compile,
-		FExecuteAction::CreateRaw(this, &FNiagaraScriptToolkit::CompileScript, true));
+		FExecuteAction::CreateRaw(this, &FNiagaraScriptToolkit::CompileScript, false));
 	GetToolkitCommands()->MapAction(
 		FNiagaraEditorCommands::Get().RefreshNodes,
 		FExecuteAction::CreateRaw(this, &FNiagaraScriptToolkit::RefreshNodes));
@@ -651,10 +651,25 @@ void FNiagaraScriptToolkit::ExtendToolbar()
 					TAttribute<FText>(), TAttribute<FText>(),
 					FSlateIcon(FAppStyle::Get().GetStyleSetName(), "AssetEditor.Apply")));
 
-				InSection.AddEntry(FToolMenuEntry::InitToolBarButton(FNiagaraEditorCommands::Get().Compile,
-					TAttribute<FText>(),
-					TAttribute<FText>(NiagaraScriptToolkit.ToSharedRef(), &FNiagaraScriptToolkit::GetCompileStatusTooltip),
-					TAttribute<FSlateIcon>(NiagaraScriptToolkit.ToSharedRef(), &FNiagaraScriptToolkit::GetCompileStatusImage)));
+				InSection.AddDynamicEntry("CompileDynamic", FNewToolMenuSectionDelegate::CreateLambda([](FToolMenuSection& DynamicSection)
+				{
+					if(TSharedPtr<FNiagaraScriptToolkit> NiagaraScriptToolkit = GetNiagaraScriptToolkitFromMenuContext(DynamicSection.FindContext<UAssetEditorToolkitMenuContext>()))
+					{
+						DynamicSection.AddEntry(FToolMenuEntry::InitToolBarButton(FNiagaraEditorCommands::Get().Compile,
+						TAttribute<FText>(),
+						TAttribute<FText>(NiagaraScriptToolkit.ToSharedRef(), &FNiagaraScriptToolkit::GetCompileStatusTooltip),
+						TAttribute<FSlateIcon>(NiagaraScriptToolkit.ToSharedRef(), &FNiagaraScriptToolkit::GetCompileStatusImage)));
+
+						DynamicSection.AddEntry(FToolMenuEntry::InitComboButton("CompileOptions",
+							FUIAction(),
+							FNewToolMenuChoice(FOnGetContent::CreateRaw(NiagaraScriptToolkit.Get(), &FNiagaraScriptToolkit::GenerateCompileMenuContent)),
+							LOCTEXT("CompileCombo_Label", "Compile options"),
+							LOCTEXT("CompileComboToolTip", "Compile options menu"),
+							FSlateIcon(FAppStyle::GetAppStyleSetName(), "LevelEditor.Build"),
+							true
+							));
+					}			
+				}));
 
 				InSection.AddEntry(FToolMenuEntry::InitToolBarButton(FNiagaraEditorCommands::Get().RefreshNodes,
 					TAttribute<FText>(),
@@ -737,6 +752,24 @@ FText FNiagaraScriptToolkit::GetCompileStatusTooltip() const
 	return FNiagaraEditorUtilities::StatusToText(Status);
 }
 
+TSharedRef<SWidget> FNiagaraScriptToolkit::GenerateCompileMenuContent()
+{
+	constexpr bool bShouldCloseWindowAfterMenuSelection = true;
+	static const FName CompileStatusBackground("AssetEditor.CompileStatus.Background");
+	static const FName CompileStatusUnknown("AssetEditor.CompileStatus.Overlay.Unknown");
+	FMenuBuilder MenuBuilder(bShouldCloseWindowAfterMenuSelection, nullptr);
+
+	FUIAction ForceCompileAction(
+		FExecuteAction::CreateRaw(this, &FNiagaraScriptToolkit::CompileScript, true));
+
+	MenuBuilder.AddMenuEntry(LOCTEXT("ForceCompile", "Force Compile"),
+		LOCTEXT("ForceCompileTooltip", "Triggers a recompilation of this script, ignoring the change tracking and cached results."),
+		FSlateIcon(FAppStyle::Get().GetStyleSetName(), CompileStatusBackground, NAME_None, CompileStatusUnknown),
+		ForceCompileAction, NAME_None, EUserInterfaceActionType::Button);
+
+	return MenuBuilder.MakeWidget();
+}
+
 FSlateIcon FNiagaraScriptToolkit::GetRefreshStatusImage() const
 {
 	return FSlateIcon(FAppStyle::Get().GetStyleSetName(), "FontEditor.Update");
@@ -763,7 +796,7 @@ FText FNiagaraScriptToolkit::GetVersionButtonLabel() const
 
 void FNiagaraScriptToolkit::CompileScript(bool bForce)
 {
-	ScriptViewModel->CompileStandaloneScript();
+	ScriptViewModel->CompileStandaloneScript(bForce);
 }
 
 void FNiagaraScriptToolkit::RefreshNodes()
diff --git a/Engine/Plugins/FX/Niagara/Source/NiagaraEditor/Private/Toolkits/NiagaraScriptToolkit.h b/Engine/Plugins/FX/Niagara/Source/NiagaraEditor/Private/Toolkits/NiagaraScriptToolkit.h
--- a/Engine/Plugins/FX/Niagara/Source/NiagaraEditor/Private/Toolkits/NiagaraScriptToolkit.h
+++ b/Engine/Plugins/FX/Niagara/Source/NiagaraEditor/Private/Toolkits/NiagaraScriptToolkit.h
@@ -120,6 +120,8 @@ private:
 	const FName GetNiagaraScriptMessageLogName(FVersionedNiagaraScript InScript) const;
 	FSlateIcon GetCompileStatusImage() const;
 	FText GetCompileStatusTooltip() const;
+	
+	TSharedRef<SWidget> GenerateCompileMenuContent();
 
 	/** Builds the toolbar widget */
 	void ExtendToolbar();
