Fix for AppendToClassSchema accessing shader source when -noshadercompile was being used

diff --git a/Engine/Plugins/FX/Niagara/Source/Niagara/Private/NiagaraScript.cpp b/Engine/Plugins/FX/Niagara/Source/Niagara/Private/NiagaraScript.cpp
--- a/Engine/Plugins/FX/Niagara/Source/Niagara/Private/NiagaraScript.cpp
+++ b/Engine/Plugins/FX/Niagara/Source/Niagara/Private/NiagaraScript.cpp
@@ -160,33 +160,44 @@ namespace NiagaraScriptInternal
 			UE_CALL_ONCE(
 				[]()
 				{
-					TArray<UClass*> DataInterfaceClasses;
-					for (TObjectIterator<UClass> ClassIterator; ClassIterator; ++ClassIterator)
+					FSHA1 HashState;
+
+					// only try to access the shader source if we are allowed to compile shaders
+					if (AllowShaderCompiling())
 					{
-						UClass* Class = *ClassIterator;
-						if (Class == nullptr ||
-							Class->HasAnyClassFlags(CLASS_Abstract | CLASS_Deprecated | CLASS_NewerVersionExists) ||
-							!ClassIterator->IsChildOf(UNiagaraDataInterface::StaticClass()))
+						TArray<UClass*> DataInterfaceClasses;
+						for (TObjectIterator<UClass> ClassIterator; ClassIterator; ++ClassIterator)
 						{
-							continue;
+							UClass* Class = *ClassIterator;
+							if (Class == nullptr ||
+								Class->HasAnyClassFlags(CLASS_Abstract | CLASS_Deprecated | CLASS_NewerVersionExists) ||
+								!ClassIterator->IsChildOf(UNiagaraDataInterface::StaticClass()))
+							{
+								continue;
+							}
+							DataInterfaceClasses.Add(Class);
 						}
-						DataInterfaceClasses.Add(Class);
-					}
-					Algo::Sort(
-						DataInterfaceClasses,
-						[](UClass* Lhs, UClass* Rhs)
+						Algo::Sort(
+							DataInterfaceClasses,
+							[](UClass* Lhs, UClass* Rhs)
+							{
+								return Lhs->GetName() < Rhs->GetName();
+							}
+						);
+
+						FNiagaraCompileHashVisitor Visitor(HashState);
+						for (UClass* DIClass : DataInterfaceClasses)
 						{
-							return Lhs->GetName() < Rhs->GetName();
+							UNiagaraDataInterface* CDODataInterface = DIClass->GetDefaultObject<UNiagaraDataInterface>();
+							CDODataInterface->AppendCompileHash(&Visitor);
 						}
-					);
-
-					FSHA1 HashState;
-					FNiagaraCompileHashVisitor Visitor(HashState);
-					for (UClass* DIClass : DataInterfaceClasses)
+					}
+					else
 					{
-						UNiagaraDataInterface* CDODataInterface = DIClass->GetDefaultObject<UNiagaraDataInterface>();
-						CDODataInterface->AppendCompileHash(&Visitor);
+						const FGuid DISourcePlaceholder = FGuid(0x17B53667, 0x0A08411A, 0xBDD134D5, 0x039F91C2);
+						HashState.Update((const uint8*) &DISourcePlaceholder, sizeof(DISourcePlaceholder));
 					}
+
 					HashState.Final();
 
 					TArray<uint8> DataHash;
diff --git a/Engine/Plugins/FX/Niagara/Source/NiagaraShader/Private/NiagaraShader.cpp b/Engine/Plugins/FX/Niagara/Source/NiagaraShader/Private/NiagaraShader.cpp
--- a/Engine/Plugins/FX/Niagara/Source/NiagaraShader/Private/NiagaraShader.cpp
+++ b/Engine/Plugins/FX/Niagara/Source/NiagaraShader/Private/NiagaraShader.cpp
@@ -1222,14 +1222,17 @@ void FNiagaraShader::BuildClassSchema(FAppendToClassSchemaContext& Context)
 	Context.Update(&ShaderParamStructHash, sizeof(ShaderParamStructHash));
 
 	// Add in any referenced HLSL files.
-	const FSHAHash InstanceShaderHash = GetShaderFileHash((TEXT("/Plugin/FX/Niagara/Private/NiagaraEmitterInstanceShader.usf")), EShaderPlatform::SP_PCD3D_SM5);
-	Context.Update(&InstanceShaderHash, sizeof(InstanceShaderHash));
+	if (AllowShaderCompiling())
+	{
+		const FSHAHash InstanceShaderHash = GetShaderFileHash((TEXT("/Plugin/FX/Niagara/Private/NiagaraEmitterInstanceShader.usf")), EShaderPlatform::SP_PCD3D_SM5);
+		Context.Update(&InstanceShaderHash, sizeof(InstanceShaderHash));
 
-	const FSHAHash NiagaraShaderVersionHash = GetShaderFileHash((TEXT("/Plugin/FX/Niagara/Private/NiagaraShaderVersion.ush")), EShaderPlatform::SP_PCD3D_SM5);
-	Context.Update(&NiagaraShaderVersionHash, sizeof(NiagaraShaderVersionHash));
+		const FSHAHash NiagaraShaderVersionHash = GetShaderFileHash((TEXT("/Plugin/FX/Niagara/Private/NiagaraShaderVersion.ush")), EShaderPlatform::SP_PCD3D_SM5);
+		Context.Update(&NiagaraShaderVersionHash, sizeof(NiagaraShaderVersionHash));
 
-	const FSHAHash ShaderVersionHash = GetShaderFileHash((TEXT("/Engine/Public/ShaderVersion.ush")), EShaderPlatform::SP_PCD3D_SM5);
-	Context.Update(&ShaderVersionHash, sizeof(ShaderVersionHash));
+		const FSHAHash ShaderVersionHash = GetShaderFileHash((TEXT("/Engine/Public/ShaderVersion.ush")), EShaderPlatform::SP_PCD3D_SM5);
+		Context.Update(&ShaderVersionHash, sizeof(ShaderVersionHash));
+	}
 
 	Context.Update(&GNiagaraSkipVectorVMBackendOptimizations, sizeof(GNiagaraSkipVectorVMBackendOptimizations));
 }
