Fixes to Niagara GPU compiling with AsyncTasks when dealing with target platforms that don't require shaders (specifically the WindowsClient+WindowsServer combo)
#tests cooked niagara system with gpu emitter with the WindowsClient+WindowsServer targetplatform combo

diff --git a/Engine/Plugins/FX/Niagara/Source/NiagaraEditor/Private/NiagaraCompilationTasks.cpp b/Engine/Plugins/FX/Niagara/Source/NiagaraEditor/Private/NiagaraCompilationTasks.cpp
--- a/Engine/Plugins/FX/Niagara/Source/NiagaraEditor/Private/NiagaraCompilationTasks.cpp
+++ b/Engine/Plugins/FX/Niagara/Source/NiagaraEditor/Private/NiagaraCompilationTasks.cpp
@@ -1511,15 +1511,12 @@ void FNiagaraSystemCompilationTask::AddScript(int32 SourceEmitterIndex, UNiagara
 					TaskInfo.ComputeShaderTaskIndices.Reserve(ShaderRequests.Num());
 					for (const FShaderCompileRequest& ShaderRequest : ShaderRequests)
 					{
-						if (Script->ShouldCompile(ShaderRequest.ShaderPlatform))
-						{
-							TaskInfo.ComputeShaderTaskIndices.Add(CompileComputeShaderTasks.Num());
-							FCompileComputeShaderTaskInfo& ShaderTaskInfo = CompileComputeShaderTasks.AddDefaulted_GetRef();
+						TaskInfo.ComputeShaderTaskIndices.Add(CompileComputeShaderTasks.Num());
+						FCompileComputeShaderTaskInfo& ShaderTaskInfo = CompileComputeShaderTasks.AddDefaulted_GetRef();
 
-							ShaderTaskInfo.ParentCompileTaskIndex = CompileTaskIndex;
-							ShaderTaskInfo.ShaderMapId = ShaderRequest.ShaderMapId;
-							ShaderTaskInfo.ShaderPlatform = ShaderRequest.ShaderPlatform;
-						}
+						ShaderTaskInfo.ParentCompileTaskIndex = CompileTaskIndex;
+						ShaderTaskInfo.ShaderMapId = ShaderRequest.ShaderMapId;
+						ShaderTaskInfo.ShaderPlatform = ShaderRequest.ShaderPlatform;
 					}
 				}
 			}
diff --git a/Engine/Plugins/FX/Niagara/Source/NiagaraEditor/Private/NiagaraSystemCompilingManager.cpp b/Engine/Plugins/FX/Niagara/Source/NiagaraEditor/Private/NiagaraSystemCompilingManager.cpp
--- a/Engine/Plugins/FX/Niagara/Source/NiagaraEditor/Private/NiagaraSystemCompilingManager.cpp
+++ b/Engine/Plugins/FX/Niagara/Source/NiagaraEditor/Private/NiagaraSystemCompilingManager.cpp
@@ -314,16 +314,18 @@ FNiagaraCompilationTaskHandle FNiagaraSystemCompilingManager::AddSystem(UNiagara
 
 					if (bRequiresCompilation || !Script->IsShaderMapCached(TargetPlatform, ShaderMapId))
 					{
-						FNiagaraSystemCompilationTask::FShaderCompileRequest& Request = ShaderRequests.AddDefaulted_GetRef();
-						Request.ShaderMapId = ShaderMapId;
-						Request.ShaderPlatform = PlatformFeatureLevel.Key;
+						if (Script->ShouldCompile(PlatformFeatureLevel.Key))
+						{
+							FNiagaraSystemCompilationTask::FShaderCompileRequest& Request = ShaderRequests.AddDefaulted_GetRef();
+							Request.ShaderMapId = ShaderMapId;
+							Request.ShaderPlatform = PlatformFeatureLevel.Key;
+						}
 					}
 				}
 
-				if (!ShaderRequests.IsEmpty())
-				{
-					bRequiresCompilation = true;
-				}
+				// for GPU scripts we only need to worry about compilation if we actually have some shaders that are required.  So we
+				// override bRequiresCompilation based on that so platforms that exclude all shaders will not generate a compile request
+				bRequiresCompilation = !ShaderRequests.IsEmpty();
 			}
 
 			bHasCompilation = bHasCompilation || bRequiresCompilation;
@@ -542,10 +544,11 @@ void FNiagaraSystemCompilingManager::FindOrAddFeatureLevels(const FCompileOption
 	{
 		if (CompileOptions.TargetPlatform)
 		{
-			TArray<FPlatformFeatureLevelPair>& CachedFeatureLevels = PlatformFeatureLevels.FindOrAdd(CompileOptions.TargetPlatform);
-
-			if (CachedFeatureLevels.IsEmpty())
+			TArray<FPlatformFeatureLevelPair>* CachedFeatureLevels = PlatformFeatureLevels.Find(CompileOptions.TargetPlatform);
+			if (!CachedFeatureLevels)
 			{
+				CachedFeatureLevels = &PlatformFeatureLevels.Add(CompileOptions.TargetPlatform);
+
 				TArray<FName> DesiredShaderFormats;
 				CompileOptions.TargetPlatform->GetAllTargetedShaderFormats(DesiredShaderFormats);
 				for (const FName& ShaderFormat : DesiredShaderFormats)
@@ -555,12 +558,12 @@ void FNiagaraSystemCompilingManager::FindOrAddFeatureLevels(const FCompileOption
 
 					if (NiagaraShaderType->ShouldCompilePermutation(FShaderPermutationParameters(ShaderPlatform)))
 					{
-						CachedFeatureLevels.AddUnique(MakeTuple(ShaderPlatform, TargetFeatureLevel));
+						CachedFeatureLevels->AddUnique(MakeTuple(ShaderPlatform, TargetFeatureLevel));
 					}
 				}
 			}
 
-			FeatureLevels = CachedFeatureLevels;
+			FeatureLevels = *CachedFeatureLevels;
 		}
 		else
 		{
