Enable AutoRTFM for Niagara module

diff --git a/Engine/Plugins/FX/Niagara/Source/Niagara/Niagara.Build.cs b/Engine/Plugins/FX/Niagara/Source/Niagara/Niagara.Build.cs
--- a/Engine/Plugins/FX/Niagara/Source/Niagara/Niagara.Build.cs
+++ b/Engine/Plugins/FX/Niagara/Source/Niagara/Niagara.Build.cs
@@ -69,7 +69,5 @@ public class Niagara : ModuleRules
                 "VECTORVM_SUPPORTS_EXPERIMENTAL=1",
                 "VECTORVM_SUPPORTS_LEGACY=1"
             });
-
-		bDisableAutoRTFMInstrumentation = true;
 	}
 }
