Ensure that scripts that fail to compile because of missing function implementations (like the deleted native fast path stuff) will invalidate the compile results.  This avoids trying to optimize the byte code which would fail.
#tests Try to compile old assets that still referenced the native fast path logic

diff --git a/Engine/Plugins/FX/Niagara/Source/NiagaraEditor/Private/NiagaraCompiler.cpp b/Engine/Plugins/FX/Niagara/Source/NiagaraEditor/Private/NiagaraCompiler.cpp
--- a/Engine/Plugins/FX/Niagara/Source/NiagaraEditor/Private/NiagaraCompiler.cpp
+++ b/Engine/Plugins/FX/Niagara/Source/NiagaraEditor/Private/NiagaraCompiler.cpp
@@ -2249,18 +2249,7 @@ TOptional<FNiagaraCompileResults> FHlslNiagaraCompiler::GetCompileResult(int32 J
 		DumpHLSLText(Results.Data->LastHlslTranslation, CompilationJob->CompileResults.DumpDebugInfoPath);
 	}
 
-	if (!Results.bVMSucceeded)
-	{
-		//For now we just copy the shader code over into the script. 
-		Results.Data->ByteCode.Reset();
-		Results.Data->Attributes.Empty();
-		Results.Data->Parameters.Empty();
-		Results.Data->InternalParameters.Empty();
-		Results.Data->DataInterfaceInfo.Empty();
-		Results.Data->UObjectInfos.Empty();
-		//Eventually Niagara will have all the shader plumbing and do things like materials.
-	}
-	else
+	if (Results.bVMSucceeded)
 	{
 			//Build internal parameters
 		SCOPE_CYCLE_COUNTER(STAT_NiagaraEditor_HlslCompiler_CompileShader_VectorVMSucceeded);
@@ -2370,6 +2359,17 @@ TOptional<FNiagaraCompileResults> FHlslNiagaraCompiler::GetCompileResult(int32 J
 			}
 		}
 	}
+
+	if (!Results.bVMSucceeded)
+	{
+		Results.Data->ByteCode.Reset();
+		Results.Data->Attributes.Empty();
+		Results.Data->Parameters.Empty();
+		Results.Data->InternalParameters.Empty();
+		Results.Data->DataInterfaceInfo.Empty();
+		Results.Data->UObjectInfos.Empty();
+	}
+
 	DumpDebugInfo(CompileResults, CompilationJob->ShaderCompileJob->Input, false);
 
 	//Seems like Results is a bit of a cobbled together mess at this point.
