Improve Slate render time for large systems by wrapping the overview graph node contents in invalidation panels

#tests Tested basic interactions with the emitter nodes like hovering, context menus and changing module inputs to generate stack issues

diff --git a/Engine/Plugins/FX/Niagara/Source/NiagaraEditorWidgets/Private/SNiagaraOverviewStackNode.cpp b/Engine/Plugins/FX/Niagara/Source/NiagaraEditorWidgets/Private/SNiagaraOverviewStackNode.cpp
--- a/Engine/Plugins/FX/Niagara/Source/NiagaraEditorWidgets/Private/SNiagaraOverviewStackNode.cpp
+++ b/Engine/Plugins/FX/Niagara/Source/NiagaraEditorWidgets/Private/SNiagaraOverviewStackNode.cpp
@@ -666,7 +666,12 @@ TSharedRef<SWidget> SNiagaraOverviewStackNode::CreateNodeContentArea_Default()
 		.BorderBackgroundColor(TAttribute<FSlateColor>(this, &SNiagaraOverviewStackNode::GetScalabilityTintAlpha))
 	];
 
-	return ScalabilityWrapper.ToSharedRef();
+	ContentAreaWidget = SNew(SInvalidationPanel)
+	[
+		ScalabilityWrapper.ToSharedRef()
+	];
+
+	return ContentAreaWidget.ToSharedRef();
 }
 
 TSharedRef<SWidget> SNiagaraOverviewStackNode::CreateTopContentBar_Default()
diff --git a/Engine/Plugins/FX/Niagara/Source/NiagaraEditorWidgets/Private/SNiagaraOverviewStackNode.h b/Engine/Plugins/FX/Niagara/Source/NiagaraEditorWidgets/Private/SNiagaraOverviewStackNode.h
--- a/Engine/Plugins/FX/Niagara/Source/NiagaraEditorWidgets/Private/SNiagaraOverviewStackNode.h
+++ b/Engine/Plugins/FX/Niagara/Source/NiagaraEditorWidgets/Private/SNiagaraOverviewStackNode.h
@@ -4,6 +4,7 @@
 
 #include "SGraphNode.h"
 #include "Widgets/SBoxPanel.h"
+#include "Widgets/SInvalidationPanel.h"
 #include "SNiagaraOverviewStack.h"
 #include "ViewModels/NiagaraSystemScalabilityViewModel.h"
 
@@ -137,6 +138,7 @@ private:
 	TWeakObjectPtr<UNiagaraSystemScalabilityViewModel> ScalabilityViewModel;
 	TWeakPtr<FNiagaraEmitterHandleViewModel> EmitterHandleViewModelWeak;
 	TSharedPtr<SOverlay> ScalabilityWrapper;
+	TSharedPtr<SInvalidationPanel> ContentAreaWidget;
 
 	TArray<UNiagaraStackEntry*> RendererPreviewStackEntries;
 	bool bIsHoveringThumbnail = false;
