- Change default custom sorting binding to Particles.Age not Particles.NormalizedAge
- NormalizedAge with random lifetime will introduce flickering as particles die and Age does not

#tests Open existing assets to ensure binding is not changed, create new renderers and test sorting

diff --git a/Engine/Plugins/FX/Niagara/Source/Niagara/Classes/NiagaraConstants.h b/Engine/Plugins/FX/Niagara/Source/Niagara/Classes/NiagaraConstants.h
--- a/Engine/Plugins/FX/Niagara/Source/Niagara/Classes/NiagaraConstants.h
+++ b/Engine/Plugins/FX/Niagara/Source/Niagara/Classes/NiagaraConstants.h
@@ -107,6 +107,7 @@
 #define SYS_PARAM_PARTICLES_VELOCITY                     INiagaraModule::GetVar_Particles_Velocity()
 #define SYS_PARAM_PARTICLES_COLOR                        INiagaraModule::GetVar_Particles_Color()
 #define SYS_PARAM_PARTICLES_SPRITE_ROTATION              INiagaraModule::GetVar_Particles_SpriteRotation()
+#define SYS_PARAM_PARTICLES_AGE                          INiagaraModule::GetVar_Particles_Age()
 #define SYS_PARAM_PARTICLES_NORMALIZED_AGE               INiagaraModule::GetVar_Particles_NormalizedAge()
 #define SYS_PARAM_PARTICLES_SPRITE_SIZE                  INiagaraModule::GetVar_Particles_SpriteSize()
 #define SYS_PARAM_PARTICLES_SPRITE_FACING                INiagaraModule::GetVar_Particles_SpriteFacing()
diff --git a/Engine/Plugins/FX/Niagara/Source/Niagara/Private/NiagaraMeshRendererProperties.cpp b/Engine/Plugins/FX/Niagara/Source/Niagara/Private/NiagaraMeshRendererProperties.cpp
--- a/Engine/Plugins/FX/Niagara/Source/Niagara/Private/NiagaraMeshRendererProperties.cpp
+++ b/Engine/Plugins/FX/Niagara/Source/Niagara/Private/NiagaraMeshRendererProperties.cpp
@@ -462,6 +462,10 @@ void UNiagaraMeshRendererProperties::Serialize(FArchive& Ar)
 		{
 			bSubImageBlend = false;
 		}
+		if (NiagaraVersion < FNiagaraCustomVersion::CustomSortingBindingToAge)
+		{
+			CustomSortingBinding = FNiagaraConstants::GetAttributeDefaultBinding(SYS_PARAM_PARTICLES_NORMALIZED_AGE);
+		}
 	}
 
 	Super::Serialize(Ar);
@@ -511,7 +515,7 @@ void UNiagaraMeshRendererProperties::InitBindings()
 		MeshIndexBinding = FNiagaraConstants::GetAttributeDefaultBinding(SYS_PARAM_PARTICLES_MESH_INDEX);
 
 		//Default custom sorting to age
-		CustomSortingBinding = FNiagaraConstants::GetAttributeDefaultBinding(SYS_PARAM_PARTICLES_NORMALIZED_AGE);
+		CustomSortingBinding = FNiagaraConstants::GetAttributeDefaultBinding(SYS_PARAM_PARTICLES_AGE);
 
 		// Initialize the array with a single, defaulted entry
 		Meshes.AddDefaulted();
diff --git a/Engine/Plugins/FX/Niagara/Source/Niagara/Private/NiagaraModule.cpp b/Engine/Plugins/FX/Niagara/Source/Niagara/Private/NiagaraModule.cpp
--- a/Engine/Plugins/FX/Niagara/Source/Niagara/Private/NiagaraModule.cpp
+++ b/Engine/Plugins/FX/Niagara/Source/Niagara/Private/NiagaraModule.cpp
@@ -270,6 +270,7 @@ FNiagaraVariable INiagaraModule::Particles_Position;
 FNiagaraVariable INiagaraModule::Particles_Velocity;
 FNiagaraVariable INiagaraModule::Particles_Color;
 FNiagaraVariable INiagaraModule::Particles_SpriteRotation;
+FNiagaraVariable INiagaraModule::Particles_Age;
 FNiagaraVariable INiagaraModule::Particles_NormalizedAge;
 FNiagaraVariable INiagaraModule::Particles_SpriteSize;
 FNiagaraVariable INiagaraModule::Particles_SpriteFacing;
@@ -434,6 +435,7 @@ void INiagaraModule::StartupModule()
 	Particles_Velocity = FNiagaraVariable(FNiagaraTypeDefinition::GetVec3Def(), TEXT("Particles.Velocity"));
 	Particles_Color = FNiagaraVariable(FNiagaraTypeDefinition::GetColorDef(), TEXT("Particles.Color"));
 	Particles_SpriteRotation = FNiagaraVariable(FNiagaraTypeDefinition::GetFloatDef(), TEXT("Particles.SpriteRotation"));
+	Particles_Age = FNiagaraVariable(FNiagaraTypeDefinition::GetFloatDef(), TEXT("Particles.Age"));
 	Particles_NormalizedAge = FNiagaraVariable(FNiagaraTypeDefinition::GetFloatDef(), TEXT("Particles.NormalizedAge"));
 	Particles_SpriteSize = FNiagaraVariable(FNiagaraTypeDefinition::GetVec2Def(), TEXT("Particles.SpriteSize"));
 	Particles_SpriteFacing = FNiagaraVariable(FNiagaraTypeDefinition::GetVec3Def(), TEXT("Particles.SpriteFacing"));
diff --git a/Engine/Plugins/FX/Niagara/Source/Niagara/Private/NiagaraSpriteRendererProperties.cpp b/Engine/Plugins/FX/Niagara/Source/Niagara/Private/NiagaraSpriteRendererProperties.cpp
--- a/Engine/Plugins/FX/Niagara/Source/Niagara/Private/NiagaraSpriteRendererProperties.cpp
+++ b/Engine/Plugins/FX/Niagara/Source/Niagara/Private/NiagaraSpriteRendererProperties.cpp
@@ -230,6 +230,10 @@ void UNiagaraSpriteRendererProperties::Serialize(FStructuredArchive::FRecord Rec
 		{
 			bSubImageBlend = false;
 		}
+		if (NiagaraVersion < FNiagaraCustomVersion::CustomSortingBindingToAge)
+		{
+			CustomSortingBinding = FNiagaraConstants::GetAttributeDefaultBinding(SYS_PARAM_PARTICLES_NORMALIZED_AGE);
+		}
 	}
 
 	// MIC will replace the main material during serialize
@@ -304,7 +308,7 @@ void UNiagaraSpriteRendererProperties::InitBindings()
 		RendererVisibilityTagBinding = FNiagaraConstants::GetAttributeDefaultBinding(SYS_PARAM_PARTICLES_VISIBILITY_TAG);
 
 		//Default custom sorting to age
-		CustomSortingBinding = FNiagaraConstants::GetAttributeDefaultBinding(SYS_PARAM_PARTICLES_NORMALIZED_AGE);
+		CustomSortingBinding = FNiagaraConstants::GetAttributeDefaultBinding(SYS_PARAM_PARTICLES_AGE);
 	}
 
 	SetPreviousBindings(FVersionedNiagaraEmitter(), SourceMode);
diff --git a/Engine/Plugins/FX/Niagara/Source/Niagara/Public/NiagaraModule.h b/Engine/Plugins/FX/Niagara/Source/Niagara/Public/NiagaraModule.h
--- a/Engine/Plugins/FX/Niagara/Source/Niagara/Public/NiagaraModule.h
+++ b/Engine/Plugins/FX/Niagara/Source/Niagara/Public/NiagaraModule.h
@@ -230,6 +230,7 @@ public:
 	FORCEINLINE static const FNiagaraVariable&  GetVar_Particles_Velocity() { return Particles_Velocity; }
 	FORCEINLINE static const FNiagaraVariable&  GetVar_Particles_Color() { return Particles_Color; }
 	FORCEINLINE static const FNiagaraVariable&  GetVar_Particles_SpriteRotation() { return Particles_SpriteRotation; }
+	FORCEINLINE static const FNiagaraVariable&  GetVar_Particles_Age() { return Particles_Age; }
 	FORCEINLINE static const FNiagaraVariable&  GetVar_Particles_NormalizedAge() { return Particles_NormalizedAge; }
 	FORCEINLINE static const FNiagaraVariable&  GetVar_Particles_SpriteSize() { return Particles_SpriteSize; }
 	FORCEINLINE static const FNiagaraVariable&  GetVar_Particles_SpriteFacing() { return Particles_SpriteFacing; }
@@ -370,6 +371,7 @@ private:
 	static NIAGARA_API FNiagaraVariable Particles_Velocity;
 	static NIAGARA_API FNiagaraVariable Particles_Color;
 	static NIAGARA_API FNiagaraVariable Particles_SpriteRotation;
+	static NIAGARA_API FNiagaraVariable Particles_Age;
 	static NIAGARA_API FNiagaraVariable Particles_NormalizedAge;
 	static NIAGARA_API FNiagaraVariable Particles_SpriteSize;
 	static NIAGARA_API FNiagaraVariable Particles_SpriteFacing;
diff --git a/Engine/Plugins/FX/Niagara/Source/NiagaraCore/Public/NiagaraCustomVersion.h b/Engine/Plugins/FX/Niagara/Source/NiagaraCore/Public/NiagaraCustomVersion.h
--- a/Engine/Plugins/FX/Niagara/Source/NiagaraCore/Public/NiagaraCustomVersion.h
+++ b/Engine/Plugins/FX/Niagara/Source/NiagaraCore/Public/NiagaraCustomVersion.h
@@ -202,6 +202,8 @@ struct FNiagaraCustomVersion
 		InheritanceUxRefactor, // Decoupling of 'Template' and 'Inheritance'
 
 		NDCSpawnGroupOverrideDisabledByDefault, // NDC Read DIs will not override spawn group by default when spawning particles. Old content will remain unchanged.
+
+		CustomSortingBindingToAge, // Before it was normalized age which can introduce flickering with sorting and random lifetimes
 		
 		// DO NOT ADD A NEW VERSION UNLESS YOU HAVE TALKED TO THE NIAGARA LEAD. Mismanagement of these versions can lead to data loss if it is adjusted in multiple streams simultaneously.
 		// -----<new versions can be added above this line>  -------------------------------------------------
