- Optional debug data for simulation caches
- When enabled a separate object is captured which contains information useful for debugging issues
- Currently only override parameters are captured per frame, this can be extended as needed

diff --git a/Engine/Plugins/FX/Niagara/Source/Niagara/Classes/NiagaraSimCache.h b/Engine/Plugins/FX/Niagara/Source/Niagara/Classes/NiagaraSimCache.h
--- a/Engine/Plugins/FX/Niagara/Source/Niagara/Classes/NiagaraSimCache.h
+++ b/Engine/Plugins/FX/Niagara/Source/Niagara/Classes/NiagaraSimCache.h
@@ -11,6 +11,7 @@
 #include "NiagaraSimCache.generated.h"
 
 class UNiagaraComponent;
+class UNiagaraSimCacheDebugData;
 struct FNiagaraSimCacheDataBuffersLayout;
 
 UENUM(BlueprintType)
@@ -85,6 +86,12 @@ struct FNiagaraSimCacheCreateParameters
 	UPROPERTY(EditAnywhere, BlueprintReadWrite, Category = "SimCache")
 	uint32 bAllowSerializeLargeCache : 1 = true;
 
+	/**
+	When enabled additional information is stored that can be useful for debugging a simulation
+	*/
+	UPROPERTY(EditAnywhere, BlueprintReadWrite, Category = "SimCache")
+	uint32 bIncludeDebugData : 1 = false;
+
 	/**
 	List of Attributes to force include in the SimCache rebase, they should be the full path to the attribute
 	For example, MyEmitter.Particles.MyQuat would force the particle attribute MyQuat to be included for MyEmitter
@@ -119,6 +126,9 @@ struct FNiagaraSimCacheCreateParameters
 	*/
 	UPROPERTY(EditAnywhere, BlueprintReadWrite, Category = "SimCache", meta=(EditCondition="AttributeCaptureMode == ENiagaraSimCacheAttributeCaptureMode::ExplicitAttributes"))
 	TArray<FName> ExplicitCaptureAttributes;
+
+	// Helper function to setup the parameters for capturing a cache for debugging
+	NIAGARA_API static FNiagaraSimCacheCreateParameters CreateForDebugging();
 };
 
 USTRUCT()
@@ -462,6 +472,9 @@ public:
 	/** Returns the actual data we have captured in the SimCache for the given data interface. */
 	NIAGARA_API const UObject* GetDataInterfaceStorageObject(const FNiagaraVariableBase& DataInterface) const;
 
+	/** Returns the debug data stored inside the cache, intended for internal use only. */
+	NIAGARA_API const UNiagaraSimCacheDebugData* GetDebugData() const { return DebugData; }
+
 	/**
 	Get number of active instances for the emitter at the given frame.
 	An EmitterIndex or INDEX_NONE will return information about the system instance.
@@ -610,6 +623,9 @@ private:
 	UPROPERTY()
 	TMap<FNiagaraVariableBase, TObjectPtr<UObject>> DataInterfaceStorage;
 
+	UPROPERTY()
+	TObjectPtr<UNiagaraSimCacheDebugData> DebugData;
+
 	int32 CaptureTickCount = INDEX_NONE;
 	double CaptureStartTime = 0;
 
diff --git a/Engine/Plugins/FX/Niagara/Source/Niagara/Internal/NiagaraSimCacheDebugData.h b/Engine/Plugins/FX/Niagara/Source/Niagara/Internal/NiagaraSimCacheDebugData.h
new file mode 100644
--- /dev/null
+++ b/Engine/Plugins/FX/Niagara/Source/Niagara/Internal/NiagaraSimCacheDebugData.h
@@ -0,0 +1,33 @@
+// Copyright Epic Games, Inc. All Rights Reserved.
+
+#pragma once
+
+#include "NiagaraCommon.h"
+#include "NiagaraScript.h"
+
+#include "NiagaraSimCacheDebugData.generated.h"
+
+struct FNiagaraSimCacheHelper;
+
+USTRUCT()
+struct FNiagaraSimCacheDebugDataFrame
+{
+	GENERATED_BODY()
+
+	UPROPERTY()
+	FNiagaraParameterStore OverrideParameters;
+};
+
+// Contains data useful for debugging a Niagara system
+UCLASS(MinimalAPI)
+class UNiagaraSimCacheDebugData : public UObject
+{
+	GENERATED_UCLASS_BODY()
+
+public:
+	void CaptureFrame(FNiagaraSimCacheHelper& Helper, int FrameNumber);
+
+public:
+	UPROPERTY()
+	TArray<FNiagaraSimCacheDebugDataFrame> Frames;
+};
diff --git a/Engine/Plugins/FX/Niagara/Source/Niagara/Private/DataInterface/NiagaraDataInterfaceMemoryBuffer.cpp b/Engine/Plugins/FX/Niagara/Source/Niagara/Private/DataInterface/NiagaraDataInterfaceMemoryBuffer.cpp
--- a/Engine/Plugins/FX/Niagara/Source/Niagara/Private/DataInterface/NiagaraDataInterfaceMemoryBuffer.cpp
+++ b/Engine/Plugins/FX/Niagara/Source/Niagara/Private/DataInterface/NiagaraDataInterfaceMemoryBuffer.cpp
@@ -949,8 +949,7 @@ void UNiagaraDataInterfaceMemoryBuffer::SetShaderParameters(const FNiagaraDataIn
 
 UObject* UNiagaraDataInterfaceMemoryBuffer::SimCacheBeginWrite(UObject* SimCache, FNiagaraSystemInstance* NiagaraSystemInstance, const void* OptionalPerInstanceData, FNiagaraSimCacheFeedbackContext& FeedbackContext) const
 {
-	UNDIMemoryBufferSimCacheData* CacheData = NewObject<UNDIMemoryBufferSimCacheData>(SimCache);
-	return CacheData;
+	return OptionalPerInstanceData ? NewObject<UNDIMemoryBufferSimCacheData>(SimCache) : nullptr;
 }
 
 bool UNiagaraDataInterfaceMemoryBuffer::SimCacheWriteFrame(UObject* StorageObject, int FrameIndex, FNiagaraSystemInstance* SystemInstance, const void* OptionalPerInstanceData, FNiagaraSimCacheFeedbackContext& FeedbackContext) const
diff --git a/Engine/Plugins/FX/Niagara/Source/Niagara/Private/DataInterface/NiagaraDataInterfaceSimpleCounter.cpp b/Engine/Plugins/FX/Niagara/Source/Niagara/Private/DataInterface/NiagaraDataInterfaceSimpleCounter.cpp
--- a/Engine/Plugins/FX/Niagara/Source/Niagara/Private/DataInterface/NiagaraDataInterfaceSimpleCounter.cpp
+++ b/Engine/Plugins/FX/Niagara/Source/Niagara/Private/DataInterface/NiagaraDataInterfaceSimpleCounter.cpp
@@ -479,8 +479,7 @@ void UNiagaraDataInterfaceSimpleCounter::PushToRenderThreadImpl()
 
 UObject* UNiagaraDataInterfaceSimpleCounter::SimCacheBeginWrite(UObject* SimCache, FNiagaraSystemInstance* NiagaraSystemInstance, const void* OptionalPerInstanceData, FNiagaraSimCacheFeedbackContext& FeedbackContext) const
 {
-	UNDISimpleCounterSimCacheData* CacheData = NewObject<UNDISimpleCounterSimCacheData>(SimCache);
-	return CacheData;
+	return OptionalPerInstanceData ? NewObject<UNDISimpleCounterSimCacheData>(SimCache) : nullptr;
 }
 
 bool UNiagaraDataInterfaceSimpleCounter::SimCacheWriteFrame(UObject* StorageObject, int FrameIndex, FNiagaraSystemInstance* SystemInstance, const void* OptionalPerInstanceData, FNiagaraSimCacheFeedbackContext& FeedbackContext) const
diff --git a/Engine/Plugins/FX/Niagara/Source/Niagara/Private/NiagaraDataInterfaceArray.cpp b/Engine/Plugins/FX/Niagara/Source/Niagara/Private/NiagaraDataInterfaceArray.cpp
--- a/Engine/Plugins/FX/Niagara/Source/Niagara/Private/NiagaraDataInterfaceArray.cpp
+++ b/Engine/Plugins/FX/Niagara/Source/Niagara/Private/NiagaraDataInterfaceArray.cpp
@@ -96,8 +96,7 @@ void UNiagaraDataInterfaceArray::SetShaderParameters(const FNiagaraDataInterface
 
 UObject* UNiagaraDataInterfaceArray::SimCacheBeginWrite(UObject* SimCache, FNiagaraSystemInstance* NiagaraSystemInstance, const void* OptionalPerInstanceData, FNiagaraSimCacheFeedbackContext& FeedbackContext) const
 {
-	UNDIArraySimCacheData* CacheData = NewObject<UNDIArraySimCacheData>(SimCache);
-	return CacheData;
+	return OptionalPerInstanceData ? NewObject<UNDIArraySimCacheData>(SimCache) : nullptr;
 }
 
 bool UNiagaraDataInterfaceArray::SimCacheWriteFrame(UObject* StorageObject, int FrameIndex, FNiagaraSystemInstance* SystemInstance, const void* OptionalPerInstanceData, FNiagaraSimCacheFeedbackContext& FeedbackContext) const
diff --git a/Engine/Plugins/FX/Niagara/Source/Niagara/Private/NiagaraDataInterfaceRenderTarget2D.cpp b/Engine/Plugins/FX/Niagara/Source/Niagara/Private/NiagaraDataInterfaceRenderTarget2D.cpp
--- a/Engine/Plugins/FX/Niagara/Source/Niagara/Private/NiagaraDataInterfaceRenderTarget2D.cpp
+++ b/Engine/Plugins/FX/Niagara/Source/Niagara/Private/NiagaraDataInterfaceRenderTarget2D.cpp
@@ -578,10 +578,14 @@ bool UNiagaraDataInterfaceRenderTarget2D::RenderVariableToCanvas(FNiagaraSystemI
 
 NIAGARA_API UObject* UNiagaraDataInterfaceRenderTarget2D::SimCacheBeginWrite(UObject* InSimCache, FNiagaraSystemInstance* NiagaraSystemInstance, const void* OptionalPerInstanceData, FNiagaraSimCacheFeedbackContext& FeedbackContext) const
 {
+	if (!OptionalPerInstanceData)
+	{
+		return nullptr;
+	}
+
 	UNDIRenderTargetVolumeSimCacheData* SimCacheData = nullptr;
 	SimCacheData = NewObject<UNDIRenderTargetVolumeSimCacheData>(InSimCache);
 	SimCacheData->CompressionType = NDIRenderTarget2DLocal::GetSimCacheCompressionType();
-
 	return SimCacheData;	
 }
 
diff --git a/Engine/Plugins/FX/Niagara/Source/Niagara/Private/NiagaraDataInterfaceRenderTargetVolume.cpp b/Engine/Plugins/FX/Niagara/Source/Niagara/Private/NiagaraDataInterfaceRenderTargetVolume.cpp
--- a/Engine/Plugins/FX/Niagara/Source/Niagara/Private/NiagaraDataInterfaceRenderTargetVolume.cpp
+++ b/Engine/Plugins/FX/Niagara/Source/Niagara/Private/NiagaraDataInterfaceRenderTargetVolume.cpp
@@ -586,6 +586,11 @@ bool UNiagaraDataInterfaceRenderTargetVolume::GetExposedVariableValue(const FNia
 
 UObject* UNiagaraDataInterfaceRenderTargetVolume::SimCacheBeginWrite(UObject* InSimCache, FNiagaraSystemInstance* NiagaraSystemInstance, const void* OptionalPerInstanceData, FNiagaraSimCacheFeedbackContext& FeedbackContext) const
 {
+	if (!OptionalPerInstanceData)
+	{
+		return nullptr;
+	}
+
 	if (NDIRenderTargetVolumeLocal::GSimCacheDataStorageMode == NDIRenderTargetVolumeLocal::SimCacheStorageMode::SVT)
 	{		
 		UNiagaraSimCache* SimCache = CastChecked<UNiagaraSimCache>(InSimCache);
diff --git a/Engine/Plugins/FX/Niagara/Source/Niagara/Private/NiagaraDebuggerClient.cpp b/Engine/Plugins/FX/Niagara/Source/Niagara/Private/NiagaraDebuggerClient.cpp
--- a/Engine/Plugins/FX/Niagara/Source/Niagara/Private/NiagaraDebuggerClient.cpp
+++ b/Engine/Plugins/FX/Niagara/Source/Niagara/Private/NiagaraDebuggerClient.cpp
@@ -560,7 +560,7 @@ bool FNiagaraSimCacheCaptureInfo::Process()
 		//First Frame of Capture, init the sim cache.
 		if (ProcessedFrames == Request.CaptureDelayFrames)
 		{	
-			FNiagaraSimCacheCreateParameters CacheCreateParameters;
+			FNiagaraSimCacheCreateParameters CacheCreateParameters = FNiagaraSimCacheCreateParameters::CreateForDebugging();
 			CacheCreateParameters.bAllowSerializeLargeCache = false;
 
 			SimCache.Reset(NewObject<UNiagaraSimCache>(GetTransientPackage()));
diff --git a/Engine/Plugins/FX/Niagara/Source/Niagara/Private/NiagaraSimCache.cpp b/Engine/Plugins/FX/Niagara/Source/Niagara/Private/NiagaraSimCache.cpp
--- a/Engine/Plugins/FX/Niagara/Source/Niagara/Private/NiagaraSimCache.cpp
+++ b/Engine/Plugins/FX/Niagara/Source/Niagara/Private/NiagaraSimCache.cpp
@@ -15,6 +15,7 @@
 #include "NiagaraGpuComputeDispatchInterface.h"
 #include "NiagaraSimCacheAttributeReaderHelper.h"
 #include "NiagaraSimCacheCustomStorageInterface.h"
+#include "NiagaraSimCacheDebugData.h"
 #include "NiagaraSimCacheHelper.h"
 #include "NiagaraSystemImpl.h"
 #include "NiagaraSystemInstance.h"
@@ -31,6 +32,15 @@ namespace FNiagaraSimCacheInternal
 
 ////////////////////////////////////////////////////////////////////////////////////////////////////////////////////////////////////////////////////
 
+FNiagaraSimCacheCreateParameters FNiagaraSimCacheCreateParameters::CreateForDebugging()
+{
+	FNiagaraSimCacheCreateParameters CreateParameters;
+	CreateParameters.bIncludeDebugData = true;
+	return CreateParameters;
+}
+
+////////////////////////////////////////////////////////////////////////////////////////////////////////////////////////////////////////////////////
+
 bool FNiagaraSimCacheVariable::operator==(const FNiagaraSimCacheVariable& Other) const
 {
 	return
@@ -339,6 +349,9 @@ bool UNiagaraSimCache::BeginWrite(FNiagaraSimCacheCreateParameters InCreateParam
 	}
 	DataInterfaceStorage.Empty();
 
+	// Create our debug data (if required)
+	DebugData = CreateParameters.bIncludeDebugData ? NewObject<UNiagaraSimCacheDebugData>(this) : nullptr;
+
 	// Not explicit capture mode?  Empty the list as internally we reuse this list when in rendering only mode
 	if (CreateParameters.AttributeCaptureMode == ENiagaraSimCacheAttributeCaptureMode::ExplicitAttributes)
 	{
@@ -548,6 +561,9 @@ bool UNiagaraSimCache::BeginAppend(FNiagaraSimCacheCreateParameters InCreatePara
 		}
 	}
 
+	// We don't support debug caches when appending
+	DebugData = nullptr;
+
 	if (CacheFrames.Num() == 0)
 	{
 		StartSeconds = 0.0f;
@@ -684,6 +700,12 @@ bool UNiagaraSimCache::WriteFrame(UNiagaraComponent* NiagaraComponent, FNiagaraS
 		}
 	}
 
+	// If we have debug data time to capture that also
+	if (DebugData)
+	{
+		DebugData->CaptureFrame(Helper, CacheFrames.Num() - 1);
+	}
+
 	// Store data interface data
 	//-OPT: We shouldn't need to search all the time here
 	if (DataInterfaceStorage.IsEmpty() == false)
diff --git a/Engine/Plugins/FX/Niagara/Source/Niagara/Private/NiagaraSimCacheDebugData.cpp b/Engine/Plugins/FX/Niagara/Source/Niagara/Private/NiagaraSimCacheDebugData.cpp
new file mode 100644
--- /dev/null
+++ b/Engine/Plugins/FX/Niagara/Source/Niagara/Private/NiagaraSimCacheDebugData.cpp
@@ -0,0 +1,17 @@
+// Copyright Epic Games, Inc. All Rights Reserved.
+
+#include "NiagaraSimCacheDebugData.h"
+#include "NiagaraSimCacheHelper.h"
+
+UNiagaraSimCacheDebugData::UNiagaraSimCacheDebugData(const FObjectInitializer& ObjectInitializer)
+{
+}
+
+void UNiagaraSimCacheDebugData::CaptureFrame(FNiagaraSimCacheHelper& Helper, int FrameNumber)
+{
+   	Frames.SetNum(FrameNumber + 1);
+	if (const FNiagaraParameterStore* ParameterStore = Helper.SystemInstance->GetOverrideParameters())
+	{
+		Frames[FrameNumber].OverrideParameters = *ParameterStore;
+	}
+}
diff --git a/Engine/Plugins/FX/Niagara/Source/NiagaraEditor/Private/Customizations/NiagaraComponentDetails.cpp b/Engine/Plugins/FX/Niagara/Source/NiagaraEditor/Private/Customizations/NiagaraComponentDetails.cpp
--- a/Engine/Plugins/FX/Niagara/Source/NiagaraEditor/Private/Customizations/NiagaraComponentDetails.cpp
+++ b/Engine/Plugins/FX/Niagara/Source/NiagaraEditor/Private/Customizations/NiagaraComponentDetails.cpp
@@ -361,8 +361,7 @@ FReply FNiagaraComponentDetails::OnCaptureSelectedSystem()
 		SimCache->SetFlags(RF_Transient);
 		CapturedCaches.Emplace(SimCache);
 
-		FNiagaraSimCacheCreateParameters CreateParameters;
-		
+	
 		FNiagaraSimCacheCaptureParameters CaptureParameters;
 
 		CaptureParameters.CaptureRate = 1;
@@ -380,7 +379,7 @@ FReply FNiagaraComponentDetails::OnCaptureSelectedSystem()
 					GEditor->GetEditorSubsystem<UAssetEditorSubsystem>()->OpenEditorForAsset(CapturedSimCache, EToolkitMode::Standalone);
 				}
 			});
-			SimCacheCapture->CaptureNiagaraSimCache(SimCache, CreateParameters, NiagaraComponentToUse, CaptureParameters);
+			SimCacheCapture->CaptureNiagaraSimCache(SimCache, FNiagaraSimCacheCreateParameters::CreateForDebugging(), NiagaraComponentToUse, CaptureParameters);
 		}
 
 		
diff --git a/Engine/Plugins/FX/Niagara/Source/NiagaraEditor/Private/NiagaraEditorStyle.cpp b/Engine/Plugins/FX/Niagara/Source/NiagaraEditor/Private/NiagaraEditorStyle.cpp
--- a/Engine/Plugins/FX/Niagara/Source/NiagaraEditor/Private/NiagaraEditorStyle.cpp
+++ b/Engine/Plugins/FX/Niagara/Source/NiagaraEditor/Private/NiagaraEditorStyle.cpp
@@ -907,6 +907,9 @@ void FNiagaraEditorStyle::InitSimCacheEditor()
 	FTableRowStyle DataInterfaceStyle = EmitterStyle;
 	Set("NiagaraEditor.SimCache.DataInterfaceItem", DataInterfaceStyle);
 
+	FTableRowStyle DebugDataStyle = EmitterStyle;
+	Set("NiagaraEditor.SimCache.DebugData", DebugDataStyle);
+
 	const FLinearColor NormalColor(0.15, 0.15, 0.15, 1);
 
 	Set ("NiagaraEditor.SimCache.FilterToggleStyle", FCheckBoxStyle()
diff --git a/Engine/Plugins/FX/Niagara/Source/NiagaraEditor/Private/ViewModels/NiagaraSimCacheViewModel.cpp b/Engine/Plugins/FX/Niagara/Source/NiagaraEditor/Private/ViewModels/NiagaraSimCacheViewModel.cpp
--- a/Engine/Plugins/FX/Niagara/Source/NiagaraEditor/Private/ViewModels/NiagaraSimCacheViewModel.cpp
+++ b/Engine/Plugins/FX/Niagara/Source/NiagaraEditor/Private/ViewModels/NiagaraSimCacheViewModel.cpp
@@ -74,8 +74,20 @@ void FNiagaraSimCacheViewModel::SetupPreviewComponentAndInstance()
 	}
 }
 
-TConstArrayView<FNiagaraSimCacheViewModel::FComponentInfo> FNiagaraSimCacheViewModel::GetComponentInfos(
-	int32 InEmitterIndex) const
+TConstArrayView<FNiagaraSimCacheViewModel::FComponentInfo> FNiagaraSimCacheViewModel::GetSelectedComponentInfos() const
+{
+	if (SelectionMode == ESelectionMode::SystemInstance)
+	{
+		return GetComponentInfos(INDEX_NONE);
+	}
+	else if (SelectionMode == ESelectionMode::Emitter)
+	{
+		return GetComponentInfos(SelectedEmitterIndex);
+	}
+	return MakeArrayView<FComponentInfo>(nullptr, 0);
+}
+
+TConstArrayView<FNiagaraSimCacheViewModel::FComponentInfo> FNiagaraSimCacheViewModel::GetComponentInfos(int32 InEmitterIndex) const
 {
 	if( InEmitterIndex == INDEX_NONE)
 	{
@@ -92,7 +104,7 @@ TConstArrayView<FNiagaraSimCacheViewModel::FComponentInfo> FNiagaraSimCacheViewM
 
 FText FNiagaraSimCacheViewModel::GetComponentText(const FName ComponentName, const int32 InstanceIndex) const
 {
-	const FComponentInfo* ComponentInfo = GetCurrentComponentInfos().FindByPredicate([ComponentName](const FComponentInfo& FoundInfo) { return FoundInfo.Name == ComponentName; });
+	const FComponentInfo* ComponentInfo = GetSelectedComponentInfos().FindByPredicate([ComponentName](const FComponentInfo& FoundInfo) { return FoundInfo.Name == ComponentName; });
 
 	if (ComponentInfo)
 	{
@@ -155,19 +167,53 @@ void FNiagaraSimCacheViewModel::SetFrameIndex(const int32 InFrameIndex)
 	OnViewDataChangedDelegate.Broadcast(false);
 }
 
-const UObject* FNiagaraSimCacheViewModel::GetActiveDataInterfaceStorage() const
+const UObject* FNiagaraSimCacheViewModel::GetSelectedDataInterfaceStorage() const
 {
-	if (SimCache)
+	if (SelectionMode == ESelectionMode::DataInterface && SimCache)
 	{
-		return SimCache->GetDataInterfaceStorageObject(GetActiveDataInterface());
+		return SimCache->GetDataInterfaceStorageObject(SelectedDataInterface);
 	}
 	return nullptr;
 }
 
-void FNiagaraSimCacheViewModel::SetEmitterIndex(const int32 InEmitterIndex, FNiagaraVariableBase InActiveDataInterface)
+void FNiagaraSimCacheViewModel::SetSelectedSystemInstance()
+{
+	SelectionMode			= ESelectionMode::SystemInstance;
+	SelectedEmitterIndex	= INDEX_NONE;
+	SelectedDataInterface	= FNiagaraVariableBase();
+
+	RefreshFromSelectionChanged();
+}
+
+void FNiagaraSimCacheViewModel::SetSelectedEmitter(int32 EmitterIndex)
+{
+	SelectionMode			= ESelectionMode::Emitter;
+	SelectedEmitterIndex	= EmitterIndex;
+	SelectedDataInterface	= FNiagaraVariableBase();
+
+	RefreshFromSelectionChanged();
+}
+
+void FNiagaraSimCacheViewModel::SetSelectedDataInterface(FNiagaraVariableBase DIVariable)
+{
+	SelectionMode			= ESelectionMode::DataInterface;
+	SelectedEmitterIndex	= INDEX_NONE;
+	SelectedDataInterface	= DIVariable;
+
+	RefreshFromSelectionChanged();
+}
+
+void FNiagaraSimCacheViewModel::SetSelectedDebugData()
+{
+	SelectionMode			= ESelectionMode::DebugData;
+	SelectedEmitterIndex	= INDEX_NONE;
+	SelectedDataInterface	= FNiagaraVariableBase();
+
+	RefreshFromSelectionChanged();
+}
+
+void FNiagaraSimCacheViewModel::RefreshFromSelectionChanged()
 {
-	EmitterIndex = InEmitterIndex;
-	ActiveDataInterface = InActiveDataInterface;
 	UpdateCachedFrame();
 	UpdateCurrentEntries();
 	bComponentFilterActive = false;
@@ -190,6 +236,11 @@ FName FNiagaraSimCacheViewModel::GetEmitterLayoutName(const int32 Index) const
 	return SimCache ? SimCache->GetEmitterName(Index) : NAME_None;
 }
 
+const UNiagaraSimCacheDebugData* FNiagaraSimCacheViewModel::GetCacheDebugData() const
+{
+	return SimCache ? SimCache->GetDebugData() : nullptr;
+}
+
 FNiagaraSimCacheViewModel::FOnViewDataChanged& FNiagaraSimCacheViewModel::OnViewDataChanged()
 {
 	return OnViewDataChangedDelegate;
@@ -237,23 +288,45 @@ void FNiagaraSimCacheViewModel::UpdateCachedFrame()
 		return;
 	}
 
-	if ( EmitterIndex != INDEX_NONE && EmitterIndex >= SimCache->GetNumEmitters() )
+	// Determine if we need to read attributes from the cache
+	TOptional<int32> EmitterIndex;
+	switch (SelectionMode)
 	{
-		return;
+		case ESelectionMode::SystemInstance:
+			NumInstances = 1;
+			EmitterIndex = INDEX_NONE;
+			break;
+
+		case ESelectionMode::Emitter:
+			if (SelectedEmitterIndex >= 0 && SelectedEmitterIndex < SimCache->GetNumEmitters())
+			{
+				EmitterIndex = SelectedEmitterIndex;
+				NumInstances = SimCache->GetEmitterNumInstances(SelectedEmitterIndex, FrameIndex);
+			}
+			break;
+
+		case ESelectionMode::DataInterface:
+		case ESelectionMode::DebugData:
+			NumInstances = 1;
+			break;
 	}
 
-	NumInstances = EmitterIndex == INDEX_NONE ? 1 : SimCache->GetEmitterNumInstances(EmitterIndex, FrameIndex);
-	const FName EmitterName = EmitterIndex == INDEX_NONE ? NAME_None : SimCache->GetEmitterName(EmitterIndex);
+	// Read attributes
+	if (EmitterIndex.IsSet())
+	{
+		const FName EmitterName = EmitterIndex.GetValue() == INDEX_NONE ? NAME_None : SimCache->GetEmitterName(EmitterIndex.GetValue());
 
-	SimCache->ForEachEmitterAttribute(EmitterIndex,
-		[&](const FNiagaraSimCacheVariable& Variable)
-		{
-			// Pull in data
-			SimCache->ReadAttribute(FloatComponents, HalfComponents, Int32Components, Variable.Variable.GetName(), EmitterName, FrameIndex);
+		SimCache->ForEachEmitterAttribute(
+			EmitterIndex.GetValue(),
+			[&](const FNiagaraSimCacheVariable& Variable)
+			{
+				// Pull in data
+				SimCache->ReadAttribute(FloatComponents, HalfComponents, Int32Components, Variable.Variable.GetName(), EmitterName, FrameIndex);
 
-			return true;
-		}
-	);
+				return true;
+			}
+		);
+	}
 }
 
 void FNiagaraSimCacheViewModel::UpdateComponentInfos()
@@ -269,7 +342,6 @@ void FNiagaraSimCacheViewModel::UpdateComponentInfos()
 		return;
 	}
 	
-	
 	SimCache->ForEachEmitterAttribute(INDEX_NONE,
 		[&](const FNiagaraSimCacheVariable& Variable)
 		{
@@ -474,6 +546,18 @@ void FNiagaraSimCacheViewModel::BuildEntries(TWeakPtr<SNiagaraSimCacheTreeView>
 			RootEntries.Add(CurrentDataInterfaceItem);
 			BufferEntries.Add(CurrentDataInterfaceBufferItem);
 		}
+
+		if (SimCache->GetDebugData() != nullptr)
+		{
+			const TSharedRef<FNiagaraSimCacheDebugDataTreeItem> TreeItem = MakeShared<FNiagaraSimCacheDebugDataTreeItem>(OwningTreeView);
+			const TSharedRef<FNiagaraSimCacheOverviewDebugDataItem> DataItem = MakeShared<FNiagaraSimCacheOverviewDebugDataItem>();
+			const FText DisplayNameText = LOCTEXT("DebugData", "Debug Data");
+			TreeItem->SetDisplayName(DisplayNameText);
+			DataItem->SetDisplayName(DisplayNameText);
+
+			RootEntries.Add(TreeItem);
+			BufferEntries.Add(DataItem);
+		}
 	}
 
 	UpdateCurrentEntries();
@@ -481,19 +565,27 @@ void FNiagaraSimCacheViewModel::BuildEntries(TWeakPtr<SNiagaraSimCacheTreeView>
 
 void FNiagaraSimCacheViewModel::UpdateCurrentEntries()
 {
-	// The overview panel is in charge of building these entries.
-	// Early out if they haven't been built yet.
-	if(!RootEntries.IsValidIndex(EmitterIndex +1))
+	SelectedRootEntries.Empty();
+	switch (SelectionMode)
 	{
-		return;
+		case ESelectionMode::SystemInstance:
+			SelectedRootEntries.Add(RootEntries[0]);
+			return;
+		case ESelectionMode::Emitter:
+			if (RootEntries.IsValidIndex(SelectedEmitterIndex + 1))
+			{
+				SelectedRootEntries.Add(RootEntries[SelectedEmitterIndex + 1]);
+			}
+			return;
+		default:
+			//-TODO: DO we need to do something here?
+			break;
 	}
-	CurrentRootEntries.Empty();
-	CurrentRootEntries.Add(RootEntries[EmitterIndex + 1]);
 }
 
-TArray<TSharedRef<FNiagaraSimCacheTreeItem>>* FNiagaraSimCacheViewModel::GetCurrentRootEntries()
+TArray<TSharedRef<FNiagaraSimCacheTreeItem>>* FNiagaraSimCacheViewModel::GetSelectedRootEntries()
 {
-	return &CurrentRootEntries;
+	return &SelectedRootEntries;
 }
 
 TArray<TSharedRef<FNiagaraSimCacheOverviewItem>>* FNiagaraSimCacheViewModel::GetBufferEntries()
@@ -503,7 +595,7 @@ TArray<TSharedRef<FNiagaraSimCacheOverviewItem>>* FNiagaraSimCacheViewModel::Get
 
 bool FNiagaraSimCacheViewModel::CanCopyActiveToClipboard() const
 {
-	return IsCacheValid() && !ActiveDataInterface.IsValid();
+	return IsCacheValid() && SelectionMode == ESelectionMode::DataInterface;
 }
 
 void FNiagaraSimCacheViewModel::CopyActiveToClipboard() const
@@ -514,7 +606,7 @@ void FNiagaraSimCacheViewModel::CopyActiveToClipboard() const
 	}
 
 	FString ClipboardString;
-	TConstArrayView<FComponentInfo> ComponentInfos = GetCurrentComponentInfos();
+	TConstArrayView<FComponentInfo> ComponentInfos = GetSelectedComponentInfos();
 
 	ClipboardString.Append(TEXT("Instance"));
 	for ( int iComponent=0; iComponent < ComponentInfos.Num(); ++iComponent)
diff --git a/Engine/Plugins/FX/Niagara/Source/NiagaraEditor/Private/Widgets/SNiagaraDebugCaptureView.cpp b/Engine/Plugins/FX/Niagara/Source/NiagaraEditor/Private/Widgets/SNiagaraDebugCaptureView.cpp
--- a/Engine/Plugins/FX/Niagara/Source/NiagaraEditor/Private/Widgets/SNiagaraDebugCaptureView.cpp
+++ b/Engine/Plugins/FX/Niagara/Source/NiagaraEditor/Private/Widgets/SNiagaraDebugCaptureView.cpp
@@ -341,14 +341,12 @@ void SNiagaraDebugCaptureView::OnSingleFrameSelected()
 		TSharedPtr<ISequencer> Sequencer = SystemViewModel.Get()->GetSequencer();
 		Sequencer->OnPlay(false);
 
-		const FNiagaraSimCacheCreateParameters CreateParameters;
-
 		FQualifiedFrameTime StartTime = Sequencer->GetGlobalTime();
 		float CurrentAge = TargetComponent->GetDesiredAge();
 
 		UNiagaraSimCache* OutCache = CapturedCache.Get();
 
-		SimCacheCapture.CaptureCurrentFrameImmediate(CapturedCache.Get(), CreateParameters, TargetComponent, OutCache, true, 0.01666f);
+		SimCacheCapture.CaptureCurrentFrameImmediate(CapturedCache.Get(), FNiagaraSimCacheCreateParameters::CreateForDebugging(), TargetComponent, OutCache, true, 0.01666f);
 
 		if(OutCache)
 		{
@@ -372,7 +370,6 @@ void SNiagaraDebugCaptureView::OnMultiFrameSelected()
 	{
 		UNiagaraSimCache* MultiFrameCache = NewObject<UNiagaraSimCache>(GetTransientPackage(), GetTempCacheName(TargetComponent->GetFXSystemAsset()->GetName()));
 		MultiFrameCache->SetFlags(RF_Transient);
-		const FNiagaraSimCacheCreateParameters CreateParameters;
 		
 		FNiagaraSimCacheCaptureParameters CaptureParameters;
 		CaptureParameters.NumFrames = NumFrames;
@@ -383,7 +380,7 @@ void SNiagaraDebugCaptureView::OnMultiFrameSelected()
 
 		SimCacheCapture.OnCaptureComplete().AddSP(this, &SNiagaraDebugCaptureView::OnCaptureComplete);
 		
-		SimCacheCapture.CaptureNiagaraSimCache(MultiFrameCache, CreateParameters, TargetComponent, CaptureParameters);
+		SimCacheCapture.CaptureNiagaraSimCache(MultiFrameCache, FNiagaraSimCacheCreateParameters::CreateForDebugging(), TargetComponent, CaptureParameters);
 	}
 }
 
@@ -404,9 +401,9 @@ void SNiagaraDebugCaptureView::OnCaptureComplete(UNiagaraSimCache* CapturedSimCa
 			{
 				if(SimCacheViewModel->GetEmitterLayoutName(i) == SelectedEmitterHandleViewModel->GetName())
 				{
-					if(SimCacheViewModel->GetEmitterIndex() != i)
+					if(SimCacheViewModel->GetSelectedEmitter() != i)
 					{
-						SimCacheViewModel->SetEmitterIndex(i);
+						SimCacheViewModel->SetSelectedEmitter(i);
 					}
 					break;
 				}
diff --git a/Engine/Plugins/FX/Niagara/Source/NiagaraEditor/Private/Widgets/SNiagaraSimCacheDebugDataView.cpp b/Engine/Plugins/FX/Niagara/Source/NiagaraEditor/Private/Widgets/SNiagaraSimCacheDebugDataView.cpp
new file mode 100644
--- /dev/null
+++ b/Engine/Plugins/FX/Niagara/Source/NiagaraEditor/Private/Widgets/SNiagaraSimCacheDebugDataView.cpp
@@ -0,0 +1,202 @@
+// Copyright Epic Games, Inc. All Rights Reserved.
+
+#include "SNiagaraSimCacheDebugDataView.h"
+#include "ViewModels/NiagaraSimCacheViewModel.h"
+#include "Widgets/SNiagaraParameterName.h"
+#include "NiagaraSimCacheDebugData.h"
+
+#include "Widgets/Layout/SScrollBox.h"
+#include "Widgets/Views/SListView.h"
+
+#define LOCTEXT_NAMESPACE "NiagaraSimCacheDebugDataView"
+
+////////////////////////////////////////////////////////////////////////////////////////////////////////////////////////
+
+namespace UE::Niagara::SimCache::DebugDataUI
+{
+static FName NAME_ParameterName("ParameterName");
+static FName NAME_ParameterValue("ParameterValue");
+
+class SParameterStoreItemWidget : public SMultiColumnTableRow<TSharedPtr<FNiagaraVariableBase>>
+{
+public:
+	SLATE_BEGIN_ARGS(SParameterStoreItemWidget){}
+		SLATE_ARGUMENT(TSharedPtr<FNiagaraParameterStore>,	ParameterStore)
+		SLATE_ARGUMENT(TSharedPtr<FNiagaraVariableBase>,	ParameterVariable)
+	SLATE_END_ARGS()
+
+	void Construct(const FArguments& InArgs, const TSharedRef<STableViewBase>& InOwnerTable)
+	{
+		ParameterStore = InArgs._ParameterStore;
+		ParameterVariable = InArgs._ParameterVariable;
+	
+		SMultiColumnTableRow<TSharedPtr<FNiagaraVariableBase>>::Construct( FSuperRowType::FArguments(), InOwnerTable);
+	}
+
+	TSharedRef<SWidget> GenerateWidgetForColumn(const FName& ColumnName)
+	{
+		TSharedPtr<SWidget> ContentWidget;
+
+		if (ParameterStore && ParameterVariable)
+		{
+			if ( ColumnName == NAME_ParameterName )
+			{
+				ContentWidget =
+					SNew(SNiagaraParameterName)
+					.ParameterName(ParameterVariable->GetName())
+					.IsReadOnly(true);
+			}
+			else if ( ColumnName == NAME_ParameterValue)
+			{
+				if (const uint8* ParameterData = ParameterStore->GetParameterData(*ParameterVariable))
+				{
+					const FString ValueString = ParameterVariable->GetType().ToString(ParameterData);
+					ContentWidget =
+						SNew(STextBlock)
+						.Text(FText::FromString(ValueString));
+				}
+			}
+		}
+
+		if (ContentWidget == nullptr)
+		{
+			ContentWidget =
+				SNew(STextBlock)
+				.Text(LOCTEXT("UnknownColumn", "Unknown Column"));
+		}
+
+		return
+			SNew(SHorizontalBox)
+			+ SHorizontalBox::Slot()
+			.Padding(0.0f, 0.0f, 5.0f, 0.0f)
+			[
+				ContentWidget.ToSharedRef()
+			];
+	}
+
+	TSharedPtr<FNiagaraParameterStore>	ParameterStore;
+	TSharedPtr<FNiagaraVariableBase>	ParameterVariable;
+};
+
+class SParameterStoreListView : public SListView<TSharedPtr<FNiagaraVariableBase>>
+{
+	SLATE_BEGIN_ARGS(SParameterStoreListView) {}
+	SLATE_END_ARGS()
+
+	void Construct(const FArguments& InArgs)
+	{
+		SListView::Construct(SListView::FArguments()
+			//.Clipping(EWidgetClipping::OnDemand)
+			.ListItemsSource(&ParameterVariables)
+			.OnGenerateRow(this, &SParameterStoreListView::OnGenerateRowForEntry)
+			.HeaderRow
+			(
+				SNew(SHeaderRow)
+				+ SHeaderRow::Column(NAME_ParameterName)
+				.DefaultLabel(LOCTEXT("ParameterName", "Parameter Name"))
+				.ManualWidth(200)
+				+ SHeaderRow::Column(NAME_ParameterValue)
+				.DefaultLabel(LOCTEXT("ParameterValue", "Parameter Value"))
+				.ManualWidth(200)
+			)
+		);
+	}
+
+	void SetParameterStore(const FNiagaraParameterStore& InParameterStore)
+	{
+		ParameterVariables.Empty();
+		ParameterStore = MakeShared<FNiagaraParameterStore>(InParameterStore);
+
+		for (const FNiagaraVariableBase& Variable : ParameterStore->ReadParameterVariables())
+		{
+			if (Variable.IsDataInterface() || Variable.IsUObject())
+			{
+				continue;
+			}
+			ParameterVariables.Emplace(MakeShared<FNiagaraVariableBase>(Variable));
+		}
+		RequestListRefresh();
+	}
+
+	TSharedRef<ITableRow> OnGenerateRowForEntry(TSharedPtr<FNiagaraVariableBase> ParameterVariable, const TSharedRef<STableViewBase>& OwnerTable)
+	{
+		return
+			SNew(SParameterStoreItemWidget, OwnerTable)
+			.ParameterVariable(ParameterVariable)
+			.ParameterStore(ParameterStore);
+	}
+
+	TSharedPtr<FNiagaraParameterStore>			ParameterStore;
+	TArray<TSharedPtr<FNiagaraVariableBase>>	ParameterVariables;
+};
+
+} //namespace UE::Niagara::SimCache::DebugDataUI
+
+////////////////////////////////////////////////////////////////////////////////////////////////////////////////////////
+
+void SNiagaraSimCacheDebugDataView::Construct(const FArguments& InArgs)
+{
+	SimCacheViewModel = InArgs._SimCacheViewModel;
+
+	// Attach handlers for the cache data changing
+	SimCacheViewModel->OnViewDataChanged().AddSP(this, &SNiagaraSimCacheDebugDataView::RefreshContents);
+	SimCacheViewModel->OnSimCacheChanged().AddSP(this, &SNiagaraSimCacheDebugDataView::RefreshContents);
+	SimCacheViewModel->OnBufferChanged().AddSP(this, &SNiagaraSimCacheDebugDataView::RefreshContents);
+
+	ChildSlot
+	[
+		SNew(SVerticalBox)
+		+ SVerticalBox::Slot()
+		.AutoHeight()
+		[
+			SNew(SBorder)
+			.BorderImage(FAppStyle::Get().GetBrush("DetailsView.GridLine"))
+			[
+				SNew(SHorizontalBox)
+				+ SHorizontalBox::Slot()
+				.AutoWidth()
+				[
+					SNew(STextBlock)
+					.Text(LOCTEXT("OverrideParameters", "Override Parameters"))
+				]
+			]
+		]
+		+ SVerticalBox::Slot()
+		.AutoHeight()
+		[
+			SNew(SHorizontalBox)
+			+ SHorizontalBox::Slot()
+			.AutoWidth()
+			[
+				SAssignNew(OverrideParametersWidget, UE::Niagara::SimCache::DebugDataUI::SParameterStoreListView)
+			]
+		]
+	];
+
+	RefreshContents();
+};
+
+void SNiagaraSimCacheDebugDataView::RefreshContents()
+{
+	const UNiagaraSimCacheDebugData* DebugData = SimCacheViewModel ? SimCacheViewModel->GetCacheDebugData() : nullptr;
+	if (DebugData)
+	{
+		const int32 FrameIndex = SimCacheViewModel->GetFrameIndex();
+		if ( DebugData->Frames.IsValidIndex(FrameIndex))
+		{
+			const FNiagaraSimCacheDebugDataFrame& DebugFrameData = DebugData->Frames[FrameIndex];
+			OverrideParametersWidget->SetParameterStore(DebugFrameData.OverrideParameters);
+			return;
+		}
+	}
+
+	// If we get here we failed to set anything so update the display with defaults
+	OverrideParametersWidget->SetParameterStore(FNiagaraParameterStore());
+}
+
+void SNiagaraSimCacheDebugDataView::RefreshContents(bool)
+{
+	RefreshContents();
+}
+
+#undef LOCTEXT_NAMESPACE
diff --git a/Engine/Plugins/FX/Niagara/Source/NiagaraEditor/Private/Widgets/SNiagaraSimCacheDebugDataView.h b/Engine/Plugins/FX/Niagara/Source/NiagaraEditor/Private/Widgets/SNiagaraSimCacheDebugDataView.h
new file mode 100644
--- /dev/null
+++ b/Engine/Plugins/FX/Niagara/Source/NiagaraEditor/Private/Widgets/SNiagaraSimCacheDebugDataView.h
@@ -0,0 +1,29 @@
+// Copyright Epic Games, Inc. All Rights Reserved.
+
+#pragma once
+
+#include "Widgets/SCompoundWidget.h"
+
+class FNiagaraSimCacheViewModel;
+namespace UE::Niagara::SimCache::DebugDataUI
+{
+	class SParameterStoreListView;
+};
+
+class SNiagaraSimCacheDebugDataView : public SCompoundWidget
+{
+public:
+	SLATE_BEGIN_ARGS(SNiagaraSimCacheDebugDataView) {}
+		SLATE_ARGUMENT(TSharedPtr<FNiagaraSimCacheViewModel>, SimCacheViewModel)
+	SLATE_END_ARGS()
+
+	void Construct(const FArguments& InArgs);
+
+private:
+	void RefreshContents();
+	void RefreshContents(bool);
+
+private:
+	TSharedPtr<FNiagaraSimCacheViewModel>		SimCacheViewModel;
+	TSharedPtr<UE::Niagara::SimCache::DebugDataUI::SParameterStoreListView>			OverrideParametersWidget;
+};
\ No newline at end of file
diff --git a/Engine/Plugins/FX/Niagara/Source/NiagaraEditor/Private/Widgets/SNiagaraSimCacheOverview.cpp b/Engine/Plugins/FX/Niagara/Source/NiagaraEditor/Private/Widgets/SNiagaraSimCacheOverview.cpp
--- a/Engine/Plugins/FX/Niagara/Source/NiagaraEditor/Private/Widgets/SNiagaraSimCacheOverview.cpp
+++ b/Engine/Plugins/FX/Niagara/Source/NiagaraEditor/Private/Widgets/SNiagaraSimCacheOverview.cpp
@@ -120,16 +120,17 @@ void SNiagaraSimCacheOverview::Construct(const FArguments& InArgs)
 	];
 }
 
-TSharedRef<ITableRow> SNiagaraSimCacheOverview::OnGenerateRowForItem(TSharedRef<FNiagaraSimCacheOverviewItem> Item,
-	const TSharedRef<STableViewBase>& Owner)
+TSharedRef<ITableRow> SNiagaraSimCacheOverview::OnGenerateRowForItem(TSharedRef<FNiagaraSimCacheOverviewItem> Item, const TSharedRef<STableViewBase>& Owner)
 {
 	static const char* ItemStyles[] =
 	{
 		"NiagaraEditor.SimCache.SystemItem",
 		"NiagaraEditor.SimCache.EmitterItem",
 		"NiagaraEditor.SimCache.ComponentItem",
-		"NiagaraEditor.SimCache.DataInterfaceItem"
+		"NiagaraEditor.SimCache.DataInterfaceItem",
+		"NiagaraEditor.SimCache.DebugData",
 	};
+	static_assert(UE_ARRAY_COUNT(ItemStyles) == int(ENiagaraSimCacheOverviewItemType::MAX), "Mismatch on style count");
 
 	ENiagaraSimCacheOverviewItemType StyleType = Item->GetType();
 	
@@ -146,7 +147,24 @@ void SNiagaraSimCacheOverview::OnListSelectionChanged(TSharedPtr<FNiagaraSimCach
 {
 	if (Item.IsValid())
 	{
-		ViewModel->SetEmitterIndex(Item->GetBufferIndex(), Item->GetDataInterface());
+		switch (Item->GetType())
+		{
+			case ENiagaraSimCacheOverviewItemType::System:
+				ViewModel->SetSelectedSystemInstance();
+				break;
+
+			case ENiagaraSimCacheOverviewItemType::Emitter:
+				ViewModel->SetSelectedEmitter(Item->GetBufferIndex());
+				break;
+
+			case ENiagaraSimCacheOverviewItemType::DataInterface:
+				ViewModel->SetSelectedDataInterface(Item->GetDataInterface());
+				break;
+
+			case ENiagaraSimCacheOverviewItemType::DebugData:
+				ViewModel->SetSelectedDebugData();
+				break;
+		}
 	}
 }
 
diff --git a/Engine/Plugins/FX/Niagara/Source/NiagaraEditor/Private/Widgets/SNiagaraSimCacheOverview.h b/Engine/Plugins/FX/Niagara/Source/NiagaraEditor/Private/Widgets/SNiagaraSimCacheOverview.h
--- a/Engine/Plugins/FX/Niagara/Source/NiagaraEditor/Private/Widgets/SNiagaraSimCacheOverview.h
+++ b/Engine/Plugins/FX/Niagara/Source/NiagaraEditor/Private/Widgets/SNiagaraSimCacheOverview.h
@@ -19,49 +19,26 @@ enum class ENiagaraSimCacheOverviewItemType
 	Emitter,
 	Component,
 	DataInterface,
+	DebugData,
 	MAX,
 };
 
 struct FNiagaraSimCacheOverviewItem : public TSharedFromThis<FNiagaraSimCacheOverviewItem>
 {
+	FNiagaraSimCacheOverviewItem() = default;
 	virtual ~FNiagaraSimCacheOverviewItem() = default;
 
-	FNiagaraSimCacheOverviewItem()
-	{
-		
-	}
+	FText GetDisplayNameText() { return DisplayName; }
+	void SetDisplayName(const FText& NewName) { DisplayName = NewName; }
 
-	FText GetDisplayNameText()
-	{
-		return DisplayName;
-	}
+	virtual int32 GetBufferIndex() { return BufferIndex; }
+	void SetBufferIndex(int32 NewIndex) { BufferIndex = NewIndex; }
 
-	void SetDisplayName(const FText& NewName)
-	{
-		DisplayName = NewName;
-	}
+	virtual FNiagaraVariableBase GetDataInterface() { return FNiagaraVariableBase(); }
 
-	virtual int32 GetBufferIndex()
-	{
-		return BufferIndex;
-	}
+	virtual ENiagaraSimCacheOverviewItemType GetType () { checkNoEntry(); return ENiagaraSimCacheOverviewItemType::MAX; }
 
-	void SetBufferIndex(int32 NewIndex)
-	{
-		BufferIndex = NewIndex;
-	}
-
-	virtual FNiagaraVariableBase GetDataInterface()
-	{
-		return FNiagaraVariableBase();
-	}
-
-	virtual ENiagaraSimCacheOverviewItemType GetType () { return ENiagaraSimCacheOverviewItemType::MAX; }
-
-	virtual TSharedRef<SWidget> GetRowWidget()
-	{
-		return SNew(STextBlock).Text(GetDisplayNameText());
-	}
+	virtual TSharedRef<SWidget> GetRowWidget() { return SNew(STextBlock).Text(GetDisplayNameText()); }
 
 	int32 BufferIndex = INDEX_NONE;
 	FText DisplayName;
@@ -85,6 +62,11 @@ struct FNiagaraSimCacheOverviewDataInterfaceItem : FNiagaraSimCacheOverviewItem
 	FNiagaraVariableBase DataInterfaceReference;
 };
 
+struct FNiagaraSimCacheOverviewDebugDataItem : FNiagaraSimCacheOverviewItem
+{
+	virtual ENiagaraSimCacheOverviewItemType GetType() override { return ENiagaraSimCacheOverviewItemType::DebugData; }
+};
+
 class SNiagaraSimCacheOverview : public SCompoundWidget
 {
 public:
diff --git a/Engine/Plugins/FX/Niagara/Source/NiagaraEditor/Private/Widgets/SNiagaraSimCacheTreeView.cpp b/Engine/Plugins/FX/Niagara/Source/NiagaraEditor/Private/Widgets/SNiagaraSimCacheTreeView.cpp
--- a/Engine/Plugins/FX/Niagara/Source/NiagaraEditor/Private/Widgets/SNiagaraSimCacheTreeView.cpp
+++ b/Engine/Plugins/FX/Niagara/Source/NiagaraEditor/Private/Widgets/SNiagaraSimCacheTreeView.cpp
@@ -165,7 +165,7 @@ bool SSimCacheTreeViewVisibilityWidget::IsItemSelected() const
 
 void SNiagaraSimCacheTreeView::SetupRootEntries()
 {
-	TArray<TSharedRef<FNiagaraSimCacheTreeItem>>* RootEntries = ViewModel->GetCurrentRootEntries();
+	TArray<TSharedRef<FNiagaraSimCacheTreeItem>>* RootEntries = ViewModel->GetSelectedRootEntries();
 
 	if(RootEntries && !RootEntries->IsEmpty())
 	{
@@ -184,7 +184,7 @@ void SNiagaraSimCacheTreeView::Construct(const FArguments& InArgs)
 	
 	TreeView = SNew(STreeView<TSharedRef<FNiagaraSimCacheTreeItem>>)
 	.SelectionMode(ESelectionMode::Multi)
-	.TreeItemsSource(ViewModel->GetCurrentRootEntries())
+	.TreeItemsSource(ViewModel->GetSelectedRootEntries())
 	.OnGenerateRow(this, &SNiagaraSimCacheTreeView::OnGenerateRow)
 	.OnGeneratePinnedRow(this, &SNiagaraSimCacheTreeView::OnGenerateRow)
 	.ShouldStackHierarchyHeaders(true)
@@ -199,16 +199,17 @@ void SNiagaraSimCacheTreeView::Construct(const FArguments& InArgs)
 }
 
 
-TSharedRef<ITableRow> SNiagaraSimCacheTreeView::OnGenerateRow(TSharedRef<FNiagaraSimCacheTreeItem> Item,
-                                                              const TSharedRef<STableViewBase>& OwnerTable)
+TSharedRef<ITableRow> SNiagaraSimCacheTreeView::OnGenerateRow(TSharedRef<FNiagaraSimCacheTreeItem> Item, const TSharedRef<STableViewBase>& OwnerTable)
 {
 	static const char* ItemStyles[] =
 	{
 		"NiagaraEditor.SimCache.SystemItem",
 		"NiagaraEditor.SimCache.EmitterItem",
 		"NiagaraEditor.SimCache.ComponentItem",
-		"NiagaraEditor.SimCache.DataInterfaceItem"
+		"NiagaraEditor.SimCache.DataInterfaceItem",
+		"NiagaraEditor.SimCache.DebugData"
 	};
+	static_assert(UE_ARRAY_COUNT(ItemStyles) == int(ENiagaraSimCacheOverviewItemType::MAX), "Mismatch on style count");
 
 	ENiagaraSimCacheOverviewItemType StyleType = Item->GetType();
 	
@@ -233,7 +234,12 @@ void SNiagaraSimCacheTreeView::OnBufferChanged()
 	SelectionForFilter.Empty();
 	ViewModel->SetComponentFilters(TArray<FString>());
 	SelectAll();
-	TreeView->SetItemExpansion((*ViewModel->GetCurrentRootEntries())[0], true);
+
+	TArray<TSharedRef<FNiagaraSimCacheTreeItem>>* RootEntries = ViewModel->GetSelectedRootEntries();
+	if (RootEntries->Num() > 0)
+	{
+		TreeView->SetItemExpansion((*RootEntries)[0], true);
+	}
 }
 
 void SNiagaraSimCacheTreeView::OnSimCacheChanged()
@@ -327,16 +333,21 @@ void SNiagaraSimCacheTreeView::SelectAll()
 		return;
 	}
 	
-	RecursiveAddToSelectionFilter(*ViewModel->GetCurrentRootEntries());
+	RecursiveAddToSelectionFilter(*ViewModel->GetSelectedRootEntries());
 
 	UpdateStringFilters();
 }
 
-bool SNiagaraSimCacheTreeView::IsDataInterfaceViewActive() const
+bool SNiagaraSimCacheTreeView::ShouldShowComponentView() const
 {
 	if (ViewModel)
 	{
-		return ViewModel->GetActiveDataInterface().IsValid();
+		switch (ViewModel->GetSelectionMode())
+		{
+			case FNiagaraSimCacheViewModel::ESelectionMode::SystemInstance:
+			case FNiagaraSimCacheViewModel::ESelectionMode::Emitter:
+				return true;
+		}
 	}
 	return false;
 }
diff --git a/Engine/Plugins/FX/Niagara/Source/NiagaraEditor/Private/Widgets/SNiagaraSimCacheTreeView.h b/Engine/Plugins/FX/Niagara/Source/NiagaraEditor/Private/Widgets/SNiagaraSimCacheTreeView.h
--- a/Engine/Plugins/FX/Niagara/Source/NiagaraEditor/Private/Widgets/SNiagaraSimCacheTreeView.h
+++ b/Engine/Plugins/FX/Niagara/Source/NiagaraEditor/Private/Widgets/SNiagaraSimCacheTreeView.h
@@ -73,7 +73,8 @@ SLATE_END_ARGS()
 	void ClearFilterSelection();
 	void UpdateStringFilters();
 	void SelectAll();
-	bool IsDataInterfaceViewActive() const;
+	bool ShouldShowComponentView() const;
+
 private:
 	TSharedRef<ITableRow> OnGenerateRow(TSharedRef<FNiagaraSimCacheTreeItem> Item, const TSharedRef<STableViewBase>& OwnerTable);
 
@@ -98,7 +99,7 @@ private:
 
 struct FNiagaraSimCacheTreeItem : FNiagaraSimCacheOverviewItem
 {
-	FNiagaraSimCacheTreeItem(TWeakPtr<SNiagaraSimCacheTreeView> InOwner)
+	explicit FNiagaraSimCacheTreeItem(TWeakPtr<SNiagaraSimCacheTreeView> InOwner)
 	{
 		Owner = InOwner;
 	}
@@ -149,7 +150,7 @@ struct FNiagaraSimCacheTreeItem : FNiagaraSimCacheOverviewItem
 				{
 					if (TSharedPtr<SNiagaraSimCacheTreeView> TreeView = Owner.Pin())
 					{
-						return TreeView->IsDataInterfaceViewActive() ? EVisibility::Collapsed : EVisibility::Visible; 
+						return TreeView->ShouldShowComponentView() ? EVisibility::Visible : EVisibility::Collapsed;
 					}
 					return EVisibility::Collapsed;
 				})
@@ -166,8 +167,7 @@ struct FNiagaraSimCacheTreeItem : FNiagaraSimCacheOverviewItem
 
 struct FNiagaraSimCacheEmitterTreeItem : FNiagaraSimCacheTreeItem
 {
-	
-	FNiagaraSimCacheEmitterTreeItem(TWeakPtr<SNiagaraSimCacheTreeView> InOwner): FNiagaraSimCacheTreeItem(InOwner)
+	explicit FNiagaraSimCacheEmitterTreeItem(TWeakPtr<SNiagaraSimCacheTreeView> InOwner): FNiagaraSimCacheTreeItem(InOwner)
 	{
 	}
 
@@ -178,7 +178,7 @@ struct FNiagaraSimCacheEmitterTreeItem : FNiagaraSimCacheTreeItem
 
 struct FNiagaraSimCacheDataInterfaceTreeItem : FNiagaraSimCacheTreeItem
 {
-	FNiagaraSimCacheDataInterfaceTreeItem(TWeakPtr<SNiagaraSimCacheTreeView> InOwner): FNiagaraSimCacheTreeItem(InOwner)
+	explicit FNiagaraSimCacheDataInterfaceTreeItem(TWeakPtr<SNiagaraSimCacheTreeView> InOwner): FNiagaraSimCacheTreeItem(InOwner)
 	{
 	}
 
@@ -189,10 +189,20 @@ struct FNiagaraSimCacheDataInterfaceTreeItem : FNiagaraSimCacheTreeItem
 	FNiagaraVariableBase DataInterfaceReference;
 };
 
+struct FNiagaraSimCacheDebugDataTreeItem : FNiagaraSimCacheTreeItem
+{
+	explicit FNiagaraSimCacheDebugDataTreeItem(TWeakPtr<SNiagaraSimCacheTreeView> InOwner) : FNiagaraSimCacheTreeItem(InOwner)
+	{
+	}
+
+	virtual ~FNiagaraSimCacheDebugDataTreeItem() override {}
+
+	virtual ENiagaraSimCacheOverviewItemType GetType() override { return ENiagaraSimCacheOverviewItemType::DebugData; }
+};
+
 struct FNiagaraSimCacheComponentTreeItem : FNiagaraSimCacheTreeItem
 {
-	
-	FNiagaraSimCacheComponentTreeItem(TWeakPtr<SNiagaraSimCacheTreeView> InOwner): FNiagaraSimCacheTreeItem(InOwner)
+	explicit FNiagaraSimCacheComponentTreeItem(TWeakPtr<SNiagaraSimCacheTreeView> InOwner) : FNiagaraSimCacheTreeItem(InOwner)
 	{
 	}
 
@@ -202,7 +212,6 @@ struct FNiagaraSimCacheComponentTreeItem : FNiagaraSimCacheTreeItem
 
 	virtual TSharedRef<SWidget> GetRowWidget() override
 	{
-		
 		TSharedRef<SHorizontalBox> Contents = SNew(SHorizontalBox);
 
 		if(TypeDef.IsSet())
@@ -247,7 +256,7 @@ struct FNiagaraSimCacheComponentTreeItem : FNiagaraSimCacheTreeItem
 			{
 				if (TSharedPtr<SNiagaraSimCacheTreeView> TreeView = Owner.Pin())
 				{
-					return TreeView->IsDataInterfaceViewActive() ? EVisibility::Collapsed : EVisibility::Visible; 
+					return TreeView->ShouldShowComponentView() ? EVisibility::Visible : EVisibility::Collapsed;
 				}
 				return EVisibility::Collapsed;
 			})
diff --git a/Engine/Plugins/FX/Niagara/Source/NiagaraEditor/Private/Widgets/SNiagaraSimCacheView.cpp b/Engine/Plugins/FX/Niagara/Source/NiagaraEditor/Private/Widgets/SNiagaraSimCacheView.cpp
--- a/Engine/Plugins/FX/Niagara/Source/NiagaraEditor/Private/Widgets/SNiagaraSimCacheView.cpp
+++ b/Engine/Plugins/FX/Niagara/Source/NiagaraEditor/Private/Widgets/SNiagaraSimCacheView.cpp
@@ -2,6 +2,7 @@
 
 #include "SNiagaraSimCacheView.h"
 #include "ViewModels/NiagaraSimCacheViewModel.h"
+#include "SNiagaraSimCacheDebugDataView.h"
 
 #include "CoreMinimal.h"
 #include "NiagaraEditorModule.h"
@@ -87,7 +88,7 @@ void SNiagaraSimCacheView::Construct(const FArguments& InArgs)
 		.AlwaysShowScrollbar(true)
 		.Thickness(12.0f)
 		.Orientation(Orient_Vertical);
-	DataInterfaceScrollBar =
+	CustomDisplayScrollBar =
 		SNew(SScrollBar)
 		.AlwaysShowScrollbar(false)
 		.Thickness(12.0f)
@@ -122,7 +123,17 @@ void SNiagaraSimCacheView::Construct(const FArguments& InArgs)
 				[
 					// switcher for spreadsheet / data interfaces
 					SAssignNew(SwitchWidget, SWidgetSwitcher)
-					.WidgetIndex_Lambda([this] { return SimCacheViewModel->GetActiveDataInterface().IsValid() ? 1 : 0; })
+					.WidgetIndex_Lambda(
+						[this]
+						{
+							switch (SimCacheViewModel->GetSelectionMode())
+							{
+								case FNiagaraSimCacheViewModel::ESelectionMode::SystemInstance:
+								case FNiagaraSimCacheViewModel::ESelectionMode::Emitter:			return 0;
+								default:															return 1;
+							}
+						}
+					)
 					+SWidgetSwitcher::Slot()
 					[
 						ListViewWidget.ToSharedRef()
@@ -137,14 +148,32 @@ void SNiagaraSimCacheView::Construct(const FArguments& InArgs)
 			.AutoWidth()
 			[
 				SNew(SWidgetSwitcher)
-				.WidgetIndex_Lambda([this] { return SimCacheViewModel->GetActiveDataInterface().IsValid() ? 1 : 0; })
+					.WidgetIndex_Lambda(
+						[this]
+						{
+							switch (SimCacheViewModel->GetSelectionMode())
+							{
+								case FNiagaraSimCacheViewModel::ESelectionMode::SystemInstance:
+								case FNiagaraSimCacheViewModel::ESelectionMode::Emitter:			return 0;
+								case FNiagaraSimCacheViewModel::ESelectionMode::DataInterface:		return 1;
+								case FNiagaraSimCacheViewModel::ESelectionMode::DebugData:			return 2;
+								default:
+									checkNoEntry();
+									return 0;
+							}
+						}
+					)
 				+SWidgetSwitcher::Slot()
 				[
 					VerticalScrollBar
 				]
 				+SWidgetSwitcher::Slot()
 				[
-					DataInterfaceScrollBar.ToSharedRef()
+					CustomDisplayScrollBar.ToSharedRef()
+				]
+				+SWidgetSwitcher::Slot()
+				[
+					SNullWidget::NullWidget
 				]
 			]
 		]
@@ -188,7 +217,7 @@ void SNiagaraSimCacheView::GenerateColumns()
 	);
 	
 	// Generate a column for each component
-	for (const FNiagaraSimCacheViewModel::FComponentInfo& ComponentInfo : SimCacheViewModel->GetCurrentComponentInfos())
+	for (const FNiagaraSimCacheViewModel::FComponentInfo& ComponentInfo : SimCacheViewModel->GetSelectedComponentInfos())
 	{
 		HeaderRowWidget->AddColumn(
 			SHeaderRow::Column(ComponentInfo.Name)
@@ -206,50 +235,88 @@ void SNiagaraSimCacheView::GenerateColumns()
 
 }
 
-void SNiagaraSimCacheView::UpdateDIWidget()
+void SNiagaraSimCacheView::UpdateCustomDisplayWidget()
 {
-	for (TSharedPtr<SWidget> Widget : DIVisualizerWidgets)
+	for (TSharedPtr<SWidget> Widget : CustomDisplayWidgets)
 	{
 		SwitchWidget->RemoveSlot(Widget.ToSharedRef());
 	}
-	DIVisualizerWidgets.Empty();
+	CustomDisplayWidgets.Empty();
+
 	TSharedRef<SVerticalBox> WidgetBox = SNew(SVerticalBox);
-	if (SimCacheViewModel->GetActiveDataInterface().IsValid())
+
+	TOptional<FText> MissingCustomDisplayText;
+	switch (SimCacheViewModel->GetSelectionMode())
 	{
-		FNiagaraEditorModule& NiagaraEditorModule = FModuleManager::GetModuleChecked<FNiagaraEditorModule>("NiagaraEditor");
-		for (TSharedRef<INiagaraDataInterfaceSimCacheVisualizer> Visualizer : NiagaraEditorModule.FindDataInterfaceCacheVisualizer(SimCacheViewModel->GetActiveDataInterface().GetType().GetClass()))
+		case FNiagaraSimCacheViewModel::ESelectionMode::DataInterface:
+		{
+			const FNiagaraVariableBase DIVariable = SimCacheViewModel->GetSelectedDataInterface();
+
+			FNiagaraEditorModule& NiagaraEditorModule = FModuleManager::GetModuleChecked<FNiagaraEditorModule>("NiagaraEditor");
+			for (TSharedRef<INiagaraDataInterfaceSimCacheVisualizer> Visualizer : NiagaraEditorModule.FindDataInterfaceCacheVisualizer(DIVariable.GetType().GetClass()))
+			{
+				if (const UObject* DataObject = SimCacheViewModel->GetSelectedDataInterfaceStorage())
+				{
+					TSharedPtr<SWidget> VisualizerWidget = Visualizer->CreateWidgetFor(DataObject, SimCacheViewModel);
+					CustomDisplayWidgets.Add(VisualizerWidget);
+					WidgetBox->AddSlot()
+						.AutoHeight()
+						.AttachWidget(VisualizerWidget.ToSharedRef());
+				}
+			}
+
+			if (CustomDisplayWidgets.Num() == 0)
+			{
+				MissingCustomDisplayText = FText::Format(LOCTEXT("NoDataInterfaceVisualizer", "No valid visualizer found for data interface '{0}'"), DIVariable.GetType().GetNameText());
+			}
+			break;
+		}
+
+		case FNiagaraSimCacheViewModel::ESelectionMode::DebugData:
 		{
-			if (const UObject* DataObject = SimCacheViewModel->GetActiveDataInterfaceStorage())
+			if (SimCacheViewModel->GetCacheDebugData())
 			{
-				TSharedPtr<SWidget> DIVisualizerWidget = Visualizer->CreateWidgetFor(DataObject, SimCacheViewModel);
-				DIVisualizerWidgets.Add(DIVisualizerWidget);
+				TSharedPtr<SWidget> DebugDataWidget =
+					SNew(SNiagaraSimCacheDebugDataView)
+					.SimCacheViewModel(SimCacheViewModel);
+
+				CustomDisplayWidgets.Add(DebugDataWidget);
+
 				WidgetBox->AddSlot()
 					.AutoHeight()
-					.AttachWidget(DIVisualizerWidget.ToSharedRef());
+					.AttachWidget(DebugDataWidget.ToSharedRef());
 			}
+			else
+			{
+				MissingCustomDisplayText = LOCTEXT("NoDebugData", "Data Data not found inside cache");
+			}
+			break;
 		}
 	}
-	if (DIVisualizerWidgets.Num() == 0)
+
+	if (MissingCustomDisplayText.IsSet())
 	{
-		TSharedPtr<SWidget> DIVisualizerWidget = SNew(SBox)
+		TSharedPtr<SWidget> VisualizerWidget = SNew(SBox)
 		   .Padding(10)
 		   [
 			   SNew(STextBlock)
-			   .Text(LOCTEXT("NoDataInterface", "No valid visualizer found for data interface"))
+			   .Text(MissingCustomDisplayText.GetValue())
 		   ];
-		DIVisualizerWidgets.Add(DIVisualizerWidget);
+		CustomDisplayWidgets.Add(VisualizerWidget);
 		WidgetBox->AddSlot()
 			.AutoHeight()
-			.AttachWidget(DIVisualizerWidget.ToSharedRef());
+			.AttachWidget(VisualizerWidget.ToSharedRef());
 	}
+
 	SwitchWidget->AddSlot(1).AttachWidget(
 		SNew(SScrollBox)
 		.Orientation(Orient_Vertical)
-		.ExternalScrollbar(DataInterfaceScrollBar)
+		.ExternalScrollbar(CustomDisplayScrollBar)
 		 + SScrollBox::Slot()
 		[
 			WidgetBox
-		]);
+		]
+	);
 }
 
 void SNiagaraSimCacheView::UpdateColumns(const bool bReset)
@@ -281,7 +348,7 @@ void SNiagaraSimCacheView::OnSimCacheChanged()
 {
 	UpdateRows(true);
 	UpdateColumns(true);
-	UpdateDIWidget();
+	UpdateCustomDisplayWidget();
 }
 
 void SNiagaraSimCacheView::OnViewDataChanged(const bool bFullRefresh)
@@ -297,7 +364,7 @@ void SNiagaraSimCacheView::OnBufferChanged()
 {
 	UpdateRows(true);
 	UpdateColumns(true);
-	UpdateDIWidget();
+	UpdateCustomDisplayWidget();
 }
 
 bool SNiagaraSimCacheView::GetShouldGenerateWidget(FName Name)
diff --git a/Engine/Plugins/FX/Niagara/Source/NiagaraEditor/Private/Widgets/SNiagaraSimCacheView.h b/Engine/Plugins/FX/Niagara/Source/NiagaraEditor/Private/Widgets/SNiagaraSimCacheView.h
--- a/Engine/Plugins/FX/Niagara/Source/NiagaraEditor/Private/Widgets/SNiagaraSimCacheView.h
+++ b/Engine/Plugins/FX/Niagara/Source/NiagaraEditor/Private/Widgets/SNiagaraSimCacheView.h
@@ -39,7 +39,7 @@ public:
 private:
 	bool GetShouldGenerateWidget(FName Name);
 	void GenerateColumns();
-	void UpdateDIWidget();
+	void UpdateCustomDisplayWidget();
 
 	TArray<TSharedPtr<int32>>					RowItems;
 	TSharedPtr<FNiagaraSimCacheViewModel>		SimCacheViewModel;
@@ -47,6 +47,6 @@ private:
 	TSharedPtr<SHeaderRow>						HeaderRowWidget;
 	TSharedPtr<SListView<TSharedPtr<int32>>>	ListViewWidget;
 	TSharedPtr<SWidgetSwitcher>					SwitchWidget;
-	TArray<TSharedPtr<SWidget>>					DIVisualizerWidgets;
-	TSharedPtr<SScrollBar>						DataInterfaceScrollBar;
+	TArray<TSharedPtr<SWidget>>					CustomDisplayWidgets;
+	TSharedPtr<SScrollBar>						CustomDisplayScrollBar;
 };
\ No newline at end of file
diff --git a/Engine/Plugins/FX/Niagara/Source/NiagaraEditor/Public/ViewModels/NiagaraSimCacheViewModel.h b/Engine/Plugins/FX/Niagara/Source/NiagaraEditor/Public/ViewModels/NiagaraSimCacheViewModel.h
--- a/Engine/Plugins/FX/Niagara/Source/NiagaraEditor/Public/ViewModels/NiagaraSimCacheViewModel.h
+++ b/Engine/Plugins/FX/Niagara/Source/NiagaraEditor/Public/ViewModels/NiagaraSimCacheViewModel.h
@@ -17,6 +17,14 @@ class UNiagaraComponent;
 class FNiagaraSimCacheViewModel : public TSharedFromThis<FNiagaraSimCacheViewModel>, public FGCObject
 {
 public:
+	enum class ESelectionMode
+	{
+		SystemInstance,
+		Emitter,
+		DataInterface,
+		DebugData,
+	};
+
 	struct FComponentInfo
 	{
 		FName Name = NAME_None;
@@ -35,7 +43,7 @@ public:
 	NIAGARAEDITOR_API FNiagaraSimCacheViewModel();
 	NIAGARAEDITOR_API virtual ~FNiagaraSimCacheViewModel() override;
 	
-	NIAGARAEDITOR_API TConstArrayView<FComponentInfo> GetCurrentComponentInfos() const { return GetComponentInfos(EmitterIndex); }
+	NIAGARAEDITOR_API TConstArrayView<FComponentInfo> GetSelectedComponentInfos() const;
 	NIAGARAEDITOR_API TConstArrayView<FComponentInfo> GetComponentInfos(int32 InEmitterIndex) const;
 
 	NIAGARAEDITOR_API int32 GetNumInstances() const { return NumInstances; }
@@ -43,9 +51,6 @@ public:
 	NIAGARAEDITOR_API int32 GetFrameIndex() const { return FrameIndex; }
 	NIAGARAEDITOR_API void SetFrameIndex(const int32 InFrameIndex);
 
-	NIAGARAEDITOR_API FNiagaraVariableBase GetActiveDataInterface() const { return ActiveDataInterface; };
-	NIAGARAEDITOR_API const UObject* GetActiveDataInterfaceStorage() const;
-
 	NIAGARAEDITOR_API FOnViewDataChanged& OnViewDataChanged();
 	NIAGARAEDITOR_API FOnSimCacheChanged& OnSimCacheChanged();
 	NIAGARAEDITOR_API FOnBufferChanged& OnBufferChanged();
@@ -55,18 +60,31 @@ public:
 	TConstArrayView<FString> GetComponentFilters() const { return ComponentFilterArray; }
 	void SetComponentFilters(const TArray<FString>& NewComponentFilterArray);
 	bool IsComponentFilterActive() const { return bComponentFilterActive; }
-	
-	int32 GetEmitterIndex() const { return EmitterIndex; };
-	void SetEmitterIndex(int32 InEmitterIndex, FNiagaraVariableBase ActiveDataInterface = FNiagaraVariableBase());
+
+	ESelectionMode GetSelectionMode() const { return SelectionMode; }
+	int32 GetSelectedEmitter() const { return SelectedEmitterIndex; };
+	FNiagaraVariableBase GetSelectedDataInterface() const { return SelectedDataInterface; };
+
+	NIAGARAEDITOR_API const UObject* GetSelectedDataInterfaceStorage() const;
+
+	void SetSelectedSystemInstance();
+	void SetSelectedEmitter(int32 EmitterIndex);
+	void SetSelectedDataInterface(FNiagaraVariableBase DIVariable);
+	void SetSelectedDebugData();
+private:
+	void RefreshFromSelectionChanged();
+public:
 
 	FText GetComponentText(FName ComponentName, int32 InstanceIndex) const;
 	bool IsCacheValid() const;
 	int32 GetNumEmitterLayouts() const;
 	FName GetEmitterLayoutName(int32 Index) const;
 
+	const UNiagaraSimCacheDebugData* GetCacheDebugData() const;
+
 	// Construct entries for the tree view.
 	void BuildEntries(TWeakPtr<SNiagaraSimCacheTreeView> OwningTreeView);
-	TArray<TSharedRef<FNiagaraSimCacheTreeItem>>* GetCurrentRootEntries();
+	TArray<TSharedRef<FNiagaraSimCacheTreeItem>>* GetSelectedRootEntries();
 	TArray<TSharedRef<FNiagaraSimCacheOverviewItem>>* GetBufferEntries();
 
 	bool CanCopyActiveToClipboard() const;
@@ -102,12 +120,11 @@ private:
 
 	// Which frame of the cached sim is being viewed
 	int32 FrameIndex = 0;
-	
-	// Which Emitter of the cached sim is being viewed
-	int32 EmitterIndex = INDEX_NONE;
 
-	// Which data interface of the cached sim is being viewed
-	FNiagaraVariableBase ActiveDataInterface;
+	// Handles what is selected currently
+	ESelectionMode			SelectionMode = ESelectionMode::SystemInstance;
+	int32					SelectedEmitterIndex = INDEX_NONE;
+	FNiagaraVariableBase	SelectedDataInterface;
 
 	// Number of particles in the given frame
 	int32 NumInstances = 0;
@@ -118,7 +135,7 @@ private:
 	// TODO: Component infos for Data Interfaces?
 
 	TArray<TSharedRef<FNiagaraSimCacheTreeItem>> RootEntries;
-	TArray<TSharedRef<FNiagaraSimCacheTreeItem>> CurrentRootEntries;
+	TArray<TSharedRef<FNiagaraSimCacheTreeItem>> SelectedRootEntries;
 
 	TArray<TSharedRef<FNiagaraSimCacheOverviewItem>> BufferEntries;
 
