Fixed compile issue

diff --git a/Engine/Plugins/FX/Niagara/Source/NiagaraEditor/Public/Widgets/SNiagaraParameterMenu.cpp b/Engine/Plugins/FX/Niagara/Source/NiagaraEditor/Public/Widgets/SNiagaraParameterMenu.cpp
--- a/Engine/Plugins/FX/Niagara/Source/NiagaraEditor/Public/Widgets/SNiagaraParameterMenu.cpp
+++ b/Engine/Plugins/FX/Niagara/Source/NiagaraEditor/Public/Widgets/SNiagaraParameterMenu.cpp
@@ -2,6 +2,7 @@
 
 #include "SNiagaraParameterMenu.h"
 
+#include "EdGraphSchema_Niagara.h"
 #include "AssetRegistry/AssetRegistryModule.h"
 #include "EdGraph/EdGraphPin.h"
 #include "EdGraph/EdGraphSchema.h"
