Added Quick Init widget to the Niagara NDC Write node

diff --git a/Engine/Plugins/FX/Niagara/Source/NiagaraEditor/Private/NiagaraActions.cpp b/Engine/Plugins/FX/Niagara/Source/NiagaraEditor/Private/NiagaraActions.cpp
--- a/Engine/Plugins/FX/Niagara/Source/NiagaraEditor/Private/NiagaraActions.cpp
+++ b/Engine/Plugins/FX/Niagara/Source/NiagaraEditor/Private/NiagaraActions.cpp
@@ -811,6 +811,79 @@ void FNiagaraDataInterfaceNodeActionProvider_DataChannelWrite::GetNodeContextMen
 	Section.AddSubMenu("InitForDataChannelMenu", InitForDataChannelMenuText, InitForDataChannelMenuTooltipText,	FNewToolMenuDelegate::CreateLambda(CreateNodeContextMenu));
 }
 
+void FNiagaraDataInterfaceNodeActionProvider_DataChannelWrite::GetInlineNodeContextMenuActionsImpl(UToolMenu* ToolMenu) const
+{
+	AddDataChannelInitActions(ToolMenu);
+}
+
+INiagaraDataInterfaceNodeActionProvider::FInlineMenuDisplayOptions FNiagaraDataInterfaceNodeActionProvider_DataChannelWrite::GetInlineMenuDisplayOptionsImpl(UClass* DIClass, UEdGraphNode* Source) const
+{
+	UNiagaraNodeFunctionCall* FunctionCall = Cast<UNiagaraNodeFunctionCall>(Source);
+	if(FunctionCall == nullptr)
+	{
+		return {};
+	}
+
+	if (FunctionCall->Signature.Name != TEXT("Write"))
+	{
+		return {};
+	}
+	
+	FInlineMenuDisplayOptions DisplayOptions;
+	DisplayOptions.bDisplayInline = true;
+	DisplayOptions.DisplayBrush = FAppStyle::GetBrush("Icons.Edit");
+	DisplayOptions.TooltipText = LOCTEXT("InitWithDataChannel", "Initialize the write node using the selected Niagara Data Channel asset.");
+
+	return DisplayOptions;
+}
+
+void FNiagaraDataInterfaceNodeActionProvider_DataChannelWrite::AddDataChannelInitActions(UToolMenu* ToolMenu)
+{
+	UGraphNodeContextMenuContext* Context = ToolMenu->FindContext<UGraphNodeContextMenuContext>();
+
+	if(Context == nullptr || Context->Node == nullptr)
+	{
+		return;
+	}
+	
+	const UNiagaraNodeFunctionCall* FuncNode = CastChecked<UNiagaraNodeFunctionCall>(Context->Node);
+	TWeakObjectPtr<const UNiagaraNodeFunctionCall> WeakNode = FuncNode;
+
+	FToolMenuSection& MenuSection = ToolMenu->AddSection("DataChannelWrite", NiagaraActionsLocal::InitForDataChannelHeaderText);
+	auto InitForDataChannelSection = [&MenuSection, WeakNode](UNiagaraDataChannel* DataChannel)
+	{
+		check(DataChannel);
+
+		TWeakObjectPtr<UNiagaraDataChannel> WeakChannel = DataChannel;
+		auto CreateDataChannelActionEntry = [WeakChannel, WeakNode]()
+		{
+			UNiagaraDataChannel* Channel = WeakChannel.Get();
+			UNiagaraNodeFunctionCall* Node = const_cast<UNiagaraNodeFunctionCall*>(WeakNode.Get());
+
+			if (Channel && Node)
+			{
+				Node->RemoveAllDynamicPins();
+				TConstArrayView<FNiagaraDataChannelVariable> ChannelVars = Channel->GetVariables();
+				for (const FNiagaraDataChannelVariable& Var : ChannelVars)
+				{
+					FNiagaraTypeDefinition Type = Var.GetType();
+					if (Type.IsEnum() == false)
+					{
+						Type = FNiagaraTypeDefinition(FNiagaraTypeHelper::GetSWCStruct(Var.GetType().GetScriptStruct()));
+					}
+					FNiagaraVariable SWCVar(Type, Var.GetName());
+					Node->AddParameter(SWCVar, EEdGraphPinDirection::EGPD_Input);
+				}
+			}
+		};
+
+		MenuSection.AddMenuEntry(NAME_None, FText::FromString(DataChannel->GetAsset()->GetName()), FText::FromString(DataChannel->GetAsset()->GetName()), FSlateIcon(FAppStyle::GetAppStyleSetName(), "Icons.Import"),
+			FUIAction(FExecuteAction::CreateLambda(CreateDataChannelActionEntry), FCanExecuteAction()));
+	};
+
+	UNiagaraDataChannel::ForEachDataChannel(InitForDataChannelSection);
+}
+
 
 void FNiagaraDataInterfaceNodeActionProvider_DataChannelRead::GetNodeContextMenuActionsImpl(UToolMenu* Menu, UGraphNodeContextMenuContext* Context, FNiagaraFunctionSignature Signature) const
 {
diff --git a/Engine/Plugins/FX/Niagara/Source/NiagaraEditor/Public/NiagaraActions.h b/Engine/Plugins/FX/Niagara/Source/NiagaraEditor/Public/NiagaraActions.h
--- a/Engine/Plugins/FX/Niagara/Source/NiagaraEditor/Public/NiagaraActions.h
+++ b/Engine/Plugins/FX/Niagara/Source/NiagaraEditor/Public/NiagaraActions.h
@@ -485,6 +485,11 @@ class FNiagaraDataInterfaceNodeActionProvider_DataChannelWrite : public INiagara
 public:
 
 	virtual void GetNodeContextMenuActionsImpl(UToolMenu* Menu, UGraphNodeContextMenuContext* Context, FNiagaraFunctionSignature Signature) const override;
+	virtual void GetInlineNodeContextMenuActionsImpl(UToolMenu* ToolMenu) const override;
+	virtual FInlineMenuDisplayOptions GetInlineMenuDisplayOptionsImpl(UClass* DIClass, UEdGraphNode* Source) const override;
+
+private:
+	static void AddDataChannelInitActions(UToolMenu* ToolMenu);
 };
 
 class FNiagaraDataInterfaceNodeActionProvider_DataChannelRead : public INiagaraDataInterfaceNodeActionProvider
