Fixup missing support for MapFor in new compilation mode
#tests Make a test gpu fountain emitter and verified it compiled when using a MapFor in a scratchpad

diff --git a/Engine/Plugins/FX/Niagara/Source/NiagaraEditor/Private/NiagaraGraphDigest.cpp b/Engine/Plugins/FX/Niagara/Source/NiagaraEditor/Private/NiagaraGraphDigest.cpp
--- a/Engine/Plugins/FX/Niagara/Source/NiagaraEditor/Private/NiagaraGraphDigest.cpp
+++ b/Engine/Plugins/FX/Niagara/Source/NiagaraEditor/Private/NiagaraGraphDigest.cpp
@@ -3850,7 +3850,14 @@ void FNiagaraCompilationNodeParameterMapFor::Compile(FTranslator* Translator, TA
 		FNiagaraCompilationNodeParameterMapSet::Compile(Translator, Outputs);
 		//Translator->Message(FNiagaraCompileEventSeverity::Log,LOCTEXT("UnsupportedParamMapFor", "Parameter map for is not yet supported on cpu."), this, nullptr);
 	}
+}
+
+void FNiagaraCompilationNodeParameterMapFor::CollectInputPinsToCompile(FTranslator* Translator, FInputPinCollection& ActiveInputPins) const
+{
+	FNiagaraCompilationNodeParameterMapSet::CollectInputPinsToCompile(Translator, ActiveInputPins);
 
+	// we need to cull the iteration pin from the list of pins to compile as it's already going to be compiled during the MapFor::Compile
+	ActiveInputPins.Remove(&InputPins[1]);
 }
 
 FNiagaraCompilationNodeParameterMapForWithContinue::FNiagaraCompilationNodeParameterMapForWithContinue(const UNiagaraNodeParameterMapForWithContinue* InNode, FNiagaraCompilationGraphCreateContext& Context)
@@ -3886,6 +3893,14 @@ void FNiagaraCompilationNodeParameterMapForWithContinue::Compile(FTranslator* Tr
 	}
 }
 
+void FNiagaraCompilationNodeParameterMapForWithContinue::CollectInputPinsToCompile(FTranslator* Translator, FInputPinCollection& ActiveInputPins) const
+{
+	FNiagaraCompilationNodeParameterMapFor::CollectInputPinsToCompile(Translator, ActiveInputPins);
+
+	// we also need to cull out the iteration enabled pin
+	ActiveInputPins.Remove(&InputPins[2]);
+}
+
 FNiagaraCompilationNodeParameterMapForIndex::FNiagaraCompilationNodeParameterMapForIndex(const UNiagaraNodeParameterMapForIndex* InNode, FNiagaraCompilationGraphCreateContext& Context)
 	: FNiagaraCompilationNode(ENodeType::ParameterMapForIndex, InNode, Context)
 {
@@ -3979,21 +3994,8 @@ void FNiagaraCompilationNodeParameterMapSet::Compile(FTranslator* Translator, TA
 	check(Outputs.Num() == 0);
 	Outputs.Init(INDEX_NONE, OutputPins.Num());
 
-	TArray<const FNiagaraCompilationInputPin*, TInlineAllocator<16>> ActiveInputPins;
-	ActiveInputPins.Reserve(InputPins.Num());
-	// do a first pass over all of the pins so that we can properly cull out input pins and
-	// propagate the disabled pins up the chain
-	for (const FNiagaraCompilationInputPin& InputPin : InputPins)
-	{
-		if (Translator->IsFunctionVariableCulledFromCompilation(InputPin.PinName))
-		{
-			Translator->CullMapSetInputPin(&InputPin);
-		}
-		else
-		{
-			ActiveInputPins.Add(&InputPin);
-		}
-	}
+	FInputPinCollection ActiveInputPins;
+	CollectInputPinsToCompile(Translator, ActiveInputPins);
 
 	TArray<FTranslator::FCompiledPin, TInlineAllocator<16>> CompileInputs;
 	CompileInputs.Reserve(ActiveInputPins.Num());
@@ -4018,6 +4020,24 @@ void FNiagaraCompilationNodeParameterMapSet::Compile(FTranslator* Translator, TA
 	}
 }
 
+void FNiagaraCompilationNodeParameterMapSet::CollectInputPinsToCompile(FTranslator* Translator, FInputPinCollection& ActiveInputPins) const
+{
+	ActiveInputPins.Reserve(InputPins.Num());
+	// do a first pass over all of the pins so that we can properly cull out input pins and
+	// propagate the disabled pins up the chain
+	for (const FNiagaraCompilationInputPin& InputPin : InputPins)
+	{
+		if (Translator->IsFunctionVariableCulledFromCompilation(InputPin.PinName))
+		{
+			Translator->CullMapSetInputPin(&InputPin);
+		}
+		else
+		{
+			ActiveInputPins.Add(&InputPin);
+		}
+	}
+}
+
 FNiagaraCompilationNodeReadDataSet::FNiagaraCompilationNodeReadDataSet(const UNiagaraNodeReadDataSet* InNode, FNiagaraCompilationGraphCreateContext& Context)
 	: FNiagaraCompilationNode(ENodeType::ReadDataSet, InNode, Context)
 {
diff --git a/Engine/Plugins/FX/Niagara/Source/NiagaraEditor/Private/NiagaraGraphDigest.h b/Engine/Plugins/FX/Niagara/Source/NiagaraEditor/Private/NiagaraGraphDigest.h
--- a/Engine/Plugins/FX/Niagara/Source/NiagaraEditor/Private/NiagaraGraphDigest.h
+++ b/Engine/Plugins/FX/Niagara/Source/NiagaraEditor/Private/NiagaraGraphDigest.h
@@ -607,6 +607,10 @@ public:
 
 	virtual void BuildParameterMapHistory(FParameterMapHistoryBuilder& Builder, bool bRecursive, bool bFilterForCompilation) const override;
 	virtual void Compile(FTranslator* Translator, TArray<int32>& Outputs) const;
+
+protected:
+	using FInputPinCollection = TArray<const FNiagaraCompilationInputPin*, TInlineAllocator<16>>;
+	virtual void CollectInputPinsToCompile(FTranslator* Translator, FInputPinCollection& ActiveInputPins) const;
 };
 
 class FNiagaraCompilationNodeParameterMapFor : public FNiagaraCompilationNodeParameterMapSet
@@ -618,6 +622,9 @@ public:
 	FNiagaraCompilationNodeParameterMapFor(const FNiagaraCompilationNodeParameterMapFor& InNode, FNiagaraCompilationGraphDuplicateContext& Context);
 
 	virtual void Compile(FTranslator* Translator, TArray<int32>& Outputs) const;
+
+protected:
+	virtual void CollectInputPinsToCompile(FTranslator* Translator, FInputPinCollection& ActiveInputPins) const override;
 };
 
 class FNiagaraCompilationNodeParameterMapForWithContinue : public FNiagaraCompilationNodeParameterMapFor
@@ -629,6 +636,9 @@ public:
 	FNiagaraCompilationNodeParameterMapForWithContinue(const FNiagaraCompilationNodeParameterMapForWithContinue& InNode, FNiagaraCompilationGraphDuplicateContext& Context);
 
 	virtual void Compile(FTranslator* Translator, TArray<int32>& Outputs) const;
+
+protected:
+	virtual void CollectInputPinsToCompile(FTranslator* Translator, FInputPinCollection& ActiveInputPins) const override;
 };
 
 class FNiagaraCompilationNodeParameterMapForIndex : public FNiagaraCompilationNode
