* Fixed a crash when undoing parameter type changes
* Added undo/redo engine tests for rename und type change

diff --git a/Engine/Plugins/FX/Niagara/Source/NiagaraEditor/Private/NiagaraGraph.cpp b/Engine/Plugins/FX/Niagara/Source/NiagaraEditor/Private/NiagaraGraph.cpp
--- a/Engine/Plugins/FX/Niagara/Source/NiagaraEditor/Private/NiagaraGraph.cpp
+++ b/Engine/Plugins/FX/Niagara/Source/NiagaraEditor/Private/NiagaraGraph.cpp
@@ -770,6 +770,7 @@ void UNiagaraGraph::DeclareConstructClasses(TArray<FTopLevelAssetPath>& OutConst
 void UNiagaraGraph::ChangeParameterType(const TArray<FNiagaraVariable>& ParametersToChange, const FNiagaraTypeDefinition& NewType, bool bAllowOrphanedPins)
 {
 	check(!bIsForCompilationOnly);
+	Modify();
 
 	const UEdGraphSchema_Niagara* NiagaraSchema = GetDefault<UEdGraphSchema_Niagara>();
 	TArray<UNiagaraNodeParameterMapBase*> NiagaraParameterNodes;
@@ -860,7 +861,7 @@ void UNiagaraGraph::ChangeParameterType(const TArray<FNiagaraVariable>& Paramete
 			if((PinConnectionsToKeep.Num() != OriginalPin->LinkedTo.Num()) || (OldDefaultValuesCache[OriginalPin] != OldAutogeneratedDefaultValuesCache[OriginalPin]))
 			{
 				ensure(OldTypesCache.Contains(OriginalPin) && OldDefaultValuesCache.Contains(OriginalPin));
-			
+				
 				UEdGraphPin* NewOrphanedPin = Node->CreatePin(OriginalPin->Direction, OriginalPin->PinType.PinCategory, OrphanedPinName);
 				NewOrphanedPin->PinType = UEdGraphSchema_Niagara::TypeDefinitionToPinType(OldTypesCache[OriginalPin]);
 				NewOrphanedPin->DefaultValue = OldDefaultValuesCache[OriginalPin];
@@ -894,7 +895,7 @@ void UNiagaraGraph::ChangeParameterType(const TArray<FNiagaraVariable>& Paramete
 			}
 		}
 		
-		return (UEdGraphPin*) nullptr;
+		return static_cast<UEdGraphPin*>(nullptr);
 	};
 		
 	for (UNiagaraNodeParameterMapBase* NiagaraNode : NiagaraParameterNodes)
