- Fix CPU stateless sim issue with varying lifetimes

#rnx

diff --git a/Engine/Plugins/FX/Niagara/Source/Niagara/Private/Stateless/NiagaraStatelessParticleSimContext.cpp b/Engine/Plugins/FX/Niagara/Source/Niagara/Private/Stateless/NiagaraStatelessParticleSimContext.cpp
--- a/Engine/Plugins/FX/Niagara/Source/Niagara/Private/Stateless/NiagaraStatelessParticleSimContext.cpp
+++ b/Engine/Plugins/FX/Niagara/Source/Niagara/Private/Stateless/NiagaraStatelessParticleSimContext.cpp
@@ -193,7 +193,7 @@ namespace NiagaraStateless
 			float Particle_Lifetime = RandomScaleBiasFloat(NumInstances, 0, EmitterData->LifetimeRange);
 			if (Particle_Lifetime <= 0.0f || Particle_Age >= Particle_Lifetime)
 			{
-				return;
+				continue;
 			}
 
 			const float PreviousAge = FMath::Max(Particle_Age - DeltaTime, 0.0f);
