Fixed duplicate loctext key

diff --git a/Engine/Plugins/FX/Niagara/Source/NiagaraEditor/Private/NiagaraActions.cpp b/Engine/Plugins/FX/Niagara/Source/NiagaraEditor/Private/NiagaraActions.cpp
--- a/Engine/Plugins/FX/Niagara/Source/NiagaraEditor/Private/NiagaraActions.cpp
+++ b/Engine/Plugins/FX/Niagara/Source/NiagaraEditor/Private/NiagaraActions.cpp
@@ -832,7 +832,7 @@ INiagaraDataInterfaceNodeActionProvider::FInlineMenuDisplayOptions FNiagaraDataI
 	FInlineMenuDisplayOptions DisplayOptions;
 	DisplayOptions.bDisplayInline = true;
 	DisplayOptions.DisplayBrush = FAppStyle::GetBrush("Icons.Edit");
-	DisplayOptions.TooltipText = LOCTEXT("InitWithDataChannel", "Initialize the write node using the selected Niagara Data Channel asset.");
+	DisplayOptions.TooltipText = LOCTEXT("InitWriteWithDataChannel", "Initialize the write node using the selected Niagara Data Channel asset.");
 
 	return DisplayOptions;
 }
@@ -929,7 +929,7 @@ INiagaraDataInterfaceNodeActionProvider::FInlineMenuDisplayOptions FNiagaraDataI
 	FInlineMenuDisplayOptions DisplayOptions;
 	DisplayOptions.bDisplayInline = true;
 	DisplayOptions.DisplayBrush = FAppStyle::GetBrush("Icons.Edit");
-	DisplayOptions.TooltipText = LOCTEXT("InitWithDataChannel", "Initialize the read node using the selected Niagara Data Channel asset.");
+	DisplayOptions.TooltipText = LOCTEXT("InitReadWithDataChannel", "Initialize the read node using the selected Niagara Data Channel asset.");
 
 	return DisplayOptions;
 }
