- Move ENiagaraExecutionStateManagement into C++ for stateless to use

#rnx

diff --git a/Engine/Plugins/FX/Niagara/Config/BaseNiagara.ini b/Engine/Plugins/FX/Niagara/Config/BaseNiagara.ini
--- a/Engine/Plugins/FX/Niagara/Config/BaseNiagara.ini
+++ b/Engine/Plugins/FX/Niagara/Config/BaseNiagara.ini
@@ -11,6 +11,8 @@
 +EnumRedirects=(OldName="/Niagara/Functions/Localspace/ENiagaraOrientationAxis.ENiagaraOrientationAxis",NewName="/Script/Niagara.ENiagaraOrientationAxis",OverrideClassName="/Script/CoreUObject.Enum",ValueChanges=(("NewEnumerator0","XAxis"),("NewEnumerator1","YAxis"),("NewEnumerator2","ZAxis"),("NewEnumerator3","Max")))
 +EnumRedirects=(OldName="/Niagara/Enums/ENiagaraCoordinateSpace.ENiagaraCoordinateSpace",NewName="/Script/Niagara.ENiagaraCoordinateSpace",OverrideClassName="/Script/CoreUObject.Enum",ValueChanges=(("NewEnumerator0","Simulation"),("NewEnumerator1","World"),("NewEnumerator2","Local"),("NewEnumerator3","Max")))
 +EnumRedirects=(OldName="/Niagara/Enums/ENiagaraOrientationAxis.ENiagaraOrientationAxis",NewName="/Script/Niagara.ENiagaraOrientationAxis",OverrideClassName="/Script/CoreUObject.Enum",ValueChanges=(("NewEnumerator0","XAxis"),("NewEnumerator1","YAxis"),("NewEnumerator2","ZAxis"),("NewEnumerator3","Max")))
++EnumRedirects=(OldName="/Niagara/Enums/ENiagaraExecutionStateManagement.ENiagaraExecutionStateManagement",NewName="/Script/Niagara.ENiagaraExecutionStateManagement",OverrideClassName="/Script/CoreUObject.Enum",ValueChanges=(("NewEnumerator0","Awaken"),("NewEnumerator1","SleepAndLetParticlesFinish"),("NewEnumerator2","SleepAndClearParticles"),("NewEnumerator3","KillImmediately"),("NewEnumerator4","KillAfterParticlesFinish")))
+
 
 +StructRedirects=(OldName="NiagaraScalabilitySettings",NewName="/Script/Niagara.NiagaraSystemScalabilitySettings")
 +StructRedirects=(OldName="NiagaraScalabilityOverrides",NewName="/Script/Niagara.NiagaraSystemScalabilityOverride")
diff --git a/Engine/Plugins/FX/Niagara/Source/Niagara/Private/NiagaraTypes.cpp b/Engine/Plugins/FX/Niagara/Source/Niagara/Private/NiagaraTypes.cpp
--- a/Engine/Plugins/FX/Niagara/Source/Niagara/Private/NiagaraTypes.cpp
+++ b/Engine/Plugins/FX/Niagara/Source/Niagara/Private/NiagaraTypes.cpp
@@ -14,6 +14,9 @@
 static FName NAME_NiagaraDouble(TEXT("NiagaraDouble"));
 static FName NAME_NiagaraPosition(TEXT("NiagaraPosition"));
 
+// Emitter stats makes some assumptions about how you can cast from a ENiagaraExecutionStateManagement to ENiagaraExecutionState and the enum values will match
+static_assert(int(ENiagaraExecutionState::Num) == int(ENiagaraExecutionStateManagement::Num), "ENiagaraExecutionState & ENiagaraExecutionStateManagement must match 1:1 as scripts assume they can cast between");
+
 void FNiagaraVariableBase::SetNamespacedName(const FString& InNamespace, FName InVariableName)
 {
 	TStringBuilder<128> NameBuilder;
diff --git a/Engine/Plugins/FX/Niagara/Source/Niagara/Public/NiagaraTypes.h b/Engine/Plugins/FX/Niagara/Source/Niagara/Public/NiagaraTypes.h
--- a/Engine/Plugins/FX/Niagara/Source/Niagara/Public/NiagaraTypes.h
+++ b/Engine/Plugins/FX/Niagara/Source/Niagara/Public/NiagaraTypes.h
@@ -577,6 +577,38 @@ enum class ENiagaraExecutionState : uint32
 	Num UMETA(Hidden)
 };
 
+UENUM()
+enum class ENiagaraExecutionStateManagement : uint32
+{
+	/** Change the execution state to active (if possible). */
+	Awaken,
+	/**
+	Change the execution state to inactivate (if possible) and stop spawning particles.
+	Age will continue to increase and can complete if the loop condition is complete.
+	Can be woken again.
+	*/
+	SleepAndLetParticlesFinish,
+	/**
+	Change the execution state to inactivate clear (if possible), and kill all existing particles.
+	Age will continue to increase and can complete if the loop condition is complete.
+	Can be woken again.
+	*/
+	SleepAndClearParticles,
+	/**
+	Change the execution state to complete (if possible), and kill all existing particles.
+	Can not wake up again.
+	*/
+	KillImmediately,
+	/**
+	Change the execution state to inactive (if possible), and stop spawning particles.
+	When we have no particles activate will change the execution state to complete (if possible).
+	Can not wake up again if particles are still active.
+	*/
+	KillAfterParticlesFinish,
+
+	Num UMETA(Hidden)
+};
+
 UENUM()
 enum class ENiagaraCoordinateSpace : uint32
 {
