Improve performance button tooltip in Niagara system editor

#rnx

diff --git a/Engine/Plugins/FX/Niagara/Source/NiagaraEditor/Private/Toolkits/SystemToolkitModes/NiagaraSystemToolkitMode_Default.cpp b/Engine/Plugins/FX/Niagara/Source/NiagaraEditor/Private/Toolkits/SystemToolkitModes/NiagaraSystemToolkitMode_Default.cpp
--- a/Engine/Plugins/FX/Niagara/Source/NiagaraEditor/Private/Toolkits/SystemToolkitModes/NiagaraSystemToolkitMode_Default.cpp
+++ b/Engine/Plugins/FX/Niagara/Source/NiagaraEditor/Private/Toolkits/SystemToolkitModes/NiagaraSystemToolkitMode_Default.cpp
@@ -176,7 +176,7 @@ void FNiagaraSystemToolkitMode_Default::ExtendToolbar()
 			{
 				ToolbarBuilder.AddToolBarButton(FNiagaraEditorCommands::Get().ToggleStatPerformance, NAME_None,
                     LOCTEXT("NiagaraShowPerformance", "Performance"),
-                    LOCTEXT("NiagaraShowPerformanceTooltip", "Show runtime performance for particle scripts."),
+                    LOCTEXT("NiagaraShowPerformanceTooltip", "Show runtime performance for particle scripts. Not supported when optimized VM byte code setting is enabled."),
                     FSlateIcon(FAppStyle::Get().GetStyleSetName(), "AssetEditor.ToggleStats"));
 				ToolbarBuilder.AddComboButton(
                     FUIAction(),
