Implement add pin customization for NDC write nodes.

diff --git a/Engine/Plugins/FX/Niagara/Source/NiagaraEditor/Private/NiagaraActions.cpp b/Engine/Plugins/FX/Niagara/Source/NiagaraEditor/Private/NiagaraActions.cpp
--- a/Engine/Plugins/FX/Niagara/Source/NiagaraEditor/Private/NiagaraActions.cpp
+++ b/Engine/Plugins/FX/Niagara/Source/NiagaraEditor/Private/NiagaraActions.cpp
@@ -811,6 +811,46 @@ void FNiagaraDataInterfaceNodeActionProvider_DataChannelWrite::GetNodeContextMen
 	Section.AddSubMenu("InitForDataChannelMenu", InitForDataChannelMenuText, InitForDataChannelMenuTooltipText,	FNewToolMenuDelegate::CreateLambda(CreateNodeContextMenu));
 }
 
+void FNiagaraDataInterfaceNodeActionProvider_DataChannelWrite::CollectAddPinActionsImpl(FNiagaraMenuActionCollector& Collector, UEdGraphPin* AddPin)const
+{
+	auto GatherAddPinsForChannel = [&](UNiagaraDataChannel* Channel)
+	{
+		TConstArrayView<FNiagaraDataChannelVariable> ChannelVars = Channel->GetVariables();
+		for (const FNiagaraDataChannelVariable& Var : ChannelVars)
+		{
+			FNiagaraTypeDefinition Type = Var.GetType();
+			if (Type.IsEnum() == false)
+			{
+				Type = FNiagaraTypeDefinition(FNiagaraTypeHelper::GetSWCStruct(Var.GetType().GetScriptStruct()));
+			}
+			FNiagaraVariable SWCVar(Type, Var.GetName());
+			FNiagaraEditorUtilities::ResetVariableToDefaultValue(SWCVar);
+
+			const UEdGraphPin* ConstAddPin = AddPin;
+
+			// The script variable is not a duplicate, add an entry for it.
+			FText Category;
+			FText DisplayName;
+			FText Tooltip;
+			{
+				Category = FText::Format(LOCTEXT("NDIWriteAddPinCatFmt", "Write to NDC {0}"), FText::FromString(Channel->GetAsset()->GetName()));
+				DisplayName = FText::FromName(SWCVar.GetName());
+				Tooltip = FText::Format(
+					LOCTEXT("NDIWritelAddPinTooltipFmt", "Write to the variable {0} from NDC {1}."),
+					FText::FromName(SWCVar.GetName()),
+					FText::FromString(Channel->GetAsset()->GetName()));
+			}
+
+			TSharedPtr<FNiagaraMenuAction> Action(new FNiagaraMenuAction(
+				Category, DisplayName, Tooltip, 0, FText::GetEmpty(),
+				FNiagaraMenuAction::FOnExecuteStackAction::CreateUObject(CastChecked<UNiagaraNodeWithDynamicPins>(AddPin->GetOwningNode()), &UNiagaraNodeWithDynamicPins::AddParameter, SWCVar, ConstAddPin->Direction)));
+			Action->SetParameterVariable(SWCVar);
+			Collector.AddAction(Action, 3);
+		}
+	};
+	UNiagaraDataChannel::ForEachDataChannel(GatherAddPinsForChannel);
+}
+
 void FNiagaraDataInterfaceNodeActionProvider_DataChannelWrite::GetInlineNodeContextMenuActionsImpl(UToolMenu* ToolMenu) const
 {
 	AddDataChannelInitActions(ToolMenu);
@@ -824,7 +864,7 @@ INiagaraDataInterfaceNodeActionProvider::FInlineMenuDisplayOptions FNiagaraDataI
 		return {};
 	}
 
-	if (FunctionCall->Signature.Name != TEXT("Write"))
+	if (FunctionCall->Signature.Name != TEXT("Write") && FunctionCall->Signature.Name != TEXT("Append"))
 	{
 		return {};
 	}
diff --git a/Engine/Plugins/FX/Niagara/Source/NiagaraEditor/Public/NiagaraActions.h b/Engine/Plugins/FX/Niagara/Source/NiagaraEditor/Public/NiagaraActions.h
--- a/Engine/Plugins/FX/Niagara/Source/NiagaraEditor/Public/NiagaraActions.h
+++ b/Engine/Plugins/FX/Niagara/Source/NiagaraEditor/Public/NiagaraActions.h
@@ -485,6 +485,7 @@ class FNiagaraDataInterfaceNodeActionProvider_DataChannelWrite : public INiagara
 public:
 
 	virtual void GetNodeContextMenuActionsImpl(UToolMenu* Menu, UGraphNodeContextMenuContext* Context, FNiagaraFunctionSignature Signature) const override;
+	virtual void CollectAddPinActionsImpl(FNiagaraMenuActionCollector& Collector, UEdGraphPin* AddPin)const override;
 	virtual void GetInlineNodeContextMenuActionsImpl(UToolMenu* ToolMenu) const override;
 	virtual FInlineMenuDisplayOptions GetInlineMenuDisplayOptionsImpl(UClass* DIClass, UEdGraphNode* Source) const override;
 
