- Added class & tag searching options to socket DI
-Will search up the attachment hierarchy until hitting a component that satisfies the requirements

#rnx
#tests Added an engine test that was used with development

#virtualized

diff --git a/Engine/Plugins/FX/Niagara/Source/Niagara/Internal/DataInterface/NiagaraDataInterfaceSocketReader.h b/Engine/Plugins/FX/Niagara/Source/Niagara/Internal/DataInterface/NiagaraDataInterfaceSocketReader.h
--- a/Engine/Plugins/FX/Niagara/Source/Niagara/Internal/DataInterface/NiagaraDataInterfaceSocketReader.h
+++ b/Engine/Plugins/FX/Niagara/Source/Niagara/Internal/DataInterface/NiagaraDataInterfaceSocketReader.h
@@ -70,9 +70,17 @@ public:
 	TLazyObjectPtr<AActor> SourceActor;
 
 	/** Source object asset to read sockets from, the transsform for these would be in relation to the Niagara system. */
-	UPROPERTY(EditAnywhere, Category = "SocketReader", meta = (EditConditionHides, EditCondition = "SourceMode == ENDISocketReaderSourceMode::Default || SourceMode == ENDISocketReaderSourceMode::SourceOnly", AllowedClasses= "/Script/Engine.StaticMesh,/Script/Engine.SkeletalMesh"))
+	UPROPERTY(EditAnywhere, Category = "SocketReader", meta = (EditConditionHides, EditCondition = "SourceMode == ENDISocketReaderSourceMode::Default || SourceMode == ENDISocketReaderSourceMode::SourceOnly", AllowedClasses="/Script/Engine.StaticMesh,/Script/Engine.SkeletalMesh"))
 	TObjectPtr<UObject> SourceAsset;
 
+	/** When looking for an attached parent component only accept this type of component. */
+	UPROPERTY(EditAnywhere, Category = "SocketReader", meta = (EditConditionHides, EditCondition = "SourceMode == ENDISocketReaderSourceMode::Default || SourceMode == ENDISocketReaderSourceMode::AttachedParentOnly", AllowedClasses="/Script/Engine.SceneComponent"))
+	TObjectPtr<UClass> AttachComponentClass;
+
+	/** When looking for an attached parent component it must have this tag to be considered. */
+	UPROPERTY(EditAnywhere, Category = "SocketReader", meta = (EditConditionHides, EditCondition = "SourceMode == ENDISocketReaderSourceMode::Default || SourceMode == ENDISocketReaderSourceMode::AttachedParentOnly", AllowedClasses = "/Script/Engine.SceneComponent"))
+	FName AttachComponentTag;
+
 	/**
 	Source object parameter binding.
 	Note: Source Mode impacts the order of operations.
diff --git a/Engine/Plugins/FX/Niagara/Source/Niagara/Private/DataInterface/NiagaraDataInterfaceSocketReader.cpp b/Engine/Plugins/FX/Niagara/Source/Niagara/Private/DataInterface/NiagaraDataInterfaceSocketReader.cpp
--- a/Engine/Plugins/FX/Niagara/Source/Niagara/Private/DataInterface/NiagaraDataInterfaceSocketReader.cpp
+++ b/Engine/Plugins/FX/Niagara/Source/Niagara/Private/DataInterface/NiagaraDataInterfaceSocketReader.cpp
@@ -231,12 +231,27 @@ namespace NDISocketReaderLocal
 		// Resolve Attached Parent?
 		if (bTryAttachedParent && !ResolvedObject)
 		{
-			USceneComponent* Component = SystemInstance->GetAttachComponent();
-			if (Component->IsA<UNiagaraComponent>())
+			USceneComponent* AttachComponent = SystemInstance->GetAttachComponent();
+			if (AttachComponent->IsA<UNiagaraComponent>())
 			{
-				ResolvedObject = Component->GetAttachParent();
-				//-TODO: Do we need to add tags or look until we find sockets?
+				AttachComponent = AttachComponent->GetAttachParent();
 			}
+
+			// Look for the first viable component that satisfies the class and tag
+			if (SocketDI->AttachComponentClass || !SocketDI->AttachComponentTag.IsNone())
+			{
+				while (AttachComponent)
+				{
+					if ((!SocketDI->AttachComponentClass || AttachComponent->IsA(SocketDI->AttachComponentClass)) &&
+						(!SocketDI->AttachComponentTag.IsNone() && AttachComponent->ComponentHasTag(SocketDI->AttachComponentTag)) )
+					{
+						break;
+					}
+					AttachComponent = AttachComponent->GetAttachParent();
+				}
+			}
+
+			ResolvedObject = AttachComponent;
 		}
 
 		// Resolve Source
@@ -810,6 +825,8 @@ bool UNiagaraDataInterfaceSocketReader::Equals(const UNiagaraDataInterface* Othe
 	#endif
 		&& OtherTyped->SourceActor == SourceActor
 		&& OtherTyped->SourceAsset == SourceAsset
+		&& OtherTyped->AttachComponentClass == AttachComponentClass
+		&& OtherTyped->AttachComponentTag == AttachComponentTag
 		&& OtherTyped->ObjectParameterBinding == ObjectParameterBinding
 		&& OtherTyped->bUpdateSocketsPerFrame == bUpdateSocketsPerFrame
 		&& OtherTyped->bRequireCurrentFrameData == bRequireCurrentFrameData;
@@ -830,6 +847,8 @@ bool UNiagaraDataInterfaceSocketReader::CopyToInternal(UNiagaraDataInterface* De
 #endif
 	OtherTyped->SourceActor = SourceActor;
 	OtherTyped->SourceAsset = SourceAsset;
+	OtherTyped->AttachComponentClass = AttachComponentClass;
+	OtherTyped->AttachComponentTag = AttachComponentTag;
 	OtherTyped->ObjectParameterBinding = ObjectParameterBinding;
 	OtherTyped->bUpdateSocketsPerFrame = bUpdateSocketsPerFrame;
 	OtherTyped->bRequireCurrentFrameData = bRequireCurrentFrameData;
