- Fix mesh renderer emitter mode material random not working

diff --git a/Engine/Plugins/FX/Niagara/Source/Niagara/Private/NiagaraRendererMeshes.cpp b/Engine/Plugins/FX/Niagara/Source/Niagara/Private/NiagaraRendererMeshes.cpp
--- a/Engine/Plugins/FX/Niagara/Source/Niagara/Private/NiagaraRendererMeshes.cpp
+++ b/Engine/Plugins/FX/Niagara/Source/Niagara/Private/NiagaraRendererMeshes.cpp
@@ -1085,6 +1085,9 @@ FNiagaraMeshUniformBufferRef FNiagaraRendererMeshes::CreateVFUniformBuffer(const
 					FMemory::Memcpy(&Params.DefaultDynamicMaterialParameter3, ParameterBoundData + VFBoundOffsetsInParamStore[i], sizeof(FVector4f));
 					Params.MaterialParamValidMask |= 0xf000;
 					break;
+				case ENiagaraMeshVFLayout::Type::MaterialRandom:
+					FMemory::Memcpy(&Params.DefaultMatRandom, ParameterBoundData + VFBoundOffsetsInParamStore[i], sizeof(float));
+					break;
 				case ENiagaraMeshVFLayout::Type::CustomSorting:
 					// unsupported for now...
 					break;
