- Add scalability overrides to stateless emitters

#rnx

diff --git a/Engine/Plugins/FX/Niagara/Source/Niagara/Internal/Stateless/NiagaraStatelessEmitter.h b/Engine/Plugins/FX/Niagara/Source/Niagara/Internal/Stateless/NiagaraStatelessEmitter.h
--- a/Engine/Plugins/FX/Niagara/Source/Niagara/Internal/Stateless/NiagaraStatelessEmitter.h
+++ b/Engine/Plugins/FX/Niagara/Source/Niagara/Internal/Stateless/NiagaraStatelessEmitter.h
@@ -5,6 +5,7 @@
 #include "NiagaraStatelessCommon.h"
 #include "NiagaraStatelessSpawnInfo.h"
 #include "NiagaraDataSet.h"
+#include "NiagaraEFfectType.h"
 #include "NiagaraRendererProperties.h"
 #include "NiagaraSystemEmitterState.h"
 #include "Stateless/NiagaraStatelessEmitterTemplate.h"
@@ -56,6 +57,7 @@ public:
 	void CacheFromCompiledData();
 protected:
 	void BuildCompiledDataSet();
+	void ResolveScalabilitySettings();
 
 public:
 	template<typename TAction>
@@ -118,6 +120,7 @@ public:
 	NIAGARA_API UNiagaraStatelessModule* GetModule(UClass* Class) const;
 
 	NIAGARA_API FNiagaraPlatformSet& GetPlatformSet() { return Platforms; }
+	NIAGARA_API FNiagaraEmitterScalabilityOverrides& GetScalabilityOverrides() { return ScalabilityOverrides; }
 
 	UNiagaraStatelessEmitter* CreateAsDuplicate(FName InDuplicateName, UNiagaraSystem& InDuplicateOwnerSystem) const;
 
@@ -160,6 +163,9 @@ protected:
 	UPROPERTY(EditAnywhere, Category = "Scalability", meta = (DisplayInScalabilityContext))
 	FNiagaraPlatformSet Platforms;
 
+	UPROPERTY(EditAnywhere, Category = "Scalability", meta = (DisplayInScalabilityContext))
+	FNiagaraEmitterScalabilityOverrides ScalabilityOverrides;
+
 	UPROPERTY()
 	FNiagaraDataSetCompiledData ParticleDataSetCompiledData;
 
diff --git a/Engine/Plugins/FX/Niagara/Source/Niagara/Internal/Stateless/NiagaraStatelessEmitterData.h b/Engine/Plugins/FX/Niagara/Source/Niagara/Internal/Stateless/NiagaraStatelessEmitterData.h
--- a/Engine/Plugins/FX/Niagara/Source/Niagara/Internal/Stateless/NiagaraStatelessEmitterData.h
+++ b/Engine/Plugins/FX/Niagara/Source/Niagara/Internal/Stateless/NiagaraStatelessEmitterData.h
@@ -44,6 +44,7 @@ struct FNiagaraStatelessEmitterData
 
 	FNiagaraEmitterStateData						EmitterState;
 	TArray<FNiagaraStatelessSpawnInfo>				SpawnInfos;
+	float											SpawnCountScale = 1.0f;
 
 	TArray<TObjectPtr<UNiagaraRendererProperties>>	RendererProperties;
 	TArray<TObjectPtr<UNiagaraParameterCollection>>	BoundParameterCollections;
diff --git a/Engine/Plugins/FX/Niagara/Source/Niagara/Private/Stateless/NiagaraStatelessEmitter.cpp b/Engine/Plugins/FX/Niagara/Source/Niagara/Private/Stateless/NiagaraStatelessEmitter.cpp
--- a/Engine/Plugins/FX/Niagara/Source/Niagara/Private/Stateless/NiagaraStatelessEmitter.cpp
+++ b/Engine/Plugins/FX/Niagara/Source/Niagara/Private/Stateless/NiagaraStatelessEmitter.cpp
@@ -334,6 +334,9 @@ void UNiagaraStatelessEmitter::CacheFromCompiledData()
 		}
 	}
 
+	// Resolve scalability settings
+	ResolveScalabilitySettings();
+
 	// Build buffers that are shared across all instances
 	//-OPT: We should be able to build and serialize this data as part of the UNiagaraStatelessEmitter, potentially all of this data even since it's immutable and does not change at runtime
 	if (StatelessEmitterData->bCanEverExecute)
@@ -470,6 +473,35 @@ void UNiagaraStatelessEmitter::BuildCompiledDataSet()
 	StatelessEmitterData->ComponentOffsets = ComponentOffsets;
 }
 
+void UNiagaraStatelessEmitter::ResolveScalabilitySettings()
+{
+	const float DefaultSpawnCountScale = 1.0f;
+	StatelessEmitterData->SpawnCountScale = DefaultSpawnCountScale;
+
+	if (UNiagaraSystem* OwnerSystem = GetTypedOuter<UNiagaraSystem>())
+	{
+		if (UNiagaraEffectType* ActualEffectType = OwnerSystem->GetEffectType())
+		{
+			const FNiagaraEmitterScalabilitySettings& ScalabilitySettings = ActualEffectType->GetActiveEmitterScalabilitySettings();
+			if (ScalabilitySettings.bScaleSpawnCount)
+			{
+				StatelessEmitterData->SpawnCountScale = ScalabilitySettings.SpawnCountScale;
+			}
+		}
+	}
+
+	for (FNiagaraEmitterScalabilityOverride& Override : ScalabilityOverrides.Overrides)
+	{
+		if (Override.Platforms.IsActive())
+		{
+			if (Override.bOverrideSpawnCountScale)
+			{
+				StatelessEmitterData->SpawnCountScale = Override.bScaleSpawnCount ? Override.SpawnCountScale : DefaultSpawnCountScale;
+			}
+		}
+	}	
+}
+
 bool UNiagaraStatelessEmitter::SetUniqueEmitterName(const FString& InName)
 {
 	if (InName != UniqueEmitterName)
diff --git a/Engine/Plugins/FX/Niagara/Source/Niagara/Private/Stateless/NiagaraStatelessEmitterInstance.cpp b/Engine/Plugins/FX/Niagara/Source/Niagara/Private/Stateless/NiagaraStatelessEmitterInstance.cpp
--- a/Engine/Plugins/FX/Niagara/Source/Niagara/Private/Stateless/NiagaraStatelessEmitterInstance.cpp
+++ b/Engine/Plugins/FX/Niagara/Source/Niagara/Private/Stateless/NiagaraStatelessEmitterInstance.cpp
@@ -401,7 +401,7 @@ void FNiagaraStatelessEmitterInstance::InitSpawnInfos()
 		{
 			if (SpawnInfo.Type == ENiagaraStatelessSpawnInfoType::Rate)
 			{
-				const float SpawnRate = RandomStream.FRandRange(SpawnInfo.Rate.Min, SpawnInfo.Rate.Max);
+				const float SpawnRate = RandomStream.FRandRange(SpawnInfo.Rate.Min, SpawnInfo.Rate.Max) * EmitterData->SpawnCountScale;
 				if (SpawnRate > 0.0f)
 				{
 					FActiveSpawnRate& ActiveSpawnRate = ActiveSpawnRates.AddDefaulted_GetRef();
@@ -482,7 +482,8 @@ void FNiagaraStatelessEmitterInstance::InitSpawnInfosForLoop()
 			continue;
 		}
 
-		const int32 SpawnAmount = RandomStream.RandRange(SpawnInfo.Amount.Min, SpawnInfo.Amount.Max);
+		const int32 UnscaledSpawnAmount = RandomStream.RandRange(SpawnInfo.Amount.Min, SpawnInfo.Amount.Max);
+		const int32 SpawnAmount = FMath::FloorToInt(float(UnscaledSpawnAmount) * EmitterData->SpawnCountScale);
 		if (SpawnAmount <= 0)
 		{
 			continue;
diff --git a/Engine/Plugins/FX/Niagara/Source/NiagaraEditor/Private/Toolkits/SystemToolkitModes/NiagaraSystemToolkitMode_Scalability.cpp b/Engine/Plugins/FX/Niagara/Source/NiagaraEditor/Private/Toolkits/SystemToolkitModes/NiagaraSystemToolkitMode_Scalability.cpp
--- a/Engine/Plugins/FX/Niagara/Source/NiagaraEditor/Private/Toolkits/SystemToolkitModes/NiagaraSystemToolkitMode_Scalability.cpp
+++ b/Engine/Plugins/FX/Niagara/Source/NiagaraEditor/Private/Toolkits/SystemToolkitModes/NiagaraSystemToolkitMode_Scalability.cpp
@@ -140,6 +140,13 @@ void FNiagaraSystemToolkitMode_Scalability::BindEmitterPreviewOverrides()
 			StatelessEmitter->GetPlatformSet().OnOverrideQualityLevelDelegate.BindUObject(ScalabilityViewModel, &UNiagaraSystemScalabilityViewModel::GetPreviewQualityLevel);
 			StatelessEmitter->GetPlatformSet().OnOverrideActiveDeviceProfileDelegate.BindUObject(ScalabilityViewModel, &UNiagaraSystemScalabilityViewModel::GetPreviewDeviceProfile);
 
+			// we do the same for the overrides
+			for (FNiagaraEmitterScalabilityOverride& EmitterScalabilityOverride : StatelessEmitter->GetScalabilityOverrides().Overrides)
+			{
+				EmitterScalabilityOverride.Platforms.OnOverrideQualityLevelDelegate.BindUObject(ScalabilityViewModel, &UNiagaraSystemScalabilityViewModel::GetPreviewQualityLevel);
+				EmitterScalabilityOverride.Platforms.OnOverrideActiveDeviceProfileDelegate.BindUObject(ScalabilityViewModel, &UNiagaraSystemScalabilityViewModel::GetPreviewDeviceProfile);
+			}
+
 			// and also for the renderers
 			StatelessEmitter->ForEachRenderer([=](UNiagaraRendererProperties* RendererProperties)
 			{
@@ -196,6 +203,12 @@ void FNiagaraSystemToolkitMode_Scalability::UnbindEmitterPreviewOverrides()
 			StatelessEmitter->GetPlatformSet().OnOverrideQualityLevelDelegate.Unbind();
 			StatelessEmitter->GetPlatformSet().OnOverrideActiveDeviceProfileDelegate.Unbind();
 
+			for (FNiagaraEmitterScalabilityOverride& EmitterScalabilityOverride : StatelessEmitter->GetScalabilityOverrides().Overrides)
+			{
+				EmitterScalabilityOverride.Platforms.OnOverrideQualityLevelDelegate.Unbind();
+				EmitterScalabilityOverride.Platforms.OnOverrideActiveDeviceProfileDelegate.Unbind();
+			}
+
 			StatelessEmitter->ForEachRenderer([](UNiagaraRendererProperties* RendererProperties)
 			{
 				RendererProperties->Platforms.OnOverrideQualityLevelDelegate.Unbind();
diff --git a/Engine/Plugins/FX/Niagara/Source/NiagaraEditor/Private/Widgets/Scalability/SNiagaraSystemResolvedScalabilitySettings.cpp b/Engine/Plugins/FX/Niagara/Source/NiagaraEditor/Private/Widgets/Scalability/SNiagaraSystemResolvedScalabilitySettings.cpp
--- a/Engine/Plugins/FX/Niagara/Source/NiagaraEditor/Private/Widgets/Scalability/SNiagaraSystemResolvedScalabilitySettings.cpp
+++ b/Engine/Plugins/FX/Niagara/Source/NiagaraEditor/Private/Widgets/Scalability/SNiagaraSystemResolvedScalabilitySettings.cpp
@@ -1,6 +1,7 @@
 // Copyright Epic Games, Inc. All Rights Reserved.
 
 #include "SNiagaraSystemResolvedScalabilitySettings.h"
+#include "Stateless/NiagaraStatelessEmitter.h"
 
 #include "IDetailTreeNode.h"
 #include "IPropertyRowGenerator.h"
@@ -160,6 +161,16 @@ TSharedRef<SWidget> SScalabilityResolvedRow::GenerateWidgetForColumn(const FName
 			}
 		}
 	}
+	else if (UNiagaraStatelessEmitter* StatelessEmitter = Cast<UNiagaraStatelessEmitter>(ScalabilityRowData->OwningObject))
+	{
+		for (FNiagaraEmitterScalabilityOverride& EmitterScalabilityOverride : StatelessEmitter->GetScalabilityOverrides().Overrides)
+		{
+			if (ParentWidget.Pin()->GetSystemViewModel()->GetScalabilityViewModel()->IsPlatformActive(EmitterScalabilityOverride.Platforms))
+			{
+				CurrentEmitterOverride = &EmitterScalabilityOverride;
+			}
+		}
+	}
 
 	// we need to check whether the property is overridden so we can determine whether to check the resolved or default property for a scalability property being active
 
diff --git a/Engine/Plugins/FX/Niagara/Source/NiagaraEditorWidgets/Private/SNiagaraOverviewStackNode.cpp b/Engine/Plugins/FX/Niagara/Source/NiagaraEditorWidgets/Private/SNiagaraOverviewStackNode.cpp
--- a/Engine/Plugins/FX/Niagara/Source/NiagaraEditorWidgets/Private/SNiagaraOverviewStackNode.cpp
+++ b/Engine/Plugins/FX/Niagara/Source/NiagaraEditorWidgets/Private/SNiagaraOverviewStackNode.cpp
@@ -890,6 +890,7 @@ EVisibility SNiagaraOverviewStackNode::GetScalabilityIndicatorVisibility() const
 			if (UNiagaraStatelessEmitter* StatelessEmitter = EmitterHandle->GetStatelessEmitter())
 			{
 				EmitterPlatformSet = &StatelessEmitter->GetPlatformSet();
+				EmitterScalabilityOverrides = &StatelessEmitter->GetScalabilityOverrides();
 			}
 		}
 
