- Fix capitalization in NiagaraEffectType.h include (EF vs Ef)

#rnx

diff --git a/Engine/Plugins/FX/Niagara/Source/Niagara/Internal/Stateless/NiagaraStatelessEmitter.h b/Engine/Plugins/FX/Niagara/Source/Niagara/Internal/Stateless/NiagaraStatelessEmitter.h
--- a/Engine/Plugins/FX/Niagara/Source/Niagara/Internal/Stateless/NiagaraStatelessEmitter.h
+++ b/Engine/Plugins/FX/Niagara/Source/Niagara/Internal/Stateless/NiagaraStatelessEmitter.h
@@ -5,7 +5,7 @@
 #include "NiagaraStatelessCommon.h"
 #include "NiagaraStatelessSpawnInfo.h"
 #include "NiagaraDataSet.h"
-#include "NiagaraEFfectType.h"
+#include "NiagaraEffectType.h"
 #include "NiagaraRendererProperties.h"
 #include "NiagaraSystemEmitterState.h"
 #include "Stateless/NiagaraStatelessEmitterTemplate.h"
