- Manual emitter dependencies
- Deps can be used to avoid races when emitter stages might execute concurrently with one another, for example if Emitter A wrote to a RT and emitter B read from the same RT it would potentially race or the order will be undefined.  By adding a dep to emitter B on emitter A we guarantee that the work will not overlap

#tests Test system as described in the CL to ensure the read / write order was now correct and could be controlled

diff --git a/Engine/Plugins/FX/Niagara/Source/Niagara/Classes/NiagaraEmitter.h b/Engine/Plugins/FX/Niagara/Source/Niagara/Classes/NiagaraEmitter.h
--- a/Engine/Plugins/FX/Niagara/Source/Niagara/Classes/NiagaraEmitter.h
+++ b/Engine/Plugins/FX/Niagara/Source/Niagara/Classes/NiagaraEmitter.h
@@ -12,6 +12,7 @@
 #include "INiagaraMergeManager.h"
 #include "NiagaraAssetTagDefinitions.h"
 #include "NiagaraEffectType.h"
+#include "NiagaraDataInterfaceEmitterBinding.h"
 #include "NiagaraDataSetAccessor.h"
 #include "NiagaraBoundsCalculator.h"
 #include "NiagaraRendererProperties.h"
@@ -343,6 +344,14 @@ struct FVersionedNiagaraEmitterData
 	UPROPERTY(EditAnywhere, AdvancedDisplay, Category = "Emitter", meta = (EditCondition = "AllocationMode != EParticleAllocationMode::AutomaticEstimate", EditConditionHides))
 	int32 PreAllocationCount = 0;
 
+	/**
+	List of emitter dependencies to use when calculating the execution order for emitter particle scripts.
+	This is generally only required when you are using advanced features, such as reading / writing to a data interface in different emitters
+	and need to ensure the emitters can not run concurrently with one another, either on the CPU or the GPU.
+	**/
+	UPROPERTY(EditAnywhere, Category = "Emitter")
+	TArray<FNiagaraDataInterfaceEmitterBinding> EmitterDependencies;
+
 	UPROPERTY()
 	FNiagaraEmitterScriptProperties UpdateScriptProps;
 
diff --git a/Engine/Plugins/FX/Niagara/Source/Niagara/Private/NiagaraDataInterfaceEmitterBinding.cpp b/Engine/Plugins/FX/Niagara/Source/Niagara/Private/NiagaraDataInterfaceEmitterBinding.cpp
--- a/Engine/Plugins/FX/Niagara/Source/Niagara/Private/NiagaraDataInterfaceEmitterBinding.cpp
+++ b/Engine/Plugins/FX/Niagara/Source/Niagara/Private/NiagaraDataInterfaceEmitterBinding.cpp
@@ -172,6 +172,38 @@ const FNiagaraEmitterHandle* FNiagaraDataInterfaceEmitterBinding::ResolveHandle(
 	return nullptr;
 }
 
+const FNiagaraEmitterHandle* FNiagaraDataInterfaceEmitterBinding::ResolveHandle(const UNiagaraSystem* OwnerSystem, const FNiagaraEmitterHandle* OwnerEmitter) const
+{
+	check(OwnerSystem);
+
+	if (BindingMode == ENiagaraDataInterfaceEmitterBindingMode::Self)
+	{
+		return OwnerEmitter;
+	}
+	else if (BindingMode == ENiagaraDataInterfaceEmitterBindingMode::Other)
+	{
+		if (!EmitterName.IsNone())
+		{
+			FNameBuilder EmitterNameString;
+			EmitterName.ToString(EmitterNameString);
+			FStringView EmitterNameStringView = EmitterNameString.ToView();
+
+			for (const FNiagaraEmitterHandle& EmitterHandle : OwnerSystem->GetEmitterHandles())
+			{
+				if (UNiagaraEmitter* NiagaraEmitter = EmitterHandle.GetInstance().Emitter)
+				{
+					//-TODO: UniqueEmitterName should probably be a FName?
+					if (EmitterNameStringView.Equals(NiagaraEmitter->GetUniqueEmitterName(), ESearchCase::IgnoreCase))
+					{
+						return &EmitterHandle;
+					}
+				}
+			}
+		}
+	}
+	return nullptr;
+}
+
 UNiagaraEmitter* FNiagaraDataInterfaceEmitterBinding::Resolve(const UNiagaraDataInterface* DataInterface) const
 {
 	const FNiagaraEmitterHandle* EmitterHandle = ResolveHandle(DataInterface);
diff --git a/Engine/Plugins/FX/Niagara/Source/Niagara/Private/NiagaraSystem.cpp b/Engine/Plugins/FX/Niagara/Source/Niagara/Private/NiagaraSystem.cpp
--- a/Engine/Plugins/FX/Niagara/Source/Niagara/Private/NiagaraSystem.cpp
+++ b/Engine/Plugins/FX/Niagara/Source/Niagara/Private/NiagaraSystem.cpp
@@ -2036,6 +2036,7 @@ void UNiagaraSystem::ComputeEmittersExecutionOrder()
 		const FNiagaraEmitterHandle& EmitterHandle = EmitterHandles[EmitterIdx];
 		FVersionedNiagaraEmitterData* EmitterData = EmitterHandle.GetEmitterData();
 
+		EmitterExecutionOrder[EmitterIdx].bStartNewOverlapGroup = false;
 		EmitterExecutionOrder[EmitterIdx].EmitterIndex = EmitterIdx;
 		EmitterPriorities[EmitterIdx] = -1;
 
@@ -2051,6 +2052,15 @@ void UNiagaraSystem::ComputeEmittersExecutionOrder()
 
 		EmitterDependencies.SetNum(0, EAllowShrinking::No);
 
+		for (const FNiagaraDataInterfaceEmitterBinding& EmitterBinding : EmitterData->EmitterDependencies)
+		{
+			const FNiagaraEmitterHandle* ResolvedHandle = EmitterBinding.ResolveHandle(this, &EmitterHandle);
+			if (ResolvedHandle)
+			{
+				EmitterDependencies.Add(ResolvedHandle->GetInstance());
+			}
+		}
+
 		if (EmitterData->SimTarget == ENiagaraSimTarget::GPUComputeSim && EmitterData->GetGPUComputeScript())
 		{
 			// GPU emitters have a combined execution context for spawn and update.
diff --git a/Engine/Plugins/FX/Niagara/Source/Niagara/Public/NiagaraDataInterfaceEmitterBinding.h b/Engine/Plugins/FX/Niagara/Source/Niagara/Public/NiagaraDataInterfaceEmitterBinding.h
--- a/Engine/Plugins/FX/Niagara/Source/Niagara/Public/NiagaraDataInterfaceEmitterBinding.h
+++ b/Engine/Plugins/FX/Niagara/Source/Niagara/Public/NiagaraDataInterfaceEmitterBinding.h
@@ -39,6 +39,9 @@ struct FNiagaraDataInterfaceEmitterBinding
 	/** Resolves the binding to an FNiagaraEmitterHandle or nullptr if it's invalid. */
 	const FNiagaraEmitterHandle* ResolveHandle(const UNiagaraDataInterface* DataInterface) const;
 
+	/** Resolves the binding to an FNiagaraEmitterHandle or nullptr if it's invalid. */
+	const FNiagaraEmitterHandle* ResolveHandle(const UNiagaraSystem* OwnerSystem, const FNiagaraEmitterHandle* OwnerEmitter) const;
+
 	/** Resolves the emitter name */
 	FString ResolveUniqueName(const UNiagaraDataInterface* DataInterface) const;
 
diff --git a/Engine/Plugins/FX/Niagara/Source/NiagaraEditor/Private/Widgets/SNiagaraSystemViewport.cpp b/Engine/Plugins/FX/Niagara/Source/NiagaraEditor/Private/Widgets/SNiagaraSystemViewport.cpp
--- a/Engine/Plugins/FX/Niagara/Source/NiagaraEditor/Private/Widgets/SNiagaraSystemViewport.cpp
+++ b/Engine/Plugins/FX/Niagara/Source/NiagaraEditor/Private/Widgets/SNiagaraSystemViewport.cpp
@@ -423,6 +423,11 @@ void FNiagaraSystemViewportClient::DrawEmitterExecutionOrder(UNiagaraComponent*
 			const FNiagaraEmitterHandle& EmitterHandle = NiagaraSystem->GetEmitterHandle(EmitterExecIndex.EmitterIndex);
 			if (FVersionedNiagaraEmitterData* NiagaraEmitter = EmitterHandle.GetEmitterData())
 			{
+				if (EmitterExecIndex.bStartNewOverlapGroup)
+				{
+					Canvas->DrawShadowedString(CurrentX, CurrentY, TEXT("<Overlap Barrier>"), Font, FLinearColor::White);
+					CurrentY += FontHeight;
+				}
 				Canvas->DrawShadowedString(CurrentX, CurrentY, *FString::Printf(TEXT("%d - %s"), ++DisplayIndex, NiagaraEmitter->GetDebugSimName()), Font, FLinearColor::White);
 				CurrentY += FontHeight;
 			}
