- Emitter scalability for stateless
- Includes min / max distance & visibility reactions

#rnx

diff --git a/Engine/Plugins/FX/Niagara/Source/Niagara/Classes/NiagaraEmitterInstance.h b/Engine/Plugins/FX/Niagara/Source/Niagara/Classes/NiagaraEmitterInstance.h
--- a/Engine/Plugins/FX/Niagara/Source/Niagara/Classes/NiagaraEmitterInstance.h
+++ b/Engine/Plugins/FX/Niagara/Source/Niagara/Classes/NiagaraEmitterInstance.h
@@ -39,7 +39,8 @@ public:
 	virtual void BindParameters(bool bExternalOnly) = 0;
 	virtual void UnbindParameters(bool bExternalOnly) = 0;
 
-	virtual void PreTick() {};
+	virtual bool ShouldTick() const = 0;
+	virtual void PreTick() {}
 	virtual void Tick(float DeltaSeconds) = 0;
 	//~End: Define Emitter Interface
 
@@ -65,8 +66,6 @@ public:
 	bool IsInactive() const { return ExecutionState == ENiagaraExecutionState::Inactive; }
 	bool IsComplete() const { return ExecutionState == ENiagaraExecutionState::Complete || ExecutionState == ENiagaraExecutionState::Disabled; }
 
-	FORCEINLINE bool ShouldTick() const { return ExecutionState == ENiagaraExecutionState::Active || GetNumParticles() > 0; }
-
 	//-TODO:Stateless: Does this need to be virtual?  Can we cache the value after ticking / allocating the data / have a cache value somewhere instead?
 	NIAGARA_API virtual int32 GetNumParticles() const;
 	int32 GetTotalSpawnedParticles() const { return TotalSpawnedParticles; }
diff --git a/Engine/Plugins/FX/Niagara/Source/Niagara/Internal/NiagaraEmitterInstanceImpl.h b/Engine/Plugins/FX/Niagara/Source/Niagara/Internal/NiagaraEmitterInstanceImpl.h
--- a/Engine/Plugins/FX/Niagara/Source/Niagara/Internal/NiagaraEmitterInstanceImpl.h
+++ b/Engine/Plugins/FX/Niagara/Source/Niagara/Internal/NiagaraEmitterInstanceImpl.h
@@ -57,6 +57,7 @@ public:
 	virtual TConstArrayView<UNiagaraRendererProperties*> GetRenderers() const override;
 	virtual void BindParameters(bool bExternalOnly) override;
 	virtual void UnbindParameters(bool bExternalOnly) override;
+	virtual bool ShouldTick() const override;
 	virtual void PreTick() override;
 	virtual void Tick(float DeltaSeconds) override;
 	// FNiagaraEmitterInstance Impl
diff --git a/Engine/Plugins/FX/Niagara/Source/Niagara/Internal/NiagaraSystemEmitterState.h b/Engine/Plugins/FX/Niagara/Source/Niagara/Internal/NiagaraSystemEmitterState.h
--- a/Engine/Plugins/FX/Niagara/Source/Niagara/Internal/NiagaraSystemEmitterState.h
+++ b/Engine/Plugins/FX/Niagara/Source/Niagara/Internal/NiagaraSystemEmitterState.h
@@ -88,21 +88,48 @@ struct FNiagaraEmitterStateData
 	UPROPERTY(EditAnywhere, Category = "Emitter State", meta = (ClampMin = "0.0"))
 	FNiagaraDistributionRangeFloat LoopDuration = FNiagaraDistributionRangeFloat(1.0f);
 
-	UPROPERTY(EditAnywhere, Category = "Emitter State", meta = (EditCondition = "LoopBehavior != ENiagaraLoopBehavior::Once", EditConditionHides))
-	bool bRecalculateDurationEachLoop = false;
-
 	UPROPERTY(EditAnywhere, Category = "Emitter State", meta = (ClampMin = "0.0"))
 	FNiagaraDistributionRangeFloat LoopDelay = FNiagaraDistributionRangeFloat(0.0f);
 
+	UPROPERTY(EditAnywhere, Category = "Emitter State", meta = (EditCondition = "LoopBehavior != ENiagaraLoopBehavior::Once", EditConditionHides, DisplayAfter = "LoopDuration"))
+	uint32 bRecalculateDurationEachLoop : 1 = false;
+
 	UPROPERTY(EditAnywhere, Category = "Emitter State", meta = (EditCondition = "LoopBehavior != ENiagaraLoopBehavior::Once", EditConditionHides))
-	bool bDelayFirstLoopOnly = false;
+	uint32 bDelayFirstLoopOnly : 1 = false;
 
 	UPROPERTY(EditAnywhere, Category = "Emitter State", meta = (EditCondition = "LoopBehavior != ENiagaraLoopBehavior::Once && !bDelayFirstLoopOnly", EditConditionHides))
-	bool bRecalculateDelayEachLoop = false;
+	uint32 bRecalculateDelayEachLoop : 1 = false;
 
-	//UPROPERTY(EditAnywhere, Category="Emitter State")
-	//ENiagaraStatelessEmitterState_SelfSystem Scalability = ENiagaraStatelessEmitterState_SelfSystem::Self;
-	//Enable Distance Culling
-	//Enable Visibility Culling
-	//Reset Age On Awaken
+	UPROPERTY(EditAnywhere, Category = "Scalability")
+	uint32 bEnableDistanceCulling : 1 = false;
+
+	UPROPERTY(EditAnywhere, Category = "Scalability", meta = (DisplayAfter = "MaxDistanceReaction"))
+	uint32 bEnableVisibilityCulling : 1 = false;
+
+	UPROPERTY(EditAnywhere, Category = "Scalability", meta = (EditCondition = "bEnableDistanceCulling", EditConditionHides))
+	uint32 bMinDistanceEnabled : 1 = false;
+
+	UPROPERTY(EditAnywhere, Category = "Scalability", meta = (EditCondition = "bEnableDistanceCulling", EditConditionHides))
+	uint32 bMaxDistanceEnabled : 1 = false;
+
+	UPROPERTY(EditAnywhere, Category = "Scalability", meta = (DisplayAfter = "VisibilityCullDelay"))
+	uint32 bResetAgeOnAwaken : 1 = false;
+
+	UPROPERTY(EditAnywhere, Category = "Scalability", meta = (EditCondition = "bEnableDistanceCulling && bMinDistanceEnabled", EditConditionHides))
+	float MinDistance = 0.0f;
+
+	UPROPERTY(EditAnywhere, Category = "Scalability", meta = (EditCondition = "bEnableDistanceCulling && bMinDistanceEnabled", EditConditionHides))
+	ENiagaraExecutionStateManagement MinDistanceReaction = ENiagaraExecutionStateManagement::Awaken;
+
+	UPROPERTY(EditAnywhere, Category = "Scalability", meta = (EditCondition = "bEnableDistanceCulling && bMaxDistanceEnabled", EditConditionHides))
+	float MaxDistance = 5000.0f;
+
+	UPROPERTY(EditAnywhere, Category = "Scalability", meta = (EditCondition = "bEnableDistanceCulling && bMaxDistanceEnabled", EditConditionHides))
+	ENiagaraExecutionStateManagement MaxDistanceReaction = ENiagaraExecutionStateManagement::SleepAndLetParticlesFinish;
+
+	UPROPERTY(EditAnywhere, Category = "Scalability", meta = (EditCondition = "bEnableVisibilityCulling", EditConditionHides))
+	ENiagaraExecutionStateManagement VisibilityCullReaction = ENiagaraExecutionStateManagement::SleepAndLetParticlesFinish;
+
+	UPROPERTY(EditAnywhere, Category = "Scalability", meta = (EditCondition = "bEnableVisibilityCulling", EditConditionHides))
+	float VisibilityCullDelay = 1.0f;
 };
diff --git a/Engine/Plugins/FX/Niagara/Source/Niagara/Internal/Stateless/NiagaraStatelessEmitterInstance.h b/Engine/Plugins/FX/Niagara/Source/Niagara/Internal/Stateless/NiagaraStatelessEmitterInstance.h
--- a/Engine/Plugins/FX/Niagara/Source/Niagara/Internal/Stateless/NiagaraStatelessEmitterInstance.h
+++ b/Engine/Plugins/FX/Niagara/Source/Niagara/Internal/Stateless/NiagaraStatelessEmitterInstance.h
@@ -65,6 +65,7 @@ public:
 	virtual TConstArrayView<UNiagaraRendererProperties*> GetRenderers() const override;
 	virtual void BindParameters(bool bExternalOnly) override;
 	virtual void UnbindParameters(bool bExternalOnly) override;
+	virtual bool ShouldTick() const override;
 	virtual void Tick(float DeltaSeconds) override;
 	// FNiagaraEmitterInstance Impl
 
@@ -82,9 +83,12 @@ private:
 	void CalculateBounds();
 	void SendRenderData();
 
-	void InitSpawnInfos();
-	void InitSpawnInfosForLoop();
+	void InitSpawnInfos(float InitializationAge);
+	void InitSpawnInfosForLoop(float InitializationAge);
 	void TickSpawnInfos();
+	void CropSpawnInfos();
+	void KillSpawnInfos();
+	void RestartSpawnInfos();
 
 	//-TODO: This can be shared perhaps?
 	void SetExecutionStateInternal(ENiagaraExecutionState InExecutionState);
@@ -107,6 +111,7 @@ private:
 	TArray<FNiagaraStatelessRuntimeSpawnInfo>	SpawnInfos;
 	TArray<FActiveSpawnRate>					ActiveSpawnRates;
 
+	ENiagaraExecutionStateManagement			ScalabilityState = ENiagaraExecutionStateManagement::Awaken;
 	int32										LoopCount = 0;
 	float										CurrentLoopDuration = 0.0f;
 	float										CurrentLoopDelay = 0.0f;
diff --git a/Engine/Plugins/FX/Niagara/Source/Niagara/Private/NiagaraEmitterInstanceImpl.cpp b/Engine/Plugins/FX/Niagara/Source/Niagara/Private/NiagaraEmitterInstanceImpl.cpp
--- a/Engine/Plugins/FX/Niagara/Source/Niagara/Private/NiagaraEmitterInstanceImpl.cpp
+++ b/Engine/Plugins/FX/Niagara/Source/Niagara/Private/NiagaraEmitterInstanceImpl.cpp
@@ -1051,6 +1051,11 @@ void FNiagaraEmitterInstanceImpl::TickRapidIterationParameters()
 }
 #endif
 
+bool FNiagaraEmitterInstanceImpl::ShouldTick() const
+{
+	return ExecutionState == ENiagaraExecutionState::Active || GetNumParticles() > 0;
+}
+
 /** 
   * PreTick - handles killing dead particles, emitter death, and buffer swaps
   */
diff --git a/Engine/Plugins/FX/Niagara/Source/Niagara/Private/Stateless/NiagaraStatelessEmitter.cpp b/Engine/Plugins/FX/Niagara/Source/Niagara/Private/Stateless/NiagaraStatelessEmitter.cpp
--- a/Engine/Plugins/FX/Niagara/Source/Niagara/Private/Stateless/NiagaraStatelessEmitter.cpp
+++ b/Engine/Plugins/FX/Niagara/Source/Niagara/Private/Stateless/NiagaraStatelessEmitter.cpp
@@ -269,6 +269,11 @@ void UNiagaraStatelessEmitter::CacheFromCompiledData()
 	// Setup emitter state
 	StatelessEmitterData->EmitterState = EmitterState;
 
+	if (StatelessEmitterData->EmitterState.bEnableDistanceCulling)
+	{
+		StatelessEmitterData->EmitterState.bEnableDistanceCulling = StatelessEmitterData->EmitterState.bMinDistanceEnabled | StatelessEmitterData->EmitterState.bMaxDistanceEnabled;
+	}
+
 	StatelessEmitterData->EmitterState.LoopDuration.Min = FMath::Max(StatelessEmitterData->EmitterState.LoopDuration.Min, UE_KINDA_SMALL_NUMBER);
 	StatelessEmitterData->EmitterState.LoopDuration.Max = FMath::Max(StatelessEmitterData->EmitterState.LoopDuration.Max, UE_KINDA_SMALL_NUMBER);
 
diff --git a/Engine/Plugins/FX/Niagara/Source/Niagara/Private/Stateless/NiagaraStatelessEmitterInstance.cpp b/Engine/Plugins/FX/Niagara/Source/Niagara/Private/Stateless/NiagaraStatelessEmitterInstance.cpp
--- a/Engine/Plugins/FX/Niagara/Source/Niagara/Private/Stateless/NiagaraStatelessEmitterInstance.cpp
+++ b/Engine/Plugins/FX/Niagara/Source/Niagara/Private/Stateless/NiagaraStatelessEmitterInstance.cpp
@@ -97,7 +97,7 @@ void FNiagaraStatelessEmitterInstance::Init(int32 InEmitterIndex)
 	GPUDataBufferInterfaces = RenderThreadDataPtr.Get();
 
 	InitEmitterState();
-	InitSpawnInfos();
+	InitSpawnInfos(0.0f);
 }
 
 void FNiagaraStatelessEmitterInstance::ResetSimulation(bool bKillExisting)
@@ -129,9 +129,10 @@ void FNiagaraStatelessEmitterInstance::ResetSimulation(bool bKillExisting)
 	bEmitterEnabled_CNC = bEmitterEnabled_GT;
 
 	InitEmitterState();
-	InitSpawnInfos();
+	InitSpawnInfos(0.0f);
 
 	ExecutionState = ENiagaraExecutionState::Active;
+	ScalabilityState = ENiagaraExecutionStateManagement::Awaken;
 	if (NiagaraStateless::FEmitterInstance_RT* RenderThreadData = RenderThreadDataPtr.Get())
 	{
 		ENQUEUE_RENDER_COMMAND(UpdateStatelessAge)(
@@ -174,7 +175,7 @@ bool FNiagaraStatelessEmitterInstance::HandleCompletion(bool bForce)
 
 int32 FNiagaraStatelessEmitterInstance::GetNumParticles() const
 {
-	return bCanEverExecute ? EmitterData->CalculateActiveParticles(RandomSeed, SpawnInfos, Age) : 0;
+	return bCanEverExecute && (SpawnInfos.Num() > 0) ? EmitterData->CalculateActiveParticles(RandomSeed, SpawnInfos, Age) : 0;
 }
 
 TConstArrayView<UNiagaraRendererProperties*> FNiagaraStatelessEmitterInstance::GetRenderers() const
@@ -220,6 +221,11 @@ void FNiagaraStatelessEmitterInstance::UnbindParameters(bool bExternalOnly)
 	}
 }
 
+bool FNiagaraStatelessEmitterInstance::ShouldTick() const
+{
+	return ExecutionState <= ENiagaraExecutionState::Inactive;
+}
+
 void FNiagaraStatelessEmitterInstance::Tick(float DeltaSeconds)
 {
 	Age += DeltaSeconds;
@@ -242,65 +248,140 @@ void FNiagaraStatelessEmitterInstance::InitEmitterState()
 
 void FNiagaraStatelessEmitterInstance::TickEmitterState()
 {
-	if (ParentSystemInstance)
+	// Update execution state based on the parent which be told to go inactive / complete
 	{
-		const ENiagaraExecutionState ParentExecutionState = ParentSystemInstance->GetActualExecutionState();
+		const ENiagaraExecutionState ParentExecutionState = ParentSystemInstance ? ParentSystemInstance->GetActualExecutionState() : ENiagaraExecutionState::Complete;
 		if (ParentExecutionState > ExecutionState)
 		{
 			SetExecutionStateInternal(ParentExecutionState);
 		}
 	}
 
-	if (ExecutionState == ENiagaraExecutionState::Active)
+	// If we are going inactive and we hit zero particles we are now complete
+	if (ExecutionState == ENiagaraExecutionState::Inactive)
 	{
-		if (Age < CurrentLoopAgeEnd)
+		if (GetNumParticles() == 0)
 		{
-			return;
+			SetExecutionStateInternal(ENiagaraExecutionState::Complete);
 		}
+	}
 
-		const FNiagaraEmitterStateData& EmitterState = EmitterData->EmitterState;
-		if (EmitterState.LoopBehavior == ENiagaraLoopBehavior::Once)
-		{
-			SetExecutionStateInternal(ENiagaraExecutionState::Inactive);
-			return;
-		}
+	// If we are not active we don't need to evaluate loops / scalability anymore
+	if ( ExecutionState != ENiagaraExecutionState::Active )
+	{
+		return;
+	}
+
+	const FNiagaraEmitterStateData& EmitterState = EmitterData->EmitterState;
 
-		// Keep looping until we find out which loop we are in as a small loop age + large DT could result in crossing multiple loops
-		do
+	// Evaluate scalability state
+	{
+		ENiagaraExecutionStateManagement RequestedScalabilityState = ENiagaraExecutionStateManagement::Awaken;
+		if (EmitterState.bEnableVisibilityCulling)
 		{
-			++LoopCount;
-			if (EmitterState.LoopBehavior == ENiagaraLoopBehavior::Multiple && LoopCount >= EmitterState.LoopCount)
+			const FNiagaraSystemParameters& SystemParameters = ParentSystemInstance->GetSystemParameters();
+			if (SystemParameters.EngineTimeSinceRendered > EmitterState.VisibilityCullDelay)
 			{
-				SetExecutionStateInternal(ENiagaraExecutionState::Inactive);
-				return;
+				RequestedScalabilityState = EmitterState.VisibilityCullReaction;
 			}
+		}
 
-			if (EmitterState.bRecalculateDurationEachLoop)
+		if (EmitterState.bEnableDistanceCulling)
+		{
+			const float LODDistance = ParentSystemInstance->GetLODDistance();
+
+			if (LODDistance > EmitterState.MaxDistance)
 			{
-				CurrentLoopDuration = RandomStream.FRandRange(EmitterState.LoopDuration.Min, EmitterState.LoopDuration.Max);
+				RequestedScalabilityState = EmitterState.MaxDistanceReaction;
 			}
-
-			if (EmitterState.bDelayFirstLoopOnly)
+			else if (LODDistance < EmitterState.MinDistance)
 			{
-				CurrentLoopDelay	= 0.0f;
+				RequestedScalabilityState = EmitterState.MinDistanceReaction;
 			}
-			else if ( EmitterState.bRecalculateDelayEachLoop )
+		}
+
+		// We need to transition the state
+		if (RequestedScalabilityState != ScalabilityState)
+		{
+			ScalabilityState = RequestedScalabilityState;
+			switch (RequestedScalabilityState)
 			{
-				CurrentLoopDelay	= RandomStream.FRandRange(EmitterState.LoopDelay.Min, EmitterState.LoopDelay.Max);
+				case ENiagaraExecutionStateManagement::Awaken:
+					if (EmitterState.bResetAgeOnAwaken)
+					{
+						ResetSimulation(false);
+					}
+					break;
+
+				case ENiagaraExecutionStateManagement::SleepAndLetParticlesFinish:
+				case ENiagaraExecutionStateManagement::KillAfterParticlesFinish:
+					CropSpawnInfos();
+					break;
+
+				case ENiagaraExecutionStateManagement::SleepAndClearParticles:
+					KillSpawnInfos();
+					break;
+
+				case ENiagaraExecutionStateManagement::KillImmediately:
+					SetExecutionStateInternal(ENiagaraExecutionState::Complete);
+					return;
 			}
+		}
 
-			CurrentLoopAgeStart	= CurrentLoopAgeEnd;
-			CurrentLoopAgeEnd	= CurrentLoopAgeStart + CurrentLoopDelay + CurrentLoopDuration;
-			InitSpawnInfosForLoop();
-		} while ( Age >= CurrentLoopAgeEnd );
+		// Perform any per frame operations for scalability state
+		switch (ScalabilityState)
+		{
+			case ENiagaraExecutionStateManagement::KillAfterParticlesFinish:
+				if ( GetNumParticles() == 0 )
+				{
+					SetExecutionStateInternal(ENiagaraExecutionState::Complete);
+					return;
+				}
+				break;
+		}
 	}
-	else if (ExecutionState == ENiagaraExecutionState::Inactive)
+
+	// Evaluate emitter state
+	if ( Age >= CurrentLoopAgeEnd )
 	{
-		if (SpawnInfos.Num() == 0 || GetNumParticles() == 0)
+		// Do we only execute a single loop?
+		if (EmitterState.LoopBehavior == ENiagaraLoopBehavior::Once)
 		{
-			SetExecutionStateInternal(ENiagaraExecutionState::Complete);
+			SetExecutionStateInternal(ENiagaraExecutionState::Inactive);
+		}
+		// Multi-loop inject our new spawn infos
+		else
+		{
+			// Keep looping until we find out which loop we are in as a small loop age + large DT could result in crossing multiple loops
+			do
+			{
+				++LoopCount;
+				if (EmitterState.LoopBehavior == ENiagaraLoopBehavior::Multiple && LoopCount >= EmitterState.LoopCount)
+				{
+					SetExecutionStateInternal(ENiagaraExecutionState::Inactive);
+					break;
+				}
+
+				if (EmitterState.bRecalculateDurationEachLoop)
+				{
+					CurrentLoopDuration = RandomStream.FRandRange(EmitterState.LoopDuration.Min, EmitterState.LoopDuration.Max);
+				}
+
+				if (EmitterState.bDelayFirstLoopOnly)
+				{
+					CurrentLoopDelay = 0.0f;
+				}
+				else if (EmitterState.bRecalculateDelayEachLoop)
+				{
+					CurrentLoopDelay = RandomStream.FRandRange(EmitterState.LoopDelay.Min, EmitterState.LoopDelay.Max);
+				}
+
+				CurrentLoopAgeStart = CurrentLoopAgeEnd;
+				CurrentLoopAgeEnd = CurrentLoopAgeStart + CurrentLoopDelay + CurrentLoopDuration;
+
+				InitSpawnInfosForLoop(CurrentLoopAgeStart);
+			} while (Age >= CurrentLoopAgeEnd);
 		}
-		return;
 	}
 }
 
@@ -391,33 +472,35 @@ void FNiagaraStatelessEmitterInstance::SendRenderData()
 	);
 }
 
-void FNiagaraStatelessEmitterInstance::InitSpawnInfos()
+void FNiagaraStatelessEmitterInstance::InitSpawnInfos(float InitializationAge)
 {
-	UniqueIndexOffset = 0;
+	// If we are not enabled, or not awake from scalability skip adding
+	if (!bEmitterEnabled_GT || (ScalabilityState != ENiagaraExecutionStateManagement::Awaken))
+	{
+		return;
+	}
 
-	if (bEmitterEnabled_CNC)
+	for (const FNiagaraStatelessSpawnInfo& SpawnInfo : EmitterData->SpawnInfos)
 	{
-		for (const FNiagaraStatelessSpawnInfo& SpawnInfo : EmitterData->SpawnInfos)
+		if (SpawnInfo.Type == ENiagaraStatelessSpawnInfoType::Rate)
 		{
-			if (SpawnInfo.Type == ENiagaraStatelessSpawnInfoType::Rate)
+			const float SpawnRate = RandomStream.FRandRange(SpawnInfo.Rate.Min, SpawnInfo.Rate.Max) * EmitterData->SpawnCountScale;
+			if (SpawnRate > 0.0f)
 			{
-				const float SpawnRate = RandomStream.FRandRange(SpawnInfo.Rate.Min, SpawnInfo.Rate.Max) * EmitterData->SpawnCountScale;
-				if (SpawnRate > 0.0f)
-				{
-					FActiveSpawnRate& ActiveSpawnRate = ActiveSpawnRates.AddDefaulted_GetRef();
-					ActiveSpawnRate.Rate = SpawnRate;
-					ActiveSpawnRate.SpawnTime = CurrentLoopDelay;
-				}
+				FActiveSpawnRate& ActiveSpawnRate = ActiveSpawnRates.AddDefaulted_GetRef();
+				ActiveSpawnRate.Rate = SpawnRate;
+				ActiveSpawnRate.SpawnTime = CurrentLoopDelay;
 			}
 		}
-
-		InitSpawnInfosForLoop();
 	}
+
+	InitSpawnInfosForLoop(InitializationAge);
 }
 
-void FNiagaraStatelessEmitterInstance::InitSpawnInfosForLoop()
+void FNiagaraStatelessEmitterInstance::InitSpawnInfosForLoop(float InitializationAge)
 {
-	if (bEmitterEnabled_CNC == false)
+	// If we are not enabled, or not awake from scalability skip adding
+	if (!bEmitterEnabled_GT || (ScalabilityState != ENiagaraExecutionStateManagement::Awaken))
 	{
 		return;
 	}
@@ -425,7 +508,7 @@ void FNiagaraStatelessEmitterInstance::InitSpawnInfosForLoop()
 	// Add the next chunk for any active spawn rates
 	for (const FActiveSpawnRate& SpawnInfo : ActiveSpawnRates)
 	{
-		float SpawnTime = FMath::Max(CurrentLoopAgeStart - SpawnInfo.SpawnTime, 0.0f);
+		float SpawnTime = FMath::Max(InitializationAge - SpawnInfo.SpawnTime, 0.0f);
 		SpawnTime = FMath::CeilToFloat(SpawnTime * SpawnInfo.Rate) / SpawnInfo.Rate;
 		SpawnTime += SpawnInfo.SpawnTime;
 		if (SpawnTime >= CurrentLoopAgeEnd)
@@ -489,11 +572,17 @@ void FNiagaraStatelessEmitterInstance::InitSpawnInfosForLoop()
 			continue;
 		}
 
+		const float SpawnTime = CurrentLoopAgeStart + CurrentLoopDelay + SpawnInfo.SpawnTime;
+		if (SpawnTime < InitializationAge)
+		{
+			continue;
+		}
+
 		FNiagaraStatelessRuntimeSpawnInfo& NewSpawnInfo = SpawnInfos.AddDefaulted_GetRef();
 		NewSpawnInfo.Type			= ENiagaraStatelessSpawnInfoType::Burst;
 		NewSpawnInfo.UniqueOffset	= UniqueIndexOffset;
-		NewSpawnInfo.SpawnTimeStart	= CurrentLoopAgeStart + CurrentLoopDelay + SpawnInfo.SpawnTime;
-		NewSpawnInfo.SpawnTimeEnd	= NewSpawnInfo.SpawnTimeStart;
+		NewSpawnInfo.SpawnTimeStart	= SpawnTime;
+		NewSpawnInfo.SpawnTimeEnd	= SpawnTime;
 		NewSpawnInfo.Amount			= SpawnAmount;
 
 		UniqueIndexOffset += SpawnAmount;
@@ -503,54 +592,19 @@ void FNiagaraStatelessEmitterInstance::InitSpawnInfosForLoop()
 
 void FNiagaraStatelessEmitterInstance::TickSpawnInfos()
 {
-	if (bEmitterEnabled_CNC != bEmitterEnabled_GT)
+	const bool bNewEmitterEnabled = bEmitterEnabled_GT && (ScalabilityState == ENiagaraExecutionStateManagement::Awaken);
+
+	if (bEmitterEnabled_CNC != bNewEmitterEnabled)
 	{
-		bEmitterEnabled_CNC = bEmitterEnabled_GT;
+		bEmitterEnabled_CNC = bNewEmitterEnabled;
 
 		if (bEmitterEnabled_CNC)
 		{
-			if ( ActiveSpawnRates.Num() > 0 )
-			{
-				for (const FActiveSpawnRate& SpawnInfo : ActiveSpawnRates)
-				{
-					float SpawnTime = FMath::Max(Age - SpawnInfo.SpawnTime, 0.0f);
-					SpawnTime = FMath::CeilToFloat(SpawnTime * SpawnInfo.Rate) / SpawnInfo.Rate;
-					SpawnTime += SpawnInfo.SpawnTime;
-					if (SpawnTime >= CurrentLoopAgeEnd)
-					{
-						continue;
-					}
-
-					FNiagaraStatelessRuntimeSpawnInfo& NewSpawnInfo = SpawnInfos.AddDefaulted_GetRef();
-					NewSpawnInfo.Type = ENiagaraStatelessSpawnInfoType::Rate;
-					NewSpawnInfo.UniqueOffset = UniqueIndexOffset;
-					NewSpawnInfo.SpawnTimeStart = SpawnTime;
-					NewSpawnInfo.SpawnTimeEnd = CurrentLoopAgeEnd;
-					NewSpawnInfo.Rate = SpawnInfo.Rate;
-
-					const float ActiveDuration = CurrentLoopAgeEnd - SpawnTime;
-					const int32 NumSpawned = FMath::FloorToInt(ActiveDuration * SpawnInfo.Rate);
-
-					UniqueIndexOffset += NumSpawned;
-					bSpawnInfosDirty = true;
-				}
-			}
+			RestartSpawnInfos();
 		}
 		else
 		{
-			for (auto SpawnInfoIt=SpawnInfos.CreateIterator(); SpawnInfoIt; ++SpawnInfoIt)
-			{
-				FNiagaraStatelessRuntimeSpawnInfo& SpawnInfo = *SpawnInfoIt;
-				if (SpawnInfo.SpawnTimeStart > Age)
-				{
-					SpawnInfoIt.RemoveCurrent();
-				}
-				else
-				{
-					SpawnInfo.SpawnTimeEnd = FMath::Min(SpawnInfo.SpawnTimeEnd, Age);
-				}
-			}
-			bSpawnInfosDirty = true;
+			CropSpawnInfos();
 		}
 	}
 
@@ -563,6 +617,49 @@ void FNiagaraStatelessEmitterInstance::TickSpawnInfos()
 	);
 }
 
+void FNiagaraStatelessEmitterInstance::CropSpawnInfos()
+{
+	if (!SpawnInfos.Num() && !ActiveSpawnRates.Num())
+	{
+		return;
+	}
+
+	ActiveSpawnRates.Empty(SpawnInfos.Num());
+
+	const float MaxLifetime = EmitterData->LifetimeRange.Max;
+	for (auto it = SpawnInfos.CreateIterator(); it; ++it)
+	{
+		FNiagaraStatelessRuntimeSpawnInfo& SpawnInfo = *it;
+		if (SpawnInfo.Type == ENiagaraStatelessSpawnInfoType::Rate)
+		{
+			SpawnInfo.SpawnTimeEnd = FMath::Min(SpawnInfo.SpawnTimeEnd, Age);
+		}
+		if (Age < SpawnInfo.SpawnTimeStart || Age >= SpawnInfo.SpawnTimeEnd + MaxLifetime)
+		{
+			it.RemoveCurrent();
+		}
+	}
+
+	bSpawnInfosDirty = true;
+}
+
+void FNiagaraStatelessEmitterInstance::KillSpawnInfos()
+{
+	if (!SpawnInfos.Num() && !ActiveSpawnRates.Num())
+	{
+		return;
+	}
+
+	SpawnInfos.Empty();
+	ActiveSpawnRates.Empty();
+	bSpawnInfosDirty = true;
+}
+
+void FNiagaraStatelessEmitterInstance::RestartSpawnInfos()
+{
+	InitSpawnInfos(Age);
+}
+
 void FNiagaraStatelessEmitterInstance::SetExecutionStateInternal(ENiagaraExecutionState RequestedExecutionState)
 {
 	if (ExecutionState == RequestedExecutionState)
@@ -579,46 +676,23 @@ void FNiagaraStatelessEmitterInstance::SetExecutionStateInternal(ENiagaraExecuti
 		case ENiagaraExecutionState::Inactive:
 			if (EmitterData->EmitterState.InactiveResponse != ENiagaraEmitterInactiveResponse::Kill)
 			{
-				ActiveSpawnRates.Empty();
-
-				// Crop & Remove Spawn Infos
-				const float MaxLifetime = EmitterData->LifetimeRange.Max;
-				for (auto it=SpawnInfos.CreateIterator(); it; ++it)
-				{
-					FNiagaraStatelessRuntimeSpawnInfo& SpawnInfo = *it;
-					if (SpawnInfo.Type == ENiagaraStatelessSpawnInfoType::Rate)
-					{
-						SpawnInfo.SpawnTimeEnd = FMath::Min(SpawnInfo.SpawnTimeEnd, Age);
-					}
-					if (Age < SpawnInfo.SpawnTimeStart || Age >= SpawnInfo.SpawnTimeEnd + MaxLifetime)
-					{
-						it.RemoveCurrent();
-					}
-				}
-
-				if (SpawnInfos.Num() > 0)
-				{
-					//-TODO: Better way to send data to renderer
-					if (NiagaraStateless::FEmitterInstance_RT* RenderThreadData = RenderThreadDataPtr.Get())
-					{
-						ENQUEUE_RENDER_COMMAND(FInitStatelessEmitter)(
-							[RenderThreadData=RenderThreadDataPtr.Get(), SpawnInfos_RT=SpawnInfos](FRHICommandListImmediate& RHICmdList) mutable
-							{
-								RenderThreadData->SpawnInfos = MoveTemp(SpawnInfos_RT);
-							}
-						);
-					}
+				CropSpawnInfos();
+			}
 
-					ExecutionState = ENiagaraExecutionState::Inactive;
-					break;
-				}
+			if (SpawnInfos.Num() > 0)
+			{
+				ExecutionState = ENiagaraExecutionState::Inactive;
+			}
+			else
+			{
+				KillSpawnInfos();
+				ExecutionState = ENiagaraExecutionState::Complete;
 			}
-			// Intentional fall through as we are complete
+			break;
 
 		case ENiagaraExecutionState::InactiveClear:
 		case ENiagaraExecutionState::Complete:
-			SpawnInfos.Empty();
-			ActiveSpawnRates.Empty();
+			KillSpawnInfos();
 			ExecutionState = ENiagaraExecutionState::Complete;
 			break;
 	}
