Fixed issue with CPU->GPU NDC transfer when buffer layouts don't perfectly match.
CPU already handles this case but GPU did not.

diff --git a/Engine/Plugins/FX/Niagara/Source/Niagara/Private/NiagaraDataSet.cpp b/Engine/Plugins/FX/Niagara/Source/Niagara/Private/NiagaraDataSet.cpp
--- a/Engine/Plugins/FX/Niagara/Source/Niagara/Private/NiagaraDataSet.cpp
+++ b/Engine/Plugins/FX/Niagara/Source/Niagara/Private/NiagaraDataSet.cpp
@@ -1281,13 +1281,15 @@ void FNiagaraDataBuffer::PushCPUBuffersToGPU(const TArray<FNiagaraDataBufferRef>
 {
 	uint32 NewCount = 0;
 	check(GetOwner()->GetSimTarget() == ENiagaraSimTarget::GPUComputeSim);
+	
+	bool bIdenticalLayout = true;
 	for (FNiagaraDataBuffer* Buffer : SourceBuffers)
 	{
 		if (Buffer)
 		{
 			NewCount += Buffer->GetNumInstances();
 
-			checkSlow(Buffer->GetOwner()->GetVariables() == GetOwner()->GetVariables())
+			bIdenticalLayout &= Buffer->GetOwner()->GetCompiledData().GetLayoutHash() == GetOwner()->GetCompiledData().GetLayoutHash();
 		}
 	}
 
@@ -1304,42 +1306,130 @@ void FNiagaraDataBuffer::PushCPUBuffersToGPU(const TArray<FNiagaraDataBufferRef>
 		uint8* MappedBufferInt32 = GPUBufferInt.Buffer ? (uint8*)RHICmdList.LockBuffer(GPUBufferInt.Buffer, 0, IntComponents * Int32Stride, RLM_WriteOnly) : nullptr;
 		uint8* MappedBufferHalf = GPUBufferHalf.Buffer ? (uint8*)RHICmdList.LockBuffer(GPUBufferHalf.Buffer, 0, HalfComponents * HalfStride, RLM_WriteOnly) : nullptr;
 
-		for (uint32 CompIdx = 0; CompIdx < FloatComponents; ++CompIdx)
+		if(bIdenticalLayout)
 		{
-			float* Dest = (float*)MappedBufferFloat;
-			MappedBufferFloat += FloatStride;
-			for (FNiagaraDataBuffer* Buffer : SourceBuffers)
+			//We have identical layout between all src buffers and the dest buffer so we can copy easily without worrying about missing variables.
+			for (uint32 CompIdx = 0; CompIdx < FloatComponents; ++CompIdx)
 			{
-				const float* Src = Buffer->GetInstancePtrFloat(CompIdx, 0);
-				int32 CopyInstances = Buffer->GetNumInstances();
-				FMemory::Memcpy(Dest, Src, CopyInstances * sizeof(float));
-				Dest += CopyInstances;
+				float* Dest = (float*)MappedBufferFloat;
+				MappedBufferFloat += FloatStride;
+				for (FNiagaraDataBuffer* Buffer : SourceBuffers)
+				{
+					const float* Src = Buffer->GetInstancePtrFloat(CompIdx, 0);
+					int32 CopyInstances = Buffer->GetNumInstances();
+					FMemory::Memcpy(Dest, Src, CopyInstances * sizeof(float));
+					Dest += CopyInstances;
+				}
 			}
-		}
 
-		for (uint32 CompIdx = 0; CompIdx < IntComponents; ++CompIdx)
-		{
-			int32* Dest = (int32*)MappedBufferInt32;
-			MappedBufferInt32 += Int32Stride;
-			for (FNiagaraDataBuffer* Buffer : SourceBuffers)
+			for (uint32 CompIdx = 0; CompIdx < IntComponents; ++CompIdx)
 			{
-				const int32* Src = Buffer->GetInstancePtrInt32(CompIdx, 0);
-				int32 CopyInstances = Buffer->GetNumInstances();
-				FMemory::Memcpy(Dest, Src, CopyInstances * sizeof(int32));
-				Dest += CopyInstances;
+				int32* Dest = (int32*)MappedBufferInt32;
+				MappedBufferInt32 += Int32Stride;
+				for (FNiagaraDataBuffer* Buffer : SourceBuffers)
+				{
+					const int32* Src = Buffer->GetInstancePtrInt32(CompIdx, 0);
+					int32 CopyInstances = Buffer->GetNumInstances();
+					FMemory::Memcpy(Dest, Src, CopyInstances * sizeof(int32));
+					Dest += CopyInstances;
+				}
 			}
-		}
 
-		for (uint32 CompIdx = 0; CompIdx < HalfComponents; ++CompIdx)
+			for (uint32 CompIdx = 0; CompIdx < HalfComponents; ++CompIdx)
+			{
+				FFloat16* Dest = (FFloat16*)MappedBufferHalf;
+				MappedBufferHalf += HalfStride;
+				for (FNiagaraDataBuffer* Buffer : SourceBuffers)
+				{
+					const FFloat16* Src = Buffer->GetInstancePtrHalf(CompIdx, 0);
+					int32 CopyInstances = Buffer->GetNumInstances();
+					FMemory::Memcpy(Dest, Src, CopyInstances * sizeof(FFloat16));
+					Dest += CopyInstances;
+				}
+			}
+		}
+		else
 		{
-			FFloat16* Dest = (FFloat16*)MappedBufferHalf;
-			MappedBufferHalf += HalfStride;
-			for (FNiagaraDataBuffer* Buffer : SourceBuffers)
+			//The layouts between src and dest buffer locations don't perfectly match so we must go variable by variable and copy over those that exist in the src data.
+			//Zeroing any that do not.
+			const FNiagaraDataSetCompiledData& CompiledData = GetOwner()->GetCompiledData();
+			
+			for (int32 VarIdx = 0; VarIdx < CompiledData.Variables.Num(); ++VarIdx)
 			{
-				const FFloat16* Src = Buffer->GetInstancePtrHalf(CompIdx, 0);
-				int32 CopyInstances = Buffer->GetNumInstances();
-				FMemory::Memcpy(Dest, Src, CopyInstances * sizeof(FFloat16));
-				Dest += CopyInstances;
+				const FNiagaraVariableBase& DestVar = CompiledData.Variables[VarIdx];
+				const FNiagaraVariableLayoutInfo& DestVarLayout = CompiledData.VariableLayouts[VarIdx];
+
+				uint32 DestInstStart = 0;
+				for (FNiagaraDataBuffer* SrcBuffer : SourceBuffers)
+				{
+					if(SrcBuffer->GetNumInstances() == 0)
+					{
+						continue;
+					}
+
+					const FNiagaraDataSetCompiledData& SrcCompiledData = SrcBuffer->GetOwner()->GetCompiledData();
+
+					const int32 SrcVarIdx = SrcCompiledData.Variables.IndexOfByKey(DestVar);
+					if (SrcVarIdx != INDEX_NONE)
+					{
+						//The variable was found so memcpy from the relevant src->dest buffer locations.
+						const FNiagaraVariableLayoutInfo& SrcVarLayout = SrcCompiledData.VariableLayouts[SrcVarIdx];
+						check(SrcVarLayout.GetNumFloatComponents() == DestVarLayout.GetNumFloatComponents());
+						for (uint32 CompIdx = 0; CompIdx < DestVarLayout.GetNumFloatComponents(); ++CompIdx)
+						{
+							const int32 SrcCompOffset = SrcVarLayout.GetFloatComponentStart() + CompIdx;
+							const int32 DestCompOffest = DestVarLayout.GetFloatComponentStart() + CompIdx;
+							const float* SrcStart = SrcBuffer->GetInstancePtrFloat(SrcCompOffset, 0);
+
+							float* Dst = (float*)(MappedBufferFloat + (DestCompOffest * FloatStride)) + DestInstStart;
+							FMemory::Memcpy(Dst, SrcStart, SrcBuffer->GetNumInstances() * sizeof(float));
+						}
+						check(SrcVarLayout.GetNumInt32Components() == DestVarLayout.GetNumInt32Components());
+						for (uint32 CompIdx = 0; CompIdx < DestVarLayout.GetNumInt32Components(); ++CompIdx)
+						{
+							const int32 SrcCompOffset = SrcVarLayout.GetInt32ComponentStart() + CompIdx;
+							const int32 DestCompOffest = DestVarLayout.GetInt32ComponentStart() + CompIdx;
+							const int32* SrcStart = SrcBuffer->GetInstancePtrInt32(SrcCompOffset, 0);
+
+							int32* Dst = (int32*)(MappedBufferInt32 + (DestCompOffest * Int32Stride)) + DestInstStart;
+							FMemory::Memcpy(Dst, SrcStart, SrcBuffer->GetNumInstances() * sizeof(int32));
+						}
+						check(SrcVarLayout.GetNumHalfComponents() == DestVarLayout.GetNumHalfComponents());
+						for (uint32 CompIdx = 0; CompIdx < DestVarLayout.GetNumHalfComponents(); ++CompIdx)
+						{
+							const int32 SrcCompOffset = SrcVarLayout.GetHalfComponentStart() + CompIdx;
+							const int32 DestCompOffest = DestVarLayout.GetHalfComponentStart() + CompIdx;
+							const FFloat16* SrcStart = SrcBuffer->GetInstancePtrHalf(SrcCompOffset, 0);
+
+							FFloat16* Dst = (FFloat16*)(MappedBufferHalf + (DestCompOffest * HalfStride)) + DestInstStart;
+							FMemory::Memcpy(Dst, SrcStart, SrcBuffer->GetNumInstances() * sizeof(FFloat16));
+						}
+					}
+					else
+					{
+						//This variable is not present in the source buffer so zero this out.
+						for (uint32 CompIdx = 0; CompIdx < DestVarLayout.GetNumFloatComponents(); ++CompIdx)
+						{
+							const int32 DestCompOffest = DestVarLayout.GetFloatComponentStart() + CompIdx;
+							float* Dst = (float*)(MappedBufferFloat + (DestCompOffest * FloatStride)) + DestInstStart;
+							FMemory::Memzero(Dst, SrcBuffer->GetNumInstances() * sizeof(float));
+						}
+						for (uint32 CompIdx = 0; CompIdx < DestVarLayout.GetNumInt32Components(); ++CompIdx)
+						{
+							const int32 DestCompOffest = DestVarLayout.GetInt32ComponentStart() + CompIdx;
+							int32* Dst = (int32*)(MappedBufferInt32 + (DestCompOffest * Int32Stride)) + DestInstStart;
+							FMemory::Memzero(Dst, SrcBuffer->GetNumInstances() * sizeof(int32));
+						}
+						for (uint32 CompIdx = 0; CompIdx < DestVarLayout.GetNumHalfComponents(); ++CompIdx)
+						{
+							const int32 DestCompOffest = DestVarLayout.GetHalfComponentStart() + CompIdx;
+							FFloat16* Dst = (FFloat16*)(MappedBufferHalf + (DestCompOffest * HalfStride)) + DestInstStart;
+							FMemory::Memzero(Dst, SrcBuffer->GetNumInstances() * sizeof(FFloat16));
+						}
+					}
+
+					DestInstStart += SrcBuffer->GetNumInstances();
+				}
 			}
 		}
 
