Substrate - added GBuffer sampling from Niagara GBuffer interface.
Moved Substrate GBuffer sampling to public and now shared with MaterialTemplate.

Fixed wrong defines preventing Substrate to work correctly with Substrate.

diff --git a/Engine/Plugins/FX/Niagara/Shaders/Private/NiagaraDataInterfaceCollisionQuery.ush b/Engine/Plugins/FX/Niagara/Shaders/Private/NiagaraDataInterfaceCollisionQuery.ush
--- a/Engine/Plugins/FX/Niagara/Shaders/Private/NiagaraDataInterfaceCollisionQuery.ush
+++ b/Engine/Plugins/FX/Niagara/Shaders/Private/NiagaraDataInterfaceCollisionQuery.ush
@@ -3,6 +3,8 @@
 #pragma once
 
 #include "/Engine/Private/Common.ush"
+
+#define SubstratePublicStruct SubstratePublic
 #include "/Engine/Public/SubstratePublic.ush"
 
 #define NDICOLLISIONQUERY_USE_GBUFFER_NORMAL    (FEATURE_LEVEL >= FEATURE_LEVEL_SM5)
@@ -49,7 +51,7 @@ void NDICollisionQuery_QuerySceneDepthGPU(in float3 In_SamplePos, in float3 In_L
 		FLWCVector3 WorldPosition = WorldPositionFromSceneDepth(ScreenPosition.xy, SceneDepth);
 
 #if NDICOLLISIONQUERY_USE_GBUFFER_NORMAL
-#if Substrate_ENABLED
+#if SUBSTRATE_ENABLED
 		uint2 ScreenPos = uint2(ScreenUV * View.BufferSizeAndInvSize.xy);
 		float3 WorldNormal = SubstratePublic_GetWorldNormal(ScreenPos);
 #else
diff --git a/Engine/Plugins/FX/Niagara/Shaders/Private/NiagaraDataInterfaceGBufferTemplate.ush b/Engine/Plugins/FX/Niagara/Shaders/Private/NiagaraDataInterfaceGBufferTemplate.ush
--- a/Engine/Plugins/FX/Niagara/Shaders/Private/NiagaraDataInterfaceGBufferTemplate.ush
+++ b/Engine/Plugins/FX/Niagara/Shaders/Private/NiagaraDataInterfaceGBufferTemplate.ush
@@ -1,6 +1,8 @@
 // Copyright Epic Games, Inc. All Rights Reserved.
 
 #include "/Engine/Private/DeferredShadingCommon.ush"
+
+#define SubstratePublicStruct SubstratePublic
 #include "/Engine/Public/SubstratePublic.ush"
 
 #define NDIGBUFFER_HAS_GBUFFER				(FEATURE_LEVEL >= FEATURE_LEVEL_SM5)
@@ -12,6 +14,14 @@
 Texture2D {ParameterName}_VelocityTexture;
 SamplerState {ParameterName}_VelocityTextureSampler;
 
+#if SUBSTRATE_ENABLED
+FSubstrateGBuffer GetNiagaraSubstrateGBuffer(in float2 ScreenUV)
+{
+	uint2 ScreenPos = uint2(ScreenUV * View.BufferSizeAndInvSize.xy);
+	return SubstratePublic_GetSubstrateGBufferFromFirstBSDF(ScreenPos);
+}
+#endif
+
 float2 ApplyViewportOffset_{ParameterName}(float2 UV, bool bApplyViewportOffset)
 {
 	if (bApplyViewportOffset)
@@ -25,9 +35,15 @@ void DecodeDiffuseColor_{ParameterName}(in float2 ScreenUV, in bool bApplyViewpo
 {
 #if NDIGBUFFER_HAS_GBUFFER
 	ScreenUV = ApplyViewportOffset_{ParameterName}(ScreenUV, bApplyViewportOffset);
-	FGBufferData GBufferData = GetGBufferData(ScreenUV);
 	IsValid = true;
-	OutValue = GBufferData.DiffuseColor;
+
+	#if SUBSTRATE_ENABLED
+		uint2 ScreenPos = uint2(ScreenUV * View.BufferSizeAndInvSize.xy);
+		OutValue = SubstratePublic_GetWorldNormal(ScreenPos);
+	#else
+		FGBufferData GBufferData = GetGBufferData(ScreenUV);
+		OutValue = GBufferData.DiffuseColor;
+	#endif
 #else
 	IsValid = false;
 	OutValue = 0;
@@ -40,7 +56,7 @@ void DecodeWorldNormal_{ParameterName}(in float2 ScreenUV, in bool bApplyViewpor
 	ScreenUV = ApplyViewportOffset_{ParameterName}(ScreenUV, bApplyViewportOffset);
 	IsValid = true;
 
-	#if Substrate_ENABLED
+	#if SUBSTRATE_ENABLED
 		uint2 ScreenPos = uint2(ScreenUV * View.BufferSizeAndInvSize.xy);
 		OutValue = SubstratePublic_GetWorldNormal(ScreenPos);
 	#else
@@ -120,9 +136,15 @@ void DecodeBaseColor_{ParameterName}(in float2 ScreenUV, in bool bApplyViewportO
 {
 #if NDIGBUFFER_HAS_GBUFFER
 	ScreenUV = ApplyViewportOffset_{ParameterName}(ScreenUV, bApplyViewportOffset);
-	FGBufferData GBufferData = GetGBufferData(ScreenUV);
 	IsValid = true;
-	OutValue = GBufferData.BaseColor;
+
+	#if SUBSTRATE_ENABLED
+		OutValue = GetNiagaraSubstrateGBuffer(ScreenUV).BaseColor;
+	#else
+		FGBufferData GBufferData = GetGBufferData(ScreenUV);
+		OutValue = GBufferData.BaseColor;
+	#endif
+
 #else
 	IsValid = false;
 	OutValue = 0;
@@ -133,9 +155,15 @@ void DecodeMetallic_{ParameterName}(in float2 ScreenUV, in bool bApplyViewportOf
 {
 #if NDIGBUFFER_HAS_GBUFFER
 	ScreenUV = ApplyViewportOffset_{ParameterName}(ScreenUV, bApplyViewportOffset);
-	FGBufferData GBufferData = GetGBufferData(ScreenUV);
 	IsValid = true;
-	OutValue = GBufferData.Metallic;
+
+	#if SUBSTRATE_ENABLED
+		OutValue = GetNiagaraSubstrateGBuffer(ScreenUV).Metallic;
+	#else
+		FGBufferData GBufferData = GetGBufferData(ScreenUV);
+		OutValue = GBufferData.Metallic;
+	#endif
+
 #else
 	IsValid = false;
 	OutValue = 0;
@@ -146,9 +174,15 @@ void DecodeSpecular_{ParameterName}(in float2 ScreenUV, in bool bApplyViewportOf
 {
 #if NDIGBUFFER_HAS_GBUFFER
 	ScreenUV = ApplyViewportOffset_{ParameterName}(ScreenUV, bApplyViewportOffset);
-	FGBufferData GBufferData = GetGBufferData(ScreenUV);
 	IsValid = true;
-	OutValue = GBufferData.Specular;
+
+	#if SUBSTRATE_ENABLED
+		OutValue = GetNiagaraSubstrateGBuffer(ScreenUV).Specular;
+	#else
+		FGBufferData GBufferData = GetGBufferData(ScreenUV);
+		OutValue = GBufferData.Specular;
+	#endif
+
 #else
 	IsValid = false;
 	OutValue = 0;
@@ -159,9 +193,16 @@ void DecodeRoughness_{ParameterName}(in float2 ScreenUV, in bool bApplyViewportO
 {
 #if NDIGBUFFER_HAS_GBUFFER
 	ScreenUV = ApplyViewportOffset_{ParameterName}(ScreenUV, bApplyViewportOffset);
-	FGBufferData GBufferData = GetGBufferData(ScreenUV);
 	IsValid = true;
-	OutValue = GBufferData.Roughness;
+
+	#if SUBSTRATE_ENABLED
+		uint2 ScreenPos = uint2(ScreenUV * View.BufferSizeAndInvSize.xy);
+		OutValue = SubstratePublic_GetRoughness(ScreenPos);
+	#else
+		FGBufferData GBufferData = GetGBufferData(ScreenUV);
+		OutValue = GBufferData.Roughness;
+	#endif
+
 #else
 	IsValid = false;
 	OutValue = 0;
@@ -255,9 +296,15 @@ void DecodeShadingModelID_{ParameterName}(in float2 ScreenUV, in bool bApplyView
 {
 #if NDIGBUFFER_HAS_GBUFFER
 	ScreenUV = ApplyViewportOffset_{ParameterName}(ScreenUV, bApplyViewportOffset);
-	FGBufferData GBufferData = GetGBufferData(ScreenUV);
 	IsValid = true;
-	OutValue = GBufferData.ShadingModelID;
+
+	#if SUBSTRATE_ENABLED
+		OutValue = int(GetNiagaraSubstrateGBuffer(ScreenUV).ShadingModelID);
+	#else
+		FGBufferData GBufferData = GetGBufferData(ScreenUV);
+		OutValue = GBufferData.ShadingModelID;
+	#endif
+
 #else
 	IsValid = false;
 	OutValue = 0;
diff --git a/Engine/Shaders/Private/MaterialTemplate.ush b/Engine/Shaders/Private/MaterialTemplate.ush
--- a/Engine/Shaders/Private/MaterialTemplate.ush
+++ b/Engine/Shaders/Private/MaterialTemplate.ush
@@ -173,7 +173,17 @@
 #endif
 
 #if SUBSTRATE_ENABLED
+
 #include "/Engine/Private/Substrate/Substrate.ush"
+
+#ifdef SubstrateStruct
+	#define SubstratePublicStruct SubstrateStruct
+#else
+	#define SubstratePublicStruct Substrate
+#endif
+#include "/Engine/Public/SubstratePublic.ush"
+#undef SubstratePublicStruct
+
 #else
 struct FSubstrateData
 {
@@ -3087,43 +3097,31 @@ float4 SceneTextureLookup(float2 UV, int SceneTextureIndex, bool bFiltered)
 	}
 	#elif SUBSTRATE_ENABLED
 	{
-		#ifdef SubstrateStruct
-			#define SubstrateSceneTexture SubstrateStruct
-		#else
-			#define SubstrateSceneTexture Substrate
-		#endif
-		FSubstrateAddressing SubstrateAddressing   = GetSubstratePixelDataByteOffset(PixelPos, uint2(View.BufferSizeAndInvSize.xy), SubstrateSceneTexture.MaxBytesPerPixel);
-		FSubstratePixelHeader SubstratePixelHeader = UnpackSubstrateHeaderIn(SubstrateSceneTexture.MaterialTextureArray, SubstrateAddressing, SubstrateSceneTexture.TopLayerTexture);
+		FSubstrateGBuffer SubstrateGBuffer = SubstratePublic_GetSubstrateGBufferFromFirstBSDF(PixelPos);
 		BRANCH
-		if (SubstratePixelHeader.ClosureCount > 0)
+		if (SubstrateGBuffer.bIsValid > 0)
 		{
-			const FSubstrateSubsurfaceHeader SSSHeader = SubstrateLoadSubsurfaceHeader(SubstrateSceneTexture.MaterialTextureArray, SubstrateSceneTexture.FirstSliceStoringSubstrateSSSData, SubstrateAddressing.PixelCoords);
-			const uint SSSType = SubstrateSubSurfaceHeaderGetSSSType(SSSHeader);
-
-			const bool bSkipSSSMaterialOverride = true; // We double not want the material data to be affected by any override.
-			FSubstrateBSDF BSDF = UnpackSubstrateBSDFIn(SubstrateSceneTexture.MaterialTextureArray, SubstrateAddressing, SubstratePixelHeader, bSkipSSSMaterialOverride);
 			switch (SceneTextureIndex)
 			{
 				// order needs to match to ESceneTextureId
-				case PPI_DiffuseColor:			return float4(SubstrateGetBSDFDiffuseColor(BSDF), 0);
-				case PPI_SpecularColor:			return float4(SubstrateGetBSDFSpecularColor(BSDF), 0);
-				case PPI_SubsurfaceColor:		return float4(SubstrateGetBSDFSubSurfaceColor(BSDF), 0);
-				case PPI_BaseColor:				return float4(SubstrateGetBSDFBaseColor(BSDF), 0);
-				case PPI_Specular:				return SubstrateGetBSDFSpecular(BSDF);
-				case PPI_Metallic:				return SubstrateGetBSDFMetallic(BSDF);
-				case PPI_WorldNormal:			return float4(SubstrateGetWorldNormal(SubstratePixelHeader, BSDF, SubstrateAddressing), 0);
-				case PPI_Opacity:				return BSDF.HasScatteringData() ? SubstrateSubSurfaceHeaderGetOpacity(SSSHeader) : 1.f;
-				case PPI_Roughness:				return SubstrateGetBSDFRoughness(BSDF);
-				case PPI_MaterialAO:			return SubstrateGetAO(SubstratePixelHeader);
-				case PPI_ShadingModelColor:		return float4(GetShadingModelColor(SubstrateGetLegacyShadingModels(BSDF)), 1);
-				case PPI_ShadingModelID:		return float4(SubstrateGetLegacyShadingModels(BSDF), 0, 0, 0);
-				case PPI_StoredBaseColor:		return float4(SubstrateGetBSDFBaseColor(BSDF), 0);
-				case PPI_StoredSpecular:		return float4(SubstrateGetBSDFSpecular(BSDF).rrr, 0);
-				case PPI_WorldTangent:			return float4(SubstrateGetWorldTangent(SubstratePixelHeader, BSDF, SubstrateAddressing), 0);
-				case PPI_Anisotropy:			return SubstrateGetBSDFAnisotropy(BSDF);
+				case PPI_DiffuseColor:			return float4(SubstrateGBuffer.DiffuseColor, 0);
+				case PPI_SpecularColor:			return float4(SubstrateGBuffer.SpecularColor, 0);
+				case PPI_SubsurfaceColor:		return float4(SubstrateGBuffer.SubsurfaceColor, 0);
+				case PPI_BaseColor:				return float4(SubstrateGBuffer.BaseColor, 0);
+				case PPI_Specular:				return SubstrateGBuffer.Specular;
+				case PPI_Metallic:				return SubstrateGBuffer.Metallic;
+				case PPI_WorldNormal:			return float4(SubstrateGBuffer.WorldNormal, 0);
+				case PPI_Opacity:				return SubstrateGBuffer.Opacity;
+				case PPI_Roughness:				return SubstrateGBuffer.Roughness;
+				case PPI_MaterialAO:			return SubstrateGBuffer.MaterialAO;
+				case PPI_ShadingModelColor:		return float4(SubstrateGBuffer.ShadingModelColor, 1);
+				case PPI_ShadingModelID:		return float4(SubstrateGBuffer.ShadingModelID, 0, 0, 0);
+				case PPI_StoredBaseColor:		return float4(SubstrateGBuffer.StoredBaseColor, 0);
+				case PPI_StoredSpecular:		return float4(SubstrateGBuffer.StoredSpecular.rrr, 0);
+				case PPI_WorldTangent:			return float4(SubstrateGBuffer.WorldTangent, 0);
+				case PPI_Anisotropy:			return SubstrateGBuffer.Anisotropy;
 			}
 		}
-		#undef SubstrateSceneTexture
 	}
 	#else // SUBSTRATE_ENABLED
 	{
diff --git a/Engine/Shaders/Public/SubstratePublic.ush b/Engine/Shaders/Public/SubstratePublic.ush
--- a/Engine/Shaders/Public/SubstratePublic.ush
+++ b/Engine/Shaders/Public/SubstratePublic.ush
@@ -8,14 +8,82 @@
 
 #include "/Engine/Private/Substrate/Substrate.ush"
 
-#if Substrate_ENABLED
+#if SUBSTRATE_ENABLED
 
 /** Public facing functions allowing external users access to some Substrate data. */
-float3 SubstratePublic_GetWorldNormal(uint2 PixelCoord)
+
+#ifndef SubstratePublicStruct
+#error One must specify the SubstratePublicStruct uniform buffer to fetch data from when using SubstratePuclic.ush
+#endif
+
+struct FSubstrateGBuffer
 {
-  return SubstrateUnpackTopLayerData(SubstratePublic.TopLayerTexture.Load(uint3(PixelCoord, 0))).WorldNormal;
+	bool bIsValid;
+
+	float3 DiffuseColor;
+	float3 SpecularColor;
+	float3 SubsurfaceColor;
+	float3 BaseColor;
+	float  Specular;
+	float  Metallic;
+	float3 WorldNormal;
+	float  Opacity;
+	float  Roughness;
+	float  MaterialAO;
+	float3 ShadingModelColor;
+	float  ShadingModelID;
+	float3 StoredBaseColor;
+	float3 StoredSpecular;
+	float3 WorldTangent;
+	float  Anisotropy;
+};
+
+FSubstrateGBuffer SubstratePublic_GetSubstrateGBufferFromFirstBSDF(uint2 PixelCoord)
+{
+	FSubstrateGBuffer SGBuffer = (FSubstrateGBuffer)0; 
+	SGBuffer.bIsValid = false;
+
+	FSubstrateAddressing SubstrateAddressing   = GetSubstratePixelDataByteOffset(PixelCoord, uint2(View.BufferSizeAndInvSize.xy), SubstratePublicStruct.MaxBytesPerPixel);
+	FSubstratePixelHeader SubstratePixelHeader = UnpackSubstrateHeaderIn(SubstratePublicStruct.MaterialTextureArray, SubstrateAddressing, SubstratePublicStruct.TopLayerTexture);
+	
+	BRANCH
+	if (SubstratePixelHeader.ClosureCount > 0)
+	{
+		const FSubstrateSubsurfaceHeader SSSHeader = SubstrateLoadSubsurfaceHeader(SubstratePublicStruct.MaterialTextureArray, SubstratePublicStruct.FirstSliceStoringSubstrateSSSData, SubstrateAddressing.PixelCoords);
+
+		const bool bSkipSSSMaterialOverride = true; // We double not want the material data to be affected by any override.
+		FSubstrateBSDF BSDF = UnpackSubstrateBSDFIn(SubstratePublicStruct.MaterialTextureArray, SubstrateAddressing, SubstratePixelHeader, bSkipSSSMaterialOverride);
+	
+		SGBuffer.bIsValid			= true;
+		SGBuffer.DiffuseColor		= SubstrateGetBSDFDiffuseColor(BSDF);
+		SGBuffer.SpecularColor		= SubstrateGetBSDFSpecularColor(BSDF);
+		SGBuffer.SubsurfaceColor	= SubstrateGetBSDFSubSurfaceColor(BSDF);
+		SGBuffer.BaseColor			= SubstrateGetBSDFBaseColor(BSDF);
+		SGBuffer.Specular			= SubstrateGetBSDFSpecular(BSDF);
+		SGBuffer.Metallic			= SubstrateGetBSDFMetallic(BSDF);
+		SGBuffer.WorldNormal		= SubstrateGetWorldNormal(SubstratePixelHeader, BSDF, SubstrateAddressing);
+		SGBuffer.Opacity			= BSDF.HasScatteringData() ? SubstrateSubSurfaceHeaderGetOpacity(SSSHeader) : 1.f;
+		SGBuffer.Roughness			= SubstrateGetBSDFRoughness(BSDF);
+		SGBuffer.MaterialAO			= SubstrateGetAO(SubstratePixelHeader);
+		SGBuffer.ShadingModelColor	= GetShadingModelColor(SubstrateGetLegacyShadingModels(BSDF));
+		SGBuffer.ShadingModelID		= SubstrateGetLegacyShadingModels(BSDF);
+		SGBuffer.StoredBaseColor	= SubstrateGetBSDFBaseColor(BSDF);
+		SGBuffer.StoredSpecular		= SubstrateGetBSDFSpecular(BSDF).rrr;
+		SGBuffer.WorldTangent		= SubstrateGetWorldTangent(SubstratePixelHeader, BSDF, SubstrateAddressing);
+		SGBuffer.Anisotropy			= SubstrateGetBSDFAnisotropy(BSDF);
+	}
+
+	return SGBuffer;
 }
 
-//TODO: Implement other accessors that were previously available via SceneTextures.
+// Optimised data fetch from the TopLayer texture
+float3 SubstratePublic_GetWorldNormal(uint2 PixelCoord)
+{
+  return SubstrateUnpackTopLayerData(SubstratePublicStruct.TopLayerTexture.Load(uint3(PixelCoord, 0))).WorldNormal;
+}
+float SubstratePublic_GetRoughness(uint2 PixelCoord)
+{
+  return SubstrateUnpackTopLayerData(SubstratePublicStruct.TopLayerTexture.Load(uint3(PixelCoord, 0))).Roughness;
+}
 
 #endif
\ No newline at end of file
diff --git a/Engine/Source/Runtime/Renderer/Private/Substrate/Substrate.cpp b/Engine/Source/Runtime/Renderer/Private/Substrate/Substrate.cpp
--- a/Engine/Source/Runtime/Renderer/Private/Substrate/Substrate.cpp
+++ b/Engine/Source/Runtime/Renderer/Private/Substrate/Substrate.cpp
@@ -588,6 +588,12 @@ void BindSubstrateBasePassUniformParameters(FRDGBuilder& GraphBuilder, const FVi
 	}
 }
 
+static FRDGTextureRef GetDefaultSubstrateMaterialTextureArray(FRDGBuilder& GraphBuilder)
+{
+	FRDGTextureRef DefaultSubstrateMaterialTextureArray = GSystemTextures.GetDefaultTexture(GraphBuilder, ETextureDimension::Texture2DArray, EPixelFormat::PF_R32_UINT, FClearValueBinding::Transparent);
+	return DefaultSubstrateMaterialTextureArray;
+}
+
 static void BindSubstrateGlobalUniformParameters(FRDGBuilder& GraphBuilder, const FSubstrateViewData* SubstrateViewData, bool bNeedsMaterialBuffer, FSubstrateGlobalUniformParameters& OutSubstrateUniformParameters)
 {
 	FSubstrateSceneData* SubstrateSceneData = SubstrateViewData->SceneData;
@@ -621,7 +627,7 @@ static void BindSubstrateGlobalUniformParameters(FRDGBuilder& GraphBuilder, cons
 			check(SubstrateSceneData->TopLayerTexture == nullptr);
 			check(SubstrateSceneData->OpaqueRoughRefractionTexture == nullptr);
 			const FRDGSystemTextures& SystemTextures = FRDGSystemTextures::Get(GraphBuilder);
-			FRDGTextureRef DefaultTextureArray = GSystemTextures.GetDefaultTexture(GraphBuilder, ETextureDimension::Texture2DArray, EPixelFormat::PF_R32_UINT, FClearValueBinding::Transparent);
+			FRDGTextureRef DefaultTextureArray = GetDefaultSubstrateMaterialTextureArray(GraphBuilder);
 			OutSubstrateUniformParameters.MaterialTextureArray = DefaultTextureArray;
 			OutSubstrateUniformParameters.TopLayerTexture = SystemTextures.DefaultNormal8Bit;
 			OutSubstrateUniformParameters.OpaqueRoughRefractionTexture = SystemTextures.Black;
@@ -630,7 +636,7 @@ static void BindSubstrateGlobalUniformParameters(FRDGBuilder& GraphBuilder, cons
 	else
 	{
 		const FRDGSystemTextures& SystemTextures = FRDGSystemTextures::Get(GraphBuilder);
-		FRDGTextureRef DefaultTextureArray = GSystemTextures.GetDefaultTexture(GraphBuilder, ETextureDimension::Texture2DArray, EPixelFormat::PF_R32_UINT, FClearValueBinding::Transparent);
+		FRDGTextureRef DefaultTextureArray = GetDefaultSubstrateMaterialTextureArray(GraphBuilder);
 		FRDGBufferSRVRef DefaultBuffer = GraphBuilder.CreateSRV(GSystemTextures.GetDefaultBuffer(GraphBuilder, 4, 0u), PF_R32_UINT);
 		OutSubstrateUniformParameters.Common = GetSubstrateCommonParameter();
 		OutSubstrateUniformParameters.SliceStoringDebugSubstrateTreeData = -1;
@@ -674,9 +680,8 @@ void BindSubstrateForwardPasslUniformParameters(FRDGBuilder& GraphBuilder, const
 	if (bCreateDummyResources)
 	{
 		const FRDGSystemTextures& SystemTextures = FRDGSystemTextures::Get(GraphBuilder);
-		FRDGTextureRef DefaultTextureArray = GSystemTextures.GetDefaultTexture(GraphBuilder, ETextureDimension::Texture2DArray, EPixelFormat::PF_R32_UINT, FClearValueBinding::Transparent);
 		OutSubstrateUniformParameters.FirstSliceStoringSubstrateSSSData = -1;
-		OutSubstrateUniformParameters.MaterialTextureArray = DefaultTextureArray;
+		OutSubstrateUniformParameters.MaterialTextureArray = GetDefaultSubstrateMaterialTextureArray(GraphBuilder);
 		OutSubstrateUniformParameters.TopLayerTexture = SystemTextures.DefaultNormal8Bit;
 	}
 }
@@ -702,17 +707,22 @@ TRDGUniformBufferRef<FSubstrateGlobalUniformParameters> BindSubstrateGlobalUnifo
 
 static void BindSubstratePublicGlobalUniformParameters(FRDGBuilder& GraphBuilder, FSubstrateSceneData* SubstrateSceneData, FSubstratePublicGlobalUniformParameters& OutSubstrateUniformParameters)
 {
-	const FRDGSystemTextures& SystemTextures = FRDGSystemTextures::Get(GraphBuilder);
 	if (SubstrateSceneData && SubstrateSceneData->TopLayerTexture)
 	{
+		OutSubstrateUniformParameters.Common = GetSubstrateCommonParameter(*SubstrateSceneData);
+		OutSubstrateUniformParameters.FirstSliceStoringSubstrateSSSData = SubstrateSceneData->FirstSliceStoringSubstrateSSSData;
+		OutSubstrateUniformParameters.MaterialTextureArray = SubstrateSceneData->MaterialTextureArray;
 		OutSubstrateUniformParameters.TopLayerTexture = SubstrateSceneData->TopLayerTexture;
 	}
 	else
 	{
+		const FRDGSystemTextures& SystemTextures = FRDGSystemTextures::Get(GraphBuilder);
+		OutSubstrateUniformParameters.Common = GetSubstrateCommonParameter();
+		OutSubstrateUniformParameters.FirstSliceStoringSubstrateSSSData = -1;
+		OutSubstrateUniformParameters.MaterialTextureArray = GetDefaultSubstrateMaterialTextureArray(GraphBuilder);
 		OutSubstrateUniformParameters.TopLayerTexture = SystemTextures.Black;
 	}
 
-	//TODO: Other Substrate scene textures or other globals.
 }
 
 static ERHIFeatureSupport SubstrateSupportsWaveOps(EShaderPlatform Platform)
diff --git a/Engine/Source/Runtime/Renderer/Private/Substrate/Substrate.h b/Engine/Source/Runtime/Renderer/Private/Substrate/Substrate.h
--- a/Engine/Source/Runtime/Renderer/Private/Substrate/Substrate.h
+++ b/Engine/Source/Runtime/Renderer/Private/Substrate/Substrate.h
@@ -73,7 +73,10 @@ BEGIN_GLOBAL_SHADER_PARAMETER_STRUCT(FSubstrateGlobalUniformParameters, RENDERER
 END_GLOBAL_SHADER_PARAMETER_STRUCT()
 
 BEGIN_GLOBAL_SHADER_PARAMETER_STRUCT(FSubstratePublicGlobalUniformParameters, RENDERER_API)
+	SHADER_PARAMETER_STRUCT_INCLUDE(FSubstrateCommonParameters, Common)
+	SHADER_PARAMETER(int32, FirstSliceStoringSubstrateSSSData)
 	SHADER_PARAMETER_RDG_TEXTURE(Texture2D<SUBSTRATE_TOP_LAYER_TYPE>, TopLayerTexture)
+	SHADER_PARAMETER_RDG_TEXTURE(Texture2DArray<uint>, MaterialTextureArray)
 END_GLOBAL_SHADER_PARAMETER_STRUCT()
 
 // This must map to the SUBSTRATE_TILE_TYPE defines.
