- Fix not killing spawn infos when inactive reaction is set to kill

#tests Asset that was using stateless did not deactivate correctly until I fixed this

diff --git a/Engine/Plugins/FX/Niagara/Source/Niagara/Private/Stateless/NiagaraStatelessEmitterInstance.cpp b/Engine/Plugins/FX/Niagara/Source/Niagara/Private/Stateless/NiagaraStatelessEmitterInstance.cpp
--- a/Engine/Plugins/FX/Niagara/Source/Niagara/Private/Stateless/NiagaraStatelessEmitterInstance.cpp
+++ b/Engine/Plugins/FX/Niagara/Source/Niagara/Private/Stateless/NiagaraStatelessEmitterInstance.cpp
@@ -674,19 +674,15 @@ void FNiagaraStatelessEmitterInstance::SetExecutionStateInternal(ENiagaraExecuti
 			break;
 
 		case ENiagaraExecutionState::Inactive:
-			if (EmitterData->EmitterState.InactiveResponse != ENiagaraEmitterInactiveResponse::Kill)
+			if (EmitterData->EmitterState.InactiveResponse == ENiagaraEmitterInactiveResponse::Kill)
 			{
-				CropSpawnInfos();
-			}
-
-			if (SpawnInfos.Num() > 0)
-			{
-				ExecutionState = ENiagaraExecutionState::Inactive;
+				KillSpawnInfos();
+				ExecutionState = ENiagaraExecutionState::Complete;
 			}
 			else
 			{
-				KillSpawnInfos();
-				ExecutionState = ENiagaraExecutionState::Complete;
+				CropSpawnInfos();
+				ExecutionState = SpawnInfos.Num() > 0 ? ENiagaraExecutionState::Inactive : ENiagaraExecutionState::Complete;
 			}
 			break;
 
