- Fast path will now run the system spawn script if required
- This allows you to have logic in the spawn script (sample a skel mesh for example) and bind stateless emitters to that data and still get the benefit of no VM execution on update
- The VM will only execute the spawn script in cases where it's required, note this does not inspect emitter renderer bindings so it might run redundantly if the attributes were not trimmed

#rnx

diff --git a/Engine/Plugins/FX/Niagara/Source/Niagara/Classes/NiagaraSystem.h b/Engine/Plugins/FX/Niagara/Source/Niagara/Classes/NiagaraSystem.h
--- a/Engine/Plugins/FX/Niagara/Source/Niagara/Classes/NiagaraSystem.h
+++ b/Engine/Plugins/FX/Niagara/Source/Niagara/Classes/NiagaraSystem.h
@@ -810,12 +810,12 @@ public:
 	bool AllowValidation() const { return bAllowValidation; }
 #endif
 
-	/** Can we run the code only system state path, i.e. we don't need to invoke the VVM / store per instance data set? */
-	bool SystemStateFastPathEnabled() const { return bSystemStateFastPathEnabled && bAllowSystemStateFastPath; }
-
 	/** Access the code system state data. */
 	const FNiagaraSystemStateData& GetSystemStateData() const { return SystemStateData; }
 
+	/** Used for debug HUD / viewport to convey what mode we are running in */
+	NIAGARA_API const TCHAR* GetSystemStateModeString() const;
+
 private:
 #if WITH_EDITORONLY_DATA
 
diff --git a/Engine/Plugins/FX/Niagara/Source/Niagara/Internal/NiagaraSystemEmitterState.h b/Engine/Plugins/FX/Niagara/Source/Niagara/Internal/NiagaraSystemEmitterState.h
--- a/Engine/Plugins/FX/Niagara/Source/Niagara/Internal/NiagaraSystemEmitterState.h
+++ b/Engine/Plugins/FX/Niagara/Source/Niagara/Internal/NiagaraSystemEmitterState.h
@@ -41,7 +41,22 @@ struct FNiagaraSystemStateData
 	GENERATED_BODY()
 
 	UPROPERTY(EditAnywhere, Category = "System State")
-	bool bIgnoreSystemState = true;
+	uint32 bRunSpawnScript : 1 = true;
+
+	UPROPERTY(EditAnywhere, Category = "System State")
+	uint32 bRunUpdateScript : 1 = true;
+
+	UPROPERTY(EditAnywhere, Category = "System State")
+	uint32 bIgnoreSystemState : 1 = true;
+
+	UPROPERTY(EditAnywhere, Category = "System State", meta = (DisplayAfter = "LoopCount", EditCondition = "LoopBehavior != ENiagaraLoopBehavior::Once", EditConditionHides))
+	uint32 bRecalculateDurationEachLoop : 1 = false;
+
+	UPROPERTY(EditAnywhere, Category = "System State", meta = (DisplayAfter = "LoopDelay", EditCondition = "LoopBehavior != ENiagaraLoopBehavior::Once", EditConditionHides))
+	uint32 bDelayFirstLoopOnly : 1 = false;
+
+	UPROPERTY(EditAnywhere, Category = "System State", meta = (DisplayAfter = "bDelayFirstLoopOnly", EditCondition = "LoopBehavior != ENiagaraLoopBehavior::Once && !bDelayFirstLoopOnly", EditConditionHides))
+	uint32 bRecalculateDelayEachLoop : 1 = false;
 
 	UPROPERTY(EditAnywhere, Category = "System State")
 	ENiagaraSystemInactiveResponse InactiveResponse = ENiagaraSystemInactiveResponse::Complete;
@@ -55,17 +70,8 @@ struct FNiagaraSystemStateData
 	UPROPERTY(EditAnywhere, Category = "System State", meta = (ClampMin = "1", EditCondition = "LoopBehavior == ENiagaraLoopBehavior::Multiple", EditConditionHides))
 	int LoopCount = 1;
 
-	UPROPERTY(EditAnywhere, Category = "System State", meta = (EditCondition = "LoopBehavior != ENiagaraLoopBehavior::Once", EditConditionHides))
-	bool bRecalculateDurationEachLoop = false;
-
 	UPROPERTY(EditAnywhere, Category = "System State", meta = (ClampMin = "0.0"))
 	FNiagaraDistributionRangeFloat LoopDelay = FNiagaraDistributionRangeFloat(0.0f);
-
-	UPROPERTY(EditAnywhere, Category = "System State", meta = (EditCondition = "LoopBehavior != ENiagaraLoopBehavior::Once", EditConditionHides))
-	bool bDelayFirstLoopOnly = false;
-
-	UPROPERTY(EditAnywhere, Category = "System State", meta = (EditCondition = "LoopBehavior != ENiagaraLoopBehavior::Once && !bDelayFirstLoopOnly", EditConditionHides))
-	bool bRecalculateDelayEachLoop = false;
 };
 
 USTRUCT()
diff --git a/Engine/Plugins/FX/Niagara/Source/Niagara/Private/NiagaraDebugHud.cpp b/Engine/Plugins/FX/Niagara/Source/Niagara/Private/NiagaraDebugHud.cpp
--- a/Engine/Plugins/FX/Niagara/Source/Niagara/Private/NiagaraDebugHud.cpp
+++ b/Engine/Plugins/FX/Niagara/Source/Niagara/Private/NiagaraDebugHud.cpp
@@ -58,20 +58,7 @@ namespace NiagaraDebugLocal
 
 	struct FCachedVariables
 	{
-		~FCachedVariables()
-		{
-#if WITH_EDITORONLY_DATA
-			if ( UNiagaraSystem* NiagaraSystem = WeakNiagaraSystem.Get() )
-			{
-				NiagaraSystem->OnSystemCompiled().Remove(CompiledDelegate);
-			}
-#endif
-		}
-
 		TWeakObjectPtr<UNiagaraSystem> WeakNiagaraSystem;
-#if WITH_EDITORONLY_DATA
-		FDelegateHandle CompiledDelegate;
-#endif
 
 		TArray<int32>										EngineVariables;			// Engine varibles that are visible on the HUD, these are special because they are not in the buffers
 		TArray<FNiagaraDataSetDebugAccessor>				SystemVariables;			// System & Emitter variables since both are inside the same DataBuffer
@@ -98,6 +85,8 @@ namespace NiagaraDebugLocal
 	static FDelegateHandle	GDebugDrawHandle;
 	static int32			GDebugDrawHandleUsers = 0;
 
+	static const TCHAR*		GCompileForEditString = TEXT(" [EditMode]");
+
 	static FVector StringToVector(const FString& Arg, const FVector& DefaultValue)
 	{
 		TArray<FString> Values;
@@ -485,9 +474,6 @@ namespace NiagaraDebugLocal
 		{
 			CachedVariables = &GCachedSystemVariables.Emplace(NiagaraSystem);
 			CachedVariables->WeakNiagaraSystem = MakeWeakObjectPtr(NiagaraSystem);
-#if WITH_EDITORONLY_DATA
-			CachedVariables->CompiledDelegate = NiagaraSystem->OnSystemCompiled().AddLambda([](UNiagaraSystem* NiagaraSystem) { GCachedSystemVariables.Remove(NiagaraSystem); });
-#endif
 
 			if (Settings.SystemVariables.Num() > 0)
 			{
@@ -831,7 +817,7 @@ FNiagaraDebugHud::FNiagaraDebugHud(UWorld* World)
 
 	WeakWorld = World;
 
-	LongestSystemName = TEXT("NS_SomeBigLongNiagaraSystemName");
+	LongestSystemPrettyName = TEXT("NS_SomeBigLongNiagaraSystemName");
 
 	if ( !GDebugDrawHandle.IsValid() )
 	{
@@ -937,6 +923,17 @@ FNiagaraDebugHud::~FNiagaraDebugHud()
 		UDebugDrawService::Unregister(GDebugDrawHandle);
 		GDebugDrawHandle.Reset();
 	}
+
+#if WITH_EDITORONLY_DATA
+	for (auto It = SystemCompiledDelegates.CreateIterator(); It; ++It)
+	{
+		if ( UNiagaraSystem* NiagaraSystem = It.Key().Get() )
+		{
+			NiagaraSystem->OnSystemCompiled().Remove(It.Value());
+		}
+	}
+	SystemCompiledDelegates.Empty();
+#endif
 }
 
 void FNiagaraDebugHud::UpdateSettings(const FNiagaraDebugHUDSettingsData& NewSettings)
@@ -1073,18 +1070,41 @@ void FNiagaraDebugHud::GatherSystemInfo()
 		if (SystemDebugInfo.SystemName.IsEmpty())
 		{
 			SystemDebugInfo.SystemName = GetNameSafe(FXComponent->GetFXSystemAsset());
+		}
+
+		if (SystemDebugInfo.SystemPrettyName.IsEmpty())
+		{
+			SystemDebugInfo.SystemPrettyName = SystemDebugInfo.SystemName;
+			if (UNiagaraSystem* NiagaraSystem = Cast<UNiagaraSystem>(FXComponent->GetFXSystemAsset()))
+			{
+			#if WITH_EDITORONLY_DATA
+				if (NiagaraSystem->GetCompileForEdit())
+				{
+					SystemDebugInfo.SystemPrettyName.Append(GCompileForEditString);
+				}
+			#endif
+				if (const TCHAR* SystemStateModeString = NiagaraSystem->GetSystemStateModeString())
+				{
+					SystemDebugInfo.SystemPrettyName.AppendChar(' ');
+					SystemDebugInfo.SystemPrettyName.Append(SystemStateModeString);
+				}
+
+			#if WITH_EDITORONLY_DATA
+				if (!SystemCompiledDelegates.Contains(NiagaraSystem))
+				{
+					SystemCompiledDelegates.Emplace(
+						NiagaraSystem,
+						NiagaraSystem->OnSystemCompiled().AddRaw(this, &FNiagaraDebugHud::OnSystemCompiled)
+					);
+				}
+			#endif
+			}
 
-			static const FString SystemNamePostFix = TEXT(" (Fast Path) ");
-			if (SystemDebugInfo.SystemName.Len() + SystemNamePostFix.Len() > LongestSystemName.Len())
+			if (SystemDebugInfo.SystemPrettyName.Len() > LongestSystemPrettyName.Len())
 			{
-				LongestSystemName = SystemDebugInfo.SystemName;
-				LongestSystemName.Append(SystemNamePostFix);
+				LongestSystemPrettyName = SystemDebugInfo.SystemPrettyName;
 			}
 		}
-	#if WITH_EDITORONLY_DATA
-		SystemDebugInfo.bCompileForEdit = NiagaraComponent ? NiagaraComponent->GetAsset()->GetCompileForEdit() : false;
-	#endif
-		SystemDebugInfo.bSystemStateFastPath = NiagaraComponent ? NiagaraComponent->GetAsset()->SystemStateFastPathEnabled() : false;
 		SystemDebugInfo.bShowInWorld = (Settings.SystemDebugVerbosity != ENiagaraDebugHudVerbosity::None) && (!Settings.bSystemFilterEnabled || Settings.SystemFilter.IsEmpty() || SystemDebugInfo.SystemName.MatchesWildcard(Settings.SystemFilter));
 		SystemDebugInfo.bPassesSystemFilter = SystemDebugInfo.bShowInWorld;
 
@@ -1488,6 +1508,18 @@ void FNiagaraDebugHud::DrawDebugGeomerty(class FNiagaraWorldManager* WorldManage
 	}
 }
 
+#if WITH_EDITORONLY_DATA
+void FNiagaraDebugHud::OnSystemCompiled(UNiagaraSystem * NiagaraSystem)
+{
+	if (FSystemDebugInfo* DebugInfo = PerSystemDebugInfo.Find(NiagaraSystem->GetFName()))
+	{
+		DebugInfo->SystemPrettyName.Empty();
+	}
+
+	NiagaraDebugLocal::GCachedSystemVariables.Remove(NiagaraSystem);
+}
+#endif
+
 template<typename T>
 struct FGraph
 {
@@ -1669,7 +1701,7 @@ void FNiagaraDebugHud::DrawOverview(class FNiagaraWorldManager* WorldManager, FC
 
 	if (Settings.bOverviewEnabled)
 	{
-		OverviewColumns.Emplace(TEXT(""), TEXT(""), TEXT("System Name"), ColumnOffset, Font, *LongestSystemName,
+		OverviewColumns.Emplace(TEXT(""), TEXT(""), TEXT("System Name"), ColumnOffset, Font, *LongestSystemPrettyName,
 			[&DetailColor, &DetailHighlightColor, &fAdvanceHeight](FCanvas* Canvas, UFont* Font, float X, float Y, FOverviewColumn& Col, const FSystemDebugInfo& SystemInfo)
 			{
 				FLinearColor RowBGColor = SystemInfo.UniqueColor;
@@ -1678,17 +1710,7 @@ void FNiagaraDebugHud::DrawOverview(class FNiagaraWorldManager* WorldManager, FC
 				const FLinearColor RowColor = SystemInfo.bShowInWorld ? DetailHighlightColor : DetailColor;
 
 				FNameBuilder SystemNameString;
-				SystemNameString.Append(*SystemInfo.SystemName);
-			#if WITH_EDITORONLY_DATA
-				if (SystemInfo.bCompileForEdit)
-				{
-					SystemNameString.Append(TEXT(" (Edit Mode)"));
-				}
-			#endif
-				if (SystemInfo.bSystemStateFastPath)
-				{
-					SystemNameString.Append(TEXT(" (Fast Path)"));
-				}
+				SystemNameString.Append(*SystemInfo.SystemPrettyName);
 				Canvas->DrawShadowedString(X, Y, SystemNameString.ToString(), Font, RowColor);
 		});
 
@@ -2349,9 +2371,6 @@ void FNiagaraDebugHud::DrawGpuComputeOverriew(class FNiagaraWorldManager* WorldM
 
 		const bool bShowDetailed = !Settings.bSystemFilterEnabled || Settings.SystemFilter.IsEmpty() || OwnerSystem->GetName().MatchesWildcard(Settings.SystemFilter);
 		SystemIt.Value().bShowDetailed = bShowDetailed;
-#if WITH_EDITORONLY_DATA
-		SystemIt.Value().bCompileForEdit = OwnerSystem->GetCompileForEdit();
-#endif
 		bHasDetailedView |= bShowDetailed;
 		bHasSimpleView |= !bShowDetailed;
 	}
@@ -2398,9 +2417,9 @@ void FNiagaraDebugHud::DrawGpuComputeOverriew(class FNiagaraWorldManager* WorldM
 					const FGpuUsagePerStage& StageUsage = StageIt.Value();
 					OwnerSystem->GetFName().AppendString(SimpleTable.GetColumnText(0));
 #if WITH_EDITORONLY_DATA
-					if (SystemIt.Value().bCompileForEdit)
+					if (OwnerSystem->GetCompileForEdit())
 					{
-						SimpleTable.GetColumnText(0).Append(TEXT(" (Edit Mode)"));
+						SimpleTable.GetColumnText(0).Append(GCompileForEditString);
 					}
 #endif
 					OwnerEmitter->GetFName().AppendString(SimpleTable.GetColumnText(1));
@@ -2448,9 +2467,9 @@ void FNiagaraDebugHud::DrawGpuComputeOverriew(class FNiagaraWorldManager* WorldM
 
 			OwnerSystem->GetFName().AppendString(SimpleTable.GetColumnText(0));
 #if WITH_EDITORONLY_DATA
-			if (SystemIt.Value().bCompileForEdit)
+			if (OwnerSystem->GetCompileForEdit())
 			{
-				SimpleTable.GetColumnText(0).Append(TEXT(" (Edit Mode)"));
+				SimpleTable.GetColumnText(0).Append(GCompileForEditString);
 			}
 #endif
 			SimpleTable.GetColumnText(1).Appendf(TEXT("%4.1f"), SystemUsage.InstanceCount.GetAverage<float>());
diff --git a/Engine/Plugins/FX/Niagara/Source/Niagara/Private/NiagaraDebugHud.h b/Engine/Plugins/FX/Niagara/Source/Niagara/Private/NiagaraDebugHud.h
--- a/Engine/Plugins/FX/Niagara/Source/Niagara/Private/NiagaraDebugHud.h
+++ b/Engine/Plugins/FX/Niagara/Source/Niagara/Private/NiagaraDebugHud.h
@@ -94,10 +94,7 @@ class FNiagaraDebugHud
 	struct FSystemDebugInfo
 	{
 		FString		SystemName;
-	#if WITH_EDITORONLY_DATA
-		bool		bCompileForEdit = false;
-	#endif
-		bool		bSystemStateFastPath = false;
+		FString		SystemPrettyName;
 
 		#if WITH_PARTICLE_PERF_STATS
 		TSharedPtr<FNiagaraDebugHUDPerfStats> PerfStats = nullptr;
@@ -251,9 +248,17 @@ private:
 	void DrawMessages(class FNiagaraWorldManager* WorldManager, class FCanvas* DrawCanvas, FVector2f& TextLocation);
 	void DrawDebugGeomerty(class FNiagaraWorldManager* WorldManager, class UCanvas* DrawCanvas);
 
+#if WITH_EDITORONLY_DATA
+	void OnSystemCompiled(UNiagaraSystem* NiagaraSystem);
+#endif
+
 private:
 	TWeakObjectPtr<class UWorld>	WeakWorld;
 
+#if WITH_EDITORONLY_DATA
+	TMap<TWeakObjectPtr<UNiagaraSystem>, FDelegateHandle> SystemCompiledDelegates;
+#endif
+
 	int32 GlobalTotalRegistered = 0;
 	int32 GlobalTotalActive = 0;
 	int32 GlobalTotalScalability = 0;
@@ -269,7 +274,7 @@ private:
 
 	int32 GlobalTotalPlayerSystems = 0;
 
-	FString LongestSystemName;
+	FString LongestSystemPrettyName;
 	TMap<FName, FSystemDebugInfo>	PerSystemDebugInfo;
 
 	TArray<TWeakObjectPtr<class UFXSystemComponent>>	InWorldComponents;
@@ -305,9 +310,6 @@ private:
 
 	struct FGpuUsagePerSystem
 	{
-#if WITH_EDITORONLY_DATA
-		bool bCompileForEdit = false;
-#endif
 		bool bShowDetailed = false;
 		FSmoothedCounter<uint32> InstanceCount;
 		FSmoothedCounter<uint64> Microseconds;
diff --git a/Engine/Plugins/FX/Niagara/Source/Niagara/Private/NiagaraSystem.cpp b/Engine/Plugins/FX/Niagara/Source/Niagara/Private/NiagaraSystem.cpp
--- a/Engine/Plugins/FX/Niagara/Source/Niagara/Private/NiagaraSystem.cpp
+++ b/Engine/Plugins/FX/Niagara/Source/Niagara/Private/NiagaraSystem.cpp
@@ -2731,7 +2731,6 @@ void UNiagaraSystem::RemoveEmitterHandlesById(const TSet<FGuid>& HandlesToRemove
 }
 #endif
 
-
 UNiagaraScript* UNiagaraSystem::GetSystemSpawnScript()
 {
 	return SystemSpawnScript;
@@ -2752,6 +2751,15 @@ const UNiagaraScript* UNiagaraSystem::GetSystemUpdateScript() const
 	return SystemUpdateScript;
 }
 
+const TCHAR* UNiagaraSystem::GetSystemStateModeString() const
+{
+	if ( SystemStateData.bRunUpdateScript == false )
+	{
+		return SystemStateData.bRunSpawnScript ? TEXT("[Fast-U]") : TEXT("[Fast-SU]");
+	}
+	return nullptr;
+}
+
 #if WITH_EDITORONLY_DATA
 
 void UNiagaraSystem::KillAllActiveCompilations()
@@ -3871,17 +3879,16 @@ void UNiagaraSystem::GatherStaticVariables(TArray<FNiagaraVariable>& OutVars, TA
 void UNiagaraSystem::ResolveRequiresScripts()
 {
 #if WITH_EDITORONLY_DATA
-	TOptional<FNiagaraSystemStateData> NewSystemStateData;
-
 	if (bAllowSystemStateFastPath)
 	{
 		INiagaraModule& NiagaraModule = FModuleManager::GetModuleChecked<INiagaraModule>("Niagara");
 		const INiagaraEditorOnlyDataUtilities& EditorOnlyDataUtilities = NiagaraModule.GetEditorOnlyDataUtilities();
-		NewSystemStateData = EditorOnlyDataUtilities.TryGetSystemStateData(*this);
+		SystemStateData = EditorOnlyDataUtilities.GetSystemStateData(*this);
+	}
+	else
+	{
+		SystemStateData = FNiagaraSystemStateData();
 	}
-
-	bSystemStateFastPathEnabled = NewSystemStateData.IsSet();
-	SystemStateData = NewSystemStateData.Get(FNiagaraSystemStateData());
 #endif //WITH_EDITORONLY_DATA
 }
 
diff --git a/Engine/Plugins/FX/Niagara/Source/Niagara/Private/NiagaraSystemSimulation.cpp b/Engine/Plugins/FX/Niagara/Source/Niagara/Private/NiagaraSystemSimulation.cpp
--- a/Engine/Plugins/FX/Niagara/Source/Niagara/Private/NiagaraSystemSimulation.cpp
+++ b/Engine/Plugins/FX/Niagara/Source/Niagara/Private/NiagaraSystemSimulation.cpp
@@ -588,7 +588,11 @@ bool FNiagaraSystemSimulation::Init(UNiagaraSystem* InSystem, UWorld* InWorld, b
 
 	bCanExecute = System->GetSystemSpawnScript()->GetVMExecutableData().IsValid() && System->GetSystemUpdateScript()->GetVMExecutableData().IsValid();
 
-	bSystemStateFastPathEnabled = System->SystemStateFastPathEnabled();
+	{
+		const FNiagaraSystemStateData& SystemState = System->GetSystemStateData();
+		bRunSpawnScript		= SystemState.bRunSpawnScript;
+		bRunUpdateScript	= SystemState.bRunUpdateScript;
+	}
 
 	MaxDeltaTime = System->GetMaxDeltaTime();
 
@@ -1010,7 +1014,7 @@ void FNiagaraSystemSimulation::FlushTickBatch(FNiagaraSystemSimulationTickContex
 	if ( Context.IsRunningAsync() )
 	{
 		FGraphEventArray FinalizePrereqArray;
-		if (bSystemStateFastPathEnabled)
+		if (bRunUpdateScript == false)
 		{
 			for (FNiagaraSystemInstance* Inst : Context.TickBatch)
 			{
@@ -1693,21 +1697,61 @@ void FNiagaraSystemSimulation::Tick_Concurrent(FNiagaraSystemSimulationTickConte
 	FScopeCycleCounter SystemStatCounter(Context.System->GetStatID(true, true));
 #endif
 
-	if (bSystemStateFastPathEnabled)
+	if (bRunUpdateScript == false)
 	{
-		//-OPT: We should be able to avoid this but will require a lot of changes to the system simulation
-		//      We might want to consider having different system simulation types
+		const int32 NumInstances = Context.Instances.Num();
+		const int32 FirstSpawnedInstance = NumInstances - Context.SpawnNum;
+		if (bRunSpawnScript)// && Context.SpawnNum > 0)
 		{
-			Context.DataSet.BeginSimulate();
-			Context.DataSet.Allocate(Context.Instances.Num());
-			Context.DataSet.GetDestinationDataChecked().SetNumInstances(Context.Instances.Num());
-			Context.DataSet.EndSimulate();
+			for (int32 iSystemInstance = FirstSpawnedInstance; iSystemInstance < NumInstances; ++iSystemInstance)
+			{
+				FNiagaraSystemInstance* SystemInstance = Context.Instances[iSystemInstance];
+				SystemInstance->TickInstanceParameters_Concurrent();
+			}
+			if (Context.SpawnNum > 0)
+			{
+				PrepareForSystemSimulate(Context);
+				SpawnSystemInstances(Context);
+
+				// Transfer results to parameter stores
+				for (int32 iSystemInstance = FirstSpawnedInstance; iSystemInstance < NumInstances; ++iSystemInstance)
+				{
+					FNiagaraSystemInstance* SystemInstance = Context.Instances[iSystemInstance];
+
+					TArrayView<FNiagaraEmitterInstanceRef> Emitters = SystemInstance->GetEmitters();
+					for (int32 iEmitter = 0; iEmitter < Emitters.Num(); ++iEmitter)
+					{
+						FNiagaraEmitterInstance& EmitterInstance = Emitters[iEmitter].Get();
+						if ( EmitterInstance.IsComplete() )
+						{
+							continue;
+						}
+						DataSetToEmitterRendererParameters[iEmitter].DataSetToParameterStore(EmitterInstance.GetRendererBoundVariables(), Context.DataSet, iSystemInstance);
+					}
+				}
+			}
+		}
+		else
+		{
+			//-OPT: We should be able to avoid this but will require a lot of changes to the system simulation
+			//      We might want to consider having different system simulation types
+			{
+				Context.DataSet.BeginSimulate();
+				Context.DataSet.Allocate(Context.Instances.Num());
+				Context.DataSet.GetDestinationDataChecked().SetNumInstances(Context.Instances.Num());
+				Context.DataSet.EndSimulate();
+			}
 		}
 
-		for (FNiagaraSystemInstance* SystemInstance : Context.Instances)
+		for ( int32 i=0; i < Context.Instances.Num(); ++i )
 		{
+			FNiagaraSystemInstance* SystemInstance = Context.Instances[i];
+
 			//-TODO: Stateless doesn't require a lot of this data
-			SystemInstance->TickInstanceParameters_Concurrent();
+			if (i < FirstSpawnedInstance)
+			{
+				SystemInstance->TickInstanceParameters_Concurrent();
+			}
 
 			SystemInstance->TickSystemState();
 
diff --git a/Engine/Plugins/FX/Niagara/Source/Niagara/Public/INiagaraEditorOnlyDataUtlities.h b/Engine/Plugins/FX/Niagara/Source/Niagara/Public/INiagaraEditorOnlyDataUtlities.h
--- a/Engine/Plugins/FX/Niagara/Source/Niagara/Public/INiagaraEditorOnlyDataUtlities.h
+++ b/Engine/Plugins/FX/Niagara/Source/Niagara/Public/INiagaraEditorOnlyDataUtlities.h
@@ -32,5 +32,5 @@ public:
 
 	virtual UNiagaraDataInterface* GetResolvedRuntimeInstanceForEditorDataInterfaceInstance(const UNiagaraSystem& OwningSystem, UNiagaraDataInterface& EditorDataInterfaceInstance) const = 0;
 
-	virtual TOptional<FNiagaraSystemStateData> TryGetSystemStateData(const UNiagaraSystem& System) const = 0;
+	virtual FNiagaraSystemStateData GetSystemStateData(const UNiagaraSystem& System) const = 0;
 };
diff --git a/Engine/Plugins/FX/Niagara/Source/Niagara/Public/NiagaraSystemSimulation.h b/Engine/Plugins/FX/Niagara/Source/Niagara/Public/NiagaraSystemSimulation.h
--- a/Engine/Plugins/FX/Niagara/Source/Niagara/Public/NiagaraSystemSimulation.h
+++ b/Engine/Plugins/FX/Niagara/Source/Niagara/Public/NiagaraSystemSimulation.h
@@ -364,6 +364,9 @@ protected:
 	uint32 bInSpawnPhase : 1;
 	uint32 bIsSolo : 1;
 
+	uint32 bRunSpawnScript : 1 = true;
+	uint32 bRunUpdateScript : 1 = true;
+
 	/** System instance per state. */
 	TArray<FNiagaraSystemInstance*> SystemInstancesPerState[int32(ENiagaraSystemInstanceState::Num)];
 
@@ -433,6 +436,4 @@ protected:
 	float FixedDeltaTickAge = 0;
 
 	FNiagaraTickInfo TickInfo;
-
-	bool bSystemStateFastPathEnabled = true;
 };
diff --git a/Engine/Plugins/FX/Niagara/Source/NiagaraEditor/Private/NiagaraEditorModule.cpp b/Engine/Plugins/FX/Niagara/Source/NiagaraEditor/Private/NiagaraEditorModule.cpp
--- a/Engine/Plugins/FX/Niagara/Source/NiagaraEditor/Private/NiagaraEditorModule.cpp
+++ b/Engine/Plugins/FX/Niagara/Source/NiagaraEditor/Private/NiagaraEditorModule.cpp
@@ -298,43 +298,57 @@ class FNiagaraEditorOnlyDataUtilities : public INiagaraEditorOnlyDataUtilities
 		return FNiagaraEditorUtilities::GetResolvedRuntimeInstanceForEditorDataInterfaceInstance(OwningSystem, EditorDataInterfaceInstance);
 	}
 
-	virtual TOptional<FNiagaraSystemStateData> TryGetSystemStateData(const UNiagaraSystem& System) const override
+	virtual FNiagaraSystemStateData GetSystemStateData(const UNiagaraSystem& System) const override
 	{
-		TOptional<FNiagaraSystemStateData> SystemStateData;
-
 		// All emitters must be stateless currently
 		// We can perhaps look at this again, but we always write Emitter.RandomSeed currently even with an empty script
 		for (const FNiagaraEmitterHandle& EmitterHandle : System.GetEmitterHandles())
 		{
 			if ( EmitterHandle.GetIsEnabled() && EmitterHandle.GetEmitterMode() != ENiagaraEmitterMode::Stateless )
 			{
-				return SystemStateData;
+				return FNiagaraSystemStateData();
 			}
 		}
 
 		// Try to resolve system state from the system scripts
 		const UNiagaraScript* Script = System.GetSystemSpawnScript();
 		const UNiagaraScriptSource* ScriptSource = Script ? Cast<UNiagaraScriptSource>(Script->GetLatestSource()) : nullptr;
-		if (ScriptSource)
+		if (!ScriptSource || !ScriptSource->NodeGraph)
 		{
-			const TCHAR* SystemStateName = TEXT("/Niagara/Modules/System/SystemState.SystemState");
-			TArray<UNiagaraNodeFunctionCall*> Nodes;
-			if (ScriptSource->NodeGraph)
-			{
-				ScriptSource->NodeGraph->GetNodesOfClass<UNiagaraNodeFunctionCall>(Nodes);
-			}
-			Nodes.RemoveAll([](UNiagaraNodeFunctionCall* Node) { return !Node || !Node->IsNodeEnabled(); });
+			return FNiagaraSystemStateData();
+		}
+
+		// Look for update script nodes to see if it's possible to avoid running the update script
+		FNiagaraSystemStateData SystemStateData;
+		if (UNiagaraNodeOutput* UpdateScriptOutput = ScriptSource->NodeGraph->FindEquivalentOutputNode(ENiagaraScriptUsage::SystemUpdateScript, FGuid()))
+		{
+			TArray<UNiagaraNodeFunctionCall*> ModuleNodes;
+			FNiagaraStackGraphUtilities::GetOrderedModuleNodes(*UpdateScriptOutput, ModuleNodes);
+			ModuleNodes.RemoveAll([](UNiagaraNodeFunctionCall* Node) { return !Node || !Node->IsNodeEnabled(); });
 
-			// No function calls, we can enable fast path with the empty system state
-			if (Nodes.Num() == 0)
+			const TCHAR* SystemStateName = TEXT("/Niagara/Modules/System/SystemState.SystemState");
+			if (ModuleNodes.Num() == 0)
 			{
-				SystemStateData.Emplace(FNiagaraSystemStateData());
+				SystemStateData.bRunUpdateScript = false;
 			}
 			//-TODO:Stateless: Single function call which is system state, attempt to extract the data
 			//else if (Nodes.Num() == 1 && Nodes[0]->FunctionScript->GetPathName() == SystemStateName)
 			//{
 			//}
 		}
+		
+		// If we don't need to execute the update script, do we need to execute the spawn script?
+		if (SystemStateData.bRunUpdateScript == false)
+		{
+			if (UNiagaraNodeOutput* SpawnScriptOutput = ScriptSource->NodeGraph->FindEquivalentOutputNode(ENiagaraScriptUsage::SystemSpawnScript, FGuid()))
+			{
+				TArray<UNiagaraNodeFunctionCall*> ModuleNodes;
+				FNiagaraStackGraphUtilities::GetOrderedModuleNodes(*SpawnScriptOutput, ModuleNodes);
+				ModuleNodes.RemoveAll([](UNiagaraNodeFunctionCall* Node) { return !Node || !Node->IsNodeEnabled(); });
+
+				SystemStateData.bRunSpawnScript = ModuleNodes.Num() != 0;
+			}
+		}
 
 		return SystemStateData;
 	}
diff --git a/Engine/Plugins/FX/Niagara/Source/NiagaraEditor/Private/Widgets/SNiagaraSystemViewport.cpp b/Engine/Plugins/FX/Niagara/Source/NiagaraEditor/Private/Widgets/SNiagaraSystemViewport.cpp
--- a/Engine/Plugins/FX/Niagara/Source/NiagaraEditor/Private/Widgets/SNiagaraSystemViewport.cpp
+++ b/Engine/Plugins/FX/Niagara/Source/NiagaraEditor/Private/Widgets/SNiagaraSystemViewport.cpp
@@ -573,10 +573,11 @@ void FNiagaraSystemViewportClient::DrawStatelessInfo(UNiagaraComponent* Componen
 		const FVector2D ScaledViewportSize = FVector2D(InViewport->GetSizeXY()) / Canvas->GetDPIScale();
 		FCanvasTextItem TextItem(FVector2D::ZeroVector, FText::GetEmpty(), Font, FLinearColor::White);
 
-		FString StatelessInfo = FString::Printf(TEXT("Stateless: %d / %d Emitters"), NumStatelessEmitters, NumEmitters);
-		if (NiagaraSystem->SystemStateFastPathEnabled())
+		FString StatelessInfo = FString::Printf(TEXT("Stateless: %d / %d Emitters "), NumStatelessEmitters, NumEmitters);
+
+		if (const TCHAR* SystemStateModeString = NiagaraSystem->GetSystemStateModeString())
 		{
-			StatelessInfo.Append(TEXT(" [FastPath]"));
+			StatelessInfo.Append(SystemStateModeString);
 		}
 
 		int32 Width = 0;
