- Stateless emitters asset registry tags

#rnx

diff --git a/Engine/Plugins/FX/Niagara/Source/Niagara/Classes/NiagaraEmitterHandle.h b/Engine/Plugins/FX/Niagara/Source/Niagara/Classes/NiagaraEmitterHandle.h
--- a/Engine/Plugins/FX/Niagara/Source/Niagara/Classes/NiagaraEmitterHandle.h
+++ b/Engine/Plugins/FX/Niagara/Source/Niagara/Classes/NiagaraEmitterHandle.h
@@ -71,6 +71,9 @@ public:
 	/** Test to see if the emitter is allowed by scalability or not */
 	NIAGARA_API bool IsAllowedByScalability() const;
 
+	/** Test to see if the emitter is enabled on the provided quality level or not. */
+	NIAGARA_API bool IsEnabledOnEffectQualityLevel(int32 QualityLevel) const;
+
 #if WITH_EDITORONLY_DATA
 	bool IsIsolated() const {	return bIsolated; }
 	void SetIsolated(bool bInIsolated) { bIsolated = bInIsolated; }
diff --git a/Engine/Plugins/FX/Niagara/Source/Niagara/Private/NiagaraEmitterHandle.cpp b/Engine/Plugins/FX/Niagara/Source/Niagara/Private/NiagaraEmitterHandle.cpp
--- a/Engine/Plugins/FX/Niagara/Source/Niagara/Private/NiagaraEmitterHandle.cpp
+++ b/Engine/Plugins/FX/Niagara/Source/Niagara/Private/NiagaraEmitterHandle.cpp
@@ -182,6 +182,19 @@ bool FNiagaraEmitterHandle::IsAllowedByScalability() const
 }
 
 #if WITH_EDITORONLY_DATA
+bool FNiagaraEmitterHandle::IsEnabledOnEffectQualityLevel(int32 QualityLevel) const
+{
+	if (EmitterMode == ENiagaraEmitterMode::Standard)
+	{
+		FVersionedNiagaraEmitterData* EmitterData = VersionedInstance.GetEmitterData();
+		return EmitterData ? EmitterData->Platforms.IsEffectQualityEnabled(QualityLevel) : false;
+	}
+	else
+	{
+		return StatelessEmitter ? StatelessEmitter->GetPlatformSet().IsEffectQualityEnabled(QualityLevel) : false;
+	}
+}
+
 void FNiagaraEmitterHandle::SetEmitterMode(UNiagaraSystem& InOwningSystem, ENiagaraEmitterMode InEmitterMode)
 {
 	if (EmitterMode != InEmitterMode)
diff --git a/Engine/Plugins/FX/Niagara/Source/Niagara/Private/NiagaraSystem.cpp b/Engine/Plugins/FX/Niagara/Source/Niagara/Private/NiagaraSystem.cpp
--- a/Engine/Plugins/FX/Niagara/Source/Niagara/Private/NiagaraSystem.cpp
+++ b/Engine/Plugins/FX/Niagara/Source/Niagara/Private/NiagaraSystem.cpp
@@ -82,6 +82,7 @@ namespace NiagaraSystemPrivate
 {
 	static const FName NAME_ActiveEmitters("ActiveEmitters");
 	static const FName NAME_ActiveRenderers("ActiveRenderers");
+	static const FName NAME_ActiveStatelessEmitters("ActiveStatelessEmitters");
 	static const FName NAME_GPUSimsMissingFixedBounds("GPUSimsMissingFixedBounds");
 	static const FName NAME_EffectType("EffectType");
 	static const FName NAME_WarmupTime("WarmupTime");
@@ -1564,35 +1565,68 @@ void UNiagaraSystem::GetAssetRegistryTags(FAssetRegistryTagsContext Context) con
 	// Gather up generic NumActive values
 	uint32 NumActiveEmitters = 0;
 	uint32 NumActiveRenderers = 0;
+	uint32 NumActiveStatelessEmitters = 0;
 	TArray<const UNiagaraRendererProperties*> ActiveRenderers;
 	for (const FNiagaraEmitterHandle& Handle : EmitterHandles)
 	{
-		if (Handle.GetIsEnabled())
+		if (!Handle.GetIsEnabled())
 		{
+			continue;
+		}
 
-			NumActiveEmitters++;
-			if (FVersionedNiagaraEmitterData* EmitterData = Handle.GetEmitterData())
+		++NumActiveEmitters;
+
+		switch ( Handle.GetEmitterMode() )
+		{
+			case ENiagaraEmitterMode::Standard:
 			{
-				// Only register fixed bounds requirement for GPU if the system itself isn't fixed bounds.
-				if (bFixedBounds == false && EmitterData->CalculateBoundsMode == ENiagaraEmitterCalculateBoundMode::Dynamic && EmitterData->SimTarget == ENiagaraSimTarget::GPUComputeSim)
+				if (FVersionedNiagaraEmitterData* EmitterData = Handle.GetEmitterData())
 				{
-					GPUSimsMissingFixedBounds++;
+					// Only register fixed bounds requirement for GPU if the system itself isn't fixed bounds.
+					if (bFixedBounds == false && EmitterData->CalculateBoundsMode == ENiagaraEmitterCalculateBoundMode::Dynamic && EmitterData->SimTarget == ENiagaraSimTarget::GPUComputeSim)
+					{
+						GPUSimsMissingFixedBounds++;
+					}
+
+					for (const UNiagaraRendererProperties* Props : EmitterData->GetRenderers())
+					{
+						if (Props)
+						{
+							NumActiveRenderers++;
+							ActiveRenderers.Add(Props);
+						}
+					}
 				}
+				break;
+			}
 
-				for (const UNiagaraRendererProperties* Props : EmitterData->GetRenderers())
+			case ENiagaraEmitterMode::Stateless:
 				{
-					if (Props)
+					if (UNiagaraStatelessEmitter* StatelessEmitter = Handle.GetStatelessEmitter())
 					{
-						NumActiveRenderers++;
-						ActiveRenderers.Add(Props);
+						++NumActiveStatelessEmitters;
+
+						for (const UNiagaraRendererProperties* Props : StatelessEmitter->GetRenderers())
+						{
+							if (Props)
+							{
+								NumActiveRenderers++;
+								ActiveRenderers.Add(Props);
+							}
+						}
 					}
+					break;
 				}
-			}
+
+			default:
+				checkNoEntry();
+				break;
 		}
 	}
 
 	Context.AddTag(FAssetRegistryTag(NiagaraSystemPrivate::NAME_ActiveEmitters, LexToString(NumActiveEmitters), FAssetRegistryTag::TT_Numerical));
 	Context.AddTag(FAssetRegistryTag(NiagaraSystemPrivate::NAME_ActiveRenderers, LexToString(NumActiveRenderers), FAssetRegistryTag::TT_Numerical));
+	Context.AddTag(FAssetRegistryTag(NiagaraSystemPrivate::NAME_ActiveStatelessEmitters, LexToString(NumActiveStatelessEmitters), FAssetRegistryTag::TT_Numerical));
 	Context.AddTag(FAssetRegistryTag(NiagaraSystemPrivate::NAME_GPUSimsMissingFixedBounds, LexToString(GPUSimsMissingFixedBounds), FAssetRegistryTag::TT_Numerical));
 	Context.AddTag(FAssetRegistryTag(NiagaraSystemPrivate::NAME_EffectType, EffectType != nullptr ? EffectType->GetName() : FString(TEXT("None")), FAssetRegistryTag::TT_Alphabetical));
 	Context.AddTag(FAssetRegistryTag(NiagaraSystemPrivate::NAME_WarmupTime, LexToString(WarmupTime), FAssetRegistryTag::TT_Numerical));
@@ -1610,19 +1644,16 @@ void UNiagaraSystem::GetAssetRegistryTags(FAssetRegistryTagsContext Context) con
 
 		for (const FNiagaraEmitterHandle& Handle : EmitterHandles)
 		{
-			if (Handle.GetIsEnabled())
+			if (!Handle.GetIsEnabled())
+			{
+				continue;
+			}
+
+			for (int32 i = 0; i < NumQualityLevels; i++)
 			{
-				FVersionedNiagaraEmitter Emitter = Handle.GetInstance();
-				FVersionedNiagaraEmitterData* EmitterData = Emitter.GetEmitterData();
-				if (EmitterData)
+				if ( Handle.IsEnabledOnEffectQualityLevel(i) )
 				{
-					for (int32 i = 0; i < NumQualityLevels; i++)
-					{
-						if (EmitterData->Platforms.IsEffectQualityEnabled(i))
-						{
-							QualityLevelsNumActive[i]++;
-						}
-					}
+					QualityLevelsNumActive[i]++;
 				}
 			}
 		}
@@ -1759,6 +1790,11 @@ void UNiagaraSystem::GetAssetRegistryTagMetadata(TMap<FName, FAssetRegistryTagMe
 		FAssetRegistryTagMetadata()
 		.SetDisplayName(LOCTEXT("ActiveRenderers", "Active Renderers"))
 	);
+	OutMetadata.Add(
+		NiagaraSystemPrivate::NAME_ActiveStatelessEmitters,
+		FAssetRegistryTagMetadata()
+		.SetDisplayName(LOCTEXT("ActiveStatelessEmitters", "Active Lightweight Emitters"))
+	);
 	OutMetadata.Add(
 		NiagaraSystemPrivate::NAME_GPUSimsMissingFixedBounds,
 		FAssetRegistryTagMetadata()
