- Fix for cache comparing when we have no UniqueID

#rnx
#tests A test I was writing failed because all the instance indices are 0

diff --git a/Engine/Plugins/FX/Niagara/Source/Niagara/Private/NiagaraSimCacheCompare.cpp b/Engine/Plugins/FX/Niagara/Source/Niagara/Private/NiagaraSimCacheCompare.cpp
--- a/Engine/Plugins/FX/Niagara/Source/Niagara/Private/NiagaraSimCacheCompare.cpp
+++ b/Engine/Plugins/FX/Niagara/Source/Niagara/Private/NiagaraSimCacheCompare.cpp
@@ -409,13 +409,12 @@ bool FNiagaraSimCacheCompare::CompareEmitter(const UNiagaraSimCache& LhsCache, c
 				return false;
 			}
 		}
-	}
-
-	if (InstanceMapping.Num() == 0)
-	{
-		for (int32 i=0; i < LhsNumInstances; ++i)
+		else
 		{
-			InstanceMapping[i] = i;
+			for (int32 i = 0; i < LhsNumInstances; ++i)
+			{
+				InstanceMapping[i] = i;
+			}
 		}
 	}
 
