- Fix for copying data from the sim cache to the clipboard
- This only works on system / emitter data at the moment

#rnx
#tests Capture a cache from attribute spreadsheet and noticed you can't copy particle / system data

diff --git a/Engine/Plugins/FX/Niagara/Source/NiagaraEditor/Private/ViewModels/NiagaraSimCacheViewModel.cpp b/Engine/Plugins/FX/Niagara/Source/NiagaraEditor/Private/ViewModels/NiagaraSimCacheViewModel.cpp
--- a/Engine/Plugins/FX/Niagara/Source/NiagaraEditor/Private/ViewModels/NiagaraSimCacheViewModel.cpp
+++ b/Engine/Plugins/FX/Niagara/Source/NiagaraEditor/Private/ViewModels/NiagaraSimCacheViewModel.cpp
@@ -595,7 +595,7 @@ TArray<TSharedRef<FNiagaraSimCacheOverviewItem>>* FNiagaraSimCacheViewModel::Get
 
 bool FNiagaraSimCacheViewModel::CanCopyActiveToClipboard() const
 {
-	return IsCacheValid() && SelectionMode == ESelectionMode::DataInterface;
+	return IsCacheValid() && (SelectionMode == ESelectionMode::SystemInstance || SelectionMode == ESelectionMode::Emitter);
 }
 
 void FNiagaraSimCacheViewModel::CopyActiveToClipboard() const
