Fixed tooltip and search keywords for lerp operation

diff --git a/Engine/Plugins/FX/Niagara/Source/NiagaraEditor/Private/NiagaraEditorCommon.cpp b/Engine/Plugins/FX/Niagara/Source/NiagaraEditor/Private/NiagaraEditorCommon.cpp
--- a/Engine/Plugins/FX/Niagara/Source/NiagaraEditor/Private/NiagaraEditorCommon.cpp
+++ b/Engine/Plugins/FX/Niagara/Source/NiagaraEditor/Private/NiagaraEditorCommon.cpp
@@ -369,8 +369,8 @@ void FNiagaraOpInfo::Init()
 		Op->FriendlyName = NSLOCTEXT("NiagaraOpInfo", "Lerp Name", "Lerp");
 		Op->CompactName = FText::FromString(TEXT("Lerp"));
 		Op->bShowPinNamesInCompactMode = true;
-		Op->Description = NSLOCTEXT("NiagaraOpInfo", "Lerp Desc", "Result = (A * (1 - C)) + (B * C)");
-		Op->Keywords = FText::FromString(TEXT("lerp"));
+		Op->Description = NSLOCTEXT("NiagaraOpInfo", "Lerp Desc", "Result = (A * (1 - Alpha)) + (B * Alpha)");
+		Op->Keywords = FText::FromString(TEXT("lerp linear interpolate"));
 
 		FText AlphaPinFriendlyName = FText::FromString("Alpha");
 		FText AlphaTooltip = NSLOCTEXT("NiagaraOpInfo", "Lerp Alpha Tooltip", "A value typically between 0 and 1. Determines the percentage to use for interpolating from A to B.");
