- Sim Cache Debugging Support for Stateless Emitters
- Adding support so that we can more easily debug what is happening in stateless emitters (what the values are, what attributes are active, etc)
- With this support we can also use automation testing to compare the cache results together

#rnx
#tests Capture stateless data inside attribute spreadsheet, etc, and validate that capturing without debugging data during sequencer doesn't pull data in

diff --git a/Engine/Plugins/FX/Niagara/Source/Niagara/Classes/NiagaraDataSet.h b/Engine/Plugins/FX/Niagara/Source/Niagara/Classes/NiagaraDataSet.h
--- a/Engine/Plugins/FX/Niagara/Source/Niagara/Classes/NiagaraDataSet.h
+++ b/Engine/Plugins/FX/Niagara/Source/Niagara/Classes/NiagaraDataSet.h
@@ -94,11 +94,12 @@ public:
 
 	NIAGARA_API void SwapInstances(uint32 OldIndex, uint32 NewIndex);
 	NIAGARA_API void KillInstance(uint32 InstanceIdx);
-	NIAGARA_API void CopyTo(FNiagaraDataBuffer& DestBuffer, int32 SrcStartIdx, int32 DestStartIdx, int32 NumInstances)const;
-	NIAGARA_API void CopyToUnrelated(FNiagaraDataBuffer& DestBuffer, int32 SrcStartIdx, int32 DestStartIdx, int32 NumInstances)const;
+	NIAGARA_API void CopyTo(FNiagaraDataBuffer& DestBuffer, int32 SrcStartIdx, int32 DestStartIdx, int32 NumInstances) const;
+	NIAGARA_API void CopyToUnrelated(FNiagaraDataBuffer& DestBuffer, int32 SrcStartIdx, int32 DestStartIdx, int32 NumInstances) const;
 	NIAGARA_API void GPUCopyFrom(const float* GPUReadBackFloat, const int* GPUReadBackInt, const FFloat16* GPUReadBackHalf, int32 StartIdx, int32 NumInstances, uint32 InSrcFloatStride, uint32 InSrcIntStride, uint32 InSrcHalfStride);
 	NIAGARA_API void PushCPUBuffersToGPU(const TArray<FNiagaraDataBufferRef>& SourceBuffers, bool bReleaseRef, FRHICommandList& RHICmdList, ERHIFeatureLevel::Type FeatureLevel, const TCHAR* DebugSimName);
-	NIAGARA_API void Dump(int32 StartIndex, int32 NumInstances, const FString& Label, const FName& SortParameterKey = FName())const;
+	NIAGARA_API void TransferGPUToCPUImmediate(FRHICommandList& RHICmdList, FNiagaraGpuComputeDispatchInterface* ComputeInterface, FNiagaraDataBuffer* CPUBuffer) const;
+	NIAGARA_API void Dump(int32 StartIndex, int32 NumInstances, const FString& Label, const FName& SortParameterKey = FName()) const;
 
 	FORCEINLINE TArrayView<uint8 const* RESTRICT const> ReadRegisterTable() const { return TArrayView<uint8 const* RESTRICT const>(RegisterTable); }
 	FORCEINLINE TArrayView<uint8* RESTRICT const> EditRegisterTable() const { return TArrayView<uint8* RESTRICT const>(RegisterTable); }
diff --git a/Engine/Plugins/FX/Niagara/Source/Niagara/Internal/Stateless/NiagaraStatelessCommon.h b/Engine/Plugins/FX/Niagara/Source/Niagara/Internal/Stateless/NiagaraStatelessCommon.h
--- a/Engine/Plugins/FX/Niagara/Source/Niagara/Internal/Stateless/NiagaraStatelessCommon.h
+++ b/Engine/Plugins/FX/Niagara/Source/Niagara/Internal/Stateless/NiagaraStatelessCommon.h
@@ -7,12 +7,12 @@
 #include "NiagaraStatelessCommon.generated.h"
 
 UENUM(meta = (Bitflags, UseEnumValuesAsMaskValuesInEditor = "true"))
-enum class ENiagaraStatelessFeatureMask : uint8
+enum class ENiagaraStatelessFeatureMask : uint32
 {
 	// We can execute on the GPU (Might be broken down into GPUCompute | GPUGraphics | GPUAsyncCompute in future but this will remain the master mask)
-	ExecuteGPU		= 1 << 0,
+	ExecuteGPU			= 1 << 0,
 	// We can execute on the CPU
-	ExecuteCPU		= 1 << 1,
+	ExecuteCPU			= 1 << 1,
 
 	None			= 0							UMETA(Hidden),
 	ExecuteAll		= ExecuteGPU | ExecuteCPU	UMETA(Hidden),
diff --git a/Engine/Plugins/FX/Niagara/Source/Niagara/Internal/Stateless/NiagaraStatelessEmitter.h b/Engine/Plugins/FX/Niagara/Source/Niagara/Internal/Stateless/NiagaraStatelessEmitter.h
--- a/Engine/Plugins/FX/Niagara/Source/Niagara/Internal/Stateless/NiagaraStatelessEmitter.h
+++ b/Engine/Plugins/FX/Niagara/Source/Niagara/Internal/Stateless/NiagaraStatelessEmitter.h
@@ -137,7 +137,23 @@ protected:
 	TObjectPtr<UClass> EmitterTemplateClass;
 
 	UPROPERTY(EditAnywhere, Category = "General")
-	bool bDeterministic = false;
+	uint32 bDeterministic : 1 = false;
+
+#if WITH_EDITORONLY_DATA
+	/**
+	When enabled the emitter will output all available attributes.
+	You should not need to modify this with the exception of debugging / testing and as it will impact cooked performance and memory
+	*/
+	UPROPERTY(EditAnywhere, Category = "General", AdvancedDisplay)
+	uint32 bForceOutputAllAttributes : 1 = false;
+
+	/**
+	When enabled the emitter will always include UniqueID in the output attributes.
+	You should not need to modify this with the exception of debugging / testing and as it will impact cooked performance and memory
+	*/
+	UPROPERTY(EditAnywhere, Category = "General", AdvancedDisplay, meta = (EditCondition = "!bForceOutputAllAttributes"))
+	uint32 bForceOutputUniqueID : 1 = false;
+#endif
 
 	UPROPERTY(EditAnywhere, Category = "General", AdvancedDisplay, meta = (Bitmask, BitMaskEnum = "/Script/Niagara.ENiagaraStatelessFeatureMask"))
 	uint32 AllowedFeatureMask = uint32(ENiagaraStatelessFeatureMask::All);
diff --git a/Engine/Plugins/FX/Niagara/Source/Niagara/Internal/Stateless/NiagaraStatelessEmitterInstance.h b/Engine/Plugins/FX/Niagara/Source/Niagara/Internal/Stateless/NiagaraStatelessEmitterInstance.h
--- a/Engine/Plugins/FX/Niagara/Source/Niagara/Internal/Stateless/NiagaraStatelessEmitterInstance.h
+++ b/Engine/Plugins/FX/Niagara/Source/Niagara/Internal/Stateless/NiagaraStatelessEmitterInstance.h
@@ -93,6 +93,9 @@ private:
 	//-TODO: This can be shared perhaps?
 	void SetExecutionStateInternal(ENiagaraExecutionState InExecutionState);
 
+public:
+	void CaptureForDebugging(FNiagaraDataBuffer* DataBuffer) const;
+
 private:
 	uint32										bCanEverExecute : 1 = false;
 	uint32										bEmitterEnabled_GT : 1 = true;
diff --git a/Engine/Plugins/FX/Niagara/Source/Niagara/Private/NiagaraDataSet.cpp b/Engine/Plugins/FX/Niagara/Source/Niagara/Private/NiagaraDataSet.cpp
--- a/Engine/Plugins/FX/Niagara/Source/Niagara/Private/NiagaraDataSet.cpp
+++ b/Engine/Plugins/FX/Niagara/Source/Niagara/Private/NiagaraDataSet.cpp
@@ -1450,6 +1450,59 @@ void FNiagaraDataBuffer::PushCPUBuffersToGPU(const TArray<FNiagaraDataBufferRef>
 	}
 }
 
+void FNiagaraDataBuffer::TransferGPUToCPUImmediate(FRHICommandList& RHICmdList, FNiagaraGpuComputeDispatchInterface* ComputeInterface, FNiagaraDataBuffer* CPUBuffer) const
+{
+	check(GetOwner()->GetCompiledData().GetLayoutHash() == CPUBuffer->GetOwner()->GetCompiledData().GetLayoutHash());
+
+	uint32 GPUNumInstances = 0;
+	if (GPUInstanceCountBufferOffset != INDEX_NONE)
+	{
+		FNiagaraGPUInstanceCountManager& CountManager = ComputeInterface->GetGPUInstanceCounterManager();
+		FRHIBuffer* CountBuffer = CountManager.GetInstanceCountBuffer().Buffer;
+		GPUNumInstances = *reinterpret_cast<const int32*>(RHICmdList.LockBuffer(CountBuffer, GPUInstanceCountBufferOffset * sizeof(int32), sizeof(int32), RLM_ReadOnly));
+		RHICmdList.UnlockBuffer(CountBuffer);
+	}
+
+	CPUBuffer->Allocate(GPUNumInstances);
+	CPUBuffer->SetNumInstances(GPUNumInstances);
+	if (GPUNumInstances == 0)
+	{
+		return;
+	}
+
+	if (GPUBufferFloat.NumBytes > 0)
+	{
+		const uint32 GPUStride	= FloatStride;
+		const uint8* GPUData	= reinterpret_cast<uint8*>(RHICmdList.LockBuffer(GPUBufferFloat.Buffer, 0, GPUBufferFloat.NumBytes, RLM_ReadOnly));
+
+		const uint32 CPUStride	= CPUBuffer->GetFloatStride();
+		uint8* CPUData			= CPUBuffer->GetComponentPtrFloat(0);
+
+		const int32 NumComponents = CPUBuffer->GetOwner()->GetNumFloatComponents();
+		for (int32 i = 0; i < NumComponents; ++i)
+		{
+			FMemory::Memcpy(CPUData + (CPUStride * i), GPUData + (GPUStride * i), GPUNumInstances * sizeof(float));
+		}
+		RHICmdList.UnlockBuffer(GPUBufferFloat.Buffer);
+	}
+
+	if (GPUBufferInt.NumBytes > 0)
+	{
+		const uint32 GPUStride	= Int32Stride;
+		const uint8* GPUData	= reinterpret_cast<uint8*>(RHICmdList.LockBuffer(GPUBufferInt.Buffer, 0, GPUBufferInt.NumBytes, RLM_ReadOnly));
+
+		const uint32 CPUStride	= CPUBuffer->GetInt32Stride();
+		uint8* CPUData			= CPUBuffer->GetComponentPtrInt32(0);
+
+		const int32 NumComponents = CPUBuffer->GetOwner()->GetNumInt32Components();
+		for (int32 i = 0; i < NumComponents; ++i)
+		{
+			FMemory::Memcpy(CPUData + (CPUStride * i), GPUData + (GPUStride * i), GPUNumInstances * sizeof(int32));
+		}
+		RHICmdList.UnlockBuffer(GPUBufferInt.Buffer);
+	}
+}
+
 void FNiagaraDataBuffer::Dump(int32 StartIndex, int32 InNumInstances, const FString& Label, const FName& SortParameterKey)const
 {
 	FNiagaraDataVariableIterator Itr(this, StartIndex);
diff --git a/Engine/Plugins/FX/Niagara/Source/Niagara/Private/NiagaraSimCache.cpp b/Engine/Plugins/FX/Niagara/Source/Niagara/Private/NiagaraSimCache.cpp
--- a/Engine/Plugins/FX/Niagara/Source/Niagara/Private/NiagaraSimCache.cpp
+++ b/Engine/Plugins/FX/Niagara/Source/Niagara/Private/NiagaraSimCache.cpp
@@ -19,6 +19,7 @@
 #include "NiagaraSimCacheHelper.h"
 #include "NiagaraSystemImpl.h"
 #include "NiagaraSystemInstance.h"
+#include "Stateless/NiagaraStatelessEmitterInstance.h"
 
 #include UE_INLINE_GENERATED_CPP_BY_NAME(NiagaraSimCache)
 
@@ -674,6 +675,14 @@ bool UNiagaraSimCache::WriteFrame(UNiagaraComponent* NiagaraComponent, FNiagaraS
 		CacheEmitterFrame.TotalSpawnedParticles = EmitterInstance.GetTotalSpawnedParticles();
 		if (CacheLayout.EmitterLayouts[i].SimTarget == ENiagaraSimTarget::GPUComputeSim)
 		{
+			if (FNiagaraStatelessEmitterInstance* StatelessEmitterInstance = EmitterInstance.AsStateless())
+			{
+				FNiagaraDataBufferRef DataBuffer = new FNiagaraDataBuffer(EmitterCurrentData->GetOwner());
+				StatelessEmitterInstance->CaptureForDebugging(DataBuffer);
+				Helper.WriteDataBuffer(*DataBuffer, CacheLayout.EmitterLayouts[i], CacheEmitterFrame.ParticleDataBuffers, 0, DataBuffer->GetNumInstances());
+				continue;
+			}
+
 			// First time we encounter a GPU emitter we need to process we also need to flush EOF updates / the compute process to ensure we get latest data
 			if (bNeedsGpuTickFlush)
 			{
diff --git a/Engine/Plugins/FX/Niagara/Source/Niagara/Private/NiagaraSimCacheHelper.h b/Engine/Plugins/FX/Niagara/Source/Niagara/Private/NiagaraSimCacheHelper.h
--- a/Engine/Plugins/FX/Niagara/Source/Niagara/Private/NiagaraSimCacheHelper.h
+++ b/Engine/Plugins/FX/Niagara/Source/Niagara/Private/NiagaraSimCacheHelper.h
@@ -11,6 +11,8 @@
 #include "NiagaraGpuComputeDispatchInterface.h"
 #include "NiagaraSimCache.h"
 #include "NiagaraSystemInstanceController.h"
+#include "Stateless/NiagaraStatelessEmitter.h"
+#include "Stateless/NiagaraStatelessEmitterData.h"
 
 struct FNiagaraSimCacheHelper
 {
@@ -318,15 +320,35 @@ struct FNiagaraSimCacheHelper
 	void BuildCacheLayoutForEmitter(const FNiagaraSimCacheCreateParameters& CreateParameters, FNiagaraSimCacheDataBuffersLayout& CacheLayout, int EmitterIndex)
 	{
 		const FNiagaraEmitterHandle& EmitterHandle = NiagaraSystem->GetEmitterHandle(EmitterIndex);
-		const FNiagaraEmitterCompiledData& EmitterCompiledData = NiagaraSystem->GetEmitterCompiledData()[EmitterIndex].Get();
-		const FVersionedNiagaraEmitterData* EmitterData = EmitterHandle.GetInstance().GetEmitterData();
-		if (EmitterHandle.GetIsEnabled() == false || EmitterData == nullptr)
+		if (!EmitterHandle.GetIsEnabled())
+		{
+			return;
+		}
+
+		const FNiagaraDataSetCompiledData* ParticleDataSet = nullptr;
+		const FVersionedNiagaraEmitterData* StatefulEmitterData = nullptr;
+
+		if (EmitterHandle.GetEmitterMode() == ENiagaraEmitterMode::Standard)
+		{
+			const FNiagaraEmitterCompiledData& EmitterCompiledData = NiagaraSystem->GetEmitterCompiledData()[EmitterIndex].Get();
+			StatefulEmitterData = EmitterHandle.GetInstance().GetEmitterData();
+			ParticleDataSet = StatefulEmitterData ? &EmitterCompiledData.DataSetCompiledData : nullptr;
+		}
+		// Include stateless emitters for debugging only.  For none debugging there is no value in storing this data as they are stateless.
+		else if ( CreateParameters.bIncludeDebugData )
+		{
+			const UNiagaraStatelessEmitter* StatelessEmitter = EmitterHandle.GetStatelessEmitter();
+			const FNiagaraStatelessEmitterDataPtr StatelessEmitterData = StatelessEmitter ? StatelessEmitter->GetEmitterData() : nullptr;
+			ParticleDataSet = StatelessEmitterData ? &StatelessEmitterData->ParticleDataSetCompiledData : nullptr;
+		}
+
+		if (ParticleDataSet == nullptr)
 		{
 			return;
 		}
 
 		// Find potential candidates for re-basing
-		CacheLayout.bLocalSpace = EmitterData->bLocalSpace;
+		CacheLayout.bLocalSpace = StatefulEmitterData ? StatefulEmitterData->bLocalSpace : true;
 
 		TArray<FName> RebaseVariableNames;
 		if ( CreateParameters.bAllowRebasing && CacheLayout.bLocalSpace == false )
@@ -357,31 +379,34 @@ struct FNiagaraSimCacheHelper
 			}
 
 		#if WITH_EDITORONLY_DATA
-			// Look for renderer attributes bound to Quat / Matrix types are we will want to rebase those
-			// We will add all Position types after this so no need to add them here
-			EmitterData->ForEachEnabledRenderer(
-				[&](UNiagaraRendererProperties* RenderProperties)
-				{
-					for (FNiagaraVariable BoundAttribute : RenderProperties->GetBoundAttributes())
+			if (StatefulEmitterData)
+			{
+				// Look for renderer attributes bound to Quat / Matrix types are we will want to rebase those
+				// We will add all Position types after this so no need to add them here
+				StatefulEmitterData->ForEachEnabledRenderer(
+					[&](UNiagaraRendererProperties* RenderProperties)
 					{
-						if ( (BoundAttribute.GetType() == FNiagaraTypeDefinition::GetQuatDef()) ||
-							 (BoundAttribute.GetType() == FNiagaraTypeDefinition::GetMatrix4Def()) )
+						for (FNiagaraVariable BoundAttribute : RenderProperties->GetBoundAttributes())
 						{
-							if (BoundAttribute.RemoveRootNamespace(FNiagaraConstants::ParticleAttributeNamespaceString))
+							if ( (BoundAttribute.GetType() == FNiagaraTypeDefinition::GetQuatDef()) ||
+								 (BoundAttribute.GetType() == FNiagaraTypeDefinition::GetMatrix4Def()) )
 							{
-								if (EmitterCompiledData.DataSetCompiledData.Variables.Contains(BoundAttribute) && ForceExcludeNames.Contains(BoundAttribute.GetName()) == false )
+								if (BoundAttribute.RemoveRootNamespace(FNiagaraConstants::ParticleAttributeNamespaceString))
 								{
-									RebaseVariableNames.AddUnique(BoundAttribute.GetName());
+									if (ParticleDataSet->Variables.Contains(BoundAttribute) && ForceExcludeNames.Contains(BoundAttribute.GetName()) == false )
+									{
+										RebaseVariableNames.AddUnique(BoundAttribute.GetName());
+									}
 								}
 							}
 						}
 					}
-				}
-			);
+				);
+			}
 		#endif
 
 			// Look for regular attributes that we are forcing to rebase or can rebase like positions
-			for (const FNiagaraVariableBase& Variable : EmitterCompiledData.DataSetCompiledData.Variables)
+			for (const FNiagaraVariableBase& Variable : ParticleDataSet->Variables)
 			{
 				if ( Variable.GetType() == FNiagaraTypeDefinition::GetPositionDef() )
 				{
@@ -400,7 +425,7 @@ struct FNiagaraSimCacheHelper
 		TArray<FName> InterpVariableNames;
 		if (CreateParameters.bAllowInterpolation)
 		{
-			for (const FNiagaraVariableBase& Variable : EmitterCompiledData.DataSetCompiledData.Variables)
+			for (const FNiagaraVariableBase& Variable : ParticleDataSet->Variables)
 			{
 				if (!CanInterpolateVariable(Variable) || CreateParameters.RebaseExcludeAttributes.Contains(Variable.GetName()))
 				{
@@ -431,7 +456,7 @@ struct FNiagaraSimCacheHelper
 			}
 		}
 
-		BuildCacheLayout(CreateParameters, CacheLayout, EmitterCompiledData.DataSetCompiledData, EmitterHandle.GetName(), MoveTemp(RebaseVariableNames), MoveTemp(InterpVariableNames), ExplicitCaptureAttributes);
+		BuildCacheLayout(CreateParameters, CacheLayout, *ParticleDataSet, EmitterHandle.GetName(), MoveTemp(RebaseVariableNames), MoveTemp(InterpVariableNames), ExplicitCaptureAttributes);
 	}
 
 	static bool BuildCacheReadMappings(FNiagaraSimCacheDataBuffersLayout& CacheLayout, const FNiagaraDataSetCompiledData& CompiledData)
diff --git a/Engine/Plugins/FX/Niagara/Source/Niagara/Private/Stateless/NiagaraStatelessComputeManager.cpp b/Engine/Plugins/FX/Niagara/Source/Niagara/Private/Stateless/NiagaraStatelessComputeManager.cpp
--- a/Engine/Plugins/FX/Niagara/Source/Niagara/Private/Stateless/NiagaraStatelessComputeManager.cpp
+++ b/Engine/Plugins/FX/Niagara/Source/Niagara/Private/Stateless/NiagaraStatelessComputeManager.cpp
@@ -16,21 +16,131 @@
 #include "RenderGraphUtils.h"
 #include "SceneView.h"
 
-bool GNiagaraStatelessComputeManager_UseCache = true;
-FAutoConsoleVariableRef CVarNiagaraStatelessComputeManager_UseCache(
-	TEXT("fx.NiagaraStateless.ComputeManager.UseCache"),
-	GNiagaraStatelessComputeManager_UseCache,
-	TEXT("When enabled we will attempt to reuse allocated buffers between frames."),
-	ECVF_Default
-);
-
-int32 GNiagaraStatelessComputeManager_CPUThreshold = 0;
-FAutoConsoleVariableRef CVarNiagaraStatelessComputeManager_CPUThreshold(
-	TEXT("fx.NiagaraStateless.ComputeManager.CPUThreshold"),
-	GNiagaraStatelessComputeManager_CPUThreshold,
-	TEXT("When lower than this particle count prefer to use the CPU over dispatching a compute shader."),
-	ECVF_Default
-);
+namespace NiagaraStatelessComputeManagerPrivate
+{
+	enum class EComputeExecutionPath
+	{
+		None,
+		CPU,
+		GPU,
+	};
+
+	bool GUseDataBufferCache = true;
+	FAutoConsoleVariableRef CVarUseCache(
+		TEXT("fx.NiagaraStateless.ComputeManager.UseCache"),
+		GUseDataBufferCache,
+		TEXT("When enabled we will attempt to reuse allocated buffers between frames."),
+		ECVF_Default
+	);
+
+	int32 GParticleCountCPUThreshold = 0;
+	FAutoConsoleVariableRef CVarCPUThreshold(
+		TEXT("fx.NiagaraStateless.ComputeManager.CPUThreshold"),
+		GParticleCountCPUThreshold,
+		TEXT("When lower than this particle count prefer to use the CPU over dispatching a compute shader."),
+		ECVF_Default
+	);
+
+	EComputeExecutionPath DetermineComputeExecutionPath(const FNiagaraStatelessEmitterData* EmitterData, uint32 ActiveParticlesEstimate)
+	{
+		const bool bAllowGPUExec = EnumHasAnyFlags(EmitterData->FeatureMask, ENiagaraStatelessFeatureMask::ExecuteGPU);
+		const bool bUseCPUExec = EnumHasAnyFlags(EmitterData->FeatureMask, ENiagaraStatelessFeatureMask::ExecuteCPU) && (!bAllowGPUExec || (ActiveParticlesEstimate <= uint32(GParticleCountCPUThreshold)));
+		if (bUseCPUExec)
+		{
+			return EComputeExecutionPath::CPU;
+		}
+		if (bAllowGPUExec)
+		{
+			return EComputeExecutionPath::GPU;
+		}
+		return EComputeExecutionPath::None;
+	}
+
+	void GenerateGPUData(FRHICommandList& RHICmdList, FNiagaraGpuComputeDispatchInterface* ComputeInterface, TConstArrayView<FNiagaraStatelessComputeManager::FStatelessDataCache*> DataToGenerate)
+	{
+		const int32 NumJobs = DataToGenerate.Num();
+
+		// Get Count Buffer
+		FNiagaraGPUInstanceCountManager& CountManager = ComputeInterface->GetGPUInstanceCounterManager();
+		FRHIUnorderedAccessView* CountBufferUAV = CountManager.GetInstanceCountBuffer().UAV;
+
+		// Build Transitions
+		TArray<FRHITransitionInfo> TransitionsBefore;
+		TArray<FRHITransitionInfo> TransitionsAfter;
+		{
+			TransitionsBefore.Reserve(1 + (NumJobs * 2));
+			TransitionsAfter.Reserve(1 + (NumJobs * 2));
+
+			TransitionsBefore.Emplace(CountManager.GetInstanceCountBuffer().Buffer, FNiagaraGPUInstanceCountManager::kCountBufferDefaultState, ERHIAccess::UAVCompute);
+			TransitionsAfter.Emplace(CountManager.GetInstanceCountBuffer().Buffer, ERHIAccess::UAVCompute, FNiagaraGPUInstanceCountManager::kCountBufferDefaultState);
+
+			for (FNiagaraStatelessComputeManager::FStatelessDataCache* CacheData : DataToGenerate)
+			{
+				const FRWBuffer& FloatBuffer = CacheData->DataBuffer->GetGPUBufferFloat();
+				if (FloatBuffer.NumBytes > 0)
+				{
+					TransitionsBefore.Emplace(FloatBuffer.Buffer, ERHIAccess::SRVMask, ERHIAccess::UAVCompute);
+					TransitionsAfter.Emplace(FloatBuffer.Buffer, ERHIAccess::UAVCompute, ERHIAccess::SRVMask);
+				}
+				const FRWBuffer& IntBuffer = CacheData->DataBuffer->GetGPUBufferInt();
+				if (IntBuffer.NumBytes > 0)
+				{
+					TransitionsBefore.Emplace(IntBuffer.Buffer, ERHIAccess::SRVMask, ERHIAccess::UAVCompute);
+					TransitionsAfter.Emplace(IntBuffer.Buffer, ERHIAccess::UAVCompute, ERHIAccess::SRVMask);
+				}
+			}
+		}
+
+		FNiagaraEmptyUAVPoolScopedAccess UAVPoolAccessScope(ComputeInterface->GetEmptyUAVPool());
+		FRHIUnorderedAccessView* EmptyIntBufferUAV = ComputeInterface->GetEmptyUAVFromPool(RHICmdList, PF_R32_SINT, ENiagaraEmptyUAVType::Buffer);
+
+		// Execute Simulations
+		RHICmdList.Transition(TransitionsBefore);
+
+		RHICmdList.BeginUAVOverlap(CountBufferUAV);
+		for (FNiagaraStatelessComputeManager::FStatelessDataCache* CacheData : DataToGenerate)
+		{
+			const NiagaraStateless::FEmitterInstance_RT* EmitterInstance = CacheData->EmitterInstance;
+			const FNiagaraStatelessEmitterData* EmitterData = CacheData->EmitterInstance->EmitterData.Get();
+
+			// Do we need to update the parameter buffer?
+			if (EmitterInstance->BindingBufferData.IsSet())
+			{
+				EmitterInstance->BindingBuffer.Release();
+				EmitterInstance->BindingBuffer.Initialize(RHICmdList, TEXT("FNiagaraStatelessEmitterInstance::BindingBuffer"), sizeof(uint32), EmitterInstance->BindingBufferData->Num() / sizeof(uint32), EPixelFormat::PF_R32_UINT, EBufferUsageFlags::Static);
+				void* LockedBuffer = RHICmdList.LockBuffer(EmitterInstance->BindingBuffer.Buffer, 0, EmitterInstance->BindingBuffer.NumBytes, RLM_WriteOnly);
+				FMemory::Memcpy(LockedBuffer, EmitterInstance->BindingBufferData->GetData(), EmitterInstance->BindingBuffer.NumBytes);
+				RHICmdList.UnlockBuffer(EmitterInstance->BindingBuffer.Buffer);
+				EmitterInstance->BindingBufferData.Reset();
+			}
+
+			// Update parameters for this compute invocation
+			NiagaraStateless::FCommonShaderParameters* ShaderParameters = EmitterInstance->ShaderParameters.Get();
+			ShaderParameters->Common_SimulationTime = EmitterInstance->Age;
+			ShaderParameters->Common_SimulationDeltaTime = EmitterInstance->DeltaTime;
+			ShaderParameters->Common_SimulationInvDeltaTime = EmitterInstance->DeltaTime > 0.0f ? (1.0f / EmitterInstance->DeltaTime) : 0.0f;
+			ShaderParameters->Common_OutputBufferStride = CacheData->DataBuffer->GetFloatStride() / sizeof(float);
+			ShaderParameters->Common_GPUCountBufferOffset = CacheData->DataBuffer->GetGPUInstanceCountBufferOffset();
+			ShaderParameters->Common_FloatOutputBuffer = CacheData->DataBuffer->GetGPUBufferFloat().UAV;
+			//ShaderParameters->Common_HalfOutputBuffer		= CacheData->DataBuffer->GetGPUBufferHalf().UAV;
+			ShaderParameters->Common_IntOutputBuffer = CacheData->DataBuffer->GetGPUBufferInt().UAV.IsValid() ? CacheData->DataBuffer->GetGPUBufferInt().UAV.GetReference() : EmptyIntBufferUAV;
+			ShaderParameters->Common_GPUCountBuffer = CountBufferUAV;
+			ShaderParameters->Common_StaticFloatBuffer = EmitterData->StaticFloatBuffer.SRV;
+			ShaderParameters->Common_ParameterBuffer = FNiagaraRenderer::GetSrvOrDefaultUInt(EmitterInstance->BindingBuffer.SRV);
+
+			// Execute the simulation
+			TShaderRef<NiagaraStateless::FSimulationShader> ComputeShader = EmitterData->GetShader();
+			FRHIComputeShader* ComputeShaderRHI = ComputeShader.GetComputeShader();
+			const uint32 NumThreadGroups = FMath::DivideAndRoundUp<uint32>(CacheData->ActiveParticles, NiagaraStateless::FSimulationShader::ThreadGroupSize);
+
+			const FIntVector NumWrappedThreadGroups = FComputeShaderUtils::GetGroupCountWrapped(NumThreadGroups);
+			FComputeShaderUtils::Dispatch(RHICmdList, ComputeShader, EmitterData->GetShaderParametersMetadata(), *ShaderParameters, NumWrappedThreadGroups);
+		}
+		RHICmdList.EndUAVOverlap(CountBufferUAV);
+
+		RHICmdList.Transition(TransitionsAfter);
+	}
+}
 
 FNiagaraStatelessComputeManager::FNiagaraStatelessComputeManager(FNiagaraGpuComputeDispatchInterface* InOwnerInterface)
 	: FNiagaraGpuComputeDataManager(InOwnerInterface)
@@ -46,6 +156,7 @@ FNiagaraStatelessComputeManager::~FNiagaraStatelessComputeManager()
 FNiagaraDataBuffer* FNiagaraStatelessComputeManager::GetDataBuffer(uintptr_t EmitterKey, const NiagaraStateless::FEmitterInstance_RT* EmitterInstance)
 {
 	using namespace NiagaraStateless;
+	using namespace NiagaraStatelessComputeManagerPrivate;
 
 	if (TUniquePtr<FStatelessDataCache>* ExistingData = UsedData.Find(EmitterKey))
 	{
@@ -84,7 +195,7 @@ FNiagaraDataBuffer* FNiagaraStatelessComputeManager::GetDataBuffer(uintptr_t Emi
 	const uint32 DataSetLayoutHash = EmitterInstance->EmitterData->ParticleDataSetCompiledData.GetLayoutHash();
 
 	FStatelessDataCache* CacheData = nullptr;
-	if (GNiagaraStatelessComputeManager_UseCache)
+	if (GUseDataBufferCache)
 	{
 		for (int32 i=0; i < FreeData.Num(); ++i)
 		{
@@ -111,40 +222,110 @@ FNiagaraDataBuffer* FNiagaraStatelessComputeManager::GetDataBuffer(uintptr_t Emi
 	FRHICommandListImmediate& RHICmdList = FRHICommandListExecutor::GetImmediateCommandList();
 	CacheData->DataBuffer->AllocateGPU(RHICmdList, CacheData->ActiveParticles, ComputeInterface->GetFeatureLevel(), TEXT("StatelessSimBuffer"));
 
-	const bool bAllowGPUExec = EnumHasAnyFlags(EmitterData->FeatureMask, ENiagaraStatelessFeatureMask::ExecuteGPU);
-	const bool bUseCPUExec = EnumHasAnyFlags(EmitterData->FeatureMask, ENiagaraStatelessFeatureMask::ExecuteCPU) && (!bAllowGPUExec || (ActiveParticles <= uint32(GNiagaraStatelessComputeManager_CPUThreshold)));
-	if (bUseCPUExec)
+	const EComputeExecutionPath ComputeExecutionPath = DetermineComputeExecutionPath(EmitterData, ActiveParticles);
+	switch (ComputeExecutionPath)
 	{
-		FParticleSimulationContext ParticleSimulation(EmitterData, EmitterInstance->BindingBufferData.Get(TArray<uint8>()));
-		ParticleSimulation.SimulateGPU(RHICmdList, EmitterInstance->RandomSeed, EmitterInstance->Age, EmitterInstance->DeltaTime, EmitterInstance->SpawnInfos, CacheData->DataBuffer);
-		if (ParticleSimulation.GetNumInstances() == 0)
+		case EComputeExecutionPath::CPU:
+		{
+			FParticleSimulationContext ParticleSimulation(EmitterData, EmitterInstance->BindingBufferData.Get(TArray<uint8>()));
+			ParticleSimulation.SimulateGPU(RHICmdList, EmitterInstance->RandomSeed, EmitterInstance->Age, EmitterInstance->DeltaTime, EmitterInstance->SpawnInfos, CacheData->DataBuffer);
+			if (ParticleSimulation.GetNumInstances() == 0)
+			{
+				FreeData.Emplace(CacheData);
+				return nullptr;
+			}
+			break;
+		}
+
+		case EComputeExecutionPath::GPU:
 		{
+			CacheData->DataBuffer->SetNumInstances(CacheData->ActiveParticles);
+
+			FNiagaraGPUInstanceCountManager& CountManager = ComputeInterface->GetGPUInstanceCounterManager();
+			const uint32 CountOffset = CountManager.AcquireOrAllocateEntry(RHICmdList);
+			CacheData->DataBuffer->SetGPUInstanceCountBufferOffset(CountOffset);
+			CountsToRelease.Add(CountOffset);
+
+			GPUDataToGenerate.Add(CacheData);
+			break;
+		}
+
+		default:
+			ensureMsgf(false, TEXT("No execution path was found for stateless emitter, data will not be generated"));
 			FreeData.Emplace(CacheData);
 			return nullptr;
-		}
+	}
+
+	UsedData.Emplace(EmitterKey, CacheData);
+	return CacheData->DataBuffer;
+}
+
+void FNiagaraStatelessComputeManager::GenerateDataBufferForDebugging(FRHICommandListImmediate& RHICmdList, FNiagaraDataBuffer* DataBuffer, const NiagaraStateless::FEmitterInstance_RT* EmitterInstance) const
+{
+	using namespace NiagaraStateless;
+	using namespace NiagaraStatelessComputeManagerPrivate;
+
+	check(IsInRenderingThread());
+
+	const FNiagaraStatelessEmitterData* EmitterData = EmitterInstance->EmitterData.Get();
+	const uint32 ActiveParticlesEstimate = EmitterData->CalculateActiveParticles(
+		EmitterInstance->RandomSeed,
+		EmitterInstance->SpawnInfos,
+		EmitterInstance->Age,
+		&EmitterInstance->ShaderParameters->SpawnParameters
+	);
 
-		//CacheData->DataBuffer->Dump(0, CacheData->DataBuffer->GetNumInstances(), TEXT("Stateless"));
+	if (ActiveParticlesEstimate == 0)
+	{
+		DataBuffer->SetNumInstances(0);
+		return;
 	}
-	else
+
+	const EComputeExecutionPath ComputeExecutionPath = DetermineComputeExecutionPath(EmitterData, ActiveParticlesEstimate);
+	switch (ComputeExecutionPath)
 	{
-		if (!ensure(bAllowGPUExec))
+		case EComputeExecutionPath::CPU:
 		{
-			FreeData.Emplace(CacheData);
-			return nullptr;
+			FParticleSimulationContext ParticleSimulation(EmitterData, EmitterInstance->BindingBufferData.Get(TArray<uint8>()));
+			ParticleSimulation.Simulate(EmitterInstance->RandomSeed, EmitterInstance->Age, EmitterInstance->DeltaTime, EmitterInstance->SpawnInfos, DataBuffer);
+			break;
 		}
 
-		CacheData->DataBuffer->SetNumInstances(CacheData->ActiveParticles);
+		case EComputeExecutionPath::GPU:
+		{
+			FNiagaraGpuComputeDispatchInterface* ComputeInterface = GetOwnerInterface();
+			FNiagaraGPUInstanceCountManager& CountManager = ComputeInterface->GetGPUInstanceCounterManager();
 
-		FNiagaraGPUInstanceCountManager& CountManager = ComputeInterface->GetGPUInstanceCounterManager();
-		const uint32 CountOffset = CountManager.AcquireOrAllocateEntry(RHICmdList);
-		CacheData->DataBuffer->SetGPUInstanceCountBufferOffset(CountOffset);
-		CountsToRelease.Add(CountOffset);
+			// Allocate counter and destination data
+			FNiagaraDataBufferRef GPUDataBuffer = new FNiagaraDataBuffer(DataBuffer->GetOwner());
 
-		GPUDataToGenerate.Add(CacheData);
-	}
+			uint32 CountIndex = CountManager.AcquireOrAllocateEntry(RHICmdList);
+			GPUDataBuffer->AllocateGPU(RHICmdList, ActiveParticlesEstimate, ComputeInterface->GetFeatureLevel(), TEXT("StatelessSimBuffer"));
+			GPUDataBuffer->SetGPUInstanceCountBufferOffset(CountIndex);
 
-	UsedData.Emplace(EmitterKey, CacheData);
-	return CacheData->DataBuffer;
+			// Generate the data
+			FStatelessDataCache DataCache;
+			DataCache.DataBuffer		= GPUDataBuffer;
+			DataCache.EmitterInstance	= EmitterInstance;
+			DataCache.ActiveParticles	= ActiveParticlesEstimate;
+
+			FStatelessDataCache* DataCachePtr = &DataCache;
+			GenerateGPUData(RHICmdList, ComputeInterface, MakeConstArrayView(&DataCachePtr, 1));
+
+			// Copy to CPU data
+			//TransferGPUToCPU(RHICmdList, ComputeInterface, GPUDataBuffer, DataBuffer);
+			GPUDataBuffer->TransferGPUToCPUImmediate(RHICmdList, ComputeInterface, DataBuffer);
+
+			// Release the GPU buffer and count
+			GPUDataBuffer->ReleaseGPU();
+			GPUDataBuffer->SetGPUInstanceCountBufferOffset(INDEX_NONE);
+			CountManager.FreeEntry(CountIndex);
+			break;
+		}
+
+		default:
+			break;
+	}
 }
 
 void FNiagaraStatelessComputeManager::OnPostPreRender(FRDGBuilder& GraphBuilder)
@@ -164,88 +345,9 @@ void FNiagaraStatelessComputeManager::OnPostPreRender(FRDGBuilder& GraphBuilder)
 		[DataToGenerate=MoveTemp(GPUDataToGenerate), ComputeInterface=GetOwnerInterface()](FRHICommandListImmediate& RHICmdList)
 		{
 			SCOPED_DRAW_EVENT(RHICmdList, FNiagaraStatelessComputeManager_OnPostPreRender);
-	
-			const int32 NumJobs = DataToGenerate.Num();
 
-			// Get Count Buffer
 			FNiagaraGPUInstanceCountManager& CountManager = ComputeInterface->GetGPUInstanceCounterManager();
-			FRHIUnorderedAccessView* CountBufferUAV = CountManager.GetInstanceCountBuffer().UAV;
-
-			// Build Transitions
-			TArray<FRHITransitionInfo> TransitionsBefore;
-			TArray<FRHITransitionInfo> TransitionsAfter;
-			{
-				TransitionsBefore.Reserve(1 + (NumJobs * 2));
-				TransitionsAfter.Reserve(1 + (NumJobs * 2));
-
-				TransitionsBefore.Emplace(CountManager.GetInstanceCountBuffer().Buffer, FNiagaraGPUInstanceCountManager::kCountBufferDefaultState, ERHIAccess::UAVCompute);
-				TransitionsAfter.Emplace(CountManager.GetInstanceCountBuffer().Buffer, ERHIAccess::UAVCompute, FNiagaraGPUInstanceCountManager::kCountBufferDefaultState);
-
-				for (FStatelessDataCache* CacheData : DataToGenerate)
-				{
-					const FRWBuffer& FloatBuffer = CacheData->DataBuffer->GetGPUBufferFloat();
-					if (FloatBuffer.NumBytes > 0)
-					{
-						TransitionsBefore.Emplace(FloatBuffer.Buffer, ERHIAccess::SRVMask, ERHIAccess::UAVCompute);
-						TransitionsAfter.Emplace(FloatBuffer.Buffer, ERHIAccess::UAVCompute, ERHIAccess::SRVMask);
-					}
-					const FRWBuffer& IntBuffer = CacheData->DataBuffer->GetGPUBufferInt();
-					if (IntBuffer.NumBytes > 0)
-					{
-						TransitionsBefore.Emplace(IntBuffer.Buffer, ERHIAccess::SRVMask, ERHIAccess::UAVCompute);
-						TransitionsAfter.Emplace(IntBuffer.Buffer, ERHIAccess::UAVCompute, ERHIAccess::SRVMask);
-					}
-				}
-			}
-
-			FNiagaraEmptyUAVPoolScopedAccess UAVPoolAccessScope(ComputeInterface->GetEmptyUAVPool());
-			FRHIUnorderedAccessView* EmptyIntBufferUAV = ComputeInterface->GetEmptyUAVFromPool(RHICmdList, PF_R32_SINT, ENiagaraEmptyUAVType::Buffer);
-
-			// Execute Simulations
-			RHICmdList.Transition(TransitionsBefore);
-
-			RHICmdList.BeginUAVOverlap(CountBufferUAV);
-			for (FStatelessDataCache* CacheData : DataToGenerate)
-			{
-				const NiagaraStateless::FEmitterInstance_RT* EmitterInstance = CacheData->EmitterInstance;
-				const FNiagaraStatelessEmitterData* EmitterData = CacheData->EmitterInstance->EmitterData.Get();
-
-				// Do we need to update the parameter buffer?
-				if (EmitterInstance->BindingBufferData.IsSet())
-				{
-					EmitterInstance->BindingBuffer.Release();
-					EmitterInstance->BindingBuffer.Initialize(RHICmdList, TEXT("FNiagaraStatelessEmitterInstance::BindingBuffer"), sizeof(uint32), EmitterInstance->BindingBufferData->Num() / sizeof(uint32), EPixelFormat::PF_R32_UINT, EBufferUsageFlags::Static);
-					void* LockedBuffer = RHICmdList.LockBuffer(EmitterInstance->BindingBuffer.Buffer, 0, EmitterInstance->BindingBuffer.NumBytes, RLM_WriteOnly);
-					FMemory::Memcpy(LockedBuffer, EmitterInstance->BindingBufferData->GetData(), EmitterInstance->BindingBuffer.NumBytes);
-					RHICmdList.UnlockBuffer(EmitterInstance->BindingBuffer.Buffer);
-					EmitterInstance->BindingBufferData.Reset();
-				}
-
-				// Update parameters for this compute invocation
-				NiagaraStateless::FCommonShaderParameters* ShaderParameters = EmitterInstance->ShaderParameters.Get();
-				ShaderParameters->Common_SimulationTime			= EmitterInstance->Age;
-				ShaderParameters->Common_SimulationDeltaTime	= EmitterInstance->DeltaTime;
-				ShaderParameters->Common_SimulationInvDeltaTime	= EmitterInstance->DeltaTime > 0.0f ? (1.0f / EmitterInstance->DeltaTime) : 0.0f;
-				ShaderParameters->Common_OutputBufferStride		= CacheData->DataBuffer->GetFloatStride() / sizeof(float);
-				ShaderParameters->Common_GPUCountBufferOffset	= CacheData->DataBuffer->GetGPUInstanceCountBufferOffset();
-				ShaderParameters->Common_FloatOutputBuffer		= CacheData->DataBuffer->GetGPUBufferFloat().UAV;
-				//ShaderParameters->Common_HalfOutputBuffer		= CacheData->DataBuffer->GetGPUBufferHalf().UAV;
-				ShaderParameters->Common_IntOutputBuffer		= CacheData->DataBuffer->GetGPUBufferInt().UAV.IsValid() ? CacheData->DataBuffer->GetGPUBufferInt().UAV.GetReference() : EmptyIntBufferUAV;
-				ShaderParameters->Common_GPUCountBuffer			= CountBufferUAV;
-				ShaderParameters->Common_StaticFloatBuffer		= EmitterData->StaticFloatBuffer.SRV;
-				ShaderParameters->Common_ParameterBuffer		= FNiagaraRenderer::GetSrvOrDefaultUInt(EmitterInstance->BindingBuffer.SRV);
-
-				// Execute the simulation
-				TShaderRef<NiagaraStateless::FSimulationShader> ComputeShader = EmitterData->GetShader();
-				FRHIComputeShader* ComputeShaderRHI = ComputeShader.GetComputeShader();
-				const uint32 NumThreadGroups = FMath::DivideAndRoundUp<uint32>(CacheData->ActiveParticles, NiagaraStateless::FSimulationShader::ThreadGroupSize);
-
-				const FIntVector NumWrappedThreadGroups = FComputeShaderUtils::GetGroupCountWrapped(NumThreadGroups);
-				FComputeShaderUtils::Dispatch(RHICmdList, ComputeShader, EmitterData->GetShaderParametersMetadata(), *ShaderParameters, NumWrappedThreadGroups);
-			}
-			RHICmdList.EndUAVOverlap(CountBufferUAV);
-
-			RHICmdList.Transition(TransitionsAfter);
+			NiagaraStatelessComputeManagerPrivate::GenerateGPUData(RHICmdList, ComputeInterface, DataToGenerate);
 		}
 	);
 	GPUDataToGenerate.Empty();
diff --git a/Engine/Plugins/FX/Niagara/Source/Niagara/Private/Stateless/NiagaraStatelessComputeManager.h b/Engine/Plugins/FX/Niagara/Source/Niagara/Private/Stateless/NiagaraStatelessComputeManager.h
--- a/Engine/Plugins/FX/Niagara/Source/Niagara/Private/Stateless/NiagaraStatelessComputeManager.h
+++ b/Engine/Plugins/FX/Niagara/Source/Niagara/Private/Stateless/NiagaraStatelessComputeManager.h
@@ -18,6 +18,7 @@ namespace NiagaraStateless
 
 class FNiagaraStatelessComputeManager final : public FNiagaraGpuComputeDataManager
 {
+public:
 	struct FStatelessDataCache
 	{
 		uint32										DataSetLayoutHash = 0;
@@ -40,6 +41,9 @@ public:
 
 	FNiagaraDataBuffer* GetDataBuffer(uintptr_t EmitterKey, const NiagaraStateless::FEmitterInstance_RT* EmitterInstance);
 
+	// Used to execute the simulation immediately into a CPU side data buffer
+	void GenerateDataBufferForDebugging(FRHICommandListImmediate& RHICmdList, FNiagaraDataBuffer* DataBuffer, const NiagaraStateless::FEmitterInstance_RT* EmitterInstance) const;
+
 private:
 	void OnPostPreRender(FRDGBuilder& GraphBuilder);
 	void OnPostPostRender(FRDGBuilder& GraphBuilder);
diff --git a/Engine/Plugins/FX/Niagara/Source/Niagara/Private/Stateless/NiagaraStatelessEmitter.cpp b/Engine/Plugins/FX/Niagara/Source/Niagara/Private/Stateless/NiagaraStatelessEmitter.cpp
--- a/Engine/Plugins/FX/Niagara/Source/Niagara/Private/Stateless/NiagaraStatelessEmitter.cpp
+++ b/Engine/Plugins/FX/Niagara/Source/Niagara/Private/Stateless/NiagaraStatelessEmitter.cpp
@@ -427,33 +427,52 @@ void UNiagaraStatelessEmitter::BuildCompiledDataSet()
 			}
 		}
 
-		// Build data set from variables that are used
-		ForEachEnabledRenderer(
-			[this, &AvailableVariables](UNiagaraRendererProperties* RendererProps)
+		// Force all the attributes in?
+		if (bForceOutputAllAttributes)
+		{
+			for (const FNiagaraVariableBase& Variable : AvailableVariables)
 			{
-				if (AvailableVariables.Num() == 0)
-				{
-					return;
-				}
-
-				for (FNiagaraVariableBase BoundAttribute : RendererProps->GetBoundAttributes())
+				ParticleDataSetCompiledData.Variables.Emplace(Variable);
+			}
+		}
+		// Build data set from variables that are used by renderers
+		else
+		{
+			ForEachEnabledRenderer(
+				[this, &AvailableVariables](UNiagaraRendererProperties* RendererProps)
 				{
-					// Edge condition with UniqueID which does not contain the Particle namespace from Ribbon Renderer
-					BoundAttribute.RemoveRootNamespace(FNiagaraConstants::ParticleAttributeNamespaceString);
+					if (AvailableVariables.Num() == 0)
+					{
+						return;
+					}
 
-					const int32 Index = AvailableVariables.IndexOfByKey(BoundAttribute);
-					if (Index != INDEX_NONE)
+					for (FNiagaraVariableBase BoundAttribute : RendererProps->GetBoundAttributes())
 					{
-						AvailableVariables.RemoveAtSwap(Index, EAllowShrinking::No);
-						ParticleDataSetCompiledData.Variables.Emplace(BoundAttribute);
-						if (AvailableVariables.Num() == 0)
+						// Edge condition with UniqueID which does not contain the Particle namespace from Ribbon Renderer
+						BoundAttribute.RemoveRootNamespace(FNiagaraConstants::ParticleAttributeNamespaceString);
+
+						const int32 Index = AvailableVariables.IndexOfByKey(BoundAttribute);
+						if (Index != INDEX_NONE)
 						{
-							return;
+							AvailableVariables.RemoveAtSwap(Index, EAllowShrinking::No);
+							ParticleDataSetCompiledData.Variables.Emplace(BoundAttribute);
+							if (AvailableVariables.Num() == 0)
+							{
+								return;
+							}
 						}
 					}
 				}
+			);
+
+			if (bForceOutputUniqueID)
+			{
+				if (AvailableVariables.Contains(FNiagaraStatelessGlobals::Get().UniqueIDVariable))
+				{
+					ParticleDataSetCompiledData.Variables.Emplace(FNiagaraStatelessGlobals::Get().UniqueIDVariable);
+				}
 			}
-		);
+		}
 
 		//-TODO: We can alias variables in the data set, for example PreviousSpriteFacing could be SpriteFacing in some cases
 		ParticleDataSetCompiledData.BuildLayout();
diff --git a/Engine/Plugins/FX/Niagara/Source/Niagara/Private/Stateless/NiagaraStatelessEmitterInstance.cpp b/Engine/Plugins/FX/Niagara/Source/Niagara/Private/Stateless/NiagaraStatelessEmitterInstance.cpp
--- a/Engine/Plugins/FX/Niagara/Source/Niagara/Private/Stateless/NiagaraStatelessEmitterInstance.cpp
+++ b/Engine/Plugins/FX/Niagara/Source/Niagara/Private/Stateless/NiagaraStatelessEmitterInstance.cpp
@@ -714,3 +714,30 @@ void FNiagaraStatelessEmitterInstance::InitEmitterData()
 		EmitterData->bCanEverExecute &&
 		EmitterHandle.GetIsEnabled();
 }
+
+void FNiagaraStatelessEmitterInstance::CaptureForDebugging(FNiagaraDataBuffer* DataBuffer) const
+{
+	check(DataBuffer);
+
+	// Set instances to zero to handle any early outs
+	DataBuffer->SetNumInstances(0);
+	if (!bCanEverExecute || IsComplete())
+	{
+		return;
+	}
+
+	ENQUEUE_RENDER_COMMAND(CaptureStatelessForDebugging)(
+		[RenderThreadData=RenderThreadDataPtr.Get(), DataBuffer](FRHICommandListImmediate& RHICmdList)
+		{
+			// No compute manager then we can not do anything
+			if ( !RenderThreadData->ComputeManager )
+			{
+				return;
+			}
+
+			RenderThreadData->ComputeManager->GenerateDataBufferForDebugging(RHICmdList, DataBuffer, RenderThreadData);
+		}
+	);
+
+	FlushRenderingCommands();
+}
