FIx for DI specific compilation environment modifications missing with new compilation mode
#tests Used the debugger to validate that the DI functions were now being called

diff --git a/Engine/Plugins/FX/Niagara/Source/NiagaraEditor/Private/NiagaraCompilationTasks.cpp b/Engine/Plugins/FX/Niagara/Source/NiagaraEditor/Private/NiagaraCompilationTasks.cpp
--- a/Engine/Plugins/FX/Niagara/Source/NiagaraEditor/Private/NiagaraCompilationTasks.cpp
+++ b/Engine/Plugins/FX/Niagara/Source/NiagaraEditor/Private/NiagaraCompilationTasks.cpp
@@ -55,7 +55,7 @@ namespace NiagaraSystemCookStats
 
 namespace NiagaraCompilationTasksImpl
 {
-	static const FGuid UE_NIAGARA_ASYNC_TASK_COMPILER_VER = FGuid(0x1D7ABD2B, 0xE882465E, 0xA93B3CA6, 0x7753CD29);
+	static const FGuid UE_NIAGARA_ASYNC_TASK_COMPILER_VER = FGuid(0xA6981076, 0x2AC24481, 0x80CFB302, 0x7EDAB79F);
 	static UE::DerivedData::FCacheBucket NiagaraDDCBucket("NiagaraScript");
 
 	void GetUsagesToDuplicate(ENiagaraScriptUsage TargetUsage, TArray<ENiagaraScriptUsage>& DuplicateUsages)
@@ -866,7 +866,13 @@ void FNiagaraSystemCompilationTask::FCompileTaskInfo::IssueCompileGpu(FNiagaraSy
 		ShaderMapCompiler->AddShaderPlatform(ComputeInfo.ShaderMapId, ComputeInfo.ShaderPlatform);
 	}
 
-	ShaderMapCompiler->CompileScript(CompileId, PrecompileData->SourceName, DebugGroupName, CompileOptions, TranslateResults, TranslateOutput, TranslatedHlsl);
+	TArray<UNiagaraDataInterface*> UniqueDI;
+	for (const auto& It : CompilationCopyData->AggregatedDataInterfaceCDODuplicates)
+	{
+		UniqueDI.AddUnique(It.Value);
+	}
+
+	ShaderMapCompiler->CompileScript(CompileId, PrecompileData->SourceName, DebugGroupName, CompileOptions, TranslateResults, TranslateOutput, TranslatedHlsl, UniqueDI);
 }
 
 TOptional<FNiagaraCompileResults> FNiagaraSystemCompilationTask::FCompileTaskInfo::HandleDeprecatedGpuScriptResults() const
diff --git a/Engine/Plugins/FX/Niagara/Source/NiagaraEditor/Private/NiagaraCompiler.cpp b/Engine/Plugins/FX/Niagara/Source/NiagaraEditor/Private/NiagaraCompiler.cpp
--- a/Engine/Plugins/FX/Niagara/Source/NiagaraEditor/Private/NiagaraCompiler.cpp
+++ b/Engine/Plugins/FX/Niagara/Source/NiagaraEditor/Private/NiagaraCompiler.cpp
@@ -2430,7 +2430,8 @@ void FNiagaraShaderMapCompiler::CompileScript(
 	const FNiagaraCompileOptions& CompileOptions,
 	const FNiagaraTranslateResults& TranslateResults,
 	const FNiagaraTranslatorOutput& TranslatorOutput,
-	const FString& TranslatedHLSL)
+	const FString& TranslatedHLSL,
+	TConstArrayView<UNiagaraDataInterface*> DataInterfaces)
 {
 	TArray<TRefCountPtr<FShaderCommonCompileJob>> CompileJobs;
 
@@ -2447,6 +2448,11 @@ void FNiagaraShaderMapCompiler::CompileScript(
 			CompilationEnvironment->CompilerFlags.Add(CFLAG_NoFastMath);
 		}
 
+		for (UNiagaraDataInterface* DataInterface : DataInterfaces)
+		{
+			DataInterface->ModifyCompilationEnvironment(ActiveCompilation.ShaderPlatform, *CompilationEnvironment.GetReference());
+		}
+
 		ActiveCompilation.ShaderMap->CreateCompileJobs(
 			ShaderType,
 			DebugGroupName,
diff --git a/Engine/Plugins/FX/Niagara/Source/NiagaraEditor/Private/NiagaraCompiler.h b/Engine/Plugins/FX/Niagara/Source/NiagaraEditor/Private/NiagaraCompiler.h
--- a/Engine/Plugins/FX/Niagara/Source/NiagaraEditor/Private/NiagaraCompiler.h
+++ b/Engine/Plugins/FX/Niagara/Source/NiagaraEditor/Private/NiagaraCompiler.h
@@ -77,7 +77,8 @@ public:
 		const FNiagaraCompileOptions& InOptions,
 		const FNiagaraTranslateResults& InTranslateResults,
 		const FNiagaraTranslatorOutput& TranslatorOutput,
-		const FString& TranslatedHLSL);
+		const FString& TranslatedHLSL,
+		TConstArrayView<UNiagaraDataInterface*> DataInterfaces);
 	void AddShaderPlatform(const FNiagaraShaderMapId& ShaderMapId, EShaderPlatform ShaderPlatform);
 
 	// returns true if there are no more results requiring processing
