- Niagara Parallel GDME

diff --git a/Engine/Plugins/FX/Niagara/Source/Niagara/Classes/NiagaraComputeExecutionContext.h b/Engine/Plugins/FX/Niagara/Source/Niagara/Classes/NiagaraComputeExecutionContext.h
--- a/Engine/Plugins/FX/Niagara/Source/Niagara/Classes/NiagaraComputeExecutionContext.h
+++ b/Engine/Plugins/FX/Niagara/Source/Niagara/Classes/NiagaraComputeExecutionContext.h
@@ -59,7 +59,7 @@ struct FNiagaraGpuSpawnInfo
 struct INiagaraComputeDataBufferInterface
 {
 	virtual bool HasTranslucentDataToRender() const = 0;
-	virtual FNiagaraDataBuffer* GetDataToRender(bool bIsLowLatencyTranslucent) const = 0;
+	virtual FNiagaraDataBuffer* GetDataToRender(FRHICommandListBase& RHICmdList, bool bIsLowLatencyTranslucent) const = 0;
 };
 
 struct FNiagaraComputeExecutionContext : public INiagaraComputeDataBufferInterface
@@ -69,7 +69,7 @@ struct FNiagaraComputeExecutionContext : public INiagaraComputeDataBufferInterfa
 
 	// Begin: INiagaraComputeDataBufferInterface
 	virtual bool HasTranslucentDataToRender() const override { return TranslucentDataToRender != nullptr; }
-	virtual FNiagaraDataBuffer* GetDataToRender(bool bIsLowLatencyTranslucent) const override
+	virtual FNiagaraDataBuffer* GetDataToRender(FRHICommandListBase& RHICmdList, bool bIsLowLatencyTranslucent) const override
 	{
 		if (bIsLowLatencyTranslucent)
 		{
diff --git a/Engine/Plugins/FX/Niagara/Source/Niagara/Classes/NiagaraDataSet.h b/Engine/Plugins/FX/Niagara/Source/Niagara/Classes/NiagaraDataSet.h
--- a/Engine/Plugins/FX/Niagara/Source/Niagara/Classes/NiagaraDataSet.h
+++ b/Engine/Plugins/FX/Niagara/Source/Niagara/Classes/NiagaraDataSet.h
@@ -88,7 +88,7 @@ public:
 	NIAGARA_API void Allocate(uint32 NumInstances, bool bMaintainExisting = false);
 	NIAGARA_API void ReleaseCPU();
 
-	NIAGARA_API void AllocateGPU(FRHICommandList& RHICmdList, uint32 InNumInstances, ERHIFeatureLevel::Type FeatureLevel, const TCHAR* DebugSimName);
+	NIAGARA_API void AllocateGPU(FRHICommandListBase& RHICmdList, uint32 InNumInstances, ERHIFeatureLevel::Type FeatureLevel, const TCHAR* DebugSimName);
 	NIAGARA_API void SwapGPU(FNiagaraDataBuffer* BufferToSwap);
 	NIAGARA_API void ReleaseGPU();
 
diff --git a/Engine/Plugins/FX/Niagara/Source/Niagara/Classes/NiagaraGPUInstanceCountManager.h b/Engine/Plugins/FX/Niagara/Source/Niagara/Classes/NiagaraGPUInstanceCountManager.h
--- a/Engine/Plugins/FX/Niagara/Source/Niagara/Classes/NiagaraGPUInstanceCountManager.h
+++ b/Engine/Plugins/FX/Niagara/Source/Niagara/Classes/NiagaraGPUInstanceCountManager.h
@@ -68,12 +68,17 @@ public:
 		check(UsedInstanceCounts <= AllocatedInstanceCounts); // Can't resize after after the buffer gets bound.
 		return CountBuffer;
 	}
-
+	
 	/** Acquire an entry from the free list, assumes this comes from being presized. */
 	uint32 AcquireEntry();
 	/** Acquire an entry, this will either come from the free list or reallocate the buffer. */
 	uint32 AcquireOrAllocateEntry(FRHICommandListImmediate& RHICmdList);
 
+	/** Allocate an entry in the buffer, this could return an existing count or count yet to be allocated */
+	uint32 AllocateEntryDeferred();
+	/** Commit any deferred count allocations */
+	void AllocateDeferredCounts(FRHICommandListImmediate& RHICmdList);
+
 	/** Free the entry and reset it to INDEX_NONE if valid. */
 	void FreeEntry(uint32& BufferOffset);
 	/** Free and array of entries, you are expected to reset or change to INDEX_NONE. */
@@ -129,6 +134,8 @@ protected:
 
 	void ReleaseCounts();
 
+	UE::FMutex AddDrawIndirectGuard;
+
 	ERHIFeatureLevel::Type FeatureLevel;
 
 	/** The current used instance counts allocated from FNiagaraDataBuffer::AllocateGPU() */
@@ -137,7 +144,7 @@ protected:
 	int32 AllocatedInstanceCounts = 0;
 
 	/** The number of culled instance counts needed from view culling */
-	int32 RequiredCulledCounts = 0;
+	std::atomic<int32> RequiredCulledCounts = 0;
 	/** The allocated instance counts in the culled count buffer*/
 	int32 AllocatedCulledCounts = 0;
 	/** Whether or not the culled counts were acquired this frame */
@@ -173,4 +180,7 @@ protected:
 	/** Buffers holding drawindirect data to render GPU emitter renderers. */
 	TArray<FIndirectArgsPoolEntryPtr> DrawIndirectPool;
 	uint32 DrawIndirectLowWaterFrames = 0;
-};
\ No newline at end of file
+
+	UE::FMutex	DeferredCountAllocationGuard;
+	int			DeferredCountAllocations = 0;
+};
diff --git a/Engine/Plugins/FX/Niagara/Source/Niagara/Internal/Stateless/NiagaraStatelessEmitterInstance.h b/Engine/Plugins/FX/Niagara/Source/Niagara/Internal/Stateless/NiagaraStatelessEmitterInstance.h
--- a/Engine/Plugins/FX/Niagara/Source/Niagara/Internal/Stateless/NiagaraStatelessEmitterInstance.h
+++ b/Engine/Plugins/FX/Niagara/Source/Niagara/Internal/Stateless/NiagaraStatelessEmitterInstance.h
@@ -32,7 +32,7 @@ namespace NiagaraStateless
 
 		// Begin: INiagaraComputeDataBufferInterface
 		virtual bool HasTranslucentDataToRender() const override { return false; }
-		virtual FNiagaraDataBuffer* GetDataToRender(bool bIsLowLatencyTranslucent) const override;
+		virtual FNiagaraDataBuffer* GetDataToRender(FRHICommandListBase& RHICmdList, bool bIsLowLatencyTranslucent) const override;
 		// End: INiagaraComputeDataBufferInterface
 	};
 }
diff --git a/Engine/Plugins/FX/Niagara/Source/Niagara/Internal/Stateless/NiagaraStatelessParticleSimContext.h b/Engine/Plugins/FX/Niagara/Source/Niagara/Internal/Stateless/NiagaraStatelessParticleSimContext.h
--- a/Engine/Plugins/FX/Niagara/Source/Niagara/Internal/Stateless/NiagaraStatelessParticleSimContext.h
+++ b/Engine/Plugins/FX/Niagara/Source/Niagara/Internal/Stateless/NiagaraStatelessParticleSimContext.h
@@ -8,7 +8,7 @@
 class FNiagaraDataBuffer;
 struct FNiagaraStatelessEmitterData;
 struct FNiagaraStatelessRuntimeSpawnInfo;
-class FRHICommandList;
+class FRHICommandListBase;
 
 namespace NiagaraStateless
 {
@@ -32,7 +32,7 @@ public:
 	explicit FParticleSimulationContext(const FNiagaraStatelessEmitterData* EmitterData, TConstArrayView<uint8> DynamicBufferData);
 
 	void Simulate(int32 EmitterRandomSeed, float EmitterAge, float DeltaTime, TConstArrayView<FNiagaraStatelessRuntimeSpawnInfo> SpawnInfos, FNiagaraDataBuffer* DestinationData);
-	void SimulateGPU(FRHICommandList& RHICmdList, int32 EmitterRandomSeed, float EmitterAge, float DeltaTime, TConstArrayView<FNiagaraStatelessRuntimeSpawnInfo> SpawnInfos, FNiagaraDataBuffer* DestinationData);
+	void SimulateGPU(FRHICommandListBase& RHICmdList, int32 EmitterRandomSeed, float EmitterAge, float DeltaTime, TConstArrayView<FNiagaraStatelessRuntimeSpawnInfo> SpawnInfos, FNiagaraDataBuffer* DestinationData);
 
 private:
 	void SimulateInternal(int32 EmitterRandomSeed, float EmitterAge, float DeltaTime, FSpawnInfoShaderParameters& SpawnParameters, uint32 ActiveParticles);
diff --git a/Engine/Plugins/FX/Niagara/Source/Niagara/Private/NiagaraComponent.cpp b/Engine/Plugins/FX/Niagara/Source/Niagara/Private/NiagaraComponent.cpp
--- a/Engine/Plugins/FX/Niagara/Source/Niagara/Private/NiagaraComponent.cpp
+++ b/Engine/Plugins/FX/Niagara/Source/Niagara/Private/NiagaraComponent.cpp
@@ -45,6 +45,14 @@ DEFINE_LOG_CATEGORY(LogNiagara);
 
 DEFINE_RENDER_COMMAND_PIPE(NiagaraDynamicData, ERenderCommandPipeFlags::None);
 
+static bool GNiagaraParallelGDME = false;
+static FAutoConsoleVariableRef CVarNiagaraParallelGDME(
+	TEXT("fx.Niagara.ParallelGDME"),
+	GNiagaraParallelGDME,
+	TEXT("Allow Niagara to run parallel in GDME."),
+	ECVF_Default
+);
+
 static int GNiagaraSoloTickEarly = 1;
 static FAutoConsoleVariableRef CVarNiagaraSoloTickEarly(
 	TEXT("fx.Niagara.Solo.TickEarly"),
@@ -221,7 +229,7 @@ FNiagaraSceneProxy::FNiagaraSceneProxy(UNiagaraComponent* InComponent)
 	// Prevent continuous VSM invalidation from the bounds of the primitive.
 	bHasDeformableMesh = false;
 	// Niagara renderers reference a lot of common contexts that aren't locked (and would otherwise introduce a lot of contention).
-	bSupportsParallelGDME = false;
+	bSupportsParallelGDME = GNiagaraParallelGDME;
 }
 
 SIZE_T FNiagaraSceneProxy::GetTypeHash() const
diff --git a/Engine/Plugins/FX/Niagara/Source/Niagara/Private/NiagaraDataSet.cpp b/Engine/Plugins/FX/Niagara/Source/Niagara/Private/NiagaraDataSet.cpp
--- a/Engine/Plugins/FX/Niagara/Source/Niagara/Private/NiagaraDataSet.cpp
+++ b/Engine/Plugins/FX/Niagara/Source/Niagara/Private/NiagaraDataSet.cpp
@@ -796,7 +796,7 @@ void FNiagaraDataBuffer::ReleaseCPU()
 	HalfData.Empty();
 }
 
-void FNiagaraDataBuffer::AllocateGPU(FRHICommandList& RHICmdList, uint32 InNumInstances, ERHIFeatureLevel::Type FeatureLevel, const TCHAR* DebugSimName)
+void FNiagaraDataBuffer::AllocateGPU(FRHICommandListBase& RHICmdList, uint32 InNumInstances, ERHIFeatureLevel::Type FeatureLevel, const TCHAR* DebugSimName)
 {
 	static constexpr EBufferUsageFlags GPUBufferFlags = BUF_Static | BUF_SourceCopy;
 
diff --git a/Engine/Plugins/FX/Niagara/Source/Niagara/Private/NiagaraGPUInstanceCountManager.cpp b/Engine/Plugins/FX/Niagara/Source/Niagara/Private/NiagaraGPUInstanceCountManager.cpp
--- a/Engine/Plugins/FX/Niagara/Source/Niagara/Private/NiagaraGPUInstanceCountManager.cpp
+++ b/Engine/Plugins/FX/Niagara/Source/Niagara/Private/NiagaraGPUInstanceCountManager.cpp
@@ -148,7 +148,7 @@ void FNiagaraGPUInstanceCountManager::ReleaseCounts()
 
 uint32 FNiagaraGPUInstanceCountManager::AcquireEntry()
 {
-	checkSlow(IsInRenderingThread());
+	check(IsInRenderingThread());
 
 	if (FreeEntries.Num())
 	{
@@ -170,7 +170,7 @@ uint32 FNiagaraGPUInstanceCountManager::AcquireEntry()
 
 uint32 FNiagaraGPUInstanceCountManager::AcquireOrAllocateEntry(FRHICommandListImmediate& RHICmdList)
 {
-	checkSlow(IsInRenderingThread());
+	check(IsInRenderingThread());
 
 	// Free entries?
 	if (FreeEntries.Num())
@@ -189,9 +189,37 @@ uint32 FNiagaraGPUInstanceCountManager::AcquireOrAllocateEntry(FRHICommandListIm
 	return UsedInstanceCounts++;
 }
 
+uint32 FNiagaraGPUInstanceCountManager::AllocateEntryDeferred()
+{
+	check(IsInParallelRenderingThread());
+
+	UE::TScopeLock LockGuard(DeferredCountAllocationGuard);
+
+	if (FreeEntries.Num())
+	{
+		return FreeEntries.Pop();
+	}
+	else if (UsedInstanceCounts < AllocatedInstanceCounts)
+	{
+		return UsedInstanceCounts++;
+	}
+
+	return AllocatedInstanceCounts + DeferredCountAllocations++;
+}
+
+void FNiagaraGPUInstanceCountManager::AllocateDeferredCounts(FRHICommandListImmediate& RHICmdList)
+{
+	if (DeferredCountAllocations == 0 )
+	{
+		return;
+	}
+	ResizeBuffers(RHICmdList, AllocatedInstanceCounts + DeferredCountAllocations);
+	DeferredCountAllocations = 0;
+}
+
 void FNiagaraGPUInstanceCountManager::FreeEntry(uint32& BufferOffset)
 {
-	checkSlow(IsInRenderingThread());
+	check(IsInRenderingThread());
 
 	if (BufferOffset != INDEX_NONE)
 	{
@@ -205,7 +233,7 @@ void FNiagaraGPUInstanceCountManager::FreeEntry(uint32& BufferOffset)
 
 void FNiagaraGPUInstanceCountManager::FreeEntryArray(TConstArrayView<uint32> EntryArray)
 {
-	checkSlow(IsInRenderingThread());
+	check(IsInRenderingThread());
 
 	const int32 NumToFree = EntryArray.Num();
 	if (NumToFree > 0)
@@ -223,6 +251,8 @@ void FNiagaraGPUInstanceCountManager::FreeEntryArray(TConstArrayView<uint32> Ent
 
 FRWBuffer* FNiagaraGPUInstanceCountManager::AcquireCulledCountsBuffer(FRHICommandListImmediate& RHICmdList)
 {
+	check(IsInRenderingThread());
+
 	if (RequiredCulledCounts > 0)
 	{
 		if (!bAcquiredCulledCounts)
@@ -317,6 +347,8 @@ void FNiagaraGPUInstanceCountManager::ResizeBuffers(FRHICommandListImmediate& RH
 
 void FNiagaraGPUInstanceCountManager::FlushIndirectArgsPool(FRHICommandListBase& RHICmdList)
 {
+	check(IsInRenderingThread());
+
 	// Cull indirect draw pool entries so that we only keep the last pool
 	while (DrawIndirectPool.Num() > 1)
 	{
@@ -348,7 +380,7 @@ void FNiagaraGPUInstanceCountManager::FlushIndirectArgsPool(FRHICommandListBase&
 
 FNiagaraGPUInstanceCountManager::FIndirectArgSlot FNiagaraGPUInstanceCountManager::AddDrawIndirect(FRHICommandListBase& RHICmdList, uint32 InstanceCountBufferOffset, uint32 NumIndicesPerInstance, uint32 StartIndexLocation, bool bIsInstancedStereoEnabled, bool bCulled, ENiagaraGpuComputeTickStage::Type ReadyTickStage)
 {
-
+	UE::TScopeLock Lock(AddDrawIndirectGuard);
 
 	const ENiagaraDrawIndirectArgGenTaskFlags TaskFlags =
 		(bIsInstancedStereoEnabled ? ENiagaraDrawIndirectArgGenTaskFlags::InstancedStereo : ENiagaraDrawIndirectArgGenTaskFlags::None)
@@ -395,6 +427,8 @@ FNiagaraGPUInstanceCountManager::FIndirectArgSlot FNiagaraGPUInstanceCountManage
 
 void FNiagaraGPUInstanceCountManager::UpdateDrawIndirectBuffers(FNiagaraGpuComputeDispatchInterface* ComputeDispatchInterface, FRHICommandList& RHICmdList, ENiagaraGPUCountUpdatePhase::Type CountPhase)
 {
+	check(IsInRenderingThread());
+
 	// Anything to process?
 	TArray<FNiagaraDrawIndirectArgGenTaskInfo>& ArgTasks = DrawIndirectArgGenTasks[CountPhase];
 	const bool bClearCounts = (CountPhase == ENiagaraGPUCountUpdatePhase::PreOpaque) && (InstanceCountClearTasks.Num() > 0) && ComputeDispatchInterface->IsFirstViewFamily();
@@ -613,6 +647,8 @@ void FNiagaraGPUInstanceCountManager::UpdateDrawIndirectBuffers(FNiagaraGpuCompu
 
 const uint32* FNiagaraGPUInstanceCountManager::GetGPUReadback()
 {
+	check(IsInRenderingThread());
+
 	if (CountReadback && CountReadbackSize && CountReadback->IsReady())
 	{
 		SCOPE_CYCLE_COUNTER(STAT_NiagaraGPUReadbackLock);
@@ -626,6 +662,7 @@ const uint32* FNiagaraGPUInstanceCountManager::GetGPUReadback()
 
 void FNiagaraGPUInstanceCountManager::ReleaseGPUReadback()
 {
+	check(IsInRenderingThread());
 	check(CountReadback && CountReadbackSize);
 	CountReadback->Unlock();
 	// Readback can only ever be done once, to prevent misusage with index lifetime
@@ -634,6 +671,7 @@ void FNiagaraGPUInstanceCountManager::ReleaseGPUReadback()
 
 void FNiagaraGPUInstanceCountManager::EnqueueGPUReadback(FRHICommandListImmediate& RHICmdList)
 {
+	check(IsInRenderingThread());
 	if (UsedInstanceCounts > 0 && (UsedInstanceCounts != FreeEntries.Num()))
 	{
 		if (!CountReadback)
@@ -649,11 +687,13 @@ void FNiagaraGPUInstanceCountManager::EnqueueGPUReadback(FRHICommandListImmediat
 
 bool FNiagaraGPUInstanceCountManager::HasPendingGPUReadback() const
 {
+	check(IsInRenderingThread());
 	return CountReadback && CountReadbackSize;
 }
 
 void FNiagaraGPUInstanceCountManager::CopyToMultiViewCountBuffer(FRHICommandListImmediate& RHICmdList)
 {
+	check(IsInRenderingThread());
 	if (AllocatedInstanceCounts > 0)
 	{
 		// Need to copy on all GPUs
diff --git a/Engine/Plugins/FX/Niagara/Source/Niagara/Private/NiagaraGpuComputeDispatch.cpp b/Engine/Plugins/FX/Niagara/Source/Niagara/Private/NiagaraGpuComputeDispatch.cpp
--- a/Engine/Plugins/FX/Niagara/Source/Niagara/Private/NiagaraGpuComputeDispatch.cpp
+++ b/Engine/Plugins/FX/Niagara/Source/Niagara/Private/NiagaraGpuComputeDispatch.cpp
@@ -1066,7 +1066,7 @@ void FNiagaraGpuComputeDispatch::ExecuteTicks(FRDGBuilder& GraphBuilder, TConstS
 						ComputeContext->SetMultiViewPreviousDataToRender(nullptr);
 
 						// Mark data as ready for anyone who picks up the buffer on the next frame (see FNiagaraGpuComputeDispatch::ExecuteTicks)
-						ComputeContext->GetDataToRender(false)->SetGPUDataReadyStage(ENiagaraGpuComputeTickStage::First);
+						ComputeContext->GetDataToRender(RHICmdList, false)->SetGPUDataReadyStage(ENiagaraGpuComputeTickStage::First);
 
 						// Clear instance count offsets (see FNiagaraSystemGpuComputeProxy::ReleaseTicks)
 						for (int i = 0; i < UE_ARRAY_COUNT(ComputeContext->DataBuffers_RT); ++i)
@@ -1965,6 +1965,8 @@ void FNiagaraGpuComputeDispatch::PostInitViews(FRDGBuilder& GraphBuilder, TConst
 {
 	LLM_SCOPE(ELLMTag::Niagara);
 
+	GPUInstanceCounterManager.AllocateDeferredCounts(GraphBuilder.RHICmdList);
+
 	bAllowGPUParticleUpdate = bAllowGPUParticleUpdate && Views.Num() > 0 && Views[0].AllowGPUParticleUpdate();
 
 	if (bAllowGPUParticleUpdate && FNiagaraUtilities::AllowGPUParticles(GetShaderPlatform()))
diff --git a/Engine/Plugins/FX/Niagara/Source/Niagara/Private/NiagaraRenderer.cpp b/Engine/Plugins/FX/Niagara/Source/Niagara/Private/NiagaraRenderer.cpp
--- a/Engine/Plugins/FX/Niagara/Source/Niagara/Private/NiagaraRenderer.cpp
+++ b/Engine/Plugins/FX/Niagara/Source/Niagara/Private/NiagaraRenderer.cpp
@@ -293,9 +293,9 @@ bool FNiagaraDynamicDataBase::IsGpuLowLatencyTranslucencyEnabled() const
 	return ComputeDataBufferInterface ? ComputeDataBufferInterface->HasTranslucentDataToRender() : false;
 }
 
-FNiagaraDataBuffer* FNiagaraDynamicDataBase::GetParticleDataToRender(bool bIsLowLatencyTranslucent) const
+FNiagaraDataBuffer* FNiagaraDynamicDataBase::GetParticleDataToRender(FRHICommandListBase& RHICmdList, bool bIsLowLatencyTranslucent) const
 {
-	FNiagaraDataBuffer* Ret = ComputeDataBufferInterface ? ComputeDataBufferInterface->GetDataToRender(bIsLowLatencyTranslucent) : CPUParticleData.GetReference();
+	FNiagaraDataBuffer* Ret = ComputeDataBufferInterface ? ComputeDataBufferInterface->GetDataToRender(RHICmdList, bIsLowLatencyTranslucent) : CPUParticleData.GetReference();
 	checkSlow(Ret == nullptr || Ret->IsBeingRead());
 	return Ret;
 }
diff --git a/Engine/Plugins/FX/Niagara/Source/Niagara/Private/NiagaraRendererMeshes.cpp b/Engine/Plugins/FX/Niagara/Source/Niagara/Private/NiagaraRendererMeshes.cpp
--- a/Engine/Plugins/FX/Niagara/Source/Niagara/Private/NiagaraRendererMeshes.cpp
+++ b/Engine/Plugins/FX/Niagara/Source/Niagara/Private/NiagaraRendererMeshes.cpp
@@ -312,7 +312,7 @@ void FNiagaraRendererMeshes::ReleaseRenderThreadResources()
 {
 }
 
-void FNiagaraRendererMeshes::PrepareParticleMeshRenderData(FParticleMeshRenderData& ParticleMeshRenderData, const FSceneViewFamily& ViewFamily, FMeshElementCollector& Collector, FNiagaraDynamicDataBase* InDynamicData, const FNiagaraSceneProxy* SceneProxy, bool bRayTracing, ENiagaraGpuComputeTickStage::Type GpuReadyTickStage) const
+void FNiagaraRendererMeshes::PrepareParticleMeshRenderData(FRHICommandListBase& RHICmdList, FParticleMeshRenderData& ParticleMeshRenderData, const FSceneViewFamily& ViewFamily, FMeshElementCollector& Collector, FNiagaraDynamicDataBase* InDynamicData, const FNiagaraSceneProxy* SceneProxy, bool bRayTracing, ENiagaraGpuComputeTickStage::Type GpuReadyTickStage) const
 {
 	ParticleMeshRenderData.Collector = &Collector;
 
@@ -331,7 +331,7 @@ void FNiagaraRendererMeshes::PrepareParticleMeshRenderData(FParticleMeshRenderDa
 	}
 
 	// Early out if we have no data or instances, this must be done before we read the material
-	FNiagaraDataBuffer* CurrentParticleData = ParticleMeshRenderData.DynamicDataMesh->GetParticleDataToRender(bGpuLowLatencyTranslucency);
+	FNiagaraDataBuffer* CurrentParticleData = ParticleMeshRenderData.DynamicDataMesh->GetParticleDataToRender(RHICmdList, bGpuLowLatencyTranslucency);
 	if (!CurrentParticleData || (SourceMode == ENiagaraRendererSourceDataMode::Particles && CurrentParticleData->GetNumInstances() == 0) || !Meshes.Num())
 	{
 		return;
@@ -369,7 +369,7 @@ void FNiagaraRendererMeshes::PrepareParticleMeshRenderData(FParticleMeshRenderDa
 			!UE::FXRenderingUtils::CanMaterialRenderBeforeFXPostOpaque(ViewFamily, *SceneProxy, Material);
 	}
 	
-	ParticleMeshRenderData.SourceParticleData = ParticleMeshRenderData.DynamicDataMesh->GetParticleDataToRender(ParticleMeshRenderData.bIsGpuLowLatencyTranslucency);
+	ParticleMeshRenderData.SourceParticleData = ParticleMeshRenderData.DynamicDataMesh->GetParticleDataToRender(RHICmdList, ParticleMeshRenderData.bIsGpuLowLatencyTranslucency);
 	
 	// Anything to render?
 	if ((ParticleMeshRenderData.SourceParticleData == nullptr) ||
@@ -475,7 +475,7 @@ void FNiagaraRendererMeshes::PrepareParticleMeshRenderData(FParticleMeshRenderDa
 	}
 }
 
-bool FNiagaraRendererMeshes::CalculateMeshUsed(FParticleMeshRenderData& ParticleMeshRenderData) const
+bool FNiagaraRendererMeshes::CalculateMeshUsed(FRHICommandListBase& RHICmdList, FParticleMeshRenderData& ParticleMeshRenderData) const
 {
 	//-OPT: Should we handle modes where the tags are inside none-particle data?  i.e. EmitterRendererVisTagOffset | EmitterMeshIndexOffset
 	if (ParticleRendererVisTagOffset == INDEX_NONE && ParticleMeshIndexOffset == INDEX_NONE)
@@ -500,7 +500,7 @@ bool FNiagaraRendererMeshes::CalculateMeshUsed(FParticleMeshRenderData& Particle
 		return true;
 	}
 
-	const FNiagaraDataBuffer* DataToRender = ParticleMeshRenderData.DynamicDataMesh->GetParticleDataToRender();
+	const FNiagaraDataBuffer* DataToRender = ParticleMeshRenderData.DynamicDataMesh->GetParticleDataToRender(RHICmdList);
 	if (ParticleRendererVisTagOffset != INDEX_NONE)
 	{
 		const int32* RendererVisValues = reinterpret_cast<const int32*>(DataToRender->GetComponentPtrInt32(ParticleRendererVisTagOffset));
@@ -1388,7 +1388,7 @@ void FNiagaraRendererMeshes::GetDynamicMeshElements(const TArray<const FSceneVie
 	// This will also determine if we have anything to render
 	// ENiagaraGpuComputeTickStage::Last is used as the GPU ready stage as we can support reading translucent data after PostRenderOpaque sims have run
 	FParticleMeshRenderData ParticleMeshRenderData;
-	PrepareParticleMeshRenderData(ParticleMeshRenderData, ViewFamily, Collector, DynamicDataRender, SceneProxy, false, ENiagaraGpuComputeTickStage::Last);
+	PrepareParticleMeshRenderData(Collector.GetRHICommandList(), ParticleMeshRenderData, ViewFamily, Collector, DynamicDataRender, SceneProxy, false, ENiagaraGpuComputeTickStage::Last);
 
 	if (ParticleMeshRenderData.SourceParticleData == nullptr || GbEnableNiagaraMeshRendering == 0)
 	{
@@ -1404,7 +1404,7 @@ void FNiagaraRendererMeshes::GetDynamicMeshElements(const TArray<const FSceneVie
 	FScopeCycleCounter EmitterStatsCounter(EmitterStatID);
 #endif
 
-	if (!CalculateMeshUsed(ParticleMeshRenderData))
+	if (!CalculateMeshUsed(Collector.GetRHICommandList(), ParticleMeshRenderData))
 	{
 		return;
 	}
@@ -1639,7 +1639,7 @@ void FNiagaraRendererMeshes::GetDynamicRayTracingInstances(FRayTracingMaterialGa
 	// This will also determine if we have anything to render
 	// ENiagaraGpuComputeTickStage::PostInitViews is used as we need the data one InitViews is complete as the HWRT BVH will be generated before other sims have run
 	FParticleMeshRenderData ParticleMeshRenderData;
-	PrepareParticleMeshRenderData(ParticleMeshRenderData, *View->Family, Context.RayTracingMeshResourceCollector, DynamicDataRender, SceneProxy, true, ENiagaraGpuComputeTickStage::PostInitViews);
+	PrepareParticleMeshRenderData(Context.RHICmdList, ParticleMeshRenderData, *View->Family, Context.RayTracingMeshResourceCollector, DynamicDataRender, SceneProxy, true, ENiagaraGpuComputeTickStage::PostInitViews);
 
 	if (ParticleMeshRenderData.SourceParticleData == nullptr || Meshes.Num() == 0)
 	{
diff --git a/Engine/Plugins/FX/Niagara/Source/Niagara/Private/NiagaraRendererRibbons.cpp b/Engine/Plugins/FX/Niagara/Source/Niagara/Private/NiagaraRendererRibbons.cpp
--- a/Engine/Plugins/FX/Niagara/Source/Niagara/Private/NiagaraRendererRibbons.cpp
+++ b/Engine/Plugins/FX/Niagara/Source/Niagara/Private/NiagaraRendererRibbons.cpp
@@ -571,6 +571,8 @@ public:
 	{
 		if (InRenderingResources->ViewResources.Num() > 0)
 		{
+			UE::TScopeLock ScopeLock(AllocateGuard);
+
 			const int32 GenerateIndex = InSourceParticleData->GetGPUDataReadyStage() == ENiagaraGpuComputeTickStage::PostOpaqueRender ? 1 : 0;
 			RenderersToGeneratePerStage[GenerateIndex].Emplace(InRenderer, InSourceParticleData, InRenderingResources);
 		}
@@ -578,6 +580,8 @@ public:
 
 	FNiagaraRibbonRenderingFrameViewResources* FindExistingRendererViewData(const FNiagaraRendererRibbons* InRenderer, const FNiagaraDataBuffer* InSourceParticleData)
 	{
+		UE::TScopeLock ScopeLock(AllocateGuard);
+
 		const int32 GenerateIndex = InSourceParticleData->GetGPUDataReadyStage() == ENiagaraGpuComputeTickStage::PostOpaqueRender ? 1 : 0;
 		for (FNiagaraRibbonGPUInitParameters& ExistingData : RenderersToGeneratePerStage[GenerateIndex])
 		{
@@ -593,6 +597,8 @@ public:
 	//-OPT: These caches should be more central and are as a simple solution to reduce memory thrashing / poor performance for ribbons
 	const TSharedPtr<FRWBuffer>& GetOrAllocateIndirectDrawBuffer(FRHICommandListBase& RHICmdList)
 	{
+		UE::TScopeLock ScopeLock(AllocateGuard);
+
 		FIndirectDrawBufferEntry* BufferEntry = IndirectDrawBufferCache.FindByPredicate(
 			[&](const FIndirectDrawBufferEntry& ExistingBuffer)
 			{
@@ -611,6 +617,8 @@ public:
 
 	const TSharedPtr<FNiagaraRibbonIndexBuffer>& GetOrAllocateIndexBuffer(FRHICommandListBase& RHICmdList, int32 NumIndices, int32 MaxIndicesEstimate)
 	{
+		UE::TScopeLock ScopeLock(AllocateGuard);
+
 		if (GNiagaraRibbonGpuBufferCachePurgeCounter >= 0)
 		{
 			NumIndices = GNiagaraRibbonGpuAllocateMaxCount == 0 ? Align(NumIndices, GNiagaraRibbonGpuBufferAlign) : MaxIndicesEstimate;
@@ -643,6 +651,8 @@ private:
 
 	uint64									FrameCounter = 0;
 
+	UE::FMutex								AllocateGuard;
+
 	void OnPostPreRender(FRHICommandListImmediate& RHICmdList)
 	{
 		if (GNiagaraRibbonGpuBufferCachePurgeCounter < 0)
@@ -813,7 +823,7 @@ void FNiagaraRendererRibbons::GetDynamicMeshElements(const TArray<const FSceneVi
 		return;
 	}
 
-	FNiagaraDataBuffer* SourceParticleData = DynamicData->GetParticleDataToRender();
+	FNiagaraDataBuffer* SourceParticleData = DynamicData->GetParticleDataToRender(Collector.GetRHICommandList());
 
 	if (GbEnableNiagaraRibbonRendering == 0 || SourceParticleData == nullptr)
 	{
@@ -1068,7 +1078,7 @@ void FNiagaraRendererRibbons::GetDynamicRayTracingInstances(FRayTracingMaterialG
 		return;
 	}
 
-	FNiagaraDataBuffer* SourceParticleData = DynamicDataRibbon->GetParticleDataToRender();
+	FNiagaraDataBuffer* SourceParticleData = DynamicDataRibbon->GetParticleDataToRender(Context.RHICmdList);
 
 	if (GbEnableNiagaraRibbonRendering == 0 || SourceParticleData == nullptr)
 	{
diff --git a/Engine/Plugins/FX/Niagara/Source/Niagara/Private/NiagaraRendererSprites.cpp b/Engine/Plugins/FX/Niagara/Source/Niagara/Private/NiagaraRendererSprites.cpp
--- a/Engine/Plugins/FX/Niagara/Source/Niagara/Private/NiagaraRendererSprites.cpp
+++ b/Engine/Plugins/FX/Niagara/Source/Niagara/Private/NiagaraRendererSprites.cpp
@@ -224,7 +224,7 @@ void FNiagaraRendererSprites::CreateRenderThreadResources(FRHICommandListBase& R
 #endif
 }
 
-void FNiagaraRendererSprites::PrepareParticleSpriteRenderData(FParticleSpriteRenderData& ParticleSpriteRenderData, const FSceneViewFamily& ViewFamily, FNiagaraDynamicDataBase* InDynamicData, const FNiagaraSceneProxy* SceneProxy, ENiagaraGpuComputeTickStage::Type GpuReadyTickStage) const
+void FNiagaraRendererSprites::PrepareParticleSpriteRenderData(FRHICommandListBase& RHICmdList, FParticleSpriteRenderData& ParticleSpriteRenderData, const FSceneViewFamily& ViewFamily, FNiagaraDynamicDataBase* InDynamicData, const FNiagaraSceneProxy* SceneProxy, ENiagaraGpuComputeTickStage::Type GpuReadyTickStage) const
 {
 	ParticleSpriteRenderData.DynamicDataSprites = static_cast<FNiagaraDynamicDataSprites*>(InDynamicData);
 	if (!ParticleSpriteRenderData.DynamicDataSprites || !SceneProxy->GetComputeDispatchInterface())
@@ -234,7 +234,7 @@ void FNiagaraRendererSprites::PrepareParticleSpriteRenderData(FParticleSpriteRen
 	}
 
 	// Early out if we have no data or instances, this must be done before we read the material
-	FNiagaraDataBuffer* CurrentParticleData = ParticleSpriteRenderData.DynamicDataSprites->GetParticleDataToRender(bGpuLowLatencyTranslucency);
+	FNiagaraDataBuffer* CurrentParticleData = ParticleSpriteRenderData.DynamicDataSprites->GetParticleDataToRender(RHICmdList, bGpuLowLatencyTranslucency);
 	if (!CurrentParticleData || (SourceMode == ENiagaraRendererSourceDataMode::Particles && CurrentParticleData->GetNumInstances() == 0) || (GbEnableNiagaraSpriteRendering == 0))
 	{
 		return;
@@ -263,7 +263,7 @@ void FNiagaraRendererSprites::PrepareParticleSpriteRenderData(FParticleSpriteRen
 	}
 
 
-	ParticleSpriteRenderData.SourceParticleData = ParticleSpriteRenderData.DynamicDataSprites->GetParticleDataToRender(bLowLatencyTranslucencyEnabled);
+	ParticleSpriteRenderData.SourceParticleData = ParticleSpriteRenderData.DynamicDataSprites->GetParticleDataToRender(RHICmdList, bLowLatencyTranslucencyEnabled);
 	if ( !ParticleSpriteRenderData.SourceParticleData || (SourceMode == ENiagaraRendererSourceDataMode::Particles && ParticleSpriteRenderData.SourceParticleData->GetNumInstances() == 0) )
 	{
 		ParticleSpriteRenderData.SourceParticleData = nullptr;
@@ -959,7 +959,7 @@ void FNiagaraRendererSprites::GetDynamicMeshElements(const TArray<const FSceneVi
 	// This will also determine if we have anything to render
 	// ENiagaraGpuComputeTickStage::Last is used as the GPU ready stage as we can support reading translucent data after PostRenderOpaque sims have run
 	FParticleSpriteRenderData ParticleSpriteRenderData;
-	PrepareParticleSpriteRenderData(ParticleSpriteRenderData, ViewFamily, DynamicDataRender, SceneProxy, ENiagaraGpuComputeTickStage::Last);
+	PrepareParticleSpriteRenderData(Collector.GetRHICommandList(), ParticleSpriteRenderData, ViewFamily, DynamicDataRender, SceneProxy, ENiagaraGpuComputeTickStage::Last);
 
 	if (ParticleSpriteRenderData.SourceParticleData == nullptr)
 	{
@@ -1084,7 +1084,7 @@ void FNiagaraRendererSprites::GetDynamicRayTracingInstances(FRayTracingMaterialG
 	// This will also determine if we have anything to render
 	// ENiagaraGpuComputeTickStage::PostInitViews is used as we need the data one InitViews is complete as the HWRT BVH will be generated before other sims have run
 	FParticleSpriteRenderData ParticleSpriteRenderData;
-	PrepareParticleSpriteRenderData(ParticleSpriteRenderData, *Context.ReferenceView->Family, DynamicDataRender, SceneProxy, ENiagaraGpuComputeTickStage::PostInitViews);
+	PrepareParticleSpriteRenderData(Context.RHICmdList, ParticleSpriteRenderData, *Context.ReferenceView->Family, DynamicDataRender, SceneProxy, ENiagaraGpuComputeTickStage::PostInitViews);
 
 	if (ParticleSpriteRenderData.SourceParticleData == nullptr)
 	{
diff --git a/Engine/Plugins/FX/Niagara/Source/Niagara/Private/NiagaraRendererVolumes.cpp b/Engine/Plugins/FX/Niagara/Source/Niagara/Private/NiagaraRendererVolumes.cpp
--- a/Engine/Plugins/FX/Niagara/Source/Niagara/Private/NiagaraRendererVolumes.cpp
+++ b/Engine/Plugins/FX/Niagara/Source/Niagara/Private/NiagaraRendererVolumes.cpp
@@ -246,7 +246,7 @@ void FNiagaraRendererVolumes::GetDynamicMeshElements(const TArray<const FSceneVi
 		return;
 	}
 
-	FNiagaraDataBuffer* ParticleDataBuffer = VolumeDynamicData->GetParticleDataToRender();
+	FNiagaraDataBuffer* ParticleDataBuffer = VolumeDynamicData->GetParticleDataToRender(Collector.GetRHICommandList());
 	if (!ParticleDataBuffer || (SourceMode == ENiagaraRendererSourceDataMode::Particles && ParticleDataBuffer->GetNumInstances() == 0))
 	{
 		return;
diff --git a/Engine/Plugins/FX/Niagara/Source/Niagara/Private/Stateless/NiagaraStatelessComputeManager.cpp b/Engine/Plugins/FX/Niagara/Source/Niagara/Private/Stateless/NiagaraStatelessComputeManager.cpp
--- a/Engine/Plugins/FX/Niagara/Source/Niagara/Private/Stateless/NiagaraStatelessComputeManager.cpp
+++ b/Engine/Plugins/FX/Niagara/Source/Niagara/Private/Stateless/NiagaraStatelessComputeManager.cpp
@@ -153,10 +153,13 @@ FNiagaraStatelessComputeManager::~FNiagaraStatelessComputeManager()
 {
 }
 
-FNiagaraDataBuffer* FNiagaraStatelessComputeManager::GetDataBuffer(uintptr_t EmitterKey, const NiagaraStateless::FEmitterInstance_RT* EmitterInstance)
+FNiagaraDataBuffer* FNiagaraStatelessComputeManager::GetDataBuffer(FRHICommandListBase& RHICmdList, uintptr_t EmitterKey, const NiagaraStateless::FEmitterInstance_RT* EmitterInstance)
 {
 	using namespace NiagaraStateless;
 	using namespace NiagaraStatelessComputeManagerPrivate;
+	
+	//-OPT: This lock is very conservative, ideally we only have it around the relevant parts
+	UE::TScopeLock ScopeLock(GetDataBufferGuard);
 
 	if (TUniquePtr<FStatelessDataCache>* ExistingData = UsedData.Find(EmitterKey))
 	{
@@ -219,7 +222,6 @@ FNiagaraDataBuffer* FNiagaraStatelessComputeManager::GetDataBuffer(uintptr_t Emi
 	CacheData->EmitterInstance = EmitterInstance;
 	CacheData->ActiveParticles = ActiveParticles;
 
-	FRHICommandListImmediate& RHICmdList = FRHICommandListExecutor::GetImmediateCommandList();
 	CacheData->DataBuffer->AllocateGPU(RHICmdList, CacheData->ActiveParticles, ComputeInterface->GetFeatureLevel(), TEXT("StatelessSimBuffer"));
 
 	const EComputeExecutionPath ComputeExecutionPath = DetermineComputeExecutionPath(EmitterData, ActiveParticles);
@@ -242,7 +244,7 @@ FNiagaraDataBuffer* FNiagaraStatelessComputeManager::GetDataBuffer(uintptr_t Emi
 			CacheData->DataBuffer->SetNumInstances(CacheData->ActiveParticles);
 
 			FNiagaraGPUInstanceCountManager& CountManager = ComputeInterface->GetGPUInstanceCounterManager();
-			const uint32 CountOffset = CountManager.AcquireOrAllocateEntry(RHICmdList);
+			const uint32 CountOffset = CountManager.AllocateEntryDeferred();
 			CacheData->DataBuffer->SetGPUInstanceCountBufferOffset(CountOffset);
 			CountsToRelease.Add(CountOffset);
 
diff --git a/Engine/Plugins/FX/Niagara/Source/Niagara/Private/Stateless/NiagaraStatelessComputeManager.h b/Engine/Plugins/FX/Niagara/Source/Niagara/Private/Stateless/NiagaraStatelessComputeManager.h
--- a/Engine/Plugins/FX/Niagara/Source/Niagara/Private/Stateless/NiagaraStatelessComputeManager.h
+++ b/Engine/Plugins/FX/Niagara/Source/Niagara/Private/Stateless/NiagaraStatelessComputeManager.h
@@ -39,7 +39,7 @@ public:
 		return ManagerName;
 	}
 
-	FNiagaraDataBuffer* GetDataBuffer(uintptr_t EmitterKey, const NiagaraStateless::FEmitterInstance_RT* EmitterInstance);
+	FNiagaraDataBuffer* GetDataBuffer(FRHICommandListBase& RHICmdList, uintptr_t EmitterKey, const NiagaraStateless::FEmitterInstance_RT* EmitterInstance);
 
 	// Used to execute the simulation immediately into a CPU side data buffer
 	void GenerateDataBufferForDebugging(FRHICommandListImmediate& RHICmdList, FNiagaraDataBuffer* DataBuffer, const NiagaraStateless::FEmitterInstance_RT* EmitterInstance) const;
@@ -56,4 +56,6 @@ private:
 	TArray<uint32>										CountsToRelease;
 
 	TArray<FStatelessDataCache*>						GPUDataToGenerate;
+
+	UE::FMutex											GetDataBufferGuard;
 };
diff --git a/Engine/Plugins/FX/Niagara/Source/Niagara/Private/Stateless/NiagaraStatelessEmitterInstance.cpp b/Engine/Plugins/FX/Niagara/Source/Niagara/Private/Stateless/NiagaraStatelessEmitterInstance.cpp
--- a/Engine/Plugins/FX/Niagara/Source/Niagara/Private/Stateless/NiagaraStatelessEmitterInstance.cpp
+++ b/Engine/Plugins/FX/Niagara/Source/Niagara/Private/Stateless/NiagaraStatelessEmitterInstance.cpp
@@ -10,9 +10,9 @@
 
 #include "Shader.h"
 
-FNiagaraDataBuffer* NiagaraStateless::FEmitterInstance_RT::GetDataToRender(bool bIsLowLatencyTranslucent) const
+FNiagaraDataBuffer* NiagaraStateless::FEmitterInstance_RT::GetDataToRender(FRHICommandListBase& RHICmdList, bool bIsLowLatencyTranslucent) const
 {
-	return ComputeManager ? ComputeManager->GetDataBuffer(uintptr_t(this), this) : nullptr;
+	return ComputeManager ? ComputeManager->GetDataBuffer(RHICmdList, uintptr_t(this), this) : nullptr;
 }
 
 FNiagaraStatelessEmitterInstance::FNiagaraStatelessEmitterInstance(FNiagaraSystemInstance* InParentSystemInstance)
diff --git a/Engine/Plugins/FX/Niagara/Source/Niagara/Private/Stateless/NiagaraStatelessParticleSimContext.cpp b/Engine/Plugins/FX/Niagara/Source/Niagara/Private/Stateless/NiagaraStatelessParticleSimContext.cpp
--- a/Engine/Plugins/FX/Niagara/Source/Niagara/Private/Stateless/NiagaraStatelessParticleSimContext.cpp
+++ b/Engine/Plugins/FX/Niagara/Source/Niagara/Private/Stateless/NiagaraStatelessParticleSimContext.cpp
@@ -80,7 +80,7 @@ namespace NiagaraStateless
 		DestinationData->SetNumInstances(NumInstances);
 	}
 
-	void FParticleSimulationContext::SimulateGPU(FRHICommandList& RHICmdList, int32 InEmitterRandomSeed, float EmitterAge, float InDeltaTime, TConstArrayView<FNiagaraStatelessRuntimeSpawnInfo> SpawnInfos, FNiagaraDataBuffer* DestinationData)
+	void FParticleSimulationContext::SimulateGPU(FRHICommandListBase& RHICmdList, int32 InEmitterRandomSeed, float EmitterAge, float InDeltaTime, TConstArrayView<FNiagaraStatelessRuntimeSpawnInfo> SpawnInfos, FNiagaraDataBuffer* DestinationData)
 	{
 		NumInstances = 0;
 
diff --git a/Engine/Plugins/FX/Niagara/Source/Niagara/Public/NiagaraGpuComputeDispatchInterface.h b/Engine/Plugins/FX/Niagara/Source/Niagara/Public/NiagaraGpuComputeDispatchInterface.h
--- a/Engine/Plugins/FX/Niagara/Source/Niagara/Public/NiagaraGpuComputeDispatchInterface.h
+++ b/Engine/Plugins/FX/Niagara/Source/Niagara/Public/NiagaraGpuComputeDispatchInterface.h
@@ -65,6 +65,9 @@ public:
 	template<typename TManager>
 	TManager& GetOrCreateDataManager()
 	{
+		check(IsInParallelRenderingThread());
+
+		UE::TScopeLock ScopeLock(ComputeManagerGuard);
 		const FName ManagerName = TManager::GetManagerName();
 		for (auto& DataManager : GpuDataManagers)
 		{
@@ -183,9 +186,9 @@ public:
 #endif
 
 	/** Event that broadcast before any rendering work is prepared / executed for Niagara. */
-	FOnPostPreRenderEvent& GetOnPreRenderEvent() { return OnPreRenderEvent; }
+	FOnPostPreRenderEvent& GetOnPreRenderEvent() { check(IsInRenderingThread()); return OnPreRenderEvent; }
 	/** Event that broadcast after all rendering for Niagara is complete. */
-	FOnPostPreRenderEvent& GetOnPostRenderEvent() { return OnPostRenderEvent; }
+	FOnPostPreRenderEvent& GetOnPostRenderEvent() { check(IsInRenderingThread()); return OnPostRenderEvent; }
 
 protected:
 	EShaderPlatform							ShaderPlatform;
@@ -209,4 +212,6 @@ protected:
 
 	FOnPostPreRenderEvent					OnPreRenderEvent;
 	FOnPostPreRenderEvent					OnPostRenderEvent;
+
+	UE::FMutex								ComputeManagerGuard;
 };
diff --git a/Engine/Plugins/FX/Niagara/Source/Niagara/Public/NiagaraRenderer.h b/Engine/Plugins/FX/Niagara/Source/Niagara/Public/NiagaraRenderer.h
--- a/Engine/Plugins/FX/Niagara/Source/Niagara/Public/NiagaraRenderer.h
+++ b/Engine/Plugins/FX/Niagara/Source/Niagara/Public/NiagaraRenderer.h
@@ -34,7 +34,7 @@ struct FNiagaraDynamicDataBase
 	UE_NONCOPYABLE(FNiagaraDynamicDataBase);
 
 	NIAGARA_API bool IsGpuLowLatencyTranslucencyEnabled() const;
-	NIAGARA_API FNiagaraDataBuffer* GetParticleDataToRender(bool bIsLowLatencyTranslucent = false) const;
+	NIAGARA_API FNiagaraDataBuffer* GetParticleDataToRender(FRHICommandListBase& RHICmdList, bool bIsLowLatencyTranslucent = false) const;
 	FORCEINLINE FMaterialRelevance GetMaterialRelevance() const { return MaterialRelevance; }
 
 	FORCEINLINE void SetMaterialRelevance(FMaterialRelevance NewRelevance) { MaterialRelevance = NewRelevance; }
diff --git a/Engine/Plugins/FX/Niagara/Source/Niagara/Public/NiagaraRendererMeshes.h b/Engine/Plugins/FX/Niagara/Source/Niagara/Public/NiagaraRendererMeshes.h
--- a/Engine/Plugins/FX/Niagara/Source/Niagara/Public/NiagaraRendererMeshes.h
+++ b/Engine/Plugins/FX/Niagara/Source/Niagara/Public/NiagaraRendererMeshes.h
@@ -129,8 +129,8 @@ protected:
 		virtual ~FGPUSceneUpdateResource() {}
 	};
 
-	NIAGARA_API void PrepareParticleMeshRenderData(FParticleMeshRenderData& ParticleMeshRenderData, const FSceneViewFamily& ViewFamily, FMeshElementCollector& Collector, FNiagaraDynamicDataBase* InDynamicData, const FNiagaraSceneProxy* SceneProxy, bool bRayTracing, ENiagaraGpuComputeTickStage::Type GpuReadyTickStage) const;
-	NIAGARA_API bool CalculateMeshUsed(FParticleMeshRenderData& ParticleMeshRenderData) const;
+	NIAGARA_API void PrepareParticleMeshRenderData(FRHICommandListBase& RHICmdList, FParticleMeshRenderData& ParticleMeshRenderData, const FSceneViewFamily& ViewFamily, FMeshElementCollector& Collector, FNiagaraDynamicDataBase* InDynamicData, const FNiagaraSceneProxy* SceneProxy, bool bRayTracing, ENiagaraGpuComputeTickStage::Type GpuReadyTickStage) const;
+	NIAGARA_API bool CalculateMeshUsed(FRHICommandListBase& RHICmdList, FParticleMeshRenderData& ParticleMeshRenderData) const;
 	NIAGARA_API void PrepareParticleRenderBuffers(FRHICommandListBase& RHICmdList, FParticleMeshRenderData& ParticleMeshRenderData, FGlobalDynamicReadBuffer& DynamicReadBuffer) const;
 	NIAGARA_API void InitializeSortInfo(const FParticleMeshRenderData& ParticleMeshRenderData, const FNiagaraSceneProxy& SceneProxy, const FSceneView& View, int32 ViewIndex, bool bIsInstancedStereo, FNiagaraGPUSortInfo& OutSortInfo) const;
 	NIAGARA_API void PreparePerMeshData(FParticleMeshRenderData& ParticleMeshRenderData, const FNiagaraMeshVertexFactory& VertexFactory, const FNiagaraSceneProxy& SceneProxy, const FMeshData& MeshData) const;
diff --git a/Engine/Plugins/FX/Niagara/Source/Niagara/Public/NiagaraRendererSprites.h b/Engine/Plugins/FX/Niagara/Source/Niagara/Public/NiagaraRendererSprites.h
--- a/Engine/Plugins/FX/Niagara/Source/Niagara/Public/NiagaraRendererSprites.h
+++ b/Engine/Plugins/FX/Niagara/Source/Niagara/Public/NiagaraRendererSprites.h
@@ -71,7 +71,7 @@ private:
 		FNiagaraSpriteUniformBufferRef	UniformBuffer;
 	};
 
-	void PrepareParticleSpriteRenderData(FParticleSpriteRenderData& ParticleSpriteRenderData, const FSceneViewFamily& ViewFamily, FNiagaraDynamicDataBase* InDynamicData, const FNiagaraSceneProxy* SceneProxy, ENiagaraGpuComputeTickStage::Type GpuReadyTickStage) const;
+	void PrepareParticleSpriteRenderData(FRHICommandListBase& RHICmdList, FParticleSpriteRenderData& ParticleSpriteRenderData, const FSceneViewFamily& ViewFamily, FNiagaraDynamicDataBase* InDynamicData, const FNiagaraSceneProxy* SceneProxy, ENiagaraGpuComputeTickStage::Type GpuReadyTickStage) const;
 	void PrepareParticleRenderBuffers(FRHICommandListBase& RHICmdList, FParticleSpriteRenderData& ParticleSpriteRenderData, FGlobalDynamicReadBuffer& DynamicReadBuffer) const;
 	void InitializeSortInfo(FParticleSpriteRenderData& ParticleSpriteRenderData, const FNiagaraSceneProxy& SceneProxy, const FSceneView& View, int32 ViewIndex, FNiagaraGPUSortInfo& OutSortInfo) const;
 	void SetupVertexFactory(FRHICommandListBase& RHICmdList, FParticleSpriteRenderData& ParticleSpriteRenderData, FNiagaraSpriteVertexFactory& VertexFactory) const;
