- Couple of fixes to CPU side stateless simulation
- Add output for Unique ID
- Fix mesh size / sprite size not taking uniform into account
- Fix using normalized vs age in solve velocities and forces

#rnx
#tests Running stateless EngineTest to fix issues

diff --git a/Engine/Plugins/FX/Niagara/Source/Niagara/Internal/Stateless/Modules/NiagaraStatelessModule_InitializeParticle.h b/Engine/Plugins/FX/Niagara/Source/Niagara/Internal/Stateless/Modules/NiagaraStatelessModule_InitializeParticle.h
--- a/Engine/Plugins/FX/Niagara/Source/Niagara/Internal/Stateless/Modules/NiagaraStatelessModule_InitializeParticle.h
+++ b/Engine/Plugins/FX/Niagara/Source/Niagara/Internal/Stateless/Modules/NiagaraStatelessModule_InitializeParticle.h
@@ -131,6 +131,9 @@ public:
 
 		const FModuleBuiltData* ModuleBuiltData = ParticleSimulationContext.ReadBuiltData<FModuleBuiltData>();
 
+		const bool bUniformSpriteSize = (ModuleBuiltData->ModuleFlags & EInitializeParticleModuleFlag_UniformSpriteSize) != 0;
+		const bool bUniformMeshScale  = (ModuleBuiltData->ModuleFlags & EInitializeParticleModuleFlag_UniformMeshScale) != 0;
+
 		for (uint32 i = 0; i < ParticleSimulationContext.GetNumInstances(); ++i)
 		{
 			const FStatelessDistributionSampler<FVector3f> PositionSampler(ParticleSimulationContext, ModuleBuiltData->InitialPosition, i, 0);
@@ -138,9 +141,9 @@ public:
 			const FVector3f		Position	= PositionSampler.GetValue(ParticleSimulationContext, 0.0f);
 			const FLinearColor	Color		= ParticleSimulationContext.RandomScaleBiasFloat(i, 1 ,ModuleBuiltData->ColorRange);
 			const float			RibbonWidth	= ParticleSimulationContext.RandomScaleBiasFloat(i, 2, ModuleBuiltData->RibbonWidthRange);
-			const FVector2f		SpriteSize	= ParticleSimulationContext.RandomScaleBiasFloat(i, 3, ModuleBuiltData->SpriteSizeRange);
+			const FVector2f		SpriteSize	= ParticleSimulationContext.RandomScaleBiasFloat(i, 3, ModuleBuiltData->SpriteSizeRange, bUniformSpriteSize);
 			const float			SpriteRot	= ParticleSimulationContext.RandomScaleBiasFloat(i, 4, ModuleBuiltData->SpriteRotationRange);
-			const FVector3f		Scale		= ParticleSimulationContext.RandomScaleBiasFloat(i, 5, ModuleBuiltData->MeshScaleRange);
+			const FVector3f		Scale		= ParticleSimulationContext.RandomScaleBiasFloat(i, 5, ModuleBuiltData->MeshScaleRange, bUniformMeshScale);
 
 			ParticleSimulationContext.WriteParticleVariable(ModuleBuiltData->PositionVariableOffset,				i, Position);
 			ParticleSimulationContext.WriteParticleVariable(ModuleBuiltData->ColorVariableOffset,					i, Color);
diff --git a/Engine/Plugins/FX/Niagara/Source/Niagara/Internal/Stateless/Modules/NiagaraStatelessModule_SolveVelocitiesAndForces.h b/Engine/Plugins/FX/Niagara/Source/Niagara/Internal/Stateless/Modules/NiagaraStatelessModule_SolveVelocitiesAndForces.h
--- a/Engine/Plugins/FX/Niagara/Source/Niagara/Internal/Stateless/Modules/NiagaraStatelessModule_SolveVelocitiesAndForces.h
+++ b/Engine/Plugins/FX/Niagara/Source/Niagara/Internal/Stateless/Modules/NiagaraStatelessModule_SolveVelocitiesAndForces.h
@@ -130,8 +130,8 @@ public:
 	{
 		const FModuleBuiltData* ModuleBuiltData = ParticleSimulationContext.ReadBuiltData<FModuleBuiltData>();
 
-		const float* NormalizedAgeData			= ParticleSimulationContext.GetParticleNormalizedAge();
-		const float* PreviousNormalizedAgeData	= ParticleSimulationContext.GetParticlePreviousNormalizedAge();
+		const float* AgeData			= ParticleSimulationContext.GetParticleAge();
+		const float* PreviousAgeData	= ParticleSimulationContext.GetParticlePreviousAge();
 
 		for (uint32 i = 0; i < ParticleSimulationContext.GetNumInstances(); ++i)
 		{
@@ -177,8 +177,8 @@ public:
 				//-TODO:
 			}
 
-			Position += IntegratePosition(NormalizedAgeData[i], Mass, Drag, InitialVelocity, Wind, Acceleration);
-			PreviousPosition += IntegratePosition(PreviousNormalizedAgeData[i], Mass, Drag, InitialVelocity, Wind, Acceleration);
+			Position += IntegratePosition(AgeData[i], Mass, Drag, InitialVelocity, Wind, Acceleration);
+			PreviousPosition += IntegratePosition(PreviousAgeData[i], Mass, Drag, InitialVelocity, Wind, Acceleration);
 
 			ParticleSimulationContext.WriteParticleVariable(ModuleBuiltData->PositionVariableOffset, i, Position);
 			ParticleSimulationContext.WriteParticleVariable(ModuleBuiltData->PreviousPositionVariableOffset, i, PreviousPosition);
diff --git a/Engine/Plugins/FX/Niagara/Source/Niagara/Internal/Stateless/NiagaraStatelessParticleSimExecData.h b/Engine/Plugins/FX/Niagara/Source/Niagara/Internal/Stateless/NiagaraStatelessParticleSimExecData.h
--- a/Engine/Plugins/FX/Niagara/Source/Niagara/Internal/Stateless/NiagaraStatelessParticleSimExecData.h
+++ b/Engine/Plugins/FX/Niagara/Source/Niagara/Internal/Stateless/NiagaraStatelessParticleSimExecData.h
@@ -41,6 +41,7 @@ public:
 	explicit FParticleSimulationExecData(const FNiagaraDataSetCompiledData& ParticleDataSetCompiledData);
 
 	TArray<FVariableOffset>	VariableComponentOffsets;		// Stored offsets per variable
+	int32					UniqueIDIndex = INDEX_NONE;		// Variable index for Unique ID, since we always generate this we need to know if we should output of keep internal
 	TArray<FCallback>		SimulateFunctions;				// Series of functions to simulate particles
 };
 
diff --git a/Engine/Plugins/FX/Niagara/Source/Niagara/Private/Stateless/NiagaraStatelessParticleSimContext.cpp b/Engine/Plugins/FX/Niagara/Source/Niagara/Private/Stateless/NiagaraStatelessParticleSimContext.cpp
--- a/Engine/Plugins/FX/Niagara/Source/Niagara/Private/Stateless/NiagaraStatelessParticleSimContext.cpp
+++ b/Engine/Plugins/FX/Niagara/Source/Niagara/Private/Stateless/NiagaraStatelessParticleSimContext.cpp
@@ -157,6 +157,12 @@ namespace NiagaraStateless
 			}
 		}
 
+		// Redirect Unique ID (if required)
+		if (ExecData->UniqueIDIndex != INDEX_NONE)
+		{
+			RequiredComponents[int(EParticleComponent::UniqueIndex)] = VariableComponents[ExecData->UniqueIDIndex];
+		}
+
 		// Setup particles
 		for (uint32 iParticle = 0; iParticle < ActiveParticles; ++iParticle)
 		{
diff --git a/Engine/Plugins/FX/Niagara/Source/Niagara/Private/Stateless/NiagaraStatelessParticleSimExecData.cpp b/Engine/Plugins/FX/Niagara/Source/Niagara/Private/Stateless/NiagaraStatelessParticleSimExecData.cpp
--- a/Engine/Plugins/FX/Niagara/Source/Niagara/Private/Stateless/NiagaraStatelessParticleSimExecData.cpp
+++ b/Engine/Plugins/FX/Niagara/Source/Niagara/Private/Stateless/NiagaraStatelessParticleSimExecData.cpp
@@ -14,6 +14,7 @@ namespace NiagaraStateless
 		for (int32 i = 0; i < NumVariables; ++i)
 		{
 			const FNiagaraVariableLayoutInfo& VariableLayout = ParticleDataSetCompiledData.VariableLayouts[i];
+			const FNiagaraVariableBase& Variable = ParticleDataSetCompiledData.Variables[i];
 			if (VariableLayout.GetNumFloatComponents() > 0)
 			{
 				check(VariableLayout.GetNumInt32Components() == 0 && VariableLayout.GetNumHalfComponents() == 0);
@@ -25,6 +26,11 @@ namespace NiagaraStateless
 				check(VariableLayout.GetNumFloatComponents() == 0 && VariableLayout.GetNumHalfComponents() == 0);
 				VariableComponentOffsets[i].Type	= 1;
 				VariableComponentOffsets[i].Offset	= VariableLayout.GetInt32ComponentStart();
+
+				if (Variable.GetName() == FNiagaraStatelessGlobals::Get().UniqueIDVariable.GetName())
+				{
+					UniqueIDIndex = i;
+				}
 			}
 			else
 			{
