Fix for intermitent crash with NiagaraWorldManager getting out of date delegate broadcasts for system edits

diff --git a/Engine/Plugins/FX/Niagara/Source/Niagara/Private/NiagaraWorldManager.cpp b/Engine/Plugins/FX/Niagara/Source/Niagara/Private/NiagaraWorldManager.cpp
--- a/Engine/Plugins/FX/Niagara/Source/Niagara/Private/NiagaraWorldManager.cpp
+++ b/Engine/Plugins/FX/Niagara/Source/Niagara/Private/NiagaraWorldManager.cpp
@@ -653,7 +653,11 @@ FNiagaraSystemSimulationRef FNiagaraWorldManager::GetSystemSimulation(ETickingGr
 	}
 
 #if WITH_EDITOR
-	System->OnSystemPostEditChange().AddRaw(this, &FNiagaraWorldManager::OnSystemPostChange);
+	if (!SystemPostChangeDelegateHandles.Contains(System))
+	{
+		SystemPostChangeDelegateHandles.Emplace(System,
+			System->OnSystemPostEditChange().AddRaw(this, &FNiagaraWorldManager::OnSystemPostChange));
+	}
 #endif
 
 	return Sim;
@@ -676,7 +680,11 @@ void FNiagaraWorldManager::DestroySystemSimulation(UNiagaraSystem* System)
 	ComponentPool->RemoveComponentsBySystem(System);
 
 #if WITH_EDITOR
-	System->OnSystemPostEditChange().RemoveAll(this);
+	if (FDelegateHandle* Handle = SystemPostChangeDelegateHandles.Find(System))
+	{
+		System->OnSystemPostEditChange().Remove(*Handle);
+		SystemPostChangeDelegateHandles.Remove(System);
+	}
 #endif
 }
 
@@ -730,9 +738,6 @@ void FNiagaraWorldManager::OnWorldCleanup(bool bSessionEnded, bool bCleanupResou
 	{
 		for (TPair<UNiagaraSystem*, FNiagaraSystemSimulationRef>& SimPair : SystemSimulations[TG])
 		{
-#if WITH_EDITOR
-			SimPair.Key->OnSystemPostEditChange().RemoveAll(this);
-#endif
 			SimPair.Value->Destroy();
 		}
 		SystemSimulations[TG].Empty();
@@ -748,6 +753,17 @@ void FNiagaraWorldManager::OnWorldCleanup(bool bSessionEnded, bool bCleanupResou
 
 	CullProxyMap.Empty();
 	DIGeneratedData.Empty();
+
+#if WITH_EDITOR
+	for (FSystemDelegateMap::TConstIterator DelegateIt(SystemPostChangeDelegateHandles); DelegateIt; ++DelegateIt)
+	{
+		if (UNiagaraSystem* System = DelegateIt.Key().Get())
+		{
+			System->OnSystemPostEditChange().Remove(DelegateIt.Value());
+		}
+	}
+	SystemPostChangeDelegateHandles.Empty();
+#endif
 }
 
 void FNiagaraWorldManager::OnPostWorldCleanup(bool bSessionEnded, bool bCleanupResources)
diff --git a/Engine/Plugins/FX/Niagara/Source/Niagara/Public/NiagaraWorldManager.h b/Engine/Plugins/FX/Niagara/Source/Niagara/Public/NiagaraWorldManager.h
--- a/Engine/Plugins/FX/Niagara/Source/Niagara/Public/NiagaraWorldManager.h
+++ b/Engine/Plugins/FX/Niagara/Source/Niagara/Public/NiagaraWorldManager.h
@@ -430,6 +430,11 @@ private:
 #if WITH_PER_FXTYPE_PARTICLE_PERF_STATS
 	FParticlePerfStatsListenerPtr FXTypeCSVListener;
 #endif
+
+#if WITH_EDITOR
+	using FSystemDelegateMap = TMap<TWeakObjectPtr<UNiagaraSystem>, FDelegateHandle>;
+	FSystemDelegateMap SystemPostChangeDelegateHandles;
+#endif
 };
 
 
