Stall detection with new niagara compilation mode and removes completed (but not retrieved) tasks from the calculation of number of assets remaining to be compiled.

diff --git a/Engine/Plugins/FX/Niagara/Source/NiagaraEditor/Private/NiagaraCompilationTasks.cpp b/Engine/Plugins/FX/Niagara/Source/NiagaraEditor/Private/NiagaraCompilationTasks.cpp
--- a/Engine/Plugins/FX/Niagara/Source/NiagaraEditor/Private/NiagaraCompilationTasks.cpp
+++ b/Engine/Plugins/FX/Niagara/Source/NiagaraEditor/Private/NiagaraCompilationTasks.cpp
@@ -505,6 +505,57 @@ void FNiagaraSystemCompilationTask::Abort()
 	CompileCompletionEvent.Trigger();
 }
 
+FString FNiagaraSystemCompilationTask::GetDescription() const
+{
+	FString TaskDescription;
+	TaskDescription.Append(SystemInfo.SystemName);
+	if (bCompileForEdit)
+	{
+		TaskDescription.Append(TEXT(" [Edit]"));
+	}
+	return TaskDescription;
+}
+
+FString FNiagaraSystemCompilationTask::GetStatusString() const
+{
+	FString StatusString;
+
+	switch (CompilationState)
+	{
+		case EState::Invalid:
+		{
+			StatusString = TEXT("Invalid");
+			break;
+		}
+
+		case EState::WaitingForProcessing:
+		{
+			StatusString = TEXT("WaitingForProcessing");
+			break;
+		}
+
+		case EState::ResultsProcessed:
+		{
+			StatusString = TEXT("ResultsProcessed");
+			break;
+		}
+
+		case EState::Completed:
+		{
+			StatusString = TEXT("Completed");
+			break;
+		}
+
+		case EState::Aborted:
+		{
+			StatusString = TEXT("Aborted");
+			break;
+		}
+	}
+
+	return StatusString;
+}
+
 const FNiagaraSystemCompilationTask::FEmitterInfo* FNiagaraSystemCompilationTask::FSystemInfo::EmitterInfoBySourceEmitter(int32 InSourceEmitterIndex) const
 {
 	return EmitterInfo.FindByPredicate([InSourceEmitterIndex](const FEmitterInfo& Info) -> bool
diff --git a/Engine/Plugins/FX/Niagara/Source/NiagaraEditor/Private/NiagaraCompilationTasks.h b/Engine/Plugins/FX/Niagara/Source/NiagaraEditor/Private/NiagaraCompilationTasks.h
--- a/Engine/Plugins/FX/Niagara/Source/NiagaraEditor/Private/NiagaraCompilationTasks.h
+++ b/Engine/Plugins/FX/Niagara/Source/NiagaraEditor/Private/NiagaraCompilationTasks.h
@@ -49,14 +49,20 @@ struct FNiagaraSystemCompilationTask
 	void IssueCompilationTasks();
 	void IssuePostResultsProcessedTasks();
 	bool HasOutstandingCompileTasks() const;
+	FString GetDescription() const;
+	FString GetStatusString() const;
 
 	double PrepareStartTime = 0.0;
 	double QueueStartTime = 0.0;
 	double LaunchStartTime = 0.0;
+	double LastStallWarningTime = 0.0;
 
 	// compilation was forced
 	bool bForced = false;
 
+	// task has been deemed to be stalled
+	bool bStalled = false;
+
 	enum class EState : uint8
 	{
 		Invalid = 0,
diff --git a/Engine/Plugins/FX/Niagara/Source/NiagaraEditor/Private/NiagaraSystemCompilingManager.cpp b/Engine/Plugins/FX/Niagara/Source/NiagaraEditor/Private/NiagaraSystemCompilingManager.cpp
--- a/Engine/Plugins/FX/Niagara/Source/NiagaraEditor/Private/NiagaraSystemCompilingManager.cpp
+++ b/Engine/Plugins/FX/Niagara/Source/NiagaraEditor/Private/NiagaraSystemCompilingManager.cpp
@@ -22,6 +22,14 @@ static FAutoConsoleVariableRef CVarNiagaraCompilationMaxActiveTaskCount(
 	ECVF_Default
 );
 
+static float GNiagaraCompilationStalledTaskWarningTime = 10.0f * 60.0f;
+static FAutoConsoleVariableRef CVarNiagaraCompilationStalledTaskWarningTime(
+	TEXT("fx.Niagara.Compilation.StalledTaskWarningTime"),
+	GNiagaraCompilationStalledTaskWarningTime,
+	TEXT("The length of time a task is being processed before warnings are generated."),
+	ECVF_Default
+);
+
 #if ENABLE_COOK_STATS
 namespace NiagaraSystemCookStats
 {
@@ -129,8 +137,12 @@ int32 FNiagaraSystemCompilingManager::GetNumRemainingAssets() const
 	int32 RemainingAssetCount = 0;
 
 	{
+		// note that we don't worry about including RequestsAwaitingRetrieval because
+		// those tasks do not reflect significant remaining work for the compilation manager.
+		// Additionally, it can cause deadlocks in some scenarios as calling code could wait
+		// for the remaining assets to get to 0 before advancing to polling for results.
 		FReadScopeLock Read(QueueLock);
-		RemainingAssetCount = QueuedRequests.Num() + ActiveTasks.Num() + RequestsAwaitingRetrieval.Num();
+		RemainingAssetCount = QueuedRequests.Num() + ActiveTasks.Num();
 	}
 
 	return RemainingAssetCount;
@@ -169,6 +181,39 @@ void FNiagaraSystemCompilingManager::Shutdown()
 {
 }
 
+void FNiagaraSystemCompilingManager::CheckStalledTask(double CurrentTime, FNiagaraSystemCompilationTask* Task) const
+{
+	const double ElapsedTime = CurrentTime - Task->LaunchStartTime;
+
+	if (ElapsedTime > GNiagaraCompilationStalledTaskWarningTime)
+	{
+		bool bWarn = true;
+
+		if (!Task->bStalled)
+		{
+			Task->LastStallWarningTime = CurrentTime;
+			Task->bStalled = true;
+
+		}
+		else
+		{
+			const double TimeSinceLastWarning = CurrentTime - Task->LastStallWarningTime;
+			if (TimeSinceLastWarning < GNiagaraCompilationStalledTaskWarningTime)
+			{
+				bWarn = false;
+			}
+		}
+
+		if (bWarn)
+		{
+			UE_LOG(LogNiagaraEditor, Warning, TEXT("NiagaraSystemCompilingManager - compilation task [%s] stalled for %f seconds.  Status - %s"),
+				*Task->GetDescription(), (float)ElapsedTime, *Task->GetStatusString());
+
+			Task->LastStallWarningTime = CurrentTime;
+		}
+	}
+}
+
 void FNiagaraSystemCompilingManager::ProcessAsyncTasks(bool bLimitExecutionTime)
 {
 	{
@@ -190,12 +235,15 @@ void FNiagaraSystemCompilingManager::ProcessAsyncTasks(bool bLimitExecutionTime)
 		}
 
 		{
+			double CurrentTime = FPlatformTime::Seconds();
+
 			FReadScopeLock ReadScope(QueueLock);
 			for (FNiagaraCompilationTaskHandle TaskHandle : ActiveTasks)
 			{
 				FTaskPtr TaskPtr = SystemRequestMap.FindRef(TaskHandle);
 				if (TaskPtr.IsValid())
 				{
+					CheckStalledTask(CurrentTime, TaskPtr.Get());
 					TaskPtr->Tick();
 				}
 			}
diff --git a/Engine/Plugins/FX/Niagara/Source/NiagaraEditor/Private/NiagaraSystemCompilingManager.h b/Engine/Plugins/FX/Niagara/Source/NiagaraEditor/Private/NiagaraSystemCompilingManager.h
--- a/Engine/Plugins/FX/Niagara/Source/NiagaraEditor/Private/NiagaraSystemCompilingManager.h
+++ b/Engine/Plugins/FX/Niagara/Source/NiagaraEditor/Private/NiagaraSystemCompilingManager.h
@@ -45,6 +45,7 @@ protected:
 	// End - IAssetCompilingManager
 
 	bool ConditionalLaunchTask();
+	void CheckStalledTask(double CurrentTime, FNiagaraSystemCompilationTask* Task) const;
 
 	FNiagaraShaderType* NiagaraShaderType = nullptr;
 
