Enable a bunch of context information for intermittent failure being seen with Niagara's new compilation mode
#tests Created artificial tests to validate log message

diff --git a/Engine/Plugins/FX/Niagara/Source/NiagaraEditor/Private/NiagaraTraversalStateContext.cpp b/Engine/Plugins/FX/Niagara/Source/NiagaraEditor/Private/NiagaraTraversalStateContext.cpp
--- a/Engine/Plugins/FX/Niagara/Source/NiagaraEditor/Private/NiagaraTraversalStateContext.cpp
+++ b/Engine/Plugins/FX/Niagara/Source/NiagaraEditor/Private/NiagaraTraversalStateContext.cpp
@@ -107,10 +107,43 @@ void FNiagaraTraversalStateContext::PushGraphInternal(const FNiagaraCompilationN
 						if (CallerInputPin)
 						{
 							const uint32 FunctionPinNodeHash = HashCombine(InnerFunctionNodeHash, GetTypeHash(ValuePin->PinName));
-							if (ensure(!FunctionDefaultValueMap.Contains(FunctionPinNodeHash)))
+
+							const FString* ExistingDefaultValue = FunctionDefaultValueMap.Find(FunctionPinNodeHash);
+							if (ensure(!ExistingDefaultValue))
 							{
 								FunctionDefaultValueMap.Add(FunctionPinNodeHash, CallerInputPin->DefaultValue);
 							}
+							else
+							{
+								// generate an error message to help track down the scenario where this has happened
+								UE_LOG(LogNiagaraEditor, Warning, TEXT("FNiagaraTraversalStateContext::PushGraphInternal() generated a non-unique function call.\n" \
+									"\t[ExistingDefaultValue] %s\n" \
+									"\t[NewDefaultValue] %s\n" \
+									"\t[ValuePin->PinName] %s\n" \
+									"\t[TaggedVariable] %s - %s\n" \
+									"\t[InnerFunctionNode] %s - %s"),
+									ExistingDefaultValue ? **ExistingDefaultValue : TEXT("<null>"),
+									*CallerInputPin->DefaultValue,
+									*ValuePin->PinName.ToString(),
+									*TaggedVariable.Key.GetName().ToString(), *TaggedVariable.Value.ToString(),
+									*InnerFunctionNode->FunctionName, *InnerFunctionNode->FunctionScriptName
+									);
+								
+								UE_LOG(LogNiagaraEditor, Warning, TEXT("FNiagaraTraversalStateContext - Stack"))
+								for (int32 StackIt = TraversalStack.Num() - 1; StackIt >= 0; --StackIt)
+								{
+									FString StackMessage = FString::Printf(TEXT("[%d] - %s, %xd"),
+										StackIt,
+										*TraversalStack[StackIt].NodeGuid.ToString(EGuidFormats::DigitsWithHyphens),
+										TraversalStack[StackIt].FullStackHash);
+
+#if WITH_NIAGARA_TRAVERSAL_FRIENDLY_NAME
+									StackMessage.Append(TEXT(", "));
+									StackMessage.Append(TraversalStack[StackIt].FriendlyName);
+#endif
+									UE_LOG(LogNiagaraEditor, Warning, TEXT("%s"),  *StackMessage);
+								}
+							}
 						}
 					}
 				}
diff --git a/Engine/Plugins/FX/Niagara/Source/NiagaraEditor/Private/NiagaraTraversalStateContext.h b/Engine/Plugins/FX/Niagara/Source/NiagaraEditor/Private/NiagaraTraversalStateContext.h
--- a/Engine/Plugins/FX/Niagara/Source/NiagaraEditor/Private/NiagaraTraversalStateContext.h
+++ b/Engine/Plugins/FX/Niagara/Source/NiagaraEditor/Private/NiagaraTraversalStateContext.h
@@ -10,7 +10,7 @@ class FNiagaraCompilationNodeEmitter;
 class FNiagaraCompilationNodeFunctionCall;
 class FNiagaraFixedConstantResolver;
 
-#define WITH_NIAGARA_TRAVERSAL_FRIENDLY_NAME (UE_BUILD_DEBUG)
+#define WITH_NIAGARA_TRAVERSAL_FRIENDLY_NAME (1)
 
 struct FNiagaraTraversalStackEntry
 {
