Allow stateless module properties to use customized widgets, similar to existing module inputs.

#tests Tested the existing stateless emitter ui for all the modules and emitter properties
#rnx

diff --git a/Engine/Plugins/FX/Niagara/Source/Niagara/Internal/NiagaraSystemEmitterState.h b/Engine/Plugins/FX/Niagara/Source/Niagara/Internal/NiagaraSystemEmitterState.h
--- a/Engine/Plugins/FX/Niagara/Source/Niagara/Internal/NiagaraSystemEmitterState.h
+++ b/Engine/Plugins/FX/Niagara/Source/Niagara/Internal/NiagaraSystemEmitterState.h
@@ -64,13 +64,13 @@ struct FNiagaraSystemStateData
 	UPROPERTY(EditAnywhere, Category = "System State")
 	ENiagaraLoopBehavior LoopBehavior = ENiagaraLoopBehavior::Once;
 
-	UPROPERTY(EditAnywhere, Category = "System State", meta = (ClampMin = "0.0"))
+	UPROPERTY(EditAnywhere, Category = "System State", meta = (ClampMin = "0.0", Units="s"))
 	FNiagaraDistributionRangeFloat LoopDuration = FNiagaraDistributionRangeFloat(0.0f);
 
 	UPROPERTY(EditAnywhere, Category = "System State", meta = (ClampMin = "1", EditCondition = "LoopBehavior == ENiagaraLoopBehavior::Multiple", EditConditionHides))
 	int LoopCount = 1;
 
-	UPROPERTY(EditAnywhere, Category = "System State", meta = (ClampMin = "0.0"))
+	UPROPERTY(EditAnywhere, Category = "System State", meta = (ClampMin = "0.0", Units="s"))
 	FNiagaraDistributionRangeFloat LoopDelay = FNiagaraDistributionRangeFloat(0.0f);
 };
 
@@ -82,19 +82,19 @@ struct FNiagaraEmitterStateData
 	//UPROPERTY(EditAnywhere, Category="Emitter State")
 	//ENiagaraStatelessEmitterState_SelfSystem LifeCycleMode = ENiagaraStatelessEmitterState_SelfSystem::Self;
 
-	UPROPERTY(EditAnywhere, Category="Emitter State")
+	UPROPERTY(EditAnywhere, Category="Emitter State", meta=(SegmentedDisplay))
 	ENiagaraEmitterInactiveResponse InactiveResponse = ENiagaraEmitterInactiveResponse::Complete;
 
-	UPROPERTY(EditAnywhere, Category="Emitter State")
+	UPROPERTY(EditAnywhere, Category="Emitter State", meta=(SegmentedDisplay))
 	ENiagaraLoopBehavior LoopBehavior = ENiagaraLoopBehavior::Infinite;
 
 	UPROPERTY(EditAnywhere, Category = "Emitter State", meta = (ClampMin = "1", EditCondition = "LoopBehavior == ENiagaraLoopBehavior::Multiple", EditConditionHides))
 	int32 LoopCount = 1;
 
-	UPROPERTY(EditAnywhere, Category = "Emitter State", meta = (ClampMin = "0.0"))
+	UPROPERTY(EditAnywhere, Category = "Emitter State", meta = (ClampMin = "0.0", Units="s"))
 	FNiagaraDistributionRangeFloat LoopDuration = FNiagaraDistributionRangeFloat(1.0f);
 
-	UPROPERTY(EditAnywhere, Category = "Emitter State", meta = (ClampMin = "0.0"))
+	UPROPERTY(EditAnywhere, Category = "Emitter State", meta = (ClampMin = "0.0", Units="s"))
 	FNiagaraDistributionRangeFloat LoopDelay = FNiagaraDistributionRangeFloat(0.0f);
 
 	UPROPERTY(EditAnywhere, Category = "Emitter State", meta = (EditCondition = "LoopBehavior != ENiagaraLoopBehavior::Once", EditConditionHides, DisplayAfter = "LoopDuration"))
diff --git a/Engine/Plugins/FX/Niagara/Source/Niagara/Internal/Stateless/Modules/NiagaraStatelessModule_AddVelocity.h b/Engine/Plugins/FX/Niagara/Source/Niagara/Internal/Stateless/Modules/NiagaraStatelessModule_AddVelocity.h
--- a/Engine/Plugins/FX/Niagara/Source/Niagara/Internal/Stateless/Modules/NiagaraStatelessModule_AddVelocity.h
+++ b/Engine/Plugins/FX/Niagara/Source/Niagara/Internal/Stateless/Modules/NiagaraStatelessModule_AddVelocity.h
@@ -23,7 +23,7 @@ class UNiagaraStatelessModule_AddVelocity : public UNiagaraStatelessModule
 	GENERATED_BODY()
 
 public:
-	UPROPERTY(EditAnywhere, Category = "Parameters")
+	UPROPERTY(EditAnywhere, Category = "Parameters", meta=(SegmentedDisplay))
 	ENSM_VelocityType VelocityType = ENSM_VelocityType::Linear;
 
 	UPROPERTY(EditAnywhere, Category = "Parameters", meta = (DisplayName = "Velocity", DisableUniformDistribution, DisableBindingDistribution, EditConditionHides, EditCondition = "VelocityType == ENSM_VelocityType::Linear"))
diff --git a/Engine/Plugins/FX/Niagara/Source/Niagara/Internal/Stateless/Modules/NiagaraStatelessModule_InitialMeshOrientation.h b/Engine/Plugins/FX/Niagara/Source/Niagara/Internal/Stateless/Modules/NiagaraStatelessModule_InitialMeshOrientation.h
--- a/Engine/Plugins/FX/Niagara/Source/Niagara/Internal/Stateless/Modules/NiagaraStatelessModule_InitialMeshOrientation.h
+++ b/Engine/Plugins/FX/Niagara/Source/Niagara/Internal/Stateless/Modules/NiagaraStatelessModule_InitialMeshOrientation.h
@@ -25,10 +25,10 @@ class UNiagaraStatelessModule_InitialMeshOrientation : public UNiagaraStatelessM
 public:
 	using FParameters = NiagaraStateless::FInitialMeshOrientationModule_ShaderParameters;
 
-	UPROPERTY(EditAnywhere, Category = "Parameters")
+	UPROPERTY(EditAnywhere, Category = "Parameters", meta=(Units="deg"))
 	FVector3f	Rotation = FVector3f::ZeroVector;
 
-	UPROPERTY(EditAnywhere, Category = "Parameters")
+	UPROPERTY(EditAnywhere, Category = "Parameters", meta=(Units="deg"))
 	FVector3f	RandomRotationRange = FVector3f(360.0f, 360.0f, 360.0f);
 
 	virtual void BuildEmitterData(const FNiagaraStatelessEmitterDataBuildContext& BuildContext) const override
diff --git a/Engine/Plugins/FX/Niagara/Source/Niagara/Internal/Stateless/Modules/NiagaraStatelessModule_InitializeParticle.h b/Engine/Plugins/FX/Niagara/Source/Niagara/Internal/Stateless/Modules/NiagaraStatelessModule_InitializeParticle.h
--- a/Engine/Plugins/FX/Niagara/Source/Niagara/Internal/Stateless/Modules/NiagaraStatelessModule_InitializeParticle.h
+++ b/Engine/Plugins/FX/Niagara/Source/Niagara/Internal/Stateless/Modules/NiagaraStatelessModule_InitializeParticle.h
@@ -50,7 +50,7 @@ class UNiagaraStatelessModule_InitializeParticle : public UNiagaraStatelessModul
 public:
 	using FParameters = NiagaraStateless::FInitializeParticleModule_ShaderParameters;
 
-	UPROPERTY(EditAnywhere, Category = "Parameters", meta = (DisplayName = "Lifetime"))
+	UPROPERTY(EditAnywhere, Category = "Parameters", meta = (DisplayName = "Lifetime", Units="s"))
 	FNiagaraDistributionRangeFloat LifetimeDistribution = FNiagaraDistributionRangeFloat(FNiagaraStatelessGlobals::GetDefaultLifetimeValue());
 
 	UPROPERTY(EditAnywhere, Category = "Parameters", meta = (DisplayName = "Color", DisableCurveDistribution))
@@ -62,7 +62,7 @@ public:
 	UPROPERTY(EditAnywhere, Category = "Parameters", meta = (DisplayName = "Sprite Size"))
 	FNiagaraDistributionRangeVector2 SpriteSizeDistribution = FNiagaraDistributionRangeVector2(FNiagaraStatelessGlobals::GetDefaultSpriteSizeValue());
 
-	UPROPERTY(EditAnywhere, Category = "Parameters", meta = (DisplayName = "Sprite Rotation"))
+	UPROPERTY(EditAnywhere, Category = "Parameters", meta = (DisplayName = "Sprite Rotation", Units="deg"))
 	FNiagaraDistributionRangeFloat SpriteRotationDistribution = FNiagaraDistributionRangeFloat(FNiagaraStatelessGlobals::GetDefaultSpriteRotationValue());
 
 	UPROPERTY(EditAnywhere, Category = "Parameters", meta = (DisplayName = "Mesh Scale"))
diff --git a/Engine/Plugins/FX/Niagara/Source/Niagara/Internal/Stateless/Modules/NiagaraStatelessModule_MeshRotationRate.h b/Engine/Plugins/FX/Niagara/Source/Niagara/Internal/Stateless/Modules/NiagaraStatelessModule_MeshRotationRate.h
--- a/Engine/Plugins/FX/Niagara/Source/Niagara/Internal/Stateless/Modules/NiagaraStatelessModule_MeshRotationRate.h
+++ b/Engine/Plugins/FX/Niagara/Source/Niagara/Internal/Stateless/Modules/NiagaraStatelessModule_MeshRotationRate.h
@@ -25,7 +25,7 @@ class UNiagaraStatelessModule_MeshRotationRate : public UNiagaraStatelessModule
 public:
 	using FParameters = NiagaraStateless::FMeshRotationRateModule_ShaderParameters;
 
-	UPROPERTY(EditAnywhere, Category = "Parameters", meta = (DisplayName = "Rotation Rate"))
+	UPROPERTY(EditAnywhere, Category = "Parameters", meta = (DisplayName = "Rotation Rate", Units="deg/s"))
 	FNiagaraDistributionRangeVector3 RotationRateDistribution = FNiagaraDistributionRangeVector3(FVector3f::ZeroVector);
 
 	virtual void BuildEmitterData(const FNiagaraStatelessEmitterDataBuildContext& BuildContext) const override
diff --git a/Engine/Plugins/FX/Niagara/Source/Niagara/Internal/Stateless/Modules/NiagaraStatelessModule_RotateAroundPoint.h b/Engine/Plugins/FX/Niagara/Source/Niagara/Internal/Stateless/Modules/NiagaraStatelessModule_RotateAroundPoint.h
--- a/Engine/Plugins/FX/Niagara/Source/Niagara/Internal/Stateless/Modules/NiagaraStatelessModule_RotateAroundPoint.h
+++ b/Engine/Plugins/FX/Niagara/Source/Niagara/Internal/Stateless/Modules/NiagaraStatelessModule_RotateAroundPoint.h
@@ -15,9 +15,9 @@ class UNiagaraStatelessModule_RotateAroundPoint : public UNiagaraStatelessModule
 public:
 	using FParameters = NiagaraStateless::FRotateAroundPointModule_ShaderParameters;
 
-	UPROPERTY(EditAnywhere, Category = "Parameters")
+	UPROPERTY(EditAnywhere, Category = "Parameters", meta=(Units="deg/s"))
 	float RateMin = 360.f;
-	UPROPERTY(EditAnywhere, Category = "Parameters")
+	UPROPERTY(EditAnywhere, Category = "Parameters", meta=(Units="deg/s"))
 	float RateMax = 360.f;
 
 	UPROPERTY(EditAnywhere, Category = "Parameters")
@@ -31,7 +31,7 @@ public:
 	float InitialPhaseMax = 0.f;
 
 	//-TODO: Add support for GPU once we settle on a feature set for this module
-	virtual ENiagaraStatelessFeatureMask GetFeatureMask() const { return ENiagaraStatelessFeatureMask::ExecuteGPU; }
+	virtual ENiagaraStatelessFeatureMask GetFeatureMask() const override { return ENiagaraStatelessFeatureMask::ExecuteGPU; }
 
 	virtual void SetShaderParameters(const FNiagaraStatelessSetShaderParameterContext& SetShaderParameterContext) const override
 	{
diff --git a/Engine/Plugins/FX/Niagara/Source/Niagara/Internal/Stateless/Modules/NiagaraStatelessModule_ShapeLocation.h b/Engine/Plugins/FX/Niagara/Source/Niagara/Internal/Stateless/Modules/NiagaraStatelessModule_ShapeLocation.h
--- a/Engine/Plugins/FX/Niagara/Source/Niagara/Internal/Stateless/Modules/NiagaraStatelessModule_ShapeLocation.h
+++ b/Engine/Plugins/FX/Niagara/Source/Niagara/Internal/Stateless/Modules/NiagaraStatelessModule_ShapeLocation.h
@@ -66,7 +66,7 @@ public:
 #if WITH_EDITOR
 	virtual bool CanDisableModule() const override { return true; }
 
-	virtual bool CanDebugDraw() const { return true; }
+	virtual bool CanDebugDraw() const override { return true; }
 	virtual void DrawDebug(const FNiagaraStatelessDrawDebugContext& DrawDebugContext) const override;
 #endif
 #if WITH_EDITORONLY_DATA
diff --git a/Engine/Plugins/FX/Niagara/Source/Niagara/Internal/Stateless/Modules/NiagaraStatelessModule_SubUVAnimation.h b/Engine/Plugins/FX/Niagara/Source/Niagara/Internal/Stateless/Modules/NiagaraStatelessModule_SubUVAnimation.h
--- a/Engine/Plugins/FX/Niagara/Source/Niagara/Internal/Stateless/Modules/NiagaraStatelessModule_SubUVAnimation.h
+++ b/Engine/Plugins/FX/Niagara/Source/Niagara/Internal/Stateless/Modules/NiagaraStatelessModule_SubUVAnimation.h
@@ -54,7 +54,7 @@ public:
 	UPROPERTY(EditAnywhere, Category = "Parameters", meta = (EditCondition = "bEndFrameRangeOverride_Enabled"))
 	int32	EndFrameRangeOverride = 0;
 	
-	UPROPERTY(EditAnywhere, Category = "Parameters")
+	UPROPERTY(EditAnywhere, Category = "Parameters", meta=(SegmentedDisplay))
 	ENSMSubUVAnimation_Mode AnimationMode = ENSMSubUVAnimation_Mode::Linear;
 
 	//-Note: Main module has PlaybackMode (Loops / FPS) to choose between loops or frames per second
diff --git a/Engine/Plugins/FX/Niagara/Source/Niagara/Private/NiagaraVariableMetaData.cpp b/Engine/Plugins/FX/Niagara/Source/Niagara/Private/NiagaraVariableMetaData.cpp
--- a/Engine/Plugins/FX/Niagara/Source/Niagara/Private/NiagaraVariableMetaData.cpp
+++ b/Engine/Plugins/FX/Niagara/Source/Niagara/Private/NiagaraVariableMetaData.cpp
@@ -3,6 +3,47 @@
 
 #include "NiagaraVariableMetaData.h"
 
+#if WITH_EDITOR
+#include "PropertyHandle.h"
+
+FNiagaraInputParameterCustomization FNiagaraInputParameterCustomization::MakeFromProperty(TSharedPtr<IPropertyHandle> PropertyHandle)
+{
+	FNiagaraInputParameterCustomization ParameterCustomization;
+	if (PropertyHandle.IsValid())
+	{
+		static const FName ClampMinName("ClampMin");
+		static const FName ClampMaxName("ClampMax");
+		static const FName DeltaName("Delta");
+		static const FName WidgetTypeName = GET_MEMBER_NAME_CHECKED(FNiagaraInputParameterCustomization, WidgetType);
+		if (PropertyHandle->HasMetaData(WidgetTypeName))
+		{
+			FString TypeString = PropertyHandle->GetMetaData(WidgetTypeName);
+			UEnum* TypeEnum = StaticEnum<ENiagaraInputWidgetType>();
+			int32 Value = TypeEnum->GetValueByNameString(TypeString);
+			if (Value != INDEX_NONE)
+			{
+				ParameterCustomization.WidgetType = static_cast<ENiagaraInputWidgetType>(Value);
+			}
+		}
+		if (PropertyHandle->HasMetaData(ClampMinName))
+		{
+			ParameterCustomization.bHasMinValue = true;
+			ParameterCustomization.MinValue = PropertyHandle->GetFloatMetaData(ClampMinName);
+		}
+		if (PropertyHandle->HasMetaData(ClampMaxName))
+		{
+			ParameterCustomization.bHasMaxValue = true;
+			ParameterCustomization.MaxValue = PropertyHandle->GetFloatMetaData(ClampMaxName);
+		}
+		if (PropertyHandle->HasMetaData(DeltaName))
+		{
+			ParameterCustomization.bHasStepWidth = true;
+			ParameterCustomization.StepWidth = PropertyHandle->GetFloatMetaData(DeltaName);
+		}
+	}
+	return ParameterCustomization;
+}
+#endif
 
 void FNiagaraVariableMetaData::CopyUserEditableMetaData(const FNiagaraVariableMetaData& OtherMetaData)
 {
diff --git a/Engine/Plugins/FX/Niagara/Source/Niagara/Public/NiagaraVariableMetaData.h b/Engine/Plugins/FX/Niagara/Source/Niagara/Public/NiagaraVariableMetaData.h
--- a/Engine/Plugins/FX/Niagara/Source/Niagara/Public/NiagaraVariableMetaData.h
+++ b/Engine/Plugins/FX/Niagara/Source/Niagara/Public/NiagaraVariableMetaData.h
@@ -8,6 +8,7 @@
 #include "Math/UnitConversion.inl"
 #include "NiagaraVariableMetaData.generated.h"
 
+class IPropertyHandle;
 
 /** Defines options for conditionally editing and showing script inputs in the UI. */
 USTRUCT()
@@ -166,6 +167,10 @@ struct FNiagaraInputParameterCustomization
 	/** If true then the input is also displayed and editable as a 3d widget in the viewport (vector and transform types only). */
 	//UPROPERTY(EditAnywhere, Category="Customization")
 	//bool bCreateViewPortEditWidget = false;
+#if WITH_EDITOR
+	// Maps existing property tags as best it can to a customization object
+	NIAGARA_API static FNiagaraInputParameterCustomization MakeFromProperty(TSharedPtr<IPropertyHandle> InPropertyHandle);
+#endif
 };
 
 UENUM()
diff --git a/Engine/Plugins/FX/Niagara/Source/NiagaraEditor/Private/Stateless/NiagaraDistributionIntPropertyCustomization.cpp b/Engine/Plugins/FX/Niagara/Source/NiagaraEditor/Private/Stateless/NiagaraDistributionIntPropertyCustomization.cpp
--- a/Engine/Plugins/FX/Niagara/Source/NiagaraEditor/Private/Stateless/NiagaraDistributionIntPropertyCustomization.cpp
+++ b/Engine/Plugins/FX/Niagara/Source/NiagaraEditor/Private/Stateless/NiagaraDistributionIntPropertyCustomization.cpp
@@ -16,6 +16,7 @@
 #include "Framework/MultiBox/MultiBoxBuilder.h"
 #include "Stateless/NiagaraStatelessDistribution.h"
 #include "Styling/StyleColors.h"
+#include "TypeEditorUtilities/NiagaraIntegerTypeEditorUtilities.h"
 #include "Widgets/DeclarativeSyntaxSupport.h"
 #include "Widgets/NiagaraDistributionEditorUtilities.h"
 #include "Widgets/SBoxPanel.h"
@@ -38,6 +39,15 @@ class SNiagaraDistributionIntPropertyWidget : public SCompoundWidget
 		WeakPropertyHandle	= InPropertyHandle;
 		WeakOwnerObject = InOwnerObject;
 		Distribution = InDistribution;
+		WidgetCustomizationOptions = FNiagaraInputParameterCustomization::MakeFromProperty(InPropertyHandle);
+
+		static const FName UnitsName("Units");
+		if (InPropertyHandle->HasMetaData(UnitsName))
+		{
+			FString UnitString = InPropertyHandle->GetMetaData(UnitsName);
+			TOptional<EUnit> PropertyUnit = FUnitConversion::UnitFromString(*UnitString);
+			DisplayUnit = PropertyUnit.Get(EUnit::Unspecified);
+		}
 
 		ChildSlot
 		[
@@ -179,7 +189,7 @@ class SNiagaraDistributionIntPropertyWidget : public SCompoundWidget
 			return SNew(SHorizontalBox)
 				+ SHorizontalBox::Slot()
 				[
-					ConstructIntWidget(0, FText())
+					ConstructCustomizableWidget(0)
 				];
 		}
 		else if (Mode == ENiagaraDistributionEditorMode::UniformRange)
@@ -197,7 +207,7 @@ class SNiagaraDistributionIntPropertyWidget : public SCompoundWidget
 		return SNullWidget::NullWidget;
 	}
 
-	TSharedRef<SWidget> ConstructIntWidget(int32 ValueIndex, FText LabelText)
+	TSharedRef<SWidget> ConstructIntWidget(int32 ValueIndex, const FText& LabelText)
 	{
 		TSharedRef<SWidget> LabelWidget = SNullWidget::NullWidget;
 		if ( !LabelText.IsEmpty() )
@@ -221,6 +231,7 @@ class SNiagaraDistributionIntPropertyWidget : public SCompoundWidget
 			.MinSliderValue(TOptional<int32>())
 			.MaxSliderValue(TOptional<int32>())
 			.BroadcastValueChangesPerKey(false)
+			.TypeInterface(MakeShareable(new TNumericUnitTypeInterface<int32>(DisplayUnit)))
 			.LabelVAlign(EVerticalAlignment::VAlign_Center)
 			.MinDesiredValueWidth(30)
 			.Label()
@@ -229,6 +240,15 @@ class SNiagaraDistributionIntPropertyWidget : public SCompoundWidget
 			];
 	}
 
+	TSharedRef<SWidget> ConstructCustomizableWidget(int32 ValueIndex)
+	{
+		return SNew(SNiagaraIntegerParameterEditor, DisplayUnit, WidgetCustomizationOptions)
+			.Value(this, &SNiagaraDistributionIntPropertyWidget::GetValueInt, ValueIndex)
+			.OnValueChanged(this, &SNiagaraDistributionIntPropertyWidget::ValueChanged, ValueIndex)
+			.OnBeginValueChange(this, &SNiagaraDistributionIntPropertyWidget::BeginValueSliderMovement)
+			.OnEndValueChange(this, &SNiagaraDistributionIntPropertyWidget::EndValueSliderMovement);
+	}
+
 	TOptional<int32> GetValue(int32 ValueIndex) const
 	{
 		TOptional<int32> Value;
@@ -241,6 +261,12 @@ class SNiagaraDistributionIntPropertyWidget : public SCompoundWidget
 		return Value;
 	}
 
+	int32 GetValueInt(int32 ValueIndex) const
+	{
+		TOptional<int32> Value = GetValue(ValueIndex);
+		return Value.Get(0);
+	}
+
 	void ValueChanged(int32 Value, int32 ValueIndex)
 	{
 		TSharedPtr<IPropertyHandle> PropertyHandle = WeakPropertyHandle.Pin();
@@ -371,10 +397,12 @@ class SNiagaraDistributionIntPropertyWidget : public SCompoundWidget
 
 private:
 	TWeakPtr<IPropertyHandle>		WeakPropertyHandle;
-	TWeakObjectPtr<UObject>			WeakOwnerObject;
+	TWeakObjectPtr<>			WeakOwnerObject;
 	FNiagaraDistributionRangeInt*	Distribution = nullptr;
 	bool							bContinuousChangeActive = false;
 	TSharedPtr<SBox>				ContentBox;
+	EUnit							DisplayUnit = EUnit::Unspecified;
+	FNiagaraInputParameterCustomization WidgetCustomizationOptions;
 };
 
 TSharedRef<IPropertyTypeCustomization> FNiagaraDistributionIntPropertyCustomization::MakeIntInstance(UObject* OptionalOuter)
diff --git a/Engine/Plugins/FX/Niagara/Source/NiagaraEditor/Private/Stateless/NiagaraDistributionPropertyCustomization.cpp b/Engine/Plugins/FX/Niagara/Source/NiagaraEditor/Private/Stateless/NiagaraDistributionPropertyCustomization.cpp
--- a/Engine/Plugins/FX/Niagara/Source/NiagaraEditor/Private/Stateless/NiagaraDistributionPropertyCustomization.cpp
+++ b/Engine/Plugins/FX/Niagara/Source/NiagaraEditor/Private/Stateless/NiagaraDistributionPropertyCustomization.cpp
@@ -29,7 +29,6 @@ public:
 	FNiagaraDistributionAdapter()
 		: PropertyHandle(nullptr)
 		, OwnerObjectWeak(nullptr)
-		, SourceDistribution(nullptr)
 		, SourceNumChannels(INDEX_NONE)
 	{
 	}
@@ -37,7 +36,7 @@ public:
 	void Initialize(TSharedPtr<IPropertyHandle> InPropertyHandle, UObject* InOwningObject, FNiagaraDistributionBase* InDistribution, int32 InNumChannels)
 	{
 		PropertyHandle = InPropertyHandle;
-		OwnerObjectWeak = TWeakObjectPtr<UObject>(InOwningObject);
+		OwnerObjectWeak = TWeakObjectPtr<>(InOwningObject);
 		SourceDistribution = InDistribution;
 		SourceNumChannels = InNumChannels;
 		if (IsValid())
@@ -49,14 +48,23 @@ public:
 			const FName DisableRangeDistributionName("DisableRangeDistribution");
 			const FName DisplayAsColorDistributionName("DisplayAsColorDistribution");
 			bAllowBinding = InDistribution->AllowBinding() && InDistribution->GetBindingTypeDef().IsValid() && (PropertyHandle ? PropertyHandle->HasMetaData(DisableBindingDistributionName) == false : false);
-			bAllowUniform = PropertyHandle ? PropertyHandle->HasMetaData(DisableUniformDistributionName) == false : false;
-			bAllowNonUniform = PropertyHandle ? PropertyHandle->HasMetaData(DisableNonUniformDistributionName) == false : false;
-			bAllowRange = PropertyHandle ? PropertyHandle->HasMetaData(DisableRangeDistributionName) == false : false;
-			bAllowCurves = InDistribution->AllowCurves() && (PropertyHandle ? PropertyHandle->HasMetaData(DisableCurveDistributionName) == false : false);
-			bDisplayAsColor = InDistribution->DisplayAsColor() || (InNumChannels >= 3 && (PropertyHandle ? PropertyHandle->HasMetaData(DisplayAsColorDistributionName) : false));
+			bAllowUniform = PropertyHandle->HasMetaData(DisableUniformDistributionName) == false;
+			bAllowNonUniform = PropertyHandle->HasMetaData(DisableNonUniformDistributionName) == false;
+			bAllowRange = PropertyHandle->HasMetaData(DisableRangeDistributionName) == false;
+			bAllowCurves = InDistribution->AllowCurves() && PropertyHandle->HasMetaData(DisableCurveDistributionName) == false;
+			bDisplayAsColor = InDistribution->DisplayAsColor() || (InNumChannels >= 3 && PropertyHandle->HasMetaData(DisplayAsColorDistributionName));
 			PropertyHandle->SetOnPropertyValueChanged(FSimpleDelegate::CreateSP(this->AsShared(), &FNiagaraDistributionAdapter::DistributionPropertyChanged));
 			EditorMode = GetDistributionEditorModeFromSourceMode(SourceNumChannels, bDisplayAsColor, SourceDistribution->Mode);
 			EditorNumChannels = FNiagaraDistributionEditorUtilities::IsUniform(EditorMode) ? 1 : SourceNumChannels;
+			WidgetCustomizationOptions = FNiagaraInputParameterCustomization::MakeFromProperty(InPropertyHandle);
+
+			static const FName UnitsName("Units");
+			if (PropertyHandle->HasMetaData(UnitsName))
+			{
+				FString UnitString = PropertyHandle->GetMetaData(UnitsName);
+				TOptional<EUnit> PropertyUnit = FUnitConversion::UnitFromString(*UnitString);
+				DisplayUnit = PropertyUnit.Get(EUnit::Unspecified);
+			}
 		}
 	}
 
@@ -116,6 +124,16 @@ public:
 		return FSlateColor(EStyleColor::AccentWhite);
 	}
 
+	virtual FNiagaraInputParameterCustomization GetWidgetCustomization() override
+	{
+		return WidgetCustomizationOptions;
+	}
+
+	virtual EUnit GetDisplayUnit() override
+	{
+		return DisplayUnit;	
+	}
+
 	virtual void GetSupportedDistributionModes(TArray<ENiagaraDistributionEditorMode>& OutSupportedModes) const override
 	{
 		if (bAllowBinding)
@@ -779,6 +797,8 @@ private:
 
 	ENiagaraDistributionEditorMode EditorMode = ENiagaraDistributionEditorMode::Constant;
 	int32 EditorNumChannels = 0;
+	FNiagaraInputParameterCustomization WidgetCustomizationOptions;
+	EUnit DisplayUnit = EUnit::Unspecified;
 
 	bool bContinuousTransactionPending = false;
 	TOptional<int32> ContinuousTransactionIndex;
diff --git a/Engine/Plugins/FX/Niagara/Source/NiagaraEditor/Private/TypeEditorUtilities/NiagaraFloatTypeEditorUtilities.cpp b/Engine/Plugins/FX/Niagara/Source/NiagaraEditor/Private/TypeEditorUtilities/NiagaraFloatTypeEditorUtilities.cpp
--- a/Engine/Plugins/FX/Niagara/Source/NiagaraEditor/Private/TypeEditorUtilities/NiagaraFloatTypeEditorUtilities.cpp
+++ b/Engine/Plugins/FX/Niagara/Source/NiagaraEditor/Private/TypeEditorUtilities/NiagaraFloatTypeEditorUtilities.cpp
@@ -16,227 +16,218 @@
 #include "Widgets/Input/SSlider.h"
 #include "Widgets/Input/SVolumeControl.h"
 
-class SNiagaraFloatParameterEditor : public SNiagaraParameterEditor
+void SNiagaraFloatParameterEditor::Construct(const FArguments& InArgs, EUnit DisplayUnit, const FNiagaraInputParameterCustomization& WidgetCustomization)
 {
-public:
-	SLATE_BEGIN_ARGS(SNiagaraFloatParameterEditor) { }
-	SLATE_END_ARGS();
+	ValueAttribute = InArgs._Value;
+	OnValueChangedEvent = InArgs._OnValueChanged;
+	OnBeginValueChangeEvent = InArgs._OnBeginValueChange;
+	OnEndValueChangeEvent = InArgs._OnEndValueChange;
+	SNiagaraParameterEditor::Construct(SNiagaraParameterEditor::FArguments()
+	                                   .MinimumDesiredWidth(DefaultInputSize)
+	                                   .MaximumDesiredWidth(DefaultInputSize));
 
-	void Construct(const FArguments& InArgs, EUnit DisplayUnit, const FNiagaraInputParameterCustomization& WidgetCustomization)
-	{
-		SNiagaraParameterEditor::Construct(SNiagaraParameterEditor::FArguments()
-			.MinimumDesiredWidth(DefaultInputSize)
-			.MaximumDesiredWidth(DefaultInputSize));
-
-		const UGraphEditorSettings* Settings = GetDefault<UGraphEditorSettings>();
-		const UNiagaraEditorSettings* NiagaraSettings = GetDefault<UNiagaraEditorSettings>();
+	const UGraphEditorSettings* Settings = GetDefault<UGraphEditorSettings>();
+	const UNiagaraEditorSettings* NiagaraSettings = GetDefault<UNiagaraEditorSettings>();
 
-		if (WidgetCustomization.WidgetType == ENiagaraInputWidgetType::Slider)
-		{
-			float MinValue = WidgetCustomization.bHasMinValue ? WidgetCustomization.MinValue : 0;
-			float MaxValue = WidgetCustomization.bHasMaxValue ? WidgetCustomization.MaxValue : 1;
-			ChildSlot
+	if (WidgetCustomization.WidgetType == ENiagaraInputWidgetType::Slider)
+	{
+		float MinValue = WidgetCustomization.bHasMinValue ? WidgetCustomization.MinValue : 0;
+		float MaxValue = WidgetCustomization.bHasMaxValue ? WidgetCustomization.MaxValue : 1;
+		ChildSlot
+		[
+			SNew(SHorizontalBox)
+			+ SHorizontalBox::Slot()
+			.VAlign(VAlign_Center)
+			.HAlign(HAlign_Center)
+			.AutoWidth()
 			[
-				SNew(SHorizontalBox)
-				+ SHorizontalBox::Slot()
-				.VAlign(VAlign_Center)
-				.HAlign(HAlign_Center)
-				.AutoWidth()
-				[
-					SNew(SImage)
-					.ColorAndOpacity(UEdGraphSchema_Niagara::GetTypeColor(FNiagaraTypeDefinition::GetFloatDef()))
-					.Image(FNiagaraEditorStyle::Get().GetBrush("NiagaraEditor.Module.TypeIconPill"))
-				]
-				+ SHorizontalBox::Slot().AutoWidth()
+				SNew(SImage)
+				.ColorAndOpacity(UEdGraphSchema_Niagara::GetTypeColor(FNiagaraTypeDefinition::GetFloatDef()))
+				.Image(FNiagaraEditorStyle::Get().GetBrush("NiagaraEditor.Module.TypeIconPill"))
+			]
+			+ SHorizontalBox::Slot().AutoWidth()
+			[
+				SNew(SBox).WidthOverride(100.0f)
 				[
-					SNew(SBox).WidthOverride(100.0f)
-					[
-						SNew(SSlider)
-						.MinValue(MinValue)
-						.MaxValue(MaxValue)
-						.Value(this, &SNiagaraFloatParameterEditor::GetSliderValue)
-						.OnValueChanged_Lambda([this, WidgetCustomization](float NewVal)
+					SNew(SSlider)
+					.MinValue(MinValue)
+					.MaxValue(MaxValue)
+					.Value(this, &SNiagaraFloatParameterEditor::GetSliderValue)
+					.OnValueChanged_Lambda([this, WidgetCustomization](float NewVal)
+					{
+						SliderValue = NewVal;
+						FloatValue = NewVal;
+						if (WidgetCustomization.bHasStepWidth && WidgetCustomization.StepWidth != 0)
 						{
-							SliderValue = NewVal;
-							FloatValue = NewVal;
-							if (WidgetCustomization.bHasStepWidth && WidgetCustomization.StepWidth != 0)
-							{
-								FloatValue = FMath::RoundToFloat(NewVal / WidgetCustomization.StepWidth) * WidgetCustomization.StepWidth;
-							}
-							ExecuteOnValueChanged();
-						})
-						.OnMouseCaptureBegin(this, &SNiagaraFloatParameterEditor::ExecuteOnBeginValueChange)
-						.OnMouseCaptureEnd(this, &SNiagaraFloatParameterEditor::ExecuteOnEndValueChange)
-					]
+							FloatValue = FMath::RoundToFloat(NewVal / WidgetCustomization.StepWidth) * WidgetCustomization.StepWidth;
+						}
+						ExecuteOnValueChanged();
+						OnValueChangedEvent.ExecuteIfBound(FloatValue);
+					})
+					.OnMouseCaptureBegin(this, &SNiagaraFloatParameterEditor::BeginSliderMovement)
+					.OnMouseCaptureEnd(this, &SNiagaraFloatParameterEditor::ExecuteOnEndValueChange)
 				]
-				+ SHorizontalBox::Slot().AutoWidth()
+			]
+			+ SHorizontalBox::Slot().AutoWidth()
+			[
+				SNew(SBox).WidthOverride(75.0f)
 				[
-					SNew(SBox).WidthOverride(75.0f)
-					[
-						SNew(SNumericEntryBox<float>)
-						.Font(FAppStyle::Get().GetFontStyle("PropertyWindow.NormalFont"))
-						.MinValue(MinValue)
-						.MaxValue(MaxValue)
-						.Value(this, &SNiagaraFloatParameterEditor::GetValue)
-						.OnValueChanged(this, &SNiagaraFloatParameterEditor::ValueChanged)
-						.OnValueCommitted(this, &SNiagaraFloatParameterEditor::ValueCommitted)
-						.TypeInterface(GetTypeInterface<float>(DisplayUnit))
-						.AllowSpin(false)
-						.Delta(WidgetCustomization.bHasStepWidth ? WidgetCustomization.StepWidth : 0)
-					]
+					SNew(SNumericEntryBox<float>)
+					.Font(FAppStyle::Get().GetFontStyle("PropertyWindow.NormalFont"))
+					.MinValue(MinValue)
+					.MaxValue(MaxValue)
+					.Value(this, &SNiagaraFloatParameterEditor::GetValue)
+					.OnValueChanged(this, &SNiagaraFloatParameterEditor::ValueChanged)
+					.OnValueCommitted(this, &SNiagaraFloatParameterEditor::ValueCommitted)
+					.TypeInterface(GetTypeInterface<float>(DisplayUnit))
+					.AllowSpin(false)
+					.Delta(WidgetCustomization.bHasStepWidth ? WidgetCustomization.StepWidth : 0)
 				]
-			];
-		}
-		else if (WidgetCustomization.WidgetType == ENiagaraInputWidgetType::Volume)
-		{
-			ChildSlot
-			[
-				SNew(SVolumeControl)
-				.Volume(this, &SNiagaraFloatParameterEditor::GetSliderValue)
-				.Muted(this, &SNiagaraFloatParameterEditor::IsMuted)
-				.OnVolumeChanged_Lambda([this](float NewVal)
-				{
-					SliderValue = NewVal;
-					FloatValue = NewVal;
-					ExecuteOnValueChanged();
-				})
-				.OnMuteChanged_Lambda([this](bool NewVal)
-				{
-					bMuted = NewVal;
-					FloatValue = bMuted ? 0 : SliderValue;
-					ExecuteOnValueChanged();
-				})
-			];
-		}
-		else if (WidgetCustomization.WidgetType == ENiagaraInputWidgetType::NumericDropdown && WidgetCustomization.InputDropdownValues.Num() > 0)
-		{
-			TArray<SNiagaraNumericDropDown<float>::FNamedValue> DropDownValues;
-			for (const FWidgetNamedInputValue& Value : WidgetCustomization.InputDropdownValues)
-			{
-				DropDownValues.Add(SNiagaraNumericDropDown<float>::FNamedValue(Value.Value, Value.DisplayName.IsEmpty() ? FText::AsNumber(Value.Value) : Value.DisplayName, Value.Tooltip));
-			}
-
-			ChildSlot
-			[
-				SNew(SNiagaraNumericDropDown<float>)
-				.DropDownValues(DropDownValues)
-				.bShowNamedValue(true)
-				.MinDesiredValueWidth(75)
-				.PillType(FNiagaraTypeDefinition::GetFloatDef())
-				.Value_Lambda([this]() 
-				{ 
-					return FloatValue;
-				})
-				.OnValueChanged_Lambda([this](float NewVal)
-				{
-					FloatValue = NewVal;
-					ExecuteOnValueChanged();
-				})
-			];
-		}
-		else
-		{
-			TOptional<float> MinValue;
-			TOptional<float> MaxValue;
-			if (WidgetCustomization.bHasMinValue)
+			]
+		];
+	}
+	else if (WidgetCustomization.WidgetType == ENiagaraInputWidgetType::Volume)
+	{
+		ChildSlot
+		[
+			SNew(SVolumeControl)
+			.Volume(this, &SNiagaraFloatParameterEditor::GetSliderValue)
+			.Muted(this, &SNiagaraFloatParameterEditor::IsMuted)
+			.OnVolumeChanged_Lambda([this](float NewVal)
 			{
-				MinValue = WidgetCustomization.MinValue;
-			}
-			if (WidgetCustomization.bHasMaxValue)
+				SliderValue = NewVal;
+				FloatValue = NewVal;
+				ExecuteOnValueChanged();
+				OnValueChangedEvent.ExecuteIfBound(NewVal);
+			})
+			.OnMuteChanged_Lambda([this](bool NewVal)
 			{
-				MaxValue = WidgetCustomization.MaxValue;
-			}
-			
-			ChildSlot
-			[
-				SNew(SNumericEntryBox<float>)
-				.Font(FAppStyle::Get().GetFontStyle("PropertyWindow.NormalFont"))
-				.MinValue(MinValue)
-				.MaxValue(MaxValue)
-				.MinSliderValue(MinValue)
-				.MaxSliderValue(MaxValue)
-				.Value(this, &SNiagaraFloatParameterEditor::GetValue)
-				.OnValueChanged(this, &SNiagaraFloatParameterEditor::ValueChanged)
-				.OnValueCommitted(this, &SNiagaraFloatParameterEditor::ValueCommitted)
-				.OnBeginSliderMovement(this, &SNiagaraFloatParameterEditor::BeginSliderMovement)
-				.OnEndSliderMovement(this, &SNiagaraFloatParameterEditor::EndSliderMovement)
-				.TypeInterface(GetTypeInterface<float>(DisplayUnit))
-				.AllowSpin(true)
-				.BroadcastValueChangesPerKey(!NiagaraSettings->GetUpdateStackValuesOnCommitOnly() && !WidgetCustomization.bBroadcastValueChangesOnCommitOnly)
-				.LabelPadding(FMargin(3))
-				.Delta(WidgetCustomization.bHasStepWidth ? WidgetCustomization.StepWidth : 0)
-				.LabelLocation(SNumericEntryBox<float>::ELabelLocation::Inside)
-				.Label()
-				[
-					SNumericEntryBox<float>::BuildNarrowColorLabel(Settings->FloatPinTypeColor)
-				]
-			];
-		}
+				bMuted = NewVal;
+				FloatValue = bMuted ? 0 : SliderValue;
+				ExecuteOnValueChanged();
+				OnValueChangedEvent.ExecuteIfBound(FloatValue);
+			})
+		];
 	}
+	else if (WidgetCustomization.WidgetType == ENiagaraInputWidgetType::NumericDropdown && WidgetCustomization.InputDropdownValues.Num() > 0)
+	{
+		TArray<SNiagaraNumericDropDown<float>::FNamedValue> DropDownValues;
+		for (const FWidgetNamedInputValue& Value : WidgetCustomization.InputDropdownValues)
+		{
+			DropDownValues.Add(SNiagaraNumericDropDown<float>::FNamedValue(Value.Value, Value.DisplayName.IsEmpty() ? FText::AsNumber(Value.Value) : Value.DisplayName, Value.Tooltip));
+		}
 
-	virtual void UpdateInternalValueFromStruct(TSharedRef<FStructOnScope> Struct) override
+		ChildSlot
+		[
+			SNew(SNiagaraNumericDropDown<float>)
+			.DropDownValues(DropDownValues)
+			.bShowNamedValue(true)
+			.MinDesiredValueWidth(75)
+			.PillType(FNiagaraTypeDefinition::GetFloatDef())
+			.Value_Lambda([this]() 
+			{ 
+				return FloatValue;
+			})
+			.OnValueChanged(this, &SNiagaraFloatParameterEditor::ValueChanged)
+		];
+	}
+	else
 	{
-		checkf(Struct->GetStruct() == FNiagaraTypeDefinition::GetFloatStruct(), TEXT("Struct type not supported."));
-		FloatValue = reinterpret_cast<FNiagaraFloat*>(Struct->GetStructMemory())->Value;
-		if (!bMuted)
+		TOptional<float> MinValue;
+		TOptional<float> MaxValue;
+		if (WidgetCustomization.bHasMinValue)
+		{
+			MinValue = WidgetCustomization.MinValue;
+		}
+		if (WidgetCustomization.bHasMaxValue)
 		{
-			SliderValue = FloatValue;
+			MaxValue = WidgetCustomization.MaxValue;
 		}
+		
+		ChildSlot
+		[
+			SNew(SNumericEntryBox<float>)
+			.Font(FAppStyle::Get().GetFontStyle("PropertyWindow.NormalFont"))
+			.MinValue(MinValue)
+			.MaxValue(MaxValue)
+			.MinSliderValue(MinValue)
+			.MaxSliderValue(MaxValue)
+			.Value(this, &SNiagaraFloatParameterEditor::GetValue)
+			.OnValueChanged(this, &SNiagaraFloatParameterEditor::ValueChanged)
+			.OnValueCommitted(this, &SNiagaraFloatParameterEditor::ValueCommitted)
+			.OnBeginSliderMovement(this, &SNiagaraFloatParameterEditor::BeginSliderMovement)
+			.OnEndSliderMovement(this, &SNiagaraFloatParameterEditor::EndSliderMovement)
+			.TypeInterface(GetTypeInterface<float>(DisplayUnit))
+			.AllowSpin(true)
+			.BroadcastValueChangesPerKey(!NiagaraSettings->GetUpdateStackValuesOnCommitOnly() && !WidgetCustomization.bBroadcastValueChangesOnCommitOnly)
+			.LabelPadding(FMargin(3))
+			.Delta(WidgetCustomization.bHasStepWidth ? WidgetCustomization.StepWidth : 0)
+			.LabelLocation(SNumericEntryBox<float>::ELabelLocation::Inside)
+			.Label()
+			[
+				SNumericEntryBox<float>::BuildNarrowColorLabel(Settings->FloatPinTypeColor)
+			]
+		];
 	}
+}
 
-	virtual void UpdateStructFromInternalValue(TSharedRef<FStructOnScope> Struct) override
+void SNiagaraFloatParameterEditor::UpdateInternalValueFromStruct(TSharedRef<FStructOnScope> Struct)
+{
+	checkf(Struct->GetStruct() == FNiagaraTypeDefinition::GetFloatStruct(), TEXT("Struct type not supported."));
+	FloatValue = reinterpret_cast<FNiagaraFloat*>(Struct->GetStructMemory())->Value;
+	if (!bMuted)
 	{
-		checkf(Struct->GetStruct() == FNiagaraTypeDefinition::GetFloatStruct(), TEXT("Struct type not supported."));
-		reinterpret_cast<FNiagaraFloat*>(Struct->GetStructMemory())->Value = FloatValue;
+		SliderValue = FloatValue;
 	}
+}
 
-	virtual bool CanChangeContinuously() const override { return true; }
+void SNiagaraFloatParameterEditor::UpdateStructFromInternalValue(TSharedRef<FStructOnScope> Struct)
+{
+	checkf(Struct->GetStruct() == FNiagaraTypeDefinition::GetFloatStruct(), TEXT("Struct type not supported."));
+	reinterpret_cast<FNiagaraFloat*>(Struct->GetStructMemory())->Value = FloatValue;
+}
 
-private:
-	void BeginSliderMovement()
-	{
-		ExecuteOnBeginValueChange();
-	}
+void SNiagaraFloatParameterEditor::BeginSliderMovement()
+{
+	ExecuteOnBeginValueChange();
+	OnBeginValueChangeEvent.ExecuteIfBound();
+}
 
-	void EndSliderMovement(float Value)
-	{
-		ExecuteOnEndValueChange();
-	}
+void SNiagaraFloatParameterEditor::EndSliderMovement(float Value)
+{
+	ExecuteOnEndValueChange();
+	OnEndValueChangeEvent.ExecuteIfBound(Value);
+}
 
-	TOptional<float> GetValue() const
-	{
-		return FloatValue;
-	}
+TOptional<float> SNiagaraFloatParameterEditor::GetValue() const
+{
+	return ValueAttribute.Get(FloatValue);
+}
 
-	float GetSliderValue() const
-	{
-		return SliderValue;
-	}
+float SNiagaraFloatParameterEditor::GetSliderValue() const
+{
+	return ValueAttribute.Get(SliderValue);
+}
 
-	bool IsMuted() const
-	{
-		return bMuted;
-	}
-	
-	void ValueChanged(float Value)
-	{
-		FloatValue = Value;
-		ExecuteOnValueChanged();
-	}
+bool SNiagaraFloatParameterEditor::IsMuted() const
+{
+	return bMuted;
+}
 
-	void ValueCommitted(float Value, ETextCommit::Type CommitInfo)
+void SNiagaraFloatParameterEditor::ValueChanged(float Value)
+{
+	FloatValue = Value;
+	ExecuteOnValueChanged();
+	OnValueChangedEvent.ExecuteIfBound(Value);
+}
+
+void SNiagaraFloatParameterEditor::ValueCommitted(float Value, ETextCommit::Type CommitInfo)
+{
+	if (CommitInfo == ETextCommit::OnEnter || CommitInfo == ETextCommit::OnUserMovedFocus)
 	{
-		if (CommitInfo == ETextCommit::OnEnter || CommitInfo == ETextCommit::OnUserMovedFocus)
-		{
-			ValueChanged(Value);
-		}
+		ValueChanged(Value);
 	}
-
-private:
-	float FloatValue = 0;
-	float SliderValue = 0;
-	bool bMuted = false;
-};
+}
 
 TSharedPtr<SNiagaraParameterEditor> FNiagaraEditorFloatTypeUtilities::CreateParameterEditor(const FNiagaraTypeDefinition& ParameterType, EUnit DisplayUnit, const FNiagaraInputParameterCustomization& WidgetCustomization) const
 {
diff --git a/Engine/Plugins/FX/Niagara/Source/NiagaraEditor/Private/TypeEditorUtilities/NiagaraFloatTypeEditorUtilities.h b/Engine/Plugins/FX/Niagara/Source/NiagaraEditor/Private/TypeEditorUtilities/NiagaraFloatTypeEditorUtilities.h
--- a/Engine/Plugins/FX/Niagara/Source/NiagaraEditor/Private/TypeEditorUtilities/NiagaraFloatTypeEditorUtilities.h
+++ b/Engine/Plugins/FX/Niagara/Source/NiagaraEditor/Private/TypeEditorUtilities/NiagaraFloatTypeEditorUtilities.h
@@ -4,8 +4,7 @@
 
 #include "INiagaraEditorTypeUtilities.h"
 #include "NiagaraTypes.h"
-
-class SNiagaraParameterEditor;
+#include "SNiagaraParameterEditor.h"
 
 /** Niagara editor utilities for the float type. */
 class FNiagaraEditorFloatTypeUtilities : public FNiagaraEditorTypeUtilities
@@ -19,7 +18,42 @@ public:
 	virtual bool SetValueFromPinDefaultString(const FString& StringValue, FNiagaraVariable& Variable) const override;
 	virtual FText GetSearchTextFromValue(const FNiagaraVariable& AllocatedVariable) const override;
 	virtual FText GetStackDisplayText(const FNiagaraVariable& Variable) const override { return FText::AsNumber(Variable.GetValue<float>()); };
-	virtual bool SupportsClipboardPortableValues() const { return true; }
+	virtual bool SupportsClipboardPortableValues() const override { return true; }
 	virtual bool TryUpdateClipboardPortableValueFromTypedValue(const FNiagaraTypeDefinition& InSourceType, const FNiagaraVariant& InSourceValue, FNiagaraClipboardPortableValue& OutTargetClipboardPortableValue) const override;
 	virtual bool TryUpdateTypedValueFromClipboardPortableValue(const FNiagaraClipboardPortableValue& InSourceClipboardPortableValue, const FNiagaraTypeDefinition& InTargetType, FNiagaraVariant& InTargetValue) const override;
+};
+
+class SNiagaraFloatParameterEditor : public SNiagaraParameterEditor
+{
+	DECLARE_DELEGATE_OneParam(FOnValueChanged, float);
+	
+public:
+	SLATE_BEGIN_ARGS(SNiagaraFloatParameterEditor) { }
+		SLATE_ATTRIBUTE(float, Value)
+		SLATE_EVENT( FOnValueChanged, OnValueChanged )
+		SLATE_EVENT( FSimpleDelegate, OnBeginValueChange )
+		SLATE_EVENT( FOnValueChanged, OnEndValueChange )
+	SLATE_END_ARGS();
+
+	void Construct(const FArguments& InArgs, EUnit DisplayUnit, const FNiagaraInputParameterCustomization& WidgetCustomization);
+	virtual void UpdateInternalValueFromStruct(TSharedRef<FStructOnScope> Struct) override;
+	virtual void UpdateStructFromInternalValue(TSharedRef<FStructOnScope> Struct) override;
+	virtual bool CanChangeContinuously() const override { return true; }
+private:
+	void BeginSliderMovement();
+	void EndSliderMovement(float Value);
+	TOptional<float> GetValue() const;
+	float GetSliderValue() const;
+	bool IsMuted() const;	
+	void ValueChanged(float Value);
+	void ValueCommitted(float Value, ETextCommit::Type CommitInfo);
+
+	float FloatValue = 0;
+	float SliderValue = 0;
+	bool bMuted = false;
+
+	TAttribute<float> ValueAttribute;
+	FOnValueChanged OnValueChangedEvent;
+	FSimpleDelegate OnBeginValueChangeEvent;
+	FOnValueChanged OnEndValueChangeEvent;
 };
\ No newline at end of file
diff --git a/Engine/Plugins/FX/Niagara/Source/NiagaraEditor/Private/TypeEditorUtilities/NiagaraIntegerTypeEditorUtilities.cpp b/Engine/Plugins/FX/Niagara/Source/NiagaraEditor/Private/TypeEditorUtilities/NiagaraIntegerTypeEditorUtilities.cpp
--- a/Engine/Plugins/FX/Niagara/Source/NiagaraEditor/Private/TypeEditorUtilities/NiagaraIntegerTypeEditorUtilities.cpp
+++ b/Engine/Plugins/FX/Niagara/Source/NiagaraEditor/Private/TypeEditorUtilities/NiagaraIntegerTypeEditorUtilities.cpp
@@ -15,223 +15,220 @@
 #include "Widgets/Input/SNumericEntryBox.h"
 #include "Widgets/Input/SSlider.h"
 
-class SNiagaraIntegerParameterEditor : public SNiagaraParameterEditor
+void SNiagaraIntegerParameterEditor::Construct(const FArguments& InArgs, EUnit DisplayUnit, const FNiagaraInputParameterCustomization& WidgetCustomization)
 {
-public:
-	SLATE_BEGIN_ARGS(SNiagaraIntegerParameterEditor) { }
-	SLATE_END_ARGS();
+	ValueAttribute = InArgs._Value;
+	OnValueChangedEvent = InArgs._OnValueChanged;
+	OnBeginValueChangeEvent = InArgs._OnBeginValueChange;
+	OnEndValueChangeEvent = InArgs._OnEndValueChange;
+	SNiagaraParameterEditor::Construct(SNiagaraParameterEditor::FArguments()
+		.MinimumDesiredWidth(DefaultInputSize)
+		.MaximumDesiredWidth(DefaultInputSize));
 
-	void Construct(const FArguments& InArgs, EUnit DisplayUnit, const FNiagaraInputParameterCustomization& WidgetCustomization)
-	{
-		SNiagaraParameterEditor::Construct(SNiagaraParameterEditor::FArguments()
-			.MinimumDesiredWidth(DefaultInputSize)
-			.MaximumDesiredWidth(DefaultInputSize));
-
-		const UGraphEditorSettings* Settings = GetDefault<UGraphEditorSettings>();
+	const UGraphEditorSettings* Settings = GetDefault<UGraphEditorSettings>();
 
-		if (WidgetCustomization.WidgetType == ENiagaraInputWidgetType::Slider)
-		{
-			float MinValue = WidgetCustomization.bHasMinValue ? WidgetCustomization.MinValue : 0;
-			float MaxValue = WidgetCustomization.bHasMaxValue ? WidgetCustomization.MaxValue : 1;
-			ChildSlot
+	if (WidgetCustomization.WidgetType == ENiagaraInputWidgetType::Slider)
+	{
+		float MinValue = WidgetCustomization.bHasMinValue ? WidgetCustomization.MinValue : 0;
+		float MaxValue = WidgetCustomization.bHasMaxValue ? WidgetCustomization.MaxValue : 1;
+		ChildSlot
+		[
+			SNew(SHorizontalBox)
+			+ SHorizontalBox::Slot()
+			.VAlign(VAlign_Center)
+			.HAlign(HAlign_Center)
+			.AutoWidth()
 			[
-				SNew(SHorizontalBox)
-				+ SHorizontalBox::Slot()
-				.VAlign(VAlign_Center)
-				.HAlign(HAlign_Center)
-				.AutoWidth()
-				[
-					SNew(SImage)
-					.ColorAndOpacity(UEdGraphSchema_Niagara::GetTypeColor(FNiagaraTypeDefinition::GetIntDef()))
-					.Image(FNiagaraEditorStyle::Get().GetBrush("NiagaraEditor.Module.TypeIconPill"))
-				]
-				+ SHorizontalBox::Slot().AutoWidth()
+				SNew(SImage)
+				.ColorAndOpacity(UEdGraphSchema_Niagara::GetTypeColor(FNiagaraTypeDefinition::GetIntDef()))
+				.Image(FNiagaraEditorStyle::Get().GetBrush("NiagaraEditor.Module.TypeIconPill"))
+			]
+			+ SHorizontalBox::Slot().AutoWidth()
+			[
+				SNew(SBox).WidthOverride(100.0f)
 				[
-					SNew(SBox).WidthOverride(100.0f)
-					[
-						SNew(SSlider)
-						.MinValue(MinValue)
-						.MaxValue(MaxValue)
-						.Value(this, &SNiagaraIntegerParameterEditor::GetSliderValue)
-						.OnValueChanged_Lambda([this, WidgetCustomization](float NewVal)
-						{
-							SliderValue = NewVal;
+					SNew(SSlider)
+					.MinValue(MinValue)
+					.MaxValue(MaxValue)
+					.Value(this, &SNiagaraIntegerParameterEditor::GetSliderValue)
+					.OnValueChanged_Lambda([this, WidgetCustomization](float NewVal)
+					{
+						SliderValue = NewVal;
 
-							// slider only works with float, so we round for the actual parameter value
-							int32 Resolution = WidgetCustomization.bHasStepWidth && WidgetCustomization.StepWidth >= 1 ? WidgetCustomization.StepWidth : 1;
-							IntValue = FMath::RoundToInt(NewVal / Resolution) * Resolution;
-							ExecuteOnValueChanged();
-						})
-						.OnMouseCaptureBegin(this, &SNiagaraIntegerParameterEditor::ExecuteOnBeginValueChange)
-						.OnMouseCaptureEnd(this, &SNiagaraIntegerParameterEditor::ExecuteOnEndValueChange)
-					]
+						// slider only works with float, so we round for the actual parameter value
+						int32 Resolution = WidgetCustomization.bHasStepWidth && WidgetCustomization.StepWidth >= 1 ? WidgetCustomization.StepWidth : 1;
+						IntValue = FMath::RoundToInt(NewVal / Resolution) * Resolution;
+						ExecuteOnValueChanged();
+						OnValueChangedEvent.ExecuteIfBound(IntValue);
+					})
+					.OnMouseCaptureBegin(this, &SNiagaraIntegerParameterEditor::ExecuteOnBeginValueChange)
+					.OnMouseCaptureEnd(this, &SNiagaraIntegerParameterEditor::ExecuteOnEndValueChange)
 				]
-				+ SHorizontalBox::Slot().AutoWidth()
+			]
+			+ SHorizontalBox::Slot().AutoWidth()
+			[
+				SNew(SBox).WidthOverride(75.0f)
 				[
-					SNew(SBox).WidthOverride(75.0f)
-					[
-						SNew(SNumericEntryBox<int32>)
-						.Font(FAppStyle::Get().GetFontStyle("PropertyWindow.NormalFont"))
-						.MinValue(MinValue)
-						.MaxValue(MaxValue)
-						.Value(this, &SNiagaraIntegerParameterEditor::GetValue)
-						.OnValueChanged(this, &SNiagaraIntegerParameterEditor::ValueChanged)
-						.OnValueCommitted(this, &SNiagaraIntegerParameterEditor::ValueCommitted)
-						.TypeInterface(GetTypeInterface<int32>(DisplayUnit))
-						.AllowSpin(false)
-						.Delta(WidgetCustomization.bHasStepWidth ? WidgetCustomization.StepWidth : 0)
-					]
+					SNew(SNumericEntryBox<int32>)
+					.Font(FAppStyle::Get().GetFontStyle("PropertyWindow.NormalFont"))
+					.MinValue(MinValue)
+					.MaxValue(MaxValue)
+					.Value(this, &SNiagaraIntegerParameterEditor::GetValue)
+					.OnValueChanged(this, &SNiagaraIntegerParameterEditor::ValueChanged)
+					.OnValueCommitted(this, &SNiagaraIntegerParameterEditor::ValueCommitted)
+					.TypeInterface(GetTypeInterface<int32>(DisplayUnit))
+					.AllowSpin(false)
+					.Delta(WidgetCustomization.bHasStepWidth ? WidgetCustomization.StepWidth : 0)
 				]
-			];
-		}
-		else if (WidgetCustomization.WidgetType == ENiagaraInputWidgetType::NumericDropdown && WidgetCustomization.InputDropdownValues.Num() > 0)
+			]
+		];
+	}
+	else if (WidgetCustomization.WidgetType == ENiagaraInputWidgetType::NumericDropdown && WidgetCustomization.InputDropdownValues.Num() > 0)
+	{
+		TArray<SNiagaraNumericDropDown<int32>::FNamedValue> DropDownValues;
+		for (const FWidgetNamedInputValue& Value : WidgetCustomization.InputDropdownValues)
 		{
-			TArray<SNiagaraNumericDropDown<int32>::FNamedValue> DropDownValues;
-			for (const FWidgetNamedInputValue& Value : WidgetCustomization.InputDropdownValues)
-			{
-				DropDownValues.Add(SNiagaraNumericDropDown<int32>::FNamedValue(Value.Value, Value.DisplayName.IsEmpty() ? FText::AsNumber(Value.Value) : Value.DisplayName, Value.Tooltip));
-			}
-
-			ChildSlot
-			[
-				SNew(SNiagaraNumericDropDown<int32>)
-				.DropDownValues(DropDownValues)
-				.bShowNamedValue(true)
-				.MinDesiredValueWidth(75)
-				.PillType(FNiagaraTypeDefinition::GetIntDef())
-				.Value_Lambda([this]() 
-				{ 
-					return IntValue;
-				})
-				.OnValueChanged_Lambda([this](int32 NewVal)
-				{
-					IntValue = NewVal;
-					ExecuteOnValueChanged();
-				})
-			];
+			DropDownValues.Add(SNiagaraNumericDropDown<int32>::FNamedValue(Value.Value, Value.DisplayName.IsEmpty() ? FText::AsNumber(Value.Value) : Value.DisplayName, Value.Tooltip));
 		}
-		else if(WidgetCustomization.WidgetType == ENiagaraInputWidgetType::EnumStyle && WidgetCustomization.EnumStyleDropdownValues.Num() > 0)
+
+		ChildSlot
+		[
+			SNew(SNiagaraNumericDropDown<int32>)
+			.DropDownValues(DropDownValues)
+			.bShowNamedValue(true)
+			.MinDesiredValueWidth(75)
+			.PillType(FNiagaraTypeDefinition::GetIntDef())
+			.Value_Lambda([this]() 
+			{ 
+				return IntValue;
+			})
+			.OnValueChanged_Lambda([this](int32 NewVal)
+			{
+				IntValue = NewVal;
+				ExecuteOnValueChanged();
+				OnValueChangedEvent.ExecuteIfBound(NewVal);
+			})
+		];
+	}
+	else if(WidgetCustomization.WidgetType == ENiagaraInputWidgetType::EnumStyle && WidgetCustomization.EnumStyleDropdownValues.Num() > 0)
+	{
+		TArray<SNiagaraNumericDropDown<int32>::FNamedValue> DropDownValues;
+		for(int32 EnumStyleValueIndex = 0; EnumStyleValueIndex < WidgetCustomization.EnumStyleDropdownValues.Num(); EnumStyleValueIndex++)
 		{
-			TArray<SNiagaraNumericDropDown<int32>::FNamedValue> DropDownValues;
-			for(int32 EnumStyleValueIndex = 0; EnumStyleValueIndex < WidgetCustomization.EnumStyleDropdownValues.Num(); EnumStyleValueIndex++)
+			FText DisplayName = WidgetCustomization.EnumStyleDropdownValues[EnumStyleValueIndex].DisplayName;
+			FText Tooltip = WidgetCustomization.EnumStyleDropdownValues[EnumStyleValueIndex].Tooltip;
+					
+			DropDownValues.Add(SNiagaraNumericDropDown<int32>::FNamedValue(EnumStyleValueIndex, DisplayName, Tooltip));
+		}
+		
+		ChildSlot
+		[
+			SNew(SNiagaraNumericDropDown<int32>)
+			.DropDownValues(DropDownValues)
+			.bAllowTyping(false)
+			.bShowNamedValue(true)
+			.MinDesiredValueWidth(75)
+			.PillType(FNiagaraTypeDefinition::GetIntDef())
+			.Value_Lambda([this]() 
+			{ 
+				return IntValue;
+			})
+			.OnValueChanged_Lambda([this](int32 NewVal)
 			{
-				FText DisplayName = WidgetCustomization.EnumStyleDropdownValues[EnumStyleValueIndex].DisplayName;
-				FText Tooltip = WidgetCustomization.EnumStyleDropdownValues[EnumStyleValueIndex].Tooltip;
-						
-				DropDownValues.Add(SNiagaraNumericDropDown<int32>::FNamedValue(EnumStyleValueIndex, DisplayName, Tooltip));
-			}
-			
-			ChildSlot
-			[
-				SNew(SNiagaraNumericDropDown<int32>)
-				.DropDownValues(DropDownValues)
-				.bAllowTyping(false)
-				.bShowNamedValue(true)
-				.MinDesiredValueWidth(75)
-				.PillType(FNiagaraTypeDefinition::GetIntDef())
-				.Value_Lambda([this]() 
-				{ 
-					return IntValue;
-				})
-				.OnValueChanged_Lambda([this](int32 NewVal)
-				{
-					IntValue = NewVal;
-					ExecuteOnValueChanged();
-				})
-			];
+				IntValue = NewVal;
+				ExecuteOnValueChanged();
+				OnValueChangedEvent.ExecuteIfBound(NewVal);
+			})
+		];
+	}
+	else
+	{
+		TOptional<int32> MinValue;
+		TOptional<int32> MaxValue;
+		if (WidgetCustomization.bHasMinValue)
+		{
+			MinValue = WidgetCustomization.MinValue;
 		}
-		else
+		if (WidgetCustomization.bHasMaxValue)
 		{
-			TOptional<int32> MinValue;
-			TOptional<int32> MaxValue;
-			if (WidgetCustomization.bHasMinValue)
-			{
-				MinValue = WidgetCustomization.MinValue;
-			}
-			if (WidgetCustomization.bHasMaxValue)
-			{
-				MaxValue = WidgetCustomization.MaxValue;
-			}
-			
-			ChildSlot
-			[
-				SNew(SNumericEntryBox<int32>)
-				.Font(FAppStyle::Get().GetFontStyle("PropertyWindow.NormalFont"))
-				.MinValue(MinValue)
-				.MaxValue(MaxValue)
-				.MinSliderValue(MinValue)
-				.MaxSliderValue(MaxValue)
-				.Value(this, &SNiagaraIntegerParameterEditor::GetValue)
-				.OnValueChanged(this, &SNiagaraIntegerParameterEditor::ValueChanged)
-				.OnValueCommitted(this, &SNiagaraIntegerParameterEditor::ValueCommitted)
-				.OnBeginSliderMovement(this, &SNiagaraIntegerParameterEditor::BeginSliderMovement)
-				.OnEndSliderMovement(this, &SNiagaraIntegerParameterEditor::EndSliderMovement)
-				.TypeInterface(GetTypeInterface<int32>(DisplayUnit))
-				.AllowSpin(true)
-				.BroadcastValueChangesPerKey(!GetDefault<UNiagaraEditorSettings>()->GetUpdateStackValuesOnCommitOnly() && !WidgetCustomization.bBroadcastValueChangesOnCommitOnly)
-				.LabelPadding(FMargin(3))
-				.LabelLocation(SNumericEntryBox<int32>::ELabelLocation::Inside)
-				.Label()
-				[
-					SNumericEntryBox<int32>::BuildNarrowColorLabel(Settings->IntPinTypeColor)
-				]
-			];
+			MaxValue = WidgetCustomization.MaxValue;
 		}
+		
+		ChildSlot
+		[
+			SNew(SNumericEntryBox<int32>)
+			.Font(FAppStyle::Get().GetFontStyle("PropertyWindow.NormalFont"))
+			.MinValue(MinValue)
+			.MaxValue(MaxValue)
+			.MinSliderValue(MinValue)
+			.MaxSliderValue(MaxValue)
+			.Value(this, &SNiagaraIntegerParameterEditor::GetValue)
+			.OnValueChanged(this, &SNiagaraIntegerParameterEditor::ValueChanged)
+			.OnValueCommitted(this, &SNiagaraIntegerParameterEditor::ValueCommitted)
+			.OnBeginSliderMovement(this, &SNiagaraIntegerParameterEditor::BeginSliderMovement)
+			.OnEndSliderMovement(this, &SNiagaraIntegerParameterEditor::EndSliderMovement)
+			.TypeInterface(GetTypeInterface<int32>(DisplayUnit))
+			.AllowSpin(true)
+			.BroadcastValueChangesPerKey(!GetDefault<UNiagaraEditorSettings>()->GetUpdateStackValuesOnCommitOnly() && !WidgetCustomization.bBroadcastValueChangesOnCommitOnly)
+			.LabelPadding(FMargin(3))
+			.LabelLocation(SNumericEntryBox<int32>::ELabelLocation::Inside)
+			.Label()
+			[
+				SNumericEntryBox<int32>::BuildNarrowColorLabel(Settings->IntPinTypeColor)
+			]
+		];
 	}
+}
 
-	virtual void UpdateInternalValueFromStruct(TSharedRef<FStructOnScope> Struct) override
-	{
-		checkf(Struct->GetStruct() == FNiagaraTypeDefinition::GetIntStruct(), TEXT("Struct type not supported."));
-		IntValue = reinterpret_cast<FNiagaraInt32*>(Struct->GetStructMemory())->Value;
-		SliderValue = IntValue;
-	}
+void SNiagaraIntegerParameterEditor::UpdateInternalValueFromStruct(TSharedRef<FStructOnScope> Struct)
+{
+	checkf(Struct->GetStruct() == FNiagaraTypeDefinition::GetIntStruct(), TEXT("Struct type not supported."));
+	IntValue = reinterpret_cast<FNiagaraInt32*>(Struct->GetStructMemory())->Value;
+	SliderValue = IntValue;
+}
 
-	virtual void UpdateStructFromInternalValue(TSharedRef<FStructOnScope> Struct) override
-	{
-		checkf(Struct->GetStruct() == FNiagaraTypeDefinition::GetIntStruct(), TEXT("Struct type not supported."));
-		reinterpret_cast<FNiagaraInt32*>(Struct->GetStructMemory())->Value = IntValue;
-	}
+void SNiagaraIntegerParameterEditor::UpdateStructFromInternalValue(TSharedRef<FStructOnScope> Struct)
+{
+	checkf(Struct->GetStruct() == FNiagaraTypeDefinition::GetIntStruct(), TEXT("Struct type not supported."));
+	reinterpret_cast<FNiagaraInt32*>(Struct->GetStructMemory())->Value = IntValue;
+}
 
-	virtual bool CanChangeContinuously() const override { return true; }
+void SNiagaraIntegerParameterEditor::BeginSliderMovement()
+{
+	ExecuteOnBeginValueChange();
+	OnBeginValueChangeEvent.ExecuteIfBound();
+}
 
-private:
-	void BeginSliderMovement()
-	{
-		ExecuteOnBeginValueChange();
-	}
+void SNiagaraIntegerParameterEditor::EndSliderMovement(int32 Value)
+{
+	ExecuteOnEndValueChange();
+	OnEndValueChangeEvent.ExecuteIfBound(Value);
+}
 
-	void EndSliderMovement(int32 Value)
-	{
-		ExecuteOnEndValueChange();
-	}
+TOptional<int32> SNiagaraIntegerParameterEditor::GetValue() const
+{
+	return ValueAttribute.Get(IntValue);
+}
 
-	TOptional<int32> GetValue() const
-	{
-		return IntValue;
-	}
-	
-	float GetSliderValue() const
-    {
-    	return SliderValue;
-    }
+float SNiagaraIntegerParameterEditor::GetSliderValue() const
+{
+    return ValueAttribute.Get(SliderValue);
+}
 
-	void ValueChanged(int32 Value)
-	{
-		IntValue = Value;
-		ExecuteOnValueChanged();
-	}
+void SNiagaraIntegerParameterEditor::ValueChanged(int32 Value)
+{
+	IntValue = Value;
+	ExecuteOnValueChanged();
+	OnValueChangedEvent.ExecuteIfBound(Value);
+}
 
-	void ValueCommitted(int32 Value, ETextCommit::Type CommitInfo)
+void SNiagaraIntegerParameterEditor::ValueCommitted(int32 Value, ETextCommit::Type CommitInfo)
+{
+	if (CommitInfo == ETextCommit::OnEnter || CommitInfo == ETextCommit::OnUserMovedFocus)
 	{
-		if (CommitInfo == ETextCommit::OnEnter || CommitInfo == ETextCommit::OnUserMovedFocus)
-		{
-			ValueChanged(Value);
-		}
+		ValueChanged(Value);
 	}
-
-	int32 IntValue = 0;
-	float SliderValue = 0;
-};
+}
 
 TSharedPtr<SNiagaraParameterEditor> FNiagaraEditorIntegerTypeUtilities::CreateParameterEditor(const FNiagaraTypeDefinition& ParameterType, EUnit DisplayUnit, const FNiagaraInputParameterCustomization& WidgetCustomization) const
 {
diff --git a/Engine/Plugins/FX/Niagara/Source/NiagaraEditor/Private/TypeEditorUtilities/NiagaraIntegerTypeEditorUtilities.h b/Engine/Plugins/FX/Niagara/Source/NiagaraEditor/Private/TypeEditorUtilities/NiagaraIntegerTypeEditorUtilities.h
--- a/Engine/Plugins/FX/Niagara/Source/NiagaraEditor/Private/TypeEditorUtilities/NiagaraIntegerTypeEditorUtilities.h
+++ b/Engine/Plugins/FX/Niagara/Source/NiagaraEditor/Private/TypeEditorUtilities/NiagaraIntegerTypeEditorUtilities.h
@@ -3,8 +3,7 @@
 #pragma once
 
 #include "INiagaraEditorTypeUtilities.h"
-
-class SNiagaraParameterEditor;
+#include "SNiagaraParameterEditor.h"
 
 class FNiagaraEditorIntegerTypeUtilities : public FNiagaraEditorTypeUtilities
 {
@@ -16,7 +15,45 @@ public:
 	virtual FString GetPinDefaultStringFromValue(const FNiagaraVariable& AllocatedVariable) const override;
 	virtual bool SetValueFromPinDefaultString(const FString& StringValue, FNiagaraVariable& Variable) const override;
 	virtual FText GetSearchTextFromValue(const FNiagaraVariable& AllocatedVariable) const override;
-	virtual bool SupportsClipboardPortableValues() const { return true; }
+	virtual bool SupportsClipboardPortableValues() const override { return true; }
 	virtual bool TryUpdateClipboardPortableValueFromTypedValue(const FNiagaraTypeDefinition& InSourceType, const FNiagaraVariant& InSourceValue, FNiagaraClipboardPortableValue& InTargetClipboardPortableValue) const override;
 	virtual bool TryUpdateTypedValueFromClipboardPortableValue(const FNiagaraClipboardPortableValue& InSourceClipboardPortableValue, const FNiagaraTypeDefinition& InTargetType, FNiagaraVariant& InTargetValue) const override;
+};
+
+class SNiagaraIntegerParameterEditor : public SNiagaraParameterEditor
+{
+	DECLARE_DELEGATE_OneParam(FOnValueChanged, int32);
+	
+public:
+	SLATE_BEGIN_ARGS(SNiagaraIntegerParameterEditor) { }
+		SLATE_ATTRIBUTE(int32, Value)
+		SLATE_EVENT( FOnValueChanged, OnValueChanged )
+		SLATE_EVENT( FSimpleDelegate, OnBeginValueChange )
+		SLATE_EVENT( FOnValueChanged, OnEndValueChange )
+	SLATE_END_ARGS();
+
+	void Construct(const FArguments& InArgs, EUnit DisplayUnit, const FNiagaraInputParameterCustomization& WidgetCustomization);
+	virtual void UpdateInternalValueFromStruct(TSharedRef<FStructOnScope> Struct) override;
+	virtual void UpdateStructFromInternalValue(TSharedRef<FStructOnScope> Struct) override;
+	virtual bool CanChangeContinuously() const override { return true; }
+private:
+	void BeginSliderMovement();
+	void EndSliderMovement(int32 Value);
+	TOptional<int32> GetValue() const;
+	float GetSliderValue() const;
+	void ValueChanged(int32 Value);
+	void ValueCommitted(int32 Value, ETextCommit::Type CommitInfo);
+
+	int32 IntValue = 0;
+	float SliderValue = 0;
+
+	TAttribute<int32> ValueAttribute;
+	FOnValueChanged OnValueChangedEvent;
+	FSimpleDelegate OnBeginValueChangeEvent;
+	FOnValueChanged OnEndValueChangeEvent;
+
+
+
+
+	
 };
\ No newline at end of file
diff --git a/Engine/Plugins/FX/Niagara/Source/NiagaraEditor/Private/ViewModels/Stack/NiagaraStackStatelessEmitterSpawnGroup.cpp b/Engine/Plugins/FX/Niagara/Source/NiagaraEditor/Private/ViewModels/Stack/NiagaraStackStatelessEmitterSpawnGroup.cpp
--- a/Engine/Plugins/FX/Niagara/Source/NiagaraEditor/Private/ViewModels/Stack/NiagaraStackStatelessEmitterSpawnGroup.cpp
+++ b/Engine/Plugins/FX/Niagara/Source/NiagaraEditor/Private/ViewModels/Stack/NiagaraStackStatelessEmitterSpawnGroup.cpp
@@ -9,6 +9,7 @@
 #include "Stateless/NiagaraDistributionPropertyCustomization.h"
 #include "Stateless/NiagaraDistributionIntPropertyCustomization.h"
 #include "Stateless/NiagaraStatelessEmitter.h"
+#include "Stateless/NiagaraStatelessModule.h"
 #include "Styling/AppStyle.h"
 #include "ViewModels/NiagaraSystemSelectionViewModel.h"
 #include "ViewModels/NiagaraSystemViewModel.h"
@@ -227,15 +228,16 @@ void UNiagaraStackStatelessEmitterSpawnItem::RefreshChildrenInternal(const TArra
 	FNiagaraStatelessSpawnInfo* SpawnInfo = GetSpawnInfo();
 	if (SpawnInfo != nullptr)
 	{
-		if (SpawnInfoStructOnScope.IsValid() == false || SpawnInfoStructOnScope->GetStructMemory() != (uint8*)SpawnInfo)
+		uint8* SpawnInfoPtr = reinterpret_cast<uint8*>(SpawnInfo);
+		if (SpawnInfoStructOnScope.IsValid() == false || SpawnInfoStructOnScope->GetStructMemory() != SpawnInfoPtr)
 		{
-			SpawnInfoStructOnScope = MakeShared<FStructOnScope>(FNiagaraStatelessSpawnInfo::StaticStruct(), (uint8*)SpawnInfo);
+			SpawnInfoStructOnScope = MakeShared<FStructOnScope>(FNiagaraStatelessSpawnInfo::StaticStruct(), SpawnInfoPtr);
 		}
 
 		UNiagaraStackObject* SpawnInfoObject = SpawnInfoObjectWeak.Get();
 		UObject* StatelessEmitterObject = StatelessEmitterWeak.Get();
 		if (SpawnInfoObject == nullptr || SpawnInfoObject->GetObject() != StatelessEmitterObject ||
-			SpawnInfoObject->GetDisplayedStruct().IsValid() == false || SpawnInfoObject->GetDisplayedStruct()->GetStructMemory() != (uint8*)SpawnInfo)
+			SpawnInfoObject->GetDisplayedStruct().IsValid() == false || SpawnInfoObject->GetDisplayedStruct()->GetStructMemory() != SpawnInfoPtr)
 		{
 			bool bIsInTopLevelStruct = true;
 			bool bHideTopLevelCategories = true;
diff --git a/Engine/Plugins/FX/Niagara/Source/NiagaraEditor/Private/Widgets/INiagaraDistributionAdapter.h b/Engine/Plugins/FX/Niagara/Source/NiagaraEditor/Private/Widgets/INiagaraDistributionAdapter.h
--- a/Engine/Plugins/FX/Niagara/Source/NiagaraEditor/Private/Widgets/INiagaraDistributionAdapter.h
+++ b/Engine/Plugins/FX/Niagara/Source/NiagaraEditor/Private/Widgets/INiagaraDistributionAdapter.h
@@ -7,6 +7,7 @@
 #include "Internationalization/Text.h"
 #include "Styling/SlateColor.h"
 #include "NiagaraCommon.h"
+#include "NiagaraVariableMetaData.h"
 
 struct FRichCurve;
 
@@ -70,6 +71,12 @@ public:
 	/** A delegate which is called whenever the distribution mode has changed. */
 	virtual FSimpleMulticastDelegate& OnDistributionEditorModeChanged() = 0;
 
+	/** Returns the widget customization options defined by the property metadata. */
+	virtual FNiagaraInputParameterCustomization GetWidgetCustomization() = 0;
+
+	/** If specified, the unit to display for the property */
+	virtual EUnit GetDisplayUnit() = 0;
+
 	/** 
 		Gets a constant or range value by channel index and value index. 
 		@param ChannelIndex the index of the channel to get.
diff --git a/Engine/Plugins/FX/Niagara/Source/NiagaraEditor/Private/Widgets/SNiagaraDistributionEditor.cpp b/Engine/Plugins/FX/Niagara/Source/NiagaraEditor/Private/Widgets/SNiagaraDistributionEditor.cpp
--- a/Engine/Plugins/FX/Niagara/Source/NiagaraEditor/Private/Widgets/SNiagaraDistributionEditor.cpp
+++ b/Engine/Plugins/FX/Niagara/Source/NiagaraEditor/Private/Widgets/SNiagaraDistributionEditor.cpp
@@ -2,10 +2,13 @@
 
 #include "Widgets/SNiagaraDistributionEditor.h"
 
+#include "NiagaraEditorModule.h"
 #include "Curves/CurveOwnerInterface.h"
 #include "Framework/MultiBox/MultiBoxBuilder.h"
 #include "NiagaraEditorStyle.h"
 #include "SColorGradientEditor.h"
+#include "SNiagaraParameterEditor.h"
+#include "TypeEditorUtilities/NiagaraFloatTypeEditorUtilities.h"
 #include "Widgets/Input/SComboButton.h"
 #include "Widgets/Input/SNumericEntryBox.h"
 #include "Widgets/Input/SVectorInputBox.h"
@@ -14,7 +17,6 @@
 #include "Widgets/NiagaraDistributionEditorUtilities.h"
 #include "Widgets/SNiagaraColorEditor.h"
 #include "Widgets/SNiagaraDistributionCurveEditor.h"
-#include "Widgets/SNiagaraExpandedToggle.h"
 #include "Widgets/SNiagaraParameterName.h"
 
 #define LOCTEXT_NAMESPACE "NiagaraDistributionEditor"
@@ -194,7 +196,7 @@ public:
 			TSharedPtr<SWidget> ValueWidget;
 			if (bIsUniform)
 			{
-				ValueWidget = ConstructFloatWidget(0, 0, FText());
+				ValueWidget = ConstructCustomizableFloatWidget(0, 0);
 			}
 			else if (bIsColor)
 			{
@@ -303,7 +305,8 @@ private:
 			: SNew(STextBlock)
 				.TextStyle(FNiagaraEditorStyle::Get(), "NiagaraEditor.ParameterText")
 				.Text(LabelText);
-
+		EUnit DisplayUnit = DistributionAdapter->GetDisplayUnit();
+		
 		return
 			SNew(SNumericEntryBox<float>)
 			.Font(FAppStyle::Get().GetFontStyle("PropertyWindow.NormalFont"))
@@ -318,6 +321,7 @@ private:
 			.MinSliderValue(TOptional<float>())
 			.MaxSliderValue(TOptional<float>())
 			.BroadcastValueChangesPerKey(false)
+			.TypeInterface(MakeShareable(new TNumericUnitTypeInterface<float>(DisplayUnit)))
 			.MinDesiredValueWidth(SNiagaraDistributionEditor::DefaultInputSize - 18)
 			.LabelVAlign(EVerticalAlignment::VAlign_Center)
 			.Label()
@@ -326,6 +330,19 @@ private:
 			];
 	}
 
+	TSharedRef<SWidget> ConstructCustomizableFloatWidget(int32 ChannelIndex, int32 ValueIndex)
+	{
+		FNiagaraEditorModule& NiagaraEditorModule = FNiagaraEditorModule::Get();
+		TSharedPtr<INiagaraEditorTypeUtilities> TypeUtilities = NiagaraEditorModule.GetTypeUtilities(FNiagaraTypeDefinition::GetFloatDef());
+		FNiagaraInputParameterCustomization WidgetCustomization = DistributionAdapter->GetWidgetCustomization();
+		EUnit DisplayUnit = DistributionAdapter->GetDisplayUnit();
+		return SNew(SNiagaraFloatParameterEditor, DisplayUnit, WidgetCustomization)
+			.Value(this, &SNiagaraDistributionValueEditor::GetValueFloat, ChannelIndex, ValueIndex)
+			.OnValueChanged(this, &SNiagaraDistributionValueEditor::ValueChanged, ChannelIndex, ValueIndex)
+			.OnBeginValueChange(this, &SNiagaraDistributionValueEditor::BeginValueChange)
+			.OnEndValueChange(this, &SNiagaraDistributionValueEditor::EndValueChange);
+	}
+
 	typedef SNumericVectorInputBox<float, UE::Math::TVector2<float>, 2> SNumericVectorInputBox2;
 	typedef SNumericVectorInputBox<float, UE::Math::TVector<float>, 3> SNumericVectorInputBox3;
 	typedef SNumericVectorInputBox<float, UE::Math::TVector4<float>, 4> SNumericVectorInputBox4;
@@ -422,6 +439,11 @@ private:
 		return DistributionAdapter->GetConstantOrRangeValue(ChannelIndex, ValueIndex);
 	}
 
+	float GetValueFloat(int32 ChannelIndex, int32 ValueIndex) const
+	{
+		return DistributionAdapter->GetConstantOrRangeValue(ChannelIndex, ValueIndex);
+	}
+
 	void ValueChanged(float Value, int32 ChannelIndex, int32 ValueIndex)
 	{
 		DistributionAdapter->SetConstantOrRangeValue(ChannelIndex, ValueIndex, Value);
