- Added space selection for position sampling offset
- Fix some issues with mesh orientation output under different spaces and local space emitters
- Added Inverse Quaternion function
- Added DebugHud data for socket data interface
- Fix quat multiply order in GPU transform utils
- Fix filtered & unfiltered socket lookup on GPU

#tests Compared static + skel mesh location against socket location to ensure orientations and positions are correct on CPU + GPU with the various sampling space modes

#virtualized

diff --git a/Engine/Plugins/FX/Niagara/Shaders/Private/NiagaraDataInterfaceSocketReaderTemplate.ush b/Engine/Plugins/FX/Niagara/Shaders/Private/NiagaraDataInterfaceSocketReaderTemplate.ush
--- a/Engine/Plugins/FX/Niagara/Shaders/Private/NiagaraDataInterfaceSocketReaderTemplate.ush
+++ b/Engine/Plugins/FX/Niagara/Shaders/Private/NiagaraDataInterfaceSocketReaderTemplate.ush
@@ -28,13 +28,13 @@ int SanitizeSocketIndex_{ParameterName}(int SocketIndex)
 int SanitizeFilteredSocketIndex_{ParameterName}(int SocketIndex)
 {
 	const int FilteredSocketIndex = SocketIndex >= 0 && SocketIndex < {ParameterName}_NumFilteredSockets ? (SocketIndex + 1) : 0;
-	return {ParameterName}_SocketData.Load(SocketIndex * 4);
+	return {ParameterName}_SocketData.Load(FilteredSocketIndex * 4);
 }
 
 int SanitizeUnfilteredSocketIndex_{ParameterName}(int SocketIndex)
 {
 	const int UnfilteredSocketIndex = SocketIndex >= 0 && SocketIndex < {ParameterName}_NumUnfilteredSockets ? (SocketIndex + 1 + {ParameterName}_NumFilteredSockets) : 0;
-	return {ParameterName}_SocketData.Load(SocketIndex * 4);
+	return {ParameterName}_SocketData.Load(UnfilteredSocketIndex * 4);
 }
 
 void GetSocketData_{ParameterName}(int SocketIndex, float Interp, out float3 OutPosition, out float4 OutRotation, out float3 OutScale, out float3 OutVelocity)
@@ -61,6 +61,13 @@ void GetSocketData_{ParameterName}(int SocketIndex, float Interp, out float3 Out
 
 ////////////////////////////////////////////////////
 
+void GetComponentToWorld_{ParameterName}(out float3 Translation, out float4 Rotation, out float3 Scale)
+{
+	Translation = {ParameterName}_ComponentToTranslatedWorld_Translation;
+	Rotation	= {ParameterName}_ComponentToTranslatedWorld_Rotation;
+	Scale		= {ParameterName}_ComponentToTranslatedWorld_Scale;
+}
+
 void IsValid_{ParameterName}(out bool bValid)
 {
 	bValid = {ParameterName}_IsDataValid != 0;
diff --git a/Engine/Plugins/FX/Niagara/Shaders/Private/NiagaraTransformUtils.ush b/Engine/Plugins/FX/Niagara/Shaders/Private/NiagaraTransformUtils.ush
--- a/Engine/Plugins/FX/Niagara/Shaders/Private/NiagaraTransformUtils.ush
+++ b/Engine/Plugins/FX/Niagara/Shaders/Private/NiagaraTransformUtils.ush
@@ -33,7 +33,7 @@ FNiagaraTransform MultiplyTransform(in FNiagaraTransform A, in FNiagaraTransform
 {
 	FNiagaraTransform R;
 	R.Translation	= RotateVectorByQuat(A.Translation * B.Scale, B.Rotation) + B.Translation;
-	R.Rotation		= MultiplyQuat(A.Rotation, B.Rotation);
+	R.Rotation		= MultiplyQuat(B.Rotation, A.Rotation);
 	R.Scale			= A.Scale * B.Scale;
 	return R;
 }
diff --git a/Engine/Plugins/FX/Niagara/Source/Niagara/Internal/DataInterface/NiagaraDataInterfaceSocketReader.h b/Engine/Plugins/FX/Niagara/Source/Niagara/Internal/DataInterface/NiagaraDataInterfaceSocketReader.h
--- a/Engine/Plugins/FX/Niagara/Source/Niagara/Internal/DataInterface/NiagaraDataInterfaceSocketReader.h
+++ b/Engine/Plugins/FX/Niagara/Source/Niagara/Internal/DataInterface/NiagaraDataInterfaceSocketReader.h
@@ -130,6 +130,10 @@ public:
 	NIAGARA_API virtual bool CopyToInternal(UNiagaraDataInterface* Destination) const override;
 
 	virtual bool HasPreSimulateTick() const override { return true; }
+
+#if WITH_NIAGARA_DEBUGGER
+	virtual void DrawDebugHud(FNDIDrawDebugHudContext& DebugHudContext) const override;
+#endif
 	//UNiagaraDataInterface Interface
 
 #if WITH_EDITORONLY_DATA
diff --git a/Engine/Plugins/FX/Niagara/Source/Niagara/Private/DataInterface/NiagaraDataInterfaceSocketReader.cpp b/Engine/Plugins/FX/Niagara/Source/Niagara/Private/DataInterface/NiagaraDataInterfaceSocketReader.cpp
--- a/Engine/Plugins/FX/Niagara/Source/Niagara/Private/DataInterface/NiagaraDataInterfaceSocketReader.cpp
+++ b/Engine/Plugins/FX/Niagara/Source/Niagara/Private/DataInterface/NiagaraDataInterfaceSocketReader.cpp
@@ -11,6 +11,7 @@
 #include "NiagaraWorldManager.h"
 
 #include "Components/SceneComponent.h"
+#include "Engine/Canvas.h"
 #include "Engine/SkeletalMesh.h"
 #include "Engine/SkeletalMeshSocket.h"
 #include "Engine/StaticMesh.h"
@@ -31,6 +32,8 @@ namespace NDISocketReaderLocal
 
 	static const FName	IsValidName("IsValid");
 
+	static const FName	GetComponentToWorldName("GetComponentToWorld");
+
 	static const FName	GetSocketCountName("GetSocketCount");
 	static const FName	GetFilteredSocketCountName("GetFilteredSocketCount");
 	static const FName	GetUnfilteredSocketCountName("GetUnfilteredSocketCount");
@@ -506,6 +509,24 @@ namespace NDISocketReaderLocal
 		}
 	}
 
+	void VMGetComponentToWorld(FVectorVMExternalFunctionContext& Context)
+	{
+		VectorVM::FUserPtrHandler<FInstanceData_GameThread> InstanceData(Context);
+		FNDIOutputParam<FVector3f>	OutPosition(Context);
+		FNDIOutputParam<FQuat4f>	OutRotation(Context);
+		FNDIOutputParam<FVector3f>	OutScale(Context);
+
+		const FVector3f Translation = InstanceData->ComponentToTranslatedWorld.GetTranslation();
+		const FQuat4f	Rotation	= InstanceData->ComponentToTranslatedWorld.GetRotation();
+		const FVector3f	Scale		= InstanceData->ComponentToTranslatedWorld.GetScale3D();
+		for (int32 i = 0; i < Context.GetNumInstances(); ++i)
+		{
+			OutPosition.SetAndAdvance(Translation);
+			OutRotation.SetAndAdvance(Rotation);
+			OutScale.SetAndAdvance(Scale);
+		}
+	}
+
 	template<ESocketReadType ReadType>
 	void VMGetSocketCount(FVectorVMExternalFunctionContext& Context)
 	{
@@ -558,7 +579,7 @@ namespace NDISocketReaderLocal
 
 				case ESocketReadType::Unfiltered:
 				{
-					const int32 UnfilteredSockedIndex = SocketIndex >= 0 && SocketIndex < InstanceData->NumFilteredSockets ? (SocketIndex + 1 + InstanceData->NumFilteredSockets) : 0;
+					const int32 UnfilteredSockedIndex = SocketIndex >= 0 && SocketIndex < InstanceData->NumUnfilteredSockets ? (SocketIndex + 1 + InstanceData->NumFilteredSockets) : 0;
 					SocketIndex = InstanceData->SocketFilterUnfilteredIndex[UnfilteredSockedIndex];
 					break;
 				}
@@ -593,6 +614,7 @@ namespace NDISocketReaderLocal
 	static const TMap<FName, FVMExternalFunction> VMFunctionBindings =
 	{
 		{IsValidName,									FVMExternalFunction::CreateStatic(VMIsValid)},
+		{GetComponentToWorldName,						FVMExternalFunction::CreateStatic(VMGetComponentToWorld)},
 		{GetSocketCountName,							FVMExternalFunction::CreateStatic(VMGetSocketCount<ESocketReadType::Default>)},
 		{GetFilteredSocketCountName,					FVMExternalFunction::CreateStatic(VMGetSocketCount<ESocketReadType::Filtered>)},
 		{GetUnfilteredSocketCountName,					FVMExternalFunction::CreateStatic(VMGetSocketCount<ESocketReadType::Unfiltered>)},
@@ -641,6 +663,13 @@ void UNiagaraDataInterfaceSocketReader::GetFunctionsInternal(TArray<FNiagaraFunc
 		FunctionSignature.Name = IsValidName;
 		FunctionSignature.Outputs.Emplace(FNiagaraTypeDefinition::GetBoolDef(), TEXT("IsValid"));
 	}
+	{
+		FNiagaraFunctionSignature& FunctionSignature = OutFunctions.Add_GetRef(ImmutableSig);
+		FunctionSignature.Name = GetComponentToWorldName;
+		FunctionSignature.Outputs.Emplace(FNiagaraTypeDefinition::GetVec3Def(), TEXT("Translation"));
+		FunctionSignature.Outputs.Emplace(FNiagaraTypeDefinition::GetQuatDef(), TEXT("Rotation"));
+		FunctionSignature.Outputs.Emplace(FNiagaraTypeDefinition::GetVec3Def(), TEXT("Scale"));
+	}
 	{
 		FNiagaraFunctionSignature FunctionSignature = ImmutableSig;
 		FunctionSignature.Outputs.Emplace(FNiagaraTypeDefinition::GetIntDef(), TEXT("Count"));
@@ -855,6 +884,39 @@ bool UNiagaraDataInterfaceSocketReader::CopyToInternal(UNiagaraDataInterface* De
 	return true;
 }
 
+#if WITH_NIAGARA_DEBUGGER
+void UNiagaraDataInterfaceSocketReader::DrawDebugHud(FNDIDrawDebugHudContext& DebugHudContext) const
+{
+	using namespace NDISocketReaderLocal;
+
+	const FInstanceData_GameThread* InstanceData_GT = DebugHudContext.GetSystemInstance()->FindTypedDataInterfaceInstanceData<FInstanceData_GameThread>(this);
+	if (InstanceData_GT == nullptr)
+	{
+		return;
+	}
+
+	UObject* ResolvedObject = InstanceData_GT->ResolvedObject.Get();
+	DebugHudContext.GetOutputString().Appendf(TEXT("ResolvedObject(%s)"), *GetNameSafe(ResolvedObject));
+
+	UCanvas* Canvas = DebugHudContext.GetCanvas();
+	if (DebugHudContext.IsVerbose() && ResolvedObject && Canvas)
+	{
+		for ( const FTransform3f SocketTransform : InstanceData_GT->SocketTransforms )
+		{
+			const FTransform WorldTransform = FTransform(SocketTransform) * InstanceData_GT->ComponentToWorld;
+			const FVector SocketLocation = WorldTransform.GetLocation();
+			const FVector ScreenPos = Canvas->Project(SocketLocation, false);
+			if (ScreenPos.Z <= 0.0f)
+			{
+				continue;
+			}
+
+			Canvas->Canvas->DrawNGon(FVector2D(ScreenPos), FColor::Red, 8, 4.0f);
+		}
+	}
+}
+#endif
+
 #if WITH_EDITORONLY_DATA
 TArray<FName> UNiagaraDataInterfaceSocketReader::GetEditorSocketNames() const
 {
