- When allocating deferred counts that don't exist in the buffer yet we need to flag them as UsedInstanceCounts otherwise the buffer can be deallocated next frame

#rnx

diff --git a/Engine/Plugins/FX/Niagara/Source/Niagara/Classes/NiagaraGPUInstanceCountManager.h b/Engine/Plugins/FX/Niagara/Source/Niagara/Classes/NiagaraGPUInstanceCountManager.h
--- a/Engine/Plugins/FX/Niagara/Source/Niagara/Classes/NiagaraGPUInstanceCountManager.h
+++ b/Engine/Plugins/FX/Niagara/Source/Niagara/Classes/NiagaraGPUInstanceCountManager.h
@@ -74,8 +74,11 @@ public:
 	/** Acquire an entry, this will either come from the free list or reallocate the buffer. */
 	uint32 AcquireOrAllocateEntry(FRHICommandListImmediate& RHICmdList);
 
-	/** Allocate an entry in the buffer, this could return an existing count or count yet to be allocated */
-	uint32 AllocateEntryDeferred();
+	/**
+	Allocate a count entry that might return an existing cleared count, or a count waiting to be allocated.
+	This method allows you to allocated counts concurrently, then commit them by calling AllocateDeferredCounts once you join back on the render thread.
+	*/
+	uint32 AllocateDeferredEntry();
 	/** Commit any deferred count allocations */
 	void AllocateDeferredCounts(FRHICommandListImmediate& RHICmdList);
 
diff --git a/Engine/Plugins/FX/Niagara/Source/Niagara/Private/NiagaraGPUInstanceCountManager.cpp b/Engine/Plugins/FX/Niagara/Source/Niagara/Private/NiagaraGPUInstanceCountManager.cpp
--- a/Engine/Plugins/FX/Niagara/Source/Niagara/Private/NiagaraGPUInstanceCountManager.cpp
+++ b/Engine/Plugins/FX/Niagara/Source/Niagara/Private/NiagaraGPUInstanceCountManager.cpp
@@ -189,7 +189,7 @@ uint32 FNiagaraGPUInstanceCountManager::AcquireOrAllocateEntry(FRHICommandListIm
 	return UsedInstanceCounts++;
 }
 
-uint32 FNiagaraGPUInstanceCountManager::AllocateEntryDeferred()
+uint32 FNiagaraGPUInstanceCountManager::AllocateDeferredEntry()
 {
 	check(IsInParallelRenderingThread());
 
@@ -214,6 +214,7 @@ void FNiagaraGPUInstanceCountManager::AllocateDeferredCounts(FRHICommandListImme
 		return;
 	}
 	ResizeBuffers(RHICmdList, AllocatedInstanceCounts + DeferredCountAllocations);
+	UsedInstanceCounts += DeferredCountAllocations;
 	DeferredCountAllocations = 0;
 }
 
diff --git a/Engine/Plugins/FX/Niagara/Source/Niagara/Private/NiagaraGpuComputeDispatch.cpp b/Engine/Plugins/FX/Niagara/Source/Niagara/Private/NiagaraGpuComputeDispatch.cpp
--- a/Engine/Plugins/FX/Niagara/Source/Niagara/Private/NiagaraGpuComputeDispatch.cpp
+++ b/Engine/Plugins/FX/Niagara/Source/Niagara/Private/NiagaraGpuComputeDispatch.cpp
@@ -1965,8 +1965,6 @@ void FNiagaraGpuComputeDispatch::PostInitViews(FRDGBuilder& GraphBuilder, TConst
 {
 	LLM_SCOPE(ELLMTag::Niagara);
 
-	GPUInstanceCounterManager.AllocateDeferredCounts(GraphBuilder.RHICmdList);
-
 	bAllowGPUParticleUpdate = bAllowGPUParticleUpdate && Views.Num() > 0 && Views[0].AllowGPUParticleUpdate();
 
 	if (bAllowGPUParticleUpdate && FNiagaraUtilities::AllowGPUParticles(GetShaderPlatform()))
diff --git a/Engine/Plugins/FX/Niagara/Source/Niagara/Private/Stateless/NiagaraStatelessComputeManager.cpp b/Engine/Plugins/FX/Niagara/Source/Niagara/Private/Stateless/NiagaraStatelessComputeManager.cpp
--- a/Engine/Plugins/FX/Niagara/Source/Niagara/Private/Stateless/NiagaraStatelessComputeManager.cpp
+++ b/Engine/Plugins/FX/Niagara/Source/Niagara/Private/Stateless/NiagaraStatelessComputeManager.cpp
@@ -244,7 +244,7 @@ FNiagaraDataBuffer* FNiagaraStatelessComputeManager::GetDataBuffer(FRHICommandLi
 			CacheData->DataBuffer->SetNumInstances(CacheData->ActiveParticles);
 
 			FNiagaraGPUInstanceCountManager& CountManager = ComputeInterface->GetGPUInstanceCounterManager();
-			const uint32 CountOffset = CountManager.AllocateEntryDeferred();
+			const uint32 CountOffset = CountManager.AllocateDeferredEntry();
 			CacheData->DataBuffer->SetGPUInstanceCountBufferOffset(CountOffset);
 			CountsToRelease.Add(CountOffset);
 
@@ -341,10 +341,18 @@ void FNiagaraStatelessComputeManager::OnPostPreRender(FRDGBuilder& GraphBuilder)
 	RDG_CSV_STAT_EXCLUSIVE_SCOPE(GraphBuilder, NiagaraStateless);
 	RDG_GPU_MASK_SCOPE(GraphBuilder, FRHIGPUMask::All());
 
+	// Ensure we allocate any deferred counts that we need
+	FNiagaraGpuComputeDispatchInterface* ComputeInterface = GetOwnerInterface();
+	{
+		FNiagaraGPUInstanceCountManager& CountManager = ComputeInterface->GetGPUInstanceCounterManager();
+		CountManager.AllocateDeferredCounts(GraphBuilder.RHICmdList);
+	}
+
+	// Execute dispatches
 	AddPass(
 		GraphBuilder,
 		RDG_EVENT_NAME("FNiagaraStatelessComputeManager::OnPostPreRender"),
-		[DataToGenerate=MoveTemp(GPUDataToGenerate), ComputeInterface=GetOwnerInterface()](FRHICommandListImmediate& RHICmdList)
+		[DataToGenerate=MoveTemp(GPUDataToGenerate), ComputeInterface](FRHICommandListImmediate& RHICmdList)
 		{
 			SCOPED_DRAW_EVENT(RHICmdList, FNiagaraStatelessComputeManager_OnPostPreRender);
 
