[Viewport Toolbar] Add new viewport toolbar to the Niagara Emitter editor

First submit of new viewport toolbar for Niagara Editor.

- Still missing viewport stats submenus shown in Figma design.

diff --git a/Engine/Plugins/FX/Niagara/Source/NiagaraEditor/Private/Widgets/NiagaraSystemEditorViewportToolbarSections.cpp b/Engine/Plugins/FX/Niagara/Source/NiagaraEditor/Private/Widgets/NiagaraSystemEditorViewportToolbarSections.cpp
new file mode 100644
--- /dev/null
+++ b/Engine/Plugins/FX/Niagara/Source/NiagaraEditor/Private/Widgets/NiagaraSystemEditorViewportToolbarSections.cpp
@@ -0,0 +1,174 @@
+﻿// Copyright Epic Games, Inc. All Rights Reserved.
+
+#include "NiagaraSystemEditorViewportToolbarSections.h"
+
+#include "EditorViewportCommands.h"
+#include "NiagaraEditorCommands.h"
+#include "SNiagaraSystemViewport.h"
+#include "ToolMenu.h"
+#include "Widgets/Input/SSpinBox.h"
+
+#define LOCTEXT_NAMESPACE "NiagaraSystemEditorViewportToolbarSections"
+
+TSharedRef<SWidget> UE::NiagaraSystemEditor::CreateShowMenuWidget(
+	const TSharedRef<SNiagaraSystemViewport>& InNiagaraSystemEditorViewport, bool bInShowViewportStatsToggle
+)
+{
+	InNiagaraSystemEditorViewport->OnFloatingButtonClicked();
+	TSharedRef<SEditorViewport> ViewportRef = InNiagaraSystemEditorViewport->GetViewportWidget();
+
+	constexpr bool bInShouldCloseWindowAfterMenuSelection = true;
+
+	// TODO move to UToolMenu*
+	FMenuBuilder ShowMenuBuilder(bInShouldCloseWindowAfterMenuSelection, ViewportRef->GetCommandList());
+	{
+		FNiagaraEditorCommands Commands = FNiagaraEditorCommands::Get();
+
+		if (bInShowViewportStatsToggle)
+		{
+			ShowMenuBuilder.AddSubMenu(
+				LOCTEXT("ViewportStatsSubMenu", "Viewport Stats"),
+				LOCTEXT("ViewportStatsSubMenu_ToolTip", "Viewport Stats settings"),
+				FNewMenuDelegate::CreateLambda(
+					[](FMenuBuilder& InMenuBuilder)
+					{
+						InMenuBuilder.BeginSection("CommonStats", LOCTEXT("CommonStatsLabel", "Common Stats"));
+
+						InMenuBuilder.AddMenuEntry(
+							FEditorViewportCommands::Get().ToggleStats,
+							"ViewportStats",
+							LOCTEXT("ViewportStatsLabel", "Show Stats")
+						);
+
+						InMenuBuilder.AddMenuEntry(
+							FEditorViewportCommands::Get().ToggleFPS, "ViewportFPS", LOCTEXT("ViewportFPSLabel", "Show FPS")
+						);
+
+						InMenuBuilder.AddMenuSeparator();
+						InMenuBuilder.EndSection();
+					}
+				),
+				false,
+				FSlateIcon()
+			);
+		}
+
+		ShowMenuBuilder.BeginSection("CommonShowFlags", LOCTEXT("CommonShowFlagsLabel", "Common Show Flags"));
+
+		ShowMenuBuilder.AddMenuEntry(Commands.TogglePreviewGrid);
+		ShowMenuBuilder.AddMenuSeparator();
+		ShowMenuBuilder.AddMenuEntry(Commands.ToggleEmitterExecutionOrder);
+		ShowMenuBuilder.AddMenuEntry(Commands.ToggleGpuTickInformation);
+		ShowMenuBuilder.AddMenuEntry(Commands.ToggleInstructionCounts);
+		ShowMenuBuilder.AddMenuEntry(Commands.ToggleMemoryInfo);
+		ShowMenuBuilder.AddMenuEntry(Commands.ToggleParticleCounts);
+		ShowMenuBuilder.AddMenuEntry(Commands.ToggleStatelessInfo);
+	}
+
+	return ShowMenuBuilder.MakeWidget();
+}
+
+FToolMenuEntry UE::NiagaraSystemEditor::CreateShowSubmenu(const TWeakPtr<SNiagaraSystemViewport>& InNiagaraSystemEditorViewport)
+{
+	return FToolMenuEntry::InitSubMenu(
+		"Show",
+		LOCTEXT("ShowSubmenuLabel", "Show"),
+		LOCTEXT("ShowSubmenuTooltip", "Show options"),
+		FNewToolMenuDelegate::CreateLambda(
+			[InNiagaraSystemEditorViewport](UToolMenu* Submenu) -> void
+			{
+				TSharedPtr<SNiagaraSystemViewport> Viewport = InNiagaraSystemEditorViewport.Pin();
+				if (!Viewport)
+				{
+					return;
+				}
+
+				FToolMenuSection& UnnamedSection = Submenu->FindOrAddSection("", LOCTEXT("UnnamedLabel", ""));
+
+				UnnamedSection.AddEntry(FToolMenuEntry::InitWidget(
+					"ShowMenuItems", UE::NiagaraSystemEditor::CreateShowMenuWidget(Viewport.ToSharedRef(), true), FText(), true
+				));
+			}
+		)
+	);
+}
+
+FToolMenuEntry UE::NiagaraSystemEditor::CreateSettingsSubmenu(const TWeakPtr<SNiagaraSystemViewport>& InNiagaraSystemEditorViewport
+)
+{
+	return FToolMenuEntry::InitSubMenu(
+		"Settings",
+		LOCTEXT("SettingsSubmenuLabel", "Settings"),
+		LOCTEXT("SettingsSubmenuTooltip", "Settings options"),
+		FNewToolMenuDelegate::CreateLambda(
+			[InNiagaraSystemEditorViewport](UToolMenu* Submenu) -> void
+			{
+				FToolMenuSection& ViewportControlsSection =
+					Submenu->FindOrAddSection("ViewportControls", LOCTEXT("ViewportControlsLabel", "Viewport Controls"));
+
+				ViewportControlsSection.AddSubMenu(
+					"MotionOptions",
+					LOCTEXT("MotionOptionsSubMenu", "Motion Options"),
+					LOCTEXT("MotionOptionsSubMenu_ToolTip", "Set Motion Options for the Niagara Component"),
+					FNewToolMenuDelegate::CreateLambda(
+						[InNiagaraSystemEditorViewport](UToolMenu* InMenu)
+						{
+							TSharedPtr<SNiagaraSystemViewport> Viewport = InNiagaraSystemEditorViewport.Pin();
+							if (!Viewport)
+							{
+								return;
+							}
+
+							InMenu->AddMenuEntry(
+								"MotionOptions",
+								FToolMenuEntry::InitWidget(
+									"MotionOptions",
+									UE::NiagaraSystemEditor::CreateMotionMenuWidget(Viewport.ToSharedRef()),
+									FText(),
+									true
+								)
+							);
+						}
+					),
+					false,
+					FSlateIcon()
+				);
+			}
+		)
+	);
+}
+
+TSharedRef<SWidget> UE::NiagaraSystemEditor::CreateMotionMenuWidget(const TSharedRef<SNiagaraSystemViewport>& InNiagaraSystemEditorViewport
+)
+{
+	constexpr bool bInShouldCloseWindowAfterMenuSelection = true;
+	FMenuBuilder MotionMenuBuilder(bInShouldCloseWindowAfterMenuSelection, InNiagaraSystemEditorViewport->GetCommandList());
+
+	MotionMenuBuilder.AddMenuEntry(FNiagaraEditorCommands::Get().ToggleMotion);
+
+	MotionMenuBuilder.AddWidget(
+		SNew(SSpinBox<float>)
+			.IsEnabled(InNiagaraSystemEditorViewport, &SNiagaraSystemViewport::IsMotionEnabled)
+			.Font(FAppStyle::GetFontStyle(TEXT("MenuItem.Font")))
+			.MinSliderValue(0.0f)
+			.MaxSliderValue(360.0f)
+			.Value(InNiagaraSystemEditorViewport, &SNiagaraSystemViewport::GetMotionRate)
+			.OnValueChanged(InNiagaraSystemEditorViewport, &SNiagaraSystemViewport::SetMotionRate),
+		LOCTEXT("MotionSpeed", "Motion Speed")
+	);
+
+	MotionMenuBuilder.AddWidget(
+		SNew(SSpinBox<float>)
+			.IsEnabled(InNiagaraSystemEditorViewport, &SNiagaraSystemViewport::IsMotionEnabled)
+			.Font(FAppStyle::GetFontStyle(TEXT("MenuItem.Font")))
+			.MinSliderValue(0.0f)
+			.MaxSliderValue(1000.0f)
+			.Value(InNiagaraSystemEditorViewport, &SNiagaraSystemViewport::GetMotionRadius)
+			.OnValueChanged(InNiagaraSystemEditorViewport, &SNiagaraSystemViewport::SetMotionRadius),
+		LOCTEXT("MotionRadius", "Motion Radius")
+	);
+
+	return MotionMenuBuilder.MakeWidget();
+}
+
+#undef LOCTEXT_NAMESPACE
diff --git a/Engine/Plugins/FX/Niagara/Source/NiagaraEditor/Private/Widgets/NiagaraSystemEditorViewportToolbarSections.h b/Engine/Plugins/FX/Niagara/Source/NiagaraEditor/Private/Widgets/NiagaraSystemEditorViewportToolbarSections.h
new file mode 100644
--- /dev/null
+++ b/Engine/Plugins/FX/Niagara/Source/NiagaraEditor/Private/Widgets/NiagaraSystemEditorViewportToolbarSections.h
@@ -0,0 +1,17 @@
+﻿// Copyright Epic Games, Inc. All Rights Reserved.
+
+#pragma once
+
+#include "Templates/SharedPointerFwd.h"
+#include "ToolMenuEntry.h"
+
+class SNiagaraSystemViewport;
+
+namespace UE::NiagaraSystemEditor
+{
+TSharedRef<SWidget> CreateShowMenuWidget(const TSharedRef<SNiagaraSystemViewport>& InNiagaraSystemEditorViewport, bool bInShowViewportStatsToggle = true);
+FToolMenuEntry CreateShowSubmenu(const TWeakPtr<SNiagaraSystemViewport>& InNiagaraSystemEditorViewport);
+
+TSharedRef<SWidget> CreateMotionMenuWidget(const TSharedRef<SNiagaraSystemViewport>& InNiagaraSystemEditorViewport);
+FToolMenuEntry CreateSettingsSubmenu(const TWeakPtr<SNiagaraSystemViewport>& InNiagaraSystemEditorViewport);
+}
diff --git a/Engine/Plugins/FX/Niagara/Source/NiagaraEditor/Private/Widgets/SNiagaraSystemViewport.cpp b/Engine/Plugins/FX/Niagara/Source/NiagaraEditor/Private/Widgets/SNiagaraSystemViewport.cpp
--- a/Engine/Plugins/FX/Niagara/Source/NiagaraEditor/Private/Widgets/SNiagaraSystemViewport.cpp
+++ b/Engine/Plugins/FX/Niagara/Source/NiagaraEditor/Private/Widgets/SNiagaraSystemViewport.cpp
@@ -23,6 +23,7 @@
 #include "NiagaraSystemInstanceController.h"
 #include "NiagaraSystemGpuComputeProxy.h"
 #include "NiagaraSystemEditorData.h"
+#include "NiagaraSystemEditorViewportToolbarSections.h"
 #include "SNiagaraSystemViewportToolBar.h"
 #include "UnrealEdGlobals.h"
 #include "Editor/EditorEngine.h"
@@ -32,6 +33,8 @@
 #include "Engine/TextureCube.h"
 #include "Slate/SceneViewport.h"
 #include "ThumbnailRendering/ThumbnailManager.h"
+#include "ToolMenus.h"
+#include "ViewportToolbar/UnrealEdViewportToolbar.h"
 #include "Widgets/Docking/SDockTab.h"
 #include "Widgets/Layout/SBox.h"
 
@@ -39,6 +42,30 @@
 
 class UNiagaraSystemEditorData;
 
+namespace UE::NiagaraSystemEditor::Private
+{
+void ExtendCameraSubmenu(FName InCameraOptionsSubmenuName)
+{
+	UToolMenu* const Submenu = UToolMenus::Get()->ExtendMenu(InCameraOptionsSubmenuName);
+
+	Submenu->AddDynamicSection(
+		"EditorCameraExtensionDynamicSection",
+		FNewToolMenuDelegate::CreateLambda(
+			[](UToolMenu* InDynamicMenu)
+			{
+				FToolMenuSection& PositioningSection = InDynamicMenu->FindOrAddSection(
+					"NavigationOptions", LOCTEXT("NavigationOptionsLabel", "Navigation Options")
+				);
+
+				PositioningSection.AddMenuEntry(FNiagaraEditorCommands::Get().ToggleOrbit);
+			}
+		)
+	);
+}
+
+}
+
+
 /** Viewport Client for the preview viewport */
 class FNiagaraSystemViewportClient : public FEditorViewportClient, public TSharedFromThis<FNiagaraSystemViewportClient>
 {
@@ -1145,10 +1172,159 @@ TSharedRef<FEditorViewportClient> SNiagaraSystemViewport::MakeEditorViewportClie
 
 TSharedPtr<SWidget> SNiagaraSystemViewport::MakeViewportToolbar()
 {
-	//return SNew(SNiagaraSystemViewportToolBar)
+	// return SNew(SNiagaraSystemViewportToolBar)
 	//.Viewport(SharedThis(this));
 	return SNew(SBox);
 }
+TSharedPtr<SWidget> SNiagaraSystemViewport::BuildViewportToolbar()
+{
+	// Register the viewport toolbar if another viewport hasn't already (it's shared).
+	const FName ViewportToolbarName = "NiagaraSystemEditor.ViewportToolbar";
+
+	if (!UToolMenus::Get()->IsMenuRegistered(ViewportToolbarName))
+	{
+		UToolMenu* const ViewportToolbarMenu = UToolMenus::Get()->RegisterMenu(
+			ViewportToolbarName, NAME_None /* parent */, EMultiBoxType::SlimHorizontalToolBar
+		);
+
+		// Add the left-aligned part of the viewport toolbar.
+		{
+			FToolMenuSection& LeftSection = ViewportToolbarMenu->FindOrAddSection("Left");
+		}
+
+		// Add the right-aligned part of the viewport toolbar.
+		{
+			// Add the submenus of this section as EToolMenuInsertType::Last to sort them after any
+			// default-positioned submenus external code might add.
+			FToolMenuSection& RightSection = ViewportToolbarMenu->FindOrAddSection("Right");
+			RightSection.Alignment = EToolMenuSectionAlign::Last;
+
+			// Add the "Camera" submenu.
+			{
+				const FName GrandParentSubmenuName = "UnrealEd.ViewportToolbar.Camera";
+				const FName ParentSubmenuName = "NiagaraSystemEditor.ViewportToolbar.Camera";
+				const FName SubmenuName = "NiagaraSystemEditor.ViewportToolbar.CameraOptions";
+
+				// Create our grandparent menu.
+				if (!UToolMenus::Get()->IsMenuRegistered(GrandParentSubmenuName))
+				{
+					UToolMenus::Get()->RegisterMenu(GrandParentSubmenuName);
+				}
+
+				// Create our parent menu.
+				if (!UToolMenus::Get()->IsMenuRegistered(ParentSubmenuName))
+				{
+					UToolMenus::Get()->RegisterMenu(ParentSubmenuName, GrandParentSubmenuName);
+				}
+
+				// Create our menu.
+				UToolMenus::Get()->RegisterMenu(SubmenuName, ParentSubmenuName);
+
+				UE::UnrealEd::ExtendCameraSubmenu(SubmenuName);
+				UE::NiagaraSystemEditor::Private::ExtendCameraSubmenu(SubmenuName);
+
+				FToolMenuEntry CameraSubmenu = UE::UnrealEd::CreateViewportToolbarCameraSubmenu();
+				CameraSubmenu.InsertPosition.Position = EToolMenuInsertType::First;
+				RightSection.AddEntry(CameraSubmenu);
+
+				// TODO: add navigation options --> Orbit Mode toggle
+			}
+
+			// Add the "View Modes" sub menu.
+			{
+				// Stay backward-compatible with the old viewport toolbar.
+				{
+					const FName ParentSubmenuName = "UnrealEd.ViewportToolbar.View";
+					// Create our parent menu.
+					if (!UToolMenus::Get()->IsMenuRegistered(ParentSubmenuName))
+					{
+						UToolMenus::Get()->RegisterMenu(ParentSubmenuName);
+					}
+
+					// Register our ToolMenu here first, before we create the submenu, so we can set our parent.
+					UToolMenus::Get()->RegisterMenu("NiagaraSystemEditor.ViewportToolbar.ViewModes", ParentSubmenuName);
+				}
+
+				FToolMenuEntry ViewModesSubmenu = UE::UnrealEd::CreateViewportToolbarViewModesSubmenu();
+				ViewModesSubmenu.InsertPosition.Position = EToolMenuInsertType::First;
+				RightSection.AddEntry(ViewModesSubmenu);
+
+				FToolMenuEntry PerformanceAndScalabilitySubmenu = UE::UnrealEd::CreatePerformanceAndScalabilitySubmenu();
+				PerformanceAndScalabilitySubmenu.InsertPosition.Position = EToolMenuInsertType::First;
+				RightSection.AddEntry(PerformanceAndScalabilitySubmenu);
+			}
+
+			// Add the Show submenu.
+			{
+				FToolMenuEntry ShowSubmenu = UE::NiagaraSystemEditor::CreateShowSubmenu(SharedThis(this));
+				ShowSubmenu.InsertPosition.Position = EToolMenuInsertType::First;
+				RightSection.AddEntry(ShowSubmenu);
+			}
+
+			// Add the Performance and Scalability submenu.
+			{
+				FToolMenuEntry PerformanceAndScalabilitySubmenu = UE::UnrealEd::CreatePerformanceAndScalabilitySubmenu();
+				PerformanceAndScalabilitySubmenu.InsertPosition.Position = EToolMenuInsertType::First;
+				RightSection.AddEntry(PerformanceAndScalabilitySubmenu);
+			}
+
+			// Add the Settings submenu.
+			{
+				FToolMenuEntry SettingsSubmenu = UE::NiagaraSystemEditor::CreateSettingsSubmenu(SharedThis(this));
+				SettingsSubmenu.InsertPosition.Position = EToolMenuInsertType::First;
+				RightSection.AddEntry(SettingsSubmenu);
+			}
+		}
+	}
+
+	FToolMenuContext ViewportToolbarContext;
+	{
+		ViewportToolbarContext.AppendCommandList(GetCommandList());
+
+		// Add the UnrealEd viewport toolbar context.
+		{
+			UUnrealEdViewportToolbarContext* const ContextObject = NewObject<UUnrealEdViewportToolbarContext>();
+			ContextObject->Viewport = SharedThis(this);
+
+			// Hook up our toolbar's filter for supported view modes.
+			ContextObject->IsViewModeSupported = UE::UnrealEd::IsViewModeSupportedDelegate::CreateLambda(
+				[](EViewModeIndex ViewModeIndex) -> bool
+				{
+					// This code is taken from SViewportToolBar::IsViewModeSupported
+					// SSCSEditorViewportToolBar does not override it, so we just take it as-is
+					// TODO: maybe create a private function for it, or move IsViewModeSupported to SEditorViewport
+
+					switch (ViewModeIndex)
+					{
+					case VMI_PrimitiveDistanceAccuracy:
+					case VMI_MaterialTextureScaleAccuracy:
+					case VMI_RequiredTextureResolution:
+						return false;
+					default:
+						return true;
+					}
+				}
+			);
+
+			ViewportToolbarContext.AddObject(ContextObject);
+		}
+	}
+
+	// clang-format off
+	const TSharedRef<SWidget> NewViewportToolbar = SNew(SBox)
+	.Visibility_Lambda(
+		[this]() -> EVisibility
+		{
+			return  UE::UnrealEd::ShowNewViewportToolbars() ? EVisibility::Visible: EVisibility::Collapsed;
+		}
+	)
+	[
+		UToolMenus::Get()->GenerateWidget(ViewportToolbarName, ViewportToolbarContext)
+	];
+	// clang-format on
+
+	return NewViewportToolbar;
+}
 
 EVisibility SNiagaraSystemViewport::OnGetViewportContentVisibility() const
 {
@@ -1186,11 +1362,23 @@ FText SNiagaraSystemViewport::GetViewportCompileStatusText() const
 
 void SNiagaraSystemViewport::PopulateViewportOverlays(TSharedRef<class SOverlay> Overlay)
 {
-	Overlay->AddSlot()
-	.VAlign(VAlign_Top)
-	[
-		SNew(SNiagaraSystemViewportToolBar, SharedThis(this)).Sequencer(Sequencer)
-	];
+	if (UE::UnrealEd::ShowOldViewportToolbars())
+	{
+		// clang-format off
+		Overlay->AddSlot()
+		.VAlign(VAlign_Top)
+		[
+			SNew(SNiagaraSystemViewportToolBar, SharedThis(this))
+			.Sequencer(Sequencer)
+			.Visibility_Lambda([]()
+				{
+					return UE::UnrealEd::ShowOldViewportToolbars() ? EVisibility::Visible : EVisibility::Collapsed;
+				}
+			)
+		];
+		// clang-format on
+	}
+
 	Overlay->AddSlot()
 	.VAlign(VAlign_Center)
 	.HAlign(HAlign_Center)
diff --git a/Engine/Plugins/FX/Niagara/Source/NiagaraEditor/Private/Widgets/SNiagaraSystemViewport.h b/Engine/Plugins/FX/Niagara/Source/NiagaraEditor/Private/Widgets/SNiagaraSystemViewport.h
--- a/Engine/Plugins/FX/Niagara/Source/NiagaraEditor/Private/Widgets/SNiagaraSystemViewport.h
+++ b/Engine/Plugins/FX/Niagara/Source/NiagaraEditor/Private/Widgets/SNiagaraSystemViewport.h
@@ -109,6 +109,7 @@ protected:
 	/** SEditorViewport interface */
 	virtual TSharedRef<FEditorViewportClient> MakeEditorViewportClient() override;
 	virtual TSharedPtr<SWidget> MakeViewportToolbar() override;
+	virtual TSharedPtr<SWidget> BuildViewportToolbar() override;
 	virtual EVisibility OnGetViewportContentVisibility() const override;
 	virtual void BindCommands() override;
 	virtual void OnFocusViewportToSelection() override;
