- Small fix to Socket Reader DI when searching for tags
- Fix engine test to work on cooked content

#rnx

#virtualized

diff --git a/Engine/Plugins/FX/Niagara/Source/Niagara/Private/DataInterface/NiagaraDataInterfaceSocketReader.cpp b/Engine/Plugins/FX/Niagara/Source/Niagara/Private/DataInterface/NiagaraDataInterfaceSocketReader.cpp
--- a/Engine/Plugins/FX/Niagara/Source/Niagara/Private/DataInterface/NiagaraDataInterfaceSocketReader.cpp
+++ b/Engine/Plugins/FX/Niagara/Source/Niagara/Private/DataInterface/NiagaraDataInterfaceSocketReader.cpp
@@ -246,7 +246,7 @@ namespace NDISocketReaderLocal
 				while (AttachComponent)
 				{
 					if ((!SocketDI->AttachComponentClass || AttachComponent->IsA(SocketDI->AttachComponentClass)) &&
-						(!SocketDI->AttachComponentTag.IsNone() && AttachComponent->ComponentHasTag(SocketDI->AttachComponentTag)) )
+						(SocketDI->AttachComponentTag.IsNone() || AttachComponent->ComponentHasTag(SocketDI->AttachComponentTag)) )
 					{
 						break;
 					}
