- Enable parallel GDME for Niagara

#rnx

diff --git a/Engine/Plugins/FX/Niagara/Source/Niagara/Private/NiagaraComponent.cpp b/Engine/Plugins/FX/Niagara/Source/Niagara/Private/NiagaraComponent.cpp
--- a/Engine/Plugins/FX/Niagara/Source/Niagara/Private/NiagaraComponent.cpp
+++ b/Engine/Plugins/FX/Niagara/Source/Niagara/Private/NiagaraComponent.cpp
@@ -45,7 +45,7 @@ DEFINE_LOG_CATEGORY(LogNiagara);
 
 DEFINE_RENDER_COMMAND_PIPE(NiagaraDynamicData, ERenderCommandPipeFlags::None);
 
-static bool GNiagaraParallelGDME = false;
+static bool GNiagaraParallelGDME = true;
 static FAutoConsoleVariableRef CVarNiagaraParallelGDME(
 	TEXT("fx.Niagara.ParallelGDME"),
 	GNiagaraParallelGDME,
