- Allow various emitter state & spawn info properties to be const / range / binding

#tests Test that the various properties on emitter state & spawn infos can be const / range / binding
#rnx

diff --git a/Engine/Plugins/FX/Niagara/Source/Niagara/Internal/Stateless/NiagaraStatelessDistribution.h b/Engine/Plugins/FX/Niagara/Source/Niagara/Internal/Stateless/NiagaraStatelessDistribution.h
--- a/Engine/Plugins/FX/Niagara/Source/Niagara/Internal/Stateless/NiagaraStatelessDistribution.h
+++ b/Engine/Plugins/FX/Niagara/Source/Niagara/Internal/Stateless/NiagaraStatelessDistribution.h
@@ -128,6 +128,16 @@ struct FNiagaraDistributionRangeFloat : public FNiagaraDistributionBase
 	NIAGARA_API virtual void UpdateValuesFromDistribution() override;
 	virtual FNiagaraTypeDefinition GetBindingTypeDef() const { return FNiagaraTypeDefinition::GetFloatDef(); }
 #endif
+	bool SerializeFromMismatchedTag(const struct FPropertyTag& Tag, FStructuredArchive::FSlot Slot);
+};
+
+template<>
+struct TStructOpsTypeTraits<FNiagaraDistributionRangeFloat> : public TStructOpsTypeTraitsBase2<FNiagaraDistributionRangeFloat>
+{
+	enum
+	{
+		WithStructuredSerializeFromMismatchedTag = true,
+	};
 };
 
 USTRUCT()
diff --git a/Engine/Plugins/FX/Niagara/Source/Niagara/Internal/Stateless/NiagaraStatelessSpawnInfo.h b/Engine/Plugins/FX/Niagara/Source/Niagara/Internal/Stateless/NiagaraStatelessSpawnInfo.h
--- a/Engine/Plugins/FX/Niagara/Source/Niagara/Internal/Stateless/NiagaraStatelessSpawnInfo.h
+++ b/Engine/Plugins/FX/Niagara/Source/Niagara/Internal/Stateless/NiagaraStatelessSpawnInfo.h
@@ -30,10 +30,10 @@ struct FNiagaraStatelessSpawnInfo
 	UPROPERTY(EditAnywhere, Category = "Spawn", meta = (EditConditionHides, EditCondition = "Type == ENiagaraStatelessSpawnInfoType::Burst"))
 	float SpawnTime = 0.0f;
 
-	UPROPERTY(EditAnywhere, Category = "Spawn", meta = (EditConditionHides, EditCondition = "Type == ENiagaraStatelessSpawnInfoType::Burst", ClampMin = "0", DisableBindingDistribution))
+	UPROPERTY(EditAnywhere, Category = "Spawn", meta = (EditConditionHides, EditCondition = "Type == ENiagaraStatelessSpawnInfoType::Burst", ClampMin = "0"))
 	FNiagaraDistributionRangeInt Amount = FNiagaraDistributionRangeInt(1);
 
-	UPROPERTY(EditAnywhere, Category = "Spawn", meta = (EditConditionHides, EditCondition = "Type == ENiagaraStatelessSpawnInfoType::Rate", ClampMin = "0.0", DisableBindingDistribution))
+	UPROPERTY(EditAnywhere, Category = "Spawn", meta = (EditConditionHides, EditCondition = "Type == ENiagaraStatelessSpawnInfoType::Rate", ClampMin = "0.0"))
 	FNiagaraDistributionRangeFloat Rate = FNiagaraDistributionRangeFloat(60.0f);
 
 	UPROPERTY(EditAnywhere, Category = "Spawn", meta = (InlineEditConditionToggle))
@@ -43,7 +43,7 @@ struct FNiagaraStatelessSpawnInfo
 	uint8 bSpawnProbabilityEnabled : 1 = false;
 
 	UPROPERTY(EditAnywhere, Category = "Spawn", meta = (EditCondition = "bSpawnProbabilityEnabled", ClampMin = "0.0", ClampMax = "1.0", UIMin = "0.0", UIMax = "1.0"))
-	float SpawnProbability = 1.0f;
+	FNiagaraDistributionRangeFloat SpawnProbability = FNiagaraDistributionRangeFloat(1.0f);
 
 	bool IsValid(TOptional<float> LoopDuration) const;
 };
diff --git a/Engine/Plugins/FX/Niagara/Source/Niagara/Private/Stateless/NiagaraStatelessDistribution.cpp b/Engine/Plugins/FX/Niagara/Source/Niagara/Private/Stateless/NiagaraStatelessDistribution.cpp
--- a/Engine/Plugins/FX/Niagara/Source/Niagara/Private/Stateless/NiagaraStatelessDistribution.cpp
+++ b/Engine/Plugins/FX/Niagara/Source/Niagara/Private/Stateless/NiagaraStatelessDistribution.cpp
@@ -442,6 +442,19 @@ FNiagaraStatelessRangeColor FNiagaraDistributionColor::CalculateRange(const FLin
 	return Range;
 }
 
+bool FNiagaraDistributionRangeFloat::SerializeFromMismatchedTag(const struct FPropertyTag& Tag, FStructuredArchive::FSlot Slot)
+{
+	if (Tag.Type == NAME_FloatProperty)
+	{
+		float Value;
+		Slot << Value;
+		*this = FNiagaraDistributionRangeFloat(Value);
+		return true;
+	}
+
+	return false;
+}
+
 #if WITH_EDITORONLY_DATA
 void FNiagaraDistributionBase::PostEditChangeProperty(UObject* OwnerObject, FPropertyChangedEvent& PropertyChangedEvent)
 {
diff --git a/Engine/Plugins/FX/Niagara/Source/Niagara/Private/Stateless/NiagaraStatelessEmitter.cpp b/Engine/Plugins/FX/Niagara/Source/Niagara/Private/Stateless/NiagaraStatelessEmitter.cpp
--- a/Engine/Plugins/FX/Niagara/Source/Niagara/Private/Stateless/NiagaraStatelessEmitter.cpp
+++ b/Engine/Plugins/FX/Niagara/Source/Niagara/Private/Stateless/NiagaraStatelessEmitter.cpp
@@ -387,7 +387,23 @@ void UNiagaraStatelessEmitter::CacheFromCompiledData()
 			}
 		}
 
-		// Prepare renderer bindings this avoid having to do this per instance spawned
+		// Populate any renderer bindings from the emitter state and spawn infos
+		EmitterBuildContext.ConvertDistributionToRange(StatelessEmitterData->EmitterState.LoopDuration, 0.0f);
+		EmitterBuildContext.ConvertDistributionToRange(StatelessEmitterData->EmitterState.LoopDelay, 0.0f);
+		for ( const FNiagaraStatelessSpawnInfo& SpawnInfo : StatelessEmitterData->SpawnInfos )
+		{
+			if (SpawnInfo.Type == ENiagaraStatelessSpawnInfoType::Rate)
+			{
+				EmitterBuildContext.ConvertDistributionToRange(SpawnInfo.Rate, 0.0f);
+			}
+			else //if (SpawnInfo.Type == ENiagaraStatelessSpawnInfoType::Burst)
+			{
+				EmitterBuildContext.ConvertDistributionToRange(SpawnInfo.Amount, 0.0f);
+				EmitterBuildContext.ConvertDistributionToRange(SpawnInfo.SpawnProbability, 0.0f);
+			}
+		}
+
+		// Prepare renderer bindings this avoids having to do this per instance spawned
 		// Note: Order is important here we detect if we need to update shader parameters on binding changes above by looking to see if we had any renderer bindings so this must be below
 		for (UNiagaraRendererProperties* Renderer : RendererProperties )
 		{
diff --git a/Engine/Plugins/FX/Niagara/Source/Niagara/Private/Stateless/NiagaraStatelessEmitterInstance.cpp b/Engine/Plugins/FX/Niagara/Source/Niagara/Private/Stateless/NiagaraStatelessEmitterInstance.cpp
--- a/Engine/Plugins/FX/Niagara/Source/Niagara/Private/Stateless/NiagaraStatelessEmitterInstance.cpp
+++ b/Engine/Plugins/FX/Niagara/Source/Niagara/Private/Stateless/NiagaraStatelessEmitterInstance.cpp
@@ -10,6 +10,35 @@
 
 #include "Shader.h"
 
+namespace NiagaraStatelessEmitterInstancePrivate
+{
+	static const float DefaultLoopDuration		= 0.001f;
+	static const float DefaultLoopDelay			= 0.0f;
+	static const float DefaultSpawnRate			= 0.0f;
+	static const float DefaultSpawnProbability	= 0.0f;
+	static const int32 DefaultSpawnAmount		= 0;
+
+	// Other areas work on built distributions so this is a helper function to work with a raw distribution
+	float EvaluateDistribution(const FNiagaraDistributionRangeFloat& Distribution, FRandomStream& RandomStream, const FNiagaraParameterStore& ParameterStore, const float DefaultValue)
+	{
+		if (Distribution.IsBinding())
+		{
+			return Distribution.ParameterBinding.IsValid() ? ParameterStore.GetParameterValueOrDefault(Distribution.ParameterBinding, DefaultValue) : DefaultValue;
+		}
+		const float Fraction = RandomStream.GetFraction();
+		return ((Distribution.Max - Distribution.Max) * Fraction) + Distribution.Min;
+	}
+
+	int32 EvaluateDistribution(const FNiagaraDistributionRangeInt& Distribution, FRandomStream& RandomStream, const FNiagaraParameterStore& ParameterStore, const int32 DefaultValue)
+	{
+		if (Distribution.IsBinding())
+		{
+			return Distribution.ParameterBinding.IsValid() ? ParameterStore.GetParameterValueOrDefault(Distribution.ParameterBinding, DefaultValue) : DefaultValue;
+		}
+		return RandomStream.RandRange(Distribution.Min, Distribution.Max);
+	}
+}
+
 FNiagaraDataBuffer* NiagaraStateless::FEmitterInstance_RT::GetDataToRender(FRHICommandListBase& RHICmdList, bool bIsLowLatencyTranslucent) const
 {
 	return ComputeManager ? ComputeManager->GetDataBuffer(RHICmdList, uintptr_t(this), this) : nullptr;
@@ -238,16 +267,22 @@ void FNiagaraStatelessEmitterInstance::Tick(float DeltaSeconds)
 
 void FNiagaraStatelessEmitterInstance::InitEmitterState()
 {
+	using namespace NiagaraStatelessEmitterInstancePrivate;
+
 	const FNiagaraEmitterStateData& EmitterState = EmitterData->EmitterState;
 	LoopCount			= 0;
-	CurrentLoopDuration = RandomStream.FRandRange(EmitterState.LoopDuration.Min, EmitterState.LoopDuration.Max);
-	CurrentLoopDelay	= RandomStream.FRandRange(EmitterState.LoopDelay.Min, EmitterState.LoopDelay.Max);
+	CurrentLoopDuration = EvaluateDistribution(EmitterState.LoopDuration, RandomStream, RendererBindings, DefaultLoopDuration);
+	CurrentLoopDuration = FMath::Max(CurrentLoopDuration, DefaultLoopDuration);
+	CurrentLoopDelay	= EvaluateDistribution(EmitterState.LoopDelay, RandomStream, RendererBindings, DefaultLoopDelay);
+	CurrentLoopDelay	= FMath::Max(CurrentLoopDelay, 0.0f);
 	CurrentLoopAgeStart	= 0.0f;
 	CurrentLoopAgeEnd	= CurrentLoopAgeStart + CurrentLoopDelay + CurrentLoopDuration;
 }
 
 void FNiagaraStatelessEmitterInstance::TickEmitterState()
 {
+	using namespace NiagaraStatelessEmitterInstancePrivate;
+
 	// Update execution state based on the parent which be told to go inactive / complete
 	{
 		const ENiagaraExecutionState ParentExecutionState = ParentSystemInstance ? ParentSystemInstance->GetActualExecutionState() : ENiagaraExecutionState::Complete;
@@ -364,7 +399,8 @@ void FNiagaraStatelessEmitterInstance::TickEmitterState()
 
 				if (EmitterState.bRecalculateDurationEachLoop)
 				{
-					CurrentLoopDuration = RandomStream.FRandRange(EmitterState.LoopDuration.Min, EmitterState.LoopDuration.Max);
+					CurrentLoopDuration = EvaluateDistribution(EmitterState.LoopDuration, RandomStream, RendererBindings, DefaultLoopDuration);
+					CurrentLoopDuration = FMath::Max(CurrentLoopDuration, DefaultLoopDuration);
 				}
 
 				if (EmitterState.bDelayFirstLoopOnly)
@@ -373,7 +409,8 @@ void FNiagaraStatelessEmitterInstance::TickEmitterState()
 				}
 				else if (EmitterState.bRecalculateDelayEachLoop)
 				{
-					CurrentLoopDelay = RandomStream.FRandRange(EmitterState.LoopDelay.Min, EmitterState.LoopDelay.Max);
+					CurrentLoopDelay = EvaluateDistribution(EmitterState.LoopDelay, RandomStream, RendererBindings, DefaultLoopDelay);
+					CurrentLoopDelay = FMath::Max(CurrentLoopDelay, 0.0f);
 				}
 
 				CurrentLoopAgeStart = CurrentLoopAgeEnd;
@@ -429,15 +466,18 @@ void FNiagaraStatelessEmitterInstance::SendRenderData()
 	DataForRenderThread.Age				= Age;
 	DataForRenderThread.ExecutionState	= ExecutionState;
 
-	if ( EmitterData->bModulesHaveRendererBindings && RendererBindings.GetParametersDirty() )
+	if (RendererBindings.GetParametersDirty())
 	{
 		RendererBindings.Tick();
-		DataForRenderThread.bHasBindingBufferData = true;
-		DataForRenderThread.BindingBufferData = RendererBindings.GetParameterDataArray();
-		check((DataForRenderThread.BindingBufferData.Num() % sizeof(uint32)) == 0);
+		if (EmitterData->bModulesHaveRendererBindings)
+		{
+			DataForRenderThread.bHasBindingBufferData = true;
+			DataForRenderThread.BindingBufferData = RendererBindings.GetParameterDataArray();
+			check((DataForRenderThread.BindingBufferData.Num() % sizeof(uint32)) == 0);
 
-		DataForRenderThread.ShaderParameters = WeakStatelessEmitter->AllocateShaderParameters(RendererBindings);
-		DataForRenderThread.ShaderParameters->Common_RandomSeed = RandomSeed;
+			DataForRenderThread.ShaderParameters = WeakStatelessEmitter->AllocateShaderParameters(RendererBindings);
+			DataForRenderThread.ShaderParameters->Common_RandomSeed = RandomSeed;
+		}
 	}
 
 	if (bSpawnInfosDirty)
@@ -474,6 +514,8 @@ void FNiagaraStatelessEmitterInstance::SendRenderData()
 
 void FNiagaraStatelessEmitterInstance::InitSpawnInfos(float InitializationAge)
 {
+	using namespace NiagaraStatelessEmitterInstancePrivate;
+
 	// If we are not enabled, or not awake from scalability skip adding
 	if (!bEmitterEnabled_GT || (ScalabilityState != ENiagaraExecutionStateManagement::Awaken))
 	{
@@ -484,7 +526,8 @@ void FNiagaraStatelessEmitterInstance::InitSpawnInfos(float InitializationAge)
 	{
 		if (SpawnInfo.Type == ENiagaraStatelessSpawnInfoType::Rate)
 		{
-			const float SpawnRate = RandomStream.FRandRange(SpawnInfo.Rate.Min, SpawnInfo.Rate.Max) * EmitterData->SpawnCountScale;
+			float SpawnRate = EvaluateDistribution(SpawnInfo.Rate, RandomStream, RendererBindings, DefaultSpawnRate) * EmitterData->SpawnCountScale;
+			SpawnRate = FMath::Max(SpawnRate, 0.0f);
 			if (SpawnRate > 0.0f)
 			{
 				FActiveSpawnRate& ActiveSpawnRate = ActiveSpawnRates.AddDefaulted_GetRef();
@@ -499,6 +542,8 @@ void FNiagaraStatelessEmitterInstance::InitSpawnInfos(float InitializationAge)
 
 void FNiagaraStatelessEmitterInstance::InitSpawnInfosForLoop(float InitializationAge)
 {
+	using namespace NiagaraStatelessEmitterInstancePrivate;
+
 	// If we are not enabled, or not awake from scalability skip adding
 	if (!bEmitterEnabled_GT || (ScalabilityState != ENiagaraExecutionStateManagement::Awaken))
 	{
@@ -560,12 +605,17 @@ void FNiagaraStatelessEmitterInstance::InitSpawnInfosForLoop(float Initializatio
 			continue;
 		}
 
-		if (SpawnInfo.bSpawnProbabilityEnabled && SpawnInfo.SpawnProbability < RandomStream.FRand())
+		if (SpawnInfo.bSpawnProbabilityEnabled)
 		{
-			continue;
+			float SpawnProbability = EvaluateDistribution(SpawnInfo.SpawnProbability, RandomStream, RendererBindings, DefaultSpawnProbability);
+			SpawnProbability = FMath::Clamp(SpawnProbability, 0.0f, 1.0f);
+			if (SpawnProbability < RandomStream.FRand())
+			{
+				continue;
+			}
 		}
 
-		const int32 UnscaledSpawnAmount = RandomStream.RandRange(SpawnInfo.Amount.Min, SpawnInfo.Amount.Max);
+		const int32 UnscaledSpawnAmount = EvaluateDistribution(SpawnInfo.Amount, RandomStream, RendererBindings, DefaultSpawnAmount);
 		const int32 SpawnAmount = FMath::FloorToInt(float(UnscaledSpawnAmount) * EmitterData->SpawnCountScale);
 		if (SpawnAmount <= 0)
 		{
