- Downgrade warning to a log

#rnx

diff --git a/Engine/Plugins/FX/Niagara/Source/NiagaraEditor/Private/NiagaraSystemCompilingManager.cpp b/Engine/Plugins/FX/Niagara/Source/NiagaraEditor/Private/NiagaraSystemCompilingManager.cpp
--- a/Engine/Plugins/FX/Niagara/Source/NiagaraEditor/Private/NiagaraSystemCompilingManager.cpp
+++ b/Engine/Plugins/FX/Niagara/Source/NiagaraEditor/Private/NiagaraSystemCompilingManager.cpp
@@ -206,7 +206,7 @@ void FNiagaraSystemCompilingManager::CheckStalledTask(double CurrentTime, FNiaga
 
 		if (bWarn)
 		{
-			UE_LOG(LogNiagaraEditor, Warning, TEXT("NiagaraSystemCompilingManager - compilation task [%s] stalled for %f seconds.  Status - %s"),
+			UE_LOG(LogNiagaraEditor, Log, TEXT("NiagaraSystemCompilingManager - compilation task [%s] stalled for %f seconds.  Status - %s"),
 				*Task->GetDescription(), (float)ElapsedTime, *Task->GetStatusString());
 
 			Task->LastStallWarningTime = CurrentTime;
