- Fix for float range distribution emitter state randomization

#rnx

diff --git a/Engine/Plugins/FX/Niagara/Source/Niagara/Private/Stateless/NiagaraStatelessEmitterInstance.cpp b/Engine/Plugins/FX/Niagara/Source/Niagara/Private/Stateless/NiagaraStatelessEmitterInstance.cpp
--- a/Engine/Plugins/FX/Niagara/Source/Niagara/Private/Stateless/NiagaraStatelessEmitterInstance.cpp
+++ b/Engine/Plugins/FX/Niagara/Source/Niagara/Private/Stateless/NiagaraStatelessEmitterInstance.cpp
@@ -26,7 +26,7 @@ namespace NiagaraStatelessEmitterInstancePrivate
 			return Distribution.ParameterBinding.IsValid() ? ParameterStore.GetParameterValueOrDefault(Distribution.ParameterBinding, DefaultValue) : DefaultValue;
 		}
 		const float Fraction = RandomStream.GetFraction();
-		return ((Distribution.Max - Distribution.Max) * Fraction) + Distribution.Min;
+		return ((Distribution.Max - Distribution.Min) * Fraction) + Distribution.Min;
 	}
 
 	int32 EvaluateDistribution(const FNiagaraDistributionRangeInt& Distribution, FRandomStream& RandomStream, const FNiagaraParameterStore& ParameterStore, const int32 DefaultValue)
