- Fix for stateless loop duration not delaying spawn rates with each loop

#rnx

diff --git a/Engine/Plugins/FX/Niagara/Source/Niagara/Private/Stateless/NiagaraStatelessEmitterInstance.cpp b/Engine/Plugins/FX/Niagara/Source/Niagara/Private/Stateless/NiagaraStatelessEmitterInstance.cpp
--- a/Engine/Plugins/FX/Niagara/Source/Niagara/Private/Stateless/NiagaraStatelessEmitterInstance.cpp
+++ b/Engine/Plugins/FX/Niagara/Source/Niagara/Private/Stateless/NiagaraStatelessEmitterInstance.cpp
@@ -532,7 +532,6 @@ void FNiagaraStatelessEmitterInstance::InitSpawnInfos(float InitializationAge)
 			{
 				FActiveSpawnRate& ActiveSpawnRate = ActiveSpawnRates.AddDefaulted_GetRef();
 				ActiveSpawnRate.Rate = SpawnRate;
-				ActiveSpawnRate.SpawnTime = CurrentLoopDelay;
 			}
 		}
 	}
@@ -555,7 +554,7 @@ void FNiagaraStatelessEmitterInstance::InitSpawnInfosForLoop(float Initializatio
 	{
 		float SpawnTime = FMath::Max(InitializationAge - SpawnInfo.SpawnTime, 0.0f);
 		SpawnTime = FMath::CeilToFloat(SpawnTime * SpawnInfo.Rate) / SpawnInfo.Rate;
-		SpawnTime += SpawnInfo.SpawnTime;
+		SpawnTime += SpawnInfo.SpawnTime + CurrentLoopDelay;
 		if (SpawnTime >= CurrentLoopAgeEnd)
 		{
 			continue;
