- Fix for compute manager cache
- The compile data set was a raw pointer, and we might release the original asset before the cache is cleared which results in a garbage data set allocation

#rnx

diff --git a/Engine/Plugins/FX/Niagara/Source/Niagara/Internal/Stateless/NiagaraStatelessEmitterData.h b/Engine/Plugins/FX/Niagara/Source/Niagara/Internal/Stateless/NiagaraStatelessEmitterData.h
--- a/Engine/Plugins/FX/Niagara/Source/Niagara/Internal/Stateless/NiagaraStatelessEmitterData.h
+++ b/Engine/Plugins/FX/Niagara/Source/Niagara/Internal/Stateless/NiagaraStatelessEmitterData.h
@@ -32,7 +32,7 @@ struct FNiagaraStatelessEmitterData
 	FNiagaraStatelessEmitterData() = default;
 	~FNiagaraStatelessEmitterData();
 
-	FNiagaraDataSetCompiledData						ParticleDataSetCompiledData;
+	TSharedPtr<FNiagaraDataSetCompiledData>			ParticleDataSetCompiledData;
 	TArray<int32>									ComponentOffsets;
 
 	bool											bCanEverExecute = false;
diff --git a/Engine/Plugins/FX/Niagara/Source/Niagara/Private/NiagaraSimCacheHelper.h b/Engine/Plugins/FX/Niagara/Source/Niagara/Private/NiagaraSimCacheHelper.h
--- a/Engine/Plugins/FX/Niagara/Source/Niagara/Private/NiagaraSimCacheHelper.h
+++ b/Engine/Plugins/FX/Niagara/Source/Niagara/Private/NiagaraSimCacheHelper.h
@@ -339,7 +339,7 @@ struct FNiagaraSimCacheHelper
 		{
 			const UNiagaraStatelessEmitter* StatelessEmitter = EmitterHandle.GetStatelessEmitter();
 			const FNiagaraStatelessEmitterDataPtr StatelessEmitterData = StatelessEmitter ? StatelessEmitter->GetEmitterData() : nullptr;
-			ParticleDataSet = StatelessEmitterData ? &StatelessEmitterData->ParticleDataSetCompiledData : nullptr;
+			ParticleDataSet = StatelessEmitterData ? StatelessEmitterData->ParticleDataSetCompiledData.Get() : nullptr;
 		}
 
 		if (ParticleDataSet == nullptr)
diff --git a/Engine/Plugins/FX/Niagara/Source/Niagara/Private/Stateless/NiagaraStatelessComputeManager.cpp b/Engine/Plugins/FX/Niagara/Source/Niagara/Private/Stateless/NiagaraStatelessComputeManager.cpp
--- a/Engine/Plugins/FX/Niagara/Source/Niagara/Private/Stateless/NiagaraStatelessComputeManager.cpp
+++ b/Engine/Plugins/FX/Niagara/Source/Niagara/Private/Stateless/NiagaraStatelessComputeManager.cpp
@@ -195,7 +195,7 @@ FNiagaraDataBuffer* FNiagaraStatelessComputeManager::GetDataBuffer(FRHICommandLi
 
 	FNiagaraGpuComputeDispatchInterface* ComputeInterface = GetOwnerInterface();
 
-	const uint32 DataSetLayoutHash = EmitterInstance->EmitterData->ParticleDataSetCompiledData.GetLayoutHash();
+	const uint32 DataSetLayoutHash = EmitterInstance->EmitterData->ParticleDataSetCompiledData->GetLayoutHash();
 
 	FStatelessDataCache* CacheData = nullptr;
 	if (GUseDataBufferCache)
@@ -215,7 +215,8 @@ FNiagaraDataBuffer* FNiagaraStatelessComputeManager::GetDataBuffer(FRHICommandLi
 	{
 		CacheData = new FStatelessDataCache();
 		CacheData->DataSetLayoutHash = DataSetLayoutHash;
-		CacheData->DataSet.Init(&EmitterInstance->EmitterData->ParticleDataSetCompiledData);
+		CacheData->DataSetCompiledData = EmitterInstance->EmitterData->ParticleDataSetCompiledData;
+		CacheData->DataSet.Init(CacheData->DataSetCompiledData.Get());
 		CacheData->DataBuffer = new FNiagaraDataBuffer(&CacheData->DataSet);
 	}
 
diff --git a/Engine/Plugins/FX/Niagara/Source/Niagara/Private/Stateless/NiagaraStatelessComputeManager.h b/Engine/Plugins/FX/Niagara/Source/Niagara/Private/Stateless/NiagaraStatelessComputeManager.h
--- a/Engine/Plugins/FX/Niagara/Source/Niagara/Private/Stateless/NiagaraStatelessComputeManager.h
+++ b/Engine/Plugins/FX/Niagara/Source/Niagara/Private/Stateless/NiagaraStatelessComputeManager.h
@@ -22,6 +22,7 @@ public:
 	struct FStatelessDataCache
 	{
 		uint32										DataSetLayoutHash = 0;
+		TSharedPtr<FNiagaraDataSetCompiledData>		DataSetCompiledData;
 		FNiagaraDataSet								DataSet;
 		FNiagaraDataBufferRef						DataBuffer;
 
diff --git a/Engine/Plugins/FX/Niagara/Source/Niagara/Private/Stateless/NiagaraStatelessEmitter.cpp b/Engine/Plugins/FX/Niagara/Source/Niagara/Private/Stateless/NiagaraStatelessEmitter.cpp
--- a/Engine/Plugins/FX/Niagara/Source/Niagara/Private/Stateless/NiagaraStatelessEmitter.cpp
+++ b/Engine/Plugins/FX/Niagara/Source/Niagara/Private/Stateless/NiagaraStatelessEmitter.cpp
@@ -310,7 +310,7 @@ void UNiagaraStatelessEmitter::CacheFromCompiledData()
 	{
 		if (Renderer && Renderer->GetIsEnabled())
 		{
-			Renderer->CacheFromCompiledData(&StatelessEmitterData->ParticleDataSetCompiledData);
+			Renderer->CacheFromCompiledData(StatelessEmitterData->ParticleDataSetCompiledData.Get());
 		}
 	}
 
@@ -348,11 +348,11 @@ void UNiagaraStatelessEmitter::CacheFromCompiledData()
 	{
 		if (EnumHasAnyFlags(StatelessEmitterData->FeatureMask, ENiagaraStatelessFeatureMask::ExecuteCPU))
 		{
-			StatelessEmitterData->ParticleSimExecData = new NiagaraStateless::FParticleSimulationExecData(StatelessEmitterData->ParticleDataSetCompiledData);
+			StatelessEmitterData->ParticleSimExecData = new NiagaraStateless::FParticleSimulationExecData(*StatelessEmitterData->ParticleDataSetCompiledData.Get());
 		}
 
 		FNiagaraStatelessEmitterDataBuildContext EmitterBuildContext(
-			StatelessEmitterData->ParticleDataSetCompiledData,
+			*StatelessEmitterData->ParticleDataSetCompiledData.Get(),
 			StatelessEmitterData->RendererBindings,
 			StatelessEmitterData->BuiltData,
 			StatelessEmitterData->StaticFloatData,
@@ -509,7 +509,7 @@ void UNiagaraStatelessEmitter::BuildCompiledDataSet()
 		ComponentOffsets.Shrink();
 	}
 #endif
-	StatelessEmitterData->ParticleDataSetCompiledData = ParticleDataSetCompiledData;
+	StatelessEmitterData->ParticleDataSetCompiledData = MakeShared<FNiagaraDataSetCompiledData>(ParticleDataSetCompiledData);
 	StatelessEmitterData->ComponentOffsets = ComponentOffsets;
 }
 
diff --git a/Engine/Plugins/FX/Niagara/Source/Niagara/Private/Stateless/NiagaraStatelessEmitterInstance.cpp b/Engine/Plugins/FX/Niagara/Source/Niagara/Private/Stateless/NiagaraStatelessEmitterInstance.cpp
--- a/Engine/Plugins/FX/Niagara/Source/Niagara/Private/Stateless/NiagaraStatelessEmitterInstance.cpp
+++ b/Engine/Plugins/FX/Niagara/Source/Niagara/Private/Stateless/NiagaraStatelessEmitterInstance.cpp
@@ -101,7 +101,7 @@ void FNiagaraStatelessEmitterInstance::Init(int32 InEmitterIndex)
 	//CompletionAge = EmitterData->CalculateCompletionAge(RandomSeed);
 
 	// Initialize data set
-	ParticleDataSet->Init(&EmitterData->ParticleDataSetCompiledData);
+	ParticleDataSet->Init(EmitterData->ParticleDataSetCompiledData.Get());
 
 	// Prepare our parameters
 	RendererBindings = EmitterData->RendererBindings;
