Add overload to NiagaraWorldManager::GetLODDistance to allow providing the list of view locations.

Split from review 35032115.

#rnx

diff --git a/Engine/Plugins/FX/Niagara/Source/Niagara/Private/NiagaraWorldManager.cpp b/Engine/Plugins/FX/Niagara/Source/Niagara/Private/NiagaraWorldManager.cpp
--- a/Engine/Plugins/FX/Niagara/Source/Niagara/Private/NiagaraWorldManager.cpp
+++ b/Engine/Plugins/FX/Niagara/Source/Niagara/Private/NiagaraWorldManager.cpp
@@ -1556,16 +1556,24 @@ FVector::FReal FNiagaraWorldManager::GetLODDistance(FVector Location)const
 			PlayerViewLocations = World->ViewLocationsRenderedLastFrame;
 		}
 
-		if (PlayerViewLocations.Num() > 0)
+		LODDistance = GetLODDistance(Location, PlayerViewLocations);
+	}
+	return LODDistance;
+}
+
+FVector::FReal FNiagaraWorldManager::GetLODDistance(FVector Location, TConstArrayView<FVector> ViewPoints) const
+{
+	FVector::FReal LODDistance = 0.0f;
+
+	if (ViewPoints.Num() > 0)
+	{
+		FVector::FReal LODDistanceSqr = FMath::Square(WORLD_MAX);
+		for (const FVector& ViewLocation : ViewPoints)
 		{
-			FVector::FReal LODDistanceSqr = FMath::Square(WORLD_MAX);
-			for (const FVector& ViewLocation : PlayerViewLocations)
-			{
-				const FVector::FReal DistanceToEffectSqr = FVector(ViewLocation - Location).SizeSquared();
-				LODDistanceSqr = FMath::Min(LODDistanceSqr, DistanceToEffectSqr);
-			}
-			LODDistance = FVector::FReal(FMath::Sqrt(LODDistanceSqr));
+			const FVector::FReal DistanceToEffectSqr = FVector(ViewLocation - Location).SizeSquared();
+			LODDistanceSqr = FMath::Min(LODDistanceSqr, DistanceToEffectSqr);
 		}
+		LODDistance = FVector::FReal(FMath::Sqrt(LODDistanceSqr));
 	}
 	return LODDistance;
 }
diff --git a/Engine/Plugins/FX/Niagara/Source/Niagara/Public/NiagaraWorldManager.h b/Engine/Plugins/FX/Niagara/Source/Niagara/Public/NiagaraWorldManager.h
--- a/Engine/Plugins/FX/Niagara/Source/Niagara/Public/NiagaraWorldManager.h
+++ b/Engine/Plugins/FX/Niagara/Source/Niagara/Public/NiagaraWorldManager.h
@@ -167,7 +167,8 @@ public:
 	TArrayView<const FNiagaraCachedViewInfo> GetCachedViewInfo() const { return MakeArrayView(CachedViewInfo); }
 
 	//Returns the distance to the nearest viewpoint to the give location. Used for a distance on which to base LODs.
-	NIAGARA_API FVector::FReal GetLODDistance(FVector Location)const;
+	NIAGARA_API FVector::FReal GetLODDistance(FVector Location) const;
+	NIAGARA_API FVector::FReal GetLODDistance(FVector Location, TConstArrayView<FVector> ViewPoints) const;
 
 	UNiagaraComponentPool* GetComponentPool() { return ComponentPool; }
 
