Exposed Niagara Effect Type creation in the Content Browser and made cooked effect types open-able in read-only mode.

diff --git a/Engine/Plugins/FX/Niagara/Source/NiagaraEditor/Private/AssetDefinitions/AssetDefinition_NiagaraEffectType.cpp b/Engine/Plugins/FX/Niagara/Source/NiagaraEditor/Private/AssetDefinitions/AssetDefinition_NiagaraEffectType.cpp
--- a/Engine/Plugins/FX/Niagara/Source/NiagaraEditor/Private/AssetDefinitions/AssetDefinition_NiagaraEffectType.cpp
+++ b/Engine/Plugins/FX/Niagara/Source/NiagaraEditor/Private/AssetDefinitions/AssetDefinition_NiagaraEffectType.cpp
@@ -3,6 +3,7 @@
 #include "AssetDefinition_NiagaraEffectType.h"
 #include "NiagaraEditorStyle.h"
 #include "SDetailsDiff.h"
+#include "Toolkits/SimpleAssetEditor.h"
 
 #define LOCTEXT_NAMESPACE "AssetTypeActions_NiagaraEffectType"
 
@@ -27,4 +28,21 @@ EAssetCommandResult UAssetDefinition_NiagaraEffectType::PerformAssetDiff(const F
 	return EAssetCommandResult::Handled;
 }
 
+EAssetCommandResult UAssetDefinition_NiagaraEffectType::OpenAssets(const FAssetOpenArgs& OpenArgs) const
+{
+	if (OpenArgs.OpenMethod == EAssetOpenMethod::Edit || OpenArgs.OpenMethod == EAssetOpenMethod::View)
+	{
+		FSimpleAssetEditor::CreateEditor(EToolkitMode::Standalone, OpenArgs.ToolkitHost, OpenArgs.LoadObjects<UObject>());
+		return EAssetCommandResult::Handled;
+	}
+
+	return EAssetCommandResult::Unhandled;
+}
+
+FAssetOpenSupport UAssetDefinition_NiagaraEffectType::GetAssetOpenSupport(const FAssetOpenSupportArgs& OpenSupportArgs) const
+{
+	return FAssetOpenSupport(OpenSupportArgs.OpenMethod,OpenSupportArgs.OpenMethod == EAssetOpenMethod::Edit || OpenSupportArgs.OpenMethod == EAssetOpenMethod::View); 
+
+}
+
 #undef LOCTEXT_NAMESPACE
diff --git a/Engine/Plugins/FX/Niagara/Source/NiagaraEditor/Private/AssetDefinitions/AssetDefinition_NiagaraEffectType.h b/Engine/Plugins/FX/Niagara/Source/NiagaraEditor/Private/AssetDefinitions/AssetDefinition_NiagaraEffectType.h
--- a/Engine/Plugins/FX/Niagara/Source/NiagaraEditor/Private/AssetDefinitions/AssetDefinition_NiagaraEffectType.h
+++ b/Engine/Plugins/FX/Niagara/Source/NiagaraEditor/Private/AssetDefinitions/AssetDefinition_NiagaraEffectType.h
@@ -27,5 +27,7 @@ public:
 		return Categories;
 	}
 	virtual EAssetCommandResult PerformAssetDiff(const FAssetDiffArgs& DiffArgs) const override;
+	virtual EAssetCommandResult OpenAssets(const FAssetOpenArgs& OpenArgs) const override;
+	virtual FAssetOpenSupport GetAssetOpenSupport(const FAssetOpenSupportArgs& OpenSupportArgs) const override;
 	// UAssetDefinition End
 };
