Fixed mesh renderer thumbnails when multiple mesh entries without mesh (user param bound for example) exist.

Issue was the same default thumbnail due to lack of mesh was added multiple times instead of constructing multiple default thumbnails

diff --git a/Engine/Plugins/FX/Niagara/Source/Niagara/Private/NiagaraMeshRendererProperties.cpp b/Engine/Plugins/FX/Niagara/Source/Niagara/Private/NiagaraMeshRendererProperties.cpp
--- a/Engine/Plugins/FX/Niagara/Source/Niagara/Private/NiagaraMeshRendererProperties.cpp
+++ b/Engine/Plugins/FX/Niagara/Source/Niagara/Private/NiagaraMeshRendererProperties.cpp
@@ -1029,14 +1029,28 @@ void UNiagaraMeshRendererProperties::GetAdditionalVariables(TArray<FNiagaraVaria
 
 void UNiagaraMeshRendererProperties::GetRendererWidgets(const FNiagaraEmitterInstance* InEmitter, TArray<TSharedPtr<SWidget>>& OutWidgets, TSharedPtr<FAssetThumbnailPool> InThumbnailPool) const
 {
-	TSharedRef<SWidget> DefaultThumbnailWidget = SNew(SImage)
-		.Image(FSlateIconFinder::FindIconBrushForClass(StaticClass()));
-
+	// for(const FNiagaraMeshRendererMeshProperties& MeshProperties : Meshes)
+	// {
+	// 	// TSharedPtr<SWidget> ThumbnailWidget = DefaultThumbnailWidget;
+	// 	//
+	// 	// UStaticMesh* Mesh = MeshProperties.Mesh;
+	// 	// if (Mesh && Mesh->HasValidRenderData())
+	// 	// {
+	// 	// 	TSharedPtr<FAssetThumbnail> AssetThumbnail = MakeShareable(new FAssetThumbnail(Mesh, ThumbnailSize, ThumbnailSize, InThumbnailPool));
+	// 	// 	ThumbnailWidget = AssetThumbnail->MakeThumbnailWidget();
+	// 	// }
+	// 					
+	// 	OutWidgets.Add(SNullWidget::NullWidget);		
+	// }
+	
 	int32 ThumbnailSize = 32;
 	for(const FNiagaraMeshRendererMeshProperties& MeshProperties : Meshes)
 	{
+		TSharedRef<SWidget> DefaultThumbnailWidget = SNew(SImage)
+		.Image(FSlateIconFinder::FindIconBrushForClass(StaticClass()));
+		
 		TSharedPtr<SWidget> ThumbnailWidget = DefaultThumbnailWidget;
-
+	
 		UStaticMesh* Mesh = MeshProperties.Mesh;
 		if (Mesh && Mesh->HasValidRenderData())
 		{
@@ -1044,11 +1058,14 @@ void UNiagaraMeshRendererProperties::GetRendererWidgets(const FNiagaraEmitterIns
 			ThumbnailWidget = AssetThumbnail->MakeThumbnailWidget();
 		}
 		
-		OutWidgets.Add(ThumbnailWidget);		
+		OutWidgets.Add(ThumbnailWidget);
 	}
-
+	
 	if (Meshes.Num() == 0)
 	{
+		TSharedRef<SWidget> DefaultThumbnailWidget = SNew(SImage)
+		.Image(FSlateIconFinder::FindIconBrushForClass(StaticClass()));
+		
 		OutWidgets.Add(DefaultThumbnailWidget);
 	}
 }
