Removed inactive code from previous submit

diff --git a/Engine/Plugins/FX/Niagara/Source/Niagara/Private/NiagaraMeshRendererProperties.cpp b/Engine/Plugins/FX/Niagara/Source/Niagara/Private/NiagaraMeshRendererProperties.cpp
--- a/Engine/Plugins/FX/Niagara/Source/Niagara/Private/NiagaraMeshRendererProperties.cpp
+++ b/Engine/Plugins/FX/Niagara/Source/Niagara/Private/NiagaraMeshRendererProperties.cpp
@@ -1028,21 +1028,7 @@ void UNiagaraMeshRendererProperties::GetAdditionalVariables(TArray<FNiagaraVaria
 }
 
 void UNiagaraMeshRendererProperties::GetRendererWidgets(const FNiagaraEmitterInstance* InEmitter, TArray<TSharedPtr<SWidget>>& OutWidgets, TSharedPtr<FAssetThumbnailPool> InThumbnailPool) const
-{
-	// for(const FNiagaraMeshRendererMeshProperties& MeshProperties : Meshes)
-	// {
-	// 	// TSharedPtr<SWidget> ThumbnailWidget = DefaultThumbnailWidget;
-	// 	//
-	// 	// UStaticMesh* Mesh = MeshProperties.Mesh;
-	// 	// if (Mesh && Mesh->HasValidRenderData())
-	// 	// {
-	// 	// 	TSharedPtr<FAssetThumbnail> AssetThumbnail = MakeShareable(new FAssetThumbnail(Mesh, ThumbnailSize, ThumbnailSize, InThumbnailPool));
-	// 	// 	ThumbnailWidget = AssetThumbnail->MakeThumbnailWidget();
-	// 	// }
-	// 					
-	// 	OutWidgets.Add(SNullWidget::NullWidget);		
-	// }
-	
+{	
 	int32 ThumbnailSize = 32;
 	for(const FNiagaraMeshRendererMeshProperties& MeshProperties : Meshes)
 	{
