Added uproperty tags to a few enums properties

#tests Looked at the enum properties in the details panel to verify they still fit and work as expected
#rnx

diff --git a/Engine/Plugins/FX/Niagara/Source/Niagara/Classes/NiagaraEmitter.h b/Engine/Plugins/FX/Niagara/Source/Niagara/Classes/NiagaraEmitter.h
--- a/Engine/Plugins/FX/Niagara/Source/Niagara/Classes/NiagaraEmitter.h
+++ b/Engine/Plugins/FX/Niagara/Source/Niagara/Classes/NiagaraEmitter.h
@@ -295,14 +295,14 @@ struct FVersionedNiagaraEmitterData
 	uint32 bGpuAlwaysRunParticleUpdateScript : 1;
 #endif
 
-	UPROPERTY(EditAnywhere, BlueprintReadWrite, Category = "Emitter")
+	UPROPERTY(EditAnywhere, BlueprintReadWrite, Category = "Emitter", meta = (SegmentedDisplay))
 	ENiagaraSimTarget SimTarget = ENiagaraSimTarget::CPUSim;
 
 	/**
 	How should we calculate bounds for the emitter.
 	Note: If this is greyed out it means fixed bounds are enabled in the System Properties and these bounds are therefore ignored.
 	*/
-	UPROPERTY(EditAnywhere, Category = "Emitter")
+	UPROPERTY(EditAnywhere, Category = "Emitter", meta = (SegmentedDisplay))
 	ENiagaraEmitterCalculateBoundMode CalculateBoundsMode = ENiagaraEmitterCalculateBoundMode::Dynamic;
 	
 	/**
@@ -334,7 +334,7 @@ struct FVersionedNiagaraEmitterData
 	To prevent reallocations, the emitter should allocate as much memory as is needed for the max particle count.
 	This setting controls if the allocation size should be automatically determined or manually entered.
 	*/
-	UPROPERTY(EditAnywhere, BlueprintReadWrite, AdvancedDisplay, Category = "Emitter")
+	UPROPERTY(EditAnywhere, BlueprintReadWrite, AdvancedDisplay, Category = "Emitter", meta = (SegmentedDisplay))
 	EParticleAllocationMode AllocationMode = EParticleAllocationMode::AutomaticEstimate;
 	
 	/** 
@@ -447,7 +447,7 @@ struct FVersionedNiagaraEmitterData
 	TArray<FString> AttributesToPreserve;
 
 	/** This determines how emitters will be added to a system by default. If summary view is setup, consider setting this to 'Summary'. */
-	UPROPERTY(EditAnywhere, AdvancedDisplay, Category = "Asset Options")
+	UPROPERTY(EditAnywhere, AdvancedDisplay, Category = "Asset Options", meta = (SegmentedDisplay))
 	ENiagaraEmitterDefaultSummaryState AddEmitterDefaultViewState = ENiagaraEmitterDefaultSummaryState::Default;
 	
 	UPROPERTY()
diff --git a/Engine/Plugins/FX/Niagara/Source/Niagara/Classes/NiagaraScript.h b/Engine/Plugins/FX/Niagara/Source/Niagara/Classes/NiagaraScript.h
--- a/Engine/Plugins/FX/Niagara/Source/Niagara/Classes/NiagaraScript.h
+++ b/Engine/Plugins/FX/Niagara/Source/Niagara/Classes/NiagaraScript.h
@@ -729,7 +729,7 @@ public:
 	bool bUsePythonScriptConversion = false;
 
 	/** Reference to a python script that is executed when the user updates from a previous version to this version. */
-	UPROPERTY(EditAnywhere, Category = Script, meta=(EditCondition="bUsePythonScriptConversion", EditConditionHides))
+	UPROPERTY(EditAnywhere, Category = Script, meta=(EditCondition="bUsePythonScriptConversion", EditConditionHides, SegmentedDisplay))
 	ENiagaraPythonUpdateScriptReference ConversionScriptExecution = ENiagaraPythonUpdateScriptReference::ScriptAsset;
 
 	/** Python script to run when converting this script to the recommended deprecation update script. */
@@ -761,7 +761,7 @@ public:
 	FText DebugDrawMessage;
 
 	/* Defines if this script is visible to the user when searching for modules to add to an emitter.  */
-	UPROPERTY(AssetRegistrySearchable, EditAnywhere, Category = Script)
+	UPROPERTY(AssetRegistrySearchable, EditAnywhere, Category = Script, meta = (SegmentedDisplay))
 	ENiagaraScriptLibraryVisibility LibraryVisibility;
 
 	/** The mode to use when deducing the type of numeric output pins from the types of the input pins. */
