* Surface DI function descriptions as tooltips in search menu
* Updated occlusion DI functions for better consistency when it comes to naming and tooltips

#tests Ran the occulsion DI engine test and created a new scratch module to test that everything still works as before

diff --git a/Engine/Plugins/FX/Niagara/Shaders/Private/NiagaraDataInterfaceOcclusion.ush b/Engine/Plugins/FX/Niagara/Shaders/Private/NiagaraDataInterfaceOcclusion.ush
--- a/Engine/Plugins/FX/Niagara/Shaders/Private/NiagaraDataInterfaceOcclusion.ush
+++ b/Engine/Plugins/FX/Niagara/Shaders/Private/NiagaraDataInterfaceOcclusion.ush
@@ -7,7 +7,7 @@ float3				{ParameterName}_SystemLWCTile;
 Texture2D<float4>	{ParameterName}_CloudVolumetricTexture;
 SamplerState		{ParameterName}_CloudVolumetricTextureSampler;
 
-void QueryOcclusionFactorWithRectangleGPU_{ParameterName}(in float3 In_SampleCenterWorldPos, in float In_SampleWindowWidthWorld, in float In_SampleWindowHeightWorld, in float In_SampleSteps, out float Out_VisibilityFraction, out float Out_SampleFraction)
+void QueryOcclusionFactorWithRectangle_{ParameterName}(in float3 In_SampleCenterWorldPos, in float In_SampleWindowWidthWorld, in float In_SampleWindowHeightWorld, in float In_SampleSteps, out float Out_VisibilityFraction, out float Out_SampleFraction)
 {
 	FLWCVector3 LwcSamplePos = MakeLWCVector3({ParameterName}_SystemLWCTile, In_SampleCenterWorldPos);
 	float3 ResolvedLwcSamplePos = DFFastToTranslatedWorld(DFFromTileOffset_Hack(LwcSamplePos), PrimaryView.PreViewTranslation);
@@ -61,6 +61,11 @@ void QueryOcclusionFactorWithRectangleGPU_{ParameterName}(in float3 In_SampleCen
 	Out_SampleFraction = Steps == 0 ? 0 : (TotalSamples / (Steps * Steps));
 }
 
+void QueryOcclusionFactorWithRectangleGPU_{ParameterName}(in float3 In_SampleCenterWorldPos, in float In_SampleWindowWidthWorld, in float In_SampleWindowHeightWorld, in float In_SampleSteps, out float Out_VisibilityFraction, out float Out_SampleFraction)
+{
+	QueryOcclusionFactorWithRectangle_{ParameterName}(In_SampleCenterWorldPos, In_SampleWindowWidthWorld, In_SampleWindowHeightWorld, In_SampleSteps, Out_VisibilityFraction, Out_SampleFraction);
+}
+
 void QueryOcclusionFactorWithCircleGPU_{ParameterName}(in float3 In_SampleCenterWorldPos, in float In_SampleWindowDiameterWorld, in float In_SampleRays, in float In_SampleStepsPerRay, out float Out_VisibilityFraction, out float Out_SampleFraction)
 {
 	const float PI = 3.14159265;
diff --git a/Engine/Plugins/FX/Niagara/Source/Niagara/Private/NiagaraDataInterfaceOcclusion.cpp b/Engine/Plugins/FX/Niagara/Source/Niagara/Private/NiagaraDataInterfaceOcclusion.cpp
--- a/Engine/Plugins/FX/Niagara/Source/Niagara/Private/NiagaraDataInterfaceOcclusion.cpp
+++ b/Engine/Plugins/FX/Niagara/Source/Niagara/Private/NiagaraDataInterfaceOcclusion.cpp
@@ -19,7 +19,8 @@
 namespace NDIOcclusionImpl
 {
 	static const TCHAR*	TemplateShaderFilePath = TEXT("/Plugin/FX/Niagara/Private/NiagaraDataInterfaceOcclusion.ush");
-	static const FName	GetCameraOcclusionRectangleName(TEXT("QueryOcclusionFactorWithRectangleGPU"));
+	static const FName	Deprecated_GetCameraOcclusionRectangleName(TEXT("QueryOcclusionFactorWithRectangleGPU"));
+	static const FName	GetCameraOcclusionRectangleName(TEXT("QueryOcclusionFactorWithRectangle"));
 	static const FName	Deprecated_GetCameraOcclusionCircleName(TEXT("QueryOcclusionFactorWithCircleGPU"));
 	static const FName	QueryOcclusionFactorWithCircleName(TEXT("QueryOcclusionFactorWithCircle"));
 	static const FName	QueryCloudOcclusionWithCircleName(TEXT("QueryCloudOcclusionWithCircle"));
@@ -79,7 +80,9 @@ void UNiagaraDataInterfaceOcclusion::GetFunctionsInternal(TArray<FNiagaraFunctio
 
 	{
 		FNiagaraFunctionSignature& Sig = OutFunctions.Add_GetRef(DefaultSig);
-		Sig.Name = GetCameraOcclusionRectangleName;
+		FNiagaraFunctionSignature& SigNew = OutFunctions.Add_GetRef(DefaultSig);
+		Sig.Name = Deprecated_GetCameraOcclusionRectangleName;
+		Sig.bSoftDeprecatedFunction = true;
 		Sig.SetDescription(LOCTEXT("GetCameraOcclusionRectFunctionDescription", "This function returns the occlusion factor of a sprite. It samples the depth buffer in a rectangular grid around the given world position and compares each sample with the camera distance."));
 		Sig.AddInput(FNiagaraVariable(FNiagaraTypeDefinition::GetPositionDef(), TEXT("Sample Center World Position")), LOCTEXT("RectCenterPosDescription", "This world space position where the center of the sample rectangle should be."));
 		Sig.AddInput(FNiagaraVariable(FNiagaraTypeDefinition::GetFloatDef(), TEXT("Sample Window Width World")), LOCTEXT("SampleWindowWidthWorldDescription", "The total width of the sample rectangle in world space.\nIf the particle is a camera-aligned sprite then this is the sprite width."));
@@ -87,6 +90,10 @@ void UNiagaraDataInterfaceOcclusion::GetFunctionsInternal(TArray<FNiagaraFunctio
 		Sig.AddInput(FNiagaraVariable(FNiagaraTypeDefinition::GetFloatDef(), TEXT("Sample Steps Per Line")), LOCTEXT("StepsPerLineDescription", "The number of samples to take horizontally. The total number of samples is this value squared."));
 		Sig.AddOutput(FNiagaraVariable(FNiagaraTypeDefinition::GetFloatDef(), TEXT("Visibility Fraction")), VisibilityFractionDescription);
 		Sig.AddOutput(FNiagaraVariable(FNiagaraTypeDefinition::GetFloatDef(), TEXT("Sample Fraction")), SampleFractionDescription);
+
+		SigNew = Sig;
+		SigNew.bSoftDeprecatedFunction = false;
+		SigNew.Name = GetCameraOcclusionRectangleName;
 	}
 
 	{
@@ -104,7 +111,7 @@ void UNiagaraDataInterfaceOcclusion::GetFunctionsInternal(TArray<FNiagaraFunctio
 	{
 		FNiagaraFunctionSignature& Sig = OutFunctions.Add_GetRef(DefaultSig);
 		Sig.Name = QueryOcclusionFactorWithCircleName;
-		Sig.SetDescription(LOCTEXT("QueryCloudOcclusionWithCircleDescription", "Returns the cloud occlusion factor for the world position. "));
+		Sig.SetDescription(LOCTEXT("QueryCloudOcclusionWithCircleDescription", "This function returns the occlusion factor of a sprite. It samples the depth buffer in a spiral pattern around the given world position and compares each sample with the camera distance."));
 		Sig.Inputs.Emplace(FNiagaraTypeDefinition::GetPositionDef(), TEXT("WorldPosition"));
 		Sig.Inputs.Emplace(FNiagaraTypeDefinition::GetFloatDef(), TEXT("WorldDiameter"));
 		Sig.Inputs.Emplace_GetRef(FNiagaraTypeDefinition::GetBoolDef(), TEXT("IncludeCenterSample")).SetValue(true);
@@ -174,6 +181,7 @@ bool UNiagaraDataInterfaceOcclusion::GetFunctionHLSL(const FNiagaraDataInterface
 		QueryOcclusionFactorWithCircleName,
 		QueryCloudOcclusionWithCircleName,
 		Deprecated_GetCameraOcclusionCircleName,
+		Deprecated_GetCameraOcclusionRectangleName,
 	};
 	return ValidGpuFunctions.Contains(FunctionInfo.DefinitionName);
 }
diff --git a/Engine/Plugins/FX/Niagara/Source/NiagaraEditor/Private/EdGraphSchema_Niagara.cpp b/Engine/Plugins/FX/Niagara/Source/NiagaraEditor/Private/EdGraphSchema_Niagara.cpp
--- a/Engine/Plugins/FX/Niagara/Source/NiagaraEditor/Private/EdGraphSchema_Niagara.cpp
+++ b/Engine/Plugins/FX/Niagara/Source/NiagaraEditor/Private/EdGraphSchema_Niagara.cpp
@@ -812,7 +812,7 @@ PRAGMA_ENABLE_DEPRECATION_WARNINGS
 					continue;
 
 				UNiagaraNodeFunctionCall* FuncNode = NewObject<UNiagaraNodeFunctionCall>(OwnerOfTemporaries);
-				AddNewNodeMenuAction(NewActions, FuncNode, FText::FromString(FName::NameToDisplayString(Sig.GetNameString(), false)), ENiagaraMenuSections::General, {MenuCat.ToString()}, FText::GetEmpty(), FText::GetEmpty());
+				AddNewNodeMenuAction(NewActions, FuncNode, FText::FromString(FName::NameToDisplayString(Sig.GetNameString(), false)), ENiagaraMenuSections::General, {MenuCat.ToString()}, Sig.GetDescription(), FText::GetEmpty());
 				FuncNode->Signature = Sig;
 			}
 		}
