Eliminate MSVC vtable padding and other packing cleanup in the associated classes

#rnx

diff --git a/Engine/Plugins/Animation/ControlRig/Source/ControlRig/Public/Tools/ControlRigPose.h b/Engine/Plugins/Animation/ControlRig/Source/ControlRig/Public/Tools/ControlRigPose.h
--- a/Engine/Plugins/Animation/ControlRig/Source/ControlRig/Public/Tools/ControlRigPose.h
+++ b/Engine/Plugins/Animation/ControlRig/Source/ControlRig/Public/Tools/ControlRigPose.h
@@ -25,8 +25,8 @@ struct CONTROLRIG_API FRigControlCopy
 		FRigControlCopy()
 		: Name(NAME_None)
 		, ControlType(ERigControlType::Transform)
-		, Value()
 		, ParentKey()
+		, Value()
 		, OffsetTransform(FTransform::Identity)
 		, ParentTransform(FTransform::Identity)
 		, LocalTransform(FTransform::Identity)
@@ -47,7 +47,6 @@ struct CONTROLRIG_API FRigControlCopy
 		LocalTransform = InHierarchy->GetTransform(InControlElement, ERigTransformType::CurrentLocal);
 		GlobalTransform = InHierarchy->GetTransform(InControlElement, ERigTransformType::CurrentGlobal);
 	}
-	virtual ~FRigControlCopy() {}
 
 	UPROPERTY(BlueprintReadWrite, EditAnywhere, Category = "Names")
 	FName Name;
@@ -55,12 +54,12 @@ struct CONTROLRIG_API FRigControlCopy
 	UPROPERTY(BlueprintReadWrite, EditAnywhere, Category = "Type")
 	ERigControlType ControlType;
 
-	UPROPERTY()
-	FRigControlValue Value;
-
 	UPROPERTY(BlueprintReadWrite, EditAnywhere, Category = "Names")
 	FRigElementKey ParentKey;
 
+	UPROPERTY()
+	FRigControlValue Value;
+
 	UPROPERTY(BlueprintReadWrite, EditAnywhere, Category = "Transforms")
 	FTransform OffsetTransform;
 
diff --git a/Engine/Plugins/Experimental/Mutable/Source/MutableRuntime/Internal/MuR/SystemPrivate.h b/Engine/Plugins/Experimental/Mutable/Source/MutableRuntime/Internal/MuR/SystemPrivate.h
--- a/Engine/Plugins/Experimental/Mutable/Source/MutableRuntime/Internal/MuR/SystemPrivate.h
+++ b/Engine/Plugins/Experimental/Mutable/Source/MutableRuntime/Internal/MuR/SystemPrivate.h
@@ -1888,7 +1888,7 @@ namespace mu
     public:
 
         Private( Ptr<Settings>, const TSharedPtr<ExtensionDataStreamer>& );
-        virtual ~Private();
+        ~Private();
 
         //-----------------------------------------------------------------------------------------
         //! Own interface
@@ -1919,15 +1919,22 @@ namespace mu
 
 		TSharedPtr<ExternalResourceProvider> ExternalResourceProvider;
 
+		/** Counter used to generate unique IDs for every new instance created in the system. */
+        Instance::ID LastInstanceID = 0;
+
+	private:
+		/** This flag is turned on when a streaming error or similar happens. Results are not usable.
+		* This should only happen in-editor.
+		*/
+		bool bUnrecoverableError = false;
+
+	public:
 		/** If this is set, it will be tried first instead of the internal formatting function. */
 		FImageOperator::FImagePixelFormatFunc ImagePixelFormatOverride;
 
 		/** */
 		FWorkingMemoryManager WorkingMemoryManager;
 
-		/** Counter used to generate unique IDs for every new instance created in the system. */
-        Instance::ID LastInstanceID = 0;
-
 		/** The pointer returned by this function is only valid for the duration of the current mutable operation. */
 		inline FLiveInstance* FindLiveInstance(Instance::ID id);
 
@@ -1944,12 +1951,6 @@ namespace mu
 
 		/** Owned by this system. */
 		TSharedPtr<ExtensionDataStreamer> ExtensionDataStreamer = nullptr;
-
-		/** This flag is turned on when a streaming error or similar happens. Results are not usable.
-		* This should only happen in-editor.
-		*/
-		bool bUnrecoverableError = false;
-
     };
 
 }
diff --git a/Engine/Plugins/FX/Niagara/Source/Niagara/Private/NiagaraSystemInstance.cpp b/Engine/Plugins/FX/Niagara/Source/Niagara/Private/NiagaraSystemInstance.cpp
--- a/Engine/Plugins/FX/Niagara/Source/Niagara/Private/NiagaraSystemInstance.cpp
+++ b/Engine/Plugins/FX/Niagara/Source/Niagara/Private/NiagaraSystemInstance.cpp
@@ -104,13 +104,13 @@ static FAutoConsoleVariableRef CVarNiagaraAllowDeferredReset(
 
 FNiagaraSystemInstance::FNiagaraSystemInstance(UWorld& InWorld, UNiagaraSystem& InSystem, FNiagaraUserRedirectionParameterStore* InOverrideParameters,
                                                USceneComponent* InAttachComponent, ENiagaraTickBehavior InTickBehavior, bool bInPooled)
-	: SystemInstanceIndex(INDEX_NONE)
+	: TickBehavior(InTickBehavior)
+	  , SystemInstanceIndex(INDEX_NONE)
 	  , SignificanceIndex(INDEX_NONE)
 	  , World(&InWorld)
 	  , System(&InSystem)
 	  , OverrideParameters(InOverrideParameters)
 	  , AttachComponent(InAttachComponent)
-	  , TickBehavior(InTickBehavior)
 	  , Age(0.0f)
 	  , LastRenderTime(0.0f)
 	  , TickCount(0)
diff --git a/Engine/Plugins/FX/Niagara/Source/Niagara/Public/NiagaraSystemInstance.h b/Engine/Plugins/FX/Niagara/Source/Niagara/Public/NiagaraSystemInstance.h
--- a/Engine/Plugins/FX/Niagara/Source/Niagara/Public/NiagaraSystemInstance.h
+++ b/Engine/Plugins/FX/Niagara/Source/Niagara/Public/NiagaraSystemInstance.h
@@ -67,7 +67,14 @@ private:
 #endif
 };
 
-class FNiagaraSystemInstance 
+// Hacky base class to avoid 8 bytes of padding after the vtable
+class FNiagaraSystemInstanceFixLayout
+{
+public:
+	virtual ~FNiagaraSystemInstanceFixLayout() = default;
+};
+
+class FNiagaraSystemInstance : public FNiagaraSystemInstanceFixLayout
 {
 	friend class FNiagaraSystemSimulation;
 	friend class FNiagaraGPUSystemTick;
@@ -88,7 +95,7 @@ public:
 public:
 
 	/** Defines modes for resetting the System instance. */
-	enum class EResetMode
+	enum class EResetMode : uint8
 	{
 		/** Resets the System instance and simulations. */
 		ResetAll,
@@ -102,6 +109,13 @@ public:
 
 	ENiagaraSystemInstanceState SystemInstanceState = ENiagaraSystemInstanceState::None;
 
+	ENiagaraTickBehavior TickBehavior;
+
+private:
+	/** If async work was running when we request an Activate we will store the reset mode and perform in finalize to avoid stalling the GameThread. */
+	EResetMode DeferredResetMode = EResetMode::None;
+
+public:
 	FORCEINLINE bool GetAreDataInterfacesInitialized() const { return bDataInterfacesInitialized; }
 
 	/** Creates a new Niagara system instance. */
@@ -489,8 +503,6 @@ private:
 	FTransform WorldTransform;
 	TOptional<FVector> PreviousLocation;
 
-	ENiagaraTickBehavior TickBehavior;
-
 	/** The age of the System instance. */
 	float Age;
 
@@ -513,6 +525,9 @@ private:
 	int32 WarmupTickCount = -1;
 	float WarmupTickDelta = 0;
 	
+	FName IDName;
+	FNiagaraSystemInstanceID ID;
+
 	//-TODO:Stateless:
 	//TArray<TSharedRef<FNiagaraEmitterInstanceImpl, ESPMode::ThreadSafe>> Emitters;
 	TArray<FNiagaraEmitterInstanceRef> Emitters;
@@ -534,9 +549,7 @@ private:
 	TMap<FGuid, TSharedPtr<TArray<TSharedPtr<struct FNiagaraScriptDebuggerInfo, ESPMode::ThreadSafe>>, ESPMode::ThreadSafe>> CapturedFrames;
 #endif
 
-	FNiagaraSystemInstanceID ID;
-	FName IDName;
-	
+
 	/** Per instance data for any data interfaces requiring it. */
 	TArray<uint8, TAlignedHeapAllocator<16>> DataInterfaceInstanceData;
 	TArray<int32> PreTickDataInterfaces;
@@ -562,44 +575,46 @@ private:
 	FNiagaraOwnerParameters OwnerParameters[ParameterBufferCount];
 	TArray<FNiagaraEmitterParameters> EmitterParameters;
 
-	/** Used for double buffered global/system/emitter parameters */
-	uint32 CurrentFrameIndex : 1;
-	uint32 ParametersValid : 1;
-
 	// registered events for each of the emitters
 	typedef TPair<FName, FName> EmitterEventKey;
 	typedef TMap<EmitterEventKey, FNiagaraDataSet*> EventDataSetMap;
 	EventDataSetMap EmitterEventDataSetMap;
 
+	/** Used for double buffered global/system/emitter parameters */
+	uint8 CurrentFrameIndex : 1;
+	uint8 ParametersValid : 1;
+
 	/** Indicates whether this instance must update itself rather than being batched up as most instances are. */
-	uint32 bSolo : 1;
-	uint32 bForceSolo : 1;
+	uint8 bSolo : 1;
+	uint8 bForceSolo : 1;
 
-	uint32 bNotifyOnCompletion : 1;
+	uint8 bNotifyOnCompletion : 1;
 
 	/** If this system has emitters that will run GPU Simulations */
-	uint32 bHasGPUEmitters : 1;
+	uint8 bHasGPUEmitters : 1;
 	/** The system contains data interfaces that can have tick group prerequisites. */
-	uint32 bDataInterfacesHaveTickPrereqs : 1;
+	uint8 bDataInterfacesHaveTickPrereqs : 1;
 	/** The system contains data interfaces that can have tick group post requisites. */
-	uint32 bDataInterfacesHaveTickPostreqs : 1;
+	uint8 bDataInterfacesHaveTickPostreqs : 1;
 
-	uint32 bDataInterfacesInitialized : 1;
+	uint8 bDataInterfacesInitialized : 1;
 
-	uint32 bAlreadyBound : 1;
+	uint8 bAlreadyBound : 1;
 
-	uint32 bLODDistanceIsValid : 1;
-	uint32 bLODDistanceIsOverridden : 1;
+	uint8 bLODDistanceIsValid : 1;
+	uint8 bLODDistanceIsOverridden : 1;
 
 	/** True if the system instance is pooled. Prevents unbinding of parameters on completing the system */
-	uint32 bPooled : 1;
+	uint8 bPooled : 1;
 
 #if WITH_EDITOR
-	uint32 bNeedsUIResync : 1;
+	uint8 bNeedsUIResync : 1;
 #endif
 
-	/** If async work was running when we request an Activate we will store the reset mode and perform in finalize to avoid stalling the GameThread. */
-	EResetMode DeferredResetMode = EResetMode::None;
+public:
+	uint8 GPUParamIncludeInterpolation : 1 = false;
+
+private:
 
 	/** Graph event to track pending concurrent work. */
 	FGraphEventRef ConcurrentTickGraphEvent;
@@ -645,7 +660,6 @@ public:
 	uint32 ActiveGPUEmitterCount = 0;
 
 	int32 GPUDataInterfaceInstanceDataSize = 0;
-	bool GPUParamIncludeInterpolation = false;
 	TArray<TPair<TWeakObjectPtr<UNiagaraDataInterface>, int32>> GPUDataInterfaces;
 
 	struct FInstanceParameters
@@ -682,12 +696,13 @@ public:
 		}
 	};
 
+	FRandomStream	SystemState_RandomStream;
+
 	FInstanceParameters GatheredInstanceParameters;
 
 	void InitSystemState();
 	void TickSystemState();
 
-	FRandomStream	SystemState_RandomStream;
 	int32			SystemState_LoopCount = 0;
 	float			SystemState_CurrentLoopDuration = 0.0f;
 	float			SystemState_CurrentLoopDelay = 0.0f;
diff --git a/Engine/Plugins/Runtime/AudioModulation/Source/AudioModulation/Public/SoundModulationGenerator.h b/Engine/Plugins/Runtime/AudioModulation/Source/AudioModulation/Public/SoundModulationGenerator.h
--- a/Engine/Plugins/Runtime/AudioModulation/Source/AudioModulation/Public/SoundModulationGenerator.h
+++ b/Engine/Plugins/Runtime/AudioModulation/Source/AudioModulation/Public/SoundModulationGenerator.h
@@ -19,7 +19,14 @@ namespace AudioModulation
 	class IGenerator;
 	using FGeneratorPtr = TUniquePtr<IGenerator>;
 
-	class AUDIOMODULATION_API IGenerator
+	// Hacky base class to avoid 8 bytes of padding after the vtable
+	class IGeneratorFixLayout
+	{
+	public:
+		virtual ~IGeneratorFixLayout() = default;
+	};
+
+	class AUDIOMODULATION_API IGenerator : public IGeneratorFixLayout
 	{
 	public:
 		virtual ~IGenerator() = default;
diff --git a/Engine/Plugins/Runtime/GameplayAbilities/Source/GameplayAbilities/Public/Abilities/GameplayAbilityTargetTypes.h b/Engine/Plugins/Runtime/GameplayAbilities/Source/GameplayAbilities/Public/Abilities/GameplayAbilityTargetTypes.h
--- a/Engine/Plugins/Runtime/GameplayAbilities/Source/GameplayAbilities/Public/Abilities/GameplayAbilityTargetTypes.h
+++ b/Engine/Plugins/Runtime/GameplayAbilities/Source/GameplayAbilities/Public/Abilities/GameplayAbilityTargetTypes.h
@@ -294,17 +294,26 @@ struct TStructOpsTypeTraits<FGameplayAbilityTargetDataHandle> : public TStructOp
 	};
 };
 
+// Hacky base class to avoid 8 bytes of padding after the vtable
+struct FGameplayAbilityTargetingLocationInfoFixLayout
+{
+	virtual ~FGameplayAbilityTargetingLocationInfoFixLayout() = default;
+};
+
 /** Structure that stores a location in one of several different formats */
 USTRUCT(BlueprintType)
 struct GAMEPLAYABILITIES_API FGameplayAbilityTargetingLocationInfo
+#if CPP
+	: public FGameplayAbilityTargetingLocationInfoFixLayout
+#endif
 {
 	GENERATED_USTRUCT_BODY()
 
 	FGameplayAbilityTargetingLocationInfo()
-	: LocationType(EGameplayAbilityTargetingLocationType::LiteralTransform)
-	, SourceActor(nullptr)
+	: SourceActor(nullptr)
 	, SourceComponent(nullptr)
 	, SourceAbility(nullptr)
+	, LocationType(EGameplayAbilityTargetingLocationType::LiteralTransform)
 	{ }
 
 	virtual ~FGameplayAbilityTargetingLocationInfo() {}
@@ -328,14 +337,6 @@ struct GAMEPLAYABILITIES_API FGameplayAbilityTargetingLocationInfo
 	/** Initializes new actor list target data, and sets this as the origin */
 	FGameplayAbilityTargetDataHandle MakeTargetDataHandleFromActors(const TArray<TWeakObjectPtr<AActor>>& TargetActors, bool OneActorPerHandle = false) const;
 
-	/** Type of location used - will determine what data is transmitted over the network and what fields are used when calculating position. */
-	UPROPERTY(BlueprintReadWrite, meta = (ExposeOnSpawn = true), Category = Targeting)
-	TEnumAsByte<EGameplayAbilityTargetingLocationType::Type> LocationType;
-
-	/** A literal world transform can be used, if one has been calculated outside of the actor using the ability. */
-	UPROPERTY(BlueprintReadWrite, meta = (ExposeOnSpawn = true), Category = Targeting)
-	FTransform LiteralTransform;
-
 	/** A source actor is needed for Actor-based targeting, but not for Socket-based targeting. */
 	UPROPERTY(BlueprintReadWrite, meta = (ExposeOnSpawn = true), Category = Targeting)
 	TObjectPtr<AActor> SourceActor;
@@ -348,10 +349,18 @@ struct GAMEPLAYABILITIES_API FGameplayAbilityTargetingLocationInfo
 	UPROPERTY(BlueprintReadWrite, meta = (ExposeOnSpawn = true), Category = Targeting)
 	TObjectPtr<UGameplayAbility> SourceAbility;
 
+	/** A literal world transform can be used, if one has been calculated outside of the actor using the ability. */
+	UPROPERTY(BlueprintReadWrite, meta = (ExposeOnSpawn = true), Category = Targeting)
+	FTransform LiteralTransform;
+
 	/** If SourceComponent is valid, this is the name of the socket transform that will be used. If no Socket is provided, SourceComponent's transform will be used. */
 	UPROPERTY(BlueprintReadWrite, meta = (ExposeOnSpawn = true), Category = Targeting)
 	FName SourceSocketName;
 
+	/** Type of location used - will determine what data is transmitted over the network and what fields are used when calculating position. */
+	UPROPERTY(BlueprintReadWrite, meta = (ExposeOnSpawn = true), Category = Targeting)
+	TEnumAsByte<EGameplayAbilityTargetingLocationType::Type> LocationType;
+
 	/** Optimized serialize function */
 	bool NetSerialize(FArchive& Ar, class UPackageMap* Map, bool& bOutSuccess);
 };
diff --git a/Engine/Plugins/Runtime/GeometryCache/Source/GeometryCache/Public/GeometryCacheSceneProxy.h b/Engine/Plugins/Runtime/GeometryCache/Source/GeometryCache/Public/GeometryCacheSceneProxy.h
--- a/Engine/Plugins/Runtime/GeometryCache/Source/GeometryCache/Public/GeometryCacheSceneProxy.h
+++ b/Engine/Plugins/Runtime/GeometryCache/Source/GeometryCache/Public/GeometryCacheSceneProxy.h
@@ -140,6 +140,13 @@ public:
 	void Init(FRHICommandListBase& RHICmdList, const FVertexBuffer* PositionBuffer, const FVertexBuffer* MotionBlurDataBuffer, const FVertexBuffer* TangentXBuffer, const FVertexBuffer* TangentZBuffer, const FVertexBuffer* TextureCoordinateBuffer, const FVertexBuffer* ColorBuffer);
 };
 
+// Hacky base class to avoid 8 bytes of padding after the vtable
+class FGeomCacheTrackProxyFixLayout
+{
+public:
+	virtual ~FGeomCacheTrackProxyFixLayout() = default;
+};
+
 /**
  * This the track proxy has some "double double buffering" going on.
  * First we keep two mesh frames. The one just before the current time and the one just after the current time. This is the full mesh and
@@ -147,7 +154,7 @@ public:
  * Secondly we have two position buffers. The one for the current rendered frame and the one from the previous rendered frame (this is not the same as
  * the mesh frame, the mesh may be at say 10 fps then get interpolated to 60 fps rendered frames)
  */
-class GEOMETRYCACHE_API FGeomCacheTrackProxy
+class GEOMETRYCACHE_API FGeomCacheTrackProxy : FGeomCacheTrackProxyFixLayout
 {
 public:
 
@@ -241,6 +248,8 @@ public:
 	float PositionBufferFrameTimes[2]; // Exact time after interpolation of the positions in the position buffer.
 	uint32 CurrentPositionBufferIndex; // CurrentPositionBufferIndex%2  is the last updated position buffer
 
+	int32 UploadedSampleIndex;
+
 	FGeomCacheTangentBuffer TangentXBuffer;
 	FGeomCacheTangentBuffer TangentZBuffer;
 	FGeomCacheVertexBuffer TextureCoordinatesBuffer;
@@ -252,13 +261,11 @@ public:
 	/** Vertex factory for this Track */
 	FGeomCacheVertexFactory VertexFactory;
 
-	/** World Matrix for this Track */
-	FMatrix WorldMatrix;
-
 	/** The GeometryCacheTrack to which the proxy is associated */
 	UGeometryCacheTrack* Track;
 
-	int32 UploadedSampleIndex;
+	/** World Matrix for this Track */
+	FMatrix WorldMatrix;
 
 	/** Flag to indicate which frame mesh data was selected during the update */
 	bool bNextFrameMeshDataSelected;
diff --git a/Engine/Plugins/Runtime/GeometryProcessing/Source/DynamicMesh/Public/MeshRefinerBase.h b/Engine/Plugins/Runtime/GeometryProcessing/Source/DynamicMesh/Public/MeshRefinerBase.h
--- a/Engine/Plugins/Runtime/GeometryProcessing/Source/DynamicMesh/Public/MeshRefinerBase.h
+++ b/Engine/Plugins/Runtime/GeometryProcessing/Source/DynamicMesh/Public/MeshRefinerBase.h
@@ -15,11 +15,18 @@ namespace Geometry
 
 class FDynamicMeshChangeTracker;
 
+// Hacky base class to avoid 8 bytes of padding after the vtable
+class FMeshRefinerBaseFixLayout
+{
+public:
+	virtual ~FMeshRefinerBaseFixLayout() = default;
+};
+
 /**
  * This is a base class that implements common functionality for various triangle mesh resampling strategies
  * (ie FRemesher and FReducer). You probably should not use this class directly.
  */
-class DYNAMICMESH_API FMeshRefinerBase
+class DYNAMICMESH_API FMeshRefinerBase : public FMeshRefinerBaseFixLayout
 {
 protected:
 	/** Mesh that will be refined */
@@ -31,6 +38,8 @@ protected:
 	/** Vertices can be projected onto this surface when they are modified */
 	IProjectionTarget* ProjTarget = nullptr;
 
+	FDynamicMeshChangeTracker* ActiveChangeTracker = nullptr;
+
 	FMeshRefinerBase(FDynamicMesh3* MeshIn)
 	{
 		this->Mesh = MeshIn;
@@ -43,14 +52,7 @@ protected:
 
 public:
 
-	/**
-	 * If true, then when two Fixed vertices have the same non-invalid SetID,
-	 * we treat them as not fixed and allow collapse
-	 */
-	bool AllowCollapseFixedVertsWithSameSetID = true;
-
-
-	enum class EVertexControl
+	enum class EVertexControl : uint8
 	{
 		AllowAll = 0,
 		NoSmooth = 1,
@@ -63,9 +65,14 @@ public:
 	 */
 	TFunction<EVertexControl(int)> VertexControlF = nullptr;
 
+	/** Set this to be able to cancel running Remesher/Reducer*/
+	FProgressCancel* Progress = nullptr;
+
+	/** This is a debugging aid, will break to debugger if these edges are touched, in debug builds */
+	TArray<int> DebugEdges;
 
 	/** Options for projecting vertices onto target surface */
-	enum class ETargetProjectionMode
+	enum class ETargetProjectionMode : uint8
 	{
 		NoProjection = 0,		// disable projection
 		AfterRefinement = 1,	// do all projection after the refine/smooth pass
@@ -78,14 +85,11 @@ public:
 	ETargetProjectionMode ProjectionMode = ETargetProjectionMode::NoProjection;
 
 
-
-	/** Set this to be able to cancel running Remesher/Reducer*/
-	FProgressCancel* Progress = nullptr;
-
-
-
-	/** This is a debugging aid, will break to debugger if these edges are touched, in debug builds */
-	TArray<int> DebugEdges;
+	/**
+	* If true, then when two Fixed vertices have the same non-invalid SetID,
+	* we treat them as not fixed and allow collapse
+	*/
+	bool AllowCollapseFixedVertsWithSameSetID = true;
 
 	/** Set to true to profile various passes @todo re-enable this! */
 	bool ENABLE_PROFILING = false;
@@ -289,7 +293,6 @@ public:
 	void SetMeshChangeTracker(FDynamicMeshChangeTracker* Tracker);
 
 protected:
-	FDynamicMeshChangeTracker* ActiveChangeTracker = nullptr;
 	virtual void SaveTriangleBeforeModify(int32 TriangleID);
 	virtual void SaveEdgeBeforeModify(int32 EdgeID);
 	virtual void SaveVertexTrianglesBeforeModify(int32 VertexID);
diff --git a/Engine/Plugins/Runtime/GeometryProcessing/Source/DynamicMesh/Public/Operations/ExtrudeMesh.h b/Engine/Plugins/Runtime/GeometryProcessing/Source/DynamicMesh/Public/Operations/ExtrudeMesh.h
--- a/Engine/Plugins/Runtime/GeometryProcessing/Source/DynamicMesh/Public/Operations/ExtrudeMesh.h
+++ b/Engine/Plugins/Runtime/GeometryProcessing/Source/DynamicMesh/Public/Operations/ExtrudeMesh.h
@@ -18,6 +18,13 @@ namespace Geometry
 class FDynamicMesh3;
 class FMeshNormals;
 
+// Hacky base class to avoid 8 bytes of padding after the vtable
+class FExtrudeMeshFixLayout
+{
+public:
+	virtual ~FExtrudeMeshFixLayout() = default;
+};
+
 /**
  * Note: FExtrudeMesh might someday be removed. Consider using FOffsetMeshRegion instead,
  * which can extrude all or part of a mesh.
@@ -29,7 +36,7 @@ class FMeshNormals;
  * 
  * Each quad of the border loop is assigned it's own normal and UVs (ie each is a separate UV-island)
  */
-class DYNAMICMESH_API FExtrudeMesh
+class DYNAMICMESH_API FExtrudeMesh : public FExtrudeMeshFixLayout
 {
 public:
 
@@ -49,12 +56,11 @@ public:
 	/** if Extrusion is "negative" (ie negative distance, inset, etc) then this value must be set to false or the output will have incorrect winding orientation */
 	bool IsPositiveOffset = true;
 
-	/** quads on the stitch loop are planar-projected and scaled by this amount */
-	float UVScaleFactor = 1.0f;
-
 	/** If true, skip closed components */
 	bool bSkipClosedComponents = false;
 
+	/** quads on the stitch loop are planar-projected and scaled by this amount */
+	float UVScaleFactor = 1.0f;
 
 	//
 	// Outputs
diff --git a/Engine/Plugins/Runtime/GeometryProcessing/Source/DynamicMesh/Public/Operations/MeshPlaneCut.h b/Engine/Plugins/Runtime/GeometryProcessing/Source/DynamicMesh/Public/Operations/MeshPlaneCut.h
--- a/Engine/Plugins/Runtime/GeometryProcessing/Source/DynamicMesh/Public/Operations/MeshPlaneCut.h
+++ b/Engine/Plugins/Runtime/GeometryProcessing/Source/DynamicMesh/Public/Operations/MeshPlaneCut.h
@@ -21,6 +21,13 @@ namespace Geometry
 class FDynamicMesh3;
 template<typename RealType> class TDynamicMeshScalarTriangleAttribute;
 
+// Hacky base class to avoid 8 bytes of padding after the vtable
+class FMeshPlaneCutFixLayout
+{
+public:
+	virtual ~FMeshPlaneCutFixLayout() = default;
+};
+
 /**
  * Cut the Mesh with the Plane. The *positive* side, ie (p-o).n > 0, is removed.
  * If possible, returns boundary loop(s) along cut
@@ -38,7 +45,7 @@ template<typename RealType> class TDynamicMeshScalarTriangleAttribute;
  *	  6) (optionally) change an attribute tag for all triangles on positive side
  *    7) find loops through valid boundary edges (ie connected to splits, or on-plane edges) (if second half was kept, do this separately for each separate mesh ID label)
  */
-class DYNAMICMESH_API FMeshPlaneCut
+class DYNAMICMESH_API FMeshPlaneCut : public FMeshPlaneCutFixLayout
 {
 public:
 
@@ -53,19 +60,19 @@ public:
 	 */
 	TUniqueFunction<bool(int32)> EdgeFilterFunc = nullptr;
 
+	/** Control whether we attempt to auto-simplify the small planar triangles that the plane cut operation tends to generate */
+	bool bSimplifyAlongNewEdges = false;
+
 	bool bCollapseDegenerateEdgesOnCut = true;
-	double DegenerateEdgeTol = FMathd::ZeroTolerance;
 
 	/** UVs on any hole fill surfaces are scaled by this amount */
 	float UVScaleFactor = 1.0f;
 
+	double DegenerateEdgeTol = FMathd::ZeroTolerance;
+
 	/** Tolerance distance for considering a vertex to be 'on plane' */
 	double PlaneTolerance = FMathf::ZeroTolerance * 10.0;
 
-
-	/** Control whether we attempt to auto-simplify the small planar triangles that the plane cut operation tends to generate */
-	bool bSimplifyAlongNewEdges = false;
-
 	/** Settings to apply if bSimplifyAlongNewEdges == true */
 	FLocalPlanarSimplify SimplifySettings;
 
diff --git a/Engine/Plugins/Runtime/GeometryProcessing/Source/DynamicMesh/Public/Operations/OffsetMeshRegion.h b/Engine/Plugins/Runtime/GeometryProcessing/Source/DynamicMesh/Public/Operations/OffsetMeshRegion.h
--- a/Engine/Plugins/Runtime/GeometryProcessing/Source/DynamicMesh/Public/Operations/OffsetMeshRegion.h
+++ b/Engine/Plugins/Runtime/GeometryProcessing/Source/DynamicMesh/Public/Operations/OffsetMeshRegion.h
@@ -18,6 +18,13 @@ class FDynamicMesh3;
 class FMeshNormals;
 class FDynamicMeshChangeTracker;
 
+// Hacky base class to avoid 8 bytes of padding after the vtable
+class FOffsetMeshRegionFixLayout
+{
+public:
+	virtual ~FOffsetMeshRegionFixLayout() = default;
+};
+
 /**
  * FOffsetMeshRegion implements local extrusion/offset of a mesh region. 
  * The selected triangles are separated and then stitched back together, creating
@@ -28,7 +35,7 @@ class FDynamicMeshChangeTracker;
  * 
  * Each quad of the border loop is assigned it's own normal and UVs (ie each is a separate UV-island)
  */
-class DYNAMICMESH_API FOffsetMeshRegion
+class DYNAMICMESH_API FOffsetMeshRegion : public FOffsetMeshRegionFixLayout
 {
 public:
 
@@ -124,6 +131,12 @@ public:
 	 */
 	bool bUVIslandPerGroup = true;
 
+	/** 
+	* If true, Material IDs around the border of the extrude area are propagated "down" the extrusion tube 
+	* This parameter is only supported in EVersion::Version1 and later
+	*/
+	bool bInferMaterialID = true;
+
 	/**
 	 * Split the extrude "tube" into separate groups based on the opening angle between quads.
 	 * This split is done relative to the above group options, so even if the tube would 
@@ -140,13 +153,6 @@ public:
 	 */
 	int SetMaterialID = 0;
 
-	/** 
-	 * If true, Material IDs around the border of the extrude area are propagated "down" the extrusion tube 
-	 * This parameter is only supported in EVersion::Version1 and later
-	 */
-	bool bInferMaterialID = true;
-
-
 	/**
 	 * Support for different versions of the OffsetMeshRegion geometric operation.
 	 * We default to the latest version.
diff --git a/Engine/Plugins/Runtime/GeometryProcessing/Source/DynamicMesh/Public/Operations/PolygroupRemesh.h b/Engine/Plugins/Runtime/GeometryProcessing/Source/DynamicMesh/Public/Operations/PolygroupRemesh.h
--- a/Engine/Plugins/Runtime/GeometryProcessing/Source/DynamicMesh/Public/Operations/PolygroupRemesh.h
+++ b/Engine/Plugins/Runtime/GeometryProcessing/Source/DynamicMesh/Public/Operations/PolygroupRemesh.h
@@ -23,13 +23,18 @@ namespace UE
 namespace Geometry
 {
 
-
+// Hacky base class to avoid 8 bytes of padding after the vtable
+class FPolygroupRemeshFixLayout
+{
+public:
+	virtual ~FPolygroupRemeshFixLayout() = default;
+};
 
 /**
  * PolygroupRemesh -- remesh only considering polygroup features (topological corners and bends on polygroup edges)
  * This can help clean up low poly meshes that have extra vertices along straight edges e.g. after mesh Boolean operations
  */
-class DYNAMICMESH_API FPolygroupRemesh
+class DYNAMICMESH_API FPolygroupRemesh : public FPolygroupRemeshFixLayout
 {
 public:
 
diff --git a/Engine/Plugins/Runtime/GeometryProcessing/Source/DynamicMesh/Public/Sampling/MeshSurfaceSampler.h b/Engine/Plugins/Runtime/GeometryProcessing/Source/DynamicMesh/Public/Sampling/MeshSurfaceSampler.h
--- a/Engine/Plugins/Runtime/GeometryProcessing/Source/DynamicMesh/Public/Sampling/MeshSurfaceSampler.h
+++ b/Engine/Plugins/Runtime/GeometryProcessing/Source/DynamicMesh/Public/Sampling/MeshSurfaceSampler.h
@@ -103,7 +103,7 @@ struct FMeshUVSampleInfo
 
 
 /** Types of query that FMeshSurfaceUVSampler/TMeshSurfaceUVSampler supports */
-enum class EMeshSurfaceSamplerQueryType
+enum class EMeshSurfaceSamplerQueryType : uint8
 {
 	/** Query with arbitrary UV value */
 	UVOnly,
@@ -111,6 +111,12 @@ enum class EMeshSurfaceSamplerQueryType
 	TriangleAndUV
 };
 
+// Hacky base class to avoid 8 bytes of padding after the vtable
+class FMeshSurfaceUVSamplerFixLayout
+{
+public:
+	virtual ~FMeshSurfaceUVSamplerFixLayout() = default;
+};
 
 /**
  * FMeshSurfaceUVSampler computes FMeshUVSampleInfo's on a given mesh at given UV-space positions, this info can then
@@ -122,10 +128,9 @@ enum class EMeshSurfaceSamplerQueryType
  * signature and requires that the sample type is passed as a template parameter. Also the old class is more awkward to
  * use when sampling multiple things, see the documentation of TMeshSurfaceUVSampler for more details.
  */
-class FMeshSurfaceUVSampler
+class FMeshSurfaceUVSampler : public FMeshSurfaceUVSamplerFixLayout
 {
 public:
-	virtual ~FMeshSurfaceUVSampler() {}
 
 	/**
 	 * Initialize the sampler.
@@ -154,10 +159,11 @@ protected:
 	const FDynamicMeshUVOverlay* UVOverlay = nullptr;
 	EMeshSurfaceSamplerQueryType QueryType = EMeshSurfaceSamplerQueryType::TriangleAndUV;
 
+	bool bUVMeshSpatialInitialized = false;
+
 	// BV tree for finding triangle for a given UV. Not always initialized.
 	FDynamicMeshUVMesh UVMeshAdapter;
 	TMeshAABBTree3<FDynamicMeshUVMesh> UVMeshSpatial;
-	bool bUVMeshSpatialInitialized = false;
 	void InitializeUVMeshSpatial();
 };
 
@@ -178,10 +184,9 @@ protected:
  * and call CachedSampleUV(), to avoid expensive BVTree constructions and UV-to-3D recalculation.
  */
 template<typename SampleType>
-class TMeshSurfaceUVSampler
+class TMeshSurfaceUVSampler : public FMeshSurfaceUVSamplerFixLayout
 {
 public:
-	virtual ~TMeshSurfaceUVSampler() {}
 
 	/**
 	 * Configure the sampler.
@@ -230,9 +235,10 @@ public:
 	virtual bool CachedSampleUV(const FMeshUVSampleInfo& CachedSampleInfo, SampleType& ResultOut);
 
 protected:
+	EMeshSurfaceSamplerQueryType QueryType = EMeshSurfaceSamplerQueryType::TriangleAndUV;
+	bool bUVSpatialValid = false;
 	const FDynamicMesh3* Mesh = nullptr;
 	const FDynamicMeshUVOverlay* UVOverlay = nullptr;
-	EMeshSurfaceSamplerQueryType QueryType = EMeshSurfaceSamplerQueryType::TriangleAndUV;
 
 	TUniqueFunction<void(const FMeshUVSampleInfo& SampleInfo, SampleType& SampleValueOut)> ValueFunction;
 
@@ -241,7 +247,6 @@ protected:
 	// BV tree for finding triangle for a given UV. Not always initialized.
 	FDynamicMeshUVMesh UVMeshAdapter;
 	TMeshAABBTree3<FDynamicMeshUVMesh> UVBVTree;
-	bool bUVSpatialValid = false;
 	void InitializeBVTree();
 };
 
diff --git a/Engine/Plugins/Runtime/Harmonix/Source/HarmonixDsp/Private/HarmonixDsp/FusionSampler/FusionVoice.cpp b/Engine/Plugins/Runtime/Harmonix/Source/HarmonixDsp/Private/HarmonixDsp/FusionSampler/FusionVoice.cpp
--- a/Engine/Plugins/Runtime/Harmonix/Source/HarmonixDsp/Private/HarmonixDsp/FusionSampler/FusionVoice.cpp
+++ b/Engine/Plugins/Runtime/Harmonix/Source/HarmonixDsp/Private/HarmonixDsp/FusionSampler/FusionVoice.cpp
@@ -28,13 +28,6 @@ namespace FusionVoice
 const double FFusionVoice::kMaxPitchOffsetCents = 12 * HarmonixDsp::kCentsPerOctave;
 
 FFusionVoice::FFusionVoice()
-	: MySampler(nullptr)
-	, bWaitingForAttack(false)
-	, KeyZone(nullptr)
-	, PitchShifter(nullptr)
-	, OctaveShift(0.0f)
-	, MaxAudioLevel(0.0f)
-	, VoicePool(nullptr)
 {
 
 	ActiveRenderer = MakeShared<FStreamingAudioRendererV2>();
diff --git a/Engine/Plugins/Runtime/Harmonix/Source/HarmonixDsp/Public/HarmonixDsp/FusionSampler/FusionVoice.h b/Engine/Plugins/Runtime/Harmonix/Source/HarmonixDsp/Public/HarmonixDsp/FusionSampler/FusionVoice.h
--- a/Engine/Plugins/Runtime/Harmonix/Source/HarmonixDsp/Public/HarmonixDsp/FusionSampler/FusionVoice.h
+++ b/Engine/Plugins/Runtime/Harmonix/Source/HarmonixDsp/Public/HarmonixDsp/FusionSampler/FusionVoice.h
@@ -27,7 +27,14 @@ class FFusionVoicePool;
 struct FKeyzoneSettings;
 struct FLfoSettings;
 
-class HARMONIXDSP_API alignas(16) FFusionVoice
+// Hacky base class to avoid 8 bytes of padding after the vtable
+class FFusionVoiceFixLayout
+{
+public:
+	virtual ~FFusionVoiceFixLayout() = default;
+};
+
+class HARMONIXDSP_API alignas(16) FFusionVoice : FFusionVoiceFixLayout
 {
 public:
 
@@ -131,9 +138,9 @@ public:
 protected:
 
 	FMidiVoiceId VoiceID;
+	float TargetMidiNote = 0.0f;
 	FFusionSampler* MySampler = nullptr;
 	bool bIsRendererForAlias = false;
-	float TargetMidiNote = 0.0f;
 	uint8 TriggeredMidiNote = 0;
 
 private:
@@ -161,6 +168,7 @@ private:
 	uint32 DebugID;
 
 	float VelocityGain = 0.0f;
+	float MaxAudioLevel = 0.0f;
 
 	const FKeyzoneSettings* KeyZone = nullptr;
 
@@ -184,13 +192,12 @@ private:
 	// we would call this "frequency", but we want
 	// the units to be "octaves" so it can be linear.
 	// the actual frequency is  20*(2^octave)
-	Harmonix::Dsp::Modulators::FModulatorTarget FilterOctaveTarget;
 	float OctaveShift = 0.0f;
+	Harmonix::Dsp::Modulators::FModulatorTarget FilterOctaveTarget;
 
 	TLinearRamper<float> FilterGainRamper;
 
 	TAudioBuffer<float> OutputBuffer;
-	float MaxAudioLevel = 0.0f;
 
 	double FileToOutputSampleRatio = 0.0;
 
diff --git a/Engine/Plugins/Runtime/MeshModelingToolset/Source/MeshModelingTools/Private/Operations/ExtrudeBoundaryEdges.h b/Engine/Plugins/Runtime/MeshModelingToolset/Source/MeshModelingTools/Private/Operations/ExtrudeBoundaryEdges.h
--- a/Engine/Plugins/Runtime/MeshModelingToolset/Source/MeshModelingTools/Private/Operations/ExtrudeBoundaryEdges.h
+++ b/Engine/Plugins/Runtime/MeshModelingToolset/Source/MeshModelingTools/Private/Operations/ExtrudeBoundaryEdges.h
@@ -15,7 +15,14 @@ namespace UE::Geometry
 
 class FDynamicMesh3;
 
-class FExtrudeBoundaryEdges
+// Hacky base class to avoid 8 bytes of padding after the vtable
+class FExtrudeBoundaryEdgesFixLayout
+{
+public:
+	virtual ~FExtrudeBoundaryEdgesFixLayout() = default;
+};
+
+class FExtrudeBoundaryEdges : public FExtrudeBoundaryEdgesFixLayout
 {
 public:
 	using FFrame3d = UE::Geometry::FFrame3d;
@@ -56,6 +63,9 @@ public:
 	/** Whether to calculate local extrude frames and supply them to OffsetPositionFunc */
 	bool bUsePerVertexExtrudeFrames = true;
 
+	/** When generating extrude frames, whether to use unselected neighbors for setting the frame. */
+	bool bAssignAnyBoundaryNeighborToUnmatched = false;
+
 	/**
 	 * Function queried for new vertex positions. ExtrudeFrame origin is Position, unless it is not initialized
 	 *  due to bUsePerVertexExtrudeFrames being false.
@@ -67,9 +77,6 @@ public:
 		return ExtrudeFrame.Frame.FromFramePoint(FVector3d(this->DefaultOffsetDistance, 0, 0) * ExtrudeFrame.Scaling);
 	};
 
-	/** When generating extrude frames, whether to use unselected neighbors for setting the frame. */
-	bool bAssignAnyBoundaryNeighborToUnmatched = false;
-
 	/** 
 	 * If greater than 1, maximal amount by which a vertex can be moved in an attempt to keep edges parallel to
 	 *  original edges while extruding. This "movement" is done by scaling the X axis of the extrude frame.
diff --git a/Engine/Plugins/Runtime/RigVM/Source/RigVM/Public/RigVMCore/RigVMExecuteContext.h b/Engine/Plugins/Runtime/RigVM/Source/RigVM/Public/RigVMCore/RigVMExecuteContext.h
--- a/Engine/Plugins/Runtime/RigVM/Source/RigVM/Public/RigVMCore/RigVMExecuteContext.h
+++ b/Engine/Plugins/Runtime/RigVM/Source/RigVM/Public/RigVMCore/RigVMExecuteContext.h
@@ -244,12 +244,21 @@ struct RIGVM_API FRigVMRuntimeSettings
 	}
 };
 
+// Hacky base class to avoid 8 bytes of padding after the vtable
+struct FRigVMExecuteContextFixLayout
+{
+	virtual ~FRigVMExecuteContextFixLayout() = default;
+};
+
 /**
  * The execute context is used for mutable nodes to
  * indicate execution order.
  */
 USTRUCT(BlueprintType, meta=(DisplayName="Execute Context"))
-struct FRigVMExecuteContext
+struct FRigVMExecuteContext 
+#if CPP
+	: public FRigVMExecuteContextFixLayout
+#endif
 {
 	GENERATED_BODY()
 
diff --git a/Engine/Plugins/Runtime/SignificanceManager/Source/SignificanceManager/Public/SignificanceManager.h b/Engine/Plugins/Runtime/SignificanceManager/Source/SignificanceManager/Public/SignificanceManager.h
--- a/Engine/Plugins/Runtime/SignificanceManager/Source/SignificanceManager/Public/SignificanceManager.h
+++ b/Engine/Plugins/Runtime/SignificanceManager/Source/SignificanceManager/Public/SignificanceManager.h
@@ -53,7 +53,13 @@ public:
 		Sequential
 	};
 
-	struct SIGNIFICANCEMANAGER_API FManagedObjectInfo
+	// Hacky base class to avoid 8 bytes of padding after the vtable
+	struct FManagedObjectInfoFixLayout
+	{
+		virtual ~FManagedObjectInfoFixLayout() = default;
+	};
+
+	struct SIGNIFICANCEMANAGER_API FManagedObjectInfo : public FManagedObjectInfoFixLayout
 	{
 		FManagedObjectInfo()
 			: Object(nullptr)
@@ -81,8 +87,6 @@ public:
 			}
 		}
 
-		virtual ~FManagedObjectInfo() { }
-
 		FORCEINLINE UObject* GetObject() const { return Object; }
 		FORCEINLINE FName GetTag() const { return Tag; }
 		FORCEINLINE float GetSignificance() const { return Significance; }
diff --git a/Engine/Source/Runtime/AudioCaptureImplementations/Windows/AudioCaputureWasapi/Private/WasapiInputStream.h b/Engine/Source/Runtime/AudioCaptureImplementations/Windows/AudioCaputureWasapi/Private/WasapiInputStream.h
--- a/Engine/Source/Runtime/AudioCaptureImplementations/Windows/AudioCaputureWasapi/Private/WasapiInputStream.h
+++ b/Engine/Source/Runtime/AudioCaptureImplementations/Windows/AudioCaputureWasapi/Private/WasapiInputStream.h
@@ -43,7 +43,7 @@ namespace Audio
 			uint32 InNumFramesDesired,
 			FWasapiOnAudioCaptureFunction InCallback);
 
-		virtual ~FWasapiInputStream();
+		~FWasapiInputStream();
 
 		FWasapiInputStream& operator=(FWasapiInputStream&& InOther) = delete;
 		FWasapiInputStream& operator=(const FWasapiInputStream& InOther) = delete;
@@ -70,8 +70,6 @@ namespace Audio
 
 	private:
 
-		/** Indicates if this object has been successfully initialized. */
-		bool bIsInitialized = false;
 		/** COM pointer to the WASAPI audio client object. */
 		TComPtr<IAudioClient3> AudioClient;
 		/** COM pointer to the WASAPI capture client object. */
@@ -79,10 +77,9 @@ namespace Audio
 
 		/** Holds the audio format configuration for this stream. */
 		FWasapiAudioFormat AudioFormat;
-		/** The callback used to periodically deliver new audio data during capture. */
-		FWasapiOnAudioCaptureFunction OnAudioCaptureCallback;
-		/** Buffer used when WASAPI indicates that silence should be output for a given callback. */
-		TArray<uint8> SilienceBuffer;
+
+		/** Indicates if this object has been successfully initialized. */
+		bool bIsInitialized = false;
 
 		/** Number of frames of audio data which will be used for each audio callback during capture. */
 		uint32 NumFramesPerBuffer = 0;
@@ -94,6 +91,11 @@ namespace Audio
 		 */
 		HANDLE EventHandle = nullptr;
 
+		/** The callback used to periodically deliver new audio data during capture. */
+		FWasapiOnAudioCaptureFunction OnAudioCaptureCallback;
+		/** Buffer used when WASAPI indicates that silence should be output for a given callback. */
+		TArray<uint8> SilienceBuffer;
+
 		/**
 		 * DrainInputBuffer - Prepares the capture client by draining the input buffer of any audio data
 		 * which may have accumulated.
diff --git a/Engine/Source/Runtime/Core/Public/Async/AsyncFileHandle.h b/Engine/Source/Runtime/Core/Public/Async/AsyncFileHandle.h
--- a/Engine/Source/Runtime/Core/Public/Async/AsyncFileHandle.h
+++ b/Engine/Source/Runtime/Core/Public/Async/AsyncFileHandle.h
@@ -20,7 +20,14 @@ DECLARE_DWORD_ACCUMULATOR_STAT_EXTERN(TEXT("Num Async File Requests"), STAT_Asyn
 class IAsyncReadRequest;
 typedef TFunction<void(bool bWasCancelled, IAsyncReadRequest*)> FAsyncFileCallBack;
 
-class IAsyncReadRequest
+// Hacky base class to avoid 8 bytes of padding after the vtable
+class IAsyncReadRequestFixLayout
+{
+public:
+	virtual ~IAsyncReadRequestFixLayout() = default;
+};
+
+class IAsyncReadRequest : public IAsyncReadRequestFixLayout
 {
 protected:
 	union
diff --git a/Engine/Source/Runtime/D3D12RHI/Private/D3D12Submission.h b/Engine/Source/Runtime/D3D12RHI/Private/D3D12Submission.h
--- a/Engine/Source/Runtime/D3D12RHI/Private/D3D12Submission.h
+++ b/Engine/Source/Runtime/D3D12RHI/Private/D3D12Submission.h
@@ -169,8 +169,14 @@ struct FD3D12BatchedPayloadObjects
 	}
 };
 
+// Hacky base class to avoid 8 bytes of padding after the vtable
+struct FD3D12PayloadBaseFixLayout
+{
+	virtual ~FD3D12PayloadBaseFixLayout() = default;
+};
+
 // A single unit of work (specific to a single GPU node and queue type) to be processed by the submission thread.
-struct FD3D12PayloadBase
+struct FD3D12PayloadBase : public FD3D12PayloadBaseFixLayout
 {
 	// Used to signal FD3D12ManualFence instances on the submission thread.
 	struct FManualFence
@@ -213,6 +219,10 @@ struct FD3D12PayloadBase
 	// UpdateReservedResources
 	TArray<FD3D12CommitReservedResourceDesc> ReservedResourcesToCommit;
 
+	// Flags.
+	bool bAlwaysSignal = false;
+	std::atomic<bool> bSubmitted { false };
+
 	// Used by RHIRunOnQueue
 	TFunction<void(ID3D12CommandQueue*)> PreExecuteCallback;
 
@@ -229,10 +239,6 @@ struct FD3D12PayloadBase
 
 	TOptional<FD3D12Timing*> Timing;
 
-	// Flags.
-	bool bAlwaysSignal = false;
-	std::atomic<bool> bSubmitted { false };
-
 	// Cleanup
 	TArray<FD3D12CommandAllocator*> AllocatorsToRelease;
 
diff --git a/Engine/Source/Runtime/Engine/Classes/GameFramework/RootMotionSource.h b/Engine/Source/Runtime/Engine/Classes/GameFramework/RootMotionSource.h
--- a/Engine/Source/Runtime/Engine/Classes/GameFramework/RootMotionSource.h
+++ b/Engine/Source/Runtime/Engine/Classes/GameFramework/RootMotionSource.h
@@ -159,21 +159,27 @@ struct FRootMotionFinishVelocitySettings
 	UPROPERTY()
 	ERootMotionFinishVelocityMode Mode;
 
-	// Set Velocity if Mode == SetVelocity
-	UPROPERTY()
-	FVector SetVelocity;
-
 	// Clamp Velocity if Move == ClampVelocity
 	UPROPERTY()
 	float ClampVelocity;
 
+	// Set Velocity if Mode == SetVelocity
+	UPROPERTY()
+	FVector SetVelocity;
+
 	FRootMotionFinishVelocitySettings()
 		: Mode(ERootMotionFinishVelocityMode::MaintainLastRootMotionVelocity)
-		, SetVelocity(FVector::ZeroVector)
 		, ClampVelocity(0.f)
+		, SetVelocity(FVector::ZeroVector)
 	{}
 };
 
+// Hacky base class to avoid 8 bytes of padding after the vtable
+struct FRootMotionSourceFixLayout
+{
+	virtual ~FRootMotionSourceFixLayout() = default;
+};
+
 /** 
 *	Generalized source of Root Motion to a CharacterMovementComponent.
 *	
@@ -215,6 +221,9 @@ struct FRootMotionFinishVelocitySettings
 */
 USTRUCT()
 struct FRootMotionSource
+#if CPP
+	: public FRootMotionSourceFixLayout
+#endif
 {
 	GENERATED_USTRUCT_BODY()
 
@@ -227,12 +236,6 @@ struct FRootMotionSource
 	UPROPERTY()
 	uint16 LocalID;
 
-	/** 
-	 *  Accumulation mode for this source (whether or not to additively apply this root motion or override completely)
-	 **/
-	UPROPERTY()
-	ERootMotionAccumulateMode AccumulateMode;
-
 	/** 
 	 *	This name allows us to find the source later so that we can end it. 
 	 **/
@@ -267,6 +270,12 @@ struct FRootMotionSource
 	UPROPERTY(NotReplicated)
 	FRootMotionSourceSettings Settings;
 
+	/** 
+	*  Accumulation mode for this source (whether or not to additively apply this root motion or override completely)
+	**/
+	UPROPERTY()
+	ERootMotionAccumulateMode AccumulateMode;
+
 	/** True when this RootMotionSource is contributing local space accumulation (false for world space) */
 	UPROPERTY()
 	bool bInLocalSpace;
diff --git a/Engine/Source/Runtime/Engine/Public/LightSceneProxy.h b/Engine/Source/Runtime/Engine/Public/LightSceneProxy.h
--- a/Engine/Source/Runtime/Engine/Public/LightSceneProxy.h
+++ b/Engine/Source/Runtime/Engine/Public/LightSceneProxy.h
@@ -27,12 +27,19 @@ enum ELightShaderParameterFlags
 	RectAsSpotLight=1,
 };
 
+// Hacky base class to avoid 8 bytes of padding after the vtable
+class FLightSceneProxyFixLayout
+{
+public:
+	virtual ~FLightSceneProxyFixLayout() = default;
+};
+
 /**
  * Encapsulates the data which is used to render a light by the rendering thread.
  * The constructor is called from the game thread, and after that the rendering thread owns the object.
  * FLightSceneProxy is in the engine module and is subclassed to implement various types of lights.
  */
-class FLightSceneProxy
+class FLightSceneProxy : public FLightSceneProxyFixLayout
 {
 public:
 
@@ -310,6 +317,9 @@ protected:
 	/** The scene the primitive is in. */
 	FSceneInterface* SceneInterface;
 
+	/** The light's scene info. */
+	class FLightSceneInfo* LightSceneInfo;
+
 	/** The homogeneous position of the light. */
 	FVector4 Position;
 
@@ -322,9 +332,6 @@ protected:
 	/** A transform from light space into world space. */
 	FMatrix LightToWorld;
 
-	/** The light's scene info. */
-	class FLightSceneInfo* LightSceneInfo;
-
 	/** Scale for indirect lighting from this light.  When 0, indirect lighting is disabled. */
 	float IndirectLightingScale;
 
@@ -452,6 +459,12 @@ protected:
     /** Is the light selected in the editor? */
 	uint8 bSelected : 1;
 
+	/**
+	* The light index in order to be able to read matrix and parameters when reading the light function atlas for that light.
+	* A value of 0 means this is the default identity light function and no light function sampling will be done in shader.
+	*/
+	uint8 LightFunctionAtlasLightIndex;
+
 	/** The index of the atmospheric light. Multiple lights can be considered when computing the sky/atmospheric scattering. */
 	const uint8 AtmosphereSunLightIndex;
 
@@ -497,12 +510,6 @@ protected:
 	/** IES texture atlas id. */
 	uint32 IESAtlasId;
 
-	/**
-	 * The light index in order to be able to read matrix and parameters when reading the light function atlas for that light.
-	 * A value of 0 means this is the default identity light function and no light function sampling will be done in shader.
-	 */
-	uint8 LightFunctionAtlasLightIndex;
-
 	/**
 	 * Updates the light proxy's cached transforms.
 	 * @param InLightToWorld - The new light-to-world transform.
diff --git a/Engine/Source/Runtime/Engine/Public/ParticleEmitterInstances.h b/Engine/Source/Runtime/Engine/Public/ParticleEmitterInstances.h
--- a/Engine/Source/Runtime/Engine/Public/ParticleEmitterInstances.h
+++ b/Engine/Source/Runtime/Engine/Public/ParticleEmitterInstances.h
@@ -208,10 +208,16 @@ struct FParticleEmitterBuildInfo
 	FParticleEmitterBuildInfo();
 };
 
+// Hacky base class to avoid 8 bytes of padding after the vtable
+struct FParticleEmitterInstanceFixLayout
+{
+	virtual ~FParticleEmitterInstanceFixLayout() = default;
+};
+
 /*-----------------------------------------------------------------------------
 	FParticleEmitterInstance
 -----------------------------------------------------------------------------*/
-struct FParticleEmitterInstance
+struct FParticleEmitterInstance : FParticleEmitterInstanceFixLayout
 {
 public:
 	/** The maximum DeltaTime allowed for updating PeakActiveParticle tracking.
@@ -223,10 +229,10 @@ public:
 	UParticleEmitter* SpriteTemplate;
 	/** The component who owns it.										*/
 	UParticleSystemComponent* Component;
-	/** The index of the currently set LOD level.						*/
-	int32 CurrentLODLevelIndex;
 	/** The currently set LOD level.									*/
 	UParticleLODLevel* CurrentLODLevel;
+	/** The index of the currently set LOD level.						*/
+	int32 CurrentLODLevelIndex;
 	/** The offset to the TypeData payload in the particle data.		*/
 	int32 TypeDataOffset;
 	/** The offset to the TypeData instance payload.					*/
@@ -242,6 +248,8 @@ public:
 	int32 OrbitModuleOffset;
 	/** The offset to the Camera payload in the particle data.			*/
 	int32 CameraPayloadOffset;
+	/** The offset to the particle data.								*/
+	int32 PayloadOffset;
 	/** The location of the emitter instance							*/
 	FVector Location;
 	/** Transform from emitter local space to simulation space.			*/
@@ -249,27 +257,29 @@ public:
 	/** Transform from simulation space to world space.					*/
 	FMatrix SimulationToWorld;
 	/** Component can disable Tick and Rendering of this emitter. */
-	uint32 bEnabled : 1;
+	uint8 bEnabled : 1;
 	/** If true, kill this emitter instance when it is deactivated.		*/
-	uint32 bKillOnDeactivate:1;
+	uint8 bKillOnDeactivate:1;
 	/** if true, kill this emitter instance when it has completed.		*/
-	uint32 bKillOnCompleted:1;
+	uint8 bKillOnCompleted:1;
 	/** Whether this emitter requires sorting as specified by artist.	*/
-	uint32 bRequiresSorting:1;
+	uint8 bRequiresSorting:1;
 	/** If true, halt spawning for this instance.						*/
-	uint32 bHaltSpawning : 1;
+	uint8 bHaltSpawning : 1;
 	/** If true, this emitter has been disabled by game code and some systems to re-enable are not allowed. */
-	uint32 bHaltSpawningExternal : 1;
+	uint8 bHaltSpawningExternal : 1;
 	/** If true, the emitter has modules that require loop notification.*/
-	uint32 bRequiresLoopNotification:1;
+	uint8 bRequiresLoopNotification:1;
 	/** If true, the emitter ignores the component's scale. (Mesh emitters only). */
-	uint32 bIgnoreComponentScale:1;
+	uint8 bIgnoreComponentScale:1;
 	/** Hack: Make sure this is a Beam type to avoid casting from/to wrong types. */
-	uint32 bIsBeam:1;
+	uint8 bIsBeam:1;
 	/** Whether axis lock is enabled, cached here to avoid finding it from the module each frame */
-	uint32 bAxisLockEnabled : 1;
+	uint8 bAxisLockEnabled : 1;
 	/** When true and spawning is supressed, the bursts will be faked so that when spawning is enabled again, the bursts don't fire late. */
-	uint32 bFakeBurstsWhenSpawningSupressed : 1;
+	uint8 bFakeBurstsWhenSpawningSupressed : 1;
+	/** true if the emitter has no active particles and will no longer spawn any in the future */
+	uint8 bEmitterIsDone:1;
 	/** Axis lock flags, cached here to avoid finding it from the module each frame */
 	TEnumAsByte<EParticleAxisLock> LockAxisFlags;
 	/** The sort mode to use for this emitter as specified by artist.	*/
@@ -282,8 +292,6 @@ public:
 	uint8* InstanceData;
 	/** The size of the Instance data array.							*/
 	int32 InstancePayloadSize;
-	/** The offset to the particle data.								*/
-	int32 PayloadOffset;
 	/** The total size of a particle (in bytes).						*/
 	int32 ParticleSize;
 	/** The stride between particles in the ParticleData array.			*/
@@ -326,8 +334,6 @@ public:
 	TArray<float> EmitterDurations;
 	/** The emitter's delay for the current loop		*/
 	float CurrentDelay;
-	/** true if the emitter has no active particles and will no longer spawn any in the future */
-	bool bEmitterIsDone;
 
 	/** The number of triangles to render								*/
 	int32	TrianglesToRender;
diff --git a/Engine/Source/Runtime/Engine/Public/SceneView.h b/Engine/Source/Runtime/Engine/Public/SceneView.h
--- a/Engine/Source/Runtime/Engine/Public/SceneView.h
+++ b/Engine/Source/Runtime/Engine/Public/SceneView.h
@@ -1287,10 +1287,17 @@ namespace EDrawDynamicFlags
 	};
 }
 
+// Hacky base class to avoid 8 bytes of padding after the vtable
+class FSceneViewFixLayout
+{
+public:
+	virtual ~FSceneViewFixLayout() = default;
+};
+
 /**
  * A projection from scene space into a 2D screen region.
  */
-class FSceneView
+class FSceneView : public FSceneViewFixLayout
 {
 public:
 	const FSceneViewFamily* Family;
diff --git a/Engine/Source/Runtime/Experimental/Chaos/Public/Chaos/ISpatialAcceleration.h b/Engine/Source/Runtime/Experimental/Chaos/Public/Chaos/ISpatialAcceleration.h
--- a/Engine/Source/Runtime/Experimental/Chaos/Public/Chaos/ISpatialAcceleration.h
+++ b/Engine/Source/Runtime/Experimental/Chaos/Public/Chaos/ISpatialAcceleration.h
@@ -271,7 +271,7 @@ public:
 	using TPayload = TPayloadType;
 
 	ISpatialAcceleration(SpatialAccelerationType InType = static_cast<SpatialAccelerationType>(ESpatialAcceleration::Unknown))
-		: Type(InType), SyncTimestamp(0), AsyncTimeSlicingComplete(true)
+		: SyncTimestamp(0), AsyncTimeSlicingComplete(true), Type(InType)
 	{}
 
 	ISpatialAcceleration(ESpatialAcceleration InType)
@@ -413,9 +413,9 @@ protected:
 	virtual void SetAsyncTimeSlicingComplete(bool InState) { AsyncTimeSlicingComplete = InState; }
 
 private:
-	SpatialAccelerationType Type;
 	int32 SyncTimestamp;	//The set of inputs the acceleration structure is in sync with. GT moves forward in time and enqueues inputs
 	bool AsyncTimeSlicingComplete;
+	SpatialAccelerationType Type;
 };
 
 template <typename TBase, typename TDerived>
diff --git a/Engine/Source/Runtime/GeometryCore/Public/Implicit/Morphology.h b/Engine/Source/Runtime/GeometryCore/Public/Implicit/Morphology.h
--- a/Engine/Source/Runtime/GeometryCore/Public/Implicit/Morphology.h
+++ b/Engine/Source/Runtime/GeometryCore/Public/Implicit/Morphology.h
@@ -40,10 +40,6 @@ public:
 		Open = 3
 	};
 
-	virtual ~TImplicitMorphology()
-	{
-	}
-
 	///
 	/// Inputs
 	///
diff --git a/Engine/Source/Runtime/GeometryCore/Public/Implicit/Solidify.h b/Engine/Source/Runtime/GeometryCore/Public/Implicit/Solidify.h
--- a/Engine/Source/Runtime/GeometryCore/Public/Implicit/Solidify.h
+++ b/Engine/Source/Runtime/GeometryCore/Public/Implicit/Solidify.h
@@ -28,10 +28,6 @@ public:
 	{
 	}
 
-	virtual ~TImplicitSolidify()
-	{
-	}
-
 	///
 	/// Inputs
 	///
diff --git a/Engine/Source/Runtime/Landscape/Private/LandscapeGrassWeightExporter.h b/Engine/Source/Runtime/Landscape/Private/LandscapeGrassWeightExporter.h
--- a/Engine/Source/Runtime/Landscape/Private/LandscapeGrassWeightExporter.h
+++ b/Engine/Source/Runtime/Landscape/Private/LandscapeGrassWeightExporter.h
@@ -11,8 +11,15 @@ class FLandscapeComponentSceneProxy;
 class FTextureRenderTarget2DResource;
 class UTextureRenderTarget2D;
 
+// Hacky base class to avoid 8 bytes of padding after the vtable
+class FLandscapeGrassWeightExporter_RenderThread_FixLayout
+{
+public:
+	virtual ~FLandscapeGrassWeightExporter_RenderThread_FixLayout() = default;
+};
+
 // data also accessible by render thread
-class FLandscapeGrassWeightExporter_RenderThread
+class FLandscapeGrassWeightExporter_RenderThread : public FLandscapeGrassWeightExporter_RenderThread_FixLayout
 {
 	FLandscapeGrassWeightExporter_RenderThread(const TArray<int32>& InHeightMips)
 		: HeightMips(InHeightMips)
@@ -70,11 +77,12 @@ public:
 	TArray<int32> HeightMips;
 	float PassOffsetX;
 	FVector ViewOrigin;
-	FMatrix ViewRotationMatrix;
-	FMatrix ProjectionMatrix;
 
 	FLandscapeAsyncTextureReadback* AsyncReadbackPtr = nullptr;
 
+	FMatrix ViewRotationMatrix;
+	FMatrix ProjectionMatrix;
+
 	void RenderLandscapeComponentToTexture_RenderThread(FRHICommandListImmediate& RHICmdList);
 };
 
diff --git a/Engine/Source/Runtime/Net/Core/Public/Net/Core/Connection/NetConnectionFaultRecoveryBase.h b/Engine/Source/Runtime/Net/Core/Public/Net/Core/Connection/NetConnectionFaultRecoveryBase.h
--- a/Engine/Source/Runtime/Net/Core/Public/Net/Core/Connection/NetConnectionFaultRecoveryBase.h
+++ b/Engine/Source/Runtime/Net/Core/Public/Net/Core/Connection/NetConnectionFaultRecoveryBase.h
@@ -143,18 +143,21 @@ inline uint8 ToInt(ENetFaultCounterCategory CategoryVal)
 	return static_cast<uint8>(CategoryVal);
 }
 
+// Hacky base class to avoid 8 bytes of padding after the vtable
+class FNetConnectionFaultRecoveryBaseFixLayout
+{
+public:
+	virtual ~FNetConnectionFaultRecoveryBaseFixLayout() = default;
+};
 
 /**
  * Implements the base/public interface for FNetConnectionFaultRecovery - defined here, to eliminate the need for Engine dependencies.
  */
-class FNetConnectionFaultRecoveryBase
+class FNetConnectionFaultRecoveryBase : public FNetConnectionFaultRecoveryBaseFixLayout
 {
 public:
 	using FNetFaultEscalationHandler = TEscalationManager<ENetFaultCounters, FNetFaultState, ENetFaultCounterCategory>;
 
-public:
-	virtual ~FNetConnectionFaultRecoveryBase() = default;
-
 	/**
 	 * (For Counter Categories) Use this to notify fault recovery of a successfully handled fault it should be aware of, within HandleNetResult;
 	 * this is required for escalation tracking, if that fault type is registered with a counter category.
@@ -206,12 +209,12 @@ protected:
 	/** Whether or not fault recovery has disconnected the NetConnection */
 	bool bDisconnected = false;
 
-	/** The escalation manager instance used for tracking/implementing recovery for the various fault types */
-	TUniquePtr<FNetFaultEscalationHandler> NetFaultEscalationManager;
-
 	/** Tracks the index of the last added counter, before NetFaultEscalationManager is created */
 	int32 LastCounterIndex = static_cast<int32>(ENetFaultCounters::Max);
 
+	/** The escalation manager instance used for tracking/implementing recovery for the various fault types */
+	TUniquePtr<FNetFaultEscalationHandler> NetFaultEscalationManager;
+
 	struct FPendingCategoryRegister
 	{
 		int32 CategoryIndex		= 0;
diff --git a/Engine/Source/Runtime/PacketHandlers/PacketHandler/Private/PacketHandler.cpp b/Engine/Source/Runtime/PacketHandlers/PacketHandler/Private/PacketHandler.cpp
--- a/Engine/Source/Runtime/PacketHandlers/PacketHandler/Private/PacketHandler.cpp
+++ b/Engine/Source/Runtime/PacketHandlers/PacketHandler/Private/PacketHandler.cpp
@@ -100,15 +100,17 @@ BufferedPacket::~BufferedPacket()
 
 PacketHandler::PacketHandler(FDDoSDetection* InDDoS/*=nullptr*/)
 	: Mode(UE::Handler::Mode::Client)
+	, State(UE::Handler::State::Uninitialized)
+	, bRawSend(false)
+	, bBeganHandshaking(false)
 	, bConnectionlessHandler(false)
+	, MaxPacketBits(0)
 	, DDoS(InDDoS)
 	, LowLevelSendDel()
 	, HandshakeCompleteDel()
 	, OutgoingPacket(MAX_PACKET_SIZE * 8)
 	, IncomingPacket()
 	, HandlerComponents()
-	, MaxPacketBits(0)
-	, State(UE::Handler::State::Uninitialized)
 	, BufferedPackets()
 	, QueuedPackets()
 	, QueuedRawPackets()
@@ -116,10 +118,8 @@ PacketHandler::PacketHandler(FDDoSDetection* InDDoS/*=nullptr*/)
 	, BufferedConnectionlessPackets()
 	, QueuedConnectionlessPackets()
 	, ReliabilityComponent(nullptr)
-	, bRawSend(false)
 	, Provider()
 	, Aggregator()
-	, bBeganHandshaking(false)
 {
 	OutgoingPacket.SetAllowResize(true);
 	OutgoingPacket.AllowAppend(true);
diff --git a/Engine/Source/Runtime/PacketHandlers/PacketHandler/Public/PacketHandler.h b/Engine/Source/Runtime/PacketHandlers/PacketHandler/Public/PacketHandler.h
--- a/Engine/Source/Runtime/PacketHandlers/PacketHandler/Public/PacketHandler.h
+++ b/Engine/Source/Runtime/PacketHandlers/PacketHandler/Public/PacketHandler.h
@@ -259,8 +259,6 @@ public:
 	 */
 	PACKETHANDLER_API PacketHandler(FDDoSDetection* InDDoS=nullptr);
 
-	virtual ~PacketHandler() = default;
-
 	/**
 	 * Handles initialization of manager
 	 *
@@ -446,7 +444,7 @@ public:
 	/** Returns a pointer to the first component in the HandlerComponents array with the specified name. */
 	PACKETHANDLER_API TSharedPtr<HandlerComponent> GetComponentByName(FName ComponentName) const;
 
-	PACKETHANDLER_API virtual void CountBytes(FArchive& Ar) const;
+	PACKETHANDLER_API void CountBytes(FArchive& Ar) const;
 
 protected:
 	UE_DEPRECATED(4.26, "Incoming_Internal now uses FReceivedPacketView")
@@ -653,8 +651,21 @@ public:
 	static PACKETHANDLER_API FPacketHandlerAddComponentDelegate& GetAddComponentDelegate();
 
 private:
+
+	/** State of the handler */
+	UE::Handler::State State;
+
+	/** Whether or not outgoing packets bypass the handler */
+	uint8 bRawSend:1;
+
+	/** Whether or not component handshaking has begun */
+	uint8 bBeganHandshaking:1;
+
 	/** Whether or not this PacketHandler handles connectionless (i.e. non-UNetConnection) data */
-	bool bConnectionlessHandler;
+	uint8 bConnectionlessHandler:1;
+
+	/** The maximum supported packet size (reflects UNetConnection::MaxPacket) */
+	uint32 MaxPacketBits;
 
 	/** Mirroring UNetDriver.DDoS*/
 	FDDoSDetection* DDoS;
@@ -680,12 +691,6 @@ private:
 	/** A direct pointer to the component configured as the encryption component. Will also be present in the HandlerComponents array. */
 	TSharedPtr<FEncryptionComponent> EncryptionComponent;
 
-	/** The maximum supported packet size (reflects UNetConnection::MaxPacket) */
-	uint32 MaxPacketBits;
-
-	/** State of the handler */
-	UE::Handler::State State;
-
 	/** Packets that are buffered while HandlerComponents are being initialized */
 	TArray<BufferedPacket*> BufferedPackets;
 
@@ -707,17 +712,11 @@ private:
 	/** Reliability Handler Component */
 	TSharedPtr<ReliabilityHandlerComponent> ReliabilityComponent;
 
-	/** Whether or not outgoing packets bypass the handler */
-	bool bRawSend;
-
 	/** The analytics provider */
 	TSharedPtr<IAnalyticsProvider> Provider;
 
 	/** The NetDriver level aggregator for the analytics provider */
 	TSharedPtr<FNetAnalyticsAggregator> Aggregator;
-	
-	/** Whether or not component handshaking has begun */
-	bool bBeganHandshaking;
 };
 
 /**
diff --git a/Engine/Source/Runtime/RHI/Public/RHIResources.h b/Engine/Source/Runtime/RHI/Public/RHIResources.h
--- a/Engine/Source/Runtime/RHI/Public/RHIResources.h
+++ b/Engine/Source/Runtime/RHI/Public/RHIResources.h
@@ -3099,7 +3099,7 @@ enum ERayTracingGeometryType
 };
 DECLARE_INTRINSIC_TYPE_LAYOUT(ERayTracingGeometryType);
 
-enum class ERayTracingGeometryInitializerType
+enum class ERayTracingGeometryInitializerType : uint8
 {
 	// Fully initializes the RayTracingGeometry object: creates underlying buffer and initializes shader parameters.
 	Rendering,
@@ -3158,6 +3158,11 @@ public:
 	// Total number of primitives in all segments of the geometry. Only used for validation.
 	uint32 TotalPrimitiveCount = 0;
 
+	bool bFastBuild = false;
+	bool bAllowUpdate = false;
+	bool bAllowCompaction = true;
+	ERayTracingGeometryInitializerType Type = ERayTracingGeometryInitializerType::Rendering;
+
 	// Partitions of geometry to allow different shader and resource bindings.
 	// All ray tracing geometries must have at least one segment.
 	TArray<FRayTracingGeometrySegment> Segments;
@@ -3168,11 +3173,6 @@ public:
 	// Pointer to an existing ray tracing geometry which the new geometry is built from.
 	FRHIRayTracingGeometry* SourceGeometry = nullptr;
 
-	bool bFastBuild = false;
-	bool bAllowUpdate = false;
-	bool bAllowCompaction = true;
-	ERayTracingGeometryInitializerType Type = ERayTracingGeometryInitializerType::Rendering;
-
 	// Use FDebugName for auto-generated debug names with numbered suffixes, it is a variation of FMemoryImageName with optional number postfix.
 	FDebugName DebugName;
 	// Store the path name of the owner object for resource tracking. FMemoryImageName allows a conversion to/from FName.
diff --git a/Engine/Source/Runtime/RenderCore/Public/RayTracingGeometry.h b/Engine/Source/Runtime/RenderCore/Public/RayTracingGeometry.h
--- a/Engine/Source/Runtime/RenderCore/Public/RayTracingGeometry.h
+++ b/Engine/Source/Runtime/RenderCore/Public/RayTracingGeometry.h
@@ -16,7 +16,7 @@ namespace RayTracing
 	using GeometryGroupHandle = int32;
 }
 
-enum class ERTAccelerationStructureBuildPriority
+enum class ERTAccelerationStructureBuildPriority : uint8
 {
 	Immediate,
 	High,
@@ -157,9 +157,9 @@ protected:
 	RENDERCORE_API void RemoveBuildRequest();
 
 	friend class FRayTracingGeometryManager;
+	EGeometryStateFlags GeometryState = EGeometryStateFlags::Invalid;
 	int32 RayTracingBuildRequestIndex = INDEX_NONE;
 	int32 RayTracingGeometryHandle = INDEX_NONE; // Only valid when ray tracing is dynamic
-	EGeometryStateFlags GeometryState = EGeometryStateFlags::Invalid;
 #endif
 };
 
diff --git a/Engine/Source/Runtime/RenderCore/Public/RayTracingGeometryManagerInterface.h b/Engine/Source/Runtime/RenderCore/Public/RayTracingGeometryManagerInterface.h
--- a/Engine/Source/Runtime/RenderCore/Public/RayTracingGeometryManagerInterface.h
+++ b/Engine/Source/Runtime/RenderCore/Public/RayTracingGeometryManagerInterface.h
@@ -12,7 +12,7 @@
 class FRayTracingGeometry;
 class FRHIComputeCommandList;
 enum class EAccelerationStructureBuildMode;
-enum class ERTAccelerationStructureBuildPriority;
+enum class ERTAccelerationStructureBuildPriority : uint8;
 
 namespace RayTracing
 {
