- Fix for parallel GDME and GPU simulation that execute in PostInitViews
- Parallel GDME will overlap from during GDME to before PreRender.  Therefore Niagara's PostInitViews can execute while parallel GDME runs.  When doing so it will call SetDataToRender which is also read on the parallel GDME threads.
- This change disables parallel GDME for those proxies who run GPU compute in PostInitViews

#rnx
#tests Soak Lyra in debug build as it would hit a checkSlow

diff --git a/Engine/Plugins/FX/Niagara/Source/Niagara/Classes/NiagaraDataSet.h b/Engine/Plugins/FX/Niagara/Source/Niagara/Classes/NiagaraDataSet.h
--- a/Engine/Plugins/FX/Niagara/Source/Niagara/Classes/NiagaraDataSet.h
+++ b/Engine/Plugins/FX/Niagara/Source/Niagara/Classes/NiagaraDataSet.h
@@ -506,7 +506,7 @@ FORCEINLINE void FNiagaraDataBuffer::CheckUsage(bool bReadOnly)const
 	//For GPU sims we must be on the RT.
 	checkSlow(
 		(Owner->GetSimTarget() == ENiagaraSimTarget::CPUSim && (IsInGameThread() || bReadOnly || !GIsThreadedRendering || !IsInRenderingThread())) ||
-		(Owner->GetSimTarget() == ENiagaraSimTarget::GPUComputeSim && IsInRenderingThread())
+		(Owner->GetSimTarget() == ENiagaraSimTarget::GPUComputeSim && IsInParallelRenderingThread())
 	);
 }
 
diff --git a/Engine/Plugins/FX/Niagara/Source/Niagara/Private/NiagaraComponent.cpp b/Engine/Plugins/FX/Niagara/Source/Niagara/Private/NiagaraComponent.cpp
--- a/Engine/Plugins/FX/Niagara/Source/Niagara/Private/NiagaraComponent.cpp
+++ b/Engine/Plugins/FX/Niagara/Source/Niagara/Private/NiagaraComponent.cpp
@@ -208,6 +208,12 @@ FNiagaraSceneProxy::FNiagaraSceneProxy(UNiagaraComponent* InComponent)
 	: FPrimitiveSceneProxy(InComponent, InComponent->GetAsset() ? InComponent->GetAsset()->GetFName() : FName())
 	, OcclusionQueryMode(InComponent->GetOcclusionQueryMode())
 {
+	// Prevent continuous VSM invalidation from the bounds of the primitive.
+	bHasDeformableMesh = false;
+
+	// Optionally enable GDME for Niagara, any caveats are covered below
+	bSupportsParallelGDME = GNiagaraParallelGDME;
+
 	FNiagaraSystemInstanceControllerConstPtr SystemInstanceController = InComponent->GetSystemInstanceController();
 	UNiagaraSystem* NiagaraSystem = InComponent->GetAsset();
 	if (SystemInstanceController && NiagaraSystem)
@@ -224,12 +230,15 @@ FNiagaraSceneProxy::FNiagaraSceneProxy(UNiagaraComponent* InComponent)
 #if NIAGARAPROXY_EVENTS_ENABLED
 		SystemStatString = NiagaraSystem->GetFName().ToString();
 #endif
-	}
 
-	// Prevent continuous VSM invalidation from the bounds of the primitive.
-	bHasDeformableMesh = false;
-	// Niagara renderers reference a lot of common contexts that aren't locked (and would otherwise introduce a lot of contention).
-	bSupportsParallelGDME = GNiagaraParallelGDME;
+		// Any GPU simulations that tick in PostInitView can not run parallel in GDME as GDME runs async over the top of this call
+		//-OPT: Would it be better to run these simulations in PreRender which is after we join from parallel GDME?
+		const ENiagaraGpuComputeTickStage::Type GpuComputeTickStage = SystemInstanceController->GetGpuComputeTickStage();
+		if (GpuComputeTickStage == ENiagaraGpuComputeTickStage::PostInitViews )
+		{
+			bSupportsParallelGDME = false;
+		}
+	}
 }
 
 SIZE_T FNiagaraSceneProxy::GetTypeHash() const
diff --git a/Engine/Plugins/FX/Niagara/Source/Niagara/Private/NiagaraSystemInstanceController.cpp b/Engine/Plugins/FX/Niagara/Source/Niagara/Private/NiagaraSystemInstanceController.cpp
--- a/Engine/Plugins/FX/Niagara/Source/Niagara/Private/NiagaraSystemInstanceController.cpp
+++ b/Engine/Plugins/FX/Niagara/Source/Niagara/Private/NiagaraSystemInstanceController.cpp
@@ -5,6 +5,7 @@
 #include "Engine/StaticMesh.h"
 #include "Engine/TextureRenderTarget.h"
 #include "NiagaraEmitterInstance.h"
+#include "NiagaraSystemGpuComputeProxy.h"
 #include "NiagaraWorldManager.h"
 
 FNiagaraSystemInstanceController::FNiagaraSystemInstanceController()
@@ -237,6 +238,19 @@ UMaterialInterface* FNiagaraSystemInstanceController::GetMaterialOverride(const
 	return nullptr;
 }
 
+ENiagaraGpuComputeTickStage::Type FNiagaraSystemInstanceController::GetGpuComputeTickStage() const
+{
+	ENiagaraGpuComputeTickStage::Type TickStage = ENiagaraGpuComputeTickStage::PreInitViews;
+	if (SystemInstance.IsValid())
+	{
+		if (FNiagaraSystemGpuComputeProxy* GpuComputeProxy = SystemInstance->GetSystemGpuComputeProxy())
+		{
+			return GpuComputeProxy->GetComputeTickStage();
+		}
+	}
+	return TickStage;
+}
+
 void FNiagaraSystemInstanceController::DebugDump(bool bFullDump)
 {
 	if (SystemInstance.IsValid())
diff --git a/Engine/Plugins/FX/Niagara/Source/Niagara/Public/NiagaraSystemInstanceController.h b/Engine/Plugins/FX/Niagara/Source/Niagara/Public/NiagaraSystemInstanceController.h
--- a/Engine/Plugins/FX/Niagara/Source/Niagara/Public/NiagaraSystemInstanceController.h
+++ b/Engine/Plugins/FX/Niagara/Source/Niagara/Public/NiagaraSystemInstanceController.h
@@ -97,6 +97,8 @@ public:
 	void GetMaterialStreamingInfo(FNiagaraMaterialAndScaleArray& OutMaterialAndScales) const;
 	void GetStreamingMeshInfo(const FBoxSphereBounds& OwnerBounds, FStreamingTextureLevelContext& LevelContext, TArray<FStreamingRenderAssetPrimitiveInfo>& OutStreamingRenderAssets) const;
 
+	ENiagaraGpuComputeTickStage::Type GetGpuComputeTickStage() const;
+
 	/** Dumps system instance state and info to the log (Used by fx.Niagara.DumpComponents commandlet) */
 	void DebugDump(bool bFullDump);
 	SIZE_T GetTotalBytesUsed() const;
