- Add optional execute to all async trace functions

#tests Using collision module in RT mode and a custom test to ensure it worked and upgraded as expected

diff --git a/Engine/Plugins/FX/Niagara/Shaders/Private/NiagaraDataInterfaceAsyncGpuTraceTemplate.ush b/Engine/Plugins/FX/Niagara/Shaders/Private/NiagaraDataInterfaceAsyncGpuTraceTemplate.ush
--- a/Engine/Plugins/FX/Niagara/Shaders/Private/NiagaraDataInterfaceAsyncGpuTraceTemplate.ush
+++ b/Engine/Plugins/FX/Niagara/Shaders/Private/NiagaraDataInterfaceAsyncGpuTraceTemplate.ush
@@ -10,22 +10,43 @@ uint											{ParameterName}_IntersectionResultsOffset;
 RWBuffer<uint>									{ParameterName}_RWRayTraceCounts;
 uint											{ParameterName}_RayTraceCountsOffset;
 
-void IssueAsyncRayTraceGpu_{ParameterName}_UEImpureCall(int In_QueryID, float3 In_TraceStart, float3 In_TraceEnd, int In_CollisionGroup, out bool Out_IsQueryValid)
+void IssueAsyncRayTraceGpu_{ParameterName}_UEImpureCall(bool bExecute, int In_QueryID, float3 In_TraceStart, float3 In_TraceEnd, int In_CollisionGroup, out bool Out_IsQueryValid)
 {
-	NDIAsyncGpuTrace_IssueAsyncRayTrace({ParameterName}_MaxRayTraceCount, {ParameterName}_RWRayRequests, {ParameterName}_RayRequestsOffset, {ParameterName}_RWRayTraceCounts, {ParameterName}_RayTraceCountsOffset, {ParameterName}_SystemLWCTile, In_QueryID, In_TraceStart, In_TraceEnd, In_CollisionGroup, Out_IsQueryValid);
+	Out_IsQueryValid = false;
+	if ( bExecute )
+	{
+		NDIAsyncGpuTrace_IssueAsyncRayTrace({ParameterName}_MaxRayTraceCount, {ParameterName}_RWRayRequests, {ParameterName}_RayRequestsOffset, {ParameterName}_RWRayTraceCounts, {ParameterName}_RayTraceCountsOffset, {ParameterName}_SystemLWCTile, In_QueryID, In_TraceStart, In_TraceEnd, In_CollisionGroup, Out_IsQueryValid);
+	}
 }
 
-void ReserveAsyncRayTraceGpu_{ParameterName}_UEImpureCall(int In_TraceCount, out int Out_StartQueryID, out bool Out_IndicesValid)
+void ReserveAsyncRayTraceGpu_{ParameterName}_UEImpureCall(bool bExecute, int In_TraceCount, out int Out_StartQueryID, out bool Out_IndicesValid)
 {
-	NDIAsyncGpuTrace_ReserveRayTraceIndex({ParameterName}_MaxRayTraceCount, {ParameterName}_RWRayTraceCounts, {ParameterName}_RayTraceCountsOffset, In_TraceCount, Out_StartQueryID, Out_IndicesValid);
+	Out_StartQueryID = -1;
+	Out_IndicesValid = false;
+	if (bExecute)
+	{
+		NDIAsyncGpuTrace_ReserveRayTraceIndex({ParameterName}_MaxRayTraceCount, {ParameterName}_RWRayTraceCounts, {ParameterName}_RayTraceCountsOffset, In_TraceCount, Out_StartQueryID, Out_IndicesValid);
+	}
 }
 
-void CreateAsyncRayTraceGpu_{ParameterName}_UEImpureCall(float3 In_TraceStart, float3 In_TraceEnd, int In_CollisionGroup, out int Out_QueryID, out bool Out_IsQueryValid)
+void CreateAsyncRayTraceGpu_{ParameterName}_UEImpureCall(bool bExecute, float3 In_TraceStart, float3 In_TraceEnd, int In_CollisionGroup, out int Out_QueryID, out bool Out_IsQueryValid)
 {
-	NDIAsyncGpuTrace_CreateAsyncRayTrace({ParameterName}_MaxRayTraceCount, {ParameterName}_RWRayRequests, {ParameterName}_RayRequestsOffset, {ParameterName}_RWRayTraceCounts, {ParameterName}_RayTraceCountsOffset, {ParameterName}_SystemLWCTile, In_TraceStart, In_TraceEnd, In_CollisionGroup, Out_QueryID, Out_IsQueryValid);
+	Out_QueryID			= -1;
+	Out_IsQueryValid	= false;
+	if (bExecute)
+	{
+		NDIAsyncGpuTrace_CreateAsyncRayTrace({ParameterName}_MaxRayTraceCount, {ParameterName}_RWRayRequests, {ParameterName}_RayRequestsOffset, {ParameterName}_RWRayTraceCounts, {ParameterName}_RayTraceCountsOffset, {ParameterName}_SystemLWCTile, In_TraceStart, In_TraceEnd, In_CollisionGroup, Out_QueryID, Out_IsQueryValid);
+	}
 }
 
-void ReadAsyncRayTraceGpu_{ParameterName}(int In_PreviousFrameQueryID, out bool Out_CollisionValid, out float Out_CollisionDistance, out float3 Out_CollisionPosWorld, out float3 Out_CollisionNormal)
+void ReadAsyncRayTraceGpu_{ParameterName}(bool bExecute, int In_PreviousFrameQueryID, out bool Out_CollisionValid, out float Out_CollisionDistance, out float3 Out_CollisionPosWorld, out float3 Out_CollisionNormal)
 {
-	NDIAsyncGpuTrace_ReadAsyncRayTrace({ParameterName}_MaxRayTraceCount, {ParameterName}_IntersectionResults, {ParameterName}_IntersectionResultsOffset, {ParameterName}_SystemLWCTile, In_PreviousFrameQueryID, Out_CollisionValid, Out_CollisionDistance, Out_CollisionPosWorld, Out_CollisionNormal);
-}
\ No newline at end of file
+	Out_CollisionValid		= false;
+	Out_CollisionDistance	= 0.0f;
+	Out_CollisionPosWorld	= float3(0.0f, 0.0f, 0.0f);
+	Out_CollisionNormal		= float3(0.0f, 0.0f, 0.0f);
+	if (bExecute)
+	{
+		NDIAsyncGpuTrace_ReadAsyncRayTrace({ParameterName}_MaxRayTraceCount, {ParameterName}_IntersectionResults, {ParameterName}_IntersectionResultsOffset, {ParameterName}_SystemLWCTile, In_PreviousFrameQueryID, Out_CollisionValid, Out_CollisionDistance, Out_CollisionPosWorld, Out_CollisionNormal);
+	}
+}
diff --git a/Engine/Plugins/FX/Niagara/Source/Niagara/Private/DataInterface/NiagaraDataInterfaceAsyncGpuTrace.cpp b/Engine/Plugins/FX/Niagara/Source/Niagara/Private/DataInterface/NiagaraDataInterfaceAsyncGpuTrace.cpp
--- a/Engine/Plugins/FX/Niagara/Source/Niagara/Private/DataInterface/NiagaraDataInterfaceAsyncGpuTrace.cpp
+++ b/Engine/Plugins/FX/Niagara/Source/Niagara/Private/DataInterface/NiagaraDataInterfaceAsyncGpuTrace.cpp
@@ -52,6 +52,7 @@ struct FNiagaraAsyncGpuTraceDIFunctionVersion
 	enum Type
 	{
 		InitialVersion = 0,
+		AddOptionalExecution = 1,
 
 		VersionPlusOne,
 		LatestVersion = VersionPlusOne - 1
@@ -254,6 +255,7 @@ void UNiagaraDataInterfaceAsyncGpuTrace::GetFunctionsInternal(TArray<FNiagaraFun
 		IssueRayTrace.ExperimentalMessage = ExperimentalMessage;
 #endif
 		IssueRayTrace.AddInput(FNiagaraVariable(FNiagaraTypeDefinition(GetClass()), TEXT("AsyncGpuTrace")));
+		IssueRayTrace.Inputs.Emplace_GetRef(FNiagaraTypeDefinition::GetBoolDef(), TEXT("Execute")).SetValue(true);
 		IssueRayTrace.AddInput(FNiagaraVariable(FNiagaraTypeDefinition::GetIntDef(), TEXT("QueryID")), QueryIDDescription);
 		IssueRayTrace.AddInput(FNiagaraVariable(FNiagaraTypeDefinition::GetPositionDef(), TEXT("TraceStartWorld")), RayTraceStartWorldDescription);
 		IssueRayTrace.AddInput(FNiagaraVariable(FNiagaraTypeDefinition::GetPositionDef(), TEXT("TraceEndWorld")), RayTraceEndWorldDescription);
@@ -274,6 +276,7 @@ void UNiagaraDataInterfaceAsyncGpuTrace::GetFunctionsInternal(TArray<FNiagaraFun
 		CreateRayTrace.Description = LOCTEXT("CreateAsyncRayTraceDescription", "Creates a GPU raytrace with the result being available the following frame (index is returned)");
 		CreateRayTrace.ExperimentalMessage = ExperimentalMessage;
 		CreateRayTrace.AddInput(FNiagaraVariable(FNiagaraTypeDefinition(GetClass()), TEXT("AsyncGpuTrace")));
+		CreateRayTrace.Inputs.Emplace_GetRef(FNiagaraTypeDefinition::GetBoolDef(), TEXT("Execute")).SetValue(true);
 		CreateRayTrace.AddInput(FNiagaraVariable(FNiagaraTypeDefinition::GetPositionDef(), TEXT("TraceStartWorld")), RayTraceStartWorldDescription);
 		CreateRayTrace.AddInput(FNiagaraVariable(FNiagaraTypeDefinition::GetPositionDef(), TEXT("TraceEndWorld")), RayTraceEndWorldDescription);
 		CreateRayTrace.AddInput(CollisionGroupVariable, CollisionGroupDescription);
@@ -296,6 +299,7 @@ void UNiagaraDataInterfaceAsyncGpuTrace::GetFunctionsInternal(TArray<FNiagaraFun
 		ReserveRayTrace.Description = LOCTEXT("ReserveAsyncRayTraceDescription", "Reserves a number of ray trace request slots");
 		ReserveRayTrace.ExperimentalMessage = ExperimentalMessage;
 		ReserveRayTrace.AddInput(FNiagaraVariable(FNiagaraTypeDefinition(GetClass()), TEXT("AsyncGpuTrace")));
+		ReserveRayTrace.Inputs.Emplace_GetRef(FNiagaraTypeDefinition::GetBoolDef(), TEXT("Execute")).SetValue(true);
 		ReserveRayTrace.AddInput(FNiagaraVariable(FNiagaraTypeDefinition::GetIntDef(), TEXT("TraceCount")), TraceCountDescription);
 		ReserveRayTrace.AddOutput(FNiagaraVariable(FNiagaraTypeDefinition::GetIntDef(), TEXT("FirstQueryID")), FirstQueryIDValidDescription);
 		ReserveRayTrace.AddOutput(FNiagaraVariable(FNiagaraTypeDefinition::GetBoolDef(), TEXT("IsQueryValid")), ReserveValidDescription);
@@ -314,6 +318,7 @@ void UNiagaraDataInterfaceAsyncGpuTrace::GetFunctionsInternal(TArray<FNiagaraFun
 		ReadRayTrace.Description = LOCTEXT("ReadAsyncRayTraceDescription", "Reads the results of a previously enqueued GPU ray trace");
 		ReadRayTrace.ExperimentalMessage = ExperimentalMessage;
 		ReadRayTrace.AddInput(FNiagaraVariable(FNiagaraTypeDefinition(GetClass()), TEXT("AsyncGpuTrace")));
+		ReadRayTrace.Inputs.Emplace_GetRef(FNiagaraTypeDefinition::GetBoolDef(), TEXT("Execute")).SetValue(true);
 		ReadRayTrace.AddInput(FNiagaraVariable(FNiagaraTypeDefinition::GetIntDef(), TEXT("PreviousFrameQueryID")), PreviousFrameQueryIDDescription);
 		ReadRayTrace.AddOutput(FNiagaraVariable(FNiagaraTypeDefinition::GetBoolDef(), TEXT("CollisionValid")), CollisionValidDescription);
 		ReadRayTrace.AddOutput(FNiagaraVariable(FNiagaraTypeDefinition::GetFloatDef(), TEXT("CollisionDistance")), CollisionDistanceDescription);
