- Guard for parallel usage of custom uniform buffers
- Guard for parallel upload of ribbon data

#rnx

diff --git a/Engine/Plugins/FX/Niagara/Source/Niagara/Private/NiagaraComponent.cpp b/Engine/Plugins/FX/Niagara/Source/Niagara/Private/NiagaraComponent.cpp
--- a/Engine/Plugins/FX/Niagara/Source/Niagara/Private/NiagaraComponent.cpp
+++ b/Engine/Plugins/FX/Niagara/Source/Niagara/Private/NiagaraComponent.cpp
@@ -373,6 +373,7 @@ TUniformBuffer<FPrimitiveUniformShaderParameters>* FNiagaraSceneProxy::GetCustom
 		KeyHash = HashCombine(KeyHash, CustomFloatHash);
 	}
 
+	UE::TScopeLock LockGuard(CustomUniformBuffersGuard);
 	TUniformBuffer<FPrimitiveUniformShaderParameters>*& CustomUBRef = CustomUniformBuffers.FindOrAdd(KeyHash);
 	if (CustomUBRef == nullptr)
 	{
@@ -419,11 +420,6 @@ TUniformBuffer<FPrimitiveUniformShaderParameters>* FNiagaraSceneProxy::GetCustom
 	return CustomUBRef;
 }
 
-TUniformBuffer<FPrimitiveUniformShaderParameters>* FNiagaraSceneProxy::GetCustomUniformBufferResource(bool bHasVelocity, const FBox& InstanceBounds) const
-{
-	return GetCustomUniformBufferResource(FRHICommandListImmediate::Get(), bHasVelocity, InstanceBounds);
-}
-
 FRHIUniformBuffer* FNiagaraSceneProxy::GetCustomUniformBuffer(FRHICommandListBase& RHICmdList, bool bHasVelocity, const FBox& InstanceBounds) const
 {
 	// Default UB we create for the primitive
@@ -435,11 +431,6 @@ FRHIUniformBuffer* FNiagaraSceneProxy::GetCustomUniformBuffer(FRHICommandListBas
 	return GetCustomUniformBufferResource(RHICmdList, bHasVelocity, InstanceBounds)->GetUniformBufferRHI();
 }
 
-FRHIUniformBuffer* FNiagaraSceneProxy::GetCustomUniformBuffer(bool bHasVelocity, const FBox& InstanceBounds) const
-{
-	return GetCustomUniformBuffer(FRHICommandListImmediate::Get(), bHasVelocity, InstanceBounds);
-}
-
 uint32 FNiagaraSceneProxy::GetMemoryFootprint() const
 { 
 	return (sizeof(*this) + GetAllocatedSize()); 
diff --git a/Engine/Plugins/FX/Niagara/Source/Niagara/Private/NiagaraRendererRibbons.cpp b/Engine/Plugins/FX/Niagara/Source/Niagara/Private/NiagaraRendererRibbons.cpp
--- a/Engine/Plugins/FX/Niagara/Source/Niagara/Private/NiagaraRendererRibbons.cpp
+++ b/Engine/Plugins/FX/Niagara/Source/Niagara/Private/NiagaraRendererRibbons.cpp
@@ -2368,12 +2368,14 @@ void FNiagaraRendererRibbons::InitializeVertexBuffersResources(FRHICommandListBa
 		RenderingResources->ParticleIntDataStride = RenderingResources->ParticleData.IntStride / sizeof(int32);
 	}
 
-	
 	// If the data was generated sync it here, otherwise we rely on the generation step later to populate it
 	if (DynamicDataRibbon->GenerationOutput.IsValid() && DynamicDataRibbon->GenerationOutput->SegmentData.Num() > 0)
 	{
-		const auto& GeneratedGeometryData = *DynamicDataRibbon->GenerationOutput;
-		
+		//-OPT: We only need to update this data once for all GDME passes
+		UE::TScopeLock VertexBuffersLock(VertexBuffersGuard);
+
+		const FNiagaraRibbonCPUGeneratedVertexData& GeneratedGeometryData = *DynamicDataRibbon->GenerationOutput;
+
 		void *IndexPtr = RHICmdList.LockBuffer(VertexBuffers.SortedIndicesBuffer.Buffer, 0, GeneratedGeometryData.SortedIndices.Num() * sizeof(int32), RLM_WriteOnly);
 		FMemory::Memcpy(IndexPtr, GeneratedGeometryData.SortedIndices.GetData(), GeneratedGeometryData.SortedIndices.Num() * sizeof(int32));
 		RHICmdList.UnlockBuffer(VertexBuffers.SortedIndicesBuffer.Buffer);
diff --git a/Engine/Plugins/FX/Niagara/Source/Niagara/Public/NiagaraRendererRibbons.h b/Engine/Plugins/FX/Niagara/Source/Niagara/Public/NiagaraRendererRibbons.h
--- a/Engine/Plugins/FX/Niagara/Source/Niagara/Public/NiagaraRendererRibbons.h
+++ b/Engine/Plugins/FX/Niagara/Source/Niagara/Public/NiagaraRendererRibbons.h
@@ -395,6 +395,7 @@ protected:
 	bool bGpuRibbonLinkIsFloat = false;
 	uint32 GpuRibbonLinkOrderOffset = INDEX_NONE;
 
+	mutable UE::FMutex VertexBuffersGuard;
 	mutable FNiagaraRibbonVertexBuffers VertexBuffers;
 	
 	mutable FNiagaraRibbonTessellationSmoothingData TessellationSmoothingData;
diff --git a/Engine/Plugins/FX/Niagara/Source/Niagara/Public/NiagaraSceneProxy.h b/Engine/Plugins/FX/Niagara/Source/Niagara/Public/NiagaraSceneProxy.h
--- a/Engine/Plugins/FX/Niagara/Source/Niagara/Public/NiagaraSceneProxy.h
+++ b/Engine/Plugins/FX/Niagara/Source/Niagara/Public/NiagaraSceneProxy.h
@@ -47,14 +47,8 @@ public:
 
 	NIAGARA_API TUniformBuffer<FPrimitiveUniformShaderParameters>* GetCustomUniformBufferResource(FRHICommandListBase& RHICmdList, bool bHasVelocity, const FBox& InstanceBounds = FBox(ForceInitToZero)) const;
 
-	UE_DEPRECATED(5.4, "GetCustomUniformBufferResource requires a command list.")
-	NIAGARA_API TUniformBuffer<FPrimitiveUniformShaderParameters>* GetCustomUniformBufferResource(bool bHasVelocity, const FBox& InstanceBounds = FBox(ForceInitToZero)) const;
-
 	NIAGARA_API FRHIUniformBuffer* GetCustomUniformBuffer(FRHICommandListBase& RHICmdList, bool bHasVelocity, const FBox& InstanceBounds = FBox(ForceInitToZero)) const;
 
-	UE_DEPRECATED(5.4, "GetCustomUniformBuffer requires a command list.")
-	NIAGARA_API FRHIUniformBuffer* GetCustomUniformBuffer(bool bHasVelocity, const FBox& InstanceBounds = FBox(ForceInitToZero)) const;
-
 	NIAGARA_API virtual FPrimitiveViewRelevance GetViewRelevance(const FSceneView* View) const override;
 
 	/** Some proxy wide dynamic settings passed down with the emitter dynamic data. */
@@ -92,6 +86,7 @@ private:
 
 private:
 	/** Custom Uniform Buffers, allows us to have renderer specific data packed inside such as pre-skinned bounds. */
+	mutable UE::FMutex CustomUniformBuffersGuard;
 	mutable TMap<uint32, TUniformBuffer<FPrimitiveUniformShaderParameters>*> CustomUniformBuffers;
 
 	/** The data required to render a single instance of a NiagaraSystem */
