PR #12191: Update NiagaraPreviewGrid.h to add missing override on virtual functions

#rnx

diff --git a/Engine/Plugins/FX/Niagara/Source/Niagara/Public/NiagaraPreviewGrid.h b/Engine/Plugins/FX/Niagara/Source/Niagara/Public/NiagaraPreviewGrid.h
--- a/Engine/Plugins/FX/Niagara/Source/Niagara/Public/NiagaraPreviewGrid.h
+++ b/Engine/Plugins/FX/Niagara/Source/Niagara/Public/NiagaraPreviewGrid.h
@@ -21,7 +21,7 @@ class ANiagaraPreviewBase : public AActor
 
 public:
 	//AActor Interface
-	virtual bool ShouldTickIfViewportsOnly() const final { return true; }
+	virtual bool ShouldTickIfViewportsOnly() const override final { return true; }
 	//AActor Interface End
 
 	UFUNCTION(BlueprintCallable, BlueprintImplementableEvent, CallInEditor, Category = "Niagara|Preview")
@@ -43,12 +43,12 @@ class UNiagaraPreviewAxis : public UObject
 
 private:
 	virtual int32 Num_Implementation() { return 1; }
-	virtual void ApplyToPreview_Implementation(UNiagaraComponent* PreviewComponent, int32 PreviewIndex, bool bIsXAxis, FString& OutLabelText){ }
+	virtual void ApplyToPreview_Implementation(UNiagaraComponent* PreviewComponent, int32 PreviewIndex, bool bIsXAxis, FString& OutLabelText) { }
 public:
 
 	//UObject Interface
 #if WITH_EDITOR
-	virtual void PostEditChangeProperty(struct FPropertyChangedEvent& PropertyChangedEvent)override;
+	virtual void PostEditChangeProperty(struct FPropertyChangedEvent& PropertyChangedEvent) override;
 #endif
 	//UObject Interface END
 	
@@ -72,7 +72,7 @@ protected:
 	UPROPERTY(EditAnywhere, Category = Axis)
 	int32 Count;
 
-	virtual int32 Num_Implementation() { return Count; }
+	virtual int32 Num_Implementation() override { return Count; }
 };
 
 UCLASS(EditInlineNew, BlueprintType, Blueprintable)
@@ -87,7 +87,7 @@ private:
 	int32 Max;
 
 	//UNiagaraPreviewAxis Interface
-	virtual void ApplyToPreview_Implementation(UNiagaraComponent* PreviewComponent, int32 PreviewIndex, bool bIsXAxis, FString& OutLabelText);
+	virtual void ApplyToPreview_Implementation(UNiagaraComponent* PreviewComponent, int32 PreviewIndex, bool bIsXAxis, FString& OutLabelText) override;
 	//UNiagaraPreviewAxis Interface END
 };
 
@@ -103,7 +103,7 @@ private:
 	float Max;
 
 	//UNiagaraPreviewAxis Interface
-	virtual void ApplyToPreview_Implementation(UNiagaraComponent* PreviewComponent, int32 PreviewIndex, bool bIsXAxis, FString& OutLabelText);
+	virtual void ApplyToPreview_Implementation(UNiagaraComponent* PreviewComponent, int32 PreviewIndex, bool bIsXAxis, FString& OutLabelText) override;
 	//UNiagaraPreviewAxis Interface END
 };
 
@@ -119,7 +119,7 @@ private:
 	FVector2D Max;
 
 	//UNiagaraPreviewAxis Interface
-	virtual void ApplyToPreview_Implementation(UNiagaraComponent* PreviewComponent, int32 PreviewIndex, bool bIsXAxis, FString& OutLabelText);
+	virtual void ApplyToPreview_Implementation(UNiagaraComponent* PreviewComponent, int32 PreviewIndex, bool bIsXAxis, FString& OutLabelText) override;
 	//UNiagaraPreviewAxis Interface END
 };
 
@@ -135,7 +135,7 @@ private:
 	FVector Max;
 
 	//UNiagaraPreviewAxis Interface
-	virtual void ApplyToPreview_Implementation(UNiagaraComponent* PreviewComponent, int32 PreviewIndex, bool bIsXAxis, FString& OutLabelText);
+	virtual void ApplyToPreview_Implementation(UNiagaraComponent* PreviewComponent, int32 PreviewIndex, bool bIsXAxis, FString& OutLabelText) override;
 	//UNiagaraPreviewAxis Interface END
 };
 
@@ -151,7 +151,7 @@ private:
 	FVector4 Max;
 
 	//UNiagaraPreviewAxis Interface
-	virtual void ApplyToPreview_Implementation(UNiagaraComponent* PreviewComponent, int32 PreviewIndex, bool bIsXAxis, FString& OutLabelText);
+	virtual void ApplyToPreview_Implementation(UNiagaraComponent* PreviewComponent, int32 PreviewIndex, bool bIsXAxis, FString& OutLabelText) override;
 	//UNiagaraPreviewAxis Interface END
 };
 
@@ -167,7 +167,7 @@ private:
 	FLinearColor Max;
 
 	//UNiagaraPreviewAxis Interface
-	virtual void ApplyToPreview_Implementation(UNiagaraComponent* PreviewComponent, int32 PreviewIndex, bool bIsXAxis, FString& OutLabelText);
+	virtual void ApplyToPreview_Implementation(UNiagaraComponent* PreviewComponent, int32 PreviewIndex, bool bIsXAxis, FString& OutLabelText) override;
 	//UNiagaraPreviewAxis Interface END
 };
 
@@ -186,15 +186,15 @@ class ANiagaraPreviewGrid : public AActor
 
 public:
 	//UObject Interface
-	virtual void PostLoad()override;
-	virtual void BeginDestroy()override;
+	virtual void PostLoad() override;
+	virtual void BeginDestroy() override;
 #if WITH_EDITOR
-	virtual void PostEditChangeProperty(struct FPropertyChangedEvent& PropertyChangedEvent)override;
+	virtual void PostEditChangeProperty(struct FPropertyChangedEvent& PropertyChangedEvent) override;
 #endif
 	//UObject Interface End
 
 	//AActor Interface
-	virtual void TickActor(float DeltaTime, enum ELevelTick TickType, FActorTickFunction& ThisTickFunction)override;
+	virtual void TickActor(float DeltaTime, enum ELevelTick TickType, FActorTickFunction& ThisTickFunction) override;
 	virtual bool ShouldTickIfViewportsOnly() const override { return true; }
 	//AActor Interface End
 
