- Add stateless emitters report to audit commandlet

#rnx
#tests Ran audit on large project to gather data

diff --git a/Engine/Plugins/FX/Niagara/Source/NiagaraEditor/Private/Commandlets/NiagaraSystemAuditCommandlet.cpp b/Engine/Plugins/FX/Niagara/Source/NiagaraEditor/Private/Commandlets/NiagaraSystemAuditCommandlet.cpp
--- a/Engine/Plugins/FX/Niagara/Source/NiagaraEditor/Private/Commandlets/NiagaraSystemAuditCommandlet.cpp
+++ b/Engine/Plugins/FX/Niagara/Source/NiagaraEditor/Private/Commandlets/NiagaraSystemAuditCommandlet.cpp
@@ -236,7 +236,7 @@ bool UNiagaraSystemAuditCommandlet::ProcessNiagaraSystems()
 	UPackage* CurrentPackage = nullptr;
 
 	const bool bColledGarbage = NiagaraValidationIssues.IsValid();
-	const int32 CollectGarbageFrequency = 1024;
+	const int32 CollectGarbageFrequency = 256;
 	int32 CollectGarabageCounter = 0;
 
 	for (const FAssetData& AssetIt : AssetList)
@@ -317,9 +317,12 @@ bool UNiagaraSystemAuditCommandlet::ProcessNiagaraSystems()
 		bool bHasEvents = false;
 		bool bHasGPUEmitters = false;
 		bool bHasCPUEmitters = false;
+		bool bHasStatelessEmitters = false;
 
 		for (const FNiagaraEmitterHandle& EmitterHandle : NiagaraSystem->GetEmitterHandles())
 		{
+			bHasStatelessEmitters |= (EmitterHandle.GetEmitterMode() == ENiagaraEmitterMode::Stateless) && EmitterHandle.GetStatelessEmitter();
+
 			FVersionedNiagaraEmitterData* EmitterData = EmitterHandle.GetEmitterData();
 			if (EmitterData == nullptr)
 			{
@@ -474,6 +477,11 @@ bool UNiagaraSystemAuditCommandlet::ProcessNiagaraSystems()
 			NiagaraSystemsWithEvents.Add(NiagaraSystem->GetPathName());
 		}
 
+		if (bHasStatelessEmitters)
+		{
+			NiagaraSystemsWithStatelessEmitters.Add(NiagaraSystem->GetPathName());
+		}
+
 		if (SystemDataInterfacesWihPrereqs.Num() > 0)
 		{
 			FString DataInterfaceNames;
@@ -585,6 +593,7 @@ void UNiagaraSystemAuditCommandlet::DumpResults()
 	}
 	DumpSimpleSet(NiagaraSystemsWithSimulationStages, TEXT("NiagaraSystemsWithSimulationStages"), TEXT("System,Emitters"));
 	DumpSimpleSet(NiagaraSystemsWithCompression, TEXT("NiagaraSystemsWithCompression"), TEXT("System,ExecutionTypes"));
+	DumpSimpleSet(NiagaraSystemsWithStatelessEmitters, TEXT("NiagaraSystemsWithStatelessEmitters"), TEXT("System"));
 
 	if (NiagaraDataInterfaceUsage.Num() > 0)
 	{
diff --git a/Engine/Plugins/FX/Niagara/Source/NiagaraEditor/Public/Commandlets/NiagaraSystemAuditCommandlet.h b/Engine/Plugins/FX/Niagara/Source/NiagaraEditor/Public/Commandlets/NiagaraSystemAuditCommandlet.h
--- a/Engine/Plugins/FX/Niagara/Source/NiagaraEditor/Public/Commandlets/NiagaraSystemAuditCommandlet.h
+++ b/Engine/Plugins/FX/Niagara/Source/NiagaraEditor/Public/Commandlets/NiagaraSystemAuditCommandlet.h
@@ -30,6 +30,8 @@ class UNiagaraSystemAuditCommandlet : public UCommandlet
 	TSet<FString> NiagaraSystemsWithSimulationStages;
 	/** All Niagara systems with compression enabled */
 	TSet<FString> NiagaraSystemsWithCompression;
+	/** All Niagara systems with stateless emitters enabled */
+	TSet<FString> NiagaraSystemsWithStatelessEmitters;
 
 	/** All Validation Issues */
 	TUniquePtr<FArchive> NiagaraValidationIssues;
