- Fixes for stateless rate spawning and looping
- Some of the calculations could miss residual amounts when crossing between loops
- Float precision issues could also creep in over time resulting in divergence between stateful and stateless

#rnx

diff --git a/Engine/Plugins/FX/Niagara/Source/Niagara/Internal/Stateless/NiagaraStatelessEmitterInstance.h b/Engine/Plugins/FX/Niagara/Source/Niagara/Internal/Stateless/NiagaraStatelessEmitterInstance.h
--- a/Engine/Plugins/FX/Niagara/Source/Niagara/Internal/Stateless/NiagaraStatelessEmitterInstance.h
+++ b/Engine/Plugins/FX/Niagara/Source/Niagara/Internal/Stateless/NiagaraStatelessEmitterInstance.h
@@ -47,7 +47,7 @@ class FNiagaraStatelessEmitterInstance final : public FNiagaraEmitterInstance
 	struct FActiveSpawnRate
 	{
 		float	Rate = 0.0f;
-		float	SpawnTime = 0.0f;
+		float	ResidualSpawnTime = 0.0f;
 	};
 
 public:
diff --git a/Engine/Plugins/FX/Niagara/Source/Niagara/Internal/Stateless/NiagaraStatelessSpawnInfo.h b/Engine/Plugins/FX/Niagara/Source/Niagara/Internal/Stateless/NiagaraStatelessSpawnInfo.h
--- a/Engine/Plugins/FX/Niagara/Source/Niagara/Internal/Stateless/NiagaraStatelessSpawnInfo.h
+++ b/Engine/Plugins/FX/Niagara/Source/Niagara/Internal/Stateless/NiagaraStatelessSpawnInfo.h
@@ -55,5 +55,5 @@ struct FNiagaraStatelessRuntimeSpawnInfo
 	float	SpawnTimeStart	= 0.0f;
 	float	SpawnTimeEnd	= 0.0f;
 	float	Rate			= 0.0f;
-	int32	Amount			= 0;
+	int32	Amount			= 0;		// Note: When a burst this is the absolute amount burst, when rate this is the amount over the spawn duration
 };
diff --git a/Engine/Plugins/FX/Niagara/Source/Niagara/Private/Stateless/NiagaraStatelessEmitterData.cpp b/Engine/Plugins/FX/Niagara/Source/Niagara/Private/Stateless/NiagaraStatelessEmitterData.cpp
--- a/Engine/Plugins/FX/Niagara/Source/Niagara/Private/Stateless/NiagaraStatelessEmitterData.cpp
+++ b/Engine/Plugins/FX/Niagara/Source/Niagara/Private/Stateless/NiagaraStatelessEmitterData.cpp
@@ -88,24 +88,20 @@ uint32 FNiagaraStatelessEmitterData::CalculateActiveParticles(int32 InRandomSeed
 		const bool bIsValidForAge = !Age.IsSet() || (Age.GetValue() >= SpawnInfo.SpawnTimeStart && Age.GetValue() < SpawnInfo.SpawnTimeEnd + LifetimeRange.Max);
 		if (bIsValidForAge && GpuSpawnIndex < NiagaraStateless::MaxGpuSpawnInfos)
 		{
-			uint32	NumActive = 0;
-			uint32	ParticleOffset = 0;
-			float	SpawnRate = 0.0f;
-			float	SpawnTimeStart = SpawnInfo.SpawnTimeStart;
+			uint32	NumActive		= SpawnInfo.Amount;
+			uint32	ParticleOffset	= 0;
+			float	SpawnRate		= 0.0f;
+			float	SpawnTimeStart	= SpawnInfo.SpawnTimeStart;
 			switch (SpawnInfo.Type)
 			{
 				case ENiagaraStatelessSpawnInfoType::Burst:
-					NumActive				= SpawnInfo.Amount;
 					break;
 
 				case ENiagaraStatelessSpawnInfoType::Rate:
-					if (!Age.IsSet())
+					// If the age is set we can make a narrowed number of active particles calculation
+					if (Age.IsSet())
 					{
-						NumActive = FMath::FloorToInt((SpawnInfo.SpawnTimeEnd - SpawnInfo.SpawnTimeStart) * SpawnInfo.Rate);
-					}
-					else
-					{
-						const uint32 MaxActive = FMath::FloorToInt((SpawnInfo.SpawnTimeEnd - SpawnInfo.SpawnTimeStart) * SpawnInfo.Rate);
+						const uint32 MaxActive = SpawnInfo.Amount;
 						ParticleOffset = FMath::FloorToInt(FMath::Max(Age.GetValue() - SpawnInfo.SpawnTimeStart - LifetimeRange.Max, 0.0f) * SpawnInfo.Rate);
 						ParticleOffset = FMath::Min(ParticleOffset, MaxActive);
 						NumActive = FMath::FloorToInt(FMath::Max(Age.GetValue() - SpawnInfo.SpawnTimeStart, 0.0f) * SpawnInfo.Rate);
diff --git a/Engine/Plugins/FX/Niagara/Source/Niagara/Private/Stateless/NiagaraStatelessEmitterInstance.cpp b/Engine/Plugins/FX/Niagara/Source/Niagara/Private/Stateless/NiagaraStatelessEmitterInstance.cpp
--- a/Engine/Plugins/FX/Niagara/Source/Niagara/Private/Stateless/NiagaraStatelessEmitterInstance.cpp
+++ b/Engine/Plugins/FX/Niagara/Source/Niagara/Private/Stateless/NiagaraStatelessEmitterInstance.cpp
@@ -550,50 +550,49 @@ void FNiagaraStatelessEmitterInstance::InitSpawnInfosForLoop(float Initializatio
 	}
 
 	// Add the next chunk for any active spawn rates
-	for (const FActiveSpawnRate& SpawnInfo : ActiveSpawnRates)
+	for (FActiveSpawnRate& SpawnInfo : ActiveSpawnRates)
 	{
-		float SpawnTime = FMath::Max(InitializationAge - SpawnInfo.SpawnTime, 0.0f);
-		SpawnTime = FMath::CeilToFloat(SpawnTime * SpawnInfo.Rate) / SpawnInfo.Rate;
-		SpawnTime += SpawnInfo.SpawnTime + CurrentLoopDelay;
-		if (SpawnTime >= CurrentLoopAgeEnd)
-		{
-			continue;
-		}
+		const float SpawnAgeStart	= FMath::Min(InitializationAge + CurrentLoopDelay - SpawnInfo.ResidualSpawnTime, CurrentLoopAgeEnd);
+		const float ActiveDuration	= CurrentLoopAgeEnd - SpawnAgeStart;
+		const int32 NumSpawned		= FMath::FloorToInt(ActiveDuration * SpawnInfo.Rate);
+		const float SpawnAgeEnd		= SpawnAgeStart + (float(NumSpawned) / SpawnInfo.Rate);
 
-		// Try and append to the last info in the list if it's a rate type
-		// We do this to reduce the number of spawn infos in the common case of having a single rate info
-		bool bDidAppend = false;
-		if ( SpawnInfos.Num() > 0 )
+		if ( NumSpawned > 0 )
 		{
-			FNiagaraStatelessRuntimeSpawnInfo& ExistingInfo = SpawnInfos.Last();
-			//-TODO: Validate time start better in the case of active / deactivate
-			if ( (ExistingInfo.Type == ENiagaraStatelessSpawnInfoType::Rate) && (ExistingInfo.Rate == SpawnInfo.Rate) && (ExistingInfo.SpawnTimeStart == SpawnInfo.SpawnTime))
+			// Try and append to the last info in the list if it's a rate type
+			// We do this to reduce the number of spawn infos in the common case of having a single rate info
+			bool bDidAppend = false;
+			if ( SpawnInfos.Num() > 0 )
 			{
-				const float ActiveDuration = ExistingInfo.SpawnTimeEnd - ExistingInfo.SpawnTimeStart;
-				const int32 NumSpawned = FMath::FloorToInt(ActiveDuration * ExistingInfo.Rate);
-				if (ExistingInfo.UniqueOffset + NumSpawned == UniqueIndexOffset)
+				FNiagaraStatelessRuntimeSpawnInfo& ExistingInfo = SpawnInfos.Last();
+				if ( (ExistingInfo.Type == ENiagaraStatelessSpawnInfoType::Rate) && (ExistingInfo.Rate == SpawnInfo.Rate) && (ExistingInfo.SpawnTimeEnd == SpawnAgeStart))
 				{
-					ExistingInfo.SpawnTimeEnd = CurrentLoopAgeEnd;
-					bDidAppend = true;
+					if (ExistingInfo.UniqueOffset + ExistingInfo.Amount == UniqueIndexOffset)
+					{
+						ExistingInfo.SpawnTimeEnd = SpawnAgeEnd;
+						ExistingInfo.Amount += NumSpawned;
+						bDidAppend = true;
+					}
 				}
 			}
-		}
 
-		if (!bDidAppend)
-		{
-			FNiagaraStatelessRuntimeSpawnInfo& NewSpawnInfo = SpawnInfos.AddDefaulted_GetRef();
-			NewSpawnInfo.Type			= ENiagaraStatelessSpawnInfoType::Rate;
-			NewSpawnInfo.UniqueOffset	= UniqueIndexOffset;
-			NewSpawnInfo.SpawnTimeStart	= SpawnTime;
-			NewSpawnInfo.SpawnTimeEnd	= CurrentLoopAgeEnd;
-			NewSpawnInfo.Rate			= SpawnInfo.Rate;
-		}
+			//-TODO: We need to add a spawn info as we have something to spawn
+			if (!bDidAppend)
+			{
+				FNiagaraStatelessRuntimeSpawnInfo& NewSpawnInfo = SpawnInfos.AddDefaulted_GetRef();
+				NewSpawnInfo.Type			= ENiagaraStatelessSpawnInfoType::Rate;
+				NewSpawnInfo.UniqueOffset	= UniqueIndexOffset;
+				NewSpawnInfo.SpawnTimeStart	= SpawnAgeStart;
+				NewSpawnInfo.SpawnTimeEnd	= SpawnAgeEnd;
+				NewSpawnInfo.Rate			= SpawnInfo.Rate;
+				NewSpawnInfo.Amount			= NumSpawned;
+			}
 
-		const float ActiveDuration	= CurrentLoopAgeEnd - SpawnTime;
-		const int32 NumSpawned		= FMath::FloorToInt(ActiveDuration * SpawnInfo.Rate);
+			UniqueIndexOffset += NumSpawned;
+			bSpawnInfosDirty = true;
+		}
 
-		UniqueIndexOffset += NumSpawned;
-		bSpawnInfosDirty = true;
+		SpawnInfo.ResidualSpawnTime = CurrentLoopAgeEnd - SpawnAgeEnd;
 	}
 
 	// Add bursts that fit within the loop duration (due to loop random they might not)
@@ -682,6 +681,7 @@ void FNiagaraStatelessEmitterInstance::CropSpawnInfos()
 		if (SpawnInfo.Type == ENiagaraStatelessSpawnInfoType::Rate)
 		{
 			SpawnInfo.SpawnTimeEnd = FMath::Min(SpawnInfo.SpawnTimeEnd, Age);
+			SpawnInfo.Amount = FMath::FloorToInt((SpawnInfo.SpawnTimeEnd - SpawnInfo.SpawnTimeStart) * SpawnInfo.Rate);
 		}
 		if (Age < SpawnInfo.SpawnTimeStart || Age >= SpawnInfo.SpawnTimeEnd + MaxLifetime)
 		{
