- Fix Niagara mesh renderer not respecting the section cast shadows flag
- Fix Niagara mesh renderer not using the newer raytracing LOD representation

#rnx
#tests Tested changing mesh section cast shadows and that RT shadows worked as expected

diff --git a/Engine/Plugins/FX/Niagara/Source/Niagara/Private/NiagaraMeshRendererProperties.cpp b/Engine/Plugins/FX/Niagara/Source/Niagara/Private/NiagaraMeshRendererProperties.cpp
--- a/Engine/Plugins/FX/Niagara/Source/Niagara/Private/NiagaraMeshRendererProperties.cpp
+++ b/Engine/Plugins/FX/Niagara/Source/Niagara/Private/NiagaraMeshRendererProperties.cpp
@@ -63,7 +63,7 @@ public:
 		return LocalBounds;
 	}
 
-	void GetLODModelData(FLODModelData& OutLODModelData, int32 LODLevel) const override
+	virtual void GetLODModelData(FLODModelData& OutLODModelData, int32 LODLevel) const override
 	{
 		LODLevel = FMath::Max(MinLOD, LODLevel);
 		OutLODModelData.LODIndex = RenderData->GetCurrentFirstLODIdx(LODLevel);
@@ -80,7 +80,7 @@ public:
 		OutLODModelData.Sections = MakeArrayView(LODResources.Sections);
 		OutLODModelData.IndexBuffer = &LODResources.IndexBuffer;
 		OutLODModelData.VertexFactoryUserData = RenderData->LODVertexFactories.IsValidIndex(OutLODModelData.LODIndex) ? RenderData->LODVertexFactories[OutLODModelData.LODIndex].VertexFactory.GetUniformBuffer() : nullptr;
-		OutLODModelData.RayTracingGeometry = LODResources.RayTracingGeometry;
+		OutLODModelData.RayTracingGeometry = nullptr;
 
 		if (LODResources.AdditionalIndexBuffers != nullptr && LODResources.AdditionalIndexBuffers->WireframeIndexBuffer.IsInitialized())
 		{
@@ -89,6 +89,28 @@ public:
 		}
 	}
 
+	virtual void GetRayTraceLODModelData(FLODModelData& OutLODModelData, int32 LODLevel) const override
+	{
+		GetLODModelData(OutLODModelData, LODLevel);
+		if (OutLODModelData.LODIndex == INDEX_NONE)
+		{
+			return;
+		}
+
+		const FStaticMeshLODResources& LODResources = RenderData->LODResources[OutLODModelData.LODIndex];
+		OutLODModelData.RayTracingGeometry = LODResources.RayTracingGeometry;
+
+		if (FStaticMeshRayTracingProxy* RayTracingProxy = RenderData->RayTracingProxy)
+		{
+			const int32 RayTracingLOD = RayTracingProxy->bUsingRenderingLODs ? OutLODModelData.LODIndex : 0;
+
+			if (RayTracingProxy->LODs.IsValidIndex(RayTracingLOD))
+			{
+				OutLODModelData.RayTracingGeometry = RayTracingProxy->LODs[RayTracingLOD].RayTracingGeometry;
+			}
+		}
+	}
+
 	virtual FIntVector2 GetLODRange() const override
 	{
 		return FIntVector2(MeshMinLOD, RenderData->LODResources.Num());
diff --git a/Engine/Plugins/FX/Niagara/Source/Niagara/Private/NiagaraRendererMeshes.cpp b/Engine/Plugins/FX/Niagara/Source/Niagara/Private/NiagaraRendererMeshes.cpp
--- a/Engine/Plugins/FX/Niagara/Source/Niagara/Private/NiagaraRendererMeshes.cpp
+++ b/Engine/Plugins/FX/Niagara/Source/Niagara/Private/NiagaraRendererMeshes.cpp
@@ -1283,9 +1283,9 @@ void FNiagaraRendererMeshes::CreateMeshBatchForSection(
 	MeshBatch.VertexFactory = &VertexFactory;
 	MeshBatch.LCI = NULL;
 	MeshBatch.ReverseCulling = SceneProxy.IsLocalToWorldDeterminantNegative();
-	MeshBatch.CastShadow = SceneProxy.CastsDynamicShadow() && bCastShadows;
+	MeshBatch.CastShadow = SceneProxy.CastsDynamicShadow() && bCastShadows && Section.bCastShadow;
 #if RHI_RAYTRACING
-	MeshBatch.CastRayTracedShadow = SceneProxy.CastsDynamicShadow() && bCastShadows;
+	MeshBatch.CastRayTracedShadow = SceneProxy.CastsDynamicShadow() && bCastShadows && Section.bCastShadow;
 #endif
 	MeshBatch.DepthPriorityGroup = (ESceneDepthPriorityGroup)SceneProxy.GetDepthPriorityGroup(&View);
 
@@ -1674,7 +1674,7 @@ void FNiagaraRendererMeshes::GetDynamicRayTracingInstances(FRayTracingMaterialGa
 		const FMeshData& MeshData = Meshes[MeshIndex];
 
 		INiagaraRenderableMesh::FLODModelData LODModel;
-		MeshData.RenderableMesh->GetLODModelData(LODModel, MeshData.LODLevel);
+		MeshData.RenderableMesh->GetRayTraceLODModelData(LODModel, MeshData.LODLevel);
 		if (LODModel.LODIndex == INDEX_NONE || LODModel.RayTracingGeometry == nullptr)
 		{
 			continue;
diff --git a/Engine/Plugins/FX/Niagara/Source/Niagara/Public/NiagaraRenderableMeshInterface.h b/Engine/Plugins/FX/Niagara/Source/Niagara/Public/NiagaraRenderableMeshInterface.h
--- a/Engine/Plugins/FX/Niagara/Source/Niagara/Public/NiagaraRenderableMeshInterface.h
+++ b/Engine/Plugins/FX/Niagara/Source/Niagara/Public/NiagaraRenderableMeshInterface.h
@@ -44,6 +44,8 @@ public:
 	virtual FBox GetLocalBounds() const = 0;
 	// Gather all the relevant mesh data to render the mesh
 	virtual void GetLODModelData(FLODModelData& OutLODModelData, int32 LODLevel) const = 0;
+	// Gather all the relevant mesh data to render the mesh in the raytracing scene
+	virtual void GetRayTraceLODModelData(FLODModelData& OutLODModelData, int32 LODLevel) const { GetLODModelData(OutLODModelData, LODLevel); }
 	// Setup the vertex factory for the mesh
 	virtual void SetupVertexFactory(FRHICommandListBase& RHICmdList, class FNiagaraMeshVertexFactory& InVertexFactory, const FLODModelData& LODModelData) const = 0;
 
