- Add variable type to name for parameter outputs

#rnx
#tests Observed results in module parameter reads / writes

diff --git a/Engine/Plugins/FX/Niagara/Source/NiagaraEditor/Private/ViewModels/Stack/NiagaraStackModuleItemOutput.cpp b/Engine/Plugins/FX/Niagara/Source/NiagaraEditor/Private/ViewModels/Stack/NiagaraStackModuleItemOutput.cpp
--- a/Engine/Plugins/FX/Niagara/Source/NiagaraEditor/Private/ViewModels/Stack/NiagaraStackModuleItemOutput.cpp
+++ b/Engine/Plugins/FX/Niagara/Source/NiagaraEditor/Private/ViewModels/Stack/NiagaraStackModuleItemOutput.cpp
@@ -61,6 +61,9 @@ const UNiagaraStackModuleItemOutput::FCollectedUsageData& UNiagaraStackModuleIte
 FText UNiagaraStackModuleItemOutput::GetTooltipText() const
 {
 	FNiagaraVariable ValueVariable(OutputType, OutputParameterHandle.GetParameterHandleString());
+
+	FText VariableText = FText::Format(LOCTEXT("OutputVariableTooltip", "{0} ({1})"), FText::FromName(ValueVariable.GetName()), ValueVariable.GetType().GetNameText());
+
 	if (FunctionCallNode.IsValid() && FunctionCallNode->FunctionScript != nullptr)
 	{
 		UNiagaraScriptSource* Source = FunctionCallNode->GetFunctionScriptSource();
@@ -73,13 +76,12 @@ FText UNiagaraStackModuleItemOutput::GetTooltipText() const
 		{
 			MetaData = Source->NodeGraph->GetMetaData(ValueVariable);
 		}
-
-		if (MetaData.IsSet())
+		if (MetaData.IsSet() && !MetaData->Description.IsEmpty())
 		{
-			return MetaData->Description;
+			return FText::Format(LOCTEXT("OutputVariableWithDescTooltip", "{0}\n{1}"), VariableText, MetaData->Description);
 		}
 	}
-	return FText::FromName(ValueVariable.GetName());
+	return VariableText;
 }
 
 bool UNiagaraStackModuleItemOutput::GetIsEnabled() const
