- Handle a null data interface inside GenerateDefaultFunctionBindings
- It is unclear how this occurs, but the ensure will give us more clues

#rnx
#tests Forced a DI to be null to ensure we didn't crash elsewhere

diff --git a/Engine/Plugins/FX/Niagara/Source/Niagara/Private/NiagaraScript.cpp b/Engine/Plugins/FX/Niagara/Source/Niagara/Private/NiagaraScript.cpp
--- a/Engine/Plugins/FX/Niagara/Source/Niagara/Private/NiagaraScript.cpp
+++ b/Engine/Plugins/FX/Niagara/Source/Niagara/Private/NiagaraScript.cpp
@@ -1817,10 +1817,13 @@ void UNiagaraScript::GenerateDefaultFunctionBindings()
 			for (int32 DataInterfaceIt = 0; DataInterfaceIt < DataInterfaceCount; ++DataInterfaceIt)
 			{
 				const FNiagaraScriptDataInterfaceCompileInfo& ScriptInfo = CachedScriptVM.DataInterfaceInfo[DataInterfaceIt];
-
 				if (ScriptInfo.UserPtrIdx == INDEX_NONE && ScriptInfo.Name == BindingInfo.OwnerName)
 				{
-					ScriptDataInterfaces[DataInterfaceIt]->GetVMExternalFunction(BindingInfo, nullptr, FuncBind);
+					UNiagaraDataInterface* ScriptDataInterface = ScriptDataInterfaces[DataInterfaceIt];
+					if ( ensureMsgf(ScriptDataInterface, TEXT("Script(%s) Data Interface Slot %d is nullptr"), *GetFullNameSafe(this), DataInterfaceIt) )
+					{
+						ScriptDataInterfaces[DataInterfaceIt]->GetVMExternalFunction(BindingInfo, nullptr, FuncBind);
+					}
 				}
 			}
 		}
