* Made implicit type truncations explicit
* Removed deprecated kismet math library methods
* fixed some bugs caused by incorrect type conversions

#rnx

diff --git a/Engine/Plugins/FX/Niagara/Source/Niagara/Internal/DataInterface/NiagaraDataInterfaceDataChannelRead.h b/Engine/Plugins/FX/Niagara/Source/Niagara/Internal/DataInterface/NiagaraDataInterfaceDataChannelRead.h
--- a/Engine/Plugins/FX/Niagara/Source/Niagara/Internal/DataInterface/NiagaraDataInterfaceDataChannelRead.h
+++ b/Engine/Plugins/FX/Niagara/Source/Niagara/Internal/DataInterface/NiagaraDataInterfaceDataChannelRead.h
@@ -212,8 +212,12 @@ private:
 	float Scale = 1.0f;
 
 public:
-	uint32 Get()const { return Count * Scale; }
-	
+	uint32 Get() const
+	{
+		float Value = Count * Scale;
+		return static_cast<uint32>(Value > MAX_int32 ? MAX_int32 : Value);
+	}
+
 	void SetCount(uint32 NewCount) { Count = NewCount; }
 	void Append(uint32 NewCount) { Count += NewCount; }
 	void SetScale(float NewScale) { Scale = NewScale; }
diff --git a/Engine/Plugins/FX/Niagara/Source/Niagara/Internal/Stateless/NiagaraStatelessParticleSimContext.h b/Engine/Plugins/FX/Niagara/Source/Niagara/Internal/Stateless/NiagaraStatelessParticleSimContext.h
--- a/Engine/Plugins/FX/Niagara/Source/Niagara/Internal/Stateless/NiagaraStatelessParticleSimContext.h
+++ b/Engine/Plugins/FX/Niagara/Source/Niagara/Internal/Stateless/NiagaraStatelessParticleSimContext.h
@@ -66,8 +66,8 @@ public:
 
 	template<typename TType>
 	const TType ReadParticleVariableComponent(int32 iVariable, int32 iInstance, int32 iComponent) const { return *reinterpret_cast<const TType*>(VariableComponents[iVariable] + (BufferStride * iComponent) + (iInstance * sizeof(float))); }
-	int32 ReadParticleVariable(int32 iVariable, uint32 iInstance, int32 DefaultValue) const { check(iInstance < NumInstances); return iVariable >= 0 ? ReadParticleVariableComponent<float>(iVariable, iInstance, 0) : DefaultValue; }
-	float ReadParticleVariable(int32 iVariable, uint32 iInstance, float DefaultValue) const { check(iInstance < NumInstances); return iVariable >= 0 ? ReadParticleVariableComponent<int32>(iVariable, iInstance, 0) : DefaultValue; }
+	int32 ReadParticleVariable(int32 iVariable, uint32 iInstance, int32 DefaultValue) const { check(iInstance < NumInstances); return iVariable >= 0 ? ReadParticleVariableComponent<int32>(iVariable, iInstance, 0) : DefaultValue; }
+	float ReadParticleVariable(int32 iVariable, uint32 iInstance, float DefaultValue) const { check(iInstance < NumInstances); return iVariable >= 0 ? ReadParticleVariableComponent<float>(iVariable, iInstance, 0) : DefaultValue; }
 	FVector2f ReadParticleVariable(int32 iVariable, uint32 iInstance, const FVector2f& DefaultValue) const { check(iInstance < NumInstances); return iVariable >= 0 ? FVector2f(ReadParticleVariableComponent<float>(iVariable, iInstance, 0), ReadParticleVariableComponent<float>(iVariable, iInstance, 1)) : DefaultValue; }
 	FVector3f ReadParticleVariable(int32 iVariable, uint32 iInstance, const FVector3f& DefaultValue) const { check(iInstance < NumInstances); return iVariable >= 0 ? FVector3f(ReadParticleVariableComponent<float>(iVariable, iInstance, 0), ReadParticleVariableComponent<float>(iVariable, iInstance, 1), ReadParticleVariableComponent<float>(iVariable, iInstance, 2)) : DefaultValue; }
 	FVector4f ReadParticleVariable(int32 iVariable, uint32 iInstance, const FVector4f& DefaultValue) const { check(iInstance < NumInstances); return iVariable >= 0 ? FVector4f(ReadParticleVariableComponent<float>(iVariable, iInstance, 0), ReadParticleVariableComponent<float>(iVariable, iInstance, 1), ReadParticleVariableComponent<float>(iVariable, iInstance, 2), ReadParticleVariableComponent<float>(iVariable, iInstance, 3)) : DefaultValue; }
diff --git a/Engine/Plugins/FX/Niagara/Source/Niagara/Private/DataInterface/NiagaraDataInterfaceDataChannelRead.cpp b/Engine/Plugins/FX/Niagara/Source/Niagara/Private/DataInterface/NiagaraDataInterfaceDataChannelRead.cpp
--- a/Engine/Plugins/FX/Niagara/Source/Niagara/Private/DataInterface/NiagaraDataInterfaceDataChannelRead.cpp
+++ b/Engine/Plugins/FX/Niagara/Source/Niagara/Private/DataInterface/NiagaraDataInterfaceDataChannelRead.cpp
@@ -146,7 +146,7 @@ namespace NDIDataChannelReadLocal
 			Sig.AddInput(FNiagaraVariable(FNiagaraTypeDefinition(UNiagaraDataInterfaceDataChannelRead::StaticClass()), TEXT("DataChannel interface")));
 			Sig.AddInputWithDefault(FNiagaraVariable(FNiagaraTypeDefinition::GetIntDef(), TEXT("Index")), 0, LOCTEXT("ConsumeIndexInputDesc", "The index to read."));
 			Sig.AddOutput(FNiagaraVariable(FNiagaraTypeDefinition::GetBoolDef(), TEXT("Success")), LOCTEXT("ConsumeSuccessOutputDesc", "True if all reads succeeded."));
-			Sig.RequiredOutputs = Sig.Outputs.Num();//The user defines what we read in the graph.			
+			Sig.RequiredOutputs = IntCastChecked<int16>(Sig.Outputs.Num()); //The user defines what we read in the graph.			
 		}
 		return Sig;
 	}
@@ -167,7 +167,7 @@ namespace NDIDataChannelReadLocal
 			Sig.AddInputWithDefault(FNiagaraVariable(FNiagaraTypeDefinition::GetBoolDef(), TEXT("Consume")), FNiagaraBool(true), LOCTEXT("ConsumeInputDesc", "True if this instance (particle/emitter etc) should consume data from the data channel in this call."));
 			Sig.AddOutput(FNiagaraVariable(FNiagaraTypeDefinition::GetBoolDef(), TEXT("Success")), LOCTEXT("ConsumeSuccessOutputDesc", "True if all reads succeeded."));
 			Sig.AddOutput(FNiagaraVariable(FNiagaraTypeDefinition::GetIntDef(), TEXT("Index")), LOCTEXT("ConsumeIndexOutputDesc", "The index we actually read from. If reading failed this can be -1. This allows subsequent reads of the data channel at this index."));
-			Sig.RequiredOutputs = Sig.Outputs.Num();//The user defines what we read in the graph.			
+			Sig.RequiredOutputs = IntCastChecked<int16>(Sig.Outputs.Num()); //The user defines what we read in the graph.			
 		}
 		return Sig;
 	}
@@ -192,7 +192,7 @@ namespace NDIDataChannelReadLocal
 		For compound data types that contain multiple component floats or ints, comparissons are done on a per component basis.\n\
 		For example if you add a Vector condition parameter it will be compared against each component of the corresponding Vector in the Data Channel.\n\
 		Result = (Param.X == ChannelValue.X) && (Param.Y == ChannelValue.Y) && (Param.Z == ChannelValue.Z)");
-			Sig.FunctionVersion = (uint32)FunctionVersion_SpawnConditional::EmitterIDParameter;
+			Sig.FunctionVersion = static_cast<uint32>(FunctionVersion_SpawnConditional::EmitterIDParameter);
 #endif
 			Sig.bMemberFunction = true;
 			Sig.bExperimental = true;
@@ -205,7 +205,7 @@ namespace NDIDataChannelReadLocal
 			Sig.AddInput(FNiagaraVariable(StaticEnum<ENiagaraConditionalOperator>(), TEXT("Operator")), LOCTEXT("SpawnCondOpInputDesc", "The comparison operator to use when comparing values in the data channel to conditional parameters."));
 			Sig.AddInputWithDefault(FNiagaraVariable(FNiagaraTypeDefinition::GetIntDef(), TEXT("Min Spawn Count")), 1, LOCTEXT("MinSpawnCountInputDesc", "Minimum number of particles to spawn for each element in the data channel."));
 			Sig.AddInputWithDefault(FNiagaraVariable(FNiagaraTypeDefinition::GetIntDef(), TEXT("Max Spawn Count")), 1, LOCTEXT("MaxSpawnCountInputDesc", "Maximum number of particles to spawn for each element in the data channel."));
-			Sig.RequiredInputs = Sig.Inputs.Num();
+			Sig.RequiredInputs = IntCastChecked<int16>(Sig.Inputs.Num());
 		}
 		return Sig;
 	}
@@ -1734,7 +1734,7 @@ void UNiagaraDataInterfaceDataChannelRead::SpawnConditional(FVectorVMExternalFun
 template<typename T> 
 float NDCValueSize(const T& Value){ return 1.0f; }
 
-template<> float NDCValueSize<int32>(const int32& Value){ return FMath::Abs(Value); }
+template<> float NDCValueSize<int32>(const int32& Value){ return static_cast<float>(FMath::Abs(Value)); }
 template<> float NDCValueSize<float>(const float& Value) { return FMath::Abs(Value); }
 template<> float NDCValueSize<FVector2f>(const FVector2f& Value) { return Value.Size(); }
 template<> float NDCValueSize<FVector3f>(const FVector3f& Value) { return Value.Size(); }
@@ -1833,7 +1833,7 @@ void UNiagaraDataInterfaceDataChannelRead::SpawnDirect(FVectorVMExternalFunction
 			TSimType NDCValue = ValueData.GetSafe(DataChannelIndex, NDCValueDefault<TSimType>());
 			double VarSize = NDCValueSize(NDCValue);
 			float Scale = RandHelper.RandRange(DataChannelIndex, RandMinScale, RandMaxScale);
-			int32 ScaledCount = VarSize * Scale;
+			int32 ScaledCount = FMath::TruncToInt32(VarSize * Scale);
 			int32 FinalCount = FMath::Clamp(ScaledCount, ClampMin, ClampMax);
 
 			if (Mode == ENDIDataChannelSpawnMode::Accumulate)
diff --git a/Engine/Plugins/FX/Niagara/Source/Niagara/Private/DataInterface/NiagaraDataInterfaceDataChannelWrite.cpp b/Engine/Plugins/FX/Niagara/Source/Niagara/Private/DataInterface/NiagaraDataInterfaceDataChannelWrite.cpp
--- a/Engine/Plugins/FX/Niagara/Source/Niagara/Private/DataInterface/NiagaraDataInterfaceDataChannelWrite.cpp
+++ b/Engine/Plugins/FX/Niagara/Source/Niagara/Private/DataInterface/NiagaraDataInterfaceDataChannelWrite.cpp
@@ -799,7 +799,7 @@ void UNiagaraDataInterfaceDataChannelWrite::GetFunctionsInternal(TArray<FNiagara
 		Sig.AddInput(EmitVar);
 		Sig.AddInput(FNiagaraVariable(FNiagaraTypeDefinition::GetIntDef(), TEXT("Index")));
 		Sig.AddOutput(FNiagaraVariable(FNiagaraTypeDefinition::GetBoolDef(), TEXT("Success")));
-		Sig.RequiredInputs = Sig.Inputs.Num();//The user defines what we write in the graph.
+		Sig.RequiredInputs = IntCastChecked<int16>(Sig.Inputs.Num());//The user defines what we write in the graph.
 		OutFunctions.Add(Sig);
 	}
 	{
@@ -812,7 +812,7 @@ void UNiagaraDataInterfaceDataChannelWrite::GetFunctionsInternal(TArray<FNiagara
 		Sig.AddInput(FNiagaraVariable(FNiagaraTypeDefinition(GetClass()), TEXT("DataChannel interface")));
 		Sig.AddInput(EmitVar);
 		Sig.AddOutput(FNiagaraVariable(FNiagaraTypeDefinition::GetBoolDef(), TEXT("Success")));
-		Sig.RequiredInputs = Sig.Inputs.Num();//The user defines what we write in the graph.
+		Sig.RequiredInputs = IntCastChecked<int16>(Sig.Inputs.Num());//The user defines what we write in the graph.
 		OutFunctions.Add(Sig);
 	}
 }
diff --git a/Engine/Plugins/FX/Niagara/Source/Niagara/Private/DataInterface/NiagaraDataInterfaceDynamicMesh.cpp b/Engine/Plugins/FX/Niagara/Source/Niagara/Private/DataInterface/NiagaraDataInterfaceDynamicMesh.cpp
--- a/Engine/Plugins/FX/Niagara/Source/Niagara/Private/DataInterface/NiagaraDataInterfaceDynamicMesh.cpp
+++ b/Engine/Plugins/FX/Niagara/Source/Niagara/Private/DataInterface/NiagaraDataInterfaceDynamicMesh.cpp
@@ -669,7 +669,7 @@ namespace NDIDynamicMeshLocal
 			// Setup TexCoord Data
 			if (InstanceData->TexCoordOffset != INDEX_NONE)
 			{
-				StaticMeshData.NumTexCoords = InstanceData->NumTexCoords;
+				StaticMeshData.NumTexCoords = IntCastChecked<uint8>(InstanceData->NumTexCoords);
 				for (uint32 i=0; i < InstanceData->NumTexCoords; ++i)
 				{
 					StaticMeshData.TextureCoordinates.Emplace(&InstanceData->VertexBuffer, /*InstanceData->TexCoordOffset + (i * sizeof(FVector2f))*/0, int32(sizeof(FVector2f) * InstanceData->NumTexCoords), VET_Float2, EVertexStreamUsage::ManualFetch);
diff --git a/Engine/Plugins/FX/Niagara/Source/Niagara/Private/DataInterface/NiagaraDataInterfaceSocketReader.cpp b/Engine/Plugins/FX/Niagara/Source/Niagara/Private/DataInterface/NiagaraDataInterfaceSocketReader.cpp
--- a/Engine/Plugins/FX/Niagara/Source/Niagara/Private/DataInterface/NiagaraDataInterfaceSocketReader.cpp
+++ b/Engine/Plugins/FX/Niagara/Source/Niagara/Private/DataInterface/NiagaraDataInterfaceSocketReader.cpp
@@ -191,7 +191,7 @@ namespace NDISocketReaderLocal
 				const uint64 InitialDataSize = InstanceData.DataToUpload.Num() * InstanceData.DataToUpload.GetTypeSize();
 				const uint64 BufferSize = Align(InitialDataSize, 16);
 
-				FRDGBufferDesc BufferDesc = FRDGBufferDesc::CreateByteAddressDesc(BufferSize);
+				FRDGBufferDesc BufferDesc = FRDGBufferDesc::CreateByteAddressDesc(IntCastChecked<uint32>(BufferSize));
 				ResizeBufferIfNeeded(GraphBuilder, InstanceData.PooledBuffer, BufferDesc, TEXT("NiagaraSocketReader"));
 
 				GraphBuilder.QueueBufferUpload(
diff --git a/Engine/Plugins/FX/Niagara/Source/Niagara/Private/NiagaraComponent.cpp b/Engine/Plugins/FX/Niagara/Source/Niagara/Private/NiagaraComponent.cpp
--- a/Engine/Plugins/FX/Niagara/Source/Niagara/Private/NiagaraComponent.cpp
+++ b/Engine/Plugins/FX/Niagara/Source/Niagara/Private/NiagaraComponent.cpp
@@ -778,10 +778,10 @@ void UNiagaraComponent::ReleaseToPool()
 uint32 UNiagaraComponent::GetApproxMemoryUsage() const
 {
 	uint32 MemoryBytes = sizeof(UNiagaraComponent);
-	MemoryBytes += OverrideParameters.GetResourceSize();
+	MemoryBytes += static_cast<uint32>(OverrideParameters.GetResourceSize());
 	if (SystemInstanceController)
 	{
-		MemoryBytes += uint32(SystemInstanceController->GetTotalBytesUsed());
+		MemoryBytes += static_cast<uint32>(SystemInstanceController->GetTotalBytesUsed());
 	}
 	return MemoryBytes;
 }
diff --git a/Engine/Plugins/FX/Niagara/Source/Niagara/Private/NiagaraDataChannel.cpp b/Engine/Plugins/FX/Niagara/Source/Niagara/Private/NiagaraDataChannel.cpp
--- a/Engine/Plugins/FX/Niagara/Source/Niagara/Private/NiagaraDataChannel.cpp
+++ b/Engine/Plugins/FX/Niagara/Source/Niagara/Private/NiagaraDataChannel.cpp
@@ -237,7 +237,7 @@ void FNiagaraDataChannelGameData::WriteToDataSet(FNiagaraDataBuffer* DestBuffer,
 						for (int32 i = 0; i < NumInstances; ++i)
 						{
 							double* Src = reinterpret_cast<double*>((SrcData + i * SrcVarSize));
-							*Dest++ = *Src;
+							*Dest++ = static_cast<float>(*Src);
 						}
 					}
 					else if (SrcProperty->IsA(FFloatProperty::StaticClass()))
diff --git a/Engine/Plugins/FX/Niagara/Source/Niagara/Private/NiagaraEffectType.cpp b/Engine/Plugins/FX/Niagara/Source/Niagara/Private/NiagaraEffectType.cpp
--- a/Engine/Plugins/FX/Niagara/Source/Niagara/Private/NiagaraEffectType.cpp
+++ b/Engine/Plugins/FX/Niagara/Source/Niagara/Private/NiagaraEffectType.cpp
@@ -462,10 +462,10 @@ bool FParticlePerfStatsListener_EffectType::Tick()
 				{
 					FXType->PopulatePerfStatNames();
 
-					float TotalTime = FPlatformTime::ToMilliseconds64(Stats.GetGameThreadStats().GetTotalCycles()) * 1000.0f;
-					float GTTime = FPlatformTime::ToMilliseconds64(Stats.GetGameThreadStats().GetTotalCycles_GTOnly()) * 1000.0f;
+					float TotalTime = static_cast<float>(FPlatformTime::ToMilliseconds64(Stats.GetGameThreadStats().GetTotalCycles()) * 1000.0f);
+					float GTTime = static_cast<float>(FPlatformTime::ToMilliseconds64(Stats.GetGameThreadStats().GetTotalCycles_GTOnly()) * 1000.0f);
 					//float AvgTime = FPlatformTime::ToMilliseconds64(Stats->GetGameThreadStats().GetPerInstanceAvgCycles()) * 1000.0f;
-					int32 Count = (int32)Stats.GetGameThreadStats().NumInstances;
+					int32 Count = static_cast<int32>(Stats.GetGameThreadStats().NumInstances);
 					//float Activation = FPlatformTime::ToMilliseconds64(Stats->GetGameThreadStats().ActivationCycles) * 1000.0f;
 					//float Wait = FPlatformTime::ToMilliseconds64(Stats->GetGameThreadStats().WaitCycles) * 1000.0f;
 
@@ -518,12 +518,12 @@ void FParticlePerfStatsListener_EffectType::TickRT()
 				{
 					FXType->PopulatePerfStatNames();
 
-					const float RTTime = FPlatformTime::ToMilliseconds64(Stats.GetRenderThreadStats().GetTotalCycles()) * 1000.0f;
+					const float RTTime = static_cast<float>(FPlatformTime::ToMilliseconds64(Stats.GetRenderThreadStats().GetTotalCycles()) * 1000.0f);
 					//const float RTAvgTime = FPlatformTime::ToMilliseconds64(Stats->GetRenderThreadStats().GetPerInstanceAvgCycles()) * 1000.0f;
 					CSVProfiler->RecordCustomStat(FXType->CSVStat_RT, CSV_CATEGORY_INDEX(Particles), RTTime, ECsvCustomStatOp::Set);
 					//CSVProfiler->RecordCustomStat(System->CSVStat_InstAvgRT, CSV_CATEGORY_INDEX(Particles), RTAvgTime, ECsvCustomStatOp::Set);
 
-					const float GpuTime = float(Stats.GetGPUStats().GetTotalMicroseconds());
+					const float GpuTime = static_cast<float>(Stats.GetGPUStats().GetTotalMicroseconds());
 					//const float GpuAvgTime = float(Stats->GetGPUStats().GetPerInstanceAvgMicroseconds());
 					CSVProfiler->RecordCustomStat(FXType->CSVStat_GPU, CSV_CATEGORY_INDEX(Particles), GpuTime, ECsvCustomStatOp::Set);
 					//CSVProfiler->RecordCustomStat(System->CSVStat_InstAvgGPU, CSV_CATEGORY_INDEX(Particles), GpuAvgTime, ECsvCustomStatOp::Set);
diff --git a/Engine/Plugins/FX/Niagara/Source/Niagara/Private/NiagaraEmitterInstance.cpp b/Engine/Plugins/FX/Niagara/Source/Niagara/Private/NiagaraEmitterInstance.cpp
--- a/Engine/Plugins/FX/Niagara/Source/Niagara/Private/NiagaraEmitterInstance.cpp
+++ b/Engine/Plugins/FX/Niagara/Source/Niagara/Private/NiagaraEmitterInstance.cpp
@@ -93,12 +93,12 @@ bool FNiagaraEmitterInstance::IsDisabledFromIsolation() const
 
 int64 FNiagaraEmitterInstance::GetTotalBytesUsed() const
 {
-	int32 ByteSize = 0;
+	int64 ByteSize = 0;
 	if (ParticleDataSet)
 	{
 		ByteSize += ParticleDataSet->GetSizeBytes();
 	}
-	return 0;
+	return ByteSize;
 }
 
 const FNiagaraEmitterHandle& FNiagaraEmitterInstance::GetEmitterHandle() const
diff --git a/Engine/Plugins/FX/Niagara/Source/Niagara/Private/NiagaraGPUSystemTick.cpp b/Engine/Plugins/FX/Niagara/Source/Niagara/Private/NiagaraGPUSystemTick.cpp
--- a/Engine/Plugins/FX/Niagara/Source/Niagara/Private/NiagaraGPUSystemTick.cpp
+++ b/Engine/Plugins/FX/Niagara/Source/Niagara/Private/NiagaraGPUSystemTick.cpp
@@ -213,7 +213,7 @@ void FNiagaraGPUSystemTick::Init(FNiagaraSystemInstance* InSystemInstance)
 			for (const FNiagaraSimStageExecutionLoopData& LoopData : GPUContext->SimStageExecData->ExecutionLoops)
 			{
 				int32 NumLoops = LoopData.NumLoopsBinding.IsNone() ? LoopData.NumLoops : ParameterStore.GetParameterValueOrDefault(FNiagaraVariable(FNiagaraTypeDefinition::GetIntDef(), LoopData.NumLoopsBinding), LoopData.NumLoops);
-				NumLoops = FMath::Clamp(NumLoops, 0, TNumericLimits<uint16>::Max());
+				NumLoops = static_cast<int32>(FMath::Clamp(NumLoops, 0, TNumericLimits<uint16>::Max()));
 				for (int32 LoopIndex=0; LoopIndex < NumLoops; ++LoopIndex)
 				{
 					for (int32 SimStageIndex=LoopData.StartStageIndex; SimStageIndex <= LoopData.EndStageIndex; ++SimStageIndex)
@@ -241,7 +241,7 @@ void FNiagaraGPUSystemTick::Init(FNiagaraSystemInstance* InSystemInstance)
 
 						// Get number of iterations
 						int32 NumIterations = SimStageMetaData.NumIterationsBinding.IsNone() ? SimStageMetaData.NumIterations : ParameterStore.GetParameterValueOrDefault(FNiagaraVariable(FNiagaraTypeDefinition::GetIntDef(), SimStageMetaData.NumIterationsBinding), SimStageMetaData.NumIterations);
-						NumIterations = FMath::Clamp(NumIterations, 0, TNumericLimits<uint16>::Max());
+						NumIterations = static_cast<int32>(FMath::Clamp(NumIterations, 0, TNumericLimits<uint16>::Max()));
 						if ( NumIterations <= 0 )
 						{
 							continue;
@@ -277,10 +277,10 @@ void FNiagaraGPUSystemTick::Init(FNiagaraSystemInstance* InSystemInstance)
 
 						// Stage is live we can add it
 						FNiagaraComputeInstanceData::FPerStageInfo& NewStageInfo = InstanceData->PerStageInfo.AddDefaulted_GetRef();
-						NewStageInfo.SimStageIndex		= SimStageIndex;
-						NewStageInfo.NumIterations		= NumIterations;
-						NewStageInfo.LoopIndex			= LoopIndex;
-						NewStageInfo.NumLoops			= NumLoops;
+						NewStageInfo.SimStageIndex		= IntCastChecked<int16>(SimStageIndex);
+						NewStageInfo.NumIterations		= IntCastChecked<int16>(NumIterations);
+						NewStageInfo.LoopIndex			= IntCastChecked<int16>(LoopIndex);
+						NewStageInfo.NumLoops			= IntCastChecked<int16>(NumLoops);
 						NewStageInfo.ElementCountXYZ	= ElementCountXYZ;
 
 						InstanceData->bHasMultipleStages = true;
diff --git a/Engine/Plugins/FX/Niagara/Source/Niagara/Private/NiagaraGpuComputeDispatch.cpp b/Engine/Plugins/FX/Niagara/Source/Niagara/Private/NiagaraGpuComputeDispatch.cpp
--- a/Engine/Plugins/FX/Niagara/Source/Niagara/Private/NiagaraGpuComputeDispatch.cpp
+++ b/Engine/Plugins/FX/Niagara/Source/Niagara/Private/NiagaraGpuComputeDispatch.cpp
@@ -845,9 +845,9 @@ void FNiagaraGpuComputeDispatch::PrepareTicksForProxy(FRHICommandListImmediate&
 					FNiagaraGpuDispatchInstance& DispatchInstance = DispatchGroup.DispatchInstances.Emplace_GetRef(Tick, InstanceData);
 					FNiagaraSimStageData& SimStageData = DispatchInstance.SimStageData;
 					SimStageData.bFirstStage = bFirstStage;
-					SimStageData.StageIndex = SimStageIndex;
-					SimStageData.NumIterations = NumIterations;
-					SimStageData.IterationIndex = IterationIndex;
+					SimStageData.StageIndex = IntCastChecked<uint16>(SimStageIndex);
+					SimStageData.NumIterations = IntCastChecked<uint16>(NumIterations);
+					SimStageData.IterationIndex = IntCastChecked<uint16>(IterationIndex);
 					SimStageData.NumLoops = PerStageInfo.NumLoops;
 					SimStageData.LoopIndex = PerStageInfo.LoopIndex;
 					SimStageData.DispatchArgs.ElementCount = PerStageInfo.ElementCountXYZ;
diff --git a/Engine/Plugins/FX/Niagara/Source/Niagara/Private/NiagaraMeshRendererProperties.cpp b/Engine/Plugins/FX/Niagara/Source/Niagara/Private/NiagaraMeshRendererProperties.cpp
--- a/Engine/Plugins/FX/Niagara/Source/Niagara/Private/NiagaraMeshRendererProperties.cpp
+++ b/Engine/Plugins/FX/Niagara/Source/Niagara/Private/NiagaraMeshRendererProperties.cpp
@@ -123,7 +123,7 @@ public:
 		return FVector3f(
 			LODLevel < MaxLODLevel ? RenderData->ScreenSize[LODLevel + 1].GetValue() : 0.0f,
 			RenderData->ScreenSize[LODLevel].GetValue(),
-			RenderData->Bounds.SphereRadius
+			static_cast<float>(RenderData->Bounds.SphereRadius)
 		);
 	}
 
diff --git a/Engine/Plugins/FX/Niagara/Source/Niagara/Private/NiagaraRendererProperties.cpp b/Engine/Plugins/FX/Niagara/Source/Niagara/Private/NiagaraRendererProperties.cpp
--- a/Engine/Plugins/FX/Niagara/Source/Niagara/Private/NiagaraRendererProperties.cpp
+++ b/Engine/Plugins/FX/Niagara/Source/Niagara/Private/NiagaraRendererProperties.cpp
@@ -157,8 +157,8 @@ bool FNiagaraRendererLayout::SetVariable(const FNiagaraDataSetCompiledData* Comp
 		{
 			//For CPU Sims we pack just the required data tightly in a GPU buffer we upload. For GPU sims the data is there already so we just provide the real data location.
 			GPULocation = CompiledData->SimTarget == ENiagaraSimTarget::CPUSim ? TotalVFComponents : Offset;
-			check(int32(TotalVFComponents ) + NumComponents <= TNumericLimits<uint16>::Max());
-			TotalVFComponents += NumComponents;
+			check(static_cast<int32>(TotalVFComponents) + NumComponents <= TNumericLimits<uint16>::Max());
+			TotalVFComponents += static_cast<uint16>(NumComponents);
 		}
 	}
 
diff --git a/Engine/Plugins/FX/Niagara/Source/Niagara/Private/NiagaraScriptExecutionContext.cpp b/Engine/Plugins/FX/Niagara/Source/Niagara/Private/NiagaraScriptExecutionContext.cpp
--- a/Engine/Plugins/FX/Niagara/Source/Niagara/Private/NiagaraScriptExecutionContext.cpp
+++ b/Engine/Plugins/FX/Niagara/Source/Niagara/Private/NiagaraScriptExecutionContext.cpp
@@ -198,7 +198,7 @@ TMap<TStatIdData const*, float> FNiagaraScriptExecutionContextBase::ReportStats(
 		uint64 ExecCycles = ScopeData.ExecutionCycleCount.exchange(0);
 		if (ExecCycles > 0)
 		{
-			ExecutionTimings.FindOrAdd(ScopeData.StatId.GetRawPointer()) = ExecCycles;
+			ExecutionTimings.FindOrAdd(ScopeData.StatId.GetRawPointer()) = static_cast<float>(ExecCycles);
 		}
 	}
 	return ExecutionTimings;
diff --git a/Engine/Plugins/FX/Niagara/Source/Niagara/Private/NiagaraSimCache.cpp b/Engine/Plugins/FX/Niagara/Source/Niagara/Private/NiagaraSimCache.cpp
--- a/Engine/Plugins/FX/Niagara/Source/Niagara/Private/NiagaraSimCache.cpp
+++ b/Engine/Plugins/FX/Niagara/Source/Niagara/Private/NiagaraSimCache.cpp
@@ -130,8 +130,8 @@ struct FNiagaraSimCacheBuffersSetup
 			DataBuffer.FloatData		= TArrayView<uint8>(BaseData + FloatDataOffset, FloatDataNum * sizeof(float));
 			DataBuffer.HalfData			= TArrayView<uint8>(BaseData + HalfDataOffset, HalfDataNum * sizeof(FFloat16));
 			DataBuffer.Int32Data		= TArrayView<uint8>(BaseData + Int32DataOffset, Int32DataNum * sizeof(int32));
-			DataBuffer.IDToIndexTable	= TArrayView<int32>(reinterpret_cast<int32*>(BaseData + IDToIndexTableOffset), IDToIndexTableNum);
-			DataBuffer.InterpMapping	= TArrayView<uint32>(reinterpret_cast<uint32*>(BaseData + InterpMappingOffset), InterpMappingNum);
+			DataBuffer.IDToIndexTable	= TArrayView<int32>(reinterpret_cast<int32*>(BaseData + IDToIndexTableOffset), static_cast<int32>(IDToIndexTableNum));
+			DataBuffer.InterpMapping	= TArrayView<uint32>(reinterpret_cast<uint32*>(BaseData + InterpMappingOffset), static_cast<int32>(InterpMappingNum));
 			return true;
 		}
 	}
diff --git a/Engine/Plugins/FX/Niagara/Source/Niagara/Private/NiagaraSimCacheHelper.h b/Engine/Plugins/FX/Niagara/Source/Niagara/Private/NiagaraSimCacheHelper.h
--- a/Engine/Plugins/FX/Niagara/Source/Niagara/Private/NiagaraSimCacheHelper.h
+++ b/Engine/Plugins/FX/Niagara/Source/Niagara/Private/NiagaraSimCacheHelper.h
@@ -180,34 +180,34 @@ struct FNiagaraSimCacheHelper
 
 			CacheVariable.Variable = CompiledData.Variables[iDataSetVariable];
 			CacheVariable.FloatOffset = DataSetVariableLayout.GetNumFloatComponents() > 0 ? CacheLayout.FloatCount : INDEX_NONE;
-			CacheVariable.FloatCount = uint16(DataSetVariableLayout.GetNumFloatComponents());
+			CacheVariable.FloatCount = IntCastChecked<uint16>(DataSetVariableLayout.GetNumFloatComponents());
 			CacheVariable.HalfOffset = DataSetVariableLayout.GetNumHalfComponents() > 0 ? CacheLayout.HalfCount : INDEX_NONE;
-			CacheVariable.HalfCount = uint16(DataSetVariableLayout.GetNumHalfComponents());
+			CacheVariable.HalfCount = IntCastChecked<uint16>(DataSetVariableLayout.GetNumHalfComponents());
 			CacheVariable.Int32Offset = DataSetVariableLayout.GetNumInt32Components() > 0 ? CacheLayout.Int32Count : INDEX_NONE;
-			CacheVariable.Int32Count = uint16(DataSetVariableLayout.GetNumInt32Components());
+			CacheVariable.Int32Count = IntCastChecked<uint16>(DataSetVariableLayout.GetNumInt32Components());
 
-			CacheLayout.FloatCount += uint16(DataSetVariableLayout.GetNumFloatComponents());
-			CacheLayout.HalfCount += uint16(DataSetVariableLayout.GetNumHalfComponents());
-			CacheLayout.Int32Count += uint16(DataSetVariableLayout.GetNumInt32Components());
+			CacheLayout.FloatCount += IntCastChecked<uint16>(DataSetVariableLayout.GetNumFloatComponents());
+			CacheLayout.HalfCount += IntCastChecked<uint16>(DataSetVariableLayout.GetNumHalfComponents());
+			CacheLayout.Int32Count += IntCastChecked<uint16>(DataSetVariableLayout.GetNumInt32Components());
 		}
 
 		if (CreateParameters.bAllowVelocityExtrapolation)
 		{
 			if (const FNiagaraSimCacheVariable* VelocityCacheVariable = CacheLayout.FindCacheVariable(VelocityVariable) )
 			{
-				CacheLayout.ComponentVelocity = uint16(VelocityCacheVariable->FloatOffset);
+				CacheLayout.ComponentVelocity = IntCastChecked<uint16>(VelocityCacheVariable->FloatOffset);
 			}
 		}
 
 		if (CreateParameters.bAllowInterpolation)
 		{
 			const FNiagaraVariableBase UniqueIDVariable(FNiagaraTypeDefinition::GetIntDef(), "UniqueID");
-			CacheLayout.CacheBufferWriteInfo.ComponentUniqueID = CompiledData.Variables.IndexOfByKey(UniqueIDVariable);
-			if (CacheLayout.CacheBufferWriteInfo.ComponentUniqueID != uint16(INDEX_NONE))
+			CacheLayout.CacheBufferWriteInfo.ComponentUniqueID = IntCastChecked<uint16>(CompiledData.Variables.IndexOfByKey(UniqueIDVariable));
+			if (CacheLayout.CacheBufferWriteInfo.ComponentUniqueID != static_cast<uint16>(INDEX_NONE))
 			{
 				const FNiagaraVariableLayoutInfo& DataSetVariableLayout = CompiledData.VariableLayouts[CacheLayout.CacheBufferWriteInfo.ComponentUniqueID];
 				check(DataSetVariableLayout.GetNumInt32Components() == 1);
-				CacheLayout.CacheBufferWriteInfo.ComponentUniqueID = DataSetVariableLayout.GetInt32ComponentStart();
+				CacheLayout.CacheBufferWriteInfo.ComponentUniqueID = IntCastChecked<uint16>(DataSetVariableLayout.GetInt32ComponentStart());
 				CacheLayout.bAllowInterpolation = true;
 			}
 		}
@@ -228,19 +228,19 @@ struct FNiagaraSimCacheHelper
 
 			for (int32 iComponent=0; iComponent < CacheVariable.FloatCount; ++iComponent)
 			{
-				CacheLayout.CacheBufferWriteInfo.ComponentMappingsFromDataBuffer[FloatOffset] = uint16(DataSetVariableLayout.GetFloatComponentStart() + iComponent);
+				CacheLayout.CacheBufferWriteInfo.ComponentMappingsFromDataBuffer[FloatOffset] = IntCastChecked<uint16>(DataSetVariableLayout.GetFloatComponentStart() + iComponent);
 				++FloatOffset;
 			}
 
 			for (int32 iComponent=0; iComponent < CacheVariable.HalfCount; ++iComponent)
 			{
-				CacheLayout.CacheBufferWriteInfo.ComponentMappingsFromDataBuffer[HalfOffset] = uint16(DataSetVariableLayout.GetHalfComponentStart() + iComponent);
+				CacheLayout.CacheBufferWriteInfo.ComponentMappingsFromDataBuffer[HalfOffset] = IntCastChecked<uint16>(DataSetVariableLayout.GetHalfComponentStart() + iComponent);
 				++HalfOffset;
 			}
 
 			for (int32 iComponent=0; iComponent < CacheVariable.Int32Count; ++iComponent)
 			{
-				CacheLayout.CacheBufferWriteInfo.ComponentMappingsFromDataBuffer[Int32Offset] = uint16(DataSetVariableLayout.GetInt32ComponentStart() + iComponent);
+				CacheLayout.CacheBufferWriteInfo.ComponentMappingsFromDataBuffer[Int32Offset] = IntCastChecked<uint16>(DataSetVariableLayout.GetInt32ComponentStart() + iComponent);
 				++Int32Offset;
 			}
 		}
@@ -551,23 +551,23 @@ struct FNiagaraSimCacheHelper
 
 								CacheBufferReadInfo.VariableCopyMappingsToDataBuffer.Emplace(
 									SourceVariable.FloatOffset, PrevSourceVariable->FloatOffset,
-									uint16(DestVariableLayout->GetFloatComponentStart()), uint16(CompiledData.VariableLayouts[PreviousDataSetVariableIndex].GetFloatComponentStart()),
+									IntCastChecked<uint16>(DestVariableLayout->GetFloatComponentStart()), IntCastChecked<uint16>(CompiledData.VariableLayouts[PreviousDataSetVariableIndex].GetFloatComponentStart()),
 									bRebaseVariable ? &FNiagaraSimCacheHelper::InterpPositions<true, true> : &FNiagaraSimCacheHelper::InterpPositions<false, true>
 									
 								);
 							}
 							else
 							{
-								CacheBufferReadInfo.VariableCopyMappingsToDataBuffer.Emplace(SourceVariable.FloatOffset, uint16(DestVariableLayout->GetFloatComponentStart()), bRebaseVariable ? &FNiagaraSimCacheHelper::InterpPositions<true, false> : &FNiagaraSimCacheHelper::InterpPositions<false, false>);
+								CacheBufferReadInfo.VariableCopyMappingsToDataBuffer.Emplace(SourceVariable.FloatOffset, IntCastChecked<uint16>(DestVariableLayout->GetFloatComponentStart()), bRebaseVariable ? &FNiagaraSimCacheHelper::InterpPositions<true, false> : &FNiagaraSimCacheHelper::InterpPositions<false, false>);
 							}
 						}
 						else if ( CacheLayout.bAllowVelocityExtrapolation )
 						{
-							CacheBufferReadInfo.VariableCopyMappingsToDataBuffer.Emplace(SourceVariable.FloatOffset, uint16(DestVariableLayout->GetFloatComponentStart()), bRebaseVariable ? &FNiagaraSimCacheHelper::ExtrapolatePositions<true> : &FNiagaraSimCacheHelper::ExtrapolatePositions<false>);
+							CacheBufferReadInfo.VariableCopyMappingsToDataBuffer.Emplace(SourceVariable.FloatOffset, IntCastChecked<uint16>(DestVariableLayout->GetFloatComponentStart()), bRebaseVariable ? &FNiagaraSimCacheHelper::ExtrapolatePositions<true> : &FNiagaraSimCacheHelper::ExtrapolatePositions<false>);
 						}
 						else
 						{
-							CacheBufferReadInfo.VariableCopyMappingsToDataBuffer.Emplace(SourceVariable.FloatOffset, uint16(DestVariableLayout->GetFloatComponentStart()), &FNiagaraSimCacheHelper::RebasePositions);
+							CacheBufferReadInfo.VariableCopyMappingsToDataBuffer.Emplace(SourceVariable.FloatOffset, IntCastChecked<uint16>(DestVariableLayout->GetFloatComponentStart()), &FNiagaraSimCacheHelper::RebasePositions);
 						}
 						DestVariableLayout = nullptr;
 					}
@@ -587,19 +587,19 @@ struct FNiagaraSimCacheHelper
 
 								CacheBufferReadInfo.VariableCopyMappingsToDataBuffer.Emplace(
 									SourceVariable.FloatOffset, PrevSourceVariable->FloatOffset,
-									uint16(DestVariableLayout->GetFloatComponentStart()), uint16(CompiledData.VariableLayouts[PreviousDataSetVariableIndex].GetFloatComponentStart()),
+									IntCastChecked<uint16>(DestVariableLayout->GetFloatComponentStart()), IntCastChecked<uint16>(CompiledData.VariableLayouts[PreviousDataSetVariableIndex].GetFloatComponentStart()),
 									bRebaseVariable ? &FNiagaraSimCacheHelper::InterpQuaternions<true, true> : &FNiagaraSimCacheHelper::InterpQuaternions<false, true>
 									
 								);
 							}
 							else
 							{
-								CacheBufferReadInfo.VariableCopyMappingsToDataBuffer.Emplace(SourceVariable.FloatOffset, uint16(DestVariableLayout->GetFloatComponentStart()), bInterpVariable  ? &FNiagaraSimCacheHelper::InterpQuaternions<true, false> : &FNiagaraSimCacheHelper::InterpQuaternions<false, false>);
+								CacheBufferReadInfo.VariableCopyMappingsToDataBuffer.Emplace(SourceVariable.FloatOffset, IntCastChecked<uint16>(DestVariableLayout->GetFloatComponentStart()), bInterpVariable  ? &FNiagaraSimCacheHelper::InterpQuaternions<true, false> : &FNiagaraSimCacheHelper::InterpQuaternions<false, false>);
 							}
 						}
 						else
 						{
-							CacheBufferReadInfo.VariableCopyMappingsToDataBuffer.Emplace(SourceVariable.FloatOffset, uint16(DestVariableLayout->GetFloatComponentStart()), &FNiagaraSimCacheHelper::RebaseQuaternions);
+							CacheBufferReadInfo.VariableCopyMappingsToDataBuffer.Emplace(SourceVariable.FloatOffset, IntCastChecked<uint16>(DestVariableLayout->GetFloatComponentStart()), &FNiagaraSimCacheHelper::RebaseQuaternions);
 						}
 						DestVariableLayout = nullptr;
 					}
@@ -608,7 +608,7 @@ struct FNiagaraSimCacheHelper
 						if (bRebaseVariable)
 						{
 							check(SourceVariable.FloatCount == 16);
-							CacheBufferReadInfo.VariableCopyMappingsToDataBuffer.Emplace(SourceVariable.FloatOffset, uint16(DestVariableLayout->GetFloatComponentStart()), &FNiagaraSimCacheHelper::RebaseMatrices);
+							CacheBufferReadInfo.VariableCopyMappingsToDataBuffer.Emplace(SourceVariable.FloatOffset, IntCastChecked<uint16>(DestVariableLayout->GetFloatComponentStart()), &FNiagaraSimCacheHelper::RebaseMatrices);
 							DestVariableLayout = nullptr;
 						}
 					}
@@ -617,17 +617,17 @@ struct FNiagaraSimCacheHelper
 
 			for (int32 i = 0; i < SourceVariable.FloatCount; ++i)
 			{
-				CacheBufferReadInfo.ComponentMappingsToDataBuffer[FloatOffset++] = uint16(DestVariableLayout ? DestVariableLayout->GetFloatComponentStart() + i : INDEX_NONE);
+				CacheBufferReadInfo.ComponentMappingsToDataBuffer[FloatOffset++] = IntCastChecked<uint16>(DestVariableLayout ? DestVariableLayout->GetFloatComponentStart() + i : INDEX_NONE);
 			}
 
 			for (int32 i = 0; i < SourceVariable.HalfCount; ++i)
 			{
-				CacheBufferReadInfo.ComponentMappingsToDataBuffer[HalfOffset++] = uint16(DestVariableLayout ? DestVariableLayout->GetHalfComponentStart() + i : INDEX_NONE);
+				CacheBufferReadInfo.ComponentMappingsToDataBuffer[HalfOffset++] = IntCastChecked<uint16>(DestVariableLayout ? DestVariableLayout->GetHalfComponentStart() + i : INDEX_NONE);
 			}
 
 			for (int32 i = 0; i < SourceVariable.Int32Count; ++i)
 			{
-				CacheBufferReadInfo.ComponentMappingsToDataBuffer[Int32Offset++] = uint16(DestVariableLayout ? DestVariableLayout->GetInt32ComponentStart() + i : INDEX_NONE);
+				CacheBufferReadInfo.ComponentMappingsToDataBuffer[Int32Offset++] = IntCastChecked<uint16>(DestVariableLayout ? DestVariableLayout->GetInt32ComponentStart() + i : INDEX_NONE);
 			}
 		}
 
@@ -775,7 +775,7 @@ struct FNiagaraSimCacheHelper
 		VariableCopyDataContext.DestStride			= DestStride;
 		VariableCopyDataContext.SourceAStride		= CacheBufferA.NumInstances * sizeof(float);
 		VariableCopyDataContext.SourceBStride		= CacheBufferB.NumInstances * sizeof(float);
-		if (ComponentVelocity != uint16(INDEX_NONE))
+		if (ComponentVelocity != static_cast<uint16>(INDEX_NONE))
 		{
 			VariableCopyDataContext.Velocity = CacheBufferA.FloatData.GetData() + (ComponentVelocity * VariableCopyDataContext.SourceAStride);
 		}
@@ -786,8 +786,8 @@ struct FNiagaraSimCacheHelper
 
 		for (const FNiagaraSimCacheDataBuffersLayout::FVariableCopyMapping& VariableCopyMapping : VariableCopyMappingsToDataBuffer)
 		{
-			VariableCopyDataContext.DestCurr	= DestBuffer + (uint32(VariableCopyMapping.CurrComponentTo) * DestStride);
-			VariableCopyDataContext.DestPrev	= DestBuffer + (uint32(VariableCopyMapping.PrevComponentTo) * DestStride);
+			VariableCopyDataContext.DestCurr	= DestBuffer + (IntCastChecked<uint32>(VariableCopyMapping.CurrComponentTo) * DestStride);
+			VariableCopyDataContext.DestPrev	= DestBuffer + (IntCastChecked<uint32>(VariableCopyMapping.PrevComponentTo) * DestStride);
 			VariableCopyDataContext.SourceACurr	= CacheBufferA.FloatData.GetData() + (VariableCopyMapping.CurrComponentFrom * VariableCopyDataContext.SourceAStride);
 			VariableCopyDataContext.SourceAPrev = CacheBufferA.FloatData.GetData() + (VariableCopyMapping.PrevComponentFrom * VariableCopyDataContext.SourceAStride);
 			VariableCopyDataContext.SourceBCurr	= CacheBufferB.FloatData.GetData() + (VariableCopyMapping.CurrComponentFrom * VariableCopyDataContext.SourceBStride);
@@ -1076,14 +1076,14 @@ struct FNiagaraSimCacheHelper
 			FMatrix44d CacheMatrix;
 			for (int32 j = 0; j < 16; ++j)
 			{
-				CacheMatrix.M[j >> 2][j & 0x3] = double(SrcFloats[i + (SrcStride * j)]);
+				CacheMatrix.M[j >> 2][j & 0x3] = static_cast<double>(SrcFloats[i + (SrcStride * j)]);
 			}
 
 			CacheMatrix = CacheMatrix * RebaseMatrix;
 
 			for (int32 j = 0; j < 16; ++j)
 			{
-				DstFloats[i + (DstStride * j)] = float(CacheMatrix.M[j >> 2][j & 0x3]);
+				DstFloats[i + (DstStride * j)] = static_cast<float>(CacheMatrix.M[j >> 2][j & 0x3]);
 			}
 		}
 	}
diff --git a/Engine/Plugins/FX/Niagara/Source/Niagara/Private/NiagaraSimulationTaskPriority.cpp b/Engine/Plugins/FX/Niagara/Source/Niagara/Private/NiagaraSimulationTaskPriority.cpp
--- a/Engine/Plugins/FX/Niagara/Source/Niagara/Private/NiagaraSimulationTaskPriority.cpp
+++ b/Engine/Plugins/FX/Niagara/Source/Niagara/Private/NiagaraSimulationTaskPriority.cpp
@@ -61,7 +61,7 @@ namespace NiagaraSimulationTaskPriority
 			if (TickGroup >= 0 && TickGroup < TG_MAX)
 			{
 				int32 Priority = FCString::Atoi(*TGPair.Mid(SepIndex + 1));
-				Priority = FMath::Clamp(Priority, 0, UE_ARRAY_COUNT(GTaskPriorities) - 1);
+				Priority = static_cast<int32>(FMath::Clamp(Priority, 0, UE_ARRAY_COUNT(GTaskPriorities) - 1));
 				GTickGroupPriorities[TickGroup] = Priority;
 			}
 		}
diff --git a/Engine/Plugins/FX/Niagara/Source/Niagara/Private/NiagaraSystemInstance.cpp b/Engine/Plugins/FX/Niagara/Source/Niagara/Private/NiagaraSystemInstance.cpp
--- a/Engine/Plugins/FX/Niagara/Source/Niagara/Private/NiagaraSystemInstance.cpp
+++ b/Engine/Plugins/FX/Niagara/Source/Niagara/Private/NiagaraSystemInstance.cpp
@@ -768,7 +768,7 @@ void FNiagaraSystemInstance::Reset(EResetMode Mode)
 	// Wait for any async operations, can complete the system
 	WaitForConcurrentTickAndFinalize();
 
-	LastRenderTime = World->GetTimeSeconds();
+	LastRenderTime = static_cast<float>(World->GetTimeSeconds());
 
 	SetPaused(false);
 
@@ -1827,7 +1827,8 @@ float FNiagaraSystemInstance::GetLODDistance()
 	const FVector EffectLocation = WorldTransform.GetLocation() + (FVector(LWCTile) * FLargeWorldRenderScalar::GetTileSize());
 	LODDistance = DefaultLODDistance;
 
-	LODDistance = WorldManager->GetLODDistance(EffectLocation);
+	// truncation to float here with precision loss seems ok for LOD distance
+	LODDistance = static_cast<float>(WorldManager->GetLODDistance(EffectLocation));
 
 	bLODDistanceIsValid = true;
 	return LODDistance;
diff --git a/Engine/Plugins/FX/Niagara/Source/Niagara/Private/NiagaraTypes.cpp b/Engine/Plugins/FX/Niagara/Source/Niagara/Private/NiagaraTypes.cpp
--- a/Engine/Plugins/FX/Niagara/Source/Niagara/Private/NiagaraTypes.cpp
+++ b/Engine/Plugins/FX/Niagara/Source/Niagara/Private/NiagaraTypes.cpp
@@ -860,9 +860,9 @@ void FNiagaraTypeLayoutInfo::GenerateLayoutInfo(const UScriptStruct* Struct)
 	int32 Int32Count = 0;
 	int32 HalfCount = 0;
 	GenerateLayoutInfoInternal(Struct, FloatCount, Int32Count, HalfCount, true);
-	NumFloatComponents = FloatCount;
-	NumInt32Components = Int32Count;
-	NumHalfComponents = HalfCount;
+	NumFloatComponents = IntCastChecked<uint16>(FloatCount);
+	NumInt32Components = IntCastChecked<uint16>(Int32Count);
+	NumHalfComponents = IntCastChecked<uint16>(HalfCount);
 
 	FloatCount = 0;
 	Int32Count = 0;
diff --git a/Engine/Plugins/FX/Niagara/Source/Niagara/Private/NiagaraVariableMetaData.cpp b/Engine/Plugins/FX/Niagara/Source/Niagara/Private/NiagaraVariableMetaData.cpp
--- a/Engine/Plugins/FX/Niagara/Source/Niagara/Private/NiagaraVariableMetaData.cpp
+++ b/Engine/Plugins/FX/Niagara/Source/Niagara/Private/NiagaraVariableMetaData.cpp
@@ -19,7 +19,7 @@ FNiagaraInputParameterCustomization FNiagaraInputParameterCustomization::MakeFro
 		{
 			FString TypeString = PropertyHandle->GetMetaData(WidgetTypeName);
 			UEnum* TypeEnum = StaticEnum<ENiagaraInputWidgetType>();
-			int32 Value = TypeEnum->GetValueByNameString(TypeString);
+			int32 Value = static_cast<int32>(TypeEnum->GetValueByNameString(TypeString));
 			if (Value != INDEX_NONE)
 			{
 				ParameterCustomization.WidgetType = static_cast<ENiagaraInputWidgetType>(Value);
diff --git a/Engine/Plugins/FX/Niagara/Source/Niagara/Private/NiagaraWorldManager.cpp b/Engine/Plugins/FX/Niagara/Source/Niagara/Private/NiagaraWorldManager.cpp
--- a/Engine/Plugins/FX/Niagara/Source/Niagara/Private/NiagaraWorldManager.cpp
+++ b/Engine/Plugins/FX/Niagara/Source/Niagara/Private/NiagaraWorldManager.cpp
@@ -1530,7 +1530,7 @@ FVector::FReal FNiagaraWorldManager::GetLODDistance(FVector Location)const
 			const FVector::FReal DistanceToEffectSqr = FVector(ViewInfo.ViewToWorld.GetOrigin() - Location).SizeSquared();
 			LODDistanceSqr = FMath::Min(LODDistanceSqr, DistanceToEffectSqr);
 		}
-		LODDistance = FVector::FReal(FMath::Sqrt(LODDistanceSqr));
+		LODDistance = FMath::Sqrt(LODDistanceSqr);
 	}
 	else
 	{
@@ -1871,11 +1871,11 @@ void FNiagaraWorldManager::ViewBasedCulling(UNiagaraEffectType* EffectType, cons
 	float TimeSinceInsideView = 0.0f;
 	if (bInsideAnyView)
 	{
-		OutState.LastVisibleTime = World->GetTimeSeconds();
+		OutState.LastVisibleTime = static_cast<float>(World->GetTimeSeconds());
 	}
 	else
 	{
-		TimeSinceInsideView = World->GetTimeSeconds() - OutState.LastVisibleTime;
+		TimeSinceInsideView = static_cast<float>(World->GetTimeSeconds() - OutState.LastVisibleTime);
 	}
 
 	bool bCullByOutsideViewFrustum = ScalabilitySettings.VisibilityCulling.bCullByViewFrustum &&
@@ -1883,7 +1883,6 @@ void FNiagaraWorldManager::ViewBasedCulling(UNiagaraEffectType* EffectType, cons
 		TimeSinceInsideView > ScalabilitySettings.VisibilityCulling.MaxTimeOutsideViewFrustum;
 
 	//Check for the component having been rendered recently. If the app doesn't have focus we skip this to avoid issues when alt-tabbing away from the game/editor.
-	float TimeSinceWorldRendered = World->GetTimeSeconds() - World->LastRenderTime;	
 	bool bCullByNotRendered =	bAppHasFocus && 
 								ScalabilitySettings.VisibilityCulling.bCullWhenNotRendered && 
 		ComponentTimeSinceRendered > ScalabilitySettings.VisibilityCulling.MaxTimeWithoutRender;
diff --git a/Engine/Plugins/FX/Niagara/Source/Niagara/Private/Stateless/Modules/NiagaraStatelessModule_MeshIndex.cpp b/Engine/Plugins/FX/Niagara/Source/Niagara/Private/Stateless/Modules/NiagaraStatelessModule_MeshIndex.cpp
--- a/Engine/Plugins/FX/Niagara/Source/Niagara/Private/Stateless/Modules/NiagaraStatelessModule_MeshIndex.cpp
+++ b/Engine/Plugins/FX/Niagara/Source/Niagara/Private/Stateless/Modules/NiagaraStatelessModule_MeshIndex.cpp
@@ -64,7 +64,7 @@ namespace NSMMeshIndexPrivate
 			for (uint32 i = 0; i < ParticleSimulationContext.GetNumInstances(); ++i)
 			{
 				const FVector2f Rand = ParticleSimulationContext.RandomFloat2(i, 0);
-				MeshIndex = FMath::RoundToInt(Rand.X * float(ModuleBuiltData->TableNumElements));
+				MeshIndex = FMath::RoundToInt(Rand.X * static_cast<float>(ModuleBuiltData->TableNumElements));
 				const float Probability = ParticleSimulationContext.GetStaticFloat<float>(ModuleBuiltData->TableOffset, (MeshIndex * 2) + 0);
 				if (Rand.Y > Probability)
 				{
diff --git a/Engine/Plugins/FX/Niagara/Source/Niagara/Private/Stateless/Modules/NiagaraStatelessModule_ShapeLocation.cpp b/Engine/Plugins/FX/Niagara/Source/Niagara/Private/Stateless/Modules/NiagaraStatelessModule_ShapeLocation.cpp
--- a/Engine/Plugins/FX/Niagara/Source/Niagara/Private/Stateless/Modules/NiagaraStatelessModule_ShapeLocation.cpp
+++ b/Engine/Plugins/FX/Niagara/Source/Niagara/Private/Stateless/Modules/NiagaraStatelessModule_ShapeLocation.cpp
@@ -74,7 +74,7 @@ namespace NSMShapeLocationPrivate
 			const float Radius = ParticleSimulationContext.RandomScaleBiasFloat(iInstance, 0, RadiusScale, RadiusBias);
 			const float U = ParticleSimulationContext.RandomScaleBiasFloat(iInstance, 1, UDistributionScale, UDistributionBias);
 	
-			return FVector3f(cos(U) * Radius, sin(U) * Radius, 0.0f);
+			return FVector3f(cosf(U) * Radius, sinf(U) * Radius, 0.0f);
 		}
 	
 		// ENSM_ShapePrimitive::Sphere:
diff --git a/Engine/Plugins/FX/Niagara/Source/Niagara/Private/Stateless/NiagaraStatelessDrawDebugContext.cpp b/Engine/Plugins/FX/Niagara/Source/Niagara/Private/Stateless/NiagaraStatelessDrawDebugContext.cpp
--- a/Engine/Plugins/FX/Niagara/Source/Niagara/Private/Stateless/NiagaraStatelessDrawDebugContext.cpp
+++ b/Engine/Plugins/FX/Niagara/Source/Niagara/Private/Stateless/NiagaraStatelessDrawDebugContext.cpp
@@ -9,7 +9,7 @@
 #if WITH_EDITOR
 void FNiagaraStatelessDrawDebugContext::DrawArrow(const FVector& Origin, const FVector& DirectionWithLength, const FColor& Color) const
 {
-	const float Len = DirectionWithLength.Length();
+	const float Len = static_cast<float>(DirectionWithLength.Length());
 	if (Len > UE_SMALL_NUMBER)
 	{
 		DrawDebugDirectionalArrow(
diff --git a/Engine/Plugins/FX/Niagara/Source/Niagara/Private/Stateless/NiagaraStatelessEmitter.cpp b/Engine/Plugins/FX/Niagara/Source/Niagara/Private/Stateless/NiagaraStatelessEmitter.cpp
--- a/Engine/Plugins/FX/Niagara/Source/Niagara/Private/Stateless/NiagaraStatelessEmitter.cpp
+++ b/Engine/Plugins/FX/Niagara/Source/Niagara/Private/Stateless/NiagaraStatelessEmitter.cpp
@@ -398,7 +398,7 @@ void UNiagaraStatelessEmitter::CacheFromCompiledData()
 			}
 			else //if (SpawnInfo.Type == ENiagaraStatelessSpawnInfoType::Burst)
 			{
-				EmitterBuildContext.ConvertDistributionToRange(SpawnInfo.Amount, 0.0f);
+				EmitterBuildContext.ConvertDistributionToRange(SpawnInfo.Amount, 0);
 				EmitterBuildContext.ConvertDistributionToRange(SpawnInfo.SpawnProbability, 0.0f);
 			}
 		}
diff --git a/Engine/Plugins/FX/Niagara/Source/Niagara/Public/NiagaraComponent.h b/Engine/Plugins/FX/Niagara/Source/Niagara/Public/NiagaraComponent.h
--- a/Engine/Plugins/FX/Niagara/Source/Niagara/Public/NiagaraComponent.h
+++ b/Engine/Plugins/FX/Niagara/Source/Niagara/Public/NiagaraComponent.h
@@ -69,17 +69,17 @@ class UNiagaraComponent : public UFXSystemComponent
 public:
 
 	/********* UFXSystemComponent *********/
-	NIAGARA_API void SetBoolParameter(FName ParameterName, bool Param) override;
-	NIAGARA_API void SetIntParameter(FName ParameterName, int Param) override;
-	NIAGARA_API void SetFloatParameter(FName ParameterName, float Param) override;
-	NIAGARA_API void SetVectorParameter(FName ParameterName, FVector Param) override;
-	NIAGARA_API void SetColorParameter(FName ParameterName, FLinearColor Param) override;
-	NIAGARA_API void SetActorParameter(FName ParameterName, class AActor* Param) override;
+	NIAGARA_API virtual void SetBoolParameter(FName ParameterName, bool Param) override;
+	NIAGARA_API virtual void SetIntParameter(FName ParameterName, int Param) override;
+	NIAGARA_API virtual void SetFloatParameter(FName ParameterName, float Param) override;
+	NIAGARA_API virtual void SetVectorParameter(FName ParameterName, FVector Param) override;
+	NIAGARA_API virtual void SetColorParameter(FName ParameterName, FLinearColor Param) override;
+	NIAGARA_API virtual void SetActorParameter(FName ParameterName, class AActor* Param) override;
 
 	NIAGARA_API virtual UFXSystemAsset* GetFXSystemAsset() const override;
-	NIAGARA_API void SetEmitterEnable(FName EmitterName, bool bNewEnableState) override;
-	NIAGARA_API void ReleaseToPool() override;
-	NIAGARA_API uint32 GetApproxMemoryUsage() const override;
+	NIAGARA_API virtual void SetEmitterEnable(FName EmitterName, bool bNewEnableState) override;
+	NIAGARA_API virtual void ReleaseToPool() override;
+	NIAGARA_API virtual uint32 GetApproxMemoryUsage() const override;
 	NIAGARA_API virtual void ActivateSystem(bool bFlagAsJustAttached = false) override;
 	/********* UFXSystemComponent *********/
 
diff --git a/Engine/Plugins/FX/Niagara/Source/Niagara/Public/NiagaraSpriteRendererProperties.h b/Engine/Plugins/FX/Niagara/Source/Niagara/Public/NiagaraSpriteRendererProperties.h
--- a/Engine/Plugins/FX/Niagara/Source/Niagara/Public/NiagaraSpriteRendererProperties.h
+++ b/Engine/Plugins/FX/Niagara/Source/Niagara/Public/NiagaraSpriteRendererProperties.h
@@ -133,7 +133,7 @@ public:
 	NIAGARA_API virtual FNiagaraRenderer* CreateEmitterRenderer(ERHIFeatureLevel::Type FeatureLevel, const FNiagaraEmitterInstance* Emitter, const FNiagaraSystemInstanceController& InController) override;
 	NIAGARA_API virtual class FNiagaraBoundsCalculator* CreateBoundsCalculator() override;
 	NIAGARA_API virtual void GetUsedMaterials(const FNiagaraEmitterInstance* InEmitter, TArray<UMaterialInterface*>& OutMaterials) const override;
-	NIAGARA_API virtual float GetMaterialStreamingScale() const override { return FMath::Max(FMath::Max(SubImageSize.X, SubImageSize.Y), 1.0f); }
+	NIAGARA_API virtual float GetMaterialStreamingScale() const override { return static_cast<float>(FMath::Max(FMath::Max(SubImageSize.X, SubImageSize.Y), 1.0)); }
 	NIAGARA_API virtual const FVertexFactoryType* GetVertexFactoryType() const override;
 	virtual bool IsSimTargetSupported(ENiagaraSimTarget InSimTarget) const override { return true; };
 	NIAGARA_API virtual bool PopulateRequiredBindings(FNiagaraParameterStore& InParameterStore)  override;
diff --git a/Engine/Plugins/FX/Niagara/Source/NiagaraBlueprintNodes/NiagaraBlueprintNodes.build.cs b/Engine/Plugins/FX/Niagara/Source/NiagaraBlueprintNodes/NiagaraBlueprintNodes.build.cs
--- a/Engine/Plugins/FX/Niagara/Source/NiagaraBlueprintNodes/NiagaraBlueprintNodes.build.cs
+++ b/Engine/Plugins/FX/Niagara/Source/NiagaraBlueprintNodes/NiagaraBlueprintNodes.build.cs
@@ -7,6 +7,7 @@ public class NiagaraBlueprintNodes : ModuleRules
 	public NiagaraBlueprintNodes(ReadOnlyTargetRules Target) : base(Target)
 	{
 		PCHUsage = ModuleRules.PCHUsageMode.UseExplicitOrSharedPCHs;
+		UnsafeTypeCastWarningLevel = WarningLevel.Error;
 		
 		PublicDependencyModuleNames.AddRange(
 			new string[]
diff --git a/Engine/Plugins/FX/Niagara/Source/NiagaraCore/NiagaraCore.Build.cs b/Engine/Plugins/FX/Niagara/Source/NiagaraCore/NiagaraCore.Build.cs
--- a/Engine/Plugins/FX/Niagara/Source/NiagaraCore/NiagaraCore.Build.cs
+++ b/Engine/Plugins/FX/Niagara/Source/NiagaraCore/NiagaraCore.Build.cs
@@ -6,6 +6,8 @@ public class NiagaraCore : ModuleRules
 {
     public NiagaraCore(ReadOnlyTargetRules Target) : base(Target)
     {
+	    UnsafeTypeCastWarningLevel = WarningLevel.Error;
+	    
         PrivateDependencyModuleNames.AddRange(
             new string[] {
                 "Core",
diff --git a/Engine/Plugins/FX/Niagara/Source/NiagaraEditor/Private/Customizations/NiagaraTypeCustomizations.cpp b/Engine/Plugins/FX/Niagara/Source/NiagaraEditor/Private/Customizations/NiagaraTypeCustomizations.cpp
--- a/Engine/Plugins/FX/Niagara/Source/NiagaraEditor/Private/Customizations/NiagaraTypeCustomizations.cpp
+++ b/Engine/Plugins/FX/Niagara/Source/NiagaraEditor/Private/Customizations/NiagaraTypeCustomizations.cpp
@@ -2599,7 +2599,7 @@ public:
 		{
 			FVector2D ChildSize = ChildSlot.GetWidget()->GetDesiredSize();
 
-			float XVal = FMath::Max(MinWidthVal, ChildSize.X);
+			double XVal = FMath::Max(MinWidthVal, ChildSize.X);
 			if (MaxWidthVal > MinWidthVal)
 			{
 				XVal = FMath::Min(MaxWidthVal, XVal);
diff --git a/Engine/Plugins/FX/Niagara/Source/NiagaraEditor/Private/NiagaraCompilationDigestBridgeImpl.cpp b/Engine/Plugins/FX/Niagara/Source/NiagaraEditor/Private/NiagaraCompilationDigestBridgeImpl.cpp
--- a/Engine/Plugins/FX/Niagara/Source/NiagaraEditor/Private/NiagaraCompilationDigestBridgeImpl.cpp
+++ b/Engine/Plugins/FX/Niagara/Source/NiagaraEditor/Private/NiagaraCompilationDigestBridgeImpl.cpp
@@ -55,7 +55,7 @@ void FNiagaraCompilationDigestBridge::CustomHlslReferencesTokens(const FCustomHl
 		return;
 	}
 
-	for (SIZE_T i = 0; i < TokenStrings.Num(); i++)
+	for (int32 i = 0; i < TokenStrings.Num(); i++)
 	{
 		FNameBuilder NameBuilder(TokenStrings[i]);
 		FStringView NameString(NameBuilder.ToView());
diff --git a/Engine/Plugins/FX/Niagara/Source/NiagaraEditor/Private/NiagaraCompilationGraphBridgeImpl.cpp b/Engine/Plugins/FX/Niagara/Source/NiagaraEditor/Private/NiagaraCompilationGraphBridgeImpl.cpp
--- a/Engine/Plugins/FX/Niagara/Source/NiagaraEditor/Private/NiagaraCompilationGraphBridgeImpl.cpp
+++ b/Engine/Plugins/FX/Niagara/Source/NiagaraEditor/Private/NiagaraCompilationGraphBridgeImpl.cpp
@@ -116,7 +116,7 @@ void FNiagaraCompilationGraphBridge::CustomHlslReferencesTokens(const FCustomHls
 	TArray<FStringView> StringTokens;
 	UNiagaraNodeCustomHlsl::GetTokensFromString(CustomNode->GetCustomHlsl(), StringTokens, false, false);
 
-	for (SIZE_T i = 0; i < TokenStrings.Num(); i++)
+	for (int32 i = 0; i < TokenStrings.Num(); i++)
 	{
 		FNameBuilder NameBuilder(TokenStrings[i]);
 		FStringView NameString(NameBuilder.ToView());
diff --git a/Engine/Plugins/FX/Niagara/Source/NiagaraEditor/Private/NiagaraEditorStyle.cpp b/Engine/Plugins/FX/Niagara/Source/NiagaraEditor/Private/NiagaraEditorStyle.cpp
--- a/Engine/Plugins/FX/Niagara/Source/NiagaraEditor/Private/NiagaraEditorStyle.cpp
+++ b/Engine/Plugins/FX/Niagara/Source/NiagaraEditor/Private/NiagaraEditorStyle.cpp
@@ -558,7 +558,7 @@ void FNiagaraEditorStyle::InitStackWidgets()
 	FTextBlockStyle NoteHeaderStyle;
 	
 	FFontOutlineSettings NoteHeaderOutlineSettings;
-	NoteHeaderOutlineSettings.OutlineSize = 1.f;
+	NoteHeaderOutlineSettings.OutlineSize = 1;
 	NoteHeaderStyle.SetFont(DEFAULT_FONT("Regular", 12.f, NoteHeaderOutlineSettings));
 	NoteHeaderStyle.SetColorAndOpacity(FSlateColor(FStyleColors::White));
 	FInlineEditableTextBlockStyle NoteHeaderInlineEditableStyle;
diff --git a/Engine/Plugins/FX/Niagara/Source/NiagaraEditor/Private/NiagaraNodeSelect.cpp b/Engine/Plugins/FX/Niagara/Source/NiagaraEditor/Private/NiagaraNodeSelect.cpp
--- a/Engine/Plugins/FX/Niagara/Source/NiagaraEditor/Private/NiagaraNodeSelect.cpp
+++ b/Engine/Plugins/FX/Niagara/Source/NiagaraEditor/Private/NiagaraNodeSelect.cpp
@@ -712,7 +712,7 @@ TArray<int32> UNiagaraNodeSelect::GetOptionValues() const
 				if(FNiagaraEditorUtilities::IsEnumIndexVisible(Enum, EnumIndex))
 				{
 					ValidEnumEntryCount++;
-					SelectorValues.Add(Enum->GetValueByIndex(EnumIndex));
+					SelectorValues.Add(static_cast<int32>(Enum->GetValueByIndex(EnumIndex)));
 				}			
 			}
 		}
diff --git a/Engine/Plugins/FX/Niagara/Source/NiagaraEditor/Private/NiagaraNodeStaticSwitch.cpp b/Engine/Plugins/FX/Niagara/Source/NiagaraEditor/Private/NiagaraNodeStaticSwitch.cpp
--- a/Engine/Plugins/FX/Niagara/Source/NiagaraEditor/Private/NiagaraNodeStaticSwitch.cpp
+++ b/Engine/Plugins/FX/Niagara/Source/NiagaraEditor/Private/NiagaraNodeStaticSwitch.cpp
@@ -250,7 +250,7 @@ TArray<int32> UNiagaraNodeStaticSwitch::GetOptionValues() const
 		{
 			if (FNiagaraEditorUtilities::IsEnumIndexVisible(Enum, EnumIndex))
 			{
-				OptionValues.Add(Enum->GetValueByIndex(EnumIndex));
+				OptionValues.Add(static_cast<int32>(Enum->GetValueByIndex(EnumIndex)));
 			}
 		}
 	}
diff --git a/Engine/Plugins/FX/Niagara/Source/NiagaraEditor/Private/NiagaraNodeUsageSelector.cpp b/Engine/Plugins/FX/Niagara/Source/NiagaraEditor/Private/NiagaraNodeUsageSelector.cpp
--- a/Engine/Plugins/FX/Niagara/Source/NiagaraEditor/Private/NiagaraNodeUsageSelector.cpp
+++ b/Engine/Plugins/FX/Niagara/Source/NiagaraEditor/Private/NiagaraNodeUsageSelector.cpp
@@ -229,10 +229,10 @@ void UNiagaraNodeUsageSelector::AllocateDefaultPins()
 	NumOptionsPerVariable = (int32) ENiagaraScriptGroup::Max;
 	for (int64 i = 0; i < NumOptionsPerVariable; i++)
 	{
-		const FString PathSuffix = ENiagaraScriptGroupEnum ? ( FString::Printf(TEXT(" if %s"), *ENiagaraScriptGroupEnum->GetNameStringByValue((int64)i))) : TEXT("Error Unknown!");
+		const FString PathSuffix = ENiagaraScriptGroupEnum ? ( FString::Printf(TEXT(" if %s"), *ENiagaraScriptGroupEnum->GetNameStringByValue(i))) : TEXT("Error Unknown!");
 		for (FNiagaraVariable& Var : OutputVars)
 		{
-			AddOptionPin(Var, i);
+			AddOptionPin(Var, static_cast<int32>(i));
 		}
 	}
 
diff --git a/Engine/Plugins/FX/Niagara/Source/NiagaraEditor/Private/TypeEditorUtilities/NiagaraEnumTypeEditorUtilities.cpp b/Engine/Plugins/FX/Niagara/Source/NiagaraEditor/Private/TypeEditorUtilities/NiagaraEnumTypeEditorUtilities.cpp
--- a/Engine/Plugins/FX/Niagara/Source/NiagaraEditor/Private/TypeEditorUtilities/NiagaraEnumTypeEditorUtilities.cpp
+++ b/Engine/Plugins/FX/Niagara/Source/NiagaraEditor/Private/TypeEditorUtilities/NiagaraEnumTypeEditorUtilities.cpp
@@ -140,7 +140,7 @@ void FNiagaraEditorEnumTypeUtilities::UpdateVariableWithDefaultValue(FNiagaraVar
 	checkf(Enum != nullptr, TEXT("Variable is not an enum type."));
 
 	FNiagaraInt32 EnumIntValue;
-	EnumIntValue.Value = Enum->GetValueByIndex(0);
+	EnumIntValue.Value = static_cast<int32>(Enum->GetValueByIndex(0));
 
 	Variable.SetValue<FNiagaraInt32>(EnumIntValue);
 }
diff --git a/Engine/Plugins/FX/Niagara/Source/NiagaraEditor/Private/ViewModels/Stack/NiagaraStackGraphUtilities.cpp b/Engine/Plugins/FX/Niagara/Source/NiagaraEditor/Private/ViewModels/Stack/NiagaraStackGraphUtilities.cpp
--- a/Engine/Plugins/FX/Niagara/Source/NiagaraEditor/Private/ViewModels/Stack/NiagaraStackGraphUtilities.cpp
+++ b/Engine/Plugins/FX/Niagara/Source/NiagaraEditor/Private/ViewModels/Stack/NiagaraStackGraphUtilities.cpp
@@ -196,8 +196,8 @@ void FNiagaraStackGraphUtilities::RelayoutGraph(UEdGraph& Graph)
 			for (UEdGraphNode* Node : TraversalLevel)
 			{
 				Node->Modify();
-				Node->NodePosX = CurrentXOffset;
-				Node->NodePosY = CurrentYOffset;
+				Node->NodePosX = FMath::RoundToInt(CurrentXOffset);
+				Node->NodePosY = FMath::RoundToInt(CurrentYOffset);
 				int NumInputPins = 0;
 				int NumOutputPins = 0;
 				for (UEdGraphPin* Pin : Node->GetAllPins())
diff --git a/Engine/Plugins/FX/Niagara/Source/NiagaraEditor/Private/ViewModels/Stack/NiagaraStackItemPropertyHeaderValueShared.cpp b/Engine/Plugins/FX/Niagara/Source/NiagaraEditor/Private/ViewModels/Stack/NiagaraStackItemPropertyHeaderValueShared.cpp
--- a/Engine/Plugins/FX/Niagara/Source/NiagaraEditor/Private/ViewModels/Stack/NiagaraStackItemPropertyHeaderValueShared.cpp
+++ b/Engine/Plugins/FX/Niagara/Source/NiagaraEditor/Private/ViewModels/Stack/NiagaraStackItemPropertyHeaderValueShared.cpp
@@ -51,9 +51,9 @@ int32 FNiagaraStackItemPropertyHeaderValue::GetEnumValue() const
 			}
 			else
 			{
-				ValuePtr = (uint8*)OwnerObject + EnumProperty->GetOffset_ForInternal();
+				ValuePtr = reinterpret_cast<uint8*>(OwnerObject) + EnumProperty->GetOffset_ForInternal();
 			}
-			EnumValueCache = EnumProperty->GetUnderlyingProperty()->GetSignedIntPropertyValue(ValuePtr);
+			EnumValueCache = static_cast<uint32>(EnumProperty->GetUnderlyingProperty()->GetSignedIntPropertyValue(ValuePtr));
 		}
 		else
 		{
diff --git a/Engine/Plugins/FX/Niagara/Source/NiagaraEditorWidgets/Private/Stack/SNiagaraDeterminismToggle.cpp b/Engine/Plugins/FX/Niagara/Source/NiagaraEditorWidgets/Private/Stack/SNiagaraDeterminismToggle.cpp
--- a/Engine/Plugins/FX/Niagara/Source/NiagaraEditorWidgets/Private/Stack/SNiagaraDeterminismToggle.cpp
+++ b/Engine/Plugins/FX/Niagara/Source/NiagaraEditorWidgets/Private/Stack/SNiagaraDeterminismToggle.cpp
@@ -36,7 +36,7 @@ void SNiagaraDeterminismToggle::Construct(const FArguments& InArgs, TSharedPtr<F
 	}
 
 	FSlateFontInfo FontInfo = FSlateFontInfo(FPaths::EngineContentDir() / TEXT("Slate/Fonts/Roboto-Regular.ttf"), 9);
-	FontInfo.OutlineSettings = FFontOutlineSettings(1.f, FLinearColor::White);
+	FontInfo.OutlineSettings = FFontOutlineSettings(1, FLinearColor::White);
 	
 	ChildSlot
 	[
diff --git a/Engine/Plugins/FX/NiagaraSimCaching/Source/NiagaraSimCaching/Private/Niagara/Sequencer/MovieSceneNiagaraCacheTemplate.cpp b/Engine/Plugins/FX/NiagaraSimCaching/Source/NiagaraSimCaching/Private/Niagara/Sequencer/MovieSceneNiagaraCacheTemplate.cpp
--- a/Engine/Plugins/FX/NiagaraSimCaching/Source/NiagaraSimCaching/Private/Niagara/Sequencer/MovieSceneNiagaraCacheTemplate.cpp
+++ b/Engine/Plugins/FX/NiagaraSimCaching/Source/NiagaraSimCaching/Private/Niagara/Sequencer/MovieSceneNiagaraCacheTemplate.cpp
@@ -112,7 +112,7 @@ struct FNiagaraCacheExecutionToken : IMovieSceneExecutionToken
 		FFrameNumber SectionStartTime = Section.GetInclusiveStartFrame();
 		FFrameNumber SectionEndTime = Section.GetExclusiveEndFrame();
 		
-		const float SequenceLength = Section.Params.SectionStretchMode == ENiagaraSimCacheSectionStretchMode::TimeDilate ? InFrameRate.AsSeconds(SectionEndTime - SectionStartTime) : ComponentDuration;
+		const float SequenceLength = Section.Params.SectionStretchMode == ENiagaraSimCacheSectionStretchMode::TimeDilate ? static_cast<float>(InFrameRate.AsSeconds(SectionEndTime - SectionStartTime)) : ComponentDuration;
 		const FFrameTime AnimationLength = SequenceLength * InFrameRate;
 		const int32 LengthInFrames = AnimationLength.FrameNumber.Value + static_cast<int>(AnimationLength.GetSubFrame() + 0.5f) + 1;
 	
@@ -123,15 +123,15 @@ struct FNiagaraCacheExecutionToken : IMovieSceneExecutionToken
 
 		const float SectionPlayRate = BaseParams.PlayRate;
 		const float AnimPlayRate = FMath::IsNearlyZero(SectionPlayRate) ? 1.0f : SectionPlayRate;
-		const float SeqLength = SequenceLength - InFrameRate.AsSeconds(BaseParams.StartFrameOffset + BaseParams.EndFrameOffset);
+		const float SeqLength = SequenceLength - static_cast<float>(InFrameRate.AsSeconds(BaseParams.StartFrameOffset + BaseParams.EndFrameOffset));
 
-		float AnimPosition = FFrameTime::FromDecimal((InPosition - SectionStartTime).AsDecimal() * AnimPlayRate) / InFrameRate;
-		AnimPosition += InFrameRate.AsSeconds(BaseParams.FirstLoopStartFrameOffset);
+		float AnimPosition = static_cast<float>(FFrameTime::FromDecimal((InPosition - SectionStartTime).AsDecimal() * AnimPlayRate) / InFrameRate);
+		AnimPosition += static_cast<float>(InFrameRate.AsSeconds(BaseParams.FirstLoopStartFrameOffset));
 		if (SeqLength > 0.f && (bLooping || !FMath::IsNearlyEqual(AnimPosition, SeqLength, 1e-4f)))
 		{
 			AnimPosition = FMath::Fmod(AnimPosition, SeqLength);
 		}
-		AnimPosition += InFrameRate.AsSeconds(BaseParams.StartFrameOffset);
+		AnimPosition += static_cast<float>(InFrameRate.AsSeconds(BaseParams.StartFrameOffset));
 		if (BaseParams.bReverse)
 		{
 			AnimPosition = SequenceLength - AnimPosition;
diff --git a/Engine/Source/Runtime/Engine/Classes/Kismet/KismetMathLibrary.h b/Engine/Source/Runtime/Engine/Classes/Kismet/KismetMathLibrary.h
--- a/Engine/Source/Runtime/Engine/Classes/Kismet/KismetMathLibrary.h
+++ b/Engine/Source/Runtime/Engine/Classes/Kismet/KismetMathLibrary.h
@@ -510,16 +510,10 @@ class UKismetMathLibrary : public UBlueprintFunctionLibrary
 	UFUNCTION(BlueprintPure, meta = (DisplayName = "Power"), Category = "Math|Float")
 	static ENGINE_API double MultiplyMultiply_FloatFloat(double Base, double Exp);
 
-	UE_DEPRECATED(5.0, "This method has been deprecated and will be removed. Use the double version instead.")
-	static ENGINE_API float Multiply_FloatFloat(float A, float B);
-
 	/** Multiplication (A * B) */
 	UFUNCTION(BlueprintPure, meta = (DisplayName = "int * float", CompactNodeTitle = "*", Keywords = "* multiply"), Category = "Math|Float")
 	static ENGINE_API double Multiply_IntFloat(int32 A, double B);
 
-	UE_DEPRECATED(5.0, "This method has been deprecated and will be removed. Use the double version instead.")
-	static ENGINE_API float Divide_FloatFloat(float A, float B = 1.f);
-
 	UFUNCTION(BlueprintPure, CustomThunk, meta = (DisplayName = "% (Float)", CompactNodeTitle = "%", Keywords = "% modulus"), Category = "Math|Float")
 	static ENGINE_API double Percent_FloatFloat(double A, double B = 1.f);
 
@@ -583,27 +577,6 @@ class UKismetMathLibrary : public UBlueprintFunctionLibrary
 	UFUNCTION(BlueprintPure, meta = (DisplayName = "Equal (Float)", CompactNodeTitle = "==", Keywords = "== equal"), Category = "Math|Float")
 	static ENGINE_API bool EqualEqual_DoubleDouble(double A, double B);
 
-	UE_DEPRECATED(5.0, "This method has been deprecated and will be removed. Use the double version instead.")
-	static ENGINE_API float Add_FloatFloat(float A, float B = 1.f);
-
-	UE_DEPRECATED(5.0, "This method has been deprecated and will be removed. Use the double version instead.")
-	static ENGINE_API float Subtract_FloatFloat(float A, float B = 1.f);
-
-	UE_DEPRECATED(5.0, "This method has been deprecated and will be removed. Use the double version instead.")
-	static ENGINE_API bool Less_FloatFloat(float A, float B);
-
-	UE_DEPRECATED(5.0, "This method has been deprecated and will be removed. Use the double version instead.")
-	static ENGINE_API bool Greater_FloatFloat(float A, float B);
-
-	UE_DEPRECATED(5.0, "This method has been deprecated and will be removed. Use the double version instead.")
-	static ENGINE_API bool LessEqual_FloatFloat(float A, float B);
-
-	UE_DEPRECATED(5.0, "This method has been deprecated and will be removed. Use the double version instead.")
-	static ENGINE_API bool GreaterEqual_FloatFloat(float A, float B);
-
-	UE_DEPRECATED(5.0, "This method has been deprecated and will be removed. Use the double version instead.")
-	static ENGINE_API bool EqualEqual_FloatFloat(float A, float B);
-
 	/** Returns true if A is nearly equal to B (|A - B| < ErrorTolerance) */
 	UFUNCTION(BlueprintPure, meta = (DisplayName = "Nearly Equal (Float)", Keywords = "== equal"), Category = "Math|Float")
 	static ENGINE_API bool NearlyEqual_FloatFloat(double A, double B, double ErrorTolerance = 1.e-6);
@@ -612,9 +585,6 @@ class UKismetMathLibrary : public UBlueprintFunctionLibrary
 	UFUNCTION(BlueprintPure, meta = (DisplayName = "Not Equal (Float)", CompactNodeTitle = "!=", Keywords = "!= not equal"), Category = "Math|Float")
 	static ENGINE_API bool NotEqual_DoubleDouble(double A, double B);
 
-	UE_DEPRECATED(5.0, "This method has been deprecated and will be removed. Use the double version instead.")
-	static ENGINE_API bool NotEqual_FloatFloat(float A, float B);
-
 	/** Returns true if value is between Min and Max (V >= Min && V <= Max)
 	 * If InclusiveMin is true, value needs to be equal or larger than Min, else it needs to be larger
 	 * If InclusiveMax is true, value needs to be smaller or equal than Max, else it needs to be smaller
@@ -848,9 +818,6 @@ class UKismetMathLibrary : public UBlueprintFunctionLibrary
 	UFUNCTION(BlueprintPure, meta = (DisplayName = "Division (Whole and Remainder) to Integer64"), Category = "Math|Float")
 	static ENGINE_API int64 FMod64(double Dividend, double Divisor, double& Remainder);
 
-	UE_DEPRECATED(5.0, "This method has been deprecated and will be removed. Use the double version instead.")
-	static ENGINE_API int32 FMod(float Dividend, float Divisor, float& Remainder);
-
 	/** Sign (float, returns -1 if A < 0, 0 if A is zero, and +1 if A > 0) */
 	UFUNCTION(BlueprintPure, meta=(DisplayName = "Sign (Float)"), Category="Math|Float")
 	static ENGINE_API double SignOfFloat(double A);
@@ -1016,9 +983,6 @@ class UKismetMathLibrary : public UBlueprintFunctionLibrary
 	UFUNCTION(BlueprintPure, Category = "Math|Vector2D", meta = (Keywords = "construct build", NativeMakeFunc))
 	static ENGINE_API FVector2D MakeVector2D(double X, double Y);
 
-	UE_DEPRECATED(5.0, "This method has been deprecated and will be removed. Use the double version instead.")
-	static ENGINE_API void BreakVector2D(FVector2D InVec, float& X, float& Y);
-
 	/** Breaks a 2D vector apart into X, Y. */
 	UFUNCTION(BlueprintPure, Category = "Math|Vector2D", meta = (NativeBreakFunc))
 	static ENGINE_API void BreakVector2D(FVector2D InVec, double& X, double& Y);
@@ -1245,9 +1209,6 @@ class UKismetMathLibrary : public UBlueprintFunctionLibrary
 	UFUNCTION(BlueprintPure, meta = (DisplayName = "To Direction And Length", ScriptMethod = "ToDirectionAndLength"), Category = "Math|Vector2D")
 	static ENGINE_API void ToDirectionAndLength2D(FVector2D A, FVector2D &OutDir, double &OutLength);
 
-	UE_DEPRECATED(5.0, "This method has been deprecated and will be removed. Use the double version instead.")
-	static ENGINE_API void ToDirectionAndLength2D(FVector2D A, FVector2D& OutDir, float& OutLength);
-
 	/**
 	 * Get this vector as a vector where each component has been rounded to the nearest int.
 	 *
@@ -1340,9 +1301,6 @@ class UKismetMathLibrary : public UBlueprintFunctionLibrary
 	UFUNCTION(BlueprintCallable, meta = (ScriptMethod = "Set"), Category = "Math|Vector")
 	static ENGINE_API void Vector_Set(UPARAM(ref) FVector& A, double X, double Y, double Z);
 
-	UE_DEPRECATED(5.0, "This method has been deprecated and will be removed. Use the double version instead.")
-	static ENGINE_API void BreakVector(FVector InVec, float& X, float& Y, float& Z);
-
 	/** Breaks a vector apart into X, Y, Z */
 	UFUNCTION(BlueprintPure, Category="Math|Vector", meta=(NativeBreakFunc))
 	static ENGINE_API void BreakVector(FVector InVec, double& X, double& Y, double& Z);
@@ -2085,9 +2043,6 @@ class UKismetMathLibrary : public UBlueprintFunctionLibrary
 	UFUNCTION(BlueprintPure, meta = (Keywords = "construct build", NativeMakeFunc), Category = "Math|Vector4")
 	static ENGINE_API FVector4 MakeVector4(double X, double Y, double Z, double W);
 
-	UE_DEPRECATED(5.0, "This method has been deprecated and will be removed. Use the double version instead.")
-	static ENGINE_API void BreakVector4(const FVector4& InVec, float& X, float& Y, float& Z, float& W);
-
 	/** Breaks a 4D vector apart into X, Y, Z, W. */
 	UFUNCTION(BlueprintPure, meta = (NativeBreakFunc), Category = "Math|Vector4")
 	static ENGINE_API void BreakVector4(const FVector4& InVec, double& X, double& Y, double& Z, double& W);
@@ -4343,20 +4298,6 @@ class UKismetMathLibrary : public UBlueprintFunctionLibrary
 	// Geometry
 	//
 
-	/**  
-	 * Finds the minimum area rectangle that encloses all of the points in InVerts
-	 * Uses algorithm found in http://www.geometrictools.com/Documentation/MinimumAreaRectangle.pdf
-	 *	
-	 * @param		InVerts	- Points to enclose in the rectangle
-	 * @outparam	OutRectCenter - Center of the enclosing rectangle
-	 * @outparam	OutRectSideA - Vector oriented and sized to represent one edge of the enclosing rectangle, orthogonal to OutRectSideB
-	 * @outparam	OutRectSideB - Vector oriented and sized to represent one edge of the enclosing rectangle, orthogonal to OutRectSideA
-	*/
-	UE_DEPRECATED(5.0, "Use MinAreaRectangle instead. This deprecated version incorrectly returns the average of all input points as the rectangle center.")
-	UFUNCTION(BlueprintCallable, Category="Math|Geometry", meta=(WorldContext="WorldContextObject", CallableWithoutWorldContext,
-		DeprecatedFunction, DeprecationMessage = "Use 'Min Area Rectangle' instead; this deprecated version incorrectly returns the average of all input points as the rectangle center."))
-	static ENGINE_API void MinimumAreaRectangle(UObject* WorldContextObject, const TArray<FVector>& InVerts, const FVector& SampleSurfaceNormal, FVector& OutRectCenter, FRotator& OutRectRotation, float& OutSideLengthX, float& OutSideLengthY, bool bDebugDraw = false);
-
 	/**
 	 * Finds the minimum area rectangle that encloses a set of coplanar points.
 	 * Uses the exhaustive search algorithm in http://www.geometrictools.com/Documentation/MinimumAreaRectangle.pdf
diff --git a/Engine/Source/Runtime/Engine/Classes/Kismet/KismetMathLibrary.inl b/Engine/Source/Runtime/Engine/Classes/Kismet/KismetMathLibrary.inl
--- a/Engine/Source/Runtime/Engine/Classes/Kismet/KismetMathLibrary.inl
+++ b/Engine/Source/Runtime/Engine/Classes/Kismet/KismetMathLibrary.inl
@@ -410,12 +410,6 @@ double UKismetMathLibrary::MultiplyMultiply_FloatFloat(double Base, double Exp)
 	return FMath::Pow(Base, Exp);
 }
 
-KISMET_MATH_FORCEINLINE
-float UKismetMathLibrary::Multiply_FloatFloat(float A, float B)
-{
-	return A * B;
-}	
-
 KISMET_MATH_FORCEINLINE
 double UKismetMathLibrary::Multiply_IntFloat(int32 A, double B)
 {
@@ -495,17 +489,6 @@ double UKismetMathLibrary::Multiply_DoubleDouble(double A, double B)
 	return A * B;
 }
 
-KISMET_MATH_FORCEINLINE
-float UKismetMathLibrary::Divide_FloatFloat(float A, float B)
-{
-	if (B == 0.f)
-	{
-		ReportError_Divide_FloatFloat();
-		return 0.f;
-	}
-	return A / B;
-}
-
 KISMET_MATH_FORCEINLINE
 double UKismetMathLibrary::Divide_DoubleDouble(double A, double B)
 {
@@ -547,60 +530,12 @@ bool UKismetMathLibrary::EqualEqual_DoubleDouble(double A, double B)
 	return A == B;
 }
 
-KISMET_MATH_FORCEINLINE
-float UKismetMathLibrary::Add_FloatFloat(float A, float B)
-{
-	return A + B;
-}	
-
-KISMET_MATH_FORCEINLINE
-float UKismetMathLibrary::Subtract_FloatFloat(float A, float B)
-{
-	return A - B;
-}	
-
-KISMET_MATH_FORCEINLINE
-bool UKismetMathLibrary::Less_FloatFloat(float A, float B)
-{
-	return A < B;
-}	
-
-KISMET_MATH_FORCEINLINE
-bool UKismetMathLibrary::Greater_FloatFloat(float A, float B)
-{
-	return A > B;
-}	
-
-KISMET_MATH_FORCEINLINE
-bool UKismetMathLibrary::LessEqual_FloatFloat(float A, float B)
-{
-	return A <= B;
-}	
-
-KISMET_MATH_FORCEINLINE
-bool UKismetMathLibrary::GreaterEqual_FloatFloat(float A, float B)
-{
-	return A >= B;
-}	
-
-KISMET_MATH_FORCEINLINE
-bool UKismetMathLibrary::EqualEqual_FloatFloat(float A, float B)
-{
-	return A == B;
-}	
-
 KISMET_MATH_FORCEINLINE
 bool UKismetMathLibrary::NearlyEqual_FloatFloat(double A, double B, double ErrorTolerance)
 {
 	return FMath::IsNearlyEqual(A, B, ErrorTolerance);
 }
 
-KISMET_MATH_FORCEINLINE
-bool UKismetMathLibrary::NotEqual_FloatFloat(float A, float B)
-{
-	return A != B;
-}
-
 KISMET_MATH_FORCEINLINE
 bool UKismetMathLibrary::NotEqual_DoubleDouble(double A, double B)
 {
@@ -768,19 +703,19 @@ double UKismetMathLibrary::Square(double A)
 KISMET_MATH_FORCEINLINE
 int32 UKismetMathLibrary::Round(double A)
 {
-	return FMath::RoundToInt(A);
+	return static_cast<int32>(FMath::RoundToInt(A));
 }	
 
 KISMET_MATH_FORCEINLINE
 int32 UKismetMathLibrary::FFloor(double A)
 {
-	return FMath::FloorToInt(A);
+	return static_cast<int32>(FMath::FloorToInt(A));
 }	
 
 KISMET_MATH_FORCEINLINE
 int32 UKismetMathLibrary::FTrunc(double A)
 {
-	return FMath::TruncToInt(A);
+	return static_cast<int32>(FMath::TruncToInt(A));
 }	
 
 KISMET_MATH_FORCEINLINE
@@ -810,7 +745,7 @@ int64 UKismetMathLibrary::FCeil64(double A)
 KISMET_MATH_FORCEINLINE
 int32 UKismetMathLibrary::FCeil(double A)
 {
-	return FMath::CeilToInt(A);
+	return static_cast<int32>(FMath::CeilToInt(A));
 }	
 
 KISMET_MATH_FORCEINLINE
@@ -1144,13 +1079,6 @@ FVector2D UKismetMathLibrary::MakeVector2D(double X, double Y)
 	return FVector2D(X, Y);
 }
 
-KISMET_MATH_FORCEINLINE
-void UKismetMathLibrary::BreakVector2D(FVector2D InVec, float& X, float& Y)
-{
-	X = InVec.X;
-	Y = InVec.Y;
-}
-
 KISMET_MATH_FORCEINLINE
 void UKismetMathLibrary::BreakVector2D(FVector2D InVec, double& X, double& Y)
 {
@@ -1374,14 +1302,6 @@ FVector UKismetMathLibrary::Spherical2DToUnitCartesian(FVector2D A)
 	return A.SphericalToUnitCartesian();
 }
 
-KISMET_MATH_FORCEINLINE
-void UKismetMathLibrary::ToDirectionAndLength2D(FVector2D A, FVector2D &OutDir, float &OutLength)
-{
-	double OutLengthDbl;	// LWC_TODO: Perf pessimization
-	A.ToDirectionAndLength(OutDir, OutLengthDbl);
-	OutLength = OutLengthDbl;
-}
-
 KISMET_MATH_FORCEINLINE
 void UKismetMathLibrary::ToDirectionAndLength2D(FVector2D A, FVector2D& OutDir, double& OutLength)
 {
@@ -1482,14 +1402,6 @@ void UKismetMathLibrary::Vector_Set(FVector& A, double X, double Y, double Z)
 	A.Set(X, Y, Z);
 }
 
-KISMET_MATH_FORCEINLINE
-void UKismetMathLibrary::BreakVector(FVector InVec, float& X, float& Y, float& Z)
-{
-	X = InVec.X;
-	Y = InVec.Y;
-	Z = InVec.Z;
-}
-
 KISMET_MATH_FORCEINLINE
 void UKismetMathLibrary::BreakVector(FVector InVec, double& X, double& Y, double& Z)
 {
@@ -2067,9 +1979,9 @@ FVector UKismetMathLibrary::GetDirectionUnitVector(FVector From, FVector To)
 KISMET_MATH_FORCEINLINE
 FIntVector UKismetMathLibrary::FTruncVector(const FVector& InVector)
 {
-	return FIntVector(FMath::TruncToInt(InVector.X),
-		FMath::TruncToInt(InVector.Y),
-		FMath::TruncToInt(InVector.Z));
+	return FIntVector(static_cast<int32>(FMath::TruncToInt(InVector.X)),
+		static_cast<int32>(FMath::TruncToInt(InVector.Y)),
+		static_cast<int32>(FMath::TruncToInt(InVector.Z)));
 }
 
 
@@ -2088,15 +2000,6 @@ FVector4 UKismetMathLibrary::MakeVector4(double X, double Y, double Z, double W)
 	return FVector4(X, Y, Z, W);
 }
 
-KISMET_MATH_FORCEINLINE
-void UKismetMathLibrary::BreakVector4(const FVector4& InVec, float& X, float& Y, float& Z, float& W)
-{
-	X = InVec.X;
-	Y = InVec.Y;
-	Z = InVec.Z;
-	W = InVec.W;
-}
-
 KISMET_MATH_FORCEINLINE
 void UKismetMathLibrary::BreakVector4(const FVector4& InVec, double& X, double& Y, double& Z, double& W)
 {
@@ -2659,31 +2562,31 @@ int32 UKismetMathLibrary::GetSeconds(FTimespan A)
 KISMET_MATH_FORCEINLINE
 float UKismetMathLibrary::GetTotalDays(FTimespan A)
 {
-	return A.GetTotalDays();
+	return static_cast<float>(A.GetTotalDays());
 }
 
 KISMET_MATH_FORCEINLINE
 float UKismetMathLibrary::GetTotalHours(FTimespan A)
 {
-	return A.GetTotalHours();
+	return static_cast<float>(A.GetTotalHours());
 }
 
 KISMET_MATH_FORCEINLINE
 float UKismetMathLibrary::GetTotalMilliseconds(FTimespan A)
 {
-	return A.GetTotalMilliseconds();
+	return static_cast<float>(A.GetTotalMilliseconds());
 }
 
 KISMET_MATH_FORCEINLINE
 float UKismetMathLibrary::GetTotalMinutes(FTimespan A)
 {
-	return A.GetTotalMinutes();
+	return static_cast<float>(A.GetTotalMinutes());
 }
 
 KISMET_MATH_FORCEINLINE
 float UKismetMathLibrary::GetTotalSeconds(FTimespan A)
 {
-	return A.GetTotalSeconds();
+	return static_cast<float>(A.GetTotalSeconds());
 }
 
 KISMET_MATH_FORCEINLINE
@@ -2701,7 +2604,7 @@ FTimespan UKismetMathLibrary::TimespanMinValue()
 KISMET_MATH_INLINE
 float UKismetMathLibrary::TimespanRatio(FTimespan A, FTimespan B)
 {
-	return FTimespan::Ratio(A, B);
+	return static_cast<float>(FTimespan::Ratio(A, B));
 }
 
 KISMET_MATH_FORCEINLINE
@@ -2897,7 +2800,9 @@ FLinearColor UKismetMathLibrary::Conv_FloatToLinearColor(float InFloat)
 KISMET_MATH_FORCEINLINE
 FLinearColor UKismetMathLibrary::Conv_DoubleToLinearColor(double InDouble)
 {
-	return FLinearColor(InDouble, InDouble, InDouble);
+	// potential precision loss by float conversion
+	float Value = static_cast<float>(InDouble);
+	return FLinearColor(Value, Value, Value);
 }
 
 KISMET_MATH_FORCEINLINE
@@ -2995,9 +2900,9 @@ FRotator UKismetMathLibrary::MakeRotFromZY(const FVector& Z, const FVector& Y)
 KISMET_MATH_FORCEINLINE
 void UKismetMathLibrary::BreakRotator(FRotator InRot, float& Roll, float& Pitch, float& Yaw)
 {
-	Pitch = InRot.Pitch;
-	Yaw = InRot.Yaw;
-	Roll = InRot.Roll;
+	Pitch = static_cast<float>(InRot.Pitch);
+	Yaw = static_cast<float>(InRot.Yaw);
+	Roll = static_cast<float>(InRot.Roll);
 }
 
 
@@ -3097,13 +3002,13 @@ FMatrix UKismetMathLibrary::Matrix_GetTransposed(const FMatrix& M)
 KISMET_MATH_FORCEINLINE
 float UKismetMathLibrary::Matrix_GetDeterminant(const FMatrix& M)
 {
-	return M.Determinant();
+	return static_cast<float>(M.Determinant());
 }
 
 KISMET_MATH_FORCEINLINE
 float UKismetMathLibrary::Matrix_GetRotDeterminant(const FMatrix& M)
 {
-	return M.RotDeterminant();
+	return static_cast<float>(M.RotDeterminant());
 }
 
 KISMET_MATH_FORCEINLINE
@@ -3165,7 +3070,7 @@ FMatrix UKismetMathLibrary::Matrix_ScaleTranslation(const FMatrix& M, FVector Sc
 KISMET_MATH_FORCEINLINE
 float UKismetMathLibrary::Matrix_GetMaximumAxisScale(const FMatrix& M)
 {
-	return M.GetMaximumAxisScale();
+	return static_cast<float>(M.GetMaximumAxisScale());
 }
 
 KISMET_MATH_FORCEINLINE
@@ -3326,10 +3231,10 @@ FQuat UKismetMathLibrary::MakeQuat(float X, float Y, float Z, float W)
 KISMET_MATH_FORCEINLINE
 void UKismetMathLibrary::BreakQuat(const FQuat& InQuat, float& X, float& Y, float& Z, float& W)
 {
-	X = InQuat.X;
-	Y = InQuat.Y;
-	Z = InQuat.Z;
-	W = InQuat.W;
+	X = static_cast<float>(InQuat.X);
+	Y = static_cast<float>(InQuat.Y);
+	Z = static_cast<float>(InQuat.Z);
+	W = static_cast<float>(InQuat.W);
 }
 
 KISMET_MATH_FORCEINLINE
@@ -3359,7 +3264,7 @@ bool UKismetMathLibrary::Quat_IsNonFinite(const FQuat& Q)
 KISMET_MATH_FORCEINLINE
 float UKismetMathLibrary::Quat_AngularDistance(const FQuat& A, const FQuat& B)
 {
-	return A.AngularDistance(B);
+	return static_cast<float>(A.AngularDistance(B));
 }
 
 KISMET_MATH_FORCEINLINE
@@ -3383,7 +3288,7 @@ FQuat UKismetMathLibrary::Quat_Exp(const FQuat& Q)
 KISMET_MATH_FORCEINLINE
 float UKismetMathLibrary::Quat_GetAngle(const FQuat& Q)
 {
-	return Q.GetAngle();
+	return static_cast<float>(Q.GetAngle());
 }
 
 KISMET_MATH_FORCEINLINE
@@ -3488,13 +3393,13 @@ FQuat UKismetMathLibrary::Conv_RotatorToQuaternion(FRotator InRot)
 KISMET_MATH_FORCEINLINE
 float UKismetMathLibrary::Quat_Size(const FQuat& Q)
 {
-	return Q.Size();
+	return static_cast<float>(Q.Size());
 }
 
 KISMET_MATH_FORCEINLINE
 float UKismetMathLibrary::Quat_SizeSquared(const FQuat& Q)
 {
-	return Q.SizeSquared();
+	return static_cast<float>(Q.SizeSquared());
 }
 
 KISMET_MATH_FORCEINLINE
diff --git a/Engine/Source/Runtime/Engine/Private/KismetMathLibrary.cpp b/Engine/Source/Runtime/Engine/Private/KismetMathLibrary.cpp
--- a/Engine/Source/Runtime/Engine/Private/KismetMathLibrary.cpp
+++ b/Engine/Source/Runtime/Engine/Private/KismetMathLibrary.cpp
@@ -259,26 +259,6 @@ double UKismetMathLibrary::Log(double A, double Base)
 	return Result;
 }
 
-int32 UKismetMathLibrary::FMod(float Dividend, float Divisor, float& Remainder)
-{
-	int32 Result;
-	if (Divisor != 0.f)
-	{
-		const float Quotient = Dividend / Divisor;
-		Result = (Quotient < 0.f ? -1 : 1) * FMath::FloorToInt32(FMath::Abs(Quotient));
-		Remainder = FMath::Fmod(Dividend, Divisor);
-	}
-	else
-	{
-		FFrame::KismetExecutionMessage(TEXT("Attempted modulo 0 - returning 0."), ELogVerbosity::Warning, DivideByZeroWarning);
-
-		Result = 0;
-		Remainder = 0.f;
-	}
-
-	return Result;
-}
-
 int32 UKismetMathLibrary::FMod(double Dividend, double Divisor, double& Remainder)
 {
 	int32 Result;
@@ -1155,108 +1135,6 @@ void UKismetMathLibrary::SetRandomStreamSeed(FRandomStream& Stream, int32 NewSee
 	Stream.Initialize(NewSeed);
 }
 
-void UKismetMathLibrary::MinimumAreaRectangle(class UObject* WorldContextObject, const TArray<FVector>& InVerts, const FVector& SampleSurfaceNormal, FVector& OutRectCenter, FRotator& OutRectRotation, float& OutSideLengthX, float& OutSideLengthY, bool bDebugDraw)
-{
-	float MinArea = -1.f;
-	float CurrentArea = -1.f;
-	FVector SupportVectorA, SupportVectorB;
-	FVector RectSideA, RectSideB;
-	float MinDotResultA, MinDotResultB, MaxDotResultA, MaxDotResultB;
-	FVector TestEdge;
-	float TestEdgeDot = 0.f;
-	FVector PolyNormal(0.f, 0.f, 1.f);
-	TArray<int32> PolyVertIndices;
-
-	// Bail if we receive an empty InVerts array
-	if (InVerts.Num() == 0)
-	{
-		return;
-	}
-
-	// Compute the approximate normal of the poly, using the direction of SampleSurfaceNormal for guidance
-	PolyNormal = (InVerts[InVerts.Num() / 3] - InVerts[0]) ^ (InVerts[InVerts.Num() * 2 / 3] - InVerts[InVerts.Num() / 3]);
-	if ((PolyNormal | SampleSurfaceNormal) < 0.f)
-	{
-		PolyNormal = -PolyNormal;
-	}
-
-	// Transform the sample points to 2D
-	FMatrix SurfaceNormalMatrix = FRotationMatrix::MakeFromZX(PolyNormal, FVector(1.f, 0.f, 0.f));
-	TArray<FVector> TransformedVerts;
-	OutRectCenter = FVector(0.f);
-	for (int32 Idx = 0; Idx < InVerts.Num(); ++Idx)
-	{
-		OutRectCenter += InVerts[Idx];
-		TransformedVerts.Add(SurfaceNormalMatrix.InverseTransformVector(InVerts[Idx]));
-	}
-	OutRectCenter /= InVerts.Num();
-
-	// Compute the convex hull of the sample points
-	ConvexHull2D::ComputeConvexHullLegacy(TransformedVerts, PolyVertIndices);
-
-	// Minimum area rectangle as computed by http://www.geometrictools.com/Documentation/MinimumAreaRectangle.pdf
-	for (int32 Idx = 1; Idx < PolyVertIndices.Num() - 1; ++Idx)
-	{
-		SupportVectorA = (TransformedVerts[PolyVertIndices[Idx]] - TransformedVerts[PolyVertIndices[Idx-1]]).GetSafeNormal();
-		SupportVectorA.Z = 0.f;
-		SupportVectorB.X = -SupportVectorA.Y;
-		SupportVectorB.Y = SupportVectorA.X;
-		SupportVectorB.Z = 0.f;
-		MinDotResultA = MinDotResultB = MaxDotResultA = MaxDotResultB = 0.f;
-
-		for (int TestVertIdx = 1; TestVertIdx < PolyVertIndices.Num(); ++TestVertIdx)
-		{
-			TestEdge = TransformedVerts[PolyVertIndices[TestVertIdx]] - TransformedVerts[PolyVertIndices[0]];
-			TestEdgeDot = SupportVectorA | TestEdge;
-			if (TestEdgeDot < MinDotResultA)
-			{
-				MinDotResultA = TestEdgeDot;
-			}
-			else if (TestEdgeDot > MaxDotResultA)
-			{
-				MaxDotResultA = TestEdgeDot;
-			}
-
-			TestEdgeDot = SupportVectorB | TestEdge;
-			if (TestEdgeDot < MinDotResultB)
-			{
-				MinDotResultB = TestEdgeDot;
-			}
-			else if (TestEdgeDot > MaxDotResultB)
-			{
-				MaxDotResultB = TestEdgeDot;
-			}
-		}
-
-		CurrentArea = (MaxDotResultA - MinDotResultA) * (MaxDotResultB - MinDotResultB);
-		if (MinArea < 0.f || CurrentArea < MinArea)
-		{
-			MinArea = CurrentArea;
-			RectSideA = SupportVectorA * (MaxDotResultA - MinDotResultA);
-			RectSideB = SupportVectorB * (MaxDotResultB - MinDotResultB);
-		}
-	}
-
-	RectSideA = SurfaceNormalMatrix.TransformVector(RectSideA);
-	RectSideB = SurfaceNormalMatrix.TransformVector(RectSideB);
-	OutRectRotation = FRotationMatrix::MakeFromZX(PolyNormal, RectSideA).Rotator();
-	OutSideLengthX = RectSideA.Size();
-	OutSideLengthY = RectSideB.Size();
-
-#if ENABLE_DRAW_DEBUG
-	if (bDebugDraw)
-	{
-		if (UWorld* World = GEngine->GetWorldFromContextObject(WorldContextObject, EGetWorldErrorMode::LogAndReturnNull))
-		{
-			DrawDebugSphere(World, OutRectCenter, 10.f, 12, FColor::Yellow, true);
-			DrawDebugCoordinateSystem(World, OutRectCenter, SurfaceNormalMatrix.Rotator(), 100.f, true);
-			DrawDebugLine(World, OutRectCenter - RectSideA * 0.5f + FVector(0, 0, 10.f), OutRectCenter + RectSideA * 0.5f + FVector(0, 0, 10.f), FColor::Green, true, -1, 0, 5.f);
-			DrawDebugLine(World, OutRectCenter - RectSideB * 0.5f + FVector(0, 0, 10.f), OutRectCenter + RectSideB * 0.5f + FVector(0, 0, 10.f), FColor::Blue, true, -1, 0, 5.f);
-		}
-	}
-#endif
-}
-
 void UKismetMathLibrary::MinAreaRectangle(class UObject* WorldContextObject, const TArray<FVector>& InPoints, const FVector& SampleSurfaceNormal,
                                           FVector& OutRectCenter, FRotator& OutRectRotation, float& OutRectLengthX, float& OutRectLengthY, bool bDebugDraw)
 {
