Fixed duplicate LOCTEXT namespace

#rnx

diff --git a/Engine/Plugins/FX/Niagara/Source/Niagara/Private/NiagaraDataInterfaceOcclusion.cpp b/Engine/Plugins/FX/Niagara/Source/Niagara/Private/NiagaraDataInterfaceOcclusion.cpp
--- a/Engine/Plugins/FX/Niagara/Source/Niagara/Private/NiagaraDataInterfaceOcclusion.cpp
+++ b/Engine/Plugins/FX/Niagara/Source/Niagara/Private/NiagaraDataInterfaceOcclusion.cpp
@@ -111,7 +111,7 @@ void UNiagaraDataInterfaceOcclusion::GetFunctionsInternal(TArray<FNiagaraFunctio
 	{
 		FNiagaraFunctionSignature& Sig = OutFunctions.Add_GetRef(DefaultSig);
 		Sig.Name = QueryOcclusionFactorWithCircleName;
-		Sig.SetDescription(LOCTEXT("QueryCloudOcclusionWithCircleDescription", "This function returns the occlusion factor of a sprite. It samples the depth buffer in a spiral pattern around the given world position and compares each sample with the camera distance."));
+		Sig.SetDescription(LOCTEXT("QueryCameraOcclusionCircleDescription", "This function returns the occlusion factor of a sprite. It samples the depth buffer in a spiral pattern around the given world position and compares each sample with the camera distance."));
 		Sig.Inputs.Emplace(FNiagaraTypeDefinition::GetPositionDef(), TEXT("WorldPosition"));
 		Sig.Inputs.Emplace(FNiagaraTypeDefinition::GetFloatDef(), TEXT("WorldDiameter"));
 		Sig.Inputs.Emplace_GetRef(FNiagaraTypeDefinition::GetBoolDef(), TEXT("IncludeCenterSample")).SetValue(true);
