- Fix for SpawnProbability distribution not updating
- Evaluate SpawnProbability & SpawnRate once per loop for rate based spawning
- Add some tooltips to stateless spawn infos

#rnx
#tests Tested out spawn probability and random ranges with burst / rate.  Added ET to cover some basic cases.

#virtualized

diff --git a/Engine/Plugins/FX/Niagara/Source/Niagara/Internal/Stateless/NiagaraStatelessEmitterInstance.h b/Engine/Plugins/FX/Niagara/Source/Niagara/Internal/Stateless/NiagaraStatelessEmitterInstance.h
--- a/Engine/Plugins/FX/Niagara/Source/Niagara/Internal/Stateless/NiagaraStatelessEmitterInstance.h
+++ b/Engine/Plugins/FX/Niagara/Source/Niagara/Internal/Stateless/NiagaraStatelessEmitterInstance.h
@@ -46,8 +46,9 @@ class FNiagaraStatelessEmitterInstance final : public FNiagaraEmitterInstance
 
 	struct FActiveSpawnRate
 	{
-		float	Rate = 0.0f;
-		float	ResidualSpawnTime = 0.0f;
+		FNiagaraDistributionRangeFloat				SpawnRate = FNiagaraDistributionRangeFloat(0.0f);
+		float										ResidualSpawnTime = 0.0f;
+		TOptional<FNiagaraDistributionRangeFloat>	SpawnProbability;
 	};
 
 public:
diff --git a/Engine/Plugins/FX/Niagara/Source/Niagara/Internal/Stateless/NiagaraStatelessSpawnInfo.h b/Engine/Plugins/FX/Niagara/Source/Niagara/Internal/Stateless/NiagaraStatelessSpawnInfo.h
--- a/Engine/Plugins/FX/Niagara/Source/Niagara/Internal/Stateless/NiagaraStatelessSpawnInfo.h
+++ b/Engine/Plugins/FX/Niagara/Source/Niagara/Internal/Stateless/NiagaraStatelessSpawnInfo.h
@@ -27,12 +27,16 @@ struct FNiagaraStatelessSpawnInfo
 	UPROPERTY(EditAnywhere, Category = "Spawn", meta = (ShowInStackItemHeader, StackItemHeaderAlignment = "Left"))
 	ENiagaraStatelessSpawnInfoType	Type = ENiagaraStatelessSpawnInfoType::Burst;
 
+	// Time to spawn particles at
 	UPROPERTY(EditAnywhere, Category = "Spawn", meta = (EditConditionHides, EditCondition = "Type == ENiagaraStatelessSpawnInfoType::Burst"))
 	float SpawnTime = 0.0f;
 
+	// Number of particles to spawn at the provided time
 	UPROPERTY(EditAnywhere, Category = "Spawn", meta = (EditConditionHides, EditCondition = "Type == ENiagaraStatelessSpawnInfoType::Burst", ClampMin = "0"))
 	FNiagaraDistributionRangeInt Amount = FNiagaraDistributionRangeInt(1);
 
+	// Number of particles to spawn per second
+	// Note: In the case of a random range or binding the value is only evaluated at the start of each loop.  This varies from regular emitters.
 	UPROPERTY(EditAnywhere, Category = "Spawn", meta = (EditConditionHides, EditCondition = "Type == ENiagaraStatelessSpawnInfoType::Rate", ClampMin = "0.0"))
 	FNiagaraDistributionRangeFloat Rate = FNiagaraDistributionRangeFloat(60.0f);
 
@@ -42,6 +46,9 @@ struct FNiagaraStatelessSpawnInfo
 	UPROPERTY(EditAnywhere, Category = "Spawn", meta = (InlineEditConditionToggle))
 	uint8 bSpawnProbabilityEnabled : 1 = false;
 
+	// 0 - 1 value that can be viewed as a percentage chance that the spawn will generated particles or not.
+	// A value of 0.5 can be viewed as a 50% chance that the spawn will trigger.
+	// Note: Rate spawning will only evaluate the probability at the start of each loop.  This varies from regular emitters which evaluate each time particles are spawned.
 	UPROPERTY(EditAnywhere, Category = "Spawn", meta = (EditCondition = "bSpawnProbabilityEnabled", ClampMin = "0.0", ClampMax = "1.0", UIMin = "0.0", UIMax = "1.0"))
 	FNiagaraDistributionRangeFloat SpawnProbability = FNiagaraDistributionRangeFloat(1.0f);
 
diff --git a/Engine/Plugins/FX/Niagara/Source/Niagara/Private/Stateless/NiagaraStatelessEmitterInstance.cpp b/Engine/Plugins/FX/Niagara/Source/Niagara/Private/Stateless/NiagaraStatelessEmitterInstance.cpp
--- a/Engine/Plugins/FX/Niagara/Source/Niagara/Private/Stateless/NiagaraStatelessEmitterInstance.cpp
+++ b/Engine/Plugins/FX/Niagara/Source/Niagara/Private/Stateless/NiagaraStatelessEmitterInstance.cpp
@@ -526,12 +526,11 @@ void FNiagaraStatelessEmitterInstance::InitSpawnInfos(float InitializationAge)
 	{
 		if (SpawnInfo.Type == ENiagaraStatelessSpawnInfoType::Rate)
 		{
-			float SpawnRate = EvaluateDistribution(SpawnInfo.Rate, RandomStream, RendererBindings, DefaultSpawnRate) * EmitterData->SpawnCountScale;
-			SpawnRate = FMath::Max(SpawnRate, 0.0f);
-			if (SpawnRate > 0.0f)
+			FActiveSpawnRate& ActiveSpawnRate	= ActiveSpawnRates.AddDefaulted_GetRef();
+			ActiveSpawnRate.SpawnRate			= SpawnInfo.Rate;
+			if (SpawnInfo.bSpawnProbabilityEnabled)
 			{
-				FActiveSpawnRate& ActiveSpawnRate = ActiveSpawnRates.AddDefaulted_GetRef();
-				ActiveSpawnRate.Rate = SpawnRate;
+				ActiveSpawnRate.SpawnProbability = SpawnInfo.SpawnProbability;
 			}
 		}
 	}
@@ -552,10 +551,27 @@ void FNiagaraStatelessEmitterInstance::InitSpawnInfosForLoop(float Initializatio
 	// Add the next chunk for any active spawn rates
 	for (FActiveSpawnRate& SpawnInfo : ActiveSpawnRates)
 	{
+		// Unlike stateful emitters we evaluate the spawn probability & rate per loop
+		if (SpawnInfo.SpawnProbability.IsSet())
+		{
+			float SpawnProbability = EvaluateDistribution(SpawnInfo.SpawnProbability.GetValue(), RandomStream, RendererBindings, DefaultSpawnProbability);
+			SpawnProbability = FMath::Clamp(SpawnProbability, 0.0f, 1.0f);
+			if (SpawnProbability < RandomStream.FRand())
+			{
+				continue;
+			}
+		}
+
+		const float SpawnRate = EvaluateDistribution(SpawnInfo.SpawnRate, RandomStream, RendererBindings, DefaultSpawnRate) * EmitterData->SpawnCountScale;
+		if (SpawnRate <= 0.0f)
+		{
+			continue;
+		}
+
 		const float SpawnAgeStart	= FMath::Min(InitializationAge + CurrentLoopDelay - SpawnInfo.ResidualSpawnTime, CurrentLoopAgeEnd);
 		const float ActiveDuration	= CurrentLoopAgeEnd - SpawnAgeStart;
-		const int32 NumSpawned		= FMath::FloorToInt(ActiveDuration * SpawnInfo.Rate);
-		const float SpawnAgeEnd		= SpawnAgeStart + (float(NumSpawned) / SpawnInfo.Rate);
+		const int32 NumSpawned		= FMath::FloorToInt(ActiveDuration * SpawnRate);
+		const float SpawnAgeEnd		= SpawnAgeStart + (float(NumSpawned) / SpawnRate);
 
 		if ( NumSpawned > 0 )
 		{
@@ -565,7 +581,7 @@ void FNiagaraStatelessEmitterInstance::InitSpawnInfosForLoop(float Initializatio
 			if ( SpawnInfos.Num() > 0 )
 			{
 				FNiagaraStatelessRuntimeSpawnInfo& ExistingInfo = SpawnInfos.Last();
-				if ( (ExistingInfo.Type == ENiagaraStatelessSpawnInfoType::Rate) && (ExistingInfo.Rate == SpawnInfo.Rate) && (ExistingInfo.SpawnTimeEnd == SpawnAgeStart))
+				if ( (ExistingInfo.Type == ENiagaraStatelessSpawnInfoType::Rate) && (ExistingInfo.Rate == SpawnRate) && (ExistingInfo.SpawnTimeEnd == SpawnAgeStart))
 				{
 					if (ExistingInfo.UniqueOffset + ExistingInfo.Amount == UniqueIndexOffset)
 					{
@@ -584,7 +600,7 @@ void FNiagaraStatelessEmitterInstance::InitSpawnInfosForLoop(float Initializatio
 				NewSpawnInfo.UniqueOffset	= UniqueIndexOffset;
 				NewSpawnInfo.SpawnTimeStart	= SpawnAgeStart;
 				NewSpawnInfo.SpawnTimeEnd	= SpawnAgeEnd;
-				NewSpawnInfo.Rate			= SpawnInfo.Rate;
+				NewSpawnInfo.Rate			= SpawnRate;
 				NewSpawnInfo.Amount			= NumSpawned;
 			}
 
diff --git a/Engine/Plugins/FX/Niagara/Source/NiagaraEditor/Private/ViewModels/Stack/NiagaraStackStatelessEmitterSpawnGroup.cpp b/Engine/Plugins/FX/Niagara/Source/NiagaraEditor/Private/ViewModels/Stack/NiagaraStackStatelessEmitterSpawnGroup.cpp
--- a/Engine/Plugins/FX/Niagara/Source/NiagaraEditor/Private/ViewModels/Stack/NiagaraStackStatelessEmitterSpawnGroup.cpp
+++ b/Engine/Plugins/FX/Niagara/Source/NiagaraEditor/Private/ViewModels/Stack/NiagaraStackStatelessEmitterSpawnGroup.cpp
@@ -349,6 +349,7 @@ void UNiagaraStackStatelessEmitterSpawnItem::OnSpawnInfoModified(TArray<UObject*
 	{
 		FNiagaraStatelessSpawnInfo* SpawnInfo = GetSpawnInfo();
 		SpawnInfo->Rate.UpdateValuesFromDistribution();
+		SpawnInfo->SpawnProbability.UpdateValuesFromDistribution();
 
 		FPropertyChangedEvent EmptyPropertyUpdateStruct(nullptr);
 		StatelessEmitter->PostEditChangeProperty(EmptyPropertyUpdateStruct);
