- Cast INDEX_NONE to uint16 to fix narrowing error

#rnx
#tests Ran ET where it previously triggered

diff --git a/Engine/Plugins/FX/Niagara/Source/Niagara/Private/NiagaraSimCacheHelper.h b/Engine/Plugins/FX/Niagara/Source/Niagara/Private/NiagaraSimCacheHelper.h
--- a/Engine/Plugins/FX/Niagara/Source/Niagara/Private/NiagaraSimCacheHelper.h
+++ b/Engine/Plugins/FX/Niagara/Source/Niagara/Private/NiagaraSimCacheHelper.h
@@ -617,17 +617,17 @@ struct FNiagaraSimCacheHelper
 
 			for (int32 i = 0; i < SourceVariable.FloatCount; ++i)
 			{
-				CacheBufferReadInfo.ComponentMappingsToDataBuffer[FloatOffset++] = IntCastChecked<uint16>(DestVariableLayout ? DestVariableLayout->GetFloatComponentStart() + i : INDEX_NONE);
+				CacheBufferReadInfo.ComponentMappingsToDataBuffer[FloatOffset++] = IntCastChecked<uint16>(DestVariableLayout ? DestVariableLayout->GetFloatComponentStart() + i : uint16(INDEX_NONE));
 			}
 
 			for (int32 i = 0; i < SourceVariable.HalfCount; ++i)
 			{
-				CacheBufferReadInfo.ComponentMappingsToDataBuffer[HalfOffset++] = IntCastChecked<uint16>(DestVariableLayout ? DestVariableLayout->GetHalfComponentStart() + i : INDEX_NONE);
+				CacheBufferReadInfo.ComponentMappingsToDataBuffer[HalfOffset++] = IntCastChecked<uint16>(DestVariableLayout ? DestVariableLayout->GetHalfComponentStart() + i : uint16(INDEX_NONE));
 			}
 
 			for (int32 i = 0; i < SourceVariable.Int32Count; ++i)
 			{
-				CacheBufferReadInfo.ComponentMappingsToDataBuffer[Int32Offset++] = IntCastChecked<uint16>(DestVariableLayout ? DestVariableLayout->GetInt32ComponentStart() + i : INDEX_NONE);
+				CacheBufferReadInfo.ComponentMappingsToDataBuffer[Int32Offset++] = IntCastChecked<uint16>(DestVariableLayout ? DestVariableLayout->GetInt32ComponentStart() + i : uint16(INDEX_NONE));
 			}
 		}
 
