- Ensure ExecutionState visible to outside of stateless is consistent with stateful
- Currently we represent what is really the internal execution state, scalability can modify this (i.e. making it inactive)
- Without this change the cache, the HUD will show execution state that might be confusing to users

#rnx
#tests Scalability level testing with stateful vs stateless to confirm they respond the same

diff --git a/Engine/Plugins/FX/Niagara/Source/Niagara/Internal/Stateless/NiagaraStatelessEmitterInstance.h b/Engine/Plugins/FX/Niagara/Source/Niagara/Internal/Stateless/NiagaraStatelessEmitterInstance.h
--- a/Engine/Plugins/FX/Niagara/Source/Niagara/Internal/Stateless/NiagaraStatelessEmitterInstance.h
+++ b/Engine/Plugins/FX/Niagara/Source/Niagara/Internal/Stateless/NiagaraStatelessEmitterInstance.h
@@ -115,6 +115,7 @@ private:
 	TArray<FNiagaraStatelessRuntimeSpawnInfo>	SpawnInfos;
 	TArray<FActiveSpawnRate>					ActiveSpawnRates;
 
+	ENiagaraExecutionState						InternalExecutionState = ENiagaraExecutionState::Active;
 	ENiagaraExecutionStateManagement			ScalabilityState = ENiagaraExecutionStateManagement::Awaken;
 	int32										LoopCount = 0;
 	float										CurrentLoopDuration = 0.0f;
diff --git a/Engine/Plugins/FX/Niagara/Source/Niagara/Private/Stateless/NiagaraStatelessEmitterInstance.cpp b/Engine/Plugins/FX/Niagara/Source/Niagara/Private/Stateless/NiagaraStatelessEmitterInstance.cpp
--- a/Engine/Plugins/FX/Niagara/Source/Niagara/Private/Stateless/NiagaraStatelessEmitterInstance.cpp
+++ b/Engine/Plugins/FX/Niagara/Source/Niagara/Private/Stateless/NiagaraStatelessEmitterInstance.cpp
@@ -87,7 +87,8 @@ void FNiagaraStatelessEmitterInstance::Init(int32 InEmitterIndex)
 	InitEmitterData();
 	if (!bCanEverExecute)
 	{
-		ExecutionState = ENiagaraExecutionState::Disabled;
+		InternalExecutionState = ENiagaraExecutionState::Disabled;
+		ExecutionState = InternalExecutionState;
 		return;
 	}
 
@@ -160,7 +161,8 @@ void FNiagaraStatelessEmitterInstance::ResetSimulation(bool bKillExisting)
 	InitEmitterState();
 	InitSpawnInfos(0.0f);
 
-	ExecutionState = ENiagaraExecutionState::Active;
+	InternalExecutionState = ENiagaraExecutionState::Active;
+	ExecutionState = InternalExecutionState;
 	ScalabilityState = ENiagaraExecutionStateManagement::Awaken;
 	if (NiagaraStateless::FEmitterInstance_RT* RenderThreadData = RenderThreadDataPtr.Get())
 	{
@@ -185,7 +187,8 @@ bool FNiagaraStatelessEmitterInstance::HandleCompletion(bool bForce)
 	bool bIsComplete = IsComplete();
 	if (!bIsComplete && bForce)
 	{
-		ExecutionState = ENiagaraExecutionState::Complete;
+		InternalExecutionState = ENiagaraExecutionState::Complete;
+		ExecutionState = InternalExecutionState;
 		bIsComplete = true;
 
 		if (NiagaraStateless::FEmitterInstance_RT* RenderThreadData = RenderThreadDataPtr.Get())
@@ -252,7 +255,7 @@ void FNiagaraStatelessEmitterInstance::UnbindParameters(bool bExternalOnly)
 
 bool FNiagaraStatelessEmitterInstance::ShouldTick() const
 {
-	return ExecutionState <= ENiagaraExecutionState::Inactive;
+	return InternalExecutionState <= ENiagaraExecutionState::Inactive;
 }
 
 void FNiagaraStatelessEmitterInstance::Tick(float DeltaSeconds)
@@ -286,14 +289,14 @@ void FNiagaraStatelessEmitterInstance::TickEmitterState()
 	// Update execution state based on the parent which be told to go inactive / complete
 	{
 		const ENiagaraExecutionState ParentExecutionState = ParentSystemInstance ? ParentSystemInstance->GetActualExecutionState() : ENiagaraExecutionState::Complete;
-		if (ParentExecutionState > ExecutionState)
+		if (ParentExecutionState > InternalExecutionState)
 		{
 			SetExecutionStateInternal(ParentExecutionState);
 		}
 	}
 
 	// If we are going inactive and we hit zero particles we are now complete
-	if (ExecutionState == ENiagaraExecutionState::Inactive)
+	if (InternalExecutionState == ENiagaraExecutionState::Inactive)
 	{
 		if (GetNumParticles() == 0)
 		{
@@ -302,7 +305,7 @@ void FNiagaraStatelessEmitterInstance::TickEmitterState()
 	}
 
 	// If we are not active we don't need to evaluate loops / scalability anymore
-	if ( ExecutionState != ENiagaraExecutionState::Active )
+	if (InternalExecutionState != ENiagaraExecutionState::Active )
 	{
 		return;
 	}
@@ -338,6 +341,7 @@ void FNiagaraStatelessEmitterInstance::TickEmitterState()
 		// We need to transition the state
 		if (RequestedScalabilityState != ScalabilityState)
 		{
+			ExecutionState = InternalExecutionState;
 			ScalabilityState = RequestedScalabilityState;
 			switch (RequestedScalabilityState)
 			{
@@ -350,10 +354,12 @@ void FNiagaraStatelessEmitterInstance::TickEmitterState()
 
 				case ENiagaraExecutionStateManagement::SleepAndLetParticlesFinish:
 				case ENiagaraExecutionStateManagement::KillAfterParticlesFinish:
+					ExecutionState = ENiagaraExecutionState::Inactive;
 					CropSpawnInfos();
 					break;
 
 				case ENiagaraExecutionStateManagement::SleepAndClearParticles:
+					ExecutionState = ENiagaraExecutionState::Inactive;
 					KillSpawnInfos();
 					break;
 
@@ -464,7 +470,7 @@ void FNiagaraStatelessEmitterInstance::SendRenderData()
 
 	FDataForRenderThread DataForRenderThread;
 	DataForRenderThread.Age				= Age;
-	DataForRenderThread.ExecutionState	= ExecutionState;
+	DataForRenderThread.ExecutionState	= InternalExecutionState;
 
 	if (RendererBindings.GetParametersDirty())
 	{
@@ -727,7 +733,7 @@ void FNiagaraStatelessEmitterInstance::RestartSpawnInfos()
 
 void FNiagaraStatelessEmitterInstance::SetExecutionStateInternal(ENiagaraExecutionState RequestedExecutionState)
 {
-	if (ExecutionState == RequestedExecutionState)
+	if (InternalExecutionState == RequestedExecutionState)
 	{
 		return;
 	}
@@ -742,19 +748,22 @@ void FNiagaraStatelessEmitterInstance::SetExecutionStateInternal(ENiagaraExecuti
 			if (EmitterData->EmitterState.InactiveResponse == ENiagaraEmitterInactiveResponse::Kill)
 			{
 				KillSpawnInfos();
-				ExecutionState = ENiagaraExecutionState::Complete;
+				InternalExecutionState = ENiagaraExecutionState::Complete;
+				ExecutionState = InternalExecutionState;
 			}
 			else
 			{
 				CropSpawnInfos();
-				ExecutionState = SpawnInfos.Num() > 0 ? ENiagaraExecutionState::Inactive : ENiagaraExecutionState::Complete;
+				InternalExecutionState = SpawnInfos.Num() > 0 ? ENiagaraExecutionState::Inactive : ENiagaraExecutionState::Complete;
+				ExecutionState = InternalExecutionState;
 			}
 			break;
 
 		case ENiagaraExecutionState::InactiveClear:
 		case ENiagaraExecutionState::Complete:
 			KillSpawnInfos();
-			ExecutionState = ENiagaraExecutionState::Complete;
+			InternalExecutionState = ENiagaraExecutionState::Complete;
+			ExecutionState = InternalExecutionState;
 			break;
 	}
 }
