- Add BlockUntilGPUIdle when transferring a dataset from GPU -> CPU
- This is a temporary fix where some RHIs are not performing a block before reading back the buffer

#rnx

diff --git a/Engine/Plugins/FX/Niagara/Source/Niagara/Classes/NiagaraDataSet.h b/Engine/Plugins/FX/Niagara/Source/Niagara/Classes/NiagaraDataSet.h
--- a/Engine/Plugins/FX/Niagara/Source/Niagara/Classes/NiagaraDataSet.h
+++ b/Engine/Plugins/FX/Niagara/Source/Niagara/Classes/NiagaraDataSet.h
@@ -98,7 +98,7 @@ public:
 	NIAGARA_API void CopyToUnrelated(FNiagaraDataBuffer& DestBuffer, int32 SrcStartIdx, int32 DestStartIdx, int32 NumInstances) const;
 	NIAGARA_API void GPUCopyFrom(const float* GPUReadBackFloat, const int* GPUReadBackInt, const FFloat16* GPUReadBackHalf, int32 StartIdx, int32 NumInstances, uint32 InSrcFloatStride, uint32 InSrcIntStride, uint32 InSrcHalfStride);
 	NIAGARA_API void PushCPUBuffersToGPU(const TArray<FNiagaraDataBufferRef>& SourceBuffers, bool bReleaseRef, FRHICommandList& RHICmdList, ERHIFeatureLevel::Type FeatureLevel, const TCHAR* DebugSimName);
-	NIAGARA_API void TransferGPUToCPUImmediate(FRHICommandList& RHICmdList, FNiagaraGpuComputeDispatchInterface* ComputeInterface, FNiagaraDataBuffer* CPUBuffer) const;
+	NIAGARA_API void TransferGPUToCPUImmediate(FRHICommandListImmediate& RHICmdList, FNiagaraGpuComputeDispatchInterface* ComputeInterface, FNiagaraDataBuffer* CPUBuffer) const;
 	NIAGARA_API void Dump(int32 StartIndex, int32 NumInstances, const FString& Label, const FName& SortParameterKey = FName()) const;
 
 	FORCEINLINE TArrayView<uint8 const* RESTRICT const> ReadRegisterTable() const { return TArrayView<uint8 const* RESTRICT const>(RegisterTable); }
diff --git a/Engine/Plugins/FX/Niagara/Source/Niagara/Private/NiagaraDataSet.cpp b/Engine/Plugins/FX/Niagara/Source/Niagara/Private/NiagaraDataSet.cpp
--- a/Engine/Plugins/FX/Niagara/Source/Niagara/Private/NiagaraDataSet.cpp
+++ b/Engine/Plugins/FX/Niagara/Source/Niagara/Private/NiagaraDataSet.cpp
@@ -1450,10 +1450,13 @@ void FNiagaraDataBuffer::PushCPUBuffersToGPU(const TArray<FNiagaraDataBufferRef>
 	}
 }
 
-void FNiagaraDataBuffer::TransferGPUToCPUImmediate(FRHICommandList& RHICmdList, FNiagaraGpuComputeDispatchInterface* ComputeInterface, FNiagaraDataBuffer* CPUBuffer) const
+void FNiagaraDataBuffer::TransferGPUToCPUImmediate(FRHICommandListImmediate& RHICmdList, FNiagaraGpuComputeDispatchInterface* ComputeInterface, FNiagaraDataBuffer* CPUBuffer) const
 {
 	check(GetOwner()->GetCompiledData().GetLayoutHash() == CPUBuffer->GetOwner()->GetCompiledData().GetLayoutHash());
 
+	// Note: On some RHIs we need block until idle or the lock opersations will have the wrong information
+	RHICmdList.BlockUntilGPUIdle();
+
 	uint32 GPUNumInstances = 0;
 	if (GPUInstanceCountBufferOffset != INDEX_NONE)
 	{
