- Debug HUD bounds display changes
- Allow inactive systems to display bounds, this is useful for debugging occlusion culled instances
- Move settings into drop down menu next to the button to make it more obvious

#rnx
#tests Scalability tests with bounds enabled

diff --git a/Engine/Plugins/FX/Niagara/Source/Niagara/Private/NiagaraDebugHud.cpp b/Engine/Plugins/FX/Niagara/Source/Niagara/Private/NiagaraDebugHud.cpp
--- a/Engine/Plugins/FX/Niagara/Source/Niagara/Private/NiagaraDebugHud.cpp
+++ b/Engine/Plugins/FX/Niagara/Source/Niagara/Private/NiagaraDebugHud.cpp
@@ -140,8 +140,12 @@ namespace NiagaraDebugLocal
 		MakeTuple(TEXT("PerfGraphMode="), TEXT("Change the mode of the perf graph"), [](FString Arg) {Settings.PerfGraphMode = (ENiagaraDebugHUDPerfGraphMode)FCString::Atoi(*Arg); }),
 		MakeTuple(TEXT("PerfGraphTimeRange="), TEXT("Time range for the Y axis of the perf graph"), [](FString Arg) {Settings.PerfGraphTimeRange = FCString::Atof(*Arg); }),
 
+		// Bounds commands
+		MakeTuple(TEXT("DrawBoundsEnabled="), TEXT("Draw the component bounds"), [](FString Arg) {Settings.bDrawBoundsEnabled = FCString::Atoi(*Arg) != 0; }),
+		MakeTuple(TEXT("DrawBoundsWireframe="), TEXT("Draw the component bounds"), [](FString Arg) {Settings.bDrawBoundsWireframe = FCString::Atoi(*Arg) != 0; }),
+		MakeTuple(TEXT("DrawBoundsAlpha="), TEXT("Draw the component bounds"), [](FString Arg) {Settings.DrawBoundsAlpha = FCString::Atof(*Arg); }),
+
 		// System commands
-		MakeTuple(TEXT("SystemShowBounds="), TEXT("Show system bounds"), [](FString Arg) {Settings.bSystemShowBounds = FCString::Atoi(*Arg) != 0; }),
 		MakeTuple(TEXT("SystemShowActiveOnlyInWorld="), TEXT("When enabled only active systems are shown in world"), [](FString Arg) {Settings.bSystemShowActiveOnlyInWorld = FCString::Atoi(*Arg) != 0; }),
 		MakeTuple(TEXT("SystemDebugVerbosity="), TEXT("Set the in world system debug verbosity"), [](FString Arg) {Settings.SystemDebugVerbosity = FMath::Clamp(ENiagaraDebugHudVerbosity(FCString::Atoi(*Arg)), ENiagaraDebugHudVerbosity::None, ENiagaraDebugHudVerbosity::Verbose); }),
 		MakeTuple(TEXT("SystemEmitterVerbosity="), TEXT("Set the in world system emitter debug verbosity"), [](FString Arg) {Settings.SystemEmitterVerbosity = FMath::Clamp(ENiagaraDebugHudVerbosity(FCString::Atoi(*Arg)), ENiagaraDebugHudVerbosity::None, ENiagaraDebugHudVerbosity::Verbose); }),
@@ -600,34 +604,34 @@ namespace NiagaraDebugLocal
 		return TPair<FVector2f, FVector2f>(StringSize, OutLocation);
 	}
 
-	void DrawBox(UWorld* World, const FVector& Location, const FVector& Extents, const FLinearColor& Color, float SolidAlpha = 0.0f, float Thickness = 3.0f)
+	void DrawWireframeBox(UWorld* World, const FVector& Location, const FVector& Extents, const FLinearColor& Color, float Thickness = 3.0f)
 	{
 		if (ULineBatchComponent* LineBatcher = World->LineBatcher)
 		{
-			if (SolidAlpha > 0.0f)
-			{
-				const FBox BoundsBox(-Extents, Extents);
-				FColor BoxColor = Color.ToFColor(false);
-				BoxColor.A = uint8(FMath::Clamp(int32(SolidAlpha * 255.0f), 0, 255));
-				LineBatcher->DrawSolidBox(BoundsBox, FTransform(FQuat::Identity, Location), BoxColor, 0, 0.0f);
-			}
-			else
-			{
-				LineBatcher->DrawLine(Location + FVector(Extents.X, Extents.Y, Extents.Z), Location + FVector(Extents.X, -Extents.Y, Extents.Z), Color, 0, Thickness);
-				LineBatcher->DrawLine(Location + FVector(Extents.X, -Extents.Y, Extents.Z), Location + FVector(-Extents.X, -Extents.Y, Extents.Z), Color, 0, Thickness);
-				LineBatcher->DrawLine(Location + FVector(-Extents.X, -Extents.Y, Extents.Z), Location + FVector(-Extents.X, Extents.Y, Extents.Z), Color, 0, Thickness);
-				LineBatcher->DrawLine(Location + FVector(-Extents.X, Extents.Y, Extents.Z), Location + FVector(Extents.X, Extents.Y, Extents.Z), Color, 0, Thickness);
+			LineBatcher->DrawLine(Location + FVector(Extents.X, Extents.Y, Extents.Z), Location + FVector(Extents.X, -Extents.Y, Extents.Z), Color, 0, Thickness);
+			LineBatcher->DrawLine(Location + FVector(Extents.X, -Extents.Y, Extents.Z), Location + FVector(-Extents.X, -Extents.Y, Extents.Z), Color, 0, Thickness);
+			LineBatcher->DrawLine(Location + FVector(-Extents.X, -Extents.Y, Extents.Z), Location + FVector(-Extents.X, Extents.Y, Extents.Z), Color, 0, Thickness);
+			LineBatcher->DrawLine(Location + FVector(-Extents.X, Extents.Y, Extents.Z), Location + FVector(Extents.X, Extents.Y, Extents.Z), Color, 0, Thickness);
 
-				LineBatcher->DrawLine(Location + FVector(Extents.X, Extents.Y, -Extents.Z), Location + FVector(Extents.X, -Extents.Y, -Extents.Z), Color, 0, Thickness);
-				LineBatcher->DrawLine(Location + FVector(Extents.X, -Extents.Y, -Extents.Z), Location + FVector(-Extents.X, -Extents.Y, -Extents.Z), Color, 0, Thickness);
-				LineBatcher->DrawLine(Location + FVector(-Extents.X, -Extents.Y, -Extents.Z), Location + FVector(-Extents.X, Extents.Y, -Extents.Z), Color, 0, Thickness);
-				LineBatcher->DrawLine(Location + FVector(-Extents.X, Extents.Y, -Extents.Z), Location + FVector(Extents.X, Extents.Y, -Extents.Z), Color, 0, Thickness);
+			LineBatcher->DrawLine(Location + FVector(Extents.X, Extents.Y, -Extents.Z), Location + FVector(Extents.X, -Extents.Y, -Extents.Z), Color, 0, Thickness);
+			LineBatcher->DrawLine(Location + FVector(Extents.X, -Extents.Y, -Extents.Z), Location + FVector(-Extents.X, -Extents.Y, -Extents.Z), Color, 0, Thickness);
+			LineBatcher->DrawLine(Location + FVector(-Extents.X, -Extents.Y, -Extents.Z), Location + FVector(-Extents.X, Extents.Y, -Extents.Z), Color, 0, Thickness);
+			LineBatcher->DrawLine(Location + FVector(-Extents.X, Extents.Y, -Extents.Z), Location + FVector(Extents.X, Extents.Y, -Extents.Z), Color, 0, Thickness);
 
-				LineBatcher->DrawLine(Location + FVector(Extents.X, Extents.Y, Extents.Z), Location + FVector(Extents.X, Extents.Y, -Extents.Z), Color, 0, Thickness);
-				LineBatcher->DrawLine(Location + FVector(Extents.X, -Extents.Y, Extents.Z), Location + FVector(Extents.X, -Extents.Y, -Extents.Z), Color, 0, Thickness);
-				LineBatcher->DrawLine(Location + FVector(-Extents.X, -Extents.Y, Extents.Z), Location + FVector(-Extents.X, -Extents.Y, -Extents.Z), Color, 0, Thickness);
-				LineBatcher->DrawLine(Location + FVector(-Extents.X, Extents.Y, Extents.Z), Location + FVector(-Extents.X, Extents.Y, -Extents.Z), Color, 0, Thickness);
-			}
+			LineBatcher->DrawLine(Location + FVector(Extents.X, Extents.Y, Extents.Z), Location + FVector(Extents.X, Extents.Y, -Extents.Z), Color, 0, Thickness);
+			LineBatcher->DrawLine(Location + FVector(Extents.X, -Extents.Y, Extents.Z), Location + FVector(Extents.X, -Extents.Y, -Extents.Z), Color, 0, Thickness);
+			LineBatcher->DrawLine(Location + FVector(-Extents.X, -Extents.Y, Extents.Z), Location + FVector(-Extents.X, -Extents.Y, -Extents.Z), Color, 0, Thickness);
+			LineBatcher->DrawLine(Location + FVector(-Extents.X, Extents.Y, Extents.Z), Location + FVector(-Extents.X, Extents.Y, -Extents.Z), Color, 0, Thickness);
+		}
+	}
+
+	void DrawBox(UWorld* World, const FVector& Location, const FVector& Extents, const FLinearColor& Color)
+	{
+		if (ULineBatchComponent* LineBatcher = World->LineBatcher)
+		{
+			const FBox BoundsBox(-Extents, Extents);
+			FColor BoxColor = Color.ToFColor(false);
+			LineBatcher->DrawSolidBox(BoundsBox, FTransform(FQuat::Identity, Location), BoxColor, 0, 0.0f);
 		}
 	}
 
@@ -2808,12 +2812,22 @@ void FNiagaraDebugHud::DrawComponents(FNiagaraWorldManager* WorldManager, UCanva
 		const bool bIsActive = FXComponent->IsActive();
 
 		// Show system bounds (only active components)
-		if (Settings.bSystemShowBounds && bIsActive)
+		if (Settings.bDrawBoundsEnabled)
 		{
-			const FBox Bounds = FXComponent->CalcBounds(FXComponent->GetComponentTransform()).GetBox();
-			if (Bounds.IsValid)
+			const FBox BoundsToDraw = FXComponent->CalcBounds(FXComponent->GetComponentTransform()).GetBox();
+			if (BoundsToDraw.IsValid)
 			{
-				DrawBox(World, Bounds.GetCenter(), Bounds.GetExtent(), FColor::Red, Settings.SystemBoundsSolidBoxAlpha);
+				FLinearColor BoundsColor = bIsActive ? FLinearColor::Red : FLinearColor::Black;
+				BoundsColor.A = Settings.DrawBoundsAlpha;
+
+				if (Settings.bDrawBoundsWireframe)
+				{
+					DrawWireframeBox(World, BoundsToDraw.GetCenter(), BoundsToDraw.GetExtent(), BoundsColor);
+				}
+				else
+				{
+					DrawBox(World, BoundsToDraw.GetCenter(), BoundsToDraw.GetExtent(), BoundsColor);
+				}
 			}
 		}
 
diff --git a/Engine/Plugins/FX/Niagara/Source/Niagara/Public/NiagaraDebuggerCommon.h b/Engine/Plugins/FX/Niagara/Source/Niagara/Public/NiagaraDebuggerCommon.h
--- a/Engine/Plugins/FX/Niagara/Source/Niagara/Public/NiagaraDebuggerCommon.h
+++ b/Engine/Plugins/FX/Niagara/Source/Niagara/Public/NiagaraDebuggerCommon.h
@@ -471,10 +471,6 @@ struct FNiagaraDebugHUDSettingsData
 	UPROPERTY(VisibleAnywhere, Category = "Debug Overview")
 	bool bOverviewEnabled = false;
 
-	/** When enabled will show the system bounds for all filtered systems. */
-	UPROPERTY(VisibleAnywhere, Category = "Debug Overview")
-	bool bSystemShowBounds = false;
-
 	UPROPERTY(EditAnywhere, Category = "Debug Overview", meta = (DisplayName = "Debug Overview Mode", EditCondition = "bOverviewEnabled"))
 	ENiagaraDebugHUDOverviewMode OverviewMode = ENiagaraDebugHUDOverviewMode::Overview;
 
@@ -719,9 +715,14 @@ struct FNiagaraDebugHUDSettingsData
 	UPROPERTY(EditAnywhere, Category = "Style", meta = (EditCondition = "bShowParticleVariables"))
 	FNiagaraDebugHudTextOptions ParticleTextOptions;
 
-	/** When bounds display is enabled allows you to draw a solid box if alpha is > 0. */
-	UPROPERTY(EditAnywhere, Category = "Style", meta = (EditCondition = "bSystemShowBounds"))
-	float SystemBoundsSolidBoxAlpha = 0.0f;
+	UPROPERTY()
+	bool bDrawBoundsEnabled = false;
+
+	UPROPERTY()
+	bool bDrawBoundsWireframe = true;
+
+	UPROPERTY()
+	float DrawBoundsAlpha = 1.0f;
 
 	// Default background color used generally for panels
 	UPROPERTY(EditAnywhere, Category = "Style")
diff --git a/Engine/Plugins/FX/Niagara/Source/NiagaraEditor/Private/Customizations/NiagaraDebugHUDCustomization.cpp b/Engine/Plugins/FX/Niagara/Source/NiagaraEditor/Private/Customizations/NiagaraDebugHUDCustomization.cpp
--- a/Engine/Plugins/FX/Niagara/Source/NiagaraEditor/Private/Customizations/NiagaraDebugHUDCustomization.cpp
+++ b/Engine/Plugins/FX/Niagara/Source/NiagaraEditor/Private/Customizations/NiagaraDebugHUDCustomization.cpp
@@ -247,7 +247,6 @@ void FNiagaraDebugHUDSettingsDetailsCustomization::CustomizeDetails(IDetailLayou
 		IDetailCategoryBuilder& OverviewCategory = DetailBuilder.EditCategory("Debug Overview");
 		DetailBuilder.HideProperty(DetailBuilder.GetProperty(GET_MEMBER_NAME_CHECKED(FNiagaraDebugHUDSettingsData, bValidationEnabled), FNiagaraDebugHUDSettingsData::StaticStruct()));
 		DetailBuilder.HideProperty(DetailBuilder.GetProperty(GET_MEMBER_NAME_CHECKED(FNiagaraDebugHUDSettingsData, bOverviewEnabled), FNiagaraDebugHUDSettingsData::StaticStruct()));
-		DetailBuilder.HideProperty(DetailBuilder.GetProperty(GET_MEMBER_NAME_CHECKED(FNiagaraDebugHUDSettingsData, bSystemShowBounds), FNiagaraDebugHUDSettingsData::StaticStruct()));
 	}
 	
 	// Customize Filters
diff --git a/Engine/Plugins/FX/Niagara/Source/NiagaraEditor/Private/Widgets/SNiagaraDebugger.cpp b/Engine/Plugins/FX/Niagara/Source/NiagaraEditor/Private/Widgets/SNiagaraDebugger.cpp
--- a/Engine/Plugins/FX/Niagara/Source/NiagaraEditor/Private/Widgets/SNiagaraDebugger.cpp
+++ b/Engine/Plugins/FX/Niagara/Source/NiagaraEditor/Private/Widgets/SNiagaraDebugger.cpp
@@ -24,6 +24,7 @@
 #include "Widgets/Input/SComboButton.h"
 #include "Widgets/Input/SNumericEntryBox.h"
 #include "Widgets/Input/SEditableTextBox.h"
+#include "Widgets/Input/SSpinBox.h"
 #include "Widgets/Layout/SScrollBox.h"
 #include "Widgets/Layout/SUniformGridPanel.h"
 #include "Widgets/SBoxPanel.h"
@@ -1177,9 +1178,9 @@ TSharedRef<SWidget> SNiagaraDebugger::MakeToolbar()
 	{
 		ToolbarBuilder.AddToolBarButton(
 			FUIAction(
-				FExecuteAction::CreateLambda([=]() {Settings->Data.bSystemShowBounds = !Settings->Data.bSystemShowBounds; Settings->NotifyPropertyChanged(); }),
+				FExecuteAction::CreateLambda([=]() {Settings->Data.bDrawBoundsEnabled = !Settings->Data.bDrawBoundsEnabled; Settings->NotifyPropertyChanged(); }),
 				HudEnabledAction,
-				FIsActionChecked::CreateLambda([=]() { return Settings->Data.bSystemShowBounds; })
+				FIsActionChecked::CreateLambda([=]() { return Settings->Data.bDrawBoundsEnabled; })
 			),
 			NAME_None,
 			LOCTEXT("BoundsLabel", "Bounds"),
@@ -1187,6 +1188,15 @@ TSharedRef<SWidget> SNiagaraDebugger::MakeToolbar()
 			FSlateIcon(FAppStyle::Get().GetStyleSetName(), "AssetEditor.ToggleShowBounds"),
 			EUserInterfaceActionType::ToggleButton
 		);
+
+		ToolbarBuilder.AddComboButton(
+			FUIAction(),
+			FOnGetContent::CreateSP(this, &SNiagaraDebugger::MakeBoundsOptionsMenu),
+			FText(),
+			LOCTEXT("BoundsModeOptionsTooltip", "Additional options for how we display bounds."),
+			FSlateIcon(FAppStyle::GetAppStyleSetName(), "MaterialEditor.ToggleMaterialStats"),
+			true
+		);
 	}
 
 	ToolbarBuilder.AddSeparator();
@@ -1453,6 +1463,49 @@ TSharedRef<SWidget> SNiagaraDebugger::MakeModeOptionsMenu()
 	return MenuBuilder.MakeWidget();
 }
 
+TSharedRef<SWidget> SNiagaraDebugger::MakeBoundsOptionsMenu()
+{
+	FMenuBuilder MenuBuilder(/*bInShouldCloseWindowAfterMenuSelection=*/true, nullptr);
+	UNiagaraDebugHUDSettings* Settings = GetMutableDefault<UNiagaraDebugHUDSettings>();
+
+	MenuBuilder.BeginSection(NAME_None);
+	{
+		MenuBuilder.AddMenuEntry(
+			LOCTEXT("DrawBoundsWireframe", "Wireframe"),
+			LOCTEXT("DrawBoundsWireframeTooltip", "When enabled the bounds are draw as a wireframe box, when disabled they are a solid box."),
+			FSlateIcon(),
+			FUIAction(
+				FExecuteAction::CreateLambda([=]() {Settings->Data.bDrawBoundsWireframe = !Settings->Data.bDrawBoundsWireframe; Settings->NotifyPropertyChanged(); }),
+				FCanExecuteAction::CreateLambda([=]() { return Settings->Data.bDrawBoundsEnabled; }),
+				FIsActionChecked::CreateLambda([=]() { return Settings->Data.bDrawBoundsWireframe; })
+			),
+			NAME_None, EUserInterfaceActionType::ToggleButton
+		);
+
+		MenuBuilder.AddWidget(
+			SNew(SBox)
+				.Padding(FMargin(4.0f, 0.0f, 0.0f, 0.0f))
+				[
+					SNew(SSpinBox<float>)
+					.Font(FAppStyle::GetFontStyle(TEXT("MenuItem.Font")))
+					.MinValue(0.0f)
+					.MaxValue(1.0f)
+					.MinSliderValue(0.0f)
+					.MaxSliderValue(1.0f)
+					.Value(TAttribute<float>::CreateLambda([=]() { return Settings->Data.DrawBoundsAlpha; }))
+					.OnValueChanged(SSpinBox<float>::FOnValueChanged::CreateLambda([=](float NewValue) {Settings->Data.DrawBoundsAlpha = NewValue; Settings->NotifyPropertyChanged(); }))
+				],
+			LOCTEXT("DrawBoundsAlpha", "Alpha"),
+			false,
+			true,
+			LOCTEXT("DrawBoundsAlphaTooltip", "Controls the alpha channel for the bounds rendering.")
+		);
+	}
+	MenuBuilder.EndSection();
+
+	return MenuBuilder.MakeWidget();
+}
+
 #undef LOCTEXT_NAMESPACE
 
 #endif // WITH_NIAGARA_DEBUGGER
diff --git a/Engine/Plugins/FX/Niagara/Source/NiagaraEditor/Private/Widgets/SNiagaraDebugger.h b/Engine/Plugins/FX/Niagara/Source/NiagaraEditor/Private/Widgets/SNiagaraDebugger.h
--- a/Engine/Plugins/FX/Niagara/Source/NiagaraEditor/Private/Widgets/SNiagaraDebugger.h
+++ b/Engine/Plugins/FX/Niagara/Source/NiagaraEditor/Private/Widgets/SNiagaraDebugger.h
@@ -45,6 +45,7 @@ private:
 	TSharedRef<SWidget> MakeToolbar();
 	TSharedRef<SWidget> MakePlaybackOptionsMenu();
 	TSharedRef<SWidget> MakeModeOptionsMenu();
+	TSharedRef<SWidget> MakeBoundsOptionsMenu();
 
 protected:
 	TSharedPtr<FTabManager>		TabManager;
