Substrate - fixed niagara shader compilation issue when template is included multiple times.

#rnx

diff --git a/Engine/Plugins/FX/Niagara/Shaders/Private/NiagaraDataInterfaceGBufferTemplate.ush b/Engine/Plugins/FX/Niagara/Shaders/Private/NiagaraDataInterfaceGBufferTemplate.ush
--- a/Engine/Plugins/FX/Niagara/Shaders/Private/NiagaraDataInterfaceGBufferTemplate.ush
+++ b/Engine/Plugins/FX/Niagara/Shaders/Private/NiagaraDataInterfaceGBufferTemplate.ush
@@ -15,7 +15,7 @@ Texture2D {ParameterName}_VelocityTexture;
 SamplerState {ParameterName}_VelocityTextureSampler;
 
 #if SUBSTRATE_ENABLED
-FSubstrateGBuffer GetNiagaraSubstrateGBuffer(in float2 ScreenUV)
+FSubstrateGBuffer GetNiagaraSubstrateGBuffer_{ParameterName}(in float2 ScreenUV)
 {
 	uint2 ScreenPos = uint2(ScreenUV * View.BufferSizeAndInvSize.xy);
 	return SubstratePublic_GetSubstrateGBufferFromFirstBSDF(ScreenPos);
@@ -139,7 +139,7 @@ void DecodeBaseColor_{ParameterName}(in float2 ScreenUV, in bool bApplyViewportO
 	IsValid = true;
 
 	#if SUBSTRATE_ENABLED
-		OutValue = GetNiagaraSubstrateGBuffer(ScreenUV).BaseColor;
+		OutValue = GetNiagaraSubstrateGBuffer_{ParameterName}(ScreenUV).BaseColor;
 	#else
 		FGBufferData GBufferData = GetGBufferData(ScreenUV);
 		OutValue = GBufferData.BaseColor;
@@ -158,7 +158,7 @@ void DecodeMetallic_{ParameterName}(in float2 ScreenUV, in bool bApplyViewportOf
 	IsValid = true;
 
 	#if SUBSTRATE_ENABLED
-		OutValue = GetNiagaraSubstrateGBuffer(ScreenUV).Metallic;
+		OutValue = GetNiagaraSubstrateGBuffer_{ParameterName}(ScreenUV).Metallic;
 	#else
 		FGBufferData GBufferData = GetGBufferData(ScreenUV);
 		OutValue = GBufferData.Metallic;
@@ -177,7 +177,7 @@ void DecodeSpecular_{ParameterName}(in float2 ScreenUV, in bool bApplyViewportOf
 	IsValid = true;
 
 	#if SUBSTRATE_ENABLED
-		OutValue = GetNiagaraSubstrateGBuffer(ScreenUV).Specular;
+		OutValue = GetNiagaraSubstrateGBuffer_{ParameterName}(ScreenUV).Specular;
 	#else
 		FGBufferData GBufferData = GetGBufferData(ScreenUV);
 		OutValue = GBufferData.Specular;
@@ -299,7 +299,7 @@ void DecodeShadingModelID_{ParameterName}(in float2 ScreenUV, in bool bApplyView
 	IsValid = true;
 
 	#if SUBSTRATE_ENABLED
-		OutValue = int(GetNiagaraSubstrateGBuffer(ScreenUV).ShadingModelID);
+		OutValue = int(GetNiagaraSubstrateGBuffer_{ParameterName}(ScreenUV).ShadingModelID);
 	#else
 		FGBufferData GBufferData = GetGBufferData(ScreenUV);
 		OutValue = GBufferData.ShadingModelID;
