Some improvements to the HWRT async traces within Niagara.  Adds support for inline HW traces, where supported, through a compute shader (in some artificial tests it results in 50% performance improvement, but results will vary).  Also fixes up collision group masking.

diff --git a/Engine/Plugins/FX/Niagara/Shaders/Private/NiagaraRayTracingShaders.usf b/Engine/Plugins/FX/Niagara/Shaders/Private/NiagaraRayTracingShaders.usf
--- a/Engine/Plugins/FX/Niagara/Shaders/Private/NiagaraRayTracingShaders.usf
+++ b/Engine/Plugins/FX/Niagara/Shaders/Private/NiagaraRayTracingShaders.usf
@@ -15,7 +15,7 @@ uint											RaysOffset;
 RWStructuredBuffer<FNiagaraAsyncGpuTraceResult>	CollisionOutput;
 uint											CollisionOutputOffset;
 
-#if NIAGARA_RAYTRACE_FAKE_INDIRECT
+#if NIAGARA_RAYTRACE_FAKE_INDIRECT || COMPUTESHADER
 Buffer<uint>                                  RayTraceCounts;
 #endif
 
@@ -25,7 +25,7 @@ uint MaxRetraces;
 #if NIAGARA_SUPPORTS_COLLISION_GROUPS
 bool FindCollisionGroup(int GPUSceneIndexId, out int CollisionGroup)
 {
-	uint Index = 0;
+	uint Index = MurmurMix(GPUSceneIndexId);
 	if (HashTableFind(GPUSceneIndexId, Index))
 	{
 		CollisionGroup = HashToCollisionGroups[Index];
@@ -36,12 +36,18 @@ bool FindCollisionGroup(int GPUSceneIndexId, out int CollisionGroup)
 }
 #endif
 
-
+#if COMPUTESHADER
+[numthreads(RAY_TRACING_THREAD_GROUP_SIZE_X, 1, 1)]
+void NiagaraCollisionRayTraceCS(uint3 DispatchThreadIndex : SV_DispatchThreadID, uint3 DispatchGroupId : SV_GroupID, uint DispatchGroupIndex : SV_GroupIndex)
+{
+	const uint RayIndex = DispatchThreadIndex.x;
+#else
 RAY_TRACING_ENTRY_RAYGEN(NiagaraCollisionRayTraceRG)
 {
 	const uint RayIndex = DispatchRaysIndex().x;
+#endif
 
-#if NIAGARA_RAYTRACE_FAKE_INDIRECT
+#if NIAGARA_RAYTRACE_FAKE_INDIRECT || COMPUTESHADER
 	if (RayIndex >= RayTraceCounts[0]) 
 	{
 		return;
@@ -69,16 +75,26 @@ RAY_TRACING_ENTRY_RAYGEN(NiagaraCollisionRayTraceRG)
 		while (NumCasts++ <= MaxRetraces)
 		{
 			Payload = (FVFXTracePayload)0;
-			VfxTraceRay(
-				TLAS,   // AccelerationStructure
-				RayFlags,
-				RAY_TRACING_MASK_OPAQUE,          // InstanceInclusionMask
-				RAY_TRACING_SHADER_SLOT_MATERIAL, // RayContributionToHitGroupIndex
-				RAY_TRACING_NUM_SHADER_SLOTS,     // MultiplierForGeometryContributionToShaderIndex
-				0,      // MissShaderIndex
-				Ray.GetNativeDesc(),    // RayDesc
-				Payload // Payload
-			);
+			#if COMPUTESHADER
+				VfxTraceRayInline(
+					TLAS,   // AccelerationStructure
+					RayFlags,
+					RAY_TRACING_MASK_OPAQUE,          // InstanceInclusionMask
+					Ray.GetNativeDesc(),    // FRayDesc
+					Payload // Payload
+				);
+			#else
+				VfxTraceRay(
+					TLAS,   // AccelerationStructure
+					RayFlags,
+					RAY_TRACING_MASK_OPAQUE,          // InstanceInclusionMask
+					RAY_TRACING_SHADER_SLOT_MATERIAL, // RayContributionToHitGroupIndex
+					RAY_TRACING_NUM_SHADER_SLOTS,     // MultiplierForGeometryContributionToShaderIndex
+					0,      // MissShaderIndex
+					Ray.GetNativeDesc(),    // RayDesc
+					Payload // Payload
+				); 
+			#endif
 
 			uint PrimCollisionGroup; 
 			if (FindCollisionGroup(Payload.GPUSceneInstanceId, PrimCollisionGroup))
@@ -100,16 +116,26 @@ RAY_TRACING_ENTRY_RAYGEN(NiagaraCollisionRayTraceRG)
 	else
 #endif
 	{
-		VfxTraceRay(
-			TLAS,   // AccelerationStructure
-			RayFlags,
-			RAY_TRACING_MASK_OPAQUE,          // InstanceInclusionMask
-			RAY_TRACING_SHADER_SLOT_MATERIAL, // RayContributionToHitGroupIndex
-			RAY_TRACING_NUM_SHADER_SLOTS,     // MultiplierForGeometryContributionToShaderIndex
-			0,      // MissShaderIndex
-			Ray.GetNativeDesc(),    // FRayDesc
-			Payload // Payload
-		);
+		#if COMPUTESHADER
+			VfxTraceRayInline(
+				TLAS,   // AccelerationStructure
+				RayFlags,
+				RAY_TRACING_MASK_OPAQUE,          // InstanceInclusionMask
+				Ray.GetNativeDesc(),    // FRayDesc
+				Payload // Payload
+			);
+		#else
+			VfxTraceRay(
+				TLAS,   // AccelerationStructure
+				RayFlags,
+				RAY_TRACING_MASK_OPAQUE,          // InstanceInclusionMask
+				RAY_TRACING_SHADER_SLOT_MATERIAL, // RayContributionToHitGroupIndex
+				RAY_TRACING_NUM_SHADER_SLOTS,     // MultiplierForGeometryContributionToShaderIndex
+				0,      // MissShaderIndex
+				Ray.GetNativeDesc(),    // FRayDesc
+				Payload // Payload
+			);
+		#endif
 	}
 
 	uint OutputIndex = RayIndex + CollisionOutputOffset;
diff --git a/Engine/Plugins/FX/Niagara/Source/NiagaraShader/Private/NiagaraAsyncGpuTraceProviderHwrt.cpp b/Engine/Plugins/FX/Niagara/Source/NiagaraShader/Private/NiagaraAsyncGpuTraceProviderHwrt.cpp
--- a/Engine/Plugins/FX/Niagara/Source/NiagaraShader/Private/NiagaraAsyncGpuTraceProviderHwrt.cpp
+++ b/Engine/Plugins/FX/Niagara/Source/NiagaraShader/Private/NiagaraAsyncGpuTraceProviderHwrt.cpp
@@ -16,6 +16,7 @@
 #include "ShaderParameterStruct.h"
 #include "Templates/Function.h"
 #include "RayTracingPayloadType.h"
+#include "DataDrivenShaderPlatformInfo.h"
 
 static int GNiagaraAsyncGpuTraceHwrtEnabled = 1;
 static FAutoConsoleVariableRef CVarNiagaraAsyncGpuTraceHwrtEnabled(
@@ -25,6 +26,18 @@ static FAutoConsoleVariableRef CVarNiagaraAsyncGpuTraceHwrtEnabled(
 	ECVF_Default
 );
 
+static int GNiagaraAsyncGpuTraceHwrtInline = 1;
+static FAutoConsoleVariableRef CVarNiagaraAsyncGpuTraceHwrtInline(
+	TEXT("fx.Niagara.AsyncGpuTrace.HWRayTrace.Inline"),
+	GNiagaraAsyncGpuTraceHwrtInline,
+	TEXT("If disabled AsyncGpuTrace will not be supported against the HW ray tracing scene."),
+	ECVF_Default);
+
+bool SupportsNiagaraAsyncGpuTraceHwrtInline()
+{
+	return GRHISupportsInlineRayTracing && GNiagaraAsyncGpuTraceHwrtInline;
+}
+
 /// TODO
 ///  -get geometry masking working when an environmental mask is implemented
 
@@ -40,13 +53,15 @@ struct FVFXTracePayload
 
 IMPLEMENT_RT_PAYLOAD_TYPE(ERayTracingPayloadType::VFX, sizeof(FVFXTracePayload));
 
-
-class FNiagaraCollisionRayTraceRG : public FGlobalShader
+class FNiagaraCollisionRayTrace : public FGlobalShader
 {
-	DECLARE_GLOBAL_SHADER(FNiagaraCollisionRayTraceRG)
-	SHADER_USE_ROOT_PARAMETER_STRUCT(FNiagaraCollisionRayTraceRG, FGlobalShader)
+public:
+	FNiagaraCollisionRayTrace() = default;
+	FNiagaraCollisionRayTrace(const ShaderMetaType::CompiledShaderInitializerType& Initializer)
+		: FGlobalShader(Initializer)
+	{}\
 
-	BEGIN_SHADER_PARAMETER_STRUCT(FParameters, )
+		BEGIN_SHADER_PARAMETER_STRUCT(FParameters, )
 		SHADER_PARAMETER_STRUCT_REF(FViewUniformShaderParameters, View)
 		SHADER_PARAMETER_STRUCT_REF(FSceneUniformParameters, Scene)
 
@@ -61,9 +76,15 @@ class FNiagaraCollisionRayTraceRG : public FGlobalShader
 		SHADER_PARAMETER(uint32, CollisionOutputOffset)
 		SHADER_PARAMETER_SRV(Buffer<UINT>, RayTraceCounts)
 		SHADER_PARAMETER(uint32, MaxRetraces)
-	END_SHADER_PARAMETER_STRUCT()
+		END_SHADER_PARAMETER_STRUCT()
+};
+
+class FNiagaraCollisionRayTraceRG : public FNiagaraCollisionRayTrace
+{
+	DECLARE_GLOBAL_SHADER(FNiagaraCollisionRayTraceRG)
+	SHADER_USE_ROOT_PARAMETER_STRUCT(FNiagaraCollisionRayTraceRG, FNiagaraCollisionRayTrace)
 
-	class FFakeIndirectDispatch : SHADER_PERMUTATION_BOOL("NIAGARA_RAYTRACE_FAKE_INDIRECT");
+		class FFakeIndirectDispatch : SHADER_PERMUTATION_BOOL("NIAGARA_RAYTRACE_FAKE_INDIRECT");
 	class FSupportsCollisionGroups : SHADER_PERMUTATION_BOOL("NIAGARA_SUPPORTS_COLLISION_GROUPS");
 	using FPermutationDomain = TShaderPermutationDomain<FFakeIndirectDispatch, FSupportsCollisionGroups>;
 
@@ -79,6 +100,20 @@ class FNiagaraCollisionRayTraceRG : public FGlobalShader
 	static bool SupportsIndirectDispatch();
 };
 
+class FNiagaraCollisionRayTraceCS : public FNiagaraCollisionRayTrace
+{
+	DECLARE_GLOBAL_SHADER(FNiagaraCollisionRayTraceCS)
+	SHADER_USE_PARAMETER_STRUCT(FNiagaraCollisionRayTraceCS, FNiagaraCollisionRayTrace);
+
+	class FSupportsCollisionGroups : SHADER_PERMUTATION_BOOL("NIAGARA_SUPPORTS_COLLISION_GROUPS");
+	using FPermutationDomain = TShaderPermutationDomain<FSupportsCollisionGroups>;
+
+	static bool ShouldCompilePermutation(const FGlobalShaderPermutationParameters& Parameters);
+	static void ModifyCompilationEnvironment(const FGlobalShaderPermutationParameters& Parameters, FShaderCompilerEnvironment& OutEnvironment);
+
+	static const uint32 DispatchSize = 32;
+};
+
 class FNiagaraCollisionRayTraceCH : public FGlobalShader
 {
 	DECLARE_GLOBAL_SHADER(FNiagaraCollisionRayTraceCH);
@@ -121,6 +156,23 @@ void FNiagaraCollisionRayTraceRG::ModifyCompilationEnvironment(const FGlobalShad
 	OutEnvironment.SetDefine(TEXT("VF_SUPPORTS_PRIMITIVE_SCENE_DATA"), 1);
 }
 
+bool FNiagaraCollisionRayTraceCS::ShouldCompilePermutation(const FGlobalShaderPermutationParameters& Parameters)
+{
+	return IsRayTracingEnabledForProject(Parameters.Platform) && RHISupportsRayTracing(Parameters.Platform) && RHISupportsInlineRayTracing(Parameters.Platform);
+}
+
+void FNiagaraCollisionRayTraceCS::ModifyCompilationEnvironment(const FGlobalShaderPermutationParameters& Parameters, FShaderCompilerEnvironment& OutEnvironment)
+{
+	FGlobalShader::ModifyCompilationEnvironment(Parameters, OutEnvironment);
+	OutEnvironment.CompilerFlags.Add(CFLAG_InlineRayTracing);
+	OutEnvironment.CompilerFlags.Add(CFLAG_Wave32);
+
+	OutEnvironment.SetDefine(TEXT("RAY_TRACING_THREAD_GROUP_SIZE_X"), FNiagaraCollisionRayTraceCS::DispatchSize); \
+	OutEnvironment.SetDefine(TEXT("NIAGARA_SUPPORTS_RAY_TRACING"), 1);
+	OutEnvironment.SetDefine(TEXT("VF_SUPPORTS_PRIMITIVE_SCENE_DATA"), 1);
+	OutEnvironment.SetDefine(TEXT("NIAGARA_RAYTRACE_FAKE_INDIRECT"), 0);
+}
+
 TShaderRef<FNiagaraCollisionRayTraceRG> FNiagaraCollisionRayTraceRG::GetShader(FGlobalShaderMap* ShaderMap, bool SupportsCollisionGroups)
 {
 	FPermutationDomain PermutationVector;
@@ -170,6 +222,7 @@ FNiagaraCollisionRayTraceMiss::FNiagaraCollisionRayTraceMiss(const ShaderMetaTyp
 	: FGlobalShader(Initializer)
 {}
 
+IMPLEMENT_GLOBAL_SHADER(FNiagaraCollisionRayTraceCS, "/Plugin/FX/Niagara/Private/NiagaraRayTracingShaders.usf", "NiagaraCollisionRayTraceCS", SF_Compute);
 IMPLEMENT_GLOBAL_SHADER(FNiagaraCollisionRayTraceRG, "/Plugin/FX/Niagara/Private/NiagaraRayTracingShaders.usf", "NiagaraCollisionRayTraceRG", SF_RayGen);
 IMPLEMENT_GLOBAL_SHADER(FNiagaraCollisionRayTraceCH, "/Plugin/FX/Niagara/Private/NiagaraRayTracingShaders.usf", "NiagaraCollisionRayTraceCH", SF_RayHitGroup);
 IMPLEMENT_GLOBAL_SHADER(FNiagaraCollisionRayTraceMiss, "/Plugin/FX/Niagara/Private/NiagaraRayTracingShaders.usf", "NiagaraCollisionRayTraceMiss", SF_RayMiss);
@@ -320,34 +373,37 @@ void FNiagaraAsyncGpuTraceProviderHwrt::PostRenderOpaque(FRHICommandList& RHICmd
 		RayTracingSceneView = UE::FXRenderingUtils::RayTracing::GetRayTracingSceneView(RHICmdList, Scene);
 		ViewUniformBuffer = ReferenceView.ViewUniformBuffer;
 
-		FGlobalShaderMap* ShaderMap = GetGlobalShaderMap(ShaderPlatform);
-		auto RayGenShader = FNiagaraCollisionRayTraceRG::GetRayTracingShader(ShaderMap, CollisionGroupHash != nullptr);
-		auto ClosestHitShader = ShaderMap->GetShader<FNiagaraCollisionRayTraceCH>().GetRayTracingShader();
-		auto MissShader = ShaderMap->GetShader<FNiagaraCollisionRayTraceMiss>().GetRayTracingShader();
-
-		uint32 MaxLocalBindingDataSize = 0;
-		RayTracingPipelineState = CreateNiagaraRayTracingPipelineState(
-			ShaderPlatform,
-			RHICmdList,
-			RayGenShader,
-			ClosestHitShader,
-			MissShader,
-			MaxLocalBindingDataSize);
-
-		RayTracingSBT = UE::FXRenderingUtils::RayTracing::CreateShaderBindingTable(RHICmdList, Scene, MaxLocalBindingDataSize);
-
-		// some options for what we want with our per MeshCommand user data.  For now we'll ignore it, but possibly
-		// something we'd want to incorporate.  Some examples could be if the material is translucent, or possibly the physical material?
-		auto BakeTranslucent = [&](const FRayTracingMeshCommand& MeshCommand) {	return (MeshCommand.bIsTranslucent != 0) & 0x1;	};
-		auto BakeDefault = [&](const FRayTracingMeshCommand& MeshCommand) { return 0; };
-
-		BindNiagaraRayTracingMeshCommands(
-			RHICmdList,
-			RayTracingSBT,
-			ViewUniformBuffer,
-			UE::FXRenderingUtils::RayTracing::GetVisibleRayTracingMeshCommands(ReferenceView),
-			RayTracingPipelineState,
-			BakeDefault);
+		if (!SupportsNiagaraAsyncGpuTraceHwrtInline() && GRHISupportsRayTracingShaders)
+		{
+			FGlobalShaderMap* ShaderMap = GetGlobalShaderMap(ShaderPlatform);
+			auto RayGenShader = FNiagaraCollisionRayTraceRG::GetRayTracingShader(ShaderMap, CollisionGroupHash != nullptr);
+			auto ClosestHitShader = ShaderMap->GetShader<FNiagaraCollisionRayTraceCH>().GetRayTracingShader();
+			auto MissShader = ShaderMap->GetShader<FNiagaraCollisionRayTraceMiss>().GetRayTracingShader();
+
+			uint32 MaxLocalBindingDataSize = 0;
+			RayTracingPipelineState = CreateNiagaraRayTracingPipelineState(
+				ShaderPlatform,
+				RHICmdList,
+				RayGenShader,
+				ClosestHitShader,
+				MissShader,
+				MaxLocalBindingDataSize);
+
+			RayTracingSBT = UE::FXRenderingUtils::RayTracing::CreateShaderBindingTable(RHICmdList, Scene, MaxLocalBindingDataSize);
+
+			// some options for what we want with our per MeshCommand user data.  For now we'll ignore it, but possibly
+			// something we'd want to incorporate.  Some examples could be if the material is translucent, or possibly the physical material?
+			auto BakeTranslucent = [&](const FRayTracingMeshCommand& MeshCommand) {	return (MeshCommand.bIsTranslucent != 0) & 0x1;	};
+			auto BakeDefault = [&](const FRayTracingMeshCommand& MeshCommand) { return 0; };
+
+			BindNiagaraRayTracingMeshCommands(
+				RHICmdList,
+				RayTracingSBT,
+				ViewUniformBuffer,
+				UE::FXRenderingUtils::RayTracing::GetVisibleRayTracingMeshCommands(ReferenceView),
+				RayTracingPipelineState,
+				BakeDefault);
+		}
 	}
 	else
 	{
@@ -358,7 +414,6 @@ void FNiagaraAsyncGpuTraceProviderHwrt::PostRenderOpaque(FRHICommandList& RHICmd
 void FNiagaraAsyncGpuTraceProviderHwrt::IssueTraces(FRHICommandList& RHICmdList, const FDispatchRequest& Request, TUniformBufferRef<FSceneUniformParameters> SceneUniformBufferRHI, FCollisionGroupHashMap* CollisionGroupHash)
 {
 	check(IsAvailable());
-	check(RayTracingPipelineState);
 	check(RayTracingSceneView);
 
 	if (Request.MaxTraceCount == 0)
@@ -370,9 +425,7 @@ void FNiagaraAsyncGpuTraceProviderHwrt::IssueTraces(FRHICommandList& RHICmdList,
 
 	FGlobalShaderMap* ShaderMap = GetGlobalShaderMap(ShaderPlatform);
 
-	TShaderRef<FNiagaraCollisionRayTraceRG> RGShader = FNiagaraCollisionRayTraceRG::GetShader(ShaderMap, CollisionGroupHash != nullptr);
-
-	FNiagaraCollisionRayTraceRG::FParameters Params;
+	FNiagaraCollisionRayTrace::FParameters Params;
 
 	Params.View = GetShaderBinding(ViewUniformBuffer);
 	Params.Scene = SceneUniformBufferRHI;
@@ -391,36 +444,59 @@ void FNiagaraAsyncGpuTraceProviderHwrt::IssueTraces(FRHICommandList& RHICmdList,
 	Params.CollisionOutputOffset = Request.ResultsOffset;
 	Params.MaxRetraces = Request.MaxRetraceCount;
 
-	if (FNiagaraCollisionRayTraceRG::SupportsIndirectDispatch())
-	{
-		FRHIBatchedShaderParameters& GlobalResources = RHICmdList.GetScratchShaderParameters();
-		SetShaderParameters(GlobalResources, RGShader, Params);
-
-		//Can we wrangle things so we can have one indirect dispatch with each internal dispatch pointing to potentially different Ray and Results buffers?
-		//For now have a each as a unique dispatch.
-		RHICmdList.RayTraceDispatchIndirect(
-			RayTracingPipelineState,
-			RGShader.GetRayTracingShader(),
-			RayTracingSBT,
-			GlobalResources,
-			Request.TraceCountsBuffer->Buffer,
-			Request.TraceCountsOffset * sizeof(uint32));
-	}
-	else
+	if (SupportsNiagaraAsyncGpuTraceHwrtInline())
 	{
+		FNiagaraCollisionRayTraceCS::FPermutationDomain PermutationVector;
+		PermutationVector.Set<FNiagaraCollisionRayTraceCS::FSupportsCollisionGroups>(CollisionGroupHash != nullptr);
+		TShaderRef<FNiagaraCollisionRayTraceCS> Shader = ShaderMap->GetShader<FNiagaraCollisionRayTraceCS>(PermutationVector);
+		FRHIComputeShader* ShaderRHI = Shader.GetComputeShader();
+
 		Params.RayTraceCounts = Request.TraceCountsBuffer->SRV;
 
-		FRHIBatchedShaderParameters& GlobalResources = RHICmdList.GetScratchShaderParameters();
-		SetShaderParameters(GlobalResources, RGShader, Params);
-
-		RHICmdList.RayTraceDispatch(
-			RayTracingPipelineState,
-			RGShader.GetRayTracingShader(),
-			RayTracingSBT,
-			GlobalResources,
-			Request.MaxTraceCount,
-			1
-		);
+		SetComputePipelineState(RHICmdList, ShaderRHI);
+		SetShaderParameters(RHICmdList, Shader, ShaderRHI, Params);
+
+
+		RHICmdList.DispatchComputeShader(FMath::DivideAndRoundUp(Request.MaxTraceCount, FNiagaraCollisionRayTraceCS::DispatchSize), 1, 1);
+
+		UnsetShaderUAVs(RHICmdList, Shader, ShaderRHI);
+	}
+	else if (GRHISupportsRayTracingShaders)
+	{
+		check(RayTracingPipelineState);
+
+		TShaderRef<FNiagaraCollisionRayTraceRG> RGShader = FNiagaraCollisionRayTraceRG::GetShader(ShaderMap, CollisionGroupHash != nullptr);
+		if (FNiagaraCollisionRayTraceRG::SupportsIndirectDispatch())
+		{
+			FRHIBatchedShaderParameters& GlobalResources = RHICmdList.GetScratchShaderParameters();
+			SetShaderParameters(GlobalResources, RGShader, Params);
+
+			// Can we wrangle things so we can have one indirect dispatch with each internal dispatch pointing to potentially different Ray and Results buffers?
+			// For now have a each as a unique dispatch.
+			RHICmdList.RayTraceDispatchIndirect(
+				RayTracingPipelineState,
+				RGShader.GetRayTracingShader(),
+				RayTracingSBT,
+				GlobalResources,
+				Request.TraceCountsBuffer->Buffer,
+				Request.TraceCountsOffset * sizeof(uint32));
+		}
+		else
+		{
+			Params.RayTraceCounts = Request.TraceCountsBuffer->SRV;
+
+			FRHIBatchedShaderParameters& GlobalResources = RHICmdList.GetScratchShaderParameters();
+			SetShaderParameters(GlobalResources, RGShader, Params);
+
+			RHICmdList.RayTraceDispatch(
+				RayTracingPipelineState,
+				RGShader.GetRayTracingShader(),
+				RayTracingSBT,
+				GlobalResources,
+				Request.MaxTraceCount,
+				1
+			);
+		}
 	}
 }
 
diff --git a/Engine/Shaders/Private/RayTracing/VFXTraceRay.ush b/Engine/Shaders/Private/RayTracing/VFXTraceRay.ush
--- a/Engine/Shaders/Private/RayTracing/VFXTraceRay.ush
+++ b/Engine/Shaders/Private/RayTracing/VFXTraceRay.ush
@@ -12,6 +12,24 @@
 #include "RayTracingCommon.ush"
 #include "RayTracingHitGroupCommon.ush"
 
+#if COMPUTESHADER
+#include "TraceRayInline.ush"
+
+struct FVFXTracePayload
+{
+	float HitT;
+	uint GPUSceneInstanceId;
+	float2 Barycentrics;
+	float3 TranslatedWorldPosition;
+	float3 WorldNormal;
+
+	bool IsMiss() { return HitT < 0; }
+	bool IsHit() { return !IsMiss(); }
+
+	void SetMiss() { HitT = -1; }
+};
+
+#endif
 
 #ifdef OVERRIDE_VFXTRACERAY_USH
 #include "/Platform/Private/VFXTraceRay.ush"
@@ -40,6 +58,8 @@ void InitVFXTracePayloadFromHitShader(FRayTracingIntersectionAttributes Attribut
 	OutPayload.WorldNormal = normalize(mul(LocalNormal, InverseTranspose3x3));
 }
 
+#endif //#ifndef OVERRIDE_VFXTRACERAY_USH
+
 void VfxTraceRay(
 	RaytracingAccelerationStructure AccelerationStructure,
 	uint RayFlags,
@@ -53,6 +73,32 @@ void VfxTraceRay(
 	TraceRay(AccelerationStructure, RayFlags, InstanceInclusionMask, RayContributionToHitGroupIndex, MultiplierForGeometryContributionToHitGroupIndex, MissShaderIndex, InRay, Payload);
 }
 
-#endif //#ifndef UE_TRACE_RAY
+#if COMPUTESHADER
+
+void VfxTraceRayInline(
+	RaytracingAccelerationStructure AccelerationStructure,
+	uint RayFlags,
+	uint InstanceInclusionMask,
+	RayDesc InRay,
+
+	inout FVFXTracePayload Payload)
+{
+	FTraceRayInlineContext TraceContext = CreateTraceRayInlineContext();
+	FTraceRayInlineResult TraceResult = TraceRayInline(AccelerationStructure, RayFlags, InstanceInclusionMask, InRay, TraceContext);
+	if (TraceResult.IsHit())
+	{
+		Payload.HitT = TraceResult.HitT;
+		Payload.GPUSceneInstanceId = TraceResult.InstanceID;
+		Payload.Barycentrics = TraceResult.Barycentrics;
+		Payload.WorldNormal = TraceResult.WorldGeometryNormal;
+		Payload.TranslatedWorldPosition = InRay.Origin + InRay.Direction * TraceResult.HitT;
+	}
+	else
+	{
+		Payload.SetMiss();
+	}
+}
+
+#endif
 
 #endif // UE_TRACE_RAY_USH_INCLUDED // Workarround for UE-66460
