- Specialize cache compare for quaternions
- Quats can have different representations of the same rotation so are 'equal' and comparing the float values would fail

#rnx
#tests Run EngineTest cache tests, a future CL will submit a stateless test where this was failing before

diff --git a/Engine/Plugins/FX/Niagara/Source/Niagara/Classes/NiagaraSimCacheCompare.h b/Engine/Plugins/FX/Niagara/Source/Niagara/Classes/NiagaraSimCacheCompare.h
--- a/Engine/Plugins/FX/Niagara/Source/Niagara/Classes/NiagaraSimCacheCompare.h
+++ b/Engine/Plugins/FX/Niagara/Source/Niagara/Classes/NiagaraSimCacheCompare.h
@@ -16,6 +16,9 @@ public:
 	/** Default float tolerance when comparing caches */
 	float DefaultFloatTolerance = 0.01f;
 
+	/** Tolerance when comparing quaternions, if not set uses the default UE tolerance. */
+	TOptional<float> DefaultQuaternionTolerance;
+
 	/** Should we exclude known none deterministic variables, i.e. ID / MaterialRandom. */
 	bool bExcludeNoneDeterministicVariables = true;
 
diff --git a/Engine/Plugins/FX/Niagara/Source/Niagara/Private/NiagaraSimCacheCompare.cpp b/Engine/Plugins/FX/Niagara/Source/Niagara/Private/NiagaraSimCacheCompare.cpp
--- a/Engine/Plugins/FX/Niagara/Source/Niagara/Private/NiagaraSimCacheCompare.cpp
+++ b/Engine/Plugins/FX/Niagara/Source/Niagara/Private/NiagaraSimCacheCompare.cpp
@@ -85,6 +85,33 @@ namespace NiagaraSimCacheCompareInternal
 		return true;
 	}
 
+	template<typename TValueString>
+	bool CompareQuatAttributeData(const TArray<int32>& InstanceMapping, const TArray<float>& LhsData, const TArray<float>& RhsData, TValueString& ValueDiffString, float ErrorTolerance)
+	{
+		const int32 NumInstances = InstanceMapping.Num();
+		for (int32 i = 0; i < NumInstances; ++i)
+		{
+			const int32 LhsInstance = i;
+			const int32 RhsInstance = InstanceMapping[i];
+
+			const FQuat4f LhsValue(LhsData[LhsInstance], LhsData[LhsInstance + (NumInstances * 1)], LhsData[LhsInstance + (NumInstances * 2)], LhsData[LhsInstance + (NumInstances * 3)]);
+			const FQuat4f RhsValue(RhsData[RhsInstance], RhsData[RhsInstance + (NumInstances * 1)], RhsData[RhsInstance + (NumInstances * 2)], RhsData[RhsInstance + (NumInstances * 3)]);
+
+			if (!LhsValue.Equals(RhsValue, ErrorTolerance))
+			{
+				ValueDiffString.Append(TEXT("Particle ID "));
+				ValueDiffString.Append(FString::FromInt(LhsInstance));
+				ValueDiffString.Append(TEXT(", Expected: "));
+				ValueDiffString.Appendf(TEXT("(%s, %s, %s, %s)"), *NumberToString(LhsValue.X), *NumberToString(LhsValue.Y), *NumberToString(LhsValue.Z), *NumberToString(LhsValue.W));
+				ValueDiffString.Append(TEXT(", Actual: "));
+				ValueDiffString.Appendf(TEXT("(%s, %s, %s, %s)"), *NumberToString(RhsValue.X), *NumberToString(RhsValue.Y), *NumberToString(RhsValue.Z), *NumberToString(RhsValue.W));
+				return false;
+			}
+		}
+
+		return true;
+	}
+
 	void AddError(FString& OutErrors, const FString& Str)
 	{
 		if (OutErrors.Len() < MaxErrorStrLen)
@@ -432,12 +459,26 @@ bool FNiagaraSimCacheCompare::CompareEmitter(const UNiagaraSimCache& LhsCache, c
 		LhsCache.ReadAttribute(LhsFloats, LhsHalfs, LhsInts, AtributeName, EmitterName, FrameIndex);
 		RhsCache.ReadAttribute(RhsFloats, RhsHalfs, RhsInts, AtributeName, EmitterName, FrameIndex);
 
-		const float ErrorTolerance = VariableToFloatTolerances.FindRef(CompareVariable.Variable, DefaultFloatTolerance);
-
+		// Specialize comparison for quaternions
 		TStringBuilder<1024> ValueDiff;
-		if (!CompareAttributeData(InstanceMapping, LhsFloats, RhsFloats, CompareVariable.FloatCount, ValueDiff, ErrorTolerance) ||
-			!CompareAttributeData(InstanceMapping, LhsHalfs,  RhsHalfs,  CompareVariable.HalfCount,  ValueDiff, ErrorTolerance) ||
-			!CompareAttributeData(InstanceMapping, LhsInts,   RhsInts,   CompareVariable.Int32Count, ValueDiff, ErrorTolerance))
+		bool bIsError = false;
+
+		if (CompareVariable.Variable.GetType() == FNiagaraTypeDefinition::GetQuatDef())
+		{
+			const float ErrorTolerance = DefaultQuaternionTolerance.Get(UE_KINDA_SMALL_NUMBER);
+
+			bIsError = !CompareQuatAttributeData(InstanceMapping, LhsFloats, RhsFloats, ValueDiff, ErrorTolerance);
+		}
+		else
+		{
+			const float ErrorTolerance = VariableToFloatTolerances.FindRef(CompareVariable.Variable, DefaultFloatTolerance);
+
+			bIsError  = !CompareAttributeData(InstanceMapping, LhsFloats, RhsFloats, CompareVariable.FloatCount, ValueDiff, ErrorTolerance);
+			bIsError |= !CompareAttributeData(InstanceMapping, LhsHalfs,  RhsHalfs,  CompareVariable.HalfCount,  ValueDiff, ErrorTolerance);
+			bIsError |= !CompareAttributeData(InstanceMapping, LhsInts,   RhsInts,   CompareVariable.Int32Count, ValueDiff, ErrorTolerance);
+		}
+
+		if ( bIsError )
 		{
 			AddError(
 				OutDifferences,
diff --git a/Engine/Plugins/FX/Niagara/Source/Niagara/Private/NiagaraSimCacheHelper.h b/Engine/Plugins/FX/Niagara/Source/Niagara/Private/NiagaraSimCacheHelper.h
--- a/Engine/Plugins/FX/Niagara/Source/Niagara/Private/NiagaraSimCacheHelper.h
+++ b/Engine/Plugins/FX/Niagara/Source/Niagara/Private/NiagaraSimCacheHelper.h
@@ -617,17 +617,17 @@ struct FNiagaraSimCacheHelper
 
 			for (int32 i = 0; i < SourceVariable.FloatCount; ++i)
 			{
-				CacheBufferReadInfo.ComponentMappingsToDataBuffer[FloatOffset++] = IntCastChecked<uint16>(DestVariableLayout ? DestVariableLayout->GetFloatComponentStart() + i : uint16(INDEX_NONE));
+				CacheBufferReadInfo.ComponentMappingsToDataBuffer[FloatOffset++] = DestVariableLayout ? IntCastChecked<uint16>(DestVariableLayout->GetFloatComponentStart() + i) : uint16(INDEX_NONE);
 			}
 
 			for (int32 i = 0; i < SourceVariable.HalfCount; ++i)
 			{
-				CacheBufferReadInfo.ComponentMappingsToDataBuffer[HalfOffset++] = IntCastChecked<uint16>(DestVariableLayout ? DestVariableLayout->GetHalfComponentStart() + i : uint16(INDEX_NONE));
+				CacheBufferReadInfo.ComponentMappingsToDataBuffer[HalfOffset++] = DestVariableLayout ? IntCastChecked<uint16>(DestVariableLayout->GetHalfComponentStart() + i) : uint16(INDEX_NONE);
 			}
 
 			for (int32 i = 0; i < SourceVariable.Int32Count; ++i)
 			{
-				CacheBufferReadInfo.ComponentMappingsToDataBuffer[Int32Offset++] = IntCastChecked<uint16>(DestVariableLayout ? DestVariableLayout->GetInt32ComponentStart() + i : uint16(INDEX_NONE));
+				CacheBufferReadInfo.ComponentMappingsToDataBuffer[Int32Offset++] = DestVariableLayout ? IntCastChecked<uint16>(DestVariableLayout->GetInt32ComponentStart() + i) : uint16(INDEX_NONE);
 			}
 		}
 
