- Stateless Initial Mesh Orientation updated to include more options & support bindings

#rnx
#tests Validated new options against stateful counterpart

#virtualized

diff --git a/Engine/Plugins/FX/Niagara/Shaders/Private/Stateless/Modules/NiagaraStatelessModule_InitialMeshOrientation.ush b/Engine/Plugins/FX/Niagara/Shaders/Private/Stateless/Modules/NiagaraStatelessModule_InitialMeshOrientation.ush
--- a/Engine/Plugins/FX/Niagara/Shaders/Private/Stateless/Modules/NiagaraStatelessModule_InitialMeshOrientation.ush
+++ b/Engine/Plugins/FX/Niagara/Shaders/Private/Stateless/Modules/NiagaraStatelessModule_InitialMeshOrientation.ush
@@ -2,12 +2,12 @@
 
 #include "../NiagaraStatelessCommon.ush"
 
-float3	InitialMeshOrientation_Rotation;
-float3	InitialMeshOrientation_RandomRangeScale;
+float3	InitialMeshOrientation_RotationScale;
+float3	InitialMeshOrientation_RotationBias;
 
 void InitialMeshOrientation_Simulate(inout FStatelessParticle Particle)
 {
-	const float3 InitialRotation = InitialMeshOrientation_Rotation + (RandomFloat3(0) * InitialMeshOrientation_RandomRangeScale);
+	const float3 InitialRotation		= RandomScaleBiasFloat(0, InitialMeshOrientation_RotationScale, InitialMeshOrientation_RotationBias);
 	Particle.MeshOrientation			= RotatorToQuat(InitialRotation);
 	Particle.PreviousMeshOrientation	= Particle.MeshOrientation;
 }
diff --git a/Engine/Plugins/FX/Niagara/Source/Niagara/Internal/Stateless/Modules/NiagaraStatelessModule_InitialMeshOrientation.h b/Engine/Plugins/FX/Niagara/Source/Niagara/Internal/Stateless/Modules/NiagaraStatelessModule_InitialMeshOrientation.h
--- a/Engine/Plugins/FX/Niagara/Source/Niagara/Internal/Stateless/Modules/NiagaraStatelessModule_InitialMeshOrientation.h
+++ b/Engine/Plugins/FX/Niagara/Source/Niagara/Internal/Stateless/Modules/NiagaraStatelessModule_InitialMeshOrientation.h
@@ -5,85 +5,57 @@
 #include "Stateless/NiagaraStatelessModule.h"
 #include "Stateless/NiagaraStatelessEmitterDataBuildContext.h"
 #include "Stateless/NiagaraStatelessModuleShaderParameters.h"
-#include "Stateless/NiagaraStatelessParticleSimContext.h"
 
 #include "NiagaraStatelessModule_InitialMeshOrientation.generated.h"
 
+UENUM()
+enum class ENSMInitialMeshOrientationMode
+{
+	None,
+	Random,
+	//System,
+	OrientToAxis,
+	//OrientToMatrix,
+	//OrientToQuaternion,
+};
+
 UCLASS(MinimalAPI, EditInlineNew, meta = (DisplayName = "Initial Mesh Orientation"))
 class UNiagaraStatelessModule_InitialMeshOrientation : public UNiagaraStatelessModule
 {
 	GENERATED_BODY()
 
-	struct FModuleBuiltData
-	{
-		FVector3f	Rotation = FVector3f::ZeroVector;
-		FVector3f	RandomRotationRange = FVector3f::ZeroVector;
-		int32		MeshOrientationVariableOffset = INDEX_NONE;
-		int32		PreviousMeshOrientationVariableOffset = INDEX_NONE;
-	};
-
 public:
 	using FParameters = NiagaraStateless::FInitialMeshOrientationModule_ShaderParameters;
 
-	UPROPERTY(EditAnywhere, Category = "Parameters", meta=(Units="deg"))
-	FVector3f	Rotation = FVector3f::ZeroVector;
-
-	UPROPERTY(EditAnywhere, Category = "Parameters", meta=(Units="deg"))
-	FVector3f	RandomRotationRange = FVector3f(360.0f, 360.0f, 360.0f);
-
-	virtual void BuildEmitterData(const FNiagaraStatelessEmitterDataBuildContext& BuildContext) const override
-	{
-		FModuleBuiltData* BuiltData = BuildContext.AllocateBuiltData<FModuleBuiltData>();
+	UPROPERTY(EditAnywhere, Category = "Parameters")
+	ENSMInitialMeshOrientationMode	MeshOrientationMode = ENSMInitialMeshOrientationMode::None;
 
-		const FNiagaraStatelessGlobals& StatelessGlobals	= FNiagaraStatelessGlobals::Get();
-		BuiltData->MeshOrientationVariableOffset			= BuildContext.FindParticleVariableIndex(StatelessGlobals.MeshOrientationVariable);
-		BuiltData->PreviousMeshOrientationVariableOffset	= BuildContext.FindParticleVariableIndex(StatelessGlobals.PreviousMeshOrientationVariable);
+	// Establish an initial orientation around which to yaw, pitch, or roll. Can be overriden with any vector, for instance the normalized velocity vector, to accomplish more elaborate behavior.
+	UPROPERTY(EditAnywhere, Category = "Parameters", meta = (EditConditionHides, EditCondition  = "MeshOrientationMode == ENSMInitialMeshOrientationMode::OrientToAxis", DisableRangeDistribution, DisableUniformDistribution))
+	FNiagaraDistributionRangeVector3 OrientationVector = FNiagaraDistributionRangeVector3(FVector3f::XAxisVector);
 
-		if (IsModuleEnabled())
-		{
-			BuiltData->Rotation				= Rotation / 360.0f;
-			BuiltData->RandomRotationRange	= RandomRotationRange / 360.0f;
-		}
+	// This represents the Axis on which the model was first imported from your DCC package.
+	// This vector is then rotated in the direction of the Orientation Vector input.
+	// If your mesh was imported on a different axis than X forward, you can change it here.
+	UPROPERTY(EditAnywhere, Category = "Parameters", meta = (EditConditionHides, EditCondition = "MeshOrientationMode == ENSMInitialMeshOrientationMode::OrientToAxis", DisableRangeDistribution, DisableUniformDistribution))
+	FNiagaraDistributionRangeVector3 MeshAxisToOrient = FNiagaraDistributionRangeVector3(FVector3f::XAxisVector);
 
-		if (BuiltData->MeshOrientationVariableOffset != INDEX_NONE || BuiltData->PreviousMeshOrientationVariableOffset != INDEX_NONE)
-		{
-			BuildContext.AddParticleSimulationExecSimulate(&UNiagaraStatelessModule_InitialMeshOrientation::ParticleSimulate);
-		}
-	}
+	// Rotation in Degrees, this is applied after any other orientation is calculated and in the space of that orientation
+	UPROPERTY(EditAnywhere, Category = "Parameters", meta=(EditConditionHides, EditCondition = "MeshOrientationMode != ENSMInitialMeshOrientationMode::Random", DisableUniformDistribution, Units = "deg"))
+	FNiagaraDistributionRangeVector3 Rotation = FNiagaraDistributionRangeVector3(FVector3f::ZeroVector);
 
-	virtual void SetShaderParameters(const FNiagaraStatelessSetShaderParameterContext& SetShaderParameterContext) const override
-	{
-		FParameters* Parameters = SetShaderParameterContext.GetParameterNestedStruct<FParameters>();
-
-		const FModuleBuiltData* ModuleBuiltData = SetShaderParameterContext.ReadBuiltData<FModuleBuiltData>();
-		Parameters->InitialMeshOrientation_Rotation			= ModuleBuiltData->Rotation;
-		Parameters->InitialMeshOrientation_RandomRangeScale	= ModuleBuiltData->RandomRotationRange;
-	}
-
-	static void ParticleSimulate(const NiagaraStateless::FParticleSimulationContext& ParticleSimulationContext)
-	{
-		using namespace NiagaraStateless;
-
-		const FModuleBuiltData* ModuleBuiltData = ParticleSimulationContext.ReadBuiltData<FModuleBuiltData>();
-		for (uint32 i = 0; i < ParticleSimulationContext.GetNumInstances(); ++i)
-		{
-			const FVector3f	Rotation = ModuleBuiltData->Rotation + (ParticleSimulationContext.RandomFloat3(i, 0) * ModuleBuiltData->RandomRotationRange);
-			const FQuat4f Quat = ParticleSimulationContext.RotatorToQuat(Rotation);
-
-			ParticleSimulationContext.WriteParticleVariable(ModuleBuiltData->MeshOrientationVariableOffset, i, Quat);
-			ParticleSimulationContext.WriteParticleVariable(ModuleBuiltData->PreviousMeshOrientationVariableOffset, i, Quat);
-		}
-	}
+	virtual void BuildEmitterData(const FNiagaraStatelessEmitterDataBuildContext& BuildContext) const override;
+	virtual void SetShaderParameters(const FNiagaraStatelessSetShaderParameterContext& SetShaderParameterContext) const override;
 
 #if WITH_EDITOR
 	virtual bool CanDisableModule() const override { return true; }
 #endif
 #if WITH_EDITORONLY_DATA
-	virtual void GetOutputVariables(TArray<FNiagaraVariableBase>& OutVariables) const override
-	{
-		const FNiagaraStatelessGlobals& StatelessGlobals = FNiagaraStatelessGlobals::Get();
-		OutVariables.AddUnique(StatelessGlobals.MeshOrientationVariable);
-		OutVariables.AddUnique(StatelessGlobals.PreviousMeshOrientationVariable);
-	}
+	virtual void GetOutputVariables(TArray<FNiagaraVariableBase>& OutVariables) const override;
+	virtual void PostLoad() override;
+
+private:
+	UPROPERTY()
+	FVector3f	RandomRotationRange_DEPRECATED = FVector3f(360.0f, 360.0f, 360.0f);
 #endif
 };
diff --git a/Engine/Plugins/FX/Niagara/Source/Niagara/Internal/Stateless/NiagaraStatelessDistribution.h b/Engine/Plugins/FX/Niagara/Source/Niagara/Internal/Stateless/NiagaraStatelessDistribution.h
--- a/Engine/Plugins/FX/Niagara/Source/Niagara/Internal/Stateless/NiagaraStatelessDistribution.h
+++ b/Engine/Plugins/FX/Niagara/Source/Niagara/Internal/Stateless/NiagaraStatelessDistribution.h
@@ -131,15 +131,6 @@ struct FNiagaraDistributionRangeFloat : public FNiagaraDistributionBase
 	bool SerializeFromMismatchedTag(const struct FPropertyTag& Tag, FStructuredArchive::FSlot Slot);
 };
 
-template<>
-struct TStructOpsTypeTraits<FNiagaraDistributionRangeFloat> : public TStructOpsTypeTraitsBase2<FNiagaraDistributionRangeFloat>
-{
-	enum
-	{
-		WithStructuredSerializeFromMismatchedTag = true,
-	};
-};
-
 USTRUCT()
 struct FNiagaraDistributionRangeVector2 : public FNiagaraDistributionBase
 {
@@ -188,6 +179,7 @@ struct FNiagaraDistributionRangeVector3 : public FNiagaraDistributionBase
 	virtual void UpdateValuesFromDistribution() override;
 	virtual FNiagaraTypeDefinition GetBindingTypeDef() const { return FNiagaraTypeDefinition::GetVec3Def(); }
 #endif
+	bool SerializeFromMismatchedTag(const struct FPropertyTag& Tag, FStructuredArchive::FSlot Slot);
 };
 
 USTRUCT()
@@ -334,3 +326,22 @@ struct FNiagaraDistributionColor : public FNiagaraDistributionBase
 	virtual FNiagaraTypeDefinition GetBindingTypeDef() const { return FNiagaraTypeDefinition::GetColorDef(); }
 #endif
 };
+
+template<>
+struct TStructOpsTypeTraits<FNiagaraDistributionRangeFloat> : public TStructOpsTypeTraitsBase2<FNiagaraDistributionRangeFloat>
+{
+	enum
+	{
+		WithStructuredSerializeFromMismatchedTag = true,
+	};
+};
+
+template<>
+struct TStructOpsTypeTraits<FNiagaraDistributionRangeVector3> : public TStructOpsTypeTraitsBase2<FNiagaraDistributionRangeVector3>
+{
+	enum
+	{
+		WithStructuredSerializeFromMismatchedTag = true,
+	};
+};
+
diff --git a/Engine/Plugins/FX/Niagara/Source/Niagara/Internal/Stateless/NiagaraStatelessParticleSimContext.h b/Engine/Plugins/FX/Niagara/Source/Niagara/Internal/Stateless/NiagaraStatelessParticleSimContext.h
--- a/Engine/Plugins/FX/Niagara/Source/Niagara/Internal/Stateless/NiagaraStatelessParticleSimContext.h
+++ b/Engine/Plugins/FX/Niagara/Source/Niagara/Internal/Stateless/NiagaraStatelessParticleSimContext.h
@@ -140,9 +140,9 @@ public:
 	FVector4f		RandomFloat4(uint32 iInstance, uint32 RandomSeedOffset) const;
 
 	float			RandomScaleBiasFloat(uint32 iInstance, uint32 RandomSeedOffset, const float Scale, const float Bias) const { return Bias + (RandomFloat(iInstance, RandomSeedOffset) * Scale); }
-	FVector2f		RandomScaleBiasFloat(uint32 iInstance, uint32 RandomSeedOffset, const FVector2f& Scale, const FVector2f& Bias) const { return Bias + (RandomFloat(iInstance, RandomSeedOffset) * Scale); }
-	FVector3f		RandomScaleBiasFloat(uint32 iInstance, uint32 RandomSeedOffset, const FVector3f& Scale, const FVector3f& Bias) const { return Bias + (RandomFloat(iInstance, RandomSeedOffset) * Scale); }
-	FVector4f		RandomScaleBiasFloat(uint32 iInstance, uint32 RandomSeedOffset, const FVector4f& Scale, const FVector4f& Bias) const { return Bias + (RandomFloat(iInstance, RandomSeedOffset) * Scale); }
+	FVector2f		RandomScaleBiasFloat(uint32 iInstance, uint32 RandomSeedOffset, const FVector2f& Scale, const FVector2f& Bias) const { return Bias + (RandomFloat2(iInstance, RandomSeedOffset) * Scale); }
+	FVector3f		RandomScaleBiasFloat(uint32 iInstance, uint32 RandomSeedOffset, const FVector3f& Scale, const FVector3f& Bias) const { return Bias + (RandomFloat3(iInstance, RandomSeedOffset) * Scale); }
+	FVector4f		RandomScaleBiasFloat(uint32 iInstance, uint32 RandomSeedOffset, const FVector4f& Scale, const FVector4f& Bias) const { return Bias + (RandomFloat4(iInstance, RandomSeedOffset) * Scale); }
 
 	float			RandomScaleBiasFloat(uint32 iInstance, uint32 RandomSeedOffset, const FNiagaraStatelessRangeFloat& Range) const { return Range.Min + (RandomFloat(iInstance, RandomSeedOffset) * Range.GetScale()); }
 	FVector2f		RandomScaleBiasFloat(uint32 iInstance, uint32 RandomSeedOffset, const FNiagaraStatelessRangeVector2& Range) const { return Range.Min + (RandomFloat2(iInstance, RandomSeedOffset) * Range.GetScale()); }
@@ -164,6 +164,12 @@ public:
 	FVector2f		SafeNormalize(const FVector2f& v) const { return SafeNormalize(v, FVector2f(1.0f, 0.0f)); }
 	FVector3f		SafeNormalize(const FVector3f& v) const { return SafeNormalize(v, FVector3f(1.0f, 0.0f, 0.0f)); }
 
+	void			ConvertRangeToScaleBias(const FNiagaraStatelessRangeFloat&   Range, float& OutScale, float& OutBias) const { OutScale = Range.ParameterOffset == INDEX_NONE ? Range.GetScale() : 0.0f; OutBias = Range.ParameterOffset == INDEX_NONE ? Range.Min : GetParameterBufferValue<float>(Range.ParameterOffset, 0); }
+	void			ConvertRangeToScaleBias(const FNiagaraStatelessRangeVector2& Range, FVector2f& OutScale, FVector2f& OutBias) const { OutScale = Range.ParameterOffset == INDEX_NONE ? Range.GetScale() : FVector2f::ZeroVector; OutBias = Range.ParameterOffset == INDEX_NONE ? Range.Min : GetParameterBufferValue<FVector2f>(Range.ParameterOffset, 0); }
+	void			ConvertRangeToScaleBias(const FNiagaraStatelessRangeVector3& Range, FVector3f& OutScale, FVector3f& OutBias) const { OutScale = Range.ParameterOffset == INDEX_NONE ? Range.GetScale() : FVector3f::ZeroVector; OutBias = Range.ParameterOffset == INDEX_NONE ? Range.Min : GetParameterBufferValue<FVector3f>(Range.ParameterOffset, 0); }
+	void			ConvertRangeToScaleBias(const FNiagaraStatelessRangeVector4& Range, FVector4f& OutScale, FVector4f& OutBias) const { OutScale = Range.ParameterOffset == INDEX_NONE ? Range.GetScale() : FVector4f::Zero(); OutBias = Range.ParameterOffset == INDEX_NONE ? Range.Min : GetParameterBufferValue<FVector4f>(Range.ParameterOffset, 0); }
+	void			ConvertRangeToScaleBias(const FNiagaraStatelessRangeColor&   Range, FLinearColor& OutScale, FLinearColor& OutBias) const { OutScale = Range.ParameterOffset == INDEX_NONE ? Range.GetScale() : FLinearColor(0.0f, 0.0f, 0.0f, 0.0f); OutBias = Range.ParameterOffset == INDEX_NONE ? Range.Min : GetParameterBufferValue<FLinearColor>(Range.ParameterOffset, 0); }
+
 	template<typename T>
 	const T* ReadBuiltData() const
 	{
@@ -173,7 +179,6 @@ public:
 		return reinterpret_cast<const T*>(BuiltData.GetData() + Offset);
 	}
 
-	//-TODO: Reconcile why this differs from C++ version and match shader accordingly
 	static FQuat4f RotatorToQuat(FVector3f Rotator)
 	{
 		Rotator.X = FMath::Fractional(Rotator.X) * UE_PI;
diff --git a/Engine/Plugins/FX/Niagara/Source/Niagara/Private/Stateless/Modules/NiagaraStatelessModule_InitialMeshOrientation.cpp b/Engine/Plugins/FX/Niagara/Source/Niagara/Private/Stateless/Modules/NiagaraStatelessModule_InitialMeshOrientation.cpp
new file mode 100644
--- /dev/null
+++ b/Engine/Plugins/FX/Niagara/Source/Niagara/Private/Stateless/Modules/NiagaraStatelessModule_InitialMeshOrientation.cpp
@@ -0,0 +1,137 @@
+// Copyright Epic Games, Inc. All Rights Reserved.
+
+#include "Stateless/Modules/NiagaraStatelessModule_InitialMeshOrientation.h"
+#include "Stateless/NiagaraStatelessParticleSimContext.h"
+#include "NiagaraCustomVersion.h"
+
+namespace NSMInitialMeshOrientationPrivate
+{
+	struct FModuleBuiltData
+	{
+		ENSMInitialMeshOrientationMode	Mode = ENSMInitialMeshOrientationMode::None;
+		FNiagaraStatelessRangeVector3	OrientationVector;
+		FNiagaraStatelessRangeVector3	MeshAxisToOrient;
+		FNiagaraStatelessRangeVector3	RotationRange;
+		int32							MeshOrientationVariableOffset = INDEX_NONE;
+		int32							PreviousMeshOrientationVariableOffset = INDEX_NONE;
+	};
+
+	template<typename TExecContext>
+	static void ResolveRotationScaleBias(const TExecContext& ExecContext, const FModuleBuiltData* ModuleBuiltData, FVector3f& RotationScale, FVector3f& RotationBias)
+	{
+		RotationScale = FVector3f::ZeroVector;
+		RotationBias = FVector3f::ZeroVector;
+		ExecContext.ConvertRangeToScaleBias(ModuleBuiltData->RotationRange, RotationScale, RotationBias);
+
+		if (ModuleBuiltData->Mode == ENSMInitialMeshOrientationMode::OrientToAxis)
+		{
+			FVector3f FromVector = FVector3f::XAxisVector;
+			FVector3f ToVector = FVector3f::XAxisVector;
+			FVector3f DummyVector;
+			ExecContext.ConvertRangeToScaleBias(ModuleBuiltData->MeshAxisToOrient, DummyVector, FromVector);
+			ExecContext.ConvertRangeToScaleBias(ModuleBuiltData->OrientationVector, DummyVector, ToVector);
+
+			FromVector = FromVector.GetSafeNormal(UE_KINDA_SMALL_NUMBER, FVector3f::XAxisVector);
+			ToVector = ToVector.GetSafeNormal(UE_KINDA_SMALL_NUMBER, FVector3f::XAxisVector);
+
+			const FQuat4f QuatBetween = FQuat4f::FindBetweenVectors(FromVector, ToVector);
+			const FRotator3f Rotator = QuatBetween.Rotator();
+			RotationBias.X += Rotator.Roll;
+			RotationBias.Y += Rotator.Pitch;
+			RotationBias.Z += Rotator.Yaw;
+		}
+		RotationScale /= 360.0f;
+		RotationBias /= 360.0f;
+	}
+
+	static void ParticleSimulate(const NiagaraStateless::FParticleSimulationContext& ParticleSimulationContext)
+	{
+		using namespace NiagaraStateless;
+
+		const FModuleBuiltData* ModuleBuiltData = ParticleSimulationContext.ReadBuiltData<FModuleBuiltData>();
+
+		FVector3f RotationScale;
+		FVector3f RotationBias;
+		ResolveRotationScaleBias(ParticleSimulationContext, ModuleBuiltData, RotationScale, RotationBias);
+
+		for (uint32 i = 0; i < ParticleSimulationContext.GetNumInstances(); ++i)
+		{
+			const FVector3f	Rotation = ParticleSimulationContext.RandomScaleBiasFloat(i, 0, RotationScale, RotationBias);
+			const FQuat4f Quat = ParticleSimulationContext.RotatorToQuat(Rotation);
+
+			ParticleSimulationContext.WriteParticleVariable(ModuleBuiltData->MeshOrientationVariableOffset, i, Quat);
+			ParticleSimulationContext.WriteParticleVariable(ModuleBuiltData->PreviousMeshOrientationVariableOffset, i, Quat);
+		}
+	}
+}
+
+void UNiagaraStatelessModule_InitialMeshOrientation::BuildEmitterData(const FNiagaraStatelessEmitterDataBuildContext& BuildContext) const
+{
+	using namespace NSMInitialMeshOrientationPrivate;
+
+	FModuleBuiltData* BuiltData = BuildContext.AllocateBuiltData<FModuleBuiltData>();
+
+	const FNiagaraStatelessGlobals& StatelessGlobals	= FNiagaraStatelessGlobals::Get();
+	BuiltData->MeshOrientationVariableOffset			= BuildContext.FindParticleVariableIndex(StatelessGlobals.MeshOrientationVariable);
+	BuiltData->PreviousMeshOrientationVariableOffset	= BuildContext.FindParticleVariableIndex(StatelessGlobals.PreviousMeshOrientationVariable);
+
+	if (IsModuleEnabled())
+	{
+		BuiltData->Mode = MeshOrientationMode;
+		switch (MeshOrientationMode)
+		{
+			case ENSMInitialMeshOrientationMode::None:
+				BuiltData->RotationRange = BuildContext.ConvertDistributionToRange(Rotation, FVector3f::ZeroVector);
+				break;
+
+			case ENSMInitialMeshOrientationMode::Random:
+				BuiltData->RotationRange.Min = FVector3f::ZeroVector;
+				BuiltData->RotationRange.Max = FVector3f(360.f, 360.0f, 360.0f);
+				break;
+
+			case ENSMInitialMeshOrientationMode::OrientToAxis:
+			{
+				BuiltData->RotationRange = BuildContext.ConvertDistributionToRange(Rotation, FVector3f::ZeroVector);
+				BuiltData->OrientationVector = BuildContext.ConvertDistributionToRange(OrientationVector, FVector3f::ZeroVector);
+				BuiltData->MeshAxisToOrient = BuildContext.ConvertDistributionToRange(MeshAxisToOrient, FVector3f::ZeroVector);
+				break;
+			}
+		}
+	}
+
+	if (BuiltData->MeshOrientationVariableOffset != INDEX_NONE || BuiltData->PreviousMeshOrientationVariableOffset != INDEX_NONE)
+	{
+		BuildContext.AddParticleSimulationExecSimulate(&ParticleSimulate);
+	}
+}
+
+void UNiagaraStatelessModule_InitialMeshOrientation::SetShaderParameters(const FNiagaraStatelessSetShaderParameterContext& SetShaderParameterContext) const
+{
+	using namespace NSMInitialMeshOrientationPrivate;
+	FParameters* Parameters = SetShaderParameterContext.GetParameterNestedStruct<FParameters>();
+
+	const FModuleBuiltData* ModuleBuiltData = SetShaderParameterContext.ReadBuiltData<FModuleBuiltData>();
+
+	ResolveRotationScaleBias(SetShaderParameterContext, ModuleBuiltData, Parameters->InitialMeshOrientation_RotationScale, Parameters->InitialMeshOrientation_RotationBias);
+}
+
+#if WITH_EDITORONLY_DATA
+void UNiagaraStatelessModule_InitialMeshOrientation::GetOutputVariables(TArray<FNiagaraVariableBase>& OutVariables) const
+{
+	const FNiagaraStatelessGlobals& StatelessGlobals = FNiagaraStatelessGlobals::Get();
+	OutVariables.AddUnique(StatelessGlobals.MeshOrientationVariable);
+	OutVariables.AddUnique(StatelessGlobals.PreviousMeshOrientationVariable);
+}
+
+void UNiagaraStatelessModule_InitialMeshOrientation::PostLoad()
+{
+	Super::PostLoad();
+
+	const int32 NiagaraVersion = GetLinkerCustomVersion(FNiagaraCustomVersion::GUID);
+
+	if (NiagaraVersion < FNiagaraCustomVersion::StatelessInitialMeshOrientationV1)
+	{
+		Rotation.Max += RandomRotationRange_DEPRECATED;
+	}
+}
+#endif
diff --git a/Engine/Plugins/FX/Niagara/Source/Niagara/Private/Stateless/NiagaraStatelessDistribution.cpp b/Engine/Plugins/FX/Niagara/Source/Niagara/Private/Stateless/NiagaraStatelessDistribution.cpp
--- a/Engine/Plugins/FX/Niagara/Source/Niagara/Private/Stateless/NiagaraStatelessDistribution.cpp
+++ b/Engine/Plugins/FX/Niagara/Source/Niagara/Private/Stateless/NiagaraStatelessDistribution.cpp
@@ -442,19 +442,6 @@ FNiagaraStatelessRangeColor FNiagaraDistributionColor::CalculateRange(const FLin
 	return Range;
 }
 
-bool FNiagaraDistributionRangeFloat::SerializeFromMismatchedTag(const struct FPropertyTag& Tag, FStructuredArchive::FSlot Slot)
-{
-	if (Tag.Type == NAME_FloatProperty)
-	{
-		float Value;
-		Slot << Value;
-		*this = FNiagaraDistributionRangeFloat(Value);
-		return true;
-	}
-
-	return false;
-}
-
 #if WITH_EDITORONLY_DATA
 void FNiagaraDistributionBase::PostEditChangeProperty(UObject* OwnerObject, FPropertyChangedEvent& PropertyChangedEvent)
 {
@@ -599,3 +586,29 @@ void FNiagaraDistributionColor::UpdateValuesFromDistribution()
 
 #endif
 
+bool FNiagaraDistributionRangeFloat::SerializeFromMismatchedTag(const struct FPropertyTag& Tag, FStructuredArchive::FSlot Slot)
+{
+	if (Tag.Type == NAME_FloatProperty)
+	{
+		float Value;
+		Slot << Value;
+		*this = FNiagaraDistributionRangeFloat(Value);
+		return true;
+	}
+
+	return false;
+}
+
+bool FNiagaraDistributionRangeVector3::SerializeFromMismatchedTag(const struct FPropertyTag& Tag, FStructuredArchive::FSlot Slot)
+{
+	if (Tag.GetType().IsStruct(NAME_Vector3f))
+	{
+		FVector3f Value;
+		Slot << Value;
+		*this = FNiagaraDistributionRangeVector3(Value);
+		return true;
+	}
+
+	return false;
+}
+
diff --git a/Engine/Plugins/FX/Niagara/Source/NiagaraCore/Public/NiagaraCustomVersion.h b/Engine/Plugins/FX/Niagara/Source/NiagaraCore/Public/NiagaraCustomVersion.h
--- a/Engine/Plugins/FX/Niagara/Source/NiagaraCore/Public/NiagaraCustomVersion.h
+++ b/Engine/Plugins/FX/Niagara/Source/NiagaraCore/Public/NiagaraCustomVersion.h
@@ -204,6 +204,8 @@ struct FNiagaraCustomVersion
 		NDCSpawnGroupOverrideDisabledByDefault, // NDC Read DIs will not override spawn group by default when spawning particles. Old content will remain unchanged.
 
 		CustomSortingBindingToAge, // Before it was normalized age which can introduce flickering with sorting and random lifetimes
+
+		StatelessInitialMeshOrientationV1,	// Update Initial Mesh Orientation Module
 		
 		// DO NOT ADD A NEW VERSION UNLESS YOU HAVE TALKED TO THE NIAGARA LEAD. Mismanagement of these versions can lead to data loss if it is adjusted in multiple streams simultaneously.
 		// -----<new versions can be added above this line>  -------------------------------------------------
diff --git a/Engine/Plugins/FX/Niagara/Source/NiagaraShader/Internal/Stateless/NiagaraStatelessModuleShaderParameters.h b/Engine/Plugins/FX/Niagara/Source/NiagaraShader/Internal/Stateless/NiagaraStatelessModuleShaderParameters.h
--- a/Engine/Plugins/FX/Niagara/Source/NiagaraShader/Internal/Stateless/NiagaraStatelessModuleShaderParameters.h
+++ b/Engine/Plugins/FX/Niagara/Source/NiagaraShader/Internal/Stateless/NiagaraStatelessModuleShaderParameters.h
@@ -31,8 +31,8 @@ namespace NiagaraStateless
 	END_SHADER_PARAMETER_STRUCT()
 
 	BEGIN_SHADER_PARAMETER_STRUCT(FInitialMeshOrientationModule_ShaderParameters, NIAGARASHADER_API)
-		SHADER_PARAMETER(FVector3f,	InitialMeshOrientation_Rotation)
-		SHADER_PARAMETER(FVector3f,	InitialMeshOrientation_RandomRangeScale)
+		SHADER_PARAMETER(FVector3f, InitialMeshOrientation_RotationScale)
+		SHADER_PARAMETER(FVector3f, InitialMeshOrientation_RotationBias)
 	END_SHADER_PARAMETER_STRUCT()
 		
 
