- Capture more parameter stores to the debug cache
- Useful for debugging NPCs and what data is being send to renderers

#rnx
#test Captured various caches to inspect values from different sources, like NPC bindings in both stateful and stateless.  Ran cache tests.

diff --git a/Engine/Plugins/FX/Niagara/Source/Niagara/Internal/NiagaraSimCacheDebugData.h b/Engine/Plugins/FX/Niagara/Source/Niagara/Internal/NiagaraSimCacheDebugData.h
--- a/Engine/Plugins/FX/Niagara/Source/Niagara/Internal/NiagaraSimCacheDebugData.h
+++ b/Engine/Plugins/FX/Niagara/Source/Niagara/Internal/NiagaraSimCacheDebugData.h
@@ -15,7 +15,7 @@ struct FNiagaraSimCacheDebugDataFrame
 	GENERATED_BODY()
 
 	UPROPERTY()
-	FNiagaraParameterStore OverrideParameters;
+	TMap<FString, FNiagaraParameterStore> DebugParameterStores;
 };
 
 // Contains data useful for debugging a Niagara system
diff --git a/Engine/Plugins/FX/Niagara/Source/Niagara/Private/NiagaraSimCacheDebugData.cpp b/Engine/Plugins/FX/Niagara/Source/Niagara/Private/NiagaraSimCacheDebugData.cpp
--- a/Engine/Plugins/FX/Niagara/Source/Niagara/Private/NiagaraSimCacheDebugData.cpp
+++ b/Engine/Plugins/FX/Niagara/Source/Niagara/Private/NiagaraSimCacheDebugData.cpp
@@ -2,6 +2,38 @@
 
 #include "NiagaraSimCacheDebugData.h"
 #include "NiagaraSimCacheHelper.h"
+#include "NiagaraEmitterInstanceImpl.h"
+
+namespace NSCDebugDataPrivate
+{
+	template<typename TStringType>
+	void AddParameterStore(FNiagaraSimCacheDebugDataFrame& FrameData, TStringType StoreName, const FNiagaraParameterStore& ParameterStore)
+	{
+		if (ParameterStore.Num() > 0)
+		{
+			FrameData.DebugParameterStores.FindOrAdd(StoreName) = ParameterStore;
+		}
+	}
+
+	template<typename TStringType>
+	void AddParameterStore(FNiagaraSimCacheDebugDataFrame& FrameData, TStringType StoreName, const FNiagaraScriptExecutionContextBase* ExecContext)
+	{
+		if (ExecContext && ExecContext->Parameters.ReadParameterVariables().Num() > 0)
+		{
+			FNiagaraParameterStore& DestStore = FrameData.DebugParameterStores.FindOrAdd(StoreName);
+			ExecContext->Parameters.CopyParametersTo(DestStore, false, FNiagaraParameterStore::EDataInterfaceCopyMethod::Value);
+		}
+	}
+
+	template<typename TStringType>
+	void AddParameterStore(FNiagaraSimCacheDebugDataFrame& FrameData, TStringType StoreName, const FNiagaraComputeExecutionContext* ExecContext)
+	{
+		if (ExecContext)
+		{
+			AddParameterStore(FrameData, StoreName, ExecContext->CombinedParamStore);
+		}
+	}
+}
 
 UNiagaraSimCacheDebugData::UNiagaraSimCacheDebugData(const FObjectInitializer& ObjectInitializer)
 {
@@ -9,9 +41,41 @@ UNiagaraSimCacheDebugData::UNiagaraSimCacheDebugData(const FObjectInitializer& O
 
 void UNiagaraSimCacheDebugData::CaptureFrame(FNiagaraSimCacheHelper& Helper, int FrameNumber)
 {
+	using namespace NSCDebugDataPrivate;
+
    	Frames.SetNum(FrameNumber + 1);
-	if (const FNiagaraParameterStore* ParameterStore = Helper.SystemInstance->GetOverrideParameters())
+
+	FNiagaraSimCacheDebugDataFrame& FrameData = Frames[FrameNumber];
+
+	// Add Override Parameters
+	if (const FNiagaraParameterStore* OverrideParameterStore = Helper.SystemInstance->GetOverrideParameters())
+	{
+		AddParameterStore(FrameData, TEXT("OverrideParameters"), *OverrideParameterStore);
+	}
+
+	// Add Instance Parameters
+	AddParameterStore(FrameData, TEXT("InstanceParameters"), Helper.SystemInstance->GetInstanceParameters());
+
+	// Add System Script Parameters
+	if (FNiagaraSystemSimulation* SystemSimulation = Helper.SystemInstance->GetSystemSimulation().Get())
 	{
-		Frames[FrameNumber].OverrideParameters = *ParameterStore;
+		AddParameterStore(FrameData, TEXT("System Spawn"), SystemSimulation->GetSpawnExecutionContext());
+		AddParameterStore(FrameData, TEXT("System Update"), SystemSimulation->GetUpdateExecutionContext());
+	}
+
+	// Add Emitter Parameters
+	for (const FNiagaraEmitterInstanceRef& EmitterRef : Helper.SystemInstance->GetEmitters())
+	{
+		const FNiagaraEmitterHandle& EmitterHandle = EmitterRef->GetEmitterHandle();
+		const FString EmitterName = EmitterHandle.GetName().ToString();
+
+		AddParameterStore(FrameData, FString::Printf(TEXT("%s RendererBindings"), *EmitterName), EmitterRef->GetRendererBoundVariables());
+		AddParameterStore(FrameData, FString::Printf(TEXT("%s GPUContext"), *EmitterName), EmitterRef->GetGPUContext());
+
+		if (FNiagaraEmitterInstanceImpl* StatefulEmitter = EmitterRef->AsStateful())
+		{
+			AddParameterStore(FrameData, FString::Printf(TEXT("%s Spawn"), *EmitterName), &StatefulEmitter->GetSpawnExecutionContext());
+			AddParameterStore(FrameData, FString::Printf(TEXT("%s Update"), *EmitterName), &StatefulEmitter->GetUpdateExecutionContext());
+		}
 	}
 }
diff --git a/Engine/Plugins/FX/Niagara/Source/NiagaraEditor/Private/Widgets/SNiagaraSimCacheDebugDataView.cpp b/Engine/Plugins/FX/Niagara/Source/NiagaraEditor/Private/Widgets/SNiagaraSimCacheDebugDataView.cpp
--- a/Engine/Plugins/FX/Niagara/Source/NiagaraEditor/Private/Widgets/SNiagaraSimCacheDebugDataView.cpp
+++ b/Engine/Plugins/FX/Niagara/Source/NiagaraEditor/Private/Widgets/SNiagaraSimCacheDebugDataView.cpp
@@ -5,6 +5,8 @@
 #include "Widgets/SNiagaraParameterName.h"
 #include "NiagaraSimCacheDebugData.h"
 
+#include "Framework/MultiBox/MultiBoxBuilder.h"
+#include "Widgets/Input/SComboButton.h"
 #include "Widgets/Layout/SScrollBox.h"
 #include "Widgets/Views/SListView.h"
 
@@ -143,6 +145,12 @@ void SNiagaraSimCacheDebugDataView::Construct(const FArguments& InArgs)
 	SimCacheViewModel->OnSimCacheChanged().AddSP(this, &SNiagaraSimCacheDebugDataView::RefreshContents);
 	SimCacheViewModel->OnBufferChanged().AddSP(this, &SNiagaraSimCacheDebugDataView::RefreshContents);
 
+
+	if (const FNiagaraSimCacheDebugDataFrame* FrameData = GetCurrentFrameData())
+	{
+		SelectedParameterStoreName = FrameData->DebugParameterStores.CreateConstIterator().Key();
+	}
+
 	ChildSlot
 	[
 		SNew(SVerticalBox)
@@ -157,7 +165,21 @@ void SNiagaraSimCacheDebugDataView::Construct(const FArguments& InArgs)
 				.AutoWidth()
 				[
 					SNew(STextBlock)
-					.Text(LOCTEXT("OverrideParameters", "Override Parameters"))
+					.Text(LOCTEXT("ParameterStoreSelection", "Parameter Store Selection:"))
+					.Margin(FMargin(0.0, 0.0, 5.0, 0.0))
+				]
+				+ SHorizontalBox::Slot()
+				.AutoWidth()
+				[
+					SNew(SComboButton)
+					.OnGetMenuContent(this, &SNiagaraSimCacheDebugDataView::GetParameterStoreSelectionMenu)
+					.ButtonContent()
+					[
+						SNew(STextBlock)
+						.Text(
+							TAttribute<FText>::Create(TAttribute<FText>::FGetter::CreateLambda([this]() { return FText::FromString(SelectedParameterStoreName); }))
+						)
+					]
 				]
 			]
 		]
@@ -176,22 +198,56 @@ void SNiagaraSimCacheDebugDataView::Construct(const FArguments& InArgs)
 	RefreshContents();
 };
 
-void SNiagaraSimCacheDebugDataView::RefreshContents()
+const FNiagaraSimCacheDebugDataFrame* SNiagaraSimCacheDebugDataView::GetCurrentFrameData() const
 {
 	const UNiagaraSimCacheDebugData* DebugData = SimCacheViewModel ? SimCacheViewModel->GetCacheDebugData() : nullptr;
 	if (DebugData)
 	{
 		const int32 FrameIndex = SimCacheViewModel->GetFrameIndex();
-		if ( DebugData->Frames.IsValidIndex(FrameIndex))
+		if (DebugData->Frames.IsValidIndex(FrameIndex))
+		{
+			return &DebugData->Frames[FrameIndex];
+		}
+	}
+	return nullptr;
+}
+
+TSharedRef<SWidget> SNiagaraSimCacheDebugDataView::GetParameterStoreSelectionMenu()
+{
+	FMenuBuilder MenuBuilder(true, nullptr);
+	if (const FNiagaraSimCacheDebugDataFrame* FrameData = GetCurrentFrameData())
+	{
+		for ( auto It=FrameData->DebugParameterStores.CreateConstIterator(); It; ++It)
 		{
-			const FNiagaraSimCacheDebugDataFrame& DebugFrameData = DebugData->Frames[FrameIndex];
-			OverrideParametersWidget->SetParameterStore(DebugFrameData.OverrideParameters);
-			return;
+			MenuBuilder.AddMenuEntry(
+				FText::FromString(It->Key),
+				FText(),
+				FSlateIcon(),
+				FUIAction(
+					FExecuteAction::CreateLambda(
+						[this, SelectedValue=It->Key]()
+						{
+							SelectedParameterStoreName = SelectedValue;
+							RefreshContents();
+						}
+					)
+				)
+			);
 		}
 	}
 
-	// If we get here we failed to set anything so update the display with defaults
-	OverrideParametersWidget->SetParameterStore(FNiagaraParameterStore());
+	return MenuBuilder.MakeWidget();
+}
+
+void SNiagaraSimCacheDebugDataView::RefreshContents()
+{
+	const FNiagaraParameterStore* ParameterStore = nullptr;
+	if (const FNiagaraSimCacheDebugDataFrame* FrameData = GetCurrentFrameData())
+	{
+		ParameterStore = FrameData->DebugParameterStores.Find(SelectedParameterStoreName);
+	}
+
+	OverrideParametersWidget->SetParameterStore(ParameterStore ? *ParameterStore : FNiagaraParameterStore());
 }
 
 void SNiagaraSimCacheDebugDataView::RefreshContents(bool)
diff --git a/Engine/Plugins/FX/Niagara/Source/NiagaraEditor/Private/Widgets/SNiagaraSimCacheDebugDataView.h b/Engine/Plugins/FX/Niagara/Source/NiagaraEditor/Private/Widgets/SNiagaraSimCacheDebugDataView.h
--- a/Engine/Plugins/FX/Niagara/Source/NiagaraEditor/Private/Widgets/SNiagaraSimCacheDebugDataView.h
+++ b/Engine/Plugins/FX/Niagara/Source/NiagaraEditor/Private/Widgets/SNiagaraSimCacheDebugDataView.h
@@ -5,6 +5,7 @@
 #include "Widgets/SCompoundWidget.h"
 
 class FNiagaraSimCacheViewModel;
+struct FNiagaraSimCacheDebugDataFrame;
 namespace UE::Niagara::SimCache::DebugDataUI
 {
 	class SParameterStoreListView;
@@ -20,10 +21,14 @@ public:
 	void Construct(const FArguments& InArgs);
 
 private:
+	const FNiagaraSimCacheDebugDataFrame* GetCurrentFrameData() const;
+	TSharedRef<SWidget> GetParameterStoreSelectionMenu();
+
 	void RefreshContents();
 	void RefreshContents(bool);
 
 private:
-	TSharedPtr<FNiagaraSimCacheViewModel>		SimCacheViewModel;
-	TSharedPtr<UE::Niagara::SimCache::DebugDataUI::SParameterStoreListView>			OverrideParametersWidget;
-};
\ No newline at end of file
+	TSharedPtr<FNiagaraSimCacheViewModel> SimCacheViewModel;
+	TSharedPtr<UE::Niagara::SimCache::DebugDataUI::SParameterStoreListView> OverrideParametersWidget;
+	FString	SelectedParameterStoreName;
+};
