- Set owner for the parameter store, otherwise data interface copy will fail

#rnx
#tests Tested cache capture with more complex emitters with many data interfaces

diff --git a/Engine/Plugins/FX/Niagara/Source/Niagara/Private/NiagaraSimCacheDebugData.cpp b/Engine/Plugins/FX/Niagara/Source/Niagara/Private/NiagaraSimCacheDebugData.cpp
--- a/Engine/Plugins/FX/Niagara/Source/Niagara/Private/NiagaraSimCacheDebugData.cpp
+++ b/Engine/Plugins/FX/Niagara/Source/Niagara/Private/NiagaraSimCacheDebugData.cpp
@@ -7,30 +7,33 @@
 namespace NSCDebugDataPrivate
 {
 	template<typename TStringType>
-	void AddParameterStore(FNiagaraSimCacheDebugDataFrame& FrameData, TStringType StoreName, const FNiagaraParameterStore& ParameterStore)
+	void AddParameterStore(UObject* Owner, FNiagaraSimCacheDebugDataFrame& FrameData, TStringType StoreName, const FNiagaraParameterStore& ParameterStore)
 	{
 		if (ParameterStore.Num() > 0)
 		{
-			FrameData.DebugParameterStores.FindOrAdd(StoreName) = ParameterStore;
+			FNiagaraParameterStore& DestStore = FrameData.DebugParameterStores.FindOrAdd(StoreName);
+			DestStore.SetOwner(Owner);
+			DestStore = ParameterStore;
 		}
 	}
 
 	template<typename TStringType>
-	void AddParameterStore(FNiagaraSimCacheDebugDataFrame& FrameData, TStringType StoreName, const FNiagaraScriptExecutionContextBase* ExecContext)
+	void AddParameterStore(UObject* Owner, FNiagaraSimCacheDebugDataFrame& FrameData, TStringType StoreName, const FNiagaraScriptExecutionContextBase* ExecContext)
 	{
 		if (ExecContext && ExecContext->Parameters.ReadParameterVariables().Num() > 0)
 		{
 			FNiagaraParameterStore& DestStore = FrameData.DebugParameterStores.FindOrAdd(StoreName);
+			DestStore.SetOwner(Owner);
 			ExecContext->Parameters.CopyParametersTo(DestStore, false, FNiagaraParameterStore::EDataInterfaceCopyMethod::Value);
 		}
 	}
 
 	template<typename TStringType>
-	void AddParameterStore(FNiagaraSimCacheDebugDataFrame& FrameData, TStringType StoreName, const FNiagaraComputeExecutionContext* ExecContext)
+	void AddParameterStore(UObject* Owner, FNiagaraSimCacheDebugDataFrame& FrameData, TStringType StoreName, const FNiagaraComputeExecutionContext* ExecContext)
 	{
 		if (ExecContext)
 		{
-			AddParameterStore(FrameData, StoreName, ExecContext->CombinedParamStore);
+			AddParameterStore(Owner, FrameData, StoreName, ExecContext->CombinedParamStore);
 		}
 	}
 }
@@ -50,17 +53,17 @@ void UNiagaraSimCacheDebugData::CaptureFrame(FNiagaraSimCacheHelper& Helper, int
 	// Add Override Parameters
 	if (const FNiagaraParameterStore* OverrideParameterStore = Helper.SystemInstance->GetOverrideParameters())
 	{
-		AddParameterStore(FrameData, TEXT("OverrideParameters"), *OverrideParameterStore);
+		AddParameterStore(this, FrameData, TEXT("OverrideParameters"), *OverrideParameterStore);
 	}
 
 	// Add Instance Parameters
-	AddParameterStore(FrameData, TEXT("InstanceParameters"), Helper.SystemInstance->GetInstanceParameters());
+	AddParameterStore(this, FrameData, TEXT("InstanceParameters"), Helper.SystemInstance->GetInstanceParameters());
 
 	// Add System Script Parameters
 	if (FNiagaraSystemSimulation* SystemSimulation = Helper.SystemInstance->GetSystemSimulation().Get())
 	{
-		AddParameterStore(FrameData, TEXT("System Spawn"), SystemSimulation->GetSpawnExecutionContext());
-		AddParameterStore(FrameData, TEXT("System Update"), SystemSimulation->GetUpdateExecutionContext());
+		AddParameterStore(this, FrameData, TEXT("System Spawn"), SystemSimulation->GetSpawnExecutionContext());
+		AddParameterStore(this, FrameData, TEXT("System Update"), SystemSimulation->GetUpdateExecutionContext());
 	}
 
 	// Add Emitter Parameters
@@ -69,13 +72,13 @@ void UNiagaraSimCacheDebugData::CaptureFrame(FNiagaraSimCacheHelper& Helper, int
 		const FNiagaraEmitterHandle& EmitterHandle = EmitterRef->GetEmitterHandle();
 		const FString EmitterName = EmitterHandle.GetName().ToString();
 
-		AddParameterStore(FrameData, FString::Printf(TEXT("%s RendererBindings"), *EmitterName), EmitterRef->GetRendererBoundVariables());
-		AddParameterStore(FrameData, FString::Printf(TEXT("%s GPUContext"), *EmitterName), EmitterRef->GetGPUContext());
+		AddParameterStore(this, FrameData, FString::Printf(TEXT("%s RendererBindings"), *EmitterName), EmitterRef->GetRendererBoundVariables());
+		AddParameterStore(this, FrameData, FString::Printf(TEXT("%s GPUContext"), *EmitterName), EmitterRef->GetGPUContext());
 
 		if (FNiagaraEmitterInstanceImpl* StatefulEmitter = EmitterRef->AsStateful())
 		{
-			AddParameterStore(FrameData, FString::Printf(TEXT("%s Spawn"), *EmitterName), &StatefulEmitter->GetSpawnExecutionContext());
-			AddParameterStore(FrameData, FString::Printf(TEXT("%s Update"), *EmitterName), &StatefulEmitter->GetUpdateExecutionContext());
+			AddParameterStore(this, FrameData, FString::Printf(TEXT("%s Spawn"), *EmitterName), &StatefulEmitter->GetSpawnExecutionContext());
+			AddParameterStore(this, FrameData, FString::Printf(TEXT("%s Update"), *EmitterName), &StatefulEmitter->GetUpdateExecutionContext());
 		}
 	}
 }
