Clear out resolved data on scripts belonging to standalone emitters when they are updated from the transient system during editing.  This prevents an issue where referenced objects parented to the transient system could be getting saved (incorrectly) into the standalone emitter.
#tests Using an emitter from a UDN licensee, open a standalone emitter which used a number of DI in the emitter scripts; on saving and reopening these would generate an ensure because the resolved DIs were null.  This was resolved by clearing out entries when updating the standalone emitter.

diff --git a/Engine/Plugins/FX/Niagara/Source/Niagara/Classes/NiagaraScript.h b/Engine/Plugins/FX/Niagara/Source/Niagara/Classes/NiagaraScript.h
--- a/Engine/Plugins/FX/Niagara/Source/Niagara/Classes/NiagaraScript.h
+++ b/Engine/Plugins/FX/Niagara/Source/Niagara/Classes/NiagaraScript.h
@@ -879,6 +879,10 @@ public:
 	/** Creates a shallow transient copy of this script for compilation purposes. */
 	NIAGARA_API UNiagaraScript* CreateCompilationCopy();
 
+	/** Clears out resolved runtime data.  For scripts belonging to standalone emitters the resolved data could be references to the transient
+		NiagaraSystem used for editing.  In the future all this data will be moved to live on the System itself so this will no longer be required. */
+	NIAGARA_API void ClearResolvedData();
+
 	/** A set of rules to apply when this script is used in the stack. To create your own rules, write a custom class that extends UNiagaraValidationRule. */
 	UPROPERTY(EditAnywhere, Category = "Validation", Instanced)
 	TArray<TObjectPtr<UNiagaraValidationRule>> ValidationRules;
diff --git a/Engine/Plugins/FX/Niagara/Source/Niagara/Private/NiagaraScript.cpp b/Engine/Plugins/FX/Niagara/Source/Niagara/Private/NiagaraScript.cpp
--- a/Engine/Plugins/FX/Niagara/Source/Niagara/Private/NiagaraScript.cpp
+++ b/Engine/Plugins/FX/Niagara/Source/Niagara/Private/NiagaraScript.cpp
@@ -794,6 +794,15 @@ UNiagaraScript* UNiagaraScript::CreateCompilationCopy()
 	return Result;
 }
 
+// intended to be used for scripts belonging to standalone emitters.  This ensures that no runtime data generated from
+// transient systems are stored in the script.
+void UNiagaraScript::ClearResolvedData()
+{
+	ResolvedDataInterfaces.Empty();
+	ResolvedUserDataInterfaceBindings.Empty();
+	ResolvedUObjectInfos.Empty();
+}
+
 #endif
 
 FVersionedNiagaraEmitter UNiagaraScript::GetOuterEmitter() const
diff --git a/Engine/Plugins/FX/Niagara/Source/NiagaraEditor/Private/Toolkits/NiagaraSystemToolkit.cpp b/Engine/Plugins/FX/Niagara/Source/NiagaraEditor/Private/Toolkits/NiagaraSystemToolkit.cpp
--- a/Engine/Plugins/FX/Niagara/Source/NiagaraEditor/Private/Toolkits/NiagaraSystemToolkit.cpp
+++ b/Engine/Plugins/FX/Niagara/Source/NiagaraEditor/Private/Toolkits/NiagaraSystemToolkit.cpp
@@ -1146,6 +1146,17 @@ void FNiagaraSystemToolkit::UpdateOriginalEmitter()
 		// Restore RF_Standalone and RF_Public on the original emitter, as it had been removed from the preview emitter so that it could be GC'd.
 		Source->SetFlags(RF_Standalone | RF_Public);
 
+		// clear out any resolved data on the scripts for the emitter so that we don't reference anything connected to the transient system
+		Source->ForEachVersionData([&](const FVersionedNiagaraEmitterData& EmitterData)
+		{
+			TArray<UNiagaraScript*> AllSourceScripts;
+			EmitterData.GetScripts(AllSourceScripts, true);
+			for (UNiagaraScript* SourceScript : AllSourceScripts)
+			{
+				SourceScript->ClearResolvedData();
+			}
+		});
+
 		Source->PostEditChange();
 
 		TArray<UNiagaraScript*> EmitterScripts;
