Memory savings: Niagara
UNiagaraRendererProperties 176 -> 168
UNiagaraSystem 1376 -> 1352
FNiagaraEmitterHandle 72 -> 64

Saves 0.5 MB (MB3)

#rnx

diff --git a/Engine/Plugins/FX/Niagara/Source/Niagara/Classes/NiagaraSystem.h b/Engine/Plugins/FX/Niagara/Source/Niagara/Classes/NiagaraSystem.h
--- a/Engine/Plugins/FX/Niagara/Source/Niagara/Classes/NiagaraSystem.h
+++ b/Engine/Plugins/FX/Niagara/Source/Niagara/Classes/NiagaraSystem.h
@@ -228,7 +228,7 @@ public:
 	NIAGARA_API ~UNiagaraSystem();
 
 	//~ UObject interface
-	NIAGARA_API void PostInitProperties();
+	NIAGARA_API virtual void PostInitProperties() override;
 	NIAGARA_API virtual void Serialize(FArchive& Ar) override;
 	NIAGARA_API virtual void PostLoad() override;
 	NIAGARA_API virtual void PostDuplicate(bool bDuplicateForPIE) override;
@@ -335,7 +335,7 @@ public:
 		return EmitterHandles[Idx];
 	};
 
-	int GetNumEmitters()
+	int GetNumEmitters() const
 	{
 		return EmitterHandles.Num();
 	}
@@ -512,9 +512,8 @@ public:
 
 	FORCEINLINE bool IsInitialOwnerVelocityFromActor() const { return bInitialOwnerVelocityFromActor; }
 
-	NIAGARA_API void ReportAnalyticsData(bool bIsCooking);
-
 #if WITH_EDITORONLY_DATA
+	NIAGARA_API void ReportAnalyticsData(bool bIsCooking);
 	NIAGARA_API bool UsesEmitter(UNiagaraEmitter* Emitter) const;
 	NIAGARA_API bool UsesEmitter(const FVersionedNiagaraEmitter& VersionedEmitter) const;
 	NIAGARA_API bool UsesScript(const UNiagaraScript* Script)const; 
@@ -717,18 +716,23 @@ public:
 	}
 	
 	UPROPERTY(EditAnywhere, Category = "Debug", Transient, AdvancedDisplay)
-	bool bDumpDebugSystemInfo = false;
+	uint8 bDumpDebugSystemInfo: 1 = false;
 
 	UPROPERTY(EditAnywhere, Category = "Debug", Transient, AdvancedDisplay)
-	bool bDumpDebugEmitterInfo = false;
+	uint8 bDumpDebugEmitterInfo: 1  = false;
 
-	bool bFullyLoaded = false;
+	uint8 bFullyLoaded: 1  = false;
 
 	/** When enabled, we follow the settings on the UNiagaraComponent for tick order. When this option is disabled, we ignore any dependencies from data interfaces or other variables and instead fire off the simulation as early in the frame as possible. This greatly
 	reduces overhead and allows the game thread to run faster, but comes at a tradeoff if the dependencies might leave gaps or other visual artifacts.*/
 	UPROPERTY(EditAnywhere, Category = "Performance", AdvancedDisplay)
-	bool bRequireCurrentFrameData = true;
+	uint8 bRequireCurrentFrameData: 1  = true;
 
+protected:
+	UPROPERTY(EditAnywhere, Category = "Scalability", meta=(DisplayInScalabilityContext))
+	uint8 bOverrideScalabilitySettings : 1;
+public:
+	
 	FORCEINLINE bool HasDIsWithPostSimulateTick() const { return bHasDIsWithPostSimulateTick; }
 	FORCEINLINE bool AllDIsPostSimulateCanOverlapFrames() const { return bAllDIsPostSimulateCanOverlapFrames; }
 	FORCEINLINE bool AsyncWorkCanOverlapTickGroups() const { return bAllDIsPostStageCanOverlapTickGroups; }
@@ -860,9 +864,6 @@ protected:
 	UPROPERTY(EditAnywhere, Category = "System")
 	TObjectPtr<UNiagaraEffectType> EffectType;
 
-	UPROPERTY(EditAnywhere, Category = "Scalability", meta=(DisplayInScalabilityContext))
-	bool bOverrideScalabilitySettings;
-
 	/** Controls whether we should override the Effect Type value for bAllowCullingForLocalPlayers. */
 	UPROPERTY(EditAnywhere, Category = "Scalability", meta = (InlineEditConditionToggle, EditCondition = bOverrideScalabilitySettings))
 	uint32 bOverrideAllowCullingForLocalPlayers : 1;
@@ -941,12 +942,18 @@ protected:
 	UPROPERTY(EditAnywhere, Category = "System", meta = (SkipSystemResetOnChange = "true", EditCondition = "bFixedBounds"))
 	FBox FixedBounds;
 
+	UPROPERTY()
+	bool bNeedsGPUContextInitForDataInterfaces;
+
 	/**
 	When disabled we will generate a RandomSeed per instance on reset which is not deterministic.
 	When enabled we will always use the RandomSeed from the system plus the components RandomSeedOffset, this allows for determinism but variance between components.
 	*/
 	UPROPERTY(EditAnywhere, Category = "System")
 	bool bDeterminism = false;
+	
+	UPROPERTY(EditAnywhere, Category = "System", meta = (InlineEditConditionToggle, DisplayAfter="WarmupTickDelta"))
+	bool bFixedTickDelta = false;
 
 	/** Seed used for system script random number generator. */
 	UPROPERTY(EditAnywhere, Category = "System", meta = (EditCondition = "bDeterminism", EditConditionHides))
@@ -964,9 +971,6 @@ protected:
 	UPROPERTY(EditAnywhere, Category = "System", meta = (ForceUnits=s, EditCondition = "WarmupTime > 0.0", EditConditionHides))
 	float WarmupTickDelta = 1.0f / 15.0f;
 
-	UPROPERTY(EditAnywhere, Category = "System", meta = (InlineEditConditionToggle))
-	bool bFixedTickDelta = false;
-
 	/**
 	If activated, the system ticks with a fixed delta time instead of the varying game thread delta time. This leads to much more stable simulations.
 	When the fixed tick delta is smaller than the game thread tick time the simulation is substepping by executing multiple ticks per frame.
@@ -987,9 +991,6 @@ protected:
 	TObjectPtr<UNiagaraBakerSettings> BakerGeneratedSettings;
 #endif
 
-	UPROPERTY()
-	bool bNeedsGPUContextInitForDataInterfaces;
-
 	/** Array of emitter indices sorted by execution priority. The emitters will be ticked in this order. Please note that some indices may have the top bit set (kStartNewOverlapGroupBit)
 	* to indicate synchronization points in parallel execution, so mask it out before using the values as indices in the emitters array.
 	*/
@@ -1041,10 +1042,10 @@ protected:
 
 	//Scalability settings
 	FNiagaraSystemScalabilitySettings& CurrentScalabilitySettings;
-	bool bAllowCullingForLocalPlayers = false;
 
 	mutable FString CrashReporterTag;
 
+	uint32 bAllowCullingForLocalPlayers : 1 = false;
 	uint32 bHasDIsWithPostSimulateTick : 1;
 	uint32 bAllDIsPostSimulateCanOverlapFrames : 1;
 	uint32 bAllDIsPostStageCanOverlapTickGroups : 1;
diff --git a/Engine/Plugins/FX/Niagara/Source/Niagara/Internal/NiagaraSystemEmitterState.h b/Engine/Plugins/FX/Niagara/Source/Niagara/Internal/NiagaraSystemEmitterState.h
--- a/Engine/Plugins/FX/Niagara/Source/Niagara/Internal/NiagaraSystemEmitterState.h
+++ b/Engine/Plugins/FX/Niagara/Source/Niagara/Internal/NiagaraSystemEmitterState.h
@@ -8,7 +8,7 @@
 #include "NiagaraSystemEmitterState.generated.h"
 
 UENUM()
-enum class ENiagaraSystemInactiveResponse
+enum class ENiagaraSystemInactiveResponse : uint8
 {
 	/** Let Emitters Finish Then Kill Emitter */
 	Complete,
@@ -17,7 +17,7 @@ enum class ENiagaraSystemInactiveResponse
 };
 
 UENUM()
-enum class ENiagaraEmitterInactiveResponse
+enum class ENiagaraEmitterInactiveResponse : uint8
 {
 	/** Let Particles Finish Then Kill Emitter */
 	Complete,
@@ -28,7 +28,7 @@ enum class ENiagaraEmitterInactiveResponse
 };
 
 UENUM()
-enum class ENiagaraLoopBehavior
+enum class ENiagaraLoopBehavior : uint8
 {
 	Infinite,
 	Multiple,
diff --git a/Engine/Plugins/FX/Niagara/Source/Niagara/Public/NiagaraCommon.h b/Engine/Plugins/FX/Niagara/Source/Niagara/Public/NiagaraCommon.h
--- a/Engine/Plugins/FX/Niagara/Source/Niagara/Public/NiagaraCommon.h
+++ b/Engine/Plugins/FX/Niagara/Source/Niagara/Public/NiagaraCommon.h
@@ -107,7 +107,7 @@ enum class ENiagaraGpuBufferFormat : uint8
 };
 
 UENUM()
-enum class ENiagaraGpuSyncMode
+enum class ENiagaraGpuSyncMode : uint8
 {
 	/** Data will not be automatically pushed and could diverge between Cpu & Gpu. */
 	None,
@@ -147,7 +147,7 @@ enum class ENiagaraDefaultMode : uint8
 
 /** How to handle how Niagara rendered effects should generate motion vectors by default (can still be overridden on a case-by-case basis) */
 UENUM()
-enum class ENiagaraDefaultRendererMotionVectorSetting
+enum class ENiagaraDefaultRendererMotionVectorSetting : uint8
 {
 	/**
 	 * Motion vectors generated are precise (ideal for motion blur and temporal anti-aliasing).
@@ -163,7 +163,7 @@ enum class ENiagaraDefaultRendererMotionVectorSetting
 
 /** How a given Niagara renderer should handle motion vector generation. */
 UENUM()
-enum class ENiagaraRendererMotionVectorSetting
+enum class ENiagaraRendererMotionVectorSetting : uint8
 {
 	/** Determines the best method to employ when generating motion vectors (accurate vs. approximate) based on project and renderer settings */
 	AutoDetect,
