Fixes presizing buffer for bones and sockets when using Niagara's skel mesh DI
#tests compared replay performance

diff --git a/Engine/Plugins/FX/Niagara/Source/Niagara/Private/NiagaraDataInterfaceSkeletalMesh.cpp b/Engine/Plugins/FX/Niagara/Source/Niagara/Private/NiagaraDataInterfaceSkeletalMesh.cpp
--- a/Engine/Plugins/FX/Niagara/Source/Niagara/Private/NiagaraDataInterfaceSkeletalMesh.cpp
+++ b/Engine/Plugins/FX/Niagara/Source/Niagara/Private/NiagaraDataInterfaceSkeletalMesh.cpp
@@ -1236,7 +1236,10 @@ void FSkeletalMeshGpuDynamicBufferProxy::NewFrame(const FNDISkeletalMesh_Instanc
 			}
 
 			// Fill BoneSamplingData
-			BoneSamplingData.Reserve((SamplingBoneCount + SamplingSocketCount) * 2);
+			const TArray<FTransform3f>& FilteredSocketsCurrBuffer = InstanceData->GetFilteredSocketsCurrBuffer();
+			const int32 BoneSamplingDataCount = (BoneTransforms.Num() + FilteredSocketsCurrBuffer.Num()) * 3;
+
+			BoneSamplingData.Reserve(BoneSamplingDataCount);
 			for (int i=0; i < BoneTransforms.Num(); ++i )
 			{
 				const FTransform& BoneTransform = BoneTransforms[i];
@@ -1248,7 +1251,7 @@ void FSkeletalMeshGpuDynamicBufferProxy::NewFrame(const FNDISkeletalMesh_Instanc
 			}
 
 			// Append sockets
-			for (const FTransform3f& SocketTransform : InstanceData->GetFilteredSocketsCurrBuffer())
+			for (const FTransform3f& SocketTransform : FilteredSocketsCurrBuffer)
 			{
 				const FQuat4f Rotation = SocketTransform.GetRotation();
 				const int32 ParentIndex = -1;
