Reduce calls to USoundCue::IsLooping() with the Niagara audio player DI
#tests compared replay performance

diff --git a/Engine/Plugins/FX/Niagara/Source/Niagara/Classes/NiagaraDataInterfaceAudioPlayer.h b/Engine/Plugins/FX/Niagara/Source/Niagara/Classes/NiagaraDataInterfaceAudioPlayer.h
--- a/Engine/Plugins/FX/Niagara/Source/Niagara/Classes/NiagaraDataInterfaceAudioPlayer.h
+++ b/Engine/Plugins/FX/Niagara/Source/Niagara/Classes/NiagaraDataInterfaceAudioPlayer.h
@@ -80,6 +80,7 @@ struct FAudioPlayerInterface_InstanceData
 	int32 MaxPlaysPerTick = 0;
 	bool bStopWhenComponentIsDestroyed = true;
 	bool bValidOneShotSound = false;
+	bool bSoundToPlayIsLooping = false;
 #if WITH_EDITORONLY_DATA
 	bool bOnlyActiveDuringGameplay = false;
 #endif
diff --git a/Engine/Plugins/FX/Niagara/Source/Niagara/Private/NiagaraDataInterfaceAudioPlayer.cpp b/Engine/Plugins/FX/Niagara/Source/Niagara/Private/NiagaraDataInterfaceAudioPlayer.cpp
--- a/Engine/Plugins/FX/Niagara/Source/Niagara/Private/NiagaraDataInterfaceAudioPlayer.cpp
+++ b/Engine/Plugins/FX/Niagara/Source/Niagara/Private/NiagaraDataInterfaceAudioPlayer.cpp
@@ -139,10 +139,15 @@ bool UNiagaraDataInterfaceAudioPlayer::PerInstanceTick(void* PerInstanceData, FN
 	
 	if (IsValid(SoundToPlay) && SystemInstance)
 	{
-		PIData->SoundToPlay = SoundToPlay;
+		if (PIData->SoundToPlay != SoundToPlay)
+		{
+			PIData->SoundToPlay = SoundToPlay;
+			PIData->bSoundToPlayIsLooping = SoundToPlay->IsLooping();
+		}
+
 		PIData->Attenuation = Attenuation;
 		PIData->Concurrency = Concurrency;
-		PIData->bValidOneShotSound = SoundToPlay->IsLooping() ? bAllowLoopingOneShotSounds : true;
+		PIData->bValidOneShotSound = bAllowLoopingOneShotSounds || !PIData->bSoundToPlayIsLooping;
 
 		if (ConfigurationUserParameter.Parameter.IsValid())
 		{
@@ -158,6 +163,7 @@ bool UNiagaraDataInterfaceAudioPlayer::PerInstanceTick(void* PerInstanceData, FN
 		PIData->Concurrency.Reset();
 		PIData->CachedUserParam.Reset();
 		PIData->bValidOneShotSound = false;
+		PIData->bSoundToPlayIsLooping = false;
 	}
 
 	
