Small improvements to component renderer runtime to avoid string creation and static fname creation each frame
#tests compared replay performance

diff --git a/Engine/Plugins/FX/Niagara/Source/Niagara/Private/NiagaraRendererComponents.cpp b/Engine/Plugins/FX/Niagara/Source/Niagara/Private/NiagaraRendererComponents.cpp
--- a/Engine/Plugins/FX/Niagara/Source/Niagara/Private/NiagaraRendererComponents.cpp
+++ b/Engine/Plugins/FX/Niagara/Source/Niagara/Private/NiagaraRendererComponents.cpp
@@ -10,6 +10,7 @@
 #include "NiagaraComponentRendererProperties.h"
 #include "NiagaraSystem.h"
 #include "NiagaraSystemInstance.h"
+#include "String/ParseTokens.h"
 #include "Async/Async.h"
 
 #if WITH_EDITOR
@@ -35,6 +36,11 @@ static FAutoConsoleVariableRef CVarNiagaraComponentRenderPoolInactiveTimeLimit(
 	ECVF_Default
 );
 
+namespace NiagaraRendererComponentsImpl
+{
+	static const FName NAME_UniqueID("UniqueID");
+}
+
 //////////////////////////////////////////////////////////////////////////
 
 template<typename T>
@@ -161,7 +167,11 @@ void InvokeSetterFunction(UObject* InRuntimeObject, UFunction* Setter, const uin
 			if (Property->IsInContainer(ParmsSize))
 			{
 				// Check for a default value
-				const FString* DefaultValue = SetterDefaultValues.Find(Property->GetName());
+				FNameBuilder PropertyName(Property->GetFName());
+				FStringView PropertyNameView = PropertyName.ToView();
+				uint32 PropertyNameHash = GetTypeHash(PropertyNameView);
+
+				const FString* DefaultValue = SetterDefaultValues.FindByHash(PropertyNameHash, PropertyNameView);
 				if (DefaultValue)
 				{
 					Property->ImportText_Direct(**DefaultValue, Property->ContainerPtrToValuePtr<uint8>(Params), InRuntimeObject, PPF_None);
@@ -414,7 +424,7 @@ void FNiagaraRendererComponents::PostSystemTick_GameThread(const UNiagaraRendere
 	const FNiagaraDataBuffer& ParticleData = Data.GetCurrentDataChecked();
 	FNiagaraDataSetReaderInt32<FNiagaraBool> EnabledAccessor = FNiagaraDataSetAccessor<FNiagaraBool>::CreateReader(Data, Properties->EnabledBinding.GetDataSetBindableVariable().GetName());
 	FNiagaraDataSetReaderInt32<int32> VisTagAccessor = FNiagaraDataSetAccessor<int32>::CreateReader(Data, Properties->RendererVisibilityTagBinding.GetDataSetBindableVariable().GetName());
-	FNiagaraDataSetReaderInt32<int32> UniqueIDAccessor = FNiagaraDataSetAccessor<int32>::CreateReader(Data, FName("UniqueID"));
+	FNiagaraDataSetReaderInt32<int32> UniqueIDAccessor = FNiagaraDataSetAccessor<int32>::CreateReader(Data, NiagaraRendererComponentsImpl::NAME_UniqueID);
 
 	const bool bIsRendererEnabled = IsRendererEnabled(Properties, Emitter);
 
@@ -813,10 +823,10 @@ void FNiagaraRendererComponents::TickPropertyBindings(
 	}
 }
 
-FNiagaraRendererComponents::FComponentPropertyAddress FNiagaraRendererComponents::FindPropertyRecursive(void* BasePointer, UStruct* InStruct, TArray<FString>& InPropertyNames, uint32 Index)
+FNiagaraRendererComponents::FComponentPropertyAddress FNiagaraRendererComponents::FindPropertyRecursive(void* BasePointer, UStruct* InStruct, const TArray<FName>& InPropertyNames, uint32 Index)
 {
 	FComponentPropertyAddress NewAddress;
-	FProperty* Property = FindFProperty<FProperty>(InStruct, *InPropertyNames[Index]);
+	FProperty* Property = FindFProperty<FProperty>(InStruct, InPropertyNames[Index]);
 
 	if (FStructProperty* StructProp = CastField<FStructProperty>(Property))
 	{
@@ -844,8 +854,8 @@ FNiagaraRendererComponents::FComponentPropertyAddress FNiagaraRendererComponents
 
 FNiagaraRendererComponents::FComponentPropertyAddress FNiagaraRendererComponents::FindProperty(const UObject& InObject, const FString& InPropertyPath)
 {
-	TArray<FString> PropertyNames;
-	InPropertyPath.ParseIntoArray(PropertyNames, TEXT("."), true);
+	TArray<FName> PropertyNames;
+	UE::String::ParseTokens(InPropertyPath, TEXT('.'), [&PropertyNames](FStringView Token) { PropertyNames.Emplace(Token); });
 
 	if (IsValid(&InObject) && PropertyNames.Num() > 0)
 	{
diff --git a/Engine/Plugins/FX/Niagara/Source/Niagara/Public/NiagaraRendererComponents.h b/Engine/Plugins/FX/Niagara/Source/Niagara/Public/NiagaraRendererComponents.h
--- a/Engine/Plugins/FX/Niagara/Source/Niagara/Public/NiagaraRendererComponents.h
+++ b/Engine/Plugins/FX/Niagara/Source/Niagara/Public/NiagaraRendererComponents.h
@@ -71,7 +71,7 @@ private:
 		const FNiagaraLWCConverter& LwcConverter);
 
 	// These property accessor methods are largely copied over from MovieSceneCommonHelpers.h
-	static NIAGARA_API FComponentPropertyAddress FindPropertyRecursive(void* BasePointer, UStruct* InStruct, TArray<FString>& InPropertyNames, uint32 Index);
+	static NIAGARA_API FComponentPropertyAddress FindPropertyRecursive(void* BasePointer, UStruct* InStruct, const TArray<FName>& InPropertyNames, uint32 Index);
 	static NIAGARA_API FComponentPropertyAddress FindProperty(const UObject& Object, const FString& InPropertyPath);
 
 	// this key is used to check if the template object was changed, e.g. a blueprint compilation
