Fixed problems with module version upgrades:
* fixed crash when setting bool values from python
* fixed old inputs creating a warning after version upgrade
* fixed a crash when using the quick fix version upgrade option

#tests Tested switching versions of the curl noise module with "remap range" being on and off

diff --git a/Engine/Plugins/FX/Niagara/Source/NiagaraEditor/Private/NiagaraNodeFunctionCall.cpp b/Engine/Plugins/FX/Niagara/Source/NiagaraEditor/Private/NiagaraNodeFunctionCall.cpp
--- a/Engine/Plugins/FX/Niagara/Source/NiagaraEditor/Private/NiagaraNodeFunctionCall.cpp
+++ b/Engine/Plugins/FX/Niagara/Source/NiagaraEditor/Private/NiagaraNodeFunctionCall.cpp
@@ -1551,16 +1551,25 @@ void UNiagaraNodeFunctionCall::UpdateOverridePins(const FNiagaraScriptVersionUpg
 	if (FunctionScript)
 	{
 		// Automatically remove old inputs so it does not show a bunch of warnings to the user
-		TMap<FName, FNiagaraTypeDefinition> FunctionInputNames;
-		FPinCollectorArray OverridePins;
+		FPinCollectorArray InputPins;
+		GetInputPins(InputPins);
+		for (UEdGraphPin* Pin : InputPins)
+		{
+			if (Pin->bOrphanedPin)
+			{
+				RemovePin(Pin);
+			}
+		}
 		UNiagaraNodeParameterMapSet* OverrideNode = FNiagaraStackGraphUtilities::GetStackFunctionOverrideNode(*this);
 		if (OverrideNode != nullptr)
 		{
+			FPinCollectorArray OverridePins;
 			OverrideNode->Modify();
 			OverrideNode->GetInputPins(OverridePins);
 
 			if (!OverridePins.IsEmpty())
 			{
+				TMap<FName, FNiagaraTypeDefinition> FunctionInputNames;
 				TArray<FNiagaraVariable> ModuleInputVariables;
 				FNiagaraStackGraphUtilities::GetStackFunctionInputs(*this, ModuleInputVariables, UpgradeContext.ConstantResolver, FNiagaraStackGraphUtilities::ENiagaraGetStackFunctionInputPinsOptions::ModuleInputsOnly);
 				for (const FNiagaraVariable& InputVariable : ModuleInputVariables)
diff --git a/Engine/Plugins/FX/Niagara/Source/NiagaraEditor/Private/UpgradeNiagaraScriptResults.cpp b/Engine/Plugins/FX/Niagara/Source/NiagaraEditor/Private/UpgradeNiagaraScriptResults.cpp
--- a/Engine/Plugins/FX/Niagara/Source/NiagaraEditor/Private/UpgradeNiagaraScriptResults.cpp
+++ b/Engine/Plugins/FX/Niagara/Source/NiagaraEditor/Private/UpgradeNiagaraScriptResults.cpp
@@ -232,7 +232,7 @@ void UUpgradeNiagaraScriptResults::SetBoolInput(const FString& InputName, bool V
 	UNiagaraPythonScriptModuleInput* ModuleInput = GetNewInput(FName(InputName));
 	if (ModuleInput && ModuleInput->Input->InputType == FNiagaraTypeDefinition::GetBoolDef())
 	{
-		NiagaraScriptResults::SetValue(ModuleInput, Value);
+		ModuleInput->Input = UNiagaraClipboardEditorScriptingUtilities::CreateBoolLocalValueInput(ModuleInput, FName(InputName), false, false, Value);
 	}
 }
 
diff --git a/Engine/Plugins/FX/Niagara/Source/NiagaraEditor/Private/ViewModels/Stack/NiagaraStackModuleItem.cpp b/Engine/Plugins/FX/Niagara/Source/NiagaraEditor/Private/ViewModels/Stack/NiagaraStackModuleItem.cpp
--- a/Engine/Plugins/FX/Niagara/Source/NiagaraEditor/Private/ViewModels/Stack/NiagaraStackModuleItem.cpp
+++ b/Engine/Plugins/FX/Niagara/Source/NiagaraEditor/Private/ViewModels/Stack/NiagaraStackModuleItem.cpp
@@ -855,32 +855,36 @@ UNiagaraStackEntry::FStackIssueFixDelegate UNiagaraStackModuleItem::GetUpgradeVe
 	{
 		return FStackIssueFixDelegate();
 	}
-	return FStackIssueFixDelegate::CreateLambda([this]()
+	return FStackIssueFixDelegate::CreateLambda([WeakModule=MakeWeakObjectPtr(this)]()
     {
-        FScopedTransaction ScopedTransaction(LOCTEXT("UpgradeVersionFix", "Change module version"));
+	    UNiagaraStackModuleItem* Item = WeakModule.Get();
+		if (!Item)
+		{
+			return;
+		}
+		
+	    FScopedTransaction ScopedTransaction(LOCTEXT("UpgradeVersionFix", "Change module version"));
+		TStrongObjectPtr ModuleItem(Item); // gc lock
         FNiagaraScriptVersionUpgradeContext UpgradeContext;
-		UpgradeContext.CreateClipboardCallback = [this](UNiagaraClipboardContent* ClipboardContent)
+		UpgradeContext.CreateClipboardCallback = [ModuleItem](UNiagaraClipboardContent* ClipboardContent)
 	    {
-	        RefreshChildren();
-	        Copy(ClipboardContent);
+	        ModuleItem->RefreshChildren();
+	        ModuleItem->Copy(ClipboardContent);
 	        if (ClipboardContent->Functions.Num() > 0)
 	        {
 	            ClipboardContent->FunctionInputs = ClipboardContent->Functions[0]->Inputs;
 	            ClipboardContent->Functions.Empty();
 	        }
 	    };
-        UpgradeContext.ApplyClipboardCallback = [this](UNiagaraClipboardContent* ClipboardContent, FText& OutWarning) { Paste(ClipboardContent, OutWarning); };
-		UpgradeContext.ConstantResolver = GetEmitterViewModel().IsValid() ?
-	        FCompileConstantResolver(GetEmitterViewModel()->GetEmitter(), FNiagaraStackGraphUtilities::GetOutputNodeUsage(*FunctionCallNode)) :
-	        FCompileConstantResolver(&GetSystemViewModel()->GetSystem(), FNiagaraStackGraphUtilities::GetOutputNodeUsage(*FunctionCallNode));
-        FunctionCallNode->ChangeScriptVersion(FunctionCallNode->FunctionScript->GetExposedVersion().VersionGuid, UpgradeContext, true);
-        if (FunctionCallNode->RefreshFromExternalChanges())
-        {
-            FunctionCallNode->GetNiagaraGraph()->NotifyGraphNeedsRecompile();
-            GetSystemViewModel()->ResetSystem();
-        }
-
-		ReportScriptVersionChange();
+        UpgradeContext.ApplyClipboardCallback = [ModuleItem](UNiagaraClipboardContent* ClipboardContent, FText& OutWarning) { ModuleItem->Paste(ClipboardContent, OutWarning); };
+		UpgradeContext.ConstantResolver = ModuleItem->GetEmitterViewModel().IsValid() ?
+	        FCompileConstantResolver(ModuleItem->GetEmitterViewModel()->GetEmitter(), FNiagaraStackGraphUtilities::GetOutputNodeUsage(*ModuleItem->FunctionCallNode)) :
+	        FCompileConstantResolver(&ModuleItem->GetSystemViewModel()->GetSystem(), FNiagaraStackGraphUtilities::GetOutputNodeUsage(*ModuleItem->FunctionCallNode));
+        FGuid NewScriptVersion = ModuleItem->FunctionCallNode->FunctionScript->GetExposedVersion().VersionGuid;
+        ModuleItem->FunctionCallNode->ChangeScriptVersion(NewScriptVersion, UpgradeContext, true);
+        ModuleItem->Refresh();
+
+		ModuleItem->ReportScriptVersionChange();
     });
 }
 
