Niagara Script Editor - Hierarchy Editor Integration
List of changes
----> Script Hierarchy Editor
- Graph now owns a HierarchyRoot object
- PostLoad previous implicit hierarchy data into explicit hierarchy objects
- 'Edit Input Hierarchy' button in Parameters tab in Script Editor & ScratchPad Parameters tab
- New 'Input Preview' tab in Script Editor. Basic for now, can be expanded in the future
- Renamed SNiagaraHierarchy -> SNiagaraHierarchyEditor (low amount of actual changes in there)
- Added new default style for hierarchy categories and increased padding
- HierarchyItemViewModels for source items can now opt-in to let the details panel display an item. Useful if it represents an external item and the hierarchy object is just a wrapper. Previously we didn't allow this because source items are transient, so editing made no sense, but if they point at persistent data it's alright.
---> Copy Paste
- Added OnCopyPaste Delegate to stack entries. Stack Rows bind to it to play a new pulse animation when copy or paste occurs on a given stack row.
- Copy Paste support for Categories added
- Copy Paste support for input + children inputs added
----> Stack Rework
- Removed UNiagaraStackFunctionInputCollection (Renamed file to UNiagaraStackValueCollection to reflect the remaining class)
- Replaced with new UNiagaraStackScriptHierarchyRoot that fuses new Hierarchy paradigm with stack logic
- StackFunctionInputs now own children inputs

#tests qa test pass & content cook & engine tests

diff --git a/Engine/Plugins/FX/Niagara/Source/Niagara/Classes/NiagaraScript.h b/Engine/Plugins/FX/Niagara/Source/Niagara/Classes/NiagaraScript.h
--- a/Engine/Plugins/FX/Niagara/Source/Niagara/Classes/NiagaraScript.h
+++ b/Engine/Plugins/FX/Niagara/Source/Niagara/Classes/NiagaraScript.h
@@ -27,6 +27,7 @@ class UNiagaraDataInterface;
 class FNiagaraCompileRequestDataBase;
 class FNiagaraCompileRequestDuplicateDataBase;
 class UNiagaraConvertInPlaceUtilityBase;
+class UNiagaraHierarchyRoot;
 
 #define NIAGARA_INVALID_MEMORY (0xBA)
 
@@ -793,9 +794,6 @@ public:
 
 	UPROPERTY(EditAnywhere, Category = Script, DisplayName = "Script Metadata", meta = (ToolTip = "Script Metadata"))
 	TMap<FName, FString> ScriptMetaData;
-	
-	UPROPERTY(EditAnywhere, Category = Script)
-	TArray<FNiagaraStackSection> InputSections;
 
 	/** Adjusted every time ComputeVMCompilationId is called.*/
 	UPROPERTY()
@@ -817,7 +815,11 @@ public:
 	UPROPERTY()
 	TArray<FParameterDefinitionsSubscription> ParameterDefinitionsSubscriptions;
 	NIAGARA_API TArray<ENiagaraScriptUsage> GetSupportedUsageContexts() const;
+	
+	NIAGARA_API class UNiagaraScriptSourceBase* GetSource() { return Source; }
 
+	UPROPERTY()
+	TArray<FNiagaraStackSection> InputSections_DEPRECATED;
 private:
 	friend class UNiagaraScript;
 
@@ -1126,6 +1128,8 @@ public:
 	virtual void PreSave(FObjectPreSaveContext ObjectSaveContext) override;
 	virtual void Serialize(FArchive& Ar)override;
 	virtual void PostLoad() override;
+	
+
 #if WITH_EDITOR
 	virtual void PostEditChangeProperty(FPropertyChangedEvent& PropertyChangedEvent) override;
 	virtual void PostEditChangeVersionedProperty(FPropertyChangedEvent& PropertyChangedEvent, const FGuid& Version);
diff --git a/Engine/Plugins/FX/Niagara/Source/Niagara/Classes/NiagaraScriptSourceBase.h b/Engine/Plugins/FX/Niagara/Source/Niagara/Classes/NiagaraScriptSourceBase.h
--- a/Engine/Plugins/FX/Niagara/Source/Niagara/Classes/NiagaraScriptSourceBase.h
+++ b/Engine/Plugins/FX/Niagara/Source/Niagara/Classes/NiagaraScriptSourceBase.h
@@ -16,6 +16,7 @@ struct FNiagaraParameterStore;
 struct FNiagaraScriptHashCollector;
 class UNiagaraScriptSourceBase;
 struct FNiagaraVMExecutableDataId;
+struct FVersionedNiagaraScriptData;
 
 struct EditorExposedVectorConstant
 {
@@ -154,6 +155,7 @@ class UNiagaraScriptSourceBase : public UObject
 	/** Adds a module if it isn't already in the graph. If the module isn't found bOutFoundModule will be false. If it is found and it did need to be added, the function returns true. If it already exists, it returns false. */
 	virtual bool AddModuleIfMissing(FString ModulePath, ENiagaraScriptUsage Usage, bool& bOutFoundModule) { bOutFoundModule = false; return false; }
 
+	virtual void PostLoadFromOwner(FVersionedNiagaraScriptData& OwnerData) {}
 #if WITH_EDITOR
 	virtual void CleanUpOldAndInitializeNewRapidIterationParameters(const FVersionedNiagaraEmitter& Emitter, ENiagaraScriptUsage ScriptUsage, FGuid ScriptUsageId, FNiagaraParameterStore& RapidIterationParameters) const { checkf(false, TEXT("Not implemented")); }
 
diff --git a/Engine/Plugins/FX/Niagara/Source/Niagara/Public/NiagaraConvertInPlaceUtilityBase.h b/Engine/Plugins/FX/Niagara/Source/Niagara/Public/NiagaraConvertInPlaceUtilityBase.h
--- a/Engine/Plugins/FX/Niagara/Source/Niagara/Public/NiagaraConvertInPlaceUtilityBase.h
+++ b/Engine/Plugins/FX/Niagara/Source/Niagara/Public/NiagaraConvertInPlaceUtilityBase.h
@@ -11,12 +11,12 @@
 class UNiagaraScript;
 class UNiagaraNodeFunctionCall;
 class UNiagaraClipboardContent;
-class UNiagaraStackFunctionInputCollection;
+class UNiagaraStackScriptHierarchyRoot;
 
 UCLASS(Abstract, MinimalAPI)
 class UNiagaraConvertInPlaceUtilityBase : public UObject
 {
 	GENERATED_BODY()
 public:
-	virtual bool Convert(UNiagaraScript* InOldScript, UNiagaraClipboardContent* InOldClipboardContent, UNiagaraScript* InNewScript, UNiagaraStackFunctionInputCollection* InInputCollection, UNiagaraClipboardContent* InNewClipboardContent, UNiagaraNodeFunctionCall* InCallingNode, FText& OutMessage) { return true; };
+	virtual bool Convert(UNiagaraScript* InOldScript, UNiagaraClipboardContent* InOldClipboardContent, UNiagaraScript* InNewScript, UNiagaraStackScriptHierarchyRoot* InHierarchyRoot, UNiagaraClipboardContent* InNewClipboardContent, UNiagaraNodeFunctionCall* InCallingNode, FText& OutMessage) { return true; };
 };
diff --git a/Engine/Plugins/FX/Niagara/Source/Niagara/Public/NiagaraVariableMetaData.h b/Engine/Plugins/FX/Niagara/Source/Niagara/Public/NiagaraVariableMetaData.h
--- a/Engine/Plugins/FX/Niagara/Source/Niagara/Public/NiagaraVariableMetaData.h
+++ b/Engine/Plugins/FX/Niagara/Source/Niagara/Public/NiagaraVariableMetaData.h
@@ -222,17 +222,14 @@ struct FNiagaraVariableMetaData
 		, bOverrideColor(false)
 		, InlineParameterColorOverride(FLinearColor(ForceInit))
 		, bEnableBoolOverride(false)
-		, EditorSortPriority(0)
 		, bInlineEditConditionToggle(false)
 		, bIsStaticSwitch_DEPRECATED(false)
 		, StaticSwitchDefaultValue_DEPRECATED(0)
+		, EditorSortPriority_DEPRECATED(0)
 	{};
 
 	UPROPERTY(EditAnywhere, Category = "Variable", DisplayName="Tooltip", meta = (MultiLine = true, SkipForCompileHash = "true"))
 	FText Description;
-
-	UPROPERTY(EditAnywhere, Category = "Variable", meta = (SkipForCompileHash = "true"))
-	FText CategoryName;
 	
 	/** The unit to display next to input fields for this parameter - note that this is only a visual indicator and does not change any of the calculations. */
 	UPROPERTY(EditAnywhere, Category = "Variable", meta = (SkipForCompileHash = "true"))
@@ -263,9 +260,6 @@ struct FNiagaraVariableMetaData
 	
 	UPROPERTY(EditAnywhere, Category = "Variable", meta = (EditCondition="bEnableBoolOverride", ToolTip = "Useful to override inline bool visualization in the overview.", SkipForCompileHash = "true"))
 	FNiagaraBoolParameterMetaData InlineParameterBoolOverride;
-	
-	UPROPERTY(EditAnywhere, Category = "Variable", meta = (ToolTip = "Affects the sort order in the editor stacks. Use a smaller number to push it to the top. Defaults to zero.", SkipForCompileHash = "true"))
-	int32 EditorSortPriority;
 
 	/** Declares the associated input is used as an inline edit condition toggle, so it should be hidden and edited as a 
 	checkbox inline with the input which was designated as its edit condition. */
@@ -283,9 +277,6 @@ struct FNiagaraVariableMetaData
 	UPROPERTY(EditAnywhere, Category = "Variable", DisplayName = "Property Metadata", meta = (ToolTip = "Property Metadata", SkipForCompileHash = "true"))
 	TMap<FName, FString> PropertyMetaData;
 
-	UPROPERTY(EditAnywhere, Category = "Variable", meta = (ToolTip = "If set, this attribute is visually displayed as a child under the given parent attribute. Currently, only static switches are supported as parent attributes!", SkipForCompileHash = "true"))
-	FName ParentAttribute;
-
 	UPROPERTY(EditAnywhere, Category = "Variable", DisplayName = "Alternate Aliases For Variable", AdvancedDisplay, meta = (ToolTip = "List of alternate/previous names for this variable. Note that this is not normally needed if you rename through the UX. However, if you delete and then add a different variable, intending for it to match, you will likely want to add the prior name here.\n\nYou may need to restart and reload assets after making this change to have it take effect on already loaded assets."))
 	TArray<FName> AlternateAliases;
 
@@ -297,6 +288,12 @@ struct FNiagaraVariableMetaData
 
 	int32 GetStaticSwitchDefaultValue_DEPRECATED() const { return StaticSwitchDefaultValue_DEPRECATED; };
 
+	int32 GetEditorSortPriority_DEPRECATED() const { return EditorSortPriority_DEPRECATED; }
+	
+	FText GetCategoryName_DEPRECATED() const { return CategoryName_DEPRECATED; }
+
+	FName GetParentAttribute_DEPRECATED() const { return ParentAttribute_DEPRECATED; }
+	
 	/** Copies all the properties that are marked as editable for the user (e.g. EditAnywhere). */
 	NIAGARA_API void CopyUserEditableMetaData(const FNiagaraVariableMetaData& OtherMetaData);
 
@@ -322,4 +319,13 @@ private:
 	 */
 	UPROPERTY()
 	int32 StaticSwitchDefaultValue_DEPRECATED;  // TODO: This should be moved to the UNiagaraScriptVariable in the future
+
+	UPROPERTY()
+	FText CategoryName_DEPRECATED;
+
+	UPROPERTY()
+	FName ParentAttribute_DEPRECATED;
+
+	UPROPERTY()
+	int32 EditorSortPriority_DEPRECATED;
 };
diff --git a/Engine/Plugins/FX/Niagara/Source/NiagaraCore/Public/NiagaraCustomVersion.h b/Engine/Plugins/FX/Niagara/Source/NiagaraCore/Public/NiagaraCustomVersion.h
--- a/Engine/Plugins/FX/Niagara/Source/NiagaraCore/Public/NiagaraCustomVersion.h
+++ b/Engine/Plugins/FX/Niagara/Source/NiagaraCore/Public/NiagaraCustomVersion.h
@@ -206,6 +206,8 @@ struct FNiagaraCustomVersion
 		CustomSortingBindingToAge, // Before it was normalized age which can introduce flickering with sorting and random lifetimes
 
 		StatelessInitialMeshOrientationV1,	// Update Initial Mesh Orientation Module
+
+		HierarchyEditorScriptSupport, // Hierarchy Editor was implemented
 		
 		// DO NOT ADD A NEW VERSION UNLESS YOU HAVE TALKED TO THE NIAGARA LEAD. Mismanagement of these versions can lead to data loss if it is adjusted in multiple streams simultaneously.
 		// -----<new versions can be added above this line>  -------------------------------------------------
diff --git a/Engine/Plugins/FX/Niagara/Source/NiagaraEditor/Private/Commandlets/NiagaraDumpModuleInfoCommandlet.cpp b/Engine/Plugins/FX/Niagara/Source/NiagaraEditor/Private/Commandlets/NiagaraDumpModuleInfoCommandlet.cpp
--- a/Engine/Plugins/FX/Niagara/Source/NiagaraEditor/Private/Commandlets/NiagaraDumpModuleInfoCommandlet.cpp
+++ b/Engine/Plugins/FX/Niagara/Source/NiagaraEditor/Private/Commandlets/NiagaraDumpModuleInfoCommandlet.cpp
@@ -295,7 +295,6 @@ void UNiagaraDumpModuleInfoCommandlet::ProcessNiagaraScripts()
 						OutputStream->Logf(TEXT("\t\t<Type>%s</Type>"), *ScriptVar->Variable.GetType().GetName());
 						OutputStream->Logf(TEXT("\t<Guid>%s</Guid>"), *ScriptVar->Metadata.GetVariableGuid().ToString());
 						OutputStream->Logf(TEXT("\t\t<Description>%s</Description>"), *ScriptVar->Metadata.Description.ToString());
-						OutputStream->Logf(TEXT("\t\t<Category>%s</Category>"), *ScriptVar->Metadata.CategoryName.ToString());
 
 						OutputStream->Logf(TEXT("\t\t<Advanced>%s</Advanced>"), ScriptVar->Metadata.bAdvancedDisplay ? TEXT("true") : TEXT("false"));
 						OutputStream->Logf(TEXT("\t\t<DisplayInOverviewStack>%s</DisplayInOverviewStack>"), ScriptVar->Metadata.bDisplayInOverviewStack ? TEXT("true") : TEXT("false"));
diff --git a/Engine/Plugins/FX/Niagara/Source/NiagaraEditor/Private/Customizations/NiagaraScriptVariableCustomization.cpp b/Engine/Plugins/FX/Niagara/Source/NiagaraEditor/Private/Customizations/NiagaraScriptVariableCustomization.cpp
--- a/Engine/Plugins/FX/Niagara/Source/NiagaraEditor/Private/Customizations/NiagaraScriptVariableCustomization.cpp
+++ b/Engine/Plugins/FX/Niagara/Source/NiagaraEditor/Private/Customizations/NiagaraScriptVariableCustomization.cpp
@@ -754,8 +754,7 @@ void FNiagaraScriptVariableHierarchyDetails::CustomizeDetails(IDetailLayoutBuild
 	DetailBuilder.HideProperty(GET_MEMBER_NAME_CHECKED(UNiagaraScriptVariable, DefaultBinding));
 	// member isn't public so we can't use GET_MEMBER_NAME_CHECKED
 	DetailBuilder.HideProperty("DefaultValueVariant", UNiagaraScriptVariable::StaticClass());
-
-	DetailBuilder.HideProperty(GET_MEMBER_NAME_CHECKED(UNiagaraScriptVariable, Metadata.CategoryName));
+	
 	DetailBuilder.HideProperty(GET_MEMBER_NAME_CHECKED(UNiagaraScriptVariable, Metadata.DisplayUnit));
 	DetailBuilder.HideProperty(GET_MEMBER_NAME_CHECKED(UNiagaraScriptVariable, Metadata.bAdvancedDisplay));
 	DetailBuilder.HideProperty(GET_MEMBER_NAME_CHECKED(UNiagaraScriptVariable, Metadata.bOverrideColor));
@@ -763,9 +762,7 @@ void FNiagaraScriptVariableHierarchyDetails::CustomizeDetails(IDetailLayoutBuild
 	DetailBuilder.HideProperty(GET_MEMBER_NAME_CHECKED(UNiagaraScriptVariable, Metadata.bInlineEditConditionToggle));
 	DetailBuilder.HideProperty(GET_MEMBER_NAME_CHECKED(UNiagaraScriptVariable, Metadata.EditCondition));
 	DetailBuilder.HideProperty(GET_MEMBER_NAME_CHECKED(UNiagaraScriptVariable, Metadata.VisibleCondition));
-	DetailBuilder.HideProperty(GET_MEMBER_NAME_CHECKED(UNiagaraScriptVariable, Metadata.EditorSortPriority));
 	DetailBuilder.HideProperty(GET_MEMBER_NAME_CHECKED(UNiagaraScriptVariable, Metadata.PropertyMetaData));
-	DetailBuilder.HideProperty(GET_MEMBER_NAME_CHECKED(UNiagaraScriptVariable, Metadata.ParentAttribute));
 	DetailBuilder.HideProperty(GET_MEMBER_NAME_CHECKED(UNiagaraScriptVariable, Metadata.bDisplayInOverviewStack));
 	DetailBuilder.HideProperty(GET_MEMBER_NAME_CHECKED(UNiagaraScriptVariable, Metadata.InlineParameterSortPriority));
 	DetailBuilder.HideProperty(GET_MEMBER_NAME_CHECKED(UNiagaraScriptVariable, Metadata.InlineParameterEnumOverrides));
diff --git a/Engine/Plugins/FX/Niagara/Source/NiagaraEditor/Private/Customizations/UserParameters/NiagaraUserParametersBuilderBase.cpp b/Engine/Plugins/FX/Niagara/Source/NiagaraEditor/Private/Customizations/UserParameters/NiagaraUserParametersBuilderBase.cpp
--- a/Engine/Plugins/FX/Niagara/Source/NiagaraEditor/Private/Customizations/UserParameters/NiagaraUserParametersBuilderBase.cpp
+++ b/Engine/Plugins/FX/Niagara/Source/NiagaraEditor/Private/Customizations/UserParameters/NiagaraUserParametersBuilderBase.cpp
@@ -15,6 +15,7 @@
 #include "Widgets/SNiagaraSystemUserParameters.h"
 #include "Widgets/Input/SCheckBox.h"
 #include "NiagaraEditorDataBase.h"
+#include "NiagaraEditorUtilities.h"
 #include "PropertyEditorClipboard.h"
 #include "Customizations/NiagaraComponentDetails.h"
 
diff --git a/Engine/Plugins/FX/Niagara/Source/NiagaraEditor/Private/Customizations/UserParameters/NiagaraUserParametersBuilderBase.h b/Engine/Plugins/FX/Niagara/Source/NiagaraEditor/Private/Customizations/UserParameters/NiagaraUserParametersBuilderBase.h
--- a/Engine/Plugins/FX/Niagara/Source/NiagaraEditor/Private/Customizations/UserParameters/NiagaraUserParametersBuilderBase.h
+++ b/Engine/Plugins/FX/Niagara/Source/NiagaraEditor/Private/Customizations/UserParameters/NiagaraUserParametersBuilderBase.h
@@ -8,6 +8,8 @@
 #include "SNiagaraParameterEditor.h"
 #include "ViewModels/HierarchyEditor/NiagaraHierarchyViewModelBase.h"
 
+class UNiagaraScriptVariable;
+
 /** A base class for generating rows based on a given hierarchy root that contains sections, categories and user parameters. */
 class FNiagaraUserParameterNodeBuilder : public IDetailCustomNodeBuilder, public TSharedFromThis<FNiagaraUserParameterNodeBuilder>, FSelfRegisteringEditorUndoClient, FGCObject
 {
diff --git a/Engine/Plugins/FX/Niagara/Source/NiagaraEditor/Private/NiagaraConvertInPlaceEmitterAndSystemState.cpp b/Engine/Plugins/FX/Niagara/Source/NiagaraEditor/Private/NiagaraConvertInPlaceEmitterAndSystemState.cpp
--- a/Engine/Plugins/FX/Niagara/Source/NiagaraEditor/Private/NiagaraConvertInPlaceEmitterAndSystemState.cpp
+++ b/Engine/Plugins/FX/Niagara/Source/NiagaraEditor/Private/NiagaraConvertInPlaceEmitterAndSystemState.cpp
@@ -2,11 +2,12 @@
 
 #include "NiagaraConvertInPlaceEmitterAndSystemState.h"
 #include "NiagaraClipboard.h"
-#include "ViewModels/Stack/NiagaraStackFunctionInputCollection.h"
+#include "ViewModels/Stack/NiagaraStackValueCollection.h"
+#include "ViewModels/Stack/NiagaraStackScriptHierarchyRoot.h"
 
 #include UE_INLINE_GENERATED_CPP_BY_NAME(NiagaraConvertInPlaceEmitterAndSystemState)
 
-bool UNiagaraConvertInPlaceEmitterAndSystemState::Convert(UNiagaraScript* InOldScript, UNiagaraClipboardContent* InOldClipboardContent, UNiagaraScript* InNewScript, UNiagaraStackFunctionInputCollection* InInputCollection, UNiagaraClipboardContent* InNewClipboardContent, UNiagaraNodeFunctionCall* InCallingNode, FText& OutMessage)
+bool UNiagaraConvertInPlaceEmitterAndSystemState::Convert(UNiagaraScript* InOldScript, UNiagaraClipboardContent* InOldClipboardContent, UNiagaraScript* InNewScript, UNiagaraStackScriptHierarchyRoot* InHierarchyRoot, UNiagaraClipboardContent* InNewClipboardContent, UNiagaraNodeFunctionCall* InCallingNode, FText& OutMessage)
 {
 	UE_LOG(LogNiagaraEditor, Log, TEXT("%s to %s"), *InOldScript->GetPathName(), *InNewScript->GetPathName());
 
@@ -182,7 +183,7 @@ bool UNiagaraConvertInPlaceEmitterAndSystemState::Convert(UNiagaraScript* InOldS
 				}
 			}
 
-			InInputCollection->SetValuesFromClipboardFunctionInputs(NewInputs);
+			InHierarchyRoot->SetValuesFromClipboardFunctionInputs(NewInputs);
 			//InCallingNode->RefreshFromExternalChanges();
 
 			// We have to do this as a second pass as Loop Count is behind a static switch set above.
@@ -194,7 +195,7 @@ bool UNiagaraConvertInPlaceEmitterAndSystemState::Convert(UNiagaraScript* InOldS
 				NewInputs.Add(NewInput);
 				ConvertedItems.Add(NewInput->InputName);
 
-				InInputCollection->SetValuesFromClipboardFunctionInputs(NewInputs);
+				InHierarchyRoot->SetValuesFromClipboardFunctionInputs(NewInputs);
 			}
 
 			// Use first frame delay is behind a static switch set above.
@@ -206,7 +207,7 @@ bool UNiagaraConvertInPlaceEmitterAndSystemState::Convert(UNiagaraScript* InOldS
 				NewInputs.Add(NewInput);
 				ConvertedItems.Add(NewInput->InputName);
 
-				InInputCollection->SetValuesFromClipboardFunctionInputs(NewInputs);
+				InHierarchyRoot->SetValuesFromClipboardFunctionInputs(NewInputs);
 			}
 
 			FString ConcatString = FString::Printf(TEXT("Copied over %d properties\nPlease double-check for correctness\n"), ConvertedItems.Num());
diff --git a/Engine/Plugins/FX/Niagara/Source/NiagaraEditor/Private/NiagaraConvertInPlaceEmitterAndSystemState.h b/Engine/Plugins/FX/Niagara/Source/NiagaraEditor/Private/NiagaraConvertInPlaceEmitterAndSystemState.h
--- a/Engine/Plugins/FX/Niagara/Source/NiagaraEditor/Private/NiagaraConvertInPlaceEmitterAndSystemState.h
+++ b/Engine/Plugins/FX/Niagara/Source/NiagaraEditor/Private/NiagaraConvertInPlaceEmitterAndSystemState.h
@@ -12,5 +12,5 @@ class  UNiagaraConvertInPlaceEmitterAndSystemState : public UNiagaraConvertInPla
 {
 	GENERATED_BODY()
 public:
-	NIAGARAEDITOR_API virtual bool Convert(UNiagaraScript* InOldScript, UNiagaraClipboardContent* InOldClipboardContent, UNiagaraScript* InNewScript, UNiagaraStackFunctionInputCollection* InInputCollection, UNiagaraClipboardContent* InNewClipboardContent, UNiagaraNodeFunctionCall* InCallingNode, FText& OutMessage) override;
+	NIAGARAEDITOR_API virtual bool Convert(UNiagaraScript* InOldScript, UNiagaraClipboardContent* InOldClipboardContent, UNiagaraScript* InNewScript, UNiagaraStackScriptHierarchyRoot* InHierarchyRoot, UNiagaraClipboardContent* InNewClipboardContent, UNiagaraNodeFunctionCall* InCallingNode, FText& OutMessage) override;
 };
diff --git a/Engine/Plugins/FX/Niagara/Source/NiagaraEditor/Private/NiagaraEditorStyle.cpp b/Engine/Plugins/FX/Niagara/Source/NiagaraEditor/Private/NiagaraEditorStyle.cpp
--- a/Engine/Plugins/FX/Niagara/Source/NiagaraEditor/Private/NiagaraEditorStyle.cpp
+++ b/Engine/Plugins/FX/Niagara/Source/NiagaraEditor/Private/NiagaraEditorStyle.cpp
@@ -306,8 +306,7 @@ void FNiagaraEditorStyle::InitParameters()
 		.SetTextStyle(InlineEditableTextBlockReadOnly)
 		.SetEditableTextBoxStyle(InlineEditableTextBlockEditable);
 	Set("NiagaraEditor.Graph.Node.InlineEditablePinName", InlineEditableTextBlockStyle);
-
-
+	
 	Set("NiagaraEditor.StaticIcon", new IMAGE_BRUSH("Icons/staticpill_16x", Icon16x16));
 	Set("NiagaraEditor.Pins.StaticConnected", new IMAGE_BRUSH("Icons/StaticPin_Connected", Icon16x16));
 	Set("NiagaraEditor.Pins.StaticDisconnected", new IMAGE_BRUSH("Icons/StaticPin_Disconnected", Icon16x16));
@@ -861,6 +860,31 @@ void FNiagaraEditorStyle::InitHierarchyEditor()
 		.SetPressedPadding(FMargin(8.f, 3.f, 8.f, 1.f));
 
 	Set("NiagaraEditor.HierarchyEditor.ButtonStyle", ButtonStyle);
+
+	const FSlateColor SelectorColor = FAppStyle::GetSlateColor("SelectorColor");
+	const FSlateColor SelectionColor = FAppStyle::GetSlateColor("SelectionColor");
+	const FSlateColor SelectionColor_Inactive = FAppStyle::GetSlateColor("SelectionColor_Inactive");
+	FLinearColor HighlightColor = FLinearColor(0.068f, 0.068f, 0.068f);
+	
+	Set("NiagaraEditor.Color.Highlight", HighlightColor);
+	
+	FTableRowStyle HierarchyEditorCategoryRowStyle = FTableRowStyle()
+		.SetEvenRowBackgroundBrush(*FAppStyle::GetBrush("DetailsView.CategoryTop"))
+		.SetOddRowBackgroundBrush(*FAppStyle::GetBrush("DetailsView.CategoryTop"))
+		.SetEvenRowBackgroundHoveredBrush(*FAppStyle::GetBrush("DetailsView.CategoryTop_Hovered"))
+		.SetOddRowBackgroundHoveredBrush(*FAppStyle::GetBrush("DetailsView.CategoryTop_Hovered"))
+		.SetSelectorFocusedBrush( CORE_BORDER_BRUSH( "Common/Selector", FMargin(4.f/16.f), SelectorColor))
+		.SetActiveBrush( CORE_IMAGE_BRUSH( "Common/Selection", Icon8x8, SelectionColor))
+		.SetActiveHoveredBrush( CORE_IMAGE_BRUSH( "Common/Selection", Icon8x8, SelectionColor))
+		.SetInactiveBrush( CORE_IMAGE_BRUSH( "Common/Selection", Icon8x8, SelectionColor_Inactive))
+		.SetInactiveHoveredBrush( CORE_IMAGE_BRUSH( "Common/Selection", Icon8x8, SelectionColor_Inactive))
+		.SetActiveHighlightedBrush( CORE_IMAGE_BRUSH("Common/Selection", Icon8x8, HighlightColor))
+		.SetInactiveHighlightedBrush( CORE_IMAGE_BRUSH("Common/Selection", Icon8x8, HighlightColor))
+		.SetDropIndicator_Above(CORE_BOX_BRUSH("Common/DropZoneIndicator_Above", FMargin(10.0f / 16.0f, 10.0f / 16.0f, 0, 0), FStarshipCoreStyle::GetCoreStyle().GetSlateColor("SelectorColor")))
+		.SetDropIndicator_Onto(CORE_BOX_BRUSH("Common/DropZoneIndicator_Onto", FMargin(4.0f / 16.0f), FStarshipCoreStyle::GetCoreStyle().GetSlateColor("SelectorColor")))
+		.SetDropIndicator_Below(CORE_BOX_BRUSH("Common/DropZoneIndicator_Below", FMargin(10.0f / 16.0f, 0, 0, 10.0f / 16.0f), FStarshipCoreStyle::GetCoreStyle().GetSlateColor("SelectorColor")));
+	
+	Set("NiagaraEditor.HierarchyEditor.Row.Category", HierarchyEditorCategoryRowStyle);
 }
 
 void FNiagaraEditorStyle::InitSimCacheEditor()
diff --git a/Engine/Plugins/FX/Niagara/Source/NiagaraEditor/Private/NiagaraEditorUtilities.cpp b/Engine/Plugins/FX/Niagara/Source/NiagaraEditor/Private/NiagaraEditorUtilities.cpp
--- a/Engine/Plugins/FX/Niagara/Source/NiagaraEditor/Private/NiagaraEditorUtilities.cpp
+++ b/Engine/Plugins/FX/Niagara/Source/NiagaraEditor/Private/NiagaraEditorUtilities.cpp
@@ -79,6 +79,8 @@
 #include "ViewModels/NiagaraEmitterHandleViewModel.h"
 #include "ViewModels/NiagaraEmitterViewModel.h"
 #include "ViewModels/NiagaraParameterPanelViewModel.h"
+#include "Widgets/SNiagaraHierarchyEditor.h"
+#include "ViewModels/HierarchyEditor/NiagaraHierarchyScriptParametersViewModel.h"
 #include "Widgets/SToolTip.h"
 
 #define LOCTEXT_NAMESPACE "FNiagaraEditorUtilities"
@@ -4638,6 +4640,34 @@ UNiagaraDataInterface* FNiagaraEditorUtilities::GetResolvedRuntimeInstanceForEdi
 	return nullptr;
 }
 
+TSharedRef<SWidget> FNiagaraEditorUtilities::HierarchyEditor::Scripts::GenerateRowContentForScriptParameterHierarchyEditor(TSharedRef<FNiagaraHierarchyItemViewModelBase> HierarchyItem)
+{
+	if(HierarchyItem->GetDataMutable()->IsA<UNiagaraHierarchyCategory>())
+	{
+		TSharedRef<FNiagaraHierarchyCategoryViewModel> TreeViewCategory = StaticCastSharedRef<FNiagaraHierarchyCategoryViewModel>(HierarchyItem);
+		return SNew(SNiagaraHierarchyCategory, TreeViewCategory);
+	}
+	else if(const UNiagaraHierarchyScriptParameter* ScriptParameter = Cast<UNiagaraHierarchyScriptParameter>(HierarchyItem->GetData()))
+	{
+		FNiagaraParameterUtilities::FNiagaraParameterWidgetOptions WidgetOptions;
+		WidgetOptions.bAddTypeIcon = true;
+		WidgetOptions.bShowEditConditionIcon = true;
+		WidgetOptions.bShowVisibilityConditionIcon = true;
+		WidgetOptions.bShowAdvanced = true;
+		TSharedRef<SWidget> ParameterWidget = FNiagaraParameterUtilities::GetParameterWidget(ScriptParameter->GetVariable(), ScriptParameter->GetScriptVariable()->Metadata, WidgetOptions);
+		const UNiagaraScriptVariable* ScriptVariable = ScriptParameter->GetScriptVariable();
+
+		ParameterWidget->SetToolTipText(TAttribute<FText>::CreateLambda([ScriptVariable]()
+		{
+			return ScriptVariable->Metadata.Description;
+		}));
+		
+		return ParameterWidget;
+	}
+
+	return SNullWidget::NullWidget;
+}
+
 TSharedRef<SToolTip> FNiagaraEditorUtilities::Tooltips::CreateStackNoteTooltip(UNiagaraStackNote& StackNote)
 {
 	if(StackNote.GetTargetStackNoteData().IsSet() == false)
diff --git a/Engine/Plugins/FX/Niagara/Source/NiagaraEditor/Private/NiagaraEmitterEditorData.cpp b/Engine/Plugins/FX/Niagara/Source/NiagaraEditor/Private/NiagaraEmitterEditorData.cpp
--- a/Engine/Plugins/FX/Niagara/Source/NiagaraEditor/Private/NiagaraEmitterEditorData.cpp
+++ b/Engine/Plugins/FX/Niagara/Source/NiagaraEditor/Private/NiagaraEmitterEditorData.cpp
@@ -159,16 +159,16 @@ void UNiagaraEmitterEditorData::PostLoad_TransferSummaryDataToNewFormat()
 						if(MatchingScriptVariable != nullptr)
 						{
 							ScriptVariable = *MatchingScriptVariable;
-							OriginalCategoryName = FName(ScriptVariable->Metadata.CategoryName.ToString());
+							OriginalCategoryName = FName(ScriptVariable->Metadata.GetCategoryName_DEPRECATED().ToString());
 							ParentInputIdentity.Guids.Add(ScriptVariable->Metadata.GetVariableGuid());
 							AssetVariableMetadata.Add(ParentInputIdentity, ScriptVariable->Metadata);
 
 							// now that we found the matching parent input, we look for child inputs
 							for(UNiagaraScriptVariable* CandidateChildScriptVariable : ScriptVariables)
 							{
-								if(CandidateChildScriptVariable != nullptr && !CandidateChildScriptVariable->Metadata.ParentAttribute.IsNone())
+								if(CandidateChildScriptVariable != nullptr && !CandidateChildScriptVariable->Metadata.GetParentAttribute_DEPRECATED().IsNone())
 								{
-									if(CandidateChildScriptVariable->Metadata.ParentAttribute.IsEqual(ScriptVariable->Variable.GetName()))
+									if(CandidateChildScriptVariable->Metadata.GetParentAttribute_DEPRECATED().IsEqual(ScriptVariable->Variable.GetName()))
 									{
 										FNiagaraHierarchyIdentity ChildIdentity;
 										ChildIdentity.Guids.Add(FunctionCall->NodeGuid);
@@ -298,13 +298,13 @@ void UNiagaraEmitterEditorData::PostLoad_TransferSummaryDataToNewFormat()
 			int32 SortOrderA = SummarySortOrder[&ItemA];
 			if(SortOrderA == INDEX_NONE)
 			{
-				SortOrderA = MetaDataA.EditorSortPriority;
+				SortOrderA = MetaDataA.GetEditorSortPriority_DEPRECATED();
 			}
 
 			int32 SortOrderB = SummarySortOrder[&ItemB];
 			if(SortOrderB == INDEX_NONE)
 			{
-				SortOrderB = MetaDataB.EditorSortPriority;
+				SortOrderB = MetaDataB.GetEditorSortPriority_DEPRECATED();
 			}
 				
 			return SortOrderA < SortOrderB;
@@ -333,7 +333,7 @@ void UNiagaraEmitterEditorData::PostLoad_TransferSummaryDataToNewFormat()
 				{
 					if(AssetVariableMetadata.Contains(MinimumItem->GetPersistentIdentity()))
 					{
-						MinimumCategorySortOrder = AssetVariableMetadata[MinimumItem->GetPersistentIdentity()].EditorSortPriority;
+						MinimumCategorySortOrder = AssetVariableMetadata[MinimumItem->GetPersistentIdentity()].GetEditorSortPriority_DEPRECATED();
 					}
 				}
 				
diff --git a/Engine/Plugins/FX/Niagara/Source/NiagaraEditor/Private/NiagaraGraph.cpp b/Engine/Plugins/FX/Niagara/Source/NiagaraEditor/Private/NiagaraGraph.cpp
--- a/Engine/Plugins/FX/Niagara/Source/NiagaraEditor/Private/NiagaraGraph.cpp
+++ b/Engine/Plugins/FX/Niagara/Source/NiagaraEditor/Private/NiagaraGraph.cpp
@@ -37,6 +37,7 @@
 #include "Modules/ModuleManager.h"
 #include "String/ParseTokens.h"
 #include "ViewModels/NiagaraScriptViewModel.h"
+#include "ViewModels/HierarchyEditor/NiagaraHierarchyScriptParametersViewModel.h"
 
 #include UE_INLINE_GENERATED_CPP_BY_NAME(NiagaraGraph)
 
@@ -207,16 +208,6 @@ bool FNiagaraGraphScriptUsageInfo::IsValid() const
 	return true;
 }
 
-UNiagaraGraph::UNiagaraGraph(const FObjectInitializer& ObjectInitializer)
-	: Super(ObjectInitializer)
-	, bNeedNumericCacheRebuilt(true)
-	, bIsRenamingParameter(false)
-	, bParameterReferenceRefreshPending(true)
-{
-	Schema = UEdGraphSchema_Niagara::StaticClass();
-	ChangeId = FGuid::NewGuid();
-}
-
 FDelegateHandle UNiagaraGraph::AddOnGraphNeedsRecompileHandler(const FOnGraphChanged::FDelegate& InHandler)
 {
 	return OnGraphNeedsRecompile.Add(InHandler);
@@ -441,6 +432,17 @@ void UNiagaraGraph::PostLoad_ManageScriptVariables(int32 NiagaraVersion)
 	}
 }
 
+UNiagaraGraph::UNiagaraGraph()
+	: bNeedNumericCacheRebuilt(true)
+	, bIsRenamingParameter(false)
+	, bParameterReferenceRefreshPending(true)
+{
+	Schema = UEdGraphSchema_Niagara::StaticClass();
+	ChangeId = FGuid::NewGuid();
+	
+	ParameterHierarchyRoot = CreateDefaultSubobject<UNiagaraHierarchyRoot>("ParameterHierarchyRoot");
+}
+
 void UNiagaraGraph::PostLoad()
 {
 	LLM_SCOPE(ELLMTag::Niagara);
@@ -1006,6 +1008,7 @@ void UNiagaraGraph::ChangeParameterType(const TArray<FNiagaraVariable>& Paramete
 	}
 	
 	InvalidateCachedParameterData();
+	OnParametersChangedDelegate.Broadcast(TOptional<FParametersChangedData>());
 }
 
 void UNiagaraGraph::ReplaceScriptReferences(UNiagaraScript* OldScript, UNiagaraScript* NewScript)
@@ -1027,6 +1030,11 @@ void UNiagaraGraph::ReplaceScriptReferences(UNiagaraScript* OldScript, UNiagaraS
 	}
 }
 
+UNiagaraGraph::FOnParametersChanged& UNiagaraGraph::OnParametersChanged()
+{
+	return OnParametersChangedDelegate;
+}
+
 TArray<UEdGraphPin*> UNiagaraGraph::FindParameterMapDefaultValuePins(const FName VariableName) const
 {
 	TArray<UEdGraphPin*> OutDefaultPins;
@@ -2221,6 +2229,21 @@ void UNiagaraGraph::GetParameters(TArray<FNiagaraVariable>& Inputs, TArray<FNiag
 // 	Outputs.Sort(SortVars);
 }
 
+void UNiagaraGraph::GetAllInputScriptVariables(TArray<UNiagaraScriptVariable*>& OutScriptVariables) const
+{
+	TArray<UNiagaraScriptVariable*> AllInputScriptVariables;
+
+	for(const auto& VariablePair : GetAllMetaData())
+	{
+		if(VariablePair.Key.IsInNameSpace(FNiagaraConstants::ModuleNamespace) || VariablePair.Value->GetIsStaticSwitch())
+		{
+			AllInputScriptVariables.Add(VariablePair.Value);
+		}
+	}
+
+	OutScriptVariables.Append(AllInputScriptVariables);
+}
+
 const UNiagaraGraph::FScriptVariableMap& UNiagaraGraph::GetAllMetaData() const
 {
 	check(!bIsForCompilationOnly);
@@ -2292,7 +2315,7 @@ UNiagaraScriptVariable* UNiagaraGraph::GetScriptVariable(FGuid VariableGuid) con
 	return nullptr;
 }
 
-TArray<UNiagaraScriptVariable*> UNiagaraGraph::GetChildScriptVariablesForInput(FGuid VariableGuid) const
+TArray<UNiagaraScriptVariable*> UNiagaraGraph::GetChildScriptVariablesForInput_Deprecated(FGuid VariableGuid) const
 {
 	check(!bIsForCompilationOnly);
 
@@ -2311,7 +2334,7 @@ TArray<UNiagaraScriptVariable*> UNiagaraGraph::GetChildScriptVariablesForInput(F
 	{
 		for (auto& VariableToScriptVariableItem : VariableToScriptVariable)
 		{
-			if(VariableToScriptVariableItem.Value->Metadata.ParentAttribute == FoundVariable->GetName())
+			if(VariableToScriptVariableItem.Value->Metadata.GetParentAttribute_DEPRECATED() == FoundVariable->GetName())
 			{
 				ChildrenVariables.Add(VariableToScriptVariableItem.Value);
 			}
@@ -2320,7 +2343,7 @@ TArray<UNiagaraScriptVariable*> UNiagaraGraph::GetChildScriptVariablesForInput(F
 
 	ChildrenVariables.Sort([&](const UNiagaraScriptVariable& VariableA, const UNiagaraScriptVariable& VariableB)
 	{
-		return VariableA.Metadata.EditorSortPriority < VariableB.Metadata.EditorSortPriority;
+		return VariableA.Metadata.GetEditorSortPriority_DEPRECATED() < VariableB.Metadata.GetEditorSortPriority_DEPRECATED();
 	});
 	
 	return ChildrenVariables;
@@ -2331,7 +2354,7 @@ TArray<FGuid> UNiagaraGraph::GetChildScriptVariableGuidsForInput(FGuid VariableG
 	check(!bIsForCompilationOnly);
 
 	TArray<FGuid> ChildrenVariableGuids;
-	TMap<FGuid, int32> SortOrderMap;
+	
 	TOptional<FNiagaraVariable> FoundVariable;
 	for (auto& VariableToScriptVariableItem : VariableToScriptVariable)
 	{
@@ -2342,21 +2365,35 @@ TArray<FGuid> UNiagaraGraph::GetChildScriptVariableGuidsForInput(FGuid VariableG
 	}
 
 	if(FoundVariable.IsSet())
-	{
-		for (auto& VariableToScriptVariableItem : VariableToScriptVariable)
+	{	
+		TArray<FNiagaraVariable> Variables;
+		GetAllVariables(Variables);
+	
+		for(FNiagaraVariable& Variable : Variables)
 		{
-			if(VariableToScriptVariableItem.Value->Metadata.ParentAttribute == FoundVariable->GetName())
+			TOptional<FNiagaraVariableMetaData> MetaData = GetMetaData(Variable);
+
+			if(MetaData.IsSet())
 			{
-				ChildrenVariableGuids.Add(VariableToScriptVariableItem.Value->Metadata.GetVariableGuid());
-				SortOrderMap.Add(VariableToScriptVariableItem.Value->Metadata.GetVariableGuid(), VariableToScriptVariableItem.Value->Metadata.EditorSortPriority);
+				FNiagaraHierarchyIdentity Identity({MetaData.GetValue().GetVariableGuid()}, {});
+				if(UNiagaraHierarchyItemBase* FoundHierarchyParameter = ParameterHierarchyRoot->FindChildWithIdentity(Identity, true))
+				{
+					TArray<UNiagaraHierarchyScriptParameter*> ChildrenScriptParameters;
+					FoundHierarchyParameter->GetChildrenOfType(ChildrenScriptParameters, false);
+
+					for(UNiagaraHierarchyScriptParameter* ChildrenScriptParameter : ChildrenScriptParameters)
+					{
+						ChildrenVariableGuids.Add(ChildrenScriptParameter->GetScriptVariable()->Metadata.GetVariableGuid());
+					}
+				}
+			
+			}
+			else
+			{
+				UE_LOG(LogNiagaraEditor, Log, TEXT("Couldn't find metadata for variable"));
 			}
 		}
 	}
-
-	ChildrenVariableGuids.Sort([&](const FGuid& GuidA, const FGuid& GuidB)
-	{
-		return SortOrderMap[GuidA] < SortOrderMap[GuidB];
-	});
 	
 	return ChildrenVariableGuids;
 }
@@ -2419,6 +2456,7 @@ UNiagaraScriptVariable* UNiagaraGraph::AddParameter(const UNiagaraScriptVariable
 		}
  		ParameterToReferencesMap.Emplace(NewScriptVariable->Variable, {true /*bInCreated*/});
 		VariableToScriptVariable.Add(NewScriptVariable->Variable, NewScriptVariable);
+		OnParametersChangedDelegate.Broadcast(FParametersChangedData{NewScriptVariable});
 		NotifyGraphChanged();
 		return NewScriptVariable;
 	}
@@ -2445,6 +2483,7 @@ UNiagaraScriptVariable* UNiagaraGraph::CreateScriptVariableInternal(const FNiaga
 	}
 
 	VariableToScriptVariable.Add(Parameter, NewScriptVariable);
+	OnParametersChangedDelegate.Broadcast(FParametersChangedData{NewScriptVariable});
 
 	return NewScriptVariable;
 }
@@ -2529,6 +2568,7 @@ void UNiagaraGraph::RemoveParameter(const FNiagaraVariable& Parameter, bool bAll
 	for (FNiagaraVariable& Var : VarsToRemove)
 	{
 		VariableToScriptVariable.Remove(Var);
+		OnParametersChangedDelegate.Broadcast(TOptional<FParametersChangedData>());
 	}
 }
 
@@ -2803,6 +2843,8 @@ bool UNiagaraGraph::RenameParameter(const FNiagaraVariable& Parameter, FName New
 		*bOutMerged = bParameterMerged;
 	}
 
+	OnParametersChangedDelegate.Broadcast(TOptional<FParametersChangedData>());
+
 	return true;
 }
 
diff --git a/Engine/Plugins/FX/Niagara/Source/NiagaraEditor/Private/NiagaraScriptSource.cpp b/Engine/Plugins/FX/Niagara/Source/NiagaraEditor/Private/NiagaraScriptSource.cpp
--- a/Engine/Plugins/FX/Niagara/Source/NiagaraEditor/Private/NiagaraScriptSource.cpp
+++ b/Engine/Plugins/FX/Niagara/Source/NiagaraEditor/Private/NiagaraScriptSource.cpp
@@ -21,6 +21,7 @@
 #include "NiagaraNodeParameterMapSet.h"
 #include "NiagaraScript.h"
 #include "Logging/LogMacros.h"
+#include "ViewModels/HierarchyEditor/NiagaraHierarchyScriptParametersViewModel.h"
 #include "ViewModels/Stack/NiagaraStackGraphUtilities.h"
 
 #include UE_INLINE_GENERATED_CPP_BY_NAME(NiagaraScriptSource)
@@ -28,11 +29,6 @@
 DECLARE_CYCLE_STAT(TEXT("Niagara - ScriptSource - Compile"), STAT_NiagaraEditor_ScriptSource_Compile, STATGROUP_NiagaraEditor);
 DECLARE_CYCLE_STAT(TEXT("Niagara - ScriptSource - InitializeNewRapidIterationParameters"), STAT_NiagaraEditor_ScriptSource_InitializeNewRapidIterationParameters, STATGROUP_NiagaraEditor);
 
-UNiagaraScriptSource::UNiagaraScriptSource(const FObjectInitializer& ObjectInitializer)
-	: Super(ObjectInitializer)
-{
-}
-
 void UNiagaraScriptSource::RegisterVMCompilationIdDependencies(FNiagaraScriptHashCollector& Collector, ENiagaraScriptUsage InUsage, const FGuid& InUsageId) const
 {
 	if (!AllowShaderCompiling())
@@ -165,7 +161,6 @@ void UNiagaraScriptSource::RefreshFromExternalChanges()
 	}
 }
 
-
 void UNiagaraScriptSource::PostLoad()
 {
 	Super::PostLoad();
@@ -375,6 +370,11 @@ bool UNiagaraScriptSource::AddModuleIfMissing(FString ModulePath, ENiagaraScript
 	return false;
 }
 
+void UNiagaraScriptSource::PostLoadFromOwner(FVersionedNiagaraScriptData& OwnerData)
+{
+	PostLoadTransferScriptHierarchyData(OwnerData);
+}
+
 void UNiagaraScriptSource::FixupRenamedParameters(UNiagaraNodeFunctionCall* FunctionCallNode, TConstArrayView<FNiagaraVariable> ModuleInputVariables, FNiagaraParameterStore& RapidIterationParameters, const TArray<FNiagaraVariable>& OldRapidIterationVariables, const FVersionedNiagaraEmitter& VersionedEmitter, ENiagaraScriptUsage ScriptUsage) const
 {
 	// the rapid iteration parameters and the function input pins use different variable naming schemes, so most of this is just used to convert one name to the other 
@@ -669,6 +669,227 @@ void UNiagaraScriptSource::OnGraphDataInterfaceChanged()
 	OnChangedDelegate.Broadcast();
 }
 
+void UNiagaraScriptSource::PostLoadTransferScriptHierarchyData(FVersionedNiagaraScriptData& OwnerData)
+{
+	UNiagaraGraph* Graph = Cast<UNiagaraScriptSource>(OwnerData.GetSource())->NodeGraph;
+
+	if(Graph == nullptr)
+	{
+		UE_LOG(LogNiagaraEditor, Verbose, TEXT("Failed post loading script hierarchy of versioned script data. No graph found."));
+		return;
+	}
+	
+	const int32 CustomVersion = Graph->GetLinkerCustomVersion(FNiagaraCustomVersion::GUID);
+
+	if(CustomVersion < FNiagaraCustomVersion::HierarchyEditorScriptSupport)
+	{
+		UNiagaraHierarchyRoot* ParameterRoot = Graph->GetScriptParameterHierarchyRoot();
+
+		TMap<FName, UNiagaraHierarchySection*> HierarchySections;
+		TMap<FName, UNiagaraHierarchyCategory*> IdentityCategoryMap;
+		TMap<FName, FName> SectionCategoryMapping;
+		
+		// First we transfer all sections and create contained categories
+		for(const FNiagaraStackSection& StackSection : OwnerData.InputSections_DEPRECATED)
+		{
+			FText SectionDisplayName = StackSection.SectionDisplayName.IsEmptyOrWhitespace() ? FText::FromName(StackSection.SectionIdentifier) : StackSection.SectionDisplayName;
+			UNiagaraHierarchySection* HierarchySection = ParameterRoot->AddSection(SectionDisplayName);
+			HierarchySection->SetTooltip(StackSection.Tooltip);
+
+			FNiagaraHierarchyIdentity SectionIdentity;
+			SectionIdentity.Names.Add(StackSection.SectionIdentifier);
+			HierarchySection->SetIdentity(SectionIdentity);
+			
+			HierarchySections.Add(StackSection.SectionIdentifier, HierarchySection);
+
+			for(const FText& CategoryIdentity : StackSection.Categories)
+			{
+				// We skip categories that don't have a proper identifier setup
+				if(CategoryIdentity.IsEmptyOrWhitespace())
+				{
+					continue;
+				}
+				
+				FName CategoryName = FName(CategoryIdentity.ToString());
+				UNiagaraHierarchyScriptCategory* HierarchyCategory = ParameterRoot->AddChild<UNiagaraHierarchyScriptCategory>();
+				HierarchyCategory->SetSection(HierarchySection);
+				HierarchyCategory->SetCategoryName(CategoryName);
+				IdentityCategoryMap.Add(CategoryName, HierarchyCategory);
+				SectionCategoryMapping.Add(CategoryName, HierarchySection->GetSectionName());
+			}
+		}
+
+		OwnerData.InputSections_DEPRECATED.Empty();
+		
+		TMap<FNiagaraVariable, TObjectPtr<UNiagaraScriptVariable>> ScriptVariablesMap = Graph->GetAllMetaData();
+
+		TArray<UNiagaraScriptVariable*> ParentLevelInputs;
+		TMap<UNiagaraScriptVariable*, TArray<UNiagaraScriptVariable*>> ParentChildInputMap;
+		// We determine all parent & child level module inputs & static switches
+		for(const auto& ScriptVariablePair : ScriptVariablesMap)
+		{
+			if(ScriptVariablePair.Key.IsInNameSpace(FNiagaraConstants::ModuleNamespace) == false && ScriptVariablePair.Value->GetIsStaticSwitch() == false)
+			{
+				continue;
+			}
+			
+			UNiagaraScriptVariable* ScriptVariable = ScriptVariablePair.Value;
+
+			if(ScriptVariable->Metadata.GetParentAttribute_DEPRECATED().IsNone())
+			{
+				ParentLevelInputs.Add(ScriptVariable);
+				ParentChildInputMap.Add(ScriptVariable, Graph->GetChildScriptVariablesForInput_Deprecated(ScriptVariable->Metadata.GetVariableGuid()));
+			}
+		}
+
+		// Now we create the actual hierarchy elements
+		for(const auto& ParentChildPair : ParentChildInputMap)
+		{
+			// By default, we add to the root directly, unless a category was specified
+			UNiagaraHierarchyItemBase* Owner = ParameterRoot;
+			
+			FText SpecifiedCategoryName = ParentChildPair.Key->Metadata.GetCategoryName_DEPRECATED();
+			if(SpecifiedCategoryName.IsEmptyOrWhitespace() == false)
+			{
+				// Categories might already exist due to section info
+				if(IdentityCategoryMap.Contains(FName(SpecifiedCategoryName.ToString())))
+				{
+					Owner = IdentityCategoryMap[FName(SpecifiedCategoryName.ToString())];
+				}
+				// But if not we create a new category
+				else
+				{
+					UNiagaraHierarchyScriptCategory* Category = ParameterRoot->AddChild<UNiagaraHierarchyScriptCategory>();
+					Category->SetCategoryName(FName(SpecifiedCategoryName.ToString()));
+					IdentityCategoryMap.Add(Category->GetCategoryName(), Category);
+
+					Owner = Category;
+				}
+			}
+
+			// We add the parent input to either the root, or a given category
+			UNiagaraHierarchyScriptParameter* ParentHierarchyInput = Owner->AddChild<UNiagaraHierarchyScriptParameter>();
+			ParentHierarchyInput->Initialize(*ParentChildPair.Key);
+
+			// Children inputs can be added directly to the parent input and are already sorted
+			for(UNiagaraScriptVariable* ChildInput : ParentChildPair.Value)
+			{
+				UNiagaraHierarchyScriptParameter* ChildHierarchyInput = ParentHierarchyInput->AddChild<UNiagaraHierarchyScriptParameter>();
+				ChildHierarchyInput->Initialize(*ChildInput);
+			}
+		}
+
+		// now that we have created all inputs, categories and sections, we fix up the order
+		TArray<UNiagaraHierarchyCategory*> TopLevelCategories;
+		ParameterRoot->GetChildrenOfType(TopLevelCategories, false);
+		
+		// inputs are generally sorted by editor sort priority, but within a given context (within a category etc.)
+		// child inputs also already are sorted when constructed
+		auto SortInputsPredicate = ([&](UNiagaraHierarchyItemBase& ItemA, UNiagaraHierarchyItemBase& ItemB)
+		{
+			UNiagaraHierarchyScriptParameter* ScriptParameterA = Cast<UNiagaraHierarchyScriptParameter>(&ItemA);
+			UNiagaraHierarchyScriptParameter* ScriptParameterB = Cast<UNiagaraHierarchyScriptParameter>(&ItemB);
+
+			if(ScriptParameterA == nullptr || ScriptParameterB == nullptr)
+			{
+				return false;
+			}
+
+			const FNiagaraVariableMetaData& MetaDataA = ScriptParameterA->GetScriptVariable()->Metadata;
+			const FNiagaraVariableMetaData& MetaDataB = ScriptParameterB->GetScriptVariable()->Metadata;
+
+			if(MetaDataA.bAdvancedDisplay && !MetaDataB.bAdvancedDisplay)
+			{
+				return false;
+			}
+			else if(!MetaDataA.bAdvancedDisplay && MetaDataB.bAdvancedDisplay)
+			{
+				return true;
+			}
+			
+			if(MetaDataA.GetEditorSortPriority_DEPRECATED() != MetaDataB.GetEditorSortPriority_DEPRECATED())
+			{
+				return MetaDataA.GetEditorSortPriority_DEPRECATED() < MetaDataB.GetEditorSortPriority_DEPRECATED();
+			}
+			
+			return ItemA.ToString().Compare(ItemB.ToString(), ESearchCase::Type::IgnoreCase) < 0;
+		});
+
+		// We now sort inputs, and keep track of the minimum sort order per category
+		// This is because the stack determined category order by input sort priority (the category with the parameter of the lowest priority value would come first)
+		TMap<UNiagaraHierarchyItemBase*, int32> CategoryMinimumInputSortOrderMap;
+		for(UNiagaraHierarchyCategory* Category : TopLevelCategories)
+		{
+			Category->SortChildren(SortInputsPredicate, false);
+			if(Category->GetChildren().Num() > 0)
+			{
+				// since we have sorted the children by sort priority, child 0 has the lowest sort priority number (= highest priority)
+				UNiagaraHierarchyScriptParameter* MinimumItem = CastChecked<UNiagaraHierarchyScriptParameter>(Category->GetChildren()[0]);
+				int32 MinimumCategorySortOrder = MinimumItem->GetScriptVariable()->Metadata.GetEditorSortPriority_DEPRECATED();
+			
+				CategoryMinimumInputSortOrderMap.Add(Category, MinimumCategorySortOrder);
+			}
+			else
+			{
+				CategoryMinimumInputSortOrderMap.Add(Category, 0);
+			}
+		}
+
+		auto SortRootChildrenPredicate = ([&](UNiagaraHierarchyItemBase& ItemA, UNiagaraHierarchyItemBase& ItemB)
+		{
+			if(ItemA.IsA<UNiagaraHierarchyCategory>() && ItemB.IsA<UNiagaraHierarchyCategory>())
+			{
+				if(CategoryMinimumInputSortOrderMap[&ItemA] != CategoryMinimumInputSortOrderMap[&ItemB])
+				{
+					return CategoryMinimumInputSortOrderMap[&ItemA] < CategoryMinimumInputSortOrderMap[&ItemB];
+				}
+
+				return ItemA.ToString().Compare(ItemB.ToString(), ESearchCase::Type::IgnoreCase) < 0;
+			}
+
+			// Categories always come after items to ensure consistency with old behavior
+			if(ItemA.IsA<UNiagaraHierarchyCategory>() && ItemB.IsA<UNiagaraHierarchyItem>())
+			{
+				return false;
+			}
+			
+			if(ItemB.IsA<UNiagaraHierarchyCategory>() && ItemA.IsA<UNiagaraHierarchyItem>())
+			{
+				return true;
+			}
+			
+			if(ItemA.IsA<UNiagaraHierarchyScriptParameter>() && ItemB.IsA<UNiagaraHierarchyScriptParameter>())
+			{
+				UNiagaraHierarchyScriptParameter* ScriptParameterA = Cast<UNiagaraHierarchyScriptParameter>(&ItemA);
+				UNiagaraHierarchyScriptParameter* ScriptParameterB = Cast<UNiagaraHierarchyScriptParameter>(&ItemB);
+
+				const FNiagaraVariableMetaData& MetaDataA = ScriptParameterA->GetScriptVariable()->Metadata;
+				const FNiagaraVariableMetaData& MetaDataB = ScriptParameterB->GetScriptVariable()->Metadata;
+
+				if(MetaDataA.bAdvancedDisplay && !MetaDataB.bAdvancedDisplay)
+				{
+					return false;
+				}
+				else if(!MetaDataA.bAdvancedDisplay && MetaDataB.bAdvancedDisplay)
+				{
+					return true;
+				}
+			
+				if(MetaDataA.GetEditorSortPriority_DEPRECATED() != MetaDataB.GetEditorSortPriority_DEPRECATED())
+				{
+					return MetaDataA.GetEditorSortPriority_DEPRECATED() < MetaDataB.GetEditorSortPriority_DEPRECATED();
+				}
+			
+				return ItemA.ToString().Compare(ItemB.ToString(), ESearchCase::Type::IgnoreCase) < 0;
+			}
+			
+			return false;
+		});
+
+		ParameterRoot->SortChildren(SortRootChildrenPredicate, false);
+	}
+}
+
 FGuid UNiagaraScriptSource::GetChangeID() 
 { 
 	return NodeGraph->GetChangeID(); 
diff --git a/Engine/Plugins/FX/Niagara/Source/NiagaraEditor/Private/NiagaraStackEditorData.cpp b/Engine/Plugins/FX/Niagara/Source/NiagaraEditor/Private/NiagaraStackEditorData.cpp
--- a/Engine/Plugins/FX/Niagara/Source/NiagaraEditor/Private/NiagaraStackEditorData.cpp
+++ b/Engine/Plugins/FX/Niagara/Source/NiagaraEditor/Private/NiagaraStackEditorData.cpp
@@ -179,6 +179,11 @@ void UNiagaraStackEditorData::SetStackEntryActiveSection(const FString& StackEnt
 	}
 }
 
+void UNiagaraStackEditorData::ClearStackEntryActiveSection(const FString& StackEntryKey)
+{
+	StackEntryKeyToActiveSectionMap.Remove(StackEntryKey);
+}
+
 const FText* UNiagaraStackEditorData::GetStackEntryDisplayName(const FString& StackEntryKey) const
 {
 	return StackEntryKeyToDisplayName.Find(StackEntryKey);
diff --git a/Engine/Plugins/FX/Niagara/Source/NiagaraEditor/Private/Tests/NiagaraStackViewModelTests.cpp b/Engine/Plugins/FX/Niagara/Source/NiagaraEditor/Private/Tests/NiagaraStackViewModelTests.cpp
--- a/Engine/Plugins/FX/Niagara/Source/NiagaraEditor/Private/Tests/NiagaraStackViewModelTests.cpp
+++ b/Engine/Plugins/FX/Niagara/Source/NiagaraEditor/Private/Tests/NiagaraStackViewModelTests.cpp
@@ -9,7 +9,7 @@
 #include "ViewModels/Stack/NiagaraStackRoot.h"
 #include "ViewModels/Stack/NiagaraStackScriptItemGroup.h"
 #include "ViewModels/Stack/NiagaraStackModuleItem.h"
-#include "ViewModels/Stack/NiagaraStackFunctionInputCollection.h"
+#include "ViewModels/Stack/NiagaraStackValueCollection.h"
 #include "ViewModels/Stack/NiagaraStackFunctionInput.h"
 #include "NiagaraNodeFunctionCall.h"
 #include "NiagaraEditorTestUtilities.h"
diff --git a/Engine/Plugins/FX/Niagara/Source/NiagaraEditor/Private/Toolkits/NiagaraScriptToolkit.cpp b/Engine/Plugins/FX/Niagara/Source/NiagaraEditor/Private/Toolkits/NiagaraScriptToolkit.cpp
--- a/Engine/Plugins/FX/Niagara/Source/NiagaraEditor/Private/Toolkits/NiagaraScriptToolkit.cpp
+++ b/Engine/Plugins/FX/Niagara/Source/NiagaraEditor/Private/Toolkits/NiagaraScriptToolkit.cpp
@@ -40,11 +40,15 @@
 #include "ViewModels/NiagaraParameterDefinitionsPanelViewModel.h"
 #include "ViewModels/NiagaraParameterPanelViewModel.h"
 #include "ViewModels/NiagaraScriptViewModel.h"
+#include "ViewModels/HierarchyEditor/NiagaraHierarchyScriptParametersViewModel.h"
+#include "Widgets/SNiagaraHierarchyEditor.h"
 #include "Widgets/Docking/SDockTab.h"
 #include "Widgets/SNiagaraParameterDefinitionsPanel.h"
 #include "Widgets/SNiagaraParameterPanel.h"
 #include "Widgets/SNiagaraScriptGraph.h"
+#include "Widgets/SNiagaraScriptInputPreviewPanel.h"
 #include "Widgets/SNiagaraSelectedObjectsDetails.h"
+#include "WIdgets/Input/SButton.h"
 #include "ToolMenus.h"
 
 #define LOCTEXT_NAMESPACE "NiagaraScriptToolkit"
@@ -55,7 +59,8 @@ const FName FNiagaraScriptToolkit::NodeGraphTabId(TEXT("NiagaraEditor_NodeGraph"
 const FName FNiagaraScriptToolkit::ScriptDetailsTabId(TEXT("NiagaraEditor_ScriptDetails"));
 const FName FNiagaraScriptToolkit::SelectedDetailsTabId(TEXT("NiagaraEditor_SelectedDetails"));
 const FName FNiagaraScriptToolkit::ParametersTabId(TEXT("NiagaraEditor_Parameters"));
-const FName FNiagaraScriptToolkit::ParametersTabId2(TEXT("NiagaraEditor_Paramters2"));
+const FName FNiagaraScriptToolkit::InputPreviewTabId(TEXT("NiagaraEditor_InputPreview"));
+const FName FNiagaraScriptToolkit::HierarchyEditor_ParametersTabId(TEXT("NiagaraEditor_HierarchyEditor_Parameters"));
 const FName FNiagaraScriptToolkit::ParameterDefinitionsTabId(TEXT("NiagaraEditor_ParameterDefinitions"));
 const FName FNiagaraScriptToolkit::StatsTabId(TEXT("NiagaraEditor_Stats"));
 const FName FNiagaraScriptToolkit::MessageLogTabID(TEXT("NiagaraEditor_MessageLog"));
@@ -77,6 +82,9 @@ FNiagaraScriptToolkit::~FNiagaraScriptToolkit()
 		ParameterDefinitionsPanelViewModel->Cleanup();
 	}
 
+	ParametersHierarchyViewModel->Finalize();
+	ParametersHierarchyViewModel = nullptr;
+
 	EditedNiagaraScript.Script->OnVMScriptCompiled().RemoveAll(this);
 	ScriptViewModel->GetGraphViewModel()->GetGraph()->RemoveOnGraphNeedsRecompileHandler(OnEditedScriptGraphChangedHandle);
 
@@ -113,7 +121,16 @@ void FNiagaraScriptToolkit::RegisterTabSpawners(const TSharedRef<class FTabManag
 		.SetDisplayName(LOCTEXT("ParametersTab", "Parameters"))
 		.SetGroup(WorkspaceMenuCategoryRef)
 		.SetIcon(FSlateIcon(FNiagaraEditorStyle::Get().GetStyleSetName(), "Tab.Parameters"));
-		
+
+	InTabManager->RegisterTabSpawner(InputPreviewTabId, FOnSpawnTab::CreateSP(this, &FNiagaraScriptToolkit::SpawnTabInputsPreview))
+		.SetDisplayName(LOCTEXT("InputPreviewParametersTab", "Input Preview"))
+		.SetGroup(WorkspaceMenuCategoryRef)
+		.SetIcon(FSlateIcon(FNiagaraEditorStyle::Get().GetStyleSetName(), "Tab.Parameters"));
+
+	InTabManager->RegisterTabSpawner(HierarchyEditor_ParametersTabId, FOnSpawnTab::CreateSP(this, &FNiagaraScriptToolkit::SpawnTabHierarchyEditor))
+		.SetDisplayName(LOCTEXT("HierarchyEditorTab", "Hierarchy Editor"))
+		.SetGroup(WorkspaceMenuCategoryRef)
+		.SetIcon(FSlateIcon(FNiagaraEditorStyle::Get().GetStyleSetName(), "Tab.Parameters"));
 	//@todo(ng) disable parameter definitions panel pending bug fixes
 // 	InTabManager->RegisterTabSpawner(ParameterDefinitionsTabId, FOnSpawnTab::CreateSP(this, &FNiagaraScriptToolkit::SpawnTabParameterDefinitions))
 // 		.SetDisplayName(LOCTEXT("ParameterDefinitions", "Parameter Definitions"))
@@ -143,7 +160,7 @@ void FNiagaraScriptToolkit::UnregisterTabSpawners(const TSharedRef<class FTabMan
 	InTabManager->UnregisterTabSpawner(ScriptDetailsTabId);
 	InTabManager->UnregisterTabSpawner(SelectedDetailsTabId);
 	InTabManager->UnregisterTabSpawner(ParametersTabId);
-	InTabManager->UnregisterTabSpawner(ParametersTabId2);
+	InTabManager->UnregisterTabSpawner(InputPreviewTabId);
 	//@todo(ng) disable parameter definitions panel pending bug fixes
 	//InTabManager->UnregisterTabSpawner(ParameterDefinitionsTabId);
 	InTabManager->UnregisterTabSpawner(StatsTabId);
@@ -188,11 +205,16 @@ void FNiagaraScriptToolkit::Initialize( const EToolkitMode::Type Mode, const TSh
 
 	const FGuid MessageLogGuidKey = FGuid::NewGuid();
 	NiagaraMessageLogViewModel = MakeShared<FNiagaraMessageLogViewModel>(GetNiagaraScriptMessageLogName(EditedNiagaraScript), MessageLogGuidKey, NiagaraMessageLog);
-
+	
 	bool bIsForDataProcessingOnly = false;
 	ScriptViewModel = MakeShareable(new FNiagaraStandaloneScriptViewModel(GetGraphEditorDisplayName(), ENiagaraParameterEditMode::EditAll, NiagaraMessageLogViewModel, MessageLogGuidKey, bIsForDataProcessingOnly));
 	ScriptViewModel->Initialize(EditedNiagaraScript, OriginalNiagaraScript);
 
+	ParametersHierarchyViewModel = NewObject<UNiagaraHierarchyScriptParametersViewModel>();
+	ParametersHierarchyViewModel->Initialize(ScriptViewModel.ToSharedRef());
+	ParametersHierarchyViewModel->OnHierarchyChanged().AddSP(this, &FNiagaraScriptToolkit::OnHierarchyChanged);
+	ParametersHierarchyViewModel->OnHierarchyPropertiesChanged().AddSP(this, &FNiagaraScriptToolkit::OnHierarchyPropertiesChanged);
+	
 	ParameterPanelViewModel = MakeShareable(new FNiagaraScriptToolkitParameterPanelViewModel(ScriptViewModel));
 	ParameterDefinitionsPanelViewModel = MakeShareable(new FNiagaraScriptToolkitParameterDefinitionsPanelViewModel(ScriptViewModel));
 
@@ -235,7 +257,7 @@ void FNiagaraScriptToolkit::Initialize( const EToolkitMode::Type Mode, const TSh
 			EditedNiagaraScript.Script->PostEditChangeVersionedProperty(ChangeEvent, SelectedVersion);
 		});
 
-	TSharedRef<FTabManager::FLayout> StandaloneDefaultLayout = FTabManager::NewLayout("Standalone_Niagara_Layout_v12")
+	TSharedRef<FTabManager::FLayout> StandaloneDefaultLayout = FTabManager::NewLayout("Standalone_Niagara_Layout_v13")
 	->AddArea
 	(
 		FTabManager::NewPrimaryArea()->SetOrientation(Orient_Vertical)
@@ -258,7 +280,8 @@ void FNiagaraScriptToolkit::Initialize( const EToolkitMode::Type Mode, const TSh
 					FTabManager::NewStack()
 					->SetSizeCoefficient(0.4f)
 					->AddTab(ParametersTabId, ETabState::OpenedTab)
-					->SetForegroundTab(ParametersTabId)
+					->AddTab(InputPreviewTabId, ETabState::OpenedTab)
+					->SetForegroundTab(InputPreviewTabId)
 				)
 				->Split
 				(
@@ -456,6 +479,16 @@ void FNiagaraScriptToolkit::OnVMScriptCompiled(UNiagaraScript*, const FGuid&)
 	}
 }
 
+void FNiagaraScriptToolkit::OnHierarchyChanged()
+{
+	ScriptViewModel->GetStandaloneScript().GetScriptData()->GetSource()->MarkNotSynchronized("Input Hierarchy Changed");
+}
+
+void FNiagaraScriptToolkit::OnHierarchyPropertiesChanged()
+{
+	ScriptViewModel->GetStandaloneScript().GetScriptData()->GetSource()->MarkNotSynchronized("Input Hierarchy Properties Changed");
+}
+
 TSharedRef<SDockTab> FNiagaraScriptToolkit::SpawnTabScriptDetails(const FSpawnTabArgs& Args)
 {
 	checkf(Args.GetTabId().TabType == ScriptDetailsTabId, TEXT("Wrong tab ID in NiagaraScriptToolkit"));
@@ -506,6 +539,39 @@ TSharedRef<SDockTab> FNiagaraScriptToolkit::SpawnTabScriptParameters(const FSpaw
 		SNew(SDockTab)
 		[
 			SNew(SNiagaraParameterPanel, ParameterPanelViewModel, GetToolkitCommands())
+			.SearchAdjacentWidget()
+			[
+				SNew(SButton)
+				.OnClicked(this, &FNiagaraScriptToolkit::SummonHierarchyEditor)
+				.ButtonStyle(FAppStyle::Get(), "RoundButton")
+				[
+					SNew(STextBlock)
+					.Text(LOCTEXT("EditHierarchy_ScriptInputs", "Edit Input Hierarchy"))
+				]
+			]
+		];
+
+	return SpawnedTab;
+}
+
+TSharedRef<SDockTab> FNiagaraScriptToolkit::SpawnTabInputsPreview(const FSpawnTabArgs& Args)
+{
+	checkf(Args.GetTabId().TabType == InputPreviewTabId, TEXT("Wrong tab ID in NiagaraScriptToolkit"));
+
+	TSharedRef<SDockTab> SpawnedTab =
+		SNew(SDockTab)
+		[
+			SNew(SVerticalBox)
+			+ SVerticalBox::Slot()
+			.AutoHeight()
+			[
+				SNew(SButton)
+				.OnClicked(this, &FNiagaraScriptToolkit::SummonParametersEditor)
+			]
+			+ SVerticalBox::Slot()
+			[
+				SNew(SNiagaraScriptInputPreviewPanel, *Cast<UNiagaraScriptSource>(EditedNiagaraScript.GetScriptData()->GetSource())->NodeGraph->GetScriptParameterHierarchyRoot(), FNiagaraScriptToolkit::SharedThis(this), ScriptViewModel->GetVariableSelection())
+			]
 		];
 
 	return SpawnedTab;
@@ -541,6 +607,25 @@ TSharedRef<SDockTab> FNiagaraScriptToolkit::SpawnTabStats(const FSpawnTabArgs& A
 	return SpawnedTab;
 }
 
+TSharedRef<SDockTab> FNiagaraScriptToolkit::SpawnTabHierarchyEditor(const FSpawnTabArgs& Args)
+{
+	check(Args.GetTabId().TabType == HierarchyEditor_ParametersTabId);
+	
+	TSharedRef<SDockTab> SpawnedTab = SNew(SDockTab)
+		.Label(LOCTEXT("UserParametersHierarchyTitle", "Parameters Editor"))
+		[
+			SNew(SBox)
+			.AddMetaData<FTagMetaData>(FTagMetaData(TEXT("UserParameters")))
+			[
+				SNew(SNiagaraHierarchyEditor, ParametersHierarchyViewModel)
+				.OnGenerateRowContentWidget_Static(&FNiagaraEditorUtilities::HierarchyEditor::Scripts::GenerateRowContentForScriptParameterHierarchyEditor)
+				//.OnGenerateCustomDetailsPanelNameWidget_Static(&GenerateCustomDetailsPanelNameWidgetForUserParameterEditor)
+			]
+		];
+	
+	return SpawnedTab;
+}
+
 TSharedRef<SDockTab> FNiagaraScriptToolkit::SpawnTabVersioning(const FSpawnTabArgs& Args)
 {
 	check(Args.GetTabId() == VersioningTabID);
@@ -575,6 +660,18 @@ TSharedRef<SDockTab> FNiagaraScriptToolkit::SpawnTabMessageLog(const FSpawnTabAr
 	return SpawnedTab;
 }
 
+FReply FNiagaraScriptToolkit::SummonParametersEditor()
+{
+	TSharedPtr<SDockTab> DockTab = TabManager->TryInvokeTab(HierarchyEditor_ParametersTabId);
+
+	if(DockTab)
+	{
+		DockTab->FlashTab();
+	}
+	
+	return FReply::Handled();
+}
+
 TSharedRef<SWidget> FNiagaraScriptToolkit::GenerateVersioningDropdownMenu(TSharedRef<FUICommandList> InCommandList)
 {
 	const bool bShouldCloseWindowAfterMenuSelection = true;
@@ -853,6 +950,12 @@ void FNiagaraScriptToolkit::AddReferencedObjects(FReferenceCollector& Collector)
 	Collector.AddReferencedObject(OriginalNiagaraScript.Script);
 	Collector.AddReferencedObject(EditedNiagaraScript.Script);
 	Collector.AddReferencedObject(VersionMetadata);
+	Collector.AddReferencedObject(ParametersHierarchyViewModel);
+}
+
+UNiagaraHierarchyScriptParametersViewModel* FNiagaraScriptToolkit::GetHierarchyViewModel() const
+{
+	return ParametersHierarchyViewModel;
 }
 
 void FNiagaraScriptToolkit::UpdateModuleStats()
@@ -1036,6 +1139,12 @@ void FNiagaraScriptToolkit::FocusGraphElementIfSameScriptID(const FNiagaraScript
 	}
 }
 
+FReply FNiagaraScriptToolkit::SummonHierarchyEditor() const
+{
+	TabManager->TryInvokeTab(HierarchyEditor_ParametersTabId);
+	return FReply::Handled();
+}
+
 void FNiagaraScriptToolkit::RefreshDetailsPanel()
 {
 	if (SelectedDetailsWidget.IsValid())
diff --git a/Engine/Plugins/FX/Niagara/Source/NiagaraEditor/Private/Toolkits/NiagaraScriptToolkit.h b/Engine/Plugins/FX/Niagara/Source/NiagaraEditor/Private/Toolkits/NiagaraScriptToolkit.h
--- a/Engine/Plugins/FX/Niagara/Source/NiagaraEditor/Private/Toolkits/NiagaraScriptToolkit.h
+++ b/Engine/Plugins/FX/Niagara/Source/NiagaraEditor/Private/Toolkits/NiagaraScriptToolkit.h
@@ -11,7 +11,9 @@
 #include "EditorUndoClient.h"
 #include "NiagaraEditorCommon.h"
 #include "NiagaraScript.h"
+#include "ViewModels/HierarchyEditor/NiagaraHierarchyScriptParametersViewModel.h"
 
+class UNiagaraHierarchyScriptParametersViewModel;
 class UNiagaraVersionMetaData;
 struct FCustomExpanderData;
 class IDetailsView;
@@ -36,13 +38,13 @@ public:
 
 	virtual void RegisterTabSpawners(const TSharedRef<class FTabManager>& TabManager) override;
 	virtual void UnregisterTabSpawners(const TSharedRef<class FTabManager>& TabManager) override;
-
+	
 	/** Edits the specified Niagara Script */
 	void Initialize( const EToolkitMode::Type Mode, const TSharedPtr< class IToolkitHost >& InitToolkitHost, UNiagaraScript* Script );
 
 	/** Destructor */
-	virtual ~FNiagaraScriptToolkit();
-
+	virtual ~FNiagaraScriptToolkit() override;
+	
 	//~ Begin IToolkit Interface
 	virtual FName GetToolkitFName() const override;
 	virtual FText GetBaseToolkitName() const override;
@@ -62,7 +64,9 @@ public:
 	{
 		return TEXT("FNiagaraScriptToolkit");
 	}
-
+	
+	UNiagaraHierarchyScriptParametersViewModel* GetHierarchyViewModel() const;
+	
 	/**
 	* Updates list of module info used to show stats
 	*/
@@ -91,6 +95,9 @@ private:
 
 	void OnVMScriptCompiled(UNiagaraScript* InScript, const FGuid& ScriptVersion);
 
+	void OnHierarchyChanged();
+	void OnHierarchyPropertiesChanged();
+	
 	/** Spawns the tab with the update graph inside */
 	TSharedRef<SDockTab> SpawnTabNodeGraph(const FSpawnTabArgs& Args);
 
@@ -102,8 +109,10 @@ private:
 
 	/** Spawns the tab with the parameter view. */
 	TSharedRef<SDockTab> SpawnTabScriptParameters(const FSpawnTabArgs& Args);
+	TSharedRef<SDockTab> SpawnTabInputsPreview(const FSpawnTabArgs& Args);
 	TSharedRef<SDockTab> SpawnTabParameterDefinitions(const FSpawnTabArgs& Args);
 	TSharedRef<SDockTab> SpawnTabStats(const FSpawnTabArgs& Args);
+	TSharedRef<SDockTab> SpawnTabHierarchyEditor(const FSpawnTabArgs& Args);
 
 	/** Spawns the tab with the version management. */
 	TSharedRef<SDockTab> SpawnTabVersioning(const FSpawnTabArgs& Args);
@@ -112,6 +121,7 @@ private:
 
 	TSharedRef<SDockTab> SpawnTabMessageLog(const FSpawnTabArgs& Args);
 
+	FReply SummonParametersEditor();
 	TSharedRef<SWidget> GenerateVersioningDropdownMenu(TSharedRef<FUICommandList> InCommandList);
 
 	/** Sets up commands for the toolkit toolbar. */
@@ -161,11 +171,13 @@ private:
 	*/
 	void FocusGraphElementIfSameScriptID(const FNiagaraScriptIDAndGraphFocusInfo* ElementToFocus);
 
+	FReply SummonHierarchyEditor() const;
 private:
-
 	/** The Script being edited */
 	TSharedPtr<FNiagaraStandaloneScriptViewModel> ScriptViewModel;
 
+	TObjectPtr<UNiagaraHierarchyScriptParametersViewModel> ParametersHierarchyViewModel;
+
 	/** The Parameter Panel displaying graph variables */
 	TSharedPtr<FNiagaraScriptToolkitParameterPanelViewModel> ParameterPanelViewModel;
 
@@ -179,17 +191,6 @@ private:
 	TSharedPtr<FNiagaraMessageLogViewModel> NiagaraMessageLogViewModel;
 	TSharedPtr<class SWidget> NiagaraMessageLog;
 
-	/**	The tab ids for the Niagara editor */
-	static const FName NodeGraphTabId; 
-	static const FName ScriptDetailsTabId;
-	static const FName SelectedDetailsTabId;
-	static const FName ParametersTabId;
-	static const FName ParametersTabId2;
-	static const FName ParameterDefinitionsTabId;
-	static const FName StatsTabId;
-	static const FName MessageLogTabID;
-	static const FName VersioningTabID;
-
 	/** Stats log, with the log listing that it reflects */
 	TSharedPtr<class SWidget> Stats;
 	TSharedPtr<class IMessageLogListing> StatsListing;
@@ -207,8 +208,19 @@ private:
 	TSharedPtr<class IDetailsView> DetailsView;
 	TObjectPtr<UNiagaraVersionMetaData> VersionMetadata = nullptr;
 	FText GetGraphEditorDisplayName() const;
-
-private:
+	
 	void RefreshDetailsPanel();
 
+public:
+	/**	The tab ids for the Niagara editor */
+	static const FName NodeGraphTabId; 
+	static const FName ScriptDetailsTabId;
+	static const FName SelectedDetailsTabId;
+	static const FName ParametersTabId;
+	static const FName InputPreviewTabId;
+	static const FName HierarchyEditor_ParametersTabId;
+	static const FName ParameterDefinitionsTabId;
+	static const FName StatsTabId;
+	static const FName MessageLogTabID;
+	static const FName VersioningTabID;
 };
diff --git a/Engine/Plugins/FX/Niagara/Source/NiagaraEditor/Private/Toolkits/SystemToolkitModes/NiagaraSystemToolkitModeBase.cpp b/Engine/Plugins/FX/Niagara/Source/NiagaraEditor/Private/Toolkits/SystemToolkitModes/NiagaraSystemToolkitModeBase.cpp
--- a/Engine/Plugins/FX/Niagara/Source/NiagaraEditor/Private/Toolkits/SystemToolkitModes/NiagaraSystemToolkitModeBase.cpp
+++ b/Engine/Plugins/FX/Niagara/Source/NiagaraEditor/Private/Toolkits/SystemToolkitModes/NiagaraSystemToolkitModeBase.cpp
@@ -28,7 +28,7 @@
 #include "Widgets/SNiagaraSelectedObjectsDetails.h"
 #include "Widgets/SNiagaraSpreadsheetView.h"
 #include "Widgets/SNiagaraGeneratedCodeView.h"
-#include "Widgets/SNiagaraHierarchy.h"
+#include "Widgets/SNiagaraHierarchyEditor.h"
 #include "Widgets/Layout/SWidgetSwitcher.h"
 #include "Widgets/Layout/SBox.h"
 #include "Widgets/Input/SButton.h"
@@ -44,6 +44,7 @@
 #include "Widgets/SNiagaraSystemUserParameters.h"
 #include "Customizations/NiagaraComponentDetails.h"
 #include "Styling/StyleColors.h"
+#include "ViewModels/NiagaraScratchPadViewModel.h"
 #include "ViewModels/NiagaraSystemEditorDocumentsViewModel.h"
 #include "ViewModels/HierarchyEditor/NiagaraSummaryViewViewModel.h"
 #include "ViewModels/HierarchyEditor/NiagaraUserParametersHierarchyViewModel.h"
@@ -72,6 +73,7 @@ const FName FNiagaraSystemToolkitModeBase::ScratchPadScriptsTabID(TEXT("NiagaraS
 const FName FNiagaraSystemToolkitModeBase::UserParametersTabID(TEXT("NiagaraSystemEditor_UserParameters"));
 const FName FNiagaraSystemToolkitModeBase::UserParametersHierarchyTabID(TEXT("NiagaraSystemEditor_UserParametersHierarchy"));
 const FName FNiagaraSystemToolkitModeBase::EmitterSummaryViewEditorTabID(TEXT("NiagaraSystemEditor_SummaryViewHierarchyEditor"));
+const FName FNiagaraSystemToolkitModeBase::ScratchPadHierarchyEditorTabID(TEXT("NiagaraSystemEditor_ScratchPadHierarchyEditor"));
 
 FNiagaraSystemToolkitModeBase::FNiagaraSystemToolkitModeBase(FName InModeName, TWeakPtr<FNiagaraSystemToolkit> InSystemToolkit) : FApplicationMode(InModeName), SystemToolkit(InSystemToolkit), SwitcherIdx(0)
 {
@@ -84,38 +86,22 @@ FNiagaraSystemToolkitModeBase::~FNiagaraSystemToolkitModeBase()
 	if (SystemToolkit.Pin().IsValid())
 	{
 		SystemToolkit.Pin()->GetSystemViewModel()->GetDocumentViewModel()->OnActiveDocumentChanged().Remove(DocChangedHandle);
-		if (LastActiveDocumentModel.IsValid())
-		{
-			TSharedPtr<FNiagaraScratchPadScriptViewModel> OldScratchScriptVM = LastActiveDocumentModel.Pin();
-			OldScratchScriptVM->GetGraphViewModel()->GetNodeSelection()->OnSelectedObjectsChanged().Remove(LastSelectionUpdateDelegate);
-		}
 		
 		if(UpdateSummaryViewHandle.IsValid())
 		{
 			SystemToolkit.Pin()->GetSystemViewModel()->GetSelectionViewModel()->OnEmitterHandleIdSelectionChanged().Remove(UpdateSummaryViewHandle);
 		}
+
+		if(UpdateScratchPadScriptHierarchyHandle.IsValid())
+		{
+			SystemToolkit.Pin()->GetSystemViewModel()->GetDocumentViewModel()->OnActiveDocumentChanged().Remove(UpdateScratchPadScriptHierarchyHandle);
+		}
 	}
 }
 
 void FNiagaraSystemToolkitModeBase::OnActiveDocumentChanged(TSharedPtr<SDockTab> NewActiveTab)
 {
 	TSharedPtr<FNiagaraSystemToolkit> Toolkit = StaticCastSharedPtr<FNiagaraSystemToolkit>(SystemToolkit.Pin());
-	if (LastActiveDocumentModel.IsValid())
-	{
-		TSharedPtr<FNiagaraScratchPadScriptViewModel> OldScratchScriptVM = LastActiveDocumentModel.Pin();
-		OldScratchScriptVM->GetGraphViewModel()->GetNodeSelection()->OnSelectedObjectsChanged().Remove(LastSelectionUpdateDelegate);
-
-		TArray<UNiagaraGraph*> EdGraphs = OldScratchScriptVM->GetEditableGraphs();
-		check(EdGraphs.Num() <= 1);
-		if (EdGraphs.Num() == 1)
-		{
-			EdGraphs[0]->OnSubObjectSelectionChanged().Remove(LastParamPanelSelectionUpdateDelegate);
-		}
-
-		Toolkit->GetSystemViewModel()->GetSelectionViewModel()->OnEntrySelectionChanged().Remove(LastSystemSelectionUpdateDelegate);
-		OldScratchScriptVM->GetGraphViewModel()->GetGraph()->RemoveOnGraphNeedsRecompileHandler(LastGraphEditDelegate);
-
-	}
 
 	SwitcherIdx = 0;
 	TSharedPtr<SDockTab>  ActiveTab = Toolkit->GetSystemViewModel()->GetDocumentViewModel()->GetActiveDocumentTab().Pin();
@@ -185,13 +171,18 @@ void FNiagaraSystemToolkitModeBase::OnSystemSelectionChanged()
 
 void FNiagaraSystemToolkitModeBase::PostActivateMode()
 {
-	// by default we want to close the user parameters hierarchy tab if it has been open before. T
-	// his fixes the issue of an empty tab being summoned from a system's cached layout when applied to an emitter.
+	// by default we want to close the user parameters hierarchy tab if it has been open before.
+	// This fixes the issue of an empty tab being summoned from a system's cached layout when applied to an emitter.
 	// for consistency, we close it not only for emitters but for systems too
 	if(TSharedPtr<SDockTab> UserParametersHierarchyDockTab = SystemToolkit.Pin()->GetTabManager()->FindExistingLiveTab(UserParametersHierarchyTabID))
 	{
 		UserParametersHierarchyDockTab->RequestCloseTab();
 	}
+
+	if(TSharedPtr<SDockTab> ScratchPadScriptHierarchyDockTab = SystemToolkit.Pin()->GetTabManager()->FindExistingLiveTab(ScratchPadHierarchyEditorTabID))
+	{
+		ScratchPadScriptHierarchyDockTab->RequestCloseTab();
+	}
 }
 
 void FNiagaraSystemToolkitModeBase::UpdateSelectionForActiveDocument()
@@ -347,6 +338,11 @@ void FNiagaraSystemToolkitModeBase::RegisterTabFactories(TSharedPtr<FTabManager>
 		.SetDisplayName(LOCTEXT("SummaryViewEditorTitle", "Edit Summary View"))
 		.SetGroup(WorkspaceMenuCategory.ToSharedRef())
 		.SetIcon(FSlateIcon(FNiagaraEditorStyle::Get().GetStyleSetName(), "Tab.UserParameterHierarchy"));
+
+	InTabManager->RegisterTabSpawner(ScratchPadHierarchyEditorTabID, FOnSpawnTab::CreateSP(this, &FNiagaraSystemToolkitModeBase::SpawnTab_ScratchPadHierarchyEditor))
+		.SetDisplayName(LOCTEXT("SummaryViewEditorTitle", "Edit Scratch Pad Hierarchy"))
+		.SetGroup(WorkspaceMenuCategory.ToSharedRef())
+		.SetIcon(FSlateIcon(FNiagaraEditorStyle::Get().GetStyleSetName(), "Tab.UserParameterHierarchy"));
 }
 
 int FNiagaraSystemToolkitModeBase::GetActiveSelectionDetailsIndex() const
@@ -580,6 +576,17 @@ TSharedRef<SDockTab> FNiagaraSystemToolkitModeBase::SpawnTab_SystemParameters(co
 		[
 			SAssignNew(SystemToolkit.Pin()->ParameterPanel, SNiagaraParameterPanel, SystemToolkit.Pin()->ParameterPanelViewModel, SystemToolkit.Pin()->GetToolkitCommands())
 			.ShowParameterSynchronizingWithLibraryIconExternallyReferenced(false)
+			.SearchAdjacentWidget()
+			[
+				SNew(SButton)
+				.OnClicked(this, &FNiagaraSystemToolkitModeBase::SummonScratchPadScriptHierarchyEditor)
+				.ButtonStyle(FAppStyle::Get(), "RoundButton")
+				.Visibility(this, &FNiagaraSystemToolkitModeBase::GetSummonScratchPadHierarchyEditorButtonVisibility)
+				[
+					SNew(STextBlock)
+					.Text(LOCTEXT("EditHierarchy_ScriptInputs", "Edit Input Hierarchy"))
+				]
+			]
 		];
 	SystemToolkit.Pin()->RefreshParameters();
 
@@ -946,7 +953,7 @@ TSharedRef<SDockTab> FNiagaraSystemToolkitModeBase::SpawnTab_UserParametersHiera
 			SNew(SBox)
 			.AddMetaData<FTagMetaData>(FTagMetaData(TEXT("UserParameters")))
 			[
-				SNew(SNiagaraHierarchy, SystemToolkit.Pin()->GetSystemViewModel()->GetUserParametersHierarchyViewModel())
+				SNew(SNiagaraHierarchyEditor, SystemToolkit.Pin()->GetSystemViewModel()->GetUserParametersHierarchyViewModel())
 				.OnGenerateRowContentWidget_Static(&GenerateRowContentForUserParameterHierarchyEditor, SystemToolkit.Pin()->GetSystemViewModel())
 				.OnGenerateCustomDetailsPanelNameWidget_Static(&GenerateCustomDetailsPanelNameWidgetForUserParameterEditor)
 			]
@@ -1094,7 +1101,7 @@ TSharedRef<SWidget> FNiagaraSystemToolkitModeBase::CreateSummaryViewWidget() con
 		TSharedPtr<FNiagaraEmitterViewModel> EmitterViewModel = SystemToolkit.Pin()->GetSystemViewModel()->GetEmitterHandleViewModels()[0]->GetEmitterViewModel();
 		UNiagaraSummaryViewViewModel* SummaryViewHierarchyViewModel = EmitterViewModel->GetSummaryHierarchyViewModel();
 		
-		ContentWidget = SNew(SNiagaraHierarchy, SummaryViewHierarchyViewModel)
+		ContentWidget = SNew(SNiagaraHierarchyEditor, SummaryViewHierarchyViewModel)
 		.OnGenerateRowContentWidget_Static(&GenerateRowContentForSummaryViewHierarchyEditor, EmitterViewModel)
 		.OnGenerateCustomDetailsPanelNameWidget_Static(&GenerateCustomDetailsPanelNameWidgetForSummaryViewEditor);
 	}
@@ -1114,7 +1121,7 @@ TSharedRef<SWidget> FNiagaraSystemToolkitModeBase::CreateSummaryViewWidget() con
 			TSharedPtr<FNiagaraEmitterViewModel> EmitterViewModel = SystemToolkit.Pin()->GetSystemViewModel()->GetEmitterHandleViewModelById(SelectedEmitterHandleGuids[0])->GetEmitterViewModel();
 			UNiagaraSummaryViewViewModel* SummaryViewHierarchyViewModel = EmitterViewModel->GetSummaryHierarchyViewModel();
 		
-			ContentWidget = SNew(SNiagaraHierarchy, SummaryViewHierarchyViewModel)
+			ContentWidget = SNew(SNiagaraHierarchyEditor, SummaryViewHierarchyViewModel)
 			.OnGenerateRowContentWidget_Static(&GenerateRowContentForSummaryViewHierarchyEditor, EmitterViewModel)
 			.OnGenerateCustomDetailsPanelNameWidget_Static(&GenerateCustomDetailsPanelNameWidgetForSummaryViewEditor);
 		}
@@ -1136,6 +1143,64 @@ void FNiagaraSystemToolkitModeBase::OnSummaryViewEditorClosed(TSharedRef<SDockTa
 	SummaryViewContainer->SetContent(SNullWidget::NullWidget);
 }
 
+FReply FNiagaraSystemToolkitModeBase::SummonScratchPadScriptHierarchyEditor()
+{
+	SystemToolkit.Pin()->GetTabManager()->TryInvokeTab(ScratchPadHierarchyEditorTabID);
+	return FReply::Handled();
+}
+
+TSharedRef<SWidget> FNiagaraSystemToolkitModeBase::CreateScratchPadHierarchyWidget()
+{
+	TSharedPtr<FNiagaraScratchPadScriptViewModel> ActiveScratchPadViewModel = SystemToolkit.Pin()->GetSystemViewModel()->GetScriptScratchPadViewModel()->GetActiveScriptViewModel();
+
+	ActiveScratchPadViewModel = SystemToolkit.Pin()->GetSystemViewModel()->GetDocumentViewModel()->GetActiveScratchPadViewModelIfSet();
+	
+	TSharedPtr<SWidget> ContentWidget = SNullWidget::NullWidget;
+	if(ActiveScratchPadViewModel.IsValid())
+	{
+		UNiagaraHierarchyScriptParametersViewModel* ScriptHierarchyViewModel = ActiveScratchPadViewModel->GetHierarchyViewModel();
+		
+		ContentWidget = SNew(SNiagaraHierarchyEditor, ScriptHierarchyViewModel)
+		.OnGenerateRowContentWidget_Static(&FNiagaraEditorUtilities::HierarchyEditor::Scripts::GenerateRowContentForScriptParameterHierarchyEditor);
+		//.OnGenerateCustomDetailsPanelNameWidget_Static(&GenerateCustomDetailsPanelNameWidgetForSummaryViewEditor);
+	}
+	else
+	{
+		ContentWidget = SNew(SBox)
+		.HAlign(HAlign_Center)
+		.VAlign(VAlign_Center)
+		[
+			SNew(STextBlock).Text(LOCTEXT("ScratchPadHierarchyEditorInvalidSelection", "Please select a scratch pad."))
+		];		
+	}
+
+	return ContentWidget.ToSharedRef();
+}
+
+void FNiagaraSystemToolkitModeBase::UpdateScratchPadActiveScriptChanged(TSharedPtr<SDockTab> DockTab)
+{
+	if(SystemToolkit.Pin()->GetTabManager()->FindExistingLiveTab(ScratchPadHierarchyEditorTabID).IsValid() && ScratchPadHierarchyContainer.IsValid())
+	{
+		TSharedPtr<FNiagaraScratchPadScriptViewModel> ActiveScratchPadViewModel = SystemToolkit.Pin()->GetSystemViewModel()->GetDocumentViewModel()->GetActiveScratchPadViewModelIfSet();
+		if(LastActiveScratchPadViewModel != ActiveScratchPadViewModel)
+		{
+			ScratchPadHierarchyContainer->SetContent(CreateScratchPadHierarchyWidget());
+		}
+			
+		LastActiveScratchPadViewModel = ActiveScratchPadViewModel;
+	}
+}
+
+void FNiagaraSystemToolkitModeBase::OnScratchPadHierarchyEditorClosed(TSharedRef<SDockTab> DockTab)
+{
+	ScratchPadHierarchyContainer->SetContent(SNullWidget::NullWidget);
+}
+
+EVisibility FNiagaraSystemToolkitModeBase::GetSummonScratchPadHierarchyEditorButtonVisibility() const
+{
+	return SystemToolkit.Pin()->GetSystemViewModel()->GetDocumentViewModel()->GetActiveScratchPadViewModelIfSet() != nullptr ? EVisibility::Visible : EVisibility::Collapsed;
+}
+
 TSharedRef<SDockTab> FNiagaraSystemToolkitModeBase::SpawnTab_SummaryViewEditor(const FSpawnTabArgs& Args)
 {
 	check(Args.GetTabId().TabType == EmitterSummaryViewEditorTabID);
@@ -1163,4 +1228,31 @@ TSharedRef<SDockTab> FNiagaraSystemToolkitModeBase::SpawnTab_SummaryViewEditor(c
 	return SpawnedTab;
 }
 
+TSharedRef<SDockTab> FNiagaraSystemToolkitModeBase::SpawnTab_ScratchPadHierarchyEditor(const FSpawnTabArgs& Args)
+{
+	check(Args.GetTabId().TabType == ScratchPadHierarchyEditorTabID);
+
+	if(UpdateScratchPadScriptHierarchyHandle.IsValid())
+	{
+		SystemToolkit.Pin()->GetSystemViewModel()->GetDocumentViewModel()->OnActiveDocumentChanged().Remove(UpdateScratchPadScriptHierarchyHandle);
+	}
+	
+	UpdateScratchPadScriptHierarchyHandle = SystemToolkit.Pin()->GetSystemViewModel()->GetDocumentViewModel()->OnActiveDocumentChanged().AddSP(this, &FNiagaraSystemToolkitModeBase::UpdateScratchPadActiveScriptChanged);
+	
+	TSharedRef<SDockTab> SpawnedTab = SNew(SDockTab)
+		.OnTabClosed(this, &FNiagaraSystemToolkitModeBase::OnScratchPadHierarchyEditorClosed)
+		.Label(LOCTEXT("ScratchPadHierarchyTitle", "Edit Hierarchy"))
+		[
+			SAssignNew(ScratchPadHierarchyContainer, SBox)
+			.AddMetaData<FTagMetaData>(FTagMetaData(TEXT("ScratchPadHierarchy")))
+		];
+
+	TSharedPtr<FNiagaraScratchPadScriptViewModel> ActiveScratchPadViewModel = SystemToolkit.Pin()->GetSystemViewModel()->GetDocumentViewModel()->GetActiveScratchPadViewModelIfSet();
+	LastActiveScratchPadViewModel = ActiveScratchPadViewModel;
+	
+	ScratchPadHierarchyContainer->SetContent(CreateScratchPadHierarchyWidget());
+	
+	return SpawnedTab;
+}
+
 #undef LOCTEXT_NAMESPACE
diff --git a/Engine/Plugins/FX/Niagara/Source/NiagaraEditor/Private/ViewModels/HierarchyEditor/NiagaraHierarchyScriptParametersViewModel.cpp b/Engine/Plugins/FX/Niagara/Source/NiagaraEditor/Private/ViewModels/HierarchyEditor/NiagaraHierarchyScriptParametersViewModel.cpp
new file mode 100644
--- /dev/null
+++ b/Engine/Plugins/FX/Niagara/Source/NiagaraEditor/Private/ViewModels/HierarchyEditor/NiagaraHierarchyScriptParametersViewModel.cpp
@@ -0,0 +1,269 @@
+﻿// Fill out your copyright notice in the Description page of Project Settings.
+
+
+#include "ViewModels/HierarchyEditor/NiagaraHierarchyScriptParametersViewModel.h"
+
+#include "NiagaraConstants.h"
+#include "NiagaraEditorUtilities.h"
+#include "NiagaraGraph.h"
+#include "NiagaraScriptSource.h"
+#include "ViewModels/NiagaraScriptViewModel.h"
+
+#define LOCTEXT_NAMESPACE "NiagaraScriptParameterHierarchyEditor"
+
+void UNiagaraHierarchyScriptParameter::Initialize(UNiagaraScriptVariable& InParameterScriptVariable)
+{
+	ParameterScriptVariable = &InParameterScriptVariable;
+	SetIdentity(FNiagaraHierarchyIdentity({InParameterScriptVariable.Metadata.GetVariableGuid()}, {}));
+}
+
+void UNiagaraHierarchyScriptParametersViewModel::Initialize(TSharedRef<FNiagaraScriptViewModel> InScriptViewModel)
+{
+	ScriptViewModelWeak = InScriptViewModel;
+
+	Cast<UNiagaraScriptSource>(ScriptViewModelWeak.Pin()->GetStandaloneScript().GetScriptData()->GetSource())->NodeGraph->OnParametersChanged().AddUObject(this, &UNiagaraHierarchyScriptParametersViewModel::OnParametersChanged);
+	UNiagaraHierarchyViewModelBase::Initialize();
+
+	UNiagaraHierarchyScriptParameterRefreshContext* ScriptParameterRefreshContext = NewObject<UNiagaraHierarchyScriptParameterRefreshContext>(this);
+	ScriptParameterRefreshContext->SetNiagaraGraph(Cast<UNiagaraScriptSource>(InScriptViewModel->GetStandaloneScript().GetScriptData()->GetSource())->NodeGraph);
+	SetRefreshContext(ScriptParameterRefreshContext);
+}
+
+TSharedRef<FNiagaraScriptViewModel> UNiagaraHierarchyScriptParametersViewModel::GetScriptViewModel() const
+{
+	TSharedPtr<FNiagaraScriptViewModel> ScriptViewModelPinned = ScriptViewModelWeak.Pin();
+	checkf(ScriptViewModelPinned.IsValid(), TEXT("Script view model destroyed before parameters hierarchy view model."));
+	return ScriptViewModelPinned.ToSharedRef();
+}
+
+
+UNiagaraHierarchyRoot* UNiagaraHierarchyScriptParametersViewModel::GetHierarchyRoot() const
+{
+	const TArray<FVersionedNiagaraScriptWeakPtr>& Scripts = GetScriptViewModel()->GetScripts();
+	if(ensure(Scripts.Num() >= 1 && Scripts[0].Pin().Script != nullptr) == false)
+	{
+		return nullptr;
+	}
+
+	FVersionedNiagaraScriptData* ScriptData = Scripts[0].Pin().GetScriptData();
+	if(ensure(ScriptData != nullptr) == false)
+	{
+		return nullptr;
+	}
+
+	return Cast<UNiagaraScriptSource>(ScriptData->GetSource())->NodeGraph->GetScriptParameterHierarchyRoot();
+}
+
+TSubclassOf<UNiagaraHierarchyCategory> UNiagaraHierarchyScriptParametersViewModel::GetCategoryDataClass() const
+{
+	return UNiagaraHierarchyScriptCategory::StaticClass();
+}
+
+TSharedPtr<FNiagaraHierarchyRootViewModel> UNiagaraHierarchyScriptParametersViewModel::CreateRootViewModelForData(UNiagaraHierarchyRoot* Root, bool bIsForHierarchy)
+{
+	return MakeShared<FNiagaraHierarchyScriptRootViewModel>(Root, this, bIsForHierarchy);
+}
+
+TSharedPtr<FNiagaraHierarchyItemViewModelBase> UNiagaraHierarchyScriptParametersViewModel::CreateViewModelForData(UNiagaraHierarchyItemBase* ItemBase, TSharedPtr<FNiagaraHierarchyItemViewModelBase> Parent)
+{
+	if(UNiagaraHierarchyScriptParameter* Item = Cast<UNiagaraHierarchyScriptParameter>(ItemBase))
+	{
+		return MakeShared<FNiagaraHierarchyScriptParameterViewModel>(Item, Parent, this, Parent.IsValid() ? Parent->IsForHierarchy() : false);
+	}
+	else if(UNiagaraHierarchyCategory* Category = Cast<UNiagaraHierarchyCategory>(ItemBase))
+	{
+		return MakeShared<FNiagaraHierarchyCategoryViewModel>(Category, Parent, this, Parent.IsValid() ? Parent->IsForHierarchy() : false);
+	}
+
+	return nullptr;
+}
+
+void UNiagaraHierarchyScriptParametersViewModel::PrepareSourceItems(UNiagaraHierarchyRoot* SourceRoot, TSharedPtr<FNiagaraHierarchyRootViewModel> SourceRootViewModel)
+{
+	UNiagaraScriptSourceBase* SourceBase = GetScriptViewModel()->GetStandaloneScript().GetScriptData()->GetSource();
+	UNiagaraScriptSource* ScriptSource = Cast<UNiagaraScriptSource>(SourceBase);
+	const TMap<FNiagaraVariable, TObjectPtr<UNiagaraScriptVariable>>& ScriptVariableMap = ScriptSource->NodeGraph->GetAllMetaData();
+
+	TArray<UNiagaraHierarchyItemBase*> OldChildren = SourceRoot->GetChildrenMutable();
+	SourceRoot->GetChildrenMutable().Empty();
+
+	for(const auto& ScriptVariablePair : ScriptVariableMap)
+	{
+		// We only want to be able to organize module inputs & static switches 
+		if(ScriptVariablePair.Key.IsInNameSpace(FNiagaraConstants::ModuleNamespace) == false && ScriptVariablePair.Value->GetIsStaticSwitch() == false)
+		{
+			continue;
+		}
+		
+		UNiagaraScriptVariable* ScriptVariable = ScriptVariablePair.Value;
+		if(UNiagaraHierarchyItemBase** ExistingChild = OldChildren.FindByPredicate([ScriptVariable](UNiagaraHierarchyItemBase* ItemBase)
+		{
+			return ItemBase->GetPersistentIdentity().Guids[0] == ScriptVariable->Metadata.GetVariableGuid();
+		}))
+		{
+			SourceRoot->GetChildrenMutable().Add(*ExistingChild);
+			continue;
+		}
+
+		// since the source items are transient we need to create them here and keep them around until the end of the tool's lifetime
+		UNiagaraHierarchyScriptParameter* ScriptParameterHierarchyObject = NewObject<UNiagaraHierarchyScriptParameter>(SourceRoot);
+		ScriptParameterHierarchyObject->Initialize(*ScriptVariable);
+		SourceRoot->GetChildrenMutable().Add(ScriptParameterHierarchyObject);
+	}
+}
+
+void UNiagaraHierarchyScriptParametersViewModel::SetupCommands()
+{
+	Super::SetupCommands();
+}
+
+TSharedRef<FNiagaraHierarchyDragDropOp> UNiagaraHierarchyScriptParametersViewModel::CreateDragDropOp(TSharedRef<FNiagaraHierarchyItemViewModelBase> Item)
+{
+	if(UNiagaraHierarchyCategory* HierarchyCategory = Cast<UNiagaraHierarchyCategory>(Item->GetDataMutable()))
+	{
+		TSharedRef<FNiagaraHierarchyDragDropOp> CategoryDragDropOp = MakeShared<FNiagaraHierarchyDragDropOp>(Item);
+		CategoryDragDropOp->SetAdditionalLabel(FText::FromString(Item->ToString()));
+		CategoryDragDropOp->Construct();
+		return CategoryDragDropOp;
+	}
+	else if(UNiagaraHierarchyScriptParameter* ScriptParameter = Cast<UNiagaraHierarchyScriptParameter>(Item->GetDataMutable()))
+	{
+		TSharedPtr<FNiagaraHierarchyItemViewModel> ScriptParameterViewModel = StaticCastSharedRef<FNiagaraHierarchyItemViewModel>(Item);
+		TSharedRef<FNiagaraHierarchyDragDropOp> ScriptParameterDragDropOp = MakeShared<FNiagaraHierarchyScriptParameterDragDropOp>(ScriptParameterViewModel);
+		ScriptParameterDragDropOp->Construct();
+		return ScriptParameterDragDropOp;
+	}
+	
+	check(false);
+	return MakeShared<FNiagaraHierarchyDragDropOp>(nullptr);
+}
+
+void UNiagaraHierarchyScriptParametersViewModel::FinalizeInternal()
+{
+	if(ScriptViewModelWeak.IsValid())
+	{
+		Cast<UNiagaraScriptSource>(ScriptViewModelWeak.Pin()->GetStandaloneScript().GetScriptData()->GetSource())->NodeGraph->OnParametersChanged().RemoveAll(this);
+	}
+
+	Super::FinalizeInternal();
+}
+
+void UNiagaraHierarchyScriptParametersViewModel::OnParametersChanged(TOptional<UNiagaraGraph::FParametersChangedData> ParametersChangedData)
+{
+	ForceFullRefresh();
+}
+
+TSharedRef<SWidget> FNiagaraHierarchyScriptParameterDragDropOp::CreateCustomDecorator() const
+{
+	if(const UNiagaraHierarchyScriptParameter* ScriptParameter = Cast<UNiagaraHierarchyScriptParameter>(DraggedItem.Pin()->GetData()))
+	{
+		return FNiagaraParameterUtilities::GetParameterWidget(ScriptParameter->GetVariable(), false, false);
+	}
+
+	return SNullWidget::NullWidget;
+}
+
+bool FNiagaraHierarchyScriptParameterViewModel::DoesExternalDataStillExist(const UNiagaraHierarchyDataRefreshContext* Context) const
+{
+	const UNiagaraHierarchyScriptParameterRefreshContext* RefreshContext = CastChecked<UNiagaraHierarchyScriptParameterRefreshContext>(Context);
+	if(RefreshContext->GetNiagaraGraph()->GetAllMetaData().Contains(Cast<UNiagaraHierarchyScriptParameter>(GetDataMutable())->GetVariable()) == false)
+	{
+		return false;
+	}
+
+	const UNiagaraScriptVariable* ScriptVariable = Cast<UNiagaraHierarchyScriptParameter>(GetDataMutable())->GetScriptVariable();
+
+	// We make sure that the variable not only still exists but also qualifies for the hierarchy (namespace can change for example)
+	if(ScriptVariable->GetIsStaticSwitch() == false && ScriptVariable->Variable.IsInNameSpace(FNiagaraConstants::ModuleNamespace) == false)
+	{
+		return false;
+	}
+	
+	return true;
+}
+
+FNiagaraHierarchyItemViewModelBase::FCanPerformActionResults FNiagaraHierarchyScriptParameterViewModel::CanDropOnInternal(TSharedPtr<FNiagaraHierarchyItemViewModelBase> DraggedItem, EItemDropZone ItemDropZone)
+{
+	// if the input isn't editable, we don't allow any drops on/above/below the item.
+	// Even though it technically works, the merge process will only re-add the item at the end and not preserve order so there is no point in allowing dropping above/below
+	if(IsEditableByUser().bCanPerform == false)
+	{
+		return false;
+	}
+
+	FCanPerformActionResults AllowDrop(false);
+	
+	TSharedPtr<FNiagaraHierarchyItemViewModelBase> TargetDropItem = AsShared();
+
+	// we only allow drops if some general conditions are fulfilled
+	if(DraggedItem->GetData() != TargetDropItem->GetData() &&
+		(!DraggedItem->HasParent(TargetDropItem, false) || ItemDropZone != EItemDropZone::OntoItem)  &&
+		!TargetDropItem->HasParent(DraggedItem, true))
+	{
+		if(ItemDropZone == EItemDropZone::OntoItem)
+		{
+			// We support nested inputs
+			if(DraggedItem->GetData()->IsA<UNiagaraHierarchyScriptParameter>() && TargetDropItem->GetData()->IsA<UNiagaraHierarchyScriptParameter>())
+			{
+				FText BaseMessage = LOCTEXT("DroppingInputOnInputNestedChild", "This will nest input {0} under input {1}");
+				AllowDrop.CanPerformMessage = FText::FormatOrdered(BaseMessage, DraggedItem->ToStringAsText(), TargetDropItem->ToStringAsText());
+				AllowDrop.bCanPerform = true;
+			}
+		}
+		else
+		{
+			// if the dragged item is an input, we generally allow above/below, even for nested child inputs
+			if(DraggedItem->GetData()->IsA<UNiagaraHierarchyScriptParameter>())
+			{
+				AllowDrop.bCanPerform = true;
+			}
+			else
+			{
+				// we use default logic only if there is no parent input. Nested children are not allowed to contain anything but other inputs.
+				if(TargetDropItem->GetParent().Pin()->GetData<UNiagaraHierarchyScriptParameter>() == nullptr)
+				{
+					AllowDrop = FNiagaraHierarchyItemViewModel::CanDropOnInternal(DraggedItem, ItemDropZone);
+				}
+			}
+		}
+	}
+
+	return AllowDrop;
+}
+
+void FNiagaraHierarchyScriptParameterViewModel::OnDroppedOnInternal(TSharedPtr<FNiagaraHierarchyItemViewModelBase> DroppedItem, EItemDropZone ItemDropZone)
+{
+	if(ItemDropZone != EItemDropZone::OntoItem)
+	{
+		return FNiagaraHierarchyItemViewModel::OnDroppedOnInternal(DroppedItem, ItemDropZone);
+	}
+	else
+	{
+		FScopedTransaction Transaction(LOCTEXT("Transaction_AddedChildInput", "Added child input"));
+		HierarchyViewModel->GetHierarchyRoot()->Modify();
+		
+		if(DroppedItem->IsForHierarchy() == false)
+		{
+			TSharedRef<FNiagaraHierarchyItemViewModelBase> AddedItemViewModel = DuplicateToThis(DroppedItem);
+			AddedItemViewModel->SyncViewModelsToData();
+		}
+		else
+		{
+			TSharedRef<FNiagaraHierarchyItemViewModelBase> ReparentedViewModel = ReparentToThis(DroppedItem);
+			ReparentedViewModel->SyncViewModelsToData();
+		}
+
+		HierarchyViewModel->RefreshHierarchyView();
+		HierarchyViewModel->RefreshSourceView();
+	}
+}
+
+void FNiagaraHierarchyScriptRootViewModel::SortChildrenData() const
+{
+	GetDataMutable()->GetChildrenMutable().StableSort([](const UNiagaraHierarchyItemBase& ItemA, const UNiagaraHierarchyItemBase& ItemB)
+		{
+			return ItemA.IsA<UNiagaraHierarchyItem>() && ItemB.IsA<UNiagaraHierarchyCategory>();
+		});
+}
+
+#undef LOCTEXT_NAMESPACE
diff --git a/Engine/Plugins/FX/Niagara/Source/NiagaraEditor/Private/ViewModels/HierarchyEditor/NiagaraHierarchyViewModelBase.cpp b/Engine/Plugins/FX/Niagara/Source/NiagaraEditor/Private/ViewModels/HierarchyEditor/NiagaraHierarchyViewModelBase.cpp
--- a/Engine/Plugins/FX/Niagara/Source/NiagaraEditor/Private/ViewModels/HierarchyEditor/NiagaraHierarchyViewModelBase.cpp
+++ b/Engine/Plugins/FX/Niagara/Source/NiagaraEditor/Private/ViewModels/HierarchyEditor/NiagaraHierarchyViewModelBase.cpp
@@ -9,7 +9,7 @@
 #include "ViewModels/NiagaraSystemViewModel.h"
 #include "ViewModels/Stack/NiagaraStackViewModel.h"
 #include "Widgets/Views/STableRow.h"
-#include "Widgets/SNiagaraHierarchy.h"
+#include "Widgets/SNiagaraHierarchyEditor.h"
 #include "Widgets/Text/SInlineEditableTextBlock.h"
 #include "Widgets/SToolTip.h"
 #include "Widgets/Input/SCheckBox.h"
@@ -301,14 +301,17 @@ FNiagaraHierarchyItemViewModelBase::FCanPerformActionResults FNiagaraHierarchyCa
 					return Results;
 				}
 
+				Results.CanPerformMessage = LOCTEXT("MoveCategoryText", "Move category here");
+
 				// if we are making a category a sibling of another at the root level, the section will be set to the currently active section. Let that be known.
 				if(Parent.Pin()->GetData()->IsA<UNiagaraHierarchyRoot>())
 				{
 					UNiagaraHierarchyCategory* DraggedCategory = Cast<UNiagaraHierarchyCategory>(DraggedItem->GetDataMutable());
 					if(DraggedCategory->GetSection() != HierarchyViewModel->GetActiveHierarchySectionData())
 					{
-						FText BaseMessage = LOCTEXT("CategorySectionWillUpdateDueToDrop", "The section of the category will change to {0} after the drop");
-						Results.CanPerformMessage = FText::FormatOrdered(BaseMessage, HierarchyViewModel->GetActiveHierarchySectionData() == nullptr ? FText::FromString("All") : HierarchyViewModel->GetActiveHierarchySectionData()->GetSectionNameAsText());
+						FText SectionChangeBaseText = LOCTEXT("CategorySectionWillUpdateDueToDrop", "The section of the category will change to {0} after the drop");
+						FText ActualSectionChangeText = FText::FormatOrdered(SectionChangeBaseText, HierarchyViewModel->GetActiveHierarchySectionData() == nullptr ? FText::FromString("All") : HierarchyViewModel->GetActiveHierarchySectionData()->GetSectionNameAsText());
+						Results.CanPerformMessage = FText::FormatOrdered(FText::AsCultureInvariant("{0}\n{1}"), Results.CanPerformMessage, ActualSectionChangeText);
 					}
 				}
 			}
@@ -325,6 +328,8 @@ FNiagaraHierarchyItemViewModelBase::FCanPerformActionResults FNiagaraHierarchyCa
 					Results.CanPerformMessage = LOCTEXT("CantDropCategoryOnCategorySameChildCategoryName", "A sub-category of the same name already exists! Please rename your category first.");
 					return Results;
 				}
+
+				Results.CanPerformMessage = LOCTEXT("CreateSubcategory", "Drop category here to create a sub-category");
 			}
 
 			Results.bCanPerform = true;
@@ -410,8 +415,8 @@ void UNiagaraHierarchyViewModelBase::Initialize()
 	}
 	
 	SetupCommands();
-
-	HierarchyRootViewModel = MakeShared<FNiagaraHierarchyRootViewModel>(HierarchyRoot.Get(), this, true);
+	
+	HierarchyRootViewModel = CreateRootViewModelForData(HierarchyRoot.Get(), true);
 	HierarchyRootViewModel->Initialize();
 	HierarchyRootViewModel->AddChildFilter(FNiagaraHierarchyItemViewModelBase::FOnFilterChild::CreateUObject(this, &UNiagaraHierarchyViewModelBase::FilterForHierarchySection));
 	HierarchyRootViewModel->AddChildFilter(FNiagaraHierarchyItemViewModelBase::FOnFilterChild::CreateUObject(this, &UNiagaraHierarchyViewModelBase::FilterForUncategorizedRootItemsInAllSection));
@@ -432,6 +437,11 @@ void UNiagaraHierarchyViewModelBase::Finalize()
 	FinalizeInternal();
 }
 
+TSubclassOf<UNiagaraHierarchyCategory> UNiagaraHierarchyViewModelBase::GetCategoryDataClass() const
+{
+	return UNiagaraHierarchyCategory::StaticClass();
+}
+
 void UNiagaraHierarchyViewModelBase::ForceFullRefresh()
 {
 	RefreshSourceItemsRequestedDelegate.ExecuteIfBound();
@@ -522,6 +532,11 @@ bool UNiagaraHierarchyViewModelBase::FilterForUncategorizedRootItemsInAllSection
 	return true;
 }
 
+TSharedPtr<FNiagaraHierarchyRootViewModel> UNiagaraHierarchyViewModelBase::CreateRootViewModelForData(UNiagaraHierarchyRoot* Root, bool bIsForHierarchy)
+{
+	return MakeShared<FNiagaraHierarchyRootViewModel>(Root, this, bIsForHierarchy);
+}
+
 TSharedPtr<FNiagaraHierarchyItemViewModelBase> UNiagaraHierarchyViewModelBase::CreateViewModelForData(UNiagaraHierarchyItemBase* ItemBase, TSharedPtr<FNiagaraHierarchyItemViewModelBase> Parent)
 {
 	if(UNiagaraHierarchyItem* Item = Cast<UNiagaraHierarchyItem>(ItemBase))
@@ -614,6 +629,9 @@ void FNiagaraHierarchyItemViewModelBase::RefreshChildrenData()
 		}
 	}
 
+	/** Every item view model can define its own sort order for its children. */
+	SortChildrenData();
+	
 	RefreshChildrenDataInternal();
 
 	/** All remaining children are supposed to exist at this point, as internal data won't be removed by refreshing & external data was cleaned up already.
@@ -667,12 +685,6 @@ void FNiagaraHierarchyItemViewModelBase::SyncViewModelsToData()
 	/** Give the view models a chance to further customize the children sync process. */
 	SyncViewModelsToDataInternal();	
 
-	// first we sort the data. Categories before items.
-	GetDataMutable()->GetChildrenMutable().StableSort([](const UNiagaraHierarchyItemBase& ItemA, const UNiagaraHierarchyItemBase& ItemB)
-		{
-			return ItemA.IsA<UNiagaraHierarchyCategory>() && ItemB.IsA<UNiagaraHierarchyItem>();
-		});
-
 	// then we sort the view models according to the data order as this is what will determine widget order created from the view models
 	Children.Sort([this](const TSharedPtr<FNiagaraHierarchyItemViewModelBase>& ItemA, const TSharedPtr<FNiagaraHierarchyItemViewModelBase>& ItemB)
 		{
@@ -714,6 +726,14 @@ const TArray<TSharedPtr<FNiagaraHierarchyItemViewModelBase>>& FNiagaraHierarchyI
 	return FilteredChildren;
 }
 
+void FNiagaraHierarchyItemViewModelBase::SortChildrenData() const
+{
+	GetDataMutable()->GetChildrenMutable().StableSort([](const UNiagaraHierarchyItemBase& ItemA, const UNiagaraHierarchyItemBase& ItemB)
+		{
+			return ItemA.IsA<UNiagaraHierarchyCategory>() && ItemB.IsA<UNiagaraHierarchyItem>();
+		});
+}
+
 int32 FNiagaraHierarchyItemViewModelBase::GetHierarchyDepth() const
 {
 	if(Parent.IsValid())
@@ -888,7 +908,8 @@ void FNiagaraHierarchyItemViewModelBase::Delete()
 
 void FNiagaraHierarchyItemViewModelBase::DeleteChild(TSharedPtr<FNiagaraHierarchyItemViewModelBase> Child)
 {	
-	ensure(Child->GetParent().Pin() == AsShared());	
+	ensure(Child->GetParent().Pin() == AsShared());
+	GetDataMutable()->Modify();
 	GetDataMutable()->GetChildrenMutable().Remove(Child->GetDataMutable());
 	Children.Remove(Child);
 }
@@ -1490,7 +1511,7 @@ void UNiagaraHierarchyViewModelBase::AddCategory(TSharedPtr<FNiagaraHierarchyIte
 		return;
 	}
 	
-	TSharedPtr<FNiagaraHierarchyItemViewModelBase> ViewModel = CategoryParent->AddNewItem(UNiagaraHierarchyCategory::StaticClass());
+	TSharedPtr<FNiagaraHierarchyItemViewModelBase> ViewModel = CategoryParent->AddNewItem(GetCategoryDataClass());
 
 	if(ensure(ViewModel.IsValid()))
 	{
diff --git a/Engine/Plugins/FX/Niagara/Source/NiagaraEditor/Private/ViewModels/HierarchyEditor/NiagaraSummaryViewViewModel.cpp b/Engine/Plugins/FX/Niagara/Source/NiagaraEditor/Private/ViewModels/HierarchyEditor/NiagaraSummaryViewViewModel.cpp
--- a/Engine/Plugins/FX/Niagara/Source/NiagaraEditor/Private/ViewModels/HierarchyEditor/NiagaraSummaryViewViewModel.cpp
+++ b/Engine/Plugins/FX/Niagara/Source/NiagaraEditor/Private/ViewModels/HierarchyEditor/NiagaraSummaryViewViewModel.cpp
@@ -1,6 +1,8 @@
 // Copyright Epic Games, Inc. All Rights Reserved.
 
 #include "ViewModels/HierarchyEditor/NiagaraSummaryViewViewModel.h"
+
+#include "EdGraphSchema_Niagara.h"
 #include "GraphEditAction.h"
 #include "NiagaraConstants.h"
 #include "NiagaraEditorStyle.h"
@@ -290,7 +292,7 @@ void FNiagaraFunctionViewModel::RefreshChildrenInputs(bool bClearCache) const
 						}
 					}
 					
-					return VariableGuidMetadataMap[GuidA].EditorSortPriority < VariableGuidMetadataMap[GuidB].EditorSortPriority;
+					return VariableGuidMetadataMap[GuidA].GetEditorSortPriority_DEPRECATED() < VariableGuidMetadataMap[GuidB].GetEditorSortPriority_DEPRECATED();
 				});
 				
 				for(const FGuid& VariableGuid : VariableGuids)
@@ -623,12 +625,12 @@ void FNiagaraModuleInputViewModel::AddNativeChildrenInputs()
 			{
 				if(UNiagaraScriptVariable* ScriptVariable = Graph->GetScriptVariable(Variable))
 				{
-					if(!ScriptVariable->Metadata.ParentAttribute.IsNone() && ScriptVariable->Metadata.ParentAttribute.IsEqual(GetInputData()->InputName))
+					if(!ScriptVariable->Metadata.GetParentAttribute_DEPRECATED().IsNone() && ScriptVariable->Metadata.GetParentAttribute_DEPRECATED().IsEqual(GetInputData()->InputName))
 					{
 						FNiagaraHierarchyIdentity ChildIdentity;
 						ChildIdentity.Guids.Add(GetData()->GetPersistentIdentity().Guids[0]);
 						ChildIdentity.Guids.Add(ScriptVariable->Metadata.GetVariableGuid());
-						ChildSortOrderMap.Add(ChildIdentity, ScriptVariable->Metadata.EditorSortPriority);
+						ChildSortOrderMap.Add(ChildIdentity, ScriptVariable->Metadata.GetEditorSortPriority_DEPRECATED());
 					}
 				}
 			}
@@ -694,7 +696,7 @@ TArray<FNiagaraHierarchyIdentity> FNiagaraModuleInputViewModel::GetNativeChildIn
 			{
 				if(UNiagaraScriptVariable* ScriptVariable = Graph->GetScriptVariable(Variable))
 				{
-					if(!ScriptVariable->Metadata.ParentAttribute.IsNone() && ScriptVariable->Metadata.ParentAttribute.IsEqual(GetInputData()->InputName))
+					if(!ScriptVariable->Metadata.GetParentAttribute_DEPRECATED().IsNone() && ScriptVariable->Metadata.GetParentAttribute_DEPRECATED().IsEqual(GetInputData()->InputName))
 					{
 						FNiagaraHierarchyIdentity ChildIdentity;
 						ChildIdentity.Guids.Add(GetData()->GetPersistentIdentity().Guids[0]);
