Niagara Script Editor - Hierarchy Editor Integration

diff --git a/Engine/Plugins/FX/Niagara/Source/NiagaraEditor/Private/ViewModels/NiagaraScratchPadScriptViewModel.cpp b/Engine/Plugins/FX/Niagara/Source/NiagaraEditor/Private/ViewModels/NiagaraScratchPadScriptViewModel.cpp
--- a/Engine/Plugins/FX/Niagara/Source/NiagaraEditor/Private/ViewModels/NiagaraScratchPadScriptViewModel.cpp
+++ b/Engine/Plugins/FX/Niagara/Source/NiagaraEditor/Private/ViewModels/NiagaraScratchPadScriptViewModel.cpp
@@ -77,6 +77,12 @@ void FNiagaraScratchPadScriptViewModel::Initialize(UNiagaraScript* Script, UNiag
 		FSimpleDelegate::CreateSP(GetVariableSelection(), &FNiagaraObjectSelection::Refresh)
 	);
 	ParameterPaneViewModel->Init(UIContext);
+
+	ScriptHierarchyViewModel = NewObject<UNiagaraHierarchyScriptParametersViewModel>();
+	ScriptHierarchyViewModel->Initialize(AsShared());
+
+	ScriptHierarchyViewModel->OnHierarchyChanged().AddSP(this, &FNiagaraScratchPadScriptViewModel::OnHierarchyChanged);
+	ScriptHierarchyViewModel->OnHierarchyPropertiesChanged().AddSP(this, &FNiagaraScratchPadScriptViewModel::OnHierarchyPropertiesChanged);
 }
 
 bool FNiagaraScratchPadScriptViewModel::IsValid() const
@@ -97,11 +103,18 @@ void FNiagaraScratchPadScriptViewModel::Finalize()
 		ParameterPaneViewModel->GetOnExternalSelectionChangedDelegate().Remove(ExternalSelectionChangedDelegate);
 		ParameterPaneViewModel.Reset();
 	}
+
+	if(ScriptHierarchyViewModel != nullptr)
+	{
+		ScriptHierarchyViewModel->Finalize();
+		ScriptHierarchyViewModel = nullptr;
+	}
 }
 
 void FNiagaraScratchPadScriptViewModel::AddReferencedObjects(FReferenceCollector& Collector)
 {
 	Collector.AddReferencedObject(EditScript.Script);
+	Collector.AddReferencedObject(ScriptHierarchyViewModel);
 }
 
 UNiagaraScript* FNiagaraScratchPadScriptViewModel::GetOriginalScript() const
@@ -124,6 +137,11 @@ TArray<UNiagaraGraph*> FNiagaraScratchPadScriptViewModel::GetEditableGraphs()
 	return EditableGraphs;
 }
 
+UNiagaraHierarchyScriptParametersViewModel* FNiagaraScratchPadScriptViewModel::GetHierarchyViewModel()
+{
+	return ScriptHierarchyViewModel;
+}
+
 TSharedPtr<INiagaraParameterPanelViewModel> FNiagaraScratchPadScriptViewModel::GetParameterPanelViewModel() const
 {
 	return ParameterPaneViewModel;
@@ -273,6 +291,24 @@ FSimpleDelegate& FNiagaraScratchPadScriptViewModel::OnRequestDiscardChanges()
 	return OnRequestDiscardChangesDelegate;
 }
 
+void FNiagaraScratchPadScriptViewModel::OnHierarchyChanged()
+{
+	if (bHasPendingChanges == false)
+	{
+		bHasPendingChanges = true;
+		OnHasUnappliedChangesChangedDelegate.Broadcast();
+	}
+}
+
+void FNiagaraScratchPadScriptViewModel::OnHierarchyPropertiesChanged()
+{
+	if (bHasPendingChanges == false)
+	{
+		bHasPendingChanges = true;
+		OnHasUnappliedChangesChangedDelegate.Broadcast();
+	}
+}
+
 FText FNiagaraScratchPadScriptViewModel::GetDisplayNameInternal() const
 {
 	return FText::FromString(OriginalScript->GetName());
diff --git a/Engine/Plugins/FX/Niagara/Source/NiagaraEditor/Private/ViewModels/NiagaraSystemEditorDocumentsViewModel.cpp b/Engine/Plugins/FX/Niagara/Source/NiagaraEditor/Private/ViewModels/NiagaraSystemEditorDocumentsViewModel.cpp
--- a/Engine/Plugins/FX/Niagara/Source/NiagaraEditor/Private/ViewModels/NiagaraSystemEditorDocumentsViewModel.cpp
+++ b/Engine/Plugins/FX/Niagara/Source/NiagaraEditor/Private/ViewModels/NiagaraSystemEditorDocumentsViewModel.cpp
@@ -308,6 +308,8 @@ TSharedRef<SNiagaraScratchPadScriptEditor> UNiagaraSystemEditorDocumentsViewMode
 	}
 
 	TSharedRef<SNiagaraScratchPadScriptEditor> Editor = SNew(SNiagaraScratchPadScriptEditor, GraphViewModel);
+	// We add a tag here to identify what the tab contains. The tab type doesn't work for that purpose as it can be invalid once a tab is activated
+	Editor->SetTag(FName("ScratchPadScriptEditor"));
 	return Editor;
 }
 
@@ -358,16 +360,19 @@ bool UNiagaraSystemEditorDocumentsViewModel::FindOpenTabsContainingDocument(cons
 
 void UNiagaraSystemEditorDocumentsViewModel::SetActiveDocumentTab(TSharedPtr<SDockTab> Tab)
 {
-	ActiveDocumentTab = Tab;
-	ActiveDocumentTabScriptViewModel = GetActiveScratchPadViewModelIfSet();
+	if(ActiveDocumentTab != Tab)
+	{
+		ActiveDocumentTab = Tab;
+		ActiveDocumentTabScriptViewModel = GetActiveScratchPadViewModelIfSet();
 
-	ActiveDocChangedDelegate.Broadcast(Tab);
+		ActiveDocChangedDelegate.Broadcast(Tab);
 
-	// We need to update the parameter panel view model with new parameters potentially
-	INiagaraParameterPanelViewModel* PanelVM = GetSystemViewModel()->GetParameterPanelViewModel();
-	if (PanelVM)
-	{
-		PanelVM->RefreshNextTick();
+		// We need to update the parameter panel view model with new parameters potentially
+		INiagaraParameterPanelViewModel* PanelVM = GetSystemViewModel()->GetParameterPanelViewModel();
+		if (PanelVM)
+		{
+			PanelVM->RefreshNextTick();
+		}
 	}
 }
 
@@ -416,7 +421,8 @@ TSharedPtr<FNiagaraScratchPadScriptViewModel> UNiagaraSystemEditorDocumentsViewM
 TSharedPtr<FNiagaraScratchPadScriptViewModel> UNiagaraSystemEditorDocumentsViewModel::GetActiveScratchPadViewModelIfSet()
 {
 	TSharedPtr<SDockTab> Tab = ActiveDocumentTab.Pin();
-	if (Tab.IsValid() && Tab->GetLayoutIdentifier().TabType == TEXT("Document"))
+	// When this is called during tab construction, we can't rely on the TabID to be initialized to check if this is a scratch pad tab
+	if (Tab.IsValid() && Tab->GetContent()->GetTag() == "ScratchPadScriptEditor")
 	{
 		TSharedRef<SNiagaraScratchPadScriptEditor> GraphEditor = StaticCastSharedRef<SNiagaraScratchPadScriptEditor>(Tab->GetContent());
 		if (GraphEditor->GetGraphEditor())
diff --git a/Engine/Plugins/FX/Niagara/Source/NiagaraEditor/Private/ViewModels/Stack/NiagaraStackEmitterSettingsGroup.cpp b/Engine/Plugins/FX/Niagara/Source/NiagaraEditor/Private/ViewModels/Stack/NiagaraStackEmitterSettingsGroup.cpp
--- a/Engine/Plugins/FX/Niagara/Source/NiagaraEditor/Private/ViewModels/Stack/NiagaraStackEmitterSettingsGroup.cpp
+++ b/Engine/Plugins/FX/Niagara/Source/NiagaraEditor/Private/ViewModels/Stack/NiagaraStackEmitterSettingsGroup.cpp
@@ -22,7 +22,7 @@
 #include "NiagaraNodeFunctionCall.h"
 #include "NiagaraSettings.h"
 #include "Toolkits/SystemToolkitModes/NiagaraSystemToolkitModeBase.h"
-#include "Widgets/SNiagaraHierarchy.h"
+#include "Widgets/SNiagaraHierarchyEditor.h"
 
 #include UE_INLINE_GENERATED_CPP_BY_NAME(NiagaraStackEmitterSettingsGroup)
 
diff --git a/Engine/Plugins/FX/Niagara/Source/NiagaraEditor/Private/ViewModels/Stack/NiagaraStackEntry.cpp b/Engine/Plugins/FX/Niagara/Source/NiagaraEditor/Private/ViewModels/Stack/NiagaraStackEntry.cpp
--- a/Engine/Plugins/FX/Niagara/Source/NiagaraEditor/Private/ViewModels/Stack/NiagaraStackEntry.cpp
+++ b/Engine/Plugins/FX/Niagara/Source/NiagaraEditor/Private/ViewModels/Stack/NiagaraStackEntry.cpp
@@ -278,6 +278,11 @@ bool UNiagaraStackEntry::IsExpandedByDefault() const
 	return true;
 }
 
+bool UNiagaraStackEntry::KeepExpanderIndentation() const
+{
+	return false;
+}
+
 bool UNiagaraStackEntry::GetIsExpanded() const
 {
 	if (GetShouldShowInStack() == false || GetCanExpand() == false)
diff --git a/Engine/Plugins/FX/Niagara/Source/NiagaraEditor/Private/ViewModels/Stack/NiagaraStackFunctionInput.cpp b/Engine/Plugins/FX/Niagara/Source/NiagaraEditor/Private/ViewModels/Stack/NiagaraStackFunctionInput.cpp
--- a/Engine/Plugins/FX/Niagara/Source/NiagaraEditor/Private/ViewModels/Stack/NiagaraStackFunctionInput.cpp
+++ b/Engine/Plugins/FX/Niagara/Source/NiagaraEditor/Private/ViewModels/Stack/NiagaraStackFunctionInput.cpp
@@ -46,7 +46,7 @@
 #include "ViewModels/NiagaraSystemViewModel.h"
 #include "ViewModels/Stack/NiagaraStackViewModel.h"
 #include "ViewModels/Stack/NiagaraParameterHandle.h"
-#include "ViewModels/Stack/NiagaraStackFunctionInputCollection.h"
+#include "ViewModels/Stack/NiagaraStackValueCollection.h"
 #include "ViewModels/Stack/NiagaraStackGraphUtilities.h"
 #include "ViewModels/Stack/NiagaraStackInputCategory.h"
 #include "ViewModels/Stack/NiagaraStackObject.h"
@@ -131,6 +131,13 @@ void UNiagaraStackFunctionInput::Initialize(
 	OwningFunctionCallNode = &InInputFunctionCallNode;
 	OwningFunctionCallInitialScript = OwningFunctionCallNode->FunctionScript;
 	OwningAssignmentNode = Cast<UNiagaraNodeAssignment>(OwningFunctionCallNode.Get());
+	
+	TOptional<FGuid> VariableGuid = OwningFunctionCallNode->GetFunctionScriptSource()->NodeGraph->GetScriptVariableGuid(FNiagaraVariable(InInputType, InInputParameterHandle));
+	if(VariableGuid.IsSet())
+	{
+		UNiagaraHierarchyRoot* HierarchyScriptParameterRoot = OwningFunctionCallNode->GetFunctionScriptSource()->NodeGraph->GetScriptParameterHierarchyRoot();
+		HierarchyScriptParameter = Cast<UNiagaraHierarchyScriptParameter>(HierarchyScriptParameterRoot->FindChildWithIdentity(FNiagaraHierarchyIdentity({VariableGuid.GetValue()}, {}), true));
+	}
 
 	UNiagaraSystem& ParentSystem = GetSystemViewModel()->GetSystem();
 	FVersionedNiagaraEmitter ParentEmitter = GetEmitterViewModel().IsValid() ? GetEmitterViewModel()->GetEmitter() : FVersionedNiagaraEmitter();
@@ -183,6 +190,8 @@ void UNiagaraStackFunctionInput::Initialize(
 	MessageLogGuid = GetSystemViewModel()->GetMessageLogGuid();
 
 	AddChildFilter(FOnFilterChild::CreateUObject(this, &UNiagaraStackFunctionInput::FilterInlineChildren));
+	AddChildFilter(FOnFilterChild::CreateUObject(this, &UNiagaraStackFunctionInput::FilterForVisibleCondition));
+	AddChildFilter(FOnFilterChild::CreateUObject(this, &UNiagaraStackFunctionInput::FilterForIsInlineEditConditionToggle));
 }
 
 void UNiagaraStackFunctionInput::FinalizeInternal()
@@ -399,6 +408,45 @@ FText UNiagaraStackFunctionInput::GetPasteTransactionText(const UNiagaraClipboar
 	return LOCTEXT("PasteInputTransactionText", "Paste Niagara inputs");
 }
 
+void UNiagaraStackFunctionInput::PasteFunctionInput(const UNiagaraClipboardFunctionInput* ClipboardInput)
+{
+	// For the main input (could be a parent input), the user is intentional about copy pasting, so we only require assignable types
+	if (FNiagaraEditorUtilities::AreTypesAssignable(ClipboardInput->InputType, InputType))
+	{
+		SetValueFromClipboardFunctionInput(*ClipboardInput);
+	}
+	// or alternatively check for a conversion script
+	else
+	{
+		SetClipboardContentViaConversionScript(*ClipboardInput);
+	}
+
+	OnCopyPasteDelegate.ExecuteIfBound();
+	
+	// For children inputs we are more strict and require name & type match as the user is less intentional about it when pasting
+	if(ClipboardInput->ChildrenInputs.Num() > 0)
+	{
+		TArray<UNiagaraStackFunctionInput*> StackChildrenInputs;
+		GetUnfilteredChildrenOfType(StackChildrenInputs, false);
+		for(const UNiagaraClipboardFunctionInput* ClipboardChildInput : ClipboardInput->ChildrenInputs)
+		{
+			if(ClipboardChildInput == nullptr)
+			{
+				continue;
+			}
+			
+			for(UNiagaraStackFunctionInput* StackChildInput : StackChildrenInputs)
+			{
+				if(StackChildInput->GetInputType() == ClipboardChildInput->GetTypeDef() &&
+					StackChildInput->GetInputParameterHandle().GetName() == ClipboardChildInput->InputName)
+				{
+					StackChildInput->PasteFunctionInput(ClipboardChildInput);
+				}
+			}
+		}
+	}
+}
+
 void UNiagaraStackFunctionInput::Paste(const UNiagaraClipboardContent* ClipboardContent, FText& OutPasteWarning)
 {
 	if (ensureMsgf(ClipboardContent != nullptr && (ClipboardContent->FunctionInputs.Num() == 1 || ClipboardContent->PortableValues.Num() == 1),
@@ -409,14 +457,7 @@ void UNiagaraStackFunctionInput::Paste(const UNiagaraClipboardContent* Clipboard
 		{
 			if (const UNiagaraClipboardFunctionInput* ClipboardInput = ClipboardContent->FunctionInputs[0])
 			{
-				if (FNiagaraEditorUtilities::AreTypesAssignable(ClipboardInput->InputType, InputType))
-				{
-					SetValueFromClipboardFunctionInput(*ClipboardInput);
-				}
-				else
-				{
-					SetClipboardContentViaConversionScript(*ClipboardInput);
-				}
+				PasteFunctionInput(ClipboardInput);
 			}
 		}
 		else if (ClipboardContent->PortableValues.Num() == 1)
@@ -479,16 +520,16 @@ FNiagaraHierarchyIdentity UNiagaraStackFunctionInput::DetermineSummaryIdentity()
 	return Identity;
 }
 
+bool UNiagaraStackFunctionInput::GetCanExpand() const
+{
+	// If we are in local mode, we want to always be expanded (by setting GetCanExpand == false), to ensure potential children inputs always show up 
+	return GetValueMode() != EValueMode::Local;
+}
+
 TArray<UNiagaraStackFunctionInput*> UNiagaraStackFunctionInput::GetChildInputs() const
 {
-	TArray<UNiagaraStackFunctionInputCollection*> DynamicInputCollections;
-	GetUnfilteredChildrenOfType(DynamicInputCollections);
 	TArray<UNiagaraStackFunctionInput*> ChildInputs;
-	for (UNiagaraStackFunctionInputCollection* DynamicInputCollection : DynamicInputCollections)
-	{
-		DynamicInputCollection->GetChildInputs(ChildInputs);
-	}
-	
+	GetUnfilteredChildrenOfType(ChildInputs, true);	
 	return ChildInputs;
 }
 
@@ -499,12 +540,7 @@ TOptional<FNiagaraVariableMetaData> UNiagaraStackFunctionInput::GetInputMetaData
 
 void UNiagaraStackFunctionInput::GetFilteredChildInputs(TArray<UNiagaraStackFunctionInput*>& OutFilteredChildInputs) const
 {
-	TArray<UNiagaraStackFunctionInputCollection*> DynamicInputCollections;
-	GetUnfilteredChildrenOfType(DynamicInputCollections);
-	for (UNiagaraStackFunctionInputCollection* DynamicInputCollection : DynamicInputCollections)
-	{
-		DynamicInputCollection->GetFilteredChildInputs(OutFilteredChildInputs);
-	}
+	GetFilteredChildrenOfType(OutFilteredChildInputs, true);
 }
 
 UNiagaraStackObject* UNiagaraStackFunctionInput::GetChildDataObject() const
@@ -532,6 +568,18 @@ bool UNiagaraStackFunctionInput::FilterInlineChildren(const UNiagaraStackEntry&
 	return GbEnableExperimentalInlineDynamicInputs == 0 || GetInlineDisplayMode() == ENiagaraStackEntryInlineDisplayMode::None;
 }
 
+bool UNiagaraStackFunctionInput::FilterForVisibleCondition(const UNiagaraStackEntry& Child) const
+{
+	const UNiagaraStackFunctionInput* StackFunctionInputChild = Cast<UNiagaraStackFunctionInput>(&Child);
+	return StackFunctionInputChild == nullptr || StackFunctionInputChild->GetShouldPassFilterForVisibleCondition();
+}
+
+bool UNiagaraStackFunctionInput::FilterForIsInlineEditConditionToggle(const UNiagaraStackEntry& Child) const
+{
+	const UNiagaraStackFunctionInput* StackFunctionInputChild = Cast<UNiagaraStackFunctionInput>(&Child);
+	return StackFunctionInputChild == nullptr || StackFunctionInputChild->GetIsInlineEditConditionToggle() == false;
+}
+
 void UNiagaraStackFunctionInput::ReportScriptVersionChange() const
 {
 	// analytics
@@ -826,7 +874,44 @@ void UNiagaraStackFunctionInput::RefreshChildrenInternal(const TArray<UNiagaraSt
             }
         ));
 	}
+	
+	if(HierarchyScriptParameter.IsValid())
+	{
+		TArray<UNiagaraHierarchyScriptParameter*> ChildScriptParameters;
+		HierarchyScriptParameter->GetChildrenOfType(ChildScriptParameters);
+		
+		for(UNiagaraHierarchyScriptParameter* ChildScriptParameter : ChildScriptParameters)
+		{
+			FNiagaraVariable InputVariable = ChildScriptParameter->GetVariable();
 
+			if(ScriptInstanceData.UsedInputs.Contains(InputVariable) == false)
+			{
+				continue;
+			}
+			
+			UNiagaraStackFunctionInput* InputChild = FindCurrentChildOfTypeByPredicate<UNiagaraStackFunctionInput>(CurrentChildren, [&](UNiagaraStackFunctionInput* CurrentInput) 
+			{ 
+				return CurrentInput->GetInputParameterHandle() == FNiagaraParameterHandle(InputVariable.GetName()) && CurrentInput->GetInputType() == ChildScriptParameter->GetVariable().GetType() && &CurrentInput->GetInputFunctionCallNode() == OwningFunctionCallNode;
+			});
+
+			if (InputChild == nullptr)
+			{
+				EStackParameterBehavior Behavior = ChildScriptParameter->GetScriptVariable()->GetIsStaticSwitch() ? EStackParameterBehavior::Static : EStackParameterBehavior::Dynamic;
+				InputChild = NewObject<UNiagaraStackFunctionInput>(this);
+				InputChild->Initialize(CreateDefaultChildRequiredData(), *OwningModuleNode, *OwningFunctionCallNode,
+					InputVariable.GetName(), InputVariable.GetType(), Behavior, GetOwnerStackItemEditorDataKey());
+			}
+			
+			InputChild->SetScriptInstanceData(ScriptInstanceData);
+			FGuid VariableGuid = ChildScriptParameter->GetScriptVariable()->Metadata.GetVariableGuid();
+			// Would be nice if eventually inputs could easily access the cache so that they don't rely on the parent pushing data in to determine visibility
+			// If we let inputs use GetStackFunctionInputs & GetStackFunctionStaticSwitchPins there is a noticeable perf decrease, due to the static switch as they are less cached
+			InputChild->SetIsHidden(ScriptInstanceData.PerInputInstanceData[VariableGuid].bIsHidden);
+				
+			NewChildren.Add(InputChild);	
+		}
+	}
+	
 	if (InputValues.Mode == EValueMode::Expression && SupportsCustomExpressions() == false)
 	{
 		NewIssues.Add(FStackIssue(
@@ -842,16 +927,16 @@ void UNiagaraStackFunctionInput::RefreshChildrenInternal(const TArray<UNiagaraSt
 		FVersionedNiagaraScriptData* ScriptData = InputValues.DynamicNode->GetScriptData();
 		if (ScriptData != nullptr)
 		{
-			UNiagaraStackFunctionInputCollection* DynamicInputEntry = FindCurrentChildOfTypeByPredicate<UNiagaraStackFunctionInputCollection>(CurrentChildren,
-				[this](UNiagaraStackFunctionInputCollection* CurrentFunctionInputEntry)
+			UNiagaraStackScriptHierarchyRoot* DynamicInputEntry = FindCurrentChildOfTypeByPredicate<UNiagaraStackScriptHierarchyRoot>(CurrentChildren,
+				[this](UNiagaraStackScriptHierarchyRoot* CurrentFunctionInputEntry)
 			{
-				return CurrentFunctionInputEntry->GetInputFunctionCallNode() == InputValues.DynamicNode.Get() &&
-					CurrentFunctionInputEntry->GetModuleNode() == OwningModuleNode.Get();
+				return &CurrentFunctionInputEntry->GetOwningFunctionCallNode() == InputValues.DynamicNode.Get() &&
+					&CurrentFunctionInputEntry->GetOwningModuleNode() == OwningModuleNode.Get();
 			});
 
 			if (DynamicInputEntry == nullptr)
 			{
-				DynamicInputEntry = NewObject<UNiagaraStackFunctionInputCollection>(this);
+				DynamicInputEntry = NewObject<UNiagaraStackScriptHierarchyRoot>(this);
 				DynamicInputEntry->Initialize(CreateDefaultChildRequiredData(), *OwningModuleNode, *InputValues.DynamicNode.Get(), GetOwnerStackItemEditorDataKey());
 				DynamicInputEntry->SetShouldDisplayLabel(false);
 			}
@@ -2724,13 +2809,13 @@ void UNiagaraStackFunctionInput::ReassignDynamicInputScript(UNiagaraScript* Dyna
 
 			UNiagaraClipboardContent* NewClipboardContent = UNiagaraClipboardContent::Create();
 			Copy(NewClipboardContent);
-			TArray<UNiagaraStackFunctionInputCollection*> DynamicInputCollections;
-			GetUnfilteredChildrenOfType(DynamicInputCollections);
+			TArray<UNiagaraStackScriptHierarchyRoot*> DynamicInputRoots;
+			GetUnfilteredChildrenOfType(DynamicInputRoots);
 
-			if (ConversionUtility && DynamicInputCollections.Num() == 0)
+			if (ConversionUtility && DynamicInputRoots.Num() == 0)
 			{
 				FText ConvertMessage;
-				bool bConverted = ConversionUtility->Convert(OldScript, OldClipboardContent, DynamicInputScript, DynamicInputCollections[0], NewClipboardContent, InputValues.DynamicNode.Get(), ConvertMessage);
+				bool bConverted = ConversionUtility->Convert(OldScript, OldClipboardContent, DynamicInputScript, DynamicInputRoots[0], NewClipboardContent, InputValues.DynamicNode.Get(), ConvertMessage);
 				if (!ConvertMessage.IsEmptyOrWhitespace())
 				{
 					// Notify the end-user about the convert message, but continue the process as they could always undo.
@@ -2961,11 +3046,11 @@ const UNiagaraClipboardFunctionInput* UNiagaraStackFunctionInput::ToClipboardFun
 	{
 		ClipboardInput = UNiagaraClipboardFunctionInput::CreateDynamicValue(InOuter, InputName, InputType, bEditConditionValue, InputValues.DynamicNode->GetFunctionName(), InputValues.DynamicNode->FunctionScript, InputValues.DynamicNode->SelectedScriptVersion);
 
-		TArray<UNiagaraStackFunctionInputCollection*> DynamicInputCollections;
-		GetUnfilteredChildrenOfType(DynamicInputCollections);
-		for (UNiagaraStackFunctionInputCollection* DynamicInputCollection : DynamicInputCollections)
+		TArray<UNiagaraStackScriptHierarchyRoot*> DynamicInputRoots;
+		GetUnfilteredChildrenOfType(DynamicInputRoots);
+		for (UNiagaraStackScriptHierarchyRoot* DynamicInputRoot : DynamicInputRoots)
 		{
-			DynamicInputCollection->ToClipboardFunctionInputs(ClipboardInput->Dynamic, MutableView(ClipboardInput->Dynamic->Inputs));
+			DynamicInputRoot->ToClipboardFunctionInputs(ClipboardInput->Dynamic, MutableView(ClipboardInput->Dynamic->Inputs));
 		}
 
 		break;
@@ -2980,6 +3065,20 @@ const UNiagaraClipboardFunctionInput* UNiagaraStackFunctionInput::ToClipboardFun
 		ensureMsgf(false, TEXT("A new value mode was added without adding support for copy paste."));
 		break;
 	}
+
+	OnCopyPasteDelegate.ExecuteIfBound();
+
+	if(ClipboardInput)
+	{
+		TArray<UNiagaraStackFunctionInput*> ChildInputs;
+		GetFilteredChildrenOfType(ChildInputs, false);
+		
+		for(UNiagaraStackFunctionInput* ChildInput : ChildInputs)
+		{
+			const_cast<UNiagaraClipboardFunctionInput*>(ClipboardInput)->ChildrenInputs.Add(ChildInput->ToClipboardFunctionInput(const_cast<UNiagaraClipboardFunctionInput*>(ClipboardInput)));
+		}
+	}
+
 	return ClipboardInput;
 }
 
@@ -3055,11 +3154,11 @@ void UNiagaraStackFunctionInput::SetValueFromClipboardFunctionInput(const UNiaga
 					}
 					SetDynamicInput(NewDynamicInputScript, ClipboardFunctionInput.Dynamic->FunctionName, ClipboardFunctionInput.Dynamic->ScriptVersion);
 
-					TArray<UNiagaraStackFunctionInputCollection*> DynamicInputCollections;
-					GetUnfilteredChildrenOfType(DynamicInputCollections);
-					for (UNiagaraStackFunctionInputCollection* DynamicInputCollection : DynamicInputCollections)
+					TArray<UNiagaraStackScriptHierarchyRoot*> DynamicInputRoots;
+					GetUnfilteredChildrenOfType(DynamicInputRoots);
+					for (UNiagaraStackScriptHierarchyRoot* DynamicInputRoot : DynamicInputRoots)
 					{
-						DynamicInputCollection->SetValuesFromClipboardFunctionInputs(ClipboardFunctionInput.Dynamic->Inputs);
+						DynamicInputRoot->SetValuesFromClipboardFunctionInputs(ClipboardFunctionInput.Dynamic->Inputs);
 					}
 				}
 			}
@@ -3078,6 +3177,12 @@ void UNiagaraStackFunctionInput::SetValueFromClipboardFunctionInput(const UNiaga
 	TSharedRef<FNiagaraSystemViewModel> CachedSysViewModel = GetSystemViewModel();
 	if (CachedSysViewModel->GetSystemStackViewModel())
 		CachedSysViewModel->GetSystemStackViewModel()->InvalidateCachedParameterUsage();
+
+	// If we pasted into a static parameter, make sure to refresh children as this can cause new inputs to pop up as we might be in the middle of pasting multiple inputs
+	if(IsStaticParameter())
+	{
+		RefreshChildren();
+	}
 }
 
 bool UNiagaraStackFunctionInput::IsScratchDynamicInput() const
diff --git a/Engine/Plugins/FX/Niagara/Source/NiagaraEditor/Private/ViewModels/Stack/NiagaraStackGraphUtilities.cpp b/Engine/Plugins/FX/Niagara/Source/NiagaraEditor/Private/ViewModels/Stack/NiagaraStackGraphUtilities.cpp
--- a/Engine/Plugins/FX/Niagara/Source/NiagaraEditor/Private/ViewModels/Stack/NiagaraStackGraphUtilities.cpp
+++ b/Engine/Plugins/FX/Niagara/Source/NiagaraEditor/Private/ViewModels/Stack/NiagaraStackGraphUtilities.cpp
@@ -19,7 +19,7 @@
 #include "ViewModels/NiagaraEmitterHandleViewModel.h"
 #include "ViewModels/NiagaraScriptViewModel.h"
 #include "ViewModels/Stack/NiagaraStackEntry.h"
-#include "ViewModels/Stack/NiagaraStackFunctionInputCollection.h"
+#include "ViewModels/Stack/NiagaraStackValueCollection.h"
 #include "ViewModels/Stack/NiagaraStackInputCategory.h"
 #include "ViewModels/Stack/NiagaraStackFunctionInput.h"
 #include "ViewModels/Stack/NiagaraStackErrorItem.h"
@@ -911,33 +911,23 @@ DECLARE_DELEGATE_RetVal_OneParam(bool, FInputSelector, UNiagaraStackFunctionInpu
 
 void InitializeStackFunctionInputsInternal(TSharedRef<FNiagaraSystemViewModel> SystemViewModel, TSharedPtr<FNiagaraEmitterViewModel> EmitterViewModel, UNiagaraStackEditorData& StackEditorData, UNiagaraNodeFunctionCall& ModuleNode, UNiagaraNodeFunctionCall& InputFunctionCallNode, FInputSelector InputSelector)
 {
-	UNiagaraStackFunctionInputCollection* FunctionInputCollection = NewObject<UNiagaraStackFunctionInputCollection>(GetTransientPackage()); 
+	UNiagaraStackScriptHierarchyRoot* FunctionHierarchyRoot = NewObject<UNiagaraStackScriptHierarchyRoot>(GetTransientPackage()); 
 	UNiagaraStackEntry::FRequiredEntryData RequiredEntryData(SystemViewModel, EmitterViewModel, NAME_None, NAME_None, StackEditorData);
-	FunctionInputCollection->Initialize(RequiredEntryData, ModuleNode, InputFunctionCallNode, FString());
-	FunctionInputCollection->RefreshChildren();
+	FunctionHierarchyRoot->Initialize(RequiredEntryData, ModuleNode, InputFunctionCallNode, FString());
+	FunctionHierarchyRoot->RefreshChildren();
 
 	// Reset all direct inputs on this function to initialize data interfaces and default dynamic inputs.
-	TArray<UNiagaraStackEntry*> Children;
-	FunctionInputCollection->GetUnfilteredChildren(Children);
-	for (UNiagaraStackEntry* Child : Children)
+	TArray<UNiagaraStackFunctionInput*> StackFunctionInputs;
+	FunctionHierarchyRoot->GetUnfilteredChildrenOfType(StackFunctionInputs);
+	for (UNiagaraStackFunctionInput* StackFunctionInput : StackFunctionInputs)
 	{
-		UNiagaraStackInputCategory* InputCategory = Cast<UNiagaraStackInputCategory>(Child);
-		if (InputCategory != nullptr)
+		if (StackFunctionInput != nullptr && (InputSelector.IsBound() == false || InputSelector.Execute(StackFunctionInput)) && StackFunctionInput->CanReset())
 		{
-			TArray<UNiagaraStackEntry*> CategoryChildren;
-			InputCategory->GetUnfilteredChildren(CategoryChildren);
-			for (UNiagaraStackEntry* CategoryChild : CategoryChildren)
-			{
-				UNiagaraStackFunctionInput* FunctionInput = Cast<UNiagaraStackFunctionInput>(CategoryChild);
-				if (FunctionInput != nullptr && (InputSelector.IsBound() == false || InputSelector.Execute(FunctionInput)) && FunctionInput->CanReset())
-				{
-					FunctionInput->Reset();
-				}
-			}
+			StackFunctionInput->Reset();
 		}
 	}
 
-	FunctionInputCollection->Finalize();
+	FunctionHierarchyRoot->Finalize();
 	SystemViewModel->NotifyDataObjectChanged(TArray<UObject*>(), ENiagaraDataObjectChange::Unknown);
 }
 
@@ -1165,11 +1155,12 @@ void FNiagaraStackGraphUtilities::GatherInputRelationsForStack(FInputDataCollect
 				TOptional<FNiagaraVariableMetaData> VariableMetaData = AssetGraph->GetMetaData(Variable);
 				if(VariableMetaData.IsSet())
 				{
-					if(VariableMetaData.GetValue().ParentAttribute == ThisInputHandle.GetName())
+					// @TODO Handle Deprecated values
+					if(VariableMetaData.GetValue().GetParentAttribute_DEPRECATED() == ThisInputHandle.GetName())
 					{
 						HierarchyInputToChildrenGuidMap.FindOrAdd(Input).Add(VariableMetaData.GetValue().GetVariableGuid());
 						ChildrenGuidToScriptVariablesMap.Add(VariableMetaData->GetVariableGuid(), AssetGraph->GetScriptVariable(Variable));
-						ChildrenGuidSortOrderMap.Add(VariableMetaData->GetVariableGuid(), VariableMetaData->EditorSortPriority);
+						ChildrenGuidSortOrderMap.Add(VariableMetaData->GetVariableGuid(), VariableMetaData->GetEditorSortPriority_DEPRECATED());
 					}
 				}
 			}
@@ -1355,64 +1346,6 @@ UNiagaraNodeAssignment* FNiagaraStackGraphUtilities::FindAssignmentNode(FGuid As
 	return nullptr;
 }
 
-TOptional<FNiagaraStackGraphUtilities::FMatchingFunctionInputData> FNiagaraStackGraphUtilities::FindInputData(FNiagaraHierarchyIdentity ModuleInputIdentity, TSharedRef<FNiagaraEmitterViewModel> EmitterViewModel)
-{
-	FMatchingFunctionInputData Result;
-	
-	UNiagaraNodeFunctionCall* MatchingFunctionCall = FindFunctionCallNode(ModuleInputIdentity.Guids[0], EmitterViewModel);
-	if(MatchingFunctionCall)
-	{
-		return FindInputData(*MatchingFunctionCall, ModuleInputIdentity, EmitterViewModel);
-	}					
-	
-	return TOptional<FMatchingFunctionInputData>();
-}
-
-TOptional<FNiagaraStackGraphUtilities::FMatchingFunctionInputData> FNiagaraStackGraphUtilities::FindInputData(const UNiagaraNodeFunctionCall& FunctionCallNode, FNiagaraHierarchyIdentity InputIdentity, TSharedRef<FNiagaraEmitterViewModel> EmitterViewModel)
-{
-	SCOPE_CYCLE_COUNTER(STAT_NiagaraEditor_StackGraphUtilities_FindInputData)
-	
-	FMatchingFunctionInputData Result;
-	Result.FunctionCallNode = const_cast<UNiagaraNodeFunctionCall*>(&FunctionCallNode);
-	
-	if(const UNiagaraNodeAssignment* AssignmentNode = Cast<UNiagaraNodeAssignment>(&FunctionCallNode))
-	{
-		TOptional<FMatchingFunctionInputData> InputData = FindAssignmentInputData(*AssignmentNode, InputIdentity.Names[0], EmitterViewModel);
-
-		if(InputData.IsSet())
-		{
-			Result.InputName = InputData->InputName;
-			Result.Type = InputData->Type;
-			Result.bIsHidden = InputData->bIsHidden;
-			Result.bIsStatic = InputData->bIsStatic;
-			Result.MetaData = InputData->MetaData;
-			Result.ChildrenInputGuids = InputData->ChildrenInputGuids;
-			Result.HiddenChildrenInputGuids = InputData->HiddenChildrenInputGuids;
-
-			return Result;
-		}
-	}
-	else
-	{
-		TOptional<FMatchingFunctionInputData> InputData = FindModuleInputData(FunctionCallNode, InputIdentity.Guids[1], EmitterViewModel);
-
-		if(InputData.IsSet())
-		{
-			Result.InputName = InputData->InputName;
-			Result.Type = InputData->Type;
-			Result.bIsHidden = InputData->bIsHidden;
-			Result.bIsStatic = InputData->bIsStatic;
-			Result.MetaData = InputData->MetaData;
-			Result.ChildrenInputGuids = InputData->ChildrenInputGuids;
-			Result.HiddenChildrenInputGuids = InputData->HiddenChildrenInputGuids;
-
-			return Result;
-		}
-	}
-
-	return TOptional<FMatchingFunctionInputData>();
-}
-
 TOptional<FNiagaraStackGraphUtilities::FMatchingFunctionInputData> FNiagaraStackGraphUtilities::FindAssignmentInputData(const UNiagaraNodeAssignment& AssignmentNode, FName VariableName, TSharedRef<FNiagaraEmitterViewModel> EmitterViewModel)
 {
 	SCOPE_CYCLE_COUNTER(STAT_NiagaraEditor_StackGraphUtilities_FindAssignmentInputData);
@@ -1471,208 +1404,6 @@ TOptional<FNiagaraStackGraphUtilities::FMatchingFunctionInputData> FNiagaraStack
 	return TOptional<FNiagaraStackGraphUtilities::FMatchingFunctionInputData>();
 }
 
-TOptional<FNiagaraStackGraphUtilities::FMatchingFunctionInputData> FNiagaraStackGraphUtilities::FindModuleInputData(const UNiagaraNodeFunctionCall& FunctionCallNode, FGuid VariableGuid, TSharedRef<FNiagaraEmitterViewModel> EmitterViewModel, bool bIncludeChildrenInputs, TMap<FInputDataCacheKey, FMatchingFunctionInputData>* OptionalCache)
-{
-	SCOPE_CYCLE_COUNTER(STAT_NiagaraEditor_StackGraphUtilities_FindModuleInputData)
-
-	if(OptionalCache != nullptr)
-	{
-		FInputDataCacheKey CacheKey(FunctionCallNode.NodeGuid, VariableGuid);
-		if(OptionalCache->Contains(CacheKey))
-		{
-			return (*OptionalCache)[CacheKey];
-		}
-	}
-	
-	FCompileConstantResolver ConstantResolver(EmitterViewModel->GetEmitter(), GetOutputNodeUsage(FunctionCallNode));
-
-	// our goal is to find the correct variable's metadata from the graph
-	TOptional<FNiagaraVariableMetaData> MetaData;
-				
-	// first, we check the normal module inputs
-	TArray<FNiagaraVariable> InputVariables;
-	TSet<FNiagaraVariable> HiddenVariables;
-
-	FNiagaraStackGraphUtilities::ENiagaraGetStackFunctionInputPinsOptions Options = FNiagaraStackGraphUtilities::ENiagaraGetStackFunctionInputPinsOptions::ModuleInputsOnly;
-	FNiagaraStackGraphUtilities::GetStackFunctionInputs(FunctionCallNode, InputVariables, HiddenVariables, ConstantResolver, Options, false);
-	
-	for(const FNiagaraVariable& InputVariable : InputVariables)
-	{
-		if (InputVariable.GetType().IsValid() == false)
-		{
-			continue;
-		}
-		
-		MetaData = FunctionCallNode.GetCalledGraph()->GetMetaData(InputVariable);
-
-		if(MetaData.IsSet())
-		{
-			FMatchingFunctionInputData PotentialMatchingInputData;
-			PotentialMatchingInputData.FunctionCallNode = const_cast<UNiagaraNodeFunctionCall*>(&FunctionCallNode);
-			PotentialMatchingInputData.InputName = InputVariable.GetName();
-			PotentialMatchingInputData.Type = InputVariable.GetType();
-			PotentialMatchingInputData.MetaData = MetaData.GetValue();
-			PotentialMatchingInputData.bIsStatic = false;
-			PotentialMatchingInputData.bIsHidden = HiddenVariables.Contains(InputVariable);
-
-			if(bIncludeChildrenInputs)
-			{
-				PotentialMatchingInputData.ChildrenInputGuids = GetChildrenInputGuids(FunctionCallNode, PotentialMatchingInputData.InputName);
-
-				// we only care about hidden children if the parent input isn't hidden
-				if(!PotentialMatchingInputData.bIsHidden)
-				{
-					PotentialMatchingInputData.HiddenChildrenInputGuids = GetHiddenChildrenInputGuids(FunctionCallNode, PotentialMatchingInputData.InputName, EmitterViewModel).Array();
-				}
-			}
-
-			if(OptionalCache != nullptr)
-			{
-				OptionalCache->Add(FInputDataCacheKey(FunctionCallNode.NodeGuid, MetaData->GetVariableGuid()), PotentialMatchingInputData);
-			}
-
-			// if we specified an optional cache, we want to keep on creating the cache data even after we've found the result to avoid redundant calls
-			if(OptionalCache == nullptr)
-			{
-				if(MetaData.GetValue().GetVariableGuid() == VariableGuid)
-				{
-					return PotentialMatchingInputData;
-				}
-			}			
-		}
-	}
-
-	// then if we didn't find the input under module inputs, we check the static switch inputs
-	TSet<UEdGraphPin*> HiddenSwitchPins;
-	TArray<UEdGraphPin*> SwitchPins;
-				
-	FNiagaraStackGraphUtilities::GetStackFunctionStaticSwitchPins(FunctionCallNode, SwitchPins, HiddenSwitchPins, ConstantResolver);
-	for (UEdGraphPin* InputPin : SwitchPins)
-	{		
-		FNiagaraVariable InputVariable = UEdGraphSchema_Niagara::PinToNiagaraVariable(InputPin);
-		if (InputVariable.GetType().IsValid() == false)
-		{
-			continue;
-		}
-					
-		if (FunctionCallNode.GetCalledGraph() != nullptr)
-		{
-			MetaData = FunctionCallNode.GetCalledGraph()->GetMetaData(InputVariable);
-
-			if(MetaData.IsSet())
-			{
-				FMatchingFunctionInputData PotentialMatchingInputData;
-				PotentialMatchingInputData.FunctionCallNode = const_cast<UNiagaraNodeFunctionCall*>(&FunctionCallNode);
-				PotentialMatchingInputData.InputName = InputPin->PinName;
-				PotentialMatchingInputData.Type = InputVariable.GetType();
-				PotentialMatchingInputData.MetaData = MetaData.GetValue();
-				PotentialMatchingInputData.bIsStatic = true;
-				PotentialMatchingInputData.bIsHidden = HiddenSwitchPins.Contains(InputPin);
-
-				if(bIncludeChildrenInputs)
-				{
-					PotentialMatchingInputData.ChildrenInputGuids = GetChildrenInputGuids(FunctionCallNode, PotentialMatchingInputData.InputName);
-
-					// we only care about hidden children if the parent input isn't hidden
-					if(!PotentialMatchingInputData.bIsHidden)
-					{
-						PotentialMatchingInputData.HiddenChildrenInputGuids = GetHiddenChildrenInputGuids(FunctionCallNode, PotentialMatchingInputData.InputName, EmitterViewModel).Array();
-					}
-				}
-				
-				if(OptionalCache != nullptr)
-				{
-					OptionalCache->Add(FInputDataCacheKey(FunctionCallNode.NodeGuid, MetaData->GetVariableGuid()), PotentialMatchingInputData);
-				}
-
-				if(OptionalCache == nullptr)
-				{
-					// if metadata is set and the 2nd guid matches the variable, we found the input that our HierarchyModuleInput represents
-					if(MetaData.GetValue().GetVariableGuid() == VariableGuid)
-					{
-						return PotentialMatchingInputData;
-					}
-				}
-			}			
-		}
-	}
-
-	// if we specified a cache, we didn't have an early return. That means when we reach this code, the entire data should be cached now and if it's not, that means the data doesn't exist.
-	if(OptionalCache != nullptr)
-	{
-		FInputDataCacheKey InputDataCacheKey(FunctionCallNode.NodeGuid, VariableGuid);
-		if(OptionalCache->Contains(InputDataCacheKey))
-		{
-			return (*OptionalCache)[InputDataCacheKey];
-		}
-	}
-	
-	return TOptional<FMatchingFunctionInputData>();
-}
-
-TArray<FGuid> FNiagaraStackGraphUtilities::GetChildrenInputGuids(const UNiagaraNodeFunctionCall& FunctionCallNode, FName ParentInputName)
-{
-	TArray<FGuid> Result;
-	
-	TArray<FNiagaraVariable> Variables;
-	FunctionCallNode.GetCalledGraph()->GetAllVariables(Variables);
-	
-	for(FNiagaraVariable& Variable : Variables)
-	{
-		TOptional<FNiagaraVariableMetaData> MetaData = FunctionCallNode.GetCalledGraph()->GetMetaData(Variable);
-
-		if(MetaData.IsSet())
-		{
-			FNiagaraParameterHandle ParentInputHandle(ParentInputName);
-			if(MetaData->ParentAttribute == ParentInputHandle.GetName())
-			{
-				Result.Add(MetaData->GetVariableGuid());
-			}
-		}
-		else
-		{
-			UE_LOG(LogNiagaraEditor, Log, TEXT("Couldn't find metadata for variable"));
-		}
-	}
-
-	return Result;
-}
-
-TSet<FGuid> FNiagaraStackGraphUtilities::GetHiddenChildrenInputGuids(const UNiagaraNodeFunctionCall& FunctionCallNode, FName ParentInput, TSharedRef<FNiagaraEmitterViewModel> EmitterViewModel)
-{
-	TSet<FGuid> Result;
-	
-	FCompileConstantResolver ConstantResolver(EmitterViewModel->GetEmitter(), GetOutputNodeUsage(FunctionCallNode));
-
-	TArray<FGuid> ChildrenInputGuids = GetChildrenInputGuids(FunctionCallNode, ParentInput);
-	// first, we check the normal module inputs
-	TArray<FNiagaraVariable> InputVariables;
-	TSet<FNiagaraVariable> HiddenVariables;
-	FNiagaraStackGraphUtilities::ENiagaraGetStackFunctionInputPinsOptions Options = FNiagaraStackGraphUtilities::ENiagaraGetStackFunctionInputPinsOptions::AllInputs;
-	FNiagaraStackGraphUtilities::GetStackFunctionInputs(FunctionCallNode, InputVariables, HiddenVariables, ConstantResolver, Options, true);
-
-	TArray<FGuid> ChildrenGuids = GetChildrenInputGuids(FunctionCallNode, ParentInput);
-	
-	for(const FNiagaraVariable& Variable : HiddenVariables.Array())
-	{
-		TOptional<FNiagaraVariableMetaData> MetaData = FunctionCallNode.GetCalledGraph()->GetMetaData(Variable);
-
-		if(MetaData.IsSet())
-		{
-			if(ChildrenGuids.Contains(MetaData->GetVariableGuid()))
-			{
-				Result.Add(MetaData->GetVariableGuid());
-			}
-		}
-		else
-		{
-			UE_LOG(LogNiagaraEditor, Log, TEXT("Couldn't find metadata for variable"));
-		}
-	}
-
-	return Result;
-}
-
 TArray<UEdGraphPin*> FNiagaraStackGraphUtilities::GetUnusedFunctionInputPins(const UNiagaraNodeFunctionCall& FunctionCallNode, FCompileConstantResolver ConstantResolver)
 {
 	UNiagaraGraph* FunctionGraph = FunctionCallNode.GetCalledGraph();
diff --git a/Engine/Plugins/FX/Niagara/Source/NiagaraEditor/Private/ViewModels/Stack/NiagaraStackInputCategory.cpp b/Engine/Plugins/FX/Niagara/Source/NiagaraEditor/Private/ViewModels/Stack/NiagaraStackInputCategory.cpp
--- a/Engine/Plugins/FX/Niagara/Source/NiagaraEditor/Private/ViewModels/Stack/NiagaraStackInputCategory.cpp
+++ b/Engine/Plugins/FX/Niagara/Source/NiagaraEditor/Private/ViewModels/Stack/NiagaraStackInputCategory.cpp
@@ -3,33 +3,35 @@
 #include "ViewModels/Stack/NiagaraStackInputCategory.h"
 
 #include "IDetailTreeNode.h"
-#include "IPropertyRowGenerator.h"
 #include "ViewModels/Stack/NiagaraStackFunctionInput.h"
 #include "NiagaraNodeFunctionCall.h"
 #include "NiagaraClipboard.h"
 #include "NiagaraConstants.h"
 #include "NiagaraSimulationStageBase.h"
 #include "ViewModels/NiagaraEmitterViewModel.h"
+#include "ViewModels/HierarchyEditor/NiagaraHierarchyScriptParametersViewModel.h"
 #include "ViewModels/HierarchyEditor/NiagaraSummaryViewViewModel.h"
 #include "ViewModels/Stack/NiagaraStackEmitterSettingsGroup.h"
 #include "ViewModels/Stack/NiagaraStackEventScriptItemGroup.h"
-#include "ViewModels/Stack/NiagaraStackFunctionInputCollection.h"
 #include "ViewModels/Stack/NiagaraStackModuleItem.h"
 #include "ViewModels/Stack/NiagaraStackObject.h"
 #include "ViewModels/Stack/NiagaraStackPropertyRow.h"
 #include "ViewModels/Stack/NiagaraStackRendererItem.h"
 #include "ViewModels/Stack/NiagaraStackRenderersOwner.h"
 #include "ViewModels/Stack/NiagaraStackSimulationStageGroup.h"
+#include "Algo/Accumulate.h"
 
 #include UE_INLINE_GENERATED_CPP_BY_NAME(NiagaraStackInputCategory)
 
 using namespace FNiagaraStackGraphUtilities;
 
+#define LOCTEXT_NAMESPACE "NiagaraStack"
+
 void UNiagaraStackCategory::Initialize(FRequiredEntryData InRequiredEntryData, FString InOwningStackItemEditorDataKey, FString InStackEditorDataKey)
 {
 	Super::Initialize(InRequiredEntryData, InOwningStackItemEditorDataKey, InStackEditorDataKey);
-	AddChildFilter(FOnFilterChild::CreateUObject(this, &UNiagaraStackInputCategory::FilterForVisibleCondition));
-	AddChildFilter(FOnFilterChild::CreateUObject(this, &UNiagaraStackInputCategory::FilterForIsInlineEditConditionToggle));
+	AddChildFilter(FOnFilterChild::CreateUObject(this, &UNiagaraStackCategory::FilterForVisibleCondition));
+	AddChildFilter(FOnFilterChild::CreateUObject(this, &UNiagaraStackCategory::FilterForIsInlineEditConditionToggle));
 	bShouldShowInStack = true;
 	CategorySpacer = nullptr;
 }
@@ -77,7 +79,7 @@ void UNiagaraStackCategory::RefreshChildrenInternal(const TArray<UNiagaraStackEn
 		{
 			CategorySpacer = NewObject<UNiagaraStackSpacer>(this);
 			TAttribute<bool> ShouldShowSpacerInStack;
-			ShouldShowSpacerInStack.BindUObject(this, &UNiagaraStackInputCategory::GetShouldShowInStack);
+			ShouldShowSpacerInStack.BindUObject(this, &UNiagaraStackCategory::GetShouldShowInStack);
 			CategorySpacer->Initialize(CreateDefaultChildRequiredData(), 6, ShouldShowSpacerInStack, GetStackEditorDataKey());
 		}
 		NewChildren.Add(CategorySpacer);
@@ -96,141 +98,178 @@ bool UNiagaraStackCategory::FilterForIsInlineEditConditionToggle(const UNiagaraS
 	return StackFunctionInputChild == nullptr || StackFunctionInputChild->GetIsInlineEditConditionToggle() == false;
 }
 
-void UNiagaraStackInputCategory::Initialize(
-	FRequiredEntryData InRequiredEntryData,
-	FString InputCategoryStackEditorDataKey,
-	FText InCategoryName,
-	bool bInIsTopLevelCategory,
-	FString InOwnerStackItemEditorDataKey)
+void UNiagaraStackScriptHierarchyCategory::Initialize(FRequiredEntryData InRequiredEntryData, const UNiagaraHierarchyCategory& InHierarchyCategory, FString InOwningStackItemEditorDataKey, FString InStackEditorDataKey)
 {
-	bool bCategoryIsAdvanced = false;
-	Super::Initialize(InRequiredEntryData, InOwnerStackItemEditorDataKey, InputCategoryStackEditorDataKey);
-	CategoryName = InCategoryName;
-	bIsTopLevelCategory = bInIsTopLevelCategory;
+	Super::Initialize(InRequiredEntryData, InOwningStackItemEditorDataKey, InStackEditorDataKey);
+	HierarchyCategory = &InHierarchyCategory;
 }
 
-bool UNiagaraStackInputCategory::GetIsEnabled() const
+void UNiagaraStackScriptHierarchyCategory::Copy(UNiagaraClipboardContent* ClipboardContent) const
 {
-	for (const auto& Input : Inputs)
-	{
-		if (Input.InputFunctionCallNode->GetDesiredEnabledState() == ENodeEnabledState::Enabled)
-		{
-			return true;
-		}
-	}
-
-	return false;
+	ClipboardContent->FunctionInputs.Append(ToClipboardFunctionInputs(ClipboardContent));
 }
 
-FText UNiagaraStackInputCategory::GetDisplayName() const
+void UNiagaraStackScriptHierarchyCategory::Paste(const UNiagaraClipboardContent* ClipboardContent, FText& OutPasteWarning)
 {
-	return CategoryName;
+	PasteFromClipboard(ClipboardContent);
 }
 
-void UNiagaraStackInputCategory::ResetInputs()
+bool UNiagaraStackScriptHierarchyCategory::TestCanCopyWithMessage(FText& OutMessage) const
 {
-	Inputs.Empty();
+	TArray<UNiagaraStackFunctionInput*> StackFunctionInputs;
+	GetFilteredChildrenOfType(StackFunctionInputs, true);
+
+	if(StackFunctionInputs.Num() == 0)
+	{
+		OutMessage = LOCTEXT("CantCopyCategory_NoInputs", "No inputs available for copying.");
+		return false;
+	}
+
+	OutMessage = LOCTEXT("TestCanCopyCategory_Success", "Copy all inputs of this category.");
+	return true;
 }
 
-void UNiagaraStackInputCategory::AddInput(UNiagaraNodeFunctionCall* InModuleNode, UNiagaraNodeFunctionCall* InInputFunctionCallNode, FName InInputParameterHandle, FNiagaraTypeDefinition InInputType, EStackParameterBehavior InParameterBehavior, TOptional<FText> InOptionalDisplayName, bool bIsInputHidden, bool bIsChildInput)
+bool UNiagaraStackScriptHierarchyCategory::TestCanPasteWithMessage(const UNiagaraClipboardContent* ClipboardContent, FText& OutMessage) const
 {
-	Inputs.Add({ InModuleNode, InInputFunctionCallNode, InInputParameterHandle, InInputType, InParameterBehavior, InOptionalDisplayName, bIsInputHidden, bIsChildInput });
+	if(ClipboardContent->FunctionInputs.Num() == 0)
+	{
+		OutMessage = LOCTEXT("TestCanPasteCategory_NoCopiedInputsFound", "No inputs were copied.");
+		return false;
+	}
+	
+	TArray<UNiagaraStackFunctionInput*> AllContainedInputs;
+	GetFilteredChildrenOfType(AllContainedInputs, true);
+
+	int32 MatchingInputsNum = Algo::Accumulate(AllContainedInputs, 0, [ClipboardContent](int32 Value, UNiagaraStackFunctionInput* Input)
+	{
+		Value += ClipboardContent->FunctionInputs.ContainsByPredicate([Input](const UNiagaraClipboardFunctionInput* ClipboardFunctionInputCandidate)
+		{
+			return ClipboardFunctionInputCandidate->InputName == Input->GetInputParameterHandle().GetName() && ClipboardFunctionInputCandidate->InputType == Input->GetInputType();
+		}) ? 1 : 0;
+
+		return Value;
+	});
+	
+	if(MatchingInputsNum == 0)
+	{
+		OutMessage = LOCTEXT("TestCanPasteCategory_NoMatchingInputsFound", "No matching copied inputs were found for pasting into this category.");
+		return false;
+	}
+
+	OutMessage = LOCTEXT("TestCanPasteCategory_Success", "Paste matching inputs into this category.");
+	return true;
 }
 
-void UNiagaraStackInputCategory::RefreshChildrenInternal(const TArray<UNiagaraStackEntry*>& CurrentChildren, TArray<UNiagaraStackEntry*>& NewChildren, TArray<FStackIssue>& NewIssues)
+void UNiagaraStackScriptHierarchyCategory::PasteFromClipboard(const UNiagaraClipboardContent* ClipboardContent)
 {
-	for (FInputParameterHandleAndType& Input : Inputs)
+	TArray<UNiagaraStackFunctionInput*> AllContainedInputs;
+	GetUnfilteredChildrenOfType(AllContainedInputs, true);
+	
+	for(const UNiagaraClipboardFunctionInput* ClipboardInput : ClipboardContent->FunctionInputs)
 	{
-		UNiagaraStackFunctionInput* InputChild = FindCurrentChildOfTypeByPredicate<UNiagaraStackFunctionInput>(CurrentChildren, [&](UNiagaraStackFunctionInput* CurrentInput) 
-		{ 
-			return CurrentInput->GetInputParameterHandle() == Input.ParameterHandle && CurrentInput->GetInputType() == Input.Type && CurrentInput->GetInputFunctionCallInitialScript() == Input.InputFunctionCallNode->FunctionScript;
-		});
+		if(ClipboardInput == nullptr)
+		{
+			continue;
+		}
 
-		if (InputChild == nullptr)
+		// Since we pasting into a category, we require name & type match to ensure we aren't writing into unintended inputs
+		for(UNiagaraStackFunctionInput* StackInput : AllContainedInputs)
 		{
-			InputChild = NewObject<UNiagaraStackFunctionInput>(this);
-			InputChild->Initialize(CreateDefaultChildRequiredData(), *Input.ModuleNode, *Input.InputFunctionCallNode,
-				Input.ParameterHandle, Input.Type, Input.ParameterBehavior, GetOwnerStackItemEditorDataKey());
+			if(StackInput->GetInputParameterHandle().GetName() == ClipboardInput->InputName && StackInput->GetInputType() == ClipboardInput->InputType)
+			{
+				StackInput->PasteFunctionInput(ClipboardInput);
+			}
 		}
-		InputChild->SetIsHidden(Input.bIsHidden);
-		InputChild->SetSemanticChild(Input.bIsChildInput);
-		NewChildren.Add(InputChild);
 	}
+}
 
-	Super::RefreshChildrenInternal(CurrentChildren, NewChildren, NewIssues);
+TArray<const UNiagaraClipboardFunctionInput*> UNiagaraStackScriptHierarchyCategory::ToClipboardFunctionInputs(UObject* InOuter) const
+{
+	TArray<const UNiagaraClipboardFunctionInput*> Result;
+	
+	TArray<UNiagaraStackScriptHierarchyCategory*> SubCategories;
+	GetFilteredChildrenOfType(SubCategories, false);
+
+	for(UNiagaraStackScriptHierarchyCategory* SubCategory : SubCategories)
+	{
+		Result.Append(SubCategory->ToClipboardFunctionInputs(InOuter));
+	}
+
+	TArray<UNiagaraStackFunctionInput*> ContainedInputs;
+	GetFilteredChildrenOfType(ContainedInputs, false);
+
+	for(UNiagaraStackFunctionInput* ContainedInput : ContainedInputs)
+	{
+		Result.Add(ContainedInput->ToClipboardFunctionInput(InOuter));
+	}
+
+	return Result;
 }
 
-void UNiagaraStackInputCategory::SetShouldShowInStack(bool bInShouldShowInStack)
+FText UNiagaraStackScriptHierarchyCategory::GetDisplayName() const
 {
-	bShouldShowInStack = bInShouldShowInStack;
+	return HierarchyCategory->GetCategoryAsText();
 }
 
-void UNiagaraStackInputCategory::ToClipboardFunctionInputs(UObject* InOuter, TArray<const UNiagaraClipboardFunctionInput*>& OutClipboardFunctionInputs) const
+FText UNiagaraStackScriptHierarchyCategory::GetTooltipText() const
 {
-	TArray<UNiagaraStackFunctionInput*> ChildInputs;
-	GetUnfilteredChildrenOfType(ChildInputs);
-	for (UNiagaraStackFunctionInput* ChildInput : ChildInputs)
-	{
-		const UNiagaraClipboardFunctionInput* FunctionInput = ChildInput->ToClipboardFunctionInput(InOuter);
-		if (FunctionInput != nullptr)
-		{
-			OutClipboardFunctionInputs.Add(FunctionInput);
-		}
-	}
+	return HierarchyCategory->GetTooltip();
 }
 
-bool UNiagaraStackInputCategory::TrySetStaticSwitchValuesFromClipboardFunctionInput(const UNiagaraClipboardFunctionInput& ClipboardFunctionInput)
+void UNiagaraStackScriptHierarchyCategory::RefreshChildrenInternal(const TArray<UNiagaraStackEntry*>& CurrentChildren, TArray<UNiagaraStackEntry*>& NewChildren, TArray<FStackIssue>& NewIssues)
 {
-	TArray<UNiagaraStackFunctionInput*> ChildInputs;
-	GetUnfilteredChildrenOfType(ChildInputs);
-	for (UNiagaraStackFunctionInput* ChildInput : ChildInputs)
+	for(const UNiagaraHierarchyItemBase* ChildHierarchyItem : HierarchyCategory->GetChildren())
 	{
-		if (ChildInput->IsStaticParameter() &&
-			ChildInput->GetInputParameterHandle().GetName() == ClipboardFunctionInput.InputName &&
-			ChildInput->GetInputType() == ClipboardFunctionInput.InputType)
+		if(const UNiagaraHierarchyScriptParameter* HierarchyParameter = Cast<UNiagaraHierarchyScriptParameter>(ChildHierarchyItem))
 		{
-			if (ClipboardFunctionInput.ValueMode == ENiagaraClipboardFunctionInputValueMode::ResetToDefault)
+			FNiagaraVariable InputVariable = HierarchyParameter->GetVariable();
+			
+			if(ScriptInstanceData.UsedInputs.Contains(InputVariable) == false)
 			{
-				ChildInput->Reset();
+				continue;
 			}
-			else
+			
+			UNiagaraStackFunctionInput* InputChild = FindCurrentChildOfTypeByPredicate<UNiagaraStackFunctionInput>(CurrentChildren, [&](UNiagaraStackFunctionInput* CurrentInput) 
+			{ 
+				return CurrentInput->GetInputParameterHandle() == FNiagaraParameterHandle(InputVariable.GetName()) && CurrentInput->GetInputType() == HierarchyParameter->GetVariable().GetType() && &CurrentInput->GetInputFunctionCallNode() == OwningFunctionCallNode.Get();
+			});
+
+			if (InputChild == nullptr)
 			{
-				ChildInput->SetValueFromClipboardFunctionInput(ClipboardFunctionInput);
+				EStackParameterBehavior Behavior = HierarchyParameter->GetScriptVariable()->GetIsStaticSwitch() ? EStackParameterBehavior::Static : EStackParameterBehavior::Dynamic;
+				InputChild = NewObject<UNiagaraStackFunctionInput>(this);
+				InputChild->Initialize(CreateDefaultChildRequiredData(), *OwningModuleNode.Get(), *OwningFunctionCallNode.Get(),
+					InputVariable.GetName(), InputVariable.GetType(), Behavior, GetOwnerStackItemEditorDataKey());
 			}
-			return true;
-		}
-	}
-	return false;
-}
 
-void  UNiagaraStackInputCategory::SetStandardValuesFromClipboardFunctionInputs(const TArray<const UNiagaraClipboardFunctionInput*>& ClipboardFunctionInputs)
-{
-	TArray<UNiagaraStackFunctionInput*> ChildInputs;
-	GetUnfilteredChildrenOfType(ChildInputs);
-	for (const UNiagaraClipboardFunctionInput* ClipboardFunctionInput : ClipboardFunctionInputs)
-	{
-		for (UNiagaraStackFunctionInput* ChildInput : ChildInputs)
+			InputChild->SetScriptInstanceData(ScriptInstanceData);
+			FGuid VariableGuid = HierarchyParameter->GetScriptVariable()->Metadata.GetVariableGuid();
+			InputChild->SetIsHidden(ScriptInstanceData.PerInputInstanceData[VariableGuid].bIsHidden);
+			
+			NewChildren.Add(InputChild);
+		}
+		if(const UNiagaraHierarchyCategory* ChildHierarchyCategory = Cast<UNiagaraHierarchyCategory>(ChildHierarchyItem))
 		{
-			if (ChildInput->IsStaticParameter() == false && 
-				ChildInput->GetInputParameterHandle().GetName() == ClipboardFunctionInput->InputName &&
-				ChildInput->GetInputType() == ClipboardFunctionInput->InputType)
+			// Try to find an already existing category to reuse
+			UNiagaraStackScriptHierarchyCategory* ChildCategory = FindCurrentChildOfTypeByPredicate<UNiagaraStackScriptHierarchyCategory>(CurrentChildren,
+				[&](UNiagaraStackScriptHierarchyCategory* CurrentCategory) { return CurrentCategory->GetHierarchyCategory() == ChildHierarchyItem; });
+
+			if (ChildCategory == nullptr)
 			{
-				if (ClipboardFunctionInput->ValueMode == ENiagaraClipboardFunctionInputValueMode::ResetToDefault)
-				{
-					ChildInput->Reset();
-				}
-				else
-				{
-					ChildInput->SetValueFromClipboardFunctionInput(*ClipboardFunctionInput);
-				}
+				// If we don't have a current child for this category make a new one.
+				ChildCategory = NewObject<UNiagaraStackScriptHierarchyCategory>(this);
+				ChildCategory->SetOwningModuleNode(*OwningModuleNode.Get());
+				ChildCategory->SetOwningFunctionCallNode(*OwningFunctionCallNode.Get());
+				FString InputCategoryStackEditorDataKey = FString::Printf(TEXT("%s-InputCategory-%s"), *OwningFunctionCallNode->NodeGuid.ToString(EGuidFormats::DigitsWithHyphens), *ChildHierarchyCategory->ToString());
+				ChildCategory->Initialize(CreateDefaultChildRequiredData(), *ChildHierarchyCategory, GetOwnerStackItemEditorDataKey(), InputCategoryStackEditorDataKey);
 			}
-		}
-	}
-}
 
-void UNiagaraStackInputCategory::GetFilteredChildInputs(TArray<UNiagaraStackFunctionInput*>& OutFilteredChildInputs) const
-{
-	GetFilteredChildrenOfType(OutFilteredChildInputs);
+			ChildCategory->SetScriptInstanceData(ScriptInstanceData);
+			NewChildren.Add(ChildCategory);
+		}		
+	}
+	
+	Super::RefreshChildrenInternal(CurrentChildren, NewChildren, NewIssues);
 }
 
 void UNiagaraStackSummaryCategory::Initialize(FRequiredEntryData InRequiredEntryData, TSharedPtr<FNiagaraHierarchyCategoryViewModel> InCategoryViewModel, FString InOwnerStackItemEditorDataKey)
@@ -699,3 +738,5 @@ int32 UNiagaraStackSummaryCategory::GetChildIndentLevel() const
 {
 	return UNiagaraStackEntry::GetChildIndentLevel();
 }
+
+#undef LOCTEXT_NAMESPACE
diff --git a/Engine/Plugins/FX/Niagara/Source/NiagaraEditor/Private/ViewModels/Stack/NiagaraStackModuleItem.cpp b/Engine/Plugins/FX/Niagara/Source/NiagaraEditor/Private/ViewModels/Stack/NiagaraStackModuleItem.cpp
--- a/Engine/Plugins/FX/Niagara/Source/NiagaraEditor/Private/ViewModels/Stack/NiagaraStackModuleItem.cpp
+++ b/Engine/Plugins/FX/Niagara/Source/NiagaraEditor/Private/ViewModels/Stack/NiagaraStackModuleItem.cpp
@@ -32,7 +32,7 @@
 #include "ViewModels/NiagaraScratchPadViewModel.h"
 #include "ViewModels/NiagaraParameterPanelViewModel.h"
 #include "ViewModels/NiagaraSystemViewModel.h"
-#include "ViewModels/Stack/NiagaraStackFunctionInputCollection.h"
+#include "ViewModels/Stack/NiagaraStackValueCollection.h"
 #include "ViewModels/Stack/NiagaraStackGraphUtilities.h"
 #include "ViewModels/Stack/NiagaraStackModuleItemLinkedInputCollection.h"
 #include "ViewModels/Stack/NiagaraStackModuleItemOutputCollection.h"
@@ -59,7 +59,6 @@ TArray<ENiagaraScriptUsage> UsagePriority = { // Ordered such as the highest pri
 UNiagaraStackModuleItem::UNiagaraStackModuleItem()
 	: FunctionCallNode(nullptr)
 	, bCanRefresh(false)
-	, InputCollection(nullptr)
 	, bIsModuleScriptReassignmentPending(false)
 {
 }
@@ -153,17 +152,15 @@ void UNiagaraStackModuleItem::RefreshChildrenInternal(const TArray<UNiagaraStack
 			bCanRefresh = true;
 		}
 		
-		if (InputCollection == nullptr)
+		if(InputRoot == nullptr)
 		{
-			TArray<FString> InputParameterHandlePath;
-			InputCollection = NewObject<UNiagaraStackFunctionInputCollection>(this);
-			InputCollection->Initialize(CreateDefaultChildRequiredData(), *FunctionCallNode, *FunctionCallNode, GetStackEditorDataKey());
+			InputRoot = NewObject<UNiagaraStackScriptHierarchyRoot>(this);
+			InputRoot->Initialize(CreateDefaultChildRequiredData(), *FunctionCallNode, *FunctionCallNode, GetStackEditorDataKey());
 		}
 
 		// NiagaraNodeAssignments should not display OutputCollection and LinkedInputCollection as they effectively handle this through their InputCollection 
 		if (!FunctionCallNode->IsA<UNiagaraNodeAssignment>())
 		{
-
 			if (LinkedInputCollection == nullptr)
 			{
 				LinkedInputCollection = NewObject<UNiagaraStackModuleItemLinkedInputCollection>(this);
@@ -177,20 +174,17 @@ void UNiagaraStackModuleItem::RefreshChildrenInternal(const TArray<UNiagaraStack
 				OutputCollection->Initialize(CreateDefaultChildRequiredData(), *FunctionCallNode);
 				OutputCollection->AddChildFilter(FOnFilterChild::CreateUObject(this, &UNiagaraStackModuleItem::FilterOutputCollectionChild));
 			}
-
-			InputCollection->SetShouldDisplayLabel(GetStackEditorData().GetShowOutputs() || GetStackEditorData().GetShowLinkedInputs());
-
-			NewChildren.Add(InputCollection);
+			
+			InputRoot->SetShouldDisplayLabel(GetStackEditorData().GetShowOutputs() || GetStackEditorData().GetShowLinkedInputs());
+			NewChildren.Add(InputRoot);
+			
 			NewChildren.Add(LinkedInputCollection);
 			NewChildren.Add(OutputCollection);
 		
 		}
 		else
 		{
-			// We do not show the expander arrow for InputCollections of NiagaraNodeAssignments as they only have this one collection
-			InputCollection->SetShouldDisplayLabel(false);
-
-			NewChildren.Add(InputCollection);
+			NewChildren.Add(InputRoot);
 
 			UNiagaraNodeAssignment* AssignmentNode = CastChecked<UNiagaraNodeAssignment>(FunctionCallNode);
 			if (AssignmentNode->GetAssignmentTargets().Num() == 0)
@@ -253,13 +247,6 @@ const UNiagaraStackModuleItem::FCollectedUsageData& UNiagaraStackModuleItem::Get
 		}
 
 
-		if (InputCollection)
-		{
-			if (InputCollection->GetCollectedUsageData().bHasReferencedParameterRead)
-				CachedCollectedUsageData.GetValue().bHasReferencedParameterRead = true;
-			if (InputCollection->GetCollectedUsageData().bHasReferencedParameterWrite)
-				CachedCollectedUsageData.GetValue().bHasReferencedParameterWrite = true;
-		}
 
 		if (OutputCollection)
 		{
@@ -1435,7 +1422,7 @@ void UNiagaraStackModuleItem::ReassignModuleScript(UNiagaraScript* ModuleScript)
 			if (ConversionUtility )
 			{
 				FText ConvertMessage;
-				bool bConverted = ConversionUtility->Convert(OldScript, OldClipboardContent, ModuleScript, InputCollection, NewClipboardContent, FunctionCallNode, ConvertMessage);
+				bool bConverted = ConversionUtility->Convert(OldScript, OldClipboardContent, ModuleScript, InputRoot, NewClipboardContent, FunctionCallNode, ConvertMessage);
 				if (!ConvertMessage.IsEmptyOrWhitespace())
 				{
 					// Notify the end-user about the convert message, but continue the process as they could always undo.
@@ -1535,17 +1522,17 @@ void UNiagaraStackModuleItem::ChangeScriptVersion(FGuid NewScriptVersion)
 
 void UNiagaraStackModuleItem::SetInputValuesFromClipboardFunctionInputs(const TArray<const UNiagaraClipboardFunctionInput*>& ClipboardFunctionInputs)
 {
-	InputCollection->SetValuesFromClipboardFunctionInputs(ClipboardFunctionInputs);
+	InputRoot->SetValuesFromClipboardFunctionInputs(ClipboardFunctionInputs);
 }
 
 void UNiagaraStackModuleItem::GetParameterInputs(TArray<UNiagaraStackFunctionInput*>& OutResult) const
 {
-	return InputCollection->GetChildInputs(OutResult);
+	InputRoot->GetChildInputs(OutResult);
 }
 
 TArray<UNiagaraStackFunctionInput*> UNiagaraStackModuleItem::GetInlineParameterInputs() const
 {
-	return InputCollection->GetInlineParameterInputs();
+	return InputRoot->GetInlineParameters();
 }
 
 bool UNiagaraStackModuleItem::TestCanCutWithMessage(FText& OutMessage) const
@@ -1613,8 +1600,10 @@ void UNiagaraStackModuleItem::Copy(UNiagaraClipboardContent* ClipboardContent) c
 
 	ClipboardFunction->DisplayName = GetAlternateDisplayName().Get(FText::GetEmpty());
 
-	InputCollection->ToClipboardFunctionInputs(ClipboardFunction, MutableView(ClipboardFunction->Inputs));
+	InputRoot->ToClipboardFunctionInputs(ClipboardFunction, MutableView(ClipboardFunction->Inputs));
 	ClipboardContent->Functions.Add(ClipboardFunction);
+
+	OnCopyPasteDelegate.ExecuteIfBound();
 }
 
 bool UNiagaraStackModuleItem::TestCanPasteWithMessage(const UNiagaraClipboardContent* ClipboardContent, FText& OutMessage) const
diff --git a/Engine/Plugins/FX/Niagara/Source/NiagaraEditor/Private/ViewModels/Stack/NiagaraStackScriptItemGroup.cpp b/Engine/Plugins/FX/Niagara/Source/NiagaraEditor/Private/ViewModels/Stack/NiagaraStackScriptItemGroup.cpp
--- a/Engine/Plugins/FX/Niagara/Source/NiagaraEditor/Private/ViewModels/Stack/NiagaraStackScriptItemGroup.cpp
+++ b/Engine/Plugins/FX/Niagara/Source/NiagaraEditor/Private/ViewModels/Stack/NiagaraStackScriptItemGroup.cpp
@@ -543,6 +543,8 @@ void UNiagaraStackScriptItemGroup::Paste(const UNiagaraClipboardContent* Clipboa
 		GetUnfilteredChildrenOfType(ModuleItems);
 		int32 PasteIndex = ModuleItems.Num() > 0 ? ModuleItems.Last()->GetModuleIndex() + 1 : 0;
 		PasteModules(ClipboardContent, PasteIndex, OutPasteWarning);
+
+		OnCopyPasteDelegate.ExecuteIfBound();
 	}
 }
 
diff --git a/Engine/Plugins/FX/Niagara/Source/NiagaraEditor/Private/ViewModels/Stack/NiagaraStackSummaryViewInputCollection.cpp b/Engine/Plugins/FX/Niagara/Source/NiagaraEditor/Private/ViewModels/Stack/NiagaraStackSummaryViewInputCollection.cpp
--- a/Engine/Plugins/FX/Niagara/Source/NiagaraEditor/Private/ViewModels/Stack/NiagaraStackSummaryViewInputCollection.cpp
+++ b/Engine/Plugins/FX/Niagara/Source/NiagaraEditor/Private/ViewModels/Stack/NiagaraStackSummaryViewInputCollection.cpp
@@ -16,8 +16,9 @@
 #include "ViewModels/Stack/NiagaraStackEmitterPropertiesGroup.h"
 #include "ViewModels/Stack/NiagaraStackEmitterSettingsGroup.h"
 #include "ViewModels/Stack/NiagaraStackEventScriptItemGroup.h"
+#include "ViewModels/Stack/NiagaraStackFunctionInput.h"
 #include "ViewModels/Stack/NiagaraStackGraphUtilities.h"
-#include "ViewModels/Stack/NiagaraStackFunctionInputCollection.h"
+#include "ViewModels/Stack/NiagaraStackValueCollection.h"
 #include "ViewModels/Stack/NiagaraStackInputCategory.h"
 #include "ViewModels/Stack/NiagaraStackModuleItem.h"
 #include "ViewModels/Stack/NiagaraStackObject.h"
diff --git a/Engine/Plugins/FX/Niagara/Source/NiagaraEditor/Private/Widgets/AssetBrowser/NiagaraAssetBrowserPreview.cpp b/Engine/Plugins/FX/Niagara/Source/NiagaraEditor/Private/Widgets/AssetBrowser/NiagaraAssetBrowserPreview.cpp
--- a/Engine/Plugins/FX/Niagara/Source/NiagaraEditor/Private/Widgets/AssetBrowser/NiagaraAssetBrowserPreview.cpp
+++ b/Engine/Plugins/FX/Niagara/Source/NiagaraEditor/Private/Widgets/AssetBrowser/NiagaraAssetBrowserPreview.cpp
@@ -5,6 +5,7 @@
 #include "CanvasItem.h"
 #include "CanvasTypes.h"
 #include "NiagaraCommon.h"
+#include "NiagaraEditorUtilities.h"
 #include "NiagaraSystemEditorData.h"
 #include "NiagaraSystemFactoryNew.h"
 #include "NiagaraSystemInstanceController.h"
