Niagara Script Editor - Hierarchy Editor Integration

diff --git a/Engine/Plugins/FX/Niagara/Source/NiagaraEditor/Private/Widgets/SNiagaraHierarchyEditor.cpp b/Engine/Plugins/FX/Niagara/Source/NiagaraEditor/Private/Widgets/SNiagaraHierarchyEditor.cpp
new file mode 100644
--- /dev/null
+++ b/Engine/Plugins/FX/Niagara/Source/NiagaraEditor/Private/Widgets/SNiagaraHierarchyEditor.cpp
@@ -0,0 +1,1723 @@
+// Copyright Epic Games, Inc. All Rights Reserved.
+
+#include "Widgets/SNiagaraHierarchyEditor.h"
+
+#include "NiagaraEditorCommon.h"
+#include "UObject/Package.h"
+#include "Framework/Commands/GenericCommands.h"
+#include "Styling/StyleColors.h"
+#include "Styling/AppStyle.h"
+#include "ViewModels/HierarchyEditor/NiagaraHierarchyViewModelBase.h"
+#include "Widgets/SNiagaraParameterName.h"
+#include "Widgets/Layout/SSeparator.h"
+#include "Widgets/Layout/SWrapBox.h"
+#include "Widgets/Input/SCheckBox.h"
+#include "Widgets/Input/SButton.h"
+#include "Widgets/Images/SImage.h"
+#include "Widgets/Text/SInlineEditableTextBlock.h"
+#include "Widgets/Text/STextBlock.h"
+#include "ScopedTransaction.h"
+#include "Modules/ModuleManager.h"
+#include "PropertyEditorModule.h"
+#include "Framework/Application/MenuStack.h"
+#include "Framework/Application/SlateApplication.h"
+#include "ToolMenus.h"
+#include "ViewModels/HierarchyEditor/NiagaraHierarchyCommands.h"
+#include "Widgets/Input/SSearchBox.h"
+
+#define LOCTEXT_NAMESPACE "NiagaraHierarchy"
+
+TSharedRef<SWidget> SummonContextMenu(TArray<TSharedPtr<FNiagaraHierarchyItemViewModelBase>> Items, bool bFromHierarchy, TSharedPtr<const SWidget> HierarchyWidget)
+{
+	UNiagaraHierarchyMenuContext* MenuContextObject = NewObject<UNiagaraHierarchyMenuContext>();
+	MenuContextObject->Items = Items;
+	MenuContextObject->bFromHierarchy = bFromHierarchy;
+
+	FToolMenuContext MenuContext(MenuContextObject);
+	TWeakObjectPtr<UNiagaraHierarchyViewModelBase> ViewModel = Items[0]->GetHierarchyViewModel();
+	MenuContext.AppendCommandList(ViewModel->GetCommands());
+		
+	UToolMenu* Menu = UToolMenus::Get()->GenerateMenu("NiagaraHierarchyMenu", MenuContext);
+	FToolMenuSection* BaseSection = Menu->FindSection("Base");
+
+	if(ensure(BaseSection != nullptr))
+	{
+		if(Items.Num() == 1)
+		{
+			if(Items[0]->CanRename())
+			{
+				BaseSection->AddMenuEntry(FGenericCommands::Get().Rename);
+			}
+		}
+
+		if(Items.Num() == 1 && Items[0]->GetData()->IsA<UNiagaraHierarchySection>())
+		{
+			BaseSection->AddMenuEntry(FNiagaraHierarchyEditorCommands::Get().DeleteSection);
+		}
+		
+		// the generic delete command handles hierarchy items (not sections)
+		BaseSection->AddMenuEntry(FGenericCommands::Get().Delete);
+		
+		if(Items.Num() == 1 && bFromHierarchy == false)
+		{
+			if(UNiagaraHierarchyItemBase* FoundItem = ViewModel->GetHierarchyRoot()->FindChildWithIdentity(Items[0]->GetData()->GetPersistentIdentity(), true))
+			{
+				FExecuteAction ExecuteAction = FExecuteAction::CreateLambda([Items, ViewModel, HierarchyWidget, Identity = FoundItem->GetPersistentIdentity()]()
+				{
+					TSharedPtr<FNiagaraHierarchyItemViewModelBase> RespectiveHierarchyItem = ViewModel->GetHierarchyRootViewModel()->FindViewModelForChild(Identity, true);
+					if(RespectiveHierarchyItem.IsValid())
+					{
+						TSharedPtr<const SNiagaraHierarchyEditor> CastHierarchyWidget = StaticCastSharedPtr<const SNiagaraHierarchyEditor>(HierarchyWidget);
+						CastHierarchyWidget->NavigateToHierarchyItem(RespectiveHierarchyItem);
+					}
+				});
+				FToolUIActionChoice Action(ExecuteAction);
+				BaseSection->AddMenuEntry(FName("FindInHierarchy"), LOCTEXT("FindInHierarchyLabel", "Find in Hierarchy"), FText::GetEmpty(), FSlateIcon(), Action);
+			}			
+		}
+	}
+
+	if(Items.Num() == 1)
+	{
+		FToolMenuSection& DynamicSection = Menu->FindOrAddSection("Dynamic");		
+		Items[0]->PopulateDynamicContextMenuSection(DynamicSection);
+	}	
+		
+	TSharedRef<SWidget> MenuWidget = UToolMenus::Get()->GenerateWidget(Menu);
+	return MenuWidget;
+}
+
+void SNiagaraHierarchyCategory::Construct(const FArguments& InArgs,	TSharedPtr<FNiagaraHierarchyCategoryViewModel> InCategoryViewModel)
+{
+	CategoryViewModel = InCategoryViewModel;
+	
+	InCategoryViewModel->GetOnRequestRename().BindSP(this, &SNiagaraHierarchyCategory::EnterEditingMode);
+
+	const UNiagaraHierarchyCategory* CategoryData = InCategoryViewModel->GetData<UNiagaraHierarchyCategory>();
+	SetToolTipText(TAttribute<FText>::CreateLambda([CategoryData]
+	{
+		return CategoryData->GetTooltip();
+	}));
+	
+	ChildSlot
+	[
+		SAssignNew(InlineEditableTextBlock, SInlineEditableTextBlock)
+		.Style(&FNiagaraEditorStyle::Get().GetWidgetStyle<FInlineEditableTextBlockStyle>("NiagaraEditor.HierarchyEditor.CategoryTextStyle"))
+		.Text(this, &SNiagaraHierarchyCategory::GetCategoryText)
+		.OnTextCommitted(this, &SNiagaraHierarchyCategory::OnRenameCategory)
+		.OnVerifyTextChanged(this, &SNiagaraHierarchyCategory::OnVerifyCategoryRename)
+		.IsSelected(InArgs._IsSelected)
+	];
+}
+
+void SNiagaraHierarchyCategory::EnterEditingMode() const
+{
+	if(CategoryViewModel.Pin()->CanRename())
+	{
+		InlineEditableTextBlock->EnterEditingMode();
+	}
+}
+
+bool SNiagaraHierarchyCategory::OnVerifyCategoryRename(const FText& NewName, FText& OutTooltip) const
+{
+	TArray<TSharedPtr<FNiagaraHierarchyCategoryViewModel>> SiblingCategoryViewModels;
+	CategoryViewModel.Pin()->GetParent().Pin()->GetChildrenViewModelsForType<UNiagaraHierarchyCategory, FNiagaraHierarchyCategoryViewModel>(SiblingCategoryViewModels);
+
+	if(GetCategoryText().ToString() != NewName.ToString())
+	{
+		TSet<FString> CategoryNames;
+		for(const auto& SiblingCategoryViewModel : SiblingCategoryViewModels)
+		{
+			CategoryNames.Add(Cast<UNiagaraHierarchyCategory>(SiblingCategoryViewModel->GetDataMutable())->GetCategoryName().ToString());
+		}
+
+		if(CategoryNames.Contains(NewName.ToString()))
+		{
+			OutTooltip = LOCTEXT("HierarchyCategoryCantRename_DuplicateOnLayer", "Another category of the same name already exists!");
+			return false;
+		}
+	}
+
+	return true;
+}
+
+FText SNiagaraHierarchyCategory::GetCategoryText() const
+{
+	return FText::FromString(CategoryViewModel.Pin()->ToString());
+}
+
+void SNiagaraHierarchyCategory::OnRenameCategory(const FText& NewText, ETextCommit::Type) const
+{
+	FScopedTransaction Transaction(LOCTEXT("Transaction_Rename_Category", "Renamed hierarchy category"));
+	CategoryViewModel.Pin()->GetHierarchyViewModel()->GetHierarchyRoot()->Modify();
+	
+	CategoryViewModel.Pin()->Rename(FName(NewText.ToString()));
+}
+
+void SNiagaraHierarchySection::Construct(const FArguments& InArgs, TSharedPtr<FNiagaraHierarchySectionViewModel> InSection, TWeakObjectPtr<UNiagaraHierarchyViewModelBase> InHierarchyViewModel)
+{
+	SectionViewModel = InSection;
+	HierarchyViewModel = InHierarchyViewModel;
+
+	IsSectionActive = InArgs._IsSectionActive;
+	OnSectionActivatedDelegate = InArgs._OnSectionActivated;
+	bForbidDropOn = InArgs._bForbidDropOn;
+	
+	if(SectionViewModel != nullptr)
+	{
+		SectionViewModel->GetOnRequestRename().BindSP(this, &SNiagaraHierarchySection::TryEnterEditingMode);
+
+		SDropTarget::FArguments LeftDropTargetArgs;
+		SDropTarget::FArguments OntoDropTargetArgs;
+		SDropTarget::FArguments RightDropTargetArgs;
+
+		LeftDropTargetArgs
+			.OnAllowDrop(this, &SNiagaraHierarchySection::OnCanAcceptDrop, EItemDropZone::AboveItem)
+			.OnDropped(this, &SNiagaraHierarchySection::OnDroppedOn, EItemDropZone::AboveItem)
+			.VerticalImage(FAppStyle::GetNoBrush())
+			.HorizontalImage(FNiagaraEditorStyle::Get().GetBrush("NiagaraEditor.Stack.DropTarget.BorderHorizontal"));
+
+		OntoDropTargetArgs
+			.OnAllowDrop(this, &SNiagaraHierarchySection::OnCanAcceptDrop, EItemDropZone::OntoItem)
+			.OnDropped(this, &SNiagaraHierarchySection::OnDroppedOn, EItemDropZone::OntoItem);
+		
+		RightDropTargetArgs
+			.OnAllowDrop(this, &SNiagaraHierarchySection::OnCanAcceptDrop, EItemDropZone::BelowItem)
+			.OnDropped(this, &SNiagaraHierarchySection::OnDroppedOn, EItemDropZone::BelowItem)
+			.VerticalImage(FAppStyle::GetNoBrush())
+			.HorizontalImage(FNiagaraEditorStyle::Get().GetBrush("NiagaraEditor.Stack.DropTarget.BorderHorizontal"));
+
+		SetToolTipText(TAttribute<FText>::CreateSP(this, &SNiagaraHierarchySection::GetTooltipText));
+
+		ChildSlot
+		[
+			SNew(SHorizontalBox)
+			+ SHorizontalBox::Slot()
+			.AutoWidth()
+			[
+				SNew(SNiagaraSectionDragDropTarget, SectionViewModel, EItemDropZone::AboveItem)
+				.DropTargetArgs(LeftDropTargetArgs)
+			]
+			+ SHorizontalBox::Slot()
+			[
+				SNew(SNiagaraSectionDragDropTarget, SectionViewModel, EItemDropZone::OntoItem)
+				.DropTargetArgs(OntoDropTargetArgs
+				.Content()
+				[
+					SAssignNew(CheckBox, SCheckBox)
+					.Visibility(EVisibility::HitTestInvisible)
+					.Style(FAppStyle::Get(), "DetailsView.SectionButton")
+					.OnCheckStateChanged(this, &SNiagaraHierarchySection::OnSectionCheckChanged)
+					.IsChecked(this, &SNiagaraHierarchySection::GetSectionCheckState)
+					.Padding(FMargin(8.f, 4.f))
+					[
+						SNew(SHorizontalBox)
+						+ SHorizontalBox::Slot()
+						.AutoWidth()
+						.Padding(2.f)
+						[
+							SNew(SImage).Image(SectionViewModel->GetSectionImage() != nullptr ? SectionViewModel->GetSectionImage() : FStyleDefaults::GetNoBrush())
+						]
+						+ SHorizontalBox::Slot()
+						[
+							SAssignNew(InlineEditableTextBlock, SInlineEditableTextBlock)
+							.Visibility(EVisibility::HitTestInvisible)
+							.Text(this, &SNiagaraHierarchySection::GetText)
+							.OnTextCommitted(this, &SNiagaraHierarchySection::OnRenameSection)
+							.OnVerifyTextChanged(this, &SNiagaraHierarchySection::OnVerifySectionRename)
+							.IsSelected(this, &SNiagaraHierarchySection::IsSectionSelected)
+							.IsReadOnly(this, &SNiagaraHierarchySection::IsSectionReadOnly)
+						]
+					]
+				])
+			]
+			+ SHorizontalBox::Slot()
+			.AutoWidth()
+			[
+				SNew(SNiagaraSectionDragDropTarget, SectionViewModel, EItemDropZone::BelowItem)
+				.DropTargetArgs(RightDropTargetArgs)				
+			]
+		];		
+	}
+	// if this section doesn't represent data, it's the "All" widget
+	else
+	{
+		ChildSlot
+		[
+			SNew(SDropTarget)
+			.OnAllowDrop(this, &SNiagaraHierarchySection::OnCanAcceptDrop, EItemDropZone::OntoItem)
+			.OnDropped(this, &SNiagaraHierarchySection::OnDroppedOn, EItemDropZone::OntoItem)
+			[
+				SNew(SCheckBox)
+				.Style(FAppStyle::Get(), "DetailsView.SectionButton")
+				.OnCheckStateChanged(this, &SNiagaraHierarchySection::OnSectionCheckChanged)
+				.IsChecked(this, &SNiagaraHierarchySection::GetSectionCheckState)
+				.Padding(FMargin(8.f, 4.f))
+				[
+					SNew(STextBlock)
+					.Text(LOCTEXT("AllSection", "All"))
+				]
+			]			
+		];
+	}
+}
+
+SNiagaraHierarchySection::~SNiagaraHierarchySection()
+{
+	SectionViewModel.Reset();
+}
+
+void SNiagaraHierarchySection::TryEnterEditingMode() const
+{
+	if(SectionViewModel.IsValid() && SectionViewModel->CanRename())
+	{
+		InlineEditableTextBlock->EnterEditingMode();
+	}
+}
+
+TSharedPtr<FNiagaraHierarchySectionViewModel> SNiagaraHierarchySection::GetSectionViewModel()
+{
+	return SectionViewModel;
+}
+
+bool SNiagaraHierarchySection::OnCanAcceptDrop(TSharedPtr<FDragDropOperation> DragDropOperation, EItemDropZone ItemDropZone) const
+{
+	if(bForbidDropOn)
+	{
+		return false;
+	}
+
+	if(DragDropOperation->IsOfType<FNiagaraHierarchyDragDropOp>())
+	{
+		TSharedPtr<FNiagaraHierarchyDragDropOp> HierarchyDragDropOp = StaticCastSharedPtr<FNiagaraHierarchyDragDropOp>(DragDropOperation);
+		TSharedPtr<FNiagaraHierarchyItemViewModelBase> DraggedItem = HierarchyDragDropOp->GetDraggedItem().Pin();
+		
+		if(SectionViewModel.IsValid())
+		{
+			return SectionViewModel->CanDropOn(DraggedItem, ItemDropZone).bCanPerform;
+		}
+		
+		// for the All section which has no valid view model, we simply do a check if the sections of the dragged categories are different
+		if(DragDropOperation->IsOfType<FNiagaraHierarchyDragDropOp>())
+		{
+			if(const UNiagaraHierarchyCategory* Category = Cast<UNiagaraHierarchyCategory>(DraggedItem->GetData()))
+			{
+				return TryGetSectionData() != Category->GetSection();
+			}
+		}
+	}
+
+	return false;
+}
+
+FReply SNiagaraHierarchySection::OnDroppedOn(const FGeometry&, const FDragDropEvent& DragDropEvent, EItemDropZone DropZone) const
+{
+	bDraggedOn = false;
+	if(TSharedPtr<FNiagaraHierarchyDragDropOp> HierarchyDragDropOp = DragDropEvent.GetOperationAs<FNiagaraHierarchyDragDropOp>())
+	{
+		TSharedPtr<FNiagaraHierarchyItemViewModelBase> DraggedItem = HierarchyDragDropOp->GetDraggedItem().Pin();
+		if(SectionViewModel.IsValid())
+		{
+			SectionViewModel->OnDroppedOn(DraggedItem, DropZone);
+			return FReply::Handled();
+		}
+		
+		if(UNiagaraHierarchyCategory* Category = Cast<UNiagaraHierarchyCategory>(DraggedItem->GetDataMutable()))
+		{
+			Category->SetSection(nullptr);
+			TArray<UNiagaraHierarchyCategory*> ChildrenCategories;
+			Category->GetChildrenOfType<>(ChildrenCategories, true);
+
+			for(UNiagaraHierarchyCategory* ChildCategory : ChildrenCategories)
+			{
+				ChildCategory->SetSection(nullptr);
+			}
+
+			// we only need to reparent if the parent isn't already the root. This stops unnecessary reordering
+			if(HierarchyDragDropOp->GetDraggedItem().Pin()->GetParent() != HierarchyViewModel->GetHierarchyRootViewModel())
+			{
+				HierarchyViewModel->GetHierarchyRootViewModel()->ReparentToThis(HierarchyDragDropOp->GetDraggedItem().Pin());
+			}
+		
+			HierarchyViewModel->RefreshHierarchyView();
+
+			return FReply::Handled();
+		}
+	}
+
+	return FReply::Unhandled();
+}
+
+FReply SNiagaraHierarchySection::OnMouseButtonDown(const FGeometry& MyGeometry, const FPointerEvent& MouseEvent)
+{
+	if(SectionViewModel.IsValid())
+	{
+		// we handle the event here so we can react on mouse button up
+		if(MouseEvent.IsMouseButtonDown(EKeys::RightMouseButton))
+		{
+			return FReply::Handled();
+		}
+		else if(MouseEvent.IsMouseButtonDown(EKeys::LeftMouseButton))
+		{			
+			OnSectionActivatedDelegate.ExecuteIfBound(SectionViewModel);
+			return FReply::Handled().DetectDrag(AsShared(), EKeys::LeftMouseButton).SetUserFocus(AsShared());	
+		}
+	}
+
+	return FReply::Unhandled();
+}
+
+FReply SNiagaraHierarchySection::OnMouseButtonUp(const FGeometry& MyGeometry, const FPointerEvent& MouseEvent)
+{
+	if(SectionViewModel.IsValid() && SectionViewModel->IsForHierarchy())
+	{
+		if(MouseEvent.GetEffectingButton() == EKeys::RightMouseButton)
+		{			
+			// Show menu to choose getter vs setter
+			FSlateApplication::Get().PushMenu(
+				AsShared(),
+				FWidgetPath(),
+				SummonContextMenu({GetSectionViewModel()}, true, AsShared()),
+				FSlateApplication::Get().GetCursorPos(),
+				FPopupTransitionEffect(FPopupTransitionEffect::ContextMenu)
+			);
+				
+			OnSectionActivatedDelegate.ExecuteIfBound(SectionViewModel);			
+			return FReply::Handled();
+		}
+	}
+
+	return FReply::Unhandled();
+}
+
+void SNiagaraHierarchySection::OnDragEnter(const FGeometry& MyGeometry, const FDragDropEvent& DragDropEvent)
+{
+	bDraggedOn = true;
+	if(DragDropEvent.GetOperationAs<FNiagaraHierarchyDragDropOp>() && DragDropEvent.GetOperationAs<FNiagaraSectionDragDropOp>() == nullptr)
+	{
+		RegisterActiveTimer(1.0f, FWidgetActiveTimerDelegate::CreateSP(this, &SNiagaraHierarchySection::ActivateSectionIfDragging));
+	}
+}
+
+FReply SNiagaraHierarchySection::OnDragDetected(const FGeometry& MyGeometry, const FPointerEvent& MouseEvent)
+{
+	if(MouseEvent.IsMouseButtonDown(EKeys::LeftMouseButton) && SectionViewModel.IsValid() && SectionViewModel->CanDrag().bCanPerform)
+	{
+		TSharedRef<FNiagaraSectionDragDropOp> SectionDragDropOp = MakeShared<FNiagaraSectionDragDropOp>(SectionViewModel);
+		SectionDragDropOp->Construct();
+		return FReply::Handled().BeginDragDrop(SectionDragDropOp);
+	}
+	
+	return FReply::Unhandled();
+}
+
+void SNiagaraHierarchySection::OnDragLeave(const FDragDropEvent& DragDropEvent)
+{
+	bDraggedOn = false;
+	if(TSharedPtr<FNiagaraHierarchyDragDropOp> HierarchyDragDropOp = DragDropEvent.GetOperationAs<FNiagaraHierarchyDragDropOp>())
+	{
+		HierarchyDragDropOp->SetDescription(FText::GetEmpty());
+	}
+}
+
+FReply SNiagaraHierarchySection::OnKeyDown(const FGeometry& MyGeometry, const FKeyEvent& InKeyEvent)
+{
+	if(SectionViewModel.IsValid() && InKeyEvent.GetKey() == EKeys::Delete && SectionViewModel->CanDelete())
+	{
+		SectionViewModel->Delete();
+		return FReply::Handled();
+	}
+
+	return FReply::Unhandled();
+}
+
+void SNiagaraHierarchySection::OnMouseLeave(const FPointerEvent& MouseEvent)
+{
+	//bDraggedOn = false;
+	SCompoundWidget::OnMouseLeave(MouseEvent);
+}
+
+UNiagaraHierarchySection* SNiagaraHierarchySection::TryGetSectionData() const
+{
+	return SectionViewModel.IsValid() ? Cast<UNiagaraHierarchySection>(SectionViewModel->GetDataMutable()) : nullptr;
+}
+
+FText SNiagaraHierarchySection::GetText() const
+{
+	return SectionViewModel->GetSectionNameAsText();
+}
+
+FText SNiagaraHierarchySection::GetTooltipText() const
+{
+	return SectionViewModel->GetSectionTooltip();
+}
+
+void SNiagaraHierarchySection::OnRenameSection(const FText& Text, ETextCommit::Type CommitType) const
+{
+	FScopedTransaction Transaction(LOCTEXT("Transaction_Rename_Section", "Renamed hierarchy section"));
+	HierarchyViewModel->GetHierarchyRoot()->Modify();
+	
+	SectionViewModel->Rename(FName(Text.ToString()));
+}
+
+bool SNiagaraHierarchySection::OnVerifySectionRename(const FText& NewName, FText& OutTooltip) const
+{
+	// this function shouldn't be used in case the section isn't valid but we'll make sure regardless
+	if(!SectionViewModel.IsValid())
+	{
+		return false;
+	}
+
+	if(SectionViewModel->GetSectionName().ToString() != NewName.ToString())
+	{
+		TArray<FString> SectionNames;
+
+		SectionNames.Add("All");
+		for(auto& Section : HierarchyViewModel->GetHierarchyRootViewModel()->GetSectionViewModels())
+		{
+			SectionNames.Add(Section->GetSectionName().ToString());
+		}
+
+		if(SectionNames.Contains(NewName.ToString()))
+		{
+			OutTooltip = LOCTEXT("HierarchySectionCantRename_Duplicate", "A section with that name already exists!");
+			return false;
+		}
+	}
+
+	return true;
+}
+
+bool SNiagaraHierarchySection::IsSectionSelected() const
+{
+	return GetSectionCheckState() == ECheckBoxState::Checked ? true : false;
+}
+
+bool SNiagaraHierarchySection::IsSectionReadOnly() const
+{
+	return SectionViewModel.IsValid() ? !SectionViewModel->CanRename() : true;
+}
+
+ECheckBoxState SNiagaraHierarchySection::GetSectionCheckState() const
+{
+	return IsSectionActive.Get();
+}
+
+void SNiagaraHierarchySection::OnSectionCheckChanged(ECheckBoxState NewState)
+{
+	OnSectionActivatedDelegate.ExecuteIfBound(SectionViewModel);
+}
+
+EActiveTimerReturnType SNiagaraHierarchySection::ActivateSectionIfDragging(double CurrentTime, float DeltaTime) const
+{
+	if(bDraggedOn && FSlateApplication::Get().GetDragDroppingContent().IsValid() && FSlateApplication::Get().GetDragDroppingContent()->IsOfType<FNiagaraHierarchyDragDropOp>())
+	{
+		if(IsSectionSelected() == false)
+		{
+			OnSectionActivatedDelegate.ExecuteIfBound(SectionViewModel);
+		}
+	}
+
+	return EActiveTimerReturnType::Stop;
+}
+
+void SNiagaraHierarchyEditor::Construct(const FArguments& InArgs, TObjectPtr<UNiagaraHierarchyViewModelBase> InHierarchyViewModel)
+{
+	HierarchyViewModel = InHierarchyViewModel;
+	
+	SourceRoot = NewObject<UNiagaraHierarchyRoot>(GetTransientPackage(), MakeUniqueObjectName(GetTransientPackage(), UNiagaraHierarchyRoot::StaticClass()));
+	SourceRootViewModel = HierarchyViewModel->CreateRootViewModelForData(SourceRoot, false);
+	SourceRootViewModel->Initialize();
+	SourceRootViewModel->AddChildFilter(FNiagaraHierarchyItemViewModelBase::FOnFilterChild::CreateSP(this, &SNiagaraHierarchyEditor::FilterForSourceSection));
+	SourceRootViewModel->OnSyncPropagated().BindSP(this, &SNiagaraHierarchyEditor::RequestRefreshSourceViewNextFrame, false);
+	SourceRootViewModel->OnSectionsChanged().BindSP(this, &SNiagaraHierarchyEditor::RefreshSectionsView);
+
+	HierarchyViewModel->OnInitialized().BindSP(this, &SNiagaraHierarchyEditor::Reinitialize);
+	HierarchyViewModel->OnNavigateToItemInHierarchyRequested().BindSP(this, &SNiagaraHierarchyEditor::NavigateToHierarchyItem);
+	HierarchyViewModel->OnRefreshSourceItemsRequested().BindSP(this, &SNiagaraHierarchyEditor::RefreshSourceItems);
+	HierarchyViewModel->OnRefreshViewRequested().BindSP(this, &SNiagaraHierarchyEditor::RefreshAllViews);
+	HierarchyViewModel->OnRefreshSourceView().BindSP(this, &SNiagaraHierarchyEditor::RefreshSourceView);
+	HierarchyViewModel->OnRefreshHierarchyView().BindSP(this, &SNiagaraHierarchyEditor::RefreshHierarchyView);
+	HierarchyViewModel->OnRefreshSectionsView().BindSP(this, &SNiagaraHierarchyEditor::RefreshSectionsView);
+	HierarchyViewModel->OnHierarchySectionActivated().BindSP(this, &SNiagaraHierarchyEditor::OnHierarchySectionActivated);
+	HierarchyViewModel->OnItemAdded().BindSP(this, &SNiagaraHierarchyEditor::OnItemAdded);
+
+	BindToHierarchyRootViewModel();
+	
+	OnGenerateRowContentWidget = InArgs._OnGenerateRowContentWidget;
+	OnGenerateCustomDetailsPanelNameWidget = InArgs._OnGenerateCustomDetailsPanelNameWidget;
+	CategoryRowStyle = InArgs._CategoryRowStyle;
+	ItemRowStyle = InArgs._ItemRowStyle;
+	
+	SSplitter::FSlot* DetailsPanelSlot = nullptr;
+
+	TSharedRef<SWidget> AddSectionButton = SNew(SHorizontalBox)
+		+ SHorizontalBox::Slot()
+		.AutoWidth()
+		[
+			SNew(SButton)
+			.OnClicked_Lambda([this]()
+			{
+				HierarchyViewModel.Get()->AddSection();
+				return FReply::Handled();
+			})
+			.ButtonStyle(FNiagaraEditorStyle::Get(), "NiagaraEditor.HierarchyEditor.ButtonStyle")
+			[
+				SNew(SHorizontalBox)
+				+ SHorizontalBox::Slot()
+				.AutoWidth()
+				.Padding(2.f)
+				[
+					SNew(SImage)
+					.Image(FAppStyle::Get().GetBrush("Icons.PlusCircle"))
+				]
+				+ SHorizontalBox::Slot()
+				.VAlign(VAlign_Center)
+				.AutoWidth()
+				[			
+					SNew(STextBlock)
+					.Text(LOCTEXT("AddSectionLabel","Add Section"))				
+				]
+			]
+		];
+
+	TSharedRef<SWidget> AddCategoryButton = SNew(SHorizontalBox)
+		+ SHorizontalBox::Slot()
+		.AutoWidth()
+		[
+			SNew(SButton)
+			.OnClicked_Lambda([this]()
+			{
+				TArray<TSharedPtr<FNiagaraHierarchyItemViewModelBase>> SelectedItems;
+				HierarchyTreeView->GetSelectedItems(SelectedItems);
+
+				if(SelectedItems.Num() == 1 && SelectedItems[0]->GetData()->IsA<UNiagaraHierarchyCategory>())
+				{
+					HierarchyViewModel.Get()->AddCategory(SelectedItems[0]);
+				}
+				else
+				{
+					HierarchyViewModel.Get()->AddCategory(HierarchyViewModel->GetHierarchyRootViewModel());
+				}
+				return FReply::Handled();
+			})
+			.ButtonStyle(FNiagaraEditorStyle::Get(), "NiagaraEditor.HierarchyEditor.ButtonStyle")
+			[
+				SNew(SHorizontalBox)
+				+ SHorizontalBox::Slot()
+				.AutoWidth()
+				.Padding(2.f)
+				[
+					SNew(SImage)
+					.Image(FAppStyle::Get().GetBrush("Icons.PlusCircle"))
+				]
+				+ SHorizontalBox::Slot()
+				.VAlign(VAlign_Center)
+				.AutoWidth()
+				[			
+					SNew(STextBlock)
+					.Text(LOCTEXT("AddCategoryLabel","Add Category"))				
+				]
+			]
+		];
+	
+	ChildSlot
+	[
+		SNew(SBorder)
+		.Padding(0.f)
+		.BorderImage(FAppStyle::Get().GetBrush("Brushes.Panel"))
+		[
+			SNew(SSplitter)
+			.Orientation(Orient_Horizontal)
+			.PhysicalSplitterHandleSize(2.f)
+			+ SSplitter::Slot()
+			.Value(0.3f)
+			.MinSize(0.1f)
+			[
+				SNew(SVerticalBox)
+				+ SVerticalBox::Slot()
+				.AutoHeight()
+				.Padding(2.f)
+				[
+					SAssignNew(SourceSearchBox, SSearchBox)
+					.OnTextChanged(this, &SNiagaraHierarchyEditor::OnSourceSearchTextChanged)
+					.OnTextCommitted(this, &SNiagaraHierarchyEditor::OnSourceSearchTextCommitted)
+					.OnSearch(SSearchBox::FOnSearch::CreateSP(this, &SNiagaraHierarchyEditor::OnSearchButtonClicked))
+					.DelayChangeNotificationsWhileTyping(true)
+					.SearchResultData(this, &SNiagaraHierarchyEditor::GetSearchResultData)
+				]
+				+ SVerticalBox::Slot()
+				.AutoHeight()
+				.Padding(2.f)
+				[
+					SAssignNew(SourceSectionBox, SWrapBox)
+					.UseAllottedSize(true)
+				]
+				+ SVerticalBox::Slot()
+				.Padding(1.f, 2.f)
+				[
+					SAssignNew(SourceTreeView, STreeView<TSharedPtr<FNiagaraHierarchyItemViewModelBase>>)
+					.TreeItemsSource(&GetSourceItems())
+					.OnSelectionChanged(this, &SNiagaraHierarchyEditor::OnSelectionChanged, false)
+					.OnGenerateRow(this, &SNiagaraHierarchyEditor::GenerateSourceItemRow)
+					.OnGetChildren_UObject(InHierarchyViewModel.Get(), &UNiagaraHierarchyViewModelBase::OnGetChildren)
+					.OnItemToString_Debug_UObject(InHierarchyViewModel.Get(), &UNiagaraHierarchyViewModelBase::OnItemToStringDebug)
+					.OnContextMenuOpening(this, &SNiagaraHierarchyEditor::SummonContextMenuForSelectedRows, false)
+				]
+			]
+			+ SSplitter::Slot()
+			.Value(0.4f)
+			.MinSize(0.1f)
+			[
+				SNew(SVerticalBox)
+				+ SVerticalBox::Slot()
+				.AutoHeight()
+				[
+					SNew(SVerticalBox)
+					+ SVerticalBox::Slot()
+					.AutoHeight()
+					[
+						SNew(SBorder)
+						.Padding(0.f)
+						.BorderImage(FAppStyle::Get().GetBrush("Brushes.Header"))
+						[
+							AddSectionButton
+						]
+					]
+					+ SVerticalBox::Slot()
+					.AutoHeight()
+					.Padding(1.f)
+					[
+						SAssignNew(HierarchySectionBox, SWrapBox)
+						.UseAllottedSize(true)
+					]				
+				]
+				+ SVerticalBox::Slot()
+				[
+					SNew(SVerticalBox)
+					+ SVerticalBox::Slot()
+					.AutoHeight()
+					[
+						SNew(SBorder)
+						.Padding(0.f)
+						.BorderImage(FAppStyle::Get().GetBrush("Brushes.Header"))
+						[
+							AddCategoryButton
+						]
+					]
+					+ SVerticalBox::Slot()
+					.FillHeight(0.1f)
+					.Padding(1.f, 4.f, 1.f, 0.f)
+					[
+						SNew(SDropTarget)
+						.OnDropped(this, &SNiagaraHierarchyEditor::HandleHierarchyRootDrop)
+						.OnAllowDrop(this, &SNiagaraHierarchyEditor::OnCanDropOnRoot)
+						.OnDragEnter(this, &SNiagaraHierarchyEditor::OnRootDragEnter)
+						.OnDragLeave(this, &SNiagaraHierarchyEditor::OnRootDragLeave)
+						[
+							SNew(SBorder)
+							.Padding(0.f)
+							.BorderImage(FAppStyle::Get().GetBrush("Brushes.Recessed"))
+							[
+								SNew(SBorder)
+								.Padding(1.f)
+								.BorderImage(FAppStyle::GetBrush("DashedBorder"))
+								.BorderBackgroundColor(FLinearColor(0.2f, 0.2f, 0.2f, 0.5f))
+								[
+									SNew(SBox)
+									.HAlign(HAlign_Center)
+									.VAlign(VAlign_Center)
+									[
+										SNew(SImage)
+										.Image(FNiagaraEditorStyle::Get().GetBrush("NiagaraEditor.HierarchyEditor.RootDropIcon"))
+										.ColorAndOpacity(this, &SNiagaraHierarchyEditor::GetRootIconColor)
+									]
+								]
+							]
+						]
+					]
+					+ SVerticalBox::Slot()
+					.Padding(1.f, 0.f)
+					[
+						SAssignNew(HierarchyTreeView, STreeView<TSharedPtr<FNiagaraHierarchyItemViewModelBase>>)
+						.TreeItemsSource(&InHierarchyViewModel->GetHierarchyItems())
+						.OnSelectionChanged(this, &SNiagaraHierarchyEditor::OnSelectionChanged, true)
+						.OnGenerateRow(this, &SNiagaraHierarchyEditor::GenerateHierarchyItemRow)
+						.OnGetChildren_UObject(InHierarchyViewModel.Get(), &UNiagaraHierarchyViewModelBase::OnGetChildren)
+						.OnItemToString_Debug_UObject(InHierarchyViewModel.Get(), &UNiagaraHierarchyViewModelBase::OnItemToStringDebug)
+						.OnContextMenuOpening(this, &SNiagaraHierarchyEditor::SummonContextMenuForSelectedRows, true)
+					]
+				]
+			]
+			+ SSplitter::Slot()
+			.Value(0.3f)
+			.MinSize(0.1f)
+			.Expose(DetailsPanelSlot)		
+		]
+	];
+
+	if(InHierarchyViewModel->SupportsDetailsPanel())
+	{
+		FPropertyEditorModule& PropertyEditorModule = FModuleManager::Get().LoadModuleChecked<FPropertyEditorModule>("PropertyEditor");
+
+		FDetailsViewArgs DetailsViewArgs;
+		DetailsViewArgs.NameAreaSettings = FDetailsViewArgs::ENameAreaSettings::ObjectsUseNameArea;
+		DetailsViewArgs.bShowObjectLabel = false;
+		DetailsViewArgs.bAllowSearch = false;
+		DetailsViewArgs.NotifyHook = this;
+
+		DetailsPanel = PropertyEditorModule.CreateDetailView(DetailsViewArgs);
+		
+		if(OnGenerateCustomDetailsPanelNameWidget.IsBound())
+		{
+			TSharedRef<SWidget> CustomDetailsPanelNameWidget = OnGenerateCustomDetailsPanelNameWidget.Execute(nullptr);
+			DetailsPanel->SetNameAreaCustomContent(CustomDetailsPanelNameWidget);
+		}
+
+		DetailsPanel->SetIsPropertyEditingEnabledDelegate(FIsPropertyEditingEnabled::CreateSP(this, &SNiagaraHierarchyEditor::IsDetailsPanelEditingAllowed));
+		
+		for(auto& Customizations : HierarchyViewModel->GetInstanceCustomizations())
+		{
+			DetailsPanel->RegisterInstancedCustomPropertyLayout(Customizations.Key, Customizations.Value);
+		}
+		
+		DetailsPanelSlot->AttachWidget(DetailsPanel.ToSharedRef());
+	}
+		
+	HierarchyViewModel->GetCommands()->MapAction(FGenericCommands::Get().Rename,
+		FExecuteAction::CreateSP(this, &SNiagaraHierarchyEditor::RequestRenameSelectedItem),
+		FCanExecuteAction::CreateSP(this, &SNiagaraHierarchyEditor::CanRequestRenameSelectedItem));
+
+	HierarchyViewModel->GetCommands()->MapAction(FGenericCommands::Get().Delete,
+		FExecuteAction::CreateSP(this, &SNiagaraHierarchyEditor::DeleteSelectedHierarchyItems),
+		FCanExecuteAction::CreateSP(this, &SNiagaraHierarchyEditor::CanDeleteSelectedHierarchyItems),
+		FIsActionChecked(), FIsActionButtonVisible::CreateSP(this, &SNiagaraHierarchyEditor::CanDeleteSelectedHierarchyItems));
+
+	// HierarchyViewModel->GetCommands()->MapAction(FNiagaraHierarchyEditorCommands::Get().RenameSection,
+	// 	FExecuteAction::CreateSP(this, &SNiagaraHierarchy::DeleteActiveSection),
+	// 	FCanExecuteAction::CreateSP(this, &SNiagaraHierarchy::CanDeleteActiveSection),
+	// 	FIsActionChecked(), FIsActionButtonVisible::CreateSP(this, &SNiagaraHierarchy::CanDeleteActiveSection));
+	
+	HierarchyViewModel->GetCommands()->MapAction(FNiagaraHierarchyEditorCommands::Get().DeleteSection,
+		FExecuteAction::CreateSP(this, &SNiagaraHierarchyEditor::DeleteActiveSection),
+		FCanExecuteAction::CreateSP(this, &SNiagaraHierarchyEditor::CanDeleteActiveSection),
+		FIsActionChecked(), FIsActionButtonVisible::CreateSP(this, &SNiagaraHierarchyEditor::CanDeleteActiveSection));
+	
+	HierarchyViewModel->ForceFullRefresh();
+}
+
+SNiagaraHierarchyEditor::~SNiagaraHierarchyEditor()
+{
+	SourceSearchResults.Empty();
+	FocusedSearchResult.Reset();
+	
+	ClearSourceItems();
+	
+	if(HierarchyViewModel.IsValid())
+	{
+		HierarchyViewModel->OnInitialized().Unbind();
+		HierarchyViewModel->OnNavigateToItemInHierarchyRequested().Unbind();
+		HierarchyViewModel->OnRefreshSourceItemsRequested().Unbind();
+		HierarchyViewModel->OnRefreshViewRequested().Unbind();
+		HierarchyViewModel->OnRefreshSourceView().Unbind();
+		HierarchyViewModel->OnRefreshHierarchyView().Unbind();
+		HierarchyViewModel->OnRefreshSectionsView().Unbind();
+		HierarchyViewModel->OnHierarchySectionActivated().Unbind();
+		HierarchyViewModel->OnItemAdded().Unbind();
+
+		UnbindFromHierarchyRootViewModel();
+	
+		HierarchyViewModel->GetCommands()->UnmapAction(FGenericCommands::Get().Delete);
+		HierarchyViewModel->GetCommands()->UnmapAction(FGenericCommands::Get().Rename);			
+	}
+	
+	SourceRootViewModel->OnSyncPropagated().Unbind();
+	SourceRootViewModel->OnSectionsChanged().Unbind();
+	SourceRootViewModel.Reset();
+	SourceRoot->ConditionalBeginDestroy();
+	SourceRoot = nullptr;
+}
+
+void SNiagaraHierarchyEditor::RefreshSourceItems()
+{
+	HierarchyViewModel->PrepareSourceItems(SourceRoot, SourceRootViewModel);
+	SourceRootViewModel->SyncViewModelsToData();
+	RefreshSourceView(false);
+	RefreshSectionsView();
+}
+
+void SNiagaraHierarchyEditor::RefreshAllViews(bool bFullRefresh)
+{
+	RefreshSourceView(bFullRefresh);
+	RefreshHierarchyView(bFullRefresh);
+	RefreshSectionsView();
+}
+
+void SNiagaraHierarchyEditor::RequestRefreshAllViewsNextFrame(bool bFullRefresh)
+{
+	RequestRefreshSourceViewNextFrame(bFullRefresh);
+	RequestRefreshHierarchyViewNextFrame(bFullRefresh);
+	RequestRefreshSectionsViewNextFrame();
+}
+
+FReply SNiagaraHierarchyEditor::OnKeyDown(const FGeometry& MyGeometry, const FKeyEvent& InKeyEvent)
+{
+	return HierarchyViewModel->GetCommands()->ProcessCommandBindings(InKeyEvent) ? FReply::Handled() : FReply::Unhandled();
+}
+
+FReply SNiagaraHierarchyEditor::OnMouseButtonDown(const FGeometry& MyGeometry, const FPointerEvent& MouseEvent)
+{
+	// we catch any mouse button down event so that we can continue using our commands
+	return FReply::Handled().SetUserFocus(AsShared(), EFocusCause::Mouse, true);
+}
+
+FReply SNiagaraHierarchyEditor::OnMouseButtonUp(const FGeometry& MyGeometry, const FPointerEvent& MouseEvent)
+{
+	return FReply::Handled().SetUserFocus(AsShared(), EFocusCause::Mouse, true);
+}
+
+TSharedPtr<SWidget> SNiagaraHierarchyEditor::SummonContextMenuForSelectedRows(bool bFromHierarchy) const
+{
+	TArray<TSharedPtr<FNiagaraHierarchyItemViewModelBase>> ViewModels;
+	if(bFromHierarchy)
+	{
+		HierarchyTreeView->GetSelectedItems(ViewModels);
+	}
+	else
+	{
+		SourceTreeView->GetSelectedItems(ViewModels);
+	}
+
+	if(ViewModels.IsEmpty())
+	{
+		return nullptr;
+	}
+	
+	return SummonContextMenu(ViewModels, bFromHierarchy, AsShared());
+}
+
+void SNiagaraHierarchyEditor::RefreshSourceView(bool bFullRefresh) const
+{
+	SourceTreeView->SetTreeItemsSource(&GetSourceItems());
+	if(bFullRefresh)
+	{
+		SourceTreeView->RebuildList();
+	}
+	else
+	{
+		SourceTreeView->RequestTreeRefresh();
+	}
+}
+
+void SNiagaraHierarchyEditor::RequestRefreshSourceViewNextFrame(bool bFullRefresh)
+{
+	if(!RefreshSourceViewNextFrameHandle.IsValid())
+	{
+		RefreshSourceViewNextFrameHandle = RegisterActiveTimer(0.f, FWidgetActiveTimerDelegate::CreateLambda([this, bFullRefresh](double CurrentTime, float DeltaTime)
+		{
+			RefreshSourceView(bFullRefresh);
+			RefreshSourceViewNextFrameHandle.Reset();
+			return EActiveTimerReturnType::Stop;
+		}));
+	}
+}
+
+void SNiagaraHierarchyEditor::RefreshHierarchyView(bool bFullRefresh) const
+{
+	// the top layer objects might have changed due to filtering. We need to refresh these too.
+	HierarchyTreeView->SetTreeItemsSource(&HierarchyViewModel->GetHierarchyItems());
+	if(bFullRefresh)
+	{
+		HierarchyTreeView->RebuildList();
+	}
+	else
+	{
+		HierarchyTreeView->RequestTreeRefresh();
+	}
+}
+
+void SNiagaraHierarchyEditor::RequestRefreshHierarchyViewNextFrame(bool bFullRefresh)
+{
+	if(!RefreshHierarchyViewNextFrameHandle.IsValid())
+	{
+		RefreshHierarchyViewNextFrameHandle = RegisterActiveTimer(0.f, FWidgetActiveTimerDelegate::CreateLambda([this, bFullRefresh](double CurrentTime, float DeltaTime)
+		{
+			RefreshHierarchyView(bFullRefresh);
+			RefreshHierarchyViewNextFrameHandle.Reset();
+			return EActiveTimerReturnType::Stop;
+		}));
+	}
+}
+
+void SNiagaraHierarchyEditor::RefreshSectionsView()
+{
+	SourceSectionBox->ClearChildren();
+	HierarchySectionBox->ClearChildren();
+
+	for (TSharedPtr<FNiagaraHierarchySectionViewModel>& SourceSection : SourceRootViewModel->GetSectionViewModels())
+	{
+		TSharedPtr<SNiagaraHierarchySection> SectionWidget = SNew(SNiagaraHierarchySection, SourceSection, HierarchyViewModel)
+		.IsSectionActive_Lambda([this, SourceSection]()
+		{
+			return GetActiveSourceSection() == SourceSection ? ECheckBoxState::Checked : ECheckBoxState::Unchecked;
+		})
+		.OnSectionActivated_Lambda([this](TSharedPtr<FNiagaraHierarchySectionViewModel> SectionViewModel)
+		{
+			SetActiveSourceSection(SectionViewModel);
+		});
+		
+		SourceSectionBox->AddSlot()
+		.HAlign(HAlign_Center)
+		.VAlign(VAlign_Center)
+		[
+			SectionWidget.ToSharedRef()		
+		];
+	}
+
+	if(SourceRootViewModel->GetSectionViewModels().Num() > 0)
+	{
+		TSharedPtr<SNiagaraHierarchySection> DefaultSourceSection = SNew(SNiagaraHierarchySection, nullptr, HierarchyViewModel)
+		.IsSectionActive_Lambda([this]()
+		{
+			return GetActiveSourceSection() == nullptr ? ECheckBoxState::Checked : ECheckBoxState::Unchecked;
+		})
+		.OnSectionActivated(this, &SNiagaraHierarchyEditor::SetActiveSourceSection)
+		// we forbid drop on here as the 'All' sections don't have a valid view model to determine it instead
+		.bForbidDropOn(true);
+
+		SourceSectionBox->AddSlot()
+		.HAlign(HAlign_Center)
+		.VAlign(VAlign_Center)
+		[
+			DefaultSourceSection.ToSharedRef()
+		];
+	}
+	
+	for (TSharedPtr<FNiagaraHierarchySectionViewModel>& HierarchySection : HierarchyViewModel->GetHierarchyRootViewModel()->GetSectionViewModels())
+	{
+		TSharedPtr<SNiagaraHierarchySection> SectionWidget = SNew(SNiagaraHierarchySection, HierarchySection, HierarchyViewModel)
+		.IsSectionActive_Lambda([this, HierarchySection]()
+		{
+			return HierarchyViewModel->GetActiveHierarchySection() == HierarchySection ? ECheckBoxState::Checked : ECheckBoxState::Unchecked;
+		})
+		.OnSectionActivated_Lambda([this](TSharedPtr<FNiagaraHierarchySectionViewModel> SectionViewModel)
+		{
+			HierarchyViewModel->SetActiveHierarchySection(SectionViewModel);
+		});
+		
+		HierarchySectionBox->AddSlot()
+		.HAlign(HAlign_Center)
+		.VAlign(VAlign_Center)
+		[
+			SectionWidget.ToSharedRef()		
+		];
+	}
+	
+	TSharedPtr<SNiagaraHierarchySection> DefaultHierarchySection = SNew(SNiagaraHierarchySection, nullptr, HierarchyViewModel)
+	.IsSectionActive_Lambda([this]()
+	{
+		return HierarchyViewModel->GetActiveHierarchySection() == nullptr ? ECheckBoxState::Checked : ECheckBoxState::Unchecked;
+	})
+	.OnSectionActivated_Lambda([this](TSharedPtr<FNiagaraHierarchySectionViewModel> SectionViewModel)
+	{
+		HierarchyViewModel->SetActiveHierarchySection(SectionViewModel);
+	});
+
+	HierarchySectionBox->AddSlot()
+	.HAlign(HAlign_Center)
+	.VAlign(VAlign_Center)
+	[
+		DefaultHierarchySection.ToSharedRef()
+	];
+}
+
+void SNiagaraHierarchyEditor::RequestRefreshSectionsViewNextFrame()
+{
+	if(!RefreshSectionsViewNextFrameHandle.IsValid())
+	{
+		RefreshSectionsViewNextFrameHandle = RegisterActiveTimer(0.f, FWidgetActiveTimerDelegate::CreateLambda([this](double CurrentTime, float DeltaTime)
+		{
+			RefreshSectionsView();
+			RefreshSectionsViewNextFrameHandle.Reset();
+			return EActiveTimerReturnType::Stop;
+		}));
+	}
+}
+
+void SNiagaraHierarchyEditor::NavigateToHierarchyItem(FNiagaraHierarchyIdentity Identity) const
+{
+	if(TSharedPtr<FNiagaraHierarchyItemViewModelBase> ViewModel = HierarchyViewModel->GetHierarchyRootViewModel()->FindViewModelForChild(Identity, true))
+	{
+		NavigateToHierarchyItem(ViewModel);
+	}
+}
+
+void SNiagaraHierarchyEditor::NavigateToHierarchyItem(TSharedPtr<FNiagaraHierarchyItemViewModelBase> Item) const
+{
+	TArray<TSharedPtr<FNiagaraHierarchyItemViewModelBase>> ParentChain;
+	for(TWeakPtr<FNiagaraHierarchyItemViewModelBase> Parent = Item->GetParent(); Parent.IsValid(); Parent = Parent.Pin()->GetParent())
+	{
+		ParentChain.Add(Parent.Pin());
+	}
+
+	for(int32 ParentIndex = ParentChain.Num()-1; ParentIndex >= 0; ParentIndex--)
+	{
+		HierarchyTreeView->SetItemExpansion(ParentChain[ParentIndex], true);
+	}
+		
+	HierarchyTreeView->SetSelection(Item);
+	HierarchyTreeView->RequestScrollIntoView(Item);
+}
+
+bool SNiagaraHierarchyEditor::IsItemSelected(TSharedPtr<FNiagaraHierarchyItemViewModelBase> Item) const
+{
+	return HierarchyTreeView->IsItemSelected(Item);
+}
+
+TSharedRef<ITableRow> SNiagaraHierarchyEditor::GenerateSourceItemRow(TSharedPtr<FNiagaraHierarchyItemViewModelBase> HierarchyItem, const TSharedRef<STableViewBase>& TableViewBase)
+{
+	return SNew(STableRow<TSharedPtr<FNiagaraHierarchyItemViewModelBase>>, TableViewBase)
+	.Style(HierarchyItem->GetData()->IsA<UNiagaraHierarchyItem>() ? ItemRowStyle : CategoryRowStyle)
+	.OnDragDetected(HierarchyItem.ToSharedRef(), &FNiagaraHierarchyItemViewModelBase::OnDragDetected, true)
+	.Padding(FMargin(2))
+	[
+		SNew(SHorizontalBox)
+		+ SHorizontalBox::Slot()
+		.AutoWidth()
+		.Padding(1.f)
+		.HAlign(HAlign_Center)
+		.VAlign(VAlign_Center)
+		[
+			SNew(SBox)
+			.HeightOverride(10.f)
+			.WidthOverride(10.f)
+			.Visibility_Lambda([HierarchyItem, Identity = HierarchyItem->GetData()->GetPersistentIdentity(), this]()
+			{
+				TArray<TSharedPtr<FNiagaraHierarchyItemViewModelBase>> AllChildren;
+				HierarchyItem->GetChildrenViewModelsForType<UNiagaraHierarchyItemBase, FNiagaraHierarchyItemViewModelBase>(AllChildren, true);
+
+				bool bCanDrag = HierarchyItem->GetHierarchyViewModel()->GetHierarchyRootViewModel()->FindViewModelForChild(HierarchyItem->GetData()->GetPersistentIdentity(), true) == nullptr;			
+
+				if(bCanDrag)
+				{
+					for(TSharedPtr<FNiagaraHierarchyItemViewModelBase>& ItemViewModel : AllChildren)
+					{
+						if(HierarchyItem->GetHierarchyViewModel()->GetHierarchyRootViewModel()->FindViewModelForChild(ItemViewModel->GetData()->GetPersistentIdentity(), true) != nullptr)
+						{
+							bCanDrag = false;
+							break;
+						}
+					}
+				}
+
+				return bCanDrag ? EVisibility::Collapsed : EVisibility::Visible;
+			})
+			[
+				SNew(SImage)
+				.Image(FAppStyle::GetBrush("Icons.Lock"))
+				.ToolTipText(LOCTEXT("CantDragItemAlreadyInHierarchyTooltip", "This item already exists within the hierarchy and can not be dragged. Drag the existing one within the hierarchy directly."))
+			]
+		]
+		+ SHorizontalBox::Slot()
+		[
+			OnGenerateRowContentWidget.Execute(HierarchyItem.ToSharedRef())
+		]
+	];
+}
+
+TSharedRef<ITableRow> SNiagaraHierarchyEditor::GenerateHierarchyItemRow(TSharedPtr<FNiagaraHierarchyItemViewModelBase> HierarchyItem, const TSharedRef<STableViewBase>& TableViewBase)
+{
+	return SNew(STableRow<TSharedPtr<FNiagaraHierarchyItemViewModelBase>>, TableViewBase)
+	.Style(HierarchyItem->GetData()->IsA<UNiagaraHierarchyItem>() ? ItemRowStyle : CategoryRowStyle)
+	.OnAcceptDrop(HierarchyItem.ToSharedRef(), &FNiagaraHierarchyItemViewModelBase::OnDroppedOnRow)
+	.OnCanAcceptDrop(HierarchyItem.ToSharedRef(), &FNiagaraHierarchyItemViewModelBase::OnCanRowAcceptDrop)
+	.OnDragDetected(HierarchyItem.ToSharedRef(), &FNiagaraHierarchyItemViewModelBase::OnDragDetected, false)
+	.OnDragLeave(HierarchyItem.ToSharedRef(), &FNiagaraHierarchyItemViewModelBase::OnRowDragLeave)
+	.Padding(FMargin(2))
+	[
+		SNew(SHorizontalBox)
+		+ SHorizontalBox::Slot()
+		.AutoWidth()
+		.Padding(1.f)
+		.HAlign(HAlign_Center)
+		.VAlign(VAlign_Center)
+		[
+			SNew(SBox)
+			.HeightOverride(10.f)
+			.WidthOverride(10.f)
+			.Visibility_Lambda([HierarchyItem]()
+			{
+				return HierarchyItem->IsEditableByUser().bCanPerform ? EVisibility::Collapsed : EVisibility::Visible;
+			})
+			[
+				SNew(SImage)
+				.Image(FAppStyle::GetBrush("Icons.Lock"))
+				.ToolTipText_Lambda([HierarchyItem]()
+				{
+					FNiagaraHierarchyItemViewModelBase::FCanPerformActionResults IsEditableResults = HierarchyItem->IsEditableByUser();
+					if(IsEditableResults.bCanPerform == false)
+					{
+						return IsEditableResults.CanPerformMessage;
+					}
+
+					return FText::GetEmpty();
+				})
+			]
+		]
+		+ SHorizontalBox::Slot()
+		[
+			OnGenerateRowContentWidget.Execute(HierarchyItem.ToSharedRef())
+		]
+	];
+}
+
+bool SNiagaraHierarchyEditor::FilterForSourceSection(TSharedPtr<const FNiagaraHierarchyItemViewModelBase> ItemViewModel) const
+{
+	if(ActiveSourceSection.IsValid())
+	{
+		return GetActiveSourceSectionData() == ItemViewModel->GetSection();
+	}
+
+	return true;
+}
+
+void SNiagaraHierarchyEditor::Reinitialize()
+{
+	// the hierarchy root view model has been recreated if the summary view model reinitialized. Therefore we update the bindings.
+	BindToHierarchyRootViewModel();
+	RefreshSourceItems();
+	RefreshAllViews(true);
+}
+
+void SNiagaraHierarchyEditor::BindToHierarchyRootViewModel()
+{
+	HierarchyViewModel->GetHierarchyRootViewModel()->OnSyncPropagated().BindSP(this, &SNiagaraHierarchyEditor::RequestRefreshHierarchyViewNextFrame, false);
+	HierarchyViewModel->GetHierarchyRootViewModel()->OnSectionsChanged().BindSP(this, &SNiagaraHierarchyEditor::RefreshSectionsView);
+	HierarchyViewModel->GetHierarchyRootViewModel()->OnSectionAdded().BindSP(this, &SNiagaraHierarchyEditor::OnHierarchySectionAdded);
+	HierarchyViewModel->GetHierarchyRootViewModel()->OnSectionDeleted().BindSP(this, &SNiagaraHierarchyEditor::OnHierarchySectionDeleted);
+}
+
+void SNiagaraHierarchyEditor::UnbindFromHierarchyRootViewModel() const
+{
+	if(HierarchyViewModel->GetHierarchyRootViewModel().IsValid())
+	{
+		HierarchyViewModel->GetHierarchyRootViewModel()->OnSyncPropagated().Unbind();
+		HierarchyViewModel->GetHierarchyRootViewModel()->OnSectionsChanged().Unbind();
+		HierarchyViewModel->GetHierarchyRootViewModel()->OnSectionAdded().Unbind();
+		HierarchyViewModel->GetHierarchyRootViewModel()->OnSectionDeleted().Unbind();
+	}
+}
+
+const TArray<TSharedPtr<FNiagaraHierarchyItemViewModelBase>>& SNiagaraHierarchyEditor::GetSourceItems() const
+{
+	return SourceRootViewModel->GetFilteredChildren();
+}
+
+bool SNiagaraHierarchyEditor::IsDetailsPanelEditingAllowed() const
+{
+	return SelectedDetailsPanelItemViewModel.IsValid() && SelectedDetailsPanelItemViewModel.Pin()->IsEditableByUser().bCanPerform;
+}
+
+void SNiagaraHierarchyEditor::RequestRenameSelectedItem()
+{
+	TArray<TSharedPtr<FNiagaraHierarchyItemViewModelBase>> SelectedItems = HierarchyTreeView->GetSelectedItems();
+
+	if(SelectedItems.Num() == 0)
+	{
+		TSharedPtr<FNiagaraHierarchySectionViewModel> ActiveSection = HierarchyViewModel->GetActiveHierarchySection();
+		if(ActiveSection)
+		{
+			SelectedItems = { ActiveSection };
+		}
+	}
+	
+	if(SelectedItems.Num() == 1)
+	{
+		SelectedItems[0]->RequestRename();
+	}
+}
+
+bool SNiagaraHierarchyEditor::CanRequestRenameSelectedItem() const
+{
+	TArray<TSharedPtr<FNiagaraHierarchyItemViewModelBase>> SelectedItems = HierarchyTreeView->GetSelectedItems();
+
+	if(SelectedItems.Num() == 0)
+	{
+		TSharedPtr<FNiagaraHierarchySectionViewModel> ActiveSection = HierarchyViewModel->GetActiveHierarchySection();
+		if(ActiveSection)
+		{
+			SelectedItems = { ActiveSection };
+		}
+	}
+	
+	if(SelectedItems.Num() == 1)
+	{
+		return SelectedItems[0]->CanRename();
+	}
+
+	return false;
+}
+
+void SNiagaraHierarchyEditor::ClearSourceItems() const
+{
+	SourceRoot->GetChildrenMutable().Empty();
+	SourceRoot->GetSectionDataMutable().Empty();
+	SourceRootViewModel->GetChildrenMutable().Empty();
+	SourceRootViewModel->GetSectionViewModels().Empty();
+}
+
+void SNiagaraHierarchyEditor::DeleteItems(TArray<TSharedPtr<FNiagaraHierarchyItemViewModelBase>> ItemsToDelete) const
+{
+	TArray<FNiagaraHierarchyIdentity> DeletionIdentities;
+	
+	for(TSharedPtr<FNiagaraHierarchyItemViewModelBase>& SelectedItem : ItemsToDelete)
+	{
+		DeletionIdentities.Add(SelectedItem->GetData()->GetPersistentIdentity());
+	}
+
+	HierarchyViewModel->DeleteItemsWithIdentities(DeletionIdentities);
+}
+
+void SNiagaraHierarchyEditor::DeleteSelectedHierarchyItems() const
+{	
+	TArray<TSharedPtr<FNiagaraHierarchyItemViewModelBase>> SelectedItems = HierarchyTreeView->GetSelectedItems();
+	DeleteItems(SelectedItems);
+}
+
+bool SNiagaraHierarchyEditor::CanDeleteSelectedHierarchyItems() const
+{
+	TArray<TSharedPtr<FNiagaraHierarchyItemViewModelBase>> SelectedItems = HierarchyTreeView->GetSelectedItems();
+
+	if(SelectedItems.Num() > 0)
+	{
+		bool bCanDelete = true;
+		for(TSharedPtr<FNiagaraHierarchyItemViewModelBase> SelectedItem : SelectedItems)
+		{
+			bCanDelete &= SelectedItems[0]->CanDelete();
+		}
+
+		return bCanDelete;
+	}
+
+	return false;
+}
+
+void SNiagaraHierarchyEditor::DeleteActiveSection() const
+{
+	if(TSharedPtr<FNiagaraHierarchySectionViewModel> ActiveHierarchySectionViewModel = HierarchyViewModel->GetActiveHierarchySection())
+	{
+		DeleteItems({ActiveHierarchySectionViewModel});
+	}
+}
+
+bool SNiagaraHierarchyEditor::CanDeleteActiveSection() const
+{
+	return HierarchyViewModel->GetActiveHierarchySection() != nullptr;
+}
+
+void SNiagaraHierarchyEditor::OnItemAdded(TSharedPtr<FNiagaraHierarchyItemViewModelBase> AddedItem)
+{
+	// when a new item is created (opposed to dragged & dropped from source view, i.e. only categories so far)
+	// we make sure to request a tree refresh, select the row, and request a pending rename since the widget will created a frame later
+	if(AddedItem->GetData()->IsA<UNiagaraHierarchyItem>() || AddedItem->GetData()->IsA<UNiagaraHierarchyCategory>())
+	{
+		HierarchyTreeView->RequestTreeRefresh();
+		NavigateToHierarchyItem(AddedItem);
+	}
+	else if(AddedItem->GetData()->IsA<UNiagaraHierarchySection>())
+	{
+		RefreshSectionsView();
+	}
+
+	AddedItem->RequestRenamePending();
+}
+
+void SNiagaraHierarchyEditor::OnHierarchySectionActivated(TSharedPtr<FNiagaraHierarchySectionViewModel> Section)
+{
+	OnSelectionChanged(Section, ESelectInfo::Direct, true);
+}
+
+void SNiagaraHierarchyEditor::OnSourceSectionActivated(TSharedPtr<FNiagaraHierarchySectionViewModel> Section)
+{
+	OnSelectionChanged(Section, ESelectInfo::Direct, false);
+	RunSourceSearch();
+}
+
+void SNiagaraHierarchyEditor::OnHierarchySectionAdded(TSharedPtr<FNiagaraHierarchySectionViewModel> AddedSection)
+{
+	HierarchyViewModel->SetActiveHierarchySection(AddedSection);
+	AddedSection->RequestRenamePending();
+}
+
+void SNiagaraHierarchyEditor::OnHierarchySectionDeleted(TSharedPtr<FNiagaraHierarchySectionViewModel> DeletedSection)
+{
+	if(HierarchyViewModel->GetActiveHierarchySection() == DeletedSection)
+	{
+		HierarchyViewModel->SetActiveHierarchySection(nullptr);
+	}
+}
+
+void SNiagaraHierarchyEditor::SetActiveSourceSection(TSharedPtr<FNiagaraHierarchySectionViewModel> Section)
+{
+	ActiveSourceSection = Section;	
+	RefreshSourceView(true);
+	OnSourceSectionActivated(Section);
+}
+
+TSharedPtr<FNiagaraHierarchySectionViewModel> SNiagaraHierarchyEditor::GetActiveSourceSection() const
+{
+	return ActiveSourceSection.IsValid() ? ActiveSourceSection.Pin() : nullptr;
+}
+
+UNiagaraHierarchySection* SNiagaraHierarchyEditor::GetActiveSourceSectionData() const
+{
+	return ActiveSourceSection.IsValid() ? ActiveSourceSection.Pin()->GetDataMutable<UNiagaraHierarchySection>() : nullptr;
+}
+
+void SNiagaraHierarchyEditor::OnSelectionChanged(TSharedPtr<FNiagaraHierarchyItemViewModelBase> HierarchyItem, ESelectInfo::Type Type, bool bFromHierarchy) const
+{
+	SelectedDetailsPanelItemViewModel.Reset();
+	if(DetailsPanel.IsValid())
+	{
+		if(HierarchyItem.IsValid() && HierarchyItem->AllowEditingInDetailsPanel())
+		{
+			// when se select a section, and the previous item selection is no longer available due to it, we would get a selection refresh next tick
+			// to wipe out the current selection. We want to avoid that, so we manually clear the selected items in that case.
+			if(HierarchyItem->GetData()->IsA<UNiagaraHierarchySection>())
+			{
+				HierarchyTreeView->ClearSelection();
+			}
+
+			// we clear the selection of the other tree view to avoid 
+			if(bFromHierarchy)
+			{
+				SourceTreeView->ClearSelection();
+			}
+			else
+			{
+				HierarchyTreeView->ClearSelection();
+			}
+			
+			UObject* DataForEditing = HierarchyItem->GetDataForEditing();
+			DataForEditing->SetFlags(RF_Transactional);
+
+			// we make sure the object we are editing is transactional
+			DetailsPanel->SetObject(DataForEditing);
+			SelectedDetailsPanelItemViewModel = HierarchyItem;
+		}
+		else
+		{
+			SelectedDetailsPanelItemViewModel.Reset();
+			DetailsPanel->SetObject(nullptr);
+		}
+	}
+	
+	if(DetailsPanel.IsValid() && OnGenerateCustomDetailsPanelNameWidget.IsBound() && SelectedDetailsPanelItemViewModel.IsValid())
+	{
+		TSharedRef<SWidget> NameWidget = SNew(SHorizontalBox)
+		+ SHorizontalBox::Slot()
+		.AutoWidth()
+		.Padding(2.f)
+		.HAlign(HAlign_Center)
+		.VAlign(VAlign_Center)
+		[
+			SNew(SImage)
+			.Image(FAppStyle::GetBrush("Icons.Lock"))
+			.Visibility(SelectedDetailsPanelItemViewModel.Pin()->IsEditableByUser().bCanPerform ? EVisibility::Collapsed : EVisibility::Visible)
+			.ToolTipText(SelectedDetailsPanelItemViewModel.Pin()->IsEditableByUser().CanPerformMessage)
+		]
+		+ SHorizontalBox::Slot()
+		.AutoWidth()
+		.Padding(2.f)
+		[
+			OnGenerateCustomDetailsPanelNameWidget.Execute(SelectedDetailsPanelItemViewModel.Pin())
+		];
+		
+		DetailsPanel->SetNameAreaCustomContent(NameWidget);
+	}
+}
+
+void SNiagaraHierarchyEditor::RunSourceSearch()
+{
+	if(!SourceSearchBox->GetText().IsEmpty())
+	{
+		OnSourceSearchTextChanged(SourceSearchBox->GetText());
+	}
+}
+
+void SNiagaraHierarchyEditor::OnSourceSearchTextChanged(const FText& Text)
+{
+	SourceSearchResults.Empty();
+	FocusedSearchResult.Reset();
+	SourceTreeView->ClearSelection();
+
+	if(!Text.IsEmpty())
+	{
+		FString TextAsString = Text.ToString();;
+
+		TArray<FSearchItem> SearchItems;
+		GenerateSearchItems(SourceRootViewModel.ToSharedRef(), {}, SearchItems);
+		
+		for(const FSearchItem& SearchItem : SearchItems)
+		{
+			for(const FString& SearchTerm : SearchItem.GetEntry()->GetSearchTerms())
+			{
+				if(SearchTerm.Contains(TextAsString))
+				{
+					SourceSearchResults.Add(SearchItem);
+				}
+			}
+		}
+
+		ExpandSourceSearchResults();
+		SelectNextSourceSearchResult();
+	}
+	else
+	{
+		SourceTreeView->ClearExpandedItems();
+	}
+}
+
+void SNiagaraHierarchyEditor::OnSourceSearchTextCommitted(const FText& Text, ETextCommit::Type CommitType)
+{
+	bool bIsShiftDown = FSlateApplication::Get().GetModifierKeys().IsShiftDown();
+	if(CommitType == ETextCommit::OnEnter)
+	{
+		if(bIsShiftDown == false)
+		{
+			SelectNextSourceSearchResult();
+		}
+		else
+		{
+			SelectPreviousSourceSearchResult();
+		}
+	}
+}
+
+void SNiagaraHierarchyEditor::OnSearchButtonClicked(SSearchBox::SearchDirection SearchDirection)
+{
+	if(SearchDirection == SSearchBox::Next)
+	{
+		SelectNextSourceSearchResult();
+	}
+	else
+	{
+		SelectPreviousSourceSearchResult();
+	}
+}
+
+void SNiagaraHierarchyEditor::GenerateSearchItems(TSharedRef<FNiagaraHierarchyItemViewModelBase> Root, TArray<TSharedPtr<FNiagaraHierarchyItemViewModelBase>> ParentChain, TArray<FSearchItem>& OutSearchItems)
+{
+	const TArray<TSharedPtr<FNiagaraHierarchyItemViewModelBase>> FilteredChildren = Root->GetFilteredChildren();
+	ParentChain.Add(Root);
+	OutSearchItems.Add(FSearchItem{ParentChain});
+	for(TSharedPtr<FNiagaraHierarchyItemViewModelBase> Child : FilteredChildren)
+	{
+		GenerateSearchItems(Child.ToSharedRef(), ParentChain, OutSearchItems);
+	}
+}
+
+void SNiagaraHierarchyEditor::ExpandSourceSearchResults()
+{
+	SourceTreeView->ClearExpandedItems();
+
+	for(FSearchItem& SearchResult : SourceSearchResults)
+	{
+		for(TSharedPtr<FNiagaraHierarchyItemViewModelBase>& EntryInPath : SearchResult.Path)
+		{
+			SourceTreeView->SetItemExpansion(EntryInPath, true);
+		}
+	}
+}
+
+void SNiagaraHierarchyEditor::SelectNextSourceSearchResult()
+{
+	if(SourceSearchResults.IsEmpty())
+	{
+		return;
+	}
+	
+	if(!FocusedSearchResult.IsSet())
+	{
+		FocusedSearchResult = SourceSearchResults[0];
+	}
+	else
+	{
+		int32 CurrentSearchResultIndex = SourceSearchResults.Find(FocusedSearchResult.GetValue());
+		if(SourceSearchResults.IsValidIndex(CurrentSearchResultIndex+1))
+		{
+			FocusedSearchResult = SourceSearchResults[CurrentSearchResultIndex+1];
+		}
+		else
+		{
+			FocusedSearchResult = SourceSearchResults[0];
+		}
+	}
+
+	SourceTreeView->ClearSelection();
+	SourceTreeView->RequestScrollIntoView(FocusedSearchResult.GetValue().GetEntry());
+	SourceTreeView->SetItemSelection(FocusedSearchResult.GetValue().GetEntry(), true);
+}
+
+void SNiagaraHierarchyEditor::SelectPreviousSourceSearchResult()
+{
+	if(SourceSearchResults.IsEmpty())
+	{
+		return;
+	}
+	
+	if(!FocusedSearchResult.IsSet())
+	{
+		FocusedSearchResult = SourceSearchResults[0];
+	}
+	else
+	{
+		int32 CurrentSearchResultIndex = SourceSearchResults.Find(FocusedSearchResult.GetValue());
+		if(SourceSearchResults.IsValidIndex(CurrentSearchResultIndex-1))
+		{
+			FocusedSearchResult = SourceSearchResults[CurrentSearchResultIndex-1];
+		}
+		else
+		{
+			FocusedSearchResult = SourceSearchResults[SourceSearchResults.Num()-1];
+		}
+	}
+
+	SourceTreeView->ClearSelection();
+	SourceTreeView->RequestScrollIntoView(FocusedSearchResult.GetValue().GetEntry());
+	SourceTreeView->SetItemSelection(FocusedSearchResult.GetValue().GetEntry(), true);
+}
+
+TOptional<SSearchBox::FSearchResultData> SNiagaraHierarchyEditor::GetSearchResultData() const
+{
+	if(SourceSearchResults.Num() > 0)
+	{
+		SSearchBox::FSearchResultData SearchResultData;
+		SearchResultData.NumSearchResults = SourceSearchResults.Num();
+
+		if(FocusedSearchResult.IsSet())
+		{
+			// we add one just to make it look nicer as this is merely for cosmetic purposes
+			SearchResultData.CurrentSearchResultIndex = SourceSearchResults.Find(FocusedSearchResult.GetValue()) + 1;
+		}
+		else
+		{
+			SearchResultData.CurrentSearchResultIndex = INDEX_NONE;
+		}
+
+		return SearchResultData;
+	}
+
+	return TOptional<SSearchBox::FSearchResultData>();
+}
+
+FReply SNiagaraHierarchyEditor::HandleHierarchyRootDrop(const FGeometry& Geometry, const FDragDropEvent& DragDropEvent) const
+{
+	if(TSharedPtr<FNiagaraHierarchyDragDropOp> DragDropOp = DragDropEvent.GetOperationAs<FNiagaraHierarchyDragDropOp>())
+	{
+		HierarchyViewModel->GetHierarchyRootViewModel()->OnDroppedOn(DragDropOp->GetDraggedItem().Pin(), EItemDropZone::OntoItem);
+		return FReply::Handled();
+	}
+
+	return FReply::Unhandled();
+}
+
+FNiagaraHierarchyItemViewModelBase::FCanPerformActionResults SNiagaraHierarchyEditor::CanDropOnRoot(TSharedPtr<FNiagaraHierarchyItemViewModelBase> DraggedItem) const
+{
+	return HierarchyViewModel->GetHierarchyRootViewModel()->CanDropOnInternal(DraggedItem, EItemDropZone::OntoItem);
+}
+
+bool SNiagaraHierarchyEditor::OnCanDropOnRoot(TSharedPtr<FDragDropOperation> DragDropOperation) const
+{
+	if(DragDropOperation->IsOfType<FNiagaraHierarchyDragDropOp>())
+	{
+		TSharedPtr<FNiagaraHierarchyDragDropOp> HierarchyDragDropOp = StaticCastSharedPtr<FNiagaraHierarchyDragDropOp>(DragDropOperation);
+		TSharedPtr<FNiagaraHierarchyItemViewModelBase> DraggedItem = HierarchyDragDropOp->GetDraggedItem().Pin();
+		return CanDropOnRoot(DraggedItem).bCanPerform;
+	}
+	
+	return false;
+}
+
+void SNiagaraHierarchyEditor::OnRootDragEnter(const FDragDropEvent& DragDropEvent) const
+{
+	if(TSharedPtr<FNiagaraHierarchyDragDropOp> HierarchyDragDropOp = DragDropEvent.GetOperationAs<FNiagaraHierarchyDragDropOp>())
+	{
+		FNiagaraHierarchyItemViewModelBase::FCanPerformActionResults CanPerformActionResults = CanDropOnRoot(HierarchyDragDropOp->GetDraggedItem().Pin());
+		HierarchyDragDropOp->SetDescription(CanPerformActionResults.CanPerformMessage);
+	}
+}
+
+void SNiagaraHierarchyEditor::OnRootDragLeave(const FDragDropEvent& DragDropEvent) const
+{
+	if(TSharedPtr<FNiagaraHierarchyDragDropOp> HierarchyDragDropOp = DragDropEvent.GetOperationAs<FNiagaraHierarchyDragDropOp>())
+	{
+		HierarchyDragDropOp->SetDescription(FText::GetEmpty());
+	}
+}
+
+FSlateColor SNiagaraHierarchyEditor::GetRootIconColor() const
+{
+	if(FSlateApplication::Get().IsDragDropping())
+	{
+		if(FSlateApplication::Get().GetDragDroppingContent()->IsOfType<FNiagaraHierarchyDragDropOp>())
+		{
+			TSharedPtr<FNiagaraHierarchyDragDropOp> HierarchyDragDropOp = StaticCastSharedPtr<FNiagaraHierarchyDragDropOp>(FSlateApplication::Get().GetDragDroppingContent());
+			if(CanDropOnRoot(HierarchyDragDropOp->GetDraggedItem().Pin()).bCanPerform)
+			{
+				return FLinearColor(0.8f, 0.8f, 0.8f, 0.8f);
+			}
+		}
+	}
+
+	return FLinearColor(0.2f, 0.2f, 0.2f, 0.5f);
+}
+
+FString SNiagaraHierarchyEditor::GetReferencerName() const
+{
+	return TEXT("Niagara Hierarchy");
+}
+
+void SNiagaraHierarchyEditor::AddReferencedObjects(FReferenceCollector& Collector)
+{
+	Collector.AddReferencedObject(SourceRoot);
+}
+
+void SNiagaraHierarchyEditor::NotifyPostChange(const FPropertyChangedEvent& PropertyChangedEvent,	FProperty* PropertyThatChanged)
+{
+	HierarchyViewModel->OnHierarchyPropertiesChanged().Broadcast();
+}
+
+void SNiagaraSectionDragDropTarget::Construct(const FArguments& InArgs,TSharedPtr<FNiagaraHierarchySectionViewModel> InOwningSection, EItemDropZone InItemDropZone)
+{
+	OwningSection = InOwningSection;
+	DropZone = InItemDropZone;
+	
+	SDropTarget::Construct(InArgs._DropTargetArgs);
+}
+
+FReply SNiagaraSectionDragDropTarget::OnDragOver(const FGeometry& MyGeometry, const FDragDropEvent& DragDropEvent)
+{
+	if(TSharedPtr<FNiagaraHierarchyDragDropOp> HierarchyDragDropOp = DragDropEvent.GetOperationAs<FNiagaraHierarchyDragDropOp>())
+	{
+		TSharedPtr<FNiagaraHierarchyItemViewModelBase> DraggedItem = HierarchyDragDropOp->GetDraggedItem().Pin();
+		FNiagaraHierarchyItemViewModelBase::FCanPerformActionResults Results = OwningSection->CanDropOn(DraggedItem, DropZone);
+		HierarchyDragDropOp->SetDescription(Results.CanPerformMessage);
+	}
+
+	return SDropTarget::OnDragOver(MyGeometry, DragDropEvent);
+}
+
+void SNiagaraSectionDragDropTarget::OnDragLeave(const FDragDropEvent& DragDropEvent)
+{
+	TSharedPtr<FNiagaraHierarchyDragDropOp> HierarchyDragDropOp = DragDropEvent.GetOperationAs<FNiagaraHierarchyDragDropOp>();
+
+	if(HierarchyDragDropOp.IsValid())
+	{
+		HierarchyDragDropOp->SetDescription(FText::GetEmpty());
+	}
+	
+	SDropTarget::OnDragLeave(DragDropEvent);
+}
+
+#undef LOCTEXT_NAMESPACE
diff --git a/Engine/Plugins/FX/Niagara/Source/NiagaraEditor/Private/Widgets/SNiagaraParameterPanel.cpp b/Engine/Plugins/FX/Niagara/Source/NiagaraEditor/Private/Widgets/SNiagaraParameterPanel.cpp
--- a/Engine/Plugins/FX/Niagara/Source/NiagaraEditor/Private/Widgets/SNiagaraParameterPanel.cpp
+++ b/Engine/Plugins/FX/Niagara/Source/NiagaraEditor/Private/Widgets/SNiagaraParameterPanel.cpp
@@ -74,7 +74,7 @@ void SNiagaraParameterPanel::Construct(const FArguments& InArgs, const TSharedPt
 	{
 		FilterCategoryExpandedDelegate.BindRaw(ParameterPanelViewModel.Get(), &INiagaraImmutableParameterPanelViewModel::IsCategoryExpandedByDefault);
 	}
-
+	
 	SAssignNew(ItemSelector, SNiagaraParameterPanelSelector)
 	.PreserveSelectionOnRefresh(true)
 	.PreserveExpansionOnRefresh(true)
@@ -102,7 +102,11 @@ void SNiagaraParameterPanel::Construct(const FArguments& InArgs, const TSharedPt
 	.CategoryBorderBackgroundPadding(FMargin(0.0f, 3.0f))
 	.OnGetKeyForItem(this, &SNiagaraParameterPanel::OnGetKeyForItem)
 	.OnGetKeyForCategory(this, &SNiagaraParameterPanel::OnGetKeyForCategory)
-	.ExpandInitially(false);
+	.ExpandInitially(false)
+	.SearchBoxAdjacentContent()
+	[
+		InArgs._SearchAdjacentWidget.Widget
+	];
 
 	// Finalize the widget
 	ChildSlot
diff --git a/Engine/Plugins/FX/Niagara/Source/NiagaraEditor/Private/Widgets/SNiagaraSystemViewport.cpp b/Engine/Plugins/FX/Niagara/Source/NiagaraEditor/Private/Widgets/SNiagaraSystemViewport.cpp
--- a/Engine/Plugins/FX/Niagara/Source/NiagaraEditor/Private/Widgets/SNiagaraSystemViewport.cpp
+++ b/Engine/Plugins/FX/Niagara/Source/NiagaraEditor/Private/Widgets/SNiagaraSystemViewport.cpp
@@ -35,6 +35,8 @@
 #include "ThumbnailRendering/ThumbnailManager.h"
 #include "ToolMenus.h"
 #include "ViewportToolbar/UnrealEdViewportToolbar.h"
+#include "ViewModels/NiagaraEmitterHandleViewModel.h"
+#include "ViewModels/NiagaraSystemViewModel.h"
 #include "Widgets/Docking/SDockTab.h"
 #include "Widgets/Layout/SBox.h"
 
diff --git a/Engine/Plugins/FX/Niagara/Source/NiagaraEditor/Public/NiagaraClipboard.h b/Engine/Plugins/FX/Niagara/Source/NiagaraEditor/Public/NiagaraClipboard.h
--- a/Engine/Plugins/FX/Niagara/Source/NiagaraEditor/Public/NiagaraClipboard.h
+++ b/Engine/Plugins/FX/Niagara/Source/NiagaraEditor/Public/NiagaraClipboard.h
@@ -11,6 +11,7 @@
 #include "UObject/SoftObjectPtr.h"
 #include "NiagaraClipboard.generated.h"
 
+class UNiagaraHierarchyCategory;
 class IPropertyHandle;
 class UNiagaraDataInterface;
 class UNiagaraScript;
@@ -85,6 +86,9 @@ public:
 	UPROPERTY()
 	TObjectPtr<UNiagaraClipboardFunction> Dynamic;
 
+	UPROPERTY()
+	TArray<TObjectPtr<const UNiagaraClipboardFunctionInput>> ChildrenInputs;
+
 	NIAGARAEDITOR_API bool CopyValuesFrom(const UNiagaraClipboardFunctionInput* InOther);
 
 	const FNiagaraTypeDefinition& GetTypeDef() const { return InputType; };
@@ -252,7 +256,7 @@ public:
 
 	UPROPERTY()
 	TArray<TObjectPtr<const UNiagaraClipboardFunctionInput>> FunctionInputs;
-
+	
 	UPROPERTY()
 	TArray<TObjectPtr<const UNiagaraClipboardRenderer>> Renderers;
 
diff --git a/Engine/Plugins/FX/Niagara/Source/NiagaraEditor/Public/NiagaraEditorUtilities.h b/Engine/Plugins/FX/Niagara/Source/NiagaraEditor/Public/NiagaraEditorUtilities.h
--- a/Engine/Plugins/FX/Niagara/Source/NiagaraEditor/Public/NiagaraEditorUtilities.h
+++ b/Engine/Plugins/FX/Niagara/Source/NiagaraEditor/Public/NiagaraEditorUtilities.h
@@ -12,6 +12,7 @@
 #include "NiagaraEditorSettings.h"
 #include "ViewModels/NiagaraSystemScalabilityViewModel.h"
 #include "ViewModels/NiagaraSystemViewModel.h"
+#include "ViewModels/HierarchyEditor/NiagaraHierarchyViewModelBase.h"
 
 class UNiagaraClipboardContent;
 class UNiagaraNodeInput;
@@ -471,6 +472,14 @@ namespace FNiagaraEditorUtilities
 		EAssetTagSectionSource GetAssetTagDefinitionSource(const FAssetData& AssetData);
 		FText GetAssetTagSectionNameFromSource(EAssetTagSectionSource Source);
 	}
+
+	namespace HierarchyEditor
+	{		
+		namespace Scripts
+		{
+			TSharedRef<SWidget> GenerateRowContentForScriptParameterHierarchyEditor(TSharedRef<FNiagaraHierarchyItemViewModelBase> HierarchyItem);
+		}
+	}
 }
 
 namespace FNiagaraParameterUtilities
diff --git a/Engine/Plugins/FX/Niagara/Source/NiagaraEditor/Public/NiagaraGraph.h b/Engine/Plugins/FX/Niagara/Source/NiagaraEditor/Public/NiagaraGraph.h
--- a/Engine/Plugins/FX/Niagara/Source/NiagaraEditor/Public/NiagaraGraph.h
+++ b/Engine/Plugins/FX/Niagara/Source/NiagaraEditor/Public/NiagaraGraph.h
@@ -6,12 +6,14 @@
 #include "NiagaraCommon.h"
 #include "EdGraph/EdGraph.h"
 #include "NiagaraScript.h"
+#include "NiagaraScriptVariable.h"
 #include "NiagaraVariableMetaData.h"
 #include "NiagaraGraph.generated.h"
 
 class UNiagaraNodeStaticSwitch;
 class UNiagaraParameterDefinitions;
 class UNiagaraScriptVariable;
+class UNiagaraHierarchyRoot;
 struct FSynchronizeWithParameterDefinitionsArgs;
 struct FNiagaraScriptVariableData;
 struct FNiagaraStaticVariableSearchContext;
@@ -179,12 +181,21 @@ struct FNiagaraFindInputNodeOptions
 UCLASS(MinimalAPI)
 class UNiagaraGraph : public UEdGraph
 {
-	GENERATED_UCLASS_BODY()
+	GENERATED_BODY()
 
+public:
+	struct FParametersChangedData
+	{
+		const UNiagaraScriptVariable* ChangedParameter;
+	};
+	
+	DECLARE_MULTICAST_DELEGATE_OneParam(FOnParametersChanged, TOptional<FParametersChangedData>)
 	DECLARE_MULTICAST_DELEGATE(FOnDataInterfaceChanged);
 	DECLARE_MULTICAST_DELEGATE_OneParam(FOnSubObjectSelectionChanged, const UObject*);
 	DECLARE_DELEGATE_RetVal_OneParam(TArray<UNiagaraParameterDefinitions*> /*AvailableParameterDefinitions*/, FOnGetParameterDefinitionsForDetailsCustomization, bool /*bSkipSubscribedParameterDefinitions*/)
 
+	UNiagaraGraph();
+	
 	//~ Begin UObject Interface
 	virtual void PostLoad() override;
 #if WITH_EDITORONLY_DATA
@@ -231,6 +242,9 @@ class UNiagaraGraph : public UEdGraph
 	/** Generates a list of unique input and output parameters for when this script is used as a function. */
 	void GetParameters(TArray<FNiagaraVariable>& Inputs, TArray<FNiagaraVariable>& Outputs) const;
 
+	/** Generates a list of all input parameters - including parameters in the Input namespace & Static Switch inputs. */
+	void GetAllInputScriptVariables(TArray<UNiagaraScriptVariable*>& OutScriptVariables) const;
+	
 	/** Returns the index of this variable in the output node of the graph. INDEX_NONE if this is not a valid attribute. */
 	int32 GetOutputNodeVariableIndex(const FNiagaraVariable& Attr)const;
 	void GetOutputNodeVariables(TArray< FNiagaraVariable >& OutAttributes)const;
@@ -305,11 +319,13 @@ class UNiagaraGraph : public UEdGraph
 	NIAGARAEDITOR_API const FScriptVariableMap& GetAllMetaData() const;
 	NIAGARAEDITOR_API FScriptVariableMap& GetAllMetaData();
 
+	NIAGARAEDITOR_API UNiagaraHierarchyRoot* GetScriptParameterHierarchyRoot() const { return ParameterHierarchyRoot; }
+
 	UNiagaraScriptVariable* GetScriptVariable(FNiagaraVariable Parameter) const;
 	NIAGARAEDITOR_API UNiagaraScriptVariable* GetScriptVariable(FName ParameterName) const;
 	NIAGARAEDITOR_API UNiagaraScriptVariable* GetScriptVariable(FGuid VariableGuid) const;
-	NIAGARAEDITOR_API TArray<UNiagaraScriptVariable*> GetChildScriptVariablesForInput(FGuid VariableGuid) const;
 	NIAGARAEDITOR_API TArray<FGuid> GetChildScriptVariableGuidsForInput(FGuid VariableGuid) const;
+	TArray<UNiagaraScriptVariable*> GetChildScriptVariablesForInput_Deprecated(FGuid VariableGuid) const;
 
 	/** Adds parameter to the VariableToScriptVariable map.*/
 	UNiagaraScriptVariable* AddParameter(const FNiagaraVariable& Parameter, bool bIsStaticSwitch = false);
@@ -348,6 +364,9 @@ class UNiagaraGraph : public UEdGraph
 
 	void ReplaceScriptReferences(UNiagaraScript* OldScript, UNiagaraScript* NewScript);
 
+	/** Gets a delegate which is called whenever a parameter is added or removed */
+	FOnParametersChanged& OnParametersChanged();
+	
 	/** Gets a delegate which is called whenever a contained data interfaces changes. */
 	FOnDataInterfaceChanged& OnDataInterfaceChanged();
 
@@ -508,6 +527,9 @@ private:
 	UPROPERTY()
 	TMap<FNiagaraVariable, TObjectPtr<UNiagaraScriptVariable>> VariableToScriptVariable;
 
+	UPROPERTY()
+	TObjectPtr<UNiagaraHierarchyRoot> ParameterHierarchyRoot;
+	
 	/** A map of parameters in the graph to their referencers. */
 	UPROPERTY(Transient)
 	mutable TMap<FNiagaraVariable, FNiagaraGraphParameterReferenceCollection> ParameterToReferencesMap;
@@ -520,6 +542,7 @@ private:
 	UPROPERTY()
 	bool bHasValidLastBuiltScriptVersionId = false;
 
+	FOnParametersChanged OnParametersChangedDelegate;
 	FOnDataInterfaceChanged OnDataInterfaceChangedDelegate;
 	FOnSubObjectSelectionChanged OnSelectedSubObjectChanged;
 	FOnGetParameterDefinitionsForDetailsCustomization OnGetParameterDefinitionsForDetailsCustomizationDelegate;
diff --git a/Engine/Plugins/FX/Niagara/Source/NiagaraEditor/Public/NiagaraScriptSource.h b/Engine/Plugins/FX/Niagara/Source/NiagaraEditor/Public/NiagaraScriptSource.h
--- a/Engine/Plugins/FX/Niagara/Source/NiagaraEditor/Public/NiagaraScriptSource.h
+++ b/Engine/Plugins/FX/Niagara/Source/NiagaraEditor/Public/NiagaraScriptSource.h
@@ -17,12 +17,13 @@ struct FNiagaraVMExecutableDataId;
 UCLASS(MinimalAPI)
 class UNiagaraScriptSource final : public UNiagaraScriptSourceBase
 {
-	GENERATED_UCLASS_BODY()
+	GENERATED_BODY()
 
+public:
 	/** Graph for particle update expression */
 	UPROPERTY()
 	TObjectPtr<class UNiagaraGraph>	NodeGraph = nullptr;
-
+	
 	bool bIsCompilationCopy = false;
 	bool bIsReleased = false;
 	
@@ -59,6 +60,7 @@ class UNiagaraScriptSource final : public UNiagaraScriptSourceBase
 
 	NIAGARAEDITOR_API virtual bool AddModuleIfMissing(FString ModulePath, ENiagaraScriptUsage Usage, bool& bOutFoundModule)override;
 
+	virtual void PostLoadFromOwner(FVersionedNiagaraScriptData& OwnerData) override;
 
 	void FixupRenamedParameters(UNiagaraNodeFunctionCall* FunctionCallNode, TConstArrayView<FNiagaraVariable> ModuleInputVariables, FNiagaraParameterStore& RapidIterationParameters, const TArray<FNiagaraVariable>& OldRapidIterationVariables, const FVersionedNiagaraEmitter& Emitter, ENiagaraScriptUsage ScriptUsage) const;
 	void InitializeNewParameters(UNiagaraNodeFunctionCall* FunctionCallNode, TConstArrayView<FNiagaraVariable> ModuleInputVariables, FNiagaraParameterStore& RapidIterationParameters, const FVersionedNiagaraEmitter& VersionedEmitter, ENiagaraScriptUsage ScriptUsage, TSet<FNiagaraVariableBase>& ValidRapidIterationParameters) const;
@@ -96,4 +98,6 @@ class UNiagaraScriptSource final : public UNiagaraScriptSourceBase
 private:
 	void OnGraphChanged(const FEdGraphEditAction &Action);
 	void OnGraphDataInterfaceChanged();
+
+	void PostLoadTransferScriptHierarchyData(FVersionedNiagaraScriptData& OwnerData);
 };
diff --git a/Engine/Plugins/FX/Niagara/Source/NiagaraEditor/Public/NiagaraStackEditorData.h b/Engine/Plugins/FX/Niagara/Source/NiagaraEditor/Public/NiagaraStackEditorData.h
--- a/Engine/Plugins/FX/Niagara/Source/NiagaraEditor/Public/NiagaraStackEditorData.h
+++ b/Engine/Plugins/FX/Niagara/Source/NiagaraEditor/Public/NiagaraStackEditorData.h
@@ -138,6 +138,8 @@ public:
 	*/
 	void SetStackEntryActiveSection(const FString& StackEntryKey, FText ActiveSection);
 
+	void ClearStackEntryActiveSection(const FString& StackEntryKey);
+	
 	/*
 	* Gets a stack entry's display name. Returns null if none is found.
 	* @param StackEntryKey A unique key for the stack entry.
diff --git a/Engine/Plugins/FX/Niagara/Source/NiagaraEditor/Public/Toolkits/SystemToolkitModes/NiagaraSystemToolkitModeBase.h b/Engine/Plugins/FX/Niagara/Source/NiagaraEditor/Public/Toolkits/SystemToolkitModes/NiagaraSystemToolkitModeBase.h
--- a/Engine/Plugins/FX/Niagara/Source/NiagaraEditor/Public/Toolkits/SystemToolkitModes/NiagaraSystemToolkitModeBase.h
+++ b/Engine/Plugins/FX/Niagara/Source/NiagaraEditor/Public/Toolkits/SystemToolkitModes/NiagaraSystemToolkitModeBase.h
@@ -15,7 +15,7 @@ class FNiagaraSystemToolkitModeBase : public FApplicationMode
 {
 public:
 	FNiagaraSystemToolkitModeBase(FName InModeName, TWeakPtr<FNiagaraSystemToolkit> InSystemToolkit);
-	~FNiagaraSystemToolkitModeBase();
+	virtual ~FNiagaraSystemToolkitModeBase() override;
 	
 	virtual void RegisterTabFactories(TSharedPtr<FTabManager> InTabManager) override;
 
@@ -47,6 +47,8 @@ private:
 	TSharedRef<SDockTab> SpawnTab_UserParameters(const FSpawnTabArgs& Args);
 	TSharedRef<SDockTab> SpawnTab_UserParametersHierarchyEditor(const FSpawnTabArgs& Args);
 	TSharedRef<SDockTab> SpawnTab_SummaryViewEditor(const FSpawnTabArgs& Args);
+	TSharedRef<SDockTab> SpawnTab_ScratchPadHierarchyEditor(const FSpawnTabArgs& Args);
+
 
 protected:
 	TWeakPtr<FNiagaraSystemToolkit> SystemToolkit;
@@ -56,6 +58,8 @@ protected:
 	FDelegateHandle LastSystemSelectionUpdateDelegate;
 	FDelegateHandle LastGraphEditDelegate;
 	FDelegateHandle UpdateSummaryViewHandle;
+	FDelegateHandle UpdateScratchPadScriptHierarchyHandle;
+
 	int32 SwitcherIdx;
 	TSharedPtr<FNiagaraObjectSelection> ObjectSelection;
 	FText ObjectSelectionSubHeaderText;
@@ -69,11 +73,18 @@ protected:
 	
 	TSharedRef<SWidget> CreateSummaryViewWidget() const;
 	void UpdateSummaryViewOnSelectionChanged() const;
-	
 	void OnSummaryViewEditorClosed(TSharedRef<SDockTab> DockTab) const;
 
-	TWeakPtr< class FNiagaraScratchPadScriptViewModel > LastActiveDocumentModel;
+	FReply SummonScratchPadScriptHierarchyEditor();
+	TSharedRef<SWidget> CreateScratchPadHierarchyWidget();
+	void UpdateScratchPadActiveScriptChanged(TSharedPtr<SDockTab> DockTab);
+	void OnScratchPadHierarchyEditorClosed(TSharedRef<SDockTab> DockTab);
+	EVisibility GetSummonScratchPadHierarchyEditorButtonVisibility() const;
+	
 	TSharedPtr<class SBox> SummaryViewContainer;
+	TSharedPtr<class SBox> ScratchPadHierarchyContainer;
+	TWeakPtr<class FNiagaraScratchPadScriptViewModel> LastActiveScratchPadViewModel;
+
 public:
 	static const FName ViewportTabID;
 	static const FName CurveEditorTabID;
@@ -95,6 +106,7 @@ public:
 	static const FName VersioningTabID;
 	static const FName UserParametersTabID;
 	static const FName UserParametersHierarchyTabID;
-	static const FName EmitterSummaryViewEditorTabID;;
+	static const FName EmitterSummaryViewEditorTabID;
+	static const FName ScratchPadHierarchyEditorTabID;
 };
 
diff --git a/Engine/Plugins/FX/Niagara/Source/NiagaraEditor/Public/ViewModels/HierarchyEditor/NiagaraHierarchyScriptParametersViewModel.h b/Engine/Plugins/FX/Niagara/Source/NiagaraEditor/Public/ViewModels/HierarchyEditor/NiagaraHierarchyScriptParametersViewModel.h
new file mode 100644
--- /dev/null
+++ b/Engine/Plugins/FX/Niagara/Source/NiagaraEditor/Public/ViewModels/HierarchyEditor/NiagaraHierarchyScriptParametersViewModel.h
@@ -0,0 +1,117 @@
+﻿// Fill out your copyright notice in the Description page of Project Settings.
+
+#pragma once
+
+#include "NiagaraGraph.h"
+#include "ViewModels/HierarchyEditor/NiagaraHierarchyViewModelBase.h"
+#include "NiagaraScriptVariable.h"
+#include "ViewModels/NiagaraScriptViewModel.h"
+#include "NiagaraHierarchyScriptParametersViewModel.generated.h"
+
+class UNiagaraGraph;
+
+/** The refresh context is used to determine if hierarchy script variables should be removed. */
+UCLASS()
+class UNiagaraHierarchyScriptParameterRefreshContext : public UNiagaraHierarchyDataRefreshContext
+{
+	GENERATED_BODY()
+
+public:
+	void SetNiagaraGraph(UNiagaraGraph* InGraph) { NiagaraGraph = InGraph; }
+	const UNiagaraGraph* GetNiagaraGraph() const { return NiagaraGraph; }
+private:
+	UPROPERTY(Transient)
+	TObjectPtr<UNiagaraGraph> NiagaraGraph;
+};
+
+/** A hierarchy script parameter is an optional object embedded in the hierarchy. */
+UCLASS()
+class UNiagaraHierarchyScriptParameter : public UNiagaraHierarchyItem
+{
+	GENERATED_BODY()
+
+public:
+	void Initialize(UNiagaraScriptVariable& InParameterScriptVariable);
+
+	virtual FString ToString() const override { return GetVariable().GetName().ToString(); }
+	FText GetTooltip() const { return ParameterScriptVariable->Metadata.Description; }
+	
+	UNiagaraScriptVariable* GetScriptVariable() const { return ParameterScriptVariable; }
+	FNiagaraVariable GetVariable() const { return ParameterScriptVariable->Variable; }
+private:
+	UPROPERTY()
+	TObjectPtr<UNiagaraScriptVariable> ParameterScriptVariable;
+};
+
+/** The category class used for the script hierarchy editor. It lets us add additional data later on. */
+UCLASS()
+class UNiagaraHierarchyScriptCategory : public UNiagaraHierarchyCategory
+{
+	GENERATED_BODY()
+};
+
+/**
+ * The view model that defines the script editor's hierarchy editor for input parameters.
+ */
+UCLASS()
+class NIAGARAEDITOR_API UNiagaraHierarchyScriptParametersViewModel : public UNiagaraHierarchyViewModelBase
+{
+	GENERATED_BODY()
+public:
+	void Initialize(TSharedRef<FNiagaraScriptViewModel> InScriptViewModel);
+	TSharedRef<FNiagaraScriptViewModel> GetScriptViewModel() const;
+
+	/** UNiagaraHierarchyViewModelBase */
+	virtual UNiagaraHierarchyRoot* GetHierarchyRoot() const override;
+	virtual TSubclassOf<UNiagaraHierarchyCategory> GetCategoryDataClass() const override;
+	virtual TSharedPtr<FNiagaraHierarchyRootViewModel> CreateRootViewModelForData(UNiagaraHierarchyRoot* Root, bool bIsForHierarchy) override;
+	virtual TSharedPtr<FNiagaraHierarchyItemViewModelBase> CreateViewModelForData(UNiagaraHierarchyItemBase* ItemBase, TSharedPtr<FNiagaraHierarchyItemViewModelBase> Parent) override;
+	virtual void PrepareSourceItems(UNiagaraHierarchyRoot* SourceRoot, TSharedPtr<FNiagaraHierarchyRootViewModel> SourceRootViewModel) override;
+	virtual void SetupCommands() override;
+	virtual TSharedRef<FNiagaraHierarchyDragDropOp> CreateDragDropOp(TSharedRef<FNiagaraHierarchyItemViewModelBase> Item) override;
+	virtual bool SupportsDetailsPanel() override { return true; }
+	virtual void FinalizeInternal() override;
+private:
+	void OnParametersChanged(TOptional<UNiagaraGraph::FParametersChangedData> ParametersChangedData);
+private:
+	TWeakPtr<FNiagaraScriptViewModel> ScriptViewModelWeak;
+};
+
+class FNiagaraHierarchyScriptParameterDragDropOp : public FNiagaraHierarchyDragDropOp
+{
+public:
+	DRAG_DROP_OPERATOR_TYPE(FNiagaraHierarchyScriptParameterDragDropOp, FNiagaraHierarchyDragDropOp)
+
+	FNiagaraHierarchyScriptParameterDragDropOp(TSharedPtr<FNiagaraHierarchyItemViewModel> InputViewModel) : FNiagaraHierarchyDragDropOp(InputViewModel) {}
+	
+	virtual TSharedRef<SWidget> CreateCustomDecorator() const override;
+};
+
+struct FNiagaraHierarchyScriptParameterViewModel : public FNiagaraHierarchyItemViewModel
+{
+	FNiagaraHierarchyScriptParameterViewModel(UNiagaraHierarchyScriptParameter* ScriptParameter, TSharedPtr<FNiagaraHierarchyItemViewModelBase> InParent, TWeakObjectPtr<UNiagaraHierarchyScriptParametersViewModel> ViewModel, bool bInIsForHierarchy)
+		: FNiagaraHierarchyItemViewModel(ScriptParameter, InParent, ViewModel, bInIsForHierarchy)	{}
+
+protected:
+	virtual bool RepresentsExternalData() const override { return true; }
+	virtual bool DoesExternalDataStillExist(const UNiagaraHierarchyDataRefreshContext* Context) const override;
+	virtual UObject* GetDataForEditing() override { return Cast<UNiagaraHierarchyScriptParameter>(GetDataMutable())->GetScriptVariable(); }
+	/** We want to be able to edit in the details panel regardless of source or hierarchy item. */
+	virtual bool AllowEditingInDetailsPanel() const override { return true; }
+	virtual bool CanRenameInternal() override { return false; }
+	
+	virtual FCanPerformActionResults CanDropOnInternal(TSharedPtr<FNiagaraHierarchyItemViewModelBase>, EItemDropZone ItemDropZone) override;
+	virtual void OnDroppedOnInternal(TSharedPtr<FNiagaraHierarchyItemViewModelBase> DroppedItem, EItemDropZone ItemDropZone) override;
+	
+	virtual bool CanHaveChildren() const override { return bIsForHierarchy == true; }
+};
+
+struct FNiagaraHierarchyScriptRootViewModel : public FNiagaraHierarchyRootViewModel
+{
+	FNiagaraHierarchyScriptRootViewModel(UNiagaraHierarchyRoot* Root, TWeakObjectPtr<UNiagaraHierarchyScriptParametersViewModel> ViewModel, bool bInIsForHierarchy)
+		: FNiagaraHierarchyRootViewModel(Root, ViewModel, bInIsForHierarchy)	{}
+
+protected:
+	// In the script, loose parameters are always added before categories (reverse of the default implementation)
+	virtual void SortChildrenData() const override;
+};
diff --git a/Engine/Plugins/FX/Niagara/Source/NiagaraEditor/Public/ViewModels/HierarchyEditor/NiagaraHierarchyViewModelBase.h b/Engine/Plugins/FX/Niagara/Source/NiagaraEditor/Public/ViewModels/HierarchyEditor/NiagaraHierarchyViewModelBase.h
--- a/Engine/Plugins/FX/Niagara/Source/NiagaraEditor/Public/ViewModels/HierarchyEditor/NiagaraHierarchyViewModelBase.h
+++ b/Engine/Plugins/FX/Niagara/Source/NiagaraEditor/Public/ViewModels/HierarchyEditor/NiagaraHierarchyViewModelBase.h
@@ -3,13 +3,15 @@
 #pragma once
 
 #include "CoreMinimal.h"
-#include "NiagaraEditorUtilities.h"
 #include "PropertyEditorDelegates.h"
-#include "ViewModels/Stack/NiagaraStackFunctionInput.h"
 #include "ToolMenuSection.h"
 #include "IPropertyRowGenerator.h"
 #include "Misc/TransactionObjectEvent.h"
 #include "ScopedTransaction.h"
+#include "EditorUndoClient.h"
+#include "Widgets/Views/STableRow.h"
+#include "TickableEditorObject.h"
+#include "Engine/TimerHandle.h"
 #include "NiagaraHierarchyViewModelBase.generated.h"
 
 USTRUCT()
@@ -308,6 +310,7 @@ public:
 	NIAGARAEDITOR_API void SetSectionNameAsText(const FText& Text);
 	FText GetSectionNameAsText() const { return FText::FromName(Section); }
 
+	void SetTooltip(const FText& InTooltip) { Tooltip = InTooltip; }
 	FText GetTooltip() const { return Tooltip; }
 
 	virtual FString ToString() const override { return Section.ToString(); }
@@ -392,7 +395,8 @@ public:
 
 	void Initialize();
 	void Finalize();
-	
+
+	virtual TSubclassOf<UNiagaraHierarchyCategory> GetCategoryDataClass() const;
 	void AddCategory(TSharedPtr<FNiagaraHierarchyItemViewModelBase> CategoryParent) const;
 	void AddSection() const;
 
@@ -409,7 +413,9 @@ public:
 	/** The hierarchy root the widget is editing. This should point to persistent data stored somewhere else as the serialized root of the hierarchy. */
 	virtual UNiagaraHierarchyRoot* GetHierarchyRoot() const PURE_VIRTUAL(UNiagaraHierarchyViewModelBase::GetHierarchyRoot, return nullptr;);
 	TSharedPtr<struct FNiagaraHierarchyRootViewModel> GetHierarchyRootViewModel() const { return HierarchyRootViewModel; }
-	
+
+	/** Create the root view model for the root entry. Handled separately from all other elements. */
+	virtual TSharedPtr<FNiagaraHierarchyRootViewModel> CreateRootViewModelForData(UNiagaraHierarchyRoot* Root, bool bIsForHierarchy);
 	/** This function will create the view model for a given item. Override to customize view model behavior by providing custom classes. */
 	virtual TSharedPtr<FNiagaraHierarchyItemViewModelBase> CreateViewModelForData(UNiagaraHierarchyItemBase* ItemBase, TSharedPtr<FNiagaraHierarchyItemViewModelBase> Parent);
 	
@@ -549,10 +555,14 @@ struct FNiagaraHierarchyItemViewModelBase : TSharedFromThis<FNiagaraHierarchyIte
 
 	NIAGARAEDITOR_API void RefreshChildrenData();
 	NIAGARAEDITOR_API void SyncViewModelsToData();
+	
+	/** Every item view model can define its own sort order for its children. By default we put categories above items. */
+	virtual void SortChildrenData() const;
+	
 	const TArray<TSharedPtr<FNiagaraHierarchyItemViewModelBase>>& GetChildren() const { return Children; }
 	TArray<TSharedPtr<FNiagaraHierarchyItemViewModelBase>>& GetChildrenMutable() { return Children; }
 	NIAGARAEDITOR_API const TArray<TSharedPtr<FNiagaraHierarchyItemViewModelBase>>& GetFilteredChildren() const;
-
+	
 	NIAGARAEDITOR_API void AddChildFilter(FOnFilterChild InFilterChild);
 	
 	template<class DataClass, class ViewModelChildClass>
@@ -662,6 +672,9 @@ struct FNiagaraHierarchyItemViewModelBase : TSharedFromThis<FNiagaraHierarchyIte
 
 	/** The UObject we display in the details panel when this item is selected. By default it's the item the view model represents. */
 	virtual UObject* GetDataForEditing() { return ItemBase; }
+	/** Source items are transient, which is why we don't allow editing by default.
+	 * This is useful to override if source data points at actual data to edit. */
+	virtual bool AllowEditingInDetailsPanel() const { return bIsForHierarchy; }
 	
 	/** Used to create customized drag drop ops. */
 	NIAGARAEDITOR_API TSharedRef<class FNiagaraHierarchyDragDropOp> CreateDragDropOp();
@@ -825,7 +838,7 @@ struct FNiagaraHierarchyItemViewModel : FNiagaraHierarchyItemViewModelBase
 	virtual bool CanHaveChildren() const override { return false; }
 	virtual FCanPerformActionResults CanDragInternal() override { return true; }
 
-	NIAGARAEDITOR_API virtual FCanPerformActionResults CanDropOnInternal(TSharedPtr<FNiagaraHierarchyItemViewModelBase>, EItemDropZone ItemDropZone) override;
+	NIAGARAEDITOR_API virtual FCanPerformActionResults CanDropOnInternal(TSharedPtr<FNiagaraHierarchyItemViewModelBase> DraggedItem, EItemDropZone ItemDropZone) override;
 	NIAGARAEDITOR_API virtual void OnDroppedOnInternal(TSharedPtr<FNiagaraHierarchyItemViewModelBase> DroppedItem, EItemDropZone ItemDropZone) override;
 };
 
diff --git a/Engine/Plugins/FX/Niagara/Source/NiagaraEditor/Public/ViewModels/HierarchyEditor/NiagaraUserParametersHierarchyViewModel.h b/Engine/Plugins/FX/Niagara/Source/NiagaraEditor/Public/ViewModels/HierarchyEditor/NiagaraUserParametersHierarchyViewModel.h
--- a/Engine/Plugins/FX/Niagara/Source/NiagaraEditor/Public/ViewModels/HierarchyEditor/NiagaraUserParametersHierarchyViewModel.h
+++ b/Engine/Plugins/FX/Niagara/Source/NiagaraEditor/Public/ViewModels/HierarchyEditor/NiagaraUserParametersHierarchyViewModel.h
@@ -6,6 +6,8 @@
 #include "NiagaraScriptVariable.h"
 #include "NiagaraUserParametersHierarchyViewModel.generated.h"
 
+class FNiagaraSystemViewModel;
+
 UCLASS()
 class UNiagaraHierarchyUserParameterRefreshContext : public UNiagaraHierarchyDataRefreshContext
 {
@@ -47,6 +49,8 @@ struct FNiagaraHierarchyUserParameterViewModel : public FNiagaraHierarchyItemVie
 	
 	/** For editing in the details panel we want to handle the script variable that is represented by the hierarchy item, not the hierarchy item itself. */
 	virtual UObject* GetDataForEditing() override;
+	/** We want to be able to edit in the details panel regardless of source or hierarchy item. */
+	virtual bool AllowEditingInDetailsPanel() const override { return true; }
 
 	virtual bool RepresentsExternalData() const override { return true; }
 	virtual bool DoesExternalDataStillExist(const UNiagaraHierarchyDataRefreshContext* Context) const override;
@@ -85,7 +89,7 @@ private:
 class FNiagaraUserParameterHierarchyDragDropOp : public FNiagaraHierarchyDragDropOp
 {
 public:
-	DRAG_DROP_OPERATOR_TYPE(FNiagaraUserParameterDragDropOp, FNiagaraHierarchyDragDropOp)
+	DRAG_DROP_OPERATOR_TYPE(FNiagaraUserParameterHierarchyDragDropOp, FNiagaraHierarchyDragDropOp)
 
 	FNiagaraUserParameterHierarchyDragDropOp(TSharedPtr<FNiagaraHierarchyItemViewModelBase> UserParameterItem) : FNiagaraHierarchyDragDropOp(UserParameterItem) {}
 
diff --git a/Engine/Plugins/FX/Niagara/Source/NiagaraEditor/Public/ViewModels/NiagaraScratchPadScriptViewModel.h b/Engine/Plugins/FX/Niagara/Source/NiagaraEditor/Public/ViewModels/NiagaraScratchPadScriptViewModel.h
--- a/Engine/Plugins/FX/Niagara/Source/NiagaraEditor/Public/ViewModels/NiagaraScratchPadScriptViewModel.h
+++ b/Engine/Plugins/FX/Niagara/Source/NiagaraEditor/Public/ViewModels/NiagaraScratchPadScriptViewModel.h
@@ -4,6 +4,8 @@
 
 #include "ViewModels/NiagaraScriptViewModel.h"
 #include "GraphEditAction.h"
+#include "HierarchyEditor/NiagaraHierarchyScriptParametersViewModel.h"
+#include "HierarchyEditor/NiagaraSummaryViewViewModel.h"
 
 class INiagaraParameterPanelViewModel;
 class FNiagaraScriptToolkitParameterPanelViewModel;
@@ -25,7 +27,7 @@ public:
 
 	NIAGARAEDITOR_API ~FNiagaraScratchPadScriptViewModel();
 
-	NIAGARAEDITOR_API void Initialize(UNiagaraScript* Script, UNiagaraScript* InEditScript, TWeakPtr<class FNiagaraSystemViewModel> InSystemViewModel);
+	void Initialize(UNiagaraScript* Script, UNiagaraScript* InEditScript, TWeakPtr<class FNiagaraSystemViewModel> InSystemViewModel);
 
 	NIAGARAEDITOR_API bool IsValid() const;
 	
@@ -39,6 +41,8 @@ public:
 
 	NIAGARAEDITOR_API TArray<UNiagaraGraph*> GetEditableGraphs();
 
+	NIAGARAEDITOR_API UNiagaraHierarchyScriptParametersViewModel* GetHierarchyViewModel();
+
 	//~ Begin NiagaraParameterDefinitionsSubscriberViewModel Interface
 protected:
 	virtual INiagaraParameterDefinitionsSubscriber* GetParameterDefinitionsSubscriber() override { return &EditScript; };
@@ -94,6 +98,9 @@ public:
 
 	NIAGARAEDITOR_API FSimpleDelegate& OnRequestDiscardChanges();
 
+	NIAGARAEDITOR_API void OnHierarchyChanged();
+	NIAGARAEDITOR_API void OnHierarchyPropertiesChanged();
+
 	NIAGARAEDITOR_API FOnNodeIDFocusRequested& OnNodeIDFocusRequested();
 	NIAGARAEDITOR_API FOnPinIDFocusRequested& OnPinIDFocusRequested();
 
@@ -126,6 +133,8 @@ private:
 	TSharedPtr<FUICommandList> ParameterPanelCommands;
 	TSharedPtr<FNiagaraScriptToolkitParameterPanelViewModel> ParameterPaneViewModel;
 
+	TObjectPtr<UNiagaraHierarchyScriptParametersViewModel> ScriptHierarchyViewModel;
+	
 	FDelegateHandle OnGraphNeedsRecompileHandle;
 
 	FOnRenamed OnRenamedDelegate;
diff --git a/Engine/Plugins/FX/Niagara/Source/NiagaraEditor/Public/ViewModels/Stack/NiagaraStackEntry.h b/Engine/Plugins/FX/Niagara/Source/NiagaraEditor/Public/ViewModels/Stack/NiagaraStackEntry.h
--- a/Engine/Plugins/FX/Niagara/Source/NiagaraEditor/Public/ViewModels/Stack/NiagaraStackEntry.h
+++ b/Engine/Plugins/FX/Niagara/Source/NiagaraEditor/Public/ViewModels/Stack/NiagaraStackEntry.h
@@ -280,6 +280,9 @@ public:
 
 	NIAGARAEDITOR_API virtual bool IsExpandedByDefault() const;
 
+	/** By default, we hide the expander if an stack entry can't be expanded. */
+	NIAGARAEDITOR_API virtual bool KeepExpanderIndentation() const;
+
 	NIAGARAEDITOR_API bool GetIsExpanded() const;
 
 	NIAGARAEDITOR_API void SetIsExpanded(bool bInExpanded);
@@ -317,7 +320,7 @@ public:
 
 	NIAGARAEDITOR_API void GetCustomFilteredChildren(TArray<UNiagaraStackEntry*>& OutFilteredChildren, const TArray<FOnFilterChild>& ChildFilters) const;
 
-	void GetFilteredChildrenOfTypes(TArray<UNiagaraStackEntry*>& OutFilteredChildren, const TSet<UClass*>& AllowedClasses) const
+	void GetFilteredChildrenOfTypes(TArray<UNiagaraStackEntry*>& OutFilteredChildren, const TSet<UClass*>& AllowedClasses, bool bRecursive = false) const
 	{
 		TArray<UNiagaraStackEntry*> FilteredChildrenTmp;
 		GetFilteredChildren(FilteredChildrenTmp);
@@ -335,13 +338,18 @@ public:
 
 			if(AllowedClasses.Contains(ChildClass))
 			{
-				OutFilteredChildren.Add(FilteredChild);	
+				OutFilteredChildren.Add(FilteredChild);
+
+				if(bRecursive)
+				{
+					FilteredChild->GetFilteredChildrenOfTypes(OutFilteredChildren, AllowedClasses, bRecursive);
+				}
 			}
 		}
 	}
 	
 	template<typename T>
-	void GetUnfilteredChildrenOfType(TArray<T*>& OutUnfilteredChildrenOfType) const
+	void GetUnfilteredChildrenOfType(TArray<T*>& OutUnfilteredChildrenOfType, bool bRecursive = false) const
 	{
 		TArray<UNiagaraStackEntry*> UnfilteredChildren;
 		GetUnfilteredChildren(UnfilteredChildren);
@@ -353,21 +361,37 @@ public:
 				OutUnfilteredChildrenOfType.Add(UnfilteredChildOfType);
 			}
 		}
+
+		if(bRecursive)
+		{
+			for (UNiagaraStackEntry* UnfilteredChild : UnfilteredChildren)
+			{
+				UnfilteredChild->GetUnfilteredChildrenOfType(OutUnfilteredChildrenOfType, bRecursive);
+			}
+		}
 	}
 
 	template<typename T>
-	void GetFilteredChildrenOfType(TArray<T*>& OutFilteredChildrenOfType) const
+	void GetFilteredChildrenOfType(TArray<T*>& OutFilteredChildrenOfType, bool bRecursive = false) const
 	{
 		TArray<UNiagaraStackEntry*> OutFilteredChildren;
 		GetFilteredChildren(OutFilteredChildren);
-		for (UNiagaraStackEntry* UnfilteredChild : OutFilteredChildren)
+		for (UNiagaraStackEntry* FilteredChild : OutFilteredChildren)
 		{
-			T* UnfilteredChildOfType = Cast<T>(UnfilteredChild);
+			T* UnfilteredChildOfType = Cast<T>(FilteredChild);
 			if (UnfilteredChildOfType != nullptr)
 			{
 				OutFilteredChildrenOfType.Add(UnfilteredChildOfType);
 			}
 		}
+
+		if(bRecursive)
+		{
+			for (UNiagaraStackEntry* FilteredChild : OutFilteredChildren)
+			{
+				FilteredChild->GetFilteredChildrenOfType(OutFilteredChildrenOfType, bRecursive);
+			}
+		}
 	}
 
 	template<typename T>
@@ -508,7 +532,7 @@ public:
 
 	virtual bool TestCanCopyWithMessage(FText& OutMessage) const { return false; }
 
-	virtual void Copy(UNiagaraClipboardContent* ClipboardContent) const { }
+	virtual void Copy(UNiagaraClipboardContent* ClipboardContent) const { OnCopyPasteDelegate.ExecuteIfBound(); }
 
 	virtual bool SupportsPaste() const { return false; }
 
@@ -516,8 +540,10 @@ public:
 
 	virtual FText GetPasteTransactionText(const UNiagaraClipboardContent* ClipboardContent) const { return FText(); }
 
-	virtual void Paste(const UNiagaraClipboardContent* ClipboardContent, FText& OutPasteWarning) { }
+	virtual void Paste(const UNiagaraClipboardContent* ClipboardContent, FText& OutPasteWarning) { OnCopyPasteDelegate.ExecuteIfBound(); }
 
+	FSimpleDelegate& OnCopyPaste() { return OnCopyPasteDelegate; }
+	
 	virtual bool SupportsDelete() const { return false; }
 
 	virtual bool TestCanDeleteWithMessage(FText& OutCanDeleteMessage) const { return false; }
@@ -603,7 +629,7 @@ protected:
 
 	bool IsSystemViewModelValid() const { return SystemViewModel.IsValid(); }
 
-
+	FSimpleDelegate OnCopyPasteDelegate;
 private:
 	NIAGARAEDITOR_API void ChildStructureChanged(ENiagaraStructureChangedFlags Info);
 
diff --git a/Engine/Plugins/FX/Niagara/Source/NiagaraEditor/Public/ViewModels/Stack/NiagaraStackFunctionInput.h b/Engine/Plugins/FX/Niagara/Source/NiagaraEditor/Public/ViewModels/Stack/NiagaraStackFunctionInput.h
--- a/Engine/Plugins/FX/Niagara/Source/NiagaraEditor/Public/ViewModels/Stack/NiagaraStackFunctionInput.h
+++ b/Engine/Plugins/FX/Niagara/Source/NiagaraEditor/Public/ViewModels/Stack/NiagaraStackFunctionInput.h
@@ -7,15 +7,16 @@
 #include "ViewModels/Stack/NiagaraParameterHandle.h"
 #include "ViewModels/Stack/NiagaraStackFunctionInputCondition.h"
 #include "NiagaraStackEditorData.h"
+#include "NiagaraStackScriptHierarchyRoot.h"
 #include "NiagaraStackFunctionInput.generated.h"
 
+class UNiagaraHierarchyScriptParameter;
 class UNiagaraNodeFunctionCall;
 class UNiagaraNodeCustomHlsl;
 class UNiagaraNodeAssignment;
 class UNiagaraNodeInput;
 class UNiagaraNodeParameterMapSet;
 class FStructOnScope;
-class UNiagaraStackFunctionInputCollection;
 class UNiagaraStackObject;
 class UNiagaraScript;
 class UEdGraphPin;
@@ -74,7 +75,7 @@ public:
 	 * @param InParameterBehavior Determines how the parameter should behave in the stack
 	 * @param InOwnerStackItemEditorDataKey The editor data key of the item that owns this input.
 	 */
-	NIAGARAEDITOR_API void Initialize(
+	void Initialize(
 		FRequiredEntryData InRequiredEntryData,
 		UNiagaraNodeFunctionCall& InModuleNode,
 		UNiagaraNodeFunctionCall& InInputFunctionCallNode,
@@ -95,6 +96,8 @@ public:
 	/** Gets the type of this input. */
 	NIAGARAEDITOR_API const FNiagaraTypeDefinition& GetInputType() const;
 
+	void SetScriptInstanceData(FNiagaraScriptInstanceData InScriptInstanceData) { ScriptInstanceData = InScriptInstanceData; }
+
 	/** Gets the unit of this input. */
 	NIAGARAEDITOR_API EUnit GetInputDisplayUnit() const;
 
@@ -117,11 +120,14 @@ public:
 	virtual bool SupportsPaste() const override { return true; }
 	NIAGARAEDITOR_API virtual bool TestCanPasteWithMessage(const UNiagaraClipboardContent* ClipboardContent, FText& OutMessage) const override;
 	NIAGARAEDITOR_API virtual FText GetPasteTransactionText(const UNiagaraClipboardContent* ClipboardContent) const override;
+	void PasteFunctionInput(const UNiagaraClipboardFunctionInput* ClipboardInput);
 	NIAGARAEDITOR_API virtual void Paste(const UNiagaraClipboardContent* ClipboardContent, FText& OutPasteWarning) override;
 	NIAGARAEDITOR_API virtual bool HasOverridenContent() const override;
 	NIAGARAEDITOR_API virtual bool SupportsSummaryView() const override;
 	NIAGARAEDITOR_API virtual FNiagaraHierarchyIdentity DetermineSummaryIdentity() const override;
-	
+	virtual bool GetCanExpand() const override;
+	virtual bool KeepExpanderIndentation() const override { return true; }
+
 	/** Gets the tooltip that should be shown for the value of this input. */
 	NIAGARAEDITOR_API FText GetValueToolTip() const;
 
@@ -407,6 +413,8 @@ private:
 	NIAGARAEDITOR_API UNiagaraScript* FindConversionScript(const FNiagaraTypeDefinition& FromType, TMap<FNiagaraTypeDefinition, UNiagaraScript*>& ConversionScriptCache, bool bIncludeConversionScripts) const;
 
 	NIAGARAEDITOR_API bool FilterInlineChildren(const UNiagaraStackEntry& Child) const;
+	NIAGARAEDITOR_API bool FilterForVisibleCondition(const UNiagaraStackEntry& NiagaraStackEntry) const;
+	bool FilterForIsInlineEditConditionToggle(const UNiagaraStackEntry& NiagaraStackEntry) const;
 
 	void ReportScriptVersionChange() const;
 
@@ -421,6 +429,9 @@ private:
 	/** The script which the owning function call is referencing. */
 	TWeakObjectPtr<UNiagaraScript> OwningFunctionCallInitialScript;
 
+	/** An optionally set hierarchy object that gives us information on child inputs etc. */
+	TWeakObjectPtr<UNiagaraHierarchyScriptParameter> HierarchyScriptParameter;
+	
 	/** The assignment node which owns this input.  This is only valid for inputs of assignment modules. */
 	TWeakObjectPtr<UNiagaraNodeAssignment> OwningAssignmentNode;
 
@@ -454,7 +465,7 @@ private:
 
 	/** The name of this input for display in the UI. */
 	FText DisplayName;
-
+	
 	/** Optional override for the display name*/
 	TOptional<FText> DisplayNameOverride;
 	TOptional<TAttribute<FText>> SummaryViewDisplayNameOverride;
@@ -507,7 +518,10 @@ private:
 	/** A multicast delegate which is called when the value of this input is changed. */
 	FOnValueChanged ValueChangedDelegate;
 
-	/** The script which owns the function which owns this input.  This is also the autoritative version of the rapid iteration parameters. */
+	/** Script instance data passed in from the parent to determine additional instance data, such as whether inputs are currently visible or not. */
+	FNiagaraScriptInstanceData ScriptInstanceData;
+	
+	/** The script which owns the function which owns this input.  This is also the authoritative version of the rapid iteration parameters. */
 	TWeakObjectPtr<UNiagaraScript> SourceScript;
 
 	/** An array of scripts which this input affects. */
