Niagara Script Editor - Hierarchy Editor Integration

diff --git a/Engine/Plugins/FX/Niagara/Source/NiagaraEditor/Public/ViewModels/Stack/NiagaraStackGraphUtilities.h b/Engine/Plugins/FX/Niagara/Source/NiagaraEditor/Public/ViewModels/Stack/NiagaraStackGraphUtilities.h
--- a/Engine/Plugins/FX/Niagara/Source/NiagaraEditor/Public/ViewModels/Stack/NiagaraStackGraphUtilities.h
+++ b/Engine/Plugins/FX/Niagara/Source/NiagaraEditor/Public/ViewModels/Stack/NiagaraStackGraphUtilities.h
@@ -183,15 +183,7 @@ namespace FNiagaraStackGraphUtilities
 	UNiagaraNodeFunctionCall* FindFunctionCallNode(FGuid FunctionCallGuid, TSharedRef<FNiagaraEmitterViewModel> EmitterViewModel);
 	TArray<UNiagaraNodeFunctionCall*> FindModuleNodesForSimulationStage(UNiagaraSimulationStageBase& SimStage, TSharedRef<FNiagaraEmitterViewModel> EmitterViewModel);
 	TArray<UNiagaraNodeFunctionCall*> FindModuleNodesForEventHandler(FNiagaraEventScriptProperties& EventScriptProperties, TSharedRef<FNiagaraEmitterViewModel> EmitterViewModel);
-	TOptional<FMatchingFunctionInputData> FindInputData(FNiagaraHierarchyIdentity ModuleInputIdentity, TSharedRef<FNiagaraEmitterViewModel> EmitterViewModel);
-	/** A more optimized version of the function above. Assumes you already have the correct node. */
-	TOptional<FMatchingFunctionInputData> FindInputData(const UNiagaraNodeFunctionCall& FunctionCallNode, FNiagaraHierarchyIdentity InputIdentity, TSharedRef<FNiagaraEmitterViewModel> EmitterViewModel);
 	TOptional<FMatchingFunctionInputData> FindAssignmentInputData(const UNiagaraNodeAssignment& AssignmentNode, FName VariableName, TSharedRef<FNiagaraEmitterViewModel> EmitterViewModel);
-	/** Finds the input data for a given function call node and a variable guid. An optional cache can be used that, if specified, will be used to get variable data if found and if not, will write into the cache
-	 * This also means maintaining the cache becomes the responsibility of the caller */
-	TOptional<FMatchingFunctionInputData> FindModuleInputData(const UNiagaraNodeFunctionCall& FunctionCallNode, FGuid VariableGuid, TSharedRef<FNiagaraEmitterViewModel> EmitterViewModel, bool bIncludeChildrenInputs = true, TMap<FInputDataCacheKey, FMatchingFunctionInputData>* OptionalCache = nullptr);
-	TArray<FGuid> GetChildrenInputGuids(const UNiagaraNodeFunctionCall& FunctionCallNode, FName ParentInputName);
-	TSet<FGuid> GetHiddenChildrenInputGuids(const UNiagaraNodeFunctionCall& FunctionCallNode, FName ParentInput, TSharedRef<FNiagaraEmitterViewModel> EmitterViewModel);
 	
 	/* Module script calls do not have direct inputs, but rely on the parameter map being initialized correctly. This utility function resolves which of the module's parameters are reachable during compilation and returns a list of pins on the parameter map node that do not have to be compiled. */
 	TArray<UEdGraphPin*> GetUnusedFunctionInputPins(const UNiagaraNodeFunctionCall& FunctionCallNode, FCompileConstantResolver ConstantResolver);
diff --git a/Engine/Plugins/FX/Niagara/Source/NiagaraEditor/Public/ViewModels/Stack/NiagaraStackInputCategory.h b/Engine/Plugins/FX/Niagara/Source/NiagaraEditor/Public/ViewModels/Stack/NiagaraStackInputCategory.h
--- a/Engine/Plugins/FX/Niagara/Source/NiagaraEditor/Public/ViewModels/Stack/NiagaraStackInputCategory.h
+++ b/Engine/Plugins/FX/Niagara/Source/NiagaraEditor/Public/ViewModels/Stack/NiagaraStackInputCategory.h
@@ -2,6 +2,8 @@
 
 #pragma once
 
+#include "NiagaraClipboard.h"
+#include "NiagaraStackScriptHierarchyRoot.h"
 #include "ViewModels/Stack/NiagaraStackItem.h"
 #include "NiagaraTypes.h"
 #include "ViewModels/HierarchyEditor/NiagaraHierarchyViewModelBase.h"
@@ -49,61 +51,45 @@ protected:
 };
 
 UCLASS(MinimalAPI)
-class UNiagaraStackInputCategory : public UNiagaraStackCategory
+class UNiagaraStackScriptHierarchyCategory : public UNiagaraStackCategory
 {
-	GENERATED_BODY() 
+	GENERATED_BODY()
 
 public:
-	NIAGARAEDITOR_API void Initialize(
-		FRequiredEntryData InRequiredEntryData,
-		FString InputCategoryStackEditorDataKey,
-		FText InCategoryName,
-		bool bInIsTopLevelCategory,
-		FString InOwnerStackItemEditorDataKey);
+	NIAGARAEDITOR_API void Initialize(FRequiredEntryData InRequiredEntryData, const UNiagaraHierarchyCategory& InHierarchyCategory, FString InOwningStackItemEditorDataKey, FString InStackEditorDataKey);
+
+	/** UNiagaraStackEntry */
+	virtual bool SupportsCopy() const override { return true; }
+	virtual bool SupportsPaste() const override { return true; }
+	virtual void Copy(UNiagaraClipboardContent* ClipboardContent) const override;
+	virtual void Paste(const UNiagaraClipboardContent* ClipboardContent, FText& OutPasteWarning) override;
+	virtual bool TestCanCopyWithMessage(FText& OutMessage) const override;
+	virtual bool TestCanPasteWithMessage(const UNiagaraClipboardContent* ClipboardContent, FText& OutMessage) const override;
+
+	void PasteFromClipboard(const UNiagaraClipboardContent* ClipboardContent);
+	TArray<const UNiagaraClipboardFunctionInput*> ToClipboardFunctionInputs(UObject* InOuter) const;
 	
-	NIAGARAEDITOR_API bool GetIsEnabled() const;
-
-	NIAGARAEDITOR_API virtual FText GetDisplayName() const override;
-
-	NIAGARAEDITOR_API void ResetInputs();
-
-	NIAGARAEDITOR_API void AddInput(UNiagaraNodeFunctionCall* InModuleNode, UNiagaraNodeFunctionCall* InInputFunctionCallNode, FName InInputParameterHandle, FNiagaraTypeDefinition InInputType, EStackParameterBehavior InParameterBehavior, TOptional<FText> InOptionalDisplayName, bool bIsHidden, bool bIsChildInput);
-
-	NIAGARAEDITOR_API void SetShouldShowInStack(bool bInShouldShowInStack);
-
-	NIAGARAEDITOR_API void ToClipboardFunctionInputs(UObject* InOuter, TArray<const UNiagaraClipboardFunctionInput*>& OutClipboardFunctionInputs) const;
-
-	NIAGARAEDITOR_API bool TrySetStaticSwitchValuesFromClipboardFunctionInput(const UNiagaraClipboardFunctionInput& ClipboardFunctionInput);
-
-	NIAGARAEDITOR_API void SetStandardValuesFromClipboardFunctionInputs(const TArray<const UNiagaraClipboardFunctionInput*>& ClipboardFunctionInputs);
-
-	NIAGARAEDITOR_API void GetFilteredChildInputs(TArray<UNiagaraStackFunctionInput*>& OutFilteredChildInputs) const;
+	virtual FText GetDisplayName() const override;
+	virtual FText GetTooltipText() const override;
+	
+	void SetOwningFunctionCallNode(UNiagaraNodeFunctionCall& InOnGetOwningFunctionCallNode) { OwningFunctionCallNode = &InOnGetOwningFunctionCallNode;}
+	void SetOwningModuleNode(UNiagaraNodeFunctionCall& InOnGetOwningModuleNode) { OwningModuleNode = &InOnGetOwningModuleNode;}
+	
+	void SetScriptInstanceData(const FNiagaraScriptInstanceData& InScriptInstanceData) { ScriptInstanceData = InScriptInstanceData; }
+	
+	const UNiagaraHierarchyCategory* GetHierarchyCategory() const { return HierarchyCategory.Get(); }
+protected:
+	virtual void RefreshChildrenInternal(const TArray<UNiagaraStackEntry*>& CurrentChildren, TArray<UNiagaraStackEntry*>& NewChildren, TArray<FStackIssue>& NewIssues) override;
 
+	virtual bool IsTopLevelCategory() const override { return HierarchyCategory->GetOuter()->IsA<UNiagaraHierarchyRoot>(); }
 protected:
-	//~ UNiagaraStackEntry interface
-	NIAGARAEDITOR_API virtual void RefreshChildrenInternal(const TArray<UNiagaraStackEntry*>& CurrentChildren, TArray<UNiagaraStackEntry*>& NewChildren, TArray<FStackIssue>& NewIssues) override;
-	virtual bool IsTopLevelCategory() const override { return bIsTopLevelCategory; }
-private:
-	struct FInputParameterHandleAndType
-	{
-		UNiagaraNodeFunctionCall* ModuleNode;
-		UNiagaraNodeFunctionCall* InputFunctionCallNode;
-		FName ParameterHandle;
-		FNiagaraTypeDefinition Type;
-		EStackParameterBehavior ParameterBehavior;
-		TOptional<FText> DisplayName;
-		bool bIsHidden;
-		bool bIsChildInput;
-	};
-
-	FText CategoryName;
-	TOptional<FText> DisplayName;
-	bool bIsTopLevelCategory;
-
-	TArray<FInputParameterHandleAndType> Inputs;
-};
+	TWeakObjectPtr<const UNiagaraHierarchyCategory> HierarchyCategory;
 
-void AddSummaryItem(UNiagaraHierarchyItemBase* HierarchyItem, UNiagaraStackEntry* Parent);
+	TWeakObjectPtr<UNiagaraNodeFunctionCall> OwningModuleNode;
+	TWeakObjectPtr<UNiagaraNodeFunctionCall> OwningFunctionCallNode;
+	
+	FNiagaraScriptInstanceData ScriptInstanceData;
+};
 
 UCLASS(MinimalAPI)
 class UNiagaraStackSummaryCategory : public UNiagaraStackCategory
diff --git a/Engine/Plugins/FX/Niagara/Source/NiagaraEditor/Public/ViewModels/Stack/NiagaraStackItemFooter.h b/Engine/Plugins/FX/Niagara/Source/NiagaraEditor/Public/ViewModels/Stack/NiagaraStackItemFooter.h
--- a/Engine/Plugins/FX/Niagara/Source/NiagaraEditor/Public/ViewModels/Stack/NiagaraStackItemFooter.h
+++ b/Engine/Plugins/FX/Niagara/Source/NiagaraEditor/Public/ViewModels/Stack/NiagaraStackItemFooter.h
@@ -6,7 +6,6 @@
 #include "NiagaraStackItemFooter.generated.h"
 
 class UNiagaraNodeFunctionCall;
-class UNiagaraStackFunctionInputCollection;
 class UNiagaraStackModuleItemOutputCollection;
 class UNiagaraNode;
 
diff --git a/Engine/Plugins/FX/Niagara/Source/NiagaraEditor/Public/ViewModels/Stack/NiagaraStackModuleItem.h b/Engine/Plugins/FX/Niagara/Source/NiagaraEditor/Public/ViewModels/Stack/NiagaraStackModuleItem.h
--- a/Engine/Plugins/FX/Niagara/Source/NiagaraEditor/Public/ViewModels/Stack/NiagaraStackModuleItem.h
+++ b/Engine/Plugins/FX/Niagara/Source/NiagaraEditor/Public/ViewModels/Stack/NiagaraStackModuleItem.h
@@ -2,13 +2,13 @@
 
 #pragma once
 
+#include "NiagaraStackScriptHierarchyRoot.h"
 #include "ViewModels/Stack/NiagaraStackItem.h"
 #include "NiagaraTypes.h"
 #include "NiagaraStackModuleItem.generated.h"
 
 class UNiagaraNodeFunctionCall;
 class UNiagaraStackModuleItemLinkedInputCollection;
-class UNiagaraStackFunctionInputCollection;
 class UNiagaraStackModuleItemOutputCollection;
 class UNiagaraScript;
 class INiagaraStackItemGroupAddUtilities;
@@ -152,10 +152,10 @@ private:
 
 	UPROPERTY()
 	TObjectPtr<UNiagaraStackModuleItemLinkedInputCollection> LinkedInputCollection;
-
+	
 	UPROPERTY()
-	TObjectPtr<UNiagaraStackFunctionInputCollection> InputCollection;
-
+	TObjectPtr<UNiagaraStackScriptHierarchyRoot> InputRoot;
+	
 	UPROPERTY()
 	TObjectPtr<UNiagaraStackModuleItemOutputCollection> OutputCollection;
 
diff --git a/Engine/Plugins/FX/Niagara/Source/NiagaraEditor/Public/ViewModels/Stack/NiagaraStackSummaryViewInputCollection.h b/Engine/Plugins/FX/Niagara/Source/NiagaraEditor/Public/ViewModels/Stack/NiagaraStackSummaryViewInputCollection.h
--- a/Engine/Plugins/FX/Niagara/Source/NiagaraEditor/Public/ViewModels/Stack/NiagaraStackSummaryViewInputCollection.h
+++ b/Engine/Plugins/FX/Niagara/Source/NiagaraEditor/Public/ViewModels/Stack/NiagaraStackSummaryViewInputCollection.h
@@ -3,7 +3,7 @@
 #pragma once
 
 #include "NiagaraEmitter.h"
-#include "NiagaraStackFunctionInputCollection.h"
+#include "NiagaraStackValueCollection.h"
 #include "ViewModels/HierarchyEditor/NiagaraHierarchyViewModelBase.h"
 #include "NiagaraStackSummaryViewInputCollection.generated.h"
 
diff --git a/Engine/Plugins/FX/Niagara/Source/NiagaraEditor/Public/Widgets/SNiagaraHierarchyEditor.h b/Engine/Plugins/FX/Niagara/Source/NiagaraEditor/Public/Widgets/SNiagaraHierarchyEditor.h
new file mode 100644
--- /dev/null
+++ b/Engine/Plugins/FX/Niagara/Source/NiagaraEditor/Public/Widgets/SNiagaraHierarchyEditor.h
@@ -0,0 +1,278 @@
+// Copyright Epic Games, Inc. All Rights Reserved.
+
+#pragma once
+
+#include "CoreMinimal.h"
+#include "SDropTarget.h"
+#include "IDetailsView.h"
+#include "NiagaraEditorStyle.h"
+#include "Widgets/Views/STreeView.h"
+#include "Widgets/Views/STableRow.h"
+#include "Widgets/Text/SInlineEditableTextBlock.h"
+#include "Widgets/Input/SMenuAnchor.h"
+#include "Styling/SlateTypes.h"
+#include "ViewModels/HierarchyEditor/NiagaraHierarchyViewModelBase.h"
+#include "Widgets/Input/SSearchBox.h"
+#include "Widgets/Input/SCheckBox.h"
+#include "Misc/NotifyHook.h"
+#include "UObject/GCObject.h"
+
+class UNiagaraHierarchyViewModelBase;
+class UNiagaraHierarchySection;
+struct FNiagaraHierarchyItemViewModelBase;
+struct FNiagaraHierarchySectionViewModel;
+
+class SNiagaraHierarchySection : public SCompoundWidget
+{
+	DECLARE_DELEGATE_OneParam(FOnSectionActivated, TSharedPtr<FNiagaraHierarchySectionViewModel> SectionViewModel)
+	
+	SLATE_BEGIN_ARGS(SNiagaraHierarchySection)
+		: _bForbidDropOn(false)
+		{}
+		SLATE_ATTRIBUTE(ECheckBoxState, IsSectionActive)
+		SLATE_EVENT(FOnSectionActivated, OnSectionActivated)
+		SLATE_ARGUMENT(bool, bForbidDropOn)
+	SLATE_END_ARGS()
+
+public:
+	void Construct(const FArguments& InArgs, TSharedPtr<struct FNiagaraHierarchySectionViewModel> InSection, TWeakObjectPtr<UNiagaraHierarchyViewModelBase> HierarchyViewModel);
+	virtual ~SNiagaraHierarchySection() override;
+
+	bool OnCanAcceptDrop(TSharedPtr<FDragDropOperation> DragDropOperation, EItemDropZone ItemDropZone) const;
+	FReply OnDroppedOn(const FGeometry&, const FDragDropEvent& DragDropEvent, EItemDropZone DropZone) const;
+
+	void TryEnterEditingMode() const;
+
+	TSharedPtr<struct FNiagaraHierarchySectionViewModel> GetSectionViewModel();
+private:
+	virtual FReply OnMouseButtonDown(const FGeometry& MyGeometry, const FPointerEvent& MouseEvent) override;
+	virtual FReply OnMouseButtonUp(const FGeometry& MyGeometry, const FPointerEvent& MouseEvent) override;
+	
+	virtual void OnDragEnter(const FGeometry& MyGeometry, const FDragDropEvent& DragDropEvent) override;
+	virtual FReply OnDragDetected(const FGeometry& MyGeometry, const FPointerEvent& MouseEvent) override;
+	virtual void OnDragLeave(const FDragDropEvent& DragDropEvent) override;
+
+	virtual FReply OnKeyDown(const FGeometry& MyGeometry, const FKeyEvent& InKeyEvent) override;
+	
+	virtual void OnMouseLeave(const FPointerEvent& MouseEvent) override;
+
+	UNiagaraHierarchySection* TryGetSectionData() const;
+
+	FText GetText() const;
+	FText GetTooltipText() const;
+	void OnRenameSection(const FText& Text, ETextCommit::Type CommitType) const;
+	bool OnVerifySectionRename(const FText& NewName, FText& OutTooltip) const;
+
+	bool IsSectionSelected() const;
+	bool IsSectionReadOnly() const;
+	ECheckBoxState GetSectionCheckState() const;
+	void OnSectionCheckChanged(ECheckBoxState NewState);
+	EActiveTimerReturnType ActivateSectionIfDragging(double CurrentTime, float DeltaTime) const;
+private:
+	TSharedPtr<SMenuAnchor> MenuAnchor;
+	TSharedPtr<SCheckBox> CheckBox;
+	TSharedPtr<SInlineEditableTextBlock> InlineEditableTextBlock;
+	TWeakObjectPtr<UNiagaraHierarchyViewModelBase> HierarchyViewModel;
+	TSharedPtr<FNiagaraHierarchySectionViewModel> SectionViewModel;
+private:
+	TAttribute<ECheckBoxState> IsSectionActive;
+	FOnSectionActivated OnSectionActivatedDelegate;
+	bool bForbidDropOn = false;
+	mutable bool bDraggedOn = false;
+};
+
+class SNiagaraHierarchyEditor : public SCompoundWidget, public FGCObject, public FNotifyHook
+{
+	DECLARE_DELEGATE_RetVal_OneParam(TSharedRef<SWidget>, FOnGenerateRowContentWidget, TSharedRef<FNiagaraHierarchyItemViewModelBase> HierarchyItem);
+	DECLARE_DELEGATE_RetVal_OneParam(TSharedRef<SWidget>, FOnGenerateCustomDetailsPanelNameWidget, TSharedPtr<FNiagaraHierarchyItemViewModelBase> HierarchyItem);
+
+	SLATE_BEGIN_ARGS(SNiagaraHierarchyEditor)
+		: _bReadOnly(true)
+		, _ItemRowStyle(&FAppStyle::Get().GetWidgetStyle<FTableRowStyle>("TableView.Row"))
+		, _CategoryRowStyle(&FNiagaraEditorStyle::Get().GetWidgetStyle<FTableRowStyle>("NiagaraEditor.HierarchyEditor.Row.Category"))
+	{}
+		SLATE_ARGUMENT(bool, bReadOnly)
+		SLATE_EVENT(FOnGenerateRowContentWidget, OnGenerateRowContentWidget)
+		SLATE_EVENT(FOnGenerateCustomDetailsPanelNameWidget, OnGenerateCustomDetailsPanelNameWidget)
+		SLATE_STYLE_ARGUMENT(FTableRowStyle, ItemRowStyle)
+		SLATE_STYLE_ARGUMENT(FTableRowStyle, CategoryRowStyle)
+	SLATE_END_ARGS()
+
+public:
+	void Construct(const FArguments& InArgs, TObjectPtr<UNiagaraHierarchyViewModelBase> InHierarchyViewModel);
+	virtual ~SNiagaraHierarchyEditor() override;
+
+	void RefreshSourceItems();
+	void RefreshAllViews(bool bFullRefresh = false);
+	void RequestRefreshAllViewsNextFrame(bool bFullRefresh = false);
+	void RefreshSourceView(bool bFullRefresh = false) const;
+	void RequestRefreshSourceViewNextFrame(bool bFullRefresh = false);
+	void RefreshHierarchyView(bool bFullRefresh = false) const;
+	void RequestRefreshHierarchyViewNextFrame(bool bFullRefresh = false);
+	void RefreshSectionsView();
+	void RequestRefreshSectionsViewNextFrame();
+
+	void NavigateToHierarchyItem(FNiagaraHierarchyIdentity Identity) const;
+	void NavigateToHierarchyItem(TSharedPtr<FNiagaraHierarchyItemViewModelBase> Item) const;
+	bool IsItemSelected(TSharedPtr<FNiagaraHierarchyItemViewModelBase> Item) const;
+	
+private:
+	// need to do this to enable focus so we can handle shortcuts
+	virtual bool SupportsKeyboardFocus() const override { return true; }
+	
+	TSharedRef<ITableRow> GenerateSourceItemRow(TSharedPtr<FNiagaraHierarchyItemViewModelBase> HierarchyItem, const TSharedRef<STableViewBase>& TableViewBase);
+	TSharedRef<ITableRow> GenerateHierarchyItemRow(TSharedPtr<FNiagaraHierarchyItemViewModelBase> HierarchyItem, const TSharedRef<STableViewBase>& TableViewBase);
+
+	bool FilterForSourceSection(TSharedPtr<const FNiagaraHierarchyItemViewModelBase> ItemViewModel) const;
+private:
+	void Reinitialize();
+
+	void BindToHierarchyRootViewModel();
+	void UnbindFromHierarchyRootViewModel() const;
+	
+	/** Source items reflect the base, unedited status of items to edit into a hierarchy */
+	const TArray<TSharedPtr<FNiagaraHierarchyItemViewModelBase>>& GetSourceItems() const;
+	
+	bool IsDetailsPanelEditingAllowed() const;
+	
+	void RequestRenameSelectedItem();
+	bool CanRequestRenameSelectedItem() const;
+
+	void ClearSourceItems() const;
+
+	void DeleteItems(TArray<TSharedPtr<FNiagaraHierarchyItemViewModelBase>> ItemsToDelete) const;
+	void DeleteSelectedHierarchyItems() const;
+	bool CanDeleteSelectedHierarchyItems() const;
+
+	void DeleteActiveSection() const;
+	bool CanDeleteActiveSection() const;
+
+	void OnItemAdded(TSharedPtr<FNiagaraHierarchyItemViewModelBase> AddedItem);
+	void OnHierarchySectionActivated(TSharedPtr<FNiagaraHierarchySectionViewModel> Section);
+	void OnSourceSectionActivated(TSharedPtr<FNiagaraHierarchySectionViewModel> Section);
+	void OnHierarchySectionAdded(TSharedPtr<FNiagaraHierarchySectionViewModel> AddedSection);
+	void OnHierarchySectionDeleted(TSharedPtr<FNiagaraHierarchySectionViewModel> DeletedSection);
+
+	void SetActiveSourceSection(TSharedPtr<struct FNiagaraHierarchySectionViewModel>);
+	TSharedPtr<FNiagaraHierarchySectionViewModel> GetActiveSourceSection() const;
+	UNiagaraHierarchySection* GetActiveSourceSectionData() const;
+	
+	void OnSelectionChanged(TSharedPtr<FNiagaraHierarchyItemViewModelBase> HierarchyItem, ESelectInfo::Type Type, bool bFromHierarchy) const;
+
+	void RunSourceSearch();
+	void OnSourceSearchTextChanged(const FText& Text);
+	void OnSourceSearchTextCommitted(const FText& Text, ETextCommit::Type CommitType);
+	void OnSearchButtonClicked(SSearchBox::SearchDirection SearchDirection);
+
+	virtual FReply OnKeyDown(const FGeometry& MyGeometry, const FKeyEvent& InKeyEvent) override;
+	virtual FReply OnMouseButtonDown(const FGeometry& MyGeometry, const FPointerEvent& MouseEvent) override;
+	virtual FReply OnMouseButtonUp(const FGeometry& MyGeometry, const FPointerEvent& MouseEvent) override;
+
+	TSharedPtr<SWidget> SummonContextMenuForSelectedRows(bool bFromHierarchy) const;
+
+	struct FSearchItem
+	{
+		TArray<TSharedPtr<FNiagaraHierarchyItemViewModelBase>> Path;
+
+		TSharedPtr<FNiagaraHierarchyItemViewModelBase> GetEntry() const
+		{
+			return Path.Num() > 0 ? 
+				Path[Path.Num() - 1] : 
+				nullptr;
+		}
+
+		bool operator==(const FSearchItem& Item) const
+		{
+			return Path == Item.Path;
+		}
+	};
+
+	/** This will recursively generated parent chain paths for all items within the given root. Used for expansion purposes. */
+	void GenerateSearchItems(TSharedRef<FNiagaraHierarchyItemViewModelBase> Root, TArray<TSharedPtr<FNiagaraHierarchyItemViewModelBase>> ParentChain, TArray<FSearchItem>& OutSearchItems);
+	void ExpandSourceSearchResults();
+	void SelectNextSourceSearchResult();
+	void SelectPreviousSourceSearchResult();
+	TOptional<SSearchBox::FSearchResultData> GetSearchResultData() const;
+	
+	FNiagaraHierarchyItemViewModelBase::FCanPerformActionResults CanDropOnRoot(TSharedPtr<FNiagaraHierarchyItemViewModelBase> DraggedItem) const;
+
+	/** Callback functions for the root widget */
+	FReply HandleHierarchyRootDrop(const FGeometry&, const FDragDropEvent& DragDropEvent) const;
+	bool OnCanDropOnRoot(TSharedPtr<FDragDropOperation> DragDropOperation) const;
+	void OnRootDragEnter(const FDragDropEvent& DragDropEvent) const;
+	void OnRootDragLeave(const FDragDropEvent& DragDropEvent) const;
+	FSlateColor GetRootIconColor() const;
+
+	virtual FString GetReferencerName() const override; 
+	virtual void AddReferencedObjects(FReferenceCollector& Collector) override;
+
+	virtual void NotifyPostChange(const FPropertyChangedEvent& PropertyChangedEvent, FProperty* PropertyThatChanged) override;
+private:
+	TWeakObjectPtr<UNiagaraHierarchyViewModelBase> HierarchyViewModel;
+	
+	TArray<FSearchItem> SourceSearchResults;
+	TOptional<FSearchItem> FocusedSearchResult;
+
+	const FTableRowStyle* CategoryRowStyle = nullptr;
+	const FTableRowStyle* ItemRowStyle = nullptr;
+	
+	mutable TWeakPtr<FNiagaraHierarchyItemViewModelBase> SelectedDetailsPanelItemViewModel;
+
+private:
+	TObjectPtr<UNiagaraHierarchyRoot> SourceRoot = nullptr;
+	TSharedPtr<FNiagaraHierarchyRootViewModel> SourceRootViewModel;
+	TSharedPtr<STreeView<TSharedPtr<FNiagaraHierarchyItemViewModelBase>>> SourceTreeView;
+	TSharedPtr<STreeView<TSharedPtr<FNiagaraHierarchyItemViewModelBase>>> HierarchyTreeView;
+	TWeakPtr<struct FNiagaraHierarchySectionViewModel> ActiveSourceSection;
+	TSharedPtr<class SWrapBox> SourceSectionBox;
+	TSharedPtr<class SWrapBox> HierarchySectionBox;
+	TSharedPtr<SSearchBox> SourceSearchBox;
+	TSharedPtr<IDetailsView> DetailsPanel;
+private:
+	FOnGenerateRowContentWidget OnGenerateRowContentWidget;
+	FOnGenerateCustomDetailsPanelNameWidget OnGenerateCustomDetailsPanelNameWidget;
+	TSharedPtr<FActiveTimerHandle> RefreshHierarchyViewNextFrameHandle;
+	TSharedPtr<FActiveTimerHandle> RefreshSourceViewNextFrameHandle;
+	TSharedPtr<FActiveTimerHandle> RefreshSectionsViewNextFrameHandle;
+};
+
+class SNiagaraHierarchyCategory : public SCompoundWidget
+{
+	SLATE_BEGIN_ARGS(SNiagaraHierarchyCategory)
+	{}
+		SLATE_EVENT(FIsSelected, IsSelected)
+	SLATE_END_ARGS()
+
+public:
+	void Construct(const FArguments& InArgs, TSharedPtr<struct FNiagaraHierarchyCategoryViewModel> InCategoryViewModel);
+
+	void EnterEditingMode() const;
+	bool OnVerifyCategoryRename(const FText& NewName, FText& OutTooltip) const;
+	
+	FText GetCategoryText() const;
+	void OnRenameCategory(const FText& NewText, ETextCommit::Type) const;
+	
+private:
+	TWeakPtr<struct FNiagaraHierarchyCategoryViewModel> CategoryViewModel;
+	TSharedPtr<SInlineEditableTextBlock> InlineEditableTextBlock;
+};
+
+class SNiagaraSectionDragDropTarget : public SDropTarget
+{
+public:
+	SLATE_BEGIN_ARGS(SNiagaraSectionDragDropTarget)
+		: _DropTargetArgs(SDropTarget::FArguments())
+	{}	
+		SLATE_ARGUMENT(SDropTarget::FArguments, DropTargetArgs)
+	SLATE_END_ARGS()
+
+	void Construct(const FArguments& InArgs, TSharedPtr<FNiagaraHierarchySectionViewModel> OwningSection, EItemDropZone ItemDropZone);
+	
+	virtual FReply OnDragOver(const FGeometry& MyGeometry, const FDragDropEvent& DragDropEvent) override;
+	virtual void OnDragLeave(const FDragDropEvent& DragDropEvent) override;
+
+private:
+	TSharedPtr<FNiagaraHierarchySectionViewModel> OwningSection;
+	/** The drop zone this drop target represents */
+	EItemDropZone DropZone = EItemDropZone::OntoItem;
+};
diff --git a/Engine/Plugins/FX/Niagara/Source/NiagaraEditor/Public/Widgets/SNiagaraParameterPanel.h b/Engine/Plugins/FX/Niagara/Source/NiagaraEditor/Public/Widgets/SNiagaraParameterPanel.h
--- a/Engine/Plugins/FX/Niagara/Source/NiagaraEditor/Public/Widgets/SNiagaraParameterPanel.h
+++ b/Engine/Plugins/FX/Niagara/Source/NiagaraEditor/Public/Widgets/SNiagaraParameterPanel.h
@@ -48,6 +48,7 @@ public:
 		SLATE_ARGUMENT(bool, ShowParameterSynchronizingWithLibraryIcon)
 		SLATE_ARGUMENT(bool, ShowParameterSynchronizingWithLibraryIconExternallyReferenced)
 		SLATE_ARGUMENT(bool, ShowParameterReferenceCounter)
+		SLATE_NAMED_SLOT(FArguments, SearchAdjacentWidget)
 	SLATE_END_ARGS();
 
 	NIAGARAEDITOR_API ~SNiagaraParameterPanel();
@@ -134,7 +135,6 @@ private:
 	void OnSectionChecked(ECheckBoxState CheckState, FText Section);
 
 	FText GetTooltipForSection(FText Section);
-
 private:
 	mutable bool bPendingRefresh;
 	mutable bool bPendingSelectionRestore;
diff --git a/Engine/Plugins/FX/Niagara/Source/NiagaraEditor/Public/Widgets/SNiagaraScriptInputPreviewPanel.h b/Engine/Plugins/FX/Niagara/Source/NiagaraEditor/Public/Widgets/SNiagaraScriptInputPreviewPanel.h
new file mode 100644
--- /dev/null
+++ b/Engine/Plugins/FX/Niagara/Source/NiagaraEditor/Public/Widgets/SNiagaraScriptInputPreviewPanel.h
@@ -0,0 +1,88 @@
+﻿// Copyright Epic Games, Inc. All Rights Reserved.
+
+#pragma once
+
+#include "CoreMinimal.h"
+#include "NiagaraGraph.h"
+#include "NiagaraScriptVariable.h"
+#include "ViewModels/HierarchyEditor/NiagaraHierarchyViewModelBase.h"
+#include "Widgets/SCompoundWidget.h"
+#include "Widgets/Input/SSearchBox.h"
+#include "Widgets/Views/STreeView.h"
+
+class UNiagaraHierarchyScriptParameter;
+class FNiagaraObjectSelection;
+class FNiagaraScriptToolkit;
+class SComboButton;
+
+class NIAGARAEDITOR_API SNiagaraScriptInputPreviewPanel : public SCompoundWidget, public FGCObject, public FSelfRegisteringEditorUndoClient
+{
+public:
+	SLATE_BEGIN_ARGS(SNiagaraScriptInputPreviewPanel)
+		{
+		}
+
+	SLATE_END_ARGS()
+	
+	void Construct(const FArguments& InArgs, UNiagaraHierarchyRoot& InRoot, TSharedRef<FNiagaraScriptToolkit> InScriptToolkit, TSharedRef<FNiagaraObjectSelection> InVariableObjectSelection);
+	virtual ~SNiagaraScriptInputPreviewPanel() override;
+	
+	void Refresh();
+private:
+	void SetupDelegates();
+	void RemoveDelegates();
+		
+	TSharedRef<ITableRow> OnGenerateRow(UNiagaraHierarchyItemBase* Item, const TSharedRef<STableViewBase>& TableViewBase) const;
+	void OnGetChildren(UNiagaraHierarchyItemBase* Item, TArray<UNiagaraHierarchyItemBase*>& OutChildren) const;
+	void OnParametersChanged(TOptional<UNiagaraGraph::FParametersChangedData> ParametersChangedData);
+
+	FReply SummonHierarchyEditor() const;
+
+	struct FSearchItem
+	{
+		TArray<UNiagaraHierarchyItemBase*> Path;
+
+		UNiagaraHierarchyItemBase* GetEntry() const
+		{
+			return Path.Num() > 0 ? 
+				Path[Path.Num() - 1] : 
+				nullptr;
+		}
+
+		bool operator==(const FSearchItem& Item) const
+		{
+			return Path == Item.Path;
+		}
+	};
+
+	void OnSearchTextChanged(const FText& Text);
+	void OnSearchTextCommitted(const FText& Text, ETextCommit::Type CommitType);
+	void OnSearchButtonClicked(SSearchBox::SearchDirection SearchDirection);
+	void GenerateSearchItems(UNiagaraHierarchyItemBase* Root, TArray<UNiagaraHierarchyItemBase*> ParentChain, TArray<FSearchItem>& OutSearchItems);
+	void ExpandSourceSearchResults();
+	void SelectNextSourceSearchResult();
+	void SelectPreviousSourceSearchResult();
+	TOptional<SSearchBox::FSearchResultData> GetSearchResultData() const;
+
+	/** FGCObject */
+	virtual FString GetReferencerName() const override;
+	virtual void AddReferencedObjects(FReferenceCollector& Collector) override;
+
+	/** FSelfRegisteringEditorUndoClient */
+	virtual void PostUndo(bool bSuccess) override;
+	virtual void PostRedo(bool bSuccess) override;
+private:
+	TWeakObjectPtr<UNiagaraHierarchyRoot> Root;
+	TArray<UNiagaraHierarchyItemBase*> RootArray;
+	TWeakPtr<FNiagaraScriptToolkit> ScriptToolkit;
+	TWeakPtr<FNiagaraObjectSelection> VariableObjectSelection;
+	TSharedPtr<SSearchBox> SearchBox;
+	TSharedPtr<STreeView<UNiagaraHierarchyItemBase*>> TreeView;
+
+	/** We construct and maintain an array of parameters _not_ in the hierarchy to ensure we display all parameters without requiring hierarchy setup. */
+	TArray<TObjectPtr<UNiagaraHierarchyScriptParameter>> TransientLeftoverParameters; 
+	
+	TArray<FSearchItem> SearchResults;
+	TOptional<FSearchItem> FocusedSearchResult;
+	TSharedPtr<SComboButton> AddParameterButton;
+};
diff --git a/Engine/Plugins/FX/Niagara/Source/NiagaraEditorWidgets/Private/SNiagaraOverviewInlineParameterBox.h b/Engine/Plugins/FX/Niagara/Source/NiagaraEditorWidgets/Private/SNiagaraOverviewInlineParameterBox.h
--- a/Engine/Plugins/FX/Niagara/Source/NiagaraEditorWidgets/Private/SNiagaraOverviewInlineParameterBox.h
+++ b/Engine/Plugins/FX/Niagara/Source/NiagaraEditorWidgets/Private/SNiagaraOverviewInlineParameterBox.h
@@ -8,7 +8,7 @@
 #include "ViewModels/Stack/NiagaraStackViewModel.h"
 #include "ViewModels/Stack/NiagaraStackFunctionInput.h"
 #include "ViewModels/Stack/NiagaraStackModuleItem.h"
-#include "ViewModels/Stack/NiagaraStackFunctionInputCollection.h"
+#include "ViewModels/Stack/NiagaraStackValueCollection.h"
 
 class SWrapBox;
 class SScrollBox;
diff --git a/Engine/Plugins/FX/Niagara/Source/NiagaraEditorWidgets/Private/SNiagaraStack.cpp b/Engine/Plugins/FX/Niagara/Source/NiagaraEditorWidgets/Private/SNiagaraStack.cpp
--- a/Engine/Plugins/FX/Niagara/Source/NiagaraEditorWidgets/Private/SNiagaraStack.cpp
+++ b/Engine/Plugins/FX/Niagara/Source/NiagaraEditorWidgets/Private/SNiagaraStack.cpp
@@ -20,7 +20,7 @@
 #include "ViewModels/Stack/NiagaraStackInputCategory.h" 
 #include "ViewModels/Stack/NiagaraStackPropertyRow.h"
 #include "ViewModels/Stack/NiagaraStackItemFooter.h"
-#include "ViewModels/Stack/NiagaraStackFunctionInputCollection.h"
+#include "ViewModels/Stack/NiagaraStackValueCollection.h"
 #include "ViewModels/Stack/NiagaraStackModuleItem.h"
 #include "ViewModels/Stack/NiagaraStackSystemSettingsGroup.h"
 #include "Framework/Commands/GenericCommands.h"
@@ -35,7 +35,7 @@
 #include "IDetailTreeNode.h"
 #include "Stack/NiagaraStackPropertyRowUtilities.h"
 #include "Stack/SNiagaraStackFunctionInputName.h"
-#include "Stack/SNiagaraStackFunctionInputCollection.h"
+#include "Stack/SNiagaraStackValueCollection.h"
 #include "Stack/SNiagaraStackFunctionInputValue.h"
 #include "Stack/SNiagaraStackInlineDynamicInput.h"
 #include "Stack/SNiagaraStackItem.h"
@@ -57,6 +57,7 @@
 #include "Widgets/SNiagaraParameterName.h"
 #include "Styling/StyleColors.h"
 #include "SResetToDefaultPropertyEditor.h"
+#include "Stack/SNiagaraStackScriptHierarchyRoot.h"
 #include "ViewModels/Stack/NiagaraStackNote.h"
 
 #define LOCTEXT_NAMESPACE "NiagaraStack"
@@ -327,7 +328,7 @@ void SNiagaraStack::Construct(const FArguments& InArgs, UNiagaraStackViewModel*
 			.TreeItemsSource(&StackViewModel->GetRootEntryAsArray())
 			.OnTreeViewScrolled(this, &SNiagaraStack::StackTreeScrolled)
 			.OnSelectionChanged(this, &SNiagaraStack::StackTreeSelectionChanged)
-			.SelectionMode(ESelectionMode::Multi)
+			.SelectionMode(ESelectionMode::Type::Multi)
 			.OnItemToString_Debug_Static(&FNiagaraStackEditorWidgetsUtilities::StackEntryToStringForListDebug)
 		]
 	];
@@ -968,7 +969,7 @@ SNiagaraStack::FRowWidgets SNiagaraStack::ConstructNameAndValueWidgetsForItem(UN
 		UNiagaraStackItemFooter* ItemExpander = CastChecked<UNiagaraStackItemFooter>(Item);
 		return FRowWidgets(SNew(SNiagaraStackItemFooter, *ItemExpander));
 	}
-	else if (Item->IsA<UNiagaraStackInputCategory>())
+	else if (Item->IsA<UNiagaraStackCategory>())
 	{
 		Container->SetOverrideNameAlignment(HAlign_Left, VAlign_Center);
 		return FRowWidgets(SNew(STextBlock)
@@ -996,7 +997,12 @@ SNiagaraStack::FRowWidgets SNiagaraStack::ConstructNameAndValueWidgetsForItem(UN
 	else if (Item->IsA<UNiagaraStackValueCollection>())
 	{
 		UNiagaraStackValueCollection* InputCollection = CastChecked<UNiagaraStackValueCollection>(Item);
-		return FRowWidgets(SNew(SNiagaraStackFunctionInputCollection, InputCollection));
+		return FRowWidgets(SNew(SNiagaraStackValueCollection, InputCollection));
+	}
+	else if (Item->IsA<UNiagaraStackScriptHierarchyRoot>())
+	{
+		UNiagaraStackScriptHierarchyRoot* ModuleHierarchyRoot = CastChecked<UNiagaraStackScriptHierarchyRoot>(Item);
+		return FRowWidgets(SNew(SNiagaraStackScriptHierarchyRoot, ModuleHierarchyRoot));
 	}
 	else if (Item->IsA<UNiagaraStackModuleItemOutputCollection>() ||
 		Item->IsA<UNiagaraStackModuleItemLinkedInputCollection>())
diff --git a/Engine/Plugins/FX/Niagara/Source/NiagaraEditorWidgets/Private/Stack/SNiagaraStackTableRow.cpp b/Engine/Plugins/FX/Niagara/Source/NiagaraEditorWidgets/Private/Stack/SNiagaraStackTableRow.cpp
--- a/Engine/Plugins/FX/Niagara/Source/NiagaraEditorWidgets/Private/Stack/SNiagaraStackTableRow.cpp
+++ b/Engine/Plugins/FX/Niagara/Source/NiagaraEditorWidgets/Private/Stack/SNiagaraStackTableRow.cpp
@@ -2,6 +2,7 @@
 
 #include "Stack/SNiagaraStackTableRow.h"
 
+#include "NiagaraEditorModule.h"
 #include "NiagaraEditorStyle.h"
 #include "NiagaraEditorWidgetsStyle.h"
 #include "Styling/AppStyle.h"
@@ -26,6 +27,8 @@
 #include "ScopedTransaction.h"
 #include "NiagaraEmitterEditorData.h"
 #include "SNiagaraStackNote.h"
+#include "Framework/Commands/GenericCommands.h"
+#include "ViewModels/NiagaraSystemViewModel.h"
 #include "ViewModels/HierarchyEditor/NiagaraSummaryViewViewModel.h"
 #include "ViewModels/Stack/NiagaraStackClipboardUtilities.h"
 #include "ViewModels/Stack/NiagaraStackNote.h"
@@ -35,6 +38,14 @@
 
 #define LOCTEXT_NAMESPACE "NiagaraStackTableRow"
 
+SNiagaraStackTableRow::~SNiagaraStackTableRow()
+{
+	if(StackEntry && StackEntry->IsFinalized() == false)
+	{
+		StackEntry->OnCopyPaste().Unbind();
+	}
+}
+
 void SNiagaraStackTableRow::Construct(const FArguments& InArgs, UNiagaraStackViewModel* InStackViewModel, UNiagaraStackEntry* InStackEntry, TSharedRef<FNiagaraStackCommandContext> InStackCommandContext, const TSharedRef<STreeView<UNiagaraStackEntry*>>& InOwnerTree)
 {
 	DefaultContentPadding = InArgs._ContentPadding;
@@ -55,6 +66,10 @@ void SNiagaraStackTableRow::Construct(const FArguments& InArgs, UNiagaraStackVie
 	ExpandedImage = FCoreStyle::Get().GetBrush("TreeArrow_Expanded");
 	CollapsedImage = FCoreStyle::Get().GetBrush("TreeArrow_Collapsed");
 
+	PulseAnimation.AddCurve(0.f, 0.5f, ECurveEaseFunction::CubicInOut);
+
+	InStackEntry->OnCopyPaste().BindSP(this, &SNiagaraStackTableRow::Pulse);
+	
 	IndicatorColor = InArgs._IndicatorColor;
 
 	ExecutionCategoryToolTipText = (InStackEntry->GetExecutionSubcategoryName() != NAME_None)
@@ -518,9 +533,49 @@ FReply SNiagaraStackTableRow::OnMouseButtonUp(const FGeometry& MyGeometry, const
 		FSlateApplication::Get().PushMenu(AsShared(), WidgetPath, MenuBuilder.MakeWidget(), MouseEvent.GetScreenSpacePosition(), FPopupTransitionEffect(FPopupTransitionEffect::ContextMenu));
 		return FReply::Handled();
 	}
+	
 	return STableRow<UNiagaraStackEntry*>::OnMouseButtonUp(MyGeometry, MouseEvent);
 }
 
+int32 SNiagaraStackTableRow::OnPaint(const FPaintArgs& Args, const FGeometry& AllottedGeometry, const FSlateRect& MyCullingRect, FSlateWindowElementList& OutDrawElements, int32 LayerId, const FWidgetStyle& InWidgetStyle, bool bParentEnabled) const
+{
+	LayerId = PaintSelection(Args, AllottedGeometry, MyCullingRect, OutDrawElements, LayerId, InWidgetStyle, bParentEnabled);
+	const FSlateBrush* BrushResource = GetBorderImage();
+		
+	const bool bEnabled = ShouldBeEnabled(bParentEnabled);
+
+	if ( BrushResource && BrushResource->DrawAs != ESlateBrushDrawType::NoDrawType )
+	{
+		FLinearColor PulseAlpha = FLinearColor(1.f, 1.f, 1.f, 1.f);
+		if(PulseAnimation.IsPlaying())
+		{
+			float Lerp = PulseAnimation.GetLerp();
+			PulseAlpha = FMath::Lerp(FAppStyle::Get().GetSlateColor("Colors.Hover2").GetSpecifiedColor(), PulseAlpha, Lerp);
+		}
+		
+		const bool bShowDisabledEffect = GetShowDisabledEffect();
+		const ESlateDrawEffect DrawEffects = (bShowDisabledEffect && !bEnabled) ? ESlateDrawEffect::DisabledEffect : ESlateDrawEffect::None;
+		FSlateDrawElement::MakeBox(
+			OutDrawElements,
+			LayerId,
+			AllottedGeometry.ToPaintGeometry(),
+			BrushResource,
+			DrawEffects,
+			BrushResource->GetTint(InWidgetStyle) * InWidgetStyle.GetColorAndOpacityTint() * GetBorderBackgroundColor().GetColor(InWidgetStyle) * PulseAlpha
+		);
+	}
+
+	LayerId = SCompoundWidget::OnPaint(Args, AllottedGeometry, MyCullingRect, OutDrawElements, LayerId, InWidgetStyle, bEnabled );
+	LayerId = PaintDropIndicator(Args, AllottedGeometry, MyCullingRect, OutDrawElements, LayerId, InWidgetStyle, bParentEnabled);
+
+	return LayerId;
+}
+
+void SNiagaraStackTableRow::Pulse()
+{
+	PulseAnimation.Play(AsShared());
+}
+
 void SetExpansionStateRecursive(UNiagaraStackEntry* StackEntry, bool bIsExpanded)
 {
 	if (StackEntry->GetCanExpand())
@@ -636,7 +691,7 @@ EVisibility SNiagaraStackTableRow::GetExpanderVisibility() const
 	}
 	else
 	{
-		return EVisibility::Collapsed;
+		return StackEntry->KeepExpanderIndentation() ? EVisibility::Hidden : EVisibility::Collapsed;
 	}
 }
 
@@ -686,6 +741,18 @@ EVisibility SNiagaraStackTableRow::GetSearchResultBorderVisibility() const
 	return StackViewModel->GetCurrentFocusedEntry() == StackEntry ? EVisibility::HitTestInvisible : EVisibility::Collapsed;
 }
 
+FSlateColor SNiagaraStackTableRow::GetInnerBackgroundColor() const
+{
+	FSlateColor Color(FLinearColor(0.f, 0.f, 0.f, 1.f));
+	if (PulseAnimation.IsPlaying())
+	{
+		float Lerp = PulseAnimation.GetLerp();
+		return FMath::Lerp(FAppStyle::Get().GetSlateColor("Colors.Hover2").GetSpecifiedColor(), Color.GetSpecifiedColor(), Lerp);
+	}
+
+	return Color;
+}
+
 void SNiagaraStackTableRow::NavigateTo(UNiagaraStackEntry* Item)
 {
 	OwnerTree->RequestNavigateToItem(Item, 0);
diff --git a/Engine/Plugins/FX/Niagara/Source/NiagaraEditorWidgets/Private/Stack/SNiagaraStackTableRow.h b/Engine/Plugins/FX/Niagara/Source/NiagaraEditorWidgets/Private/Stack/SNiagaraStackTableRow.h
--- a/Engine/Plugins/FX/Niagara/Source/NiagaraEditorWidgets/Private/Stack/SNiagaraStackTableRow.h
+++ b/Engine/Plugins/FX/Niagara/Source/NiagaraEditorWidgets/Private/Stack/SNiagaraStackTableRow.h
@@ -9,6 +9,7 @@
 #include "Widgets/Views/STreeView.h"
 #include "Styling/SlateTypes.h"
 #include "Layout/Visibility.h"
+#include "Animation/CurveSequence.h"
 
 class UNiagaraHierarchyItemBase;
 class UNiagaraStackViewModel;
@@ -43,6 +44,8 @@ public:
 		SLATE_EVENT(FOnAcceptDrop, OnAcceptDrop);
 	SLATE_END_ARGS();
 
+	virtual ~SNiagaraStackTableRow() override;
+	
 	void Construct(const FArguments& InArgs, UNiagaraStackViewModel* InStackViewModel, UNiagaraStackEntry* InStackEntry, TSharedRef<FNiagaraStackCommandContext> InStackCommandContext, const TSharedRef<STreeView<UNiagaraStackEntry*>>& InOwnerTree);
 
 	void Reset();
@@ -65,11 +68,13 @@ public:
 
 	virtual FReply OnMouseButtonDoubleClick(const FGeometry& InMyGeometry, const FPointerEvent& InMouseEvent) override;
 
-	FReply OnMouseButtonUp(const FGeometry& MyGeometry, const FPointerEvent& MouseEvent);
+	virtual FReply OnMouseButtonUp(const FGeometry& MyGeometry, const FPointerEvent& MouseEvent) override;
 
+	virtual int32 OnPaint(const FPaintArgs& Args, const FGeometry& AllottedGeometry, const FSlateRect& MyCullingRect, FSlateWindowElementList& OutDrawElements, int32 LayerId, const FWidgetStyle& InWidgetStyle, bool bParentEnabled) const override;
 private:
-	void CollapseChildren();
+	void Pulse();
 
+	void CollapseChildren();
 	void ExpandChildren();
 
 	void AddStackNote() const;
@@ -95,6 +100,8 @@ private:
 	void OnValueColumnWidthChanged(float Width);
 
 	EVisibility GetSearchResultBorderVisibility() const;
+	
+	FSlateColor GetInnerBackgroundColor() const;
 
 	void NavigateTo(UNiagaraStackEntry* Item);
 
@@ -147,4 +154,6 @@ private:
 	TArray<FOnFillRowContextMenu> OnFillRowContextMenuHanders;
 
 	TSharedPtr<FNiagaraStackCommandContext> StackCommandContext;
+
+	FCurveSequence PulseAnimation;
 };
\ No newline at end of file
