Fixed typo in include

diff --git a/Engine/Plugins/FX/Niagara/Source/NiagaraEditor/Private/Toolkits/NiagaraScriptToolkit.cpp b/Engine/Plugins/FX/Niagara/Source/NiagaraEditor/Private/Toolkits/NiagaraScriptToolkit.cpp
--- a/Engine/Plugins/FX/Niagara/Source/NiagaraEditor/Private/Toolkits/NiagaraScriptToolkit.cpp
+++ b/Engine/Plugins/FX/Niagara/Source/NiagaraEditor/Private/Toolkits/NiagaraScriptToolkit.cpp
@@ -48,7 +48,7 @@
 #include "Widgets/SNiagaraScriptGraph.h"
 #include "Widgets/SNiagaraScriptInputPreviewPanel.h"
 #include "Widgets/SNiagaraSelectedObjectsDetails.h"
-#include "WIdgets/Input/SButton.h"
+#include "Widgets/Input/SButton.h"
 #include "ToolMenus.h"
 
 #define LOCTEXT_NAMESPACE "NiagaraScriptToolkit"
