* Added plugin validation rule
* Exposed validation rule api functions

diff --git a/Engine/Plugins/FX/Niagara/Source/NiagaraEditor/Private/NiagaraValidationRules.cpp b/Engine/Plugins/FX/Niagara/Source/NiagaraEditor/Private/NiagaraValidationRules.cpp
--- a/Engine/Plugins/FX/Niagara/Source/NiagaraEditor/Private/NiagaraValidationRules.cpp
+++ b/Engine/Plugins/FX/Niagara/Source/NiagaraEditor/Private/NiagaraValidationRules.cpp
@@ -38,268 +38,6 @@
 
 #define LOCTEXT_NAMESPACE "NiagaraValidationRules"
 
-namespace NiagaraValidation
-{
-	template<typename T>
-	TArray<T*> GetStackEntries(UNiagaraStackViewModel* StackViewModel, bool bRefresh = false)
-	{
-		TArray<T*> Results;
-		TArray<UNiagaraStackEntry*> EntriesToCheck;
-		if (UNiagaraStackEntry* RootEntry = StackViewModel->GetRootEntry())
-		{
-			if (bRefresh)
-			{
-				RootEntry->RefreshChildren();
-			}
-			RootEntry->GetUnfilteredChildren(EntriesToCheck);
-		}
-		while (EntriesToCheck.Num() > 0)
-		{
-			UNiagaraStackEntry* Entry = EntriesToCheck.Pop();
-			if (T* ItemToCheck = Cast<T>(Entry))
-			{
-				Results.Add(ItemToCheck);
-			}
-			Entry->GetUnfilteredChildren(EntriesToCheck);
-		}
-		return Results;
-	}
-
-	template<typename T>
-	TArray<T*> GetAllStackEntriesInSystem(TSharedPtr<FNiagaraSystemViewModel> ViewModel, bool bRefresh = false)
-	{
-		TArray<T*> Results;
-		Results.Append(GetStackEntries<T>(ViewModel->GetSystemStackViewModel(), bRefresh));
-		TArray<TSharedRef<FNiagaraEmitterHandleViewModel>> EmitterHandleViewModels = ViewModel->GetEmitterHandleViewModels();
-		for (TSharedRef<FNiagaraEmitterHandleViewModel> EmitterHandleModel : EmitterHandleViewModels)
-		{
-			Results.Append(GetStackEntries<T>(EmitterHandleModel.Get().GetEmitterStackViewModel(), bRefresh));
-		}
-		return Results;
-	}
-
-	// helper function to retrieve a single stack entry from the system or emitter view model
-	template<typename T>
-	T* GetStackEntry(UNiagaraStackViewModel* StackViewModel, bool bRefresh = false)
-	{
-		TArray<T*> StackEntries = GetStackEntries<T>(StackViewModel, bRefresh);
-		if (StackEntries.Num() > 0)
-		{
-			return StackEntries[0];
-		}
-		return nullptr;
-	}
-
-	// helper function to get renderer stack item
-	UNiagaraStackRendererItem* GetRendererStackItem(UNiagaraStackViewModel* StackViewModel, UNiagaraRendererProperties* RendererProperties)
-	{
-		TArray<UNiagaraStackRendererItem*> RendererItems = GetStackEntries<UNiagaraStackRendererItem>(StackViewModel);
-		for (UNiagaraStackRendererItem* Item : RendererItems)
-		{
-			if (Item->GetRendererProperties() == RendererProperties)
-			{
-				return Item;
-			}
-		}
-		return nullptr;
-	}
-
-	// --------------------------------------------------------------------------------------------------------------------------------------------
-	// Common fixes and links
-	void AddGoToFXTypeLink(FNiagaraValidationResult& Result, UNiagaraEffectType* FXType)
-	{
-		if (FXType == nullptr)
-		{
-			return;
-		}
-
-		FNiagaraValidationFix& GoToValidationRulesLink = Result.Links.AddDefaulted_GetRef();
-		GoToValidationRulesLink.Description = LOCTEXT("GoToValidationRulesFix", "Go To Validation Rules");
-		TWeakObjectPtr<UNiagaraEffectType> WeakFXType = FXType;
-		GoToValidationRulesLink.FixDelegate = FNiagaraValidationFixDelegate::CreateLambda([WeakFXType]
-			{
-				FAssetToolsModule& AssetToolsModule = FModuleManager::LoadModuleChecked<FAssetToolsModule>(TEXT("AssetTools"));
-				TWeakPtr<IAssetTypeActions> WeakAssetTypeActions = AssetToolsModule.Get().GetAssetTypeActionsForClass(UNiagaraEffectType::StaticClass());
-
-				if (UNiagaraEffectType* FXType = WeakFXType.Get())
-				{
-					if (TSharedPtr<IAssetTypeActions> AssetTypeActions = WeakAssetTypeActions.Pin())
-					{
-						TArray<UObject*> AssetsToEdit;
-						AssetsToEdit.Add(FXType);
-						AssetTypeActions->OpenAssetEditor(AssetsToEdit);
-						//TODO: Is there a way for us to auto navigate to and open up the validation rules inside FXType?
-					}
-				}
-			});
-	}
-
-	FNiagaraValidationFix MakeDisableGPUSimulationFix(FVersionedNiagaraEmitterWeakPtr WeakEmitterPtr)
-	{
-		return FNiagaraValidationFix(
-			LOCTEXT("GpuUsageInfoFix_SwitchToCput", "Set emitter to CPU"),
-			FNiagaraValidationFixDelegate::CreateLambda(
-				[WeakEmitterPtr]()
-				{
-					FVersionedNiagaraEmitter VersionedEmitter = WeakEmitterPtr.ResolveWeakPtr();
-					if (FVersionedNiagaraEmitterData* VersionedEmitterData = VersionedEmitter.GetEmitterData())
-					{
-						const FScopedTransaction Transaction(LOCTEXT("SetCPUSim", "Set CPU Simulation"));
-
-						VersionedEmitter.Emitter->Modify();
-						VersionedEmitterData->SimTarget = ENiagaraSimTarget::CPUSim;
-
-						FProperty* SimTargetProperty = FindFProperty<FProperty>(FVersionedNiagaraEmitterData::StaticStruct(), GET_MEMBER_NAME_CHECKED(FVersionedNiagaraEmitterData, SimTarget));
-						FPropertyChangedEvent PropertyChangedEvent(SimTargetProperty);
-						VersionedEmitter.Emitter->PostEditChangeVersionedProperty(PropertyChangedEvent, VersionedEmitter.Version);
-
-						UNiagaraSystem::RequestCompileForEmitter(VersionedEmitter);
-					}
-				}
-			)
-		);
-	}
-
-	TArray<FNiagaraPlatformSetConflictInfo> GatherPlatformSetConflicts(const FNiagaraPlatformSet* SetA, const FNiagaraPlatformSet* SetB)
-	{
-		TArray<const FNiagaraPlatformSet*> PlatformSets = {SetA, SetB};
-		TArray<FNiagaraPlatformSetConflictInfo> Conflicts;
-		FNiagaraPlatformSet::GatherConflicts(PlatformSets, Conflicts);
-		return MoveTemp(Conflicts);
-	}
-
-	FString GetPlatformConflictsString(TConstArrayView<FNiagaraPlatformSetConflictInfo> ConflictInfos, int MaxPlatformsToShow = 4)
-	{
-		if (ConflictInfos.Num() > 0)
-		{
-			TSet<FName> ConflictPlatformNames;
-			for (const FNiagaraPlatformSetConflictInfo& ConflictInfo : ConflictInfos)
-			{
-				for (const FNiagaraPlatformSetConflictEntry& ConflictEntry : ConflictInfo.Conflicts)
-				{
-					ConflictPlatformNames.Add(ConflictEntry.ProfileName);
-				}
-			}
-
-			TStringBuilder<256> ConflictPlatformsString;
-			int NumFounds = 0;
-			for (FName PlatformName : ConflictPlatformNames)
-			{
-				if (NumFounds >= MaxPlatformsToShow)
-				{
-					ConflictPlatformsString.Append(TEXT(", ..."));
-					break;
-				}
-				if (NumFounds != 0)
-				{
-					ConflictPlatformsString.Append(TEXT(", "));
-				}
-				++NumFounds;
-				PlatformName.AppendString(ConflictPlatformsString);
-			}
-			return ConflictPlatformsString.ToString();
-		}
-		return FString();
-	}
-
-	FString GetPlatformConflictsString(const FNiagaraPlatformSet& PlatformSetA, const FNiagaraPlatformSet& PlatformSetB, int MaxPlatformsToShow = 4)
-	{
-		TArray<const FNiagaraPlatformSet*> CheckSets;
-		CheckSets.Add(&PlatformSetA);
-		CheckSets.Add(&PlatformSetB);
-
-		TArray<FNiagaraPlatformSetConflictInfo> ConflictInfos;
-		FNiagaraPlatformSet::GatherConflicts(CheckSets, ConflictInfos);
-
-		return GetPlatformConflictsString(ConflictInfos, MaxPlatformsToShow);
-	}
-
-	TSharedPtr<FNiagaraEmitterHandleViewModel> GetEmitterViewModel(const FNiagaraValidationContext& Context, UNiagaraEmitter* NiagaraEmitter)
-	{
-		if (NiagaraEmitter == nullptr)
-		{
-			return nullptr;
-		}
-
-		const TSharedRef<FNiagaraEmitterHandleViewModel>* EmitterViewModel =
-			Context.ViewModel->GetEmitterHandleViewModels().FindByPredicate(
-				[NiagaraEmitter](const TSharedRef<FNiagaraEmitterHandleViewModel>& EmitterViewModelRef)
-				{
-					FNiagaraEmitterHandle* EmitterHandle = EmitterViewModelRef->GetEmitterHandle();
-					return EmitterHandle && EmitterHandle->GetInstance().Emitter == NiagaraEmitter;
-				}
-			);
-
-		if ( EmitterViewModel )
-		{
-			return *EmitterViewModel;
-		}
-		return nullptr;
-	}
-
-	TOptional<int32> GetModuleStaticInt32Value(const UNiagaraStackModuleItem* Module, FName ParameterName)
-	{
-		TArray<UNiagaraStackFunctionInput*> ModuleInputs;
-		Module->GetParameterInputs(ModuleInputs);
-
-		for (UNiagaraStackFunctionInput* Input : ModuleInputs)
-		{
-			if (Input->IsStaticParameter() && Input->GetInputParameterHandle().GetName() == ParameterName)
-			{
-				return TOptional<int32>(*(int32*)Input->GetLocalValueStruct()->GetStructMemory());
-			}
-		}
-		return TOptional<int32>();
-	}
-
-	void SetModuleStaticInt32Value(UNiagaraStackModuleItem* Module, FName ParameterName, int32 NewValue)
-	{
-		TArray<UNiagaraStackFunctionInput*> ModuleInputs;
-		Module->GetParameterInputs(ModuleInputs);
-
-		for (UNiagaraStackFunctionInput* Input : ModuleInputs)
-		{
-			if (Input->IsStaticParameter() && Input->GetInputParameterHandle().GetName() == ParameterName)
-			{
-				TSharedRef<FStructOnScope> ValueStruct = MakeShared<FStructOnScope>(Input->GetLocalValueStruct()->GetStruct());
-				*(int32*)ValueStruct->GetStructMemory() = NewValue;
-				Input->SetLocalValue(ValueStruct);
-			}
-		}
-	};
-
-	bool StructContainsUObjectProperty(UStruct* Struct)
-	{
-		for (TFieldIterator<const FProperty> PropertyIt(Struct); PropertyIt; ++PropertyIt)
-		{
-			const FProperty* Property = *PropertyIt;
-			if (const FArrayProperty* ArrayProperty = CastField<const FArrayProperty>(Property))
-			{
-				// If we are an array change the property to be the inner one to check for struct / object
-				Property = ArrayProperty->Inner;
-			}
-
-			if (const FStructProperty* StructProperty = CastField<const FStructProperty>(Property))
-			{
-				if (StructProperty->Struct)
-				{
-					if (StructContainsUObjectProperty(StructProperty->Struct))
-					{
-						return true;
-					}
-				}
-			}
-			else if (CastField<const FWeakObjectProperty>(Property) || CastField<const FObjectProperty>(Property) || CastField<const FSoftObjectProperty>(Property))
-			{
-				return true;
-			}
-		}
-		return false;
-	}
-}
-
-// --------------------------------------------------------------------------------------------------------------------------------------------
-
 bool NiagaraValidation::HasValidationRules(UNiagaraSystem* NiagaraSystem)
 {
 	if ( NiagaraSystem != nullptr )
@@ -415,6 +153,210 @@ void NiagaraValidation::ValidateAllRulesInSystem(TSharedPtr<FNiagaraSystemViewMo
 	}
 }
 
+UNiagaraStackRendererItem* NiagaraValidation::GetRendererStackItem(UNiagaraStackViewModel* StackViewModel, UNiagaraRendererProperties* RendererProperties)
+{
+	TArray<UNiagaraStackRendererItem*> RendererItems = NiagaraValidation::GetStackEntries<UNiagaraStackRendererItem>(StackViewModel);
+	for (UNiagaraStackRendererItem* Item : RendererItems)
+	{
+		if (Item->GetRendererProperties() == RendererProperties)
+		{
+			return Item;
+		}
+	}
+	return nullptr;
+}
+
+void NiagaraValidation::AddGoToFXTypeLink(FNiagaraValidationResult& Result, UNiagaraEffectType* FXType)
+{
+	if (FXType == nullptr)
+	{
+		return;
+	}
+
+	FNiagaraValidationFix& GoToValidationRulesLink = Result.Links.AddDefaulted_GetRef();
+	GoToValidationRulesLink.Description = LOCTEXT("GoToValidationRulesFix", "Go To Validation Rules");
+	TWeakObjectPtr<UNiagaraEffectType> WeakFXType = FXType;
+	GoToValidationRulesLink.FixDelegate = FNiagaraValidationFixDelegate::CreateLambda([WeakFXType]
+		{
+			FAssetToolsModule& AssetToolsModule = FModuleManager::LoadModuleChecked<FAssetToolsModule>(TEXT("AssetTools"));
+			TWeakPtr<IAssetTypeActions> WeakAssetTypeActions = AssetToolsModule.Get().GetAssetTypeActionsForClass(UNiagaraEffectType::StaticClass());
+
+			if (UNiagaraEffectType* FXType = WeakFXType.Get())
+			{
+				if (TSharedPtr<IAssetTypeActions> AssetTypeActions = WeakAssetTypeActions.Pin())
+				{
+					TArray<UObject*> AssetsToEdit;
+					AssetsToEdit.Add(FXType);
+					AssetTypeActions->OpenAssetEditor(AssetsToEdit);
+					//TODO: Is there a way for us to auto navigate to and open up the validation rules inside FXType?
+				}
+			}
+		});
+}
+
+FNiagaraValidationFix NiagaraValidation::MakeDisableGPUSimulationFix(FVersionedNiagaraEmitterWeakPtr WeakEmitterPtr)
+{
+	return FNiagaraValidationFix(
+			LOCTEXT("GpuUsageInfoFix_SwitchToCput", "Set emitter to CPU"),
+			FNiagaraValidationFixDelegate::CreateLambda(
+				[WeakEmitterPtr]()
+				{
+					FVersionedNiagaraEmitter VersionedEmitter = WeakEmitterPtr.ResolveWeakPtr();
+					if (FVersionedNiagaraEmitterData* VersionedEmitterData = VersionedEmitter.GetEmitterData())
+					{
+						const FScopedTransaction Transaction(LOCTEXT("SetCPUSim", "Set CPU Simulation"));
+
+						VersionedEmitter.Emitter->Modify();
+						VersionedEmitterData->SimTarget = ENiagaraSimTarget::CPUSim;
+
+						FProperty* SimTargetProperty = FindFProperty<FProperty>(FVersionedNiagaraEmitterData::StaticStruct(), GET_MEMBER_NAME_CHECKED(FVersionedNiagaraEmitterData, SimTarget));
+						FPropertyChangedEvent PropertyChangedEvent(SimTargetProperty);
+						VersionedEmitter.Emitter->PostEditChangeVersionedProperty(PropertyChangedEvent, VersionedEmitter.Version);
+
+						UNiagaraSystem::RequestCompileForEmitter(VersionedEmitter);
+					}
+				}
+			)
+		);
+}
+
+TArray<FNiagaraPlatformSetConflictInfo> NiagaraValidation::GatherPlatformSetConflicts(const FNiagaraPlatformSet* SetA, const FNiagaraPlatformSet* SetB)
+{
+	TArray<const FNiagaraPlatformSet*> PlatformSets = {SetA, SetB};
+	TArray<FNiagaraPlatformSetConflictInfo> Conflicts;
+	FNiagaraPlatformSet::GatherConflicts(PlatformSets, Conflicts);
+	return MoveTemp(Conflicts);
+}
+
+FString NiagaraValidation::GetPlatformConflictsString(TConstArrayView<FNiagaraPlatformSetConflictInfo> ConflictInfos, int MaxPlatformsToShow)
+{
+	if (ConflictInfos.Num() > 0)
+	{
+		TSet<FName> ConflictPlatformNames;
+		for (const FNiagaraPlatformSetConflictInfo& ConflictInfo : ConflictInfos)
+		{
+			for (const FNiagaraPlatformSetConflictEntry& ConflictEntry : ConflictInfo.Conflicts)
+			{
+				ConflictPlatformNames.Add(ConflictEntry.ProfileName);
+			}
+		}
+
+		TStringBuilder<256> ConflictPlatformsString;
+		int NumFounds = 0;
+		for (FName PlatformName : ConflictPlatformNames)
+		{
+			if (NumFounds >= MaxPlatformsToShow)
+			{
+				ConflictPlatformsString.Append(TEXT(", ..."));
+				break;
+			}
+			if (NumFounds != 0)
+			{
+				ConflictPlatformsString.Append(TEXT(", "));
+			}
+			++NumFounds;
+			PlatformName.AppendString(ConflictPlatformsString);
+		}
+		return ConflictPlatformsString.ToString();
+	}
+	return FString();
+}
+
+FString NiagaraValidation::GetPlatformConflictsString(const FNiagaraPlatformSet& PlatformSetA, const FNiagaraPlatformSet& PlatformSetB, int MaxPlatformsToShow)
+{
+	TArray<const FNiagaraPlatformSet*> CheckSets;
+	CheckSets.Add(&PlatformSetA);
+	CheckSets.Add(&PlatformSetB);
+
+	TArray<FNiagaraPlatformSetConflictInfo> ConflictInfos;
+	FNiagaraPlatformSet::GatherConflicts(CheckSets, ConflictInfos);
+
+	return GetPlatformConflictsString(ConflictInfos, MaxPlatformsToShow);
+}
+
+TSharedPtr<FNiagaraEmitterHandleViewModel> NiagaraValidation::GetEmitterViewModel(const FNiagaraValidationContext& Context, UNiagaraEmitter* NiagaraEmitter)
+{
+	if (NiagaraEmitter == nullptr)
+	{
+		return nullptr;
+	}
+
+	const TSharedRef<FNiagaraEmitterHandleViewModel>* EmitterViewModel =
+		Context.ViewModel->GetEmitterHandleViewModels().FindByPredicate(
+			[NiagaraEmitter](const TSharedRef<FNiagaraEmitterHandleViewModel>& EmitterViewModelRef)
+			{
+				FNiagaraEmitterHandle* EmitterHandle = EmitterViewModelRef->GetEmitterHandle();
+				return EmitterHandle && EmitterHandle->GetInstance().Emitter == NiagaraEmitter;
+			}
+		);
+
+	if ( EmitterViewModel )
+	{
+		return *EmitterViewModel;
+	}
+	return nullptr;
+}
+
+TOptional<int32> NiagaraValidation::GetModuleStaticInt32Value(const UNiagaraStackModuleItem* Module, FName ParameterName)
+{
+	TArray<UNiagaraStackFunctionInput*> ModuleInputs;
+	Module->GetParameterInputs(ModuleInputs);
+
+	for (UNiagaraStackFunctionInput* Input : ModuleInputs)
+	{
+		if (Input->IsStaticParameter() && Input->GetInputParameterHandle().GetName() == ParameterName)
+		{
+			return TOptional<int32>(*(int32*)Input->GetLocalValueStruct()->GetStructMemory());
+		}
+	}
+	return TOptional<int32>();
+}
+
+void NiagaraValidation::SetModuleStaticInt32Value(UNiagaraStackModuleItem* Module, FName ParameterName, int32 NewValue)
+{
+	TArray<UNiagaraStackFunctionInput*> ModuleInputs;
+	Module->GetParameterInputs(ModuleInputs);
+
+	for (UNiagaraStackFunctionInput* Input : ModuleInputs)
+	{
+		if (Input->IsStaticParameter() && Input->GetInputParameterHandle().GetName() == ParameterName)
+		{
+			TSharedRef<FStructOnScope> ValueStruct = MakeShared<FStructOnScope>(Input->GetLocalValueStruct()->GetStruct());
+			*(int32*)ValueStruct->GetStructMemory() = NewValue;
+			Input->SetLocalValue(ValueStruct);
+		}
+	}
+}
+
+bool NiagaraValidation::StructContainsUObjectProperty(UStruct* Struct)
+{
+	for (TFieldIterator<const FProperty> PropertyIt(Struct); PropertyIt; ++PropertyIt)
+	{
+		const FProperty* Property = *PropertyIt;
+		if (const FArrayProperty* ArrayProperty = CastField<const FArrayProperty>(Property))
+		{
+			// If we are an array change the property to be the inner one to check for struct / object
+			Property = ArrayProperty->Inner;
+		}
+
+		if (const FStructProperty* StructProperty = CastField<const FStructProperty>(Property))
+		{
+			if (StructProperty->Struct)
+			{
+				if (StructContainsUObjectProperty(StructProperty->Struct))
+				{
+					return true;
+				}
+			}
+		}
+		else if (CastField<const FWeakObjectProperty>(Property) || CastField<const FObjectProperty>(Property) || CastField<const FSoftObjectProperty>(Property))
+		{
+			return true;
+		}
+	}
+	return false;
+}
+
 void UNiagaraValidationRule_NoWarmupTime::CheckValidity(const FNiagaraValidationContext& Context, TArray<FNiagaraValidationResult>& Results)  const
 {
 	UNiagaraSystem& System = Context.ViewModel->GetSystem();
diff --git a/Engine/Plugins/FX/Niagara/Source/NiagaraEditor/Public/NiagaraValidationRules.h b/Engine/Plugins/FX/Niagara/Source/NiagaraEditor/Public/NiagaraValidationRules.h
--- a/Engine/Plugins/FX/Niagara/Source/NiagaraEditor/Public/NiagaraValidationRules.h
+++ b/Engine/Plugins/FX/Niagara/Source/NiagaraEditor/Public/NiagaraValidationRules.h
@@ -6,15 +6,87 @@
 #include "NiagaraValidationRule.h"
 #include "NiagaraPlatformSet.h"
 #include "NiagaraRendererProperties.h"
+#include "ViewModels/NiagaraEmitterHandleViewModel.h"
+#include "ViewModels/NiagaraSystemViewModel.h"
+#include "ViewModels/Stack/NiagaraStackEntry.h"
+#include "ViewModels/Stack/NiagaraStackViewModel.h"
 #include "NiagaraValidationRules.generated.h"
 
 class UNiagaraEffectType;
 class UNiagaraScript;
+class UNiagaraStackModuleItem;
+struct FNiagaraPlatformSetConflictInfo;
+class UNiagaraStackRendererItem;
 
 namespace NiagaraValidation
 {
 	NIAGARAEDITOR_API bool HasValidationRules(UNiagaraSystem* NiagaraSystem);
 	NIAGARAEDITOR_API void ValidateAllRulesInSystem(TSharedPtr<FNiagaraSystemViewModel> ViewModel, TFunction<void(const FNiagaraValidationResult& Result)> ResultCallback);
+	
+	template<typename T>
+	TArray<T*> GetStackEntries(UNiagaraStackViewModel* StackViewModel, bool bRefresh = false)
+	{
+		TArray<T*> Results;
+		TArray<UNiagaraStackEntry*> EntriesToCheck;
+		if (UNiagaraStackEntry* RootEntry = StackViewModel->GetRootEntry())
+		{
+			if (bRefresh)
+			{
+				RootEntry->RefreshChildren();
+			}
+			RootEntry->GetUnfilteredChildren(EntriesToCheck);
+		}
+		while (EntriesToCheck.Num() > 0)
+		{
+			UNiagaraStackEntry* Entry = EntriesToCheck.Pop();
+			if (T* ItemToCheck = Cast<T>(Entry))
+			{
+				Results.Add(ItemToCheck);
+			}
+			Entry->GetUnfilteredChildren(EntriesToCheck);
+		}
+		return Results;
+	}
+
+	template<typename T>
+	TArray<T*> GetAllStackEntriesInSystem(TSharedPtr<FNiagaraSystemViewModel> ViewModel, bool bRefresh = false)
+	{
+		TArray<T*> Results;
+		Results.Append(NiagaraValidation::GetStackEntries<T>(ViewModel->GetSystemStackViewModel(), bRefresh));
+		TArray<TSharedRef<FNiagaraEmitterHandleViewModel>> EmitterHandleViewModels = ViewModel->GetEmitterHandleViewModels();
+		for (TSharedRef<FNiagaraEmitterHandleViewModel> EmitterHandleModel : EmitterHandleViewModels)
+		{
+			Results.Append(NiagaraValidation::GetStackEntries<T>(EmitterHandleModel.Get().GetEmitterStackViewModel(), bRefresh));
+		}
+		return Results;
+	}
+
+	// helper function to retrieve a single stack entry from the system or emitter view model
+	template<typename T>
+	T* GetStackEntry(UNiagaraStackViewModel* StackViewModel, bool bRefresh = false)
+	{
+		TArray<T*> StackEntries = NiagaraValidation::GetStackEntries<T>(StackViewModel, bRefresh);
+		if (StackEntries.Num() > 0)
+		{
+			return StackEntries[0];
+		}
+		return nullptr;
+	}
+
+	// helper function to get renderer stack item
+	NIAGARAEDITOR_API UNiagaraStackRendererItem* GetRendererStackItem(UNiagaraStackViewModel* StackViewModel, UNiagaraRendererProperties* RendererProperties);
+
+	// --------------------------------------------------------------------------------------------------------------------------------------------
+	// Common fixes and links
+	NIAGARAEDITOR_API void AddGoToFXTypeLink(FNiagaraValidationResult& Result, UNiagaraEffectType* FXType);
+	NIAGARAEDITOR_API FNiagaraValidationFix MakeDisableGPUSimulationFix(FVersionedNiagaraEmitterWeakPtr WeakEmitterPtr);
+	NIAGARAEDITOR_API TArray<FNiagaraPlatformSetConflictInfo> GatherPlatformSetConflicts(const FNiagaraPlatformSet* SetA, const FNiagaraPlatformSet* SetB);
+	NIAGARAEDITOR_API FString GetPlatformConflictsString(TConstArrayView<FNiagaraPlatformSetConflictInfo> ConflictInfos, int MaxPlatformsToShow = 4);
+	NIAGARAEDITOR_API FString GetPlatformConflictsString(const FNiagaraPlatformSet& PlatformSetA, const FNiagaraPlatformSet& PlatformSetB, int MaxPlatformsToShow = 4);
+	NIAGARAEDITOR_API TSharedPtr<FNiagaraEmitterHandleViewModel> GetEmitterViewModel(const FNiagaraValidationContext& Context, UNiagaraEmitter* NiagaraEmitter);
+	NIAGARAEDITOR_API TOptional<int32> GetModuleStaticInt32Value(const UNiagaraStackModuleItem* Module, FName ParameterName);
+	NIAGARAEDITOR_API void SetModuleStaticInt32Value(UNiagaraStackModuleItem* Module, FName ParameterName, int32 NewValue);
+	NIAGARAEDITOR_API bool StructContainsUObjectProperty(UStruct* Struct);
 }
 
 /** This validation rule ensures that systems don't have a warmup time set. */
