Fixed python upgrade scripts for Niagara not finding children inputs. This was due to Hierarchy Editor restructuring children inputs

diff --git a/Engine/Plugins/FX/Niagara/Source/NiagaraEditor/Private/NiagaraEditorUtilities.cpp b/Engine/Plugins/FX/Niagara/Source/NiagaraEditor/Private/NiagaraEditorUtilities.cpp
--- a/Engine/Plugins/FX/Niagara/Source/NiagaraEditor/Private/NiagaraEditorUtilities.cpp
+++ b/Engine/Plugins/FX/Niagara/Source/NiagaraEditor/Private/NiagaraEditorUtilities.cpp
@@ -3312,12 +3312,18 @@ void FNiagaraEditorUtilities::RefreshAllScriptsFromExternalChanges(FRefreshAllSc
 TArray<UNiagaraPythonScriptModuleInput*> GetFunctionCallInputs(UNiagaraClipboardContent* ClipboardContent)
 {
 	TArray<UNiagaraPythonScriptModuleInput*> ScriptInputs;
-	for (const UNiagaraClipboardFunctionInput* FunctionInput : ClipboardContent->FunctionInputs)
+
+	TArray<const UNiagaraClipboardFunctionInput*> FunctionInputsWorkingSet = ClipboardContent->FunctionInputs;
+	for(auto It = FunctionInputsWorkingSet.CreateIterator(); It; ++It)
 	{
 		UNiagaraPythonScriptModuleInput* ScriptInput = NewObject<UNiagaraPythonScriptModuleInput>();
-		ScriptInput->Input = FunctionInput;
+		ScriptInput->Input = *It;
 		ScriptInputs.Add(ScriptInput);
+
+		FunctionInputsWorkingSet.Append((*It)->ChildrenInputs);
+		It.RemoveCurrent();
 	}
+	
 	return ScriptInputs;
 }
 
