Fixed input preview & hierarchy editor not refreshing properly when changing script versions
- Added ToText function to allow text widget binding

diff --git a/Engine/Plugins/FX/Niagara/Source/NiagaraEditor/Private/Toolkits/NiagaraScriptToolkit.cpp b/Engine/Plugins/FX/Niagara/Source/NiagaraEditor/Private/Toolkits/NiagaraScriptToolkit.cpp
--- a/Engine/Plugins/FX/Niagara/Source/NiagaraEditor/Private/Toolkits/NiagaraScriptToolkit.cpp
+++ b/Engine/Plugins/FX/Niagara/Source/NiagaraEditor/Private/Toolkits/NiagaraScriptToolkit.cpp
@@ -357,6 +357,12 @@ void FNiagaraScriptToolkit::InitViewWithVersionedData()
 		RefreshDetailsPanelDelegate
 	);
 
+	// As the input preview panel is bound to the previously edited niagara script, make sure to remove delegates before the script view model points to the new version
+	if(InputPreviewPanel)
+	{
+		InputPreviewPanel->RemoveDelegates();
+	}
+	
 	ScriptViewModel->Initialize(EditedNiagaraScript, OriginalNiagaraScript);
 	ScriptViewModel->GetGraphViewModel()->SetDisplayName(GetGraphEditorDisplayName());
 	if (ParameterPanelViewModel)
@@ -379,6 +385,18 @@ void FNiagaraScriptToolkit::InitViewWithVersionedData()
 	{
 		DetailsView->SetObjects(DetailsScriptSelection->GetSelectedObjectsResolved().Array(), true);
 	}
+	
+	if(ParametersHierarchyViewModel)
+	{
+		ParametersHierarchyViewModel->Finalize();
+		ParametersHierarchyViewModel->Initialize(ScriptViewModel.ToSharedRef());
+	}
+
+	if(InputPreviewPanel)
+	{
+		InputPreviewPanel->SetupDelegates();
+		InputPreviewPanel->Refresh();
+	}	
 
 	// add listeners
 	OnEditedScriptGraphChangedHandle = ScriptViewModel->GetGraphViewModel()->GetGraph()->AddOnGraphNeedsRecompileHandler(
@@ -570,7 +588,7 @@ TSharedRef<SDockTab> FNiagaraScriptToolkit::SpawnTabInputsPreview(const FSpawnTa
 			]
 			+ SVerticalBox::Slot()
 			[
-				SNew(SNiagaraScriptInputPreviewPanel, *Cast<UNiagaraScriptSource>(EditedNiagaraScript.GetScriptData()->GetSource())->NodeGraph->GetScriptParameterHierarchyRoot(), FNiagaraScriptToolkit::SharedThis(this), ScriptViewModel->GetVariableSelection())
+				SAssignNew(InputPreviewPanel, SNiagaraScriptInputPreviewPanel, FNiagaraScriptToolkit::SharedThis(this), ScriptViewModel->GetVariableSelection())
 			]
 		];
 
diff --git a/Engine/Plugins/FX/Niagara/Source/NiagaraEditor/Private/Toolkits/NiagaraScriptToolkit.h b/Engine/Plugins/FX/Niagara/Source/NiagaraEditor/Private/Toolkits/NiagaraScriptToolkit.h
--- a/Engine/Plugins/FX/Niagara/Source/NiagaraEditor/Private/Toolkits/NiagaraScriptToolkit.h
+++ b/Engine/Plugins/FX/Niagara/Source/NiagaraEditor/Private/Toolkits/NiagaraScriptToolkit.h
@@ -29,6 +29,7 @@ class FNiagaraStandaloneScriptViewModel;
 class FNiagaraScriptToolkitParameterPanelViewModel;
 class FNiagaraScriptToolkitParameterDefinitionsPanelViewModel;
 class SNiagaraSelectedObjectsDetails;
+class SNiagaraScriptInputPreviewPanel;
 
 /** Viewer/editor for a DataTable */
 class FNiagaraScriptToolkit : public FAssetEditorToolkit, public FGCObject, public FTickableEditorObject, public FEditorUndoClient
@@ -205,6 +206,7 @@ private:
 	bool bShowedEditingVersionWarning = false;
 
 	TSharedPtr<class SNiagaraScriptGraph> NiagaraScriptGraphWidget;
+	TSharedPtr<SNiagaraScriptInputPreviewPanel> InputPreviewPanel;
 	TSharedPtr<class IDetailsView> DetailsView;
 	TObjectPtr<UNiagaraVersionMetaData> VersionMetadata = nullptr;
 	FText GetGraphEditorDisplayName() const;
diff --git a/Engine/Plugins/FX/Niagara/Source/NiagaraEditor/Private/Widgets/SNiagaraScriptInputPreviewPanel.cpp b/Engine/Plugins/FX/Niagara/Source/NiagaraEditor/Private/Widgets/SNiagaraScriptInputPreviewPanel.cpp
--- a/Engine/Plugins/FX/Niagara/Source/NiagaraEditor/Private/Widgets/SNiagaraScriptInputPreviewPanel.cpp
+++ b/Engine/Plugins/FX/Niagara/Source/NiagaraEditor/Private/Widgets/SNiagaraScriptInputPreviewPanel.cpp
@@ -21,9 +21,8 @@ BEGIN_SLATE_FUNCTION_BUILD_OPTIMIZATION
 
 #define LOCTEXT_NAMESPACE "NiagaraEditor"
 
-void SNiagaraScriptInputPreviewPanel::Construct(const FArguments& InArgs, UNiagaraHierarchyRoot& InRoot, TSharedRef<FNiagaraScriptToolkit> InScriptToolkit, TSharedRef<FNiagaraObjectSelection> InVariableObjectSelection)
+void SNiagaraScriptInputPreviewPanel::Construct(const FArguments& InArgs, TSharedRef<FNiagaraScriptToolkit> InScriptToolkit, TSharedRef<FNiagaraObjectSelection> InVariableObjectSelection)
 {
-	Root = &InRoot;
 	ScriptToolkit = InScriptToolkit;
 	VariableObjectSelection = InVariableObjectSelection;
 
@@ -74,6 +73,12 @@ SNiagaraScriptInputPreviewPanel::~SNiagaraScriptInputPreviewPanel()
 
 void SNiagaraScriptInputPreviewPanel::Refresh()
 {
+	if(ScriptToolkit.IsValid() == false || ScriptToolkit.Pin()->GetHierarchyViewModel() == nullptr || ScriptToolkit.Pin()->GetHierarchyViewModel()->GetHierarchyRoot() == nullptr)
+	{
+		return;
+	}
+	
+	TWeakObjectPtr<UNiagaraHierarchyRoot> Root = ScriptToolkit.Pin()->GetHierarchyViewModel()->GetHierarchyRoot();
 	if(Root.IsValid())
 	{
 		TransientLeftoverParameters.Empty();
@@ -130,7 +135,7 @@ TSharedRef<ITableRow> SNiagaraScriptInputPreviewPanel::OnGenerateRow(UNiagaraHie
 		.Padding(FMargin(0.f, 6.f))
 		[
 			SNew(SRichTextBlock)
-			.Text(FText::FromString(Item->ToString()))
+			.Text_UObject(Item, &UNiagaraHierarchyItemBase::ToText)
 			.TextStyle(&FNiagaraEditorStyle::Get().GetWidgetStyle<FTextBlockStyle>("NiagaraEditor.Parameters.HeaderText"))
 		];
 	}
diff --git a/Engine/Plugins/FX/Niagara/Source/NiagaraEditor/Public/ViewModels/HierarchyEditor/NiagaraHierarchyViewModelBase.h b/Engine/Plugins/FX/Niagara/Source/NiagaraEditor/Public/ViewModels/HierarchyEditor/NiagaraHierarchyViewModelBase.h
--- a/Engine/Plugins/FX/Niagara/Source/NiagaraEditor/Public/ViewModels/HierarchyEditor/NiagaraHierarchyViewModelBase.h
+++ b/Engine/Plugins/FX/Niagara/Source/NiagaraEditor/Public/ViewModels/HierarchyEditor/NiagaraHierarchyViewModelBase.h
@@ -127,7 +127,8 @@ public:
 	void SortChildren(const PREDICATE_CLASS& Predicate, bool bRecursive = false);
 	
 	virtual FString ToString() const { return GetName(); }
-
+	FText ToText() const { return FText::FromString(ToString()); }
+	
 	/** An identity can be optionally set to create a mapping from previously existing guids or names to hierarchy items that represent them. */
 	void SetIdentity(FNiagaraHierarchyIdentity InIdentity) { Identity = InIdentity; }
 	FNiagaraHierarchyIdentity GetPersistentIdentity() const { return Identity; }
diff --git a/Engine/Plugins/FX/Niagara/Source/NiagaraEditor/Public/Widgets/SNiagaraScriptInputPreviewPanel.h b/Engine/Plugins/FX/Niagara/Source/NiagaraEditor/Public/Widgets/SNiagaraScriptInputPreviewPanel.h
--- a/Engine/Plugins/FX/Niagara/Source/NiagaraEditor/Public/Widgets/SNiagaraScriptInputPreviewPanel.h
+++ b/Engine/Plugins/FX/Niagara/Source/NiagaraEditor/Public/Widgets/SNiagaraScriptInputPreviewPanel.h
@@ -24,14 +24,13 @@ public:
 
 	SLATE_END_ARGS()
 	
-	void Construct(const FArguments& InArgs, UNiagaraHierarchyRoot& InRoot, TSharedRef<FNiagaraScriptToolkit> InScriptToolkit, TSharedRef<FNiagaraObjectSelection> InVariableObjectSelection);
+	void Construct(const FArguments& InArgs, TSharedRef<FNiagaraScriptToolkit> InScriptToolkit, TSharedRef<FNiagaraObjectSelection> InVariableObjectSelection);
 	virtual ~SNiagaraScriptInputPreviewPanel() override;
 	
 	void Refresh();
-private:
 	void SetupDelegates();
 	void RemoveDelegates();
-		
+private:
 	TSharedRef<ITableRow> OnGenerateRow(UNiagaraHierarchyItemBase* Item, const TSharedRef<STableViewBase>& TableViewBase) const;
 	void OnGetChildren(UNiagaraHierarchyItemBase* Item, TArray<UNiagaraHierarchyItemBase*>& OutChildren) const;
 	void OnParametersChanged(TOptional<UNiagaraGraph::FParametersChangedData> ParametersChangedData);
@@ -72,7 +71,6 @@ private:
 	virtual void PostUndo(bool bSuccess) override;
 	virtual void PostRedo(bool bSuccess) override;
 private:
-	TWeakObjectPtr<UNiagaraHierarchyRoot> Root;
 	TArray<UNiagaraHierarchyItemBase*> RootArray;
 	TWeakPtr<FNiagaraScriptToolkit> ScriptToolkit;
 	TWeakPtr<FNiagaraObjectSelection> VariableObjectSelection;
