Discarding the stack input root when switching module versions.
This ensures a complete reconstruction. Previously, when only order/hierarchy would change, the caching mechanism of the stack would reuse items in the wrong manner

diff --git a/Engine/Plugins/FX/Niagara/Source/NiagaraEditor/Private/ViewModels/Stack/NiagaraStackModuleItem.cpp b/Engine/Plugins/FX/Niagara/Source/NiagaraEditor/Private/ViewModels/Stack/NiagaraStackModuleItem.cpp
--- a/Engine/Plugins/FX/Niagara/Source/NiagaraEditor/Private/ViewModels/Stack/NiagaraStackModuleItem.cpp
+++ b/Engine/Plugins/FX/Niagara/Source/NiagaraEditor/Private/ViewModels/Stack/NiagaraStackModuleItem.cpp
@@ -146,11 +146,19 @@ void UNiagaraStackModuleItem::RefreshChildrenInternal(const TArray<UNiagaraStack
 		FunctionCallNode->GetGraph() != nullptr && 
 		FunctionCallNode->GetGraph()->Nodes.Contains(FunctionCallNode))
 	{
+		bool bDiscardInputRoot = LastRefreshVersionCache != FunctionCallNode->SelectedScriptVersion;
+		LastRefreshVersionCache = FunctionCallNode->SelectedScriptVersion; 
+
 		// Determine if meta-data requires that we add our own refresh button here.
 		if (FunctionCallNode->HasValidScriptAndGraph())
 		{
 			bCanRefresh = true;
 		}
+
+		if(bDiscardInputRoot)
+		{
+			InputRoot = nullptr;
+		}
 		
 		if(InputRoot == nullptr)
 		{
diff --git a/Engine/Plugins/FX/Niagara/Source/NiagaraEditor/Public/ViewModels/Stack/NiagaraStackModuleItem.h b/Engine/Plugins/FX/Niagara/Source/NiagaraEditor/Public/ViewModels/Stack/NiagaraStackModuleItem.h
--- a/Engine/Plugins/FX/Niagara/Source/NiagaraEditor/Public/ViewModels/Stack/NiagaraStackModuleItem.h
+++ b/Engine/Plugins/FX/Niagara/Source/NiagaraEditor/Public/ViewModels/Stack/NiagaraStackModuleItem.h
@@ -171,4 +171,6 @@ private:
 	TArray<FStackIssue> MessageManagerIssues;
 
 	FGuid MessageLogGuid;
+
+	FGuid LastRefreshVersionCache;
 };
