[Niagara] Non-unity CIS fixes for CL 35563160.

diff --git a/Engine/Plugins/FX/Niagara/Source/NiagaraEditor/Private/ViewModels/Stack/NiagaraStackRenderItemGroup.cpp b/Engine/Plugins/FX/Niagara/Source/NiagaraEditor/Private/ViewModels/Stack/NiagaraStackRenderItemGroup.cpp
--- a/Engine/Plugins/FX/Niagara/Source/NiagaraEditor/Private/ViewModels/Stack/NiagaraStackRenderItemGroup.cpp
+++ b/Engine/Plugins/FX/Niagara/Source/NiagaraEditor/Private/ViewModels/Stack/NiagaraStackRenderItemGroup.cpp
@@ -13,6 +13,7 @@
 #include "NiagaraClipboard.h"
 #include "NiagaraEditorModule.h"
 #include "NiagaraEditorSettings.h"
+#include "NiagaraStackEditorData.h"
 
 #include "ScopedTransaction.h"
 #include "Widgets/Notifications/SNotificationList.h"
diff --git a/Engine/Plugins/FX/Niagara/Source/NiagaraEditor/Private/ViewModels/Stack/NiagaraStackRenderersOwner.cpp b/Engine/Plugins/FX/Niagara/Source/NiagaraEditor/Private/ViewModels/Stack/NiagaraStackRenderersOwner.cpp
--- a/Engine/Plugins/FX/Niagara/Source/NiagaraEditor/Private/ViewModels/Stack/NiagaraStackRenderersOwner.cpp
+++ b/Engine/Plugins/FX/Niagara/Source/NiagaraEditor/Private/ViewModels/Stack/NiagaraStackRenderersOwner.cpp
@@ -6,6 +6,7 @@
 #include "NiagaraEditorModule.h"
 #include "NiagaraEditorStyle.h"
 #include "NiagaraEmitter.h"
+#include "NiagaraEmitterHandle.h"
 #include "NiagaraRendererProperties.h"
 #include "NiagaraScriptMergeManager.h"
 #include "Stateless/NiagaraStatelessEmitter.h"
diff --git a/Engine/Plugins/FX/Niagara/Source/NiagaraEditor/Private/ViewModels/Stack/NiagaraStackRoot.cpp b/Engine/Plugins/FX/Niagara/Source/NiagaraEditor/Private/ViewModels/Stack/NiagaraStackRoot.cpp
--- a/Engine/Plugins/FX/Niagara/Source/NiagaraEditor/Private/ViewModels/Stack/NiagaraStackRoot.cpp
+++ b/Engine/Plugins/FX/Niagara/Source/NiagaraEditor/Private/ViewModels/Stack/NiagaraStackRoot.cpp
@@ -13,6 +13,7 @@
 #include "ViewModels/Stack/NiagaraStackSimulationStageGroup.h"
 #include "ViewModels/Stack/NiagaraStackSystemSettingsGroup.h"
 #include "ViewModels/NiagaraSystemViewModel.h"
+#include "ViewModels/NiagaraEmitterHandleViewModel.h"
 #include "ViewModels/NiagaraEmitterViewModel.h"
 #include "ViewModels/NiagaraScriptViewModel.h"
 #include "ViewModels/NiagaraSystemScriptViewModel.h"
diff --git a/Engine/Plugins/FX/Niagara/Source/NiagaraEditorWidgets/Private/NiagaraEditorWidgetsUtilities.cpp b/Engine/Plugins/FX/Niagara/Source/NiagaraEditorWidgets/Private/NiagaraEditorWidgetsUtilities.cpp
--- a/Engine/Plugins/FX/Niagara/Source/NiagaraEditorWidgets/Private/NiagaraEditorWidgetsUtilities.cpp
+++ b/Engine/Plugins/FX/Niagara/Source/NiagaraEditorWidgets/Private/NiagaraEditorWidgetsUtilities.cpp
@@ -11,6 +11,7 @@
 #include "NiagaraEditorCommon.h"
 #include "NiagaraClipboard.h"
 #include "NiagaraEditorModule.h"
+#include "NiagaraEditorUtilities.h"
 #include "NiagaraNodeFunctionCall.h"
 
 #include "Framework/Application/SlateApplication.h"
diff --git a/Engine/Plugins/FX/Niagara/Source/NiagaraEditorWidgets/Private/Stack/SNiagaraDeterminismToggle.cpp b/Engine/Plugins/FX/Niagara/Source/NiagaraEditorWidgets/Private/Stack/SNiagaraDeterminismToggle.cpp
--- a/Engine/Plugins/FX/Niagara/Source/NiagaraEditorWidgets/Private/Stack/SNiagaraDeterminismToggle.cpp
+++ b/Engine/Plugins/FX/Niagara/Source/NiagaraEditorWidgets/Private/Stack/SNiagaraDeterminismToggle.cpp
@@ -3,6 +3,7 @@
 #include "Stack/SNiagaraDeterminismToggle.h"
 
 #include "NiagaraEditorStyle.h"
+#include "NiagaraSystem.h"
 #include "Styling/AppStyle.h"
 #include "Widgets/Input/SButton.h"
 #include "Widgets/Images/SImage.h"
diff --git a/Engine/Plugins/FX/Niagara/Source/NiagaraEditorWidgets/Private/Stack/SNiagaraSimTargetToggle.cpp b/Engine/Plugins/FX/Niagara/Source/NiagaraEditorWidgets/Private/Stack/SNiagaraSimTargetToggle.cpp
--- a/Engine/Plugins/FX/Niagara/Source/NiagaraEditorWidgets/Private/Stack/SNiagaraSimTargetToggle.cpp
+++ b/Engine/Plugins/FX/Niagara/Source/NiagaraEditorWidgets/Private/Stack/SNiagaraSimTargetToggle.cpp
@@ -3,6 +3,7 @@
 #include "Stack/SNiagaraSimTargetToggle.h"
 
 #include "NiagaraEditorStyle.h"
+#include "NiagaraSystem.h"
 #include "Styling/AppStyle.h"
 #include "Widgets/Input/SButton.h"
 #include "Widgets/Images/SImage.h"
