Fixed a bug where the narrowing type conversion check failed on INDEX_NONE, resulting in a crash when trying to record a sim cache from sequencer

#tests Tried to cache a default fountain system in a level sequence with and without the fix
#rnx

diff --git a/Engine/Plugins/FX/Niagara/Source/Niagara/Private/NiagaraSimCacheHelper.h b/Engine/Plugins/FX/Niagara/Source/Niagara/Private/NiagaraSimCacheHelper.h
--- a/Engine/Plugins/FX/Niagara/Source/Niagara/Private/NiagaraSimCacheHelper.h
+++ b/Engine/Plugins/FX/Niagara/Source/Niagara/Private/NiagaraSimCacheHelper.h
@@ -202,7 +202,8 @@ struct FNiagaraSimCacheHelper
 		if (CreateParameters.bAllowInterpolation)
 		{
 			const FNiagaraVariableBase UniqueIDVariable(FNiagaraTypeDefinition::GetIntDef(), "UniqueID");
-			CacheLayout.CacheBufferWriteInfo.ComponentUniqueID = IntCastChecked<uint16>(CompiledData.Variables.IndexOfByKey(UniqueIDVariable));
+			const int32 IDVarIndex = CompiledData.Variables.IndexOfByKey(UniqueIDVariable);
+			CacheLayout.CacheBufferWriteInfo.ComponentUniqueID = IDVarIndex == INDEX_NONE ? static_cast<uint16>(INDEX_NONE) : IntCastChecked<uint16>(IDVarIndex);
 			if (CacheLayout.CacheBufferWriteInfo.ComponentUniqueID != static_cast<uint16>(INDEX_NONE))
 			{
 				const FNiagaraVariableLayoutInfo& DataSetVariableLayout = CompiledData.VariableLayouts[CacheLayout.CacheBufferWriteInfo.ComponentUniqueID];
