* Allow data interface inputs to be flagged as required, so they don't accept the pin default values and need to be wired to something
* Change type of function signature output variables, as they can't have default values

#tests Tested by adding the spawn condition function from the data channel reader DI in a scratch pad. Translation fails now when the emitter ID is not provided from the graph.

diff --git a/Engine/Plugins/FX/Niagara/Source/Niagara/Private/DataInterface/NiagaraDataInterfaceDataChannelRead.cpp b/Engine/Plugins/FX/Niagara/Source/Niagara/Private/DataInterface/NiagaraDataInterfaceDataChannelRead.cpp
--- a/Engine/Plugins/FX/Niagara/Source/Niagara/Private/DataInterface/NiagaraDataInterfaceDataChannelRead.cpp
+++ b/Engine/Plugins/FX/Niagara/Source/Niagara/Private/DataInterface/NiagaraDataInterfaceDataChannelRead.cpp
@@ -122,7 +122,7 @@ namespace NDIDataChannelReadLocal
 			Sig.bExperimental = true;
 			Sig.ModuleUsageBitmask = ENiagaraScriptUsageMask::Particle;
 			Sig.AddInput(FNiagaraVariable(FNiagaraTypeDefinition(UNiagaraDataInterfaceDataChannelRead::StaticClass()), TEXT("DataChannel interface")));
-			Sig.AddInput(FNiagaraVariable(FNiagaraTypeDefinition(FNiagaraEmitterID::StaticStruct()), TEXT("Emitter ID")), LOCTEXT("EmitterIDDesc", "ID of the emitter we'd like to spawn into. This can be obtained from Engine.Emitter.ID."));
+			Sig.AddInputWithoutDefault(FNiagaraVariable(FNiagaraTypeDefinition(FNiagaraEmitterID::StaticStruct()), TEXT("Emitter ID")), LOCTEXT("EmitterIDDesc", "ID of the emitter we'd like to spawn into. This can be obtained from Engine.Emitter.ID."));
 			Sig.AddInputWithDefault(FNiagaraVariable(FNiagaraTypeDefinition::GetIntDef(), TEXT("Spawned Particle Exec Index")), 0, LOCTEXT("GetNDCSpawnData_InExecIndexDesc","The execution index of the spawned particle."));
 			Sig.AddOutput(FNiagaraVariable(FNiagaraTypeDefinition::GetIntDef(), TEXT("NDC Index")), LOCTEXT("GetNDCSpawnData_OutNDCIndexDesc","Index of the NDC item that spawned this particle."));
 			Sig.AddOutput(FNiagaraVariable(FNiagaraTypeDefinition::GetIntDef(), TEXT("NDC Spawn Index")), LOCTEXT("GetNDCSpawnData_OutNDCSpawnIndexDesc","The index of this particle in relation to all the particle spawned by the same NDC item."));
@@ -200,7 +200,7 @@ namespace NDIDataChannelReadLocal
 			Sig.ModuleUsageBitmask = ENiagaraScriptUsageMask::Emitter | ENiagaraScriptUsageMask::System;
 			Sig.AddInput(FNiagaraVariable(FNiagaraTypeDefinition(UNiagaraDataInterfaceDataChannelRead::StaticClass()), TEXT("DataChannel interface")));
 			Sig.AddInputWithDefault(FNiagaraVariable(FNiagaraTypeDefinition::GetBoolDef(), TEXT("Enable")), FNiagaraBool(true), LOCTEXT("SpawnEnableInputDesc", "Enable or disable this function call. If false, this call with have no effetcs."));
-			Sig.AddInput(FNiagaraVariable(FNiagaraTypeDefinition(FNiagaraEmitterID::StaticStruct()), TEXT("Emitter ID")), LOCTEXT("EmitterIDDesc", "ID of the emitter we'd like to spawn into. This can be obtained from Engine.Emitter.ID."));
+			Sig.AddInputWithoutDefault(FNiagaraVariable(FNiagaraTypeDefinition(FNiagaraEmitterID::StaticStruct()), TEXT("Emitter ID")), LOCTEXT("EmitterIDDesc", "ID of the emitter we'd like to spawn into. This can be obtained from Engine.Emitter.ID."));
 			Sig.AddInput(FNiagaraVariable(StaticEnum<ENDIDataChannelSpawnMode>(), TEXT("Mode")), LOCTEXT("SpawnCondModeInputDesc", "A mode switch that controls how this funciton will behave and interact with other calls to spawn functions."));
 			Sig.AddInput(FNiagaraVariable(StaticEnum<ENiagaraConditionalOperator>(), TEXT("Operator")), LOCTEXT("SpawnCondOpInputDesc", "The comparison operator to use when comparing values in the data channel to conditional parameters."));
 			Sig.AddInputWithDefault(FNiagaraVariable(FNiagaraTypeDefinition::GetIntDef(), TEXT("Min Spawn Count")), 1, LOCTEXT("MinSpawnCountInputDesc", "Minimum number of particles to spawn for each element in the data channel."));
@@ -228,7 +228,7 @@ namespace NDIDataChannelReadLocal
 			Sig.ModuleUsageBitmask = ENiagaraScriptUsageMask::Emitter | ENiagaraScriptUsageMask::System;
 			Sig.AddInput(FNiagaraVariable(FNiagaraTypeDefinition(UNiagaraDataInterfaceDataChannelRead::StaticClass()), TEXT("DataChannel interface")));
 			Sig.AddInputWithDefault(FNiagaraVariable(FNiagaraTypeDefinition::GetBoolDef(), TEXT("Enable")), FNiagaraBool(true), LOCTEXT("SpawnEnableInputDesc", "Enable or disable this function call. If false, this call with have no effetcs."));
-			Sig.AddInput(FNiagaraVariable(FNiagaraTypeDefinition(FNiagaraEmitterID::StaticStruct()), TEXT("Emitter ID")), LOCTEXT("EmitterIDDesc", "ID of the emitter we'd like to spawn into. This can be obtained from Engine.Emitter.ID."));
+			Sig.AddInputWithoutDefault(FNiagaraVariable(FNiagaraTypeDefinition(FNiagaraEmitterID::StaticStruct()), TEXT("Emitter ID")), LOCTEXT("EmitterIDDesc", "ID of the emitter we'd like to spawn into. This can be obtained from Engine.Emitter.ID."));
 			Sig.AddInput(FNiagaraVariable(StaticEnum<ENDIDataChannelSpawnMode>(), TEXT("Mode")), LOCTEXT("SpawnCondModeInputDesc", "A mode switch that controls how this funciton will behave and interact with other calls to spawn functions."));			
 			Sig.AddInputWithDefault(FNiagaraVariable(FNiagaraTypeDefinition::GetFloatDef(), TEXT("RandomScaleMin")), 1.0f, LOCTEXT("SpawnDirectRandomScaleMinInputDesc", "Minimum value for an additional random scale applied to each NDC spawn count."));
 			Sig.AddInputWithDefault(FNiagaraVariable(FNiagaraTypeDefinition::GetFloatDef(), TEXT("RandomScaleMax")), 1.0f, LOCTEXT("SpawnDirectRandomScaleMaxInputDesc", "Maximum value for an additional random scale applied to each NDC spawn count."));
@@ -256,7 +256,7 @@ namespace NDIDataChannelReadLocal
 			Sig.ModuleUsageBitmask = ENiagaraScriptUsageMask::Emitter | ENiagaraScriptUsageMask::System;
 			Sig.AddInput(FNiagaraVariable(FNiagaraTypeDefinition(UNiagaraDataInterfaceDataChannelRead::StaticClass()), TEXT("DataChannel interface")));
 			Sig.AddInputWithDefault(FNiagaraVariable(FNiagaraTypeDefinition::GetBoolDef(), TEXT("Enable")), FNiagaraBool(true), LOCTEXT("SpawnEnableInputDesc", "Enable or disable this function call. If false, this call with have no effetcs."));
-			Sig.AddInput(FNiagaraVariable(FNiagaraTypeDefinition(FNiagaraEmitterID::StaticStruct()), TEXT("Emitter ID")), LOCTEXT("EmitterIDDesc", "ID of the emitter we'd like to spawn into. This can be obtained from Engine.Emitter.ID."));
+			Sig.AddInputWithoutDefault(FNiagaraVariable(FNiagaraTypeDefinition(FNiagaraEmitterID::StaticStruct()), TEXT("Emitter ID")), LOCTEXT("EmitterIDDesc", "ID of the emitter we'd like to spawn into. This can be obtained from Engine.Emitter.ID."));
 			Sig.AddInput(FNiagaraVariable(StaticEnum<ENDIDataChannelSpawnScaleMode>(), TEXT("Mode")), LOCTEXT("SpawnScaleModeInputDesc", "Control whether to override or combine this scale with previously set scales when calling this function multiple times."));
 			Sig.AddInputWithDefault(FNiagaraVariable(FNiagaraTypeDefinition::GetFloatDef(), TEXT("RandomScaleMin")), 1.0f, LOCTEXT("ScaleSpawnCountRandomScaleMinInputDesc", "Minimum value for a random additional scale applied to each NDC spawn."));
 			Sig.AddInputWithDefault(FNiagaraVariable(FNiagaraTypeDefinition::GetFloatDef(), TEXT("RandomScaleMax")), 1.0f, LOCTEXT("ScaleSpawnCountRandomScaleMaxInputDesc", "Maximum value for a random additional scale applied to each NDC spawn."));
@@ -2169,11 +2169,11 @@ void UNiagaraDataInterfaceDataChannelRead::GetParameterDefinitionHLSL(FNiagaraDa
 			static FString ParamPrefix = TEXT(", ");
 			static FString OutputParamPrefix = TEXT("out");
 			static FString InputParamPrefix = TEXT("in");
-			auto HandleParameters = [&](TConstArrayView<FNiagaraVariable> Parameters, int32 StartParamIdx, bool bOutput)
+			auto HandleParameters = [&](TConstArrayView<FNiagaraVariableBase> Parameters, int32 StartParamIdx, bool bOutput)
 			{
 				for (int32 ParamIdx = StartParamIdx; ParamIdx < Parameters.Num(); ++ParamIdx)
 				{
-					const FNiagaraVariable& Param = Parameters[ParamIdx];
+					const FNiagaraVariableBase& Param = Parameters[ParamIdx];
 
 					//Reset or init per parameter hlsl template args.
 
@@ -2208,7 +2208,7 @@ void UNiagaraDataInterfaceDataChannelRead::GetParameterDefinitionHLSL(FNiagaraDa
 				}
 			};
 
-			HandleParameters(Signature.Inputs, Signature.NumRequiredInputs(), false);
+			HandleParameters(Signature.GetInputs(), Signature.NumRequiredInputs(), false);
 			HandleParameters(Signature.Outputs, Signature.NumRequiredOutputs(), true);
 
 			//Grab the template file for the function.
diff --git a/Engine/Plugins/FX/Niagara/Source/Niagara/Public/NiagaraCommon.h b/Engine/Plugins/FX/Niagara/Source/Niagara/Public/NiagaraCommon.h
--- a/Engine/Plugins/FX/Niagara/Source/Niagara/Public/NiagaraCommon.h
+++ b/Engine/Plugins/FX/Niagara/Source/Niagara/Public/NiagaraCommon.h
@@ -363,9 +363,9 @@ struct FNiagaraFunctionSignature
 	/** Input parameters to this function. The data stored in the variables is used for default values. */
 	UPROPERTY()
 	TArray<FNiagaraVariable> Inputs;
-	/** Output parameters of this function. The data stored in the variables is used for default values. */
+	/** Output parameters of this function. */
 	UPROPERTY()
-	TArray<FNiagaraVariable> Outputs;
+	TArray<FNiagaraVariableBase> Outputs;
 	/** Id of the owner is this is a member function. */
 	UPROPERTY()
 	FName OwnerName;
@@ -389,6 +389,13 @@ struct FNiagaraFunctionSignature
 	/** Per function version, it is up to the discretion of the function as to what the version means. */
 	UPROPERTY()
 	uint32 FunctionVersion = 0;
+
+	/**
+	* All the inputs here must be wired to something in the graph and will produce a compile error if that's not the case.
+	* This is useful if the default value for an input makes no sense and would lead to a failed function call at runtime.
+	 */
+	UPROPERTY()
+	TArray<FNiagaraVariableBase> NoDefaultValueInputs;
 #endif
 
 	/** Support running on the CPU. */
@@ -611,7 +618,7 @@ struct FNiagaraFunctionSignature
 
 	FString GetNameString() const { return Name.ToString(); }
 
-	void AddInput(FNiagaraVariable InputVar, FText Tooltip = FText())
+	void AddInput(const FNiagaraVariable& InputVar, const FText& Tooltip = FText())
 	{
 		Inputs.Add(InputVar);
 	#if WITH_EDITORONLY_DATA
@@ -623,7 +630,7 @@ struct FNiagaraFunctionSignature
 	}
 	
 	template<typename T>
-	void AddInputWithDefault(FNiagaraVariable InputVar, const T& Default, FText Tooltip = FText())
+	void AddInputWithDefault(FNiagaraVariable InputVar, const T& Default, const FText& Tooltip = FText())
 	{
 		Inputs.Add_GetRef(InputVar).SetData((uint8*)&Default);
 #if WITH_EDITORONLY_DATA
@@ -634,7 +641,19 @@ struct FNiagaraFunctionSignature
 #endif
 	}
 
-	void AddOutput(FNiagaraVariable OutputVar, const FText& Tooltip = FText())
+	void AddInputWithoutDefault(const FNiagaraVariable& InputVar, const FText& Tooltip = FText())
+	{
+		Inputs.Add(InputVar);
+#if WITH_EDITORONLY_DATA
+		if (!Tooltip.IsEmpty())
+		{
+			InputDescriptions.Add(InputVar, Tooltip);
+		}
+		NoDefaultValueInputs.AddUnique(InputVar);
+#endif
+	}
+
+	void AddOutput(const FNiagaraVariable& OutputVar, const FText& Tooltip = FText())
 	{
 		Outputs.Add(OutputVar);
 	#if WITH_EDITORONLY_DATA
@@ -693,6 +712,16 @@ struct FNiagaraFunctionSignature
 	int32 NumRequiredOutputs()const { return RequiredOutputs == INDEX_NONE ? Outputs.Num() : RequiredOutputs; }
 	int32 NumOptionalOutputs()const { return Outputs.Num() - NumRequiredOutputs(); }
 
+#if WITH_EDITORONLY_DATA
+	TArray<FNiagaraVariableBase> GetInputs() const
+	{
+		TArray<FNiagaraVariableBase> AllInputs;
+		AllInputs.Append(Inputs);
+		return AllInputs;
+	}
+	TConstArrayView<FNiagaraVariableBase> GetOutputs() const { return Outputs; }
+#endif
+	
 	NIAGARA_API void GetVariadicInputs(TArray<FNiagaraVariableBase>& OutVariadicInputs, bool bStripNonExecution = true)const;
 	NIAGARA_API void GetVariadicOutputs(TArray<FNiagaraVariableBase>& OutVariadicOutputs, bool bStripNonExecution = true)const;
 };
diff --git a/Engine/Plugins/FX/Niagara/Source/NiagaraEditor/Private/NiagaraHlslTranslator.cpp b/Engine/Plugins/FX/Niagara/Source/NiagaraEditor/Private/NiagaraHlslTranslator.cpp
--- a/Engine/Plugins/FX/Niagara/Source/NiagaraEditor/Private/NiagaraHlslTranslator.cpp
+++ b/Engine/Plugins/FX/Niagara/Source/NiagaraEditor/Private/NiagaraHlslTranslator.cpp
@@ -7975,6 +7975,19 @@ void TNiagaraHlslTranslator<GraphBridge>::FunctionCall(const FFunctionCallNode*
 		}
 	}
 
+	if (Signature.NoDefaultValueInputs.Num() > 0)
+	{
+		for (int32 i = 0; i < CallInputs.Num(); i++)
+		{
+			const FInputPin* InPin = CallInputs[i];
+			FNiagaraVariableBase Var = GraphBridge::GetPinVariable(InPin, false, ENiagaraStructConversion::Simulation);
+			if (Signature.NoDefaultValueInputs.Contains(Var) && GraphBridge::GetLinkedOutputPin(InPin) == nullptr)
+			{
+				Error(LOCTEXT("NoDefaultValueInput", "Function call input pin doesn't support default values and needs to be wired in."), FunctionNode, InPin);
+			}
+		}
+	}
+
 	if (const FCustomHlslNode* CustomFunctionHlsl = GraphBridge::AsCustomHlslNode(FunctionNode))
 	{
 		// All of the arguments here are resolved withing the HandleCustomHlsl function..
@@ -8505,8 +8518,8 @@ void TNiagaraHlslTranslator<GraphBridge>::ProcessCustomHlsl(const FString& InCus
 	OutSignature.Inputs = SigInputs;
 
 	// Resolve the names of any internal variables from the output variables.
-	TArray<FNiagaraVariable> SigOutputs;
-	for (FNiagaraVariable Output : OutSignature.Outputs)
+	TArray<FNiagaraVariableBase> SigOutputs;
+	for (FNiagaraVariableBase& Output : OutSignature.Outputs)
 	{
 		if (Output.GetType() == FNiagaraTypeDefinition::GetParameterMapDef())
 		{
@@ -9107,18 +9120,19 @@ void FNiagaraHlslTranslator::GenerateFunctionCall(ENiagaraScriptUsage ScriptUsag
 	FString DefStr = GetFunctionSignatureSymbol(FunctionSignature) + TEXT("(");
 	for (int32 i = 0; i < FunctionSignature.Inputs.Num(); ++i)
 	{
-		FNiagaraTypeDefinition Type = FunctionSignature.Inputs[i].GetType();
+		const FNiagaraVariable& InputVar = FunctionSignature.Inputs[i];
+		FNiagaraTypeDefinition Type = InputVar.GetType();
 		if (Type.UnderlyingType != 0 && Type.ClassStructOrEnum == nullptr)
 		{
-			Error(FText::Format(LOCTEXT("InvalidTypeDefError", "Invalid data in niagara type definition, might be due to broken serialization or missing DI implementation! Variable: {0}"), FText::FromName(FunctionSignature.Inputs[i].GetName())));
+			Error(FText::Format(LOCTEXT("InvalidTypeDefError", "Invalid data in niagara type definition, might be due to broken serialization or missing DI implementation! Variable: {0}"), FText::FromName(InputVar.GetName())));
 			continue;
 		}
 
 		if (!ensure(i < Inputs.Num()))
 		{
-			Error(FText::Format(LOCTEXT("InvalidInputNum", "Functon Input of %d is out of bounds in function signature! Variable: {0}"), 
+			Error(FText::Format(LOCTEXT("InvalidInputNum", "Function Input of %d is out of bounds in function signature! Variable: {0}"), 
 			FText::AsNumber(i),
-			FText::FromName(FunctionSignature.Inputs[i].GetName())));
+			FText::FromName(InputVar.GetName())));
 			continue;
 		}
 
@@ -9127,13 +9141,13 @@ void FNiagaraHlslTranslator::GenerateFunctionCall(ENiagaraScriptUsage ScriptUsag
 		{
 			if (!AddStructToDefinitionSet(Type))
 			{
-				Error(FText::Format(LOCTEXT("GetConstantFailTypeVar2", "Cannot handle type {0}! Variable: {1}"), Type.GetNameText(), FText::FromName(FunctionSignature.Inputs[i].GetName())));
+				Error(FText::Format(LOCTEXT("GetConstantFailTypeVar2", "Cannot handle type {0}! Variable: {1}"), Type.GetNameText(), FText::FromName(InputVar.GetName())));
 			}
 
 			int32 Input = Inputs[i];
 			bool bSkip = false;
 
-			if (FunctionSignature.Inputs[i].GetType() == FNiagaraTypeDefinition::GetParameterMapDef())
+			if (InputVar.GetType() == FNiagaraTypeDefinition::GetParameterMapDef())
 			{
 				Input = INDEX_NONE;
 				bSkip = true;
@@ -9149,7 +9163,7 @@ void FNiagaraHlslTranslator::GenerateFunctionCall(ENiagaraScriptUsage ScriptUsag
 				Params.Add(Input);
 				if (Input == INDEX_NONE)
 				{
-					MissingParameters.Add(FunctionSignature.Inputs[i].GetName().ToString());
+					MissingParameters.Add(InputVar.GetName().ToString());
 				}
 				else
 				{
@@ -9162,7 +9176,7 @@ void FNiagaraHlslTranslator::GenerateFunctionCall(ENiagaraScriptUsage ScriptUsag
 
 	for (int32 i = 0; i < FunctionSignature.Outputs.Num(); ++i)
 	{
-		FNiagaraVariable& OutVar = FunctionSignature.Outputs[i];
+		FNiagaraVariableBase& OutVar = FunctionSignature.Outputs[i];
 		FNiagaraTypeDefinition Type = ConvertToSimulationVariable(OutVar).GetType();
 
 		//We don't write class types as real params in the hlsl
diff --git a/Engine/Plugins/FX/Niagara/Source/NiagaraEditor/Private/NiagaraNodeFunctionCall.cpp b/Engine/Plugins/FX/Niagara/Source/NiagaraEditor/Private/NiagaraNodeFunctionCall.cpp
--- a/Engine/Plugins/FX/Niagara/Source/NiagaraEditor/Private/NiagaraNodeFunctionCall.cpp
+++ b/Engine/Plugins/FX/Niagara/Source/NiagaraEditor/Private/NiagaraNodeFunctionCall.cpp
@@ -470,7 +470,7 @@ void UNiagaraNodeFunctionCall::AllocateDefaultPins()
 			NewPin->bDefaultValueIsIgnored = true;
 		}
 
-		for (FNiagaraVariable& Output : Signature.Outputs)
+		for (FNiagaraVariableBase& Output : Signature.Outputs)
 		{
 			UEdGraphPin* NewPin = CreatePin(EGPD_Output, Schema->TypeDefinitionToPinType(Output.GetType()), Output.GetName());
 			NewPin->bDefaultValueIsIgnored = true;
@@ -1059,7 +1059,7 @@ void UNiagaraNodeFunctionCall::UpdatePinTooltips()
 	{
 		for (int i = 0; i < OutputPins.Num(); i++)
 		{
-			FNiagaraVariable& Output = Signature.Outputs[i];
+			FNiagaraVariableBase& Output = Signature.Outputs[i];
 			OutputPins[i]->PinToolTip = Signature.OutputDescriptions.Contains(Output) ? Signature.OutputDescriptions[Output].ToString() : FString();
 		}
 	}
@@ -1802,7 +1802,7 @@ bool UNiagaraNodeFunctionCall::IsBaseSignatureOfDataInterfaceFunction(const UEdG
 		if (FNiagaraFunctionSignature* BaseSig = BaseDIFuncs.FindByPredicate([&](const FNiagaraFunctionSignature& CheckSig) { return Signature.Name == CheckSig.Name; }))
 		{
 			FPinCollectorArray FoundPins;
-			TArray<FNiagaraVariable> InputOrOutputVariables;
+			TArray<FNiagaraVariableBase> InputOrOutputVariables;
 
 			if(BaseSig->bRequiresExecPin && IsExecPin(Pin))
 			{
@@ -1812,7 +1812,7 @@ bool UNiagaraNodeFunctionCall::IsBaseSignatureOfDataInterfaceFunction(const UEdG
 			if(Pin->Direction == EGPD_Input)
 			{
 				GetInputPins(FoundPins);
-				InputOrOutputVariables = BaseSig->Inputs;
+				InputOrOutputVariables = BaseSig->GetInputs();
 			}
 			else
 			{
@@ -1820,7 +1820,7 @@ bool UNiagaraNodeFunctionCall::IsBaseSignatureOfDataInterfaceFunction(const UEdG
 				InputOrOutputVariables = BaseSig->Outputs;
 			}
 			
-			FNiagaraVariable Variable = UEdGraphSchema_Niagara::PinToNiagaraVariable(Pin);
+			FNiagaraVariableBase Variable = UEdGraphSchema_Niagara::PinToNiagaraVariable(Pin);
 
 			if(InputOrOutputVariables.Contains(Variable))
 			{
