Fixed Linux compile errors

#rnx

diff --git a/Engine/Plugins/FX/Niagara/Source/Niagara/Private/DataInterface/NiagaraDataInterfaceDataChannelRead.cpp b/Engine/Plugins/FX/Niagara/Source/Niagara/Private/DataInterface/NiagaraDataInterfaceDataChannelRead.cpp
--- a/Engine/Plugins/FX/Niagara/Source/Niagara/Private/DataInterface/NiagaraDataInterfaceDataChannelRead.cpp
+++ b/Engine/Plugins/FX/Niagara/Source/Niagara/Private/DataInterface/NiagaraDataInterfaceDataChannelRead.cpp
@@ -2041,7 +2041,7 @@ void UNiagaraDataInterfaceDataChannelRead::GetParameterDefinitionHLSL(FNiagaraDa
 		{
 			Ret = HashCombine(Ret, GetTypeHash(Input));
 		}
-		for (const FNiagaraVariable& Output : Sig.Outputs)
+		for (const FNiagaraVariableBase& Output : Sig.Outputs)
 		{
 			Ret = HashCombine(Ret, GetTypeHash(Output));
 		}
@@ -2063,7 +2063,7 @@ void UNiagaraDataInterfaceDataChannelRead::GetParameterDefinitionHLSL(FNiagaraDa
 
 		for (int32 OutputIdx = Signature.NumRequiredOutputs(); OutputIdx < Signature.Outputs.Num(); ++OutputIdx)
 		{
-			const FNiagaraVariable& OutputParam = Signature.Outputs[OutputIdx];
+			const FNiagaraVariableBase& OutputParam = Signature.Outputs[OutputIdx];
 			ParametersAccessed.AddUnique(OutputParam);
 		}
 	}
diff --git a/Engine/Plugins/FX/Niagara/Source/NiagaraEditor/Private/NiagaraHlslTranslator.cpp b/Engine/Plugins/FX/Niagara/Source/NiagaraEditor/Private/NiagaraHlslTranslator.cpp
--- a/Engine/Plugins/FX/Niagara/Source/NiagaraEditor/Private/NiagaraHlslTranslator.cpp
+++ b/Engine/Plugins/FX/Niagara/Source/NiagaraEditor/Private/NiagaraHlslTranslator.cpp
@@ -9452,7 +9452,7 @@ FString FNiagaraHlslTranslator::GetFunctionSignature(const FNiagaraFunctionSigna
 
 	for (int32 i = 0; i < Sig.Outputs.Num(); ++i)
 	{
-		const FNiagaraVariable& Output = Sig.Outputs[i];
+		const FNiagaraVariableBase& Output = Sig.Outputs[i];
 		//We don't write class types as real params in the hlsl
 		if (Output.GetType().GetClass() == nullptr)
 		{
diff --git a/Engine/Plugins/FX/Niagara/Source/NiagaraEditor/Private/NiagaraHlslTranslator.h b/Engine/Plugins/FX/Niagara/Source/NiagaraEditor/Private/NiagaraHlslTranslator.h
--- a/Engine/Plugins/FX/Niagara/Source/NiagaraEditor/Private/NiagaraHlslTranslator.h
+++ b/Engine/Plugins/FX/Niagara/Source/NiagaraEditor/Private/NiagaraHlslTranslator.h
@@ -93,7 +93,7 @@ FORCEINLINE uint32 GetTypeHash(const FNiagaraFunctionSignature& Sig)
 	{
 		Hash = HashCombine(Hash, GetTypeHash(Var));
 	}
-	for (const FNiagaraVariable& Var : Sig.Outputs)
+	for (const FNiagaraVariableBase& Var : Sig.Outputs)
 	{
 		Hash = HashCombine(Hash, GetTypeHash(Var));
 	}
diff --git a/Engine/Plugins/FX/Niagara/Source/NiagaraEditor/Private/NiagaraNodeFunctionCall.cpp b/Engine/Plugins/FX/Niagara/Source/NiagaraEditor/Private/NiagaraNodeFunctionCall.cpp
--- a/Engine/Plugins/FX/Niagara/Source/NiagaraEditor/Private/NiagaraNodeFunctionCall.cpp
+++ b/Engine/Plugins/FX/Niagara/Source/NiagaraEditor/Private/NiagaraNodeFunctionCall.cpp
@@ -1777,7 +1777,7 @@ void UNiagaraNodeFunctionCall::RefreshSignature()
 				for (int32 i = 0; i < FoundPins.Num(); ++i)
 				{
 					UEdGraphPin* OutputPin = FoundPins[i];
-					FNiagaraVariable InputVariable = UEdGraphSchema_Niagara::PinToNiagaraVariable(OutputPin);
+					FNiagaraVariableBase InputVariable = UEdGraphSchema_Niagara::PinToNiagaraVariable(OutputPin);
 
 					if(!BaseSig->Outputs.Contains(InputVariable) && IsAddPin(OutputPin) == false && IsExecPin(OutputPin) == false)
 					{
