- Sim Cache buffer & attribute change to single selection

#rnx

diff --git a/Engine/Plugins/FX/Niagara/Source/NiagaraEditor/Private/Widgets/SNiagaraSimCacheDebugDataView.cpp b/Engine/Plugins/FX/Niagara/Source/NiagaraEditor/Private/Widgets/SNiagaraSimCacheDebugDataView.cpp
--- a/Engine/Plugins/FX/Niagara/Source/NiagaraEditor/Private/Widgets/SNiagaraSimCacheDebugDataView.cpp
+++ b/Engine/Plugins/FX/Niagara/Source/NiagaraEditor/Private/Widgets/SNiagaraSimCacheDebugDataView.cpp
@@ -148,7 +148,10 @@ void SNiagaraSimCacheDebugDataView::Construct(const FArguments& InArgs)
 
 	if (const FNiagaraSimCacheDebugDataFrame* FrameData = GetCurrentFrameData())
 	{
-		SelectedParameterStoreName = FrameData->DebugParameterStores.CreateConstIterator().Key();
+		if (FrameData->DebugParameterStores.Num() > 0)
+		{
+			SelectedParameterStoreName = FrameData->DebugParameterStores.CreateConstIterator().Key();
+		}
 	}
 
 	ChildSlot
diff --git a/Engine/Plugins/FX/Niagara/Source/NiagaraEditor/Private/Widgets/SNiagaraSimCacheOverview.cpp b/Engine/Plugins/FX/Niagara/Source/NiagaraEditor/Private/Widgets/SNiagaraSimCacheOverview.cpp
--- a/Engine/Plugins/FX/Niagara/Source/NiagaraEditor/Private/Widgets/SNiagaraSimCacheOverview.cpp
+++ b/Engine/Plugins/FX/Niagara/Source/NiagaraEditor/Private/Widgets/SNiagaraSimCacheOverview.cpp
@@ -51,7 +51,8 @@ void SNiagaraSimCacheOverview::Construct(const FArguments& InArgs)
 	SAssignNew(BufferListView, SListView<TSharedRef<FNiagaraSimCacheOverviewItem>>)
 	.ListItemsSource(ViewModel->GetBufferEntries())
 	.OnGenerateRow(this, &SNiagaraSimCacheOverview::OnGenerateRowForItem)
-	.OnSelectionChanged(this, &SNiagaraSimCacheOverview::OnListSelectionChanged);
+	.OnSelectionChanged(this, &SNiagaraSimCacheOverview::OnListSelectionChanged)
+	.SelectionMode(ESelectionMode::Single);
 
 	ViewModel.Get()->OnSimCacheChanged().AddSP(this, &SNiagaraSimCacheOverview::OnSimCacheChanged);
 
diff --git a/Engine/Plugins/FX/Niagara/Source/NiagaraEditor/Private/Widgets/SNiagaraSimCacheTreeView.cpp b/Engine/Plugins/FX/Niagara/Source/NiagaraEditor/Private/Widgets/SNiagaraSimCacheTreeView.cpp
--- a/Engine/Plugins/FX/Niagara/Source/NiagaraEditor/Private/Widgets/SNiagaraSimCacheTreeView.cpp
+++ b/Engine/Plugins/FX/Niagara/Source/NiagaraEditor/Private/Widgets/SNiagaraSimCacheTreeView.cpp
@@ -183,7 +183,7 @@ void SNiagaraSimCacheTreeView::Construct(const FArguments& InArgs)
 	ViewModel->BuildEntries(SharedThis(this));
 	
 	TreeView = SNew(STreeView<TSharedRef<FNiagaraSimCacheTreeItem>>)
-	.SelectionMode(ESelectionMode::Multi)
+	.SelectionMode(ESelectionMode::Single)
 	.TreeItemsSource(ViewModel->GetSelectedRootEntries())
 	.OnGenerateRow(this, &SNiagaraSimCacheTreeView::OnGenerateRow)
 	.OnGeneratePinnedRow(this, &SNiagaraSimCacheTreeView::OnGenerateRow)
