Fixed a bug where adding module inputs to the emitter summary was not captured in the transaction buffer

#tests Added a module input to the emitter summary, then hit undo to check that it's no longer there

diff --git a/Engine/Plugins/FX/Niagara/Source/NiagaraEditor/Private/ViewModels/HierarchyEditor/NiagaraHierarchyViewModelBase.cpp b/Engine/Plugins/FX/Niagara/Source/NiagaraEditor/Private/ViewModels/HierarchyEditor/NiagaraHierarchyViewModelBase.cpp
--- a/Engine/Plugins/FX/Niagara/Source/NiagaraEditor/Private/ViewModels/HierarchyEditor/NiagaraHierarchyViewModelBase.cpp
+++ b/Engine/Plugins/FX/Niagara/Source/NiagaraEditor/Private/ViewModels/HierarchyEditor/NiagaraHierarchyViewModelBase.cpp
@@ -1545,6 +1545,13 @@ void UNiagaraHierarchyViewModelBase::AddSection() const
 	OnHierarchyChangedDelegate.Broadcast();
 }
 
+UNiagaraHierarchyItemBase* UNiagaraHierarchyViewModelBase::AddItem(TSubclassOf<UNiagaraHierarchyItemBase> NewChildClass, FNiagaraHierarchyIdentity ChildIdentity)
+{
+	FScopedTransaction Transaction(LOCTEXT("Transaction_AddItem", "Add hierarchy item"));
+	HierarchyRoot->Modify();
+	return GetHierarchyRootViewModel()->AddChild(NewChildClass, ChildIdentity);
+}
+
 void UNiagaraHierarchyViewModelBase::DeleteItemWithIdentity(FNiagaraHierarchyIdentity Identity)
 {
 	FScopedTransaction Transaction(LOCTEXT("Transaction_DeleteItem", "Deleted hierarchy item"));
diff --git a/Engine/Plugins/FX/Niagara/Source/NiagaraEditor/Public/ViewModels/HierarchyEditor/NiagaraHierarchyViewModelBase.h b/Engine/Plugins/FX/Niagara/Source/NiagaraEditor/Public/ViewModels/HierarchyEditor/NiagaraHierarchyViewModelBase.h
--- a/Engine/Plugins/FX/Niagara/Source/NiagaraEditor/Public/ViewModels/HierarchyEditor/NiagaraHierarchyViewModelBase.h
+++ b/Engine/Plugins/FX/Niagara/Source/NiagaraEditor/Public/ViewModels/HierarchyEditor/NiagaraHierarchyViewModelBase.h
@@ -401,6 +401,8 @@ public:
 	void AddCategory(TSharedPtr<FNiagaraHierarchyItemViewModelBase> CategoryParent) const;
 	void AddSection() const;
 
+	// adds a new child item to the hierarchy root
+	NIAGARAEDITOR_API UNiagaraHierarchyItemBase* AddItem(TSubclassOf<UNiagaraHierarchyItemBase> NewChildClass, FNiagaraHierarchyIdentity ChildIdentity);
 	NIAGARAEDITOR_API void DeleteItemWithIdentity(FNiagaraHierarchyIdentity Identity);
 	NIAGARAEDITOR_API void DeleteItemsWithIdentities(TArray<FNiagaraHierarchyIdentity> Identities);
 
diff --git a/Engine/Plugins/FX/Niagara/Source/NiagaraEditorWidgets/Private/Stack/SNiagaraStackTableRow.cpp b/Engine/Plugins/FX/Niagara/Source/NiagaraEditorWidgets/Private/Stack/SNiagaraStackTableRow.cpp
--- a/Engine/Plugins/FX/Niagara/Source/NiagaraEditorWidgets/Private/Stack/SNiagaraStackTableRow.cpp
+++ b/Engine/Plugins/FX/Niagara/Source/NiagaraEditorWidgets/Private/Stack/SNiagaraStackTableRow.cpp
@@ -774,7 +774,7 @@ void SNiagaraStackTableRow::ToggleShowInSummaryView() const
 			{
 				UClass* HierarchyClass = DetermineHierarchyClassForSummaryView();
 				check(HierarchyClass);
-				StackEntry->GetEmitterViewModel()->GetSummaryHierarchyViewModel()->GetHierarchyRootViewModel()->AddChild(HierarchyClass, Identity);
+				StackEntry->GetEmitterViewModel()->GetSummaryHierarchyViewModel()->AddItem(HierarchyClass, Identity);
 			}
 		}
 	}
