Improve Content Browser tooltips for Niagara asset types

diff --git a/Engine/Plugins/FX/Niagara/Source/Niagara/Classes/NiagaraEffectType.h b/Engine/Plugins/FX/Niagara/Source/Niagara/Classes/NiagaraEffectType.h
--- a/Engine/Plugins/FX/Niagara/Source/Niagara/Classes/NiagaraEffectType.h
+++ b/Engine/Plugins/FX/Niagara/Source/Niagara/Classes/NiagaraEffectType.h
@@ -390,7 +390,13 @@ public:
 
 //////////////////////////////////////////////////////////////////////////
 
-/** Contains settings and working data shared among many NiagaraSystems that share some commonality of type. For example ImpactFX vs EnvironmentalFX. */
+/** Contains settings and working data shared among many Niagara systems that share some commonality of type, for example ImpactFX vs EnvironmentalFX.
+ *  Main usage of effect types is to control scalability settings for a group of effects, setting visibility and cull reactions on a per-platform basis.
+ *
+ *  Effect types can also be used for validation, checking that the content passes all the configured validation rule sets.
+ *
+ *  The effect type is set in Niagara systems in the system properties.
+ */
 UCLASS(config = Niagara, perObjectConfig, MinimalAPI)
 class UNiagaraEffectType : public UObject
 {
diff --git a/Engine/Plugins/FX/Niagara/Source/Niagara/Classes/NiagaraEmitter.h b/Engine/Plugins/FX/Niagara/Source/Niagara/Classes/NiagaraEmitter.h
--- a/Engine/Plugins/FX/Niagara/Source/Niagara/Classes/NiagaraEmitter.h
+++ b/Engine/Plugins/FX/Niagara/Source/Niagara/Classes/NiagaraEmitter.h
@@ -575,8 +575,11 @@ private:
 };
 
 /** 
- *	Niagara Emitters are particle spawners that can be reused for different effects by putting them into Niagara Systems.
+ *	Niagara Emitters are particle spawners that can be reused for different effects by putting them into Niagara systems.
  *	Emitters render their particles using different renderers, such as Sprite Renderers or Mesh Renderers to produce different effects.
+ *
+ *	Emitter assets cannot be spawned or used in a level directly, but need to be placed in a Niagara system. Emitters support inheritance, so that
+ *	changes to the base asset are automatically picked up by child emitter assets and emitters in system assets.
  */
 UCLASS(MinimalAPI)
 class UNiagaraEmitter : public UObject, public INiagaraParameterDefinitionsSubscriber, public FNiagaraVersionedObject
diff --git a/Engine/Plugins/FX/Niagara/Source/Niagara/Classes/NiagaraParameterCollection.h b/Engine/Plugins/FX/Niagara/Source/Niagara/Classes/NiagaraParameterCollection.h
--- a/Engine/Plugins/FX/Niagara/Source/Niagara/Classes/NiagaraParameterCollection.h
+++ b/Engine/Plugins/FX/Niagara/Source/Niagara/Classes/NiagaraParameterCollection.h
@@ -11,6 +11,10 @@ class UMaterialParameterCollection;
 class UMaterialParameterCollectionInstance;
 class UNiagaraParameterCollection;
 
+/**
+ * Can be used to override selected parameters from a Niagara parameter collection with another value.
+ * The values in the parameter collection instance can be set from Blueprint or C++, same as the regular parameter collection. 
+ */
 UCLASS(MinimalAPI)
 class UNiagaraParameterCollectionInstance : public UObject
 {
@@ -131,7 +135,12 @@ public:
 	NIAGARA_API void SetQuatParameter(const FString& InVariableName, const FQuat& InValue);
 };
 
-/** Asset containing a collection of global parameters usable by Niagara. */
+/** Asset containing a collection of global parameters usable by Niagara. Similar to Material parameter collections,
+ *  any number of Niagara assets may reference attributes from this parameter collection and will get new values when they are changed.
+ *
+ *  A Niagara parameter collection can reference a Material parameter collection, so it is in sync with the values provided to a Material.
+ *  To use a value from a parameter collection in a Niagara system or emitter, add a reference to it from the Parameters panel (in the Niagara Parameter Collection section).
+ */
 UCLASS(MinimalAPI)
 class UNiagaraParameterCollection : public UObject
 {
diff --git a/Engine/Plugins/FX/Niagara/Source/Niagara/Classes/NiagaraScript.h b/Engine/Plugins/FX/Niagara/Source/Niagara/Classes/NiagaraScript.h
--- a/Engine/Plugins/FX/Niagara/Source/Niagara/Classes/NiagaraScript.h
+++ b/Engine/Plugins/FX/Niagara/Source/Niagara/Classes/NiagaraScript.h
@@ -829,7 +829,13 @@ private:
 #endif	
 };
 
-/** Runtime script for a Niagara system */
+/** Scripts are function graphs that define the runtime execution for a Niagara system (similar to a Blueprint).
+ *
+ * There are three types of scripts:
+ * 1) Module: can be added as a standalone part to the emitter stack and encapsulates a single behavior, for example "Add Velocity"
+ * 2) Dynamic input: has a single output value and can be added to any input in the stack to compute such a value, for example "Random Vector" 
+ * 3) Function: usually reserved for helper functions; can only be called from within modules or dynamic inputs 
+ */
 UCLASS(MinimalAPI)
 class UNiagaraScript : public UNiagaraScriptBase, public FNiagaraVersionedObject
 {
diff --git a/Engine/Plugins/FX/Niagara/Source/Niagara/Classes/NiagaraSimCache.h b/Engine/Plugins/FX/Niagara/Source/Niagara/Classes/NiagaraSimCache.h
--- a/Engine/Plugins/FX/Niagara/Source/Niagara/Classes/NiagaraSimCache.h
+++ b/Engine/Plugins/FX/Niagara/Source/Niagara/Classes/NiagaraSimCache.h
@@ -379,7 +379,17 @@ struct FNiagaraSimCacheFeedbackContext
 	TArray<FString> Warnings;
 };
 
-UCLASS(Experimental, BlueprintType, MinimalAPI)
+/**
+ * Recording of multiple frames of simulation data from a running Niagara system.
+ * Can be used to play back the captured recording or to inspect the captured data for debug purposes.
+ * Depending on the capture settings, not all attributes from the simulation are present in the cache.
+ *
+ * To capture a cache, either
+ * (1) use the baker tool in the system editor,
+ * (2) use the Niagara component cache track in sequencer or
+ * (3) manually capture a running system with the "CaptureNiagaraSimCache" Blueprint functions
+ */
+UCLASS(BlueprintType, MinimalAPI)
 class UNiagaraSimCache : public UObject
 {
 	friend struct FNiagaraSimCacheAttributeReaderHelper;
diff --git a/Engine/Plugins/FX/Niagara/Source/Niagara/Classes/NiagaraValidationRuleSet.h b/Engine/Plugins/FX/Niagara/Source/Niagara/Classes/NiagaraValidationRuleSet.h
--- a/Engine/Plugins/FX/Niagara/Source/Niagara/Classes/NiagaraValidationRuleSet.h
+++ b/Engine/Plugins/FX/Niagara/Source/Niagara/Classes/NiagaraValidationRuleSet.h
@@ -5,7 +5,9 @@
 #include "NiagaraValidationRule.h"
 #include "NiagaraValidationRuleSet.generated.h"
 
-/** A set of reusable validation rules. */
+/** A set of reusable validation rules to check Niagara System assets.
+ * Once a number of rules are added to the rule set, it can be used either in effect types or configured as a global rule set in the Niagara plugin settings.
+ */
 UCLASS(BlueprintType, MinimalAPI)
 class UNiagaraValidationRuleSet : public UObject
 {
diff --git a/Engine/Plugins/FX/Niagara/Source/Niagara/Private/NiagaraSystem.cpp b/Engine/Plugins/FX/Niagara/Source/Niagara/Private/NiagaraSystem.cpp
--- a/Engine/Plugins/FX/Niagara/Source/Niagara/Private/NiagaraSystem.cpp
+++ b/Engine/Plugins/FX/Niagara/Source/Niagara/Private/NiagaraSystem.cpp
@@ -1793,7 +1793,7 @@ void UNiagaraSystem::GetAssetRegistryTagMetadata(TMap<FName, FAssetRegistryTagMe
 		NiagaraSystemPrivate::NAME_ActiveEmitters,
 		FAssetRegistryTagMetadata()
 		.SetDisplayName(LOCTEXT("ActiveEmitters", "Active Emitters"))
-		.SetTooltip(LOCTEXT("ActiveEmittersTooltip", "The nunmber of active emitters in the system"))
+		.SetTooltip(LOCTEXT("ActiveEmittersTooltip", "The number of active emitters in the system"))
 	);
 	OutMetadata.Add(
 		NiagaraSystemPrivate::NAME_ActiveRenderers,
diff --git a/Engine/Plugins/FX/Niagara/Source/Niagara/Public/NiagaraAssetTagDefinitions.h b/Engine/Plugins/FX/Niagara/Source/Niagara/Public/NiagaraAssetTagDefinitions.h
--- a/Engine/Plugins/FX/Niagara/Source/Niagara/Public/NiagaraAssetTagDefinitions.h
+++ b/Engine/Plugins/FX/Niagara/Source/Niagara/Public/NiagaraAssetTagDefinitions.h
@@ -111,7 +111,12 @@ private:
 	FGuid AssetTagDefinitionGuid;
 };
 
-/** An asset to define multiple tag definitions; used to sort and categorize Niagara assets. */
+/** An Asset Tag Definition defines a tag that can be added to various Niagara assets for sorting & filtering purposes.
+ * For example, custom tags will show up in the Create Niagara System dialog to filter available emitters.
+ * They can also be used to filter assets in the content browser, when used with the custom filter option.
+ *
+ * You can modify asset tags in the content browser by right-clicking on a Niagara asset, then use the "Manage Tags" submenu to add or remove them.
+ */
 UCLASS()
 class NIAGARA_API UNiagaraAssetTagDefinitions : public UObject
 {
diff --git a/Engine/Plugins/FX/Niagara/Source/NiagaraEditor/Public/NiagaraParameterDefinitions.h b/Engine/Plugins/FX/Niagara/Source/NiagaraEditor/Public/NiagaraParameterDefinitions.h
--- a/Engine/Plugins/FX/Niagara/Source/NiagaraEditor/Public/NiagaraParameterDefinitions.h
+++ b/Engine/Plugins/FX/Niagara/Source/NiagaraEditor/Public/NiagaraParameterDefinitions.h
@@ -44,7 +44,9 @@ public:
 };
 
 
-/** Collection of UNiagaraScriptVariables to synchronize between UNiagaraScripts. */
+/** Collection of Niagara parameters and their metadata, which is synced to all other assets in the project.
+ *  This can be used to enforce consistent naming and default values for commonly used parameters.
+ */
 UCLASS(MinimalAPI)
 class UNiagaraParameterDefinitions : public UNiagaraParameterDefinitionsBase
 {
