- Add custom deletor for stateless shader parameters

#rnx
#tests Ran ET PF to confirm Linux Asan did not trigger

diff --git a/Engine/Plugins/FX/Niagara/Source/Niagara/Internal/Stateless/NiagaraStatelessEmitterInstance.h b/Engine/Plugins/FX/Niagara/Source/Niagara/Internal/Stateless/NiagaraStatelessEmitterInstance.h
--- a/Engine/Plugins/FX/Niagara/Source/Niagara/Internal/Stateless/NiagaraStatelessEmitterInstance.h
+++ b/Engine/Plugins/FX/Niagara/Source/Niagara/Internal/Stateless/NiagaraStatelessEmitterInstance.h
@@ -12,6 +12,14 @@ class FShaderParametersMetadata;
 
 namespace NiagaraStateless
 {
+	struct FShaderParametersDeletor
+	{
+		void operator()(NiagaraStateless::FCommonShaderParameters* Ptr) const
+		{
+			FMemory::Free(Ptr);
+		}
+	};
+
 	class FEmitterInstance_RT : public INiagaraComputeDataBufferInterface
 	{
 	public:
@@ -25,7 +33,7 @@ namespace NiagaraStateless
 		float													DeltaTime = 0.0f;		//-OPT: We should be able to pull this from the view which is important for zero tick path
 		ENiagaraExecutionState									ExecutionState = ENiagaraExecutionState::Active;
 		TArray<FNiagaraStatelessRuntimeSpawnInfo>				SpawnInfos;
-		TUniquePtr<NiagaraStateless::FCommonShaderParameters>	ShaderParameters;
+		TUniquePtr<NiagaraStateless::FCommonShaderParameters, FShaderParametersDeletor>	ShaderParameters;
 
 		mutable TOptional<TArray<uint8>>						BindingBufferData;
 		mutable FReadBuffer										BindingBuffer;
