- Fix for false data interface ensures firing during transactions

#rnx

diff --git a/Engine/Plugins/FX/Niagara/Source/Niagara/Private/NiagaraScript.cpp b/Engine/Plugins/FX/Niagara/Source/Niagara/Private/NiagaraScript.cpp
--- a/Engine/Plugins/FX/Niagara/Source/Niagara/Private/NiagaraScript.cpp
+++ b/Engine/Plugins/FX/Niagara/Source/Niagara/Private/NiagaraScript.cpp
@@ -1816,8 +1816,16 @@ void UNiagaraScript::GenerateDefaultFunctionBindings()
 		const auto& ScriptDataInterfaces = ScriptParameterStore->GetDataInterfaces();
 
 		const int32 DataInterfaceCount = FMath::Min(CachedScriptVM.DataInterfaceInfo.Num(), ScriptDataInterfaces.Num());
-		ensureMsgf(DataInterfaceCount == CachedScriptVM.DataInterfaceInfo.Num(), TEXT("DataInterface count does not match VM this is likely caused by missing data interface classes"));
-		ensureMsgf(DataInterfaceCount == ScriptDataInterfaces.Num(), TEXT("DataInterface count does not match script data interfaces this is likely caused by missing data interface classes"));
+
+		// If we are transacting do not perform the ensures here as they can result in false positives UE-201266
+		// For example, if we were to undo / redo these ensures will execute and ScriptDataInterfaces could be out of date resulting in a false positive, we will recache the script data post action and get the correct data
+	#if WITH_EDITORONLY_DATA
+		if ( !GIsTransacting )
+	#endif
+		{
+			ensureMsgf(DataInterfaceCount == CachedScriptVM.DataInterfaceInfo.Num(), TEXT("DataInterface count does not match VM this is likely caused by missing data interface classes"));
+			ensureMsgf(DataInterfaceCount == ScriptDataInterfaces.Num(), TEXT("DataInterface count does not match script data interfaces this is likely caused by missing data interface classes"));
+		}
 
 		for (const FVMExternalFunctionBindingInfo& BindingInfo : CachedScriptVM.CalledVMExternalFunctions)
 		{
