Fixed a bug where hitting undo when editing a color user parameter resulted in a crash

diff --git a/Engine/Plugins/FX/Niagara/Source/NiagaraEditor/Private/Customizations/UserParameters/NiagaraComponentUserParametersBuilder.cpp b/Engine/Plugins/FX/Niagara/Source/NiagaraEditor/Private/Customizations/UserParameters/NiagaraComponentUserParametersBuilder.cpp
--- a/Engine/Plugins/FX/Niagara/Source/NiagaraEditor/Private/Customizations/UserParameters/NiagaraComponentUserParametersBuilder.cpp
+++ b/Engine/Plugins/FX/Niagara/Source/NiagaraEditor/Private/Customizations/UserParameters/NiagaraComponentUserParametersBuilder.cpp
@@ -120,6 +120,12 @@ public:
 		TSharedRef<FStructOnScope> ParameterEditorData = MakeShared<FStructOnScope>(ParameterKey.GetType().GetStruct());
 		ParameterEditor->UpdateStructFromInternalValue(ParameterEditorData);
 		
+		if (GIsTransacting)
+		{
+			// don't start a new transaction if we're currently applying another one (e.g. during an undo)
+			return;
+		}
+		
 		FScopedTransaction ScopedTransaction(FNiagaraUserParameterNodeBuilder::ChangedUserParameterTransactionText);
 		
 		bool bCancelTransaction = true;
