Improved reference count widget in parameter panel with detailed tooltip info

diff --git a/Engine/Plugins/FX/Niagara/Source/NiagaraEditor/Private/NiagaraEditorStyle.cpp b/Engine/Plugins/FX/Niagara/Source/NiagaraEditor/Private/NiagaraEditorStyle.cpp
--- a/Engine/Plugins/FX/Niagara/Source/NiagaraEditor/Private/NiagaraEditorStyle.cpp
+++ b/Engine/Plugins/FX/Niagara/Source/NiagaraEditor/Private/NiagaraEditorStyle.cpp
@@ -321,6 +321,8 @@ void FNiagaraEditorStyle::InitParameterMapView()
 	Set("NiagaraEditor.Stack.FlatButtonColor", FLinearColor(FColor(205, 205, 205)));
 
 	//Parameters panel
+	Set("NiagaraEditor.Parameters.ReadColor", FLinearColor(FColor(133, 108, 176)));
+	Set("NiagaraEditor.Parameters.WriteColor", FLinearColor(FColor(96, 184, 125)));
 	const FTableRowStyle TreeViewStyle = FAppStyle::GetWidgetStyle<FTableRowStyle>("DetailsView.TreeView.TableRow");
 	FTableRowStyle ParameterPanelRowStyle = FTableRowStyle(TreeViewStyle)
 		.SetTextColor(FLinearColor::White)
diff --git a/Engine/Plugins/FX/Niagara/Source/NiagaraEditor/Private/ViewModels/NiagaraParameterPanelViewModel.cpp b/Engine/Plugins/FX/Niagara/Source/NiagaraEditor/Private/ViewModels/NiagaraParameterPanelViewModel.cpp
--- a/Engine/Plugins/FX/Niagara/Source/NiagaraEditor/Private/ViewModels/NiagaraParameterPanelViewModel.cpp
+++ b/Engine/Plugins/FX/Niagara/Source/NiagaraEditor/Private/ViewModels/NiagaraParameterPanelViewModel.cpp
@@ -1894,12 +1894,9 @@ TArray<FNiagaraParameterPanelItem> FNiagaraSystemToolkitParameterPanelViewModel:
 		{
 			UNiagaraScriptVariable* ScriptVar = FNiagaraEditorUtilities::GetScriptVariableForUserParameter(Var, SystemViewModel);
 
-			FNiagaraParameterPanelItem Item = FNiagaraParameterPanelItem();
+			FNiagaraParameterPanelItem Item;
 			Item.ScriptVariable = ScriptVar;
 			Item.NamespaceMetaData = FNiagaraEditorUtilities::GetNamespaceMetaDataForVariableName(Var.GetName());
-			Item.bExternallyReferenced = false;
-			Item.bSourcedFromCustomStackContext = false;
-			Item.ReferenceCount = 0;
 			
 			// Determine whether the item is name aliasing a parameter definition's parameter.
 			Item.DefinitionMatchState = FNiagaraParameterDefinitionsUtilities::GetDefinitionMatchStateForParameter(ScriptVar->Variable);
@@ -1922,12 +1919,9 @@ TArray<FNiagaraParameterPanelItem> FNiagaraSystemToolkitParameterPanelViewModel:
 			ParameterToScriptVariableMap.Add(EditorOnlyScriptVar->Variable, EditorOnlyScriptVar);
 
 			const FNiagaraVariable& Var = EditorOnlyScriptVar->Variable;
-			FNiagaraParameterPanelItem Item = FNiagaraParameterPanelItem();
+			FNiagaraParameterPanelItem Item;
 			Item.ScriptVariable = EditorOnlyScriptVar;
 			Item.NamespaceMetaData = FNiagaraEditorUtilities::GetNamespaceMetaDataForVariableName(Var.GetName());
-			Item.bExternallyReferenced = false;
-			Item.bSourcedFromCustomStackContext = false;
-			Item.ReferenceCount = 0;
 
 			// Determine whether the item is name aliasing a parameter definition's parameter.
 			Item.DefinitionMatchState = FNiagaraParameterDefinitionsUtilities::GetDefinitionMatchStateForParameter(EditorOnlyScriptVar->Variable);
@@ -2018,6 +2012,31 @@ TArray<FNiagaraParameterPanelItem> FNiagaraSystemToolkitParameterPanelViewModel:
 					ParameterToScriptVariableMap.Append(VisitedExternalGraph->GetAllMetaData());
 				}
 
+				auto AddReadReference = [EditableGraph](const FNiagaraParameterMapHistory::FReadHistory& History, FNiagaraParameterPanelItem& Item)
+				{
+					FNiagaraParameterReferencePath Source;
+					Source.SourceGraph = EditableGraph;
+					Source.ModuleName = History.ReadPin.ModuleName;
+					if (Source.ModuleName == NAME_None && History.ReadPin.Pin->LinkedTo.Num() > 0)
+					{
+						// use the name of the module input for usages as stack function inputs
+						Source.ModuleName = History.ReadPin.Pin->LinkedTo[0]->PinName;
+					}
+					Item.AddToReadCount(Source);
+				};
+				auto AddWriteReference = [EditableGraph](const FNiagaraParameterMapHistory::FModuleScopedPin& History, FNiagaraParameterPanelItem& Item)
+				{
+					if (Cast<UNiagaraNodeParameterMapGet>(History.Pin->GetOwningNode()))
+					{
+						// ignore writes from map get nodes, as it's just the default value pin
+						return;
+					}
+					FNiagaraParameterReferencePath Source;
+					Source.SourceGraph = EditableGraph;
+					Source.ModuleName = History.ModuleName;
+					Item.AddToWriteCount(Source);
+				};
+
 				const TArray<FName>& CustomIterationSourceNamespaces = Builder.Histories[0].IterationNamespaceOverridesEncountered;
 				for (int32 VariableIndex = 0; VariableIndex < Builder.Histories[0].Variables.Num(); VariableIndex++)
 				{
@@ -2028,10 +2047,19 @@ TArray<FNiagaraParameterPanelItem> FNiagaraSystemToolkitParameterPanelViewModel:
 						continue;
 					}
 
+					TArray<TNiagaraParameterMapHistory<FNiagaraCompilationGraphBridge>::FReadHistory> ReadHistory = Builder.Histories[0].PerVariableReadHistory[VariableIndex];
+					TArray<TModuleScopedPin<UEdGraphPin>> WriteHistory = Builder.Histories[0].PerVariableWriteHistory[VariableIndex];
 					if (FNiagaraParameterPanelItem* ItemPtr = VisitedParameterToItemMap.Find(Var))
 					{
-						// This variable has already been registered, increment the reference count.
-						ItemPtr->ReferenceCount += Builder.Histories[0].PerVariableReadHistory[VariableIndex].Num() + Builder.Histories[0].PerVariableWriteHistory[VariableIndex].Num();
+						// This variable has already been registered, increment the reference counts.
+						for (const FNiagaraParameterMapHistory::FReadHistory& History : ReadHistory)
+						{
+							AddReadReference(History, *ItemPtr);
+						}
+						for (const FNiagaraParameterMapHistory::FModuleScopedPin& History : WriteHistory)
+						{
+							AddWriteReference(History, *ItemPtr);
+						}
 					}
 					else  // Add newly found variables
 					{
@@ -2063,12 +2091,12 @@ TArray<FNiagaraParameterPanelItem> FNiagaraSystemToolkitParameterPanelViewModel:
 						bool bVarOnlyInTopLevelGraph = true;
 						if (!bForceScript)
 						{
-							for (FNiagaraParameterMapHistory::FModuleScopedPin& WritePin : Builder.Histories[0].PerVariableWriteHistory[VariableIndex])
+							for (FNiagaraParameterMapHistory::FModuleScopedPin& WritePin : WriteHistory)
 							{
 								UEdGraphNode* VariableOwningNode = WritePin.Pin->GetOwningNode();
 								bVarOnlyInTopLevelGraph &= AllGraphs.Contains(static_cast<const UNiagaraGraph*>(VariableOwningNode->GetGraph()));
 							}
-							for (FNiagaraParameterMapHistory::FReadHistory& ReadPins : Builder.Histories[0].PerVariableReadHistory[VariableIndex])
+							for (FNiagaraParameterMapHistory::FReadHistory& ReadPins : ReadHistory)
 							{
 								UEdGraphNode* VariableOwningNode = ReadPins.ReadPin.Pin->GetOwningNode();
 								bVarOnlyInTopLevelGraph &= !AllGraphs.Contains(static_cast<const UNiagaraGraph*>(VariableOwningNode->GetGraph()));
@@ -2095,7 +2123,14 @@ TArray<FNiagaraParameterPanelItem> FNiagaraSystemToolkitParameterPanelViewModel:
 						Item.DefinitionMatchState = FNiagaraParameterDefinitionsUtilities::GetDefinitionMatchStateForParameter(Item.ScriptVariable->Variable);
 
 						// -Increment the reference count.
-						Item.ReferenceCount += Builder.Histories[0].PerVariableReadHistory[VariableIndex].Num() + Builder.Histories[0].PerVariableWriteHistory[VariableIndex].Num();
+						for (const FNiagaraParameterMapHistory::FReadHistory& History : ReadHistory)
+						{
+							AddReadReference(History, Item);
+						}
+						for (const FNiagaraParameterMapHistory::FModuleScopedPin& History : WriteHistory)
+						{
+							AddWriteReference(History, Item);
+						}
 
 						VisitedParameterToItemMap.Add(Var, Item);
 					}
@@ -2108,8 +2143,7 @@ TArray<FNiagaraParameterPanelItem> FNiagaraSystemToolkitParameterPanelViewModel:
 		{
 			if (EmitterVMS.Get().IsValid() && EmitterVMS.Get().GetIsEnabled() && EmitterVMS.Get().GetEmitterHandle())
 			{
-				FVersionedNiagaraEmitterData* ED = EmitterVMS.Get().GetEmitterHandle()->GetEmitterData();
-				if (ED)
+				if (FVersionedNiagaraEmitterData* ED = EmitterVMS.Get().GetEmitterHandle()->GetEmitterData())
 				{
 					FNiagaraAliasContext ResolveAliasesContext(FNiagaraAliasContext::ERapidIterationParameterMode::EmitterOrParticleScript);
 					ResolveAliasesContext.ChangeEmitterNameToEmitter(EmitterVMS.Get().GetEmitterHandle()->GetUniqueInstanceName());
@@ -2117,15 +2151,22 @@ TArray<FNiagaraParameterPanelItem> FNiagaraSystemToolkitParameterPanelViewModel:
 					ED->ForEachEnabledRenderer(
 						[&](UNiagaraRendererProperties* RenderProperties)
 						{
-							for (FNiagaraVariableBase BoundAttribute : RenderProperties->GetBoundAttributes())
+							for (const FNiagaraVariableAttributeBinding* Binding : RenderProperties->GetAttributeBindings())
 							{
-								BoundAttribute = FNiagaraUtilities::ResolveAliases(BoundAttribute, ResolveAliasesContext);
-
-								if (FNiagaraParameterPanelItem* ItemPtr = VisitedParameterToItemMap.Find(BoundAttribute))
+								FNiagaraVariable BoundAttribute = RenderProperties->GetBoundAttribute(Binding);
+								if (BoundAttribute.IsValid())
 								{
-									// This variable has already been registered, increment the reference count. Otherwise, it is 
-									// not a live binding and we can skip.
-									ItemPtr->ReferenceCount++;
+									BoundAttribute = FNiagaraUtilities::ResolveAliases(BoundAttribute, ResolveAliasesContext);
+
+									if (FNiagaraParameterPanelItem* ItemPtr = VisitedParameterToItemMap.Find(BoundAttribute))
+									{
+										// This variable has already been registered, increment the reference count. Otherwise, it is 
+										// not a live binding and we can skip.
+										FNiagaraParameterReferencePath Source;
+										Source.SourceGraph = Cast<UNiagaraScriptSource>(ED->GraphSource)->NodeGraph;
+										Source.ModuleName = FName(RenderProperties->GetWidgetDisplayName().ToString() + " (" + Binding->GetDataSetBindableVariable().GetName().ToString() + " Binding)");
+										ItemPtr->AddToReadCount(Source);
+									}
 								}
 							}
 						}
@@ -3146,11 +3187,9 @@ TArray<FNiagaraParameterPanelItem> FNiagaraScriptToolkitParameterPanelViewModel:
 				continue;
 			}
 
-			FNiagaraParameterPanelItem Item = FNiagaraParameterPanelItem();
+			FNiagaraParameterPanelItem Item;
 			Item.ScriptVariable = ScriptVariable;
 			Item.NamespaceMetaData = CandidateNamespaceMetaData;
-			Item.bExternallyReferenced = false;
-			Item.bSourcedFromCustomStackContext = false;
 
 			// Determine whether the item is name aliasing a parameter definition's parameter.
 			Item.DefinitionMatchState = FNiagaraParameterDefinitionsUtilities::GetDefinitionMatchStateForParameter(Parameter);
@@ -3159,11 +3198,17 @@ TArray<FNiagaraParameterPanelItem> FNiagaraScriptToolkitParameterPanelViewModel:
 			// for example a static switch that was added via propagation but then removed again will have no references. There might be other exceptions, too.
 			if(ReferenceCollection.Contains(Parameter))
 			{
-				Item.ReferenceCount = ReferenceCollection[Parameter].ParameterReferences.Num();
-			}
-			else
-			{
-				Item.ReferenceCount = 0;
+				for (const FNiagaraGraphParameterReference& Reference : ReferenceCollection[Parameter].ParameterReferences)
+				{
+					if (Cast<UNiagaraNodeParameterMapGet>(Reference.Value.Get()))
+					{
+						Item.ReadReferenceCount++;
+					}
+					else
+					{
+						Item.WriteReferenceCount++;
+					}
+				}
 			}
 
 			VisitedParameterToItemMap.Add(Parameter, Item);
@@ -3412,9 +3457,7 @@ TArray<FNiagaraParameterPanelItem> FNiagaraParameterDefinitionsToolkitParameterP
 		FNiagaraParameterPanelItem& Item = CachedViewedItems.AddDefaulted_GetRef();
 		Item.ScriptVariable = ScriptVar;
 		Item.NamespaceMetaData = FNiagaraEditorUtilities::GetNamespaceMetaDataForVariableName(ScriptVar->Variable.GetName());
-		Item.bExternallyReferenced = false;
-		Item.bSourcedFromCustomStackContext = false;
-		Item.ReferenceCount = 1;
+		Item.ReadReferenceCount = 1;
 		Item.DefinitionMatchState = EParameterDefinitionMatchState::MatchingOneDefinition;
 	}
 	return CachedViewedItems;
diff --git a/Engine/Plugins/FX/Niagara/Source/NiagaraEditor/Private/Widgets/SNiagaraParameterPanel.cpp b/Engine/Plugins/FX/Niagara/Source/NiagaraEditor/Private/Widgets/SNiagaraParameterPanel.cpp
--- a/Engine/Plugins/FX/Niagara/Source/NiagaraEditor/Private/Widgets/SNiagaraParameterPanel.cpp
+++ b/Engine/Plugins/FX/Niagara/Source/NiagaraEditor/Private/Widgets/SNiagaraParameterPanel.cpp
@@ -17,7 +17,6 @@
 #include "SDropTarget.h"
 #include "Widgets/SNiagaraPinTypeSelector.h"
 #include "ViewModels/NiagaraParameterPanelViewModel.h"
-#include "ViewModels/Stack/NiagaraStackGraphUtilities.h"
 #include "Widgets/SToolTip.h"
 #include "Widgets/Layout/SScaleBox.h"
 #include "Widgets/Text/SRichTextBlock.h"
@@ -26,6 +25,7 @@
 #include "Styling/AppStyle.h"
 #include "Widgets/Layout/SWrapBox.h"
 #include "Widgets/Input/SCheckBox.h"
+#include "Widgets/Input/SButton.h"
 
 #define LOCTEXT_NAMESPACE "NiagaraParameterPanel"
 
@@ -454,23 +454,70 @@ TSharedRef<SWidget> SNiagaraParameterPanel::OnGenerateWidgetForItem(const FNiaga
 
 	if (bShowParameterReferenceCounter)
 	{
+		TMap<FName, TArray<FNiagaraParameterReferencePath>> PathByParent;
+		for (const FNiagaraParameterReferencePath& Path : Item.ReferencePaths)
+		{
+			if (const UNiagaraGraph* NiagaraGraph = Path.SourceGraph.Get())
+			{
+				FName ParentName = NiagaraGraph->GetFName();
+				if (NiagaraGraph->GetOwningEmitter().Emitter)
+				{
+					ParentName = NiagaraGraph->GetOwningEmitter().Emitter.GetFName();
+				}
+				PathByParent.FindOrAdd(ParentName).Add(Path);
+			}
+		}
+
+		FString DetailReferences;
+		for (const TPair<FName, TArray<FNiagaraParameterReferencePath>>& Pair : PathByParent)
+		{
+			DetailReferences += "\n" + Pair.Key.ToString() + ":\n";
+			for (const FNiagaraParameterReferencePath& Path : Pair.Value)
+			{
+				FString Prefix = (Path.bRead && Path.bWrite) ? "R+W" : (Path.bRead ? "R" : "W");
+				DetailReferences += FString("\t") + Prefix + ": " +  Path.ModuleName.ToString() + "\n";
+			}
+		}
+		FText TooltipBaseText = LOCTEXT("ReferenceCountTooltip", "Shows the total number of references for this parameter across all selected emitters in the format 'Reads / Writes'.\nThis includes usages inside module scripts. If no emitter is selected then it only shows usages in the system script.");
+		FText FormattedTooltip = DetailReferences.IsEmpty() ? TooltipBaseText : FText::Format(LOCTEXT("DetailReferenceFmt", "{0}\nDetailed references:\n{1}"), TooltipBaseText, FText::FromString(DetailReferences));
+		
 		ItemWidgetHorizontalBox->AddSlot()
 		.VAlign(VAlign_Center)
 		.AutoWidth()
-		.Padding(3, 0)
 		[
-			SNew(SComboButton)
-			.HasDownArrow(false)
+			SNew(SButton)
 			.ButtonStyle(FAppStyle::Get(), "RoundButton")
+			.ToolTipText(FormattedTooltip)
 			.ForegroundColor(FSlateColor::UseForeground())
-			.ContentPadding(FMargin(2.0f))
 			.HAlign(HAlign_Right)
 			.VAlign(VAlign_Center)
-			.ButtonContent()
 			[
-				SNew(STextBlock)
-				.Text(FText::AsNumber(Item.ReferenceCount))
-				.Font(FCoreStyle::GetDefaultFontStyle(FName("Italic"), 10))
+				SNew(SHorizontalBox)
+				+SHorizontalBox::Slot()
+				.AutoWidth()
+				[
+					SNew(STextBlock)
+					.ColorAndOpacity(FNiagaraEditorStyle::Get().GetColor("NiagaraEditor.Parameters.ReadColor"))
+					.Text(FText::AsNumber(Item.ReadReferenceCount))
+					.Font(FCoreStyle::GetDefaultFontStyle(FName("Bold"), 10))
+				]
+				+SHorizontalBox::Slot()
+				.AutoWidth()
+				[
+					SNew(STextBlock)
+					.Margin(FMargin(2, 0))
+					.Text(FText::FromString(TEXT("/")))
+					.Font(FCoreStyle::GetDefaultFontStyle(FName("Regular"), 10))
+				]
+				+SHorizontalBox::Slot()
+				.AutoWidth()
+				[
+					SNew(STextBlock)
+					.ColorAndOpacity(FNiagaraEditorStyle::Get().GetColor("NiagaraEditor.Parameters.WriteColor"))
+					.Text(FText::AsNumber(Item.WriteReferenceCount))
+					.Font(FCoreStyle::GetDefaultFontStyle(FName("Bold"), 10))
+				]
+				
 			]
 		];
 	}
diff --git a/Engine/Plugins/FX/Niagara/Source/NiagaraEditor/Public/NiagaraParameterPanelTypes.h b/Engine/Plugins/FX/Niagara/Source/NiagaraEditor/Public/NiagaraParameterPanelTypes.h
--- a/Engine/Plugins/FX/Niagara/Source/NiagaraEditor/Public/NiagaraParameterPanelTypes.h
+++ b/Engine/Plugins/FX/Niagara/Source/NiagaraEditor/Public/NiagaraParameterPanelTypes.h
@@ -4,7 +4,9 @@
 
 #include "CoreMinimal.h"
 #include "NiagaraTypes.h"
+#include "NiagaraGraph.h"
 #include "NiagaraEditorSettings.h"
+#include "NiagaraEditorCommon.h"
 #include "Misc/Guid.h"
 #include "NiagaraScriptVariable.h"
 
@@ -35,6 +37,19 @@ public:
 	FNiagaraNamespaceMetadata NamespaceMetaData;
 };
 
+struct FNiagaraParameterReferencePath
+{
+	TWeakObjectPtr<const UNiagaraGraph> SourceGraph;
+	FName ModuleName = NAME_None;
+	bool bRead = false;
+	bool bWrite = false;
+
+	bool operator==(const FNiagaraParameterReferencePath& Other) const
+	{
+		return Other.SourceGraph == SourceGraph && Other.ModuleName == ModuleName;
+	};
+};
+
 struct FNiagaraParameterPanelItem : public FNiagaraParameterPanelItemBase
 {
 public:
@@ -52,18 +67,45 @@ public:
 
 	void RequestRenameNamespaceModifier() const { check(OnRequestRenameNamespaceModifierDelegate.IsBound()); OnRequestRenameNamespaceModifierDelegate.ExecuteIfBound(); };
 
-public:
+	void AddToReadCount(const FNiagaraParameterReferencePath& SourcePath)
+	{
+		ReadReferenceCount++;
+		FNiagaraParameterReferencePath* RefPath = ReferencePaths.FindByKey(SourcePath);
+		if (!RefPath)
+		{
+			RefPath = &ReferencePaths.Add_GetRef(SourcePath);
+		}
+		RefPath->bRead = true;
+	}
+
+	void AddToWriteCount(const FNiagaraParameterReferencePath& SourcePath)
+	{
+		WriteReferenceCount++;
+		FNiagaraParameterReferencePath* RefPath = ReferencePaths.FindByKey(SourcePath);
+		if (!RefPath)
+		{
+			RefPath = &ReferencePaths.Add_GetRef(SourcePath);
+		}
+		RefPath->bWrite = true;
+	}
+
 	/* For script variables; if true, the variable is sourced from a script that is not owned by the emitter/system the parameter panel is referencing. */
-	bool bExternallyReferenced;
+	bool bExternallyReferenced = false;
 
 	/* For script variables; if true, the variable is a member of a custom stack context for an emitter/system. */
-	bool bSourcedFromCustomStackContext;
+	bool bSourcedFromCustomStackContext = false;
+
+	/* Count of read references to the variable in graphs viewed by a parameter panel view model. */
+	int32 ReadReferenceCount = 0;
+
+	/* Count of write references to the variable in graphs viewed by a parameter panel view model. */
+	int32 WriteReferenceCount = 0;
 
-	/* Count of references to the variable in graphs viewed by a parameter panel view model. */
-	int32 ReferenceCount;
+	/* A detailed list of unique references (i.e. modules) to display in the parameters panel */
+	TArray<FNiagaraParameterReferencePath> ReferencePaths;
 
 	/* The relation of this parameter item to all parameter definitions it is matching. Whether the parameter item is subscribed to a definition is tracked by the UNiagaraScriptVariable's bSubscribedToParameterDefinitions member. */
-	EParameterDefinitionMatchState DefinitionMatchState;
+	EParameterDefinitionMatchState DefinitionMatchState = EParameterDefinitionMatchState::NoMatchingDefinitions;
 
 private:
 	mutable FOnRequestRename OnRequestRenameDelegate;
