Assignment node/parameter set renaming changes for a given parameter:
Now actually renames the parameter instead of changing the assignment target ('renaming' an assignment target float -> float2 would create a new attribute called float2 instead of renaming float)
This requires that only editable parameters are allowed to be renamed. Particles.Age is externally referenced in plenty of scripts, you don't want to be able to rename it.
Also added the same lock icon and a tooltip next to the parameter in Set Parameters UI as in the parameters panel so the user has an explanation

diff --git a/Engine/Plugins/FX/Niagara/Source/NiagaraEditor/Private/ViewModels/Stack/NiagaraStackFunctionInput.cpp b/Engine/Plugins/FX/Niagara/Source/NiagaraEditor/Private/ViewModels/Stack/NiagaraStackFunctionInput.cpp
--- a/Engine/Plugins/FX/Niagara/Source/NiagaraEditor/Private/ViewModels/Stack/NiagaraStackFunctionInput.cpp
+++ b/Engine/Plugins/FX/Niagara/Source/NiagaraEditor/Private/ViewModels/Stack/NiagaraStackFunctionInput.cpp
@@ -2564,8 +2564,30 @@ void UNiagaraStackFunctionInput::OnMessageManagerRefresh(const TArray<TSharedRef
 bool UNiagaraStackFunctionInput::SupportsRename() const
 {
 	// Only module level assignment node inputs can be renamed.
-	return OwningAssignmentNode.IsValid() && InputParameterHandlePath.Num() == 1 &&
-		OwningAssignmentNode->FindAssignmentTarget(InputParameterHandle.GetName()) != INDEX_NONE;
+	if(OwningAssignmentNode.IsValid() == false || InputParameterHandlePath.Num() != 1 || OwningAssignmentNode->FindAssignmentTarget(InputParameterHandle.GetName()) == INDEX_NONE)
+	{
+		return false;
+	}
+	
+	if(bIsReferencingExternalParameter.IsSet() == false)
+	{
+		const FNiagaraParameterPanelItem* Item = GetSystemViewModel()->GetParameterPanelViewModel()->GetCachedViewedParameterItems().FindByPredicate([this](const FNiagaraParameterPanelItem& ItemCandidate)
+		{
+			if(ItemCandidate.GetVariable().GetName() == this->GetInputParameterHandle().GetName())
+			{
+				return true;
+			}
+
+			return false;
+		});
+
+		if(Item)
+		{
+			bIsReferencingExternalParameter = Item->bExternallyReferenced;
+		}
+	}
+
+	return bIsReferencingExternalParameter.Get(false) == false;
 }
 
 void UNiagaraStackFunctionInput::OnRenamed(FText NewNameText)
@@ -2573,30 +2595,13 @@ void UNiagaraStackFunctionInput::OnRenamed(FText NewNameText)
 	FName NewName(*NewNameText.ToString());
 	FNiagaraVariable OldVar = FNiagaraVariable(InputType, InputParameterHandle.GetName());
 	FNiagaraVariable NewVar = FNiagaraVariable(InputType, NewName);
-	if (InputParameterHandle.GetName() != NewName && OwningAssignmentNode.IsValid() && SourceScript.IsValid())
-	{
-		TSharedRef<FNiagaraSystemViewModel> CachedSysViewModel = GetSystemViewModel();
-		TSharedPtr<FNiagaraEmitterViewModel> CachedEmitterViewModel = GetEmitterViewModel();
-		UNiagaraSystem& System = GetSystemViewModel()->GetSystem();
-		FVersionedNiagaraEmitter Emitter = GetEmitterViewModel().IsValid() ? GetEmitterViewModel()->GetEmitter() : FVersionedNiagaraEmitter();
-
-		FScopedTransaction ScopedTransaction(LOCTEXT("RenameInput", "Rename this function's input."));
-		FNiagaraStackGraphUtilities::RenameAssignmentTarget(
-			System,
-			Emitter,
-			*SourceScript.Get(),
-			*OwningAssignmentNode.Get(),
-			OldVar,
-			NewName);
-		ensureMsgf(IsFinalized(), TEXT("Input not finalized when renamed."));
-
-
-		if (CachedSysViewModel->GetSystemStackViewModel())
-			CachedSysViewModel->GetSystemStackViewModel()->InvalidateCachedParameterUsage();
-
-		CachedSysViewModel->NotifyParameterRenamedExternally(OldVar, NewVar, Emitter.Emitter);
 
+	if(InputParameterHandle.GetName() != NewName && OwningAssignmentNode.IsValid() && SourceScript.IsValid() && bIsReferencingExternalParameter.Get(false) == false)
+	{
+		GetSystemViewModel()->RenameParameter(OldVar, NewName, ENiagaraGetGraphParameterReferencesMode::SelectedGraphs);
 	}
+
+	bIsReferencingExternalParameter.Reset();
 }
 
 bool UNiagaraStackFunctionInput::CanDeleteInput() const
diff --git a/Engine/Plugins/FX/Niagara/Source/NiagaraEditor/Public/ViewModels/Stack/NiagaraStackFunctionInput.h b/Engine/Plugins/FX/Niagara/Source/NiagaraEditor/Public/ViewModels/Stack/NiagaraStackFunctionInput.h
--- a/Engine/Plugins/FX/Niagara/Source/NiagaraEditor/Public/ViewModels/Stack/NiagaraStackFunctionInput.h
+++ b/Engine/Plugins/FX/Niagara/Source/NiagaraEditor/Public/ViewModels/Stack/NiagaraStackFunctionInput.h
@@ -477,6 +477,8 @@ private:
 	/** Pointers and handles to the various values this input can have. */
 	FInputValues InputValues;
 
+	mutable TOptional<bool> bIsReferencingExternalParameter;
+	
 	TSharedPtr<FNiagaraPlaceholderDataInterfaceHandle> PlaceholderDataInterfaceHandle;
 
 	/** A cached pointer to the override node for this input if it exists.  This value is cached here since the
diff --git a/Engine/Plugins/FX/Niagara/Source/NiagaraEditorWidgets/Private/Stack/SNiagaraStackFunctionInputName.cpp b/Engine/Plugins/FX/Niagara/Source/NiagaraEditorWidgets/Private/Stack/SNiagaraStackFunctionInputName.cpp
--- a/Engine/Plugins/FX/Niagara/Source/NiagaraEditorWidgets/Private/Stack/SNiagaraStackFunctionInputName.cpp
+++ b/Engine/Plugins/FX/Niagara/Source/NiagaraEditorWidgets/Private/Stack/SNiagaraStackFunctionInputName.cpp
@@ -1,6 +1,8 @@
 // Copyright Epic Games, Inc. All Rights Reserved.
 
 #include "Stack/SNiagaraStackFunctionInputName.h"
+
+#include "EditorFontGlyphs.h"
 #include "NiagaraEditorWidgetsStyle.h"
 #include "NiagaraEditorStyle.h"
 #include "ViewModels/Stack/NiagaraStackViewModel.h"
@@ -71,6 +73,14 @@ void SNiagaraStackFunctionInputName::Construct(const FArguments& InArgs, UNiagar
 			.OnTextCommitted(this, &SNiagaraStackFunctionInputName::OnNameTextCommitted)
 			.HighlightText_UObject(InStackViewModel, &UNiagaraStackViewModel::GetCurrentSearchText)
 			.ToolTipText(this, &SNiagaraStackFunctionInputName::GetToolTipText)
+			.Decorator()
+			[
+				SNew(STextBlock)
+				.Visibility(FunctionInput->SupportsRename() == false ? EVisibility::Visible : EVisibility::Collapsed)
+				.Font(FAppStyle::Get().GetFontStyle("FontAwesome.8"))
+				.Text(FEditorFontGlyphs::Lock)
+				.ToolTipText(LOCTEXT("LockedToolTip", "This parameter is used in a referenced external graph. You can set its value, but not rename it."))
+			]
 		];
 	}
 	else
