Nuke the legacy attribute spreadsheet

diff --git a/Engine/Plugins/FX/Niagara/Source/NiagaraEditor/Private/Toolkits/NiagaraSystemToolkit.cpp b/Engine/Plugins/FX/Niagara/Source/NiagaraEditor/Private/Toolkits/NiagaraSystemToolkit.cpp
--- a/Engine/Plugins/FX/Niagara/Source/NiagaraEditor/Private/Toolkits/NiagaraSystemToolkit.cpp
+++ b/Engine/Plugins/FX/Niagara/Source/NiagaraEditor/Private/Toolkits/NiagaraSystemToolkit.cpp
@@ -19,7 +19,6 @@
 #include "Widgets/SNiagaraSystemScript.h"
 #include "Widgets/SNiagaraSystemViewport.h"
 #include "Widgets/SNiagaraParameterPanel.h"
-#include "Widgets/SNiagaraSpreadsheetView.h"
 #include "Widgets/SNiagaraDebugger.h"
 #include "NiagaraEditorCommands.h"
 #include "NiagaraEditorStyle.h"
diff --git a/Engine/Plugins/FX/Niagara/Source/NiagaraEditor/Private/Toolkits/SystemToolkitModes/NiagaraSystemToolkitModeBase.cpp b/Engine/Plugins/FX/Niagara/Source/NiagaraEditor/Private/Toolkits/SystemToolkitModes/NiagaraSystemToolkitModeBase.cpp
--- a/Engine/Plugins/FX/Niagara/Source/NiagaraEditor/Private/Toolkits/SystemToolkitModes/NiagaraSystemToolkitModeBase.cpp
+++ b/Engine/Plugins/FX/Niagara/Source/NiagaraEditor/Private/Toolkits/SystemToolkitModes/NiagaraSystemToolkitModeBase.cpp
@@ -26,7 +26,6 @@
 #include "Widgets/SNiagaraParameterDefinitionsPanel.h"
 #include "Widgets/SNiagaraScriptGraph.h"
 #include "Widgets/SNiagaraSelectedObjectsDetails.h"
-#include "Widgets/SNiagaraSpreadsheetView.h"
 #include "Widgets/SNiagaraGeneratedCodeView.h"
 #include "Widgets/SNiagaraHierarchyEditor.h"
 #include "Widgets/Layout/SWidgetSwitcher.h"
@@ -60,7 +59,6 @@ const FName FNiagaraSystemToolkitModeBase::SystemParametersTabID(TEXT("NiagaraSy
 const FName FNiagaraSystemToolkitModeBase::SystemParameterDefinitionsTabID(TEXT("NiagaraSystemEditor_SystemParameterDefinitions"));
 const FName FNiagaraSystemToolkitModeBase::DetailsTabID(TEXT("NiagaraSystemEditor_Details"));
 const FName FNiagaraSystemToolkitModeBase::SelectedEmitterGraphTabID(TEXT("NiagaraSystemEditor_SelectedEmitterGraph"));
-const FName FNiagaraSystemToolkitModeBase::DebugSpreadsheetTabID(TEXT("NiagaraSystemEditor_DebugAttributeSpreadsheet"));
 const FName FNiagaraSystemToolkitModeBase::DebugCacheSpreadsheetTabID(TEXT("NiagaraSystemEditor_DebugCacheSpreadsheet"));
 const FName FNiagaraSystemToolkitModeBase::PreviewSettingsTabId(TEXT("NiagaraSystemEditor_PreviewSettings"));
 const FName FNiagaraSystemToolkitModeBase::GeneratedCodeTabID(TEXT("NiagaraSystemEditor_GeneratedCode"));
@@ -274,11 +272,6 @@ void FNiagaraSystemToolkitModeBase::RegisterTabFactories(TSharedPtr<FTabManager>
 		.SetGroup(WorkspaceMenuCategory.ToSharedRef())
 		.SetAutoGenerateMenuEntry(GbShowNiagaraDeveloperWindows != 0);
 
-	InTabManager->RegisterTabSpawner(DebugSpreadsheetTabID, FOnSpawnTab::CreateSP(this, &FNiagaraSystemToolkitModeBase::SpawnTab_DebugSpreadsheet))
-		.SetDisplayName(LOCTEXT("DebugSpreadsheet", "Attribute Spreadsheet (Legacy)"))
-		.SetGroup(WorkspaceMenuCategory.ToSharedRef())
-		.SetIcon(FSlateIcon(FNiagaraEditorStyle::Get().GetStyleSetName(), "Tab.Spreadsheet"));
-
 	InTabManager->RegisterTabSpawner(DebugCacheSpreadsheetTabID, FOnSpawnTab::CreateSP(this, &FNiagaraSystemToolkitModeBase::SpawnTab_DebugCacheSpreadsheet))
 		.SetDisplayName(LOCTEXT("DebugSpreadshseet", "Attribute Spreadsheet"))
 		.SetGroup(WorkspaceMenuCategory.ToSharedRef())
@@ -685,19 +678,6 @@ TSharedRef<SDockTab> FNiagaraSystemToolkitModeBase::SpawnTab_SelectedEmitterGrap
 	return SpawnedTab;
 }
 
-TSharedRef<SDockTab> FNiagaraSystemToolkitModeBase::SpawnTab_DebugSpreadsheet(const FSpawnTabArgs& Args)
-{
-	check(Args.GetTabId().TabType == DebugSpreadsheetTabID);
-
-	TSharedRef<SDockTab> SpawnedTab =
-		SNew(SDockTab)
-		[
-			SNew(SNiagaraSpreadsheetView, SystemToolkit.Pin()->GetSystemViewModel().ToSharedRef())
-		];
-
-	return SpawnedTab;
-}
-
 TSharedRef<SDockTab> FNiagaraSystemToolkitModeBase::SpawnTab_DebugCacheSpreadsheet(const FSpawnTabArgs& Args)
 {
 	check(Args.GetTabId().TabType == DebugCacheSpreadsheetTabID);
diff --git a/Engine/Plugins/FX/Niagara/Source/NiagaraEditor/Private/Toolkits/SystemToolkitModes/NiagaraSystemToolkitMode_Default.cpp b/Engine/Plugins/FX/Niagara/Source/NiagaraEditor/Private/Toolkits/SystemToolkitModes/NiagaraSystemToolkitMode_Default.cpp
--- a/Engine/Plugins/FX/Niagara/Source/NiagaraEditor/Private/Toolkits/SystemToolkitModes/NiagaraSystemToolkitMode_Default.cpp
+++ b/Engine/Plugins/FX/Niagara/Source/NiagaraEditor/Private/Toolkits/SystemToolkitModes/NiagaraSystemToolkitMode_Default.cpp
@@ -74,7 +74,6 @@ FNiagaraSystemToolkitMode_Default::FNiagaraSystemToolkitMode_Default(TWeakPtr<FN
 					->AddTab(DetailsTabID, ETabState::OpenedTab)
 					->AddTab(SelectedEmitterGraphTabID, ETabState::ClosedTab)
 					->AddTab(SystemScriptTabID, ETabState::ClosedTab)
-					->AddTab(DebugSpreadsheetTabID, ETabState::ClosedTab)
 					->AddTab(PreviewSettingsTabId, ETabState::ClosedTab)
 					->AddTab(GeneratedCodeTabID, ETabState::ClosedTab)
 				)
diff --git a/Engine/Plugins/FX/Niagara/Source/NiagaraEditor/Private/Toolkits/SystemToolkitModes/NiagaraSystemToolkitMode_Scalability.cpp b/Engine/Plugins/FX/Niagara/Source/NiagaraEditor/Private/Toolkits/SystemToolkitModes/NiagaraSystemToolkitMode_Scalability.cpp
--- a/Engine/Plugins/FX/Niagara/Source/NiagaraEditor/Private/Toolkits/SystemToolkitModes/NiagaraSystemToolkitMode_Scalability.cpp
+++ b/Engine/Plugins/FX/Niagara/Source/NiagaraEditor/Private/Toolkits/SystemToolkitModes/NiagaraSystemToolkitMode_Scalability.cpp
@@ -389,7 +389,6 @@ TSharedRef<FTabManager::FLayout> FNiagaraSystemToolkitMode_Scalability::ChooseTa
 						->AddTab(DetailsTabID, ETabState::OpenedTab)
 						->AddTab(SelectedEmitterGraphTabID, ETabState::ClosedTab)
 						->AddTab(SystemScriptTabID, ETabState::ClosedTab)
-						->AddTab(DebugSpreadsheetTabID, ETabState::ClosedTab)
 						->AddTab(PreviewSettingsTabId, ETabState::ClosedTab)
 						->AddTab(GeneratedCodeTabID, ETabState::ClosedTab)
 						->SetForegroundTab(ResolvedScalabilityTabID)
@@ -460,7 +459,6 @@ TSharedRef<FTabManager::FLayout> FNiagaraSystemToolkitMode_Scalability::ChooseTa
 						->AddTab(DetailsTabID, ETabState::OpenedTab)
 						->AddTab(SelectedEmitterGraphTabID, ETabState::ClosedTab)
 						->AddTab(SystemScriptTabID, ETabState::ClosedTab)
-						->AddTab(DebugSpreadsheetTabID, ETabState::ClosedTab)
 						->AddTab(PreviewSettingsTabId, ETabState::ClosedTab)
 						->AddTab(GeneratedCodeTabID, ETabState::ClosedTab)
 					)
diff --git a/Engine/Plugins/FX/Niagara/Source/NiagaraEditor/Private/Widgets/SNiagaraSpreadsheetView.cpp b/Engine/Plugins/FX/Niagara/Source/NiagaraEditor/Private/Widgets/SNiagaraSpreadsheetView.cpp
deleted file mode 100644
diff --git a/Engine/Plugins/FX/Niagara/Source/NiagaraEditor/Private/Widgets/SNiagaraSpreadsheetView.h b/Engine/Plugins/FX/Niagara/Source/NiagaraEditor/Private/Widgets/SNiagaraSpreadsheetView.h
deleted file mode 100644
diff --git a/Engine/Plugins/FX/Niagara/Source/NiagaraEditor/Public/Toolkits/SystemToolkitModes/NiagaraSystemToolkitModeBase.h b/Engine/Plugins/FX/Niagara/Source/NiagaraEditor/Public/Toolkits/SystemToolkitModes/NiagaraSystemToolkitModeBase.h
--- a/Engine/Plugins/FX/Niagara/Source/NiagaraEditor/Public/Toolkits/SystemToolkitModes/NiagaraSystemToolkitModeBase.h
+++ b/Engine/Plugins/FX/Niagara/Source/NiagaraEditor/Public/Toolkits/SystemToolkitModes/NiagaraSystemToolkitModeBase.h
@@ -34,7 +34,6 @@ private:
 	TSharedRef<SDockTab> SpawnTab_SystemParameterDefinitions(const FSpawnTabArgs& Args);
 	TSharedRef<SDockTab> SpawnTab_Details(const FSpawnTabArgs& Args);
 	TSharedRef<SDockTab> SpawnTab_SelectedEmitterGraph(const FSpawnTabArgs& Args);
-	TSharedRef<SDockTab> SpawnTab_DebugSpreadsheet(const FSpawnTabArgs& Args);
 	TSharedRef<SDockTab> SpawnTab_DebugCacheSpreadsheet(const FSpawnTabArgs& Args);
 	TSharedRef<SDockTab> SpawnTab_PreviewSettings(const FSpawnTabArgs& Args);
 	TSharedRef<SDockTab> SpawnTab_GeneratedCode(const FSpawnTabArgs& Args);
@@ -94,7 +93,6 @@ public:
 	static const FName SystemParameterDefinitionsTabID;
 	static const FName DetailsTabID;
 	static const FName SelectedEmitterGraphTabID;
-	static const FName DebugSpreadsheetTabID;
 	static const FName DebugCacheSpreadsheetTabID;
 	static const FName PreviewSettingsTabId;
 	static const FName GeneratedCodeTabID;
