Niagara - Fix playback in the system/emitter editor so that the playback framerate matches the rate set in the sequencer timeline.  Also save the current display rate setting and the lock state in the system editor data so that it persists with the asset.

diff --git a/Engine/Plugins/FX/Niagara/Source/NiagaraEditor/Private/NiagaraSystemEditorData.cpp b/Engine/Plugins/FX/Niagara/Source/NiagaraEditor/Private/NiagaraSystemEditorData.cpp
--- a/Engine/Plugins/FX/Niagara/Source/NiagaraEditor/Private/NiagaraSystemEditorData.cpp
+++ b/Engine/Plugins/FX/Niagara/Source/NiagaraEditor/Private/NiagaraSystemEditorData.cpp
@@ -181,6 +181,16 @@ void UNiagaraSystemEditorData::SetPlaybackRange(TRange<float> InPlaybackRange)
 	PlaybackRangeMax = InPlaybackRange.GetUpperBoundValue();
 }
 
+void UNiagaraSystemEditorData::SetPlaybackFrameRate(const FFrameRate InPlaybackFrameRate)
+{
+	PlaybackFrameRate = InPlaybackFrameRate;
+}
+
+void UNiagaraSystemEditorData::SetLockPlaybackFrameRate(bool bInLockPlaybackFrameRate)
+{
+	bLockPlaybackFrameRate = bInLockPlaybackFrameRate;
+}
+
 UEdGraph* UNiagaraSystemEditorData::GetSystemOverviewGraph() const
 {
 	return SystemOverviewGraph;	
diff --git a/Engine/Plugins/FX/Niagara/Source/NiagaraEditor/Private/ViewModels/NiagaraSystemViewModel.cpp b/Engine/Plugins/FX/Niagara/Source/NiagaraEditor/Private/ViewModels/NiagaraSystemViewModel.cpp
--- a/Engine/Plugins/FX/Niagara/Source/NiagaraEditor/Private/ViewModels/NiagaraSystemViewModel.cpp
+++ b/Engine/Plugins/FX/Niagara/Source/NiagaraEditor/Private/ViewModels/NiagaraSystemViewModel.cpp
@@ -1357,6 +1357,16 @@ TStatId FNiagaraSystemViewModel::GetStatId() const
 	RETURN_QUICK_DECLARE_CYCLE_STAT(FNiagaraSystemViewModel, STATGROUP_Tickables);
 }
 
+void FNiagaraSystemViewModel::OnModifiedIndirectly(UMovieSceneSignedObject* MovieSceneSignedObject)
+{
+	SequencerMovieSceneModified(Cast<UMovieScene>(MovieSceneSignedObject));
+}
+
+void FNiagaraSystemViewModel::OnModifiedDirectly(UMovieSceneSignedObject* MovieSceneSignedObject)
+{
+	SequencerMovieSceneModified(Cast<UMovieScene>(MovieSceneSignedObject));
+}
+
 TSharedRef<FNiagaraPlaceholderDataInterfaceManager> FNiagaraSystemViewModel::GetPlaceholderDataInterfaceManager()
 {
 	return PlaceholderDataInterfaceManager.ToSharedRef();
@@ -2346,9 +2356,15 @@ void FNiagaraSystemViewModel::SetupSequencer()
 		// we don't need a sequencer when merging emitters or if we're in a commandlet with no slate application
 		return;
 	}
+
 	NiagaraSequence = NewObject<UNiagaraSequence>(GetTransientPackage());
 	UMovieScene* MovieScene = NewObject<UMovieScene>(NiagaraSequence, FName("Niagara System MovieScene"), RF_Transactional | RF_Transient);
-	MovieScene->SetDisplayRate(FFrameRate(240, 1));
+
+	MovieScene->SetDisplayRate(GetEditorData().GetPlaybackFrameRate());
+	MovieScene->SetEvaluationType(GetEditorData().GetLockPlaybackFrameRate() ? EMovieSceneEvaluationType::FrameLocked : EMovieSceneEvaluationType::WithSubFrames);
+
+	MovieSceneEventHandler.Unlink();
+	MovieScene->UMovieSceneSignedObject::EventHandlers.Link(MovieSceneEventHandler, this);
 
 	NiagaraSequence->Initialize(this, MovieScene);
 
@@ -2787,6 +2803,26 @@ void PopulateNiagaraFoldersFromMovieSceneFolders(TArrayView<UMovieSceneFolder* c
 	}
 }
 
+void FNiagaraSystemViewModel::SequencerMovieSceneModified(const UMovieScene* MovieScene)
+{
+	if (MovieScene == nullptr)
+	{
+		return;
+	}
+
+	UNiagaraSystemEditorData& SystemEditorData = GetEditorData();
+	if (SystemEditorData.GetPlaybackFrameRate() != MovieScene->GetDisplayRate())
+	{
+		SystemEditorData.SetPlaybackFrameRate(MovieScene->GetDisplayRate());
+	}
+
+	bool bMovieSceneIsFrameLocked = MovieScene->GetEvaluationType() == EMovieSceneEvaluationType::FrameLocked;
+	if (SystemEditorData.GetLockPlaybackFrameRate() != bMovieSceneIsFrameLocked)
+	{
+		SystemEditorData.SetLockPlaybackFrameRate(bMovieSceneIsFrameLocked);
+	}
+}
+
 void FNiagaraSystemViewModel::SequencerDataChanged(EMovieSceneDataChangeType DataChangeType)
 {
 	if (bUpdatingSequencerFromEmitterDataChange == false && GIsTransacting == false && NiagaraSequence)
diff --git a/Engine/Plugins/FX/Niagara/Source/NiagaraEditor/Public/NiagaraSystemEditorData.h b/Engine/Plugins/FX/Niagara/Source/NiagaraEditor/Public/NiagaraSystemEditorData.h
--- a/Engine/Plugins/FX/Niagara/Source/NiagaraEditor/Public/NiagaraSystemEditorData.h
+++ b/Engine/Plugins/FX/Niagara/Source/NiagaraEditor/Public/NiagaraSystemEditorData.h
@@ -2,9 +2,11 @@
 
 #pragma once
 
+#include "Misc/FrameRate.h"
 #include "NiagaraEditorDataBase.h"
 #include "NiagaraEditorData.h"
 #include "ViewModels/HierarchyEditor/NiagaraHierarchyViewModelBase.h"
+
 #include "NiagaraSystemEditorData.generated.h"
 
 class UNiagaraStackEditorData;
@@ -97,6 +99,14 @@ public:
 
 	NIAGARAEDITOR_API void SetPlaybackRange(TRange<float> InPlaybackRange);
 
+	FFrameRate GetPlaybackFrameRate() const { return PlaybackFrameRate; }
+
+	void SetPlaybackFrameRate(const FFrameRate InPlaybackFrameRate);
+
+	bool GetLockPlaybackFrameRate() const { return bLockPlaybackFrameRate; }
+
+	void SetLockPlaybackFrameRate(bool bInLockPlaybackFrameRate);
+
 	NIAGARAEDITOR_API UEdGraph* GetSystemOverviewGraph() const;
 
 	NIAGARAEDITOR_API const FNiagaraGraphViewSettings& GetSystemOverviewGraphViewSettings() const;
@@ -146,6 +156,12 @@ private:
 	UPROPERTY()
 	float PlaybackRangeMax;
 
+	UPROPERTY()
+	FFrameRate PlaybackFrameRate = FFrameRate(60, 1);
+
+	UPROPERTY()
+	bool bLockPlaybackFrameRate = true;
+
 	/** Graph presenting overview of the current system and its emitters. */
 	UPROPERTY()
 	TObjectPtr<UEdGraph> SystemOverviewGraph;
diff --git a/Engine/Plugins/FX/Niagara/Source/NiagaraEditor/Public/ViewModels/NiagaraSystemViewModel.h b/Engine/Plugins/FX/Niagara/Source/NiagaraEditor/Public/ViewModels/NiagaraSystemViewModel.h
--- a/Engine/Plugins/FX/Niagara/Source/NiagaraEditor/Public/ViewModels/NiagaraSystemViewModel.h
+++ b/Engine/Plugins/FX/Niagara/Source/NiagaraEditor/Public/ViewModels/NiagaraSystemViewModel.h
@@ -17,6 +17,7 @@
 #include "UObject/ObjectKey.h"
 #include "ViewModels/TNiagaraViewModelManager.h"
 #include "ViewModels/NiagaraParameterDefinitionsSubscriberViewModel.h"
+#include "EventHandlers/ISignedObjectEventHandler.h"
 
 struct FAssetData;
 struct FEdGraphEditAction;
@@ -117,6 +118,7 @@ class FNiagaraSystemViewModel
 	, public FTickableEditorObject
 	, public TNiagaraViewModelManager<UNiagaraSystem, FNiagaraSystemViewModel>
 	, public INiagaraParameterDefinitionsSubscriberViewModel
+	, public UE::MovieScene::ISignedObjectEventHandler
 {
 public:
 	DECLARE_MULTICAST_DELEGATE(FOnEmitterHandleViewModelsChanged);
@@ -279,6 +281,10 @@ public:
 	virtual bool IsTickable() const override { return true; }
 	virtual TStatId GetStatId() const override;
 
+	//~ UE::MovieScene::ISignedObjectEventHandler interface
+	virtual void OnModifiedIndirectly(UMovieSceneSignedObject* MovieSceneSignedObject) override;
+	virtual void OnModifiedDirectly(UMovieSceneSignedObject* MovieSceneSignedObject) override;
+
 	/** Resets the System instance to initial conditions. Tries to resets system simulation time. Does not reset all systems that share its emitters.
 	 * Does not reinitialize the system to pull in changes. Calls into overloaded ResetSystem(). */
 	void ResetSystem();
@@ -552,6 +558,9 @@ private:
 	/** Called when a script is compiled */
 	void ScriptCompiled(UNiagaraScript* InScript, const FGuid& ScriptVersion);
 
+	/** Called when the sequencer movie scene object is changed. */
+	void SequencerMovieSceneModified(const UMovieScene* MovieScene);
+
 	/** Called whenever the data in the sequence is changed. */
 	void SequencerDataChanged(EMovieSceneDataChangeType DataChangeType);
 
@@ -633,6 +642,9 @@ private:
 	/** The view model for the System script. */
 	TSharedPtr<FNiagaraSystemScriptViewModel> SystemScriptViewModel;
 
+	/** Handles linking change events from the sequencer movie scene. */
+	UE::MovieScene::TNonIntrusiveEventHandler<UE::MovieScene::ISignedObjectEventHandler> MovieSceneEventHandler;
+
 	/** A niagara sequence for displaying this System in the sequencer timeline. */
 	TObjectPtr<UNiagaraSequence >NiagaraSequence;
 
