Fixing crash when using a custom expression

diff --git a/Engine/Plugins/FX/Niagara/Source/NiagaraEditor/Private/ViewModels/Stack/NiagaraStackScriptHierarchyRoot.cpp b/Engine/Plugins/FX/Niagara/Source/NiagaraEditor/Private/ViewModels/Stack/NiagaraStackScriptHierarchyRoot.cpp
--- a/Engine/Plugins/FX/Niagara/Source/NiagaraEditor/Private/ViewModels/Stack/NiagaraStackScriptHierarchyRoot.cpp
+++ b/Engine/Plugins/FX/Niagara/Source/NiagaraEditor/Private/ViewModels/Stack/NiagaraStackScriptHierarchyRoot.cpp
@@ -51,7 +51,13 @@ void UNiagaraStackScriptHierarchyRoot::FinalizeInternal()
 
 const TArray<UNiagaraHierarchySection*>& UNiagaraStackScriptHierarchyRoot::GetSections() const
 {
-	return GetScriptParameterHierarchyRoot()->GetSectionData();
+	if(OwningFunctionCallNode->GetCalledUsage() == ENiagaraScriptUsage::Module)
+	{
+		return GetScriptParameterHierarchyRoot()->GetSectionData();
+	}
+	
+	static TArray<UNiagaraHierarchySection*> Dummy;
+	return Dummy;
 }
 
 const UNiagaraHierarchySection* UNiagaraStackScriptHierarchyRoot::GetActiveSection() const
@@ -156,6 +162,11 @@ void UNiagaraStackScriptHierarchyRoot::RefreshInstanceData()
 
 void UNiagaraStackScriptHierarchyRoot::RefreshChildrenInternal(const TArray<UNiagaraStackEntry*>& CurrentChildren, TArray<UNiagaraStackEntry*>& NewChildren, TArray<FStackIssue>& NewIssues)
 {
+	if(OwningFunctionCallNode->FunctionScript == nullptr)
+	{
+		return;
+	}
+	
 	RefreshInstanceData();
 
 	// First we determine the inputs that the hierarchy _does not_ take care of. We add them at the end.
