- Set wireframe color for Niagara proxies

#rnx

diff --git a/Engine/Plugins/FX/Niagara/Source/Niagara/Private/NiagaraComponent.cpp b/Engine/Plugins/FX/Niagara/Source/Niagara/Private/NiagaraComponent.cpp
--- a/Engine/Plugins/FX/Niagara/Source/Niagara/Private/NiagaraComponent.cpp
+++ b/Engine/Plugins/FX/Niagara/Source/Niagara/Private/NiagaraComponent.cpp
@@ -208,6 +208,8 @@ FNiagaraSceneProxy::FNiagaraSceneProxy(UNiagaraComponent* InComponent)
 	: FPrimitiveSceneProxy(InComponent, InComponent->GetAsset() ? InComponent->GetAsset()->GetFName() : FName())
 	, OcclusionQueryMode(InComponent->GetOcclusionQueryMode())
 {
+	SetWireframeColor(FLinearColor(3.0f, 0.0f, 0.0f));
+
 	// Prevent continuous VSM invalidation from the bounds of the primitive.
 	bHasDeformableMesh = false;
 
