- Fix for debug HUD particle limit double counting

#rnx
#test Made a 30 particle emitter, adjusted count to ensure it was correct

diff --git a/Engine/Plugins/FX/Niagara/Source/Niagara/Private/NiagaraDebugHud.cpp b/Engine/Plugins/FX/Niagara/Source/Niagara/Private/NiagaraDebugHud.cpp
--- a/Engine/Plugins/FX/Niagara/Source/Niagara/Private/NiagaraDebugHud.cpp
+++ b/Engine/Plugins/FX/Niagara/Source/Niagara/Private/NiagaraDebugHud.cpp
@@ -2936,7 +2936,7 @@ void FNiagaraDebugHud::DrawComponents(FNiagaraWorldManager* WorldManager, UCanva
 					DrawCanvas->DrawShadowedString(SizeAndLocation.Value.X, SizeAndLocation.Value.Y, FinalString, ParticleFont, TextColor);
 
 					++NumDisplayedParticles;
-					if (++NumDisplayedParticles >= MaxDisplayParticles)
+					if (NumDisplayedParticles >= MaxDisplayParticles)
 					{
 						break;
 					}
