- Stateless additions to handle different spaces for various parameters
- AddVelocity / Acceleration Force can now specify if the force is local / world / simulation rather than assuming local
- Gravity is also assumed to always be in world space

#rnx
#tests Tests with linear force emitter with gravity and rotating around in world space to ensure matches stateful

diff --git a/Engine/Plugins/FX/Niagara/Source/Niagara/Internal/Stateless/Modules/NiagaraStatelessModuleCommon.h b/Engine/Plugins/FX/Niagara/Source/Niagara/Internal/Stateless/Modules/NiagaraStatelessModuleCommon.h
--- a/Engine/Plugins/FX/Niagara/Source/Niagara/Internal/Stateless/Modules/NiagaraStatelessModuleCommon.h
+++ b/Engine/Plugins/FX/Niagara/Source/Niagara/Internal/Stateless/Modules/NiagaraStatelessModuleCommon.h
@@ -13,11 +13,20 @@ namespace NiagaraStateless
 
 		FNiagaraStatelessRangeFloat		MassRange = FNiagaraStatelessRangeFloat(1.0f);
 		FNiagaraStatelessRangeFloat		DragRange = FNiagaraStatelessRangeFloat(0.0f);
+
+		ENiagaraCoordinateSpace			VelocityCoordinateSpace = ENiagaraCoordinateSpace::Local;
 		FNiagaraStatelessRangeVector3	VelocityRange = FNiagaraStatelessRangeVector3(FVector3f::ZeroVector);
+
+		ENiagaraCoordinateSpace			WindCoordinateSpace = ENiagaraCoordinateSpace::Local;
 		FNiagaraStatelessRangeVector3	WindRange = FNiagaraStatelessRangeVector3(FVector3f::ZeroVector);
+
+		ENiagaraCoordinateSpace			AccelerationCoordinateSpace = ENiagaraCoordinateSpace::Local;
 		FNiagaraStatelessRangeVector3	AccelerationRange = FNiagaraStatelessRangeVector3(FVector3f::ZeroVector);
 
+		FNiagaraStatelessRangeVector3	GravityRange = FNiagaraStatelessRangeVector3(FVector3f::ZeroVector);
+
 		bool							bConeVelocity = false;
+		ENiagaraCoordinateSpace			ConeCoordinateSpace = ENiagaraCoordinateSpace::Local;
 		FQuat4f							ConeQuat = FQuat4f::Identity;
 		FNiagaraStatelessRangeFloat		ConeVelocityRange = FNiagaraStatelessRangeFloat(0.0f);
 		float							ConeOuterAngle = 0.0f;
@@ -25,6 +34,7 @@ namespace NiagaraStateless
 		float							ConeVelocityFalloff = 0.0f;
 
 		bool							bPointVelocity = false;
+		ENiagaraCoordinateSpace			PointCoordinateSpace = ENiagaraCoordinateSpace::Local;
 		FNiagaraStatelessRangeFloat		PointVelocityRange = FNiagaraStatelessRangeFloat(0.0f);
 		float							PointVelocityMax = 0.0f;
 		FVector3f						PointOrigin = FVector3f::ZeroVector;
diff --git a/Engine/Plugins/FX/Niagara/Source/Niagara/Internal/Stateless/Modules/NiagaraStatelessModule_AccelerationForce.h b/Engine/Plugins/FX/Niagara/Source/Niagara/Internal/Stateless/Modules/NiagaraStatelessModule_AccelerationForce.h
--- a/Engine/Plugins/FX/Niagara/Source/Niagara/Internal/Stateless/Modules/NiagaraStatelessModule_AccelerationForce.h
+++ b/Engine/Plugins/FX/Niagara/Source/Niagara/Internal/Stateless/Modules/NiagaraStatelessModule_AccelerationForce.h
@@ -18,6 +18,9 @@ public:
 	UPROPERTY(EditAnywhere, Category = "Parameters", meta = (DisplayName = "Acceleration", DisableUniformDistribution, DisableBindingDistribution))
 	FNiagaraDistributionRangeVector3 AccelerationDistribution = FNiagaraDistributionRangeVector3(FVector3f::ZeroVector);
 
+	UPROPERTY(EditAnywhere, Category = "Parameters", meta = (SegmentedDisplay))
+	ENiagaraCoordinateSpace CoordinateSpace = ENiagaraCoordinateSpace::Local;
+
 	virtual void BuildEmitterData(const FNiagaraStatelessEmitterDataBuildContext& BuildContext) const override
 	{
 		if (!IsModuleEnabled())
@@ -28,6 +31,7 @@ public:
 		const FNiagaraStatelessRangeVector3 AccelerationRange = AccelerationDistribution.CalculateRange(FVector3f::ZeroVector);
 
 		NiagaraStateless::FPhysicsBuildData& PhysicsBuildData = BuildContext.GetTransientBuildData<NiagaraStateless::FPhysicsBuildData>();
+		PhysicsBuildData.AccelerationCoordinateSpace = CoordinateSpace;
 		PhysicsBuildData.AccelerationRange.Min += AccelerationRange.Min;
 		PhysicsBuildData.AccelerationRange.Max += AccelerationRange.Max;
 	}
diff --git a/Engine/Plugins/FX/Niagara/Source/Niagara/Internal/Stateless/Modules/NiagaraStatelessModule_AddVelocity.h b/Engine/Plugins/FX/Niagara/Source/Niagara/Internal/Stateless/Modules/NiagaraStatelessModule_AddVelocity.h
--- a/Engine/Plugins/FX/Niagara/Source/Niagara/Internal/Stateless/Modules/NiagaraStatelessModule_AddVelocity.h
+++ b/Engine/Plugins/FX/Niagara/Source/Niagara/Internal/Stateless/Modules/NiagaraStatelessModule_AddVelocity.h
@@ -50,6 +50,9 @@ public:
 	UPROPERTY(EditAnywhere, Category = "Parameters", meta = (EditCondition = "bSpeedFalloffFromConeAxisEnabled && VelocityType == ENSM_VelocityType::InCone", ClampMin = "0.0", ClampMax = "1.0", UIMin = "0.0", UIMax = "1.0"))
 	float SpeedFalloffFromConeAxis = 0.0f;
 
+	UPROPERTY(EditAnywhere, Category = "Parameters", meta = (SegmentedDisplay))
+	ENiagaraCoordinateSpace CoordinateSpace = ENiagaraCoordinateSpace::Local;
+
 	virtual void BuildEmitterData(const FNiagaraStatelessEmitterDataBuildContext& BuildContext) const override;
 
 #if WITH_EDITOR
diff --git a/Engine/Plugins/FX/Niagara/Source/Niagara/Internal/Stateless/Modules/NiagaraStatelessModule_GravityForce.h b/Engine/Plugins/FX/Niagara/Source/Niagara/Internal/Stateless/Modules/NiagaraStatelessModule_GravityForce.h
--- a/Engine/Plugins/FX/Niagara/Source/Niagara/Internal/Stateless/Modules/NiagaraStatelessModule_GravityForce.h
+++ b/Engine/Plugins/FX/Niagara/Source/Niagara/Internal/Stateless/Modules/NiagaraStatelessModule_GravityForce.h
@@ -27,8 +27,8 @@ public:
 		const FNiagaraStatelessRangeVector3 GravityRange = GravityDistribution.CalculateRange(GetDefaultValue());
 
 		NiagaraStateless::FPhysicsBuildData& PhysicsBuildData = BuildContext.GetTransientBuildData<NiagaraStateless::FPhysicsBuildData>();
-		PhysicsBuildData.AccelerationRange.Min += GravityRange.Min;
-		PhysicsBuildData.AccelerationRange.Max += GravityRange.Max;
+		PhysicsBuildData.GravityRange.Min = GravityRange.Min;
+		PhysicsBuildData.GravityRange.Max = GravityRange.Max;
 	}
 
 	static FVector3f GetDefaultValue() { return FVector3f(0.0f, 0.0f, -980.0f); }
diff --git a/Engine/Plugins/FX/Niagara/Source/Niagara/Internal/Stateless/Modules/NiagaraStatelessModule_SolveVelocitiesAndForces.h b/Engine/Plugins/FX/Niagara/Source/Niagara/Internal/Stateless/Modules/NiagaraStatelessModule_SolveVelocitiesAndForces.h
--- a/Engine/Plugins/FX/Niagara/Source/Niagara/Internal/Stateless/Modules/NiagaraStatelessModule_SolveVelocitiesAndForces.h
+++ b/Engine/Plugins/FX/Niagara/Source/Niagara/Internal/Stateless/Modules/NiagaraStatelessModule_SolveVelocitiesAndForces.h
@@ -61,42 +61,47 @@ public:
 	virtual void SetShaderParameters(const FNiagaraStatelessSetShaderParameterContext& SetShaderParameterContext) const override
 	{
 		const FModuleBuiltData* ModuleBuiltData = SetShaderParameterContext.ReadBuiltData<FModuleBuiltData>();
+		const NiagaraStateless::FPhysicsBuildData& PhysicsData = ModuleBuiltData->PhysicsData;
+
+		const FNiagaraStatelessSpaceTransforms& SpaceTransforms = SetShaderParameterContext.GetSpaceTransforms();
 
 		FParameters* Parameters = SetShaderParameterContext.GetParameterNestedStruct<FParameters>();
-		Parameters->SolveVelocitiesAndForces_MassScale				= ModuleBuiltData->PhysicsData.MassRange.GetScale();
-		Parameters->SolveVelocitiesAndForces_MassBias				= ModuleBuiltData->PhysicsData.MassRange.Min;
-		Parameters->SolveVelocitiesAndForces_DragScale				= ModuleBuiltData->PhysicsData.DragRange.GetScale();
-		Parameters->SolveVelocitiesAndForces_DragBias				= ModuleBuiltData->PhysicsData.DragRange.Min;
-		Parameters->SolveVelocitiesAndForces_VelocityScale			= ModuleBuiltData->PhysicsData.VelocityRange.GetScale();
-		Parameters->SolveVelocitiesAndForces_VelocityBias			= ModuleBuiltData->PhysicsData.VelocityRange.Min;
-		Parameters->SolveVelocitiesAndForces_WindScale				= ModuleBuiltData->PhysicsData.WindRange.GetScale();
-		Parameters->SolveVelocitiesAndForces_WindBias				= ModuleBuiltData->PhysicsData.WindRange.Min;
-		Parameters->SolveVelocitiesAndForces_AccelerationScale		= ModuleBuiltData->PhysicsData.AccelerationRange.GetScale();
-		Parameters->SolveVelocitiesAndForces_AccelerationBias		= ModuleBuiltData->PhysicsData.AccelerationRange.Min;
-
-		Parameters->SolveVelocitiesAndForces_ConeVelocityEnabled	= ModuleBuiltData->PhysicsData.bConeVelocity ? 1 : 0;
-		Parameters->SolveVelocitiesAndForces_ConeQuat				= ModuleBuiltData->PhysicsData.ConeQuat;
-		Parameters->SolveVelocitiesAndForces_ConeVelocityScale		= ModuleBuiltData->PhysicsData.ConeVelocityRange.GetScale();
-		Parameters->SolveVelocitiesAndForces_ConeVelocityBias		= ModuleBuiltData->PhysicsData.ConeVelocityRange.Min;
-		Parameters->SolveVelocitiesAndForces_ConeAngleScale			= (ModuleBuiltData->PhysicsData.ConeOuterAngle - ModuleBuiltData->PhysicsData.ConeInnerAngle) * (UE_PI / 360.0f);
-		Parameters->SolveVelocitiesAndForces_ConeAngleBias			= ModuleBuiltData->PhysicsData.ConeInnerAngle * (UE_PI / 360.0f);
-		Parameters->SolveVelocitiesAndForces_ConeVelocityFalloff	= ModuleBuiltData->PhysicsData.ConeVelocityFalloff;
-
-		Parameters->SolveVelocitiesAndForces_PontVelocityEnabled	= ModuleBuiltData->PhysicsData.bPointVelocity ? 1 : 0;
-		Parameters->SolveVelocitiesAndForces_PointVelocityScale		= ModuleBuiltData->PhysicsData.PointVelocityRange.GetScale();
-		Parameters->SolveVelocitiesAndForces_PointVelocityBias		= ModuleBuiltData->PhysicsData.PointVelocityRange.Min;
-		Parameters->SolveVelocitiesAndForces_PointOrigin			= ModuleBuiltData->PhysicsData.PointOrigin;
-
-		Parameters->SolveVelocitiesAndForces_NoiseEnabled			= ModuleBuiltData->PhysicsData.bNoiseEnabled ? 1 : 0;
-		Parameters->SolveVelocitiesAndForces_NoiseAmplitude			= ModuleBuiltData->PhysicsData.NoiseAmplitude;
-		Parameters->SolveVelocitiesAndForces_NoiseFrequency			= FVector3f(ModuleBuiltData->PhysicsData.NoiseFrequency, ModuleBuiltData->PhysicsData.NoiseFrequency, ModuleBuiltData->PhysicsData.NoiseFrequency);
+		Parameters->SolveVelocitiesAndForces_MassScale				= PhysicsData.MassRange.GetScale();
+		Parameters->SolveVelocitiesAndForces_MassBias				= PhysicsData.MassRange.Min;
+		Parameters->SolveVelocitiesAndForces_DragScale				= PhysicsData.DragRange.GetScale();
+		Parameters->SolveVelocitiesAndForces_DragBias				= PhysicsData.DragRange.Min;
+		Parameters->SolveVelocitiesAndForces_VelocityScale			= SpaceTransforms.TransformVector(PhysicsData.VelocityCoordinateSpace, PhysicsData.VelocityRange.GetScale());
+		Parameters->SolveVelocitiesAndForces_VelocityBias			= SpaceTransforms.TransformVector(PhysicsData.VelocityCoordinateSpace, PhysicsData.VelocityRange.Min);
+		Parameters->SolveVelocitiesAndForces_WindScale				= SpaceTransforms.TransformVector(PhysicsData.WindCoordinateSpace, PhysicsData.WindRange.GetScale());
+		Parameters->SolveVelocitiesAndForces_WindBias				= SpaceTransforms.TransformVector(PhysicsData.WindCoordinateSpace, PhysicsData.WindRange.Min);
+		Parameters->SolveVelocitiesAndForces_AccelerationScale		= SpaceTransforms.TransformVector(PhysicsData.AccelerationCoordinateSpace, PhysicsData.AccelerationRange.GetScale());
+		Parameters->SolveVelocitiesAndForces_AccelerationBias		= SpaceTransforms.TransformVector(PhysicsData.AccelerationCoordinateSpace, PhysicsData.AccelerationRange.Min);
+		Parameters->SolveVelocitiesAndForces_AccelerationScale		+= SpaceTransforms.TransformVector(ENiagaraCoordinateSpace::World, PhysicsData.GravityRange.GetScale());
+		Parameters->SolveVelocitiesAndForces_AccelerationBias		+= SpaceTransforms.TransformVector(ENiagaraCoordinateSpace::World, PhysicsData.GravityRange.Min);
+
+		Parameters->SolveVelocitiesAndForces_ConeVelocityEnabled	= PhysicsData.bConeVelocity ? 1 : 0;
+		Parameters->SolveVelocitiesAndForces_ConeQuat				= SpaceTransforms.TransformRotation(PhysicsData.ConeCoordinateSpace, PhysicsData.ConeQuat);
+		Parameters->SolveVelocitiesAndForces_ConeVelocityScale		= PhysicsData.ConeVelocityRange.GetScale();
+		Parameters->SolveVelocitiesAndForces_ConeVelocityBias		= PhysicsData.ConeVelocityRange.Min;
+		Parameters->SolveVelocitiesAndForces_ConeAngleScale			= (PhysicsData.ConeOuterAngle - PhysicsData.ConeInnerAngle) * (UE_PI / 360.0f);
+		Parameters->SolveVelocitiesAndForces_ConeAngleBias			= PhysicsData.ConeInnerAngle * (UE_PI / 360.0f);
+		Parameters->SolveVelocitiesAndForces_ConeVelocityFalloff	= PhysicsData.ConeVelocityFalloff;
+
+		Parameters->SolveVelocitiesAndForces_PontVelocityEnabled	= PhysicsData.bPointVelocity ? 1 : 0;
+		Parameters->SolveVelocitiesAndForces_PointVelocityScale		= PhysicsData.PointVelocityRange.GetScale();
+		Parameters->SolveVelocitiesAndForces_PointVelocityBias		= PhysicsData.PointVelocityRange.Min;
+		Parameters->SolveVelocitiesAndForces_PointOrigin			= SpaceTransforms.TransformPosition(PhysicsData.PointCoordinateSpace, PhysicsData.PointOrigin);
+
+		Parameters->SolveVelocitiesAndForces_NoiseEnabled			= PhysicsData.bNoiseEnabled ? 1 : 0;
+		Parameters->SolveVelocitiesAndForces_NoiseAmplitude			= PhysicsData.NoiseAmplitude;
+		Parameters->SolveVelocitiesAndForces_NoiseFrequency			= FVector3f(PhysicsData.NoiseFrequency, PhysicsData.NoiseFrequency, PhysicsData.NoiseFrequency);
 		//SetShaderParameterContext.SetTextureResource(&Parameters->SolveVelocitiesAndForces_NoiseTexture, ModuleBuildData->NoiseTexture);
-		Parameters->SolveVelocitiesAndForces_NoiseMode				= ModuleBuiltData->PhysicsData.NoiseMode;
-		Parameters->SolveVelocitiesAndForces_NoiseLUTOffset			= ModuleBuiltData->PhysicsData.NoiseLUTOffset;
-		Parameters->SolveVelocitiesAndForces_NoiseLUTNumChannel		= ModuleBuiltData->PhysicsData.NoiseLUTNumChannel;
-		Parameters->SolveVelocitiesAndForces_NoiseLUTChannelWidth	= ModuleBuiltData->PhysicsData.NoiseLUTChannelWidth;
+		Parameters->SolveVelocitiesAndForces_NoiseMode				= PhysicsData.NoiseMode;
+		Parameters->SolveVelocitiesAndForces_NoiseLUTOffset			= PhysicsData.NoiseLUTOffset;
+		Parameters->SolveVelocitiesAndForces_NoiseLUTNumChannel		= PhysicsData.NoiseLUTNumChannel;
+		Parameters->SolveVelocitiesAndForces_NoiseLUTChannelWidth	= PhysicsData.NoiseLUTChannelWidth;
 
-		FVectorFieldTextureAccessor TextureAccessor(Cast<UVectorField>(ModuleBuiltData->PhysicsData.NoiseTexture));
+		FVectorFieldTextureAccessor TextureAccessor(Cast<UVectorField>(PhysicsData.NoiseTexture));
 
 		ENQUEUE_RENDER_COMMAND(FNaughtyTest)(
 			[Parameters, TextureAccessor](FRHICommandListImmediate& RHICmdList)
diff --git a/Engine/Plugins/FX/Niagara/Source/Niagara/Internal/Stateless/NiagaraStatelessCommon.h b/Engine/Plugins/FX/Niagara/Source/Niagara/Internal/Stateless/NiagaraStatelessCommon.h
--- a/Engine/Plugins/FX/Niagara/Source/Niagara/Internal/Stateless/NiagaraStatelessCommon.h
+++ b/Engine/Plugins/FX/Niagara/Source/Niagara/Internal/Stateless/NiagaraStatelessCommon.h
@@ -20,6 +20,37 @@ enum class ENiagaraStatelessFeatureMask : uint32
 };
 ENUM_CLASS_FLAGS(ENiagaraStatelessFeatureMask);
 
+// Helper structure to transform in / out of various spaces
+// The transforms are all expected to be in tile relative space, i.e. not LWC space
+struct FNiagaraStatelessSpaceTransforms
+{
+	FNiagaraStatelessSpaceTransforms() { InitializeTransforms(true, FTransform3f::Identity); }
+
+	void InitializeTransforms(bool bIsLocalSpace, const FTransform3f& LocalToWorld);
+	bool UpdateTransforms(const FTransform3f& LocalToWorld);
+	
+	const FTransform3f& GetLocalToWorld() const { return Transforms[LocalToWorldIndex]; }
+	const FTransform3f& GetWorldToLocal() const { return Transforms[WorldToLocalIndex]; }
+
+	FQuat4f   TransformRotation(ENiagaraCoordinateSpace SourceSpace, ENiagaraCoordinateSpace DestinationSpace, FQuat4f Rotation) const;
+	FVector3f TransformPosition(ENiagaraCoordinateSpace SourceSpace, ENiagaraCoordinateSpace DestinationSpace, FVector3f Position) const;
+	FVector3f TransformVector(ENiagaraCoordinateSpace SourceSpace, ENiagaraCoordinateSpace DestinationSpace, FVector3f Vector) const;
+	FVector3f TransformVectorNoScale(ENiagaraCoordinateSpace SourceSpace, ENiagaraCoordinateSpace DestinationSpace, FVector3f Vector) const;
+
+	FQuat4f   TransformRotation(ENiagaraCoordinateSpace SourceSpace, FQuat4f Rotation) const { return TransformRotation(SourceSpace, ENiagaraCoordinateSpace::Simulation, Rotation); }
+	FVector3f TransformPosition(ENiagaraCoordinateSpace SourceSpace, FVector3f Position) const { return TransformPosition(SourceSpace, ENiagaraCoordinateSpace::Simulation, Position); }
+	FVector3f TransformVector(ENiagaraCoordinateSpace SourceSpace, FVector3f Vector) const { return TransformVector(SourceSpace, ENiagaraCoordinateSpace::Simulation, Vector); }
+	FVector3f TransformVectorNoScale(ENiagaraCoordinateSpace SourceSpace, FVector3f Vector) const { return TransformVectorNoScale(SourceSpace, ENiagaraCoordinateSpace::Simulation, Vector); }
+
+private:
+	static constexpr int LocalToWorldIndex = 0;
+	static constexpr int WorldToLocalIndex = 1;
+	static constexpr int TransformRemapSize = 3;	// Local / World / Simulation
+
+	int8			TransformRemap[TransformRemapSize][TransformRemapSize];	// Remaps from transforms from world / local / simulation into an index below, where INDEX_NONE is null op
+	FTransform3f	Transforms[2];			// World to Local & Local to World
+};
+
 template<typename TType>
 struct FNiagaraStatelessRange
 {
diff --git a/Engine/Plugins/FX/Niagara/Source/Niagara/Internal/Stateless/NiagaraStatelessEmitter.h b/Engine/Plugins/FX/Niagara/Source/Niagara/Internal/Stateless/NiagaraStatelessEmitter.h
--- a/Engine/Plugins/FX/Niagara/Source/Niagara/Internal/Stateless/NiagaraStatelessEmitter.h
+++ b/Engine/Plugins/FX/Niagara/Source/Niagara/Internal/Stateless/NiagaraStatelessEmitter.h
@@ -88,7 +88,7 @@ public:
 	NIAGARA_API bool SetUniqueEmitterName(const FString& InName);
 
 	FNiagaraStatelessEmitterDataPtr GetEmitterData() const { return StatelessEmitterData; }
-	NiagaraStateless::FCommonShaderParameters* AllocateShaderParameters(const FNiagaraParameterStore& RendererBindings) const;
+	NiagaraStateless::FCommonShaderParameters* AllocateShaderParameters(const FNiagaraStatelessSpaceTransforms& SpaceTransforms, const FNiagaraParameterStore& RendererBindings) const;
 
 	NIAGARA_API bool IsAllowedByScalability() const;
 
diff --git a/Engine/Plugins/FX/Niagara/Source/Niagara/Internal/Stateless/NiagaraStatelessEmitterInstance.h b/Engine/Plugins/FX/Niagara/Source/Niagara/Internal/Stateless/NiagaraStatelessEmitterInstance.h
--- a/Engine/Plugins/FX/Niagara/Source/Niagara/Internal/Stateless/NiagaraStatelessEmitterInstance.h
+++ b/Engine/Plugins/FX/Niagara/Source/Niagara/Internal/Stateless/NiagaraStatelessEmitterInstance.h
@@ -117,6 +117,8 @@ private:
 	FNiagaraStatelessEmitterDataPtr				EmitterData;
 	TWeakObjectPtr<UNiagaraStatelessEmitter>	WeakStatelessEmitter;
 
+	FNiagaraStatelessSpaceTransforms			EmitterTransforms;
+
 	float										Age = 0.0f;
 
 	uint32										UniqueIndexOffset = 0;
diff --git a/Engine/Plugins/FX/Niagara/Source/Niagara/Internal/Stateless/NiagaraStatelessModule.h b/Engine/Plugins/FX/Niagara/Source/Niagara/Internal/Stateless/NiagaraStatelessModule.h
--- a/Engine/Plugins/FX/Niagara/Source/Niagara/Internal/Stateless/NiagaraStatelessModule.h
+++ b/Engine/Plugins/FX/Niagara/Source/Niagara/Internal/Stateless/NiagaraStatelessModule.h
@@ -21,8 +21,9 @@ class FNiagaraStatelessSetShaderParameterContext
 public:
 	UE_NONCOPYABLE(FNiagaraStatelessSetShaderParameterContext);
 
-	explicit FNiagaraStatelessSetShaderParameterContext(TConstArrayView<uint8> InRendererParameterData, TConstArrayView<uint8> InBuiltData, const FShaderParametersMetadata* InShaderParametersMetadata, uint8* InShaderParameters)
-		: RendererParameterData(InRendererParameterData)
+	explicit FNiagaraStatelessSetShaderParameterContext(const FNiagaraStatelessSpaceTransforms& InSpaceTransforms, TConstArrayView<uint8> InRendererParameterData, TConstArrayView<uint8> InBuiltData, const FShaderParametersMetadata* InShaderParametersMetadata, uint8* InShaderParameters)
+		: SpaceTransforms(InSpaceTransforms)
+		, RendererParameterData(InRendererParameterData)
 		, BuiltData(InBuiltData)
 		, ShaderParametersBase(InShaderParameters)
 		, ParameterOffset(0)
@@ -30,6 +31,8 @@ public:
 	{
 	}
 
+	const FNiagaraStatelessSpaceTransforms& GetSpaceTransforms() const { return SpaceTransforms; }
+
 	template<typename T>
 	T* GetParameterNestedStruct() const
 	{
@@ -76,12 +79,13 @@ protected:
 #endif
 
 private:
-	TConstArrayView<uint8>				RendererParameterData;
-	TConstArrayView<uint8>				BuiltData;
-	mutable int32						BuiltDataOffset = 0;
-	uint8*								ShaderParametersBase = nullptr;
-	mutable uint32						ParameterOffset = 0;
-	const FShaderParametersMetadata*	ShaderParametersMetadata = nullptr;
+	const FNiagaraStatelessSpaceTransforms&	SpaceTransforms;
+	TConstArrayView<uint8>					RendererParameterData;
+	TConstArrayView<uint8>					BuiltData;
+	mutable int32							BuiltDataOffset = 0;
+	uint8*									ShaderParametersBase = nullptr;
+	mutable uint32							ParameterOffset = 0;
+	const FShaderParametersMetadata*		ShaderParametersMetadata = nullptr;
 };
 
 ////////////////////////////////////////////////////////////////////////////////////////////////////////////////////////////////////////////////////////////////
diff --git a/Engine/Plugins/FX/Niagara/Source/Niagara/Private/Stateless/Modules/NiagaraStatelessModule_AddVelocity.cpp b/Engine/Plugins/FX/Niagara/Source/Niagara/Private/Stateless/Modules/NiagaraStatelessModule_AddVelocity.cpp
--- a/Engine/Plugins/FX/Niagara/Source/Niagara/Private/Stateless/Modules/NiagaraStatelessModule_AddVelocity.cpp
+++ b/Engine/Plugins/FX/Niagara/Source/Niagara/Private/Stateless/Modules/NiagaraStatelessModule_AddVelocity.cpp
@@ -14,6 +14,7 @@ void UNiagaraStatelessModule_AddVelocity::BuildEmitterData(const FNiagaraStatele
 	if (VelocityType == ENSM_VelocityType::Linear)
 	{
 		const FNiagaraStatelessRangeVector3 VelocityRange = LinearVelocityDistribution.CalculateRange(FVector3f::ZeroVector);
+		PhysicsBuildData.VelocityCoordinateSpace = CoordinateSpace;
 		PhysicsBuildData.VelocityRange.Min += VelocityRange.Min * LinearVelocityScale;
 		PhysicsBuildData.VelocityRange.Max += VelocityRange.Max * LinearVelocityScale;
 	}
@@ -22,6 +23,7 @@ void UNiagaraStatelessModule_AddVelocity::BuildEmitterData(const FNiagaraStatele
 		ensureMsgf(PhysicsBuildData.bPointVelocity == false, TEXT("Only a single point force is supported at the moment."));
 
 		PhysicsBuildData.bPointVelocity = true;
+		PhysicsBuildData.PointCoordinateSpace = CoordinateSpace;
 		PhysicsBuildData.PointVelocityRange = PointVelocityDistribution.CalculateRange(0.0f);
 		PhysicsBuildData.PointOrigin = PointOrigin;
 	}
@@ -30,6 +32,7 @@ void UNiagaraStatelessModule_AddVelocity::BuildEmitterData(const FNiagaraStatele
 		ensureMsgf(PhysicsBuildData.bConeVelocity == false, TEXT("Only a single cone force is supported at the moment."));
 
 		PhysicsBuildData.bConeVelocity = true;
+		PhysicsBuildData.ConeCoordinateSpace = CoordinateSpace;
 		PhysicsBuildData.ConeQuat = FQuat4f(ConeRotation.Quaternion());
 		PhysicsBuildData.ConeVelocityRange = ConeVelocityDistribution.CalculateRange(0.0f);
 		PhysicsBuildData.ConeOuterAngle = ConeAngle;
diff --git a/Engine/Plugins/FX/Niagara/Source/Niagara/Private/Stateless/NiagaraStatelessCommon.cpp b/Engine/Plugins/FX/Niagara/Source/Niagara/Private/Stateless/NiagaraStatelessCommon.cpp
--- a/Engine/Plugins/FX/Niagara/Source/Niagara/Private/Stateless/NiagaraStatelessCommon.cpp
+++ b/Engine/Plugins/FX/Niagara/Source/Niagara/Private/Stateless/NiagaraStatelessCommon.cpp
@@ -7,6 +7,67 @@
 
 #include "UObject/UObjectIterator.h"
 
+static_assert(int(ENiagaraCoordinateSpace::Local) >= 0 && int(ENiagaraCoordinateSpace::Local) < 3, "Local space transform out of expected range");
+static_assert(int(ENiagaraCoordinateSpace::World) >= 0 && int(ENiagaraCoordinateSpace::World) < 3, "World space transform out of expected range");
+static_assert(int(ENiagaraCoordinateSpace::Simulation) >= 0 && int(ENiagaraCoordinateSpace::Simulation) < 3, "Simulation space transform out of expected range");
+
+void FNiagaraStatelessSpaceTransforms::InitializeTransforms(bool bIsLocalSpace, const FTransform3f& LocalToWorld)
+{
+	TransformRemap[int(ENiagaraCoordinateSpace::Local)][int(ENiagaraCoordinateSpace::Local)]			= INDEX_NONE;
+	TransformRemap[int(ENiagaraCoordinateSpace::Local)][int(ENiagaraCoordinateSpace::World)]			= LocalToWorldIndex;
+	TransformRemap[int(ENiagaraCoordinateSpace::Local)][int(ENiagaraCoordinateSpace::Simulation)]		= bIsLocalSpace ? INDEX_NONE : LocalToWorldIndex;
+
+	TransformRemap[int(ENiagaraCoordinateSpace::World)][int(ENiagaraCoordinateSpace::Local)]			= WorldToLocalIndex;
+	TransformRemap[int(ENiagaraCoordinateSpace::World)][int(ENiagaraCoordinateSpace::World)]			= INDEX_NONE;
+	TransformRemap[int(ENiagaraCoordinateSpace::World)][int(ENiagaraCoordinateSpace::Simulation)]		= bIsLocalSpace ? WorldToLocalIndex : INDEX_NONE;
+
+	TransformRemap[int(ENiagaraCoordinateSpace::Simulation)][int(ENiagaraCoordinateSpace::Local)]		= bIsLocalSpace ? INDEX_NONE : WorldToLocalIndex;
+	TransformRemap[int(ENiagaraCoordinateSpace::Simulation)][int(ENiagaraCoordinateSpace::World)]		= bIsLocalSpace ? LocalToWorldIndex : INDEX_NONE;
+	TransformRemap[int(ENiagaraCoordinateSpace::Simulation)][int(ENiagaraCoordinateSpace::Simulation)]	= INDEX_NONE;
+
+	Transforms[LocalToWorldIndex] = LocalToWorld;
+	Transforms[WorldToLocalIndex] = LocalToWorld.Inverse();
+}
+
+bool FNiagaraStatelessSpaceTransforms::UpdateTransforms(const FTransform3f& LocalToWorld)
+{
+	if (LocalToWorld.Equals(Transforms[0]))
+	{
+		return false;
+	}
+	Transforms[LocalToWorldIndex] = LocalToWorld;
+	Transforms[WorldToLocalIndex] = LocalToWorld.Inverse();
+	return true;
+}
+
+FQuat4f FNiagaraStatelessSpaceTransforms::TransformRotation(ENiagaraCoordinateSpace SourceSpace, ENiagaraCoordinateSpace DestinationSpace, FQuat4f Rotation) const
+{
+	check(int(SourceSpace) >= 0 && int(SourceSpace) < TransformRemapSize && int(DestinationSpace) >= 0 && int(DestinationSpace) < TransformRemapSize);
+	const int8 i = TransformRemap[int(SourceSpace)][int(DestinationSpace)];
+	return i == INDEX_NONE ? Rotation : Transforms[i].TransformRotation(Rotation);
+}
+
+FVector3f FNiagaraStatelessSpaceTransforms::TransformPosition(ENiagaraCoordinateSpace SourceSpace, ENiagaraCoordinateSpace DestinationSpace, FVector3f Position) const
+{
+	check(int(SourceSpace) >= 0 && int(SourceSpace) < TransformRemapSize && int(DestinationSpace) >= 0 && int(DestinationSpace) < TransformRemapSize);
+	const int8 i = TransformRemap[int(SourceSpace)][int(DestinationSpace)];
+	return i == INDEX_NONE ? Position : Transforms[i].TransformPosition(Position);
+}
+
+FVector3f FNiagaraStatelessSpaceTransforms::TransformVector(ENiagaraCoordinateSpace SourceSpace, ENiagaraCoordinateSpace DestinationSpace, FVector3f Vector) const
+{
+	check(int(SourceSpace) >= 0 && int(SourceSpace) < TransformRemapSize && int(DestinationSpace) >= 0 && int(DestinationSpace) < TransformRemapSize);
+	const int8 i = TransformRemap[int(SourceSpace)][int(DestinationSpace)];
+	return i == INDEX_NONE ? Vector : Transforms[i].TransformVector(Vector);
+}
+
+FVector3f FNiagaraStatelessSpaceTransforms::TransformVectorNoScale(ENiagaraCoordinateSpace SourceSpace, ENiagaraCoordinateSpace DestinationSpace, FVector3f Vector) const
+{
+	check(int(SourceSpace) >= 0 && int(SourceSpace) < TransformRemapSize && int(DestinationSpace) >= 0 && int(DestinationSpace) < TransformRemapSize);
+	const int8 i = TransformRemap[int(SourceSpace)][int(DestinationSpace)];
+	return i == INDEX_NONE ? Vector : Transforms[i].TransformVectorNoScale(Vector);
+}
+
 namespace NiagaraStatelessCommon
 {
 	static FNiagaraStatelessGlobals					GGlobals;
diff --git a/Engine/Plugins/FX/Niagara/Source/Niagara/Private/Stateless/NiagaraStatelessEmitter.cpp b/Engine/Plugins/FX/Niagara/Source/Niagara/Private/Stateless/NiagaraStatelessEmitter.cpp
--- a/Engine/Plugins/FX/Niagara/Source/Niagara/Private/Stateless/NiagaraStatelessEmitter.cpp
+++ b/Engine/Plugins/FX/Niagara/Source/Niagara/Private/Stateless/NiagaraStatelessEmitter.cpp
@@ -569,7 +569,7 @@ bool UNiagaraStatelessEmitter::SetUniqueEmitterName(const FString& InName)
 	return false;
 }
 
-NiagaraStateless::FCommonShaderParameters* UNiagaraStatelessEmitter::AllocateShaderParameters(const FNiagaraParameterStore& RendererBindings) const
+NiagaraStateless::FCommonShaderParameters* UNiagaraStatelessEmitter::AllocateShaderParameters(const FNiagaraStatelessSpaceTransforms& SpaceTransforms, const FNiagaraParameterStore& RendererBindings) const
 {
 	// Allocate parameters
 	const FShaderParametersMetadata* ShaderParametersMetadata = StatelessEmitterData->GetShaderParametersMetadata();
@@ -579,6 +579,7 @@ NiagaraStateless::FCommonShaderParameters* UNiagaraStatelessEmitter::AllocateSha
 
 	// Fill in all of the shader parameters
 	FNiagaraStatelessSetShaderParameterContext SetShaderParametersContext(
+		SpaceTransforms,
 		RendererBindings.GetParameterDataArray(),
 		StatelessEmitterData->BuiltData,
 		ShaderParametersMetadata,
diff --git a/Engine/Plugins/FX/Niagara/Source/Niagara/Private/Stateless/NiagaraStatelessEmitterInstance.cpp b/Engine/Plugins/FX/Niagara/Source/Niagara/Private/Stateless/NiagaraStatelessEmitterInstance.cpp
--- a/Engine/Plugins/FX/Niagara/Source/Niagara/Private/Stateless/NiagaraStatelessEmitterInstance.cpp
+++ b/Engine/Plugins/FX/Niagara/Source/Niagara/Private/Stateless/NiagaraStatelessEmitterInstance.cpp
@@ -107,6 +107,8 @@ void FNiagaraStatelessEmitterInstance::Init(int32 InEmitterIndex)
 	// Prepare our parameters
 	RendererBindings = EmitterData->RendererBindings;
 
+	EmitterTransforms.InitializeTransforms(IsLocalSpace(), FTransform3f(ParentSystemInstance->GetWorldTransform()));
+
 	// Allocate and fill shader parameters
 	RenderThreadDataPtr.Reset(new NiagaraStateless::FEmitterInstance_RT());
 	RenderThreadDataPtr->EmitterData = EmitterData;
@@ -114,7 +116,7 @@ void FNiagaraStatelessEmitterInstance::Init(int32 InEmitterIndex)
 	RenderThreadDataPtr->Age = 0.0f;
 	RenderThreadDataPtr->DeltaTime = 0.0f;
 	RenderThreadDataPtr->ExecutionState = ENiagaraExecutionState::Active;
-	RenderThreadDataPtr->ShaderParameters.Reset(WeakStatelessEmitter->AllocateShaderParameters(RendererBindings));
+	RenderThreadDataPtr->ShaderParameters.Reset(WeakStatelessEmitter->AllocateShaderParameters(EmitterTransforms, RendererBindings));
 	RenderThreadDataPtr->ShaderParameters->Common_RandomSeed = RandomSeed;
 
 	ENQUEUE_RENDER_COMMAND(FInitStatelessEmitter)(
@@ -472,20 +474,29 @@ void FNiagaraStatelessEmitterInstance::SendRenderData()
 	DataForRenderThread.Age				= Age;
 	DataForRenderThread.ExecutionState	= InternalExecutionState;
 
+	bool bNeedsShaderParametersUpdate = false;
 	if (RendererBindings.GetParametersDirty())
 	{
 		RendererBindings.Tick();
+
 		if (EmitterData->bModulesHaveRendererBindings)
 		{
+			bNeedsShaderParametersUpdate = true;
 			DataForRenderThread.bHasBindingBufferData = true;
 			DataForRenderThread.BindingBufferData = RendererBindings.GetParameterDataArray();
 			check((DataForRenderThread.BindingBufferData.Num() % sizeof(uint32)) == 0);
-
-			DataForRenderThread.ShaderParameters = WeakStatelessEmitter->AllocateShaderParameters(RendererBindings);
-			DataForRenderThread.ShaderParameters->Common_RandomSeed = RandomSeed;
 		}
 	}
 
+	const FTransform3f ParentTransform = FTransform3f(ParentSystemInstance->GetWorldTransform());
+	bNeedsShaderParametersUpdate = EmitterTransforms.UpdateTransforms(ParentTransform);
+
+	if (bNeedsShaderParametersUpdate)
+	{
+		DataForRenderThread.ShaderParameters = WeakStatelessEmitter->AllocateShaderParameters(EmitterTransforms, RendererBindings);
+		DataForRenderThread.ShaderParameters->Common_RandomSeed = RandomSeed;
+	}
+
 	if (bSpawnInfosDirty)
 	{
 		DataForRenderThread.bHasSpawnInfoData = true;
