Added additional safety check for Hierarchy Editor for Scripts that stops Undo/Redo from crashing

diff --git a/Engine/Plugins/FX/Niagara/Source/NiagaraEditor/Private/ViewModels/HierarchyEditor/NiagaraHierarchyScriptParametersViewModel.cpp b/Engine/Plugins/FX/Niagara/Source/NiagaraEditor/Private/ViewModels/HierarchyEditor/NiagaraHierarchyScriptParametersViewModel.cpp
--- a/Engine/Plugins/FX/Niagara/Source/NiagaraEditor/Private/ViewModels/HierarchyEditor/NiagaraHierarchyScriptParametersViewModel.cpp
+++ b/Engine/Plugins/FX/Niagara/Source/NiagaraEditor/Private/ViewModels/HierarchyEditor/NiagaraHierarchyScriptParametersViewModel.cpp
@@ -164,6 +164,12 @@ TSharedRef<SWidget> FNiagaraHierarchyScriptParameterDragDropOp::CreateCustomDeco
 
 bool FNiagaraHierarchyScriptParameterViewModel::DoesExternalDataStillExist(const UNiagaraHierarchyDataRefreshContext* Context) const
 {
+	// During undo/redo it's possible the script variable becomes nullptr. If so, there is no need for this view model either
+	if(Cast<UNiagaraHierarchyScriptParameter>(GetDataMutable())->GetScriptVariable() == nullptr)
+	{
+		return false;
+	}
+	
 	const UNiagaraHierarchyScriptParameterRefreshContext* RefreshContext = CastChecked<UNiagaraHierarchyScriptParameterRefreshContext>(Context);
 	if(RefreshContext->GetNiagaraGraph()->GetAllMetaData().Contains(Cast<UNiagaraHierarchyScriptParameter>(GetDataMutable())->GetVariable()) == false)
 	{
