Fixed a bug where undo/redo didn't work correctly for some emitter properties

diff --git a/Engine/Plugins/FX/Niagara/Source/NiagaraEditor/Private/NiagaraEmitterDetailsCustomization.cpp b/Engine/Plugins/FX/Niagara/Source/NiagaraEditor/Private/NiagaraEmitterDetailsCustomization.cpp
--- a/Engine/Plugins/FX/Niagara/Source/NiagaraEditor/Private/NiagaraEmitterDetailsCustomization.cpp
+++ b/Engine/Plugins/FX/Niagara/Source/NiagaraEditor/Private/NiagaraEmitterDetailsCustomization.cpp
@@ -53,9 +53,15 @@ void CustomizeEmitterData(IDetailLayoutBuilder& InDetailLayout, UNiagaraEmitter*
 				FPropertyChangedEvent ChangeEvent(Property);
 				VersionedNiagaraEmitter.Emitter->PostEditChangeVersionedProperty(ChangeEvent, VersionedNiagaraEmitter.Version);
 			};
+			const auto& PreEditChangeLambda = [VersionedNiagaraEmitter, Property]
+			{
+				VersionedNiagaraEmitter.Emitter->PreEditChange(Property);
+			};
 
 			PropertyRow->GetPropertyHandle()->SetOnChildPropertyValueChanged(FSimpleDelegate::CreateLambda(PostEditChangeLambda));
 			PropertyRow->GetPropertyHandle()->SetOnPropertyValueChanged(FSimpleDelegate::CreateLambda(PostEditChangeLambda));
+			PropertyRow->GetPropertyHandle()->SetOnPropertyValuePreChange(FSimpleDelegate::CreateLambda(PreEditChangeLambda));
+			PropertyRow->GetPropertyHandle()->SetOnChildPropertyValuePreChange(FSimpleDelegate::CreateLambda(PreEditChangeLambda));
 		}
 	}
 
