Fixed compile error when using "draw torus" function in gpu emitters. It's still not doing anything, but at least it compiles without error.

diff --git a/Engine/Plugins/FX/Niagara/Shaders/Private/NiagaraDataInterfaceDebugDrawTemplate.ush b/Engine/Plugins/FX/Niagara/Shaders/Private/NiagaraDataInterfaceDebugDrawTemplate.ush
--- a/Engine/Plugins/FX/Niagara/Shaders/Private/NiagaraDataInterfaceDebugDrawTemplate.ush
+++ b/Engine/Plugins/FX/Niagara/Shaders/Private/NiagaraDataInterfaceDebugDrawTemplate.ush
@@ -35,7 +35,7 @@ void DrawCone_{ParameterName}_UEImpureCall(bool bExecute, float3 Location, float
 {
 }
 
-void DrawTorus_{ParameterName}_UEImpureCall(bool bExecute, float3 Location, float3 Axis, float MajorRadius, float MinorRadius, float MajorRadiusSegments, float MinorRadiusSegments, float4 Color)
+void DrawTorus_{ParameterName}_UEImpureCall(bool bExecute, float3 Location, float3 Axis, float MajorRadius, float MinorRadius, int MajorRadiusSegments, int MinorRadiusSegments, float4 Color)
 {
 }
 
