Niagara - Add insert options to simulation stage rows.

diff --git a/Engine/Plugins/FX/Niagara/Source/NiagaraEditor/Private/ViewModels/Stack/NiagaraStackEmitterPropertiesGroup.cpp b/Engine/Plugins/FX/Niagara/Source/NiagaraEditor/Private/ViewModels/Stack/NiagaraStackEmitterPropertiesGroup.cpp
--- a/Engine/Plugins/FX/Niagara/Source/NiagaraEditor/Private/ViewModels/Stack/NiagaraStackEmitterPropertiesGroup.cpp
+++ b/Engine/Plugins/FX/Niagara/Source/NiagaraEditor/Private/ViewModels/Stack/NiagaraStackEmitterPropertiesGroup.cpp
@@ -21,77 +21,32 @@
 
 #define LOCTEXT_NAMESPACE "StackEmitterProperties"
 
-class FNiagaraStackEmitterPropertiesGroupAddUtilities : public FNiagaraStackItemGroupAddUtilities
+FNiagaraStackEmitterStageAddUtilities::FNiagaraStackEmitterStageAddUtilities(TSharedPtr<FNiagaraEmitterViewModel> InEmitterViewModel, FOnItemAdded InOnItemAdded, bool bInAllowEventHandlers, bool bInAllowSimulationStages)
+	: FNiagaraStackItemGroupAddUtilities(LOCTEXT("AddUtilitiesName", "Stage"), EAddMode::AddFromAction, true, true)
+	, EmitterViewModelWeak(TWeakPtr<FNiagaraEmitterViewModel>(InEmitterViewModel))
+	, OnItemAdded(InOnItemAdded)
+	, bAllowEventHandlers(bInAllowEventHandlers)
+	, bAllowSimulationStages(bInAllowSimulationStages)
 {
-private:
-	enum class EEmitterAddMode
-	{
-		Event,
-		SimulationStage
-	};
-
-	class FAddEmitterStageAction : public INiagaraStackItemGroupAddAction
-	{
-	public:
-		FAddEmitterStageAction(EEmitterAddMode InAddMode, UClass* InSimulationStageClass, const TArray<FString>& InCategories, FText InDisplayName, FText InDescription, FText InKeywords)
-			: AddMode(InAddMode), SimulationStageClass(InSimulationStageClass), Categories(InCategories), DisplayName(InDisplayName), Description(InDescription), Keywords(InKeywords)
-		{
-		}
-
-		virtual TArray<FString> GetCategories() const override
-		{
-			return Categories;
-		}
-
-		virtual FText GetDisplayName() const override
-		{
-			return DisplayName;
-		}
-
-		virtual FText GetDescription() const override
-		{
-			return Description;
-		}
-
-		virtual FText GetKeywords() const override
-		{
-			return Keywords;
-		}
-
-		EEmitterAddMode AddMode;
-		UClass* SimulationStageClass;
-		TArray<FString> Categories;
-		FText DisplayName;
-		FText Description;
-		FText Keywords;
-	};
-
-public:
-	DECLARE_DELEGATE_TwoParams(FOnItemAdded, FGuid /* AddedEventHandlerId */, UNiagaraSimulationStageBase* /* AddedSimulationStage */);
-
-public:
-	FNiagaraStackEmitterPropertiesGroupAddUtilities(TSharedPtr<FNiagaraEmitterViewModel> InEmitterViewModel, FOnItemAdded InOnItemAdded)
-		: FNiagaraStackItemGroupAddUtilities(LOCTEXT("AddUtilitiesName", "Stage"), EAddMode::AddFromAction, true, true)
-		, EmitterViewModelWeak(TWeakPtr<FNiagaraEmitterViewModel>(InEmitterViewModel))
-		, OnItemAdded(InOnItemAdded)
-	{
-	}
-
-	virtual void AddItemDirectly() override { unimplemented(); }
+}
 
-	virtual void GenerateAddActions(TArray<TSharedRef<INiagaraStackItemGroupAddAction>>& OutAddActions, const FNiagaraStackItemGroupAddOptions& AddProperties) const override
+void FNiagaraStackEmitterStageAddUtilities::GenerateAddActions(TArray<TSharedRef<INiagaraStackItemGroupAddAction>>& OutAddActions, const FNiagaraStackItemGroupAddOptions& AddProperties) const
+{
+	if (bAllowEventHandlers)
 	{
 		TArray<FString> EventsCategories = { "Events" };
-		TArray<FString> SimulationStageCategories = { "Simulation Stages" };
-
 		OutAddActions.Add(MakeShared<FAddEmitterStageAction>(
-			EEmitterAddMode::Event,
+			EStageAddMode::Event,
 			nullptr,
 			EventsCategories,
 			LOCTEXT("AddEventActionDisplayName", "Event Handler"),
 			LOCTEXT("AddEventActionDescription", "Add an event handler to this emitter."),
 			LOCTEXT("AddEventActionKeywords", "Event")));
+	}
 
+	if (bAllowSimulationStages)
+	{
+		TArray<FString> SimulationStageCategories = { "Simulation Stages" };
 		const UNiagaraEditorSettings* NiagaraEditorSettings = GetDefault<UNiagaraEditorSettings>();
 		TArray<UClass*> SimulationStageClasses;
 		GetDerivedClasses(UNiagaraSimulationStageBase::StaticClass(), SimulationStageClasses);
@@ -100,7 +55,7 @@ public:
 			if (NiagaraEditorSettings->IsVisibleClass(SimulationStageClass))
 			{
 				OutAddActions.Add(MakeShared<FAddEmitterStageAction>(
-					EEmitterAddMode::SimulationStage,
+					EStageAddMode::SimulationStage,
 					SimulationStageClass,
 					SimulationStageCategories,
 					SimulationStageClass->GetDisplayNameText(),
@@ -109,74 +64,74 @@ public:
 			}
 		}
 	}
+}
 
-	virtual void ExecuteAddAction(TSharedRef<INiagaraStackItemGroupAddAction> AddAction, int32 TargetIndex) override
+void FNiagaraStackEmitterStageAddUtilities::ExecuteAddAction(TSharedRef<INiagaraStackItemGroupAddAction> AddAction, int32 TargetIndex)
+{
+	TSharedPtr<FNiagaraEmitterViewModel> EmitterViewModel = EmitterViewModelWeak.Pin();
+	if (EmitterViewModel.IsValid() == false)
 	{
-		TSharedPtr<FNiagaraEmitterViewModel> EmitterViewModel = EmitterViewModelWeak.Pin();
-		if (EmitterViewModel.IsValid() == false)
-		{
-			return;
-		}
+		return;
+	}
 
-		FVersionedNiagaraEmitter VersionedEmitter = EmitterViewModel->GetEmitter();
-		UNiagaraEmitter* Emitter = VersionedEmitter.Emitter;
-		UNiagaraScriptSource* Source = EmitterViewModel->GetSharedScriptViewModel()->GetGraphViewModel()->GetScriptSource();
-		UNiagaraGraph* Graph = EmitterViewModel->GetSharedScriptViewModel()->GetGraphViewModel()->GetGraph();
+	FVersionedNiagaraEmitter VersionedEmitter = EmitterViewModel->GetEmitter();
+	UNiagaraEmitter* Emitter = VersionedEmitter.Emitter;
+	UNiagaraScriptSource* Source = EmitterViewModel->GetSharedScriptViewModel()->GetGraphViewModel()->GetScriptSource();
+	UNiagaraGraph* Graph = EmitterViewModel->GetSharedScriptViewModel()->GetGraphViewModel()->GetGraph();
 
-		// The stack should not have been created if any of these are null, so bail out if it happens somehow rather than try to handle all of these cases.
-		checkf(Emitter != nullptr && Source != nullptr && Graph != nullptr, TEXT("Stack created for invalid emitter or graph."));
+	// The stack should not have been created if any of these are null, so bail out if it happens somehow rather than try to handle all of these cases.
+	checkf(Emitter != nullptr && Source != nullptr && Graph != nullptr, TEXT("Stack created for invalid emitter or graph."));
 
-		TSharedRef<FAddEmitterStageAction> AddEmitterStageAction = StaticCastSharedRef<FAddEmitterStageAction>(AddAction);
-		FGuid AddedEventHandlerId;
-		UNiagaraSimulationStageBase* AddedSimulationStage = nullptr;
-		if (AddEmitterStageAction->AddMode == EEmitterAddMode::Event)
+	TSharedRef<FAddEmitterStageAction> AddEmitterStageAction = StaticCastSharedRef<FAddEmitterStageAction>(AddAction);
+	FGuid AddedEventHandlerId;
+	UNiagaraSimulationStageBase* AddedSimulationStage = nullptr;
+	if (AddEmitterStageAction->AddMode == EStageAddMode::Event)
+	{
+		// since this is potentially modifying live data we need to kill off any existing instances that might be in flight before we make
+		// the change.
+		if (const UNiagaraSystem* EmitterSystem = Cast<UNiagaraSystem>(Emitter->GetOuter()))
 		{
-			// since this is potentially modifying live data we need to kill off any existing instances that might be in flight before we make
-			// the change.
-			if (const UNiagaraSystem* EmitterSystem = Cast<UNiagaraSystem>(Emitter->GetOuter()))
-			{
-				FNiagaraEditorUtilities::KillSystemInstances(*EmitterSystem);
-			}
-
-			FScopedTransaction ScopedTransaction(LOCTEXT("AddNewEventHandlerTransaction", "Add new event handler"));
-
-			Emitter->Modify();
-			FNiagaraEventScriptProperties EventScriptProperties;
-			EventScriptProperties.Script = NewObject<UNiagaraScript>(Emitter, MakeUniqueObjectName(Emitter, UNiagaraScript::StaticClass(), "EventScript"), EObjectFlags::RF_Transactional);
-			EventScriptProperties.Script->SetUsage(ENiagaraScriptUsage::ParticleEventScript);
-			EventScriptProperties.Script->SetUsageId(FGuid::NewGuid());
-			EventScriptProperties.Script->SetLatestSource(Source);
-			Emitter->AddEventHandler(EventScriptProperties, VersionedEmitter.Version);
-			FNiagaraStackGraphUtilities::ResetGraphForOutput(*Graph, ENiagaraScriptUsage::ParticleEventScript, EventScriptProperties.Script->GetUsageId());
-			AddedEventHandlerId = EventScriptProperties.Script->GetUsageId();
+			FNiagaraEditorUtilities::KillSystemInstances(*EmitterSystem);
 		}
-		else if (AddEmitterStageAction->AddMode == EEmitterAddMode::SimulationStage && AddEmitterStageAction->SimulationStageClass != nullptr)
-		{
-			FScopedTransaction ScopedTransaction(LOCTEXT("AddNewSimulationStagesTransaction", "Add new shader stage"));
 
-			Emitter->Modify();
-			AddedSimulationStage = NewObject<UNiagaraSimulationStageBase>(Emitter, AddEmitterStageAction->SimulationStageClass, NAME_None, RF_Transactional);
-			AddedSimulationStage->Script = NewObject<UNiagaraScript>(AddedSimulationStage, MakeUniqueObjectName(AddedSimulationStage, UNiagaraScript::StaticClass(), "SimulationStage"), EObjectFlags::RF_Transactional);
-			AddedSimulationStage->Script->SetUsage(ENiagaraScriptUsage::ParticleSimulationStageScript);
-			AddedSimulationStage->Script->SetUsageId(AddedSimulationStage->GetMergeId());
-			AddedSimulationStage->Script->SetLatestSource(Source);
-			Emitter->AddSimulationStage(AddedSimulationStage, VersionedEmitter.Version);
-			FNiagaraStackGraphUtilities::ResetGraphForOutput(*Graph, ENiagaraScriptUsage::ParticleSimulationStageScript, AddedSimulationStage->Script->GetUsageId());
+		FScopedTransaction ScopedTransaction(LOCTEXT("AddNewEventHandlerTransaction", "Add new event handler"));
+
+		Emitter->Modify();
+		FNiagaraEventScriptProperties EventScriptProperties;
+		EventScriptProperties.Script = NewObject<UNiagaraScript>(Emitter, MakeUniqueObjectName(Emitter, UNiagaraScript::StaticClass(), "EventScript"), EObjectFlags::RF_Transactional);
+		EventScriptProperties.Script->SetUsage(ENiagaraScriptUsage::ParticleEventScript);
+		EventScriptProperties.Script->SetUsageId(FGuid::NewGuid());
+		EventScriptProperties.Script->SetLatestSource(Source);
+		Emitter->AddEventHandler(EventScriptProperties, VersionedEmitter.Version);
+		FNiagaraStackGraphUtilities::ResetGraphForOutput(*Graph, ENiagaraScriptUsage::ParticleEventScript, EventScriptProperties.Script->GetUsageId());
+		AddedEventHandlerId = EventScriptProperties.Script->GetUsageId();
+	}
+	else if (AddEmitterStageAction->AddMode == EStageAddMode::SimulationStage && AddEmitterStageAction->SimulationStageClass != nullptr)
+	{
+		FScopedTransaction ScopedTransaction(LOCTEXT("AddNewSimulationStagesTransaction", "Add new simulation stage"));
+
+		Emitter->Modify();
+		AddedSimulationStage = NewObject<UNiagaraSimulationStageBase>(Emitter, AddEmitterStageAction->SimulationStageClass, NAME_None, RF_Transactional);
+		AddedSimulationStage->Script = NewObject<UNiagaraScript>(AddedSimulationStage, MakeUniqueObjectName(AddedSimulationStage, UNiagaraScript::StaticClass(), "SimulationStage"), EObjectFlags::RF_Transactional);
+		AddedSimulationStage->Script->SetUsage(ENiagaraScriptUsage::ParticleSimulationStageScript);
+		AddedSimulationStage->Script->SetUsageId(AddedSimulationStage->GetMergeId());
+		AddedSimulationStage->Script->SetLatestSource(Source);
+		Emitter->AddSimulationStage(AddedSimulationStage, VersionedEmitter.Version);
+		if (TargetIndex != INDEX_NONE)
+		{
+			Emitter->MoveSimulationStageToIndex(AddedSimulationStage, TargetIndex, VersionedEmitter.Version);
 		}
-
-		// Set the emitter here so that the internal state of the view model is updated.
-		// TODO: Move the logic for managing additional scripts into the emitter view model or script view model.
-		TWeakPtr<FNiagaraEmitterInstance, ESPMode::ThreadSafe> Simulation = EmitterViewModel->GetSimulation();
-		EmitterViewModel->Reset();
-		EmitterViewModel->Initialize(VersionedEmitter, Simulation);
-
-		OnItemAdded.ExecuteIfBound(AddedEventHandlerId, AddedSimulationStage);
+		FNiagaraStackGraphUtilities::ResetGraphForOutput(*Graph, ENiagaraScriptUsage::ParticleSimulationStageScript, AddedSimulationStage->Script->GetUsageId());
 	}
 
-private:
-	TWeakPtr<FNiagaraEmitterViewModel> EmitterViewModelWeak;
-	FOnItemAdded OnItemAdded;
-};
+	// Set the emitter here so that the internal state of the view model is updated.
+	// TODO: Move the logic for managing additional scripts into the emitter view model or script view model.
+	TWeakPtr<FNiagaraEmitterInstance, ESPMode::ThreadSafe> Simulation = EmitterViewModel->GetSimulation();
+	EmitterViewModel->Reset();
+	EmitterViewModel->Initialize(VersionedEmitter, Simulation);
+
+	OnItemAdded.ExecuteIfBound(AddedEventHandlerId, AddedSimulationStage);
+}
 
 UNiagaraStackEmitterPropertiesGroup::UNiagaraStackEmitterPropertiesGroup()
 	: PropertiesItem(nullptr)
@@ -187,8 +142,8 @@ void UNiagaraStackEmitterPropertiesGroup::Initialize(FRequiredEntryData InRequir
 {
 	FText DisplayName = LOCTEXT("EmitterPropertiesGroupName", "Properties");
 	FText Tooltip = LOCTEXT("EmitterPropertiesTooltip", "Properties which are set per Emitter.");
-	AddUtilities = MakeShared<FNiagaraStackEmitterPropertiesGroupAddUtilities>(InRequiredEntryData.EmitterViewModel.ToSharedRef(), 
-		FNiagaraStackEmitterPropertiesGroupAddUtilities::FOnItemAdded::CreateUObject(this, &UNiagaraStackEmitterPropertiesGroup::ItemAddedFromUtilities));
+	AddUtilities = MakeShared<FNiagaraStackEmitterStageAddUtilities>(InRequiredEntryData.EmitterViewModel.ToSharedRef(),
+		FNiagaraStackEmitterStageAddUtilities::FOnItemAdded::CreateUObject(this, &UNiagaraStackEmitterPropertiesGroup::ItemAddedFromUtilities));
 	Super::Initialize(InRequiredEntryData, DisplayName, Tooltip, AddUtilities.Get());
 }
 
diff --git a/Engine/Plugins/FX/Niagara/Source/NiagaraEditor/Private/ViewModels/Stack/NiagaraStackSimulationStageGroup.cpp b/Engine/Plugins/FX/Niagara/Source/NiagaraEditor/Private/ViewModels/Stack/NiagaraStackSimulationStageGroup.cpp
--- a/Engine/Plugins/FX/Niagara/Source/NiagaraEditor/Private/ViewModels/Stack/NiagaraStackSimulationStageGroup.cpp
+++ b/Engine/Plugins/FX/Niagara/Source/NiagaraEditor/Private/ViewModels/Stack/NiagaraStackSimulationStageGroup.cpp
@@ -16,6 +16,7 @@
 #include "NiagaraScriptSource.h"
 #include "NiagaraScriptMergeManager.h"
 #include "ViewModels/NiagaraSystemViewModel.h"
+#include "ViewModels/Stack/NiagaraStackEmitterPropertiesGroup.h"
 
 #include "Internationalization/Internationalization.h"
 #include "ScopedTransaction.h"
@@ -164,6 +165,10 @@ FNiagaraHierarchyIdentity UNiagaraStackSimulationStagePropertiesItem::DetermineS
 
 void UNiagaraStackSimulationStagePropertiesItem::SimulationStagePropertiesChanged()
 {
+	if (SimulationStage.IsValid() == false || GetEmitterViewModel().IsValid() == false)
+	{
+		return;
+	}
 	FVersionedNiagaraEmitter VersionedEmitter = GetEmitterViewModel()->GetEmitter();
 	GetSystemViewModel()->GetEmitterHandleViewModelForEmitter(VersionedEmitter).Get()->GetEmitterStackViewModel()->RequestValidationUpdate();
 
@@ -240,6 +245,23 @@ void UNiagaraStackSimulationStageGroup::SetIsEnabled(bool bEnabled)
 	SimulationStageProperties->SetSimulationStageEnabled(bEnabled);
 }
 
+INiagaraStackItemGroupAddUtilities* UNiagaraStackSimulationStageGroup::GetEmitterStageAddUtilities() const
+{
+	if (EmitterStageAddUtilities.IsValid() == false)
+	{
+		EmitterStageAddUtilities = MakeShared<FNiagaraStackEmitterStageAddUtilities>(GetEmitterViewModel(), FNiagaraStackEmitterStageAddUtilities::FOnItemAdded::CreateUObject(this, &UNiagaraStackSimulationStageGroup::StageAdded), false, true);
+	}
+	return EmitterStageAddUtilities.Get();
+}
+
+int32 UNiagaraStackSimulationStageGroup::GetStageIndex() const
+{
+	UNiagaraSimulationStageBase* SimulationStagePinned = SimulationStage.Get();
+	return SimulationStagePinned != nullptr 
+		? GetEmitterViewModel()->GetEmitter().GetEmitterData()->GetSimulationStages().IndexOfByKey(SimulationStagePinned)
+		: INDEX_NONE;
+}
+
 void UNiagaraStackSimulationStageGroup::FinalizeInternal()
 {
 	if (SimulationStage.IsValid())
@@ -460,6 +482,11 @@ bool UNiagaraStackSimulationStageGroup::HasBaseSimulationStage() const
 	return bHasBaseSimulationStageCache.GetValue();
 }
 
+void UNiagaraStackSimulationStageGroup::StageAdded(FGuid AddedEventHandlerId, UNiagaraSimulationStageBase* AddedSimulationStage) const
+{
+	OnRequestFullRefreshDeferred().Broadcast();
+}
+
 void UNiagaraStackSimulationStageGroup::SetOnModifiedSimulationStages(FOnModifiedSimulationStages OnModifiedSimulationStages)
 {
 	OnModifiedSimulationStagesDelegate = OnModifiedSimulationStages;
diff --git a/Engine/Plugins/FX/Niagara/Source/NiagaraEditor/Public/ViewModels/Stack/NiagaraStackEmitterPropertiesGroup.h b/Engine/Plugins/FX/Niagara/Source/NiagaraEditor/Public/ViewModels/Stack/NiagaraStackEmitterPropertiesGroup.h
--- a/Engine/Plugins/FX/Niagara/Source/NiagaraEditor/Public/ViewModels/Stack/NiagaraStackEmitterPropertiesGroup.h
+++ b/Engine/Plugins/FX/Niagara/Source/NiagaraEditor/Public/ViewModels/Stack/NiagaraStackEmitterPropertiesGroup.h
@@ -12,6 +12,54 @@ class UNiagaraSimulationStageBase;
 class UNiagaraStackEmitterPropertiesItem;
 class UNiagaraStackObject;
 
+
+class FNiagaraStackEmitterStageAddUtilities : public FNiagaraStackItemGroupAddUtilities
+{
+private:
+	enum class EStageAddMode
+	{
+		Event,
+		SimulationStage,
+	};
+
+	class FAddEmitterStageAction : public INiagaraStackItemGroupAddAction
+	{
+	public:
+		FAddEmitterStageAction(EStageAddMode InAddMode, UClass* InSimulationStageClass, const TArray<FString>& InCategories, FText InDisplayName, FText InDescription, FText InKeywords)
+			: AddMode(InAddMode), SimulationStageClass(InSimulationStageClass), Categories(InCategories), DisplayName(InDisplayName), Description(InDescription), Keywords(InKeywords)
+		{
+		}
+
+		virtual TArray<FString> GetCategories() const override { return Categories;	}
+		virtual FText GetDisplayName() const override {	return DisplayName;	}
+		virtual FText GetDescription() const override {	return Description;	}
+		virtual FText GetKeywords() const override { return Keywords; }
+
+		EStageAddMode AddMode;
+		UClass* SimulationStageClass;
+		TArray<FString> Categories;
+		FText DisplayName;
+		FText Description;
+		FText Keywords;
+	};
+
+public:
+	DECLARE_DELEGATE_TwoParams(FOnItemAdded, FGuid /* AddedEventHandlerId */, UNiagaraSimulationStageBase* /* AddedSimulationStage */);
+
+public:
+	FNiagaraStackEmitterStageAddUtilities(TSharedPtr<FNiagaraEmitterViewModel> InEmitterViewModel, FOnItemAdded InOnItemAdded, bool bInAllowEventHandlers = true, bool bInAllowSimulationStages = true);
+
+	virtual void AddItemDirectly() override { unimplemented(); }
+	virtual void GenerateAddActions(TArray<TSharedRef<INiagaraStackItemGroupAddAction>>& OutAddActions, const FNiagaraStackItemGroupAddOptions& AddProperties) const override;
+	virtual void ExecuteAddAction(TSharedRef<INiagaraStackItemGroupAddAction> AddAction, int32 TargetIndex) override;
+
+private:
+	TWeakPtr<FNiagaraEmitterViewModel> EmitterViewModelWeak;
+	FOnItemAdded OnItemAdded;
+	bool bAllowEventHandlers = true;
+	bool bAllowSimulationStages = true;
+};
+
 UCLASS(MinimalAPI)
 class UNiagaraStackEmitterPropertiesGroup : public UNiagaraStackItemGroup
 {
diff --git a/Engine/Plugins/FX/Niagara/Source/NiagaraEditor/Public/ViewModels/Stack/NiagaraStackSimulationStageGroup.h b/Engine/Plugins/FX/Niagara/Source/NiagaraEditor/Public/ViewModels/Stack/NiagaraStackSimulationStageGroup.h
--- a/Engine/Plugins/FX/Niagara/Source/NiagaraEditor/Public/ViewModels/Stack/NiagaraStackSimulationStageGroup.h
+++ b/Engine/Plugins/FX/Niagara/Source/NiagaraEditor/Public/ViewModels/Stack/NiagaraStackSimulationStageGroup.h
@@ -88,6 +88,10 @@ public:
 	NIAGARAEDITOR_API virtual void SetIsEnabled(bool bEnabled) override;
 	virtual bool SupportsChangeEnabled() const override { return true; }
 
+	NIAGARAEDITOR_API INiagaraStackItemGroupAddUtilities* GetEmitterStageAddUtilities() const;
+
+	NIAGARAEDITOR_API int32 GetStageIndex() const;
+
 protected:
 	NIAGARAEDITOR_API virtual void FinalizeInternal() override;
 
@@ -103,6 +107,8 @@ private:
 	NIAGARAEDITOR_API void SimulationStagePropertiesChanged();
 
 	NIAGARAEDITOR_API bool HasBaseSimulationStage() const;
+
+	void StageAdded(FGuid AddedEventHandlerId, UNiagaraSimulationStageBase* AddedSimulationStage) const;
 	
 private:
 	TWeakObjectPtr<UNiagaraSimulationStageBase> SimulationStage;
@@ -113,4 +119,6 @@ private:
 
 	UPROPERTY()
 	TObjectPtr<UNiagaraStackSimulationStagePropertiesItem> SimulationStageProperties;
+
+	mutable TSharedPtr<INiagaraStackItemGroupAddUtilities> EmitterStageAddUtilities;
 };
diff --git a/Engine/Plugins/FX/Niagara/Source/NiagaraEditorWidgets/Private/NiagaraEditorWidgetsUtilities.cpp b/Engine/Plugins/FX/Niagara/Source/NiagaraEditorWidgets/Private/NiagaraEditorWidgetsUtilities.cpp
--- a/Engine/Plugins/FX/Niagara/Source/NiagaraEditorWidgets/Private/NiagaraEditorWidgetsUtilities.cpp
+++ b/Engine/Plugins/FX/Niagara/Source/NiagaraEditorWidgets/Private/NiagaraEditorWidgetsUtilities.cpp
@@ -29,6 +29,7 @@
 #include "ViewModels/NiagaraSystemSelectionViewModel.h"
 #include "ViewModels/Stack/NiagaraStackInputCategory.h"
 #include "ViewModels/Stack/NiagaraStackRendererItem.h"
+#include "ViewModels/Stack/NiagaraStackSimulationStageGroup.h"
 
 #define LOCTEXT_NAMESPACE "NiagaraStackEditorWidgetsUtilities"
 
@@ -470,6 +471,40 @@ bool FNiagaraStackEditorWidgetsUtilities::AddStackModuleItemContextMenuActions(F
 	return true;
 }
 
+void ShowInsertSimulationStageMenu(TWeakObjectPtr<UNiagaraStackSimulationStageGroup> StackSimulationStageGroupWeak, int32 InsertOffset, TWeakPtr<SWidget> TargetWidgetWeak)
+{
+	UNiagaraStackSimulationStageGroup* StackSimulationStageGroup = StackSimulationStageGroupWeak.Get();
+	TSharedPtr<SWidget> TargetWidget = TargetWidgetWeak.Pin();
+	if (StackSimulationStageGroup != nullptr && TargetWidget.IsValid())
+	{
+		TSharedRef<SNiagaraStackItemGroupAddMenu> MenuContent = SNew(SNiagaraStackItemGroupAddMenu, nullptr, StackSimulationStageGroup->GetEmitterStageAddUtilities(), StackSimulationStageGroup->GetStageIndex() + InsertOffset);
+		FGeometry ThisGeometry = TargetWidget->GetCachedGeometry();
+		bool bAutoAdjustForDpiScale = false; // Don't adjust for dpi scale because the push menu command is expecting an unscaled position.
+		FVector2D MenuPosition = FSlateApplication::Get().CalculatePopupWindowPosition(ThisGeometry.GetLayoutBoundingRect(), MenuContent->GetDesiredSize(), bAutoAdjustForDpiScale);
+		FSlateApplication::Get().PushMenu(TargetWidget.ToSharedRef(), FWidgetPath(), MenuContent, MenuPosition, FPopupTransitionEffect::ContextMenu);
+		FSlateApplication::Get().SetKeyboardFocus(MenuContent->GetFilterTextBox());
+	}
+}
+
+void FNiagaraStackEditorWidgetsUtilities::AddStackSimulationStageGroupContextMenuActions(FMenuBuilder& MenuBuilder, UNiagaraStackSimulationStageGroup* StackSimulationStageGroup, TSharedRef<SWidget> TargetWidget)
+{
+	MenuBuilder.BeginSection("SimulationStageActions", LOCTEXT("SimulationStageActions", "Simulation Stage Actions"));
+	{
+		MenuBuilder.AddMenuEntry(
+			LOCTEXT("InsertSimulationStageAbove", "Insert Above"),
+			LOCTEXT("InsertSimulationStageAboveToolTip", "Insert a new simulation stage above this stage in the stack."),
+			FSlateIcon(),
+			FUIAction(FExecuteAction::CreateStatic(&ShowInsertSimulationStageMenu, TWeakObjectPtr<UNiagaraStackSimulationStageGroup>(StackSimulationStageGroup), 0, TWeakPtr<SWidget>(TargetWidget))));
+
+		MenuBuilder.AddMenuEntry(
+			LOCTEXT("InsertSimulationStageBelow", "Insert Below"),
+			LOCTEXT("InsertSimulationStageBelowToolTip", "Insert a new simulation stage below this stage in the stack."),
+			FSlateIcon(),
+			FUIAction(FExecuteAction::CreateStatic(&ShowInsertSimulationStageMenu, TWeakObjectPtr<UNiagaraStackSimulationStageGroup>(StackSimulationStageGroup), 1, TWeakPtr<SWidget>(TargetWidget))));
+	}
+	MenuBuilder.EndSection();
+}
+
 TSharedRef<FDragDropOperation> FNiagaraStackEditorWidgetsUtilities::ConstructDragDropOperationForStackEntries(const TArray<UNiagaraStackEntry*>& DraggedEntries)
 {
 	TSharedRef<FNiagaraStackEntryDragDropOp> DragDropOp = MakeShared<FNiagaraStackEntryDragDropOp>(DraggedEntries);
diff --git a/Engine/Plugins/FX/Niagara/Source/NiagaraEditorWidgets/Private/NiagaraEditorWidgetsUtilities.h b/Engine/Plugins/FX/Niagara/Source/NiagaraEditorWidgets/Private/NiagaraEditorWidgetsUtilities.h
--- a/Engine/Plugins/FX/Niagara/Source/NiagaraEditorWidgets/Private/NiagaraEditorWidgetsUtilities.h
+++ b/Engine/Plugins/FX/Niagara/Source/NiagaraEditorWidgets/Private/NiagaraEditorWidgetsUtilities.h
@@ -13,6 +13,7 @@ class SWidget;
 class UNiagaraStackEntry;
 class UNiagaraStackItem;
 class UNiagaraStackModuleItem;
+class UNiagaraStackSimulationStageGroup;
 
 namespace FNiagaraStackEditorWidgetsUtilities
 {
@@ -36,6 +37,8 @@ namespace FNiagaraStackEditorWidgetsUtilities
 
 	bool AddStackModuleItemContextMenuActions(FMenuBuilder& MenuBuilder, UNiagaraStackModuleItem& StackItem, TSharedRef<SWidget> TargetWidget);
 
+	void AddStackSimulationStageGroupContextMenuActions(FMenuBuilder& MenuBuilder, UNiagaraStackSimulationStageGroup* SimulationStageGroup, TSharedRef<SWidget> TagetWidget);
+
 	TSharedRef<FDragDropOperation> ConstructDragDropOperationForStackEntries(const TArray<UNiagaraStackEntry*>& DraggedEntries);
 
 	void HandleDragLeave(const FDragDropEvent& InDragDropEvent);
diff --git a/Engine/Plugins/FX/Niagara/Source/NiagaraEditorWidgets/Private/SNiagaraOverviewStack.cpp b/Engine/Plugins/FX/Niagara/Source/NiagaraEditorWidgets/Private/SNiagaraOverviewStack.cpp
--- a/Engine/Plugins/FX/Niagara/Source/NiagaraEditorWidgets/Private/SNiagaraOverviewStack.cpp
+++ b/Engine/Plugins/FX/Niagara/Source/NiagaraEditorWidgets/Private/SNiagaraOverviewStack.cpp
@@ -26,6 +26,7 @@
 #include "ViewModels/Stack/NiagaraStackItem.h"
 #include "ViewModels/Stack/NiagaraStackModuleItem.h"
 #include "ViewModels/Stack/NiagaraStackViewModel.h"
+#include "ViewModels/Stack/NiagaraStackSimulationStageGroup.h"
 #include "NiagaraEditorWidgetsStyle.h"
 #include "NiagaraEditorWidgetsUtilities.h"
 #include "SNiagaraStack.h"
@@ -274,6 +275,12 @@ class SNiagaraSystemOverviewEntryListRow : public STableRow<UNiagaraStackEntry*>
 				bMenuItemsAdded |= FNiagaraStackEditorWidgetsUtilities::AddStackItemContextMenuActions(MenuBuilder, *CastChecked<UNiagaraStackItem>(StackEntry));
 			}
 
+			if (StackEntry->IsA<UNiagaraStackSimulationStageGroup>())
+			{
+				FNiagaraStackEditorWidgetsUtilities::AddStackSimulationStageGroupContextMenuActions(MenuBuilder, CastChecked<UNiagaraStackSimulationStageGroup>(StackEntry), this->AsShared());
+				bMenuItemsAdded = true;
+			}
+
 			bMenuItemsAdded |= FNiagaraStackEditorWidgetsUtilities::AddStackEntryAssetContextMenuActions(MenuBuilder, *StackEntry);
 			bMenuItemsAdded |= StackCommandContext->AddEditMenuItems(MenuBuilder);
 		
diff --git a/Engine/Plugins/FX/Niagara/Source/NiagaraEditorWidgets/Private/SNiagaraStack.cpp b/Engine/Plugins/FX/Niagara/Source/NiagaraEditorWidgets/Private/SNiagaraStack.cpp
--- a/Engine/Plugins/FX/Niagara/Source/NiagaraEditorWidgets/Private/SNiagaraStack.cpp
+++ b/Engine/Plugins/FX/Niagara/Source/NiagaraEditorWidgets/Private/SNiagaraStack.cpp
@@ -23,6 +23,7 @@
 #include "ViewModels/Stack/NiagaraStackValueCollection.h"
 #include "ViewModels/Stack/NiagaraStackModuleItem.h"
 #include "ViewModels/Stack/NiagaraStackSystemSettingsGroup.h"
+#include "ViewModels/Stack/NiagaraStackSimulationStageGroup.h"
 #include "Framework/Commands/GenericCommands.h"
 #include "IContentBrowserSingleton.h"
 #include "ContentBrowserModule.h"
@@ -881,7 +882,13 @@ SNiagaraStack::FRowWidgets SNiagaraStack::ConstructNameAndValueWidgetsForItem(UN
 {
 	if (Item->IsA<UNiagaraStackItemGroup>())
 	{
-		return FRowWidgets(SNew(SNiagaraStackItemGroup, *CastChecked<UNiagaraStackItemGroup>(Item), StackViewModel));
+		TSharedRef<SWidget> GroupWidget = SNew(SNiagaraStackItemGroup, *CastChecked<UNiagaraStackItemGroup>(Item), StackViewModel);
+		if (Item->IsA<UNiagaraStackSimulationStageGroup>())
+		{
+			Container->AddFillRowContextMenuHandler(SNiagaraStackTableRow::FOnFillRowContextMenu::CreateStatic(
+				&FNiagaraStackEditorWidgetsUtilities::AddStackSimulationStageGroupContextMenuActions, CastChecked<UNiagaraStackSimulationStageGroup>(Item), GroupWidget));
+		}
+		return FRowWidgets(GroupWidget);
 	}
 	else if (Item->IsA<UNiagaraStackModuleItem>())
 	{
