- Add GPU support to PlatformSet data interface

#rnx

diff --git a/Engine/Plugins/FX/Niagara/Shaders/Private/NiagaraDataInterfacePlatformSetTemplate.ush b/Engine/Plugins/FX/Niagara/Shaders/Private/NiagaraDataInterfacePlatformSetTemplate.ush
new file mode 100644
--- /dev/null
+++ b/Engine/Plugins/FX/Niagara/Shaders/Private/NiagaraDataInterfacePlatformSetTemplate.ush
@@ -0,0 +1,8 @@
+// Copyright Epic Games, Inc. All Rights Reserved.
+
+int {ParameterName}_bIsActive;
+
+void IsActive_{ParameterName}(out bool bIsActive)
+{
+	bIsActive = {ParameterName}_bIsActive != 0;
+}
diff --git a/Engine/Plugins/FX/Niagara/Source/Niagara/Classes/NiagaraDataInterfacePlatformSet.h b/Engine/Plugins/FX/Niagara/Source/Niagara/Classes/NiagaraDataInterfacePlatformSet.h
--- a/Engine/Plugins/FX/Niagara/Source/Niagara/Classes/NiagaraDataInterfacePlatformSet.h
+++ b/Engine/Plugins/FX/Niagara/Source/Niagara/Classes/NiagaraDataInterfacePlatformSet.h
@@ -24,7 +24,15 @@ public:
 
 	// UNiagaraDataInterface interface
 	NIAGARA_API virtual void GetVMExternalFunction(const FVMExternalFunctionBindingInfo& BindingInfo, void* InstanceData, FVMExternalFunction& OutFunc) override;
-	virtual bool CanExecuteOnTarget(ENiagaraSimTarget Target)const override { return Target == ENiagaraSimTarget::CPUSim; }
+	virtual bool CanExecuteOnTarget(ENiagaraSimTarget Target)const override { return true; }
+
+#if WITH_EDITORONLY_DATA
+	virtual bool AppendCompileHash(FNiagaraCompileHashVisitor* InVisitor) const override;
+	virtual bool GetFunctionHLSL(const FNiagaraDataInterfaceGPUParamInfo& ParamInfo, const FNiagaraDataInterfaceGeneratedFunction& FunctionInfo, int FunctionInstanceIndex, FString& OutHLSL) override;
+	virtual void GetParameterDefinitionHLSL(const FNiagaraDataInterfaceGPUParamInfo& ParamInfo, FString& OutHLSL) override;
+#endif
+	virtual void BuildShaderParameters(FNiagaraShaderParametersBuilder& ShaderParametersBuilder) const override;
+	virtual void SetShaderParameters(const FNiagaraDataInterfaceSetShaderParametersContext& Context) const override;
 
 	NIAGARA_API virtual bool Equals(const UNiagaraDataInterface* Other) const override;
 
diff --git a/Engine/Plugins/FX/Niagara/Source/Niagara/Private/NiagaraDataInterfacePlatformSet.cpp b/Engine/Plugins/FX/Niagara/Source/Niagara/Private/NiagaraDataInterfacePlatformSet.cpp
--- a/Engine/Plugins/FX/Niagara/Source/Niagara/Private/NiagaraDataInterfacePlatformSet.cpp
+++ b/Engine/Plugins/FX/Niagara/Source/Niagara/Private/NiagaraDataInterfacePlatformSet.cpp
@@ -1,16 +1,55 @@
 // Copyright Epic Games, Inc. All Rights Reserved.
 #include "NiagaraDataInterfacePlatformSet.h"
 #include "NiagaraTypes.h"
+#include "NiagaraCompileHashVisitor.h"
+#include "NiagaraShaderParametersBuilder.h"
+
+#include "UObject/UObjectIterator.h"
 
 #include UE_INLINE_GENERATED_CPP_BY_NAME(NiagaraDataInterfacePlatformSet)
 
+namespace NDIPlatformSetPrivate
+{
+	BEGIN_SHADER_PARAMETER_STRUCT(FShaderParameters, )
+		SHADER_PARAMETER(int, bIsActive)
+	END_SHADER_PARAMETER_STRUCT()
+
+	static const TCHAR* TemplateShaderFile = TEXT("/Plugin/FX/Niagara/Private/NiagaraDataInterfacePlatformSetTemplate.ush");
+
+	static FName IsActiveName(TEXT("IsActive"));
+
+	struct FNDIPlatformSetProxy : public FNiagaraDataInterfaceProxy
+	{
+		virtual int32 PerInstanceDataPassedToRenderThreadSize() const override { return 0; }
+
+		bool bIsActive = false;
+	};
+
+	void RefreshPlatformSet()
+	{
+		for (TObjectIterator<UNiagaraDataInterfacePlatformSet> It; It; ++It)
+		{
+			UNiagaraDataInterfacePlatformSet* DataInterface = *It;
+			if (::IsValid(DataInterface) && !DataInterface->HasAnyFlags(RF_ClassDefaultObject))
+			{
+				static_cast<FNDIPlatformSetProxy*>(DataInterface->GetProxy())->bIsActive = DataInterface->Platforms.IsActive();
+			}
+		}
+	}
+
+	FAutoConsoleVariableSink CVarSyncPlatformSet(FConsoleCommandDelegate::CreateStatic(&RefreshPlatformSet));
+}
+
 UNiagaraDataInterfacePlatformSet::UNiagaraDataInterfacePlatformSet(FObjectInitializer const& ObjectInitializer)
 	: Super(ObjectInitializer)
 {
+	using namespace NDIPlatformSetPrivate;
+	Proxy.Reset(new FNDIPlatformSetProxy());
 }
 
 void UNiagaraDataInterfacePlatformSet::PostInitProperties()
 {
+	using namespace NDIPlatformSetPrivate;
 	Super::PostInitProperties();
 
 	if (HasAnyFlags(RF_ClassDefaultObject))
@@ -18,13 +57,16 @@ void UNiagaraDataInterfacePlatformSet::PostInitProperties()
 		ENiagaraTypeRegistryFlags Flags = ENiagaraTypeRegistryFlags::AllowAnyVariable | ENiagaraTypeRegistryFlags::AllowParameter;
 		FNiagaraTypeRegistry::Register(FNiagaraTypeDefinition(GetClass()), Flags);
 	}
+	else
+	{
+		GetProxyAs<FNDIPlatformSetProxy>()->bIsActive = Platforms.IsActive();
+	}
 }
 
-static FName IsActiveName(TEXT("IsActive"));
-
 #if WITH_EDITORONLY_DATA
 void UNiagaraDataInterfacePlatformSet::GetFunctionsInternal(TArray<FNiagaraFunctionSignature>& OutFunctions) const
 {
+	using namespace NDIPlatformSetPrivate;
 	{
 		FNiagaraFunctionSignature Sig;
 		Sig.Name = IsActiveName;
@@ -42,12 +84,59 @@ void UNiagaraDataInterfacePlatformSet::GetFunctionsInternal(TArray<FNiagaraFunct
 DEFINE_NDI_DIRECT_FUNC_BINDER(UNiagaraDataInterfacePlatformSet, IsActive);
 void UNiagaraDataInterfacePlatformSet::GetVMExternalFunction(const FVMExternalFunctionBindingInfo& BindingInfo, void* InstanceData, FVMExternalFunction& OutFunc)
 {
+	using namespace NDIPlatformSetPrivate;
 	if (BindingInfo.Name == IsActiveName && BindingInfo.GetNumInputs() == 0 && BindingInfo.GetNumOutputs() == 1)
 	{
 		NDI_FUNC_BINDER(UNiagaraDataInterfacePlatformSet, IsActive)::Bind(this, OutFunc);
 	}
 }
 
+#if WITH_EDITORONLY_DATA
+bool UNiagaraDataInterfacePlatformSet::AppendCompileHash(FNiagaraCompileHashVisitor* InVisitor) const
+{
+	using namespace NDIPlatformSetPrivate;
+	if (!Super::AppendCompileHash(InVisitor))
+	{
+		return false;
+	}
+
+	InVisitor->UpdateShaderFile(TemplateShaderFile);
+	InVisitor->UpdateShaderParameters<FShaderParameters>();
+	return true;
+}
+
+bool UNiagaraDataInterfacePlatformSet::GetFunctionHLSL(const FNiagaraDataInterfaceGPUParamInfo& ParamInfo, const FNiagaraDataInterfaceGeneratedFunction& FunctionInfo, int FunctionInstanceIndex, FString& OutHLSL)
+{
+	using namespace NDIPlatformSetPrivate;
+	return FunctionInfo.DefinitionName == IsActiveName;
+}
+
+void UNiagaraDataInterfacePlatformSet::GetParameterDefinitionHLSL(const FNiagaraDataInterfaceGPUParamInfo& ParamInfo, FString& OutHLSL)
+{
+	using namespace NDIPlatformSetPrivate;
+	const TMap<FString, FStringFormatArg> TemplateArgs =
+	{
+		{TEXT("ParameterName"),	ParamInfo.DataInterfaceHLSLSymbol},
+	};
+	AppendTemplateHLSL(OutHLSL, TemplateShaderFile, TemplateArgs);
+}
+#endif
+
+void UNiagaraDataInterfacePlatformSet::BuildShaderParameters(FNiagaraShaderParametersBuilder& ShaderParametersBuilder) const
+{
+	using namespace NDIPlatformSetPrivate;
+	ShaderParametersBuilder.AddNestedStruct<FShaderParameters>();
+}
+
+void UNiagaraDataInterfacePlatformSet::SetShaderParameters(const FNiagaraDataInterfaceSetShaderParametersContext& Context) const
+{
+	using namespace NDIPlatformSetPrivate;
+	FNDIPlatformSetProxy& DIProxy = Context.GetProxy<FNDIPlatformSetProxy>();
+
+	FShaderParameters* ShaderParameters = Context.GetParameterNestedStruct<FShaderParameters>();
+	ShaderParameters->bIsActive = DIProxy.bIsActive ? 1 : 0;
+}
+
 bool UNiagaraDataInterfacePlatformSet::Equals(const UNiagaraDataInterface* Other) const
 {
 	if (!Super::Equals(Other))
