Added IsAssetVisible check for available Niagara Parameter Collections & Niagara Parameter Definitions

diff --git a/Engine/Plugins/FX/Niagara/Source/NiagaraEditor/Public/NiagaraEditorModule.h b/Engine/Plugins/FX/Niagara/Source/NiagaraEditor/Public/NiagaraEditorModule.h
--- a/Engine/Plugins/FX/Niagara/Source/NiagaraEditor/Public/NiagaraEditorModule.h
+++ b/Engine/Plugins/FX/Niagara/Source/NiagaraEditor/Public/NiagaraEditorModule.h
@@ -3,6 +3,7 @@
 #pragma once
 
 #include "CoreMinimal.h"
+#include "AssetToolsModule.h"
 #include "AssetRegistry/AssetRegistryModule.h"
 #include "Modules/ModuleInterface.h"
 #include "Toolkits/AssetEditorToolkit.h"
@@ -325,9 +326,15 @@ private:
 			TArray<FAssetData> AssetData;
 			AssetRegistryModule.GetRegistry().GetAssetsByClass(AssetType::StaticClass()->GetClassPathName(), AssetData);
 
+			FAssetToolsModule& AssetToolsModule = FModuleManager::Get().LoadModuleChecked<FAssetToolsModule>("AssetTools");
 			CachedAssets.Reset(AssetData.Num());
 			for (const FAssetData& AssetDatum : AssetData)
 			{
+				if(AssetToolsModule.Get().IsAssetVisible(AssetDatum, true) == false)
+				{
+					continue;
+				}
+				
 				if (AssetDatum.IsAssetLoaded() || (bAllowLoading && FPackageName::GetPackageMountPoint(AssetDatum.PackageName.ToString()) != NAME_None))
 				{
 					AssetType* Asset = nullptr;
