- Handle variable rename / remove for user parameter bindings on data interfaces

#rnx
#tests Added a RT DI to user / system / emitter levels bind to a user RT Object then rename that object

diff --git a/Engine/Plugins/FX/Niagara/Source/Niagara/Private/NiagaraSystem.cpp b/Engine/Plugins/FX/Niagara/Source/Niagara/Private/NiagaraSystem.cpp
--- a/Engine/Plugins/FX/Niagara/Source/Niagara/Private/NiagaraSystem.cpp
+++ b/Engine/Plugins/FX/Niagara/Source/Niagara/Private/NiagaraSystem.cpp
@@ -106,6 +106,39 @@ namespace NiagaraSystemPrivate
 		bool TryLoad(FCbFieldViewIterator& ArgField);
 		void Save(FCbWriter& Writer);
 	};
+
+	template<typename TAction>
+	void ForEachDataInterfaceUserParameterBinding(UObject* Outer, TAction Action, bool bShouldModifyObjects = true)
+	{
+		ForEachObjectWithOuter(
+			Outer,
+			[&Action, bShouldModifyObjects](UObject* Object)
+			{
+				if (Object && Object->IsA<UNiagaraDataInterface>())
+				{
+					bool bNeedsModify = bShouldModifyObjects;
+					for (TFieldIterator<const FStructProperty> PropertyIt(Object->GetClass()); PropertyIt; ++PropertyIt)
+					{
+						const FStructProperty* StructProperty = *PropertyIt;
+						if (const UStruct* InnerStruct = StructProperty->Struct)
+						{
+							if (InnerStruct->IsChildOf(FNiagaraUserParameterBinding::StaticStruct()))
+							{
+								if ( bNeedsModify )
+								{
+									bNeedsModify = false;
+									Object->Modify();
+								}
+
+								FNiagaraUserParameterBinding* ParameterBinding = StructProperty->ContainerPtrToValuePtr<FNiagaraUserParameterBinding>(Object);
+								Action(*ParameterBinding);
+							}
+						}
+					}
+				}
+			}
+		);
+	}
 #endif
 }
 
@@ -380,8 +413,21 @@ void UNiagaraSystem::HandleVariableRenamed(const FNiagaraVariable& InOldVariable
 {
 	if (InOldVariable.IsInNameSpace(FNiagaraConstants::UserNamespaceString))
 	{
+		NiagaraSystemPrivate::ForEachDataInterfaceUserParameterBinding(
+			this,
+			[&InOldVariable, &InNewVariable](FNiagaraUserParameterBinding& ParameterBinding)
+			{
+				if (ParameterBinding.Parameter == InOldVariable)
+				{
+					ParameterBinding.Parameter = InNewVariable;
+				}
+			}
+		);
+
 		if (GetExposedParameters().IndexOf(InOldVariable) != INDEX_NONE)
+		{
 			GetExposedParameters().RenameParameter(InOldVariable, InNewVariable.GetName());
+		}
 		InitSystemCompiledData();
 	}
 
@@ -404,8 +450,21 @@ void UNiagaraSystem::HandleVariableRemoved(const FNiagaraVariable& InOldVariable
 {
 	if (InOldVariable.IsInNameSpace(FNiagaraConstants::UserNamespaceString))
 	{
+		NiagaraSystemPrivate::ForEachDataInterfaceUserParameterBinding(
+			this,
+			[&InOldVariable](FNiagaraUserParameterBinding& ParameterBinding)
+			{
+				if (ParameterBinding.Parameter == InOldVariable)
+				{
+					ParameterBinding.Parameter = FNiagaraVariable();
+				}
+			}
+		);
+
 		if (GetExposedParameters().IndexOf(InOldVariable) != INDEX_NONE)
+		{
 			GetExposedParameters().RemoveParameter(InOldVariable);
+		}
 		InitSystemCompiledData();
 	}
 	for (const FNiagaraEmitterHandle& Handle : GetEmitterHandles())
