svt export from baker can bind output volume scale to user parameters

diff --git a/Engine/Plugins/FX/Niagara/Source/Niagara/Classes/NiagaraBakerOutputSparseVolumeTexture.h b/Engine/Plugins/FX/Niagara/Source/Niagara/Classes/NiagaraBakerOutputSparseVolumeTexture.h
--- a/Engine/Plugins/FX/Niagara/Source/Niagara/Classes/NiagaraBakerOutputSparseVolumeTexture.h
+++ b/Engine/Plugins/FX/Niagara/Source/Niagara/Classes/NiagaraBakerOutputSparseVolumeTexture.h
@@ -3,6 +3,7 @@
 #pragma once
 
 #include "NiagaraBakerOutput.h"
+#include "NiagaraParameterBinding.h"
 #include "NiagaraBakerOutputSparseVolumeTexture.generated.h"
 
 UCLASS(meta = (DisplayName = "SparseVolume Texture Output"), MinimalAPI)
@@ -13,11 +14,18 @@ class UNiagaraBakerOutputSparseVolumeTexture : public UNiagaraBakerOutput
 public:
 	UNiagaraBakerOutputSparseVolumeTexture(const FObjectInitializer& Init)	
 	{
+#if WITH_EDITORONLY_DATA
+		VolumeWorldSpaceSizeBinding.SetUsage(ENiagaraParameterBindingUsage::User);
+		VolumeWorldSpaceSizeBinding.SetAllowedTypeDefinitions( { FNiagaraTypeDefinition::GetVec3Def() } );
+#endif
 	}
 
 	UPROPERTY(EditAnywhere, Category = "Settings")
 	FNiagaraBakerTextureSource SourceBinding;
 
+	UPROPERTY(EditAnywhere, Category = "Settings")
+	FNiagaraParameterBinding VolumeWorldSpaceSizeBinding;
+
 	/**
 	When enabled a volume atlas is created, the atlas is along X & Y not Z based on baker settings.
 	*/
diff --git a/Engine/Plugins/FX/Niagara/Source/NiagaraEditor/Private/NiagaraBakerRendererOutputSparseVolumeTexture.cpp b/Engine/Plugins/FX/Niagara/Source/NiagaraEditor/Private/NiagaraBakerRendererOutputSparseVolumeTexture.cpp
--- a/Engine/Plugins/FX/Niagara/Source/NiagaraEditor/Private/NiagaraBakerRendererOutputSparseVolumeTexture.cpp
+++ b/Engine/Plugins/FX/Niagara/Source/NiagaraEditor/Private/NiagaraBakerRendererOutputSparseVolumeTexture.cpp
@@ -302,7 +302,29 @@ void FNiagaraBakerRendererOutputSparseVolumeTexture::BakeFrame(FNiagaraBakerFeed
 			return;
 		}
 
-		if (!SVTAsset->AppendFrame(SparseTextureData, FTransform::Identity))
+		FTransform TransformToUse = FTransform::Identity;
+								
+		FNiagaraVariableBase& BoundWorldSizeVar = OutputVolumeTexture->VolumeWorldSpaceSizeBinding.ResolvedParameter;
+
+		if (BoundWorldSizeVar.IsValid())
+		{
+			FVector3f WorldGridScale =
+				BakerRenderer.GetPreviewComponent()->GetOverrideParameters().GetParameterValueOrDefault<FVector3f>(BoundWorldSizeVar, FVector3f(1, 1, 1));
+			
+			if (WorldGridScale.Length() < SMALL_NUMBER)
+			{
+				WorldGridScale = FVector3f(1, 1, 1);
+			}
+
+			FVector WorldScaleFVector(WorldGridScale.X, WorldGridScale.Y, WorldGridScale.Z);
+
+			// scale by volume resolution to get proper world space scale rendering
+			const FVector FloatResolution(VolumeResolution.X, VolumeResolution.Y, VolumeResolution.Z);
+			WorldScaleFVector /= FloatResolution;
+			TransformToUse.SetScale3D(WorldScaleFVector);
+		}
+
+		if (!SVTAsset->AppendFrame(SparseTextureData, TransformToUse))
 		{
 			UE_LOG(LogNiagaraBaker, Error, TEXT("Cannot append frame to SVT"));
 		}
