- Change / to | to make it clear that it's not X out of Y

#rnx
#tests Open parameter panel to see change

diff --git a/Engine/Plugins/FX/Niagara/Source/NiagaraEditor/Private/Widgets/SNiagaraParameterPanel.cpp b/Engine/Plugins/FX/Niagara/Source/NiagaraEditor/Private/Widgets/SNiagaraParameterPanel.cpp
--- a/Engine/Plugins/FX/Niagara/Source/NiagaraEditor/Private/Widgets/SNiagaraParameterPanel.cpp
+++ b/Engine/Plugins/FX/Niagara/Source/NiagaraEditor/Private/Widgets/SNiagaraParameterPanel.cpp
@@ -506,7 +506,7 @@ TSharedRef<SWidget> SNiagaraParameterPanel::OnGenerateWidgetForItem(const FNiaga
 				[
 					SNew(STextBlock)
 					.Margin(FMargin(2, 0))
-					.Text(FText::FromString(TEXT("/")))
+					.Text(FText::FromString(TEXT("|")))
 					.Font(FCoreStyle::GetDefaultFontStyle(FName("Regular"), 10))
 				]
 				+SHorizontalBox::Slot()
