Fix for hash collision with Niagara's new compilation mode's traversal context

diff --git a/Engine/Plugins/FX/Niagara/Source/NiagaraEditor/Private/NiagaraTraversalStateContext.cpp b/Engine/Plugins/FX/Niagara/Source/NiagaraEditor/Private/NiagaraTraversalStateContext.cpp
--- a/Engine/Plugins/FX/Niagara/Source/NiagaraEditor/Private/NiagaraTraversalStateContext.cpp
+++ b/Engine/Plugins/FX/Niagara/Source/NiagaraEditor/Private/NiagaraTraversalStateContext.cpp
@@ -60,7 +60,7 @@ void FNiagaraTraversalStateContext::PushGraphInternal(const FNiagaraCompilationN
 
 			if (IsValueSet)
 			{
-				const uint32 SwitchNodeHash = HashCombine(TraversalStack.Top().FullStackHash, GetTypeHash(StaticSwitchNode->NodeGuid));
+				const FGuid SwitchNodeHash = FGuid::Combine(TraversalStack.Top().FullStackHash, StaticSwitchNode->NodeGuid);
 				if (ensure(!StaticSwitchValueMap.Contains(SwitchNodeHash)))
 				{
 					StaticSwitchValueMap.Add(SwitchNodeHash, SwitchValue);
@@ -75,7 +75,7 @@ void FNiagaraTraversalStateContext::PushGraphInternal(const FNiagaraCompilationN
 		}
 		else if (const FNiagaraCompilationNodeFunctionCall* InnerFunctionNode = Node->AsType<FNiagaraCompilationNodeFunctionCall>())
 		{
-			const uint32 InnerFunctionNodeHash = HashCombine(TraversalStack.Top().FullStackHash, GetTypeHash(InnerFunctionNode->NodeGuid));
+			const FGuid InnerFunctionNodeHash = FGuid::Combine(TraversalStack.Top().FullStackHash, InnerFunctionNode->NodeGuid);
 
 			// based on the original code bInheritDebugState drives whether we use the system value vs NoDebug.  Note
 			// that the serialized Debugstate is never used.
@@ -106,12 +106,12 @@ void FNiagaraTraversalStateContext::PushGraphInternal(const FNiagaraCompilationN
 
 						if (CallerInputPin)
 						{
-							const uint32 FunctionPinNodeHash = HashCombine(InnerFunctionNodeHash, GetTypeHash(ValuePin->PinName));
+							const FFunctionDefaultValueMapKey DefaultValueKey = MakeTuple(InnerFunctionNodeHash, ValuePin->PinName);
 
-							const FString* ExistingDefaultValue = FunctionDefaultValueMap.Find(FunctionPinNodeHash);
+							const FString* ExistingDefaultValue = FunctionDefaultValueMap.Find(DefaultValueKey);
 							if (ensure(!ExistingDefaultValue))
 							{
-								FunctionDefaultValueMap.Add(FunctionPinNodeHash, CallerInputPin->DefaultValue);
+								FunctionDefaultValueMap.Add(DefaultValueKey, CallerInputPin->DefaultValue);
 							}
 							else
 							{
@@ -132,10 +132,10 @@ void FNiagaraTraversalStateContext::PushGraphInternal(const FNiagaraCompilationN
 								UE_LOG(LogNiagaraEditor, Warning, TEXT("FNiagaraTraversalStateContext - Stack"))
 								for (int32 StackIt = TraversalStack.Num() - 1; StackIt >= 0; --StackIt)
 								{
-									FString StackMessage = FString::Printf(TEXT("[%d] - %s, %xd"),
+									FString StackMessage = FString::Printf(TEXT("[%d] - %s, %s"),
 										StackIt,
 										*TraversalStack[StackIt].NodeGuid.ToString(EGuidFormats::DigitsWithHyphens),
-										TraversalStack[StackIt].FullStackHash);
+										*TraversalStack[StackIt].FullStackHash.ToString(EGuidFormats::DigitsWithHyphens));
 
 #if WITH_NIAGARA_TRAVERSAL_FRIENDLY_NAME
 									StackMessage.Append(TEXT(", "));
@@ -167,7 +167,7 @@ void FNiagaraTraversalStateContext::BeginContext(const FNiagaraCompilationGraph*
 
 void FNiagaraTraversalStateContext::PushFunction(const FNiagaraCompilationNodeFunctionCall* FunctionCall, const FNiagaraFixedConstantResolver& ConstantResolver)
 {
-	TOptional<uint32> CurrentStackHash;
+	TOptional<FGuid> CurrentStackHash;
 	if (!TraversalStack.IsEmpty())
 	{
 		CurrentStackHash = TraversalStack.Top().FullStackHash;
@@ -183,10 +183,9 @@ void FNiagaraTraversalStateContext::PushFunction(const FNiagaraCompilationNodeFu
 #endif
 
 	StackTop.NodeGuid = FunctionCall->NodeGuid;
-	const uint32 NodeGuidHash = GetTypeHash(StackTop.NodeGuid);
 	StackTop.FullStackHash = CurrentStackHash.IsSet()
-		? HashCombine(*CurrentStackHash, NodeGuidHash)
-		: NodeGuidHash;
+		? FGuid::Combine(*CurrentStackHash, StackTop.NodeGuid)
+		: StackTop.NodeGuid;
 
 	if (FunctionCall->CalledGraph)
 	{
@@ -196,7 +195,7 @@ void FNiagaraTraversalStateContext::PushFunction(const FNiagaraCompilationNodeFu
 
 void FNiagaraTraversalStateContext::PushEmitter(const FNiagaraCompilationNodeEmitter* Emitter)
 {
-	TOptional<uint32> CurrentStackHash;
+	TOptional<FGuid> CurrentStackHash;
 	if (!TraversalStack.IsEmpty())
 	{
 		CurrentStackHash = TraversalStack.Top().FullStackHash;
@@ -212,10 +211,9 @@ void FNiagaraTraversalStateContext::PushEmitter(const FNiagaraCompilationNodeEmi
 #endif
 
 	StackTop.NodeGuid = Emitter->NodeGuid;
-	const uint32 NodeGuidHash = GetTypeHash(StackTop.NodeGuid);
 	StackTop.FullStackHash = CurrentStackHash.IsSet()
-		? HashCombine(*CurrentStackHash, NodeGuidHash)
-		: NodeGuidHash;
+		? FGuid::Combine(*CurrentStackHash, StackTop.NodeGuid)
+		: StackTop.NodeGuid;
 }
 
 void FNiagaraTraversalStateContext::PopFunction(const FNiagaraCompilationNodeFunctionCall* FunctionCall)
@@ -234,7 +232,7 @@ bool FNiagaraTraversalStateContext::GetStaticSwitchValue(const FGuid& NodeGuid,
 {
 	if (!TraversalStack.IsEmpty())
 	{
-		if (const int32* ValuePtr = StaticSwitchValueMap.Find(HashCombine(TraversalStack.Top().FullStackHash, GetTypeHash(NodeGuid))))
+		if (const int32* ValuePtr = StaticSwitchValueMap.Find(FGuid::Combine(TraversalStack.Top().FullStackHash, NodeGuid)))
 		{
 			StaticSwitchValue = *ValuePtr;
 			return true;
@@ -247,9 +245,8 @@ bool FNiagaraTraversalStateContext::GetFunctionDefaultValue(const FGuid& NodeGui
 {
 	if (!TraversalStack.IsEmpty())
 	{
-		const uint32 NodeHash = HashCombine(TraversalStack.Top().FullStackHash, GetTypeHash(NodeGuid));
-		const uint32 PinHash = HashCombine(NodeHash, GetTypeHash(PinName));
-		if (const FString* ValuePtr = FunctionDefaultValueMap.Find(PinHash))
+		const FGuid StackGuid = FGuid::Combine(TraversalStack.Top().FullStackHash, NodeGuid);
+		if (const FString* ValuePtr = FunctionDefaultValueMap.Find(MakeTuple(StackGuid, PinName)))
 		{
 			FunctionDefaultValue = *ValuePtr;
 			return true;
@@ -262,7 +259,7 @@ bool FNiagaraTraversalStateContext::GetFunctionDebugState(const FGuid& NodeGuid,
 {
 	if (!TraversalStack.IsEmpty())
 	{
-		const uint32 NodeHash = HashCombine(TraversalStack.Top().FullStackHash, GetTypeHash(NodeGuid));
+		const FGuid NodeHash = FGuid::Combine(TraversalStack.Top().FullStackHash, NodeGuid);
 		if (const ENiagaraFunctionDebugState* ValuePtr = FunctionDebugStateMap.Find(NodeHash))
 		{
 			DebugState = *ValuePtr;
diff --git a/Engine/Plugins/FX/Niagara/Source/NiagaraEditor/Private/NiagaraTraversalStateContext.h b/Engine/Plugins/FX/Niagara/Source/NiagaraEditor/Private/NiagaraTraversalStateContext.h
--- a/Engine/Plugins/FX/Niagara/Source/NiagaraEditor/Private/NiagaraTraversalStateContext.h
+++ b/Engine/Plugins/FX/Niagara/Source/NiagaraEditor/Private/NiagaraTraversalStateContext.h
@@ -10,12 +10,12 @@ class FNiagaraCompilationNodeEmitter;
 class FNiagaraCompilationNodeFunctionCall;
 class FNiagaraFixedConstantResolver;
 
-#define WITH_NIAGARA_TRAVERSAL_FRIENDLY_NAME (1)
+#define WITH_NIAGARA_TRAVERSAL_FRIENDLY_NAME (UE_BUILD_DEBUG)
 
 struct FNiagaraTraversalStackEntry
 {
 	FGuid NodeGuid;
-	uint32 FullStackHash = 0;
+	FGuid FullStackHash;
 #if WITH_NIAGARA_TRAVERSAL_FRIENDLY_NAME
 	FString FriendlyName;
 #endif
@@ -39,9 +39,11 @@ struct FNiagaraTraversalStateContext
 
 	TArray<FNiagaraTraversalStackEntry> TraversalStack;
 
-	TMap<uint32, int32> StaticSwitchValueMap;
-	TMap<uint32, FString> FunctionDefaultValueMap;
-	TMap<uint32, ENiagaraFunctionDebugState> FunctionDebugStateMap;
+	using FFunctionDefaultValueMapKey = TTuple<FGuid /*Traversal stack guid*/, FName /*PinName*/>;
+
+	TMap<FGuid, int32> StaticSwitchValueMap;
+	TMap<FFunctionDefaultValueMapKey, FString> FunctionDefaultValueMap;
+	TMap<FGuid, ENiagaraFunctionDebugState> FunctionDebugStateMap;
 
 protected:
 	void PushGraphInternal(const FNiagaraCompilationNode* CallingNode, const FNiagaraCompilationGraph* Graph, const FNiagaraFixedConstantResolver& ConstantResolver);
