Simple fix for GPU emitter's still working (while complaining about compilation) when event scripts are present

diff --git a/Engine/Plugins/FX/Niagara/Source/Niagara/Private/NiagaraEmitter.cpp b/Engine/Plugins/FX/Niagara/Source/Niagara/Private/NiagaraEmitter.cpp
--- a/Engine/Plugins/FX/Niagara/Source/Niagara/Private/NiagaraEmitter.cpp
+++ b/Engine/Plugins/FX/Niagara/Source/Niagara/Private/NiagaraEmitter.cpp
@@ -1520,6 +1520,12 @@ bool FVersionedNiagaraEmitterData::IsReadyToRunInternal() const
 
 	if (SimTarget == ENiagaraSimTarget::GPUComputeSim)
 	{
+		// event scripts on GPU emitters are not currently supported
+		if (!EventHandlerScriptProps.IsEmpty())
+		{
+			return false;
+		}
+
 		if (GPUComputeScript->IsScriptCompilationPending(true))
 		{
 			return false;
diff --git a/Engine/Plugins/FX/Niagara/Source/NiagaraEditor/Private/Widgets/AssetBrowser/NiagaraAssetBrowserPreview.cpp b/Engine/Plugins/FX/Niagara/Source/NiagaraEditor/Private/Widgets/AssetBrowser/NiagaraAssetBrowserPreview.cpp
--- a/Engine/Plugins/FX/Niagara/Source/NiagaraEditor/Private/Widgets/AssetBrowser/NiagaraAssetBrowserPreview.cpp
+++ b/Engine/Plugins/FX/Niagara/Source/NiagaraEditor/Private/Widgets/AssetBrowser/NiagaraAssetBrowserPreview.cpp
@@ -167,7 +167,7 @@ void FNiagaraAssetPreviewViewportClient::Draw(FViewport* InViewport, FCanvas* Ca
 	{
 		if(UNiagaraSystem* NiagaraSystem = Component->GetAsset())
 		{
-			const bool bSystemCompiling = NiagaraSystem->HasOutstandingCompilationRequests();
+			const bool bHasOutstandingCompilation = NiagaraSystem->HasOutstandingCompilationRequests();
 		
 			for (const FNiagaraEmitterHandle& EmitterHandle : NiagaraSystem->GetEmitterHandles())
 			{
@@ -177,7 +177,8 @@ void FNiagaraAssetPreviewViewportClient::Draw(FViewport* InViewport, FCanvas* Ca
 					continue;
 				}
 		
-				const bool bEmitterCompiling = bSystemCompiling || !EmitterData->IsReadyToRun();
+				// we assume that if the emitter isn't ready to run and we have compilations still oustanding that we're waiting on that
+				const bool bEmitterCompiling = bHasOutstandingCompilation && !EmitterData->IsReadyToRun();
 				if (!bEmitterCompiling)
 				{
 					continue;
diff --git a/Engine/Plugins/FX/Niagara/Source/NiagaraEditor/Private/Widgets/SNiagaraSystemViewport.cpp b/Engine/Plugins/FX/Niagara/Source/NiagaraEditor/Private/Widgets/SNiagaraSystemViewport.cpp
--- a/Engine/Plugins/FX/Niagara/Source/NiagaraEditor/Private/Widgets/SNiagaraSystemViewport.cpp
+++ b/Engine/Plugins/FX/Niagara/Source/NiagaraEditor/Private/Widgets/SNiagaraSystemViewport.cpp
@@ -411,7 +411,7 @@ void FNiagaraSystemViewportClient::DrawParticleCounts(UNiagaraComponent* Compone
 	}
 	else if ( UNiagaraSystem* NiagaraSystem = Component->GetAsset() )
 	{
-		const bool bSystemCompiling = NiagaraSystem->HasOutstandingCompilationRequests();
+		const bool bHasOutstandingCompilation = NiagaraSystem->HasOutstandingCompilationRequests();
 
 		for (const FNiagaraEmitterHandle& EmitterHandle : NiagaraSystem->GetEmitterHandles())
 		{
@@ -421,7 +421,8 @@ void FNiagaraSystemViewportClient::DrawParticleCounts(UNiagaraComponent* Compone
 				continue;
 			}
 
-			const bool bEmitterCompiling = bSystemCompiling || !EmitterData->IsReadyToRun();
+			// we assume that if the emitter isn't ready to run and we have compilations still oustanding that we're waiting on that
+			const bool bEmitterCompiling = bHasOutstandingCompilation && !EmitterData->IsReadyToRun();
 			if (!bEmitterCompiling)
 			{
 				continue;
