Clean up handling of PerformanceMode within the niagara toolkit to behave better when toggling states which would enable/disable it (the EVVM)
-toggle button is now hidden when the current system is unsupported
-UI won't crash unexpectedly when you muck around with enabling/disabling things

diff --git a/Engine/Plugins/FX/Niagara/Source/Niagara/Private/NiagaraScript.cpp b/Engine/Plugins/FX/Niagara/Source/Niagara/Private/NiagaraScript.cpp
--- a/Engine/Plugins/FX/Niagara/Source/Niagara/Private/NiagaraScript.cpp
+++ b/Engine/Plugins/FX/Niagara/Source/Niagara/Private/NiagaraScript.cpp
@@ -2525,15 +2525,21 @@ void UNiagaraScript::GenerateStatIDs()
 #elif ENABLE_STATNAMEDEVENTS
 	StatNamedEvents.Empty();
 
-	static const IConsoleVariable* CVarOptimizeVMDetailedStats = IConsoleManager::Get().FindConsoleVariable(TEXT("vm.DetailedVMScriptStats"));
-	if (CVarOptimizeVMDetailedStats && CVarOptimizeVMDetailedStats->GetInt() != 0)
+	if (UNiagaraSystem* RootSystem = FindRootSystem())
 	{
-		if (IsReadyToRun(ENiagaraSimTarget::CPUSim))
+		if (RootSystem->SupportsStatScopedPerformanceMode())
 		{
-			StatNamedEvents.Reserve(CachedScriptVM.StatScopes.Num());
-			for (FNiagaraStatScope& StatScope : CachedScriptVM.StatScopes)
+			static const IConsoleVariable* CVarOptimizeVMDetailedStats = IConsoleManager::Get().FindConsoleVariable(TEXT("vm.DetailedVMScriptStats"));
+			if (CVarOptimizeVMDetailedStats && CVarOptimizeVMDetailedStats->GetInt() != 0)
 			{
-				StatNamedEvents.Add(StatScope.FriendlyName.ToString());
+				if (IsReadyToRun(ENiagaraSimTarget::CPUSim))
+				{
+					StatNamedEvents.Reserve(CachedScriptVM.StatScopes.Num());
+					for (FNiagaraStatScope& StatScope : CachedScriptVM.StatScopes)
+					{
+						StatNamedEvents.Add(StatScope.FriendlyName.ToString());
+					}
+				}
 			}
 		}
 	}
diff --git a/Engine/Plugins/FX/Niagara/Source/NiagaraEditor/Private/Toolkits/NiagaraSystemToolkit.cpp b/Engine/Plugins/FX/Niagara/Source/NiagaraEditor/Private/Toolkits/NiagaraSystemToolkit.cpp
--- a/Engine/Plugins/FX/Niagara/Source/NiagaraEditor/Private/Toolkits/NiagaraSystemToolkit.cpp
+++ b/Engine/Plugins/FX/Niagara/Source/NiagaraEditor/Private/Toolkits/NiagaraSystemToolkit.cpp
@@ -426,11 +426,9 @@ void FNiagaraSystemToolkit::SetupCommands()
 	GetToolkitCommands()->MapAction(
 		FNiagaraEditorCommands::Get().ToggleStatPerformance,
 		FExecuteAction::CreateSP(this, &FNiagaraSystemToolkit::ToggleStatPerformance),
-		FCanExecuteAction::CreateLambda([this]()
-		{
-			return System && System->SupportsStatScopedPerformanceMode();
-		}),
-		FIsActionChecked::CreateSP(this, &FNiagaraSystemToolkit::IsStatPerformanceChecked));
+		FCanExecuteAction(),
+		FIsActionChecked::CreateSP(this, &FNiagaraSystemToolkit::IsStatPerformanceChecked),
+		FIsActionButtonVisible::CreateSP(this, &FNiagaraSystemToolkit::SupportsStatPerformance));
 	GetToolkitCommands()->MapAction(
         FNiagaraEditorCommands::Get().ClearStatPerformance,
         FExecuteAction::CreateSP(this, &FNiagaraSystemToolkit::ClearStatPerformance));
@@ -1007,10 +1005,20 @@ void FNiagaraSystemToolkit::OnToggleBoundsSetFixedBounds_System()
 	SystemViewModel->UpdateSystemFixedBounds();
 }
 
+bool FNiagaraSystemToolkit::SupportsStatPerformance() const
+{
+#if STATS
+	return SystemViewModel->GetSystem().SupportsStatScopedPerformanceMode();
+#else
+	return false;
+#endif
+}
+
+
 void FNiagaraSystemToolkit::ClearStatPerformance()
 {
 #if STATS
-	SystemViewModel->GetSystem().GetStatData().ClearStatCaptures();
+	SystemViewModel->ClearSystemStats();
 	SystemViewModel->ClearEmitterStats();
 #endif
 }
diff --git a/Engine/Plugins/FX/Niagara/Source/NiagaraEditor/Private/Toolkits/NiagaraSystemToolkit.h b/Engine/Plugins/FX/Niagara/Source/NiagaraEditor/Private/Toolkits/NiagaraSystemToolkit.h
--- a/Engine/Plugins/FX/Niagara/Source/NiagaraEditor/Private/Toolkits/NiagaraSystemToolkit.h
+++ b/Engine/Plugins/FX/Niagara/Source/NiagaraEditor/Private/Toolkits/NiagaraSystemToolkit.h
@@ -141,6 +141,7 @@ protected:
 	void OnToggleBoundsSetFixedBounds_Emitters();
 	void OnToggleBoundsSetFixedBounds_System();
 
+	bool SupportsStatPerformance() const;
 	void ClearStatPerformance();
 	void ToggleStatPerformance();
 	bool IsStatPerformanceChecked();
diff --git a/Engine/Plugins/FX/Niagara/Source/NiagaraEditor/Private/ViewModels/NiagaraSystemViewModel.cpp b/Engine/Plugins/FX/Niagara/Source/NiagaraEditor/Private/ViewModels/NiagaraSystemViewModel.cpp
--- a/Engine/Plugins/FX/Niagara/Source/NiagaraEditor/Private/ViewModels/NiagaraSystemViewModel.cpp
+++ b/Engine/Plugins/FX/Niagara/Source/NiagaraEditor/Private/ViewModels/NiagaraSystemViewModel.cpp
@@ -1928,6 +1928,8 @@ void FNiagaraSystemViewModel::RefreshAll()
 	ResetSystem(ETimeResetMode::AllowResetTime, EMultiResetMode::AllowResetAllInstances, EReinitMode::ReinitializeSystem);
 	RefreshEmitterHandleViewModels();
 	RefreshSequencerTracks();
+	ClearSystemStats();
+	ClearEmitterStats();
 	InvalidateCachedCompileStatus();
 	ScriptScratchPadViewModel->RefreshScriptViewModels();
 	CurveSelectionViewModel->Refresh();
@@ -3196,6 +3198,18 @@ void FNiagaraSystemViewModel::UpdateSystemFixedBounds()
 	}
 }
 
+bool FNiagaraSystemViewModel::SupportsPerformanceMode() const
+{
+	return System->SupportsStatScopedPerformanceMode();
+}
+
+void FNiagaraSystemViewModel::ClearSystemStats()
+{
+#if STATS
+	System->GetStatData().ClearStatCaptures();
+#endif
+}
+
 void FNiagaraSystemViewModel::ClearEmitterStats()
 {
 #if STATS
diff --git a/Engine/Plugins/FX/Niagara/Source/NiagaraEditor/Public/ViewModels/NiagaraSystemViewModel.h b/Engine/Plugins/FX/Niagara/Source/NiagaraEditor/Public/ViewModels/NiagaraSystemViewModel.h
--- a/Engine/Plugins/FX/Niagara/Source/NiagaraEditor/Public/ViewModels/NiagaraSystemViewModel.h
+++ b/Engine/Plugins/FX/Niagara/Source/NiagaraEditor/Public/ViewModels/NiagaraSystemViewModel.h
@@ -339,6 +339,12 @@ public:
 	/** Updates the current system's fixed bounds with its current dynamic bounds. */
 	void UpdateSystemFixedBounds();
 
+	/** Returns whether the current system is in a state that supports PerformanceMode/stats collection. */
+	NIAGARAEDITOR_API bool SupportsPerformanceMode() const;
+
+	/** Clear the captures stats for the current system. */
+	void ClearSystemStats();
+
 	/** Clear the captures stats for all the emitters in the current system. */
 	void ClearEmitterStats();
 
diff --git a/Engine/Plugins/FX/Niagara/Source/NiagaraEditorWidgets/Private/Stack/SNiagaraStackRowPerfWidget.cpp b/Engine/Plugins/FX/Niagara/Source/NiagaraEditorWidgets/Private/Stack/SNiagaraStackRowPerfWidget.cpp
--- a/Engine/Plugins/FX/Niagara/Source/NiagaraEditorWidgets/Private/Stack/SNiagaraStackRowPerfWidget.cpp
+++ b/Engine/Plugins/FX/Niagara/Source/NiagaraEditorWidgets/Private/Stack/SNiagaraStackRowPerfWidget.cpp
@@ -227,7 +227,12 @@ FLinearColor SNiagaraStackRowPerfWidget::GetPlaceholderBrushColor() const
 bool SNiagaraStackRowPerfWidget::HasPerformanceData() const
 {
 #if STATS
-	bool IsPerfCaptureEnabled = StatEnabledVar && StatEnabledVar->GetBool();
+	if (IsFinalized())
+	{
+		return false;
+	}
+
+	bool IsPerfCaptureEnabled = StatEnabledVar && StatEnabledVar->GetBool() && StackEntry->GetSystemViewModel()->SupportsPerformanceMode();
 	if (IsPerfCaptureEnabled && (IsSystemStack() || IsEmitterStack()))
 	{
 		return (IsGroupHeaderEntry() || IsModuleEntry()) && (StackEntry->GetExecutionSubcategoryName() != UNiagaraStackEntry::FExecutionSubcategoryNames::Settings);
@@ -262,8 +267,19 @@ EVisibility SNiagaraStackRowPerfWidget::IsVisible() const
 	return HasPerformanceData() ? EVisibility::Visible : EVisibility::Collapsed;
 }
 
+bool SNiagaraStackRowPerfWidget::IsFinalized() const
+{
+	return !StackEntry.IsValid() || StackEntry->IsFinalized();
+}
+
 FText SNiagaraStackRowPerfWidget::GetPerformanceDisplayText() const
 {
+	// guard against the stack entry having been finalized
+	if (IsFinalized())
+	{
+		return FText::FromString("N/A");
+	}
+
 	FNumberFormattingOptions TimeFormatOptions;
 	TimeFormatOptions.MinimumIntegralDigits = 1;
 	TimeFormatOptions.MinimumFractionalDigits = 2;
@@ -310,6 +326,11 @@ FText SNiagaraStackRowPerfWidget::GetEvalTypeDisplayText() const
 
 FSlateColor SNiagaraStackRowPerfWidget::GetPerformanceDisplayTextColor() const
 {
+	// guard against the stack entry having been finalized
+	if (IsFinalized())
+	{
+		return FSlateColor(FLinearColor::Black);
+	}
 	if (IsEntrySelected())
 	{
 		return FSlateColor(FLinearColor::Black);
diff --git a/Engine/Plugins/FX/Niagara/Source/NiagaraEditorWidgets/Private/Stack/SNiagaraStackRowPerfWidget.h b/Engine/Plugins/FX/Niagara/Source/NiagaraEditorWidgets/Private/Stack/SNiagaraStackRowPerfWidget.h
--- a/Engine/Plugins/FX/Niagara/Source/NiagaraEditorWidgets/Private/Stack/SNiagaraStackRowPerfWidget.h
+++ b/Engine/Plugins/FX/Niagara/Source/NiagaraEditorWidgets/Private/Stack/SNiagaraStackRowPerfWidget.h
@@ -45,6 +45,7 @@ private:
 	ENiagaraStatDisplayMode GetDisplayMode() const;
 	bool IsInterpolatedSpawnEnabled() const;
 	bool IsGpuEmitter() const;
+	bool IsFinalized() const;
 	
 	float CalculateGroupOverallTime(FString StatScopeName) const;
 	float CalculateStackEntryTime() const;
