used TypedOuter instead of GetRootSystem for CIS failure with non-editor builds

diff --git a/Engine/Plugins/FX/Niagara/Source/Niagara/Private/NiagaraScript.cpp b/Engine/Plugins/FX/Niagara/Source/Niagara/Private/NiagaraScript.cpp
--- a/Engine/Plugins/FX/Niagara/Source/Niagara/Private/NiagaraScript.cpp
+++ b/Engine/Plugins/FX/Niagara/Source/Niagara/Private/NiagaraScript.cpp
@@ -2525,9 +2525,9 @@ void UNiagaraScript::GenerateStatIDs()
 #elif ENABLE_STATNAMEDEVENTS
 	StatNamedEvents.Empty();
 
-	if (UNiagaraSystem* RootSystem = FindRootSystem())
+	if (UNiagaraSystem* OuterSystem = GetTypedOuter<UNiagaraSystem>())
 	{
-		if (RootSystem->SupportsStatScopedPerformanceMode())
+		if (OuterSystem->SupportsStatScopedPerformanceMode())
 		{
 			static const IConsoleVariable* CVarOptimizeVMDetailedStats = IConsoleManager::Get().FindConsoleVariable(TEXT("vm.DetailedVMScriptStats"));
 			if (CVarOptimizeVMDetailedStats && CVarOptimizeVMDetailedStats->GetInt() != 0)
