Take 2 on CIS error with handling statnamedevents/editor types

diff --git a/Engine/Plugins/FX/Niagara/Source/Niagara/Private/NiagaraScript.cpp b/Engine/Plugins/FX/Niagara/Source/Niagara/Private/NiagaraScript.cpp
--- a/Engine/Plugins/FX/Niagara/Source/Niagara/Private/NiagaraScript.cpp
+++ b/Engine/Plugins/FX/Niagara/Source/Niagara/Private/NiagaraScript.cpp
@@ -2525,20 +2525,17 @@ void UNiagaraScript::GenerateStatIDs()
 #elif ENABLE_STATNAMEDEVENTS
 	StatNamedEvents.Empty();
 
-	if (UNiagaraSystem* OuterSystem = GetTypedOuter<UNiagaraSystem>())
+	if (!CachedScriptVM.SupportsExperimentalVM())
 	{
-		if (OuterSystem->SupportsStatScopedPerformanceMode())
+		static const IConsoleVariable* CVarOptimizeVMDetailedStats = IConsoleManager::Get().FindConsoleVariable(TEXT("vm.DetailedVMScriptStats"));
+		if (CVarOptimizeVMDetailedStats && CVarOptimizeVMDetailedStats->GetInt() != 0)
 		{
-			static const IConsoleVariable* CVarOptimizeVMDetailedStats = IConsoleManager::Get().FindConsoleVariable(TEXT("vm.DetailedVMScriptStats"));
-			if (CVarOptimizeVMDetailedStats && CVarOptimizeVMDetailedStats->GetInt() != 0)
+			if (IsReadyToRun(ENiagaraSimTarget::CPUSim))
 			{
-				if (IsReadyToRun(ENiagaraSimTarget::CPUSim))
+				StatNamedEvents.Reserve(CachedScriptVM.StatScopes.Num());
+				for (FNiagaraStatScope& StatScope : CachedScriptVM.StatScopes)
 				{
-					StatNamedEvents.Reserve(CachedScriptVM.StatScopes.Num());
-					for (FNiagaraStatScope& StatScope : CachedScriptVM.StatScopes)
-					{
-						StatNamedEvents.Add(StatScope.FriendlyName.ToString());
-					}
+					StatNamedEvents.Add(StatScope.FriendlyName.ToString());
 				}
 			}
 		}
