Added Rename rightlick menu & shortcut to scratch pad script UI
- Removed closing & reopening of renamed scratch pad. No longer necessary. Works with undo & redo

diff --git a/Engine/Plugins/FX/Niagara/Source/NiagaraEditor/Private/ViewModels/NiagaraScratchPadViewModel.cpp b/Engine/Plugins/FX/Niagara/Source/NiagaraEditor/Private/ViewModels/NiagaraScratchPadViewModel.cpp
--- a/Engine/Plugins/FX/Niagara/Source/NiagaraEditor/Private/ViewModels/NiagaraScratchPadViewModel.cpp
+++ b/Engine/Plugins/FX/Niagara/Source/NiagaraEditor/Private/ViewModels/NiagaraScratchPadViewModel.cpp
@@ -779,19 +779,8 @@ void UNiagaraScratchPadViewModel::ScriptGraphNodeSelectionChanged(TWeakPtr<FNiag
 
 void UNiagaraScratchPadViewModel::ScriptViewModelScriptRenamed(TWeakPtr<FNiagaraScratchPadScriptViewModel> ScriptViewModelWeak)
 {
-	TSharedPtr<FNiagaraScratchPadScriptViewModel> ScriptViewModel = ScriptViewModelWeak.Pin();
-	if (ScriptViewModel.IsValid())
-	{
-		GetSystemViewModel()->GetDocumentViewModel()->CloseChildScript(ScriptViewModel->GetGraphViewModel()->GetGraph());
-	}
-
 	UpdateChangeId(GetSystemViewModel());
 	OnScriptRenamed().Broadcast();
-
-	if (ScriptViewModel.IsValid())
-	{
-		GetSystemViewModel()->GetDocumentViewModel()->OpenChildScript(ScriptViewModel->GetGraphViewModel()->GetGraph());
-	}
 }
 
 void UNiagaraScratchPadViewModel::ScriptViewModelGraphSelectionChanged(const UObject* Obj, TWeakPtr<FNiagaraScratchPadScriptViewModel> ScriptViewModelWeak)
diff --git a/Engine/Plugins/FX/Niagara/Source/NiagaraEditorWidgets/Private/NiagaraScratchPadCommandContext.cpp b/Engine/Plugins/FX/Niagara/Source/NiagaraEditorWidgets/Private/NiagaraScratchPadCommandContext.cpp
--- a/Engine/Plugins/FX/Niagara/Source/NiagaraEditorWidgets/Private/NiagaraScratchPadCommandContext.cpp
+++ b/Engine/Plugins/FX/Niagara/Source/NiagaraEditorWidgets/Private/NiagaraScratchPadCommandContext.cpp
@@ -44,6 +44,9 @@ void FNiagaraScratchPadCommandContext::AddMenuItems(FMenuBuilder& MenuBuilder)
 {
 	MenuBuilder.BeginSection("ScriptEdit", LOCTEXT("ScriptEditActions", "Edit"));
 	{
+		TAttribute<FText> CanRenameToolTip = TAttribute<FText>::Create(TAttribute<FText>::FGetter::CreateSP(this, &FNiagaraScratchPadCommandContext::GetCanRenameSelectedScriptToolTip));
+		MenuBuilder.AddMenuEntry(FGenericCommands::Get().Rename, NAME_None, TAttribute<FText>(), CanRenameToolTip);
+		
 		TAttribute<FText> CanCutToolTip = TAttribute<FText>::Create(TAttribute<FText>::FGetter::CreateSP(this, &FNiagaraScratchPadCommandContext::GetCanCutSelectedScriptsToolTip));
 		MenuBuilder.AddMenuEntry(FGenericCommands::Get().Cut, NAME_None, TAttribute<FText>(), CanCutToolTip);
 
@@ -70,6 +73,9 @@ void FNiagaraScratchPadCommandContext::AddMenuItems(FMenuBuilder& MenuBuilder)
 
 void FNiagaraScratchPadCommandContext::SetupCommands()
 {
+	Commands->MapAction(FGenericCommands::Get().Rename, FUIAction(
+		FExecuteAction::CreateSP(this, &FNiagaraScratchPadCommandContext::RenameSelectedScript),
+		FCanExecuteAction::CreateSP(this, &FNiagaraScratchPadCommandContext::CanRenameSelectedScript)));
 	Commands->MapAction(FGenericCommands::Get().Cut, FUIAction(
 		FExecuteAction::CreateSP(this, &FNiagaraScratchPadCommandContext::CutSelectedScripts),
 		FCanExecuteAction::CreateSP(this, &FNiagaraScratchPadCommandContext::CanCutSelectedScripts)));
@@ -96,6 +102,21 @@ void FNiagaraScratchPadCommandContext::SetupCommands()
 		FCanExecuteAction::CreateSP(this, &FNiagaraScratchPadCommandContext::CanCreateAssetFromSelectedScript)));
 }
 
+void FNiagaraScratchPadCommandContext::RenameSelectedScript() const
+{
+	ScratchPadViewModel->GetActiveScriptViewModel()->SetIsPendingRename(true);
+}
+
+bool FNiagaraScratchPadCommandContext::CanRenameSelectedScript() const
+{
+	return ScratchPadViewModel->GetActiveScriptViewModel().IsValid();
+}
+
+FText FNiagaraScratchPadCommandContext::GetCanRenameSelectedScriptToolTip() const
+{
+	return LOCTEXT("RenameToolTip", "Rename the selected scratch pad script");
+}
+
 bool FNiagaraScratchPadCommandContext::CanCutSelectedScripts() const
 {
 	return ScratchPadViewModel->GetActiveScriptViewModel().IsValid();
diff --git a/Engine/Plugins/FX/Niagara/Source/NiagaraEditorWidgets/Private/NiagaraScratchPadCommandContext.h b/Engine/Plugins/FX/Niagara/Source/NiagaraEditorWidgets/Private/NiagaraScratchPadCommandContext.h
--- a/Engine/Plugins/FX/Niagara/Source/NiagaraEditorWidgets/Private/NiagaraScratchPadCommandContext.h
+++ b/Engine/Plugins/FX/Niagara/Source/NiagaraEditorWidgets/Private/NiagaraScratchPadCommandContext.h
@@ -23,6 +23,12 @@ public:
 private:
 	void SetupCommands();
 
+	void RenameSelectedScript() const;
+
+	bool CanRenameSelectedScript() const;
+
+	FText GetCanRenameSelectedScriptToolTip() const;
+
 	bool CanCutSelectedScripts() const;
 
 	FText GetCanCutSelectedScriptsToolTip() const;
