- Allow stateless modules to read from the shader parameters on CPU simulation
- This makes it easier to share the same setup between CPU & GPU and you don't have to do conversions that are already done
- It also sets up the concept of generating shader parameter metadata if we go to more dynamic compilation of modules

#rnx
#tests Ran EngineTests and tested a few other systems

diff --git a/Engine/Plugins/FX/Niagara/Source/Niagara/Internal/Stateless/Modules/NiagaraStatelessModule_CameraOffset.h b/Engine/Plugins/FX/Niagara/Source/Niagara/Internal/Stateless/Modules/NiagaraStatelessModule_CameraOffset.h
--- a/Engine/Plugins/FX/Niagara/Source/Niagara/Internal/Stateless/Modules/NiagaraStatelessModule_CameraOffset.h
+++ b/Engine/Plugins/FX/Niagara/Source/Niagara/Internal/Stateless/Modules/NiagaraStatelessModule_CameraOffset.h
@@ -45,6 +45,11 @@ public:
 		}
 	}
 
+	virtual void BuildShaderParameters(FNiagaraStatelessShaderParametersBuilder& ShaderParametersBuilder) const override
+	{
+		ShaderParametersBuilder.AddParameterNestedStruct<FParameters>();
+	}
+
 	virtual void SetShaderParameters(const FNiagaraStatelessSetShaderParameterContext& SetShaderParameterContext) const override
 	{
 		const FModuleBuiltData* ModuleBuiltData = SetShaderParameterContext.ReadBuiltData<FModuleBuiltData>();
diff --git a/Engine/Plugins/FX/Niagara/Source/Niagara/Internal/Stateless/Modules/NiagaraStatelessModule_DynamicMaterialParameters.h b/Engine/Plugins/FX/Niagara/Source/Niagara/Internal/Stateless/Modules/NiagaraStatelessModule_DynamicMaterialParameters.h
--- a/Engine/Plugins/FX/Niagara/Source/Niagara/Internal/Stateless/Modules/NiagaraStatelessModule_DynamicMaterialParameters.h
+++ b/Engine/Plugins/FX/Niagara/Source/Niagara/Internal/Stateless/Modules/NiagaraStatelessModule_DynamicMaterialParameters.h
@@ -194,6 +194,11 @@ public:
 		}
 	}
 
+	virtual void BuildShaderParameters(FNiagaraStatelessShaderParametersBuilder& ShaderParametersBuilder) const override
+	{
+		ShaderParametersBuilder.AddParameterNestedStruct<FParameters>();
+	}
+
 	virtual void SetShaderParameters(const FNiagaraStatelessSetShaderParameterContext& SetShaderParameterContext) const override
 	{
 		const FModuleBuiltData* ModuleBuiltData = SetShaderParameterContext.ReadBuiltData<FModuleBuiltData>();
diff --git a/Engine/Plugins/FX/Niagara/Source/Niagara/Internal/Stateless/Modules/NiagaraStatelessModule_InitialMeshOrientation.h b/Engine/Plugins/FX/Niagara/Source/Niagara/Internal/Stateless/Modules/NiagaraStatelessModule_InitialMeshOrientation.h
--- a/Engine/Plugins/FX/Niagara/Source/Niagara/Internal/Stateless/Modules/NiagaraStatelessModule_InitialMeshOrientation.h
+++ b/Engine/Plugins/FX/Niagara/Source/Niagara/Internal/Stateless/Modules/NiagaraStatelessModule_InitialMeshOrientation.h
@@ -45,6 +45,7 @@ public:
 	FNiagaraDistributionRangeVector3 Rotation = FNiagaraDistributionRangeVector3(FVector3f::ZeroVector);
 
 	virtual void BuildEmitterData(const FNiagaraStatelessEmitterDataBuildContext& BuildContext) const override;
+	virtual void BuildShaderParameters(FNiagaraStatelessShaderParametersBuilder& ShaderParametersBuilder) const override;
 	virtual void SetShaderParameters(const FNiagaraStatelessSetShaderParameterContext& SetShaderParameterContext) const override;
 
 #if WITH_EDITOR
diff --git a/Engine/Plugins/FX/Niagara/Source/Niagara/Internal/Stateless/Modules/NiagaraStatelessModule_InitializeParticle.h b/Engine/Plugins/FX/Niagara/Source/Niagara/Internal/Stateless/Modules/NiagaraStatelessModule_InitializeParticle.h
--- a/Engine/Plugins/FX/Niagara/Source/Niagara/Internal/Stateless/Modules/NiagaraStatelessModule_InitializeParticle.h
+++ b/Engine/Plugins/FX/Niagara/Source/Niagara/Internal/Stateless/Modules/NiagaraStatelessModule_InitializeParticle.h
@@ -111,6 +111,11 @@ public:
 		PhysicsBuildData.MassRange			= MassDistribution.CalculateRange(FNiagaraStatelessGlobals::GetDefaultMassValue());
 	}
 
+	virtual void BuildShaderParameters(FNiagaraStatelessShaderParametersBuilder& ShaderParametersBuilder) const override
+	{
+		ShaderParametersBuilder.AddParameterNestedStruct<FParameters>();
+	}
+
 	virtual void SetShaderParameters(const FNiagaraStatelessSetShaderParameterContext& SetShaderParameterContext) const override
 	{
 		FParameters* Parameters = SetShaderParameterContext.GetParameterNestedStruct<FParameters>();
diff --git a/Engine/Plugins/FX/Niagara/Source/Niagara/Internal/Stateless/Modules/NiagaraStatelessModule_MeshIndex.h b/Engine/Plugins/FX/Niagara/Source/Niagara/Internal/Stateless/Modules/NiagaraStatelessModule_MeshIndex.h
--- a/Engine/Plugins/FX/Niagara/Source/Niagara/Internal/Stateless/Modules/NiagaraStatelessModule_MeshIndex.h
+++ b/Engine/Plugins/FX/Niagara/Source/Niagara/Internal/Stateless/Modules/NiagaraStatelessModule_MeshIndex.h
@@ -24,6 +24,7 @@ public:
 	TArray<float> MeshIndexWeight;
 
 	virtual void BuildEmitterData(const FNiagaraStatelessEmitterDataBuildContext& BuildContext) const override;
+	virtual void BuildShaderParameters(FNiagaraStatelessShaderParametersBuilder& ShaderParametersBuilder) const override;
 	virtual void SetShaderParameters(const FNiagaraStatelessSetShaderParameterContext& SetShaderParameterContext) const override;
 
 	UFUNCTION()
diff --git a/Engine/Plugins/FX/Niagara/Source/Niagara/Internal/Stateless/Modules/NiagaraStatelessModule_MeshRotationRate.h b/Engine/Plugins/FX/Niagara/Source/Niagara/Internal/Stateless/Modules/NiagaraStatelessModule_MeshRotationRate.h
--- a/Engine/Plugins/FX/Niagara/Source/Niagara/Internal/Stateless/Modules/NiagaraStatelessModule_MeshRotationRate.h
+++ b/Engine/Plugins/FX/Niagara/Source/Niagara/Internal/Stateless/Modules/NiagaraStatelessModule_MeshRotationRate.h
@@ -51,6 +51,11 @@ public:
 		}
 	}
 
+	virtual void BuildShaderParameters(FNiagaraStatelessShaderParametersBuilder& ShaderParametersBuilder) const override
+	{
+		ShaderParametersBuilder.AddParameterNestedStruct<FParameters>();
+	}
+
 	virtual void SetShaderParameters(const FNiagaraStatelessSetShaderParameterContext& SetShaderParameterContext) const override
 	{
 		FParameters* Parameters = SetShaderParameterContext.GetParameterNestedStruct<FParameters>();
diff --git a/Engine/Plugins/FX/Niagara/Source/Niagara/Internal/Stateless/Modules/NiagaraStatelessModule_RotateAroundPoint.h b/Engine/Plugins/FX/Niagara/Source/Niagara/Internal/Stateless/Modules/NiagaraStatelessModule_RotateAroundPoint.h
--- a/Engine/Plugins/FX/Niagara/Source/Niagara/Internal/Stateless/Modules/NiagaraStatelessModule_RotateAroundPoint.h
+++ b/Engine/Plugins/FX/Niagara/Source/Niagara/Internal/Stateless/Modules/NiagaraStatelessModule_RotateAroundPoint.h
@@ -33,6 +33,11 @@ public:
 	//-TODO: Add support for GPU once we settle on a feature set for this module
 	virtual ENiagaraStatelessFeatureMask GetFeatureMask() const override { return ENiagaraStatelessFeatureMask::ExecuteGPU; }
 
+	virtual void BuildShaderParameters(FNiagaraStatelessShaderParametersBuilder& ShaderParametersBuilder) const override
+	{
+		ShaderParametersBuilder.AddParameterNestedStruct<FParameters>();
+	}
+
 	virtual void SetShaderParameters(const FNiagaraStatelessSetShaderParameterContext& SetShaderParameterContext) const override
 	{
 		FParameters* Parameters = SetShaderParameterContext.GetParameterNestedStruct<FParameters>();
diff --git a/Engine/Plugins/FX/Niagara/Source/Niagara/Internal/Stateless/Modules/NiagaraStatelessModule_ScaleColor.h b/Engine/Plugins/FX/Niagara/Source/Niagara/Internal/Stateless/Modules/NiagaraStatelessModule_ScaleColor.h
--- a/Engine/Plugins/FX/Niagara/Source/Niagara/Internal/Stateless/Modules/NiagaraStatelessModule_ScaleColor.h
+++ b/Engine/Plugins/FX/Niagara/Source/Niagara/Internal/Stateless/Modules/NiagaraStatelessModule_ScaleColor.h
@@ -45,6 +45,11 @@ public:
 		}
 	}
 
+	virtual void BuildShaderParameters(FNiagaraStatelessShaderParametersBuilder& ShaderParametersBuilder) const override
+	{
+		ShaderParametersBuilder.AddParameterNestedStruct<FParameters>();
+	}
+
 	virtual void SetShaderParameters(const FNiagaraStatelessSetShaderParameterContext& SetShaderParameterContext) const override
 	{
 		const FModuleBuiltData* ModuleBuiltData = SetShaderParameterContext.ReadBuiltData<FModuleBuiltData>();
diff --git a/Engine/Plugins/FX/Niagara/Source/Niagara/Internal/Stateless/Modules/NiagaraStatelessModule_ScaleMeshSize.h b/Engine/Plugins/FX/Niagara/Source/Niagara/Internal/Stateless/Modules/NiagaraStatelessModule_ScaleMeshSize.h
--- a/Engine/Plugins/FX/Niagara/Source/Niagara/Internal/Stateless/Modules/NiagaraStatelessModule_ScaleMeshSize.h
+++ b/Engine/Plugins/FX/Niagara/Source/Niagara/Internal/Stateless/Modules/NiagaraStatelessModule_ScaleMeshSize.h
@@ -75,6 +75,11 @@ public:
 		BuildContext.AddParticleSimulationExecSimulate(&UNiagaraStatelessModule_ScaleMeshSize::ParticleSimulate);
 	}
 
+	virtual void BuildShaderParameters(FNiagaraStatelessShaderParametersBuilder& ShaderParametersBuilder) const override
+	{
+		ShaderParametersBuilder.AddParameterNestedStruct<FParameters>();
+	}
+
 	virtual void SetShaderParameters(const FNiagaraStatelessSetShaderParameterContext& SetShaderParameterContext) const override
 	{
 		const FModuleBuiltData* ModuleBuiltData = SetShaderParameterContext.ReadBuiltData<FModuleBuiltData>();
diff --git a/Engine/Plugins/FX/Niagara/Source/Niagara/Internal/Stateless/Modules/NiagaraStatelessModule_ScaleMeshSizeBySpeed.h b/Engine/Plugins/FX/Niagara/Source/Niagara/Internal/Stateless/Modules/NiagaraStatelessModule_ScaleMeshSizeBySpeed.h
--- a/Engine/Plugins/FX/Niagara/Source/Niagara/Internal/Stateless/Modules/NiagaraStatelessModule_ScaleMeshSizeBySpeed.h
+++ b/Engine/Plugins/FX/Niagara/Source/Niagara/Internal/Stateless/Modules/NiagaraStatelessModule_ScaleMeshSizeBySpeed.h
@@ -68,6 +68,11 @@ public:
 		BuildContext.AddParticleSimulationExecSimulate(&UNiagaraStatelessModule_ScaleMeshSizeBySpeed::ParticleSimulate);
 	}
 
+	virtual void BuildShaderParameters(FNiagaraStatelessShaderParametersBuilder& ShaderParametersBuilder) const override
+	{
+		ShaderParametersBuilder.AddParameterNestedStruct<FParameters>();
+	}
+
 	virtual void SetShaderParameters(const FNiagaraStatelessSetShaderParameterContext& SetShaderParameterContext) const override
 	{
 		const FModuleBuiltData* ModuleBuiltData = SetShaderParameterContext.ReadBuiltData<FModuleBuiltData>();
diff --git a/Engine/Plugins/FX/Niagara/Source/Niagara/Internal/Stateless/Modules/NiagaraStatelessModule_ScaleSpriteSize.h b/Engine/Plugins/FX/Niagara/Source/Niagara/Internal/Stateless/Modules/NiagaraStatelessModule_ScaleSpriteSize.h
--- a/Engine/Plugins/FX/Niagara/Source/Niagara/Internal/Stateless/Modules/NiagaraStatelessModule_ScaleSpriteSize.h
+++ b/Engine/Plugins/FX/Niagara/Source/Niagara/Internal/Stateless/Modules/NiagaraStatelessModule_ScaleSpriteSize.h
@@ -70,6 +70,11 @@ public:
 		}
 	}
 
+	virtual void BuildShaderParameters(FNiagaraStatelessShaderParametersBuilder& ShaderParametersBuilder) const override
+	{
+		ShaderParametersBuilder.AddParameterNestedStruct<FParameters>();
+	}
+
 	virtual void SetShaderParameters(const FNiagaraStatelessSetShaderParameterContext& SetShaderParameterContext) const override
 	{
 		const FModuleBuiltData* ModuleBuiltData = SetShaderParameterContext.ReadBuiltData<FModuleBuiltData>();
diff --git a/Engine/Plugins/FX/Niagara/Source/Niagara/Internal/Stateless/Modules/NiagaraStatelessModule_ScaleSpriteSizeBySpeed.h b/Engine/Plugins/FX/Niagara/Source/Niagara/Internal/Stateless/Modules/NiagaraStatelessModule_ScaleSpriteSizeBySpeed.h
--- a/Engine/Plugins/FX/Niagara/Source/Niagara/Internal/Stateless/Modules/NiagaraStatelessModule_ScaleSpriteSizeBySpeed.h
+++ b/Engine/Plugins/FX/Niagara/Source/Niagara/Internal/Stateless/Modules/NiagaraStatelessModule_ScaleSpriteSizeBySpeed.h
@@ -68,6 +68,11 @@ public:
 		BuildContext.AddParticleSimulationExecSimulate(&UNiagaraStatelessModule_ScaleSpriteSizeBySpeed::ParticleSimulate);
 	}
 
+	virtual void BuildShaderParameters(FNiagaraStatelessShaderParametersBuilder& ShaderParametersBuilder) const override
+	{
+		ShaderParametersBuilder.AddParameterNestedStruct<FParameters>();
+	}
+
 	virtual void SetShaderParameters(const FNiagaraStatelessSetShaderParameterContext& SetShaderParameterContext) const override
 	{
 		const FModuleBuiltData* ModuleBuiltData = SetShaderParameterContext.ReadBuiltData<FModuleBuiltData>();
diff --git a/Engine/Plugins/FX/Niagara/Source/Niagara/Internal/Stateless/Modules/NiagaraStatelessModule_ShapeLocation.h b/Engine/Plugins/FX/Niagara/Source/Niagara/Internal/Stateless/Modules/NiagaraStatelessModule_ShapeLocation.h
--- a/Engine/Plugins/FX/Niagara/Source/Niagara/Internal/Stateless/Modules/NiagaraStatelessModule_ShapeLocation.h
+++ b/Engine/Plugins/FX/Niagara/Source/Niagara/Internal/Stateless/Modules/NiagaraStatelessModule_ShapeLocation.h
@@ -61,6 +61,7 @@ public:
 	float SphereMax = 100.0f;
 
 	virtual void BuildEmitterData(const FNiagaraStatelessEmitterDataBuildContext& BuildContext) const override;
+	virtual void BuildShaderParameters(FNiagaraStatelessShaderParametersBuilder& ShaderParametersBuilder) const override;
 	virtual void SetShaderParameters(const FNiagaraStatelessSetShaderParameterContext& SetShaderParameterContext) const override;
 
 #if WITH_EDITOR
diff --git a/Engine/Plugins/FX/Niagara/Source/Niagara/Internal/Stateless/Modules/NiagaraStatelessModule_SolveVelocitiesAndForces.h b/Engine/Plugins/FX/Niagara/Source/Niagara/Internal/Stateless/Modules/NiagaraStatelessModule_SolveVelocitiesAndForces.h
--- a/Engine/Plugins/FX/Niagara/Source/Niagara/Internal/Stateless/Modules/NiagaraStatelessModule_SolveVelocitiesAndForces.h
+++ b/Engine/Plugins/FX/Niagara/Source/Niagara/Internal/Stateless/Modules/NiagaraStatelessModule_SolveVelocitiesAndForces.h
@@ -58,6 +58,11 @@ public:
 		}
 	}
 
+	virtual void BuildShaderParameters(FNiagaraStatelessShaderParametersBuilder& ShaderParametersBuilder) const override
+	{
+		ShaderParametersBuilder.AddParameterNestedStruct<FParameters>();
+	}
+
 	virtual void SetShaderParameters(const FNiagaraStatelessSetShaderParameterContext& SetShaderParameterContext) const override
 	{
 		const FModuleBuiltData* ModuleBuiltData = SetShaderParameterContext.ReadBuiltData<FModuleBuiltData>();
@@ -134,34 +139,35 @@ public:
 	static void ParticleSimulate(const NiagaraStateless::FParticleSimulationContext& ParticleSimulationContext)
 	{
 		const FModuleBuiltData* ModuleBuiltData = ParticleSimulationContext.ReadBuiltData<FModuleBuiltData>();
+		const FParameters* Parameters = ParticleSimulationContext.ReadParameterNestedStruct<FParameters>();
 
 		const float* AgeData			= ParticleSimulationContext.GetParticleAge();
 		const float* PreviousAgeData	= ParticleSimulationContext.GetParticlePreviousAge();
 
 		for (uint32 i = 0; i < ParticleSimulationContext.GetNumInstances(); ++i)
 		{
-			const float Mass				= ParticleSimulationContext.RandomScaleBiasFloat(i, 0, ModuleBuiltData->PhysicsData.MassRange);
-			const float Drag				= ParticleSimulationContext.RandomScaleBiasFloat(i, 1, ModuleBuiltData->PhysicsData.DragRange);
-			FVector3f InitialVelocity		= ParticleSimulationContext.RandomScaleBiasFloat(i, 2, ModuleBuiltData->PhysicsData.VelocityRange);
-			const FVector3f Wind			= ParticleSimulationContext.RandomScaleBiasFloat(i, 3, ModuleBuiltData->PhysicsData.WindRange);
-			const FVector3f Acceleration	= ParticleSimulationContext.RandomScaleBiasFloat(i, 4, ModuleBuiltData->PhysicsData.AccelerationRange);
+			const float Mass				= ParticleSimulationContext.RandomScaleBiasFloat(i, 0, Parameters->SolveVelocitiesAndForces_MassScale, Parameters->SolveVelocitiesAndForces_MassBias);
+			const float Drag				= ParticleSimulationContext.RandomScaleBiasFloat(i, 1, Parameters->SolveVelocitiesAndForces_DragScale, Parameters->SolveVelocitiesAndForces_DragBias);
+			FVector3f InitialVelocity		= ParticleSimulationContext.RandomScaleBiasFloat(i, 2, Parameters->SolveVelocitiesAndForces_VelocityScale, Parameters->SolveVelocitiesAndForces_VelocityBias);
+			const FVector3f Wind			= ParticleSimulationContext.RandomScaleBiasFloat(i, 3, Parameters->SolveVelocitiesAndForces_WindScale, Parameters->SolveVelocitiesAndForces_WindBias);
+			const FVector3f Acceleration	= ParticleSimulationContext.RandomScaleBiasFloat(i, 4, Parameters->SolveVelocitiesAndForces_AccelerationScale, Parameters->SolveVelocitiesAndForces_AccelerationBias);
 
 			FVector3f Position			= ParticleSimulationContext.ReadParticleVariable(ModuleBuiltData->PositionVariableOffset, i, FVector3f::ZeroVector);
 			FVector3f PreviousPosition	= ParticleSimulationContext.ReadParticleVariable(ModuleBuiltData->PreviousPositionVariableOffset, i, FVector3f::ZeroVector);
 
 			if (ModuleBuiltData->PhysicsData.bConeVelocity)
 			{
-				const float ConeAngle = ParticleSimulationContext.RandomScaleBiasFloat(i, 5, FNiagaraStatelessRangeFloat(ModuleBuiltData->PhysicsData.ConeInnerAngle, ModuleBuiltData->PhysicsData.ConeOuterAngle)) * (UE_PI / 360.0f);
+				const float ConeAngle = ParticleSimulationContext.RandomScaleBiasFloat(i, 5, Parameters->SolveVelocitiesAndForces_ConeAngleScale, Parameters->SolveVelocitiesAndForces_ConeAngleBias);
 				const float ConeRotation = ParticleSimulationContext.RandomFloat(i, 6) * UE_TWO_PI;
 				const FVector2f scAng = FVector2f(FMath::Sin(ConeAngle), FMath::Cos(ConeAngle));
 				const FVector2f scRot = FVector2f(FMath::Sin(ConeRotation), FMath::Cos(ConeRotation));
 				const FVector3f Direction = FVector3f(scRot.X * scAng.X, scRot.Y * scAng.X, scAng.Y);
 
-				float VelocityScale = ParticleSimulationContext.RandomScaleBiasFloat(i, 7, ModuleBuiltData->PhysicsData.ConeVelocityRange);
+				float VelocityScale = ParticleSimulationContext.RandomScaleBiasFloat(i, 7, Parameters->SolveVelocitiesAndForces_ConeVelocityScale, Parameters->SolveVelocitiesAndForces_ConeVelocityBias);
 				if (ModuleBuiltData->PhysicsData.bConeVelocity)
 				{
-					const float pf = FMath::Pow(FMath::Clamp(scAng.Y, 0.0f, 1.0f), ModuleBuiltData->PhysicsData.ConeVelocityFalloff * 10.0f);
-					VelocityScale *= FMath::Lerp(1.0f, pf, ModuleBuiltData->PhysicsData.ConeVelocityFalloff);
+					const float pf = FMath::Pow(FMath::Clamp(scAng.Y, 0.0f, 1.0f), Parameters->SolveVelocitiesAndForces_ConeVelocityFalloff * 10.0f);
+					VelocityScale *= FMath::Lerp(1.0f, pf, Parameters->SolveVelocitiesAndForces_ConeVelocityFalloff);
 				}
 
 				InitialVelocity += ModuleBuiltData->PhysicsData.ConeQuat.RotateVector(Direction) * VelocityScale;
@@ -170,9 +176,9 @@ public:
 			if (ModuleBuiltData->PhysicsData.bPointVelocity)
 			{
 				const FVector3f FallbackDir	= ParticleSimulationContext.RandomUnitFloat3(i, 8);
-				const FVector3f Delta		= Position - ModuleBuiltData->PhysicsData.PointOrigin;
+				const FVector3f Delta		= Position - Parameters->SolveVelocitiesAndForces_PointOrigin;
 				const FVector3f Direction	= ParticleSimulationContext.SafeNormalize(Delta, FallbackDir);
-				const float		VelocityScale = ParticleSimulationContext.RandomScaleBiasFloat(i, 9, ModuleBuiltData->PhysicsData.PointVelocityRange);
+				const float		VelocityScale = ParticleSimulationContext.RandomScaleBiasFloat(i, 9, Parameters->SolveVelocitiesAndForces_PointVelocityScale, Parameters->SolveVelocitiesAndForces_PointVelocityBias);
 
 				InitialVelocity += Direction * VelocityScale;
 			}
diff --git a/Engine/Plugins/FX/Niagara/Source/Niagara/Internal/Stateless/Modules/NiagaraStatelessModule_SpriteFacingAndAlignment.h b/Engine/Plugins/FX/Niagara/Source/Niagara/Internal/Stateless/Modules/NiagaraStatelessModule_SpriteFacingAndAlignment.h
--- a/Engine/Plugins/FX/Niagara/Source/Niagara/Internal/Stateless/Modules/NiagaraStatelessModule_SpriteFacingAndAlignment.h
+++ b/Engine/Plugins/FX/Niagara/Source/Niagara/Internal/Stateless/Modules/NiagaraStatelessModule_SpriteFacingAndAlignment.h
@@ -72,6 +72,11 @@ public:
 		BuildContext.AddParticleSimulationExecSimulate(&UNiagaraStatelessModule_SpriteFacingAndAlignment::ParticleSimulate);
 	}
 
+	virtual void BuildShaderParameters(FNiagaraStatelessShaderParametersBuilder& ShaderParametersBuilder) const override
+	{
+		ShaderParametersBuilder.AddParameterNestedStruct<FParameters>();
+	}
+
 	virtual void SetShaderParameters(const FNiagaraStatelessSetShaderParameterContext& SetShaderParameterContext) const override
 	{
 		FParameters* Parameters = SetShaderParameterContext.GetParameterNestedStruct<FParameters>();
diff --git a/Engine/Plugins/FX/Niagara/Source/Niagara/Internal/Stateless/Modules/NiagaraStatelessModule_SpriteRotationRate.h b/Engine/Plugins/FX/Niagara/Source/Niagara/Internal/Stateless/Modules/NiagaraStatelessModule_SpriteRotationRate.h
--- a/Engine/Plugins/FX/Niagara/Source/Niagara/Internal/Stateless/Modules/NiagaraStatelessModule_SpriteRotationRate.h
+++ b/Engine/Plugins/FX/Niagara/Source/Niagara/Internal/Stateless/Modules/NiagaraStatelessModule_SpriteRotationRate.h
@@ -49,6 +49,11 @@ public:
 		BuildContext.AddParticleSimulationExecSimulate(&UNiagaraStatelessModule_SpriteRotationRate::ParticleSimulate);
 	}
 
+	virtual void BuildShaderParameters(FNiagaraStatelessShaderParametersBuilder& ShaderParametersBuilder) const override
+	{
+		ShaderParametersBuilder.AddParameterNestedStruct<FParameters>();
+	}
+
 	virtual void SetShaderParameters(const FNiagaraStatelessSetShaderParameterContext& SetShaderParameterContext) const override
 	{
 		FParameters* Parameters = SetShaderParameterContext.GetParameterNestedStruct<FParameters>();
diff --git a/Engine/Plugins/FX/Niagara/Source/Niagara/Internal/Stateless/Modules/NiagaraStatelessModule_SubUVAnimation.h b/Engine/Plugins/FX/Niagara/Source/Niagara/Internal/Stateless/Modules/NiagaraStatelessModule_SubUVAnimation.h
--- a/Engine/Plugins/FX/Niagara/Source/Niagara/Internal/Stateless/Modules/NiagaraStatelessModule_SubUVAnimation.h
+++ b/Engine/Plugins/FX/Niagara/Source/Niagara/Internal/Stateless/Modules/NiagaraStatelessModule_SubUVAnimation.h
@@ -123,6 +123,11 @@ public:
 		BuildContext.AddParticleSimulationExecSimulate(&UNiagaraStatelessModule_SubUVAnimation::ParticleSimulate);
 	}
 
+	virtual void BuildShaderParameters(FNiagaraStatelessShaderParametersBuilder& ShaderParametersBuilder) const override
+	{
+		ShaderParametersBuilder.AddParameterNestedStruct<FParameters>();
+	}
+
 	virtual void SetShaderParameters(const FNiagaraStatelessSetShaderParameterContext& SetShaderParameterContext) const override
 	{
 		FParameters* Parameters = SetShaderParameterContext.GetParameterNestedStruct<FParameters>();
diff --git a/Engine/Plugins/FX/Niagara/Source/Niagara/Internal/Stateless/NiagaraStatelessEmitterDataBuildContext.h b/Engine/Plugins/FX/Niagara/Source/Niagara/Internal/Stateless/NiagaraStatelessEmitterDataBuildContext.h
--- a/Engine/Plugins/FX/Niagara/Source/Niagara/Internal/Stateless/NiagaraStatelessEmitterDataBuildContext.h
+++ b/Engine/Plugins/FX/Niagara/Source/Niagara/Internal/Stateless/NiagaraStatelessEmitterDataBuildContext.h
@@ -30,7 +30,7 @@ public:
 	{
 	}
 
-	void PreModuleBuild();
+	void PreModuleBuild(int32 InShaderParameterOffset);
 
 	uint32 AddStaticData(TConstArrayView<float> FloatData) const;
 	uint32 AddStaticData(TConstArrayView<FVector2f> FloatData) const;
@@ -152,6 +152,7 @@ private:
 	NiagaraStateless::FParticleSimulationExecData*	ParticleExecData = nullptr;
 
 	int32											ModuleBuiltDataOffset = 0;
+	int32											ShaderParameterOffset = 0;
 	int32											RandomSeedOffest = 0;
 
 	struct FTransientObject
diff --git a/Engine/Plugins/FX/Niagara/Source/Niagara/Internal/Stateless/NiagaraStatelessModule.h b/Engine/Plugins/FX/Niagara/Source/Niagara/Internal/Stateless/NiagaraStatelessModule.h
--- a/Engine/Plugins/FX/Niagara/Source/Niagara/Internal/Stateless/NiagaraStatelessModule.h
+++ b/Engine/Plugins/FX/Niagara/Source/Niagara/Internal/Stateless/NiagaraStatelessModule.h
@@ -5,7 +5,8 @@
 #include "NiagaraMergeable.h"
 #include "NiagaraStatelessCommon.h"
 #include "NiagaraStatelessDistribution.h"
-#include "ShaderParameterStruct.h"
+#include "NiagaraStatelessSetShaderParameterContext.h"
+#include "NiagaraStatelessShaderParametersBuilder.h"
 
 #include "NiagaraStatelessModule.generated.h"
 
@@ -14,82 +15,6 @@ class FNiagaraStatelessEmitterDataBuildContext;
 struct FNiagaraStatelessDrawDebugContext;
 #endif
 
-////////////////////////////////////////////////////////////////////////////////////////////////////////////////////////////////////////////////////////////////
-
-class FNiagaraStatelessSetShaderParameterContext
-{
-public:
-	UE_NONCOPYABLE(FNiagaraStatelessSetShaderParameterContext);
-
-	explicit FNiagaraStatelessSetShaderParameterContext(const FNiagaraStatelessSpaceTransforms& InSpaceTransforms, TConstArrayView<uint8> InRendererParameterData, TConstArrayView<uint8> InBuiltData, const FShaderParametersMetadata* InShaderParametersMetadata, uint8* InShaderParameters)
-		: SpaceTransforms(InSpaceTransforms)
-		, RendererParameterData(InRendererParameterData)
-		, BuiltData(InBuiltData)
-		, ShaderParametersBase(InShaderParameters)
-		, ParameterOffset(0)
-		, ShaderParametersMetadata(InShaderParametersMetadata)
-	{
-	}
-
-	const FNiagaraStatelessSpaceTransforms& GetSpaceTransforms() const { return SpaceTransforms; }
-
-	template<typename T>
-	T* GetParameterNestedStruct() const
-	{
-		const uint32 StructOffset = Align(ParameterOffset, TShaderParameterStructTypeInfo<T>::Alignment);
-	#if DO_CHECK
-		ValidateIncludeStructType(StructOffset, TShaderParameterStructTypeInfo<T>::GetStructMetadata());
-	#endif
-		ParameterOffset = StructOffset + TShaderParameterStructTypeInfo<T>::GetStructMetadata()->GetSize();
-		return reinterpret_cast<T*>(ShaderParametersBase + StructOffset);
-	}
-
-	template<typename T>
-	const T* ReadBuiltData() const
-	{
-		const int32 Offset = Align(BuiltDataOffset, alignof(T));
-		BuiltDataOffset = Offset + sizeof(T);
-		check(BuiltDataOffset <= BuiltData.Num());
-		return reinterpret_cast<const T*>(BuiltData.GetData() + Offset);
-	}
-
-	template<typename T>
-	void GetRendererParameterValue(T& OutValue, int32 Offset, const T& DefaultValue) const
-	{
-		if (Offset != INDEX_NONE)
-		{
-			Offset *= sizeof(uint32);
-			check(Offset >= 0 && Offset + sizeof(T) <= RendererParameterData.Num());
-			FMemory::Memcpy(&OutValue, RendererParameterData.GetData() + Offset, sizeof(T));
-		}
-		else
-		{
-			OutValue = DefaultValue;
-		}
-	}
-
-	void ConvertRangeToScaleBias(const FNiagaraStatelessRangeFloat& Range, float& OutScale, float& OutBias) const { OutScale = Range.GetScale(); GetRendererParameterValue(OutBias, Range.ParameterOffset, Range.Min); }
-	void ConvertRangeToScaleBias(const FNiagaraStatelessRangeVector2& Range, FVector2f& OutScale, FVector2f& OutBias) const { OutScale = Range.GetScale(); GetRendererParameterValue(OutBias, Range.ParameterOffset, Range.Min); }
-	void ConvertRangeToScaleBias(const FNiagaraStatelessRangeVector3& Range, FVector3f& OutScale, FVector3f& OutBias) const { OutScale = Range.GetScale(); GetRendererParameterValue(OutBias, Range.ParameterOffset, Range.Min); }
-	void ConvertRangeToScaleBias(const FNiagaraStatelessRangeColor& Range, FLinearColor& OutScale, FLinearColor& OutBias) const { OutScale = Range.GetScale(); GetRendererParameterValue(OutBias, Range.ParameterOffset, Range.Min); }
-
-protected:
-#if DO_CHECK
-	void ValidateIncludeStructType(uint32 StructOffset, const FShaderParametersMetadata* StructMetaData) const;
-#endif
-
-private:
-	const FNiagaraStatelessSpaceTransforms&	SpaceTransforms;
-	TConstArrayView<uint8>					RendererParameterData;
-	TConstArrayView<uint8>					BuiltData;
-	mutable int32							BuiltDataOffset = 0;
-	uint8*									ShaderParametersBase = nullptr;
-	mutable uint32							ParameterOffset = 0;
-	const FShaderParametersMetadata*		ShaderParametersMetadata = nullptr;
-};
-
-////////////////////////////////////////////////////////////////////////////////////////////////////////////////////////////////////////////////////////////////
-
 UCLASS(MinimalAPI, abstract, EditInlineNew)
 class UNiagaraStatelessModule : public UNiagaraMergeable
 {
@@ -119,7 +44,8 @@ public:
 	virtual ENiagaraStatelessFeatureMask GetFeatureMask() const { return ENiagaraStatelessFeatureMask::All; }
 
 	virtual void BuildEmitterData(const FNiagaraStatelessEmitterDataBuildContext& BuildContext) const {}
-	virtual void SetShaderParameters(const FNiagaraStatelessSetShaderParameterContext& SetShaderParameterContext) const { }
+	virtual void BuildShaderParameters(FNiagaraStatelessShaderParametersBuilder& ShaderParametersBuilder) const {}
+	virtual void SetShaderParameters(const FNiagaraStatelessSetShaderParameterContext& SetShaderParameterContext) const {}
 
 #if WITH_EDITOR
 	virtual bool CanDisableModule() const { return false; }
diff --git a/Engine/Plugins/FX/Niagara/Source/Niagara/Internal/Stateless/NiagaraStatelessParticleSimContext.h b/Engine/Plugins/FX/Niagara/Source/Niagara/Internal/Stateless/NiagaraStatelessParticleSimContext.h
--- a/Engine/Plugins/FX/Niagara/Source/Niagara/Internal/Stateless/NiagaraStatelessParticleSimContext.h
+++ b/Engine/Plugins/FX/Niagara/Source/Niagara/Internal/Stateless/NiagaraStatelessParticleSimContext.h
@@ -5,6 +5,8 @@
 #include "NiagaraStatelessCommon.h"
 #include "NiagaraStatelessBuiltDistribution.h"
 
+#include "ShaderParameterMacros.h"
+
 class FNiagaraDataBuffer;
 struct FNiagaraStatelessEmitterData;
 struct FNiagaraStatelessRuntimeSpawnInfo;
@@ -29,7 +31,7 @@ enum class EParticleComponent
 class FParticleSimulationContext
 {
 public:
-	explicit FParticleSimulationContext(const FNiagaraStatelessEmitterData* EmitterData, TConstArrayView<uint8> DynamicBufferData);
+	explicit FParticleSimulationContext(const FNiagaraStatelessEmitterData* EmitterData, const void* InShaderParametersData, TConstArrayView<uint8> DynamicBufferData);
 
 	void Simulate(int32 EmitterRandomSeed, float EmitterAge, float DeltaTime, TConstArrayView<FNiagaraStatelessRuntimeSpawnInfo> SpawnInfos, FNiagaraDataBuffer* DestinationData);
 	void SimulateGPU(FRHICommandListBase& RHICmdList, int32 EmitterRandomSeed, float EmitterAge, float DeltaTime, TConstArrayView<FNiagaraStatelessRuntimeSpawnInfo> SpawnInfos, FNiagaraDataBuffer* DestinationData);
@@ -179,6 +181,14 @@ public:
 		return reinterpret_cast<const T*>(BuiltData.GetData() + Offset);
 	}
 
+	template<typename T>
+	const T* ReadParameterNestedStruct() const
+	{
+		const uint32 Offset = Align(ShaderParameterOffset, TShaderParameterStructTypeInfo<T>::Alignment);
+		ShaderParameterOffset = Offset + TShaderParameterStructTypeInfo<T>::GetStructMetadata()->GetSize();
+		return reinterpret_cast<const T*>(ShaderParametersData + Offset);
+	}
+
 	static FQuat4f RotatorToQuat(FVector3f Rotator)
 	{
 		Rotator.X = FMath::Fractional(Rotator.X) * UE_PI;
@@ -216,6 +226,8 @@ private:
 
 	const TConstArrayView<uint8>			BuiltData;
 	mutable int32							BuiltDataOffset = 0;
+	const uint8*							ShaderParametersData = nullptr;
+	mutable int32							ShaderParameterOffset = 0;
 	const TConstArrayView<float>			StaticFloatData;
 	const TConstArrayView<uint8>			DynamicBufferData;
 };
diff --git a/Engine/Plugins/FX/Niagara/Source/Niagara/Internal/Stateless/NiagaraStatelessParticleSimExecData.h b/Engine/Plugins/FX/Niagara/Source/Niagara/Internal/Stateless/NiagaraStatelessParticleSimExecData.h
--- a/Engine/Plugins/FX/Niagara/Source/Niagara/Internal/Stateless/NiagaraStatelessParticleSimExecData.h
+++ b/Engine/Plugins/FX/Niagara/Source/Niagara/Internal/Stateless/NiagaraStatelessParticleSimExecData.h
@@ -16,15 +16,17 @@ public:
 	struct FCallback
 	{
 		FCallback() = default;
-		explicit FCallback(TFunction<void(const FParticleSimulationContext&)> InFunc, int32 InBuiltDataOffset, uint32 InRandomSeedOffset)
+		explicit FCallback(TFunction<void(const FParticleSimulationContext&)> InFunc, int32 InBuiltDataOffset, int32 InShaderParameterOffset, uint32 InRandomSeedOffset)
 			: Function(MoveTemp(InFunc))
 			, BuiltDataOffset(InBuiltDataOffset)
+			, ShaderParameterOffset(InShaderParameterOffset)
 			, RandomSeedOffset(InRandomSeedOffset)
 		{
 		}
 
 		TFunction<void(const FParticleSimulationContext&)>	Function;
 		int32	BuiltDataOffset = 0;
+		int32	ShaderParameterOffset = 0;
 		uint32	RandomSeedOffset = 0;
 	};
 
diff --git a/Engine/Plugins/FX/Niagara/Source/Niagara/Internal/Stateless/NiagaraStatelessSetShaderParameterContext.h b/Engine/Plugins/FX/Niagara/Source/Niagara/Internal/Stateless/NiagaraStatelessSetShaderParameterContext.h
new file mode 100644
--- /dev/null
+++ b/Engine/Plugins/FX/Niagara/Source/Niagara/Internal/Stateless/NiagaraStatelessSetShaderParameterContext.h
@@ -0,0 +1,78 @@
+// Copyright Epic Games, Inc. All Rights Reserved.
+
+#pragma once
+
+#include "NiagaraStatelessCommon.h"
+#include "ShaderParameterStruct.h"
+
+class FNiagaraStatelessSetShaderParameterContext
+{
+public:
+	UE_NONCOPYABLE(FNiagaraStatelessSetShaderParameterContext);
+
+	explicit FNiagaraStatelessSetShaderParameterContext(const FNiagaraStatelessSpaceTransforms& InSpaceTransforms, TConstArrayView<uint8> InRendererParameterData, TConstArrayView<uint8> InBuiltData, const FShaderParametersMetadata* InShaderParametersMetadata, uint8* InShaderParameters)
+		: SpaceTransforms(InSpaceTransforms)
+		, RendererParameterData(InRendererParameterData)
+		, BuiltData(InBuiltData)
+		, ShaderParametersBase(InShaderParameters)
+		, ParameterOffset(0)
+		, ShaderParametersMetadata(InShaderParametersMetadata)
+	{
+	}
+
+	const FNiagaraStatelessSpaceTransforms& GetSpaceTransforms() const { return SpaceTransforms; }
+
+	template<typename T>
+	T* GetParameterNestedStruct() const
+	{
+		const uint32 StructOffset = Align(ParameterOffset, TShaderParameterStructTypeInfo<T>::Alignment);
+	#if DO_CHECK
+		ValidateIncludeStructType(StructOffset, TShaderParameterStructTypeInfo<T>::GetStructMetadata());
+	#endif
+		ParameterOffset = StructOffset + TShaderParameterStructTypeInfo<T>::GetStructMetadata()->GetSize();
+		return reinterpret_cast<T*>(ShaderParametersBase + StructOffset);
+	}
+
+	template<typename T>
+	const T* ReadBuiltData() const
+	{
+		const int32 Offset = Align(BuiltDataOffset, alignof(T));
+		BuiltDataOffset = Offset + sizeof(T);
+		check(BuiltDataOffset <= BuiltData.Num());
+		return reinterpret_cast<const T*>(BuiltData.GetData() + Offset);
+	}
+
+	template<typename T>
+	void GetRendererParameterValue(T& OutValue, int32 Offset, const T& DefaultValue) const
+	{
+		if (Offset != INDEX_NONE)
+		{
+			Offset *= sizeof(uint32);
+			check(Offset >= 0 && Offset + sizeof(T) <= RendererParameterData.Num());
+			FMemory::Memcpy(&OutValue, RendererParameterData.GetData() + Offset, sizeof(T));
+		}
+		else
+		{
+			OutValue = DefaultValue;
+		}
+	}
+
+	void ConvertRangeToScaleBias(const FNiagaraStatelessRangeFloat& Range, float& OutScale, float& OutBias) const { OutScale = Range.GetScale(); GetRendererParameterValue(OutBias, Range.ParameterOffset, Range.Min); }
+	void ConvertRangeToScaleBias(const FNiagaraStatelessRangeVector2& Range, FVector2f& OutScale, FVector2f& OutBias) const { OutScale = Range.GetScale(); GetRendererParameterValue(OutBias, Range.ParameterOffset, Range.Min); }
+	void ConvertRangeToScaleBias(const FNiagaraStatelessRangeVector3& Range, FVector3f& OutScale, FVector3f& OutBias) const { OutScale = Range.GetScale(); GetRendererParameterValue(OutBias, Range.ParameterOffset, Range.Min); }
+	void ConvertRangeToScaleBias(const FNiagaraStatelessRangeColor& Range, FLinearColor& OutScale, FLinearColor& OutBias) const { OutScale = Range.GetScale(); GetRendererParameterValue(OutBias, Range.ParameterOffset, Range.Min); }
+
+protected:
+#if DO_CHECK
+	void ValidateIncludeStructType(uint32 StructOffset, const FShaderParametersMetadata* StructMetaData) const;
+#endif
+
+private:
+	const FNiagaraStatelessSpaceTransforms&	SpaceTransforms;
+	TConstArrayView<uint8>					RendererParameterData;
+	TConstArrayView<uint8>					BuiltData;
+	mutable int32							BuiltDataOffset = 0;
+	uint8*									ShaderParametersBase = nullptr;
+	mutable uint32							ParameterOffset = 0;
+	const FShaderParametersMetadata*		ShaderParametersMetadata = nullptr;
+};
diff --git a/Engine/Plugins/FX/Niagara/Source/Niagara/Internal/Stateless/NiagaraStatelessShaderParametersBuilder.h b/Engine/Plugins/FX/Niagara/Source/Niagara/Internal/Stateless/NiagaraStatelessShaderParametersBuilder.h
new file mode 100644
--- /dev/null
+++ b/Engine/Plugins/FX/Niagara/Source/Niagara/Internal/Stateless/NiagaraStatelessShaderParametersBuilder.h
@@ -0,0 +1,27 @@
+// Copyright Epic Games, Inc. All Rights Reserved.
+
+#pragma once
+
+#include "NiagaraStatelessCommon.h"
+#include "ShaderParameterStruct.h"
+
+class FNiagaraStatelessShaderParametersBuilder
+{
+public:
+	// Find the current size of the shader parameter structure
+	int32 GetParametersStructSize() const
+	{
+		return ParameterOffset;
+	}
+
+	// Adds a shader parameters structure that is scoped to the data interface
+	// i.e. if the structured contained "MyFloat" the shader variable would be "UniqueDataInterfaceName_MyFloat"
+	template<typename T> void AddParameterNestedStruct()
+	{
+		const int32 StructOffset = Align(ParameterOffset, TShaderParameterStructTypeInfo<T>::Alignment);
+		ParameterOffset = StructOffset + TShaderParameterStructTypeInfo<T>::GetStructMetadata()->GetSize();
+	}
+
+private:
+	int32	ParameterOffset = 0;
+};
diff --git a/Engine/Plugins/FX/Niagara/Source/Niagara/Private/Stateless/Modules/NiagaraStatelessModule_InitialMeshOrientation.cpp b/Engine/Plugins/FX/Niagara/Source/Niagara/Private/Stateless/Modules/NiagaraStatelessModule_InitialMeshOrientation.cpp
--- a/Engine/Plugins/FX/Niagara/Source/Niagara/Private/Stateless/Modules/NiagaraStatelessModule_InitialMeshOrientation.cpp
+++ b/Engine/Plugins/FX/Niagara/Source/Niagara/Private/Stateless/Modules/NiagaraStatelessModule_InitialMeshOrientation.cpp
@@ -105,6 +105,11 @@ void UNiagaraStatelessModule_InitialMeshOrientation::BuildEmitterData(const FNia
 	}
 }
 
+void UNiagaraStatelessModule_InitialMeshOrientation::BuildShaderParameters(FNiagaraStatelessShaderParametersBuilder& ShaderParametersBuilder) const
+{
+	ShaderParametersBuilder.AddParameterNestedStruct<FParameters>();
+}
+
 void UNiagaraStatelessModule_InitialMeshOrientation::SetShaderParameters(const FNiagaraStatelessSetShaderParameterContext& SetShaderParameterContext) const
 {
 	using namespace NSMInitialMeshOrientationPrivate;
diff --git a/Engine/Plugins/FX/Niagara/Source/Niagara/Private/Stateless/Modules/NiagaraStatelessModule_MeshIndex.cpp b/Engine/Plugins/FX/Niagara/Source/Niagara/Private/Stateless/Modules/NiagaraStatelessModule_MeshIndex.cpp
--- a/Engine/Plugins/FX/Niagara/Source/Niagara/Private/Stateless/Modules/NiagaraStatelessModule_MeshIndex.cpp
+++ b/Engine/Plugins/FX/Niagara/Source/Niagara/Private/Stateless/Modules/NiagaraStatelessModule_MeshIndex.cpp
@@ -133,6 +133,11 @@ void UNiagaraStatelessModule_MeshIndex::BuildEmitterData(const FNiagaraStateless
 	BuildContext.AddParticleSimulationExecSimulate(&ParticleSimulate);
 }
 
+void UNiagaraStatelessModule_MeshIndex::BuildShaderParameters(FNiagaraStatelessShaderParametersBuilder& ShaderParametersBuilder) const
+{
+	ShaderParametersBuilder.AddParameterNestedStruct<FParameters>();
+}
+
 void UNiagaraStatelessModule_MeshIndex::SetShaderParameters(const FNiagaraStatelessSetShaderParameterContext& SetShaderParameterContext) const
 {
 	using namespace NSMMeshIndexPrivate;
diff --git a/Engine/Plugins/FX/Niagara/Source/Niagara/Private/Stateless/Modules/NiagaraStatelessModule_ShapeLocation.cpp b/Engine/Plugins/FX/Niagara/Source/Niagara/Private/Stateless/Modules/NiagaraStatelessModule_ShapeLocation.cpp
--- a/Engine/Plugins/FX/Niagara/Source/Niagara/Private/Stateless/Modules/NiagaraStatelessModule_ShapeLocation.cpp
+++ b/Engine/Plugins/FX/Niagara/Source/Niagara/Private/Stateless/Modules/NiagaraStatelessModule_ShapeLocation.cpp
@@ -180,6 +180,11 @@ void UNiagaraStatelessModule_ShapeLocation::BuildEmitterData(const FNiagaraState
 	BuildContext.AddParticleSimulationExecSimulate(&ParticleSimulate);
 }
 
+void UNiagaraStatelessModule_ShapeLocation::BuildShaderParameters(FNiagaraStatelessShaderParametersBuilder& ShaderParametersBuilder) const
+{
+	ShaderParametersBuilder.AddParameterNestedStruct<FParameters>();
+}
+
 void UNiagaraStatelessModule_ShapeLocation::SetShaderParameters(const FNiagaraStatelessSetShaderParameterContext& SetShaderParameterContext) const
 {
 	using namespace NSMShapeLocationPrivate;
diff --git a/Engine/Plugins/FX/Niagara/Source/Niagara/Private/Stateless/NiagaraStatelessComputeManager.cpp b/Engine/Plugins/FX/Niagara/Source/Niagara/Private/Stateless/NiagaraStatelessComputeManager.cpp
--- a/Engine/Plugins/FX/Niagara/Source/Niagara/Private/Stateless/NiagaraStatelessComputeManager.cpp
+++ b/Engine/Plugins/FX/Niagara/Source/Niagara/Private/Stateless/NiagaraStatelessComputeManager.cpp
@@ -230,7 +230,7 @@ FNiagaraDataBuffer* FNiagaraStatelessComputeManager::GetDataBuffer(FRHICommandLi
 	{
 		case EComputeExecutionPath::CPU:
 		{
-			FParticleSimulationContext ParticleSimulation(EmitterData, EmitterInstance->BindingBufferData.Get(TArray<uint8>()));
+			FParticleSimulationContext ParticleSimulation(EmitterData, EmitterInstance->ShaderParameters.Get(), EmitterInstance->BindingBufferData.Get(TArray<uint8>()));
 			ParticleSimulation.SimulateGPU(RHICmdList, EmitterInstance->RandomSeed, EmitterInstance->Age, EmitterInstance->DeltaTime, EmitterInstance->SpawnInfos, CacheData->DataBuffer);
 			if (ParticleSimulation.GetNumInstances() == 0)
 			{
@@ -289,7 +289,7 @@ void FNiagaraStatelessComputeManager::GenerateDataBufferForDebugging(FRHICommand
 	{
 		case EComputeExecutionPath::CPU:
 		{
-			FParticleSimulationContext ParticleSimulation(EmitterData, EmitterInstance->BindingBufferData.Get(TArray<uint8>()));
+			FParticleSimulationContext ParticleSimulation(EmitterData, EmitterInstance->ShaderParameters.Get(), EmitterInstance->BindingBufferData.Get(TArray<uint8>()));
 			ParticleSimulation.Simulate(EmitterInstance->RandomSeed, EmitterInstance->Age, EmitterInstance->DeltaTime, EmitterInstance->SpawnInfos, DataBuffer);
 			break;
 		}
diff --git a/Engine/Plugins/FX/Niagara/Source/Niagara/Private/Stateless/NiagaraStatelessEmitter.cpp b/Engine/Plugins/FX/Niagara/Source/Niagara/Private/Stateless/NiagaraStatelessEmitter.cpp
--- a/Engine/Plugins/FX/Niagara/Source/Niagara/Private/Stateless/NiagaraStatelessEmitter.cpp
+++ b/Engine/Plugins/FX/Niagara/Source/Niagara/Private/Stateless/NiagaraStatelessEmitter.cpp
@@ -359,9 +359,13 @@ void UNiagaraStatelessEmitter::CacheFromCompiledData()
 			StatelessEmitterData->ParticleSimExecData
 		);
 
+		FNiagaraStatelessShaderParametersBuilder ShaderParametersBuilder;
+		ShaderParametersBuilder.AddParameterNestedStruct<NiagaraStateless::FCommonShaderParameters>();
+
 		for (const UNiagaraStatelessModule* Module : Modules)
 		{
-			EmitterBuildContext.PreModuleBuild();
+			EmitterBuildContext.PreModuleBuild(ShaderParametersBuilder.GetParametersStructSize());
+			Module->BuildShaderParameters(ShaderParametersBuilder);
 			Module->BuildEmitterData(EmitterBuildContext);
 		}
 		StatelessEmitterData->bModulesHaveRendererBindings = StatelessEmitterData->RendererBindings.Num() > 0;
diff --git a/Engine/Plugins/FX/Niagara/Source/Niagara/Private/Stateless/NiagaraStatelessEmitterDataBuildContext.cpp b/Engine/Plugins/FX/Niagara/Source/Niagara/Private/Stateless/NiagaraStatelessEmitterDataBuildContext.cpp
--- a/Engine/Plugins/FX/Niagara/Source/Niagara/Private/Stateless/NiagaraStatelessEmitterDataBuildContext.cpp
+++ b/Engine/Plugins/FX/Niagara/Source/Niagara/Private/Stateless/NiagaraStatelessEmitterDataBuildContext.cpp
@@ -8,9 +8,10 @@
 #include "NiagaraParameterStore.h"
 #include "NiagaraParameterStore.h"
 
-void FNiagaraStatelessEmitterDataBuildContext::PreModuleBuild()
+void FNiagaraStatelessEmitterDataBuildContext::PreModuleBuild(int32 InShaderParameterOffset)
 {
 	ModuleBuiltDataOffset = BuiltData.Num();
+	ShaderParameterOffset = InShaderParameterOffset;
 	++RandomSeedOffest;
 }
 
@@ -105,7 +106,7 @@ void FNiagaraStatelessEmitterDataBuildContext::AddParticleSimulationExecSimulate
 		return;
 	}
 
-	ParticleExecData->SimulateFunctions.Emplace(MoveTemp(Func), ModuleBuiltDataOffset, RandomSeedOffest);
+	ParticleExecData->SimulateFunctions.Emplace(MoveTemp(Func), ModuleBuiltDataOffset, ShaderParameterOffset, RandomSeedOffest);
 }
 
 int32 FNiagaraStatelessEmitterDataBuildContext::FindParticleVariableIndex(const FNiagaraVariableBase& Variable) const
diff --git a/Engine/Plugins/FX/Niagara/Source/Niagara/Private/Stateless/NiagaraStatelessParticleSimContext.cpp b/Engine/Plugins/FX/Niagara/Source/Niagara/Private/Stateless/NiagaraStatelessParticleSimContext.cpp
--- a/Engine/Plugins/FX/Niagara/Source/Niagara/Private/Stateless/NiagaraStatelessParticleSimContext.cpp
+++ b/Engine/Plugins/FX/Niagara/Source/Niagara/Private/Stateless/NiagaraStatelessParticleSimContext.cpp
@@ -49,9 +49,10 @@ namespace NiagaraStateless
 		return v;
 	}
 
-	FParticleSimulationContext::FParticleSimulationContext(const FNiagaraStatelessEmitterData* InEmitterData, TConstArrayView<uint8> InDynamicBufferData)
+	FParticleSimulationContext::FParticleSimulationContext(const FNiagaraStatelessEmitterData* InEmitterData, const void* InShaderParametersData, TConstArrayView<uint8> InDynamicBufferData)
 		: EmitterData(InEmitterData)
 		, BuiltData(EmitterData->BuiltData)
+		, ShaderParametersData(static_cast<const uint8*>(InShaderParametersData))
 		, StaticFloatData(EmitterData->StaticFloatData)
 		, DynamicBufferData(InDynamicBufferData)
 	{
@@ -224,6 +225,7 @@ namespace NiagaraStateless
 		for (const auto& Callback : ExecData->SimulateFunctions)
 		{
 			BuiltDataOffset = Callback.BuiltDataOffset;
+			ShaderParameterOffset = Callback.ShaderParameterOffset;
 			ModuleRandomSeed = Callback.RandomSeedOffset;
 			Callback.Function(*this);
 		}
