- Fix log spam when we intentionally set no skeletal mesh

#rnx

diff --git a/Engine/Plugins/FX/Niagara/Source/Niagara/Private/NiagaraDataInterfaceSkeletalMesh.cpp b/Engine/Plugins/FX/Niagara/Source/Niagara/Private/NiagaraDataInterfaceSkeletalMesh.cpp
--- a/Engine/Plugins/FX/Niagara/Source/Niagara/Private/NiagaraDataInterfaceSkeletalMesh.cpp
+++ b/Engine/Plugins/FX/Niagara/Source/Niagara/Private/NiagaraDataInterfaceSkeletalMesh.cpp
@@ -2377,7 +2377,8 @@ void UNiagaraDataInterfaceSkeletalMesh::GetVMExternalFunction(const FVMExternalF
 	BindTriangleSamplingFunction(BindingInfo, InstData, OutFunc);
 	if (OutFunc.IsBound())
 	{
-		if (!InstData->bAllowCPUMeshDataAccess)
+		// Only warn about failed access if we are bound to a skeletal mesh, there are cases where the user chooses to have no SKM so this warning is invalid
+		if (!InstData->bAllowCPUMeshDataAccess && InstData->SkeletalMesh.IsValid())
 		{
 			UE_LOG(LogNiagara, Log, TEXT("Skeletal Mesh Data Interface is trying to use triangle sampling function '%s', but either no CPU access is set on the mesh or the data is invalid. Interface: %s"),
 				*BindingInfo.Name.ToString(), *GetFullName());
@@ -2389,7 +2390,8 @@ void UNiagaraDataInterfaceSkeletalMesh::GetVMExternalFunction(const FVMExternalF
 	BindVertexSamplingFunction(BindingInfo, InstData, OutFunc);
 	if (OutFunc.IsBound())
 	{
-		if (!InstData->bAllowCPUMeshDataAccess)
+		// Only warn about failed access if we are bound to a skeletal mesh, there are cases where the user chooses to have no SKM so this warning is invalid
+		if (!InstData->bAllowCPUMeshDataAccess && InstData->SkeletalMesh.IsValid())
 		{
 			UE_LOG(LogNiagara, Log, TEXT("Skeletal Mesh Data Interface is trying to use vertex sampling function '%s' but either no CPU access is set on the mesh, or the data is invalid. Interface: %s"),
 				*BindingInfo.Name.ToString(), *GetFullName());
