Niagara - Fix playback in the system/emitter editor so that the playback framerate matches the rate set in the sequencer timeline.  Also save the current display rate setting and the lock state in the system editor data so that it persists with the asset.  *Additional fix.

diff --git a/Engine/Plugins/FX/Niagara/Source/NiagaraEditor/Private/ViewModels/NiagaraSystemViewModel.cpp b/Engine/Plugins/FX/Niagara/Source/NiagaraEditor/Private/ViewModels/NiagaraSystemViewModel.cpp
--- a/Engine/Plugins/FX/Niagara/Source/NiagaraEditor/Private/ViewModels/NiagaraSystemViewModel.cpp
+++ b/Engine/Plugins/FX/Niagara/Source/NiagaraEditor/Private/ViewModels/NiagaraSystemViewModel.cpp
@@ -1906,6 +1906,7 @@ void FNiagaraSystemViewModel::SetupPreviewComponentAndInstance()
 		PreviewComponent->SetAsset(System);
 		PreviewComponent->SetForceSolo(true);
 		PreviewComponent->SetAgeUpdateMode(ENiagaraAgeUpdateMode::DesiredAge);
+		PreviewComponent->SetSeekDelta((float)GetEditorData().GetPlaybackFrameRate().AsInterval());
 		PreviewComponent->SetCanRenderWhileSeeking(false);
 		PreviewComponent->Activate(true);
 
@@ -2816,6 +2817,10 @@ void FNiagaraSystemViewModel::SequencerMovieSceneModified(const UMovieScene* Mov
 	if (SystemEditorData.GetPlaybackFrameRate() != MovieScene->GetDisplayRate())
 	{
 		SystemEditorData.SetPlaybackFrameRate(MovieScene->GetDisplayRate());
+		if (PreviewComponent != nullptr)
+		{
+			PreviewComponent->SetSeekDelta((float)SystemEditorData.GetPlaybackFrameRate().AsInterval());
+		}
 	}
 
 	bool bMovieSceneIsFrameLocked = MovieScene->GetEvaluationType() == EMovieSceneEvaluationType::FrameLocked;
