- Change readback to use staging buffers and lock / unlock API is not consistent across all RHIs
- Fix RHI validation error when computing stateless for debugging due to missing UAV overlap enable / disable

#rnx

diff --git a/Engine/Plugins/FX/Niagara/Source/Niagara/Private/NiagaraDataSet.cpp b/Engine/Plugins/FX/Niagara/Source/Niagara/Private/NiagaraDataSet.cpp
--- a/Engine/Plugins/FX/Niagara/Source/Niagara/Private/NiagaraDataSet.cpp
+++ b/Engine/Plugins/FX/Niagara/Source/Niagara/Private/NiagaraDataSet.cpp
@@ -1454,16 +1454,60 @@ void FNiagaraDataBuffer::TransferGPUToCPUImmediate(FRHICommandListImmediate& RHI
 {
 	check(GetOwner()->GetCompiledData().GetLayoutHash() == CPUBuffer->GetOwner()->GetCompiledData().GetLayoutHash());
 
-	// Note: On some RHIs we need block until idle or the lock opersations will have the wrong information
+	// No count offset means no data
+	if (GPUInstanceCountBufferOffset == INDEX_NONE)
+	{
+		CPUBuffer->Allocate(0);
+		CPUBuffer->SetNumInstances(0);
+		return;
+	}
+
+	// Enqueue buffers readbacks
+	FStagingBufferRHIRef CountReadbackBuffer = RHICreateStagingBuffer();
+	{
+		FNiagaraGPUInstanceCountManager& CountManager = ComputeInterface->GetGPUInstanceCounterManager();
+		RHICmdList.Transition(FRHITransitionInfo(CountManager.GetInstanceCountBuffer().Buffer, FNiagaraGPUInstanceCountManager::kCountBufferDefaultState, ERHIAccess::CopySrc));
+		RHICmdList.CopyToStagingBuffer(CountManager.GetInstanceCountBuffer().Buffer, CountReadbackBuffer, GPUInstanceCountBufferOffset * sizeof(int32), sizeof(int32));
+		RHICmdList.Transition(FRHITransitionInfo(CountManager.GetInstanceCountBuffer().Buffer, ERHIAccess::CopySrc, FNiagaraGPUInstanceCountManager::kCountBufferDefaultState));
+	}
+
+	FStagingBufferRHIRef FloatReadbackBuffer = nullptr;
+	if (GPUBufferFloat.NumBytes > 0)
+	{
+		FloatReadbackBuffer = RHICreateStagingBuffer();
+		RHICmdList.Transition(FRHITransitionInfo(GPUBufferFloat.Buffer, ERHIAccess::SRVMask, ERHIAccess::CopySrc));
+		RHICmdList.CopyToStagingBuffer(GPUBufferFloat.Buffer, FloatReadbackBuffer, 0, GPUBufferFloat.NumBytes);
+		RHICmdList.Transition(FRHITransitionInfo(GPUBufferFloat.Buffer, ERHIAccess::CopySrc, ERHIAccess::SRVMask));
+	}
+
+	FStagingBufferRHIRef HalfReadbackBuffer = nullptr;
+	if (GPUBufferHalf.NumBytes > 0)
+	{
+		HalfReadbackBuffer = RHICreateStagingBuffer();
+		RHICmdList.Transition(FRHITransitionInfo(GPUBufferHalf.Buffer, ERHIAccess::SRVMask, ERHIAccess::CopySrc));
+		RHICmdList.CopyToStagingBuffer(GPUBufferHalf.Buffer, HalfReadbackBuffer, 0, GPUBufferHalf.NumBytes);
+		RHICmdList.Transition(FRHITransitionInfo(GPUBufferHalf.Buffer, ERHIAccess::CopySrc, ERHIAccess::SRVMask));
+	}
+
+	FStagingBufferRHIRef Int32ReadbackBuffer = nullptr;
+	if (GPUBufferInt.NumBytes > 0)
+	{
+		Int32ReadbackBuffer = RHICreateStagingBuffer();
+		RHICmdList.Transition(FRHITransitionInfo(GPUBufferInt.Buffer, ERHIAccess::SRVMask, ERHIAccess::CopySrc));
+		RHICmdList.CopyToStagingBuffer(GPUBufferInt.Buffer, Int32ReadbackBuffer, 0, GPUBufferInt.NumBytes);
+		RHICmdList.Transition(FRHITransitionInfo(GPUBufferInt.Buffer, ERHIAccess::CopySrc, ERHIAccess::SRVMask));
+	}
+
+	// Ensure all readbacks are submitted and complete
+	RHICmdList.SubmitCommandsAndFlushGPU();
 	RHICmdList.BlockUntilGPUIdle();
 
+	// Read the count and allocato space for results
 	uint32 GPUNumInstances = 0;
-	if (GPUInstanceCountBufferOffset != INDEX_NONE)
 	{
-		FNiagaraGPUInstanceCountManager& CountManager = ComputeInterface->GetGPUInstanceCounterManager();
-		FRHIBuffer* CountBuffer = CountManager.GetInstanceCountBuffer().Buffer;
-		GPUNumInstances = *reinterpret_cast<const int32*>(RHICmdList.LockBuffer(CountBuffer, GPUInstanceCountBufferOffset * sizeof(int32), sizeof(int32), RLM_ReadOnly));
-		RHICmdList.UnlockBuffer(CountBuffer);
+		const uint32* DataPtr = static_cast<const uint32*>(RHILockStagingBuffer(CountReadbackBuffer, 0, sizeof(int32)));
+		GPUNumInstances = *DataPtr;
+		RHIUnlockStagingBuffer(CountReadbackBuffer);
 	}
 
 	CPUBuffer->Allocate(GPUNumInstances);
@@ -1473,10 +1517,11 @@ void FNiagaraDataBuffer::TransferGPUToCPUImmediate(FRHICommandListImmediate& RHI
 		return;
 	}
 
-	if (GPUBufferFloat.NumBytes > 0)
+	// Copy out buffers
+	if (FloatReadbackBuffer)
 	{
 		const uint32 GPUStride	= FloatStride;
-		const uint8* GPUData	= reinterpret_cast<uint8*>(RHICmdList.LockBuffer(GPUBufferFloat.Buffer, 0, GPUBufferFloat.NumBytes, RLM_ReadOnly));
+		const uint8* GPUData	= reinterpret_cast<uint8*>(RHILockStagingBuffer(FloatReadbackBuffer, 0, GPUBufferFloat.NumBytes));
 
 		const uint32 CPUStride	= CPUBuffer->GetFloatStride();
 		uint8* CPUData			= CPUBuffer->GetComponentPtrFloat(0);
@@ -1486,23 +1531,42 @@ void FNiagaraDataBuffer::TransferGPUToCPUImmediate(FRHICommandListImmediate& RHI
 		{
 			FMemory::Memcpy(CPUData + (CPUStride * i), GPUData + (GPUStride * i), GPUNumInstances * sizeof(float));
 		}
-		RHICmdList.UnlockBuffer(GPUBufferFloat.Buffer);
+
+		RHIUnlockStagingBuffer(FloatReadbackBuffer);
 	}
 
-	if (GPUBufferInt.NumBytes > 0)
+	if (HalfReadbackBuffer)
+	{
+		const uint32 GPUStride = HalfStride;
+		const uint8* GPUData = reinterpret_cast<uint8*>(RHILockStagingBuffer(HalfReadbackBuffer, 0, GPUBufferHalf.NumBytes));
+
+		const uint32 CPUStride = CPUBuffer->GetHalfStride();
+		uint8* CPUData = CPUBuffer->GetComponentPtrHalf(0);
+
+		const int32 NumComponents = CPUBuffer->GetOwner()->GetNumHalfComponents();
+		for (int32 i = 0; i < NumComponents; ++i)
+		{
+			FMemory::Memcpy(CPUData + (CPUStride * i), GPUData + (GPUStride * i), GPUNumInstances * sizeof(FFloat16));
+		}
+
+		RHIUnlockStagingBuffer(HalfReadbackBuffer);
+	}
+
+	if (Int32ReadbackBuffer)
 	{
-		const uint32 GPUStride	= Int32Stride;
-		const uint8* GPUData	= reinterpret_cast<uint8*>(RHICmdList.LockBuffer(GPUBufferInt.Buffer, 0, GPUBufferInt.NumBytes, RLM_ReadOnly));
+		const uint32 GPUStride = Int32Stride;
+		const uint8* GPUData = reinterpret_cast<uint8*>(RHILockStagingBuffer(Int32ReadbackBuffer, 0, GPUBufferInt.NumBytes));
 
-		const uint32 CPUStride	= CPUBuffer->GetInt32Stride();
-		uint8* CPUData			= CPUBuffer->GetComponentPtrInt32(0);
+		const uint32 CPUStride = CPUBuffer->GetInt32Stride();
+		uint8* CPUData = CPUBuffer->GetComponentPtrInt32(0);
 
 		const int32 NumComponents = CPUBuffer->GetOwner()->GetNumInt32Components();
 		for (int32 i = 0; i < NumComponents; ++i)
 		{
 			FMemory::Memcpy(CPUData + (CPUStride * i), GPUData + (GPUStride * i), GPUNumInstances * sizeof(int32));
 		}
-		RHICmdList.UnlockBuffer(GPUBufferInt.Buffer);
+
+		RHIUnlockStagingBuffer(Int32ReadbackBuffer);
 	}
 }
 
diff --git a/Engine/Plugins/FX/Niagara/Source/Niagara/Private/Stateless/NiagaraStatelessComputeManager.cpp b/Engine/Plugins/FX/Niagara/Source/Niagara/Private/Stateless/NiagaraStatelessComputeManager.cpp
--- a/Engine/Plugins/FX/Niagara/Source/Niagara/Private/Stateless/NiagaraStatelessComputeManager.cpp
+++ b/Engine/Plugins/FX/Niagara/Source/Niagara/Private/Stateless/NiagaraStatelessComputeManager.cpp
@@ -313,7 +313,10 @@ void FNiagaraStatelessComputeManager::GenerateDataBufferForDebugging(FRHICommand
 			DataCache.ActiveParticles	= ActiveParticlesEstimate;
 
 			FStatelessDataCache* DataCachePtr = &DataCache;
+
+			RHICmdList.BeginUAVOverlap();
 			GenerateGPUData(RHICmdList, ComputeInterface, MakeConstArrayView(&DataCachePtr, 1));
+			RHICmdList.EndUAVOverlap();
 
 			// Copy to CPU data
 			//TransferGPUToCPU(RHICmdList, ComputeInterface, GPUDataBuffer, DataBuffer);
