- Add clarity to tooltips for emitter properties DI

#rnx

diff --git a/Engine/Plugins/FX/Niagara/Source/Niagara/Private/DataInterface/NiagaraDataInterfaceEmitterProperties.cpp b/Engine/Plugins/FX/Niagara/Source/Niagara/Private/DataInterface/NiagaraDataInterfaceEmitterProperties.cpp
--- a/Engine/Plugins/FX/Niagara/Source/Niagara/Private/DataInterface/NiagaraDataInterfaceEmitterProperties.cpp
+++ b/Engine/Plugins/FX/Niagara/Source/Niagara/Private/DataInterface/NiagaraDataInterfaceEmitterProperties.cpp
@@ -229,7 +229,7 @@ void UNiagaraDataInterfaceEmitterProperties::GetFunctionsInternal(TArray<FNiagar
 		Sig.Outputs.Emplace(FNiagaraTypeDefinition::GetVec3Def(), TEXT("Local Min"));
 		Sig.Outputs.Emplace(FNiagaraTypeDefinition::GetVec3Def(), TEXT("Local Max"));
 #if WITH_EDITORONLY_DATA
-		Sig.SetDescription(LOCTEXT("GetEmitterFixedBounds", "Get local space fixed bounds for the emitter.  If Valid is false then the emitter has no fixed bounds and is using dynamically generated bounds."));
+		Sig.SetDescription(LOCTEXT("GetEmitterFixedBounds", "Get the runtime set local space fixed bounds for the emitter.  If Valid is false then no runtime fixed bounds has been set, the emitter will be use the asset fixed bounds or dynamically generated bounds."));
 #endif
 	}
 	{
@@ -243,7 +243,7 @@ void UNiagaraDataInterfaceEmitterProperties::GetFunctionsInternal(TArray<FNiagar
 		Sig.Inputs.Emplace(FNiagaraTypeDefinition::GetVec3Def(), TEXT("Local Min"));
 		Sig.Inputs.Emplace(FNiagaraTypeDefinition::GetVec3Def(), TEXT("Local Max"));
 #if WITH_EDITORONLY_DATA
-		Sig.SetDescription(LOCTEXT("SetEmitterFixedBounds", "Set local space fixed bounds for the emitter, this stomps any fixed bounds set from Blueprint.  If Valid is set to false the emitter will fallback to the asset set fixed bounds (if valid) or dynamically generating bounds."));
+		Sig.SetDescription(LOCTEXT("SetEmitterFixedBounds", "Set the runtime local space fixed bounds for the emitter, this stomps any fixed bounds set from Blueprint.  If Valid is set to false the emitter will fallback to the asset bounds setup, i.e. fixed / dynamic / programmable."));
 #endif
 	}
 }
