- Hide stateless raw object by default

#rnx

diff --git a/Engine/Plugins/FX/Niagara/Source/NiagaraEditor/Private/ViewModels/Stack/NiagaraStackStatelessEmitterGroup.cpp b/Engine/Plugins/FX/Niagara/Source/NiagaraEditor/Private/ViewModels/Stack/NiagaraStackStatelessEmitterGroup.cpp
--- a/Engine/Plugins/FX/Niagara/Source/NiagaraEditor/Private/ViewModels/Stack/NiagaraStackStatelessEmitterGroup.cpp
+++ b/Engine/Plugins/FX/Niagara/Source/NiagaraEditor/Private/ViewModels/Stack/NiagaraStackStatelessEmitterGroup.cpp
@@ -12,6 +12,17 @@
 
 #define LOCTEXT_NAMESPACE "NiagaraEmitterStatelessGroup"
 
+namespace NiagaraStackStatelessEmitterGroupPrivate
+{
+	bool bShowRawObject = false;
+	FAutoConsoleVariableRef CVarShowRawObject(
+		TEXT("fx.NiagaraStateless.UI.ShowRawObject"),
+		bShowRawObject,
+		TEXT("When enabled we will show the raw object in the stateless emitter UI."),
+		ECVF_Default
+	);
+}
+
 void UNiagaraStackStatelessEmitterGroup::Initialize(FRequiredEntryData InRequiredEntryData, UNiagaraStatelessEmitter* InStatelessEmitter)
 {
 	Super::Initialize(
@@ -34,20 +45,23 @@ void UNiagaraStackStatelessEmitterGroup::RefreshChildrenInternal(const TArray<UN
 	UNiagaraStatelessEmitter* StatelessEmitter = StatelessEmitterWeak.Get();
 	if (StatelessEmitter != nullptr)
 	{
-		UNiagaraStackStatelessEmitterObjectItem* RawObjectItem = RawObjectItemWeak.Get();
-		if (RawObjectItem == nullptr || RawObjectItem->GetStatelessEmitter() != StatelessEmitter)
+		if ( NiagaraStackStatelessEmitterGroupPrivate::bShowRawObject )
 		{
-			bool bExpandedByDefault = false;
-			RawObjectItem = NewObject<UNiagaraStackStatelessEmitterObjectItem>(this);
-			RawObjectItem->Initialize(
-				CreateDefaultChildRequiredData(),
-				StatelessEmitter,
-				LOCTEXT("EmitterObjectDisplayName", "Raw Object"),
-				bExpandedByDefault,
-				FNiagaraStackObjectShared::FOnFilterDetailNodes());
-			RawObjectItemWeak = RawObjectItem;
+			UNiagaraStackStatelessEmitterObjectItem* RawObjectItem = RawObjectItemWeak.Get();
+			if (RawObjectItem == nullptr || RawObjectItem->GetStatelessEmitter() != StatelessEmitter)
+			{
+				bool bExpandedByDefault = false;
+				RawObjectItem = NewObject<UNiagaraStackStatelessEmitterObjectItem>(this);
+				RawObjectItem->Initialize(
+					CreateDefaultChildRequiredData(),
+					StatelessEmitter,
+					LOCTEXT("EmitterObjectDisplayName", "Raw Object"),
+					bExpandedByDefault,
+					FNiagaraStackObjectShared::FOnFilterDetailNodes());
+				RawObjectItemWeak = RawObjectItem;
+			}
+			NewChildren.Add(RawObjectItem);
 		}
-		NewChildren.Add(RawObjectItem);
 
 		UNiagaraStackStatelessEmitterObjectItem* FilteredObjectItem = FilteredObjectItemWeak.Get();
 		if (FilteredObjectItem == nullptr || FilteredObjectItem->GetStatelessEmitter() != StatelessEmitter)
