Fixed a crash where the niagara system finalize task could be called from another object's postload method

diff --git a/Engine/Plugins/FX/Niagara/Source/Niagara/Private/NiagaraComponent.cpp b/Engine/Plugins/FX/Niagara/Source/Niagara/Private/NiagaraComponent.cpp
--- a/Engine/Plugins/FX/Niagara/Source/Niagara/Private/NiagaraComponent.cpp
+++ b/Engine/Plugins/FX/Niagara/Source/Niagara/Private/NiagaraComponent.cpp
@@ -28,6 +28,7 @@
 #include "NiagaraCullProxyComponent.h"
 #include "SceneInterface.h"
 #include "UObject/UObjectIterator.h"
+#include "UObject/UObjectThreadContext.h"
 #include "PrimitiveUniformShaderParametersBuilder.h"
 #include "NiagaraActor.h"
 
@@ -1645,6 +1646,12 @@ void UNiagaraComponent::PostSystemTick_GameThread()
 {
 	check(SystemInstanceController.IsValid()); // sanity
 
+	if (FUObjectThreadContext::Get().IsRoutingPostLoad)
+	{
+		// return early when we're in the middle of another object's postload phase, as ticking data interfaces and renderers might lead to failed assertions
+		return;
+	}
+
 #if WITH_EDITOR
 	if (SystemInstanceController->HandleNeedsUIResync())
 	{
