Added a validation rule to check if a map for node is used in a cpu script

#tests added a map for node to a module and used it in cpu and gpu emitters

diff --git a/Engine/Plugins/FX/Niagara/Source/NiagaraEditor/Private/NiagaraNodeParameterMapFor.cpp b/Engine/Plugins/FX/Niagara/Source/NiagaraEditor/Private/NiagaraNodeParameterMapFor.cpp
--- a/Engine/Plugins/FX/Niagara/Source/NiagaraEditor/Private/NiagaraNodeParameterMapFor.cpp
+++ b/Engine/Plugins/FX/Niagara/Source/NiagaraEditor/Private/NiagaraNodeParameterMapFor.cpp
@@ -75,7 +75,6 @@ void UNiagaraNodeParameterMapFor::Compile(FTranslator* Translator, TArray<int32>
 		else
 		{
 			UNiagaraNodeParameterMapSet::Compile(Translator, Outputs);
-			//Translator->Message(FNiagaraCompileEventSeverity::Log,LOCTEXT("UnsupportedParamMapFor", "Parameter map for is not yet supported on cpu."), this, nullptr);
 		}
 	}
 }
@@ -119,7 +118,6 @@ void UNiagaraNodeParameterMapForWithContinue::Compile(FTranslator* Translator, T
 		else
 		{
 			UNiagaraNodeParameterMapSet::Compile(Translator, Outputs);
-			//Translator->Message(FNiagaraCompileEventSeverity::Log,LOCTEXT("UnsupportedParamMapFor", "Parameter map for is not yet supported on cpu."), this, nullptr);
 		}
 	}
 }
diff --git a/Engine/Plugins/FX/Niagara/Source/NiagaraEditor/Private/NiagaraValidationRules.cpp b/Engine/Plugins/FX/Niagara/Source/NiagaraEditor/Private/NiagaraValidationRules.cpp
--- a/Engine/Plugins/FX/Niagara/Source/NiagaraEditor/Private/NiagaraValidationRules.cpp
+++ b/Engine/Plugins/FX/Niagara/Source/NiagaraEditor/Private/NiagaraValidationRules.cpp
@@ -29,6 +29,8 @@
 #include "ViewModels/NiagaraEmitterViewModel.h"
 
 #include "AssetToolsModule.h"
+#include "NiagaraNodeOutput.h"
+#include "NiagaraNodeParameterMapFor.h"
 #include "Materials/MaterialInterface.h"
 
 #include "DeviceProfiles/DeviceProfile.h"
@@ -1441,4 +1443,57 @@ void UNiagaraValidationRule_SingletonModule::CheckValidity(const FNiagaraValidat
 	}
 }
 
+void UNiagaraValidationRule_NoMapForOnCpu::CheckValidity(const FNiagaraValidationContext& Context, TArray<FNiagaraValidationResult>& OutResults) const
+{
+	// gather all the modules in the system
+	TArray<UNiagaraStackModuleItem*> AllModules;
+	AllModules.Append(NiagaraValidation::GetStackEntries<UNiagaraStackModuleItem>(Context.ViewModel->GetSystemStackViewModel()));
+	TArray<TSharedRef<FNiagaraEmitterHandleViewModel>> EmitterHandleViewModels = Context.ViewModel->GetEmitterHandleViewModels();
+	for (TSharedRef<FNiagaraEmitterHandleViewModel> EmitterHandleModel : EmitterHandleViewModels)
+	{
+		AllModules.Append(NiagaraValidation::GetStackEntries<UNiagaraStackModuleItem>(EmitterHandleModel.Get().GetEmitterStackViewModel()));
+	}
+
+	for (UNiagaraStackModuleItem* Module : AllModules)
+	{
+		ENiagaraScriptUsage ScriptUsage = Module->GetOutputNode()->GetUsage();
+		FNiagaraEmitterViewModel* EmitterViewModel = Module->GetEmitterViewModel().Get();
+		if (EmitterViewModel && EmitterViewModel->GetEmitter().GetEmitterData()->SimTarget == ENiagaraSimTarget::GPUComputeSim && FNiagaraUtilities::ConvertScriptUsageToStaticSwitchContext(ScriptUsage) == ENiagaraScriptContextStaticSwitch::Particle)
+		{
+			// modules used in gpu scripts are ignored 
+			continue;
+		}
+		if (UNiagaraGraph* Graph = Module->GetModuleNode().GetCalledGraph())
+		{
+			FObjectKey GraphKey(Graph);
+			FGraphCheckResult& CheckResult = CachedResults.FindOrAdd(GraphKey);
+			if (CheckResult.ChangeID != Graph->GetChangeID())
+			{
+				CheckResult.ChangeID = Graph->GetChangeID();
+				CheckResult.bContainsMapForNode = false;
+				
+				TArray<UNiagaraNode*> TraversalNodes;
+				Graph->BuildTraversal(TraversalNodes, ENiagaraScriptUsage::Module, FGuid());
+				for (UNiagaraNode* TraversalNode : TraversalNodes)
+				{
+					if (TraversalNode->IsA<UNiagaraNodeParameterMapFor>() || TraversalNode->IsA<UNiagaraNodeParameterMapForWithContinue>())
+					{
+						CheckResult.bContainsMapForNode = true;
+						break;
+					}
+				}
+			}
+			if (CheckResult.bContainsMapForNode)
+			{
+				OutResults.Emplace(
+					Severity,
+					LOCTEXT("NoMapForOnCpu", "Map for node doesn't work in cpu scripts"),
+					LOCTEXT("NoMapForOnCpuDetailed", "This module contains a map for node, which does not work yet for cpu scripts.\nMap for nodes only work in particle gpu scripts. System and emitter scripts always run on the cpu."),
+					Module
+				);
+			}
+		}
+	}
+}
+
 #undef LOCTEXT_NAMESPACE
diff --git a/Engine/Plugins/FX/Niagara/Source/NiagaraEditor/Public/NiagaraValidationRules.h b/Engine/Plugins/FX/Niagara/Source/NiagaraEditor/Public/NiagaraValidationRules.h
--- a/Engine/Plugins/FX/Niagara/Source/NiagaraEditor/Public/NiagaraValidationRules.h
+++ b/Engine/Plugins/FX/Niagara/Source/NiagaraEditor/Public/NiagaraValidationRules.h
@@ -477,4 +477,27 @@ public:
 	/** If true then the check is not emitter-wide, but only within the same context (e.g. particle update). */
 	UPROPERTY(EditAnywhere, Category = Validation)
 	bool bCheckDetailedUsageContext = false;
+};
+
+/** This validation rule checks that map for nodes are not used with cpu scripts (as they only work on gpu). */
+UCLASS(Category = "Validation", DisplayName = "MapFor on CPU Check")
+class UNiagaraValidationRule_NoMapForOnCpu : public UNiagaraValidationRule
+{
+	GENERATED_BODY()
+
+public:
+	virtual void CheckValidity(const FNiagaraValidationContext& Context, TArray<FNiagaraValidationResult>& OutResults) const override;
+
+	/** How do we want to report the error in the stack */
+	UPROPERTY(EditAnywhere, Category = Validation)
+	ENiagaraValidationSeverity Severity = ENiagaraValidationSeverity::Error;
+
+private:
+
+	struct FGraphCheckResult
+	{
+		FGuid ChangeID;
+		bool bContainsMapForNode = false;
+	};
+	mutable TMap<FObjectKey, FGraphCheckResult> CachedResults;
 };
\ No newline at end of file
