Niagara - Fix various parameter panel category expansion issues in the system/emitter editor.
+ Move UI handling code out of the view model and into methods on the widget in the SItemSelector.
+ Add a category expanded changed event to the SItemSelector.
+ Watch for expansion changes in the system parameter panel, and update the cached UI state then rather than only when chaging tabs.

diff --git a/Engine/Plugins/FX/Niagara/Source/NiagaraEditor/Private/ViewModels/NiagaraParameterPanelViewModel.cpp b/Engine/Plugins/FX/Niagara/Source/NiagaraEditor/Private/ViewModels/NiagaraParameterPanelViewModel.cpp
--- a/Engine/Plugins/FX/Niagara/Source/NiagaraEditor/Private/ViewModels/NiagaraParameterPanelViewModel.cpp
+++ b/Engine/Plugins/FX/Niagara/Source/NiagaraEditor/Private/ViewModels/NiagaraParameterPanelViewModel.cpp
@@ -1769,25 +1769,22 @@ const TArray<UNiagaraParameterDefinitions*> FNiagaraSystemToolkitParameterPanelV
 	return SystemViewModel->GetAvailableParameterDefinitions(bSkipSubscribedParameterDefinitions);
 }
 
-void FNiagaraSystemToolkitParameterPanelViewModel::PreSectionChange(const TArray<FNiagaraParameterPanelCategory>& ExpandedItems)
+void FNiagaraSystemToolkitParameterPanelViewModel::UpdateCategoryExpansionState(const FNiagaraParameterPanelCategory& Category, bool bIsExpanded)
 {
-	// Before we go to a different setup, cache the existing expanded states.
+	// Update the cached expansion state for the active section if it's valid.
 	if (Sections.IsValidIndex(ActiveSectionIndex))
 	{
 		UNiagaraEditorSettings* Settings = GetMutableDefault<UNiagaraEditorSettings>();
 		bool bAdded = false;
 		FNiagaraParameterPanelSectionStorage& Storage = Settings->FindOrAddParameterPanelSectionStorage(Sections[ActiveSectionIndex].SectionId, bAdded);
-
-		TArray<FGuid> ExpandedCategories;
-
-		for (const FNiagaraParameterPanelCategory& Item : ExpandedItems)
+		if (bIsExpanded)
 		{
-			if (Item.NamespaceMetaData.IsValid() && Item.NamespaceMetaData.GetGuid().IsValid())
-			{
-				ExpandedCategories.AddUnique(Item.NamespaceMetaData.GetGuid());
-			}
+			Storage.ExpandedCategories.AddUnique(Category.NamespaceMetaData.GetGuid());
+		}
+		else
+		{
+			Storage.ExpandedCategories.Remove(Category.NamespaceMetaData.GetGuid());
 		}
-		Storage.ExpandedCategories = ExpandedCategories;
 		Settings->SaveConfig();
 	}
 }
diff --git a/Engine/Plugins/FX/Niagara/Source/NiagaraEditor/Private/Widgets/SNiagaraParameterPanel.cpp b/Engine/Plugins/FX/Niagara/Source/NiagaraEditor/Private/Widgets/SNiagaraParameterPanel.cpp
--- a/Engine/Plugins/FX/Niagara/Source/NiagaraEditor/Private/Widgets/SNiagaraParameterPanel.cpp
+++ b/Engine/Plugins/FX/Niagara/Source/NiagaraEditor/Private/Widgets/SNiagaraParameterPanel.cpp
@@ -96,12 +96,13 @@ void SNiagaraParameterPanel::Construct(const FArguments& InArgs, const TSharedPt
 	.ClearSelectionOnClick(true)
 	.CategoryRowStyle(FNiagaraEditorStyle::Get(), "NiagaraEditor.Parameters.TableRow")
 	.OnGetCategoryBackgroundImage(this, &SNiagaraParameterPanel::GetCategoryBackgroundImage)
-	.OnItemExpandedInitially(FilterCategoryExpandedDelegate)
+	.OnGetCategoryExpandedInitially(FilterCategoryExpandedDelegate)
 	.CategoryBorderBackgroundColor(FLinearColor(.6, .6, .6, 1.0f))
 	.CategoryChildSlotPadding(FMargin(0.0f, 2.0f, 0.0f, 0.0f))
 	.CategoryBorderBackgroundPadding(FMargin(0.0f, 3.0f))
 	.OnGetKeyForItem(this, &SNiagaraParameterPanel::OnGetKeyForItem)
 	.OnGetKeyForCategory(this, &SNiagaraParameterPanel::OnGetKeyForCategory)
+	.OnCategoryExpansionChanged(this, &SNiagaraParameterPanel::OnCategoryExpansionChanged)
 	.ExpandInitially(false)
 	.SearchBoxAdjacentContent()
 	[
@@ -190,9 +191,6 @@ void SNiagaraParameterPanel::OnSectionChecked(ECheckBoxState CheckState, FText S
 {
 	if (CheckState == ECheckBoxState::Checked)
 	{
-		TArray<FNiagaraParameterPanelCategory> ExpandedItems;
-		ItemSelector->GetExpandedCategoryItems(ExpandedItems);
-		ParameterPanelViewModel->PreSectionChange(ExpandedItems);
 		ParameterPanelViewModel->SetActiveSection(Section);
 		Refresh();
 	}
@@ -218,6 +216,11 @@ const FGuid& SNiagaraParameterPanel::OnGetKeyForCategory(const FNiagaraParameter
 	return Category.NamespaceMetaData.GetGuid();
 }
 
+void SNiagaraParameterPanel::OnCategoryExpansionChanged(const FNiagaraParameterPanelCategory& Category, bool bIsExpanded)
+{
+	ParameterPanelViewModel->UpdateCategoryExpansionState(Category, bIsExpanded);
+}
+
 bool SNiagaraParameterPanel::OnCompareItemsForEquality(const FNiagaraParameterPanelItem& ItemA, const FNiagaraParameterPanelItem& ItemB) const
 {
 	return ItemA.GetVariable()  == ItemB.GetVariable();
diff --git a/Engine/Plugins/FX/Niagara/Source/NiagaraEditor/Public/ViewModels/NiagaraParameterPanelViewModel.h b/Engine/Plugins/FX/Niagara/Source/NiagaraEditor/Public/ViewModels/NiagaraParameterPanelViewModel.h
--- a/Engine/Plugins/FX/Niagara/Source/NiagaraEditor/Public/ViewModels/NiagaraParameterPanelViewModel.h
+++ b/Engine/Plugins/FX/Niagara/Source/NiagaraEditor/Public/ViewModels/NiagaraParameterPanelViewModel.h
@@ -133,7 +133,7 @@ public:
 	// API for interacting with category expand/collapse within the view model
 	virtual bool UsesCategoryFilteringForInitialExpansion() const { return false; }
 	virtual bool IsCategoryExpandedByDefault(const FNiagaraParameterPanelCategory& Category) const { return false; };
-	virtual void PreSectionChange(const TArray<FNiagaraParameterPanelCategory>& ExpandedItems) {};
+	virtual void UpdateCategoryExpansionState(const FNiagaraParameterPanelCategory& Category, bool bIsExpanded) {};
 
 protected:
 	FOnRequestRefresh OnRequestRefreshDelegate;
@@ -397,7 +397,7 @@ public:
 	virtual bool IsCategoryExpandedByDefault(const FNiagaraParameterPanelCategory& Category) const override;
 	virtual bool UsesCategoryFilteringForInitialExpansion() const override { return true; }
 
-	virtual void PreSectionChange(const TArray<FNiagaraParameterPanelCategory>& ExpandedItems) override;
+	virtual void UpdateCategoryExpansionState(const FNiagaraParameterPanelCategory& Category, bool bIsExpanded) override;
 protected:
 	virtual void OnINiagaraParameterPanelViewModelSelectionChanged(UNiagaraScriptVariable* InVar) override;
 	virtual bool IncludeViewItemsInSelectParameterItem() const { return true; }
diff --git a/Engine/Plugins/FX/Niagara/Source/NiagaraEditor/Public/Widgets/SItemSelector.h b/Engine/Plugins/FX/Niagara/Source/NiagaraEditor/Public/Widgets/SItemSelector.h
--- a/Engine/Plugins/FX/Niagara/Source/NiagaraEditor/Public/Widgets/SItemSelector.h
+++ b/Engine/Plugins/FX/Niagara/Source/NiagaraEditor/Public/Widgets/SItemSelector.h
@@ -62,6 +62,7 @@ public:
 	DECLARE_DELEGATE_RetVal_OneParam(TSharedRef<SWidget>, FOnGenerateWidgetForCategory, const CategoryType& /* Category */);
 	DECLARE_DELEGATE_RetVal_OneParam(TSharedRef<SWidget>, FOnGenerateWidgetForItem, const ItemType& /* Item */);
 	DECLARE_DELEGATE(FOnSelectionChanged);
+	DECLARE_DELEGATE_TwoParams(FOnCategoryExpansionChanged, const CategoryType& /* Category */, bool /* bIsExpanded */);
 	DECLARE_DELEGATE_RetVal_OneParam(FSectionData, FOnGetSectionData, const SectionType& /* Section */);
 	DECLARE_DELEGATE_OneParam(FOnItemActivated, const ItemType& /* Item */);
 	DECLARE_DELEGATE_OneParam(FOnCategoryActivated, const CategoryType& /* Category */);
@@ -182,6 +183,9 @@ public:
 		/** A delegate which is called when the selection changes. */
 		SLATE_EVENT(FOnSelectionChanged, OnSelectionChanged)
 
+		/** An optional delegate called when a categories expansion state changes. */
+		SLATE_EVENT(FOnCategoryExpansionChanged, OnCategoryExpansionChanged)
+
 		/** An optional delegate which is used to determine the suggested item of a search. The highest weighted item will be selected. */
 		SLATE_EVENT(FOnGetItemWeight, OnGetItemWeight)
 	
@@ -215,14 +219,12 @@ public:
 		/** An optional delegate called to get a stable key (supports GetTypeHash()) for a section. If set, this delegate is used to uniquely key sections to preserve selection and expansion state when refresh is called. */
 		SLATE_EVENT(FOnGetKeyForSection, OnGetKeyForSection)
 
-		/** An optional array of delegates to refresh the item selector view when executed. */
-		SLATE_ARGUMENT(TArray<FRefreshItemSelectorDelegate*>, RefreshItemSelectorDelegates)
+		SLATE_EVENT(FOnCategoryPassesFilter, OnGetCategoryExpandedInitially)
 	
 		/** Whether we want to expand the tree initially or not. */
 		SLATE_ARGUMENT(bool, ExpandInitially)
 
 		/** An optional delegate called when initially creating this and ExpandInitially is false. */
-		SLATE_EVENT(FOnCategoryPassesFilter, OnItemExpandedInitially)
 
 		/** Slot for additional widget content to go adjacent to right of the search box. */
 		SLATE_NAMED_SLOT(FArguments, SearchBoxAdjacentContent)
@@ -1280,59 +1282,17 @@ private:
 			return RootViewModel->GetSections();
 		}
 		
-		void Refresh(const TArray<ItemType>& InItems, const TArray<TArray<CategoryType>>& InDefaultCategoryPaths, const TSharedPtr<STreeView<TSharedRef<FItemSelectorItemViewModel>>>& ItemTreeRef)
+		void Refresh(const TArray<ItemType>& InItems, const TArray<TArray<CategoryType>>& InDefaultCategoryPaths)
 		{
 			if (RootViewModel.IsValid())
 			{
-				// If preserving expansion or selection, cache the key for items/categories/sections before clearing out item viewmodels so that they may be re-marked as expanded or selected after regeneration.
-				if(bPreserveExpansionOnRefresh)
-				{
-					ExpandedCategoryKeyCache.Reset();
-					ExpandedSectionKeyCache.Reset();
-					TSet<TSharedRef<FItemSelectorItemViewModel>> ExpandedItemViewModels;
-					ItemTreeRef->GetExpandedItems(ExpandedItemViewModels);
-					for (const TSharedRef<FItemSelectorItemViewModel>& ExpandedItemViewModel : ExpandedItemViewModels)
-					{
-						if (ExpandedItemViewModel->GetType() == EItemSelectorItemViewModelType::Category)
-						{
-							ExpandedCategoryKeyCache.Add(OnGetKeyForCategory.Execute(StaticCastSharedRef<FItemSelectorItemCategoryViewModel>(ExpandedItemViewModel)->GetCategory()));
-						}
-						else if (ExpandedItemViewModel->GetType() == EItemSelectorItemViewModelType::Section)
-						{
-							ExpandedSectionKeyCache.Add(OnGetKeyForSection.Execute(StaticCastSharedRef<FSectionViewModel>(ExpandedItemViewModel)->GetSection()));
-						}
-					}
-				}
-				if (bPreserveSelectionOnRefresh)
-				{
-					SelectedItemKeyCache.Reset();
-					SelectedCategoryKeyCache.Reset();
-					SelectedSectionKeyCache.Reset();
-					TArray<TSharedRef<FItemSelectorItemViewModel>> SelectedItemViewModels;
-					ItemTreeRef->GetSelectedItems(SelectedItemViewModels);
-					for (const TSharedRef<FItemSelectorItemViewModel>& SelectedItemViewModel : SelectedItemViewModels)
-					{
-						if (SelectedItemViewModel->GetType() == EItemSelectorItemViewModelType::Item)
-						{
-							SelectedItemKeyCache.Add(StaticCastSharedRef<FItemSelectorItemContainerViewModel>(SelectedItemViewModel)->GetItemKey());
-						}
-						else if (SelectedItemViewModel->GetType() == EItemSelectorItemViewModelType::Category)
-						{
-							SelectedCategoryKeyCache.Add(OnGetKeyForCategory.Execute(StaticCastSharedRef<FItemSelectorItemCategoryViewModel>(SelectedItemViewModel)->GetCategory()));
-						}
-						else if (SelectedItemViewModel->GetType() == EItemSelectorItemViewModelType::Section)
-						{
-							SelectedSectionKeyCache.Add(OnGetKeyForSection.Execute(StaticCastSharedRef<FSectionViewModel>(SelectedItemViewModel)->GetSection()));
-						}
-					}
-				}
-				
 				RootViewModel.Reset();
 				RootTreeCategories.Empty();
 				FilteredFlattenedItems.Empty();
 				CurrentSuggestionIndex = INDEX_NONE;
 				CurrentMaxWeight = INDEX_NONE;
 			}
+
 			Items = InItems;
 			DefaultCategoryPaths = InDefaultCategoryPaths;
 			GetRootItems();
@@ -1342,69 +1302,6 @@ private:
 			{
 				UpdateSuggestedItem();
 			}
-
-			TArray<TSharedRef<FItemSelectorItemViewModel>> Children;
-			// Set expansion state if necessary.
-			if(bPreserveExpansionOnRefresh && ( ExpandedCategoryKeyCache.Num() > 0 || ExpandedSectionKeyCache.Num() > 0 ))
-			{
-				if (Children.Num() == 0)
-				{
-					GetChildrenRecursive(Children);
-				}
-				
-				for(const TSharedRef<FItemSelectorItemViewModel>& ChildItemViewModel : Children)
-				{
-					if (ChildItemViewModel->GetType() == EItemSelectorItemViewModelType::Category)
-					{
-						if (ExpandedCategoryKeyCache.Contains(OnGetKeyForCategory.Execute(StaticCastSharedRef<FItemSelectorItemCategoryViewModel>(ChildItemViewModel)->GetCategory())))
-						{
-							ItemTreeRef->SetItemExpansion(ChildItemViewModel, true);
-						}
-					}
-					else if (ChildItemViewModel->GetType() == EItemSelectorItemViewModelType::Section)
-					{
-						if (ExpandedSectionKeyCache.Contains(OnGetKeyForSection.Execute(StaticCastSharedRef<FSectionViewModel>(ChildItemViewModel)->GetSection())))
-						{
-							ItemTreeRef->SetItemExpansion(ChildItemViewModel, true);
-						}
-					}	
-				}
-			}
-
-			// Set selection state if necessary.
-			if (bPreserveSelectionOnRefresh && ( SelectedItemKeyCache.Num() > 0 || SelectedCategoryKeyCache.Num() > 0 || SelectedSectionKeyCache.Num() > 0 ))
-			{
-				if (Children.Num() == 0)
-				{
-					GetChildrenRecursive(Children);
-				}
-
-				ItemTreeRef->ClearSelection();
-				for (const TSharedRef<FItemSelectorItemViewModel>& ChildItemViewModel : Children)
-				{
-					if (ChildItemViewModel->GetType() == EItemSelectorItemViewModelType::Item)
-					{
-						if (SelectedItemKeyCache.Contains(OnGetKeyForItem.Execute(StaticCastSharedRef<FItemSelectorItemContainerViewModel>(ChildItemViewModel)->GetItem())))
-						{
-							ItemTreeRef->SetItemSelection(ChildItemViewModel, true);
-						}
-					}
-					else if (ChildItemViewModel->GetType() == EItemSelectorItemViewModelType::Category)
-					{
-						if (SelectedCategoryKeyCache.Contains(OnGetKeyForCategory.Execute(StaticCastSharedRef<FItemSelectorItemCategoryViewModel>(ChildItemViewModel)->GetCategory())))
-						{
-							ItemTreeRef->SetItemSelection(ChildItemViewModel, true);
-						}
-					}
-					else if (ChildItemViewModel->GetType() == EItemSelectorItemViewModelType::Section)
-					{
-						if (SelectedSectionKeyCache.Contains(OnGetKeyForSection.Execute(StaticCastSharedRef<FSectionViewModel>(ChildItemViewModel)->GetSection())))
-						{
-							ItemTreeRef->SetItemSelection(ChildItemViewModel, true);
-						}
-					}
-				}
-			}
 		}
 
 		int32 NumFlattenedItems()
@@ -1446,12 +1343,6 @@ private:
 		FText FilterText;
 
 		FOnSuggestionUpdated OnSuggestionUpdated;
-		
-		TSet<CategoryKeyType> ExpandedCategoryKeyCache;
-		TSet<SectionKeyType> ExpandedSectionKeyCache;
-		TSet<ItemKeyType> SelectedItemKeyCache;
-		TSet<CategoryKeyType> SelectedCategoryKeyCache;
-		TSet<SectionKeyType> SelectedSectionKeyCache;
 	};
 
 	typedef STableRow<TSharedRef<FItemSelectorItemViewModel>> SItemSelectorTableRow;
@@ -1592,6 +1483,7 @@ public:
 		OnItemActivated = InArgs._OnItemActivated;
 		OnCategoryActivated = InArgs._OnCategoryActivated;
 		OnSelectionChanged = InArgs._OnSelectionChanged;
+		OnCategoryExpansionChanged = InArgs._OnCategoryExpansionChanged;
 		OnDoesItemPassCustomFilter = InArgs._OnDoesItemPassCustomFilter;
 		OnDoesSectionPassCustomFilter = InArgs._OnDoesSectionPassCustomFilter;
 		OnGetSectionData = InArgs._OnGetSectionData;
@@ -1602,18 +1494,8 @@ public:
 		OnGetKeyForCategory = InArgs._OnGetKeyForCategory;
 		OnGetKeyForSection = InArgs._OnGetKeyForSection;
 		SearchBoxAdjacentContentWidget = InArgs._SearchBoxAdjacentContent.Widget;
-		bIsSettingSelection = false;
 
 		OnSuggestionUpdated = FOnSuggestionUpdated::CreateSP(this, &SItemSelector::OnSuggestionChanged);
-		
-		// Bind the on refresh delegates.
-		for (auto DelegateIt = InArgs._RefreshItemSelectorDelegates.CreateConstIterator(); DelegateIt; ++DelegateIt)
-		{
-			if ((*DelegateIt) != nullptr)
-			{
-				(**DelegateIt) = FRefreshItemSelectorDelegate::CreateSP(this, &SItemSelector::RefreshAllCurrentItems, false);
-			}
-		}
 
 		// Validate bindings for options set.
 		checkf(DefaultCategoryPaths.Num() == 0 || OnCompareCategoriesForEquality.IsBound(), TEXT("OnCompareCategoriesForEquality must be bound if default categories are supplied."));
@@ -1694,6 +1576,7 @@ public:
 				.OnMouseButtonDoubleClick(this, &SItemSelector::OnMouseDoubleClick)
 				.OnContextMenuOpening(InArgs._OnContextMenuOpening)
 				.OnSelectionChanged(this, &SItemSelector::OnTreeSelectionChanged)
+				.OnExpansionChanged(this, &SItemSelector::OnTreeExpansionChanged)
 				.TreeItemsSource(ViewModelUtilities->GetRootItems())
 				.ClearSelectionOnClick(InArgs._ClearSelectionOnClick)
 			]
@@ -1703,9 +1586,9 @@ public:
 		{
 			ExpandTree();
 		}
-		else if (InArgs._OnItemExpandedInitially.IsBound())
+		else if (InArgs._OnGetCategoryExpandedInitially.IsBound())
 		{
-			ExpandTreeByFilter(InArgs._OnItemExpandedInitially);
+			ExpandTreeByFilter(InArgs._OnGetCategoryExpandedInitially);
 		}
 
 		ExpandSections();
@@ -1884,11 +1767,37 @@ public:
 
 	void RefreshItemsAndDefaultCategories(const TArray<ItemType>& InItems, const TArray<TArray<CategoryType>>& InDefaultCategoryPaths)
 	{
-		ViewModelUtilities->Refresh(InItems, InDefaultCategoryPaths, ItemTree);
-		if (!bPreserveExpansionOnRefresh)
+		TSet<CategoryKeyType> ExpandedCategoryKeys;
+		TSet<SectionKeyType> ExpandedSectionKeys;
+		if (bPreserveExpansionOnRefresh)
+		{
+			GetExpandedKeys(ExpandedCategoryKeys, ExpandedSectionKeys);
+		}
+
+		TSet<ItemKeyType> SelectedItemKeys;
+		TSet<CategoryKeyType> SelectedCategoryKeys;
+		TSet<SectionKeyType> SelectedSectionKeys;
+		if (bPreserveSelectionOnRefresh)
+		{
+			GetSelectedKeys(SelectedItemKeys, SelectedCategoryKeys, SelectedSectionKeys);
+		}
+
+		ViewModelUtilities->Refresh(InItems, InDefaultCategoryPaths);
+
+		if (bPreserveExpansionOnRefresh)
+		{
+			RestoreExpansionState(ExpandedCategoryKeys, ExpandedSectionKeys);
+		}
+		else
 		{
 			ExpandTree();
 		}
+
+		if (bPreserveSelectionOnRefresh)
+		{
+			RestoreSelectionState(SelectedItemKeys, SelectedCategoryKeys, SelectedSectionKeys);
+		}
+
 		ItemTree->RequestTreeRefresh();
 	}
 
@@ -1911,14 +1820,11 @@ public:
 
 	void RefreshAllCurrentItems(bool bForceExpansion = false)
 	{
-		ViewModelUtilities->Refresh(Items, DefaultCategoryPaths, ItemTree);
-
+		RefreshItemsAndDefaultCategories(Items, DefaultCategoryPaths);
 		if(bForceExpansion)
 		{
 			ExpandTree();
 		}
-		
-		ItemTree->RequestTreeRefresh();
 	}
 
 	bool IsSearching() const
@@ -1961,6 +1867,8 @@ public:
 
 	void ExpandTreeByFilter(const FOnCategoryPassesFilter& InFilter)
 	{
+		TGuardValue<bool> SettingExpansionGuard(bIsSettingExpansion, true);
+
 		TArray<TSharedRef<FItemSelectorItemViewModel>> ItemsToProcess;
 		ItemsToProcess.Append(*ViewModelUtilities->GetRootItems());
 		while (ItemsToProcess.Num() > 0)
@@ -1990,23 +1898,6 @@ public:
 		return ItemTree->IsPendingRefresh();
 	}
 
-	void GetExpandedCategoryItems(TArray<CategoryType >& OutExpandedItems)
-	{
-		if (ItemTree.IsValid())
-		{
-			TSet<TSharedRef<FItemSelectorItemViewModel>> ExpandedItemViewModels;
-			ItemTree->GetExpandedItems(ExpandedItemViewModels);
-			for (const TSharedRef<FItemSelectorItemViewModel>& ExpandedItemViewModel : ExpandedItemViewModels)
-			{
-				if (ExpandedItemViewModel->GetType() == EItemSelectorItemViewModelType::Category)
-				{					
-					OutExpandedItems.AddUnique(StaticCastSharedRef<FItemSelectorItemCategoryViewModel>(ExpandedItemViewModel)->GetCategory());
-				}
-			}
-		}
-	}
-
-
 private:
 	void ExpandTreeByFilterRecursive(TSharedRef<FItemSelectorItemViewModel>& ItemToProcess, const FOnCategoryPassesFilter& InFilter, TOptional<bool> bForce)
 	{
@@ -2195,6 +2086,18 @@ private:
 			}
 		}
 	}
+
+	void OnTreeExpansionChanged(TSharedRef<FItemSelectorItemViewModel> ChangedItem, bool bIsExpanded)
+	{
+		if (bIsSettingExpansion == false)
+		{
+			if (OnCategoryExpansionChanged.IsBound() && ChangedItem->GetType() == EItemSelectorItemViewModelType::Category)
+			{
+				TSharedRef<FItemSelectorItemCategoryViewModel> ChangedCategory = StaticCastSharedRef<FItemSelectorItemCategoryViewModel>(ChangedItem);
+				OnCategoryExpansionChanged.Execute(ChangedCategory->GetCategory(), bIsExpanded);
+			}
+		}
+	}
 	
 	void OnSuggestionChanged(int32 SuggestionIndex)
 	{
@@ -2218,6 +2121,112 @@ private:
 		}
 	}
 
+	void GetExpandedKeys(TSet<CategoryKeyType>& OutExpandedCategoryKeys, TSet<SectionKeyType>& OutExpandedSectionKeys)
+	{
+		TSet<TSharedRef<FItemSelectorItemViewModel>> ExpandedItemViewModels;
+		ItemTree->GetExpandedItems(ExpandedItemViewModels);
+		for (const TSharedRef<FItemSelectorItemViewModel>& ExpandedItemViewModel : ExpandedItemViewModels)
+		{
+			if (ExpandedItemViewModel->GetType() == EItemSelectorItemViewModelType::Category)
+			{
+				OutExpandedCategoryKeys.Add(OnGetKeyForCategory.Execute(StaticCastSharedRef<FItemSelectorItemCategoryViewModel>(ExpandedItemViewModel)->GetCategory()));
+			}
+			else if (ExpandedItemViewModel->GetType() == EItemSelectorItemViewModelType::Section)
+			{
+				OutExpandedSectionKeys.Add(OnGetKeyForSection.Execute(StaticCastSharedRef<FSectionViewModel>(ExpandedItemViewModel)->GetSection()));
+			}
+		}
+	}
+
+	void GetSelectedKeys(TSet<ItemKeyType>& OutSelectedItemKeys, TSet<CategoryKeyType>& OutSelectedCategoryKeys, TSet<SectionKeyType>& OutSelectedSectionKeys)
+	{
+		TArray<TSharedRef<FItemSelectorItemViewModel>> SelectedItemViewModels;
+		ItemTree->GetSelectedItems(SelectedItemViewModels);
+		for (const TSharedRef<FItemSelectorItemViewModel>& SelectedItemViewModel : SelectedItemViewModels)
+		{
+			if (SelectedItemViewModel->GetType() == EItemSelectorItemViewModelType::Item)
+			{
+				OutSelectedItemKeys.Add(StaticCastSharedRef<FItemSelectorItemContainerViewModel>(SelectedItemViewModel)->GetItemKey());
+			}
+			else if (SelectedItemViewModel->GetType() == EItemSelectorItemViewModelType::Category)
+			{
+				OutSelectedCategoryKeys.Add(OnGetKeyForCategory.Execute(StaticCastSharedRef<FItemSelectorItemCategoryViewModel>(SelectedItemViewModel)->GetCategory()));
+			}
+			else if (SelectedItemViewModel->GetType() == EItemSelectorItemViewModelType::Section)
+			{
+				OutSelectedSectionKeys.Add(OnGetKeyForSection.Execute(StaticCastSharedRef<FSectionViewModel>(SelectedItemViewModel)->GetSection()));
+			}
+		}
+	}
+
+	void RestoreExpansionState(const TSet<CategoryKeyType>& ExpandedCategoryKeys, const TSet<SectionKeyType>& ExpandedSectionKeys)
+	{
+		if (ExpandedCategoryKeys.Num() == 0 && ExpandedSectionKeys.Num() == 0)
+		{
+			return;
+		}	
+
+		TGuardValue<bool> SettingExpansionGuard(bIsSettingExpansion, true);
+		TArray<TSharedRef<FItemSelectorItemViewModel>> Children;
+		ViewModelUtilities->GetChildrenRecursive(Children);
+
+		for (const TSharedRef<FItemSelectorItemViewModel>& ChildItemViewModel : Children)
+		{
+			if (ChildItemViewModel->GetType() == EItemSelectorItemViewModelType::Category)
+			{
+				if (ExpandedCategoryKeys.Contains(OnGetKeyForCategory.Execute(StaticCastSharedRef<FItemSelectorItemCategoryViewModel>(ChildItemViewModel)->GetCategory())))
+				{
+					ItemTree->SetItemExpansion(ChildItemViewModel, true);
+				}
+			}
+			else if (ChildItemViewModel->GetType() == EItemSelectorItemViewModelType::Section)
+			{
+				if (ExpandedSectionKeys.Contains(OnGetKeyForSection.Execute(StaticCastSharedRef<FSectionViewModel>(ChildItemViewModel)->GetSection())))
+				{
+					ItemTree->SetItemExpansion(ChildItemViewModel, true);
+				}
+			}
+		}
+	}
+
+	void RestoreSelectionState(const TSet<ItemKeyType>& SelectedItemKeys, const TSet<CategoryKeyType>& SelectedCategoryKeys, const TSet<SectionKeyType>& SelectedSectionKeys)
+	{
+		if (SelectedItemKeys.Num() == 0 && SelectedCategoryKeys.Num() == 0 && SelectedSectionKeys.Num() == 0)
+		{
+			return;
+		}
+
+		TGuardValue<bool> SettingSelectionGuard(bIsSettingSelection, true);
+		TArray<TSharedRef<FItemSelectorItemViewModel>> Children;
+		ViewModelUtilities->GetChildrenRecursive(Children);
+
+		ItemTree->ClearSelection();
+		for (const TSharedRef<FItemSelectorItemViewModel>& ChildItemViewModel : Children)
+		{
+			if (ChildItemViewModel->GetType() == EItemSelectorItemViewModelType::Item)
+			{
+				if (SelectedItemKeys.Contains(OnGetKeyForItem.Execute(StaticCastSharedRef<FItemSelectorItemContainerViewModel>(ChildItemViewModel)->GetItem())))
+				{
+					ItemTree->SetItemSelection(ChildItemViewModel, true);
+				}
+			}
+			else if (ChildItemViewModel->GetType() == EItemSelectorItemViewModelType::Category)
+			{
+				if (SelectedCategoryKeys.Contains(OnGetKeyForCategory.Execute(StaticCastSharedRef<FItemSelectorItemCategoryViewModel>(ChildItemViewModel)->GetCategory())))
+				{
+					ItemTree->SetItemSelection(ChildItemViewModel, true);
+				}
+			}
+			else if (ChildItemViewModel->GetType() == EItemSelectorItemViewModelType::Section)
+			{
+				if (SelectedSectionKeys.Contains(OnGetKeyForSection.Execute(StaticCastSharedRef<FSectionViewModel>(ChildItemViewModel)->GetSection())))
+				{
+					ItemTree->SetItemSelection(ChildItemViewModel, true);
+				}
+			}
+		}
+	}
+
 private:
 	TArray<ItemType> Items;
 	TArray<TArray<CategoryType>> DefaultCategoryPaths;
@@ -2253,6 +2262,7 @@ private:
 	FOnSelectionChanged OnSelectionChanged;
 	FOnDoesItemPassCustomFilter OnDoesItemPassCustomFilter;
 	FOnDoesSectionPassCustomFilter OnDoesSectionPassCustomFilter;
+	FOnCategoryExpansionChanged OnCategoryExpansionChanged;
 	FOnGetSectionData OnGetSectionData;
 	TAttribute<bool> HideSingleSection;
 	FOnItemSelected OnItemSelected;
@@ -2265,7 +2275,8 @@ private:
 	TSharedPtr<STreeView<TSharedRef<FItemSelectorItemViewModel>>> ItemTree;
 	FOnSuggestionUpdated OnSuggestionUpdated;
 
-	bool bIsSettingSelection;
+	bool bIsSettingSelection = false;
+	bool bIsSettingExpansion = false;
 
 	TSharedPtr<SSearchBox> SearchBox;
 	TSharedPtr<SWidget> SearchBoxAdjacentContentWidget;
diff --git a/Engine/Plugins/FX/Niagara/Source/NiagaraEditor/Public/Widgets/SNiagaraParameterPanel.h b/Engine/Plugins/FX/Niagara/Source/NiagaraEditor/Public/Widgets/SNiagaraParameterPanel.h
--- a/Engine/Plugins/FX/Niagara/Source/NiagaraEditor/Public/Widgets/SNiagaraParameterPanel.h
+++ b/Engine/Plugins/FX/Niagara/Source/NiagaraEditor/Public/Widgets/SNiagaraParameterPanel.h
@@ -60,6 +60,7 @@ public:
 	bool OnCompareCategoriesForEquality(const FNiagaraParameterPanelCategory& CategoryA, const FNiagaraParameterPanelCategory& CategoryB) const;
 	bool OnCompareCategoriesForSorting(const FNiagaraParameterPanelCategory& CategoryA, const FNiagaraParameterPanelCategory& CategoryB) const;
 	const FGuid& OnGetKeyForCategory(const FNiagaraParameterPanelCategory& Category) const;
+	void OnCategoryExpansionChanged(const FNiagaraParameterPanelCategory& Category, bool bIsExpanded);
 	bool OnCompareItemsForEquality(const FNiagaraParameterPanelItem& ItemA, const FNiagaraParameterPanelItem& ItemB) const;
 	bool OnCompareItemsForSorting(const FNiagaraParameterPanelItem& ItemA, const FNiagaraParameterPanelItem& ItemB) const;
 	const FNiagaraVariableBase& OnGetKeyForItem(const FNiagaraParameterPanelItem& Item) const;
