* Added support for custom module wizards that can generate scatch pad modules
* Added wizards for data channel read and write modules

diff --git a/Engine/Plugins/FX/Niagara/Source/Niagara/Private/DataInterface/NiagaraDataInterfaceDataChannelRead.cpp b/Engine/Plugins/FX/Niagara/Source/Niagara/Private/DataInterface/NiagaraDataInterfaceDataChannelRead.cpp
--- a/Engine/Plugins/FX/Niagara/Source/Niagara/Private/DataInterface/NiagaraDataInterfaceDataChannelRead.cpp
+++ b/Engine/Plugins/FX/Niagara/Source/Niagara/Private/DataInterface/NiagaraDataInterfaceDataChannelRead.cpp
@@ -1182,6 +1182,9 @@ bool UNiagaraDataInterfaceDataChannelRead::CopyToInternal(UNiagaraDataInterface*
 #if WITH_EDITORONLY_DATA
 void UNiagaraDataInterfaceDataChannelRead::GetFunctionsInternal(TArray<FNiagaraFunctionSignature>& OutFunctions) const
 {
+	/////
+	// NOTE: *any* changes to function inputs or outputs needs to be included in FReadNDCModel::GenerateNewModuleContent()
+	/////
 	OutFunctions.Add(NDIDataChannelReadLocal::GetFunctionSig_Num());
 	OutFunctions.Add(NDIDataChannelReadLocal::GetFunctionSig_GetNDCSpawnData());
 	OutFunctions.Add(NDIDataChannelReadLocal::GetFunctionSig_Read());
diff --git a/Engine/Plugins/FX/Niagara/Source/Niagara/Private/DataInterface/NiagaraDataInterfaceDataChannelWrite.cpp b/Engine/Plugins/FX/Niagara/Source/Niagara/Private/DataInterface/NiagaraDataInterfaceDataChannelWrite.cpp
--- a/Engine/Plugins/FX/Niagara/Source/Niagara/Private/DataInterface/NiagaraDataInterfaceDataChannelWrite.cpp
+++ b/Engine/Plugins/FX/Niagara/Source/Niagara/Private/DataInterface/NiagaraDataInterfaceDataChannelWrite.cpp
@@ -774,9 +774,14 @@ bool UNiagaraDataInterfaceDataChannelWrite::CopyToInternal(UNiagaraDataInterface
 #if WITH_EDITORONLY_DATA
 void UNiagaraDataInterfaceDataChannelWrite::GetFunctionsInternal(TArray<FNiagaraFunctionSignature>& OutFunctions) const
 {
+
+	/////
+	// NOTE: *any* changes to function inputs or outputs here needs to be included in FWriteNDCModel::GenerateNewModuleContent()
+	/////
 	{
 		FNiagaraFunctionSignature Sig;
 		Sig.Name = NDIDataChannelWriteLocal::NumName;
+		NIAGARA_ADD_FUNCTION_SOURCE_INFO(Sig)
 		Sig.Description = LOCTEXT("NumFunctionDescription", "Returns the current number of DataChannel accessible by this interface.");
 		Sig.bMemberFunction = true;
 		Sig.bExperimental = true;
@@ -791,12 +796,13 @@ void UNiagaraDataInterfaceDataChannelWrite::GetFunctionsInternal(TArray<FNiagara
 	{
 		FNiagaraFunctionSignature Sig;
 		Sig.Name = NDIDataChannelWriteLocal::WriteName;
+		NIAGARA_ADD_FUNCTION_SOURCE_INFO(Sig)
 		Sig.Description = LOCTEXT("WriteFunctionDescription", "Writes DataChannel data at a specific index.  Values in the DataChannel that are not written here are set to their defaults. Returns success if an DataChannel was written to.");
 		Sig.bMemberFunction = true;
 		Sig.bRequiresExecPin = true;
 		Sig.bExperimental = true;
 		Sig.AddInput(FNiagaraVariable(FNiagaraTypeDefinition(GetClass()), TEXT("DataChannel interface")));
-		Sig.AddInput(EmitVar);
+		Sig.AddInput(EmitVar, LOCTEXT("ExecuteFlagTooltip", "If true then the write is executed, if false then the whole write is skipped"));
 		Sig.AddInput(FNiagaraVariable(FNiagaraTypeDefinition::GetIntDef(), TEXT("Index")));
 		Sig.AddOutput(FNiagaraVariable(FNiagaraTypeDefinition::GetBoolDef(), TEXT("Success")));
 		Sig.RequiredInputs = IntCastChecked<int16>(Sig.Inputs.Num());//The user defines what we write in the graph.
@@ -805,12 +811,13 @@ void UNiagaraDataInterfaceDataChannelWrite::GetFunctionsInternal(TArray<FNiagara
 	{
 		FNiagaraFunctionSignature Sig;
 		Sig.Name = NDIDataChannelWriteLocal::AppendName;
+		NIAGARA_ADD_FUNCTION_SOURCE_INFO(Sig)
 		Sig.Description = LOCTEXT("AppendFunctionDescription", "Appends a new DataChannel to the end of the DataChannel array and writes the specified values. Values in the DataChannel that are not written here are set to their defaults. Returns success if an DataChannel was successfully pushed.");
 		Sig.bMemberFunction = true;
 		Sig.bRequiresExecPin = true;
 		Sig.bExperimental = true;
 		Sig.AddInput(FNiagaraVariable(FNiagaraTypeDefinition(GetClass()), TEXT("DataChannel interface")));
-		Sig.AddInput(EmitVar);
+		Sig.AddInput(EmitVar, LOCTEXT("ExecuteAppendFlagTooltip", "If true then the append is executed, if false then the append call is skipped"));
 		Sig.AddOutput(FNiagaraVariable(FNiagaraTypeDefinition::GetBoolDef(), TEXT("Success")));
 		Sig.RequiredInputs = IntCastChecked<int16>(Sig.Inputs.Num());//The user defines what we write in the graph.
 		OutFunctions.Add(Sig);
@@ -868,7 +875,6 @@ void UNiagaraDataInterfaceDataChannelWrite::Write(FVectorVMExternalFunctionConte
 	FNDIInputParam<FNiagaraBool> InEmit(Context);
 	FNDIInputParam<int32> InIndex(Context);
 
-	const FNDIDataChannelFunctionInfo& FuncInfo = CompiledData.GetFunctionInfo()[FuncIdx];
 	const FNDIDataChannel_FunctionToDataSetBinding* BindingInfo = InstData->FunctionToDatSetBindingInfo.IsValidIndex(FuncIdx) ? InstData->FunctionToDatSetBindingInfo[FuncIdx].Get() : nullptr;
 	FNDIVariadicInputHandler<16> VariadicInputs(Context, BindingInfo);//TODO: Make static / avoid allocation
 
@@ -883,13 +889,14 @@ void UNiagaraDataInterfaceDataChannelWrite::Write(FVectorVMExternalFunctionConte
 		{			
 			bAllFailedFallback = false;
 			uint32 MaxLocalIndex = 0;
-			int32 NumAllocated = (int32)Data->GetNumInstancesAllocated();
+			int32 NumAllocated = IntCastChecked<int32>(Data->GetNumInstancesAllocated());
 			for (int32 i = 0; i < Context.GetNumInstances(); ++i)
 			{
-				int32 Index = InIndex.GetAndAdvance();
-				bool bEmit = InEmit.GetAndAdvance() && Index >= 0 && Index < NumAllocated;
+				int32 RawIndex = InIndex.GetAndAdvance();
+				bool bEmit = InEmit.GetAndAdvance() && RawIndex >= 0 && RawIndex < NumAllocated;
 
-				MaxLocalIndex = bEmit ? FMath::Max((uint32)Index, MaxLocalIndex) : MaxLocalIndex;
+				uint32 Index = static_cast<uint32>(RawIndex);
+				MaxLocalIndex = bEmit ? FMath::Max(Index, MaxLocalIndex) : MaxLocalIndex;
 
 				bool bSuccess = false;
 
@@ -898,17 +905,17 @@ void UNiagaraDataInterfaceDataChannelWrite::Write(FVectorVMExternalFunctionConte
 				auto FloatFunc = [Data, Index](const FNDIDataChannelRegisterBinding& VMBinding, FNDIInputParam<float>& FloatData)
 				{
 					if (VMBinding.DataSetRegisterIndex != INDEX_NONE)
-						*Data->GetInstancePtrFloat(VMBinding.DataSetRegisterIndex, (uint32)Index) = FloatData.GetAndAdvance();
+						*Data->GetInstancePtrFloat(VMBinding.DataSetRegisterIndex, Index) = FloatData.GetAndAdvance();
 				};
 				auto IntFunc = [Data, Index](const FNDIDataChannelRegisterBinding& VMBinding, FNDIInputParam<int32>& IntData)
 				{
 					if (VMBinding.DataSetRegisterIndex != INDEX_NONE)
-						*Data->GetInstancePtrInt32(VMBinding.DataSetRegisterIndex, (uint32)Index) = IntData.GetAndAdvance();
+						*Data->GetInstancePtrInt32(VMBinding.DataSetRegisterIndex, Index) = IntData.GetAndAdvance();
 				};
 				auto HalfFunc = [Data, Index](const FNDIDataChannelRegisterBinding& VMBinding, FNDIInputParam<FFloat16>& HalfData)
 				{
 					if (VMBinding.DataSetRegisterIndex != INDEX_NONE)
-						*Data->GetInstancePtrHalf(VMBinding.DataSetRegisterIndex, (uint32)Index) = HalfData.GetAndAdvance();
+						*Data->GetInstancePtrHalf(VMBinding.DataSetRegisterIndex, Index) = HalfData.GetAndAdvance();
 				};
 
 				bSuccess = VariadicInputs.Process(bEmit, 1, BindingInfo, FloatFunc, IntFunc, HalfFunc);
@@ -920,7 +927,6 @@ void UNiagaraDataInterfaceDataChannelWrite::Write(FVectorVMExternalFunctionConte
 			}
 
 			//Update the shared instance count with an updated max.
-			bool bAtomicNumInstancesWritten = false;
 			uint32 CurrNumInstances = AtomicNumInstances;
 			while(CurrNumInstances < MaxLocalIndex && !AtomicNumInstances.compare_exchange_weak(CurrNumInstances, MaxLocalIndex))
 			{
@@ -947,7 +953,6 @@ void UNiagaraDataInterfaceDataChannelWrite::Append(FVectorVMExternalFunctionCont
 	VectorVM::FUserPtrHandler<FNDIDataChannelWriteInstanceData> InstData(Context);
 	FNDIInputParam<FNiagaraBool> InEmit(Context);
 
-	const FNDIDataChannelFunctionInfo& FuncInfo = CompiledData.GetFunctionInfo()[FuncIdx];
 	const FNDIDataChannel_FunctionToDataSetBinding* BindingInfo = InstData->FunctionToDatSetBindingInfo.IsValidIndex(FuncIdx) ? InstData->FunctionToDatSetBindingInfo[FuncIdx].Get() : nullptr;
 	FNDIVariadicInputHandler<16> VariadicInputs(Context, BindingInfo);//TODO: Make static / avoid allocation
 
@@ -962,7 +967,7 @@ void UNiagaraDataInterfaceDataChannelWrite::Append(FVectorVMExternalFunctionCont
 		{
 			//Get the total number to emit.
 			//Allows going via a faster write path if we're emiting every instance.
-			//Aslo needed to update the atomic num instances and get our start index for writing.
+			//Also needed to update the atomic num instances and get our start index for writing.
 			uint32 LocalNumToEmit = 0;
 			if(InEmit.IsConstant())
 			{
diff --git a/Engine/Plugins/FX/Niagara/Source/NiagaraEditor/Private/NiagaraEditorModule.cpp b/Engine/Plugins/FX/Niagara/Source/NiagaraEditor/Private/NiagaraEditorModule.cpp
--- a/Engine/Plugins/FX/Niagara/Source/NiagaraEditor/Private/NiagaraEditorModule.cpp
+++ b/Engine/Plugins/FX/Niagara/Source/NiagaraEditor/Private/NiagaraEditorModule.cpp
@@ -179,6 +179,7 @@
 #include "ViewModels/HierarchyEditor/NiagaraHierarchyCommands.h"
 #include "Widgets/AssetBrowser/SNiagaraSelectedAssetDetails.h"
 #include "NiagaraRecentAndFavoritesManager.h"
+#include "Widgets/DataChannel/NiagaraDataChannelWizard.h"
 
 #include UE_INLINE_GENERATED_CPP_BY_NAME(NiagaraEditorModule)
 
@@ -1557,6 +1558,8 @@ void FNiagaraEditorModule::StartupModule()
 		}
 	}
 
+	RegisterModuleWizards(UE::Niagara::Wizard::DataChannel::CreateNDCWizardGenerator());
+
 #if NIAGARA_PERF_BASELINES
 	UNiagaraEffectType::OnGeneratePerfBaselines().BindRaw(this, &FNiagaraEditorModule::GeneratePerfBaselines);
 #endif
@@ -1832,6 +1835,11 @@ void FNiagaraEditorModule::RegisterPropertyUtilities(const UScriptStruct* InStru
 	TypeEditorsCS.Unlock();
 }
 
+void FNiagaraEditorModule::RegisterModuleWizards(TSharedRef<UE::Niagara::Wizard::FModuleWizardGenerator> WizardGenerator)
+{
+	ModuleWizards.Add(WizardGenerator);
+}
+
 TSharedPtr<INiagaraEditorTypeUtilities, ESPMode::ThreadSafe> FNiagaraEditorModule::GetTypeUtilities(const FNiagaraTypeDefinition& Type)
 {
 	TypeEditorsCS.Lock();
@@ -2051,7 +2059,7 @@ void FNiagaraEditorModule::GetDataInterfaceFeedbackSafe(UNiagaraDataInterface* I
 
 	if (OwningSystem == nullptr)
 	{
-		// If no outer was find try to find one by componenet.
+		// If no outer was found, try to find one by component.
 		if (OwningComponent != nullptr)
 		{
 			OwningSystem = OwningComponent->GetAsset();
@@ -2060,7 +2068,7 @@ void FNiagaraEditorModule::GetDataInterfaceFeedbackSafe(UNiagaraDataInterface* I
 
 	if (OwningSystem == nullptr)
 	{
-		// If no outer information is available check system view models for placeholder DIs.
+		// If no outer information is available, check system view models for placeholder DIs.
 		TArray<TSharedRef<FNiagaraSystemViewModel>> SystemViewModels;
 		FNiagaraSystemViewModel::GetAllViewModels(SystemViewModels);
 		for (TSharedRef<FNiagaraSystemViewModel> SystemViewModel : SystemViewModels)
diff --git a/Engine/Plugins/FX/Niagara/Source/NiagaraEditor/Private/NiagaraNodeWithDynamicPins.cpp b/Engine/Plugins/FX/Niagara/Source/NiagaraEditor/Private/NiagaraNodeWithDynamicPins.cpp
--- a/Engine/Plugins/FX/Niagara/Source/NiagaraEditor/Private/NiagaraNodeWithDynamicPins.cpp
+++ b/Engine/Plugins/FX/Niagara/Source/NiagaraEditor/Private/NiagaraNodeWithDynamicPins.cpp
@@ -397,6 +397,12 @@ void UNiagaraNodeWithDynamicPins::GetNodeContextMenuActions(UToolMenu* Menu, UGr
 }
 
 void UNiagaraNodeWithDynamicPins::AddParameter(FNiagaraVariable Parameter, TEnumAsByte<EEdGraphPinDirection> Direction)
+{
+	// discard return pin, so the function can be used as delegate
+	AddParameterPin(Parameter, Direction);
+}
+
+UEdGraphPin* UNiagaraNodeWithDynamicPins::AddParameterPin(FNiagaraVariable Parameter, TEnumAsByte<EEdGraphPinDirection> Direction)
 {
 	if (this->IsA<UNiagaraNodeParameterMapBase>())
 	{
@@ -416,15 +422,18 @@ void UNiagaraNodeWithDynamicPins::AddParameter(FNiagaraVariable Parameter, TEnum
 		Graph->AddParameter(Parameter);
 
 		Modify();
-		UEdGraphPin* Pin = this->RequestNewTypedPin(Direction, Parameter.GetType(), Parameter.GetName());
-	}
-	else
-	{
-		RequestNewTypedPin(Direction, Parameter.GetType(), Parameter.GetName());
+		return this->RequestNewTypedPin(Direction, Parameter.GetType(), Parameter.GetName());
 	}
+	return RequestNewTypedPin(Direction, Parameter.GetType(), Parameter.GetName());
 }
 
 void UNiagaraNodeWithDynamicPins::AddParameter(const UNiagaraScriptVariable* ScriptVar, TEnumAsByte<EEdGraphPinDirection> Direction)
+{
+	// discard return pin, so the function can be used as delegate
+	AddParameterPin(ScriptVar, Direction);
+}
+
+UEdGraphPin* UNiagaraNodeWithDynamicPins::AddParameterPin(const UNiagaraScriptVariable* ScriptVar, TEnumAsByte<EEdGraphPinDirection> Direction)
 {
 	const FNiagaraVariable& Parameter = ScriptVar->Variable;
 	if (this->IsA<UNiagaraNodeParameterMapBase>())
@@ -439,12 +448,9 @@ void UNiagaraNodeWithDynamicPins::AddParameter(const UNiagaraScriptVariable* Scr
 		Graph->AddParameter(ScriptVar);
 
 		Modify();
-		UEdGraphPin* Pin = this->RequestNewTypedPin(Direction, Parameter.GetType(), Parameter.GetName());
-	}
-	else
-	{
-		RequestNewTypedPin(Direction, Parameter.GetType(), Parameter.GetName());
+		return this->RequestNewTypedPin(Direction, Parameter.GetType(), Parameter.GetName());
 	}
+	return RequestNewTypedPin(Direction, Parameter.GetType(), Parameter.GetName());
 }
 
 void UNiagaraNodeWithDynamicPins::AddExistingParameter(FNiagaraVariable Parameter, const UEdGraphPin* AddPin)
diff --git a/Engine/Plugins/FX/Niagara/Source/NiagaraEditor/Private/ViewModels/Stack/NiagaraStackGraphUtilities.cpp b/Engine/Plugins/FX/Niagara/Source/NiagaraEditor/Private/ViewModels/Stack/NiagaraStackGraphUtilities.cpp
--- a/Engine/Plugins/FX/Niagara/Source/NiagaraEditor/Private/ViewModels/Stack/NiagaraStackGraphUtilities.cpp
+++ b/Engine/Plugins/FX/Niagara/Source/NiagaraEditor/Private/ViewModels/Stack/NiagaraStackGraphUtilities.cpp
@@ -121,13 +121,14 @@ void FNiagaraStackGraphUtilities::RelayoutGraph(UEdGraph& Graph)
 	TArray<TArray<TArray<UEdGraphNode*>>> OutputNodeTraversalStacks;
 	TArray<UNiagaraNodeOutput*> OutputNodes;
 	Graph.GetNodesOfClass(OutputNodes);
-	TSet<UEdGraphNode*> AllTraversedNodes;
+	TMap<UEdGraphNode*, int32> AllTraversedNodes; // track indentation level per node
 	for (UNiagaraNodeOutput* OutputNode : OutputNodes)
 	{
-		TSet<UEdGraphNode*> TraversedNodes;
+		TMap<UEdGraphNode*, int32> TraversedNodes; // map to track node level
 		TArray<TArray<UEdGraphNode*>> TraversalStack;
 		TArray<UEdGraphNode*> CurrentNodesToTraverse;
 		CurrentNodesToTraverse.Add(OutputNode);
+		int32 NodeLevel = 0;
 		while (CurrentNodesToTraverse.Num() > 0)
 		{
 			TArray<UEdGraphNode*> TraversedNodesThisLevel;
@@ -136,6 +137,11 @@ void FNiagaraStackGraphUtilities::RelayoutGraph(UEdGraph& Graph)
 			{
 				if (TraversedNodes.Contains(CurrentNodeToTraverse))
 				{
+					if (NodeLevel > TraversedNodes[CurrentNodeToTraverse])
+					{
+						// if the node is referenced in a longer chain we pull it back
+						TraversedNodes[CurrentNodeToTraverse] = NodeLevel;
+					}
 					continue;
 				}
 				
@@ -152,18 +158,31 @@ void FNiagaraStackGraphUtilities::RelayoutGraph(UEdGraph& Graph)
 						}
 					}
 				}
-				TraversedNodes.Add(CurrentNodeToTraverse);
+				TraversedNodes.Add(CurrentNodeToTraverse, NodeLevel);
 				TraversedNodesThisLevel.Add(CurrentNodeToTraverse);
 			}
 			TraversalStack.Add(TraversedNodesThisLevel);
+			NodeLevel++;
 			CurrentNodesToTraverse.Empty();
 			CurrentNodesToTraverse.Append(NextNodesToTraverse);
 		}
 		OutputNodeTraversalStacks.Add(TraversalStack);
-		AllTraversedNodes = AllTraversedNodes.Union(TraversedNodes);
+
+		for (const auto& Pair : TraversedNodes)
+		{
+			if (AllTraversedNodes.Contains(Pair.Key))
+			{
+				// if the node was already traversed, update the indentation level to get the longest chain
+				AllTraversedNodes[Pair.Key] = FMath::Max(AllTraversedNodes[Pair.Key], Pair.Value);
+			}
+			else
+			{
+				AllTraversedNodes.Add(Pair);
+			}
+		}
 	}
 
-	// Find all nodes which were not traversed and put them them in a separate traversal stack.
+	// Find all nodes which were not traversed and put them in a separate traversal stack.
 	TArray<UEdGraphNode*> UntraversedNodes;
 	for (UEdGraphNode* Node : Graph.Nodes)
 	{
@@ -180,21 +199,23 @@ void FNiagaraStackGraphUtilities::RelayoutGraph(UEdGraph& Graph)
 		UntraversedNodeStack.Add(UntraversedStackItem);
 	}
 	OutputNodeTraversalStacks.Add(UntraversedNodeStack);
-
+	
 	// Layout the traversed node stacks.
 	float YOffset = 0;
-	float XDistance = 400;
+	float XDistance = -400;
 	float YDistance = 50;
 	float YPinDistance = 50;
 	for (const TArray<TArray<UEdGraphNode*>>& TraversalStack : OutputNodeTraversalStacks)
 	{
-		float CurrentXOffset = 0;
+		
 		float MaxYOffset = YOffset;
 		for (const TArray<UEdGraphNode*>& TraversalLevel : TraversalStack)
 		{
 			float CurrentYOffset = YOffset;
 			for (UEdGraphNode* Node : TraversalLevel)
 			{
+				int32 Level = AllTraversedNodes.FindOrAdd(Node);
+				float CurrentXOffset = Level * XDistance;
 				Node->Modify();
 				Node->NodePosX = FMath::RoundToInt(CurrentXOffset);
 				Node->NodePosY = FMath::RoundToInt(CurrentYOffset);
@@ -215,7 +236,6 @@ void FNiagaraStackGraphUtilities::RelayoutGraph(UEdGraph& Graph)
 				CurrentYOffset += YDistance + (MaxPins * YPinDistance);
 			}
 			MaxYOffset = FMath::Max(MaxYOffset, CurrentYOffset);
-			CurrentXOffset -= XDistance;
 		}
 		YOffset = MaxYOffset + YDistance;
 	}
diff --git a/Engine/Plugins/FX/Niagara/Source/NiagaraEditor/Private/ViewModels/Stack/NiagaraStackScriptItemGroup.cpp b/Engine/Plugins/FX/Niagara/Source/NiagaraEditor/Private/ViewModels/Stack/NiagaraStackScriptItemGroup.cpp
--- a/Engine/Plugins/FX/Niagara/Source/NiagaraEditor/Private/ViewModels/Stack/NiagaraStackScriptItemGroup.cpp
+++ b/Engine/Plugins/FX/Niagara/Source/NiagaraEditor/Private/ViewModels/Stack/NiagaraStackScriptItemGroup.cpp
@@ -17,6 +17,7 @@
 #include "ScopedTransaction.h"
 #include "DragAndDrop/AssetDragDropOp.h"
 #include "Framework/Notifications/NotificationManager.h"
+#include "Interfaces/IMainFrameModule.h"
 #include "Internationalization/Internationalization.h"
 #include "Modules/ModuleManager.h"
 #include "Toolkits/NiagaraSystemToolkit.h"
@@ -31,11 +32,15 @@
 #include "ViewModels/Stack/NiagaraStackModuleItem.h"
 #include "Widgets/Notifications/SNotificationList.h"
 #include "Toolkits/SystemToolkitModes/NiagaraSystemToolkitModeBase.h"
+#include "Widgets/DataChannel/NiagaraDataChannelWizard.h"
+#include "Widgets/Wizard/SNiagaraModuleWizard.h"
 
 #include UE_INLINE_GENERATED_CPP_BY_NAME(NiagaraStackScriptItemGroup)
 
 #define LOCTEXT_NAMESPACE "UNiagaraStackScriptItemGroup"
 
+using namespace UE::Niagara::Wizard;
+
 class FScriptGroupAddAction : public INiagaraStackItemGroupAddAction
 {
 public:
@@ -123,6 +128,13 @@ public:
 		return MakeShareable(new FScriptGroupAddAction( DisplayName, {}, Description, FText(), true, true, FNiagaraVariable(), false, FAssetData(), nullptr, false, true));
 	}
 
+	static TSharedRef<FScriptGroupAddAction> CreateWizardAction(const FModuleWizardGenerator::FAction& Action)
+	{
+		TSharedRef<FScriptGroupAddAction> AddAction = MakeShareable(new FScriptGroupAddAction( Action.DisplayName, {}, Action.Description, Action.Keywords, Action.bSuggestedAction, true, FNiagaraVariable(), false, FAssetData(), nullptr, false, true));
+		AddAction->ModuleWizardModel = Action.WizardModel;
+		return AddAction;
+	}
+
 	virtual TArray<FString> GetCategories() const override
 	{
 		return Categories;
@@ -178,16 +190,27 @@ public:
 		return ModuleScript;
 	}
 
-	bool GetIsNewScratchModuleAction() const
+	bool IsNewScratchModuleAction() const
 	{
 		return bIsNewScratchModuleAction;
 	}
 
-	bool GetIsNewSetSpecificModuleAction() const
+	bool IsNewSetSpecificModuleAction() const
 	{
 		return bIsNewSetSpecificModuleAction;
 	}
 
+	bool HasWizardWidget() const
+	{
+		return ModuleWizardModel.IsValid();
+	}
+
+	TSharedRef<SNiagaraModuleWizard> CreateWizard() const
+	{
+		check(HasWizardWidget());
+		return SNew(SNiagaraModuleWizard, ModuleWizardModel.ToSharedRef());
+	}
+
 private:
 	FScriptGroupAddAction(FText InDisplayName, TArray<FString> InCategories, FText InDescription, FText InKeywords, bool bInSuggested, bool bInIsInLibrary,
 		FNiagaraVariable InModuleParameterVariable, bool bInRenameParameterOnAdd,
@@ -223,6 +246,7 @@ private:
 	bool bIsNewScratchModuleAction;
 	bool bIsNewSetSpecificModuleAction;
 	FNiagaraActionSourceData SourceData;
+	TSharedPtr<FModuleWizardModel> ModuleWizardModel;
 };
 
 class FScriptItemGroupAddUtilities : public TNiagaraStackItemGroupAddUtilities<UNiagaraNodeFunctionCall*>
@@ -253,9 +277,10 @@ public:
 
 		// Generate actions for adding script asset modules.
 		TArray<FAssetData> ModuleAssets;
+		ENiagaraScriptUsage TargetUsage = OutputNode->GetUsage();
 		FNiagaraEditorUtilities::FGetFilteredScriptAssetsOptions ModuleScriptFilterOptions;
 		ModuleScriptFilterOptions.ScriptUsageToInclude = ENiagaraScriptUsage::Module;
-		ModuleScriptFilterOptions.TargetUsageToMatch = OutputNode->GetUsage();
+		ModuleScriptFilterOptions.TargetUsageToMatch = TargetUsage;
 		ModuleScriptFilterOptions.bIncludeDeprecatedScripts = AddProperties.bIncludeDeprecated;
 		ModuleScriptFilterOptions.bIncludeNonLibraryScripts = AddProperties.bIncludeNonLibrary;
 		FNiagaraEditorUtilities::GetFilteredScriptAssets(ModuleScriptFilterOptions, ModuleAssets);
@@ -271,7 +296,7 @@ public:
 			if (ScratchPadScript->Usage == ENiagaraScriptUsage::Module)
 			{
 				TArray<ENiagaraScriptUsage> SupportedUsages = ScratchPadScript->GetLatestScriptData()->GetSupportedUsageContexts();
-				if (SupportedUsages.Contains(OutputNode->GetUsage()))
+				if (SupportedUsages.Contains(TargetUsage))
 				{
 					OutAddActions.Add(FScriptGroupAddAction::CreateModuleActionFromScratchPadScript(ScratchPadScript));
 				}
@@ -280,6 +305,16 @@ public:
 
 		OutAddActions.Add(FScriptGroupAddAction::CreateNewScratchModuleAction());
 		OutAddActions.Add(FScriptGroupAddAction::CreateNewSetSpecificModuleAction());
+
+		// Generate wizard actions
+		FNiagaraEditorModule& NiagaraEditorModule = FModuleManager::GetModuleChecked<FNiagaraEditorModule>("NiagaraEditor");
+		for (TSharedRef<FModuleWizardGenerator> WizardGenerator : NiagaraEditorModule.GetModuleWizards())
+		{
+			for (const FModuleWizardGenerator::FAction& Action : WizardGenerator->CreateWizardActions(TargetUsage))
+			{
+				OutAddActions.Add(FScriptGroupAddAction::CreateWizardAction(Action));
+			}
+		}
 	}
 
 	virtual void ExecuteAddAction(TSharedRef<INiagaraStackItemGroupAddAction> AddAction, int32 TargetIndex) override
@@ -308,7 +343,52 @@ public:
 		{
 			NewModuleNode = FNiagaraStackGraphUtilities::AddScriptModuleToStack(ScriptGroupAddAction->GetModuleScript(), *OutputNode, TargetIndex);
 		}
-		else if (ScriptGroupAddAction->GetIsNewScratchModuleAction())
+		else if (ScriptGroupAddAction->HasWizardWidget())
+		{
+			// create the wizard widget
+			TSharedRef<SNiagaraModuleWizard> ModuleWizard = ScriptGroupAddAction->CreateWizard();
+			ModuleWizard->Model->TargetUsage = FNiagaraStackGraphUtilities::GetOutputNodeUsage(*OutputNode);
+			ModuleWizard->OnCreateModule = FSimpleDelegate::CreateLambda([TargetIndex, ModuleWizard = ModuleWizard.ToSharedPtr(), SystemViewModel = SystemViewModel, EmitterViewModel = EmitterViewModel, OutputNode = TWeakObjectPtr<UNiagaraNodeOutput>(OutputNode), StackEditorData = TWeakObjectPtr<UNiagaraStackEditorData>(&StackEditorData)]()
+			{
+				if (SystemViewModel.IsValid() && OutputNode.IsValid() && StackEditorData.IsValid())
+				{
+					// create the module and add it to the stack
+					FScopedTransaction AddModuleTransaction(LOCTEXT("AddWizardModuleTransaction", "Add module via wizard"));
+					TSharedPtr<FNiagaraSystemViewModel> SystemModel = SystemViewModel.Pin();
+					SystemModel->GetSelectionViewModel()->EmptySelection();
+					TArray<const UNiagaraNodeFunctionCall*> AddedModules = ModuleWizard->AddModulesToStack(SystemModel, OutputNode.Get(), TargetIndex, EmitterViewModel.Pin(), StackEditorData.Get());
+					for (const UNiagaraNodeFunctionCall* ModuleNode : AddedModules)
+					{
+						SystemModel->GetSelectionViewModel()->AddEntryToSelectionByDisplayedObjectDeferred(ModuleNode);
+					}
+					SystemModel->RefreshAll();
+				}
+			});
+
+			// create a new window with the wizard
+			const FVector2D CursorPos = FSlateApplication::Get().GetCursorPos();
+			const TSharedRef<SWindow> AddModuleWindow = SNew(SWindow)
+				.Title(LOCTEXT("Wizard_Title", "Add new module"))
+				.CreateTitleBar(true)
+				.ScreenPosition(CursorPos)
+				.FocusWhenFirstShown(true)
+				.ClientSize(FVector2D(700, 500))
+				[
+					ModuleWizard
+				];
+
+			const IMainFrameModule& MainFrameModule = FModuleManager::LoadModuleChecked<IMainFrameModule>(TEXT("MainFrame"));
+			const TSharedPtr<SWindow> ParentWindow = MainFrameModule.GetParentWindow();
+			if (ParentWindow.IsValid())
+			{
+				FSlateApplication::Get().AddWindowAsNativeChild(AddModuleWindow, ParentWindow.ToSharedRef());
+			}
+			else
+			{
+				FSlateApplication::Get().AddWindow(AddModuleWindow);
+			}
+		}
+		else if (ScriptGroupAddAction->IsNewScratchModuleAction())
 		{
 			TSharedPtr<FNiagaraScratchPadScriptViewModel> ScratchPadScriptViewModel = SystemViewModel.Pin()->GetScriptScratchPadViewModel()->CreateNewScript(ENiagaraScriptUsage::Module, OutputNode->GetUsage(), FNiagaraTypeDefinition());
 			if (ScratchPadScriptViewModel.IsValid())
@@ -319,7 +399,7 @@ public:
 				ScratchPadScriptViewModel->SetIsPendingRename(true);
 			}
 		}
-		else if (ScriptGroupAddAction->GetIsNewSetSpecificModuleAction())
+		else if (ScriptGroupAddAction->IsNewSetSpecificModuleAction())
 		{
 			NewModuleNode = AddParameterModule(FNiagaraVariable(), false, TargetIndex);
 		}
@@ -1003,8 +1083,6 @@ TOptional<UNiagaraStackEntry::FDropRequestResponse> UNiagaraStackScriptItemGroup
 TOptional<UNiagaraStackEntry::FDropRequestResponse> UNiagaraStackScriptItemGroup::CanDropScriptsOnTarget(const UNiagaraStackEntry& TargetEntry, const FDropRequest& DropRequest)
 {
 	TSharedRef<const FNiagaraScriptDragOperation> DragDropOp = StaticCastSharedRef<const FNiagaraScriptDragOperation>(DropRequest.DragDropOperation);
-	const UEnum* NiagaraScriptUsageEnum = FindObjectChecked<UEnum>(nullptr, TEXT("/Script/Niagara.ENiagaraScriptUsage"), true);
-
 	if (&TargetEntry != this && TargetEntry.IsA<UNiagaraStackModuleItem>() == false)
 	{
 		// Only handle drops onto this script group, or child drop requests from module items.
diff --git a/Engine/Plugins/FX/Niagara/Source/NiagaraEditor/Private/Widgets/DataChannel/NiagaraDataChannelWizard.cpp b/Engine/Plugins/FX/Niagara/Source/NiagaraEditor/Private/Widgets/DataChannel/NiagaraDataChannelWizard.cpp
new file mode 100644
--- /dev/null
+++ b/Engine/Plugins/FX/Niagara/Source/NiagaraEditor/Private/Widgets/DataChannel/NiagaraDataChannelWizard.cpp
@@ -0,0 +1,631 @@
+﻿// Copyright Epic Games, Inc. All Rights Reserved.
+
+#include "NiagaraDataChannelWizard.h"
+
+#include "EdGraphSchema_Niagara.h"
+#include "NiagaraClipboard.h"
+#include "NiagaraDataChannel.h"
+#include "NiagaraDataChannelPublic.h"
+#include "NiagaraEditorStyle.h"
+#include "NiagaraGraph.h"
+#include "NiagaraNodeFunctionCall.h"
+#include "NiagaraNodeParameterMapGet.h"
+#include "NiagaraNodeParameterMapSet.h"
+#include "DataInterface/NiagaraDataInterfaceDataChannelRead.h"
+#include "DataInterface/NiagaraDataInterfaceDataChannelWrite.h"
+#include "ViewModels/NiagaraScratchPadScriptViewModel.h"
+#include "ViewModels/NiagaraScriptGraphViewModel.h"
+#include "IDetailsView.h"
+#include "Widgets/Views/SListView.h"
+#include "Widgets/Layout/SSeparator.h"
+#include "Widgets/Text/STextBlock.h"
+#include "Widgets/Input/SEditableTextBox.h"
+#include "Widgets/Input/SCheckBox.h"
+
+
+using namespace UE::Niagara::Wizard;
+
+#define LOCTEXT_NAMESPACE "NiagaraDataChannelWizard"
+
+namespace UE::Niagara::Wizard::DataChannel
+{
+	struct FSelectAssetPageBase : FModuleWizardPage
+	{
+		FSelectAssetPageBase()
+		{
+			Name = LOCTEXT("AssetPageName", "Select asset");
+		}
+		virtual ~FSelectAssetPageBase() override = default;
+
+		virtual bool CanGoToNextPage() const override { return GetAsset() != nullptr; };
+		virtual bool CanCompleteWizard() const override { return CanGoToNextPage(); };
+
+		virtual UNiagaraDataChannelAsset* GetAsset() const = 0;
+		
+		UNiagaraDataChannel* GetDataChannel() const
+		{
+			if (UNiagaraDataChannelAsset* ChannelAsset = GetAsset())
+			{
+				return ChannelAsset->Get();
+			}			
+			return nullptr;
+		}
+		
+		TSharedRef<SWidget> GetDetailsViewContent(UObject* DetailsViewObject)
+		{
+			TSharedRef<IDetailsView> DetailsView = Utilities::CreateDetailsView();
+			DetailsView->SetObject(DetailsViewObject, true);
+			
+			return SNew(SVerticalBox)
+				+SVerticalBox::Slot()
+				.Padding(15)
+				.AutoHeight()
+				[
+					SNew(STextBlock)
+					.Text(LOCTEXT("AssetPageLabel", "Please select which data channel you want to use"))
+				]
+				+SVerticalBox::Slot()
+				[
+					DetailsView
+				];
+		}
+	};
+
+	struct FSelectVariablesPageBase : FModuleWizardPage
+	{
+		explicit FSelectVariablesPageBase(FSelectAssetPageBase* InPreviousPage) : PreviousPage(InPreviousPage)
+		{
+			Name = LOCTEXT("VariablesPageName", "Select variables");
+		}
+
+		virtual ~FSelectVariablesPageBase() override = default;
+
+		virtual bool CanGoToNextPage() const override { return VariablesToProcess.Num() > 0; };
+		virtual bool CanCompleteWizard() const override { return AllVariables.IsEmpty() || CanGoToNextPage(); };
+
+		virtual FText GetHeaderLabel() = 0;
+
+		virtual void RefreshContent() override
+		{
+			FObjectKey NewDataChannelRef;
+			TArray<FNiagaraDataChannelVariable> DataChannelVariables;
+			if (UNiagaraDataChannelAsset* ChannelAsset = PreviousPage->GetAsset())
+			{
+				NewDataChannelRef = ChannelAsset;
+				DataChannelVariables = ChannelAsset->Get()->GetVariables();
+			}
+			if (NewDataChannelRef != LastDataChannelRef)
+			{
+				ModuleName = CreateNewModuleName(); 
+			} 
+
+			bool bCheckAll = AllVariables.IsEmpty() || NewDataChannelRef != LastDataChannelRef;
+			AllVariables.Empty(DataChannelVariables.Num());
+			for (const FNiagaraDataChannelVariable& Var : DataChannelVariables)
+			{
+				*AllVariables.Add_GetRef(MakeShared<FNiagaraDataChannelVariable>()).Get() = Var;
+				if (bCheckAll)
+				{
+					VariablesToProcess.Add(Var.Version);
+				}
+			}
+
+			if (VarListView.IsValid())
+			{
+				VarListView->RebuildList();
+			}
+			LastDataChannelRef = NewDataChannelRef;
+		}
+
+		virtual TSharedRef<SWidget> GetContent() override
+		{
+			return SNew(SVerticalBox)
+				+SVerticalBox::Slot()
+				.Padding(2)
+				.AutoHeight()
+				[
+					SNew(SSeparator)
+					.Orientation(Orient_Horizontal)
+				]
+				+SVerticalBox::Slot()
+				.Padding(15)
+				.AutoHeight()
+				[
+					SNew(STextBlock)
+					.Text(GetHeaderLabel())
+				]
+				+SVerticalBox::Slot()
+				[
+					SAssignNew(VarListView, SListView<TSharedPtr<FNiagaraDataChannelVariable>>)
+						.ListItemsSource(&AllVariables)
+						.OnGenerateRow(this, &FSelectVariablesPageBase::GenerateRow)
+						.SelectionMode(ESelectionMode::Single)
+				]
+				+SVerticalBox::Slot()
+				.Padding(0, 10)
+				.AutoHeight()
+				[
+					SNew(SHorizontalBox)
+					+ SHorizontalBox::Slot()
+					.AutoWidth()
+					.VAlign(VAlign_Center)
+					[
+						SNew(STextBlock)
+						.Text(LOCTEXT("WriteModuleNameText", "Module Name: "))
+					]
+					+ SHorizontalBox::Slot()
+					.AutoWidth()
+					.VAlign(VAlign_Center)
+					.Padding(4.0f, 0)
+					[
+						SNew( SEditableTextBox )
+						.MinDesiredWidth(200)
+						.Text(this, &FSelectVariablesPageBase::GetModuleNameText)
+						.SelectAllTextWhenFocused(true)
+						.ClearKeyboardFocusOnCommit(false)
+						.OnTextCommitted(this, &FSelectVariablesPageBase::SetModuleName)
+					]
+				];
+		}
+
+		FText GetModuleNameText() const
+		{
+			return ModuleName;
+		}
+
+		void SetModuleName(const FText& NewName, ETextCommit::Type)
+		{
+			ModuleName = NewName;
+		}
+
+		FText CreateNewModuleName() const
+		{
+			FText AssetName;
+			if (UNiagaraDataChannelAsset* DataChannelAsset = PreviousPage->GetAsset())
+			{
+				AssetName = FText::FromString(DataChannelAsset->GetName());
+			}
+			return GetFormattedModuleName(AssetName);
+		}
+
+		virtual FText GetFormattedModuleName(const FText& AssetName) const = 0;
+
+		TSharedRef<ITableRow> GenerateRow(const TSharedPtr<FNiagaraDataChannelVariable> Var, const TSharedRef<STableViewBase>& OwnerTable)
+		{
+			FLinearColor TypeColor = UEdGraphSchema_Niagara::GetTypeColor(Var->GetType());
+			return SNew(STableRow<TSharedPtr<FString>>, OwnerTable)
+				.Padding(FMargin(5, 0))
+				[
+					SNew(SCheckBox)
+					.OnCheckStateChanged(this, &FSelectVariablesPageBase::OnCheckStateChanged, *Var.Get())
+					.IsChecked(this, &FSelectVariablesPageBase::OnGetCheckState, *Var.Get())
+					.ToolTipText(FText::Format(LOCTEXT("VariablesSelectionTooltipFmt", "Name: {0}\nType: {1}"), FText::FromName(Var->GetName()), Var->GetType().GetNameText()))
+					[
+						SNew(SHorizontalBox)
+						+ SHorizontalBox::Slot()
+						.VAlign(VAlign_Center)
+						.HAlign(HAlign_Center)
+						.AutoWidth()
+						[
+							SNew(SImage)
+							.ColorAndOpacity(TypeColor)
+							.Image(FNiagaraEditorStyle::Get().GetBrush("NiagaraEditor.Module.TypeIconPill"))
+						]
+						+ SHorizontalBox::Slot()
+						.Padding(4, 2, 2, 2)
+						[
+							SNew(STextBlock)
+							.MinDesiredWidth(150)
+							.Text(FText::FromName(Var->GetName()))
+						]
+					]
+				];
+		}
+
+		void OnCheckStateChanged(const ECheckBoxState NewState, FNiagaraDataChannelVariable Var)
+		{
+			if (NewState == ECheckBoxState::Checked)
+			{
+				VariablesToProcess.Add(Var.Version);
+			}
+			else if (NewState == ECheckBoxState::Unchecked)
+			{
+				VariablesToProcess.Remove(Var.Version);
+			}
+		}
+
+		ECheckBoxState OnGetCheckState(FNiagaraDataChannelVariable Var) const
+		{
+			return VariablesToProcess.Contains(Var.Version) ? ECheckBoxState::Checked : ECheckBoxState::Unchecked;
+		}
+
+		TArray<TSharedPtr<FNiagaraDataChannelVariable>> AllVariables;
+		TSet<FGuid> VariablesToProcess;
+		FSelectAssetPageBase* PreviousPage;
+		FObjectKey LastDataChannelRef;
+		FText ModuleName;
+
+		TSharedPtr<SListView<TSharedPtr<FNiagaraDataChannelVariable>>> VarListView;
+	};
+}
+
+TSharedRef<FModuleWizardModel> DataChannel::CreateReadNDCModuleWizardModel()
+{
+	struct FSelectAssetPage : FSelectAssetPageBase
+	{
+		virtual UNiagaraDataChannelAsset* GetAsset() const override
+		{
+			if (UNiagaraDataChannelReadModuleData* ModuleData = Data.Get())
+			{
+				return ModuleData->DataChannel;
+			}
+			return nullptr;
+		}
+		
+		virtual TSharedRef<SWidget> GetContent() override
+		{
+			Data.Reset(NewObject<UNiagaraDataChannelReadModuleData>());
+			return GetDetailsViewContent(Data.Get());
+		}
+
+		
+		TStrongObjectPtr<UNiagaraDataChannelReadModuleData> Data;
+	};
+
+	struct FSelectVariablesPage : FSelectVariablesPageBase
+	{
+		explicit FSelectVariablesPage(FSelectAssetPage* InPreviousPage) : FSelectVariablesPageBase(InPreviousPage)
+		{}
+		virtual ~FSelectVariablesPage() override = default;
+
+		virtual FText GetHeaderLabel() override
+		{
+			return LOCTEXT("VariablesPageLabel", "Please select the variables to read from the data channel");
+		}
+		
+		virtual FText GetFormattedModuleName(const FText& AssetName) const override
+		{
+			return FText::Format(LOCTEXT("ReadModuleNameFmt", "Read {0}"), AssetName);
+		}
+	};
+	
+	struct FReadNDCModel : FModuleWizardModel
+	{
+		FReadNDCModel()
+		{
+			AssetPage = MakeShared<FSelectAssetPage>();
+			VariablesPage = MakeShared<FSelectVariablesPage>(AssetPage.Get());
+			Pages.Add(AssetPage.ToSharedRef());
+			Pages.Add(VariablesPage.ToSharedRef());
+		}
+		virtual ~FReadNDCModel() override = default;
+		
+		virtual void GenerateNewModuleContent(TSharedPtr<FNiagaraScratchPadScriptViewModel> ScratchPadScriptViewModel, const TArray<const UNiagaraNodeFunctionCall*>& PreviousModules) override
+		{
+			FText ScriptName = VariablesPage->ModuleName;
+			ScratchPadScriptViewModel->SetScriptName(ScriptName.IsEmptyOrWhitespace() ? VariablesPage->CreateNewModuleName() : ScriptName);
+			
+			UNiagaraDataChannel* Channel = AssetPage->GetDataChannel();
+			UNiagaraGraph* Graph = ScratchPadScriptViewModel->GetGraphViewModel()->GetGraph();
+			if (Channel && Graph)
+			{
+				const UEdGraphSchema_Niagara* GraphSchema = Graph->GetNiagaraSchema();
+				UNiagaraNodeParameterMapGet* MapGetNode = Utilities::FindSingleNodeChecked<UNiagaraNodeParameterMapGet>(Graph);
+				UNiagaraNodeParameterMapSet* MapSetNode = Utilities::FindSingleNodeChecked<UNiagaraNodeParameterMapSet>(Graph);
+
+				// Add inputs
+				UEdGraphPin* DIPin = Utilities::AddReadParameterPin(FNiagaraTypeDefinition(UNiagaraDataInterfaceDataChannelRead::StaticClass()), FName("Data Channel"), MapGetNode);
+				UEdGraphPin* IndexPin = Utilities::AddReadParameterPin(FNiagaraTypeDefinition::GetIntDef(), FName("Read Index"), MapGetNode);
+
+				// Call read function
+				if (UNiagaraNodeFunctionCall* ReadFunction = Utilities::CreateDataInterfaceFunctionNode(UNiagaraDataInterfaceDataChannelRead::StaticClass(), FName("Read"), Graph))
+				{
+					// connect index input pins
+					ReadFunction->AutowireNewNode(DIPin);
+					UEdGraphPin* IndexInput = ReadFunction->GetInputPin(1);
+					if (IndexInput && IndexInput->GetName() == TEXT("Index"))
+					{
+						GraphSchema->TryCreateConnection(IndexPin, IndexInput);
+					}
+
+					// create and connect read success output pin
+					UEdGraphPin* SuccessVarPin = Utilities::AddWriteParameterPin(FNiagaraTypeDefinition::GetBoolDef(), FName("Output.Module.ReadSuccess"), MapSetNode);
+					UEdGraphPin* SuccessOutPin = ReadFunction->GetOutputPin(0);
+					if (SuccessOutPin && SuccessOutPin->GetName() == TEXT("Success"))
+					{
+						GraphSchema->TryCreateConnection(SuccessOutPin, SuccessVarPin);
+					}
+
+					// add channel variable pins to read node
+					for (const FNiagaraDataChannelVariable& Var : Channel->GetVariables())
+					{
+						if (!VariablesPage->VariablesToProcess.Contains(Var.Version))
+						{
+							continue;
+						}
+						
+						FNiagaraTypeDefinition SwcType = Var.GetType();
+						if (SwcType.IsEnum() == false)
+						{
+							SwcType = FNiagaraTypeDefinition(FNiagaraTypeHelper::GetSWCStruct(Var.GetType().GetScriptStruct()));
+						}
+						FNiagaraVariable SWCVar(SwcType, Var.GetName());
+						UEdGraphPin* ReadParamPin = ReadFunction->AddParameterPin(SWCVar, EGPD_Output);
+
+						// add matching node on map set and connect them
+						UEdGraphPin* SetVarPin = Utilities::AddWriteParameterPin(SwcType, FName(TEXT("Output.Module.") + Var.GetName().ToString()), MapSetNode);
+						if (ReadParamPin && SetVarPin)
+						{
+							GraphSchema->TryCreateConnection(ReadParamPin, SetVarPin);
+						}
+					}
+				}
+				
+				FNiagaraStackGraphUtilities::RelayoutGraph(*Graph);
+				ScratchPadScriptViewModel->ApplyChanges();
+			}
+		}
+		
+		virtual bool UpdateModuleInputs(UNiagaraClipboardContent* NewModule, const TArray<const UNiagaraNodeFunctionCall*>& PreviousModules) override
+		{
+			if (UNiagaraDataChannelAsset* Channel = AssetPage->GetAsset())
+			{
+				TArray<TObjectPtr<const UNiagaraClipboardFunctionInput>> FunctionInputs = NewModule->FunctionInputs;
+				for (const UNiagaraClipboardFunctionInput* FunctionInput : FunctionInputs)
+				{
+					if (FunctionInput->InputType == FNiagaraTypeDefinition(UNiagaraDataInterfaceDataChannelRead::StaticClass()))
+					{
+						// set data interface module input
+						if (UNiagaraDataInterfaceDataChannelRead* DataInterface = Cast<UNiagaraDataInterfaceDataChannelRead>(FunctionInput->Data))
+						{
+							DataInterface->Channel = Channel;
+							DataInterface->bReadCurrentFrame = AssetPage->Data->bReadCurrentFrame;
+							DataInterface->bUpdateSourceDataEveryTick = AssetPage->Data->bUpdateSourceDataEveryTick;
+						}
+					}
+
+					// bind index input for particle scripts. System and emitter scripts default to 0
+					if (FunctionInput->InputType == FNiagaraTypeDefinition::GetIntDef() && FNiagaraUtilities::ConvertScriptUsageToStaticSwitchContext(TargetUsage) == ENiagaraScriptContextStaticSwitch::Particle)
+					{
+						UNiagaraClipboardFunctionInput* EditableInput = const_cast<UNiagaraClipboardFunctionInput*>(FunctionInput);
+						EditableInput->ValueMode = ENiagaraClipboardFunctionInputValueMode::Linked;
+						EditableInput->Linked = SYS_PARAM_PARTICLES_UNIQUE_ID.GetName();
+					}
+				}
+				return true;
+			}
+			return false;
+		}
+
+		TSharedPtr<FSelectAssetPage> AssetPage;
+		TSharedPtr<FSelectVariablesPage> VariablesPage;
+	};
+	
+	return MakeShared<FReadNDCModel>();
+}
+
+TSharedRef<FModuleWizardModel> DataChannel::CreateWriteNDCModuleWizardModel()
+{
+	struct FSelectAssetPage : FSelectAssetPageBase
+	{
+		virtual UNiagaraDataChannelAsset* GetAsset() const override
+		{
+			if (UNiagaraDataChannelWriteModuleData* ModuleData = Data.Get())
+			{
+				return ModuleData->DataChannel;
+			}
+			return nullptr;
+		}
+		
+		virtual TSharedRef<SWidget> GetContent() override
+		{
+			Data.Reset(NewObject<UNiagaraDataChannelWriteModuleData>());
+			return GetDetailsViewContent(Data.Get());
+		}
+
+		TStrongObjectPtr<UNiagaraDataChannelWriteModuleData> Data;
+	};
+
+	struct FSelectVariablesPage : FSelectVariablesPageBase
+	{
+		explicit FSelectVariablesPage(FSelectAssetPage* InPreviousPage) : FSelectVariablesPageBase(InPreviousPage)
+		{}
+		virtual ~FSelectVariablesPage() override = default;
+
+		virtual FText GetHeaderLabel() override
+		{
+			return LOCTEXT("VariablesWritePageLabel", "Please select which data channel variables to write to");
+		}
+
+		virtual FText GetFormattedModuleName(const FText& AssetName) const override
+		{
+			return FText::Format(LOCTEXT("WriteModuleNameFmt", "Write {0}"), AssetName);
+		}
+	};
+	
+	struct FWriteNDCModel : FModuleWizardModel
+	{
+		FWriteNDCModel()
+		{
+			AssetPage = MakeShared<FSelectAssetPage>();
+			VariablesPage = MakeShared<FSelectVariablesPage>(AssetPage.Get());
+			Pages.Add(AssetPage.ToSharedRef());
+			Pages.Add(VariablesPage.ToSharedRef());
+		}
+		virtual ~FWriteNDCModel() override = default;
+		
+		virtual void GenerateNewModuleContent(TSharedPtr<FNiagaraScratchPadScriptViewModel> ScratchPadScriptViewModel, const TArray<const UNiagaraNodeFunctionCall*>& PreviousModules) override
+		{
+			FText ScriptName = VariablesPage->ModuleName;
+			ScratchPadScriptViewModel->SetScriptName(ScriptName.IsEmptyOrWhitespace() ? VariablesPage->CreateNewModuleName() : ScriptName);
+			
+			UNiagaraDataChannel* Channel = AssetPage->GetDataChannel();
+			UNiagaraGraph* Graph = ScratchPadScriptViewModel->GetGraphViewModel()->GetGraph();
+			if (Channel == nullptr || Graph == nullptr)
+			{
+				return;
+			}
+			const UEdGraphSchema_Niagara* GraphSchema = Graph->GetNiagaraSchema();
+			UNiagaraNodeParameterMapGet* MapGetNode = Utilities::FindSingleNodeChecked<UNiagaraNodeParameterMapGet>(Graph);
+			UNiagaraNodeParameterMapSet* MapSetNode = Utilities::FindSingleNodeChecked<UNiagaraNodeParameterMapSet>(Graph);
+			UNiagaraNodeInput* InputNode = Utilities::FindSingleNodeChecked<UNiagaraNodeInput>(Graph);
+			ENiagaraDataChanneWriteModuleMode WriteMode = AssetPage->Data->WriteMode;
+			
+			// Add inputs
+			UEdGraphPin* DIPin = Utilities::AddReadParameterPin(FNiagaraTypeDefinition(UNiagaraDataInterfaceDataChannelWrite::StaticClass()), FName("Data Channel"), MapGetNode);
+			UEdGraphPin* ExecWritePin = Utilities::AddReadParameterPin(FNiagaraTypeDefinition::GetBoolDef(), FName("Execute Write"), MapGetNode);
+			SetBoolDefaultValue(Graph, ExecWritePin->PinName, true);
+			UEdGraphPin* IndexPin = nullptr;
+			if (WriteMode == ENiagaraDataChanneWriteModuleMode::WriteToExistingElement)
+			{
+				IndexPin = Utilities::AddReadParameterPin(FNiagaraTypeDefinition::GetIntDef(), FName("Write Index"), MapGetNode);
+			}
+
+			// Call write function
+			FName FunctionName = WriteMode == ENiagaraDataChanneWriteModuleMode::AppendNewElement ? FName("Append") : FName("Write");
+			if (UNiagaraNodeFunctionCall* WriteFunction = Utilities::CreateDataInterfaceFunctionNode(UNiagaraDataInterfaceDataChannelWrite::StaticClass(), FunctionName, Graph))
+			{
+				// connect default function pins
+				WriteFunction->AutowireNewNode(DIPin);
+				GraphSchema->TryCreateConnection(InputNode->GetOutputPin(0), WriteFunction->GetInputPin(0));
+				GraphSchema->TryCreateConnection(WriteFunction->GetOutputPin(0), MapSetNode->GetInputPin(0));
+				
+				UEdGraphPin* ExecInput = WriteFunction->GetInputPin(2);
+				if (ExecInput && ExecInput->GetName() == TEXT("Emit"))
+				{
+					GraphSchema->TryCreateConnection(ExecWritePin, ExecInput);
+				}
+				
+				UEdGraphPin* IndexInput = WriteFunction->GetInputPin(3);
+				if (IndexInput && IndexInput->GetName() == TEXT("Index"))
+				{
+					GraphSchema->TryCreateConnection(IndexPin, IndexInput);
+				}
+
+				// create and connect write success output pin
+				UEdGraphPin* SuccessVarPin = Utilities::AddWriteParameterPin(FNiagaraTypeDefinition::GetBoolDef(), FName("Output.Module.WriteSuccess"), MapSetNode);
+				UEdGraphPin* SuccessOutPin = WriteFunction->GetOutputPin(1);
+				if (SuccessOutPin && SuccessOutPin->GetName() == TEXT("Success"))
+				{
+					GraphSchema->TryCreateConnection(SuccessOutPin, SuccessVarPin);
+				}
+
+				// add channel variable pins to write node
+				for (const FNiagaraDataChannelVariable& Var : Channel->GetVariables())
+				{
+					if (!VariablesPage->VariablesToProcess.Contains(Var.Version))
+					{
+						continue;
+					}
+					
+					FNiagaraTypeDefinition SwcType = Var.GetType();
+					if (SwcType.IsEnum() == false)
+					{
+						SwcType = FNiagaraTypeDefinition(FNiagaraTypeHelper::GetSWCStruct(Var.GetType().GetScriptStruct()));
+					}
+					FNiagaraVariable SWCVar(SwcType, Var.GetName());
+					UEdGraphPin* WriteParamPin = WriteFunction->AddParameterPin(SWCVar, EGPD_Input);
+
+					// add matching node on map get and connect them
+					UEdGraphPin* SetVarPin = Utilities::AddReadParameterPin(SwcType, FName(TEXT("Module.") + Var.GetName().ToString()), MapGetNode);
+					if (WriteParamPin && SetVarPin)
+					{
+						GraphSchema->TryCreateConnection(WriteParamPin, SetVarPin);
+					}
+				}
+				
+				FNiagaraStackGraphUtilities::RelayoutGraph(*Graph);
+				ScratchPadScriptViewModel->ApplyChanges();
+			}
+		}
+		
+		virtual bool UpdateModuleInputs(UNiagaraClipboardContent* NewModule, const TArray<const UNiagaraNodeFunctionCall*>& PreviousModules) override
+		{
+			if (UNiagaraDataChannelAsset* Channel = AssetPage->GetAsset())
+			{
+				TArray<TObjectPtr<const UNiagaraClipboardFunctionInput>> FunctionInputs = NewModule->FunctionInputs;
+				for (const UNiagaraClipboardFunctionInput* FunctionInput : FunctionInputs)
+				{
+					if (FunctionInput->InputType == FNiagaraTypeDefinition(UNiagaraDataInterfaceDataChannelWrite::StaticClass()))
+					{
+						// set data interface module input
+						if (UNiagaraDataInterfaceDataChannelWrite* DataInterface = Cast<UNiagaraDataInterfaceDataChannelWrite>(FunctionInput->Data))
+						{
+							DataInterface->Channel = Channel;
+							DataInterface->bPublishToGame = AssetPage->Data->bPublishToGame;
+							DataInterface->bPublishToCPU = AssetPage->Data->bPublishToCPU;
+							DataInterface->bPublishToGPU = AssetPage->Data->bPublishToGPU;
+							DataInterface->AllocationCount = AssetPage->Data->AllocationCount;
+							DataInterface->AllocationMode = AssetPage->Data->AllocationMode;
+							DataInterface->bUpdateDestinationDataEveryTick = AssetPage->Data->bUpdateDestinationDataEveryTick;
+						}
+					}
+
+					// bind index input for particle scripts. System and emitter scripts default to 0
+					if (FunctionInput->InputType == FNiagaraTypeDefinition::GetIntDef() && FunctionInput->InputName == FName("Write Index") && FNiagaraUtilities::ConvertScriptUsageToStaticSwitchContext(TargetUsage) == ENiagaraScriptContextStaticSwitch::Particle)
+					{
+						UNiagaraClipboardFunctionInput* EditableInput = const_cast<UNiagaraClipboardFunctionInput*>(FunctionInput);
+						EditableInput->ValueMode = ENiagaraClipboardFunctionInputValueMode::Linked;
+						EditableInput->Linked = SYS_PARAM_PARTICLES_UNIQUE_ID.GetName();
+					}
+				}
+				return true;
+			}
+			return false;
+		}
+
+		void SetBoolDefaultValue(UNiagaraGraph* Graph, const FName& VarName, bool Value) const
+		{
+			if (UNiagaraScriptVariable* ScriptVariable = Graph->GetScriptVariable(VarName))
+			{
+				FNiagaraVariable Var(FNiagaraTypeDefinition::GetBoolDef(), FName("Var"));
+				Var.SetValue(Value);
+				ScriptVariable->SetDefaultValueData(Var.GetData());
+				Graph->ScriptVariableChanged(ScriptVariable->Variable);
+			}
+		}
+
+		TSharedPtr<FSelectAssetPage> AssetPage;
+		TSharedPtr<FSelectVariablesPage> VariablesPage;
+	};
+		
+	return MakeShared<FWriteNDCModel>();
+}
+
+TSharedRef<FModuleWizardGenerator> DataChannel::CreateNDCWizardGenerator()
+{
+	class NDCWizardGenerator : public FModuleWizardGenerator
+	{
+	public:
+		virtual TArray<FAction> CreateWizardActions(ENiagaraScriptUsage Usage) override
+		{
+			TArray<FAction> WizardActions;
+
+			FAction& ReadAction = WizardActions.AddDefaulted_GetRef();
+			ReadAction.DisplayName = LOCTEXT("NewReadNDCModuleName", "Read From Data Channel...");
+			ReadAction.Description = LOCTEXT("NewReadNDCModuleDescription", "Description: Create a new scratch pad module to read attributes from a data channel");
+			ReadAction.Keywords = LOCTEXT("NewReadNDCModuleKeywords", "ndc reader datachannel get external");
+			ReadAction.bSuggestedAction = true;
+			ReadAction.WizardModel = CreateReadNDCModuleWizardModel();
+
+			FAction& WriteAction = WizardActions.AddDefaulted_GetRef();
+			WriteAction. DisplayName = LOCTEXT("NewWriteNDCModuleName", "Write To Data Channel...");
+			WriteAction.Description = LOCTEXT("NewWriteNDCModuleDescription", "Description: Create a new scratch pad module to write attributes to a data channel");
+			WriteAction.Keywords = LOCTEXT("NewWriteNDCModuleKeywords", "ndc writer datachannel save append external");
+			WriteAction.bSuggestedAction = true;
+			WriteAction.WizardModel = CreateWriteNDCModuleWizardModel();
+
+			if (Usage == ENiagaraScriptUsage::EmitterUpdateScript)
+			{
+				//TODO: add spawn from ndc wizard
+			}
+			
+			return WizardActions;
+		}
+
+		virtual ~NDCWizardGenerator() override = default;
+	};
+	return MakeShared<NDCWizardGenerator>();
+}
+
+#undef LOCTEXT_NAMESPACE
diff --git a/Engine/Plugins/FX/Niagara/Source/NiagaraEditor/Private/Widgets/DataChannel/NiagaraDataChannelWizard.h b/Engine/Plugins/FX/Niagara/Source/NiagaraEditor/Private/Widgets/DataChannel/NiagaraDataChannelWizard.h
new file mode 100644
--- /dev/null
+++ b/Engine/Plugins/FX/Niagara/Source/NiagaraEditor/Private/Widgets/DataChannel/NiagaraDataChannelWizard.h
@@ -0,0 +1,104 @@
+﻿// Copyright Epic Games, Inc. All Rights Reserved.
+
+#pragma once
+
+#include "NiagaraDataChannelPublic.h"
+#include "DataInterface/NiagaraDataInterfaceDataChannelCommon.h"
+#include "Widgets/Wizard/SNiagaraModuleWizard.h"
+#include "Templates/SharedPointer.h"
+
+#include "NiagaraDataChannelWizard.generated.h"
+
+class UNiagaraNodeParameterMapGet;
+class UNiagaraNodeParameterMapSet;
+class UNiagaraNodeFunctionCall;
+class UNiagaraGraph;
+class IDetailsView;
+
+/**
+ * Helper class to display properties in the read data channel module wizard
+ */
+UCLASS()
+class UNiagaraDataChannelReadModuleData : public UObject
+{
+	GENERATED_BODY()
+
+public:
+	/** The source asset to read from */
+	UPROPERTY(EditAnywhere, Category = "Source")
+	TObjectPtr<UNiagaraDataChannelAsset> DataChannel;
+
+	/** True if this reader will read the current frame's data. If false, we read the previous frame.
+	* Reading the previous frame's data introduces a frame of latency but ensures we never miss any data as we have access to the whole frame.
+	*/
+	UPROPERTY(EditAnywhere, Category = "Advanced")
+	bool bReadCurrentFrame = false;
+
+	/**
+	Whether this DI should request updated source data from the Data Channel each tick.
+	Some Data Channels have multiple separate source data elements for things such as spatial subdivision. 
+	Each DI will request the correct one for it's owning system instance from the data channel. 
+	Depending on the data channel this could be an expensive search so we should avoid doing this every tick if possible.
+	*/
+	UPROPERTY(EditAnywhere, Category = "Advanced")
+	bool bUpdateSourceDataEveryTick = true;
+};
+
+UENUM()
+enum class ENiagaraDataChanneWriteModuleMode : uint8
+{
+	AppendNewElement,
+	WriteToExistingElement
+};
+
+/**
+ * Helper class to display properties in the write data channel module wizard
+ */
+UCLASS()
+class UNiagaraDataChannelWriteModuleData : public UObject
+{
+	GENERATED_BODY()
+
+public:
+	/** The source asset to read from */
+	UPROPERTY(EditAnywhere, Category = "Source")
+	TObjectPtr<UNiagaraDataChannelAsset> DataChannel;
+
+	/** Should the generated module append to the data channel or write to an existing element? */
+	UPROPERTY(EditAnywhere, Category = "Source")
+	ENiagaraDataChanneWriteModuleMode WriteMode = ENiagaraDataChanneWriteModuleMode::AppendNewElement;
+
+	/** Whether the data generated by the niagara data interface should be published to the world game data channel. This is required to allow game BP and C++ to read this data. Setting this to true will have a minor performance impact. */
+	UPROPERTY(EditAnywhere, Category = "Source", meta = (DisplayName = "Visible to Blueprint"))
+	bool bPublishToGame = true;
+
+	/** Whether the data generated by the niagara data interface should be published to CPU emitters in other Niagara systems. Setting this to true will have a minor performance impact. */
+	UPROPERTY(EditAnywhere, Category = "Source", meta = (DisplayName = "Visible to CPU systems"))
+	bool bPublishToCPU = true;
+
+	/** Whether the data generated by the niagara data interface should be published to GPU emitters in other Niagara systems. Setting this to true will have a minor performance impact.  */
+	UPROPERTY(EditAnywhere, Category = "Source", meta = (DisplayName = "Visible to GPU systems"))
+	bool bPublishToGPU = true;
+
+	/** How should we allocate the buffer into which we write data. */
+	UPROPERTY(EditAnywhere, Category = "Advanced")
+	ENiagaraDataChannelAllocationMode AllocationMode = ENiagaraDataChannelAllocationMode::Static;
+
+	/** How many elements to allocate for writing per frame? Usage is defendant on AllocationMode. */
+	UPROPERTY(EditAnywhere, Category = "Advanced")
+	uint32 AllocationCount = 0;
+
+	/**
+	Whether this DI should request updated destination data from the Data Channel each tick.
+	Depending on the data channel this could be an expensive search so we should avoid doing this every tick if possible.
+	*/
+	UPROPERTY(EditAnywhere, Category = "Advanced")
+	bool bUpdateDestinationDataEveryTick = true;
+};
+
+namespace UE::Niagara::Wizard::DataChannel
+{
+	TSharedRef<FModuleWizardModel> CreateReadNDCModuleWizardModel();
+	TSharedRef<FModuleWizardModel> CreateWriteNDCModuleWizardModel();
+	TSharedRef<FModuleWizardGenerator> CreateNDCWizardGenerator();
+};
diff --git a/Engine/Plugins/FX/Niagara/Source/NiagaraEditor/Private/Widgets/Wizard/SNiagaraModuleWizard.cpp b/Engine/Plugins/FX/Niagara/Source/NiagaraEditor/Private/Widgets/Wizard/SNiagaraModuleWizard.cpp
new file mode 100644
--- /dev/null
+++ b/Engine/Plugins/FX/Niagara/Source/NiagaraEditor/Private/Widgets/Wizard/SNiagaraModuleWizard.cpp
@@ -0,0 +1,217 @@
+﻿// Copyright Epic Games, Inc. All Rights Reserved.
+
+#include "Widgets/Wizard/SNiagaraModuleWizard.h"
+
+#include "NiagaraClipboard.h"
+#include "NiagaraNodeParameterMapGet.h"
+#include "NiagaraNodeParameterMapSet.h"
+#include "ViewModels/NiagaraScratchPadScriptViewModel.h"
+#include "ViewModels/NiagaraScratchPadViewModel.h"
+#include "ViewModels/NiagaraSystemViewModel.h"
+#include "ViewModels/Stack/NiagaraStackFunctionInput.h"
+#include "ViewModels/Stack/NiagaraStackInputCategory.h"
+#include "Framework/Application/SlateApplication.h"
+#include "ViewModels/NiagaraSystemSelectionViewModel.h"
+
+#define LOCTEXT_NAMESPACE "FNiagaraWizard"
+
+using namespace UE::Niagara::Wizard;
+
+TSharedRef<SWidget> FModuleWizardPage::GetContent()
+{
+	return SNullWidget::NullWidget;
+}
+
+void FModuleWizardPage::RefreshContent()
+{
+}
+
+TArray<FModuleWizardModel::FModuleCreationEntry> FModuleWizardModel::GetModulesToCreate(UNiagaraNodeOutput* ProvidedOutputNode, int32 ProvidedTargetIndex, TSharedPtr<FNiagaraSystemViewModel> SystemModel, TSharedPtr<FNiagaraEmitterViewModel> EmitterViewModel)
+{
+	TArray<FModuleCreationEntry> Result;
+	Result.Add({ProvidedOutputNode, ProvidedTargetIndex});
+	return Result;
+}
+
+void FModuleWizardModel::GenerateNewModuleContent(TSharedPtr<FNiagaraScratchPadScriptViewModel> ScratchPadScriptViewModel, const TArray<const UNiagaraNodeFunctionCall*>& PreviouslyCreatedModules)
+{
+	ScratchPadScriptViewModel->SetScriptName(FText::FromString(TEXT("Read")));
+}
+
+bool FModuleWizardModel::UpdateModuleInputs(UNiagaraClipboardContent* NewModule, const TArray<const UNiagaraNodeFunctionCall*>& PreviouslyCreatedModules)
+{
+	return false;
+}
+
+void SNiagaraModuleWizard::Construct(const FArguments& InArgs, TSharedRef<FModuleWizardModel> InModel)
+{
+	Model = InModel;
+	OnCreateModule = InArgs._OnCreateModule;
+
+	SWizard::FArguments SuperArgs;
+	SuperArgs
+	.ShowPageList(false)
+	.ShowBreadcrumbs(true)
+	.ShowPageTitle(false)
+	.FinishButtonText(LOCTEXT("FinishButtonLabel", "Create Module"))
+	.FinishButtonToolTip(LOCTEXT("FinishButtonTooltip", "Create the module and close the wizard"))
+	.OnFinished(this, &SNiagaraModuleWizard::OnFinished)
+	.OnCanceled(this, &SNiagaraModuleWizard::CloseContainingWindow)
+	.CanFinish(this, &SNiagaraModuleWizard::CanFinish);
+
+	for (int i = 0; i < Model->Pages.Num(); i++)
+	{
+		TSharedRef<FModuleWizardPage> Page = Model->Pages[i];
+		FWizardPage::FArguments PageArgs = SWizard::Page()
+			.CanShow_Lambda([this, Index = i]()
+			{
+				return Index == 0 ? true : Model->Pages[Index - 1]->CanGoToNextPage();
+			})
+			.OnEnter(Page, &FModuleWizardPage::RefreshContent)
+			.Name(Page->Name)
+			[
+				Page->GetContent()
+			];
+		SuperArgs.Slots.Add(new FWizardPage(PageArgs));
+	}
+	SWizard::Construct(SuperArgs);
+}
+
+TArray<const UNiagaraNodeFunctionCall*> SNiagaraModuleWizard::AddModulesToStack(TSharedPtr<FNiagaraSystemViewModel> SystemViewModel, UNiagaraNodeOutput* OutputNode, int32 TargetIndex, TSharedPtr<FNiagaraEmitterViewModel> EmitterViewModel, UNiagaraStackEditorData* StackEditorData)
+{
+	// the wizard can potentially create more than one module
+	TArray<const UNiagaraNodeFunctionCall*> CreatedModules;
+	TArray<FModuleWizardModel::FModuleCreationEntry> TargetModules = Model->GetModulesToCreate(OutputNode, TargetIndex, SystemViewModel, EmitterViewModel);
+	for (const FModuleWizardModel::FModuleCreationEntry& Module : TargetModules)
+	{
+		check(Module.OutputNode);
+		TSharedPtr<FNiagaraScratchPadScriptViewModel> ScratchPadScriptViewModel = SystemViewModel->GetScriptScratchPadViewModel()->CreateNewScript(ENiagaraScriptUsage::Module, Module.OutputNode->GetUsage(), FNiagaraTypeDefinition());
+		if (!ScratchPadScriptViewModel.IsValid())
+		{
+			return CreatedModules;
+		}
+		Model->GenerateNewModuleContent(ScratchPadScriptViewModel, CreatedModules);
+		UNiagaraNodeFunctionCall* NewModule = FNiagaraStackGraphUtilities::AddScriptModuleToStack(ScratchPadScriptViewModel->GetOriginalScript(), *Module.OutputNode, Module.TargetIndex);
+
+		// open newly create scratch pad?
+		//SystemModel->GetScriptScratchPadViewModel()->FocusScratchPadScriptViewModel(ScratchPadScriptViewModel.ToSharedRef());
+		
+		UNiagaraStackScriptHierarchyRoot* FunctionHierarchyRoot = NewObject<UNiagaraStackScriptHierarchyRoot>(GetTransientPackage()); 
+		UNiagaraStackEntry::FRequiredEntryData RequiredEntryData(SystemViewModel.ToSharedRef(), EmitterViewModel, NAME_None, NAME_None, *StackEditorData);
+		FunctionHierarchyRoot->Initialize(RequiredEntryData, *NewModule, *NewModule, FString());
+		FunctionHierarchyRoot->RefreshChildren();
+
+		// Reset all direct inputs on this function to initialize data interfaces and default dynamic inputs.
+		TArray<UNiagaraStackFunctionInput*> StackFunctionInputs;
+		FunctionHierarchyRoot->GetUnfilteredChildrenOfType(StackFunctionInputs, true);
+		for (UNiagaraStackFunctionInput* StackFunctionInput : StackFunctionInputs)
+		{
+			if (StackFunctionInput != nullptr && StackFunctionInput->CanReset())
+			{
+				StackFunctionInput->Reset();
+			}
+		}
+
+		// allow the wizard model to change the inputs from the default
+		UNiagaraClipboardContent* ModuleInputs = UNiagaraClipboardContent::Create();
+		FunctionHierarchyRoot->ToClipboardFunctionInputs(ModuleInputs, MutableView(ModuleInputs->FunctionInputs));
+		if (Model->UpdateModuleInputs(ModuleInputs, CreatedModules))
+		{
+			FunctionHierarchyRoot->SetValuesFromClipboardFunctionInputs(ModuleInputs->FunctionInputs);
+		}
+		FunctionHierarchyRoot->Finalize();
+		FunctionHierarchyRoot->ConditionalBeginDestroy();
+		SystemViewModel->NotifyDataObjectChanged(TArray<UObject*>(), ENiagaraDataObjectChange::Unknown);
+	
+		FNiagaraStackGraphUtilities::RelayoutGraph(*OutputNode->GetGraph());
+		CreatedModules.Add(NewModule);
+	}
+	return CreatedModules;
+}
+
+bool SNiagaraModuleWizard::CanFinish() const
+{
+	for (TSharedRef<FModuleWizardPage> Page : Model->Pages)
+	{
+		if (!Page->CanCompleteWizard())
+		{
+			return false;
+		}
+	}
+	return true;
+}
+
+void SNiagaraModuleWizard::OnFinished()
+{
+	if (Model.IsValid())
+	{
+		for (TSharedRef<FModuleWizardPage> Page : Model->Pages)
+		{
+			Page->RefreshContent();
+		}
+	}
+	OnCreateModule.ExecuteIfBound();
+
+	CloseContainingWindow();
+}
+
+void SNiagaraModuleWizard::CloseContainingWindow()
+{
+	TSharedPtr<SWindow> ContainingWindow = FSlateApplication::Get().FindWidgetWindow(AsShared());
+	if (ContainingWindow.IsValid())
+	{
+		ContainingWindow->RequestDestroyWindow();
+	}
+}
+
+
+UEdGraphPin* Utilities::AddReadParameterPin(const FNiagaraTypeDefinition& Type, const FName& Name, UNiagaraNodeParameterMapGet* MapGetNode)
+{
+	UEdGraphPin* Pin = MapGetNode->AddParameterPin(FNiagaraVariable(Type, Name), EGPD_Output);
+	MapGetNode->CancelEditablePinName(FText::GetEmpty(), Pin);
+	return Pin;
+}
+
+UEdGraphPin* Utilities::AddWriteParameterPin(const FNiagaraTypeDefinition& Type, const FName& Name, UNiagaraNodeParameterMapSet* MapSetNode)
+{
+	UEdGraphPin* Pin = MapSetNode->AddParameterPin(FNiagaraVariable(Type, Name), EGPD_Input);
+	MapSetNode->CancelEditablePinName(FText::GetEmpty(), Pin);
+	return Pin;
+}
+
+TSharedRef<IDetailsView> Utilities::CreateDetailsView()
+{
+	FPropertyEditorModule& PropertyModule = FModuleManager::LoadModuleChecked<FPropertyEditorModule>("PropertyEditor");
+	FDetailsViewArgs DetailsViewArgs;
+	DetailsViewArgs.bAllowSearch = false;
+	DetailsViewArgs.bHideSelectionTip = true;
+	DetailsViewArgs.NameAreaSettings = FDetailsViewArgs::HideNameArea;
+	DetailsViewArgs.bShowObjectLabel = false;
+
+	return PropertyModule.CreateDetailView(DetailsViewArgs);
+}
+
+UNiagaraNodeFunctionCall* Utilities::CreateDataInterfaceFunctionNode(const TSubclassOf<UNiagaraDataInterface>& DataInterfaceClass, const FName& FunctionName, UNiagaraGraph* Graph)
+{
+	UNiagaraDataInterface* DataInterfaceCDO = DataInterfaceClass.GetDefaultObject();
+	TArray<FNiagaraFunctionSignature> Functions;
+	DataInterfaceCDO->GetFunctionSignatures(Functions);
+	for (FNiagaraFunctionSignature& Sig : Functions)
+	{
+		if (Sig.Name == FunctionName)
+		{
+			UNiagaraNodeFunctionCall* FuncNode = NewObject<UNiagaraNodeFunctionCall>(Graph);
+			FuncNode->Signature = Sig;
+			FuncNode->SetFlags(RF_Transactional);
+			Graph->AddNode(FuncNode, false, false);
+
+			FuncNode->CreateNewGuid();
+			FuncNode->PostPlacedNewNode();
+			FuncNode->AllocateDefaultPins();
+			return FuncNode;
+		}
+	}
+	return nullptr;
+}
+
+#undef LOCTEXT_NAMESPACE
diff --git a/Engine/Plugins/FX/Niagara/Source/NiagaraEditor/Public/NiagaraEditorModule.h b/Engine/Plugins/FX/Niagara/Source/NiagaraEditor/Public/NiagaraEditorModule.h
--- a/Engine/Plugins/FX/Niagara/Source/NiagaraEditor/Public/NiagaraEditorModule.h
+++ b/Engine/Plugins/FX/Niagara/Source/NiagaraEditor/Public/NiagaraEditorModule.h
@@ -16,6 +16,11 @@
 #include "Customizations/NiagaraDataInterfaceSimCacheVisualizer.h"
 #include "NiagaraEditorModule.generated.h"
 
+namespace UE::Niagara::Wizard
+{
+	class FModuleWizardGenerator;
+}
+
 class FNiagaraRecentAndFavoritesManager;
 class IAssetTools;
 class IAssetTypeActions;
@@ -192,6 +197,10 @@ public:
 	/** Registers niagara editor property utilities for a specific struct type. */
 	void RegisterPropertyUtilities(const UScriptStruct* InStruct, TSharedRef<INiagaraEditorPropertyUtilities, ESPMode::ThreadSafe> PropertyUtilities);
 
+	/** Registers niagara wizards that can be used to add generated modules to the stack. */
+	void RegisterModuleWizards(TSharedRef<UE::Niagara::Wizard::FModuleWizardGenerator> WizardGenerator);
+	TConstArrayView<TSharedRef<UE::Niagara::Wizard::FModuleWizardGenerator>> GetModuleWizards() const { return ModuleWizards; }
+
 	/** Register/unregister niagara editor settings. */
 	void RegisterSettings();
 	void UnregisterSettings();
@@ -414,6 +423,7 @@ private:
 
 	/** All created asset type actions.  Cached here so that we can unregister it during shutdown. */
 	TArray< TSharedPtr<IAssetTypeActions> > CreatedAssetTypeActions;
+	TArray<TSharedRef<UE::Niagara::Wizard::FModuleWizardGenerator>> ModuleWizards;
 
 	FCriticalSection TypeEditorsCS;
 	TMap<FNiagaraTypeDefinition, TSharedRef<INiagaraEditorTypeUtilities, ESPMode::ThreadSafe>> TypeToEditorUtilitiesMap;
diff --git a/Engine/Plugins/FX/Niagara/Source/NiagaraEditor/Public/NiagaraNodeWithDynamicPins.h b/Engine/Plugins/FX/Niagara/Source/NiagaraEditor/Public/NiagaraNodeWithDynamicPins.h
--- a/Engine/Plugins/FX/Niagara/Source/NiagaraEditor/Public/NiagaraNodeWithDynamicPins.h
+++ b/Engine/Plugins/FX/Niagara/Source/NiagaraEditor/Public/NiagaraNodeWithDynamicPins.h
@@ -44,6 +44,8 @@ public:
 	/** Request a new pin. */
 	void AddParameter(FNiagaraVariable Parameter, TEnumAsByte<EEdGraphPinDirection> Direction);
 	void AddParameter(const UNiagaraScriptVariable* ScriptVar, TEnumAsByte<EEdGraphPinDirection> Direction);
+	UEdGraphPin* AddParameterPin(FNiagaraVariable Parameter, TEnumAsByte<EEdGraphPinDirection> Direction);
+	UEdGraphPin* AddParameterPin(const UNiagaraScriptVariable* ScriptVar, TEnumAsByte<EEdGraphPinDirection> Direction);
 	void AddExistingParameter(FNiagaraVariable Parameter, const UEdGraphPin* AddPin);
 
 	/** Convenience method to determine whether this Node is a Map Get or Map Set when adding a parameter through the parameter panel. */
diff --git a/Engine/Plugins/FX/Niagara/Source/NiagaraEditor/Public/Widgets/Wizard/SNiagaraModuleWizard.h b/Engine/Plugins/FX/Niagara/Source/NiagaraEditor/Public/Widgets/Wizard/SNiagaraModuleWizard.h
new file mode 100644
--- /dev/null
+++ b/Engine/Plugins/FX/Niagara/Source/NiagaraEditor/Public/Widgets/Wizard/SNiagaraModuleWizard.h
@@ -0,0 +1,123 @@
+﻿// Copyright Epic Games, Inc. All Rights Reserved.
+
+#pragma once
+
+#include "CoreMinimal.h"
+#include "NiagaraCommon.h"
+#include "NiagaraGraph.h"
+#include "NiagaraNodeOutput.h"
+#include "Templates/SharedPointer.h"
+#include "Widgets/Workflow/SWizard.h"
+
+class UNiagaraClipboardContent;
+class FNiagaraScratchPadScriptViewModel;
+class IDetailsView;
+
+namespace UE::Niagara::Wizard
+{
+	// a single page in the wizard
+	struct NIAGARAEDITOR_API FModuleWizardPage : TSharedFromThis<FModuleWizardPage>
+	{
+		FModuleWizardPage() = default;
+		virtual ~FModuleWizardPage() = default;
+		
+		// If there is a follow-up page, this enables the button for it
+		virtual bool CanGoToNextPage() const { return true; };
+
+		// If true then the wizard as a whole can be finished, even if it's not the last page
+		virtual bool CanCompleteWizard() const { return true; };
+
+		virtual TSharedRef<SWidget> GetContent();
+
+		// called when the page is activated
+		virtual void RefreshContent();
+
+		// name in the breadcrumb view
+		FText Name;
+	};
+
+	// data model for the wizard holding all the pages in order. The model has no concept of branching, it follows the pages in linear order.
+	struct NIAGARAEDITOR_API FModuleWizardModel : TSharedFromThis<FModuleWizardModel>
+	{
+		FModuleWizardModel() = default;
+		virtual ~FModuleWizardModel() = default;
+
+		struct FModuleCreationEntry
+		{
+			UNiagaraNodeOutput* OutputNode = nullptr;
+			int32 TargetIndex = INDEX_NONE;
+		};
+
+		// child classes can override this if they want to generate more than one module in the system or change the location of the generated module
+		virtual TArray<FModuleCreationEntry> GetModulesToCreate(UNiagaraNodeOutput* ProvidedOutputNode, int32 ProvidedTargetIndex, TSharedPtr<FNiagaraSystemViewModel> SystemModel, TSharedPtr<FNiagaraEmitterViewModel> EmitterViewModel); 
+
+		// called when the scratch pad is generated and the graph can be modified, but before the scratch pad is added to the stack
+		virtual void GenerateNewModuleContent(TSharedPtr<FNiagaraScratchPadScriptViewModel> ScratchPadScriptViewModel, const TArray<const UNiagaraNodeFunctionCall*>& PreviouslyCreatedModules);
+
+		// called after the generated module is added to the stack - returning true will apply the modified clipboard content to the module inputs
+		virtual bool UpdateModuleInputs(UNiagaraClipboardContent* NewModule, const TArray<const UNiagaraNodeFunctionCall*>& PreviouslyCreatedModules);
+
+		// target usage of the generated module
+		ENiagaraScriptUsage TargetUsage;
+		
+		TArray<TSharedRef<FModuleWizardPage>> Pages;
+	};
+
+	// a wizard that is used to create new scratch pad modules
+	class SNiagaraModuleWizard : public SWizard
+	{
+	public:
+
+		SLATE_BEGIN_ARGS(SNiagaraModuleWizard) {}
+			SLATE_EVENT(FSimpleDelegate, OnCreateModule)
+		SLATE_END_ARGS()
+		
+		void Construct(const FArguments& InArgs, TSharedRef<FModuleWizardModel> Model);
+		TArray<const UNiagaraNodeFunctionCall*> AddModulesToStack(TSharedPtr<FNiagaraSystemViewModel> SystemViewModel, UNiagaraNodeOutput* OutputNode, int32 TargetIndex, TSharedPtr<FNiagaraEmitterViewModel> EmitterViewModel, UNiagaraStackEditorData* StackEditorData);
+
+		FSimpleDelegate OnCreateModule;
+
+		TSharedPtr<FModuleWizardModel> Model;
+		
+	private:
+		bool CanFinish() const;
+		void OnFinished();
+		void CloseContainingWindow();
+	};
+
+	// subclasses can be registered with FNiagaraEditorModule::RegisterModuleWizards() to automatically add wizards to the "add module" menu in the stack 
+	class NIAGARAEDITOR_API FModuleWizardGenerator
+	{
+	public:
+		struct FAction
+		{
+			TSharedPtr<FModuleWizardModel> WizardModel;
+			FText DisplayName;
+			FText Description;
+			FText Keywords;
+			bool bSuggestedAction = false;
+		};
+		
+		virtual TArray<FAction> CreateWizardActions(ENiagaraScriptUsage Usage) = 0;
+		
+		virtual ~FModuleWizardGenerator() = default;
+	};
+
+	namespace Utilities
+	{	
+		template<typename NodeType>
+		NodeType* FindSingleNodeChecked(UNiagaraGraph* Graph)
+		{
+			check(Graph);
+			TArray<NodeType*> NiagaraNodes;
+			Graph->GetNodesOfClass(NiagaraNodes);
+			check(NiagaraNodes.Num() == 1);
+			return NiagaraNodes[0];
+		}
+
+		NIAGARAEDITOR_API UEdGraphPin* AddReadParameterPin(const FNiagaraTypeDefinition& Type, const FName& Name, UNiagaraNodeParameterMapGet* MapGetNode);
+		NIAGARAEDITOR_API UEdGraphPin* AddWriteParameterPin(const FNiagaraTypeDefinition& Type, const FName& Name, UNiagaraNodeParameterMapSet* MapSetNode);
+		NIAGARAEDITOR_API TSharedRef<IDetailsView> CreateDetailsView();
+		NIAGARAEDITOR_API UNiagaraNodeFunctionCall* CreateDataInterfaceFunctionNode(const TSubclassOf<UNiagaraDataInterface>& DataInterfaceClass, const FName& FunctionName, UNiagaraGraph* Graph);
+	}
+}
diff --git a/Engine/Plugins/FX/Niagara/Source/NiagaraEditorWidgets/Private/DetailCustomizations/NiagaraDataChannelDetails.cpp b/Engine/Plugins/FX/Niagara/Source/NiagaraEditorWidgets/Private/DetailCustomizations/NiagaraDataChannelDetails.cpp
--- a/Engine/Plugins/FX/Niagara/Source/NiagaraEditorWidgets/Private/DetailCustomizations/NiagaraDataChannelDetails.cpp
+++ b/Engine/Plugins/FX/Niagara/Source/NiagaraEditorWidgets/Private/DetailCustomizations/NiagaraDataChannelDetails.cpp
@@ -21,7 +21,7 @@
 
 #include "NiagaraSystem.h"
 #include "NiagaraEditorUtilities.h"
-#include "K2Node_WriteDataChannel.h"
+#include "K2Node_DataChannelBase.h"
 
 #include "NiagaraDataInterfaceDetails.h"
 
