- Add missing tooltips to stateless modules / spawn infos

#rnx
#tests Add stateless module(s) and hover hover to see the tooltip(s)

diff --git a/Engine/Plugins/FX/Niagara/Source/Niagara/Internal/Stateless/Modules/NiagaraStatelessModule_AccelerationForce.h b/Engine/Plugins/FX/Niagara/Source/Niagara/Internal/Stateless/Modules/NiagaraStatelessModule_AccelerationForce.h
--- a/Engine/Plugins/FX/Niagara/Source/Niagara/Internal/Stateless/Modules/NiagaraStatelessModule_AccelerationForce.h
+++ b/Engine/Plugins/FX/Niagara/Source/Niagara/Internal/Stateless/Modules/NiagaraStatelessModule_AccelerationForce.h
@@ -9,6 +9,7 @@
 
 #include "NiagaraStatelessModule_AccelerationForce.generated.h"
 
+// Apply acceleration which accumulates per frame
 UCLASS(MinimalAPI, EditInlineNew, meta = (DisplayName = "Acceleration Force"))
 class UNiagaraStatelessModule_AccelerationForce : public UNiagaraStatelessModule
 {
diff --git a/Engine/Plugins/FX/Niagara/Source/Niagara/Internal/Stateless/Modules/NiagaraStatelessModule_AddVelocity.h b/Engine/Plugins/FX/Niagara/Source/Niagara/Internal/Stateless/Modules/NiagaraStatelessModule_AddVelocity.h
--- a/Engine/Plugins/FX/Niagara/Source/Niagara/Internal/Stateless/Modules/NiagaraStatelessModule_AddVelocity.h
+++ b/Engine/Plugins/FX/Niagara/Source/Niagara/Internal/Stateless/Modules/NiagaraStatelessModule_AddVelocity.h
@@ -17,6 +17,7 @@ enum class ENSM_VelocityType
 	InCone,			// Max, Min, Cone| Cone Axis, Cone Angle, Inner Cone
 };
 
+// Add to the velocity applied to all particles
 UCLASS(MinimalAPI, EditInlineNew, meta = (DisplayName = "Add Velocity"))
 class UNiagaraStatelessModule_AddVelocity : public UNiagaraStatelessModule
 {
diff --git a/Engine/Plugins/FX/Niagara/Source/Niagara/Internal/Stateless/Modules/NiagaraStatelessModule_CameraOffset.h b/Engine/Plugins/FX/Niagara/Source/Niagara/Internal/Stateless/Modules/NiagaraStatelessModule_CameraOffset.h
--- a/Engine/Plugins/FX/Niagara/Source/Niagara/Internal/Stateless/Modules/NiagaraStatelessModule_CameraOffset.h
+++ b/Engine/Plugins/FX/Niagara/Source/Niagara/Internal/Stateless/Modules/NiagaraStatelessModule_CameraOffset.h
@@ -9,6 +9,7 @@
 
 #include "NiagaraStatelessModule_CameraOffset.generated.h"
 
+// Offset the particle along the vector between the particle and the camera.
 UCLASS(MinimalAPI, EditInlineNew, meta = (DisplayName = "Camera Offset"))
 class UNiagaraStatelessModule_CameraOffset : public UNiagaraStatelessModule
 {
diff --git a/Engine/Plugins/FX/Niagara/Source/Niagara/Internal/Stateless/Modules/NiagaraStatelessModule_Drag.h b/Engine/Plugins/FX/Niagara/Source/Niagara/Internal/Stateless/Modules/NiagaraStatelessModule_Drag.h
--- a/Engine/Plugins/FX/Niagara/Source/Niagara/Internal/Stateless/Modules/NiagaraStatelessModule_Drag.h
+++ b/Engine/Plugins/FX/Niagara/Source/Niagara/Internal/Stateless/Modules/NiagaraStatelessModule_Drag.h
@@ -9,6 +9,7 @@
 
 #include "NiagaraStatelessModule_Drag.generated.h"
 
+// Applies Drag directly to particle velocity, irrespective of Mass.
 UCLASS(MinimalAPI, EditInlineNew, meta = (DisplayName = "Drag"))
 class UNiagaraStatelessModule_Drag : public UNiagaraStatelessModule
 {
diff --git a/Engine/Plugins/FX/Niagara/Source/Niagara/Internal/Stateless/Modules/NiagaraStatelessModule_DynamicMaterialParameters.h b/Engine/Plugins/FX/Niagara/Source/Niagara/Internal/Stateless/Modules/NiagaraStatelessModule_DynamicMaterialParameters.h
--- a/Engine/Plugins/FX/Niagara/Source/Niagara/Internal/Stateless/Modules/NiagaraStatelessModule_DynamicMaterialParameters.h
+++ b/Engine/Plugins/FX/Niagara/Source/Niagara/Internal/Stateless/Modules/NiagaraStatelessModule_DynamicMaterialParameters.h
@@ -39,6 +39,7 @@ struct FNiagaraStatelessDynamicParameterSet
 	FNiagaraDistributionFloat WChannelDistribution = FNiagaraDistributionFloat(FNiagaraStatelessGlobals::GetDefaultDynamicMaterialParametersValue().W);
 };
 
+// Write to the Dynamic Parameters that can be read in the material vertex & pixel shader vertex
 UCLASS(MinimalAPI, EditInlineNew, meta = (DisplayName = "Dynamic Material Parameters"))
 class UNiagaraStatelessModule_DynamicMaterialParameters : public UNiagaraStatelessModule
 {
diff --git a/Engine/Plugins/FX/Niagara/Source/Niagara/Internal/Stateless/Modules/NiagaraStatelessModule_GravityForce.h b/Engine/Plugins/FX/Niagara/Source/Niagara/Internal/Stateless/Modules/NiagaraStatelessModule_GravityForce.h
--- a/Engine/Plugins/FX/Niagara/Source/Niagara/Internal/Stateless/Modules/NiagaraStatelessModule_GravityForce.h
+++ b/Engine/Plugins/FX/Niagara/Source/Niagara/Internal/Stateless/Modules/NiagaraStatelessModule_GravityForce.h
@@ -9,6 +9,7 @@
 
 #include "NiagaraStatelessModule_GravityForce.generated.h"
 
+// Applies a gravitational force (in cm/s)
 UCLASS(MinimalAPI, EditInlineNew, meta = (DisplayName = "Gravity Force"))
 class UNiagaraStatelessModule_GravityForce : public UNiagaraStatelessModule
 {
diff --git a/Engine/Plugins/FX/Niagara/Source/Niagara/Internal/Stateless/Modules/NiagaraStatelessModule_InitialMeshOrientation.h b/Engine/Plugins/FX/Niagara/Source/Niagara/Internal/Stateless/Modules/NiagaraStatelessModule_InitialMeshOrientation.h
--- a/Engine/Plugins/FX/Niagara/Source/Niagara/Internal/Stateless/Modules/NiagaraStatelessModule_InitialMeshOrientation.h
+++ b/Engine/Plugins/FX/Niagara/Source/Niagara/Internal/Stateless/Modules/NiagaraStatelessModule_InitialMeshOrientation.h
@@ -19,6 +19,7 @@ enum class ENSMInitialMeshOrientationMode
 	//OrientToQuaternion,
 };
 
+// Set the initial mesh orientation, directly, randomly or by orienting by axis
 UCLASS(MinimalAPI, EditInlineNew, meta = (DisplayName = "Initial Mesh Orientation"))
 class UNiagaraStatelessModule_InitialMeshOrientation : public UNiagaraStatelessModule
 {
diff --git a/Engine/Plugins/FX/Niagara/Source/Niagara/Internal/Stateless/Modules/NiagaraStatelessModule_InitializeParticle.h b/Engine/Plugins/FX/Niagara/Source/Niagara/Internal/Stateless/Modules/NiagaraStatelessModule_InitializeParticle.h
--- a/Engine/Plugins/FX/Niagara/Source/Niagara/Internal/Stateless/Modules/NiagaraStatelessModule_InitializeParticle.h
+++ b/Engine/Plugins/FX/Niagara/Source/Niagara/Internal/Stateless/Modules/NiagaraStatelessModule_InitializeParticle.h
@@ -13,6 +13,7 @@
 
 #include "NiagaraStatelessModule_InitializeParticle.generated.h"
 
+// Initialize common particle attributes using common settings and options.
 UCLASS(MinimalAPI, EditInlineNew, meta = (DisplayName = "Initialize Particle"))
 class UNiagaraStatelessModule_InitializeParticle : public UNiagaraStatelessModule
 {
diff --git a/Engine/Plugins/FX/Niagara/Source/Niagara/Internal/Stateless/Modules/NiagaraStatelessModule_MeshIndex.h b/Engine/Plugins/FX/Niagara/Source/Niagara/Internal/Stateless/Modules/NiagaraStatelessModule_MeshIndex.h
--- a/Engine/Plugins/FX/Niagara/Source/Niagara/Internal/Stateless/Modules/NiagaraStatelessModule_MeshIndex.h
+++ b/Engine/Plugins/FX/Niagara/Source/Niagara/Internal/Stateless/Modules/NiagaraStatelessModule_MeshIndex.h
@@ -8,6 +8,7 @@
 
 #include "NiagaraStatelessModule_MeshIndex.generated.h"
 
+// Sets the mesh index to use with mesh renderers
 UCLASS(MinimalAPI, EditInlineNew, meta = (DisplayName = "Mesh Index"))
 class UNiagaraStatelessModule_MeshIndex : public UNiagaraStatelessModule
 {
diff --git a/Engine/Plugins/FX/Niagara/Source/Niagara/Internal/Stateless/Modules/NiagaraStatelessModule_MeshRotationRate.h b/Engine/Plugins/FX/Niagara/Source/Niagara/Internal/Stateless/Modules/NiagaraStatelessModule_MeshRotationRate.h
--- a/Engine/Plugins/FX/Niagara/Source/Niagara/Internal/Stateless/Modules/NiagaraStatelessModule_MeshRotationRate.h
+++ b/Engine/Plugins/FX/Niagara/Source/Niagara/Internal/Stateless/Modules/NiagaraStatelessModule_MeshRotationRate.h
@@ -9,6 +9,7 @@
 
 #include "NiagaraStatelessModule_MeshRotationRate.generated.h"
 
+// Applies a constant rotation rate to mesh orientation
 UCLASS(MinimalAPI, EditInlineNew, meta = (DisplayName = "Mesh Rotation Rate"))
 class UNiagaraStatelessModule_MeshRotationRate : public UNiagaraStatelessModule
 {
diff --git a/Engine/Plugins/FX/Niagara/Source/Niagara/Internal/Stateless/Modules/NiagaraStatelessModule_ScaleColor.h b/Engine/Plugins/FX/Niagara/Source/Niagara/Internal/Stateless/Modules/NiagaraStatelessModule_ScaleColor.h
--- a/Engine/Plugins/FX/Niagara/Source/Niagara/Internal/Stateless/Modules/NiagaraStatelessModule_ScaleColor.h
+++ b/Engine/Plugins/FX/Niagara/Source/Niagara/Internal/Stateless/Modules/NiagaraStatelessModule_ScaleColor.h
@@ -9,6 +9,7 @@
 
 #include "NiagaraStatelessModule_ScaleColor.generated.h"
 
+// Scales the color of the particle
 UCLASS(MinimalAPI, EditInlineNew, meta = (DisplayName = "Scale Color"))
 class UNiagaraStatelessModule_ScaleColor : public UNiagaraStatelessModule
 {
diff --git a/Engine/Plugins/FX/Niagara/Source/Niagara/Internal/Stateless/Modules/NiagaraStatelessModule_ScaleMeshSize.h b/Engine/Plugins/FX/Niagara/Source/Niagara/Internal/Stateless/Modules/NiagaraStatelessModule_ScaleMeshSize.h
--- a/Engine/Plugins/FX/Niagara/Source/Niagara/Internal/Stateless/Modules/NiagaraStatelessModule_ScaleMeshSize.h
+++ b/Engine/Plugins/FX/Niagara/Source/Niagara/Internal/Stateless/Modules/NiagaraStatelessModule_ScaleMeshSize.h
@@ -11,6 +11,8 @@
 
 #include "NiagaraStatelessModule_ScaleMeshSize.generated.h"
 
+// Multiply Particle.Scale by the module calculated scale value
+// This can be a constant, random or curve indexed by Particle.NormalizedAge
 UCLASS(MinimalAPI, EditInlineNew, meta = (DisplayName = "Scale Mesh Size"))
 class UNiagaraStatelessModule_ScaleMeshSize : public UNiagaraStatelessModule
 {
diff --git a/Engine/Plugins/FX/Niagara/Source/Niagara/Internal/Stateless/Modules/NiagaraStatelessModule_ScaleMeshSizeBySpeed.h b/Engine/Plugins/FX/Niagara/Source/Niagara/Internal/Stateless/Modules/NiagaraStatelessModule_ScaleMeshSizeBySpeed.h
--- a/Engine/Plugins/FX/Niagara/Source/Niagara/Internal/Stateless/Modules/NiagaraStatelessModule_ScaleMeshSizeBySpeed.h
+++ b/Engine/Plugins/FX/Niagara/Source/Niagara/Internal/Stateless/Modules/NiagaraStatelessModule_ScaleMeshSizeBySpeed.h
@@ -9,6 +9,7 @@
 
 #include "NiagaraStatelessModule_ScaleMeshSizeBySpeed.generated.h"
 
+// Applies a modifier to mesh scale based on the velocity of the particle
 UCLASS(MinimalAPI, EditInlineNew, meta = (DisplayName = "Scale Mesh Size By Speed"))
 class UNiagaraStatelessModule_ScaleMeshSizeBySpeed : public UNiagaraStatelessModule
 {
diff --git a/Engine/Plugins/FX/Niagara/Source/Niagara/Internal/Stateless/Modules/NiagaraStatelessModule_ScaleSpriteSize.h b/Engine/Plugins/FX/Niagara/Source/Niagara/Internal/Stateless/Modules/NiagaraStatelessModule_ScaleSpriteSize.h
--- a/Engine/Plugins/FX/Niagara/Source/Niagara/Internal/Stateless/Modules/NiagaraStatelessModule_ScaleSpriteSize.h
+++ b/Engine/Plugins/FX/Niagara/Source/Niagara/Internal/Stateless/Modules/NiagaraStatelessModule_ScaleSpriteSize.h
@@ -10,6 +10,8 @@
 
 #include "NiagaraStatelessModule_ScaleSpriteSize.generated.h"
 
+// Multiply Particle.SpriteSize by the module calculated scale value
+// This can be a constant, random or curve indexed by Particle.NormalizedAge
 UCLASS(MinimalAPI, EditInlineNew, meta = (DisplayName = "Scale Sprite Size"))
 class UNiagaraStatelessModule_ScaleSpriteSize : public UNiagaraStatelessModule
 {
diff --git a/Engine/Plugins/FX/Niagara/Source/Niagara/Internal/Stateless/Modules/NiagaraStatelessModule_ScaleSpriteSizeBySpeed.h b/Engine/Plugins/FX/Niagara/Source/Niagara/Internal/Stateless/Modules/NiagaraStatelessModule_ScaleSpriteSizeBySpeed.h
--- a/Engine/Plugins/FX/Niagara/Source/Niagara/Internal/Stateless/Modules/NiagaraStatelessModule_ScaleSpriteSizeBySpeed.h
+++ b/Engine/Plugins/FX/Niagara/Source/Niagara/Internal/Stateless/Modules/NiagaraStatelessModule_ScaleSpriteSizeBySpeed.h
@@ -9,6 +9,7 @@
 
 #include "NiagaraStatelessModule_ScaleSpriteSizeBySpeed.generated.h"
 
+// Applies a modifier to sprite size based on the velocity of the particle
 UCLASS(MinimalAPI, EditInlineNew, meta = (DisplayName = "Scale Sprite Size By Speed"))
 class UNiagaraStatelessModule_ScaleSpriteSizeBySpeed : public UNiagaraStatelessModule
 {
diff --git a/Engine/Plugins/FX/Niagara/Source/Niagara/Internal/Stateless/Modules/NiagaraStatelessModule_SolveVelocitiesAndForces.h b/Engine/Plugins/FX/Niagara/Source/Niagara/Internal/Stateless/Modules/NiagaraStatelessModule_SolveVelocitiesAndForces.h
--- a/Engine/Plugins/FX/Niagara/Source/Niagara/Internal/Stateless/Modules/NiagaraStatelessModule_SolveVelocitiesAndForces.h
+++ b/Engine/Plugins/FX/Niagara/Source/Niagara/Internal/Stateless/Modules/NiagaraStatelessModule_SolveVelocitiesAndForces.h
@@ -14,6 +14,7 @@
 
 #include "NiagaraStatelessModule_SolveVelocitiesAndForces.generated.h"
 
+// Integrates all the forces applying them to position
 UCLASS(MinimalAPI, EditInlineNew, meta = (DisplayName = "Solve Forces And Velocity"))
 class UNiagaraStatelessModule_SolveVelocitiesAndForces : public UNiagaraStatelessModule
 {
diff --git a/Engine/Plugins/FX/Niagara/Source/Niagara/Internal/Stateless/Modules/NiagaraStatelessModule_SpriteFacingAndAlignment.h b/Engine/Plugins/FX/Niagara/Source/Niagara/Internal/Stateless/Modules/NiagaraStatelessModule_SpriteFacingAndAlignment.h
--- a/Engine/Plugins/FX/Niagara/Source/Niagara/Internal/Stateless/Modules/NiagaraStatelessModule_SpriteFacingAndAlignment.h
+++ b/Engine/Plugins/FX/Niagara/Source/Niagara/Internal/Stateless/Modules/NiagaraStatelessModule_SpriteFacingAndAlignment.h
@@ -9,6 +9,7 @@
 
 #include "NiagaraStatelessModule_SpriteFacingAndAlignment.generated.h"
 
+// Sets the sprite facing and alignment attributes
 UCLASS(MinimalAPI, EditInlineNew, meta = (DisplayName = "Sprite Facing Alignment"))
 class UNiagaraStatelessModule_SpriteFacingAndAlignment : public UNiagaraStatelessModule
 {
diff --git a/Engine/Plugins/FX/Niagara/Source/Niagara/Internal/Stateless/Modules/NiagaraStatelessModule_SpriteRotationRate.h b/Engine/Plugins/FX/Niagara/Source/Niagara/Internal/Stateless/Modules/NiagaraStatelessModule_SpriteRotationRate.h
--- a/Engine/Plugins/FX/Niagara/Source/Niagara/Internal/Stateless/Modules/NiagaraStatelessModule_SpriteRotationRate.h
+++ b/Engine/Plugins/FX/Niagara/Source/Niagara/Internal/Stateless/Modules/NiagaraStatelessModule_SpriteRotationRate.h
@@ -9,6 +9,7 @@
 
 #include "NiagaraStatelessModule_SpriteRotationRate.generated.h"
 
+// Applies a constant value to sprite rotation
 UCLASS(MinimalAPI, EditInlineNew, meta = (DisplayName = "Sprite Rotation Rate"))
 class UNiagaraStatelessModule_SpriteRotationRate : public UNiagaraStatelessModule
 {
diff --git a/Engine/Plugins/FX/Niagara/Source/Niagara/Internal/Stateless/Modules/NiagaraStatelessModule_SubUVAnimation.h b/Engine/Plugins/FX/Niagara/Source/Niagara/Internal/Stateless/Modules/NiagaraStatelessModule_SubUVAnimation.h
--- a/Engine/Plugins/FX/Niagara/Source/Niagara/Internal/Stateless/Modules/NiagaraStatelessModule_SubUVAnimation.h
+++ b/Engine/Plugins/FX/Niagara/Source/Niagara/Internal/Stateless/Modules/NiagaraStatelessModule_SubUVAnimation.h
@@ -17,6 +17,8 @@ enum class ENSMSubUVAnimation_Mode
 	Random,
 };
 
+// Sets the sub image frame index value based on the select animation mode
+// The sub image index is a float value where the fractional part can be used to blend frames together
 UCLASS(MinimalAPI, EditInlineNew, meta = (DisplayName = "Sub UV Animation"))
 class UNiagaraStatelessModule_SubUVAnimation : public UNiagaraStatelessModule
 {
diff --git a/Engine/Plugins/FX/Niagara/Source/NiagaraEditor/Private/ViewModels/Stack/NiagaraStackStatelessEmitterSimulateGroup.cpp b/Engine/Plugins/FX/Niagara/Source/NiagaraEditor/Private/ViewModels/Stack/NiagaraStackStatelessEmitterSimulateGroup.cpp
--- a/Engine/Plugins/FX/Niagara/Source/NiagaraEditor/Private/ViewModels/Stack/NiagaraStackStatelessEmitterSimulateGroup.cpp
+++ b/Engine/Plugins/FX/Niagara/Source/NiagaraEditor/Private/ViewModels/Stack/NiagaraStackStatelessEmitterSimulateGroup.cpp
@@ -166,6 +166,18 @@ void UNiagaraStackStatelessModuleItem::Initialize(FRequiredEntryData InRequiredE
 	DisplayName = InStatelessModule->GetClass()->GetDisplayNameText();
 }
 
+FText UNiagaraStackStatelessModuleItem::GetTooltipText() const
+{
+	if (UNiagaraStatelessModule* Module = StatelessModuleWeak.Get())
+	{
+		if (UClass* ModuleClass = Module->GetClass())
+		{
+			return ModuleClass->GetToolTipText();
+		}
+	}
+	return Super::GetTooltipText();
+}
+
 bool UNiagaraStackStatelessModuleItem::TestCanDeleteWithMessage(FText& OutCanDeleteMessage) const
 {
 	UNiagaraStatelessModule* StatelessModule = StatelessModuleWeak.Get();
diff --git a/Engine/Plugins/FX/Niagara/Source/NiagaraEditor/Private/ViewModels/Stack/NiagaraStackStatelessEmitterSpawnGroup.cpp b/Engine/Plugins/FX/Niagara/Source/NiagaraEditor/Private/ViewModels/Stack/NiagaraStackStatelessEmitterSpawnGroup.cpp
--- a/Engine/Plugins/FX/Niagara/Source/NiagaraEditor/Private/ViewModels/Stack/NiagaraStackStatelessEmitterSpawnGroup.cpp
+++ b/Engine/Plugins/FX/Niagara/Source/NiagaraEditor/Private/ViewModels/Stack/NiagaraStackStatelessEmitterSpawnGroup.cpp
@@ -174,6 +174,17 @@ FText UNiagaraStackStatelessEmitterSpawnItem::GetDisplayName(ENiagaraStatelessSp
 	}
 }
 
+FText UNiagaraStackStatelessEmitterSpawnItem::GetTooltipText(ENiagaraStatelessSpawnInfoType SpawnInfoType)
+{
+	switch (SpawnInfoType)
+	{
+		case ENiagaraStatelessSpawnInfoType::Burst:	return LOCTEXT("EmitterSpawnBurstTooltpText", "Spawns a burst of particles instantaneously.");
+		case ENiagaraStatelessSpawnInfoType::Rate:	return LOCTEXT("EmitterSpawnRateTooltpText", "Spawns particles continuously at a particular rate.");
+		default:									checkNoEntry();	return LOCTEXT("EmitterSpawnUnknownTooltpText", "Unknown");
+	}
+}
+
+
 FText UNiagaraStackStatelessEmitterSpawnItem::GetDisplayName() const
 {
 	const FNiagaraStatelessSpawnInfo* SpawnInfo = GetSpawnInfo();
@@ -181,6 +192,13 @@ FText UNiagaraStackStatelessEmitterSpawnItem::GetDisplayName() const
 	return GetDisplayName(SpawnInfoType);
 }
 
+FText UNiagaraStackStatelessEmitterSpawnItem::GetTooltipText() const
+{
+	const FNiagaraStatelessSpawnInfo* SpawnInfo = GetSpawnInfo();
+	const ENiagaraStatelessSpawnInfoType SpawnInfoType = SpawnInfo ? SpawnInfo->Type : ENiagaraStatelessSpawnInfoType::Burst;
+	return GetTooltipText(SpawnInfoType);
+}
+
 FGuid UNiagaraStackStatelessEmitterSpawnItem::GetSelectionId() const
 {
 	return SourceId;
diff --git a/Engine/Plugins/FX/Niagara/Source/NiagaraEditor/Public/ViewModels/Stack/NiagaraStackStatelessEmitterSimulateGroup.h b/Engine/Plugins/FX/Niagara/Source/NiagaraEditor/Public/ViewModels/Stack/NiagaraStackStatelessEmitterSimulateGroup.h
--- a/Engine/Plugins/FX/Niagara/Source/NiagaraEditor/Public/ViewModels/Stack/NiagaraStackStatelessEmitterSimulateGroup.h
+++ b/Engine/Plugins/FX/Niagara/Source/NiagaraEditor/Public/ViewModels/Stack/NiagaraStackStatelessEmitterSimulateGroup.h
@@ -52,6 +52,7 @@ public:
 	void Initialize(FRequiredEntryData InRequiredEntryData, UNiagaraStatelessModule* InStatelessModule);
 
 	virtual FText GetDisplayName() const override { return DisplayName; }
+	virtual FText GetTooltipText() const override;
 
 	virtual bool SupportsDelete() const override { return true; }
 	virtual bool TestCanDeleteWithMessage(FText& OutCanDeleteMessage) const override;
diff --git a/Engine/Plugins/FX/Niagara/Source/NiagaraEditor/Public/ViewModels/Stack/NiagaraStackStatelessEmitterSpawnGroup.h b/Engine/Plugins/FX/Niagara/Source/NiagaraEditor/Public/ViewModels/Stack/NiagaraStackStatelessEmitterSpawnGroup.h
--- a/Engine/Plugins/FX/Niagara/Source/NiagaraEditor/Public/ViewModels/Stack/NiagaraStackStatelessEmitterSpawnGroup.h
+++ b/Engine/Plugins/FX/Niagara/Source/NiagaraEditor/Public/ViewModels/Stack/NiagaraStackStatelessEmitterSpawnGroup.h
@@ -55,6 +55,7 @@ public:
 	void Initialize(FRequiredEntryData InRequiredEntryData, UNiagaraStatelessEmitter* InStatelessEmitter, int32 InIndex);
 
 	virtual FText GetDisplayName() const override;
+	virtual FText GetTooltipText() const override;
 	virtual FGuid GetSelectionId() const override;
 
 	virtual bool SupportsDelete() const override { return true; }
@@ -75,6 +76,7 @@ public:
 	FOnRequestDelete& OnRequestDelete() { return OnRequestDeleteDelegate; }
 
 	static FText GetDisplayName(ENiagaraStatelessSpawnInfoType SpawnInfoType);
+	static FText GetTooltipText(ENiagaraStatelessSpawnInfoType SpawnInfoType);
 
 protected:
 	virtual void RefreshChildrenInternal(const TArray<UNiagaraStackEntry*>& CurrentChildren, TArray<UNiagaraStackEntry*>& NewChildren, TArray<FStackIssue>& NewIssues) override;
