- Add debug hud commands to modify the font size

#rnx
#tests Run console commands, i.e. fx.niagara.debug.hud OverviewFont=0

diff --git a/Engine/Plugins/FX/Niagara/Source/Niagara/Private/NiagaraDebugHud.cpp b/Engine/Plugins/FX/Niagara/Source/Niagara/Private/NiagaraDebugHud.cpp
--- a/Engine/Plugins/FX/Niagara/Source/Niagara/Private/NiagaraDebugHud.cpp
+++ b/Engine/Plugins/FX/Niagara/Source/Niagara/Private/NiagaraDebugHud.cpp
@@ -108,6 +108,11 @@ namespace NiagaraDebugLocal
 		return OutValue;
 	}
 
+	static ENiagaraDebugHudFont StringToFont(const FString& Arg)
+	{
+		return FMath::Clamp(ENiagaraDebugHudFont(FCString::Atoi(*Arg)), ENiagaraDebugHudFont::Small, ENiagaraDebugHudFont::Normal);
+	}
+
 	static TTuple<const TCHAR*, const TCHAR*, TFunction<void(FString)>> GDebugConsoleCommands[] =
 	{
 		// Main HUD commands
@@ -134,6 +139,9 @@ namespace NiagaraDebugLocal
 		MakeTuple(TEXT("EmitterFilter="), TEXT("Set the emitter filter"), [](FString Arg) {Settings.EmitterFilter = Arg; Settings.bEmitterFilterEnabled = !Arg.IsEmpty(); GCachedSystemVariables.Empty(); }),
 		MakeTuple(TEXT("ActorFilter="), TEXT("Set the actor filter"), [](FString Arg) {Settings.ActorFilter = Arg; Settings.bActorFilterEnabled = !Arg.IsEmpty(); }),
 		MakeTuple(TEXT("ComponentFilter="), TEXT("Set the component filter"), [](FString Arg) {Settings.ComponentFilter = Arg; Settings.bComponentFilterEnabled = !Arg.IsEmpty(); }),
+		MakeTuple(TEXT("OverviewFont="), TEXT("Set the overview font to use"), [](FString Arg) { Settings.OverviewFont = StringToFont(Arg); }),
+		MakeTuple(TEXT("SystemTextFont="), TEXT("Set the system display text font to use"), [](FString Arg) { Settings.SystemTextOptions.Font = StringToFont(Arg); }),
+		MakeTuple(TEXT("ParticleTextFont="), TEXT("Set the particle display text font to use"), [](FString Arg) { Settings.ParticleTextOptions.Font = StringToFont(Arg); }),
 
 		MakeTuple(TEXT("ShowGlobalBudgetInfo="), TEXT("Shows global budget information"), [](FString Arg) {Settings.bShowGlobalBudgetInfo = FCString::Atoi(*Arg) != 0; }),
 
