Added duplicate user parameter helper method for the User Parameters panel, reorganized user parameter helper functions

diff --git a/Engine/Plugins/FX/Niagara/Source/NiagaraEditor/Private/Customizations/UserParameters/NiagaraSystemUserParametersBuilder.cpp b/Engine/Plugins/FX/Niagara/Source/NiagaraEditor/Private/Customizations/UserParameters/NiagaraSystemUserParametersBuilder.cpp
--- a/Engine/Plugins/FX/Niagara/Source/NiagaraEditor/Private/Customizations/UserParameters/NiagaraSystemUserParametersBuilder.cpp
+++ b/Engine/Plugins/FX/Niagara/Source/NiagaraEditor/Private/Customizations/UserParameters/NiagaraSystemUserParametersBuilder.cpp
@@ -17,6 +17,7 @@
 #include "NiagaraSystemEditorData.h"
 #include "NiagaraUserRedirectionParameterStore.h"
 #include "SNiagaraParameterEditor.h"
+#include "ViewModels/Stack/NiagaraStackViewModel.h"
 #include "Widgets/Input/SButton.h"
 #include "Widgets/SNiagaraDebugger.h"
 #include "Widgets/SNiagaraParameterName.h"
@@ -101,6 +102,9 @@ void FNiagaraSystemUserParameterBuilder::AddCustomMenuActionsForParameter(FDetai
 	FUIAction RenameAction(FExecuteAction::CreateSP(this, &FNiagaraSystemUserParameterBuilder::RequestRename, UserParameter));
 	WidgetRow.AddCustomContextMenuAction(RenameAction, LOCTEXT("RenameParameterAction", "Rename"), LOCTEXT("RenameParameterActionTooltip", "Rename this user parameter"));
 
+	FUIAction DuplicateAction(FExecuteAction::CreateSP(this, &FNiagaraSystemUserParameterBuilder::RequestDuplication, UserParameter));
+	WidgetRow.AddCustomContextMenuAction(DuplicateAction, LOCTEXT("DuplicateParameterAction", "Duplicate"), LOCTEXT("DuplicateParameterActionTooltip", "Duplicate this user parameter"));
+	
 	FUIAction DeleteAction(FExecuteAction::CreateSP(this, &FNiagaraSystemUserParameterBuilder::DeleteParameter, UserParameter));
 	WidgetRow.AddCustomContextMenuAction(DeleteAction, LOCTEXT("DeleteParameterAction", "Delete"), LOCTEXT("DeleteParameterActionTooltip", "Delete this user parameter"));
 }
@@ -392,7 +396,7 @@ TSharedRef<SWidget> FNiagaraSystemUserParameterBuilder::CreateUserParameterNameW
 	UserParamToWidgetMap.Add(UserParameter, ParameterName);
 
 	TAttribute<FText> Tooltip = FText::GetEmpty();
-	if(UNiagaraScriptVariable* ScriptVariable = FNiagaraEditorUtilities::GetScriptVariableForUserParameter(UserParameter, *GetSystem()))
+	if(UNiagaraScriptVariable* ScriptVariable = FNiagaraEditorUtilities::UserParameters::GetScriptVariableForUserParameter(UserParameter, *GetSystem()))
 	{
 		Tooltip = TAttribute<FText>::CreateLambda([ScriptVariable]
 		{
@@ -443,6 +447,17 @@ void FNiagaraSystemUserParameterBuilder::RenameParameter(FNiagaraVariable UserPa
 	{		
 		SystemViewModel.Pin()->RenameParameter(UserParameter, NewName, ENiagaraGetGraphParameterReferencesMode::AllGraphs);
 	}
+
+	Rebuild();
+}
+
+void FNiagaraSystemUserParameterBuilder::RequestDuplication(FNiagaraVariable UserParameter)
+{
+	FNiagaraVariable DuplicatedParameter = FNiagaraEditorUtilities::UserParameters::DuplicateUserParameter(UserParameter,*System.Get());
+	if(DuplicatedParameter.IsValid())
+	{
+		RequestRename(DuplicatedParameter);
+	}
 }
 
 FReply FNiagaraSystemUserParameterBuilder::GenerateParameterDragDropOp(const FGeometry& Geometry, const FPointerEvent& MouseEvent, FNiagaraVariable UserParameter) const
@@ -526,7 +541,14 @@ void FNiagaraSystemUserParameterBuilder::OnObjectAssetChanged(const FAssetData&
 FString FNiagaraSystemUserParameterBuilder::GetObjectAssetPathForUserParameter(FNiagaraVariable UserParameter) const
 {
 	ensure(UserParameter.GetType().IsUObject() && UserParameter.GetType().IsDataInterface() == false);
-	if(UObject* Object = GetCurrentParameterValue(UserParameter).GetUObject())
+	FNiagaraVariant Variant = GetCurrentParameterValue(UserParameter);
+
+	if(Variant.IsValid() == false)
+	{
+		return FString();
+	}
+	
+	if(UObject* Object = Variant.GetUObject())
 	{
 		return Object->GetPathName();
 	}
diff --git a/Engine/Plugins/FX/Niagara/Source/NiagaraEditor/Private/Customizations/UserParameters/NiagaraSystemUserParametersBuilder.h b/Engine/Plugins/FX/Niagara/Source/NiagaraEditor/Private/Customizations/UserParameters/NiagaraSystemUserParametersBuilder.h
--- a/Engine/Plugins/FX/Niagara/Source/NiagaraEditor/Private/Customizations/UserParameters/NiagaraSystemUserParametersBuilder.h
+++ b/Engine/Plugins/FX/Niagara/Source/NiagaraEditor/Private/Customizations/UserParameters/NiagaraSystemUserParametersBuilder.h
@@ -53,6 +53,8 @@ private:
 	void RequestRename(FNiagaraVariable UserParameter);
 	void RenameParameter(FNiagaraVariable UserParameter, FName NewName) const;
 
+	void RequestDuplication(FNiagaraVariable NiagaraVariable);
+
 	FReply GenerateParameterDragDropOp(const FGeometry& Geometry, const FPointerEvent& MouseEvent, FNiagaraVariable UserParameter) const;
 private:
 	FString GetObjectAssetPathForUserParameter(FNiagaraVariable UserParameter) const;
diff --git a/Engine/Plugins/FX/Niagara/Source/NiagaraEditor/Private/Customizations/UserParameters/NiagaraUserParametersBuilderBase.cpp b/Engine/Plugins/FX/Niagara/Source/NiagaraEditor/Private/Customizations/UserParameters/NiagaraUserParametersBuilderBase.cpp
--- a/Engine/Plugins/FX/Niagara/Source/NiagaraEditor/Private/Customizations/UserParameters/NiagaraUserParametersBuilderBase.cpp
+++ b/Engine/Plugins/FX/Niagara/Source/NiagaraEditor/Private/Customizations/UserParameters/NiagaraUserParametersBuilderBase.cpp
@@ -90,7 +90,7 @@ IDetailPropertyRow* FNiagaraUserParameterNodeBuilder::AddValueParameterAsRow(IDe
 	
 	if(bShouldCreateNiagaraEditor)
 	{
-		UNiagaraScriptVariable* UserParameterScriptVariable = FNiagaraEditorUtilities::GetScriptVariableForUserParameter(ChoppedUserParameter, *GetSystem());
+		UNiagaraScriptVariable* UserParameterScriptVariable = FNiagaraEditorUtilities::UserParameters::GetScriptVariableForUserParameter(ChoppedUserParameter, *GetSystem());
 		FNiagaraEditorModule& NiagaraEditorModule = FNiagaraEditorModule::Get();
 		if(TSharedPtr<INiagaraEditorTypeUtilities> TypeUtilities = NiagaraEditorModule.GetTypeUtilities(ParameterType))
 		{
@@ -413,7 +413,7 @@ void FNiagaraUserParameterNodeBuilder::AddCustomMenuActionsForParameter(FDetailW
 
 TSharedRef<SWidget> FNiagaraUserParameterNodeBuilder::CreateUserParameterNameWidget(FNiagaraVariable UserParameter)
 {
-	UNiagaraScriptVariable* ScriptVariable = FNiagaraEditorUtilities::GetScriptVariableForUserParameter(UserParameter, *GetSystem());
+	UNiagaraScriptVariable* ScriptVariable = FNiagaraEditorUtilities::UserParameters::GetScriptVariableForUserParameter(UserParameter, *GetSystem());
 	TAttribute<FText> Tooltip = FText::GetEmpty();
 	if(ScriptVariable)
 	{
diff --git a/Engine/Plugins/FX/Niagara/Source/NiagaraEditor/Private/NiagaraEditorUtilities.cpp b/Engine/Plugins/FX/Niagara/Source/NiagaraEditor/Private/NiagaraEditorUtilities.cpp
--- a/Engine/Plugins/FX/Niagara/Source/NiagaraEditor/Private/NiagaraEditorUtilities.cpp
+++ b/Engine/Plugins/FX/Niagara/Source/NiagaraEditor/Private/NiagaraEditorUtilities.cpp
@@ -2429,22 +2429,22 @@ TArray<TPair<FNiagaraVariableAttributeBinding*, ENiagaraRendererSourceDataMode>>
 	return AttributeBindings;
 }
 
-TObjectPtr<UNiagaraScriptVariable> FNiagaraEditorUtilities::GetScriptVariableForUserParameter(const FNiagaraVariable& UserParameter, TSharedPtr<FNiagaraSystemViewModel> SystemViewModel)
+TObjectPtr<UNiagaraScriptVariable> FNiagaraEditorUtilities::UserParameters::GetScriptVariableForUserParameter(const FNiagaraVariable& UserParameter, TSharedPtr<FNiagaraSystemViewModel> SystemViewModel)
 {
 	return Cast<UNiagaraSystemEditorData>(SystemViewModel->GetSystem().GetEditorData())->FindOrAddUserScriptVariable(UserParameter, SystemViewModel->GetSystem());
 }
 
-TObjectPtr<UNiagaraScriptVariable> FNiagaraEditorUtilities::GetScriptVariableForUserParameter(const FNiagaraVariable& UserParameter, UNiagaraSystem& System)
+TObjectPtr<UNiagaraScriptVariable> FNiagaraEditorUtilities::UserParameters::GetScriptVariableForUserParameter(const FNiagaraVariable& UserParameter, UNiagaraSystem& System)
 {
 	return Cast<UNiagaraSystemEditorData>(System.GetEditorData())->FindOrAddUserScriptVariable(UserParameter, System);
 }
 
-const UNiagaraScriptVariable* FNiagaraEditorUtilities::FindScriptVariableForUserParameter(const FGuid& UserParameterGuid, const UNiagaraSystem& System)
+const UNiagaraScriptVariable* FNiagaraEditorUtilities::UserParameters::FindScriptVariableForUserParameter(const FGuid& UserParameterGuid, const UNiagaraSystem& System)
 {
 	return Cast<UNiagaraSystemEditorData>(System.GetEditorData())->FindUserScriptVariable(UserParameterGuid);
 }
 
-void FNiagaraEditorUtilities::ReplaceUserParameterReferences(TSharedRef<FNiagaraEmitterViewModel> EmitterViewModel,	FNiagaraVariable OldUserParameter, FNiagaraVariable NewUserParameter)
+void FNiagaraEditorUtilities::UserParameters::ReplaceUserParameterReferences(TSharedRef<FNiagaraEmitterViewModel> EmitterViewModel,	FNiagaraVariable OldUserParameter, FNiagaraVariable NewUserParameter)
 {
 	TArray<FNiagaraUserParameterBinding*> ReferencingBindings = FNiagaraEditorUtilities::GetUserParameterBindingsForUserParameter(EmitterViewModel, OldUserParameter);
 
@@ -4646,6 +4646,49 @@ UNiagaraDataInterface* FNiagaraEditorUtilities::GetResolvedRuntimeInstanceForEdi
 	return nullptr;
 }
 
+FNiagaraVariable FNiagaraEditorUtilities::UserParameters::DuplicateUserParameter(FNiagaraVariable ParameterToDuplicate, UNiagaraSystem& System)
+{
+	FScopedTransaction AddTransaction(LOCTEXT("DuplicateUserParameter", "Duplicate User Parameter"));
+	System.Modify();
+
+	FNiagaraUserRedirectionParameterStore::MakeUserVariable(ParameterToDuplicate);
+
+	TSet<FName> ExistingParameterStoreNames;
+	TArray<FNiagaraVariable> ParameterStoreVariables;
+	System.GetExposedParameters().GetParameters(ParameterStoreVariables);
+	for (const FNiagaraVariable& Var : ParameterStoreVariables)
+	{
+		ExistingParameterStoreNames.Add(Var.GetName());
+	}
+
+	FNiagaraVariable DuplicatedParameter = ParameterToDuplicate;
+	DuplicatedParameter.SetName(FNiagaraUtilities::GetUniqueName(ParameterToDuplicate.GetName(), ExistingParameterStoreNames));
+
+	bool bSuccess = System.GetExposedParameters().AddParameter(DuplicatedParameter);
+
+	if(bSuccess)
+	{
+		System.GetExposedParameters().CopyParameterData(System.GetExposedParameters(), ParameterToDuplicate, DuplicatedParameter);
+		TObjectPtr<UNiagaraScriptVariable> ScriptVariableToDuplicate = FNiagaraEditorUtilities::UserParameters::GetScriptVariableForUserParameter(ParameterToDuplicate, System);
+		TObjectPtr<UNiagaraScriptVariable> DuplicatedScriptVariable = FNiagaraEditorUtilities::UserParameters::GetScriptVariableForUserParameter(DuplicatedParameter, System);
+		
+		for(TFieldIterator<FProperty> It(UNiagaraScriptVariable::StaticClass()); It; ++It)
+		{
+			uint8* Source = It->ContainerPtrToValuePtr<uint8>(ScriptVariableToDuplicate);
+			uint8* Target = It->ContainerPtrToValuePtr<uint8>(DuplicatedScriptVariable);
+			It->CopyCompleteValue(Target, Source);
+		}
+
+		DuplicatedScriptVariable->Metadata.CreateNewGuid();
+		DuplicatedScriptVariable->Variable.SetName(DuplicatedParameter.GetName());
+		
+		System.GetExposedParameters().TriggerOnLayoutChanged();
+		return DuplicatedParameter;
+	}
+	
+	return FNiagaraVariable();
+}
+
 TSharedRef<SWidget> FNiagaraEditorUtilities::HierarchyEditor::Scripts::GenerateRowContentForScriptParameterHierarchyEditor(TSharedRef<FNiagaraHierarchyItemViewModelBase> HierarchyItem)
 {
 	if(HierarchyItem->GetDataMutable()->IsA<UNiagaraHierarchyCategory>())
diff --git a/Engine/Plugins/FX/Niagara/Source/NiagaraEditor/Private/Toolkits/SystemToolkitModes/NiagaraSystemToolkitModeBase.cpp b/Engine/Plugins/FX/Niagara/Source/NiagaraEditor/Private/Toolkits/SystemToolkitModes/NiagaraSystemToolkitModeBase.cpp
--- a/Engine/Plugins/FX/Niagara/Source/NiagaraEditor/Private/Toolkits/SystemToolkitModes/NiagaraSystemToolkitModeBase.cpp
+++ b/Engine/Plugins/FX/Niagara/Source/NiagaraEditor/Private/Toolkits/SystemToolkitModes/NiagaraSystemToolkitModeBase.cpp
@@ -883,7 +883,7 @@ TSharedRef<SWidget> GenerateRowContentForUserParameterHierarchyEditor(TSharedRef
 	else if(const UNiagaraHierarchyUserParameter* UserParameter = Cast<UNiagaraHierarchyUserParameter>(HierarchyItem->GetData()))
 	{
 		TSharedRef<SWidget> ParameterWidget = FNiagaraParameterUtilities::GetParameterWidget(UserParameter->GetUserParameter(), true, false);
-		UNiagaraScriptVariable* ScriptVariable = FNiagaraEditorUtilities::GetScriptVariableForUserParameter(UserParameter->GetUserParameter(), SystemViewModel->GetSystem());
+		UNiagaraScriptVariable* ScriptVariable = FNiagaraEditorUtilities::UserParameters::GetScriptVariableForUserParameter(UserParameter->GetUserParameter(), SystemViewModel->GetSystem());
 
 		ParameterWidget->SetToolTipText(TAttribute<FText>::CreateLambda([ScriptVariable]()
 		{
diff --git a/Engine/Plugins/FX/Niagara/Source/NiagaraEditor/Private/ViewModels/HierarchyEditor/NiagaraUserParametersHierarchyViewModel.cpp b/Engine/Plugins/FX/Niagara/Source/NiagaraEditor/Private/ViewModels/HierarchyEditor/NiagaraUserParametersHierarchyViewModel.cpp
--- a/Engine/Plugins/FX/Niagara/Source/NiagaraEditor/Private/ViewModels/HierarchyEditor/NiagaraUserParametersHierarchyViewModel.cpp
+++ b/Engine/Plugins/FX/Niagara/Source/NiagaraEditor/Private/ViewModels/HierarchyEditor/NiagaraUserParametersHierarchyViewModel.cpp
@@ -25,7 +25,7 @@ UObject* FNiagaraHierarchyUserParameterViewModel::GetDataForEditing()
 	FNiagaraVariable ContainedVariable = HierarchyUserParameter->GetUserParameter();
 	UNiagaraUserParametersHierarchyViewModel* UserParametersHierarchyViewModel = Cast<UNiagaraUserParametersHierarchyViewModel>(GetHierarchyViewModel());
 	UserParametersHierarchyViewModel->GetSystemViewModel()->GetSystem().GetExposedParameters().RedirectUserVariable(ContainedVariable);
-	TObjectPtr<UNiagaraScriptVariable> ScriptVariable = FNiagaraEditorUtilities::GetScriptVariableForUserParameter(ContainedVariable, UserParametersHierarchyViewModel->GetSystemViewModel());
+	TObjectPtr<UNiagaraScriptVariable> ScriptVariable = FNiagaraEditorUtilities::UserParameters::GetScriptVariableForUserParameter(ContainedVariable, UserParametersHierarchyViewModel->GetSystemViewModel());
 	return ScriptVariable.Get();
 }
 
@@ -33,7 +33,7 @@ bool FNiagaraHierarchyUserParameterViewModel::DoesExternalDataStillExist(const U
 {
 	const UNiagaraHierarchyUserParameterRefreshContext* UserParameterRefreshContext = CastChecked<UNiagaraHierarchyUserParameterRefreshContext>(Context);
 	const UNiagaraSystem* System = UserParameterRefreshContext->GetSystem();
-	return FNiagaraEditorUtilities::FindScriptVariableForUserParameter(GetData()->GetPersistentIdentity().Guids[0], *System) != nullptr;
+	return FNiagaraEditorUtilities::UserParameters::FindScriptVariableForUserParameter(GetData()->GetPersistentIdentity().Guids[0], *System) != nullptr;
 }
 
 TSharedRef<FNiagaraSystemViewModel> UNiagaraUserParametersHierarchyViewModel::GetSystemViewModel() const
@@ -109,7 +109,7 @@ void UNiagaraUserParametersHierarchyViewModel::PrepareSourceItems(UNiagaraHierar
 	{
 		FNiagaraVariable RedirectedUserParameter = UserParameter;
 		GetSystemViewModel()->GetSystem().GetExposedParameters().RedirectUserVariable(RedirectedUserParameter);
-		TObjectPtr<UNiagaraScriptVariable> ScriptVariable = FNiagaraEditorUtilities::GetScriptVariableForUserParameter(RedirectedUserParameter, GetSystemViewModel());
+		TObjectPtr<UNiagaraScriptVariable> ScriptVariable = FNiagaraEditorUtilities::UserParameters::GetScriptVariableForUserParameter(RedirectedUserParameter, GetSystemViewModel());
 		
 		if(UNiagaraHierarchyItemBase** ExistingChild = OldChildren.FindByPredicate([ScriptVariable](UNiagaraHierarchyItemBase* ItemBase)
 		{
diff --git a/Engine/Plugins/FX/Niagara/Source/NiagaraEditor/Private/ViewModels/NiagaraParameterPanelViewModel.cpp b/Engine/Plugins/FX/Niagara/Source/NiagaraEditor/Private/ViewModels/NiagaraParameterPanelViewModel.cpp
--- a/Engine/Plugins/FX/Niagara/Source/NiagaraEditor/Private/ViewModels/NiagaraParameterPanelViewModel.cpp
+++ b/Engine/Plugins/FX/Niagara/Source/NiagaraEditor/Private/ViewModels/NiagaraParameterPanelViewModel.cpp
@@ -1889,7 +1889,7 @@ TArray<FNiagaraParameterPanelItem> FNiagaraSystemToolkitParameterPanelViewModel:
 		ParamStore->GetParameters(Vars);
 		for (const FNiagaraVariable& Var : Vars)
 		{
-			UNiagaraScriptVariable* ScriptVar = FNiagaraEditorUtilities::GetScriptVariableForUserParameter(Var, SystemViewModel);
+			UNiagaraScriptVariable* ScriptVar = FNiagaraEditorUtilities::UserParameters::GetScriptVariableForUserParameter(Var, SystemViewModel);
 
 			FNiagaraParameterPanelItem Item;
 			Item.ScriptVariable = ScriptVar;
diff --git a/Engine/Plugins/FX/Niagara/Source/NiagaraEditor/Private/ViewModels/NiagaraSystemViewModel.cpp b/Engine/Plugins/FX/Niagara/Source/NiagaraEditor/Private/ViewModels/NiagaraSystemViewModel.cpp
--- a/Engine/Plugins/FX/Niagara/Source/NiagaraEditor/Private/ViewModels/NiagaraSystemViewModel.cpp
+++ b/Engine/Plugins/FX/Niagara/Source/NiagaraEditor/Private/ViewModels/NiagaraSystemViewModel.cpp
@@ -818,13 +818,13 @@ void FNiagaraSystemViewModel::DuplicateEmitters(TArray<FEmitterHandleToDuplicate
 					FNiagaraVariable NewUserParameter = FNiagaraVariable(UserParameter.GetType(), FNiagaraUtilities::GetUniqueName(UserParameter.GetName(), ExistingUserParameterNames));
 					TargetParameterStore.AddParameter(NewUserParameter);
 					OriginalSystem->GetExposedParameters().CopyParameterData(TargetParameterStore, UserParameter, NewUserParameter);
-					FNiagaraEditorUtilities::ReplaceUserParameterReferences(EmitterHandleViewModel->GetEmitterViewModel(), UserParameter, NewUserParameter);
+					FNiagaraEditorUtilities::UserParameters::ReplaceUserParameterReferences(EmitterHandleViewModel->GetEmitterViewModel(), UserParameter, NewUserParameter);
 					RenamedUserParameters.Add(UserParameter, NewUserParameter);
 				}
 				else
 				{
 					FNiagaraVariable RenamedUserParameter = RenamedUserParameters[UserParameter];
-					FNiagaraEditorUtilities::ReplaceUserParameterReferences(EmitterHandleViewModel->GetEmitterViewModel(), UserParameter, RenamedUserParameter);
+					FNiagaraEditorUtilities::UserParameters::ReplaceUserParameterReferences(EmitterHandleViewModel->GetEmitterViewModel(), UserParameter, RenamedUserParameter);
 				}
 			}
 			// if the parameter doesn't exist at all, we just add it, copy the data and we're done
@@ -1409,7 +1409,7 @@ bool FNiagaraSystemViewModel::RenameParameter(const FNiagaraVariable TargetParam
 		NewVariable.SetName(NewName);
 		for(TSharedPtr<FNiagaraEmitterHandleViewModel> EmitterHandleViewModel : GetEmitterHandleViewModels())
 		{
-			FNiagaraEditorUtilities::ReplaceUserParameterReferences(EmitterHandleViewModel->GetEmitterViewModel(), TargetParameter, NewVariable);
+			FNiagaraEditorUtilities::UserParameters::ReplaceUserParameterReferences(EmitterHandleViewModel->GetEmitterViewModel(), TargetParameter, NewVariable);
 		}
 		bExposedParametersRename = true;
 	}
diff --git a/Engine/Plugins/FX/Niagara/Source/NiagaraEditor/Public/NiagaraEditorUtilities.h b/Engine/Plugins/FX/Niagara/Source/NiagaraEditor/Public/NiagaraEditorUtilities.h
--- a/Engine/Plugins/FX/Niagara/Source/NiagaraEditor/Public/NiagaraEditorUtilities.h
+++ b/Engine/Plugins/FX/Niagara/Source/NiagaraEditor/Public/NiagaraEditorUtilities.h
@@ -297,10 +297,6 @@ namespace FNiagaraEditorUtilities
 	TArray<UNiagaraNodeParameterMapGet*> GetParameterMapGetNodesWithUserParameter(TSharedRef<FNiagaraEmitterViewModel> EmitterViewModel, FNiagaraVariable UserParameter);
 	TArray<FNiagaraUserParameterBinding*> GetUserParameterBindingsForUserParameter(TSharedRef<FNiagaraEmitterViewModel> EmitterViewModel, FNiagaraVariable UserParameter);
 	TArray<TPair<FNiagaraVariableAttributeBinding*, ENiagaraRendererSourceDataMode>> GetVariableAttributeBindingsForParameter(TSharedRef<FNiagaraEmitterViewModel> EmitterViewModel, FNiagaraVariable Parameter);
-	NIAGARAEDITOR_API TObjectPtr<UNiagaraScriptVariable> GetScriptVariableForUserParameter(const FNiagaraVariable& UserParameter, TSharedPtr<FNiagaraSystemViewModel> SystemViewModel);
-	NIAGARAEDITOR_API TObjectPtr<UNiagaraScriptVariable> GetScriptVariableForUserParameter(const FNiagaraVariable& UserParameter, UNiagaraSystem& System);
-	NIAGARAEDITOR_API const UNiagaraScriptVariable* FindScriptVariableForUserParameter(const FGuid& UserParameterGuid, const UNiagaraSystem& System);
-	void ReplaceUserParameterReferences(TSharedRef<FNiagaraEmitterViewModel> EmitterViewModel, FNiagaraVariable OldUserParameter, FNiagaraVariable NewUserParameter);
 
 	NIAGARAEDITOR_API bool AddEmitterContextMenuActions(FMenuBuilder& MenuBuilder, const TSharedPtr<FNiagaraEmitterHandleViewModel>& EmitterHandleViewModel);
 
@@ -428,7 +424,15 @@ namespace FNiagaraEditorUtilities
 	NIAGARAEDITOR_API bool IsEditorDataInterfaceInstance(const UNiagaraDataInterface* DataInterface);
 
 	NIAGARAEDITOR_API UNiagaraDataInterface* GetResolvedRuntimeInstanceForEditorDataInterfaceInstance(const UNiagaraSystem& OwningSystem, UNiagaraDataInterface& EditorDataInterfaceInstance);
-	
+
+	namespace UserParameters
+	{
+		FNiagaraVariable DuplicateUserParameter(FNiagaraVariable ParameterToDuplicate, UNiagaraSystem& System);
+		NIAGARAEDITOR_API TObjectPtr<UNiagaraScriptVariable> GetScriptVariableForUserParameter(const FNiagaraVariable& UserParameter, TSharedPtr<FNiagaraSystemViewModel> SystemViewModel);
+		NIAGARAEDITOR_API TObjectPtr<UNiagaraScriptVariable> GetScriptVariableForUserParameter(const FNiagaraVariable& UserParameter, UNiagaraSystem& System);
+		NIAGARAEDITOR_API const UNiagaraScriptVariable* FindScriptVariableForUserParameter(const FGuid& UserParameterGuid, const UNiagaraSystem& System);
+		void ReplaceUserParameterReferences(TSharedRef<FNiagaraEmitterViewModel> EmitterViewModel, FNiagaraVariable OldUserParameter, FNiagaraVariable NewUserParameter);
+	}
 	namespace Tooltips
 	{
 		NIAGARAEDITOR_API TSharedRef<SToolTip> CreateStackNoteTooltip(UNiagaraStackNote& StackNote);
diff --git a/Engine/Plugins/FX/Niagara/Source/NiagaraEditorWidgets/Private/Stack/SNiagaraStackFunctionInputValue.cpp b/Engine/Plugins/FX/Niagara/Source/NiagaraEditorWidgets/Private/Stack/SNiagaraStackFunctionInputValue.cpp
--- a/Engine/Plugins/FX/Niagara/Source/NiagaraEditorWidgets/Private/Stack/SNiagaraStackFunctionInputValue.cpp
+++ b/Engine/Plugins/FX/Niagara/Source/NiagaraEditorWidgets/Private/Stack/SNiagaraStackFunctionInputValue.cpp
@@ -332,7 +332,7 @@ TSharedRef<SWidget> SNiagaraStackFunctionInputValue::ConstructValueWidgets()
 
 			if(MatchingVariable)
 			{
-				FText Tooltip = FNiagaraEditorUtilities::GetScriptVariableForUserParameter(*MatchingVariable, FunctionInput->GetSystemViewModel())->Metadata.Description;
+				FText Tooltip = FNiagaraEditorUtilities::UserParameters::GetScriptVariableForUserParameter(*MatchingVariable, FunctionInput->GetSystemViewModel())->Metadata.Description;
 				if(!Tooltip.IsEmpty())
 				{
 					ParameterWidget->SetToolTipText(Tooltip);
