- Renamed 'Empty' emitter template to 'Minimal'
- Updated Minimal emitter description
- Changed "Add empty emitter" menu entry to "Add default empty emitter", and changed the tooltip to mention that the default empty emitter is specified in Niagara settings, as well as to include the description of the selected asset. This is supposed to close the gap in understanding what's happening

#virtualized

diff --git a/Engine/Plugins/FX/Niagara/Config/BaseNiagara.ini b/Engine/Plugins/FX/Niagara/Config/BaseNiagara.ini
--- a/Engine/Plugins/FX/Niagara/Config/BaseNiagara.ini
+++ b/Engine/Plugins/FX/Niagara/Config/BaseNiagara.ini
@@ -68,7 +68,7 @@ DefaultScript=/Niagara/DefaultAssets/DefaultModule.DefaultModule
 DefaultDynamicInputScript=/Niagara/DefaultAssets/DefaultDynamicInputs.DefaultDynamicInputs
 DefaultFunctionScript=/Niagara/DefaultAssets/DefaultFunction.DefaultFunction
 DefaultModuleScript=/Niagara/DefaultAssets/DefaultModule.DefaultModule
-DefaultEmptyEmitter=/Niagara/DefaultAssets/Templates/Emitters/Empty.Empty
+DefaultEmptyEmitter=/Niagara/DefaultAssets/Templates/Emitters/Minimal.Minimal
 RequiredSystemUpdateScript=/Niagara/Modules/System/SystemState.SystemState
 +GraphCreationShortcuts=(Name="Numeric::Add",Input=(Key=A,bShift=False,bCtrl=False,bAlt=False,bCmd=False))
 +GraphCreationShortcuts=(Name="Numeric::Divide",Input=(Key=D,bShift=False,bCtrl=False,bAlt=False,bCmd=False))
diff --git a/Engine/Plugins/FX/Niagara/Source/NiagaraEditor/Private/NiagaraEditorUtilities.cpp b/Engine/Plugins/FX/Niagara/Source/NiagaraEditor/Private/NiagaraEditorUtilities.cpp
--- a/Engine/Plugins/FX/Niagara/Source/NiagaraEditor/Private/NiagaraEditorUtilities.cpp
+++ b/Engine/Plugins/FX/Niagara/Source/NiagaraEditor/Private/NiagaraEditorUtilities.cpp
@@ -4777,6 +4777,27 @@ TSharedRef<SToolTip> FNiagaraEditorUtilities::Tooltips::CreateStackNoteTooltip(U
 	];
 }
 
+FText FNiagaraEditorUtilities::Tooltips::GetMinimalEmitterCreationTooltip()
+{
+	FText EmitterDescription = LOCTEXT("TrueEmptyEmitterTooltip", "Add a completely empty emitter.");
+	FSoftObjectPath DefaultEmptyEmitter = GetDefault<UNiagaraEditorSettings>()->DefaultEmptyEmitter;
+	if(DefaultEmptyEmitter.IsValid() && DefaultEmptyEmitter.IsAsset())
+	{
+		if(UNiagaraEmitter* Emitter = Cast<UNiagaraEmitter>(DefaultEmptyEmitter.TryLoad()))
+		{
+			EmitterDescription =  Emitter->TemplateAssetDescription;
+		}
+	}
+
+	FText TooltipText = LOCTEXT("AddMinimalEmitterTooltip", "Adds a minimal emitter as specified in the Niagara Settings");
+	if(EmitterDescription.IsEmpty())
+	{
+		return TooltipText;
+	}
+			
+	return FText::FormatOrdered( FText::AsCultureInvariant("{0}\n\n{1}."), TooltipText, EmitterDescription);
+}
+
 TMap<FGuid, TArray<FNiagaraVariableBase>> FNiagaraEditorUtilities::Scripts::Validation::ValidateScriptVariableIds(UNiagaraScript* Script, FGuid VersionGuid)
 {
 	bool bReturnResults = false;
diff --git a/Engine/Plugins/FX/Niagara/Source/NiagaraEditor/Private/ViewModels/NiagaraSystemViewModel.cpp b/Engine/Plugins/FX/Niagara/Source/NiagaraEditor/Private/ViewModels/NiagaraSystemViewModel.cpp
--- a/Engine/Plugins/FX/Niagara/Source/NiagaraEditor/Private/ViewModels/NiagaraSystemViewModel.cpp
+++ b/Engine/Plugins/FX/Niagara/Source/NiagaraEditor/Private/ViewModels/NiagaraSystemViewModel.cpp
@@ -686,12 +686,12 @@ TSharedPtr<FNiagaraEmitterHandleViewModel> FNiagaraSystemViewModel::AddEmitter(c
 	return AddEmitter(*VersionedEmitter.Emitter, VersionedEmitter.Version);
 }
 
-TSharedPtr<FNiagaraEmitterHandleViewModel> FNiagaraSystemViewModel::AddDefaultEmptyEmitter()
+TSharedPtr<FNiagaraEmitterHandleViewModel> FNiagaraSystemViewModel::AddMinimalEmitter()
 {
-	FSoftObjectPath DefaultEmptyEmitter = GetDefault<UNiagaraEditorSettings>()->DefaultEmptyEmitter;
-	if(DefaultEmptyEmitter.IsValid() && DefaultEmptyEmitter.IsAsset())
+	FSoftObjectPath MinimalEmitter = GetDefault<UNiagaraEditorSettings>()->DefaultEmptyEmitter;
+	if(MinimalEmitter.IsValid() && MinimalEmitter.IsAsset())
 	{
-		if(UNiagaraEmitter* Emitter = Cast<UNiagaraEmitter>(DefaultEmptyEmitter.TryLoad()))
+		if(UNiagaraEmitter* Emitter = Cast<UNiagaraEmitter>(MinimalEmitter.TryLoad()))
 		{
 			return AddEmitterFromAssetData(FAssetData(Emitter));
 		}
diff --git a/Engine/Plugins/FX/Niagara/Source/NiagaraEditor/Private/Widgets/AssetBrowser/SNiagaraAssetBrowser.cpp b/Engine/Plugins/FX/Niagara/Source/NiagaraEditor/Private/Widgets/AssetBrowser/SNiagaraAssetBrowser.cpp
--- a/Engine/Plugins/FX/Niagara/Source/NiagaraEditor/Private/Widgets/AssetBrowser/SNiagaraAssetBrowser.cpp
+++ b/Engine/Plugins/FX/Niagara/Source/NiagaraEditor/Private/Widgets/AssetBrowser/SNiagaraAssetBrowser.cpp
@@ -1222,6 +1222,7 @@ void SNiagaraAddEmitterToSystemWindow::Construct(const FArguments& InArgs, TShar
 		   [
 			   SNew(SButton)
 			   .OnClicked(this, &SNiagaraAddEmitterToSystemWindow::AddEmptyEmitter)
+			   .ToolTipText(FNiagaraEditorUtilities::Tooltips::GetMinimalEmitterCreationTooltip())
 			   [
 				   SNew(SHorizontalBox)
 				   + SHorizontalBox::Slot()
@@ -1235,7 +1236,7 @@ void SNiagaraAddEmitterToSystemWindow::Construct(const FArguments& InArgs, TShar
 				   .AutoWidth()
 				   .Padding(2.f)
 				   [
-					   SNew(STextBlock).Text(FText::FormatOrdered(LOCTEXT("AddEmptyEmitterButtonLabel", "Add Empty {0}"), UNiagaraEmitter::StaticClass()->GetDisplayNameText()))
+					   SNew(STextBlock).Text(FText::FormatOrdered(LOCTEXT("AddEmptyEmitterButtonLabel", "Add Minimal {0}"), UNiagaraEmitter::StaticClass()->GetDisplayNameText()))
 				   ]
 			   ]
 		   ]
@@ -1283,7 +1284,7 @@ void SNiagaraAddEmitterToSystemWindow::OnAssetsActivatedInternal(const TArray<FA
 	{
 		if(AssetData.Num() == 0)
 		{
-			WeakSystemViewModel.Pin()->AddDefaultEmptyEmitter();
+			WeakSystemViewModel.Pin()->AddMinimalEmitter();
 		}
 		else
 		{
diff --git a/Engine/Plugins/FX/Niagara/Source/NiagaraEditor/Public/NiagaraEditorSettings.h b/Engine/Plugins/FX/Niagara/Source/NiagaraEditor/Public/NiagaraEditorSettings.h
--- a/Engine/Plugins/FX/Niagara/Source/NiagaraEditor/Public/NiagaraEditorSettings.h
+++ b/Engine/Plugins/FX/Niagara/Source/NiagaraEditor/Public/NiagaraEditorSettings.h
@@ -362,8 +362,8 @@ public:
 	UPROPERTY(config, EditAnywhere, Category = Niagara)
 	FSoftObjectPath DefaultModuleScript;
 
-	/** Niagara emitter to duplicate as the base for empty emitters. Useful to set up default state. If empty, a truly empty emitter will be used. */
-	UPROPERTY(config, EditAnywhere, Category = Niagara, meta=(AllowedClasses="/Script/Niagara.NiagaraEmitter"))
+	/** Niagara emitter to duplicate as the base for minimal emitters. Useful to set up default state. If empty, a truly empty emitter will be used. */
+	UPROPERTY(config, EditAnywhere, DisplayName="Minimal Emitter", Category = Niagara, meta=(AllowedClasses="/Script/Niagara.NiagaraEmitter"))
 	FSoftObjectPath DefaultEmptyEmitter;
 
 	/** Niagara script which is required in the system update script to control system state. */
diff --git a/Engine/Plugins/FX/Niagara/Source/NiagaraEditor/Public/NiagaraEditorUtilities.h b/Engine/Plugins/FX/Niagara/Source/NiagaraEditor/Public/NiagaraEditorUtilities.h
--- a/Engine/Plugins/FX/Niagara/Source/NiagaraEditor/Public/NiagaraEditorUtilities.h
+++ b/Engine/Plugins/FX/Niagara/Source/NiagaraEditor/Public/NiagaraEditorUtilities.h
@@ -436,6 +436,7 @@ namespace FNiagaraEditorUtilities
 	namespace Tooltips
 	{
 		NIAGARAEDITOR_API TSharedRef<SToolTip> CreateStackNoteTooltip(UNiagaraStackNote& StackNote);
+		NIAGARAEDITOR_API FText GetMinimalEmitterCreationTooltip();
 	}
 	
 	namespace Scripts
diff --git a/Engine/Plugins/FX/Niagara/Source/NiagaraEditor/Public/ViewModels/NiagaraSystemViewModel.h b/Engine/Plugins/FX/Niagara/Source/NiagaraEditor/Public/ViewModels/NiagaraSystemViewModel.h
--- a/Engine/Plugins/FX/Niagara/Source/NiagaraEditor/Public/ViewModels/NiagaraSystemViewModel.h
+++ b/Engine/Plugins/FX/Niagara/Source/NiagaraEditor/Public/ViewModels/NiagaraSystemViewModel.h
@@ -245,7 +245,7 @@ public:
 	NIAGARAEDITOR_API TSharedPtr<FNiagaraEmitterHandleViewModel> AddEmitter(const FVersionedNiagaraEmitter& Emitter);
 
 	/** Adds the default empty emitter to the system, or a truly empty emitter as a fallback. */
-	NIAGARAEDITOR_API TSharedPtr<FNiagaraEmitterHandleViewModel> AddDefaultEmptyEmitter();
+	NIAGARAEDITOR_API TSharedPtr<FNiagaraEmitterHandleViewModel> AddMinimalEmitter();
 
 	/** Adds a truly empty emitter to the system. */
 	NIAGARAEDITOR_API TSharedPtr<FNiagaraEmitterHandleViewModel> AddEmptyEmitter();
diff --git a/Engine/Plugins/FX/Niagara/Source/NiagaraEditorWidgets/Private/SNiagaraOverviewGraph.cpp b/Engine/Plugins/FX/Niagara/Source/NiagaraEditorWidgets/Private/SNiagaraOverviewGraph.cpp
--- a/Engine/Plugins/FX/Niagara/Source/NiagaraEditorWidgets/Private/SNiagaraOverviewGraph.cpp
+++ b/Engine/Plugins/FX/Niagara/Source/NiagaraEditorWidgets/Private/SNiagaraOverviewGraph.cpp
@@ -238,8 +238,8 @@ FActionMenuContent SNiagaraOverviewGraph::OnCreateGraphActionMenu(UEdGraph* InGr
 		{
 			MenuBuilder.AddMenuEntry(FNiagaraEditorCommands::Get().OpenAddEmitterMenu);
 			MenuBuilder.AddMenuEntry(
-				LOCTEXT("EmptyEmitterLabel", "Add empty emitter"),
-				LOCTEXT("AddEmitterToolTip", "Adds an empty emitter without any modules or renderers."),
+				LOCTEXT("MinimalEmitterLabel", "Add minimal emitter"),
+				FNiagaraEditorUtilities::Tooltips::GetMinimalEmitterCreationTooltip(),
 				FSlateIcon(),
 				FExecuteAction::CreateSP(this, &SNiagaraOverviewGraph::OnCreateEmptyEmitter));
 			MenuBuilder.AddMenuEntry(
@@ -283,7 +283,7 @@ FActionMenuContent SNiagaraOverviewGraph::OnCreateGraphActionMenu(UEdGraph* InGr
 
 void SNiagaraOverviewGraph::OnCreateEmptyEmitter()
 {
-	ViewModel->GetSystemViewModel()->AddDefaultEmptyEmitter();
+	ViewModel->GetSystemViewModel()->AddMinimalEmitter();
 }
 
 void SNiagaraOverviewGraph::OnCreateStatelessEmitter()
