- Add type interface to all distribution value editors
- Add unit customization to sprite rotation rate and modify mesh rotation to degrees rather than deg/s to reduce visual clutter

#rnx

diff --git a/Engine/Plugins/FX/Niagara/Source/Niagara/Internal/Stateless/Modules/NiagaraStatelessModule_MeshRotationRate.h b/Engine/Plugins/FX/Niagara/Source/Niagara/Internal/Stateless/Modules/NiagaraStatelessModule_MeshRotationRate.h
--- a/Engine/Plugins/FX/Niagara/Source/Niagara/Internal/Stateless/Modules/NiagaraStatelessModule_MeshRotationRate.h
+++ b/Engine/Plugins/FX/Niagara/Source/Niagara/Internal/Stateless/Modules/NiagaraStatelessModule_MeshRotationRate.h
@@ -26,7 +26,7 @@ class UNiagaraStatelessModule_MeshRotationRate : public UNiagaraStatelessModule
 public:
 	using FParameters = NiagaraStateless::FMeshRotationRateModule_ShaderParameters;
 
-	UPROPERTY(EditAnywhere, Category = "Parameters", meta = (DisplayName = "Rotation Rate", Units="deg/s"))
+	UPROPERTY(EditAnywhere, Category = "Parameters", meta = (DisplayName = "Rotation Rate", Units="deg"))
 	FNiagaraDistributionRangeVector3 RotationRateDistribution = FNiagaraDistributionRangeVector3(FVector3f::ZeroVector);
 
 	virtual void BuildEmitterData(const FNiagaraStatelessEmitterDataBuildContext& BuildContext) const override
diff --git a/Engine/Plugins/FX/Niagara/Source/Niagara/Internal/Stateless/Modules/NiagaraStatelessModule_SpriteRotationRate.h b/Engine/Plugins/FX/Niagara/Source/Niagara/Internal/Stateless/Modules/NiagaraStatelessModule_SpriteRotationRate.h
--- a/Engine/Plugins/FX/Niagara/Source/Niagara/Internal/Stateless/Modules/NiagaraStatelessModule_SpriteRotationRate.h
+++ b/Engine/Plugins/FX/Niagara/Source/Niagara/Internal/Stateless/Modules/NiagaraStatelessModule_SpriteRotationRate.h
@@ -25,7 +25,7 @@ class UNiagaraStatelessModule_SpriteRotationRate : public UNiagaraStatelessModul
 public:
 	using FParameters = NiagaraStateless::FSpriteRotationRateModule_ShaderParameters;
 
-	UPROPERTY(EditAnywhere, Category = "Parameters", meta = (DisplayName = "Rotation Rate"))
+	UPROPERTY(EditAnywhere, Category = "Parameters", meta = (DisplayName = "Rotation Rate", Units = "deg"))
 	FNiagaraDistributionRangeFloat RotationRateDistribution = FNiagaraDistributionRangeFloat(FNiagaraStatelessGlobals::GetDefaultSpriteRotationValue());
 
 	virtual void BuildEmitterData(const FNiagaraStatelessEmitterDataBuildContext& BuildContext) const override
diff --git a/Engine/Plugins/FX/Niagara/Source/NiagaraEditor/Private/Widgets/SNiagaraDistributionEditor.cpp b/Engine/Plugins/FX/Niagara/Source/NiagaraEditor/Private/Widgets/SNiagaraDistributionEditor.cpp
--- a/Engine/Plugins/FX/Niagara/Source/NiagaraEditor/Private/Widgets/SNiagaraDistributionEditor.cpp
+++ b/Engine/Plugins/FX/Niagara/Source/NiagaraEditor/Private/Widgets/SNiagaraDistributionEditor.cpp
@@ -349,6 +349,7 @@ private:
 
 	TSharedRef<SWidget> ConstructVectorWidget(int32 ChannelCount, int32 ValueIndex)
 	{
+		const EUnit DisplayUnit = DistributionAdapter->GetDisplayUnit();
 		if (ChannelCount == 2)
 		{
 			return SNew(SBox)
@@ -366,6 +367,7 @@ private:
 					.OnYCommitted(this, &SNiagaraDistributionValueEditor::ValueCommitted, 1, ValueIndex)
 					.OnBeginSliderMovement(this, &SNiagaraDistributionValueEditor::BeginValueChange)
 					.OnEndSliderMovement(this, &SNiagaraDistributionValueEditor::EndValueChange)
+					.TypeInterface(MakeShareable(new TNumericUnitTypeInterface<float>(DisplayUnit)))
 				];
 		}
 		if (ChannelCount == 3)
@@ -388,6 +390,7 @@ private:
 					.OnZCommitted(this, &SNiagaraDistributionValueEditor::ValueCommitted, 2, ValueIndex)
 					.OnBeginSliderMovement(this, &SNiagaraDistributionValueEditor::BeginValueChange)
 					.OnEndSliderMovement(this, &SNiagaraDistributionValueEditor::EndValueChange)
+					.TypeInterface(MakeShareable(new TNumericUnitTypeInterface<float>(DisplayUnit)))
 				];
 		}
 		if (ChannelCount == 4)
@@ -413,6 +416,7 @@ private:
 					.OnWCommitted(this, &SNiagaraDistributionValueEditor::ValueCommitted, 3, ValueIndex)
 					.OnBeginSliderMovement(this, &SNiagaraDistributionValueEditor::BeginValueChange)
 					.OnEndSliderMovement(this, &SNiagaraDistributionValueEditor::EndValueChange)
+					.TypeInterface(MakeShareable(new TNumericUnitTypeInterface<float>(DisplayUnit)))
 				];
 		}
 		return SNullWidget::NullWidget;
