- Update random seed if we kill particles on reset and clear the unique index
- When deterministic this will result in the same behavior on reset, when non-deterministic the behavoir will differ on reset
- This matches how stateful emitters work

#rnx

diff --git a/Engine/Plugins/FX/Niagara/Source/Niagara/Private/Stateless/NiagaraStatelessEmitterInstance.cpp b/Engine/Plugins/FX/Niagara/Source/Niagara/Private/Stateless/NiagaraStatelessEmitterInstance.cpp
--- a/Engine/Plugins/FX/Niagara/Source/Niagara/Private/Stateless/NiagaraStatelessEmitterInstance.cpp
+++ b/Engine/Plugins/FX/Niagara/Source/Niagara/Private/Stateless/NiagaraStatelessEmitterInstance.cpp
@@ -142,7 +142,12 @@ void FNiagaraStatelessEmitterInstance::ResetSimulation(bool bKillExisting)
 	if (bKillExisting)
 	{
 		SpawnInfos.Empty();
-		bSpawnInfosDirty = true;
+
+		UniqueIndexOffset = 0;
+		if (!EmitterData->bDeterministic)
+		{
+			RandomSeed ^= FPlatformTime::Cycles();
+		}
 	}
 	else
 	{
@@ -153,11 +158,11 @@ void FNiagaraStatelessEmitterInstance::ResetSimulation(bool bKillExisting)
 		}
 	}
 	ActiveSpawnRates.Empty();
+	bSpawnInfosDirty = true;
 
 	RandomStream.Initialize(RandomSeed);
 
 	Age = 0.0f;
-	UniqueIndexOffset = 0;
 	bEmitterEnabled_CNC = bEmitterEnabled_GT;
 
 	InitEmitterState();
@@ -169,11 +174,16 @@ void FNiagaraStatelessEmitterInstance::ResetSimulation(bool bKillExisting)
 	if (NiagaraStateless::FEmitterInstance_RT* RenderThreadData = RenderThreadDataPtr.Get())
 	{
 		ENQUEUE_RENDER_COMMAND(UpdateStatelessAge)(
-			[RenderThreadData](FRHICommandListImmediate& RHICmdList)
+			[RenderThreadData, RandomSeed_RT=RandomSeed](FRHICommandListImmediate& RHICmdList)
 			{
 				RenderThreadData->Age = 0.0f;
 				RenderThreadData->DeltaTime = 0.0f;
 				RenderThreadData->ExecutionState = ENiagaraExecutionState::Active;
+				RenderThreadData->RandomSeed = RandomSeed_RT;
+				if (RenderThreadData->ShaderParameters.IsValid())
+				{
+					RenderThreadData->ShaderParameters->Common_RandomSeed = RandomSeed_RT;
+				}
 			}
 		);
 	}
