- Filter the overview information when system filtering is enabled
- Change the default to filter the system overview

#rnx
#tests Tested out filtering systems & enabling + disable the filtering overview option

diff --git a/Engine/Plugins/FX/Niagara/Source/Niagara/Private/NiagaraDebugHud.cpp b/Engine/Plugins/FX/Niagara/Source/Niagara/Private/NiagaraDebugHud.cpp
--- a/Engine/Plugins/FX/Niagara/Source/Niagara/Private/NiagaraDebugHud.cpp
+++ b/Engine/Plugins/FX/Niagara/Source/Niagara/Private/NiagaraDebugHud.cpp
@@ -1117,8 +1117,8 @@ void FNiagaraDebugHud::GatherSystemInfo()
 				LongestSystemPrettyName = SystemDebugInfo.SystemPrettyName;
 			}
 		}
-		SystemDebugInfo.bShowInWorld = (Settings.SystemDebugVerbosity != ENiagaraDebugHudVerbosity::None) && (!Settings.bSystemFilterEnabled || Settings.SystemFilter.IsEmpty() || SystemDebugInfo.SystemName.MatchesWildcard(Settings.SystemFilter));
-		SystemDebugInfo.bPassesSystemFilter = SystemDebugInfo.bShowInWorld;
+		SystemDebugInfo.bPassesSystemFilter = !Settings.bSystemFilterEnabled || Settings.SystemFilter.IsEmpty() || SystemDebugInfo.SystemName.MatchesWildcard(Settings.SystemFilter);
+		SystemDebugInfo.bShowInWorld = (Settings.SystemDebugVerbosity != ENiagaraDebugHudVerbosity::None) && SystemDebugInfo.bPassesSystemFilter;
 
 		const bool bCanShowInWorld = 
 			SystemDebugInfo.bShowInWorld &&
@@ -2110,7 +2110,7 @@ void FNiagaraDebugHud::DrawOverview(class FNiagaraWorldManager* WorldManager, FC
 		{
 			const FSystemDebugInfo& SystemInfo = Pair.Value;
 			if ((SystemInfo.FramesSinceVisible >= Settings.PerfHistoryFrames) ||
-				(Settings.bOverviewShowFilteredSystemOnly && !SystemInfo.bPassesSystemFilter))
+				(Settings.bSystemFilterEnabled && Settings.bOverviewShowFilteredSystemOnly && !SystemInfo.bPassesSystemFilter))
 			{
 				continue;
 			}
diff --git a/Engine/Plugins/FX/Niagara/Source/Niagara/Public/NiagaraDebuggerCommon.h b/Engine/Plugins/FX/Niagara/Source/Niagara/Public/NiagaraDebuggerCommon.h
--- a/Engine/Plugins/FX/Niagara/Source/Niagara/Public/NiagaraDebuggerCommon.h
+++ b/Engine/Plugins/FX/Niagara/Source/Niagara/Public/NiagaraDebuggerCommon.h
@@ -464,11 +464,11 @@ struct FNiagaraDebugHUDSettingsData
 	bool bHudRenderingEnabled = true;
 
 	/** Primary control for validation settings. */
-	UPROPERTY(VisibleAnywhere, Category = "Debug Overview")
+	UPROPERTY()
 	bool bValidationEnabled = false;
 
 	/** When enabled the overview display will be enabled. */
-	UPROPERTY(VisibleAnywhere, Category = "Debug Overview")
+	UPROPERTY()
 	bool bOverviewEnabled = false;
 
 	UPROPERTY(EditAnywhere, Category = "Debug Overview", meta = (DisplayName = "Debug Overview Mode", EditCondition = "bOverviewEnabled"))
@@ -486,8 +486,8 @@ struct FNiagaraDebugHUDSettingsData
 	bool bShowRegisteredComponents = false;
 	
 	/** When enabled the overview will only show the filter system information. */
-	UPROPERTY(EditAnywhere, Category = "Debug Overview", meta = (DisplayName = "Show only filtered system info", EditCondition = "bOverviewEnabled"))
-	bool bOverviewShowFilteredSystemOnly = false;
+	UPROPERTY()
+	bool bOverviewShowFilteredSystemOnly = true;
 
 	UPROPERTY(EditAnywhere, Category = "Debug Overview", AdvancedDisplay)
 	bool bShowGlobalBudgetInfo = false;
@@ -495,7 +495,7 @@ struct FNiagaraDebugHUDSettingsData
 	/**
 	If enabled, then only systems matching the filters will be considered by the debug hud
 	*/
-	UPROPERTY(EditAnywhere, Category = "Debug Filter", meta = (DisplayName = "Filter displayed systems"))
+	UPROPERTY(EditAnywhere, Category = "Debug Filter", meta = (DisplayName = "Filter Displayed Systems"))
 	bool bSystemFilterEnabled = false;
 
 	/**
diff --git a/Engine/Plugins/FX/Niagara/Source/NiagaraEditor/Private/Customizations/NiagaraDebugHUDCustomization.cpp b/Engine/Plugins/FX/Niagara/Source/NiagaraEditor/Private/Customizations/NiagaraDebugHUDCustomization.cpp
--- a/Engine/Plugins/FX/Niagara/Source/NiagaraEditor/Private/Customizations/NiagaraDebugHUDCustomization.cpp
+++ b/Engine/Plugins/FX/Niagara/Source/NiagaraEditor/Private/Customizations/NiagaraDebugHUDCustomization.cpp
@@ -243,12 +243,6 @@ FNiagaraDebugHUDSettingsDetailsCustomization::FNiagaraDebugHUDSettingsDetailsCus
 
 void FNiagaraDebugHUDSettingsDetailsCustomization::CustomizeDetails(IDetailLayoutBuilder& DetailBuilder)
 {
-	{
-		IDetailCategoryBuilder& OverviewCategory = DetailBuilder.EditCategory("Debug Overview");
-		DetailBuilder.HideProperty(DetailBuilder.GetProperty(GET_MEMBER_NAME_CHECKED(FNiagaraDebugHUDSettingsData, bValidationEnabled), FNiagaraDebugHUDSettingsData::StaticStruct()));
-		DetailBuilder.HideProperty(DetailBuilder.GetProperty(GET_MEMBER_NAME_CHECKED(FNiagaraDebugHUDSettingsData, bOverviewEnabled), FNiagaraDebugHUDSettingsData::StaticStruct()));
-	}
-	
 	// Customize Filters
 	{
 		IDetailCategoryBuilder& FilterCategory = DetailBuilder.EditCategory("Debug Filter");
diff --git a/Engine/Plugins/FX/Niagara/Source/NiagaraEditor/Private/Widgets/SNiagaraDebugger.cpp b/Engine/Plugins/FX/Niagara/Source/NiagaraEditor/Private/Widgets/SNiagaraDebugger.cpp
--- a/Engine/Plugins/FX/Niagara/Source/NiagaraEditor/Private/Widgets/SNiagaraDebugger.cpp
+++ b/Engine/Plugins/FX/Niagara/Source/NiagaraEditor/Private/Widgets/SNiagaraDebugger.cpp
@@ -1448,6 +1448,17 @@ TSharedRef<SWidget> SNiagaraDebugger::MakeModeOptionsMenu()
 			),
 			NAME_None, EUserInterfaceActionType::ToggleButton);
 
+		MenuBuilder.AddMenuEntry(
+			LOCTEXT("OnlyShowFilteredSystems", "Show Filtered Systems Only in Overview"),
+			LOCTEXT("EnableOverviewUITooltip", "When enabled only the filtered systems will be shown in the overview."),
+			FSlateIcon(),
+			FUIAction(
+				FExecuteAction::CreateLambda([=]() {Settings->Data.bOverviewShowFilteredSystemOnly = !Settings->Data.bOverviewShowFilteredSystemOnly; Settings->NotifyPropertyChanged(); }),
+				FCanExecuteAction::CreateLambda([=]() { return Settings->Data.bHudRenderingEnabled && Settings->Data.bOverviewEnabled; }),
+				FIsActionChecked::CreateLambda([=]() { return Settings->Data.bOverviewShowFilteredSystemOnly; })
+			),
+			NAME_None, EUserInterfaceActionType::ToggleButton);
+
 		MenuBuilder.AddMenuEntry(
 			LOCTEXT("EnableValidation", "Data Validation"),
 			LOCTEXT("EnableValidationTooltip", "Will allow you to validate system and particle data buffers when enabled."),
