- Fix duplicate loc text key

#rnx

diff --git a/Engine/Plugins/FX/Niagara/Source/NiagaraEditor/Private/Widgets/SNiagaraDebugger.cpp b/Engine/Plugins/FX/Niagara/Source/NiagaraEditor/Private/Widgets/SNiagaraDebugger.cpp
--- a/Engine/Plugins/FX/Niagara/Source/NiagaraEditor/Private/Widgets/SNiagaraDebugger.cpp
+++ b/Engine/Plugins/FX/Niagara/Source/NiagaraEditor/Private/Widgets/SNiagaraDebugger.cpp
@@ -1450,7 +1450,7 @@ TSharedRef<SWidget> SNiagaraDebugger::MakeModeOptionsMenu()
 
 		MenuBuilder.AddMenuEntry(
 			LOCTEXT("OnlyShowFilteredSystems", "Show Filtered Systems Only in Overview"),
-			LOCTEXT("EnableOverviewUITooltip", "When enabled only the filtered systems will be shown in the overview."),
+			LOCTEXT("OnlyShowFilteredSystemsTooltip", "When enabled only the filtered systems will be shown in the overview."),
 			FSlateIcon(),
 			FUIAction(
 				FExecuteAction::CreateLambda([=]() {Settings->Data.bOverviewShowFilteredSystemOnly = !Settings->Data.bOverviewShowFilteredSystemOnly; Settings->NotifyPropertyChanged(); }),
