- Fix for parmeters not updating due to transform not changing

#rnx
#tests Added binding in intialize particle to user parameter, which would not set correctly until after this change

diff --git a/Engine/Plugins/FX/Niagara/Source/Niagara/Private/Stateless/NiagaraStatelessEmitterInstance.cpp b/Engine/Plugins/FX/Niagara/Source/Niagara/Private/Stateless/NiagaraStatelessEmitterInstance.cpp
--- a/Engine/Plugins/FX/Niagara/Source/Niagara/Private/Stateless/NiagaraStatelessEmitterInstance.cpp
+++ b/Engine/Plugins/FX/Niagara/Source/Niagara/Private/Stateless/NiagaraStatelessEmitterInstance.cpp
@@ -499,7 +499,7 @@ void FNiagaraStatelessEmitterInstance::SendRenderData()
 	}
 
 	const FTransform3f ParentTransform = FTransform3f(ParentSystemInstance->GetWorldTransform());
-	bNeedsShaderParametersUpdate = EmitterTransforms.UpdateTransforms(ParentTransform);
+	bNeedsShaderParametersUpdate |= EmitterTransforms.UpdateTransforms(ParentTransform);
 
 	if (bNeedsShaderParametersUpdate)
 	{
