- Separate out cache data and data generation requests

#rnx
#tests Engine Tests / local stress test, ran with asan also on console

diff --git a/Engine/Plugins/FX/Niagara/Source/Niagara/Private/Stateless/NiagaraStatelessComputeManager.cpp b/Engine/Plugins/FX/Niagara/Source/Niagara/Private/Stateless/NiagaraStatelessComputeManager.cpp
--- a/Engine/Plugins/FX/Niagara/Source/Niagara/Private/Stateless/NiagaraStatelessComputeManager.cpp
+++ b/Engine/Plugins/FX/Niagara/Source/Niagara/Private/Stateless/NiagaraStatelessComputeManager.cpp
@@ -56,9 +56,9 @@ namespace NiagaraStatelessComputeManagerPrivate
 		return EComputeExecutionPath::None;
 	}
 
-	void GenerateGPUData(FRHICommandList& RHICmdList, FNiagaraGpuComputeDispatchInterface* ComputeInterface, TConstArrayView<FNiagaraStatelessComputeManager::FStatelessDataCache*> DataToGenerate)
+	void GenerateGPUData(FRHICommandList& RHICmdList, FNiagaraGpuComputeDispatchInterface* ComputeInterface, TConstArrayView<FNiagaraStatelessComputeManager::FStatelessDataGenerationRequest> GenerationRequests)
 	{
-		const int32 NumJobs = DataToGenerate.Num();
+		const int32 NumJobs = GenerationRequests.Num();
 
 		// Get Count Buffer
 		FNiagaraGPUInstanceCountManager& CountManager = ComputeInterface->GetGPUInstanceCounterManager();
@@ -74,15 +74,17 @@ namespace NiagaraStatelessComputeManagerPrivate
 			TransitionsBefore.Emplace(CountManager.GetInstanceCountBuffer().Buffer, FNiagaraGPUInstanceCountManager::kCountBufferDefaultState, ERHIAccess::UAVCompute);
 			TransitionsAfter.Emplace(CountManager.GetInstanceCountBuffer().Buffer, ERHIAccess::UAVCompute, FNiagaraGPUInstanceCountManager::kCountBufferDefaultState);
 
-			for (FNiagaraStatelessComputeManager::FStatelessDataCache* CacheData : DataToGenerate)
+			for (const FNiagaraStatelessComputeManager::FStatelessDataGenerationRequest& GenerationRequest : GenerationRequests)
 			{
-				const FRWBuffer& FloatBuffer = CacheData->DataBuffer->GetGPUBufferFloat();
+				FNiagaraDataBuffer* DestinationData = GenerationRequest.DestinationData;
+
+				const FRWBuffer& FloatBuffer = DestinationData->GetGPUBufferFloat();
 				if (FloatBuffer.NumBytes > 0)
 				{
 					TransitionsBefore.Emplace(FloatBuffer.Buffer, ERHIAccess::SRVMask, ERHIAccess::UAVCompute);
 					TransitionsAfter.Emplace(FloatBuffer.Buffer, ERHIAccess::UAVCompute, ERHIAccess::SRVMask);
 				}
-				const FRWBuffer& IntBuffer = CacheData->DataBuffer->GetGPUBufferInt();
+				const FRWBuffer& IntBuffer = DestinationData->GetGPUBufferInt();
 				if (IntBuffer.NumBytes > 0)
 				{
 					TransitionsBefore.Emplace(IntBuffer.Buffer, ERHIAccess::SRVMask, ERHIAccess::UAVCompute);
@@ -98,10 +100,11 @@ namespace NiagaraStatelessComputeManagerPrivate
 		RHICmdList.Transition(TransitionsBefore);
 
 		RHICmdList.BeginUAVOverlap(CountBufferUAV);
-		for (FNiagaraStatelessComputeManager::FStatelessDataCache* CacheData : DataToGenerate)
+		for (const FNiagaraStatelessComputeManager::FStatelessDataGenerationRequest& GenerationRequest : GenerationRequests)
 		{
-			const NiagaraStateless::FEmitterInstance_RT* EmitterInstance = CacheData->EmitterInstance;
-			const FNiagaraStatelessEmitterData* EmitterData = CacheData->EmitterInstance->EmitterData.Get();
+			const NiagaraStateless::FEmitterInstance_RT* EmitterInstance = GenerationRequest.EmitterInstance;
+			const FNiagaraStatelessEmitterData* EmitterData = EmitterInstance->EmitterData.Get();
+			FNiagaraDataBuffer* DestinationData = GenerationRequest.DestinationData;
 
 			// Do we need to update the parameter buffer?
 			if (EmitterInstance->BindingBufferData.IsSet())
@@ -119,11 +122,11 @@ namespace NiagaraStatelessComputeManagerPrivate
 			ShaderParameters->Common_SimulationTime = EmitterInstance->Age;
 			ShaderParameters->Common_SimulationDeltaTime = EmitterInstance->DeltaTime;
 			ShaderParameters->Common_SimulationInvDeltaTime = EmitterInstance->DeltaTime > 0.0f ? (1.0f / EmitterInstance->DeltaTime) : 0.0f;
-			ShaderParameters->Common_OutputBufferStride = CacheData->DataBuffer->GetFloatStride() / sizeof(float);
-			ShaderParameters->Common_GPUCountBufferOffset = CacheData->DataBuffer->GetGPUInstanceCountBufferOffset();
-			ShaderParameters->Common_FloatOutputBuffer = CacheData->DataBuffer->GetGPUBufferFloat().UAV;
-			//ShaderParameters->Common_HalfOutputBuffer		= CacheData->DataBuffer->GetGPUBufferHalf().UAV;
-			ShaderParameters->Common_IntOutputBuffer = CacheData->DataBuffer->GetGPUBufferInt().UAV.IsValid() ? CacheData->DataBuffer->GetGPUBufferInt().UAV.GetReference() : EmptyIntBufferUAV;
+			ShaderParameters->Common_OutputBufferStride = DestinationData->GetFloatStride() / sizeof(float);
+			ShaderParameters->Common_GPUCountBufferOffset = DestinationData->GetGPUInstanceCountBufferOffset();
+			ShaderParameters->Common_FloatOutputBuffer = DestinationData->GetGPUBufferFloat().UAV;
+			//ShaderParameters->Common_HalfOutputBuffer		= DestinationData->GetGPUBufferHalf().UAV;
+			ShaderParameters->Common_IntOutputBuffer = DestinationData->GetGPUBufferInt().UAV.IsValid() ? DestinationData->GetGPUBufferInt().UAV.GetReference() : EmptyIntBufferUAV;
 			ShaderParameters->Common_GPUCountBuffer = CountBufferUAV;
 			ShaderParameters->Common_StaticFloatBuffer = EmitterData->StaticFloatBuffer.SRV;
 			ShaderParameters->Common_ParameterBuffer = FNiagaraRenderer::GetSrvOrDefaultUInt(EmitterInstance->BindingBuffer.SRV);
@@ -131,7 +134,7 @@ namespace NiagaraStatelessComputeManagerPrivate
 			// Execute the simulation
 			TShaderRef<NiagaraStateless::FSimulationShader> ComputeShader = EmitterData->GetShader();
 			FRHIComputeShader* ComputeShaderRHI = ComputeShader.GetComputeShader();
-			const uint32 NumThreadGroups = FMath::DivideAndRoundUp<uint32>(CacheData->ActiveParticles, NiagaraStateless::FSimulationShader::ThreadGroupSize);
+			const uint32 NumThreadGroups = FMath::DivideAndRoundUp<uint32>(GenerationRequest.ActiveParticles, NiagaraStateless::FSimulationShader::ThreadGroupSize);
 
 			const FIntVector NumWrappedThreadGroups = FComputeShaderUtils::GetGroupCountWrapped(NumThreadGroups);
 			FComputeShaderUtils::Dispatch(RHICmdList, ComputeShader, EmitterData->GetShaderParametersMetadata(), *ShaderParameters, NumWrappedThreadGroups);
@@ -220,10 +223,7 @@ FNiagaraDataBuffer* FNiagaraStatelessComputeManager::GetDataBuffer(FRHICommandLi
 		CacheData->DataBuffer = new FNiagaraDataBuffer(&CacheData->DataSet);
 	}
 
-	CacheData->EmitterInstance = EmitterInstance;
-	CacheData->ActiveParticles = ActiveParticles;
-
-	CacheData->DataBuffer->AllocateGPU(RHICmdList, CacheData->ActiveParticles, ComputeInterface->GetFeatureLevel(), TEXT("StatelessSimBuffer"));
+	CacheData->DataBuffer->AllocateGPU(RHICmdList, ActiveParticles, ComputeInterface->GetFeatureLevel(), TEXT("StatelessSimBuffer"));
 
 	const EComputeExecutionPath ComputeExecutionPath = DetermineComputeExecutionPath(EmitterData, ActiveParticles);
 	switch (ComputeExecutionPath)
@@ -242,14 +242,14 @@ FNiagaraDataBuffer* FNiagaraStatelessComputeManager::GetDataBuffer(FRHICommandLi
 
 		case EComputeExecutionPath::GPU:
 		{
-			CacheData->DataBuffer->SetNumInstances(CacheData->ActiveParticles);
+			CacheData->DataBuffer->SetNumInstances(ActiveParticles);
 
 			FNiagaraGPUInstanceCountManager& CountManager = ComputeInterface->GetGPUInstanceCounterManager();
 			const uint32 CountOffset = CountManager.AllocateDeferredEntry();
 			CacheData->DataBuffer->SetGPUInstanceCountBufferOffset(CountOffset);
 			CountsToRelease.Add(CountOffset);
 
-			GPUDataToGenerate.Add(CacheData);
+			GPUGenerationRequests.Emplace(CacheData->DataBuffer, EmitterInstance, ActiveParticles);
 			break;
 		}
 
@@ -307,15 +307,10 @@ void FNiagaraStatelessComputeManager::GenerateDataBufferForDebugging(FRHICommand
 			GPUDataBuffer->SetGPUInstanceCountBufferOffset(CountIndex);
 
 			// Generate the data
-			FStatelessDataCache DataCache;
-			DataCache.DataBuffer		= GPUDataBuffer;
-			DataCache.EmitterInstance	= EmitterInstance;
-			DataCache.ActiveParticles	= ActiveParticlesEstimate;
-
-			FStatelessDataCache* DataCachePtr = &DataCache;
+			FStatelessDataGenerationRequest GenerationRequest(GPUDataBuffer, EmitterInstance, ActiveParticlesEstimate);
 
 			RHICmdList.BeginUAVOverlap();
-			GenerateGPUData(RHICmdList, ComputeInterface, MakeConstArrayView(&DataCachePtr, 1));
+			GenerateGPUData(RHICmdList, ComputeInterface, MakeConstArrayView(&GenerationRequest, 1));
 			RHICmdList.EndUAVOverlap();
 
 			// Copy to CPU data
@@ -337,7 +332,7 @@ void FNiagaraStatelessComputeManager::GenerateDataBufferForDebugging(FRHICommand
 void FNiagaraStatelessComputeManager::OnPostPreRender(FRDGBuilder& GraphBuilder)
 {
 	// Anything to process?
-	if (GPUDataToGenerate.Num() == 0)
+	if (GPUGenerationRequests.Num() == 0)
 	{
 		return;
 	}
@@ -356,15 +351,15 @@ void FNiagaraStatelessComputeManager::OnPostPreRender(FRDGBuilder& GraphBuilder)
 	AddPass(
 		GraphBuilder,
 		RDG_EVENT_NAME("FNiagaraStatelessComputeManager::OnPostPreRender"),
-		[DataToGenerate=MoveTemp(GPUDataToGenerate), ComputeInterface](FRHICommandListImmediate& RHICmdList)
+		[GPUGenerationRequests_RDG=MoveTemp(GPUGenerationRequests), ComputeInterface](FRHICommandListImmediate& RHICmdList)
 		{
 			SCOPED_DRAW_EVENT(RHICmdList, FNiagaraStatelessComputeManager_OnPostPreRender);
 
 			FNiagaraGPUInstanceCountManager& CountManager = ComputeInterface->GetGPUInstanceCounterManager();
-			NiagaraStatelessComputeManagerPrivate::GenerateGPUData(RHICmdList, ComputeInterface, DataToGenerate);
+			NiagaraStatelessComputeManagerPrivate::GenerateGPUData(RHICmdList, ComputeInterface, GPUGenerationRequests_RDG);
 		}
 	);
-	GPUDataToGenerate.Empty();
+	GPUGenerationRequests.Empty();
 }
 
 void FNiagaraStatelessComputeManager::OnPostPostRender(FRDGBuilder& GraphBuilder)
diff --git a/Engine/Plugins/FX/Niagara/Source/Niagara/Private/Stateless/NiagaraStatelessComputeManager.h b/Engine/Plugins/FX/Niagara/Source/Niagara/Private/Stateless/NiagaraStatelessComputeManager.h
--- a/Engine/Plugins/FX/Niagara/Source/Niagara/Private/Stateless/NiagaraStatelessComputeManager.h
+++ b/Engine/Plugins/FX/Niagara/Source/Niagara/Private/Stateless/NiagaraStatelessComputeManager.h
@@ -25,9 +25,21 @@ public:
 		TSharedPtr<FNiagaraDataSetCompiledData>		DataSetCompiledData;
 		FNiagaraDataSet								DataSet;
 		FNiagaraDataBufferRef						DataBuffer;
+	};
+
+	struct FStatelessDataGenerationRequest
+	{
+		FStatelessDataGenerationRequest() = default;
+		explicit FStatelessDataGenerationRequest(FNiagaraDataBuffer* InDestinationData, const NiagaraStateless::FEmitterInstance_RT* InEmitterInstance, uint32 InActiveParticles)
+			: DestinationData(InDestinationData)
+			, EmitterInstance(InEmitterInstance)
+			, ActiveParticles(InActiveParticles)
+		{
+		}
 
-		const NiagaraStateless::FEmitterInstance_RT* EmitterInstance = nullptr;
-		uint32										ActiveParticles = 0;
+		FNiagaraDataBufferRef							DestinationData;
+		const NiagaraStateless::FEmitterInstance_RT*	EmitterInstance = nullptr;
+		uint32											ActiveParticles = 0;
 	};
 
 public:
@@ -56,7 +68,7 @@ private:
 	TArray<TUniquePtr<FStatelessDataCache>>				FreeData;
 	TArray<uint32>										CountsToRelease;
 
-	TArray<FStatelessDataCache*>						GPUDataToGenerate;
+	TArray<FStatelessDataGenerationRequest>				GPUGenerationRequests;
 
 	UE::FMutex											GetDataBufferGuard;
 };
