- Remove range based random calls for CPU sims as they potentially hide binding issues
- Now we can read the shader parameters we should do that instead, it's less ops and will have updated with the latest parameter store data

#rnx
#tests Run new EngineTests which uses bindings

diff --git a/Engine/Plugins/FX/Niagara/Source/Niagara/Internal/Stateless/Modules/NiagaraStatelessModule_InitializeParticle.h b/Engine/Plugins/FX/Niagara/Source/Niagara/Internal/Stateless/Modules/NiagaraStatelessModule_InitializeParticle.h
--- a/Engine/Plugins/FX/Niagara/Source/Niagara/Internal/Stateless/Modules/NiagaraStatelessModule_InitializeParticle.h
+++ b/Engine/Plugins/FX/Niagara/Source/Niagara/Internal/Stateless/Modules/NiagaraStatelessModule_InitializeParticle.h
@@ -136,6 +136,7 @@ public:
 		using namespace NiagaraStateless;
 
 		const FModuleBuiltData* ModuleBuiltData = ParticleSimulationContext.ReadBuiltData<FModuleBuiltData>();
+		const FParameters* ShaderParameters = ParticleSimulationContext.ReadParameterNestedStruct<FParameters>();
 
 		const bool bUniformSpriteSize = (ModuleBuiltData->ModuleFlags & EInitializeParticleModuleFlag_UniformSpriteSize) != 0;
 		const bool bUniformMeshScale  = (ModuleBuiltData->ModuleFlags & EInitializeParticleModuleFlag_UniformMeshScale) != 0;
@@ -145,11 +146,11 @@ public:
 			const FStatelessDistributionSampler<FVector3f> PositionSampler(ParticleSimulationContext, ModuleBuiltData->InitialPosition, i, 0);
 
 			const FVector3f		Position	= PositionSampler.GetValue(ParticleSimulationContext, 0.0f);
-			const FLinearColor	Color		= ParticleSimulationContext.RandomScaleBiasFloat(i, 1 ,ModuleBuiltData->ColorRange);
-			const float			RibbonWidth	= ParticleSimulationContext.RandomScaleBiasFloat(i, 2, ModuleBuiltData->RibbonWidthRange);
-			const FVector2f		SpriteSize	= ParticleSimulationContext.RandomScaleBiasFloat(i, 3, ModuleBuiltData->SpriteSizeRange, bUniformSpriteSize);
-			const float			SpriteRot	= ParticleSimulationContext.RandomScaleBiasFloat(i, 4, ModuleBuiltData->SpriteRotationRange);
-			const FVector3f		Scale		= ParticleSimulationContext.RandomScaleBiasFloat(i, 5, ModuleBuiltData->MeshScaleRange, bUniformMeshScale);
+			const FLinearColor	Color		= ParticleSimulationContext.RandomScaleBiasFloat(i, 1, ShaderParameters->InitializeParticle_ColorScale, ShaderParameters->InitializeParticle_ColorBias);
+			const float			RibbonWidth	= ParticleSimulationContext.RandomScaleBiasFloat(i, 2, ShaderParameters->InitializeParticle_RibbonWidthScale, ShaderParameters->InitializeParticle_RibbonWidthBias);
+			const FVector2f		SpriteSize	= ParticleSimulationContext.RandomScaleBiasFloat(i, 3, ShaderParameters->InitializeParticle_SpriteSizeScale, ShaderParameters->InitializeParticle_SpriteSizeBias, bUniformSpriteSize);
+			const float			SpriteRot	= ParticleSimulationContext.RandomScaleBiasFloat(i, 4, ShaderParameters->InitializeParticle_SpriteRotationScale, ShaderParameters->InitializeParticle_SpriteRotationBias);
+			const FVector3f		Scale		= ParticleSimulationContext.RandomScaleBiasFloat(i, 5, ShaderParameters->InitializeParticle_MeshScaleScale, ShaderParameters->InitializeParticle_MeshScaleBias, bUniformMeshScale);
 
 			ParticleSimulationContext.WriteParticleVariable(ModuleBuiltData->PositionVariableOffset,				i, Position);
 			ParticleSimulationContext.WriteParticleVariable(ModuleBuiltData->ColorVariableOffset,					i, Color);
diff --git a/Engine/Plugins/FX/Niagara/Source/Niagara/Internal/Stateless/Modules/NiagaraStatelessModule_MeshRotationRate.h b/Engine/Plugins/FX/Niagara/Source/Niagara/Internal/Stateless/Modules/NiagaraStatelessModule_MeshRotationRate.h
--- a/Engine/Plugins/FX/Niagara/Source/Niagara/Internal/Stateless/Modules/NiagaraStatelessModule_MeshRotationRate.h
+++ b/Engine/Plugins/FX/Niagara/Source/Niagara/Internal/Stateless/Modules/NiagaraStatelessModule_MeshRotationRate.h
@@ -70,12 +70,14 @@ public:
 		using namespace NiagaraStateless;
 
 		const FModuleBuiltData* ModuleBuiltData = ParticleSimulationContext.ReadBuiltData<FModuleBuiltData>();
+		const FParameters* ShaderParameters = ParticleSimulationContext.ReadParameterNestedStruct<FParameters>();
+
 		const float* AgeData = ParticleSimulationContext.GetParticleAge();
 		const float* PreviousAgeData = ParticleSimulationContext.GetParticlePreviousAge();
 
 		for (uint32 i = 0; i < ParticleSimulationContext.GetNumInstances(); ++i)
 		{
-			const FVector3f RotationRate = ParticleSimulationContext.RandomScaleBiasFloat(i, 0, ModuleBuiltData->RotationRange);
+			const FVector3f RotationRate = ParticleSimulationContext.RandomScaleBiasFloat(i, 0, ShaderParameters->MeshRotationRate_Scale, ShaderParameters->MeshRotationRate_Bias);
 			const float Age = AgeData[i];
 			const float PreviousAge = PreviousAgeData[i];
 
diff --git a/Engine/Plugins/FX/Niagara/Source/Niagara/Internal/Stateless/Modules/NiagaraStatelessModule_SpriteRotationRate.h b/Engine/Plugins/FX/Niagara/Source/Niagara/Internal/Stateless/Modules/NiagaraStatelessModule_SpriteRotationRate.h
--- a/Engine/Plugins/FX/Niagara/Source/Niagara/Internal/Stateless/Modules/NiagaraStatelessModule_SpriteRotationRate.h
+++ b/Engine/Plugins/FX/Niagara/Source/Niagara/Internal/Stateless/Modules/NiagaraStatelessModule_SpriteRotationRate.h
@@ -68,12 +68,14 @@ public:
 		using namespace NiagaraStateless;
 
 		const FModuleBuiltData* ModuleBuiltData = ParticleSimulationContext.ReadBuiltData<FModuleBuiltData>();
+		const FParameters* ShaderParameters = ParticleSimulationContext.ReadParameterNestedStruct<FParameters>();
+
 		const float* AgeData = ParticleSimulationContext.GetParticleAge();
 		const float* PreviousAgeData = ParticleSimulationContext.GetParticlePreviousAge();
 
 		for (uint32 i = 0; i < ParticleSimulationContext.GetNumInstances(); ++i)
 		{
-			const float RotationRate = ParticleSimulationContext.RandomScaleBiasFloat(i, 0, ModuleBuiltData->RotationRange);
+			const float RotationRate = ParticleSimulationContext.RandomScaleBiasFloat(i, 0, ShaderParameters->SpriteRotationRate_Scale, ShaderParameters->SpriteRotationRate_Bias);
 
 			float SpriteRotation = ParticleSimulationContext.ReadParticleVariable(ModuleBuiltData->SpriteRotationVariableOffset, i, 0.0f);
 			float PreviousSpriteRotation = ParticleSimulationContext.ReadParticleVariable(ModuleBuiltData->PreviousSpriteRotationVariableOffset, i, 0.0f);
diff --git a/Engine/Plugins/FX/Niagara/Source/Niagara/Internal/Stateless/NiagaraStatelessParticleSimContext.h b/Engine/Plugins/FX/Niagara/Source/Niagara/Internal/Stateless/NiagaraStatelessParticleSimContext.h
--- a/Engine/Plugins/FX/Niagara/Source/Niagara/Internal/Stateless/NiagaraStatelessParticleSimContext.h
+++ b/Engine/Plugins/FX/Niagara/Source/Niagara/Internal/Stateless/NiagaraStatelessParticleSimContext.h
@@ -145,18 +145,13 @@ public:
 	FVector2f		RandomScaleBiasFloat(uint32 iInstance, uint32 RandomSeedOffset, const FVector2f& Scale, const FVector2f& Bias) const { return Bias + (RandomFloat2(iInstance, RandomSeedOffset) * Scale); }
 	FVector3f		RandomScaleBiasFloat(uint32 iInstance, uint32 RandomSeedOffset, const FVector3f& Scale, const FVector3f& Bias) const { return Bias + (RandomFloat3(iInstance, RandomSeedOffset) * Scale); }
 	FVector4f		RandomScaleBiasFloat(uint32 iInstance, uint32 RandomSeedOffset, const FVector4f& Scale, const FVector4f& Bias) const { return Bias + (RandomFloat4(iInstance, RandomSeedOffset) * Scale); }
+	FLinearColor	RandomScaleBiasFloat(uint32 iInstance, uint32 RandomSeedOffset, const FLinearColor& Scale, const FLinearColor& Bias) const { return Bias + (RandomFloat4(iInstance, RandomSeedOffset) * Scale); }
 
-	float			RandomScaleBiasFloat(uint32 iInstance, uint32 RandomSeedOffset, const FNiagaraStatelessRangeFloat& Range) const { return Range.Min + (RandomFloat(iInstance, RandomSeedOffset) * Range.GetScale()); }
-	FVector2f		RandomScaleBiasFloat(uint32 iInstance, uint32 RandomSeedOffset, const FNiagaraStatelessRangeVector2& Range) const { return Range.Min + (RandomFloat2(iInstance, RandomSeedOffset) * Range.GetScale()); }
-	FVector3f		RandomScaleBiasFloat(uint32 iInstance, uint32 RandomSeedOffset, const FNiagaraStatelessRangeVector3& Range) const { return Range.Min + (RandomFloat3(iInstance, RandomSeedOffset) * Range.GetScale()); }
-	FVector4f		RandomScaleBiasFloat(uint32 iInstance, uint32 RandomSeedOffset, const FNiagaraStatelessRangeVector4& Range) const { return Range.Min + (RandomFloat4(iInstance, RandomSeedOffset) * Range.GetScale()); }
-	FLinearColor	RandomScaleBiasFloat(uint32 iInstance, uint32 RandomSeedOffset, const FNiagaraStatelessRangeColor& Range) const { return Range.Min + (RandomFloat4(iInstance, RandomSeedOffset) * Range.GetScale()); }
-
-	float			RandomScaleBiasFloat(uint32 iInstance, uint32 RandomSeedOffset, const FNiagaraStatelessRangeFloat& Range, bool bUniform) const { return RandomScaleBiasFloat(iInstance, RandomSeedOffset, Range); }
-	FVector2f		RandomScaleBiasFloat(uint32 iInstance, uint32 RandomSeedOffset, const FNiagaraStatelessRangeVector2& Range, bool bUniform) const { const FVector2f Random = RandomScaleBiasFloat(iInstance, RandomSeedOffset, Range); return bUniform ? FVector2f(Random.X, Random.X) : Random; }
-	FVector3f		RandomScaleBiasFloat(uint32 iInstance, uint32 RandomSeedOffset, const FNiagaraStatelessRangeVector3& Range, bool bUniform) const { const FVector3f Random = RandomScaleBiasFloat(iInstance, RandomSeedOffset, Range); return bUniform ? FVector3f(Random.X, Random.X, Random.X) : Random; }
-	FVector4f		RandomScaleBiasFloat(uint32 iInstance, uint32 RandomSeedOffset, const FNiagaraStatelessRangeVector4& Range, bool bUniform) const { const FVector4f Random = RandomScaleBiasFloat(iInstance, RandomSeedOffset, Range); return bUniform ? FVector4f(Random.X, Random.X, Random.X, Random.X) : Random; }
-	FLinearColor	RandomScaleBiasFloat(uint32 iInstance, uint32 RandomSeedOffset, const FNiagaraStatelessRangeColor& Range, bool bUniform) const { const FLinearColor Random = RandomScaleBiasFloat(iInstance, RandomSeedOffset, Range); return bUniform ? FLinearColor(Random.R, Random.R, Random.R, Random.R) : Random; }
+	float			RandomScaleBiasFloat(uint32 iInstance, uint32 RandomSeedOffset, const float Scale, const float Bias, bool bUniform) const { return RandomScaleBiasFloat(iInstance, RandomSeedOffset, Scale, Bias); }
+	FVector2f		RandomScaleBiasFloat(uint32 iInstance, uint32 RandomSeedOffset, const FVector2f& Scale, const FVector2f& Bias, bool bUniform) const { const FVector2f Random = RandomScaleBiasFloat(iInstance, RandomSeedOffset, Scale, Bias); return bUniform ? FVector2f(Random.X, Random.X) : Random; }
+	FVector3f		RandomScaleBiasFloat(uint32 iInstance, uint32 RandomSeedOffset, const FVector3f& Scale, const FVector3f& Bias, bool bUniform) const { const FVector3f Random = RandomScaleBiasFloat(iInstance, RandomSeedOffset, Scale, Bias); return bUniform ? FVector3f(Random.X, Random.X, Random.X) : Random; }
+	FVector4f		RandomScaleBiasFloat(uint32 iInstance, uint32 RandomSeedOffset, const FVector4f& Scale, const FVector4f& Bias, bool bUniform) const { const FVector4f Random = RandomScaleBiasFloat(iInstance, RandomSeedOffset, Scale, Bias); return bUniform ? FVector4f(Random.X, Random.X, Random.X) : Random; }
+	FLinearColor	RandomScaleBiasFloat(uint32 iInstance, uint32 RandomSeedOffset, const FLinearColor& Scale, const FLinearColor& Bias, bool bUniform) const { const FLinearColor Random = RandomScaleBiasFloat(iInstance, RandomSeedOffset, Scale, Bias); return bUniform ? FLinearColor(Random.R, Random.R, Random.R, Random.R) : Random; }
 
 	FVector2f		RandomUnitFloat2(uint32 iInstance, uint32 RandomSeedOffset) const { return SafeNormalize(RandomFloat2(iInstance, RandomSeedOffset) - 0.5f); }
 	FVector3f		RandomUnitFloat3(uint32 iInstance, uint32 RandomSeedOffset) const { return SafeNormalize(RandomFloat3(iInstance, RandomSeedOffset) - 0.5f); }
diff --git a/Engine/Plugins/FX/Niagara/Source/Niagara/Private/Stateless/NiagaraStatelessParticleSimContext.cpp b/Engine/Plugins/FX/Niagara/Source/Niagara/Private/Stateless/NiagaraStatelessParticleSimContext.cpp
--- a/Engine/Plugins/FX/Niagara/Source/Niagara/Private/Stateless/NiagaraStatelessParticleSimContext.cpp
+++ b/Engine/Plugins/FX/Niagara/Source/Niagara/Private/Stateless/NiagaraStatelessParticleSimContext.cpp
@@ -197,7 +197,7 @@ namespace NiagaraStateless
 			// Write unique index as it's needed for the random operation
 			GetParticleUniqueIndex()[NumInstances] = Particle_UniqueIndex;
 
-			float Particle_Lifetime = RandomScaleBiasFloat(NumInstances, 0, EmitterData->LifetimeRange);
+			float Particle_Lifetime = RandomScaleBiasFloat(NumInstances, 0, EmitterData->LifetimeRange.GetScale(), EmitterData->LifetimeRange.Min);
 			if (Particle_Lifetime <= 0.0f || Particle_Age >= Particle_Lifetime)
 			{
 				continue;
