[Viewport Toolbar] Niagara Emitter editor: add Preview menu

- Adds a FPreviewProfileController to SNiagaraSystemViewport to handle access to Preview Profiles (similar to what was done for Static Mesh Editor)
- Adds Preview Profile Submenu
- Removes unnecessary SNiagaraSystemViewport argument from CreateSettingsSubmenu and CreateShowSubmenu in UE::NiagaraSystemEditor. This also fixes potential issues when opening/closing multiple Niagara Editors

diff --git a/Engine/Plugins/FX/Niagara/Source/NiagaraEditor/Private/Widgets/NiagaraSystemEditorViewportToolbarSections.cpp b/Engine/Plugins/FX/Niagara/Source/NiagaraEditor/Private/Widgets/NiagaraSystemEditorViewportToolbarSections.cpp
--- a/Engine/Plugins/FX/Niagara/Source/NiagaraEditor/Private/Widgets/NiagaraSystemEditorViewportToolbarSections.cpp
+++ b/Engine/Plugins/FX/Niagara/Source/NiagaraEditor/Private/Widgets/NiagaraSystemEditorViewportToolbarSections.cpp
@@ -6,6 +6,7 @@
 #include "NiagaraEditorCommands.h"
 #include "SNiagaraSystemViewport.h"
 #include "ToolMenu.h"
+#include "ViewportToolbar/UnrealEdViewportToolbarContext.h"
 #include "Widgets/Input/SSpinBox.h"
 
 #define LOCTEXT_NAMESPACE "NiagaraSystemEditorViewportToolbarSections"
@@ -68,17 +69,25 @@ TSharedRef<SWidget> UE::NiagaraSystemEditor::CreateShowMenuWidget(
 	return ShowMenuBuilder.MakeWidget();
 }
 
-FToolMenuEntry UE::NiagaraSystemEditor::CreateShowSubmenu(const TWeakPtr<SNiagaraSystemViewport>& InNiagaraSystemEditorViewport)
+FToolMenuEntry UE::NiagaraSystemEditor::CreateShowSubmenu()
 {
 	return FToolMenuEntry::InitSubMenu(
 		"Show",
 		LOCTEXT("ShowSubmenuLabel", "Show"),
 		LOCTEXT("ShowSubmenuTooltip", "Show options"),
 		FNewToolMenuDelegate::CreateLambda(
-			[InNiagaraSystemEditorViewport](UToolMenu* Submenu) -> void
+			[](UToolMenu* Submenu) -> void
 			{
-				TSharedPtr<SNiagaraSystemViewport> Viewport = InNiagaraSystemEditorViewport.Pin();
-				if (!Viewport)
+				UUnrealEdViewportToolbarContext* ViewportToolbarContext =
+					Submenu->FindContext<UUnrealEdViewportToolbarContext>();
+				if (!ViewportToolbarContext)
+				{
+					return;
+				}
+
+				TSharedPtr<SNiagaraSystemViewport> NiagaraSystemViewport =
+					StaticCastSharedPtr<SNiagaraSystemViewport>(ViewportToolbarContext->Viewport.Pin());
+				if (!NiagaraSystemViewport)
 				{
 					return;
 				}
@@ -86,22 +95,24 @@ FToolMenuEntry UE::NiagaraSystemEditor::CreateShowSubmenu(const TWeakPtr<SNiagar
 				FToolMenuSection& UnnamedSection = Submenu->FindOrAddSection(NAME_None);
 
 				UnnamedSection.AddEntry(FToolMenuEntry::InitWidget(
-					"ShowMenuItems", UE::NiagaraSystemEditor::CreateShowMenuWidget(Viewport.ToSharedRef(), true), FText(), true
+					"ShowMenuItems",
+					UE::NiagaraSystemEditor::CreateShowMenuWidget(NiagaraSystemViewport.ToSharedRef(), true),
+					FText(),
+					true
 				));
 			}
 		)
 	);
 }
 
-FToolMenuEntry UE::NiagaraSystemEditor::CreateSettingsSubmenu(const TWeakPtr<SNiagaraSystemViewport>& InNiagaraSystemEditorViewport
-)
+FToolMenuEntry UE::NiagaraSystemEditor::CreateSettingsSubmenu()
 {
 	return FToolMenuEntry::InitSubMenu(
 		"Settings",
 		LOCTEXT("SettingsSubmenuLabel", "Settings"),
 		LOCTEXT("SettingsSubmenuTooltip", "Settings options"),
 		FNewToolMenuDelegate::CreateLambda(
-			[InNiagaraSystemEditorViewport](UToolMenu* Submenu) -> void
+			[](UToolMenu* Submenu) -> void
 			{
 				FToolMenuSection& ViewportControlsSection =
 					Submenu->FindOrAddSection("ViewportControls", LOCTEXT("ViewportControlsLabel", "Viewport Controls"));
@@ -111,10 +122,18 @@ FToolMenuEntry UE::NiagaraSystemEditor::CreateSettingsSubmenu(const TWeakPtr<SNi
 					LOCTEXT("MotionOptionsSubMenu", "Motion Options"),
 					LOCTEXT("MotionOptionsSubMenu_ToolTip", "Set Motion Options for the Niagara Component"),
 					FNewToolMenuDelegate::CreateLambda(
-						[InNiagaraSystemEditorViewport](UToolMenu* InMenu)
+						[](UToolMenu* InMenu)
 						{
-							TSharedPtr<SNiagaraSystemViewport> Viewport = InNiagaraSystemEditorViewport.Pin();
-							if (!Viewport)
+							UUnrealEdViewportToolbarContext* ViewportToolbarContext =
+								InMenu->FindContext<UUnrealEdViewportToolbarContext>();
+							if (!ViewportToolbarContext)
+							{
+								return;
+							}
+
+							TSharedPtr<SNiagaraSystemViewport> NiagaraSystemViewport =
+								StaticCastSharedPtr<SNiagaraSystemViewport>(ViewportToolbarContext->Viewport.Pin());
+							if (!NiagaraSystemViewport)
 							{
 								return;
 							}
@@ -123,7 +142,7 @@ FToolMenuEntry UE::NiagaraSystemEditor::CreateSettingsSubmenu(const TWeakPtr<SNi
 								"MotionOptions",
 								FToolMenuEntry::InitWidget(
 									"MotionOptions",
-									UE::NiagaraSystemEditor::CreateMotionMenuWidget(Viewport.ToSharedRef()),
+									UE::NiagaraSystemEditor::CreateMotionMenuWidget(NiagaraSystemViewport.ToSharedRef()),
 									FText(),
 									true
 								)
diff --git a/Engine/Plugins/FX/Niagara/Source/NiagaraEditor/Private/Widgets/NiagaraSystemEditorViewportToolbarSections.h b/Engine/Plugins/FX/Niagara/Source/NiagaraEditor/Private/Widgets/NiagaraSystemEditorViewportToolbarSections.h
--- a/Engine/Plugins/FX/Niagara/Source/NiagaraEditor/Private/Widgets/NiagaraSystemEditorViewportToolbarSections.h
+++ b/Engine/Plugins/FX/Niagara/Source/NiagaraEditor/Private/Widgets/NiagaraSystemEditorViewportToolbarSections.h
@@ -10,8 +10,8 @@ class SNiagaraSystemViewport;
 namespace UE::NiagaraSystemEditor
 {
 TSharedRef<SWidget> CreateShowMenuWidget(const TSharedRef<SNiagaraSystemViewport>& InNiagaraSystemEditorViewport, bool bInShowViewportStatsToggle = true);
-FToolMenuEntry CreateShowSubmenu(const TWeakPtr<SNiagaraSystemViewport>& InNiagaraSystemEditorViewport);
+FToolMenuEntry CreateShowSubmenu();
 
 TSharedRef<SWidget> CreateMotionMenuWidget(const TSharedRef<SNiagaraSystemViewport>& InNiagaraSystemEditorViewport);
-FToolMenuEntry CreateSettingsSubmenu(const TWeakPtr<SNiagaraSystemViewport>& InNiagaraSystemEditorViewport);
+FToolMenuEntry CreateSettingsSubmenu();
 }
diff --git a/Engine/Plugins/FX/Niagara/Source/NiagaraEditor/Private/Widgets/SNiagaraSystemViewport.cpp b/Engine/Plugins/FX/Niagara/Source/NiagaraEditor/Private/Widgets/SNiagaraSystemViewport.cpp
--- a/Engine/Plugins/FX/Niagara/Source/NiagaraEditor/Private/Widgets/SNiagaraSystemViewport.cpp
+++ b/Engine/Plugins/FX/Niagara/Source/NiagaraEditor/Private/Widgets/SNiagaraSystemViewport.cpp
@@ -24,6 +24,7 @@
 #include "NiagaraSystemGpuComputeProxy.h"
 #include "NiagaraSystemEditorData.h"
 #include "NiagaraSystemEditorViewportToolbarSections.h"
+#include "PreviewProfileController.h"
 #include "SNiagaraSystemViewportToolBar.h"
 #include "UnrealEdGlobals.h"
 #include "Editor/EditorEngine.h"
@@ -1261,7 +1262,7 @@ TSharedPtr<SWidget> SNiagaraSystemViewport::BuildViewportToolbar()
 
 			// Add the Show submenu.
 			{
-				FToolMenuEntry ShowSubmenu = UE::NiagaraSystemEditor::CreateShowSubmenu(SharedThis(this));
+				FToolMenuEntry ShowSubmenu = UE::NiagaraSystemEditor::CreateShowSubmenu();
 				ShowSubmenu.InsertPosition.Position = EToolMenuInsertType::First;
 				RightSection.AddEntry(ShowSubmenu);
 			}
@@ -1273,9 +1274,18 @@ TSharedPtr<SWidget> SNiagaraSystemViewport::BuildViewportToolbar()
 				RightSection.AddEntry(PerformanceAndScalabilitySubmenu);
 			}
 
+			// Add the "Preview Profile" sub menu.
+			{
+				PreviewProfileController = MakeShared<FPreviewProfileController>();
+				FToolMenuEntry PreviewProfileSubmenu =
+					UE::UnrealEd::CreateViewportToolbarAssetViewerProfileSubmenu(PreviewProfileController);
+				PreviewProfileSubmenu.InsertPosition.Position = EToolMenuInsertType::Last;
+				RightSection.AddEntry(PreviewProfileSubmenu);
+			}
+
 			// Add the Settings submenu.
 			{
-				FToolMenuEntry SettingsSubmenu = UE::NiagaraSystemEditor::CreateSettingsSubmenu(SharedThis(this));
+				FToolMenuEntry SettingsSubmenu = UE::NiagaraSystemEditor::CreateSettingsSubmenu();
 				SettingsSubmenu.InsertPosition.Position = EToolMenuInsertType::First;
 				RightSection.AddEntry(SettingsSubmenu);
 			}
diff --git a/Engine/Plugins/FX/Niagara/Source/NiagaraEditor/Private/Widgets/SNiagaraSystemViewport.h b/Engine/Plugins/FX/Niagara/Source/NiagaraEditor/Private/Widgets/SNiagaraSystemViewport.h
--- a/Engine/Plugins/FX/Niagara/Source/NiagaraEditor/Private/Widgets/SNiagaraSystemViewport.h
+++ b/Engine/Plugins/FX/Niagara/Source/NiagaraEditor/Private/Widgets/SNiagaraSystemViewport.h
@@ -16,6 +16,7 @@
 #include "NiagaraPerfBaseline.h"
 
 class FNiagaraSystemViewModel;
+class FPreviewProfileController;
 class UNiagaraComponent;
 class FNiagaraSystemEditorViewportClient;
 class FNiagaraSystemInstance;
@@ -142,6 +143,9 @@ private:
 	/** Level viewport client */
 	TSharedPtr<class FNiagaraSystemViewportClient> SystemViewportClient;
 
+	/** Access to Preview Profiles */
+	TSharedPtr<FPreviewProfileController> PreviewProfileController;
+
 	uint32 DrawFlags = 0;
 
 	FOnThumbnailCaptured OnThumbnailCapturedDelegate;
