- Add package redirect for Empty -> Minimal

#rnx

diff --git a/Engine/Plugins/FX/Niagara/Config/BaseNiagara.ini b/Engine/Plugins/FX/Niagara/Config/BaseNiagara.ini
--- a/Engine/Plugins/FX/Niagara/Config/BaseNiagara.ini
+++ b/Engine/Plugins/FX/Niagara/Config/BaseNiagara.ini
@@ -14,6 +14,8 @@
 +EnumRedirects=(OldName="/Niagara/Enums/ENiagaraExecutionStateManagement.ENiagaraExecutionStateManagement",NewName="/Script/Niagara.ENiagaraExecutionStateManagement",OverrideClassName="/Script/CoreUObject.Enum",ValueChanges=(("NewEnumerator0","Awaken"),("NewEnumerator1","SleepAndLetParticlesFinish"),("NewEnumerator2","SleepAndClearParticles"),("NewEnumerator3","KillImmediately"),("NewEnumerator4","KillAfterParticlesFinish")))
 
 +PackageRedirects=(OldName="/Niagara/Enums/ENiagaraExecutionStateManagement", NewName="/Script/Niagara")
++PackageRedirects=(OldName="/Niagara/DefaultAssets/Templates/Emitters/Empty", NewName="/Niagara/DefaultAssets/Templates/Emitters/Minimal")
++PackageRedirects=(OldName="/Niagara/DefaultAssets/Templates/Emitters/Empty.Empty", NewName="/Niagara/DefaultAssets/Templates/Emitters/Minimal.Minimal")
 
 +StructRedirects=(OldName="NiagaraScalabilitySettings",NewName="/Script/Niagara.NiagaraSystemScalabilitySettings")
 +StructRedirects=(OldName="NiagaraScalabilityOverrides",NewName="/Script/Niagara.NiagaraSystemScalabilityOverride")
