- Demote check to warning to avoid taking down the editor as we handle the case where the overridden pin is no longer valid
- The check is infrequently triggering and is likely due to a refresh issue + compilation combination, but it's very hard to get a consistent repro

#rnx
#tests Binding spawn time on a burst to a user parameter, which was the initial repro

diff --git a/Engine/Plugins/FX/Niagara/Source/NiagaraEditor/Private/NiagaraStackFunctionInputBinder.cpp b/Engine/Plugins/FX/Niagara/Source/NiagaraEditor/Private/NiagaraStackFunctionInputBinder.cpp
--- a/Engine/Plugins/FX/Niagara/Source/NiagaraEditor/Private/NiagaraStackFunctionInputBinder.cpp
+++ b/Engine/Plugins/FX/Niagara/Source/NiagaraEditor/Private/NiagaraStackFunctionInputBinder.cpp
@@ -235,7 +235,10 @@ TArray<uint8> FNiagaraStackFunctionInputBinder::GetData() const
 		ValidScriptGraphChangeIdForDefaultPin != GetChangeIdFromFunctionScript())
 	{
 		RefreshGraphPins();
-		checkf(OverridePin == nullptr || OverridePin->LinkedTo.Num() == 0, TEXT("Binding no longer valid because the function was overriden"));
+		if (!OverridePin || OverridePin->LinkedTo.Num() == 0)
+		{
+			UE_LOG(LogNiagaraEditor, Warning, TEXT("FNiagaraStackFunctionInputBinder - OverridePin is no longer valid because the function was overriden"));
+		}
 	}
 
 	// We copy the data to an array here rather than return a pointer to the data because values stored in pins are 
