Fixed wrong parameter context when in a scratch pad. It used system as parameter context instead of script, which impacted UI options

diff --git a/Engine/Plugins/FX/Niagara/Source/NiagaraEditor/Private/ViewModels/NiagaraParameterPanelViewModel.cpp b/Engine/Plugins/FX/Niagara/Source/NiagaraEditor/Private/ViewModels/NiagaraParameterPanelViewModel.cpp
--- a/Engine/Plugins/FX/Niagara/Source/NiagaraEditor/Private/ViewModels/NiagaraParameterPanelViewModel.cpp
+++ b/Engine/Plugins/FX/Niagara/Source/NiagaraEditor/Private/ViewModels/NiagaraParameterPanelViewModel.cpp
@@ -1489,6 +1489,12 @@ TSharedPtr<SWidget> FNiagaraSystemToolkitParameterPanelViewModel::CreateContextM
 
 FNiagaraParameterUtilities::EParameterContext FNiagaraSystemToolkitParameterPanelViewModel::GetParameterContext() const
 {
+	TSharedPtr<class FNiagaraScratchPadScriptViewModel> ScratchPadScriptVM = SystemViewModel->GetDocumentViewModel()->GetActiveScratchPadViewModelIfSet();
+	if (ScratchPadScriptVM.IsValid())
+	{
+		return FNiagaraParameterUtilities::EParameterContext::Script;
+	}
+	
 	return FNiagaraParameterUtilities::EParameterContext::System;
 }
 
