Fix for broken links between GPU multi ribbons when a single ribbon was down to a single particle (and thus no valid segments)

diff --git a/Engine/Plugins/FX/Niagara/Shaders/Private/Ribbons/NiagaraRibbonAggregationApply.usf b/Engine/Plugins/FX/Niagara/Shaders/Private/Ribbons/NiagaraRibbonAggregationApply.usf
new file mode 100644
--- /dev/null
+++ b/Engine/Plugins/FX/Niagara/Shaders/Private/Ribbons/NiagaraRibbonAggregationApply.usf
@@ -0,0 +1,58 @@
+// Copyright Epic Games, Inc. All Rights Reserved.
+
+#include "/Engine/Private/Common.ush"
+#include "NiagaraRibbonCommon.ush"
+
+Buffer<uint> SortedIndices;
+
+StructuredBuffer<FRibbonAccumulationValues> InputAccumulation;
+
+RWBuffer<float> OutputTangentsAndDistances;
+
+#if HAS_RIBBON_ID
+RWBuffer<uint> OutputMultiRibbonIndices;
+#endif
+
+#if RIBBONS_WANTS_AUTOMATIC_TESSELLATION
+RWBuffer<float> OutputTessellationStats;
+#endif
+
+RWBuffer<uint> OutputSegments;
+
+[numthreads(THREADGROUP_SIZE, 1, 1)]
+void RibbonAggregationApply(uint3 DispatchThreadId : SV_DispatchThreadID)
+{
+	const uint TotalNumParticles = GetTotalNumParticles();
+	const uint RawParticleID = DispatchThreadId.x;
+
+	if (RawParticleID < TotalNumParticles)
+	{
+		OutputTangentsAndDistances[4 * RawParticleID + 3] = InputAccumulation[RawParticleID].RibbonDistance;
+
+		// deal with the masking of invalid segments that has been inserted into the OutputSegments during initialization
+		if (OutputSegments[RawParticleID] != INDEX_NONE)
+		{
+			OutputSegments[RawParticleID] = InputAccumulation[RawParticleID].SegmentCount - 1;
+		}
+
+#if HAS_RIBBON_ID
+		OutputMultiRibbonIndices[RawParticleID] = InputAccumulation[RawParticleID].MultiRibbonCount - 1;
+#endif
+
+#if RIBBONS_WANTS_AUTOMATIC_TESSELLATION
+#if RIBBON_HAS_TWIST
+		const uint TessellationBaseOffset = RawParticleID * 5;
+#else
+		const uint TessellationBaseOffset = RawParticleID * 3;
+#endif // RIBBON_HAS_TWIST
+
+		OutputTessellationStats[TessellationBaseOffset + 0] = InputAccumulation[RawParticleID].TessTotalLength;
+		OutputTessellationStats[TessellationBaseOffset + 1] = InputAccumulation[RawParticleID].TessAvgSegmentLength;
+		OutputTessellationStats[TessellationBaseOffset + 2] = InputAccumulation[RawParticleID].TessAvgSegmentAngle;
+#if RIBBON_HAS_TWIST
+		OutputTessellationStats[TessellationBaseOffset + 3] = InputAccumulation[RawParticleID].TessTwistAvgAngle;
+		OutputTessellationStats[TessellationBaseOffset + 4] = InputAccumulation[RawParticleID].TessTwistAvgWidth;
+#endif // RIBBON_HAS_TWIST
+#endif // RIBBONS_WANTS_AUTOMATIC_TESSELLATION
+	}
+}
diff --git a/Engine/Plugins/FX/Niagara/Shaders/Private/Ribbons/NiagaraRibbonAggregationStep.usf b/Engine/Plugins/FX/Niagara/Shaders/Private/Ribbons/NiagaraRibbonAggregationStep.usf
new file mode 100644
--- /dev/null
+++ b/Engine/Plugins/FX/Niagara/Shaders/Private/Ribbons/NiagaraRibbonAggregationStep.usf
@@ -0,0 +1,64 @@
+// Copyright Epic Games, Inc. All Rights Reserved.
+
+#include "/Engine/Private/Common.ush"
+#include "NiagaraRibbonCommon.ush"
+
+Buffer<uint> SortedIndices;
+
+StructuredBuffer<FRibbonAccumulationValues> InputAccumulation;
+RWStructuredBuffer<FRibbonAccumulationValues> OutputAccumulation;
+
+uint PrefixScanStride;
+
+[numthreads(THREADGROUP_SIZE, 1, 1)]
+void RibbonAggregationStep(uint3 DispatchThreadId : SV_DispatchThreadID)
+{
+	const uint TotalNumParticles = GetTotalNumParticles();
+	const uint RawParticleID = DispatchThreadId.x;
+	
+	if (RawParticleID < TotalNumParticles)
+	{
+		if (RawParticleID >= PrefixScanStride)
+		{
+			const uint RawSourceParticleID = RawParticleID - PrefixScanStride;
+
+			// ribbon distance is treated a bit differently because we have to evaluate if the ribbon id is consistent between the two points
+#if HAS_RIBBON_ID
+			const uint SourceParticleID = SortedIndices[RawSourceParticleID];
+			const uint ParticleID = SortedIndices[RawParticleID];
+			
+			const FRibbonID SourceRibbonID = GetRibbonID(SourceParticleID);
+			const FRibbonID RibbonID = GetRibbonID(ParticleID);
+
+			if (AreRibbonIDsEqual(SourceRibbonID, RibbonID))
+#endif
+			{
+				OutputAccumulation[RawParticleID].RibbonDistance = InputAccumulation[RawParticleID].RibbonDistance + InputAccumulation[RawSourceParticleID].RibbonDistance;
+			}
+#if HAS_RIBBON_ID
+			else
+			{
+				OutputAccumulation[RawParticleID].RibbonDistance = InputAccumulation[RawParticleID].RibbonDistance;
+			}
+#endif
+
+			OutputAccumulation[RawParticleID].SegmentCount = InputAccumulation[RawParticleID].SegmentCount + InputAccumulation[RawSourceParticleID].SegmentCount;
+#if HAS_RIBBON_ID
+			OutputAccumulation[RawParticleID].MultiRibbonCount = InputAccumulation[RawParticleID].MultiRibbonCount + InputAccumulation[RawSourceParticleID].MultiRibbonCount;
+#endif // HAS_RIBBON_ID
+#if RIBBONS_WANTS_AUTOMATIC_TESSELLATION
+			OutputAccumulation[RawParticleID].TessTotalLength = InputAccumulation[RawParticleID].TessTotalLength + InputAccumulation[RawSourceParticleID].TessTotalLength;
+			OutputAccumulation[RawParticleID].TessAvgSegmentLength = InputAccumulation[RawParticleID].TessAvgSegmentLength + InputAccumulation[RawSourceParticleID].TessAvgSegmentLength;
+			OutputAccumulation[RawParticleID].TessAvgSegmentAngle = InputAccumulation[RawParticleID].TessAvgSegmentAngle + InputAccumulation[RawSourceParticleID].TessAvgSegmentAngle;
+#if RIBBON_HAS_TWIST
+			OutputAccumulation[RawParticleID].TessTwistAvgAngle = InputAccumulation[RawParticleID].TessTwistAvgAngle + InputAccumulation[RawSourceParticleID].TessTwistAvgAngle;
+			OutputAccumulation[RawParticleID].TessTwistAvgWidth = InputAccumulation[RawParticleID].TessTwistAvgWidth + InputAccumulation[RawSourceParticleID].TessTwistAvgWidth;
+#endif // RIBBON_HAS_TWIST
+#endif // RIBBONS_WANTS_AUTOMATIC_TESSELLATION
+		}
+		else
+		{
+			OutputAccumulation[RawParticleID] = InputAccumulation[RawParticleID];
+		}
+	}
+}
diff --git a/Engine/Plugins/FX/Niagara/Shaders/Private/Ribbons/NiagaraRibbonCommon.ush b/Engine/Plugins/FX/Niagara/Shaders/Private/Ribbons/NiagaraRibbonCommon.ush
--- a/Engine/Plugins/FX/Niagara/Shaders/Private/Ribbons/NiagaraRibbonCommon.ush
+++ b/Engine/Plugins/FX/Niagara/Shaders/Private/Ribbons/NiagaraRibbonCommon.ush
@@ -248,4 +248,22 @@ float AcosFast(float InX)
 	float Res = -0.156583f * X + (0.5 * PI);
 	Res *= sqrt(max(0.f, 1.0f - X));
 	return (InX >= 0) ? Res : PI - Res;
-}
\ No newline at end of file
+}
+
+struct FRibbonAccumulationValues
+{
+	float RibbonDistance;
+	uint SegmentCount;
+#if HAS_RIBBON_ID
+	uint MultiRibbonCount;
+#endif // HAS_RIBBON_ID
+#if RIBBONS_WANTS_AUTOMATIC_TESSELLATION
+	float TessTotalLength;
+	float TessAvgSegmentLength;
+	float TessAvgSegmentAngle;
+#if RIBBON_HAS_TWIST
+	float TessTwistAvgAngle;
+	float TessTwistAvgWidth;
+#endif // RIBBON_HAS_TWIST
+#endif // RIBBONS_WANTS_AUTOMATIC_TESSELLATION
+};
diff --git a/Engine/Plugins/FX/Niagara/Shaders/Private/Ribbons/NiagaraRibbonGenerateIndices.usf b/Engine/Plugins/FX/Niagara/Shaders/Private/Ribbons/NiagaraRibbonGenerateIndices.usf
--- a/Engine/Plugins/FX/Niagara/Shaders/Private/Ribbons/NiagaraRibbonGenerateIndices.usf
+++ b/Engine/Plugins/FX/Niagara/Shaders/Private/Ribbons/NiagaraRibbonGenerateIndices.usf
@@ -60,7 +60,8 @@ void GenerateIndices(uint3 DispatchThreadId : SV_DispatchThreadID)
 	const uint IndicesPerSubSegment = TrianglesPerSegment * 3;	
 	const uint SegmentIndex = Segments[RawParticleID];
 	
-	if (RawParticleID < MaxPossibleSegments)
+	// make sure to ignore those segments that have been marked as invalid (INDEX_NONE)
+	if (RawParticleID < MaxPossibleSegments && SegmentIndex != INDEX_NONE)
 	{		
 #if HAS_RIBBON_ID
 		const int ThisRibbonID = MultiRibbonIndices[RawParticleID];
diff --git a/Engine/Plugins/FX/Niagara/Shaders/Private/Ribbons/NiagaraRibbonVertexReductionFinalization.usf b/Engine/Plugins/FX/Niagara/Shaders/Private/Ribbons/NiagaraRibbonVertexReductionFinalization.usf
--- a/Engine/Plugins/FX/Niagara/Shaders/Private/Ribbons/NiagaraRibbonVertexReductionFinalization.usf
+++ b/Engine/Plugins/FX/Niagara/Shaders/Private/Ribbons/NiagaraRibbonVertexReductionFinalization.usf
@@ -19,6 +19,8 @@ Buffer<uint> Segments;
 Buffer<float> TessellationStats;
 #endif
 
+StructuredBuffer<FRibbonAccumulationValues> AccumulationBuffer;
+
 // Buffer containing basic stats about each ribbon
 // including the first/last indices, as well as UV offsets/scales/distributions for both channels
 RWBuffer<uint> PackedPerRibbonData;
@@ -113,7 +115,7 @@ void OutputRibbonStats(uint3 DispatchThreadId : SV_DispatchThreadID)
 #if !HAS_RIBBON_ID
 		const int TotalNumRibbons = 1;
 #else
-		const int TotalNumRibbons = MultiRibbonIndices[LastParticleIndex] + 1;
+		const int TotalNumRibbons = AccumulationBuffer[LastParticleIndex].MultiRibbonCount;
 #endif		
 
 		// Calculate the indirect args for the finalization shader
@@ -121,8 +123,8 @@ void OutputRibbonStats(uint3 DispatchThreadId : SV_DispatchThreadID)
 		OutputCommandBuffer[OutputCommandBufferOffset + VERTEX_GEN_OUTPUT_DATA_FINALIZATION_INDIRECT_ARGS_OFFSET + 1] = 1;
 		OutputCommandBuffer[OutputCommandBufferOffset + VERTEX_GEN_OUTPUT_DATA_FINALIZATION_INDIRECT_ARGS_OFFSET + 2] = 1;
 		
-		// Total num segments, we can get it from taking the last segment index and adding one to get the count (segments is one less in length than sorted particles, so we need the 2nd to last particle
-		const int TotalSegments = Segments[LastParticleIndex - 1] + 1;
+		// Total num segments
+		const int TotalSegments = AccumulationBuffer[LastParticleIndex].SegmentCount;
 		OutputCommandBuffer[OutputCommandBufferOffset + VERTEX_GEN_OUTPUT_DATA_TOTAL_NUM_SEGMENTS_OFFSET] = TotalSegments;
 
 		// Total num ribbons
diff --git a/Engine/Plugins/FX/Niagara/Shaders/Private/Ribbons/NiagaraRibbonVertexReductionInitialization.usf b/Engine/Plugins/FX/Niagara/Shaders/Private/Ribbons/NiagaraRibbonVertexReductionInitialization.usf
--- a/Engine/Plugins/FX/Niagara/Shaders/Private/Ribbons/NiagaraRibbonVertexReductionInitialization.usf
+++ b/Engine/Plugins/FX/Niagara/Shaders/Private/Ribbons/NiagaraRibbonVertexReductionInitialization.usf
@@ -10,54 +10,44 @@ Buffer<uint> SortedIndices;
 // prefix sum passes will be even so it takes starting with the final buffer to end in the final buffer
 
 // Tangent of the particle, and the distance along the ribbon if we're computing it
-Buffer<float> InputTangentsAndDistances;
 RWBuffer<float> OutputTangentsAndDistances;
 
 // Index of the ribbon for this particle
 // created through prefix sum across all the starting points of the ribbons
 #if HAS_RIBBON_ID
-Buffer<uint> InputMultiRibbonIndices;
 RWBuffer<uint> OutputMultiRibbonIndices;
 #endif
 
 // Index of the final segment this particle starts
-Buffer<uint> InputSegments;
 RWBuffer<uint> OutputSegments;
 
-#if RIBBONS_WANTS_AUTOMATIC_TESSELLATION
-// Tessellation reduction stats if we're computing them
-RWBuffer<float> OutputTessellationStats;
-#endif
+RWStructuredBuffer<FRibbonAccumulationValues> OutputAccumulationBuffer;
 
 float CurveTension;
-int PrefixSumStride;
-
-
-
 
 // Super simple shader just to initialize the ribbon reduction list sizes to zero
 [numthreads(THREADGROUP_SIZE, 1, 1)]
 void VertexGenInitialize(uint3 DispatchThreadId : SV_DispatchThreadID)
 {
 	const uint TotalNumParticles = GetTotalNumParticles();
-	const int RawParticleID = DispatchThreadId.x;	
+	const uint RawParticleID = DispatchThreadId.x;	
 
 	BRANCH
 	if (RawParticleID < TotalNumParticles)
 	{
+		OutputAccumulationBuffer[RawParticleID] = (FRibbonAccumulationValues)0;
+
 		const int ParticleID = SortedIndices[RawParticleID];
 #if HAS_RIBBON_ID
 		const FRibbonID RibbonID = GetRibbonID(ParticleID);
 #endif
 
-
 		bool bHasPrevParticle = RawParticleID > 0;
 		bool bHasNextParticle = (RawParticleID + 1) < TotalNumParticles;
 
 		int PrevParticleID = INDEX_NONE;
 		int NextParticleID = INDEX_NONE;
 		
-		BRANCH
 		if (bHasPrevParticle)
 		{
 			PrevParticleID = SortedIndices[RawParticleID - 1];
@@ -67,7 +57,6 @@ void VertexGenInitialize(uint3 DispatchThreadId : SV_DispatchThreadID)
 #endif
 		}
 
-		BRANCH
 		if (bHasNextParticle)
 		{
 			NextParticleID = SortedIndices[RawParticleID + 1];
@@ -76,23 +65,21 @@ void VertexGenInitialize(uint3 DispatchThreadId : SV_DispatchThreadID)
 			bHasNextParticle = AreRibbonIDsEqual(RibbonID, NextRibbonID);
 #endif
 		}
-
 		
 
 #if HAS_RIBBON_ID
 		// We start a new ribbon index each time it breaks from the previous
-		OutputMultiRibbonIndices[RawParticleID] = (!bHasPrevParticle && RawParticleID != 0) ? 1 : 0;
+		const uint MultiRibbonCountIncr = bHasPrevParticle ? 0 : 1;
+		OutputMultiRibbonIndices[RawParticleID] = MultiRibbonCountIncr;
+		OutputAccumulationBuffer[RawParticleID].MultiRibbonCount =  MultiRibbonCountIncr;
 #endif
 
 		// Track the segment offset
-		OutputSegments[RawParticleID] = (bHasNextParticle && RawParticleID != 0)? 1 : 0;
+		OutputSegments[RawParticleID] = bHasNextParticle ? 1 : INDEX_NONE;
+		OutputAccumulationBuffer[RawParticleID].SegmentCount = bHasNextParticle ? 1 : 0;
 
 		float4 TangentAndDistance = float4(0,0,0,0);
-#if RIBBONS_WANTS_AUTOMATIC_TESSELLATION
-		float4 TessellationStatsEntry = float4(0,0,0,0);
-		float TessellationTotalLength = 0;
-#endif
-		
+
 		BRANCH
 		if (bHasNextParticle && bHasPrevParticle)
 		{
@@ -128,16 +115,14 @@ void VertexGenInitialize(uint3 DispatchThreadId : SV_DispatchThreadID)
 			{
 				TangentAndDistance = float4(0,0,0,0);
 			}
-			
 
 #if RIBBONS_WANTS_AUTOMATIC_TESSELLATION
 			// TotalLength
-			TessellationTotalLength = CurrToNextSize;
-			
+			OutputAccumulationBuffer[RawParticleID].TessTotalLength = CurrToNextSize;
 			// AverageSegmentLength
-			TessellationStatsEntry.x = CurrToNextSize * CurrToNextSize;
+			OutputAccumulationBuffer[RawParticleID].TessAvgSegmentLength = CurrToNextSize * CurrToNextSize;
 			// AverageSegmentAngle
-			TessellationStatsEntry.y = CurrToNextSize * AcosFast(dot(LastToCurrVec, CurrToNextVec));
+			OutputAccumulationBuffer[RawParticleID].TessAvgSegmentAngle = CurrToNextSize * AcosFast(dot(LastToCurrVec, CurrToNextVec));
 
 #if RIBBON_HAS_TWIST
 			const float LastTwist = SafeGetFloat(TwistDataOffset, PrevParticleID, 0);
@@ -146,16 +131,11 @@ void VertexGenInitialize(uint3 DispatchThreadId : SV_DispatchThreadID)
 			//const float LastWidth = SafeGetFloat(WidthDataOffset, PrevParticleID, 1);
 			const float ThisWidth = SafeGetFloat(WidthDataOffset, ParticleID, 1);
 			// AverageTwistAngle
-			TessellationStatsEntry.z = CurrToNextSize * abs(ThisTwist - LastTwist);
-			// AverageWidth
-			TessellationStatsEntry.w = CurrToNextSize * ThisWidth;
-#else
-			// AverageTwistAngle
-			TessellationStatsEntry.z = 0;
+			OutputAccumulationBuffer[RawParticleID].TessTwistAvgAngle = CurrToNextSize * abs(ThisTwist - LastTwist);
 			// AverageWidth
-			TessellationStatsEntry.w = 0;
-#endif
-#endif
+			OutputAccumulationBuffer[RawParticleID].TessTwistAvgWidth = CurrToNextSize * ThisWidth;
+#endif // RIBBON_HAS_TWIST
+#endif // RIBBONS_WANTS_AUTOMATIC_TESSELLATION
 		}
 		else if (bHasNextParticle)
 		{
@@ -199,28 +179,7 @@ void VertexGenInitialize(uint3 DispatchThreadId : SV_DispatchThreadID)
 		OutputTangentsAndDistances[RawParticleID * 4 + 2] = TangentAndDistance.z;
 			
 		// This is the distance up to this point, we'll prefix sum this in the next stage
-		OutputTangentsAndDistances[RawParticleID * 4 + 3] = TangentAndDistance.w;
-
-#if RIBBONS_WANTS_AUTOMATIC_TESSELLATION
-#if RIBBON_HAS_TWIST
-		const uint TessellationBaseOffset = RawParticleID * 5;
-#else
-		const uint TessellationBaseOffset = RawParticleID * 3;
-#endif
-		// TotalLength
-		OutputTessellationStats[TessellationBaseOffset + 0] = TessellationTotalLength;
-		// AverageSegmentLength
-		OutputTessellationStats[TessellationBaseOffset + 1] = TessellationStatsEntry.x;
-		// AverageSegmentAngle
-		OutputTessellationStats[TessellationBaseOffset + 2] = TessellationStatsEntry.y;
-#if RIBBON_HAS_TWIST
-		// AverageTwistAngle
-		OutputTessellationStats[TessellationBaseOffset + 3] = TessellationStatsEntry.z;
-		// AverageWidth
-		OutputTessellationStats[TessellationBaseOffset + 4] = TessellationStatsEntry.w;
-#endif
-#endif
-
+		OutputAccumulationBuffer[RawParticleID].RibbonDistance = TangentAndDistance.w;
 	}
 }
 
diff --git a/Engine/Plugins/FX/Niagara/Source/Niagara/Private/NiagaraRendererRibbons.cpp b/Engine/Plugins/FX/Niagara/Source/Niagara/Private/NiagaraRendererRibbons.cpp
--- a/Engine/Plugins/FX/Niagara/Source/Niagara/Private/NiagaraRendererRibbons.cpp
+++ b/Engine/Plugins/FX/Niagara/Source/Niagara/Private/NiagaraRendererRibbons.cpp
@@ -491,13 +491,22 @@ void FNiagaraRibbonVertexBuffers::InitializeOrUpdateBuffers(FRHICommandListBase&
 struct FNiagaraRibbonGPUInitComputeBuffers
 {
 	FRWBuffer SortBuffer;
-	FRWBuffer TempSegments;
-	FRWBuffer TempDistances;
-	FRWBuffer TempMultiRibbon;
-	FRWBuffer TempTessellationStats[2];
+	FRWBuffer TransientTessellationStats;
+	FRWBufferStructured TransientAccumulation[2];
 	
 	FNiagaraRibbonGPUInitComputeBuffers() { }
 
+	uint32 GetAccumulationStructSize(bool bWantsMultiRibbon, bool bWantsTessellation, bool bWantsTessellationTwist)
+	{
+		const uint32 BaseStructSize = 2;
+		const uint32 MultiRibbonSize = 1;
+		const uint32 TessellationSize = bWantsTessellationTwist ? 5 : 3 ;
+
+		return BaseStructSize
+			+ (bWantsMultiRibbon ? MultiRibbonSize : 0)
+			+ (bWantsTessellation ? TessellationSize : 0);
+	}
+
 	void InitOrUpdateBuffers(FRHICommandListBase& RHICmdList, int32 NeededSize, bool bWantsMultiRibbon, bool bWantsTessellation, bool bWantsTessellationTwist)
 	{
 		// TODO: Downsize these when we haven't needed the size for a bit
@@ -508,27 +517,17 @@ struct FNiagaraRibbonGPUInitComputeBuffers
 			SortBuffer.Initialize(RHICmdList, TEXT("NiagarGPUInit-SortedIndices"), sizeof(uint32), NeededSize, EPixelFormat::PF_R32_UINT, InitialAccess);
 		}
 
-		if (TempSegments.NumBytes < NeededSize * sizeof(int32))
-		{
-			TempSegments.Initialize(RHICmdList, TEXT("NiagaraGPUInit-Segments"), sizeof(uint32), NeededSize, EPixelFormat::PF_R32_UINT, InitialAccess, BUF_Static);
-		}
-
-		if (TempDistances.NumBytes < NeededSize * sizeof(float) * 4)
-		{
-			TempDistances.Initialize(RHICmdList, TEXT("NiagaraGPUInit-Distances"), sizeof(float), NeededSize * 4, EPixelFormat::PF_R32_FLOAT, InitialAccess, BUF_Static);
-		}
-
-		const uint32 MultiRibbonBufferSize = NeededSize * (bWantsMultiRibbon? 1 : 0);
-		if (TempMultiRibbon.NumBytes < MultiRibbonBufferSize * sizeof(int32))
+		const uint32 TessellationBufferSize = NeededSize * (bWantsTessellation? (bWantsTessellationTwist? FTessellationStatsEntry::NumElements : FTessellationStatsEntryNoTwist::NumElements) : 0);
+		if (TransientTessellationStats.NumBytes < TessellationBufferSize * sizeof(float))
 		{
-			TempMultiRibbon.Initialize(RHICmdList, TEXT("NiagaraGPUInit-MultiRibbon"), sizeof(uint32), MultiRibbonBufferSize, EPixelFormat::PF_R32_UINT, InitialAccess, BUF_Static);
+			TransientTessellationStats.Initialize(RHICmdList, TEXT("NiagaraGPUInit-Tessellation-0"), sizeof(float), TessellationBufferSize, EPixelFormat::PF_R32_FLOAT, InitialAccess, BUF_Static);
 		}
 
-		const uint32 TessellationBufferSize = NeededSize * (bWantsTessellation? (bWantsTessellationTwist? FTessellationStatsEntry::NumElements : FTessellationStatsEntryNoTwist::NumElements) : 0);
-		if (TempTessellationStats[0].NumBytes < TessellationBufferSize * sizeof(float))
+		const uint32 AccumulationBufferStructSize = GetAccumulationStructSize(bWantsMultiRibbon, bWantsTessellation, bWantsTessellationTwist) * sizeof(float);
+		if (TransientAccumulation[0].NumBytes < (AccumulationBufferStructSize * NeededSize))
 		{
-			TempTessellationStats[0].Initialize(RHICmdList, TEXT("NiagaraGPUInit-Tessellation-0"), sizeof(float), TessellationBufferSize, EPixelFormat::PF_R32_FLOAT, InitialAccess, BUF_Static);
-			TempTessellationStats[1].Initialize(RHICmdList, TEXT("NiagaraGPUInit-Tessellation-1"), sizeof(float), TessellationBufferSize, EPixelFormat::PF_R32_FLOAT, InitialAccess, BUF_Static);
+			TransientAccumulation[0].Initialize(RHICmdList, TEXT("NiagaraGPUInit-Accumulation-0"), AccumulationBufferStructSize, NeededSize, BUF_Static, false, false, InitialAccess);
+			TransientAccumulation[1].Initialize(RHICmdList, TEXT("NiagaraGPUInit-Accumulation-1"), AccumulationBufferStructSize, NeededSize, BUF_Static, false, false, InitialAccess);
 		}
 	}	
 };
@@ -2569,9 +2568,8 @@ void FNiagaraRendererRibbons::InitializeVertexBuffersGPU(FRHICommandListImmediat
 	}
 	
 	{
-		SCOPED_DRAW_EVENT(RHICmdList, NiagaraRenderRibbonsGenVerticesReductionPhase1GPU);
-		SCOPE_CYCLE_COUNTER(STAT_NiagaraRenderRibbonsGenVerticesReductionPhase1GPU);
-		
+		SCOPED_DRAW_EVENT(RHICmdList, NiagaraRenderRibbonsAggregation);
+
 		FNiagaraRibbonVertexReductionInitializationCS::FPermutationDomain InitPermutationVector;
 		InitPermutationVector.Set<FRibbonHasFullRibbonID>(GenerationConfig.HasFullRibbonIDs());
 		InitPermutationVector.Set<FRibbonHasRibbonID>(GenerationConfig.HasSimpleRibbonIDs());
@@ -2580,123 +2578,148 @@ void FNiagaraRendererRibbons::InitializeVertexBuffersGPU(FRHICommandListImmediat
 		InitPermutationVector.Set<FRibbonHasTwist>(GenerationConfig.HasTwist());
 		TShaderMapRef<FNiagaraRibbonVertexReductionInitializationCS> ReductionInitializationShader(GetGlobalShaderMap(GMaxRHIFeatureLevel), InitPermutationVector);
 
-		FNiagaraRibbonVertexReductionPropagateCS::FPermutationDomain  PropagationPermutationVector;
-		PropagationPermutationVector.Set<FRibbonHasFullRibbonID>(GenerationConfig.HasFullRibbonIDs());
-		PropagationPermutationVector.Set<FRibbonHasRibbonID>(GenerationConfig.HasSimpleRibbonIDs());
-		PropagationPermutationVector.Set<FRibbonWantsAutomaticTessellation>(GenerationConfig.WantsAutomaticTessellation());
-		PropagationPermutationVector.Set<FRibbonWantsConstantTessellation>(GenerationConfig.WantsConstantTessellation());
-		PropagationPermutationVector.Set<FRibbonHasTwist>(GenerationConfig.HasTwist());
-		TShaderMapRef<FNiagaraRibbonVertexReductionPropagateCS> ReductionPropgateShader(GetGlobalShaderMap(GMaxRHIFeatureLevel), PropagationPermutationVector);
-		
+		FNiagaraRibbonAggregationStepCS::FPermutationDomain StepPermutationVector;
+		StepPermutationVector.Set<FRibbonHasFullRibbonID>(GenerationConfig.HasFullRibbonIDs());
+		StepPermutationVector.Set<FRibbonHasRibbonID>(GenerationConfig.HasSimpleRibbonIDs());
+		StepPermutationVector.Set<FRibbonWantsAutomaticTessellation>(GenerationConfig.WantsAutomaticTessellation());
+		StepPermutationVector.Set<FRibbonWantsConstantTessellation>(GenerationConfig.WantsConstantTessellation());
+		StepPermutationVector.Set<FRibbonHasTwist>(GenerationConfig.HasTwist());
+		TShaderMapRef<FNiagaraRibbonAggregationStepCS> AggregationStepShader(GetGlobalShaderMap(GMaxRHIFeatureLevel), StepPermutationVector);
+
+		FNiagaraRibbonAggregationApplyCS::FPermutationDomain ApplyPermutationVector;
+		ApplyPermutationVector.Set<FRibbonHasFullRibbonID>(GenerationConfig.HasFullRibbonIDs());
+		ApplyPermutationVector.Set<FRibbonHasRibbonID>(GenerationConfig.HasSimpleRibbonIDs());
+		ApplyPermutationVector.Set<FRibbonWantsAutomaticTessellation>(GenerationConfig.WantsAutomaticTessellation());
+		ApplyPermutationVector.Set<FRibbonWantsConstantTessellation>(GenerationConfig.WantsConstantTessellation());
+		ApplyPermutationVector.Set<FRibbonHasTwist>(GenerationConfig.HasTwist());
+		TShaderMapRef<FNiagaraRibbonAggregationApplyCS> AggregationApplyShader(GetGlobalShaderMap(GMaxRHIFeatureLevel), ApplyPermutationVector);
+
+		// repeatedly run the step shader to perform the aggregation
+		FUnorderedAccessViewRHIRef CurrentAccumulationUAV;
+		FShaderResourceViewRHIRef CurrentAccumulationSRV;
+
+		// Setup buffers
 		const int32 NumPrefixScanPasses = FMath::CeilLogTwo(NumExecutableInstances);
-	
-		FNiagaraRibbonVertexReductionParameters Params;
-		FMemory::Memzero(Params);
-		Params.Common = CommonParams;
-		Params.SortedIndices = VertexBuffers.SortedIndicesBuffer.SRV;
-		Params.CurveTension = GenerationConfig.GetCurveTension();
-	
+
 		uint32 CurrentBufferOrientation = 0x0;
-		const auto SwapBuffers = [&]()
+		const auto SwapAccumulationBuffers = [&]()
 		{
 			CurrentBufferOrientation ^= 0x1;
-			const bool bComputeOnOutputBuffer = CurrentBufferOrientation == 0;
-	
-			if (bComputeOnOutputBuffer)
+
+			if (CurrentBufferOrientation)
 			{
-				Params.InputTangentsAndDistances = TempBuffers.TempDistances.SRV;
-				Params.OutputTangentsAndDistances = VertexBuffers.TangentsAndDistancesBuffer.UAV;
-				Params.InputMultiRibbonIndices = TempBuffers.TempMultiRibbon.SRV;
-				Params.OutputMultiRibbonIndices = VertexBuffers.MultiRibbonIndicesBuffer.UAV;
-				Params.InputSegments = TempBuffers.TempSegments.SRV;
-				Params.OutputSegments = VertexBuffers.SegmentsBuffer.UAV;
-				Params.InputTessellationStats = TempBuffers.TempTessellationStats[1].SRV;
-				Params.OutputTessellationStats = TempBuffers.TempTessellationStats[0].UAV;
+				CurrentAccumulationSRV = TempBuffers.TransientAccumulation[0].SRV;
+				CurrentAccumulationUAV = TempBuffers.TransientAccumulation[1].UAV;
 			}
 			else
 			{
-				Params.InputTangentsAndDistances = VertexBuffers.TangentsAndDistancesBuffer.SRV;
-				Params.OutputTangentsAndDistances = TempBuffers.TempDistances.UAV;
-				Params.InputMultiRibbonIndices = VertexBuffers.MultiRibbonIndicesBuffer.SRV;
-				Params.OutputMultiRibbonIndices = TempBuffers.TempMultiRibbon.UAV;
-				Params.InputSegments = VertexBuffers.SegmentsBuffer.SRV;
-				Params.OutputSegments = TempBuffers.TempSegments.UAV;				
-				Params.InputTessellationStats = TempBuffers.TempTessellationStats[0].SRV;
-				Params.OutputTessellationStats = TempBuffers.TempTessellationStats[1].UAV;
-			}				
-		};	
-	
-		// Setup buffers
-		if (NumPrefixScanPasses % 2 == 0)
+				CurrentAccumulationSRV = TempBuffers.TransientAccumulation[1].SRV;
+				CurrentAccumulationUAV = TempBuffers.TransientAccumulation[0].UAV;
+			}
+		};
+
+		// init CurrentAccumulationBufferSRV/UAV
+		SwapAccumulationBuffers();
+
+		// aggregation involves calculating a per-particle value for a number of properties and then accumulating these values over all
+		// segments in the ribbons:
+		// Segment indices are initialized with a 1 for a valid segment, 0 otherwise.  Inclusive prefix sum accumulates the data
+		//	-note that we mark the OutputSegments field as -1 if it's an invalid segment (like if there's only one particle)
+		// Ribbon indices are initialized with a 1 for the start of a new ribbon, 0 otherwise.  Inclusive prefix sum accumulates the data
+
+		// initialize the output buffers and queue up data to be aggregated
 		{
-			CurrentBufferOrientation = 0x1;
-			SwapBuffers();
+			FNiagaraRibbonVertexReductionParameters InitParams;
+			FMemory::Memzero(InitParams);
+			InitParams.Common = CommonParams;
+			InitParams.SortedIndices = VertexBuffers.SortedIndicesBuffer.SRV;
+			InitParams.CurveTension = GenerationConfig.GetCurveTension();
+			InitParams.OutputTangentsAndDistances = VertexBuffers.TangentsAndDistancesBuffer.UAV;
+			InitParams.OutputMultiRibbonIndices = VertexBuffers.MultiRibbonIndicesBuffer.UAV;
+			InitParams.OutputSegments = VertexBuffers.SegmentsBuffer.UAV;
+			InitParams.OutputAccumulationBuffer = CurrentAccumulationUAV;
+
+			RHICmdList.Transition(FRHITransitionInfo(InitParams.OutputTangentsAndDistances, ERHIAccess::SRVMask, ERHIAccess::UAVMask));
+			RHICmdList.Transition(FRHITransitionInfo(InitParams.OutputMultiRibbonIndices, ERHIAccess::SRVMask, ERHIAccess::UAVMask));
+			RHICmdList.Transition(FRHITransitionInfo(InitParams.OutputSegments, ERHIAccess::SRVMask, ERHIAccess::UAVMask));
+			RHICmdList.Transition(FRHITransitionInfo(InitParams.OutputAccumulationBuffer, ERHIAccess::SRVMask, ERHIAccess::UAVMask));
+
+			const uint32 NumThreadGroupsInitialization = FMath::DivideAndRoundUp<uint32>(NumExecutableInstances, FNiagaraRibbonComputeCommon::VertexGenReductionInitializationThreadSize);
+			FComputeShaderUtils::Dispatch(RHICmdList, ReductionInitializationShader, InitParams, FIntVector(NumThreadGroupsInitialization, 1, 1));
+
+			RHICmdList.Transition(FRHITransitionInfo(InitParams.OutputTangentsAndDistances, ERHIAccess::UAVMask, ERHIAccess::SRVMask));
+			RHICmdList.Transition(FRHITransitionInfo(InitParams.OutputMultiRibbonIndices, ERHIAccess::UAVMask, ERHIAccess::SRVMask));
+			RHICmdList.Transition(FRHITransitionInfo(InitParams.OutputSegments, ERHIAccess::UAVMask, ERHIAccess::SRVMask));
+			RHICmdList.Transition(FRHITransitionInfo(InitParams.OutputAccumulationBuffer, ERHIAccess::UAVMask, ERHIAccess::SRVMask));
 		}
-		else
-		{
-			CurrentBufferOrientation = 0x0;
-			SwapBuffers();
-		}		
-		
-		const auto TransitionOutputBuffers = [this, &RHICmdList, &Params, &TempBuffers](ERHIAccess Previous, ERHIAccess Next)
+
 		{
-			FRHITransitionInfo DataBufferTransitions[] =
+			FNiagaraRibbonAggregationStepParameters StepParams;
+			FMemory::Memzero(StepParams);
+			StepParams.Common = CommonParams;
+			StepParams.SortedIndices = VertexBuffers.SortedIndicesBuffer.SRV;
+
+			const uint32 NumThreadGroups = FMath::DivideAndRoundUp<uint32>(NumExecutableInstances, FNiagaraRibbonComputeCommon::VertexGenReductionPropagationThreadSize);
+
+			for (StepParams.PrefixScanStride = 1; StepParams.PrefixScanStride < static_cast<uint32>(NumExecutableInstances); StepParams.PrefixScanStride *= 2)
 			{
-				FRHITransitionInfo(Params.OutputTangentsAndDistances, Previous, Next),
-				FRHITransitionInfo(Params.OutputMultiRibbonIndices, Previous, Next),
-				FRHITransitionInfo(Params.OutputSegments, Previous, Next),
-				FRHITransitionInfo(Params.OutputTessellationStats, Previous, Next),
-			};
-			RHICmdList.Transition(MakeArrayView(DataBufferTransitions, UE_ARRAY_COUNT(DataBufferTransitions)));
-		};
-		
-		{			
-			
-			const uint32 NumThreadGroupsInitialization = FMath::DivideAndRoundUp<uint32>(NumExecutableInstances, FNiagaraRibbonComputeCommon::VertexGenReductionInitializationThreadSize);	
-			TransitionOutputBuffers(ERHIAccess::SRVMask | ERHIAccess::VertexOrIndexBuffer, ERHIAccess::UAVCompute);
-			FComputeShaderUtils::Dispatch(RHICmdList, ReductionInitializationShader, Params, FIntVector(NumThreadGroupsInitialization, 1, 1));
-			TransitionOutputBuffers(ERHIAccess::UAVCompute, ERHIAccess::SRVMask | ERHIAccess::VertexOrIndexBuffer);
+				SwapAccumulationBuffers();
+
+				StepParams.InputAccumulation = CurrentAccumulationSRV;
+				StepParams.OutputAccumulation = CurrentAccumulationUAV;
+
+				RHICmdList.Transition(FRHITransitionInfo(StepParams.OutputAccumulation, ERHIAccess::SRVMask, ERHIAccess::UAVCompute));
+				FComputeShaderUtils::Dispatch(RHICmdList, AggregationStepShader, StepParams, FIntVector(NumThreadGroups, 1, 1));
+				RHICmdList.Transition(FRHITransitionInfo(StepParams.OutputAccumulation, ERHIAccess::UAVCompute, ERHIAccess::SRVMask));
+			}
 		}
-		
+
+		// apply the aggregate data onto the output buffers
 		{
-			SCOPE_CYCLE_COUNTER(STAT_NiagaraRenderRibbonsGenVerticesReductionPropagateGPU);
-			
+			SwapAccumulationBuffers();
+
+			FNiagaraRibbonAggregationApplyParameters ApplyParams;
+			FMemory::Memzero(ApplyParams);
+			ApplyParams.Common = CommonParams;
+			ApplyParams.SortedIndices = VertexBuffers.SortedIndicesBuffer.SRV;
+			ApplyParams.InputAccumulation = CurrentAccumulationSRV;
+			ApplyParams.OutputTangentsAndDistances = VertexBuffers.TangentsAndDistancesBuffer.UAV;
+			ApplyParams.OutputMultiRibbonIndices = VertexBuffers.MultiRibbonIndicesBuffer.UAV;
+			ApplyParams.OutputTessellationStats = TempBuffers.TransientTessellationStats.UAV;
+			ApplyParams.OutputSegments = VertexBuffers.SegmentsBuffer.UAV;
+
+			RHICmdList.Transition(FRHITransitionInfo(ApplyParams.OutputTangentsAndDistances, ERHIAccess::SRVMask, ERHIAccess::UAVMask));
+			RHICmdList.Transition(FRHITransitionInfo(ApplyParams.OutputMultiRibbonIndices, ERHIAccess::SRVMask, ERHIAccess::UAVMask));
+			RHICmdList.Transition(FRHITransitionInfo(ApplyParams.OutputSegments, ERHIAccess::SRVMask, ERHIAccess::UAVMask));
+			RHICmdList.Transition(FRHITransitionInfo(ApplyParams.OutputTessellationStats, ERHIAccess::SRVMask, ERHIAccess::UAVMask));
+
 			const uint32 NumThreadGroups = FMath::DivideAndRoundUp<uint32>(NumExecutableInstances, FNiagaraRibbonComputeCommon::VertexGenReductionPropagationThreadSize);
-			
-			for (Params.PrefixScanStride = 1; Params.PrefixScanStride < NumExecutableInstances; Params.PrefixScanStride *= 2)
-			{
-				SwapBuffers();
-				
-				TransitionOutputBuffers(ERHIAccess::SRVMask | ERHIAccess::VertexOrIndexBuffer, ERHIAccess::UAVCompute);
-				FComputeShaderUtils::Dispatch(RHICmdList, ReductionPropgateShader, Params, FIntVector(NumThreadGroups, 1, 1));
-				TransitionOutputBuffers(ERHIAccess::UAVCompute, ERHIAccess::SRVMask | ERHIAccess::VertexOrIndexBuffer);
-			}
+			FComputeShaderUtils::Dispatch(RHICmdList, AggregationApplyShader, ApplyParams, FIntVector(NumThreadGroups, 1, 1));
+
+			RHICmdList.Transition(FRHITransitionInfo(ApplyParams.OutputTangentsAndDistances, ERHIAccess::UAVMask, ERHIAccess::SRVMask));
+			RHICmdList.Transition(FRHITransitionInfo(ApplyParams.OutputMultiRibbonIndices, ERHIAccess::UAVMask, ERHIAccess::SRVMask));
+			RHICmdList.Transition(FRHITransitionInfo(ApplyParams.OutputSegments, ERHIAccess::UAVMask, ERHIAccess::SRVMask));
+			RHICmdList.Transition(FRHITransitionInfo(ApplyParams.OutputTessellationStats, ERHIAccess::UAVMask, ERHIAccess::SRVMask));
 		}
 
-		check(CurrentBufferOrientation == 0x0);		
-	}
-	
-	static constexpr int32 CommandBufferOffset = 0;
-		
-	{
-		SCOPED_DRAW_EVENT(RHICmdList, NiagaraRenderRibbonsGenVerticesReductionPhase2GPU);
-		SCOPE_CYCLE_COUNTER(STAT_NiagaraRenderRibbonsGenVerticesReductionPhase2GPU);
+		static constexpr int32 CommandBufferOffset = 0;
 		
-		FNiagaraRibbonVertexReductionFinalizationParameters FinalizationParams;
-		FMemory::Memzero(FinalizationParams);
-		FinalizationParams.Common = CommonParams;
-		FinalizationParams.SortedIndices = VertexBuffers.SortedIndicesBuffer.SRV;
-		FinalizationParams.TangentsAndDistances = VertexBuffers.TangentsAndDistancesBuffer.SRV;
-		FinalizationParams.MultiRibbonIndices = VertexBuffers.MultiRibbonIndicesBuffer.SRV;
-		FinalizationParams.Segments = VertexBuffers.SegmentsBuffer.SRV;
-		FinalizationParams.TessellationStats = TempBuffers.TempTessellationStats[0].SRV;
-		FinalizationParams.PackedPerRibbonData = VertexBuffers.RibbonLookupTableBuffer.UAV;
-		FinalizationParams.OutputCommandBuffer = VertexBuffers.GPUComputeCommandBuffer.UAV;
-		FinalizationParams.OutputCommandBufferIndex= CommandBufferOffset;
-		FinalizationParams.FinalizationThreadBlockSize = FNiagaraRibbonComputeCommon::VertexGenFinalizationThreadSize;
-			
 		{
-			SCOPE_CYCLE_COUNTER(STAT_NiagaraRenderRibbonsGenVerticesReductionFinalizeGPU);
+			SCOPED_DRAW_EVENT(RHICmdList, NiagaraRenderRibbonsGenVerticesReductionPhase2GPU);
+		
+			FNiagaraRibbonVertexReductionFinalizationParameters FinalizationParams;
+			FMemory::Memzero(FinalizationParams);
+			FinalizationParams.Common = CommonParams;
+			FinalizationParams.SortedIndices = VertexBuffers.SortedIndicesBuffer.SRV;
+			FinalizationParams.TangentsAndDistances = VertexBuffers.TangentsAndDistancesBuffer.SRV;
+			FinalizationParams.MultiRibbonIndices = VertexBuffers.MultiRibbonIndicesBuffer.SRV;
+			FinalizationParams.Segments = VertexBuffers.SegmentsBuffer.SRV;
+			FinalizationParams.TessellationStats = TempBuffers.TransientTessellationStats.SRV;
+			FinalizationParams.AccumulationBuffer = CurrentAccumulationSRV;
+			FinalizationParams.PackedPerRibbonData = VertexBuffers.RibbonLookupTableBuffer.UAV;
+			FinalizationParams.OutputCommandBuffer = VertexBuffers.GPUComputeCommandBuffer.UAV;
+			FinalizationParams.OutputCommandBufferIndex= CommandBufferOffset;
+			FinalizationParams.FinalizationThreadBlockSize = FNiagaraRibbonComputeCommon::VertexGenFinalizationThreadSize;
 			
 			FNiagaraRibbonVertexReductionFinalizeCS::FPermutationDomain PermutationVector;
 			PermutationVector.Set<FRibbonHasFullRibbonID>(GenerationConfig.HasFullRibbonIDs());
@@ -2721,59 +2744,59 @@ void FNiagaraRendererRibbons::InitializeVertexBuffersGPU(FRHICommandListImmediat
 			RHICmdList.Transition({
 				FRHITransitionInfo(VertexBuffers.RibbonLookupTableBuffer.Buffer, ERHIAccess::UAVCompute, ERHIAccess::UAVCompute),
 				FRHITransitionInfo(VertexBuffers.GPUComputeCommandBuffer.Buffer, ERHIAccess::UAVCompute, ERHIAccess::SRVMask | ERHIAccess::VertexOrIndexBuffer | ERHIAccess::IndirectArgs)});		
-		}		
-	}
+		}
 		
-	{
-		SCOPED_DRAW_EVENT(RHICmdList, NiagaraRenderRibbonsGenVerticesMultiRibbonInitGPU);
-		SCOPE_CYCLE_COUNTER(STAT_NiagaraRenderRibbonsGenVerticesMultiRibbonInitGPU);
+		{
+			SCOPED_DRAW_EVENT(RHICmdList, NiagaraRenderRibbonsGenVerticesMultiRibbonInitGPU);
+			SCOPE_CYCLE_COUNTER(STAT_NiagaraRenderRibbonsGenVerticesMultiRibbonInitGPU);
 		
-		FNiagaraRibbonVertexFinalizationParameters FinalizeParams;
-		FMemory::Memzero(FinalizeParams);
-		FinalizeParams.Common = CommonParams;
-		FinalizeParams.SortedIndices = VertexBuffers.SortedIndicesBuffer.SRV;
-		FinalizeParams.TangentsAndDistances = VertexBuffers.TangentsAndDistancesBuffer.UAV;
-		FinalizeParams.PackedPerRibbonData = VertexBuffers.RibbonLookupTableBuffer.UAV;
-		FinalizeParams.CommandBuffer = VertexBuffers.GPUComputeCommandBuffer.SRV;
-		FinalizeParams.CommandBufferOffset = CommandBufferOffset;
+			FNiagaraRibbonVertexFinalizationParameters FinalizeParams;
+			FMemory::Memzero(FinalizeParams);
+			FinalizeParams.Common = CommonParams;
+			FinalizeParams.SortedIndices = VertexBuffers.SortedIndicesBuffer.SRV;
+			FinalizeParams.TangentsAndDistances = VertexBuffers.TangentsAndDistancesBuffer.UAV;
+			FinalizeParams.PackedPerRibbonData = VertexBuffers.RibbonLookupTableBuffer.UAV;
+			FinalizeParams.CommandBuffer = VertexBuffers.GPUComputeCommandBuffer.SRV;
+			FinalizeParams.CommandBufferOffset = CommandBufferOffset;
 		
-		constexpr auto AddUVChannelParams = [](const FNiagaraRibbonUVSettings& Input, FNiagaraRibbonUVSettingsParams& Output)
-		{
-			Output.Offset = FVector2f(Input.Offset);
-			Output.Scale = FVector2f(Input.Scale);
-			Output.TilingLength = Input.TilingLength;
-			Output.DistributionMode = static_cast<int32>(Input.DistributionMode);
-			Output.LeadingEdgeMode = static_cast<int32>(Input.LeadingEdgeMode);
-			Output.TrailingEdgeMode = static_cast<int32>(Input.TrailingEdgeMode);
-			Output.bEnablePerParticleUOverride = Input.bEnablePerParticleUOverride ? 1 : 0;
-			Output.bEnablePerParticleVRangeOverride = Input.bEnablePerParticleVRangeOverride ? 1 : 0;			
-		};
+			constexpr auto AddUVChannelParams = [](const FNiagaraRibbonUVSettings& Input, FNiagaraRibbonUVSettingsParams& Output)
+			{
+				Output.Offset = FVector2f(Input.Offset);
+				Output.Scale = FVector2f(Input.Scale);
+				Output.TilingLength = Input.TilingLength;
+				Output.DistributionMode = static_cast<int32>(Input.DistributionMode);
+				Output.LeadingEdgeMode = static_cast<int32>(Input.LeadingEdgeMode);
+				Output.TrailingEdgeMode = static_cast<int32>(Input.TrailingEdgeMode);
+				Output.bEnablePerParticleUOverride = Input.bEnablePerParticleUOverride ? 1 : 0;
+				Output.bEnablePerParticleVRangeOverride = Input.bEnablePerParticleVRangeOverride ? 1 : 0;			
+			};
 	
-		AddUVChannelParams(UV0Settings, FinalizeParams.UV0Settings);
-		AddUVChannelParams(UV1Settings, FinalizeParams.UV1Settings);
+			AddUVChannelParams(UV0Settings, FinalizeParams.UV0Settings);
+			AddUVChannelParams(UV1Settings, FinalizeParams.UV1Settings);
 	
-		{
-			SCOPE_CYCLE_COUNTER(STAT_NiagaraRenderRibbonsGenVerticesMultiRibbonInitComputeGPU);
+			{
+				SCOPE_CYCLE_COUNTER(STAT_NiagaraRenderRibbonsGenVerticesMultiRibbonInitComputeGPU);
 			
-			FNiagaraRibbonUVParamCalculationCS::FPermutationDomain PermutationVector;
-			PermutationVector.Set<FRibbonHasFullRibbonID>(GenerationConfig.HasFullRibbonIDs());
-			PermutationVector.Set<FRibbonHasRibbonID>(GenerationConfig.HasSimpleRibbonIDs());
-			PermutationVector.Set<FRibbonWantsAutomaticTessellation>(GenerationConfig.WantsAutomaticTessellation());
-			PermutationVector.Set<FRibbonWantsConstantTessellation>(GenerationConfig.WantsConstantTessellation());
+				FNiagaraRibbonUVParamCalculationCS::FPermutationDomain PermutationVector;
+				PermutationVector.Set<FRibbonHasFullRibbonID>(GenerationConfig.HasFullRibbonIDs());
+				PermutationVector.Set<FRibbonHasRibbonID>(GenerationConfig.HasSimpleRibbonIDs());
+				PermutationVector.Set<FRibbonWantsAutomaticTessellation>(GenerationConfig.WantsAutomaticTessellation());
+				PermutationVector.Set<FRibbonWantsConstantTessellation>(GenerationConfig.WantsConstantTessellation());
 
-			TShaderMapRef<FNiagaraRibbonUVParamCalculationCS> ComputeShader(GetGlobalShaderMap(GMaxRHIFeatureLevel), PermutationVector);
+				TShaderMapRef<FNiagaraRibbonUVParamCalculationCS> ComputeShader(GetGlobalShaderMap(GMaxRHIFeatureLevel), PermutationVector);
 
-			// We don't need to transition RibbonLookupTableBuffer as it's still setup for UAV from the last shader
-			RHICmdList.Transition({
-				FRHITransitionInfo(VertexBuffers.RibbonLookupTableBuffer.Buffer, ERHIAccess::UAVCompute, ERHIAccess::UAVCompute),
-				FRHITransitionInfo(VertexBuffers.TangentsAndDistancesBuffer.Buffer, ERHIAccess::SRVMask | ERHIAccess::VertexOrIndexBuffer, ERHIAccess::UAVCompute)});
+				// We don't need to transition RibbonLookupTableBuffer as it's still setup for UAV from the last shader
+				RHICmdList.Transition({
+					FRHITransitionInfo(VertexBuffers.RibbonLookupTableBuffer.Buffer, ERHIAccess::UAVCompute, ERHIAccess::UAVCompute),
+					FRHITransitionInfo(VertexBuffers.TangentsAndDistancesBuffer.Buffer, ERHIAccess::SRVMask | ERHIAccess::VertexOrIndexBuffer, ERHIAccess::UAVCompute)});
 
-			FComputeShaderUtils::DispatchIndirect(RHICmdList, ComputeShader, FinalizeParams, VertexBuffers.GPUComputeCommandBuffer.Buffer, CommandBufferOffset * FNiagaraRibbonCommandBufferLayout::NumElements);
+				FComputeShaderUtils::DispatchIndirect(RHICmdList, ComputeShader, FinalizeParams, VertexBuffers.GPUComputeCommandBuffer.Buffer, CommandBufferOffset * FNiagaraRibbonCommandBufferLayout::NumElements);
 
-			RHICmdList.Transition({
-				FRHITransitionInfo(VertexBuffers.TangentsAndDistancesBuffer.Buffer, ERHIAccess::UAVCompute, ERHIAccess::SRVMask | ERHIAccess::VertexOrIndexBuffer),
-				FRHITransitionInfo(VertexBuffers.RibbonLookupTableBuffer.Buffer, ERHIAccess::UAVCompute, ERHIAccess::SRVMask | ERHIAccess::VertexOrIndexBuffer)});
-		}		
+				RHICmdList.Transition({
+					FRHITransitionInfo(VertexBuffers.TangentsAndDistancesBuffer.Buffer, ERHIAccess::UAVCompute, ERHIAccess::SRVMask | ERHIAccess::VertexOrIndexBuffer),
+					FRHITransitionInfo(VertexBuffers.RibbonLookupTableBuffer.Buffer, ERHIAccess::UAVCompute, ERHIAccess::SRVMask | ERHIAccess::VertexOrIndexBuffer)});
+			}		
+		}
 	}
 }
 
diff --git a/Engine/Plugins/FX/Niagara/Source/NiagaraShader/Private/NiagaraRibbonCompute.cpp b/Engine/Plugins/FX/Niagara/Source/NiagaraShader/Private/NiagaraRibbonCompute.cpp
--- a/Engine/Plugins/FX/Niagara/Source/NiagaraShader/Private/NiagaraRibbonCompute.cpp
+++ b/Engine/Plugins/FX/Niagara/Source/NiagaraShader/Private/NiagaraRibbonCompute.cpp
@@ -39,12 +39,6 @@ void FNiagaraRibbonVertexReductionInitializationCS::ModifyCompilationEnvironment
 	FNiagaraRibbonComputeCommon::ModifyCompilationEnvironment(Parameters, OutEnvironment, FNiagaraRibbonComputeCommon::VertexGenReductionInitializationThreadSize);
 }
 
-void FNiagaraRibbonVertexReductionPropagateCS::ModifyCompilationEnvironment(const FGlobalShaderPermutationParameters& Parameters, FShaderCompilerEnvironment& OutEnvironment)
-{
-	FGlobalShader::ModifyCompilationEnvironment(Parameters, OutEnvironment);
-	FNiagaraRibbonComputeCommon::ModifyCompilationEnvironment(Parameters, OutEnvironment, FNiagaraRibbonComputeCommon::VertexGenReductionPropagationThreadSize);
-}
-
 void FNiagaraRibbonVertexReductionFinalizeCS::ModifyCompilationEnvironment(const FGlobalShaderPermutationParameters& Parameters, FShaderCompilerEnvironment& OutEnvironment)
 {
 	FGlobalShader::ModifyCompilationEnvironment(Parameters, OutEnvironment);
@@ -78,15 +72,27 @@ void FNiagaraRibbonCreateIndexBufferCS::ModifyCompilationEnvironment(const FGlob
 	
 }
 
+void FNiagaraRibbonAggregationStepCS::ModifyCompilationEnvironment(const FGlobalShaderPermutationParameters& Parameters, FShaderCompilerEnvironment& OutEnvironment)
+{
+	FGlobalShader::ModifyCompilationEnvironment(Parameters, OutEnvironment);
+	FNiagaraRibbonComputeCommon::ModifyCompilationEnvironment(Parameters, OutEnvironment, FNiagaraRibbonComputeCommon::VertexGenReductionInitializationThreadSize);
+}
+
+void FNiagaraRibbonAggregationApplyCS::ModifyCompilationEnvironment(const FGlobalShaderPermutationParameters& Parameters, FShaderCompilerEnvironment& OutEnvironment)
+{
+	FGlobalShader::ModifyCompilationEnvironment(Parameters, OutEnvironment);
+	FNiagaraRibbonComputeCommon::ModifyCompilationEnvironment(Parameters, OutEnvironment, FNiagaraRibbonComputeCommon::VertexGenReductionInitializationThreadSize);
+}
 
 IMPLEMENT_GLOBAL_SHADER(FNiagaraRibbonSortPhase1CS, "/Plugin/FX/Niagara/Private/Ribbons/NiagaraRibbonSortParticles.usf", "InitialSortList", SF_Compute);
 IMPLEMENT_GLOBAL_SHADER(FNiagaraRibbonSortPhase2CS, "/Plugin/FX/Niagara/Private/Ribbons/NiagaraRibbonSortParticles.usf", "MergeSort", SF_Compute);
 
 IMPLEMENT_GLOBAL_SHADER(FNiagaraRibbonVertexReductionInitializationCS, "/Plugin/FX/Niagara/Private/Ribbons/NiagaraRibbonVertexReductionInitialization.usf", "VertexGenInitialize", SF_Compute);
-IMPLEMENT_GLOBAL_SHADER(FNiagaraRibbonVertexReductionPropagateCS, "/Plugin/FX/Niagara/Private/Ribbons/NiagaraRibbonVertexReductionPropagation.usf", "VertexGenPrefixSumPropagation", SF_Compute);
 IMPLEMENT_GLOBAL_SHADER(FNiagaraRibbonVertexReductionFinalizeCS, "/Plugin/FX/Niagara/Private/Ribbons/NiagaraRibbonVertexReductionFinalization.usf", "OutputRibbonStats", SF_Compute);
 IMPLEMENT_GLOBAL_SHADER(FNiagaraRibbonUVParamCalculationCS, "/Plugin/FX/Niagara/Private/Ribbons/NiagaraRibbonRibbonUVParamCalculation.usf", "GenerateRibbonUVParams", SF_Compute);
 
 IMPLEMENT_GLOBAL_SHADER(FNiagaraRibbonCreateIndexBufferParamsCS, "/Plugin/FX/Niagara/Private/Ribbons/NiagaraRibbonInitializeIndices.usf", "InitializeIndices", SF_Compute);
 IMPLEMENT_GLOBAL_SHADER(FNiagaraRibbonCreateIndexBufferCS, "/Plugin/FX/Niagara/Private/Ribbons/NiagaraRibbonGenerateIndices.usf", "GenerateIndices", SF_Compute);
 
+IMPLEMENT_GLOBAL_SHADER(FNiagaraRibbonAggregationStepCS, "/Plugin/FX/Niagara/Private/Ribbons/NiagaraRibbonAggregationStep.usf", "RibbonAggregationStep", SF_Compute);
+IMPLEMENT_GLOBAL_SHADER(FNiagaraRibbonAggregationApplyCS, "/Plugin/FX/Niagara/Private/Ribbons/NiagaraRibbonAggregationApply.usf", "RibbonAggregationApply", SF_Compute);
diff --git a/Engine/Plugins/FX/Niagara/Source/NiagaraShader/Public/NiagaraRibbonCompute.h b/Engine/Plugins/FX/Niagara/Source/NiagaraShader/Public/NiagaraRibbonCompute.h
--- a/Engine/Plugins/FX/Niagara/Source/NiagaraShader/Public/NiagaraRibbonCompute.h
+++ b/Engine/Plugins/FX/Niagara/Source/NiagaraShader/Public/NiagaraRibbonCompute.h
@@ -133,6 +133,7 @@ BEGIN_SHADER_PARAMETER_STRUCT(FNiagaraRibbonVertexReductionParameters, NIAGARASH
 	SHADER_PARAMETER_UAV(RWBuffer<uint32>, OutputSegments)
 	SHADER_PARAMETER_SRV(Buffer<float>, InputTessellationStats)
 	SHADER_PARAMETER_UAV(RWBuffer<float>, OutputTessellationStats)
+	SHADER_PARAMETER_UAV(RWBuffer<FRibbonAccumulationValues>, OutputAccumulationBuffer)
 	SHADER_PARAMETER(float, CurveTension)
 	SHADER_PARAMETER(int, PrefixScanStride)
 END_SHADER_PARAMETER_STRUCT()
@@ -148,17 +149,6 @@ class FNiagaraRibbonVertexReductionInitializationCS : public FGlobalShader
 	static void ModifyCompilationEnvironment(const FGlobalShaderPermutationParameters& Parameters, FShaderCompilerEnvironment& OutEnvironment);
 };
 
-class FNiagaraRibbonVertexReductionPropagateCS : public FGlobalShader
-{
-	DECLARE_EXPORTED_GLOBAL_SHADER(FNiagaraRibbonVertexReductionPropagateCS, NIAGARASHADER_API);
-	SHADER_USE_PARAMETER_STRUCT(FNiagaraRibbonVertexReductionPropagateCS, FGlobalShader);
-	
-	using FPermutationDomain = TShaderPermutationDomain<FRibbonHasFullRibbonID, FRibbonHasRibbonID, FRibbonWantsConstantTessellation, FRibbonWantsAutomaticTessellation, FRibbonHasTwist>;
-	using FParameters = FNiagaraRibbonVertexReductionParameters;
-	
-	static void ModifyCompilationEnvironment(const FGlobalShaderPermutationParameters& Parameters, FShaderCompilerEnvironment& OutEnvironment);
-};
-
 BEGIN_SHADER_PARAMETER_STRUCT(FNiagaraRibbonVertexReductionFinalizationParameters, NIAGARASHADER_API)
 	SHADER_PARAMETER_STRUCT_INCLUDE(FRibbonComputeUniformParameters, Common)
 	SHADER_PARAMETER_SRV(Buffer<uint32>, SortedIndices)
@@ -166,6 +156,7 @@ BEGIN_SHADER_PARAMETER_STRUCT(FNiagaraRibbonVertexReductionFinalizationParameter
 	SHADER_PARAMETER_SRV(Buffer<uint32>, MultiRibbonIndices)
 	SHADER_PARAMETER_SRV(Buffer<uint32>, Segments)
 	SHADER_PARAMETER_SRV(Buffer<float>, TessellationStats)
+	SHADER_PARAMETER_SRV(Buffer<FRibbonAccumulationValues>, AccumulationBuffer)
 	SHADER_PARAMETER_UAV(RWBuffer<uint32>, PackedPerRibbonData)
 	SHADER_PARAMETER_UAV(RWBuffer<uint32>, OutputCommandBuffer)
 	SHADER_PARAMETER(int, OutputCommandBufferIndex)
@@ -296,6 +287,45 @@ class FNiagaraRibbonCreateIndexBufferCS : public FGlobalShader
 	static void ModifyCompilationEnvironment(const FGlobalShaderPermutationParameters& Parameters, FShaderCompilerEnvironment& OutEnvironment);
 };
 
+BEGIN_SHADER_PARAMETER_STRUCT(FNiagaraRibbonAggregationStepParameters, NIAGARASHADER_API)
+	SHADER_PARAMETER_STRUCT_INCLUDE(FRibbonComputeUniformParameters, Common)
+	SHADER_PARAMETER_SRV(Buffer<uint32>, SortedIndices)
+	SHADER_PARAMETER_SRV(Buffer<FRibbonAccumulationValues>, InputAccumulation)
+	SHADER_PARAMETER_UAV(RWBuffer<FRibbonAccumulationValues>, OutputAccumulation)
+	SHADER_PARAMETER(uint32, PrefixScanStride)
+END_SHADER_PARAMETER_STRUCT()
+
+class FNiagaraRibbonAggregationStepCS : public FGlobalShader
+{
+	DECLARE_EXPORTED_GLOBAL_SHADER(FNiagaraRibbonAggregationStepCS, NIAGARASHADER_API);
+	SHADER_USE_PARAMETER_STRUCT(FNiagaraRibbonAggregationStepCS, FGlobalShader);
+
+	using FPermutationDomain = TShaderPermutationDomain<FRibbonHasFullRibbonID, FRibbonHasRibbonID, FRibbonWantsConstantTessellation, FRibbonWantsAutomaticTessellation, FRibbonHasTwist>;
+	using FParameters = FNiagaraRibbonAggregationStepParameters;
+
+	static void ModifyCompilationEnvironment(const FGlobalShaderPermutationParameters& Parameters, FShaderCompilerEnvironment& OutEnvironment);
+};
+
+BEGIN_SHADER_PARAMETER_STRUCT(FNiagaraRibbonAggregationApplyParameters, NIAGARASHADER_API)
+	SHADER_PARAMETER_STRUCT_INCLUDE(FRibbonComputeUniformParameters, Common)
+	SHADER_PARAMETER_SRV(Buffer<uint32>, SortedIndices)
+	SHADER_PARAMETER_SRV(Buffer<FRibbonAccumulationValues>, InputAccumulation)
+	SHADER_PARAMETER_UAV(RWBuffer<float>, OutputTangentsAndDistances)
+	SHADER_PARAMETER_UAV(RWBuffer<uint>, OutputMultiRibbonIndices)
+	SHADER_PARAMETER_UAV(RWBuffer<float>, OutputTessellationStats)
+	SHADER_PARAMETER_UAV(RWBuffer<uint>, OutputSegments)
+END_SHADER_PARAMETER_STRUCT()
+
+class FNiagaraRibbonAggregationApplyCS : public FGlobalShader
+{
+	DECLARE_EXPORTED_GLOBAL_SHADER(FNiagaraRibbonAggregationApplyCS, NIAGARASHADER_API);
+	SHADER_USE_PARAMETER_STRUCT(FNiagaraRibbonAggregationApplyCS, FGlobalShader);
+
+	using FPermutationDomain = TShaderPermutationDomain<FRibbonHasFullRibbonID, FRibbonHasRibbonID, FRibbonWantsConstantTessellation, FRibbonWantsAutomaticTessellation, FRibbonHasTwist>;
+	using FParameters = FNiagaraRibbonAggregationApplyParameters;
+
+	static void ModifyCompilationEnvironment(const FGlobalShaderPermutationParameters& Parameters, FShaderCompilerEnvironment& OutEnvironment);
+};
 
 #if UE_ENABLE_INCLUDE_ORDER_DEPRECATED_IN_5_3
 #include "CoreMinimal.h"
