- Fix for rebasing the sim cache when interpolation is enabled
- Only one frames transform was taken into account which would result in jitter on playback if the cached simulation was moving

#tests Create a sequencer with a moving Niagara System, record with interpolation & rebasing enabled, scrub the timeline and observe the jitter

diff --git a/Engine/Plugins/FX/Niagara/Source/Niagara/Classes/NiagaraSimCache.h b/Engine/Plugins/FX/Niagara/Source/Niagara/Classes/NiagaraSimCache.h
--- a/Engine/Plugins/FX/Niagara/Source/Niagara/Classes/NiagaraSimCache.h
+++ b/Engine/Plugins/FX/Niagara/Source/Niagara/Classes/NiagaraSimCache.h
@@ -266,7 +266,8 @@ struct FNiagaraSimCacheDataBuffersLayout
 		const uint8*			SourceBCurr			= nullptr;
 		uint32					SourceBStride		= 0;
 		const uint8*			Velocity			= nullptr;
-		FTransform				RebaseTransform;
+		FTransform				RebaseATransform;
+		FTransform				RebaseBTransform;
 		TConstArrayView<uint32>	InterpMappings;
 	};
 
diff --git a/Engine/Plugins/FX/Niagara/Source/Niagara/Private/NiagaraSimCache.cpp b/Engine/Plugins/FX/Niagara/Source/Niagara/Private/NiagaraSimCache.cpp
--- a/Engine/Plugins/FX/Niagara/Source/Niagara/Private/NiagaraSimCache.cpp
+++ b/Engine/Plugins/FX/Niagara/Source/Niagara/Private/NiagaraSimCache.cpp
@@ -1020,21 +1020,23 @@ bool UNiagaraSimCache::ReadFrame(int32 FrameIndex, float FrameFraction, FNiagara
 	}
 
 	const FNiagaraSimCacheFrame& CacheFrame = CacheFrames[FrameIndex];
+	const int32 NextFrameIndex = FMath::Min(FrameIndex + 1, CacheFrames.Num() - 1);
 
-	FTransform RebaseTransform = FTransform::Identity;
+	FTransform RebaseATransform = FTransform::Identity;
+	FTransform RebaseBTransform = FTransform::Identity;
 	if ( USceneComponent* AttachComponent = SystemInstance->GetAttachComponent() )
 	{
-		RebaseTransform = AttachComponent->GetComponentToWorld();
-		RebaseTransform.AddToTranslation(FVector(SystemInstance->GetLWCTile()) * -FLargeWorldRenderScalar::GetTileSize());
-		RebaseTransform = RebaseTransform * CacheFrame.LocalToWorld.Inverse();
+		FTransform LocalToWorld = AttachComponent->GetComponentToWorld();
+		LocalToWorld.AddToTranslation(FVector(SystemInstance->GetLWCTile()) * -FLargeWorldRenderScalar::GetTileSize());
+		RebaseATransform = LocalToWorld * CacheFrame.LocalToWorld.Inverse();
+		RebaseBTransform = LocalToWorld * CacheFrames[NextFrameIndex].LocalToWorld.Inverse();
 	}
 
-	const int32 NextFrameIndex = FMath::Min(FrameIndex + 1, CacheFrames.Num() - 1);
 	const float FrameDeltaSeconds = CacheFrames[NextFrameIndex].SimulationAge - CacheFrame.SimulationAge;
 	const float SimDeltaSeconds = Helper.SystemInstance->CachedDeltaSeconds;
 
 	Helper.SystemInstance->LocalBounds = CacheFrame.SystemData.LocalBounds;
-	Helper.ReadDataBuffer(FrameFraction, FrameDeltaSeconds, SimDeltaSeconds, RebaseTransform, CacheLayout.SystemLayout, CacheFrame.SystemData.SystemDataBuffers, CacheFrame.SystemData.SystemDataBuffers, Helper.GetSystemSimulationDataSet());
+	Helper.ReadDataBuffer(FrameFraction, FrameDeltaSeconds, SimDeltaSeconds, RebaseATransform, RebaseBTransform, CacheLayout.SystemLayout, CacheFrame.SystemData.SystemDataBuffers, CacheFrame.SystemData.SystemDataBuffers, Helper.GetSystemSimulationDataSet());
 
 	const int32 NumEmitters = CacheLayout.EmitterLayouts.Num();
 	for (int32 i=0; i < NumEmitters; ++i)
@@ -1052,11 +1054,11 @@ bool UNiagaraSimCache::ReadFrame(int32 FrameIndex, float FrameFraction, FNiagara
 		const FNiagaraSimCacheEmitterFrame& CacheEmitterFrameB = CacheFrames[NextFrameIndex].EmitterData[i];
 		if (CacheLayout.EmitterLayouts[i].SimTarget == ENiagaraSimTarget::GPUComputeSim)
 		{
-			Helper.ReadDataBufferGPU(FrameFraction, FrameDeltaSeconds, SimDeltaSeconds, RebaseTransform, EmitterInstance, CacheLayout.EmitterLayouts[i], CacheEmitterFrame.ParticleDataBuffers, CacheEmitterFrameB.ParticleDataBuffers, EmitterInstance.GetParticleData(), PendingCommandsInFlight);
+			Helper.ReadDataBufferGPU(FrameFraction, FrameDeltaSeconds, SimDeltaSeconds, RebaseATransform, RebaseBTransform, EmitterInstance, CacheLayout.EmitterLayouts[i], CacheEmitterFrame.ParticleDataBuffers, CacheEmitterFrameB.ParticleDataBuffers, EmitterInstance.GetParticleData(), PendingCommandsInFlight);
 		}
 		else
 		{
-			Helper.ReadDataBuffer(FrameFraction, FrameDeltaSeconds, SimDeltaSeconds, RebaseTransform, CacheLayout.EmitterLayouts[i], CacheEmitterFrame.ParticleDataBuffers, CacheEmitterFrameB.ParticleDataBuffers, EmitterInstance.GetParticleData());
+			Helper.ReadDataBuffer(FrameFraction, FrameDeltaSeconds, SimDeltaSeconds, RebaseATransform, RebaseBTransform, CacheLayout.EmitterLayouts[i], CacheEmitterFrame.ParticleDataBuffers, CacheEmitterFrameB.ParticleDataBuffers, EmitterInstance.GetParticleData());
 		}
 	}
 
diff --git a/Engine/Plugins/FX/Niagara/Source/Niagara/Private/NiagaraSimCacheHelper.h b/Engine/Plugins/FX/Niagara/Source/Niagara/Private/NiagaraSimCacheHelper.h
--- a/Engine/Plugins/FX/Niagara/Source/Niagara/Private/NiagaraSimCacheHelper.h
+++ b/Engine/Plugins/FX/Niagara/Source/Niagara/Private/NiagaraSimCacheHelper.h
@@ -758,7 +758,7 @@ struct FNiagaraSimCacheHelper
 		}
 	}
 
-	static void ReadCustomBuffers(uint32 NumInstances, float FrameFraction, float FrameDeltaSeconds, float SimDeltaSeconds, const FTransform& RebaseTransform, uint16 ComponentVelocity, const TConstArrayView<FNiagaraSimCacheDataBuffersLayout::FVariableCopyMapping>& VariableCopyMappingsToDataBuffer, const FNiagaraSimCacheDataBuffers& CacheBufferA, const FNiagaraSimCacheDataBuffers& CacheBufferB, uint8* DestBuffer, uint32 DestStride)
+	static void ReadCustomBuffers(uint32 NumInstances, float FrameFraction, float FrameDeltaSeconds, float SimDeltaSeconds, const FTransform& RebaseATransform, const FTransform& RebaseBTransform, uint16 ComponentVelocity, const TConstArrayView<FNiagaraSimCacheDataBuffersLayout::FVariableCopyMapping>& VariableCopyMappingsToDataBuffer, const FNiagaraSimCacheDataBuffers& CacheBufferA, const FNiagaraSimCacheDataBuffers& CacheBufferB, uint8* DestBuffer, uint32 DestStride)
 	{
 		if (VariableCopyMappingsToDataBuffer.Num() == 0)
 		{
@@ -771,7 +771,8 @@ struct FNiagaraSimCacheHelper
 		VariableCopyDataContext.SimDeltaSeconds		= SimDeltaSeconds;
 		VariableCopyDataContext.PrevFrameFraction	= 1.0f - FMath::Clamp(SimDeltaSeconds / (FrameDeltaSeconds * (1.0f + FrameFraction)), 0.0f, 1.0f);		//-TODO: We are assuming both frames have the same DT here
 		VariableCopyDataContext.NumInstances		= NumInstances;
-		VariableCopyDataContext.RebaseTransform		= RebaseTransform;
+		VariableCopyDataContext.RebaseATransform	= RebaseATransform;
+		VariableCopyDataContext.RebaseBTransform	= RebaseBTransform;
 		VariableCopyDataContext.InterpMappings		= CacheBufferA.InterpMapping;
 		VariableCopyDataContext.DestStride			= DestStride;
 		VariableCopyDataContext.SourceAStride		= CacheBufferA.NumInstances * sizeof(float);
@@ -796,7 +797,7 @@ struct FNiagaraSimCacheHelper
 		}
 	}
 
-	void ReadDataBuffer(float FrameFraction, float FrameDeltaSeconds, float SimDeltaSeconds, const FTransform& RebaseTransform, const FNiagaraSimCacheDataBuffersLayout& CacheLayout, const FNiagaraSimCacheDataBuffers& CacheBufferA, const FNiagaraSimCacheDataBuffers& CacheBufferB, FNiagaraDataSet& DataSet)
+	void ReadDataBuffer(float FrameFraction, float FrameDeltaSeconds, float SimDeltaSeconds, const FTransform& RebaseATransform, const FTransform& RebaseBTransform, const FNiagaraSimCacheDataBuffersLayout& CacheLayout, const FNiagaraSimCacheDataBuffers& CacheBufferA, const FNiagaraSimCacheDataBuffers& CacheBufferB, FNiagaraDataSet& DataSet)
 	{
 		const uint32 NumInstances = FMath::Min(CacheBufferA.NumInstances, DataSet.GetMaxAllocationCount());
 
@@ -809,7 +810,7 @@ struct FNiagaraSimCacheHelper
 			ReadFloatBuffers(NumInstances, iComponent, CacheLayout.FloatCount, CacheLayout.CacheBufferReadInfo.ComponentMappingsToDataBuffer, CacheBufferA, MakeArrayView(DataBuffer.GetComponentPtrFloat(0), DataBuffer.GetFloatBuffer().Num()), DataBuffer.GetFloatStride());
 			ReadHalfBuffers(NumInstances, iComponent, CacheLayout.HalfCount, CacheLayout.CacheBufferReadInfo.ComponentMappingsToDataBuffer, CacheBufferA, MakeArrayView(DataBuffer.GetComponentPtrHalf(0), DataBuffer.GetHalfBuffer().Num()), DataBuffer.GetHalfStride());
 			ReadInt32Buffers(NumInstances, iComponent, CacheLayout.Int32Count, CacheLayout.CacheBufferReadInfo.ComponentMappingsToDataBuffer, CacheBufferA, MakeArrayView(DataBuffer.GetComponentPtrInt32(0), DataBuffer.GetInt32Buffer().Num()), DataBuffer.GetInt32Stride());
-			ReadCustomBuffers(NumInstances, FrameFraction, FrameDeltaSeconds, SimDeltaSeconds, RebaseTransform, CacheLayout.ComponentVelocity, CacheLayout.CacheBufferReadInfo.VariableCopyMappingsToDataBuffer, CacheBufferA, CacheBufferB, DataBuffer.GetComponentPtrFloat(0), DataBuffer.GetFloatStride());
+			ReadCustomBuffers(NumInstances, FrameFraction, FrameDeltaSeconds, SimDeltaSeconds, RebaseATransform, RebaseBTransform, CacheLayout.ComponentVelocity, CacheLayout.CacheBufferReadInfo.VariableCopyMappingsToDataBuffer, CacheBufferA, CacheBufferB, DataBuffer.GetComponentPtrFloat(0), DataBuffer.GetFloatStride());
 		}
 
 		//-TODO:DestinationDataBuffer.SetIDTable(CacheBufferA.IDToIndexTable);
@@ -818,7 +819,7 @@ struct FNiagaraSimCacheHelper
 		DataSet.EndSimulate();
 	}
 
-	void ReadDataBufferGPU(float InFrameFraction, float InFrameDeltaSeconds, float InSimDeltaSeconds, const FTransform& InRebaseTransform, FNiagaraEmitterInstance& EmitterInstance, const FNiagaraSimCacheDataBuffersLayout& InCacheLayout, const FNiagaraSimCacheDataBuffers& InCacheBufferA, const FNiagaraSimCacheDataBuffers& InCacheBufferB, FNiagaraDataSet& InDataSet, std::atomic<int32>& InPendingCommandsCounter)
+	void ReadDataBufferGPU(float InFrameFraction, float InFrameDeltaSeconds, float InSimDeltaSeconds, const FTransform& InRebaseATransform, const FTransform& InRebaseBTransform, FNiagaraEmitterInstance& EmitterInstance, const FNiagaraSimCacheDataBuffersLayout& InCacheLayout, const FNiagaraSimCacheDataBuffers& InCacheBufferA, const FNiagaraSimCacheDataBuffers& InCacheBufferB, FNiagaraDataSet& InDataSet, std::atomic<int32>& InPendingCommandsCounter)
 	{
 		if (EmitterInstance.IsDisabled())
 		{
@@ -843,7 +844,7 @@ struct FNiagaraSimCacheHelper
 
 		FNiagaraGpuComputeDispatchInterface* DispathInterface = EmitterInstance.GetParentSystemInstance()->GetComputeDispatchInterface();
 		ENQUEUE_RENDER_COMMAND(NiagaraSimCacheGpuReadFrame)(
-			[DispathInterface, GPUExecContext=EmitterInstance.GetGPUContext(), FrameFraction=InFrameFraction, FrameDeltaSeconds=InFrameDeltaSeconds, SimDeltaSeconds=InSimDeltaSeconds, RebaseTransform=InRebaseTransform, CacheLayout=&InCacheLayout, CacheBufferA=&InCacheBufferA, CacheBufferB=&InCacheBufferB, DataSet=&InDataSet, PendingCommandsCounter=&InPendingCommandsCounter](FRHICommandListImmediate& RHICmdList)
+			[DispathInterface, GPUExecContext=EmitterInstance.GetGPUContext(), FrameFraction=InFrameFraction, FrameDeltaSeconds=InFrameDeltaSeconds, SimDeltaSeconds=InSimDeltaSeconds, RebaseATransform=InRebaseATransform, RebaseBTransform=InRebaseBTransform, CacheLayout=&InCacheLayout, CacheBufferA=&InCacheBufferA, CacheBufferB=&InCacheBufferB, DataSet=&InDataSet, PendingCommandsCounter=&InPendingCommandsCounter](FRHICommandListImmediate& RHICmdList)
 			{
 				const uint32 NumInstances = FMath::Min(CacheBufferA->NumInstances, DataSet->GetMaxAllocationCount());
 
@@ -879,7 +880,7 @@ struct FNiagaraSimCacheHelper
 						FRWBuffer& RWBuffer = DataBuffer.GetGPUBufferFloat();
 						uint8* RWBufferMemory = reinterpret_cast<uint8*>(RHICmdList.LockBuffer(RWBuffer.Buffer, 0, RWBuffer.NumBytes, RLM_WriteOnly));
 						ReadFloatBuffers(NumInstances, iComponent, CacheLayout->FloatCount, CacheLayout->CacheBufferReadInfo_RT.ComponentMappingsToDataBuffer, *CacheBufferA, MakeArrayView(RWBufferMemory, RWBuffer.NumBytes), DataBuffer.GetFloatStride());
-						ReadCustomBuffers(NumInstances, FrameFraction, FrameDeltaSeconds, SimDeltaSeconds, RebaseTransform, CacheLayout->ComponentVelocity, CacheLayout->CacheBufferReadInfo_RT.VariableCopyMappingsToDataBuffer, *CacheBufferA, *CacheBufferB, RWBufferMemory, DataBuffer.GetFloatStride());
+						ReadCustomBuffers(NumInstances, FrameFraction, FrameDeltaSeconds, SimDeltaSeconds, RebaseATransform, RebaseBTransform, CacheLayout->ComponentVelocity, CacheLayout->CacheBufferReadInfo_RT.VariableCopyMappingsToDataBuffer, *CacheBufferA, *CacheBufferB, RWBufferMemory, DataBuffer.GetFloatStride());
 						RHICmdList.UnlockBuffer(RWBuffer.Buffer);
 					}
 
@@ -938,7 +939,7 @@ struct FNiagaraSimCacheHelper
 			const FVector3f Position = SrcPositions.Get(i);
 			const FVector3f Velocity = SrcVelocities.Get(i);
 			const FVector3f ExtrapolatedPosition = Position + (Velocity * CopyDataContext.FrameFraction * CopyDataContext.FrameDeltaSeconds);
-			const FVector3f RebasedPosition = bWithRebase ? FVector3f(CopyDataContext.RebaseTransform.TransformPosition(FVector(ExtrapolatedPosition))) : ExtrapolatedPosition;
+			const FVector3f RebasedPosition = bWithRebase ? FVector3f(CopyDataContext.RebaseATransform.TransformPosition(FVector(ExtrapolatedPosition))) : ExtrapolatedPosition;
 			DstPositions.Set(i, RebasedPosition);
 		}
 	}
@@ -958,14 +959,21 @@ struct FNiagaraSimCacheHelper
 			for (uint32 iInstanceA = 0; iInstanceA < CopyDataContext.NumInstances; ++iInstanceA)
 			{
 				const uint32 iInstanceB = CopyDataContext.InterpMappings[iInstanceA];
-				const FVector3f CurrPosA = SrcAPositions.Get(iInstanceA);
-				const FVector3f CurrPosB = iInstanceB == INDEX_NONE ? CurrPosA : SrcBPositions.Get(iInstanceB);
-				const FVector3f CurrPos = FMath::Lerp(CurrPosA, CurrPosB, CopyDataContext.FrameFraction);
-				DstPositions.Set(iInstanceA, bWithRebase ? FVector3f(CopyDataContext.RebaseTransform.TransformPosition(FVector(CurrPos))) : CurrPos);
+				FVector3f CurrPosA = SrcAPositions.Get(iInstanceA);
+				FVector3f CurrPosB = iInstanceB == INDEX_NONE ? CurrPosA : SrcBPositions.Get(iInstanceB);
+				FVector3f PrevPosA = SrcAPrevPositions.Get(iInstanceA);
+				if ( bWithRebase )
+				{
+					CurrPosA = FVector3f(CopyDataContext.RebaseATransform.TransformPosition(FVector(CurrPosA)));
+					CurrPosB = FVector3f(CopyDataContext.RebaseBTransform.TransformPosition(FVector(CurrPosB)));
+					PrevPosA = FVector3f(CopyDataContext.RebaseATransform.TransformPosition(FVector(PrevPosA)));
+				}
+
 
-				const FVector3f PrevPosA = SrcAPrevPositions.Get(iInstanceA);
+				const FVector3f CurrPos = FMath::Lerp(CurrPosA, CurrPosB, CopyDataContext.FrameFraction);
 				const FVector3f PrevPos = FMath::Lerp(PrevPosA, CurrPos, CopyDataContext.PrevFrameFraction);
-				DstPrevPositions.Set(iInstanceA, bWithRebase ? FVector3f(CopyDataContext.RebaseTransform.TransformPosition(FVector(PrevPos))) : PrevPos);
+				DstPositions.Set(iInstanceA, CurrPos);
+				DstPrevPositions.Set(iInstanceA, PrevPos);
 			}
 		}
 		else
@@ -973,10 +981,15 @@ struct FNiagaraSimCacheHelper
 			for (uint32 iInstanceA = 0; iInstanceA < CopyDataContext.NumInstances; ++iInstanceA)
 			{
 				const uint32 iInstanceB = CopyDataContext.InterpMappings[iInstanceA];
-				const FVector3f CurrPosA = SrcAPositions.Get(iInstanceA);
-				const FVector3f CurrPosB = iInstanceB == INDEX_NONE ? CurrPosA : SrcBPositions.Get(iInstanceB);
-				const FVector3f CurrPos  = FMath::Lerp(CurrPosA, CurrPosB, CopyDataContext.FrameFraction);
-				DstPositions.Set(iInstanceA, bWithRebase ? FVector3f(CopyDataContext.RebaseTransform.TransformPosition(FVector(CurrPos))) : CurrPos);
+				FVector3f CurrPosA = SrcAPositions.Get(iInstanceA);
+				FVector3f CurrPosB = iInstanceB == INDEX_NONE ? CurrPosA : SrcBPositions.Get(iInstanceB);
+				if (bWithRebase)
+				{
+					CurrPosA = FVector3f(CopyDataContext.RebaseATransform.TransformPosition(FVector(CurrPosA)));
+					CurrPosB = FVector3f(CopyDataContext.RebaseBTransform.TransformPosition(FVector(CurrPosB)));
+				}
+				const FVector3f CurrPos = FMath::Lerp(CurrPosA, CurrPosB, CopyDataContext.FrameFraction);
+				DstPositions.Set(iInstanceA, CurrPos);
 			}
 		}
 	}
@@ -989,7 +1002,7 @@ struct FNiagaraSimCacheHelper
 		for (uint32 i = 0; i < CopyDataContext.NumInstances; ++i)
 		{
 			const FVector3f CachePosition = SrcPositions.Get(i);
-			const FVector3f RebasedPosition = FVector3f(CopyDataContext.RebaseTransform.TransformPosition(FVector(CachePosition)));
+			const FVector3f RebasedPosition = FVector3f(CopyDataContext.RebaseATransform.TransformPosition(FVector(CachePosition)));
 			DstPositions.Set(i, RebasedPosition);
 		}
 	}
@@ -1017,7 +1030,8 @@ struct FNiagaraSimCacheHelper
 		FSoAQuatReader SrcBQuats(CopyDataContext.SourceBCurr, CopyDataContext.SourceBStride);
 		FSoAQuatWriter DstQuats(CopyDataContext.DestCurr, CopyDataContext.DestStride);
 
-		const FQuat4f RebaseQuat(CopyDataContext.RebaseTransform.GetRotation());
+		const FQuat4f RebaseAQuat(CopyDataContext.RebaseATransform.GetRotation());
+		const FQuat4f RebaseBQuat(CopyDataContext.RebaseBTransform.GetRotation());
 		if (bWithPrevious)
 		{
 			FSoAQuatReader SrcAPrevQuats(CopyDataContext.SourceAPrev, CopyDataContext.SourceAStride);
@@ -1026,14 +1040,19 @@ struct FNiagaraSimCacheHelper
 			for (uint32 iInstanceA = 0; iInstanceA < CopyDataContext.NumInstances; ++iInstanceA)
 			{
 				const uint32 iInstanceB = CopyDataContext.InterpMappings[iInstanceA];
-				const FQuat4f CurrQuatA = SrcAQuats.Get(iInstanceA);
-				const FQuat4f CurrQuatB = iInstanceB == INDEX_NONE ? CurrQuatA : SrcBQuats.Get(iInstanceB);
+				FQuat4f CurrQuatA = SrcAQuats.Get(iInstanceA);
+				FQuat4f CurrQuatB = iInstanceB == INDEX_NONE ? CurrQuatA : SrcBQuats.Get(iInstanceB);
+				FQuat4f PrevQuatA = SrcAPrevQuats.Get(iInstanceA);
+				if (bWithRebase)
+				{
+					CurrQuatA = CurrQuatA * RebaseAQuat;
+					CurrQuatB = CurrQuatB * RebaseBQuat;
+					PrevQuatA = PrevQuatA * RebaseAQuat;
+				}
 				const FQuat4f CurrQuat  = FQuat4f::Slerp(CurrQuatA, CurrQuatB, CopyDataContext.FrameFraction);
-				DstQuats.Set(iInstanceA, bWithRebase ? CurrQuat * RebaseQuat : CurrQuat);
-
-				const FQuat4f PrevQuatA = SrcAPrevQuats.Get(iInstanceA);
 				const FQuat4f PrevQuat = FQuat4f::Slerp(PrevQuatA, CurrQuat, CopyDataContext.PrevFrameFraction);
-				DstPrevQuats.Set(iInstanceA, bWithRebase ? PrevQuat * RebaseQuat : PrevQuat);
+				DstQuats.Set(iInstanceA, CurrQuat);
+				DstPrevQuats.Set(iInstanceA, PrevQuat);
 			}
 		}
 		else
@@ -1041,10 +1060,15 @@ struct FNiagaraSimCacheHelper
 			for (uint32 iInstanceA=0; iInstanceA < CopyDataContext.NumInstances; ++iInstanceA)
 			{
 				const uint32 iInstanceB = CopyDataContext.InterpMappings[iInstanceA];
-				const FQuat4f CurrQuatA = SrcAQuats.Get(iInstanceA);
-				const FQuat4f CurrQuatB = iInstanceB == INDEX_NONE ? CurrQuatA : SrcBQuats.Get(iInstanceB);
+				FQuat4f CurrQuatA = SrcAQuats.Get(iInstanceA);
+				FQuat4f CurrQuatB = iInstanceB == INDEX_NONE ? CurrQuatA : SrcBQuats.Get(iInstanceB);
+				if (bWithRebase)
+				{
+					CurrQuatA = CurrQuatA * RebaseAQuat;
+					CurrQuatB = CurrQuatB * RebaseBQuat;
+				}
 				const FQuat4f CurrQuat = FQuat4f::Slerp(CurrQuatA, CurrQuatB, CopyDataContext.FrameFraction);
-				DstQuats.Set(iInstanceA, bWithRebase ? CurrQuat * RebaseQuat : CurrQuat);
+				DstQuats.Set(iInstanceA, CurrQuat);
 			}
 		}
 	}
@@ -1054,7 +1078,7 @@ struct FNiagaraSimCacheHelper
 		FSoAQuatReader SrcQuats(CopyDataContext.SourceACurr, CopyDataContext.SourceAStride);
 		FSoAQuatWriter DstQuats(CopyDataContext.DestCurr, CopyDataContext.DestStride);
 
-		const FQuat4f RebaseQuat(CopyDataContext.RebaseTransform.GetRotation());
+		const FQuat4f RebaseQuat(CopyDataContext.RebaseATransform.GetRotation());
 		for (uint32 i = 0; i < CopyDataContext.NumInstances; ++i)
 		{
 			const FQuat4f CacheRotation = SrcQuats.Get(i);
@@ -1071,7 +1095,7 @@ struct FNiagaraSimCacheHelper
 		const float* SrcFloats = reinterpret_cast<const float*>(CopyDataContext.SourceACurr);
 		const uint32 SrcStride = CopyDataContext.SourceAStride >> 2;
 
-		const FMatrix44d RebaseMatrix = CopyDataContext.RebaseTransform.ToMatrixWithScale();
+		const FMatrix44d RebaseMatrix = CopyDataContext.RebaseATransform.ToMatrixWithScale();
 		for (uint32 i = 0; i < CopyDataContext.NumInstances; ++i)
 		{
 			FMatrix44d CacheMatrix;
