Added a log message when ticking a niagara component during postload

diff --git a/Engine/Plugins/FX/Niagara/Source/Niagara/Private/NiagaraComponent.cpp b/Engine/Plugins/FX/Niagara/Source/Niagara/Private/NiagaraComponent.cpp
--- a/Engine/Plugins/FX/Niagara/Source/Niagara/Private/NiagaraComponent.cpp
+++ b/Engine/Plugins/FX/Niagara/Source/Niagara/Private/NiagaraComponent.cpp
@@ -1649,6 +1649,10 @@ void UNiagaraComponent::PostSystemTick_GameThread()
 	if (FUObjectThreadContext::Get().IsRoutingPostLoad)
 	{
 		// return early when we're in the middle of another object's postload phase, as ticking data interfaces and renderers might lead to failed assertions
+#if WITH_EDITOR || DO_CHECK
+		UE_LOG(LogNiagara, Warning, TEXT("Niagara component %s PostSystemTick_GameThread was called during PostLoad. This should not happen, as it can lead to dangling RT pointers."), *GetPathName());
+		FDebug::DumpStackTraceToLog(ELogVerbosity::Warning);
+#endif
 		return;
 	}
 
