[Niagara] Expose NDC writer generic function

#rnx

diff --git a/Engine/Plugins/FX/Niagara/Source/Niagara/Private/NiagaraDataChannelAccessor.cpp b/Engine/Plugins/FX/Niagara/Source/Niagara/Private/NiagaraDataChannelAccessor.cpp
--- a/Engine/Plugins/FX/Niagara/Source/Niagara/Private/NiagaraDataChannelAccessor.cpp
+++ b/Engine/Plugins/FX/Niagara/Source/Niagara/Private/NiagaraDataChannelAccessor.cpp
@@ -130,18 +130,6 @@ FNiagaraSpawnInfo UNiagaraDataChannelReader::ReadSpawnInfo(FName VarName, int32
 
 //////////////////////////////////////////////////////////////////////////
 
-template<typename T>
-void UNiagaraDataChannelWriter::WriteData(const FNiagaraVariableBase& Var, int32 Index, const T& InData)
-{
-	if (ensure(Data.IsValid()))
-	{
-		if (FNiagaraDataChannelVariableBuffer* VarBuffer = Data->FindVariableBuffer(Var))
-		{
-			VarBuffer->Write<T>(Index, InData);
-		}
-	}
-}
-
 bool UNiagaraDataChannelWriter::InitWrite(FNiagaraDataChannelSearchParameters SearchParams, int32 Count, bool bVisibleToGame, bool bVisibleToCPU, bool bVisibleToGPU, const FString& DebugSource)
 {
 	if (Count == 0)
diff --git a/Engine/Plugins/FX/Niagara/Source/Niagara/Public/NiagaraDataChannelAccessor.h b/Engine/Plugins/FX/Niagara/Source/Niagara/Public/NiagaraDataChannelAccessor.h
--- a/Engine/Plugins/FX/Niagara/Source/Niagara/Public/NiagaraDataChannelAccessor.h
+++ b/Engine/Plugins/FX/Niagara/Source/Niagara/Public/NiagaraDataChannelAccessor.h
@@ -86,10 +86,19 @@ private:
 	/** Local data buffers we're writing into. */
 	FNiagaraDataChannelGameDataPtr Data = nullptr;
 
-	template<typename T>
-	void WriteData(const FNiagaraVariableBase& Var, int32 Index, const T& InData);
 
 public:
+	template<typename T>
+	void WriteData(const FNiagaraVariableBase& Var, int32 Index, const T& InData)
+	{
+		if (ensure(Data.IsValid()))
+		{
+			if (FNiagaraDataChannelVariableBuffer* VarBuffer = Data->FindVariableBuffer(Var))
+			{
+				VarBuffer->Write<T>(Index, InData);
+			}
+		}
+	}
 
 	UPROPERTY()
 	TObjectPtr<UNiagaraDataChannelHandler> Owner;
