Fix check() firing with the new compilation mode when pins weren't connected (script will still fail to compile)
#tests Create an empty system, with a scratch pad adding two float inputs and writing the results to an output float...disconnect one of the connections between the input and the add

diff --git a/Engine/Plugins/FX/Niagara/Source/NiagaraEditor/Private/NiagaraGraphDigest.cpp b/Engine/Plugins/FX/Niagara/Source/NiagaraEditor/Private/NiagaraGraphDigest.cpp
--- a/Engine/Plugins/FX/Niagara/Source/NiagaraEditor/Private/NiagaraGraphDigest.cpp
+++ b/Engine/Plugins/FX/Niagara/Source/NiagaraEditor/Private/NiagaraGraphDigest.cpp
@@ -1070,12 +1070,16 @@ void FNiagaraCompilationGraphInstanced::ValidateRefinement() const
 
 		for (const FNiagaraCompilationInputPin& InputPin : Node->InputPins)
 		{
-			check(InputPin.Variable.GetType() != FNiagaraTypeDefinition::GetGenericNumericDef());
+			ensureMsgf(InputPin.Variable.GetType() != FNiagaraTypeDefinition::GetGenericNumericDef() || !InputPin.LinkedTo,
+				TEXT("Failed during ValidateRefinement for compilation task - %s.  Connected InputPin[%s.%s] is still a generic."),
+				*SourceScriptFullName, *InputPin.OwningNode->NodeName, *InputPin.PinName.ToString());
 		}
 
 		for (const FNiagaraCompilationOutputPin& OutputPin : Node->OutputPins)
 		{
-			check(OutputPin.Variable.GetType() != FNiagaraTypeDefinition::GetGenericNumericDef());
+			ensureMsgf(OutputPin.Variable.GetType() != FNiagaraTypeDefinition::GetGenericNumericDef() || OutputPin.LinkedTo.IsEmpty(),
+				TEXT("Failed during ValidateRefinement for compilation task - %s.  Connected OutputPin[%s.%s] is still a generic."),
+				*SourceScriptFullName, *OutputPin.OwningNode->NodeName, *OutputPin.PinName.ToString());
 		}
 	}
 }
