Niagara - UI tweaks for the script editor in graph search bar.
+ Remove the white background border.
+ Change the alignment so it doesn't stretch vertically.
+ Make the default minimum size wider to allow for a resonable amount of text.
+ Move the search result navigation controls from the left side to the right side so that the text you're entering doesn't move as results change.
+ Have the script toolkit forward the search command to the graph when nothing is focused.

diff --git a/Engine/Plugins/FX/Niagara/Source/NiagaraEditor/Private/NiagaraEditorStyle.cpp b/Engine/Plugins/FX/Niagara/Source/NiagaraEditor/Private/NiagaraEditorStyle.cpp
--- a/Engine/Plugins/FX/Niagara/Source/NiagaraEditor/Private/NiagaraEditorStyle.cpp
+++ b/Engine/Plugins/FX/Niagara/Source/NiagaraEditor/Private/NiagaraEditorStyle.cpp
@@ -673,6 +673,11 @@ void FNiagaraEditorStyle::InitScriptGraph()
 {
 	Set("NiagaraEditor.ScriptGraph.SearchBorderColor", FLinearColor(.1f, .1f, .1f, 1.f));
 	Set("NiagaraEditor.ScriptGraph.AffectedAssetsWarningColor", FLinearColor(FColor(255, 184, 0)));
+
+	FSearchBoxStyle SearchBoxStyle = FCoreStyle::Get().GetWidgetStyle<FSearchBoxStyle>("SearchBox");
+	SearchBoxStyle.SetLeftAlignGlassImageAndClearButton(true);
+	SearchBoxStyle.SetLeftAlignSearchResultButtons(false);
+	Set("NiagaraEditor.ScriptGraph.SearchBoxStyle", SearchBoxStyle);
 }
 
 void FNiagaraEditorStyle::InitDebuggerStyle()
diff --git a/Engine/Plugins/FX/Niagara/Source/NiagaraEditor/Private/Toolkits/NiagaraScriptToolkit.cpp b/Engine/Plugins/FX/Niagara/Source/NiagaraEditor/Private/Toolkits/NiagaraScriptToolkit.cpp
--- a/Engine/Plugins/FX/Niagara/Source/NiagaraEditor/Private/Toolkits/NiagaraScriptToolkit.cpp
+++ b/Engine/Plugins/FX/Niagara/Source/NiagaraEditor/Private/Toolkits/NiagaraScriptToolkit.cpp
@@ -722,8 +722,11 @@ void FNiagaraScriptToolkit::SetupCommands()
 		FNiagaraEditorCommands::Get().RefreshNodes,
 		FExecuteAction::CreateRaw(this, &FNiagaraScriptToolkit::RefreshNodes));
 	GetToolkitCommands()->MapAction(
-        FNiagaraEditorCommands::Get().ModuleVersioning,
-        FExecuteAction::CreateSP(this, &FNiagaraScriptToolkit::ManageVersions));
+		FNiagaraEditorCommands::Get().ModuleVersioning,
+		FExecuteAction::CreateSP(this, &FNiagaraScriptToolkit::ManageVersions));
+	GetToolkitCommands()->MapAction(
+		FNiagaraEditorCommands::Get().FindInCurrentView,
+		FExecuteAction::CreateRaw(this, &FNiagaraScriptToolkit::FindInCurrentView));
 }
 
 const FName FNiagaraScriptToolkit::GetNiagaraScriptMessageLogName(FVersionedNiagaraScript InScript) const
@@ -924,6 +927,14 @@ void FNiagaraScriptToolkit::ManageVersions()
 	TabManager->TryInvokeTab(VersioningTabID);
 }
 
+void FNiagaraScriptToolkit::FindInCurrentView() const
+{
+	if (NiagaraScriptGraphWidget.IsValid())
+	{
+		NiagaraScriptGraphWidget->FocusGraphSearchBox();
+	}
+}
+
 void FNiagaraScriptToolkit::SwitchToVersion(FGuid VersionGuid)
 {
 	EditedNiagaraScript.Version = VersionGuid;
diff --git a/Engine/Plugins/FX/Niagara/Source/NiagaraEditor/Private/Toolkits/NiagaraScriptToolkit.h b/Engine/Plugins/FX/Niagara/Source/NiagaraEditor/Private/Toolkits/NiagaraScriptToolkit.h
--- a/Engine/Plugins/FX/Niagara/Source/NiagaraEditor/Private/Toolkits/NiagaraScriptToolkit.h
+++ b/Engine/Plugins/FX/Niagara/Source/NiagaraEditor/Private/Toolkits/NiagaraScriptToolkit.h
@@ -146,6 +146,9 @@ private:
 	/** Opens the module versioning tab. */
 	void ManageVersions();
 
+	/** Focuses the search bar in the graph. */
+	void FindInCurrentView() const;
+
 	void InitViewWithVersionedData();
 	void SwitchToVersion(FGuid VersionGuid);
 	bool IsVersionSelected(FNiagaraAssetVersion Version) const;
diff --git a/Engine/Plugins/FX/Niagara/Source/NiagaraEditor/Private/Widgets/SNiagaraScriptGraph.cpp b/Engine/Plugins/FX/Niagara/Source/NiagaraEditor/Private/Widgets/SNiagaraScriptGraph.cpp
--- a/Engine/Plugins/FX/Niagara/Source/NiagaraEditor/Private/Widgets/SNiagaraScriptGraph.cpp
+++ b/Engine/Plugins/FX/Niagara/Source/NiagaraEditor/Private/Widgets/SNiagaraScriptGraph.cpp
@@ -173,52 +173,48 @@ TSharedRef<SGraphEditor> SNiagaraScriptGraph::ConstructGraphEditor()
 			// Search Box
 			+SOverlay::Slot()
 			.HAlign(HAlign_Right)
-			.VAlign(VAlign_Fill)
+			.VAlign(VAlign_Top)
+			.Padding(3)
 			[
-				SNew(SBorder)
-				.BorderImage(FAppStyle::GetBrush("WhiteBrush"))
-				.BorderBackgroundColor(FNiagaraEditorStyle::Get().GetColor("NiagaraEditor.ScriptGraph.SearchBorderColor"))
+				SNew(SHorizontalBox)
 				.Visibility(this, &SNiagaraScriptGraph::GetGraphSearchBoxVisibility)
-				.Padding(5)
+				+ SHorizontalBox::Slot()
+				.HAlign(HAlign_Fill)
+				.VAlign(VAlign_Fill)
+				.AutoWidth()
 				[
-					SNew(SHorizontalBox)
-					+ SHorizontalBox::Slot()
-					.HAlign(HAlign_Fill)
-					.VAlign(VAlign_Fill)
-					.AutoWidth()
+					SNew(SBox)
+					.WidthOverride(250)
 					[
-						SNew(SBox)
-						.WidthOverride(200)
-						[
-							SAssignNew(SearchBox, SSearchBox)
-							.HintText(LOCTEXT("GraphSearchBoxHint", "Search Nodes and Pins in Graph"))
-							.SearchResultData(this, &SNiagaraScriptGraph::GetSearchResultData)
-							.OnTextChanged(this, &SNiagaraScriptGraph::OnSearchTextChanged)
-							.OnTextCommitted(this, &SNiagaraScriptGraph::OnSearchBoxTextCommitted)
-							.DelayChangeNotificationsWhileTyping(true)
-							.OnSearch(this, &SNiagaraScriptGraph::OnSearchBoxSearch)
-							.OnKeyDownHandler(this, &SNiagaraScriptGraph::HandleGraphSearchBoxKeyDown)
-						]
+						SAssignNew(SearchBox, SSearchBox)
+						.Style(FNiagaraEditorStyle::Get(), "NiagaraEditor.ScriptGraph.SearchBoxStyle")
+						.HintText(LOCTEXT("GraphSearchBoxHint", "Search Nodes and Pins in Graph"))
+						.SearchResultData(this, &SNiagaraScriptGraph::GetSearchResultData)
+						.OnTextChanged(this, &SNiagaraScriptGraph::OnSearchTextChanged)
+						.OnTextCommitted(this, &SNiagaraScriptGraph::OnSearchBoxTextCommitted)
+						.DelayChangeNotificationsWhileTyping(true)
+						.OnSearch(this, &SNiagaraScriptGraph::OnSearchBoxSearch)
+						.OnKeyDownHandler(this, &SNiagaraScriptGraph::HandleGraphSearchBoxKeyDown)
 					]
-					+SHorizontalBox::Slot()
-					.HAlign(HAlign_Center)
-					.VAlign(VAlign_Center)
-					.AutoWidth()
-					.Padding(2.0f, 0.0f, 0.0f, 0.0f)
+				]
+				+SHorizontalBox::Slot()
+				.HAlign(HAlign_Center)
+				.VAlign(VAlign_Center)
+				.AutoWidth()
+				.Padding(2.0f, 0.0f, 0.0f, 0.0f)
+				[
+					SNew(SButton)
+					.ButtonStyle(FAppStyle::Get(), "HoverHintOnly")
+					.IsFocusable(false)
+					.ForegroundColor(FNiagaraEditorStyle::Get().GetColor("NiagaraEditor.Stack.FlatButtonColor"))
+					.ToolTipText(LOCTEXT("CloseGraphSearchBox", "Close Graph search box"))
+					.OnClicked(this, &SNiagaraScriptGraph::CloseGraphSearchBoxPressed)
+					.ContentPadding(3)
+					.Content()
 					[
-						SNew(SButton)
-						.ButtonStyle(FAppStyle::Get(), "HoverHintOnly")
-						.IsFocusable(false)
-						.ForegroundColor(FNiagaraEditorStyle::Get().GetColor("NiagaraEditor.Stack.FlatButtonColor"))
-						.ToolTipText(LOCTEXT("CloseGraphSearchBox", "Close Graph search box"))
-						.OnClicked(this, &SNiagaraScriptGraph::CloseGraphSearchBoxPressed)
-						.ContentPadding(3)
-						.Content()
-						[
-							SNew(STextBlock)
-							.Font(FAppStyle::Get().GetFontStyle("FontAwesome.10"))
-							.Text(FEditorFontGlyphs::Times)
-						]
+						SNew(STextBlock)
+						.Font(FAppStyle::Get().GetFontStyle("FontAwesome.10"))
+						.Text(FEditorFontGlyphs::Times)
 					]
 				]
 			]
