Fixed typo in tooltip

#rnx

#virtualized

diff --git a/Engine/Plugins/FX/Niagara/Source/NiagaraEditor/Private/ViewModels/Stack/NiagaraStackFunctionInput.cpp b/Engine/Plugins/FX/Niagara/Source/NiagaraEditor/Private/ViewModels/Stack/NiagaraStackFunctionInput.cpp
--- a/Engine/Plugins/FX/Niagara/Source/NiagaraEditor/Private/ViewModels/Stack/NiagaraStackFunctionInput.cpp
+++ b/Engine/Plugins/FX/Niagara/Source/NiagaraEditor/Private/ViewModels/Stack/NiagaraStackFunctionInput.cpp
@@ -1020,7 +1020,7 @@ void UNiagaraStackFunctionInput::RefreshChildrenInternal(const TArray<UNiagaraSt
 						FFormatNamedArguments Args;
 						Args.Add(TEXT("Function"), FText::FromString(InputValues.DynamicNode->GetFunctionName()));
 						Args.Add(TEXT("Message"), ScriptData->ExperimentalMessage);
-						ErrorMessage = FText::Format(LOCTEXT("DynamicInputScriptExperimentalReason", "The script asset for the dynamic input {Function} is experimental, reason: {Message}."), Args);
+						ErrorMessage = FText::Format(LOCTEXT("DynamicInputScriptExperimentalReason", "The script asset for the dynamic input {Function} is experimental, reason: {Message}"), Args);
 					}
 
 					NewIssues.Add(FStackIssue(
