Heavy handed approach to clean up lingering issues with RI parameters and the new compilation mode

diff --git a/Engine/Plugins/FX/Niagara/Source/NiagaraEditor/Private/NiagaraCompilationTasks.cpp b/Engine/Plugins/FX/Niagara/Source/NiagaraEditor/Private/NiagaraCompilationTasks.cpp
--- a/Engine/Plugins/FX/Niagara/Source/NiagaraEditor/Private/NiagaraCompilationTasks.cpp
+++ b/Engine/Plugins/FX/Niagara/Source/NiagaraEditor/Private/NiagaraCompilationTasks.cpp
@@ -55,7 +55,7 @@ namespace NiagaraSystemCookStats
 
 namespace NiagaraCompilationTasksImpl
 {
-	static const FGuid UE_NIAGARA_ASYNC_TASK_COMPILER_VER = FGuid(0xA6981076, 0x2AC24481, 0x80CFB302, 0x7EDAB79F);
+	static const FGuid UE_NIAGARA_ASYNC_TASK_COMPILER_VER = FGuid(0x3076A815, 0x9F4F4E23, 0x99C2573E, 0xAC120D18);
 	static UE::DerivedData::FCacheBucket NiagaraDDCBucket("NiagaraScript");
 
 	void GetUsagesToDuplicate(ENiagaraScriptUsage TargetUsage, TArray<ENiagaraScriptUsage>& DuplicateUsages)
diff --git a/Engine/Plugins/FX/Niagara/Source/NiagaraEditor/Private/NiagaraSystemCompilingManager.cpp b/Engine/Plugins/FX/Niagara/Source/NiagaraEditor/Private/NiagaraSystemCompilingManager.cpp
--- a/Engine/Plugins/FX/Niagara/Source/NiagaraEditor/Private/NiagaraSystemCompilingManager.cpp
+++ b/Engine/Plugins/FX/Niagara/Source/NiagaraEditor/Private/NiagaraSystemCompilingManager.cpp
@@ -634,6 +634,11 @@ FNiagaraCompilationTaskHandle FNiagaraEditorModule::RequestCompileSystem(UNiagar
 	const TArray<TWeakObjectPtr<UNiagaraParameterCollection>>& Collections = ParameterCollectionAssetCache.Get();
 	CompileOptions.ParameterCollections = ParameterCollectionAssetCache.Get();
 
+	// the issue here is that we aren't strictly allowed to call ComputeVMCompilationId on the worker threads which leaves us in a
+	// state where we could have broken RI parameters.  We ensure that the RI parameters are up to date before we begin.  This is a
+	// costly operation, but sadly seems necessary for some content.
+	System->PrepareRapidIterationParametersForCompilation();
+
 	return FNiagaraSystemCompilingManager::Get().AddSystem(System, CompileOptions);
 }
 
