- Don't allow GPU generation if we call GDME outside of the main InitView / PreRender section
- Shadow generation runs a separate GDME call after PreRender, this leaves dangling simulation requests for the next frame and produces incorrect results if the emitter was not visible in the main view
- Instead we will force the CPU generation path if possible, and if not an ensure will be generated

#rnx

diff --git a/Engine/Plugins/FX/Niagara/Source/Niagara/Classes/NiagaraComputeExecutionContext.h b/Engine/Plugins/FX/Niagara/Source/Niagara/Classes/NiagaraComputeExecutionContext.h
--- a/Engine/Plugins/FX/Niagara/Source/Niagara/Classes/NiagaraComputeExecutionContext.h
+++ b/Engine/Plugins/FX/Niagara/Source/Niagara/Classes/NiagaraComputeExecutionContext.h
@@ -58,6 +58,7 @@ struct FNiagaraGpuSpawnInfo
 
 struct INiagaraComputeDataBufferInterface
 {
+	virtual ~INiagaraComputeDataBufferInterface() = default;
 	virtual bool HasTranslucentDataToRender() const = 0;
 	virtual FNiagaraDataBuffer* GetDataToRender(FRHICommandListBase& RHICmdList, bool bIsLowLatencyTranslucent) const = 0;
 };
diff --git a/Engine/Plugins/FX/Niagara/Source/Niagara/Private/NiagaraGpuComputeDispatch.cpp b/Engine/Plugins/FX/Niagara/Source/Niagara/Private/NiagaraGpuComputeDispatch.cpp
--- a/Engine/Plugins/FX/Niagara/Source/Niagara/Private/NiagaraGpuComputeDispatch.cpp
+++ b/Engine/Plugins/FX/Niagara/Source/Niagara/Private/NiagaraGpuComputeDispatch.cpp
@@ -1824,6 +1824,8 @@ void FNiagaraGpuComputeDispatch::PreInitViews(FRDGBuilder& GraphBuilder, bool bA
 	RDG_CSV_STAT_EXCLUSIVE_SCOPE(GraphBuilder, Niagara);
 	LLM_SCOPE(ELLMTag::Niagara);
 
+	OnPreInitViewsEvent.Broadcast(GraphBuilder);
+
 	bIsFirstViewFamily = CurrentFamily ? CurrentFamily == ViewFamilies[0] : true;
 	bIsLastViewFamily = CurrentFamily ? CurrentFamily == ViewFamilies.Last() : true;
 	bRequiresReadback = false;
diff --git a/Engine/Plugins/FX/Niagara/Source/Niagara/Private/Stateless/NiagaraStatelessComputeManager.cpp b/Engine/Plugins/FX/Niagara/Source/Niagara/Private/Stateless/NiagaraStatelessComputeManager.cpp
--- a/Engine/Plugins/FX/Niagara/Source/Niagara/Private/Stateless/NiagaraStatelessComputeManager.cpp
+++ b/Engine/Plugins/FX/Niagara/Source/Niagara/Private/Stateless/NiagaraStatelessComputeManager.cpp
@@ -41,9 +41,9 @@ namespace NiagaraStatelessComputeManagerPrivate
 		ECVF_Default
 	);
 
-	EComputeExecutionPath DetermineComputeExecutionPath(const FNiagaraStatelessEmitterData* EmitterData, uint32 ActiveParticlesEstimate)
+	EComputeExecutionPath DetermineComputeExecutionPath(const FNiagaraStatelessEmitterData* EmitterData, uint32 ActiveParticlesEstimate, bool bAllowGPUGeneration)
 	{
-		const bool bAllowGPUExec = EnumHasAnyFlags(EmitterData->FeatureMask, ENiagaraStatelessFeatureMask::ExecuteGPU);
+		const bool bAllowGPUExec = EnumHasAnyFlags(EmitterData->FeatureMask, ENiagaraStatelessFeatureMask::ExecuteGPU) && bAllowGPUGeneration;
 		const bool bUseCPUExec = EnumHasAnyFlags(EmitterData->FeatureMask, ENiagaraStatelessFeatureMask::ExecuteCPU) && (!bAllowGPUExec || (ActiveParticlesEstimate <= uint32(GParticleCountCPUThreshold)));
 		if (bUseCPUExec)
 		{
@@ -148,7 +148,8 @@ namespace NiagaraStatelessComputeManagerPrivate
 FNiagaraStatelessComputeManager::FNiagaraStatelessComputeManager(FNiagaraGpuComputeDispatchInterface* InOwnerInterface)
 	: FNiagaraGpuComputeDataManager(InOwnerInterface)
 {
-	InOwnerInterface->GetOnPreRenderEvent().AddRaw(this, &FNiagaraStatelessComputeManager::OnPostPreRender);
+	InOwnerInterface->GetOnPreInitViewsEvent().AddRaw(this, &FNiagaraStatelessComputeManager::OnPreInitViews);
+	InOwnerInterface->GetOnPreRenderEvent().AddRaw(this, &FNiagaraStatelessComputeManager::OnPreRender);
 	InOwnerInterface->GetOnPostRenderEvent().AddRaw(this, &FNiagaraStatelessComputeManager::OnPostPostRender);
 }
 
@@ -225,7 +226,12 @@ FNiagaraDataBuffer* FNiagaraStatelessComputeManager::GetDataBuffer(FRHICommandLi
 
 	CacheData->DataBuffer->AllocateGPU(RHICmdList, ActiveParticles, ComputeInterface->GetFeatureLevel(), TEXT("StatelessSimBuffer"));
 
-	const EComputeExecutionPath ComputeExecutionPath = DetermineComputeExecutionPath(EmitterData, ActiveParticles);
+	// Until we add an extension to the render to notify about GDME start / end we can not allow GPU generation requests outside of PreInitViews / PreRender
+	// The shadow rendering, for example, will call GDME outside of this and it will result in crashes
+	// Therefore we fallback to CPU generation in these cases, if available
+	const bool bAllowGPUGeneration = bAllowDeferredGeneration;
+
+	const EComputeExecutionPath ComputeExecutionPath = DetermineComputeExecutionPath(EmitterData, ActiveParticles, bAllowGPUGeneration);
 	switch (ComputeExecutionPath)
 	{
 		case EComputeExecutionPath::CPU:
@@ -234,7 +240,9 @@ FNiagaraDataBuffer* FNiagaraStatelessComputeManager::GetDataBuffer(FRHICommandLi
 			ParticleSimulation.SimulateGPU(RHICmdList, EmitterInstance->RandomSeed, EmitterInstance->Age, EmitterInstance->DeltaTime, EmitterInstance->SpawnInfos, CacheData->DataBuffer);
 			if (ParticleSimulation.GetNumInstances() == 0)
 			{
-				FreeData.Emplace(CacheData);
+				// Note: We can not free the data when in paralell GDME mode as multiple lock / unlock operations could result in the wrong one providing the final data
+				// The plus side to adding back into the cache data is that multiple calls to get the same emitter's data will resolve to no active particles.
+				UsedData.Emplace(EmitterKey, CacheData);
 				return nullptr;
 			}
 			break;
@@ -255,6 +263,7 @@ FNiagaraDataBuffer* FNiagaraStatelessComputeManager::GetDataBuffer(FRHICommandLi
 
 		default:
 			ensureMsgf(false, TEXT("No execution path was found for stateless emitter, data will not be generated"));
+			// We can add back to free data since we never locked / unlocked the buffer
 			FreeData.Emplace(CacheData);
 			return nullptr;
 	}
@@ -284,7 +293,8 @@ void FNiagaraStatelessComputeManager::GenerateDataBufferForDebugging(FRHICommand
 		return;
 	}
 
-	const EComputeExecutionPath ComputeExecutionPath = DetermineComputeExecutionPath(EmitterData, ActiveParticlesEstimate);
+	const bool bAllowGPUGeneration = true;
+	const EComputeExecutionPath ComputeExecutionPath = DetermineComputeExecutionPath(EmitterData, ActiveParticlesEstimate, bAllowGPUGeneration);
 	switch (ComputeExecutionPath)
 	{
 		case EComputeExecutionPath::CPU:
@@ -329,8 +339,15 @@ void FNiagaraStatelessComputeManager::GenerateDataBufferForDebugging(FRHICommand
 	}
 }
 
-void FNiagaraStatelessComputeManager::OnPostPreRender(FRDGBuilder& GraphBuilder)
+void FNiagaraStatelessComputeManager::OnPreInitViews(FRDGBuilder& GraphBuilder)
 {
+	bAllowDeferredGeneration = true;
+}
+
+void FNiagaraStatelessComputeManager::OnPreRender(FRDGBuilder& GraphBuilder)
+{
+	bAllowDeferredGeneration = false;
+
 	// Anything to process?
 	if (GPUGenerationRequests.Num() == 0)
 	{
@@ -350,10 +367,10 @@ void FNiagaraStatelessComputeManager::OnPostPreRender(FRDGBuilder& GraphBuilder)
 	// Execute dispatches
 	AddPass(
 		GraphBuilder,
-		RDG_EVENT_NAME("FNiagaraStatelessComputeManager::OnPostPreRender"),
+		RDG_EVENT_NAME("FNiagaraStatelessComputeManager::OnPreRender"),
 		[GPUGenerationRequests_RDG=MoveTemp(GPUGenerationRequests), ComputeInterface](FRHICommandListImmediate& RHICmdList)
 		{
-			SCOPED_DRAW_EVENT(RHICmdList, FNiagaraStatelessComputeManager_OnPostPreRender);
+			SCOPED_DRAW_EVENT(RHICmdList, FNiagaraStatelessComputeManager_OnPreRender);
 
 			FNiagaraGPUInstanceCountManager& CountManager = ComputeInterface->GetGPUInstanceCounterManager();
 			NiagaraStatelessComputeManagerPrivate::GenerateGPUData(RHICmdList, ComputeInterface, GPUGenerationRequests_RDG);
@@ -381,6 +398,7 @@ void FNiagaraStatelessComputeManager::OnPostPostRender(FRDGBuilder& GraphBuilder
 			FreeData.Empty(UsedData.Num());
 			for (auto it=UsedData.CreateIterator(); it; ++it)
 			{
+				it.Value()->DataBuffer->SetGPUInstanceCountBufferOffset(INDEX_NONE);
 				FreeData.Emplace(it.Value().Release());
 			}
 			UsedData.Empty();
diff --git a/Engine/Plugins/FX/Niagara/Source/Niagara/Private/Stateless/NiagaraStatelessComputeManager.h b/Engine/Plugins/FX/Niagara/Source/Niagara/Private/Stateless/NiagaraStatelessComputeManager.h
--- a/Engine/Plugins/FX/Niagara/Source/Niagara/Private/Stateless/NiagaraStatelessComputeManager.h
+++ b/Engine/Plugins/FX/Niagara/Source/Niagara/Private/Stateless/NiagaraStatelessComputeManager.h
@@ -58,7 +58,8 @@ public:
 	void GenerateDataBufferForDebugging(FRHICommandListImmediate& RHICmdList, FNiagaraDataBuffer* DataBuffer, const NiagaraStateless::FEmitterInstance_RT* EmitterInstance) const;
 
 private:
-	void OnPostPreRender(FRDGBuilder& GraphBuilder);
+	void OnPreInitViews(FRDGBuilder& GraphBuilder);
+	void OnPreRender(FRDGBuilder& GraphBuilder);
 	void OnPostPostRender(FRDGBuilder& GraphBuilder);
 
 private:
@@ -68,6 +69,7 @@ private:
 	TArray<TUniquePtr<FStatelessDataCache>>				FreeData;
 	TArray<uint32>										CountsToRelease;
 
+	bool												bAllowDeferredGeneration = false;
 	TArray<FStatelessDataGenerationRequest>				GPUGenerationRequests;
 
 	UE::FMutex											GetDataBufferGuard;
diff --git a/Engine/Plugins/FX/Niagara/Source/Niagara/Public/NiagaraGpuComputeDispatchInterface.h b/Engine/Plugins/FX/Niagara/Source/Niagara/Public/NiagaraGpuComputeDispatchInterface.h
--- a/Engine/Plugins/FX/Niagara/Source/Niagara/Public/NiagaraGpuComputeDispatchInterface.h
+++ b/Engine/Plugins/FX/Niagara/Source/Niagara/Public/NiagaraGpuComputeDispatchInterface.h
@@ -29,6 +29,7 @@ class FNiagaraSystemGpuComputeProxy;
 class FNiagaraGpuComputeDispatchInterface : public FFXSystemInterface
 {
 public:
+	DECLARE_EVENT_OneParam(FNiagaraGpuComputeDispatchInterface, FOnPreInitViewsEvent, FRDGBuilder&);
 	DECLARE_EVENT_OneParam(FNiagaraGpuComputeDispatchInterface, FOnPostPreRenderEvent, FRDGBuilder&);
 
 	static NIAGARA_API FNiagaraGpuComputeDispatchInterface* Get(class UWorld* World);
@@ -185,9 +186,19 @@ public:
 	FORCEINLINE void MultiGPUResourceModified(FRHICommandList& RHICmdList, FRHITexture* Texture, bool bRequiredForSimulation, bool bRequiredForRendering) const {}
 #endif
 
-	/** Event that broadcast before any rendering work is prepared / executed for Niagara. */
+	/**
+	Event that broadcast when we enter PreInitViews.
+	*/
+	FOnPreInitViewsEvent& GetOnPreInitViewsEvent() { check(IsInRenderingThread()); return OnPreInitViewsEvent; }
+	/**
+	Event that broadcast when we endter PreRender.
+	This is called before we prepare any work or add passes for simulating.
+	*/
 	FOnPostPreRenderEvent& GetOnPreRenderEvent() { check(IsInRenderingThread()); return OnPreRenderEvent; }
-	/** Event that broadcast after all rendering for Niagara is complete. */
+	/**
+	Event that broadcast at the end of PostRenderOpaque.
+	This is called after all simulation passes have been added.
+	*/
 	FOnPostPreRenderEvent& GetOnPostRenderEvent() { check(IsInRenderingThread()); return OnPostRenderEvent; }
 
 protected:
@@ -210,6 +221,7 @@ protected:
 	bool									bIsFirstViewFamily = true;
 	bool									bIsLastViewFamily = true;
 
+	FOnPreInitViewsEvent					OnPreInitViewsEvent;
 	FOnPostPreRenderEvent					OnPreRenderEvent;
 	FOnPostPreRenderEvent					OnPostRenderEvent;
 
