- Fix memory leaks due to not destroying the per instance data

#rnx
#tests Ran emitters with the impacted data interaces and watched memory when restarting them continously

diff --git a/Engine/Plugins/FX/Niagara/Source/Niagara/Classes/NiagaraDataInterfaceCamera.h b/Engine/Plugins/FX/Niagara/Source/Niagara/Classes/NiagaraDataInterfaceCamera.h
--- a/Engine/Plugins/FX/Niagara/Source/Niagara/Classes/NiagaraDataInterfaceCamera.h
+++ b/Engine/Plugins/FX/Niagara/Source/Niagara/Classes/NiagaraDataInterfaceCamera.h
@@ -50,6 +50,7 @@ public:
 	//UNiagaraDataInterface Interface
 	NIAGARA_API virtual void GetVMExternalFunction(const FVMExternalFunctionBindingInfo& BindingInfo, void* InstanceData, FVMExternalFunction &OutFunc) override;
 	NIAGARA_API virtual bool InitPerInstanceData(void* PerInstanceData, FNiagaraSystemInstance* SystemInstance) override;
+	NIAGARA_API virtual void DestroyPerInstanceData(void* PerInstanceData, FNiagaraSystemInstance* SystemInstance) override;
 	virtual int32 PerInstanceDataSize() const override { return sizeof(FCameraDataInterface_InstanceData); }
 	NIAGARA_API virtual bool PerInstanceTick(void* PerInstanceData, FNiagaraSystemInstance* SystemInstance, float DeltaSeconds) override;
 #if WITH_EDITORONLY_DATA
diff --git a/Engine/Plugins/FX/Niagara/Source/Niagara/Classes/NiagaraDataInterfaceGrid2DCollectionReader.h b/Engine/Plugins/FX/Niagara/Source/Niagara/Classes/NiagaraDataInterfaceGrid2DCollectionReader.h
--- a/Engine/Plugins/FX/Niagara/Source/Niagara/Classes/NiagaraDataInterfaceGrid2DCollectionReader.h
+++ b/Engine/Plugins/FX/Niagara/Source/Niagara/Classes/NiagaraDataInterfaceGrid2DCollectionReader.h
@@ -32,6 +32,7 @@ public:
 
 	NIAGARA_API virtual bool Equals(const UNiagaraDataInterface* Other) const override;
 	NIAGARA_API virtual bool InitPerInstanceData(void* PerInstanceData, FNiagaraSystemInstance* SystemInstance) override;
+	NIAGARA_API virtual void DestroyPerInstanceData(void* PerInstanceData, FNiagaraSystemInstance* SystemInstance) override;
 	NIAGARA_API virtual void GetEmitterDependencies(UNiagaraSystem* Asset, TArray<FVersionedNiagaraEmitter>& Dependencies) const override;
 	NIAGARA_API virtual bool PerInstanceTickPostSimulate(void* PerInstanceData, FNiagaraSystemInstance* SystemInstance, float DeltaSeconds) override;
 
diff --git a/Engine/Plugins/FX/Niagara/Source/Niagara/Classes/NiagaraDataInterfaceGrid3DCollectionReader.h b/Engine/Plugins/FX/Niagara/Source/Niagara/Classes/NiagaraDataInterfaceGrid3DCollectionReader.h
--- a/Engine/Plugins/FX/Niagara/Source/Niagara/Classes/NiagaraDataInterfaceGrid3DCollectionReader.h
+++ b/Engine/Plugins/FX/Niagara/Source/Niagara/Classes/NiagaraDataInterfaceGrid3DCollectionReader.h
@@ -29,6 +29,7 @@ public:
 
 	NIAGARA_API virtual bool Equals(const UNiagaraDataInterface* Other) const override;
 	NIAGARA_API virtual bool InitPerInstanceData(void* PerInstanceData, FNiagaraSystemInstance* SystemInstance) override;
+	NIAGARA_API virtual void DestroyPerInstanceData(void* PerInstanceData, FNiagaraSystemInstance* SystemInstance) override;
 	NIAGARA_API virtual void GetEmitterDependencies(UNiagaraSystem* Asset, TArray<FVersionedNiagaraEmitter>& Dependencies) const override;
 	NIAGARA_API virtual bool PerInstanceTickPostSimulate(void* PerInstanceData, FNiagaraSystemInstance* SystemInstance, float DeltaSeconds) override;
 
diff --git a/Engine/Plugins/FX/Niagara/Source/Niagara/Private/NiagaraDataInterfaceCamera.cpp b/Engine/Plugins/FX/Niagara/Source/Niagara/Private/NiagaraDataInterfaceCamera.cpp
--- a/Engine/Plugins/FX/Niagara/Source/Niagara/Private/NiagaraDataInterfaceCamera.cpp
+++ b/Engine/Plugins/FX/Niagara/Source/Niagara/Private/NiagaraDataInterfaceCamera.cpp
@@ -69,6 +69,12 @@ bool UNiagaraDataInterfaceCamera::InitPerInstanceData(void* PerInstanceData, FNi
 	return true;
 }
 
+void UNiagaraDataInterfaceCamera::DestroyPerInstanceData(void* PerInstanceData, FNiagaraSystemInstance* SystemInstance)
+{
+	FCameraDataInterface_InstanceData* PIData = static_cast<FCameraDataInterface_InstanceData*>(PerInstanceData);
+	PIData->~FCameraDataInterface_InstanceData();
+}
+
 bool UNiagaraDataInterfaceCamera::PerInstanceTick(void* PerInstanceData, FNiagaraSystemInstance* SystemInstance, float DeltaSeconds)
 {
 	FCameraDataInterface_InstanceData* PIData = (FCameraDataInterface_InstanceData*)PerInstanceData;
diff --git a/Engine/Plugins/FX/Niagara/Source/Niagara/Private/NiagaraDataInterfaceGrid2DCollectionReader.cpp b/Engine/Plugins/FX/Niagara/Source/Niagara/Private/NiagaraDataInterfaceGrid2DCollectionReader.cpp
--- a/Engine/Plugins/FX/Niagara/Source/Niagara/Private/NiagaraDataInterfaceGrid2DCollectionReader.cpp
+++ b/Engine/Plugins/FX/Niagara/Source/Niagara/Private/NiagaraDataInterfaceGrid2DCollectionReader.cpp
@@ -135,6 +135,21 @@ bool UNiagaraDataInterfaceGrid2DCollectionReader::InitPerInstanceData(void* PerI
 	return false;
 }
 
+void UNiagaraDataInterfaceGrid2DCollectionReader::DestroyPerInstanceData(void* PerInstanceData, FNiagaraSystemInstance* SystemInstance)
+{
+	FGrid2DCollectionRWInstanceData_GameThread* InstanceData = static_cast<FGrid2DCollectionRWInstanceData_GameThread*>(PerInstanceData);
+	InstanceData->~FGrid2DCollectionRWInstanceData_GameThread();
+
+	SystemInstancesToProxyData_GT.Remove(SystemInstance->GetId());
+
+	FNiagaraDataInterfaceProxyGrid2DCollectionProxy* RT_Proxy = GetProxyAs<FNiagaraDataInterfaceProxyGrid2DCollectionProxy>();
+	ENQUEUE_RENDER_COMMAND(FUpdateData)(
+		[RT_Proxy, InstanceID=SystemInstance->GetId()](FRHICommandListImmediate& RHICmdList)
+		{
+			RT_Proxy->SystemInstancesToProxyData_RT.Remove(InstanceID);
+		});
+}
+
 void UNiagaraDataInterfaceGrid2DCollectionReader::GetEmitterDependencies(UNiagaraSystem* Asset, TArray<FVersionedNiagaraEmitter>& Dependencies) const
 {
 	if (!Asset)
diff --git a/Engine/Plugins/FX/Niagara/Source/Niagara/Private/NiagaraDataInterfaceGrid3DCollectionReader.cpp b/Engine/Plugins/FX/Niagara/Source/Niagara/Private/NiagaraDataInterfaceGrid3DCollectionReader.cpp
--- a/Engine/Plugins/FX/Niagara/Source/Niagara/Private/NiagaraDataInterfaceGrid3DCollectionReader.cpp
+++ b/Engine/Plugins/FX/Niagara/Source/Niagara/Private/NiagaraDataInterfaceGrid3DCollectionReader.cpp
@@ -146,6 +146,21 @@ bool UNiagaraDataInterfaceGrid3DCollectionReader::InitPerInstanceData(void* PerI
 	return false;
 }
 
+void UNiagaraDataInterfaceGrid3DCollectionReader::DestroyPerInstanceData(void* PerInstanceData, FNiagaraSystemInstance* SystemInstance)
+{
+	FGrid3DCollectionRWInstanceData_GameThread* PIData = static_cast<FGrid3DCollectionRWInstanceData_GameThread*>(PerInstanceData);
+	PIData->~FGrid3DCollectionRWInstanceData_GameThread();
+
+	SystemInstancesToProxyData_GT.Remove(SystemInstance->GetId());
+
+	FNiagaraDataInterfaceProxyGrid3DCollectionProxy* RT_Proxy = GetProxyAs<FNiagaraDataInterfaceProxyGrid3DCollectionProxy>();
+	ENQUEUE_RENDER_COMMAND(FUpdateData)(
+		[RT_Proxy, InstanceID = SystemInstance->GetId()](FRHICommandListImmediate& RHICmdList)
+		{
+			RT_Proxy->SystemInstancesToProxyData_RT.Remove(InstanceID);
+		});
+}
+
 void UNiagaraDataInterfaceGrid3DCollectionReader::GetEmitterDependencies(UNiagaraSystem* Asset, TArray<FVersionedNiagaraEmitter>& Dependencies) const
 {
 	if (!Asset)
