Added Lightweight Asset Tag for Niagara

diff --git a/Engine/Plugins/FX/Niagara/Source/Niagara/Private/NiagaraModule.cpp b/Engine/Plugins/FX/Niagara/Source/Niagara/Private/NiagaraModule.cpp
--- a/Engine/Plugins/FX/Niagara/Source/Niagara/Private/NiagaraModule.cpp
+++ b/Engine/Plugins/FX/Niagara/Source/Niagara/Private/NiagaraModule.cpp
@@ -67,6 +67,14 @@ static FAutoConsoleVariableRef CVarLogCompileIdGeneration(
 float INiagaraModule::EngineGlobalSpawnCountScale = 1.0f;
 float INiagaraModule::EngineGlobalSystemCountScale = 1.0f;
 
+const FNiagaraAssetTagDefinition INiagaraModule::LightweightTagDefinition = {
+	LOCTEXT("LightweightDisplayName", "Lightweight")
+	, (int32) (ENiagaraAssetLibraryAssetTypes::Systems)
+	, FText::GetEmpty()
+	, ENiagaraAssetTagDefinitionImportance::Primary
+	, FLinearColor::Red
+	, FGuid(0xE1737BF2, 0x52264191, 0xB9B9AA5F, 0xFB8C7412)};
+
 const FNiagaraAssetTagDefinition INiagaraModule::TemplateTagDefinition = {
 	LOCTEXT("TemplateDisplayName", "Template")
 	, (int32) (ENiagaraAssetLibraryAssetTypes::Emitters) | (int32) (ENiagaraAssetLibraryAssetTypes::Systems)
@@ -1919,6 +1927,7 @@ void INiagaraModule::RegisterInternalAssetTagDefinitions()
 {
 	// All internal tags have to be manually assigned a unique and stable guid so that names can be changed without affecting assigned assets.
 	InternalAssetTagDefinitions.Add(&TemplateTagDefinition);
+	InternalAssetTagDefinitions.Add(&LightweightTagDefinition);
 	InternalAssetTagDefinitions.Add(&LearningContentTagDefinition);
 	InternalAssetTagDefinitions.Add(&HiddenAssetTagDefinition);
 	InternalAssetTagDefinitions.Add(&DeprecatedTagDefinition);
diff --git a/Engine/Plugins/FX/Niagara/Source/Niagara/Public/NiagaraModule.h b/Engine/Plugins/FX/Niagara/Source/Niagara/Public/NiagaraModule.h
--- a/Engine/Plugins/FX/Niagara/Source/Niagara/Public/NiagaraModule.h
+++ b/Engine/Plugins/FX/Niagara/Source/Niagara/Public/NiagaraModule.h
@@ -292,6 +292,8 @@ public:
 	static NIAGARA_API bool bUseGlobalFXBudget;
 	static NIAGARA_API bool bDataChannelsEnabled;
 
+	static const FNiagaraAssetTagDefinition LightweightTagDefinition;
+
 	static const FNiagaraAssetTagDefinition TemplateTagDefinition;
 	static const FNiagaraAssetTagDefinition LearningContentTagDefinition;
 	static NIAGARA_API const FNiagaraAssetTagDefinition HiddenAssetTagDefinition;
