[Niagara] Fix log printing string that was not null terminated
#rnx

diff --git a/Engine/Plugins/FX/Niagara/Source/Niagara/Private/NiagaraDataInterfaceEmitterBinding.cpp b/Engine/Plugins/FX/Niagara/Source/Niagara/Private/NiagaraDataInterfaceEmitterBinding.cpp
--- a/Engine/Plugins/FX/Niagara/Source/Niagara/Private/NiagaraDataInterfaceEmitterBinding.cpp
+++ b/Engine/Plugins/FX/Niagara/Source/Niagara/Private/NiagaraDataInterfaceEmitterBinding.cpp
@@ -90,7 +90,7 @@ FNiagaraEmitterInstance* FNiagaraDataInterfaceEmitterBinding::Resolve(const FNia
 
 		if (FNiagaraUtilities::LogVerboseWarnings())
 		{
-			UE_LOG(LogNiagara, Warning, TEXT("EmitterBinding failed to find emitter '%s' it might not exist or has been cooked out"), EmitterNameStringView.GetData());
+			UE_LOG(LogNiagara, Warning, TEXT("EmitterBinding failed to find emitter '%s' it might not exist or has been cooked out"), *EmitterNameString);
 		}
 	}
 	return nullptr;
