Changed User Parameters Section box from Horizontal Box to Wrap Box

diff --git a/Engine/Plugins/FX/Niagara/Source/NiagaraEditor/Private/Customizations/UserParameters/NiagaraUserParametersBuilderBase.cpp b/Engine/Plugins/FX/Niagara/Source/NiagaraEditor/Private/Customizations/UserParameters/NiagaraUserParametersBuilderBase.cpp
--- a/Engine/Plugins/FX/Niagara/Source/NiagaraEditor/Private/Customizations/UserParameters/NiagaraUserParametersBuilderBase.cpp
+++ b/Engine/Plugins/FX/Niagara/Source/NiagaraEditor/Private/Customizations/UserParameters/NiagaraUserParametersBuilderBase.cpp
@@ -18,6 +18,7 @@
 #include "NiagaraEditorUtilities.h"
 #include "PropertyEditorClipboard.h"
 #include "Customizations/NiagaraComponentDetails.h"
+#include "Widgets/Layout/SWrapBox.h"
 
 #define LOCTEXT_NAMESPACE "NiagaraUserParametersBuilderBase"
 
@@ -180,7 +181,7 @@ void FNiagaraUserParameterNodeBuilder::GenerateUserParameterRows(IDetailChildren
 	// we create a custom row for user param sections here, but only if we have at least one section specified
 	if(SectionCategoryMap.Num() > 0)
 	{
-		TSharedRef<SHorizontalBox> SectionsBox = SNew(SHorizontalBox);
+		TSharedRef<SWrapBox> SectionsBox = SNew(SWrapBox).UseAllottedSize(true);
 
 		// maps don't guarantee order, so we iterate over the original section data instead
 		for(auto& Section : Root->GetSectionData())
@@ -210,7 +211,6 @@ void FNiagaraUserParameterNodeBuilder::GenerateUserParameterRows(IDetailChildren
 				if(bDoesUserParamExist)
 				{
 					SectionsBox->AddSlot()
-					.AutoWidth()
 					.Padding(FMargin(2.f))
 					[
 						SNew(SBox)
@@ -243,7 +243,6 @@ void FNiagaraUserParameterNodeBuilder::GenerateUserParameterRows(IDetailChildren
 		if(SectionsBox->GetChildren()->Num() > 0)
 		{
 			SectionsBox->AddSlot()
-			.AutoWidth()
 			.Padding(FMargin(2.f))
 			[
 				SNew(SBox)
