Improved displayed info in parameter reference widget

#rnx
#tests Tested by using dynamic inputs on a module and then linking parameters to its inputs. The reference widget now correctly displays the base module along with the dynamic input name.

diff --git a/Engine/Plugins/FX/Niagara/Source/NiagaraEditor/Private/ViewModels/NiagaraParameterPanelViewModel.cpp b/Engine/Plugins/FX/Niagara/Source/NiagaraEditor/Private/ViewModels/NiagaraParameterPanelViewModel.cpp
--- a/Engine/Plugins/FX/Niagara/Source/NiagaraEditor/Private/ViewModels/NiagaraParameterPanelViewModel.cpp
+++ b/Engine/Plugins/FX/Niagara/Source/NiagaraEditor/Private/ViewModels/NiagaraParameterPanelViewModel.cpp
@@ -2024,6 +2024,38 @@ TArray<FNiagaraParameterPanelItem> FNiagaraSystemToolkitParameterPanelViewModel:
 					{
 						// use the name of the module input for usages as stack function inputs
 						Source.ModuleName = History.ReadPin.Pin->LinkedTo[0]->PinName;
+
+						// find the name of the module this belongs to when used in a dynamic input
+						TArray<UEdGraphNode*> NodesToCheck;
+						NodesToCheck.Add(History.ReadPin.Pin->LinkedTo[0]->GetOwningNode());
+						int32 NodeIndex = 0;
+						while (NodesToCheck.IsValidIndex(NodeIndex))
+						{
+							if (UNiagaraNode* Node = Cast<UNiagaraNode>(NodesToCheck[NodeIndex]))
+							{
+								if (UNiagaraNodeFunctionCall* FunctionNode = Cast<UNiagaraNodeFunctionCall>(Node))
+								{
+									if (FunctionNode->FunctionScript->IsModuleScript() && FunctionNode->IsNodeEnabled())
+									{
+										// Found the module
+										FString CurrentName = Source.ModuleName.ToString();
+										if (!CurrentName.StartsWith(FunctionNode->GetFunctionName()))
+										{
+											Source.ModuleName = FName("[" + FunctionNode->GetFunctionName() + "] " + CurrentName);
+										}
+										break;
+									}
+								}
+								if (UEdGraphPin* OutputPin = FNiagaraStackGraphUtilities::GetParameterMapOutputPin(*Node))
+								{
+									for (UEdGraphPin* LinkedPin : OutputPin->LinkedTo)
+									{
+										NodesToCheck.AddUnique(LinkedPin->GetOwningNode());
+									}
+								}
+							}
+							NodeIndex++;
+						}
 					}
 					Item.AddToReadCount(Source);
 				};
diff --git a/Engine/Plugins/FX/Niagara/Source/NiagaraEditor/Private/Widgets/SNiagaraParameterPanel.cpp b/Engine/Plugins/FX/Niagara/Source/NiagaraEditor/Private/Widgets/SNiagaraParameterPanel.cpp
--- a/Engine/Plugins/FX/Niagara/Source/NiagaraEditor/Private/Widgets/SNiagaraParameterPanel.cpp
+++ b/Engine/Plugins/FX/Niagara/Source/NiagaraEditor/Private/Widgets/SNiagaraParameterPanel.cpp
@@ -474,7 +474,7 @@ TSharedRef<SWidget> SNiagaraParameterPanel::OnGenerateWidgetForItem(const FNiaga
 		FString DetailReferences;
 		for (const TPair<FName, TArray<FNiagaraParameterReferencePath>>& Pair : PathByParent)
 		{
-			DetailReferences += "\n" + Pair.Key.ToString() + ":\n";
+			DetailReferences += "\n" + Pair.Key.GetDisplayNameEntry()->GetPlainNameString() + ":\n";
 			for (const FNiagaraParameterReferencePath& Path : Pair.Value)
 			{
 				FString Prefix = (Path.bRead && Path.bWrite) ? "R+W" : (Path.bRead ? "R" : "W");
