Fix crash in FDynamicRayTracingInstancesContext::Finish(...) due to stale pointers to MeshBatch.
- Implementation of FNiagaraRendererMeshes::GetDynamicRayTracingInstances(...) added MeshBatch to Collector but then could early out without adding RayTracingInstance leading to MeshBatch memory being freed.

diff --git a/Engine/Plugins/FX/Niagara/Source/Niagara/Private/NiagaraRendererMeshes.cpp b/Engine/Plugins/FX/Niagara/Source/Niagara/Private/NiagaraRendererMeshes.cpp
--- a/Engine/Plugins/FX/Niagara/Source/Niagara/Private/NiagaraRendererMeshes.cpp
+++ b/Engine/Plugins/FX/Niagara/Source/Niagara/Private/NiagaraRendererMeshes.cpp
@@ -1774,8 +1774,6 @@ void FNiagaraRendererMeshes::GetDynamicRayTracingInstances(FRayTracingInstanceCo
 #endif
 
 			RayTracingInstance.Materials.Add(MoveTemp(MeshBatch));
-
-			Collector.AddMesh(ViewIndex, RayTracingInstance.Materials.Last());
 		}
 
 		if (RayTracingInstance.Materials.Num() == 0 || LODModel.Sections.Num() != RayTracingInstance.Materials.Num())
@@ -1863,6 +1861,15 @@ void FNiagaraRendererMeshes::GetDynamicRayTracingInstances(FRayTracingInstanceCo
 			RayTracingInstance.NumTransforms = NumInstances;
 		}
 
+		// Add MeshBatches to FMeshElementCollector so they're added to GPUScene if necessary
+		// TODO: should probably do it automatically in Collector.AddRayTracingInstance(...)
+		// Currently this will cache pointers to entries in RayTracingInstance.Materials
+		// so extra care must be taken to avoid reallocs that would invalidate those pointers
+		for (FMeshBatch& MeshBatch : RayTracingInstance.Materials)
+		{
+			Collector.AddMesh(ViewIndex, MeshBatch);
+		}
+
 		Collector.AddRayTracingInstance(MoveTemp(RayTracingInstance));
 	}
 }
