- Fix incorrect velocity output (needed to multiply inv DT)

#rnx
#tests Make a stateless emitter with add velocity, use approximate motion vectors on the renderer, observe incorrect motion blur

#virtualized

diff --git a/Engine/Plugins/FX/Niagara/Shaders/Private/Stateless/Modules/NiagaraStatelessModule_SolveVelocitiesAndForces.ush b/Engine/Plugins/FX/Niagara/Shaders/Private/Stateless/Modules/NiagaraStatelessModule_SolveVelocitiesAndForces.ush
--- a/Engine/Plugins/FX/Niagara/Shaders/Private/Stateless/Modules/NiagaraStatelessModule_SolveVelocitiesAndForces.ush
+++ b/Engine/Plugins/FX/Niagara/Shaders/Private/Stateless/Modules/NiagaraStatelessModule_SolveVelocitiesAndForces.ush
@@ -162,6 +162,6 @@ void SolveVelocitiesAndForces_Simulate(inout FStatelessParticle Particle)
 	SolveVelocitiesAndForces_IntegratePosition(ModuleData, Particle.Age, Particle.Position);
 	SolveVelocitiesAndForces_IntegratePosition(ModuleData, Particle.PreviousAge, Particle.PreviousPosition);
 		
-	Particle.Velocity = Particle.Position - Particle.PreviousPosition;
+	Particle.Velocity = (Particle.Position - Particle.PreviousPosition) * Particle.InvDeltaTime;
 	Particle.PreviousVelocity = Particle.Velocity;
 }
diff --git a/Engine/Plugins/FX/Niagara/Source/Niagara/Internal/Stateless/Modules/NiagaraStatelessModule_SolveVelocitiesAndForces.h b/Engine/Plugins/FX/Niagara/Source/Niagara/Internal/Stateless/Modules/NiagaraStatelessModule_SolveVelocitiesAndForces.h
--- a/Engine/Plugins/FX/Niagara/Source/Niagara/Internal/Stateless/Modules/NiagaraStatelessModule_SolveVelocitiesAndForces.h
+++ b/Engine/Plugins/FX/Niagara/Source/Niagara/Internal/Stateless/Modules/NiagaraStatelessModule_SolveVelocitiesAndForces.h
@@ -144,6 +144,7 @@ public:
 
 		const float* AgeData			= ParticleSimulationContext.GetParticleAge();
 		const float* PreviousAgeData	= ParticleSimulationContext.GetParticlePreviousAge();
+		const float InvDeltaTime		= ParticleSimulationContext.GetInvDeltaTime();
 
 		for (uint32 i = 0; i < ParticleSimulationContext.GetNumInstances(); ++i)
 		{
@@ -195,7 +196,7 @@ public:
 			ParticleSimulationContext.WriteParticleVariable(ModuleBuiltData->PositionVariableOffset, i, Position);
 			ParticleSimulationContext.WriteParticleVariable(ModuleBuiltData->PreviousPositionVariableOffset, i, PreviousPosition);
 
-			const FVector3f Velocity = Position - PreviousPosition;
+			const FVector3f Velocity = (Position - PreviousPosition) * InvDeltaTime;
 			ParticleSimulationContext.WriteParticleVariable(ModuleBuiltData->VelocityVariableOffset, i, Velocity);
 			ParticleSimulationContext.WriteParticleVariable(ModuleBuiltData->PreviousVelocityVariableOffset, i, Velocity);
 		}
