[Viewport Toolbar] Niagara Toolbar: Update Show and Motion settings menus so they use UToolMenu instead of FMenuBuilder

diff --git a/Engine/Plugins/FX/Niagara/Source/NiagaraEditor/Private/Widgets/NiagaraSystemEditorViewportToolbarSections.cpp b/Engine/Plugins/FX/Niagara/Source/NiagaraEditor/Private/Widgets/NiagaraSystemEditorViewportToolbarSections.cpp
--- a/Engine/Plugins/FX/Niagara/Source/NiagaraEditor/Private/Widgets/NiagaraSystemEditorViewportToolbarSections.cpp
+++ b/Engine/Plugins/FX/Niagara/Source/NiagaraEditor/Private/Widgets/NiagaraSystemEditorViewportToolbarSections.cpp
@@ -6,69 +6,13 @@
 #include "NiagaraEditorCommands.h"
 #include "SNiagaraSystemViewport.h"
 #include "ToolMenu.h"
+#include "ToolMenus.h"
+#include "ViewportToolbar/UnrealEdViewportToolbar.h"
 #include "ViewportToolbar/UnrealEdViewportToolbarContext.h"
 #include "Widgets/Input/SSpinBox.h"
 
 #define LOCTEXT_NAMESPACE "NiagaraSystemEditorViewportToolbarSections"
 
-TSharedRef<SWidget> UE::NiagaraSystemEditor::CreateShowMenuWidget(
-	const TSharedRef<SNiagaraSystemViewport>& InNiagaraSystemEditorViewport, bool bInShowViewportStatsToggle
-)
-{
-	InNiagaraSystemEditorViewport->OnFloatingButtonClicked();
-	TSharedRef<SEditorViewport> ViewportRef = InNiagaraSystemEditorViewport->GetViewportWidget();
-
-	constexpr bool bInShouldCloseWindowAfterMenuSelection = true;
-
-	// TODO move to UToolMenu*
-	FMenuBuilder ShowMenuBuilder(bInShouldCloseWindowAfterMenuSelection, ViewportRef->GetCommandList());
-	{
-		FNiagaraEditorCommands Commands = FNiagaraEditorCommands::Get();
-
-		if (bInShowViewportStatsToggle)
-		{
-			ShowMenuBuilder.AddSubMenu(
-				LOCTEXT("ViewportStatsSubMenu", "Viewport Stats"),
-				LOCTEXT("ViewportStatsSubMenu_ToolTip", "Viewport Stats settings"),
-				FNewMenuDelegate::CreateLambda(
-					[](FMenuBuilder& InMenuBuilder)
-					{
-						InMenuBuilder.BeginSection("CommonStats", LOCTEXT("CommonStatsLabel", "Common Stats"));
-
-						InMenuBuilder.AddMenuEntry(
-							FEditorViewportCommands::Get().ToggleStats,
-							"ViewportStats",
-							LOCTEXT("ViewportStatsLabel", "Show Stats")
-						);
-
-						InMenuBuilder.AddMenuEntry(
-							FEditorViewportCommands::Get().ToggleFPS, "ViewportFPS", LOCTEXT("ViewportFPSLabel", "Show FPS")
-						);
-
-						InMenuBuilder.AddMenuSeparator();
-						InMenuBuilder.EndSection();
-					}
-				),
-				false,
-				FSlateIcon()
-			);
-		}
-
-		ShowMenuBuilder.BeginSection("CommonShowFlags", LOCTEXT("CommonShowFlagsLabel", "Common Show Flags"));
-
-		ShowMenuBuilder.AddMenuEntry(Commands.TogglePreviewGrid);
-		ShowMenuBuilder.AddMenuSeparator();
-		ShowMenuBuilder.AddMenuEntry(Commands.ToggleEmitterExecutionOrder);
-		ShowMenuBuilder.AddMenuEntry(Commands.ToggleGpuTickInformation);
-		ShowMenuBuilder.AddMenuEntry(Commands.ToggleInstructionCounts);
-		ShowMenuBuilder.AddMenuEntry(Commands.ToggleMemoryInfo);
-		ShowMenuBuilder.AddMenuEntry(Commands.ToggleParticleCounts);
-		ShowMenuBuilder.AddMenuEntry(Commands.ToggleStatelessInfo);
-	}
-
-	return ShowMenuBuilder.MakeWidget();
-}
-
 FToolMenuEntry UE::NiagaraSystemEditor::CreateShowSubmenu()
 {
 	return FToolMenuEntry::InitSubMenu(
@@ -78,31 +22,54 @@ FToolMenuEntry UE::NiagaraSystemEditor::CreateShowSubmenu()
 		FNewToolMenuDelegate::CreateLambda(
 			[](UToolMenu* Submenu) -> void
 			{
-				UUnrealEdViewportToolbarContext* ViewportToolbarContext =
-					Submenu->FindContext<UUnrealEdViewportToolbarContext>();
-				if (!ViewportToolbarContext)
-				{
-					return;
-				}
+				UE::NiagaraSystemEditor::FillShowSubmenu(Submenu);
+			}
+		)
+	);
+}
 
-				TSharedPtr<SNiagaraSystemViewport> NiagaraSystemViewport =
-					StaticCastSharedPtr<SNiagaraSystemViewport>(ViewportToolbarContext->Viewport.Pin());
-				if (!NiagaraSystemViewport)
+void UE::NiagaraSystemEditor::FillShowSubmenu(UToolMenu* InMenu, bool bInShowViewportStatsToggle)
+{
+	const FNiagaraEditorCommands& Commands = FNiagaraEditorCommands::Get();
+
+	if (bInShowViewportStatsToggle)
+	{
+		FToolMenuSection& UnnamedSection = InMenu->FindOrAddSection(NAME_None);
+		UnnamedSection.AddSubMenu(
+			"ViewportStats",
+			LOCTEXT("ViewportStatsSubMenu", "Viewport Stats"),
+			LOCTEXT("CameraSubmenuTooltip", "Camera options"),
+			FNewToolMenuDelegate::CreateLambda(
+				[](UToolMenu* Submenu) -> void
 				{
-					return;
-				}
+					FToolMenuSection& CommonStatsSection =
+						Submenu->FindOrAddSection("CommonStats", LOCTEXT("CommonStatsLabel", "Common Stats"));
 
-				FToolMenuSection& UnnamedSection = Submenu->FindOrAddSection(NAME_None);
+					CommonStatsSection.AddMenuEntry(
+						FEditorViewportCommands::Get().ToggleStats, LOCTEXT("ViewportStatsLabel", "Show Stats")
+					);
 
-				UnnamedSection.AddEntry(FToolMenuEntry::InitWidget(
-					"ShowMenuItems",
-					UE::NiagaraSystemEditor::CreateShowMenuWidget(NiagaraSystemViewport.ToSharedRef(), true),
-					FText(),
-					true
-				));
-			}
-		)
-	);
+					CommonStatsSection.AddSeparator(NAME_None);
+
+					CommonStatsSection.AddMenuEntry(
+						FEditorViewportCommands::Get().ToggleFPS, LOCTEXT("ViewportFPSLabel", "Show FPS")
+					);
+				}
+			)
+		);
+	}
+
+	FToolMenuSection& CommonShowFlagsSection =
+		InMenu->FindOrAddSection("CommonShowFlags", LOCTEXT("CommonShowFlagsLabel", "Common Show Flags"));
+
+	CommonShowFlagsSection.AddMenuEntry(Commands.TogglePreviewGrid);
+	CommonShowFlagsSection.AddSeparator(NAME_None);
+	CommonShowFlagsSection.AddMenuEntry(Commands.ToggleEmitterExecutionOrder);
+	CommonShowFlagsSection.AddMenuEntry(Commands.ToggleGpuTickInformation);
+	CommonShowFlagsSection.AddMenuEntry(Commands.ToggleInstructionCounts);
+	CommonShowFlagsSection.AddMenuEntry(Commands.ToggleMemoryInfo);
+	CommonShowFlagsSection.AddMenuEntry(Commands.ToggleParticleCounts);
+	CommonShowFlagsSection.AddMenuEntry(Commands.ToggleStatelessInfo);
 }
 
 FToolMenuEntry UE::NiagaraSystemEditor::CreateSettingsSubmenu()
@@ -114,44 +81,7 @@ FToolMenuEntry UE::NiagaraSystemEditor::CreateSettingsSubmenu()
 		FNewToolMenuDelegate::CreateLambda(
 			[](UToolMenu* Submenu) -> void
 			{
-				FToolMenuSection& ViewportControlsSection =
-					Submenu->FindOrAddSection("ViewportControls", LOCTEXT("ViewportControlsLabel", "Viewport Controls"));
-
-				ViewportControlsSection.AddSubMenu(
-					"MotionOptions",
-					LOCTEXT("MotionOptionsSubMenu", "Motion Options"),
-					LOCTEXT("MotionOptionsSubMenu_ToolTip", "Set Motion Options for the Niagara Component"),
-					FNewToolMenuDelegate::CreateLambda(
-						[](UToolMenu* InMenu)
-						{
-							UUnrealEdViewportToolbarContext* ViewportToolbarContext =
-								InMenu->FindContext<UUnrealEdViewportToolbarContext>();
-							if (!ViewportToolbarContext)
-							{
-								return;
-							}
-
-							TSharedPtr<SNiagaraSystemViewport> NiagaraSystemViewport =
-								StaticCastSharedPtr<SNiagaraSystemViewport>(ViewportToolbarContext->Viewport.Pin());
-							if (!NiagaraSystemViewport)
-							{
-								return;
-							}
-
-							InMenu->AddMenuEntry(
-								"MotionOptions",
-								FToolMenuEntry::InitWidget(
-									"MotionOptions",
-									UE::NiagaraSystemEditor::CreateMotionMenuWidget(NiagaraSystemViewport.ToSharedRef()),
-									FText(),
-									true
-								)
-							);
-						}
-					),
-					false,
-					FSlateIcon()
-				);
+				UE::NiagaraSystemEditor::FillSettingsSubmenu(Submenu);
 			}
 		)
 	);
@@ -160,34 +90,105 @@ FToolMenuEntry UE::NiagaraSystemEditor::CreateSettingsSubmenu()
 TSharedRef<SWidget> UE::NiagaraSystemEditor::CreateMotionMenuWidget(const TSharedRef<SNiagaraSystemViewport>& InNiagaraSystemEditorViewport
 )
 {
-	constexpr bool bInShouldCloseWindowAfterMenuSelection = true;
-	FMenuBuilder MotionMenuBuilder(bInShouldCloseWindowAfterMenuSelection, InNiagaraSystemEditorViewport->GetCommandList());
-
-	MotionMenuBuilder.AddMenuEntry(FNiagaraEditorCommands::Get().ToggleMotion);
-
-	MotionMenuBuilder.AddWidget(
-		SNew(SSpinBox<float>)
-			.IsEnabled(InNiagaraSystemEditorViewport, &SNiagaraSystemViewport::IsMotionEnabled)
-			.Font(FAppStyle::GetFontStyle(TEXT("MenuItem.Font")))
-			.MinSliderValue(0.0f)
-			.MaxSliderValue(360.0f)
-			.Value(InNiagaraSystemEditorViewport, &SNiagaraSystemViewport::GetMotionRate)
-			.OnValueChanged(InNiagaraSystemEditorViewport, &SNiagaraSystemViewport::SetMotionRate),
-		LOCTEXT("MotionSpeed", "Motion Speed")
-	);
+	// We generate a menu via UToolMenus, so we can use FillShowSubmenu call from both old and new toolbar
+	FName OldShowMenuName = "NiagaraSystemEditor.OldViewportToolbar.Motion";
 
-	MotionMenuBuilder.AddWidget(
-		SNew(SSpinBox<float>)
-			.IsEnabled(InNiagaraSystemEditorViewport, &SNiagaraSystemViewport::IsMotionEnabled)
-			.Font(FAppStyle::GetFontStyle(TEXT("MenuItem.Font")))
-			.MinSliderValue(0.0f)
-			.MaxSliderValue(1000.0f)
-			.Value(InNiagaraSystemEditorViewport, &SNiagaraSystemViewport::GetMotionRadius)
-			.OnValueChanged(InNiagaraSystemEditorViewport, &SNiagaraSystemViewport::SetMotionRadius),
-		LOCTEXT("MotionRadius", "Motion Radius")
-	);
+	if (!UToolMenus::Get()->IsMenuRegistered(OldShowMenuName))
+	{
+		UToolMenu* Menu = UToolMenus::Get()->RegisterMenu(OldShowMenuName, NAME_None, EMultiBoxType::Menu, false);
+		Menu->AddDynamicSection(
+			"BaseSection",
+			FNewToolMenuDelegate::CreateLambda(
+				[](UToolMenu* InMenu)
+				{
+					FToolMenuSection& UnnamedSection = InMenu->FindOrAddSection(NAME_None);
+					UE::NiagaraSystemEditor::AddMotionSettingsToSection(UnnamedSection);
+				}
+			)
+		);
+	}
+
+	FToolMenuContext MenuContext;
+	{
+		MenuContext.AppendCommandList(InNiagaraSystemEditorViewport->GetCommandList());
+
+		// Add the UnrealEd viewport toolbar context.
+		{
+			UUnrealEdViewportToolbarContext* const ContextObject =
+				UE::UnrealEd::CreateViewportToolbarDefaultContext(InNiagaraSystemEditorViewport);
+
+			MenuContext.AddObject(ContextObject);
+		}
+	}
+
+	return UToolMenus::Get()->GenerateWidget(OldShowMenuName, MenuContext);
+}
+
+void UE::NiagaraSystemEditor::FillSettingsSubmenu(UToolMenu* InMenu)
+{
+	FToolMenuSection& ViewportControlsSection =
+		InMenu->FindOrAddSection("ViewportControls", LOCTEXT("ViewportControlsLabel", "Viewport Controls"));
+
+	AddMotionSettingsToSection(ViewportControlsSection);
+}
+void UE::NiagaraSystemEditor::AddMotionSettingsToSection(FToolMenuSection& InSection)
+{
+	InSection.AddSubMenu(
+		"MotionOptions",
+		LOCTEXT("MotionOptionsSubMenu", "Motion Options"),
+		LOCTEXT("MotionOptionsSubMenu_ToolTip", "Set Motion Options for the Niagara Component"),
+		FNewToolMenuDelegate::CreateLambda(
+			[](UToolMenu* InMenu)
+			{
+				FToolMenuSection& UnnamedSection = InMenu->FindOrAddSection(NAME_None);
 
-	return MotionMenuBuilder.MakeWidget();
+				UnnamedSection.AddMenuEntry(FNiagaraEditorCommands::Get().ToggleMotion);
+
+				UUnrealEdViewportToolbarContext* ViewportToolbarContext =
+					InMenu->FindContext<UUnrealEdViewportToolbarContext>();
+				if (!ViewportToolbarContext)
+				{
+					return;
+				}
+
+				TSharedPtr<SNiagaraSystemViewport> NiagaraSystemViewport =
+					StaticCastSharedPtr<SNiagaraSystemViewport>(ViewportToolbarContext->Viewport.Pin());
+				if (!NiagaraSystemViewport)
+				{
+					return;
+				}
+
+				FToolMenuEntry MotionRateEntry = FToolMenuEntry::InitWidget(
+					"MotionRate",
+					SNew(SSpinBox<float>)
+						.IsEnabled(NiagaraSystemViewport.Get(), &SNiagaraSystemViewport::IsMotionEnabled)
+						.Font(FAppStyle::GetFontStyle(TEXT("MenuItem.Font")))
+						.MinSliderValue(0.0f)
+						.MaxSliderValue(360.0f)
+						.Value(NiagaraSystemViewport.Get(), &SNiagaraSystemViewport::GetMotionRate)
+						.OnValueChanged(NiagaraSystemViewport.Get(), &SNiagaraSystemViewport::SetMotionRate),
+					LOCTEXT("MotionSpeed", "Motion Speed")
+				);
+
+				FToolMenuEntry MotionRadiusEntry = FToolMenuEntry::InitWidget(
+					"MotionRadius",
+					SNew(SSpinBox<float>)
+						.IsEnabled(NiagaraSystemViewport.Get(), &SNiagaraSystemViewport::IsMotionEnabled)
+						.Font(FAppStyle::GetFontStyle(TEXT("MenuItem.Font")))
+						.MinSliderValue(0.0f)
+						.MaxSliderValue(1000.0f)
+						.Value(NiagaraSystemViewport.Get(), &SNiagaraSystemViewport::GetMotionRadius)
+						.OnValueChanged(NiagaraSystemViewport.Get(), &SNiagaraSystemViewport::SetMotionRadius),
+					LOCTEXT("MotionRadius", "Motion Radius")
+				);
+
+				UnnamedSection.AddEntry(MotionRateEntry);
+				UnnamedSection.AddEntry(MotionRadiusEntry);
+			}
+		),
+		false,
+		FSlateIcon()
+	);
 }
 
 #undef LOCTEXT_NAMESPACE
diff --git a/Engine/Plugins/FX/Niagara/Source/NiagaraEditor/Private/Widgets/NiagaraSystemEditorViewportToolbarSections.h b/Engine/Plugins/FX/Niagara/Source/NiagaraEditor/Private/Widgets/NiagaraSystemEditorViewportToolbarSections.h
--- a/Engine/Plugins/FX/Niagara/Source/NiagaraEditor/Private/Widgets/NiagaraSystemEditorViewportToolbarSections.h
+++ b/Engine/Plugins/FX/Niagara/Source/NiagaraEditor/Private/Widgets/NiagaraSystemEditorViewportToolbarSections.h
@@ -9,9 +9,11 @@ class SNiagaraSystemViewport;
 
 namespace UE::NiagaraSystemEditor
 {
-TSharedRef<SWidget> CreateShowMenuWidget(const TSharedRef<SNiagaraSystemViewport>& InNiagaraSystemEditorViewport, bool bInShowViewportStatsToggle = true);
 FToolMenuEntry CreateShowSubmenu();
+void FillShowSubmenu(UToolMenu* InMenu, bool bInShowViewportStatsToggle = true);
 
 TSharedRef<SWidget> CreateMotionMenuWidget(const TSharedRef<SNiagaraSystemViewport>& InNiagaraSystemEditorViewport);
 FToolMenuEntry CreateSettingsSubmenu();
+void FillSettingsSubmenu(UToolMenu* InMenu);
+void AddMotionSettingsToSection(FToolMenuSection& InSection);
 }
diff --git a/Engine/Plugins/FX/Niagara/Source/NiagaraEditor/Private/Widgets/SNiagaraSystemViewportToolBar.cpp b/Engine/Plugins/FX/Niagara/Source/NiagaraEditor/Private/Widgets/SNiagaraSystemViewportToolBar.cpp
--- a/Engine/Plugins/FX/Niagara/Source/NiagaraEditor/Private/Widgets/SNiagaraSystemViewportToolBar.cpp
+++ b/Engine/Plugins/FX/Niagara/Source/NiagaraEditor/Private/Widgets/SNiagaraSystemViewportToolBar.cpp
@@ -7,6 +7,7 @@
 #include "NiagaraEditorCommands.h"
 #include "EditorViewportCommands.h"
 #include "NiagaraEditorSettings.h"
+#include "NiagaraSystemEditorViewportToolbarSections.h"
 #include "SEditorViewportToolBarButton.h"
 #include "PreviewProfileController.h"
 
@@ -92,11 +93,11 @@ void SNiagaraSystemViewportToolBar::ExtendOptionsMenu(FMenuBuilder& OptionsMenuB
 {
 	OptionsMenuBuilder.BeginSection("LevelViewportNavigationOptions", LOCTEXT("NavOptionsMenuHeader", "Navigation Options"));
 	OptionsMenuBuilder.AddMenuEntry(FNiagaraEditorCommands::Get().ToggleOrbit);
-	OptionsMenuBuilder.AddSubMenu(
-		LOCTEXT("MotionOptions", "Motion Options"),
-		LOCTEXT("MotionOptionsToolTip", "Set Motion Options For the Niagara Component"),
-		FNewMenuDelegate::CreateSP(this, &SNiagaraSystemViewportToolBar::MakeMotionMenu)
-	);
+
+	if (TSharedPtr<SNiagaraSystemViewport> NiagaraSystemViewport = WeakViewport.Pin())
+	{
+		OptionsMenuBuilder.AddWidget(UE::NiagaraSystemEditor::CreateMotionMenuWidget(NiagaraSystemViewport.ToSharedRef()), FText());
+	}
 
 	OptionsMenuBuilder.EndSection();
 }
@@ -123,38 +124,4 @@ void SNiagaraSystemViewportToolBar::ExtendLeftAlignedToolbarSlots(TSharedPtr<SHo
 	];
 }
 
-void SNiagaraSystemViewportToolBar::MakeMotionMenu(FMenuBuilder& MenuBuilder) const
-{
-	MenuBuilder.AddMenuEntry(FNiagaraEditorCommands::Get().ToggleMotion);
-
-	if (TSharedPtr<SNiagaraSystemViewport> ViewportSP = WeakViewport.Pin())
-	{
-		TSharedRef<SNiagaraSystemViewport> Viewport = ViewportSP.ToSharedRef();
-
-		MenuBuilder.AddWidget(
-			SNew(SSpinBox<float>)
-				.IsEnabled(Viewport, &SNiagaraSystemViewport::IsMotionEnabled)
-				.Font(FAppStyle::GetFontStyle(TEXT("MenuItem.Font")))
-				.MinSliderValue(0.0f)
-				.MaxSliderValue(360.0f)
-				.Value(Viewport, &SNiagaraSystemViewport::GetMotionRate)
-				.OnValueChanged(Viewport, &SNiagaraSystemViewport::SetMotionRate)
-			,
-			LOCTEXT("MotionSpeed", "Motion Speed")
-		);
-
-		MenuBuilder.AddWidget(
-			SNew(SSpinBox<float>)
-			.IsEnabled(Viewport, &SNiagaraSystemViewport::IsMotionEnabled)
-			.Font(FAppStyle::GetFontStyle(TEXT("MenuItem.Font")))
-			.MinSliderValue(0.0f)
-			.MaxSliderValue(1000.0f)
-			.Value(Viewport, &SNiagaraSystemViewport::GetMotionRadius)
-			.OnValueChanged(Viewport, &SNiagaraSystemViewport::SetMotionRadius)
-			,
-			LOCTEXT("MotionRadius", "Motion Radius")
-		);
-	}
-}
-
 #undef LOCTEXT_NAMESPACE
diff --git a/Engine/Plugins/FX/Niagara/Source/NiagaraEditor/Private/Widgets/SNiagaraSystemViewportToolBar.h b/Engine/Plugins/FX/Niagara/Source/NiagaraEditor/Private/Widgets/SNiagaraSystemViewportToolBar.h
--- a/Engine/Plugins/FX/Niagara/Source/NiagaraEditor/Private/Widgets/SNiagaraSystemViewportToolBar.h
+++ b/Engine/Plugins/FX/Niagara/Source/NiagaraEditor/Private/Widgets/SNiagaraSystemViewportToolBar.h
@@ -36,11 +36,9 @@ private:
 	TWeakPtr<class SNiagaraSystemViewport> WeakViewport = nullptr;
 	TWeakPtr<ISequencer> Sequencer = nullptr;
 	bool bIsSimulationRealtime = true;
-	
+
 	EVisibility GetSimulationRealtimeWarningVisibility() const;
 	FReply OnSimulationRealtimeWarningClicked() const;
 
 	FText GetSimulationSpeedText() const;
-
-	void MakeMotionMenu(FMenuBuilder& MenuBuilder) const;
 };
