- Update scale by speed modules to match stateful

#rnx
#tests Create both types of emitters and compare results

diff --git a/Engine/Plugins/FX/Niagara/Shaders/Private/Stateless/Modules/NiagaraStatelessModule_ScaleMeshSizeBySpeed.ush b/Engine/Plugins/FX/Niagara/Shaders/Private/Stateless/Modules/NiagaraStatelessModule_ScaleMeshSizeBySpeed.ush
--- a/Engine/Plugins/FX/Niagara/Shaders/Private/Stateless/Modules/NiagaraStatelessModule_ScaleMeshSizeBySpeed.ush
+++ b/Engine/Plugins/FX/Niagara/Shaders/Private/Stateless/Modules/NiagaraStatelessModule_ScaleMeshSizeBySpeed.ush
@@ -2,6 +2,8 @@
 
 #include "../NiagaraStatelessCommon.ush"
 
+float3	ScaleMeshSizeBySpeed_ScaleFactorBias;
+float3	ScaleMeshSizeBySpeed_ScaleFactorScale;
 uint2	ScaleMeshSizeBySpeed_ScaleDistribution;
 float	ScaleMeshSizeBySpeed_VelocityNorm;
 
@@ -9,10 +11,13 @@ void ScaleMeshSizeBySpeed_Simulate(inout FStatelessParticle Particle)
 {
 	if (ScaleMeshSizeBySpeed_ScaleDistribution.y > 0)
 	{
-		const float Speed = length2(Particle.Position - Particle.PreviousPosition) * Particle.InvDeltaTime;
-		const float NormSpeed = saturate(Speed * ScaleMeshSizeBySpeed_VelocityNorm);
-
-		Particle.Scale			*= LerpStaticFloat3(ScaleMeshSizeBySpeed_ScaleDistribution, NormSpeed);
-		Particle.PreviousScale	*= LerpStaticFloat3(ScaleMeshSizeBySpeed_ScaleDistribution, NormSpeed);
+		const float3 Velocity	= (Particle.Position - Particle.PreviousPosition) * Particle.InvDeltaTime;
+		const float Speed		= length2(Velocity);
+		const float NormSpeed	= saturate(Speed * ScaleMeshSizeBySpeed_VelocityNorm);
+		const float Interp		= saturate(LerpStaticFloat(ScaleMeshSizeBySpeed_ScaleDistribution, NormSpeed));
+		const float3 Scale		= ScaleMeshSizeBySpeed_ScaleFactorBias + (ScaleMeshSizeBySpeed_ScaleFactorScale * Interp);
+			
+		Particle.Scale			*= Scale;
+		Particle.PreviousScale	*= Scale;
 	}
 }
diff --git a/Engine/Plugins/FX/Niagara/Shaders/Private/Stateless/Modules/NiagaraStatelessModule_ScaleSpriteSizeBySpeed.ush b/Engine/Plugins/FX/Niagara/Shaders/Private/Stateless/Modules/NiagaraStatelessModule_ScaleSpriteSizeBySpeed.ush
--- a/Engine/Plugins/FX/Niagara/Shaders/Private/Stateless/Modules/NiagaraStatelessModule_ScaleSpriteSizeBySpeed.ush
+++ b/Engine/Plugins/FX/Niagara/Shaders/Private/Stateless/Modules/NiagaraStatelessModule_ScaleSpriteSizeBySpeed.ush
@@ -2,6 +2,8 @@
 
 #include "../NiagaraStatelessCommon.ush"
 
+float2	ScaleSpriteSizeBySpeed_ScaleFactorBias;
+float2	ScaleSpriteSizeBySpeed_ScaleFactorScale;
 uint2	ScaleSpriteSizeBySpeed_ScaleDistribution;
 float	ScaleSpriteSizeBySpeed_VelocityNorm;
 
@@ -9,10 +11,13 @@ void ScaleSpriteSizeBySpeed_Simulate(inout FStatelessParticle Particle)
 {
 	if (ScaleSpriteSizeBySpeed_ScaleDistribution.y > 0)
 	{
-		const float Speed = length2(Particle.Position - Particle.PreviousPosition) * Particle.InvDeltaTime;
-		const float NormSpeed = saturate(Speed * ScaleSpriteSizeBySpeed_VelocityNorm);
-
-		Particle.SpriteSize			*= LerpStaticFloat2(ScaleSpriteSizeBySpeed_ScaleDistribution, NormSpeed);
-		Particle.PreviousSpriteSize	*= LerpStaticFloat2(ScaleSpriteSizeBySpeed_ScaleDistribution, NormSpeed);
+		const float3 Velocity	= (Particle.Position - Particle.PreviousPosition) * Particle.InvDeltaTime;
+		const float Speed		= length2(Velocity);
+		const float NormSpeed	= saturate(Speed * ScaleSpriteSizeBySpeed_VelocityNorm);
+		const float Interp		= saturate(LerpStaticFloat(ScaleSpriteSizeBySpeed_ScaleDistribution, NormSpeed));
+		const float2 Scale		= ScaleSpriteSizeBySpeed_ScaleFactorBias + (ScaleSpriteSizeBySpeed_ScaleFactorScale * Interp);
+			
+		Particle.SpriteSize			*= Scale;
+		Particle.PreviousSpriteSize	*= Scale;
 	}
 }
diff --git a/Engine/Plugins/FX/Niagara/Source/Niagara/Internal/Stateless/Modules/NiagaraStatelessModule_ScaleMeshSizeBySpeed.h b/Engine/Plugins/FX/Niagara/Source/Niagara/Internal/Stateless/Modules/NiagaraStatelessModule_ScaleMeshSizeBySpeed.h
--- a/Engine/Plugins/FX/Niagara/Source/Niagara/Internal/Stateless/Modules/NiagaraStatelessModule_ScaleMeshSizeBySpeed.h
+++ b/Engine/Plugins/FX/Niagara/Source/Niagara/Internal/Stateless/Modules/NiagaraStatelessModule_ScaleMeshSizeBySpeed.h
@@ -5,7 +5,6 @@
 #include "Stateless/NiagaraStatelessModule.h"
 #include "Stateless/NiagaraStatelessEmitterDataBuildContext.h"
 #include "Stateless/NiagaraStatelessModuleShaderParameters.h"
-#include "Stateless/NiagaraStatelessParticleSimContext.h"
 
 #include "NiagaraStatelessModule_ScaleMeshSizeBySpeed.generated.h"
 
@@ -15,108 +14,34 @@ class UNiagaraStatelessModule_ScaleMeshSizeBySpeed : public UNiagaraStatelessMod
 {
 	GENERATED_BODY()
 
-	struct FModuleBuiltData
-	{
-		float			VelocityNorm = 0.0f;
-		FUintVector2	ScaleDistribution = FUintVector2::ZeroValue;
-
-		int32			PositionVariableOffset = INDEX_NONE;
-		int32			PreviousPositionVariableOffset = INDEX_NONE;
-		int32			ScaleVariableVariableOffset = INDEX_NONE;
-		int32			PreviousScaleVariableVariableOffset = INDEX_NONE;
-	};
+	static constexpr float DefaultVelocity = 1000.0f;
 
 public:
 	using FParameters = NiagaraStateless::FScaleMeshSizeBySpeedModule_ShaderParameters;
 
-	UPROPERTY(EditAnywhere, Category = "Parameters", meta = (ClampMin = "0.01", UIMin = "0.01"))
-	float VelocityThreshold = 1000.0f;
-
-	UPROPERTY(EditAnywhere, Category = "Parameters", meta = (DisplayName = "Scale", DisableRangeDistribution, DisableBindingDistribution))
-	FNiagaraDistributionVector3 ScaleDistribution = FNiagaraDistributionVector3(1.0f);
-
-	virtual void BuildEmitterData(const FNiagaraStatelessEmitterDataBuildContext& BuildContext) const override
-	{
-		FModuleBuiltData* BuiltData = BuildContext.AllocateBuiltData<FModuleBuiltData>();
-		if (!IsModuleEnabled())
-		{
-			return;
-		}
-
-		const FNiagaraStatelessGlobals& StatelessGlobals = FNiagaraStatelessGlobals::Get();
-		BuiltData->PositionVariableOffset				= BuildContext.FindParticleVariableIndex(StatelessGlobals.PositionVariable);
-		BuiltData->PreviousPositionVariableOffset		= BuildContext.FindParticleVariableIndex(StatelessGlobals.PreviousPositionVariable);
-		BuiltData->ScaleVariableVariableOffset			= BuildContext.FindParticleVariableIndex(StatelessGlobals.ScaleVariable);
-		BuiltData->PreviousScaleVariableVariableOffset	= BuildContext.FindParticleVariableIndex(StatelessGlobals.PreviousScaleVariable);
-
-		if (BuiltData->ScaleVariableVariableOffset == INDEX_NONE && BuiltData->PreviousScaleVariableVariableOffset == INDEX_NONE)
-		{
-			return;
-		}
-
-		BuiltData->VelocityNorm = VelocityThreshold > 0.0f ? 1.0f / (VelocityThreshold * VelocityThreshold) : 0.0f;
-		if (ScaleDistribution.IsCurve() && ScaleDistribution.Values.Num() > 1)
-		{
-			BuiltData->ScaleDistribution.X = BuildContext.AddStaticData(ScaleDistribution.Values);
-			BuiltData->ScaleDistribution.Y = ScaleDistribution.Values.Num() - 1;
-		}
-		else
-		{
-			const FVector3f Values[] = { FVector3f::One(), ScaleDistribution.Values.Num() > 0 ? ScaleDistribution.Values[0] : FVector3f::One() };
-			BuiltData->ScaleDistribution.X = BuildContext.AddStaticData(Values);
-			BuiltData->ScaleDistribution.Y = 1;
-		}
-		BuildContext.AddParticleSimulationExecSimulate(&UNiagaraStatelessModule_ScaleMeshSizeBySpeed::ParticleSimulate);
-	}
-
-	virtual void BuildShaderParameters(FNiagaraStatelessShaderParametersBuilder& ShaderParametersBuilder) const override
-	{
-		ShaderParametersBuilder.AddParameterNestedStruct<FParameters>();
-	}
-
-	virtual void SetShaderParameters(const FNiagaraStatelessSetShaderParameterContext& SetShaderParameterContext) const override
-	{
-		const FModuleBuiltData* ModuleBuiltData = SetShaderParameterContext.ReadBuiltData<FModuleBuiltData>();
-
-		FParameters* Parameters = SetShaderParameterContext.GetParameterNestedStruct<FParameters>();
-		Parameters->ScaleMeshSizeBySpeed_VelocityNorm		= ModuleBuiltData->VelocityNorm;
-		Parameters->ScaleMeshSizeBySpeed_ScaleDistribution	= ModuleBuiltData->ScaleDistribution;
-	}
-
-	static void ParticleSimulate(const NiagaraStateless::FParticleSimulationContext& ParticleSimulationContext)
-	{
-		using namespace NiagaraStateless;
+	UPROPERTY(EditAnywhere, Category = "Parameters", meta = (ClampMin = "0.01", UIMin = "0.01", DisableRangeDistribution))
+	FNiagaraDistributionRangeFloat VelocityThreshold = FNiagaraDistributionRangeFloat(DefaultVelocity);
 
-		const FModuleBuiltData* ModuleBuiltData = ParticleSimulationContext.ReadBuiltData<FModuleBuiltData>();
+	UPROPERTY(EditAnywhere, Category = "Parameters", meta = (ClampMin = "0.01", UIMin = "0.01", DisableRangeDistribution))
+	FNiagaraDistributionRangeVector3 MinScaleFactor = FNiagaraDistributionRangeVector3(FVector3f(1.0f, 1.0f, 1.0f));
 
-		for (uint32 i = 0; i < ParticleSimulationContext.GetNumInstances(); ++i)
-		{
-			const FVector3f Position			= ParticleSimulationContext.ReadParticleVariable(ModuleBuiltData->PositionVariableOffset, i, FVector3f::OneVector);
-			const FVector3f PreviousPosition	= ParticleSimulationContext.ReadParticleVariable(ModuleBuiltData->PreviousPositionVariableOffset, i, FVector3f::OneVector);
-			const float Speed					= (Position - PreviousPosition).SquaredLength() * ParticleSimulationContext.GetInvDeltaTime();
-			const float NormSpeed				= FMath::Clamp(Speed * ModuleBuiltData->VelocityNorm, 0.0f, 1.0f);
+	UPROPERTY(EditAnywhere, Category = "Parameters", meta = (ClampMin = "0.01", UIMin = "0.01", DisableRangeDistribution))
+	FNiagaraDistributionRangeVector3 MaxScaleFactor = FNiagaraDistributionRangeVector3(FVector3f(2.0f, 2.0f, 2.0f));
 
-			FVector3f Scale			= ParticleSimulationContext.ReadParticleVariable(ModuleBuiltData->ScaleVariableVariableOffset, i, FVector3f::OneVector);
-			FVector3f PreviousScale	= ParticleSimulationContext.ReadParticleVariable(ModuleBuiltData->PreviousScaleVariableVariableOffset, i, FVector3f::OneVector);
+	UPROPERTY(EditAnywhere, Category = "Parameters")
+	bool bSampleScaleFactorByCurve = false;
 
-			const FVector3f ScaleModifier = ParticleSimulationContext.LerpStaticFloat<FVector3f>(ModuleBuiltData->ScaleDistribution, NormSpeed);
-			Scale			*= ScaleModifier;
-			PreviousScale	*= ScaleModifier;
+	UPROPERTY(EditAnywhere, Category = "Parameters", meta = (DisableRangeDistribution, DisableBindingDistribution, EditCondition = "bSampleScaleFactorByCurve", EditConditionHides))
+	FNiagaraDistributionFloat SampleFactorCurve = FNiagaraDistributionFloat({0.0f, 1.0f});
 
-			ParticleSimulationContext.WriteParticleVariable(ModuleBuiltData->ScaleVariableVariableOffset, i, Scale);
-			ParticleSimulationContext.WriteParticleVariable(ModuleBuiltData->PreviousScaleVariableVariableOffset, i, PreviousScale);
-		}
-	}
+	virtual void BuildEmitterData(const FNiagaraStatelessEmitterDataBuildContext& BuildContext) const override;
+	virtual void BuildShaderParameters(FNiagaraStatelessShaderParametersBuilder& ShaderParametersBuilder) const override;
+	virtual void SetShaderParameters(const FNiagaraStatelessSetShaderParameterContext& SetShaderParameterContext) const override;
 
 #if WITH_EDITOR
 	virtual bool CanDisableModule() const override { return true; }
 #endif
 #if WITH_EDITORONLY_DATA
-	virtual void GetOutputVariables(TArray<FNiagaraVariableBase>& OutVariables) const override
-	{
-		const FNiagaraStatelessGlobals& StatelessGlobals = FNiagaraStatelessGlobals::Get();
-		OutVariables.AddUnique(StatelessGlobals.ScaleVariable);
-		OutVariables.AddUnique(StatelessGlobals.PreviousScaleVariable);
-	}
+	virtual void GetOutputVariables(TArray<FNiagaraVariableBase>& OutVariables) const override;
 #endif
 };
diff --git a/Engine/Plugins/FX/Niagara/Source/Niagara/Internal/Stateless/Modules/NiagaraStatelessModule_ScaleSpriteSizeBySpeed.h b/Engine/Plugins/FX/Niagara/Source/Niagara/Internal/Stateless/Modules/NiagaraStatelessModule_ScaleSpriteSizeBySpeed.h
--- a/Engine/Plugins/FX/Niagara/Source/Niagara/Internal/Stateless/Modules/NiagaraStatelessModule_ScaleSpriteSizeBySpeed.h
+++ b/Engine/Plugins/FX/Niagara/Source/Niagara/Internal/Stateless/Modules/NiagaraStatelessModule_ScaleSpriteSizeBySpeed.h
@@ -5,7 +5,6 @@
 #include "Stateless/NiagaraStatelessModule.h"
 #include "Stateless/NiagaraStatelessEmitterDataBuildContext.h"
 #include "Stateless/NiagaraStatelessModuleShaderParameters.h"
-#include "Stateless/NiagaraStatelessParticleSimContext.h"
 
 #include "NiagaraStatelessModule_ScaleSpriteSizeBySpeed.generated.h"
 
@@ -17,105 +16,45 @@ class UNiagaraStatelessModule_ScaleSpriteSizeBySpeed : public UNiagaraStatelessM
 
 	struct FModuleBuiltData
 	{
-		float			VelocityNorm = 0.0f;
-		FUintVector2	ScaleDistribution = FUintVector2::ZeroValue;
-
-		int32			PositionVariableOffset = INDEX_NONE;
-		int32			PreviousPositionVariableOffset = INDEX_NONE;
-		int32			SpriteSizeVariableOffset = INDEX_NONE;
-		int32			PreviousSpriteSizeVariableOffset = INDEX_NONE;
+		FNiagaraStatelessRangeVector2	MinScaleFactor = FNiagaraStatelessRangeVector2(FVector2f::One());
+		FNiagaraStatelessRangeVector2	MaxScaleFactor = FNiagaraStatelessRangeVector2(FVector2f::One());
+		FNiagaraStatelessRangeFloat		VelocityNorm = FNiagaraStatelessRangeFloat(0.0f);
+		FUintVector2					ScaleDistribution = FUintVector2::ZeroValue;
+
+		int32							PositionVariableOffset = INDEX_NONE;
+		int32							PreviousPositionVariableOffset = INDEX_NONE;
+		int32							SpriteSizeVariableOffset = INDEX_NONE;
+		int32							PreviousSpriteSizeVariableOffset = INDEX_NONE;
 	};
 
+	static constexpr float DefaultVelocity = 1000.0f;
+
 public:
 	using FParameters = NiagaraStateless::FScaleSpriteSizeBySpeedModule_ShaderParameters;
 
-	UPROPERTY(EditAnywhere, Category = "Parameters", meta = (ClampMin = "0.01", UIMin = "0.01"))
-	float VelocityThreshold = 1000.0f;
-
-	UPROPERTY(EditAnywhere, Category = "Parameters", meta = (DisplayName = "Scale", DisableRangeDistribution, DisableBindingDistribution))
-	FNiagaraDistributionVector2 ScaleDistribution = FNiagaraDistributionVector2(1.0f);
-
-	virtual void BuildEmitterData(const FNiagaraStatelessEmitterDataBuildContext& BuildContext) const override
-	{
-		FModuleBuiltData* BuiltData = BuildContext.AllocateBuiltData<FModuleBuiltData>();
-		if (!IsModuleEnabled())
-		{
-			return;
-		}
-
-		const FNiagaraStatelessGlobals& StatelessGlobals = FNiagaraStatelessGlobals::Get();
-		BuiltData->PositionVariableOffset			= BuildContext.FindParticleVariableIndex(StatelessGlobals.PositionVariable);
-		BuiltData->PreviousPositionVariableOffset	= BuildContext.FindParticleVariableIndex(StatelessGlobals.PreviousPositionVariable);
-		BuiltData->SpriteSizeVariableOffset			= BuildContext.FindParticleVariableIndex(StatelessGlobals.SpriteSizeVariable);
-		BuiltData->PreviousSpriteSizeVariableOffset	= BuildContext.FindParticleVariableIndex(StatelessGlobals.PreviousSpriteSizeVariable);
+	UPROPERTY(EditAnywhere, Category = "Parameters", meta = (ClampMin = "0.01", UIMin = "0.01", DisableRangeDistribution))
+	FNiagaraDistributionRangeFloat VelocityThreshold = FNiagaraDistributionRangeFloat(DefaultVelocity);
 
-		if (BuiltData->SpriteSizeVariableOffset == INDEX_NONE && BuiltData->PreviousSpriteSizeVariableOffset == INDEX_NONE)
-		{
-			return;
-		}
+	UPROPERTY(EditAnywhere, Category = "Parameters", meta = (ClampMin = "0.01", UIMin = "0.01", DisableRangeDistribution))
+	FNiagaraDistributionRangeVector2 MinScaleFactor = FNiagaraDistributionRangeVector2(FVector2f(1.0f, 1.0f));
 
-		BuiltData->VelocityNorm			= VelocityThreshold > 0.0f ? 1.0f / (VelocityThreshold * VelocityThreshold) : 0.0f;
-		if (ScaleDistribution.IsCurve() && ScaleDistribution.Values.Num() > 1)
-		{
-			BuiltData->ScaleDistribution.X = BuildContext.AddStaticData(ScaleDistribution.Values);
-			BuiltData->ScaleDistribution.Y = ScaleDistribution.Values.Num() - 1;
-		}
-		else
-		{
-			const FVector2f Values[] = { FVector2f::One(), ScaleDistribution.Values.Num() > 0 ? ScaleDistribution.Values[0] : FVector2f::One() };
-			BuiltData->ScaleDistribution.X = BuildContext.AddStaticData(Values);
-			BuiltData->ScaleDistribution.Y = 1;
-		}
-		BuildContext.AddParticleSimulationExecSimulate(&UNiagaraStatelessModule_ScaleSpriteSizeBySpeed::ParticleSimulate);
-	}
+	UPROPERTY(EditAnywhere, Category = "Parameters", meta = (ClampMin = "0.01", UIMin = "0.01", DisableRangeDistribution))
+	FNiagaraDistributionRangeVector2 MaxScaleFactor = FNiagaraDistributionRangeVector2(FVector2f(2.0f, 2.0f));
 
-	virtual void BuildShaderParameters(FNiagaraStatelessShaderParametersBuilder& ShaderParametersBuilder) const override
-	{
-		ShaderParametersBuilder.AddParameterNestedStruct<FParameters>();
-	}
+	UPROPERTY(EditAnywhere, Category = "Parameters")
+	bool bSampleScaleFactorByCurve = false;
 
-	virtual void SetShaderParameters(const FNiagaraStatelessSetShaderParameterContext& SetShaderParameterContext) const override
-	{
-		const FModuleBuiltData* ModuleBuiltData = SetShaderParameterContext.ReadBuiltData<FModuleBuiltData>();
+	UPROPERTY(EditAnywhere, Category = "Parameters", meta = (DisableRangeDistribution, DisableBindingDistribution, EditCondition = "bSampleScaleFactorByCurve", EditConditionHides))
+	FNiagaraDistributionFloat SampleFactorCurve = FNiagaraDistributionFloat({0.0f, 1.0f});
 
-		FParameters* Parameters = SetShaderParameterContext.GetParameterNestedStruct<FParameters>();
-		Parameters->ScaleSpriteSizeBySpeed_VelocityNorm			= ModuleBuiltData->VelocityNorm;
-		Parameters->ScaleSpriteSizeBySpeed_ScaleDistribution	= ModuleBuiltData->ScaleDistribution;
-	}
+	virtual void BuildEmitterData(const FNiagaraStatelessEmitterDataBuildContext& BuildContext) const override;
+	virtual void BuildShaderParameters(FNiagaraStatelessShaderParametersBuilder& ShaderParametersBuilder) const override;
+	virtual void SetShaderParameters(const FNiagaraStatelessSetShaderParameterContext& SetShaderParameterContext) const override;
 
-	static void ParticleSimulate(const NiagaraStateless::FParticleSimulationContext& ParticleSimulationContext)
-	{
-		using namespace NiagaraStateless;
-
-		const FModuleBuiltData* ModuleBuiltData = ParticleSimulationContext.ReadBuiltData<FModuleBuiltData>();
-
-		for (uint32 i = 0; i < ParticleSimulationContext.GetNumInstances(); ++i)
-		{
-			const FVector3f Position			= ParticleSimulationContext.ReadParticleVariable(ModuleBuiltData->PositionVariableOffset, i, FVector3f::OneVector);
-			const FVector3f PreviousPosition	= ParticleSimulationContext.ReadParticleVariable(ModuleBuiltData->PreviousPositionVariableOffset, i, FVector3f::OneVector);
-			const float Speed					= (Position - PreviousPosition).SquaredLength() * ParticleSimulationContext.GetInvDeltaTime();
-			const float NormSpeed				= FMath::Clamp(Speed * ModuleBuiltData->VelocityNorm, 0.0f, 1.0f);
-
-			FVector2f SpriteSize			= ParticleSimulationContext.ReadParticleVariable(ModuleBuiltData->SpriteSizeVariableOffset, i, FVector2f::One());
-			FVector2f PreviousSpriteSize	= ParticleSimulationContext.ReadParticleVariable(ModuleBuiltData->PreviousSpriteSizeVariableOffset, i, FVector2f::One());
-
-			const FVector2f ScaleModifier = ParticleSimulationContext.LerpStaticFloat<FVector2f>(ModuleBuiltData->ScaleDistribution, NormSpeed);
-			SpriteSize			*= ScaleModifier;
-			PreviousSpriteSize	*= ScaleModifier;
-
-			ParticleSimulationContext.WriteParticleVariable(ModuleBuiltData->SpriteSizeVariableOffset, i, SpriteSize);
-			ParticleSimulationContext.WriteParticleVariable(ModuleBuiltData->PreviousSpriteSizeVariableOffset, i, PreviousSpriteSize);
-		}
-	}
 #if WITH_EDITOR
 	virtual bool CanDisableModule() const override { return true; }
 #endif
 #if WITH_EDITORONLY_DATA
-	virtual void GetOutputVariables(TArray<FNiagaraVariableBase>& OutVariables) const override
-	{
-		const FNiagaraStatelessGlobals& StatelessGlobals = FNiagaraStatelessGlobals::Get();
-		OutVariables.AddUnique(StatelessGlobals.SpriteSizeVariable);
-		OutVariables.AddUnique(StatelessGlobals.PreviousSpriteSizeVariable);
-	}
+	virtual void GetOutputVariables(TArray<FNiagaraVariableBase>& OutVariables) const override;
 #endif
 };
diff --git a/Engine/Plugins/FX/Niagara/Source/Niagara/Internal/Stateless/NiagaraStatelessDistribution.h b/Engine/Plugins/FX/Niagara/Source/Niagara/Internal/Stateless/NiagaraStatelessDistribution.h
--- a/Engine/Plugins/FX/Niagara/Source/Niagara/Internal/Stateless/NiagaraStatelessDistribution.h
+++ b/Engine/Plugins/FX/Niagara/Source/Niagara/Internal/Stateless/NiagaraStatelessDistribution.h
@@ -215,6 +215,7 @@ struct FNiagaraDistributionFloat : public FNiagaraDistributionBase
 
 	FNiagaraDistributionFloat() = default;
 	explicit FNiagaraDistributionFloat(float ConstantValue) { InitConstant(ConstantValue); }
+	explicit FNiagaraDistributionFloat(std::initializer_list<float> CurvePoints) { InitCurve(CurvePoints); }
 
 	UPROPERTY(EditAnywhere, Category = "Parameters")
 	TArray<float> Values;
@@ -223,6 +224,10 @@ struct FNiagaraDistributionFloat : public FNiagaraDistributionBase
 	FVector2f ValuesTimeRange = FVector2f(0.0f, 1.0f);
 
 	NIAGARA_API void InitConstant(float Value);
+	NIAGARA_API void InitCurve(std::initializer_list<float> CurvePoints);
+#if WITH_EDITORONLY_DATA
+	NIAGARA_API void InitCurve(const TArray<FRichCurveKey>& CurveKeys);
+#endif
 	NIAGARA_API FNiagaraStatelessRangeFloat CalculateRange(const float Default = 0.0f) const;
 
 #if WITH_EDITORONLY_DATA
@@ -270,6 +275,8 @@ struct FNiagaraDistributionVector3 : public FNiagaraDistributionBase
 	FNiagaraDistributionVector3() = default;
 	explicit FNiagaraDistributionVector3(const float ConstantValue) { InitConstant(ConstantValue); }
 	explicit FNiagaraDistributionVector3(const FVector3f& ConstantValue) { InitConstant(ConstantValue); }
+	explicit FNiagaraDistributionVector3(std::initializer_list<float> CurvePoints) { InitCurve(CurvePoints); }
+	explicit FNiagaraDistributionVector3(std::initializer_list<FVector3f> CurvePoints) { InitCurve(CurvePoints); }
 
 	UPROPERTY(EditAnywhere, Category = "Parameters")
 	TArray<FVector3f> Values;
@@ -279,6 +286,8 @@ struct FNiagaraDistributionVector3 : public FNiagaraDistributionBase
 
 	NIAGARA_API void InitConstant(const float Value);
 	NIAGARA_API void InitConstant(const FVector3f& Value);
+	NIAGARA_API void InitCurve(std::initializer_list<float> CurvePoints);
+	NIAGARA_API void InitCurve(std::initializer_list<FVector3f> CurvePoints);
 	NIAGARA_API FNiagaraStatelessRangeVector3 CalculateRange(const FVector3f& Default = FVector3f::ZeroVector) const;
 
 #if WITH_EDITORONLY_DATA
diff --git a/Engine/Plugins/FX/Niagara/Source/Niagara/Internal/Stateless/NiagaraStatelessSetShaderParameterContext.h b/Engine/Plugins/FX/Niagara/Source/Niagara/Internal/Stateless/NiagaraStatelessSetShaderParameterContext.h
--- a/Engine/Plugins/FX/Niagara/Source/Niagara/Internal/Stateless/NiagaraStatelessSetShaderParameterContext.h
+++ b/Engine/Plugins/FX/Niagara/Source/Niagara/Internal/Stateless/NiagaraStatelessSetShaderParameterContext.h
@@ -62,6 +62,11 @@ public:
 	void ConvertRangeToScaleBias(const FNiagaraStatelessRangeVector3& Range, FVector3f& OutScale, FVector3f& OutBias) const { OutScale = Range.GetScale(); GetRendererParameterValue(OutBias, Range.ParameterOffset, Range.Min); }
 	void ConvertRangeToScaleBias(const FNiagaraStatelessRangeColor& Range, FLinearColor& OutScale, FLinearColor& OutBias) const { OutScale = Range.GetScale(); GetRendererParameterValue(OutBias, Range.ParameterOffset, Range.Min); }
 
+	float ConvertRangeToValue(const FNiagaraStatelessRangeFloat& Range) const { float OutValue; GetRendererParameterValue(OutValue, Range.ParameterOffset, Range.Min); return OutValue; }
+	FVector2f ConvertRangeToValue(const FNiagaraStatelessRangeVector2& Range) const { FVector2f OutValue; GetRendererParameterValue(OutValue, Range.ParameterOffset, Range.Min); return OutValue; }
+	FVector3f ConvertRangeToValue(const FNiagaraStatelessRangeVector3& Range) const { FVector3f OutValue; GetRendererParameterValue(OutValue, Range.ParameterOffset, Range.Min); return OutValue; }
+	FLinearColor ConvertRangeToValue(const FNiagaraStatelessRangeColor& Range) const { FLinearColor OutValue; GetRendererParameterValue(OutValue, Range.ParameterOffset, Range.Min); return OutValue; }
+
 protected:
 #if DO_CHECK
 	void ValidateIncludeStructType(uint32 StructOffset, const FShaderParametersMetadata* StructMetaData) const;
diff --git a/Engine/Plugins/FX/Niagara/Source/Niagara/Private/Stateless/Modules/NiagaraStatelessModule_ScaleMeshSizeBySpeed.cpp b/Engine/Plugins/FX/Niagara/Source/Niagara/Private/Stateless/Modules/NiagaraStatelessModule_ScaleMeshSizeBySpeed.cpp
new file mode 100644
--- /dev/null
+++ b/Engine/Plugins/FX/Niagara/Source/Niagara/Private/Stateless/Modules/NiagaraStatelessModule_ScaleMeshSizeBySpeed.cpp
@@ -0,0 +1,129 @@
+// Copyright Epic Games, Inc. All Rights Reserved.
+
+#include "Stateless/Modules/NiagaraStatelessModule_ScaleMeshSizeBySpeed.h"
+#include "Stateless/NiagaraStatelessParticleSimContext.h"
+#include "Curves/RichCurve.h"
+
+namespace NSMScaleMeshSizeBySpeedPrivate
+{
+	struct FModuleBuiltData
+	{
+		FNiagaraStatelessRangeVector3	MinScaleFactor = FNiagaraStatelessRangeVector3(FVector3f::OneVector);
+		FNiagaraStatelessRangeVector3	MaxScaleFactor = FNiagaraStatelessRangeVector3(FVector3f::OneVector);
+		FNiagaraStatelessRangeFloat		VelocityNorm = FNiagaraStatelessRangeFloat(0.0f);
+		FUintVector2					ScaleDistribution = FUintVector2::ZeroValue;
+
+		int32							PositionVariableOffset = INDEX_NONE;
+		int32							PreviousPositionVariableOffset = INDEX_NONE;
+		int32							ScaleVariableVariableOffset = INDEX_NONE;
+		int32							PreviousScaleVariableVariableOffset = INDEX_NONE;
+	};
+
+
+	static void ParticleSimulate(const NiagaraStateless::FParticleSimulationContext& ParticleSimulationContext)
+	{
+		using namespace NiagaraStateless;
+
+		const FModuleBuiltData* ModuleBuiltData = ParticleSimulationContext.ReadBuiltData<FModuleBuiltData>();
+		const UNiagaraStatelessModule_ScaleMeshSizeBySpeed::FParameters* Parameters = ParticleSimulationContext.ReadParameterNestedStruct<UNiagaraStatelessModule_ScaleMeshSizeBySpeed::FParameters>();
+
+		for (uint32 i = 0; i < ParticleSimulationContext.GetNumInstances(); ++i)
+		{
+			const FVector3f Position = ParticleSimulationContext.ReadParticleVariable(ModuleBuiltData->PositionVariableOffset, i, FVector3f::OneVector);
+			const FVector3f PreviousPosition = ParticleSimulationContext.ReadParticleVariable(ModuleBuiltData->PreviousPositionVariableOffset, i, FVector3f::OneVector);
+			const FVector3f Velocity = (Position - PreviousPosition) * ParticleSimulationContext.GetInvDeltaTime();
+			const float Speed = Velocity.SquaredLength();
+			const float NormSpeed = FMath::Clamp(Speed * Parameters->ScaleMeshSizeBySpeed_VelocityNorm, 0.0f, 1.0f);
+			const float Interp = ParticleSimulationContext.LerpStaticFloat<float>(Parameters->ScaleMeshSizeBySpeed_ScaleDistribution, NormSpeed);
+			const FVector3f	Scale = Parameters->ScaleMeshSizeBySpeed_ScaleFactorBias + (Parameters->ScaleMeshSizeBySpeed_ScaleFactorScale * FMath::Clamp(Interp, 0.0f, 1.0f));
+
+			FVector3f ParticleScale = ParticleSimulationContext.ReadParticleVariable(ModuleBuiltData->ScaleVariableVariableOffset, i, FVector3f::OneVector);
+			FVector3f ParticlePreviousScale = ParticleSimulationContext.ReadParticleVariable(ModuleBuiltData->PreviousScaleVariableVariableOffset, i, FVector3f::OneVector);
+
+			ParticleScale *= Scale;
+			ParticlePreviousScale *= Scale;
+
+			ParticleSimulationContext.WriteParticleVariable(ModuleBuiltData->ScaleVariableVariableOffset, i, ParticleScale);
+			ParticleSimulationContext.WriteParticleVariable(ModuleBuiltData->PreviousScaleVariableVariableOffset, i, ParticlePreviousScale);
+		}
+	}
+}
+
+void UNiagaraStatelessModule_ScaleMeshSizeBySpeed::BuildEmitterData(const FNiagaraStatelessEmitterDataBuildContext& BuildContext) const
+{
+	using namespace NSMScaleMeshSizeBySpeedPrivate;
+
+	FModuleBuiltData* BuiltData = BuildContext.AllocateBuiltData<FModuleBuiltData>();
+	if (!IsModuleEnabled())
+	{
+		return;
+	}
+
+	const FNiagaraStatelessGlobals& StatelessGlobals = FNiagaraStatelessGlobals::Get();
+	BuiltData->PositionVariableOffset				= BuildContext.FindParticleVariableIndex(StatelessGlobals.PositionVariable);
+	BuiltData->PreviousPositionVariableOffset		= BuildContext.FindParticleVariableIndex(StatelessGlobals.PreviousPositionVariable);
+	BuiltData->ScaleVariableVariableOffset			= BuildContext.FindParticleVariableIndex(StatelessGlobals.ScaleVariable);
+	BuiltData->PreviousScaleVariableVariableOffset	= BuildContext.FindParticleVariableIndex(StatelessGlobals.PreviousScaleVariable);
+
+	if (BuiltData->ScaleVariableVariableOffset == INDEX_NONE && BuiltData->PreviousScaleVariableVariableOffset == INDEX_NONE)
+	{
+		return;
+	}
+
+	BuiltData->VelocityNorm = BuildContext.ConvertDistributionToRange(VelocityThreshold, DefaultVelocity);
+	BuiltData->MinScaleFactor = BuildContext.ConvertDistributionToRange(MinScaleFactor, FVector3f::OneVector);
+	BuiltData->MaxScaleFactor = BuildContext.ConvertDistributionToRange(MaxScaleFactor, FVector3f::OneVector);
+	if ( bSampleScaleFactorByCurve )
+	{
+		if (SampleFactorCurve.IsCurve() && SampleFactorCurve.Values.Num() > 1)
+		{
+			BuiltData->ScaleDistribution.X = BuildContext.AddStaticData(SampleFactorCurve.Values);
+			BuiltData->ScaleDistribution.Y = SampleFactorCurve.Values.Num() - 1;
+		}
+		else
+		{
+			const float Value = SampleFactorCurve.Values.Num() > 0 ? SampleFactorCurve.Values[0] : 1.0f;
+			const float Values[] = { Value, Value };
+			BuiltData->ScaleDistribution.X = BuildContext.AddStaticData(Values);
+			BuiltData->ScaleDistribution.Y = 1;
+		}
+	}
+	else
+	{
+		const float Values[] = { 1.0f, 1.0f };
+		BuiltData->ScaleDistribution.X = BuildContext.AddStaticData(Values);
+		BuiltData->ScaleDistribution.Y = 1;
+	}
+	BuildContext.AddParticleSimulationExecSimulate(&NSMScaleMeshSizeBySpeedPrivate::ParticleSimulate);
+}
+
+void UNiagaraStatelessModule_ScaleMeshSizeBySpeed::BuildShaderParameters(FNiagaraStatelessShaderParametersBuilder& ShaderParametersBuilder) const
+{
+	ShaderParametersBuilder.AddParameterNestedStruct<FParameters>();
+}
+
+void UNiagaraStatelessModule_ScaleMeshSizeBySpeed::SetShaderParameters(const FNiagaraStatelessSetShaderParameterContext& SetShaderParameterContext) const
+{
+	using namespace NSMScaleMeshSizeBySpeedPrivate;
+
+	const FModuleBuiltData* ModuleBuiltData = SetShaderParameterContext.ReadBuiltData<FModuleBuiltData>();
+
+	const float VelocityNorm	= SetShaderParameterContext.ConvertRangeToValue(ModuleBuiltData->VelocityNorm);
+	const FVector3f MinScale	= SetShaderParameterContext.ConvertRangeToValue(ModuleBuiltData->MinScaleFactor);
+	const FVector3f MaxScale	= SetShaderParameterContext.ConvertRangeToValue(ModuleBuiltData->MaxScaleFactor);
+
+	FParameters* Parameters = SetShaderParameterContext.GetParameterNestedStruct<FParameters>();
+	Parameters->ScaleMeshSizeBySpeed_ScaleFactorBias	= MinScale;
+	Parameters->ScaleMeshSizeBySpeed_ScaleFactorScale	= MaxScale - MinScale;
+	Parameters->ScaleMeshSizeBySpeed_VelocityNorm		= VelocityNorm > 0.0f ? 1.0f / (VelocityNorm * VelocityNorm) : 0.0f;
+	Parameters->ScaleMeshSizeBySpeed_ScaleDistribution	= ModuleBuiltData->ScaleDistribution;
+}
+
+#if WITH_EDITORONLY_DATA
+void UNiagaraStatelessModule_ScaleMeshSizeBySpeed::GetOutputVariables(TArray<FNiagaraVariableBase>& OutVariables) const
+{
+	const FNiagaraStatelessGlobals& StatelessGlobals = FNiagaraStatelessGlobals::Get();
+	OutVariables.AddUnique(StatelessGlobals.ScaleVariable);
+	OutVariables.AddUnique(StatelessGlobals.PreviousScaleVariable);
+}
+#endif
diff --git a/Engine/Plugins/FX/Niagara/Source/Niagara/Private/Stateless/Modules/NiagaraStatelessModule_ScaleSpriteSizeBySpeed.cpp b/Engine/Plugins/FX/Niagara/Source/Niagara/Private/Stateless/Modules/NiagaraStatelessModule_ScaleSpriteSizeBySpeed.cpp
new file mode 100644
--- /dev/null
+++ b/Engine/Plugins/FX/Niagara/Source/Niagara/Private/Stateless/Modules/NiagaraStatelessModule_ScaleSpriteSizeBySpeed.cpp
@@ -0,0 +1,127 @@
+// Copyright Epic Games, Inc. All Rights Reserved.
+
+#include "Stateless/Modules/NiagaraStatelessModule_ScaleSpriteSizeBySpeed.h"
+#include "Stateless/NiagaraStatelessParticleSimContext.h"
+
+namespace NSMScaleSpriteSizeBySpeedPrivate
+{
+	struct FModuleBuiltData
+	{
+		FNiagaraStatelessRangeVector2	MinScaleFactor = FNiagaraStatelessRangeVector2(FVector2f::One());
+		FNiagaraStatelessRangeVector2	MaxScaleFactor = FNiagaraStatelessRangeVector2(FVector2f::One());
+		FNiagaraStatelessRangeFloat		VelocityNorm = FNiagaraStatelessRangeFloat(0.0f);
+		FUintVector2					ScaleDistribution = FUintVector2::ZeroValue;
+
+		int32							PositionVariableOffset = INDEX_NONE;
+		int32							PreviousPositionVariableOffset = INDEX_NONE;
+		int32							SpriteSizeVariableOffset = INDEX_NONE;
+		int32							PreviousSpriteSizeVariableOffset = INDEX_NONE;
+	};
+
+	static void ParticleSimulate(const NiagaraStateless::FParticleSimulationContext& ParticleSimulationContext)
+	{
+		using namespace NiagaraStateless;
+
+		const FModuleBuiltData* ModuleBuiltData = ParticleSimulationContext.ReadBuiltData<FModuleBuiltData>();
+		const UNiagaraStatelessModule_ScaleSpriteSizeBySpeed::FParameters* Parameters = ParticleSimulationContext.ReadParameterNestedStruct<UNiagaraStatelessModule_ScaleSpriteSizeBySpeed::FParameters>();
+
+		for (uint32 i = 0; i < ParticleSimulationContext.GetNumInstances(); ++i)
+		{
+			const FVector3f Position			= ParticleSimulationContext.ReadParticleVariable(ModuleBuiltData->PositionVariableOffset, i, FVector3f::OneVector);
+			const FVector3f PreviousPosition	= ParticleSimulationContext.ReadParticleVariable(ModuleBuiltData->PreviousPositionVariableOffset, i, FVector3f::OneVector);
+			const FVector3f Velocity			= (Position - PreviousPosition) * ParticleSimulationContext.GetInvDeltaTime();
+			const float Speed					= Velocity.SquaredLength();
+			const float NormSpeed				= FMath::Clamp(Speed * Parameters->ScaleSpriteSizeBySpeed_VelocityNorm, 0.0f, 1.0f);
+			const float Interp					= ParticleSimulationContext.LerpStaticFloat<float>(Parameters->ScaleSpriteSizeBySpeed_ScaleDistribution, NormSpeed);
+			const FVector2f	Scale				= Parameters->ScaleSpriteSizeBySpeed_ScaleFactorBias + (Parameters->ScaleSpriteSizeBySpeed_ScaleFactorScale * FMath::Clamp(Interp, 0.0f, 1.0f));
+
+			FVector2f SpriteSize			= ParticleSimulationContext.ReadParticleVariable(ModuleBuiltData->SpriteSizeVariableOffset, i, FVector2f::One());
+			FVector2f PreviousSpriteSize	= ParticleSimulationContext.ReadParticleVariable(ModuleBuiltData->PreviousSpriteSizeVariableOffset, i, FVector2f::One());
+
+			SpriteSize			*= Scale;
+			PreviousSpriteSize	*= Scale;
+
+			ParticleSimulationContext.WriteParticleVariable(ModuleBuiltData->SpriteSizeVariableOffset, i, SpriteSize);
+			ParticleSimulationContext.WriteParticleVariable(ModuleBuiltData->PreviousSpriteSizeVariableOffset, i, PreviousSpriteSize);
+		}
+	}
+}
+
+void UNiagaraStatelessModule_ScaleSpriteSizeBySpeed::BuildEmitterData(const FNiagaraStatelessEmitterDataBuildContext& BuildContext) const
+{
+	using namespace NSMScaleSpriteSizeBySpeedPrivate;
+
+	FModuleBuiltData* BuiltData = BuildContext.AllocateBuiltData<FModuleBuiltData>();
+	if (!IsModuleEnabled())
+	{
+		return;
+	}
+
+	const FNiagaraStatelessGlobals& StatelessGlobals = FNiagaraStatelessGlobals::Get();
+	BuiltData->PositionVariableOffset = BuildContext.FindParticleVariableIndex(StatelessGlobals.PositionVariable);
+	BuiltData->PreviousPositionVariableOffset = BuildContext.FindParticleVariableIndex(StatelessGlobals.PreviousPositionVariable);
+	BuiltData->SpriteSizeVariableOffset = BuildContext.FindParticleVariableIndex(StatelessGlobals.SpriteSizeVariable);
+	BuiltData->PreviousSpriteSizeVariableOffset = BuildContext.FindParticleVariableIndex(StatelessGlobals.PreviousSpriteSizeVariable);
+
+	if (BuiltData->SpriteSizeVariableOffset == INDEX_NONE && BuiltData->PreviousSpriteSizeVariableOffset == INDEX_NONE)
+	{
+		return;
+	}
+
+	BuiltData->VelocityNorm = BuildContext.ConvertDistributionToRange(VelocityThreshold, DefaultVelocity);
+	BuiltData->MinScaleFactor = BuildContext.ConvertDistributionToRange(MinScaleFactor, FVector2f::One());
+	BuiltData->MaxScaleFactor = BuildContext.ConvertDistributionToRange(MaxScaleFactor, FVector2f::One());
+	if (bSampleScaleFactorByCurve)
+	{
+		if (SampleFactorCurve.IsCurve() && SampleFactorCurve.Values.Num() > 1)
+		{
+			BuiltData->ScaleDistribution.X = BuildContext.AddStaticData(SampleFactorCurve.Values);
+			BuiltData->ScaleDistribution.Y = SampleFactorCurve.Values.Num() - 1;
+		}
+		else
+		{
+			const float Value = SampleFactorCurve.Values.Num() > 0 ? SampleFactorCurve.Values[0] : 1.0f;
+			const float Values[] = { Value, Value };
+			BuiltData->ScaleDistribution.X = BuildContext.AddStaticData(Values);
+			BuiltData->ScaleDistribution.Y = 1;
+		}
+	}
+	else
+	{
+		const float Values[] = { 1.0f, 1.0f };
+		BuiltData->ScaleDistribution.X = BuildContext.AddStaticData(Values);
+		BuiltData->ScaleDistribution.Y = 1;
+	}
+	BuildContext.AddParticleSimulationExecSimulate(&ParticleSimulate);
+}
+
+void UNiagaraStatelessModule_ScaleSpriteSizeBySpeed::BuildShaderParameters(FNiagaraStatelessShaderParametersBuilder& ShaderParametersBuilder) const
+{
+	ShaderParametersBuilder.AddParameterNestedStruct<FParameters>();
+}
+
+void UNiagaraStatelessModule_ScaleSpriteSizeBySpeed::SetShaderParameters(const FNiagaraStatelessSetShaderParameterContext& SetShaderParameterContext) const
+{
+	using namespace NSMScaleSpriteSizeBySpeedPrivate;
+
+	const FModuleBuiltData* ModuleBuiltData = SetShaderParameterContext.ReadBuiltData<FModuleBuiltData>();
+
+	const float VelocityNorm = SetShaderParameterContext.ConvertRangeToValue(ModuleBuiltData->VelocityNorm);
+	const FVector2f MinScale = SetShaderParameterContext.ConvertRangeToValue(ModuleBuiltData->MinScaleFactor);
+	const FVector2f MaxScale = SetShaderParameterContext.ConvertRangeToValue(ModuleBuiltData->MaxScaleFactor);
+
+	FParameters* Parameters = SetShaderParameterContext.GetParameterNestedStruct<FParameters>();
+	Parameters->ScaleSpriteSizeBySpeed_ScaleFactorBias = MinScale;
+	Parameters->ScaleSpriteSizeBySpeed_ScaleFactorScale = MaxScale - MinScale;
+	Parameters->ScaleSpriteSizeBySpeed_VelocityNorm = VelocityNorm > 0.0f ? 1.0f / (VelocityNorm * VelocityNorm) : 0.0f;
+	Parameters->ScaleSpriteSizeBySpeed_ScaleDistribution = ModuleBuiltData->ScaleDistribution;
+}
+
+#if WITH_EDITORONLY_DATA
+void UNiagaraStatelessModule_ScaleSpriteSizeBySpeed::GetOutputVariables(TArray<FNiagaraVariableBase>& OutVariables) const
+{
+	const FNiagaraStatelessGlobals& StatelessGlobals = FNiagaraStatelessGlobals::Get();
+	OutVariables.AddUnique(StatelessGlobals.SpriteSizeVariable);
+	OutVariables.AddUnique(StatelessGlobals.PreviousSpriteSizeVariable);
+}
+#endif
diff --git a/Engine/Plugins/FX/Niagara/Source/Niagara/Private/Stateless/NiagaraStatelessDistribution.cpp b/Engine/Plugins/FX/Niagara/Source/Niagara/Private/Stateless/NiagaraStatelessDistribution.cpp
--- a/Engine/Plugins/FX/Niagara/Source/Niagara/Private/Stateless/NiagaraStatelessDistribution.cpp
+++ b/Engine/Plugins/FX/Niagara/Source/Niagara/Private/Stateless/NiagaraStatelessDistribution.cpp
@@ -162,6 +162,43 @@ namespace NiagaraStatelessDistributionPrivate
 		return LUT;
 	}
 
+	void InternalCreateCurveFromValues(TConstArrayView<float> Values, TArray<FRichCurve>& ChannelCurves, TArray<float>& ChannelConstantsAndRanges, int NumChannels)
+	{
+		check((Values.Num() % NumChannels) == 0);
+
+		ChannelCurves.Reset(NumChannels);
+		ChannelCurves.AddDefaulted(NumChannels);
+
+		ChannelConstantsAndRanges = Values;
+
+		const int32 NumValues = Values.Num() / NumChannels;
+		if (NumValues > 0)
+		{
+			const float UInc = NumValues > 1 ? 1.0f / float(NumValues - 1) : 0.0f;
+			float U = 0.0f;
+
+			for (int32 iValueIndex = 0; iValueIndex < NumValues; ++iValueIndex)
+			{
+				for (int32 iChannel = 0; iChannel < NumChannels; ++iChannel)
+				{
+					ChannelCurves[iChannel].AddKey(U, Values[(iValueIndex * NumChannels) + iChannel]);
+				}
+
+				U += UInc;
+			}
+		}
+	}
+
+	void CreateCurveFromValues(TConstArrayView<float> Values, TArray<FRichCurve>& ChannelCurves, TArray<float>& ChannelConstantsAndRanges)
+	{
+		InternalCreateCurveFromValues(Values, ChannelCurves, ChannelConstantsAndRanges, 1);
+	}
+
+	void CreateCurveFromValues(TConstArrayView<FVector3f> Values, TArray<FRichCurve>& ChannelCurves, TArray<float>& ChannelConstantsAndRanges)
+	{
+		InternalCreateCurveFromValues(MakeArrayView<const float>(reinterpret_cast<const float*>(Values.GetData()), Values.Num() * 3), ChannelCurves, ChannelConstantsAndRanges, 3);
+	}
+
 	template<typename FValueContainerSetNum, typename FValueAndChannelAccessor>
 	void UpdateDistributionValues(ENiagaraDistributionMode InMode, const TArray<float>& InChannelConstantsAndRanges, const TArray<FRichCurve>& InChannelCurves, int32 InChannelCount, FVector2f& OutTimeRange, FValueContainerSetNum InContainerNum, FValueAndChannelAccessor InValueAndChannelAccessor, int32 MaxLutSampleCount)
 	{
@@ -326,6 +363,38 @@ void FNiagaraDistributionFloat::InitConstant(float Value)
 #endif
 }
 
+void FNiagaraDistributionFloat::InitCurve(std::initializer_list<float> CurvePoints)
+{
+	const int32 NumValues = CurvePoints.size();
+
+	Mode = ENiagaraDistributionMode::UniformCurve;
+	Values = CurvePoints;
+
+#if WITH_EDITORONLY_DATA
+	NiagaraStatelessDistributionPrivate::CreateCurveFromValues(Values, ChannelCurves, ChannelConstantsAndRanges);
+#endif
+}
+
+#if WITH_EDITORONLY_DATA
+void FNiagaraDistributionFloat::InitCurve(const TArray<FRichCurveKey>& CurveKeys)
+{
+	Mode = ENiagaraDistributionMode::UniformCurve;
+
+	ChannelCurves.Reset(1);
+	ChannelCurves.AddDefaulted_GetRef().SetKeys(CurveKeys);
+
+	Values.SetNumUninitialized(CurveKeys.Num());
+	ChannelConstantsAndRanges.SetNumUninitialized(CurveKeys.Num());
+
+	for (int32 iKey=0; iKey < CurveKeys.Num(); ++iKey)
+	{
+		const float KeyValue = CurveKeys[iKey].Value;
+		Values[iKey] = KeyValue;
+		ChannelConstantsAndRanges[iKey] = KeyValue;
+	}
+}
+#endif
+
 FNiagaraStatelessRangeFloat FNiagaraDistributionFloat::CalculateRange(const float Default) const
 {
 	FNiagaraStatelessRangeFloat Range(Default, Default);
@@ -394,6 +463,34 @@ void FNiagaraDistributionVector3::InitConstant(const FVector3f& Value)
 #endif
 }
 
+void FNiagaraDistributionVector3::InitCurve(std::initializer_list<float> CurvePoints)
+{
+	const int32 NumValues = CurvePoints.size();
+
+	Mode = ENiagaraDistributionMode::UniformCurve;
+	Values.Reset(NumValues);
+	for (float Value : CurvePoints)
+	{
+		Values.Emplace(Value);
+	}
+
+#if WITH_EDITORONLY_DATA
+	NiagaraStatelessDistributionPrivate::CreateCurveFromValues(Values, ChannelCurves, ChannelConstantsAndRanges);
+#endif
+}
+
+void FNiagaraDistributionVector3::InitCurve(std::initializer_list<FVector3f> CurvePoints)
+{
+	const int32 NumValues = CurvePoints.size();
+
+	Mode = ENiagaraDistributionMode::NonUniformCurve;
+	Values = TArray<FVector3f>(CurvePoints);
+
+#if WITH_EDITORONLY_DATA
+	NiagaraStatelessDistributionPrivate::CreateCurveFromValues(Values, ChannelCurves, ChannelConstantsAndRanges);
+#endif
+}
+
 FNiagaraStatelessRangeVector3 FNiagaraDistributionVector3::CalculateRange(const FVector3f& Default) const
 {
 	FNiagaraStatelessRangeVector3 Range(Default, Default);
diff --git a/Engine/Plugins/FX/Niagara/Source/NiagaraShader/Internal/Stateless/NiagaraStatelessModuleShaderParameters.h b/Engine/Plugins/FX/Niagara/Source/NiagaraShader/Internal/Stateless/NiagaraStatelessModuleShaderParameters.h
--- a/Engine/Plugins/FX/Niagara/Source/NiagaraShader/Internal/Stateless/NiagaraStatelessModuleShaderParameters.h
+++ b/Engine/Plugins/FX/Niagara/Source/NiagaraShader/Internal/Stateless/NiagaraStatelessModuleShaderParameters.h
@@ -80,6 +80,8 @@ namespace NiagaraStateless
 	END_SHADER_PARAMETER_STRUCT()
 
 	BEGIN_SHADER_PARAMETER_STRUCT(FScaleMeshSizeBySpeedModule_ShaderParameters, NIAGARASHADER_API)
+		SHADER_PARAMETER(FVector3f,		ScaleMeshSizeBySpeed_ScaleFactorBias)
+		SHADER_PARAMETER(FVector3f,		ScaleMeshSizeBySpeed_ScaleFactorScale)
 		SHADER_PARAMETER(FUintVector2,	ScaleMeshSizeBySpeed_ScaleDistribution)
 		SHADER_PARAMETER(float,			ScaleMeshSizeBySpeed_VelocityNorm)
 	END_SHADER_PARAMETER_STRUCT()
@@ -102,6 +104,8 @@ namespace NiagaraStateless
 	END_SHADER_PARAMETER_STRUCT()
 
 	BEGIN_SHADER_PARAMETER_STRUCT(FScaleSpriteSizeBySpeedModule_ShaderParameters, NIAGARASHADER_API)
+		SHADER_PARAMETER(FVector2f,		ScaleSpriteSizeBySpeed_ScaleFactorBias)
+		SHADER_PARAMETER(FVector2f,		ScaleSpriteSizeBySpeed_ScaleFactorScale)
 		SHADER_PARAMETER(FUintVector2,	ScaleSpriteSizeBySpeed_ScaleDistribution)
 		SHADER_PARAMETER(float,			ScaleSpriteSizeBySpeed_VelocityNorm)
 	END_SHADER_PARAMETER_STRUCT()
