- Fix mesh reimport crash with stateless

#rnx
#tests Reimported with stateful / stateless / stanalone emitter / instances placed in the level

diff --git a/Engine/Plugins/FX/Niagara/Source/Niagara/Private/NiagaraMeshRendererProperties.cpp b/Engine/Plugins/FX/Niagara/Source/Niagara/Private/NiagaraMeshRendererProperties.cpp
--- a/Engine/Plugins/FX/Niagara/Source/Niagara/Private/NiagaraMeshRendererProperties.cpp
+++ b/Engine/Plugins/FX/Niagara/Source/Niagara/Private/NiagaraMeshRendererProperties.cpp
@@ -1325,10 +1325,9 @@ void UNiagaraMeshRendererProperties::OnMeshChanged()
 {
 	FNiagaraSystemUpdateContext ReregisterContext;
 
-	FVersionedNiagaraEmitter Outer = GetOuterEmitter();
-	if (Outer.Emitter)
+	if (UNiagaraSystem* NiagaraSystem = GetTypedOuter<UNiagaraSystem>())
 	{
-		ReregisterContext.Add(Outer, true);
+		ReregisterContext.Add(NiagaraSystem, true);
 	}
 
 	CheckMaterialUsage();
