- Mark stateless modules as experimental

#rnx
#tests Add module and confirmed they display the warning

diff --git a/Engine/Plugins/FX/Niagara/Source/Niagara/Internal/Stateless/Modules/NiagaraStatelessModule_CurlNoiseForce.h b/Engine/Plugins/FX/Niagara/Source/Niagara/Internal/Stateless/Modules/NiagaraStatelessModule_CurlNoiseForce.h
--- a/Engine/Plugins/FX/Niagara/Source/Niagara/Internal/Stateless/Modules/NiagaraStatelessModule_CurlNoiseForce.h
+++ b/Engine/Plugins/FX/Niagara/Source/Niagara/Internal/Stateless/Modules/NiagaraStatelessModule_CurlNoiseForce.h
@@ -20,7 +20,7 @@ enum class ENSM_NoiseMode
 	LUTJacobBicubic,
 };
 
-UCLASS(MinimalAPI, EditInlineNew, meta = (DisplayName = "Curl Noise Force"))
+UCLASS(MinimalAPI, EditInlineNew, Experimental, meta = (DisplayName = "Curl Noise Force"))
 class UNiagaraStatelessModule_CurlNoiseForce : public UNiagaraStatelessModule
 {
 	GENERATED_BODY()
diff --git a/Engine/Plugins/FX/Niagara/Source/Niagara/Internal/Stateless/Modules/NiagaraStatelessModule_MeshIndex.h b/Engine/Plugins/FX/Niagara/Source/Niagara/Internal/Stateless/Modules/NiagaraStatelessModule_MeshIndex.h
--- a/Engine/Plugins/FX/Niagara/Source/Niagara/Internal/Stateless/Modules/NiagaraStatelessModule_MeshIndex.h
+++ b/Engine/Plugins/FX/Niagara/Source/Niagara/Internal/Stateless/Modules/NiagaraStatelessModule_MeshIndex.h
@@ -9,7 +9,7 @@
 #include "NiagaraStatelessModule_MeshIndex.generated.h"
 
 // Sets the mesh index to use with mesh renderers
-UCLASS(MinimalAPI, EditInlineNew, meta = (DisplayName = "Mesh Index"))
+UCLASS(MinimalAPI, EditInlineNew, Experimental, meta = (DisplayName = "Mesh Index"))
 class UNiagaraStatelessModule_MeshIndex : public UNiagaraStatelessModule
 {
 	GENERATED_BODY()
