Reset the digested graph cache on PreGC to prevent it from unnecessarily holding onto object references
#tests Ran commandlet which loads/saves/unloads assets and verifies assets are fully unloaded

diff --git a/Engine/Plugins/FX/Niagara/Source/NiagaraEditor/Private/NiagaraDigestDatabase.cpp b/Engine/Plugins/FX/Niagara/Source/NiagaraEditor/Private/NiagaraDigestDatabase.cpp
--- a/Engine/Plugins/FX/Niagara/Source/NiagaraEditor/Private/NiagaraDigestDatabase.cpp
+++ b/Engine/Plugins/FX/Niagara/Source/NiagaraEditor/Private/NiagaraDigestDatabase.cpp
@@ -10,6 +10,7 @@
 #include "NiagaraParameterCollection.h"
 #include "NiagaraScript.h"
 #include "NiagaraScriptSource.h"
+#include "UObject/UObjectGlobals.h"
 
 namespace NiagaraGraphDigestDatabaseImpl
 {
@@ -44,12 +45,17 @@ void FNiagaraDigestDatabase::Shutdown()
 FNiagaraDigestDatabase::FNiagaraDigestDatabase()
 	: CompilationGraphCache(NiagaraGraphDigestDatabaseImpl::GDigestGraphCacheSize)
 {
-
+	// we use GC as an opportunity to clear the cache, preventing the database from holding onto references and
+	// extending object lifetimes unnecessarily.  In particular the digested graphs contain duplicates of the DI
+	// which could be referencing real objects
+	FCoreUObjectDelegates::GetPreGarbageCollectDelegate().AddRaw(this, &FNiagaraDigestDatabase::ReleaseDatabase);
 }
 
 FNiagaraDigestDatabase::~FNiagaraDigestDatabase()
 {
 	ReleaseDatabase();
+
+	FCoreUObjectDelegates::GetPreGarbageCollectDelegate().RemoveAll(this);
 }
 
 void FNiagaraDigestDatabase::ReleaseDatabase()
