When querying DI usage for GPU scripts just use the metadata stored on the CachedScriptVM data rather than dealing with the ShaderScript directly.  This addresses an issue with the new compilation mode which handles cooked shaders at serialization rather than directly manipulating the current shader for the UNiagaraScript.
#tests Cooked and ran NDC test scripts for non windows platform

diff --git a/Engine/Plugins/FX/Niagara/Source/Niagara/Private/DataInterface/NiagaraDataInterfaceUtilities.cpp b/Engine/Plugins/FX/Niagara/Source/Niagara/Private/DataInterface/NiagaraDataInterfaceUtilities.cpp
--- a/Engine/Plugins/FX/Niagara/Source/Niagara/Private/DataInterface/NiagaraDataInterfaceUtilities.cpp
+++ b/Engine/Plugins/FX/Niagara/Source/Niagara/Private/DataInterface/NiagaraDataInterfaceUtilities.cpp
@@ -106,28 +106,25 @@ void ForEachGpuFunctionImpl(class UNiagaraDataInterface* RuntimeResolvedDataInte
 				return;
 			}
 
-			if (FNiagaraShaderScript* ShaderScript = Script->GetRenderThreadScript())
-			{
-				TArrayView<const FNiagaraScriptResolvedDataInterfaceInfo> ResolvedDIs = Script->GetResolvedDataInterfaces();
-				const TArray<FNiagaraDataInterfaceGPUParamInfo>& DataInterfaceParamInfos = ShaderScript->GetScriptParametersMetadata()->DataInterfaceParamInfo;
+			TArrayView<const FNiagaraScriptResolvedDataInterfaceInfo> ResolvedDIs = Script->GetResolvedDataInterfaces();
+			TConstArrayView<FNiagaraDataInterfaceGPUParamInfo> DataInterfaceParamInfos = Script->GetDataInterfaceGPUParamInfos();
 
-				const int NumDataInterfaces = FMath::Min(ResolvedDIs.Num(), DataInterfaceParamInfos.Num());	// Note: Should always be equal but lets be safe
-				for (int iDataInterface = 0; iDataInterface < NumDataInterfaces; ++iDataInterface)
+			const int NumDataInterfaces = FMath::Min(ResolvedDIs.Num(), DataInterfaceParamInfos.Num());	// Note: Should always be equal but lets be safe
+			for (int iDataInterface = 0; iDataInterface < NumDataInterfaces; ++iDataInterface)
+			{
+				const FNiagaraScriptResolvedDataInterfaceInfo& ResolvedDI = ResolvedDIs[iDataInterface];
+				if (ResolvedDI.ResolvedDataInterface == nullptr)
 				{
-					const FNiagaraScriptResolvedDataInterfaceInfo& ResolvedDI = ResolvedDIs[iDataInterface];
-					if (ResolvedDI.ResolvedDataInterface == nullptr)
-					{
-						continue;
-					}
-					if (ResolvedDI.ResolvedDataInterface == RuntimeResolvedDataInterface || (DataInterfaceUserVariable.IsValid() && ResolvedDI.ResolvedVariable == DataInterfaceUserVariable))
+					continue;
+				}
+				if (ResolvedDI.ResolvedDataInterface == RuntimeResolvedDataInterface || (DataInterfaceUserVariable.IsValid() && ResolvedDI.ResolvedVariable == DataInterfaceUserVariable))
+				{
+					for (const FNiagaraDataInterfaceGeneratedFunction& GeneratedFunction : DataInterfaceParamInfos[iDataInterface].GeneratedFunctions)
 					{
-						for (const FNiagaraDataInterfaceGeneratedFunction& GeneratedFunction : DataInterfaceParamInfos[iDataInterface].GeneratedFunctions)
+						if (Action(Script, GeneratedFunction) == false)
 						{
-							if (Action(Script, GeneratedFunction) == false)
-							{
-								bContinueSearching = false;
-								return;
-							}
+							bContinueSearching = false;
+							return;
 						}
 					}
 				}
