- Fix ensure when calling undo while tweaking a color

#rnx

diff --git a/Engine/Plugins/FX/Niagara/Source/NiagaraEditor/Private/Customizations/UserParameters/NiagaraSystemUserParametersBuilder.cpp b/Engine/Plugins/FX/Niagara/Source/NiagaraEditor/Private/Customizations/UserParameters/NiagaraSystemUserParametersBuilder.cpp
--- a/Engine/Plugins/FX/Niagara/Source/NiagaraEditor/Private/Customizations/UserParameters/NiagaraSystemUserParametersBuilder.cpp
+++ b/Engine/Plugins/FX/Niagara/Source/NiagaraEditor/Private/Customizations/UserParameters/NiagaraSystemUserParametersBuilder.cpp
@@ -486,7 +486,7 @@ void FNiagaraSystemUserParameterBuilder::OnParameterEditorValueChanged(FNiagaraV
 		return;
 	}
 	
-	FScopedTransaction ScopedTransaction(ChangedUserParameterTransactionText);
+	FScopedTransaction ScopedTransaction(ChangedUserParameterTransactionText, GIsTransacting == false);
 	SystemAsset->Modify();
 
 	// we forward the change in the parameter editor into the display data. NotifyPostChange will take care of forwarding the change into the user parameter store.
