Fix for nullptr access when collecting parameter references of customhlsl
#tests opened asset with custom hlsl and verified it didn't crash
#rnx

diff --git a/Engine/Plugins/FX/Niagara/Source/NiagaraEditor/Private/ViewModels/NiagaraParameterPanelViewModel.cpp b/Engine/Plugins/FX/Niagara/Source/NiagaraEditor/Private/ViewModels/NiagaraParameterPanelViewModel.cpp
--- a/Engine/Plugins/FX/Niagara/Source/NiagaraEditor/Private/ViewModels/NiagaraParameterPanelViewModel.cpp
+++ b/Engine/Plugins/FX/Niagara/Source/NiagaraEditor/Private/ViewModels/NiagaraParameterPanelViewModel.cpp
@@ -2035,7 +2035,8 @@ TArray<FNiagaraParameterPanelItem> FNiagaraSystemToolkitParameterPanelViewModel:
 							{
 								if (UNiagaraNodeFunctionCall* FunctionNode = Cast<UNiagaraNodeFunctionCall>(Node))
 								{
-									if (FunctionNode->FunctionScript->IsModuleScript() && FunctionNode->IsNodeEnabled())
+									const bool bIsModuleFunctionScript = FunctionNode->FunctionScript && FunctionNode->FunctionScript->IsModuleScript();
+									if (FunctionNode->IsNodeEnabled() && bIsModuleFunctionScript)
 									{
 										// Found the module
 										FString CurrentName = Source.ModuleName.ToString();
