Niagara - Parameter traversal cache.

diff --git a/Engine/Plugins/FX/Niagara/Source/Niagara/Classes/NiagaraEmitter.h b/Engine/Plugins/FX/Niagara/Source/Niagara/Classes/NiagaraEmitter.h
--- a/Engine/Plugins/FX/Niagara/Source/Niagara/Classes/NiagaraEmitter.h
+++ b/Engine/Plugins/FX/Niagara/Source/Niagara/Classes/NiagaraEmitter.h
@@ -376,7 +376,7 @@ struct FVersionedNiagaraEmitterData
 	NIAGARA_API bool IsReadyToRun() const;
 	const TArray<UNiagaraRendererProperties*>& GetRenderers() const { return RendererProperties; }
 	NIAGARA_API void GetScripts(TArray<UNiagaraScript*>& OutScripts, bool bCompilableOnly = true, bool bEnabledOnly = false) const;
-	NIAGARA_API UNiagaraScript* GetScript(ENiagaraScriptUsage Usage, FGuid UsageId);
+	NIAGARA_API UNiagaraScript* GetScript(ENiagaraScriptUsage Usage, FGuid UsageId) const;
 	UNiagaraScript* GetGPUComputeScript() { return GPUComputeScript; }
 	const UNiagaraScript* GetGPUComputeScript() const { return GPUComputeScript; }
 	FORCEINLINE const TArray<FNiagaraEventScriptProperties>& GetEventHandlers() const { return EventHandlerScriptProps; }
@@ -660,6 +660,7 @@ public:
 
 	/** Get the cached parameter map traversal for this emitter.  */
 	NIAGARA_API const TSharedPtr<FNiagaraGraphCachedDataBase, ESPMode::ThreadSafe>& GetCachedTraversalData(const FGuid& EmitterVersion) const;
+	NIAGARA_API void InvalidateCachedTraversalData(const FGuid& EmitterVersion);
 #endif
 
 	bool IsEnabledOnPlatform(const FString& PlatformName) const;
diff --git a/Engine/Plugins/FX/Niagara/Source/Niagara/Private/NiagaraEmitter.cpp b/Engine/Plugins/FX/Niagara/Source/Niagara/Private/NiagaraEmitter.cpp
--- a/Engine/Plugins/FX/Niagara/Source/Niagara/Private/NiagaraEmitter.cpp
+++ b/Engine/Plugins/FX/Niagara/Source/Niagara/Private/NiagaraEmitter.cpp
@@ -1433,6 +1433,16 @@ const TSharedPtr<FNiagaraGraphCachedDataBase, ESPMode::ThreadSafe>& UNiagaraEmit
 	EmitterData->CachedTraversalData = NiagaraModule.CacheGraphTraversal(this, EmitterVersion);
 	return EmitterData->CachedTraversalData;
 }
+
+void UNiagaraEmitter::InvalidateCachedTraversalData(const FGuid& EmitterVersion)
+{
+	const FVersionedNiagaraEmitterData* EmitterData = GetEmitterData(EmitterVersion);
+	if (EmitterData->CachedTraversalData.IsValid() && EmitterData->CachedTraversalData->IsValidForEmitter(EmitterData))
+	{
+		EmitterData->CachedTraversalData.Reset();
+	}
+}
+
 #endif // WITH_EDITORONLY_DATA
 
 bool UNiagaraEmitter::IsEnabledOnPlatform(const FString& PlatformName)const
@@ -1601,7 +1611,7 @@ void FVersionedNiagaraEmitterData::GetScripts(TArray<UNiagaraScript*>& OutScript
 	}
 }
 
-UNiagaraScript* FVersionedNiagaraEmitterData::GetScript(ENiagaraScriptUsage Usage, FGuid UsageId)
+UNiagaraScript* FVersionedNiagaraEmitterData::GetScript(ENiagaraScriptUsage Usage, FGuid UsageId) const
 {
 	TArray<UNiagaraScript*> Scripts;
 	GetScripts(Scripts, false);
diff --git a/Engine/Plugins/FX/Niagara/Source/NiagaraEditor/Private/NiagaraEditorModule.cpp b/Engine/Plugins/FX/Niagara/Source/NiagaraEditor/Private/NiagaraEditorModule.cpp
--- a/Engine/Plugins/FX/Niagara/Source/NiagaraEditor/Private/NiagaraEditorModule.cpp
+++ b/Engine/Plugins/FX/Niagara/Source/NiagaraEditor/Private/NiagaraEditorModule.cpp
@@ -181,6 +181,9 @@
 #include "NiagaraRecentAndFavoritesManager.h"
 #include "Widgets/DataChannel/NiagaraDataChannelWizard.h"
 
+#include "TraversalCache/TraversalBuilder.h"
+#include "TraversalCache/TraversalCache.h"
+
 #include UE_INLINE_GENERATED_CPP_BY_NAME(NiagaraEditorModule)
 
 IMPLEMENT_MODULE( FNiagaraEditorModule, NiagaraEditor );
@@ -219,6 +222,13 @@ static FAutoConsoleVariableRef CVarEnableCustomInlineDynamicInputFormats(
 	TEXT("If > 0 and experimental inline editors for dynamic input trees are enabled, custom formats which are defined on scripts will be applied.\n"),
 	ECVF_Default);
 
+int32 GbEnableTraversalCache = 0;
+static FAutoConsoleVariableRef CVarEnableTraversalCache(
+	TEXT("fx.Niagara.EnableTraversalCache"),
+	GbEnableTraversalCache,
+	TEXT("If > 0 the new traversal cache will be used to speed up utilities like GetStackFunctionInputs.\n"),
+	ECVF_Default);
+
 // this is required for gpu script compilation ticks
 static FNiagaraShaderQueueTickable NiagaraShaderQueueProcessor;
 
@@ -1782,6 +1792,11 @@ void FNiagaraEditorModule::OnPostEngineInit()
 	const FAssetRegistryModule& AssetRegistryModule = FModuleManager::LoadModuleChecked<FAssetRegistryModule>(TEXT("AssetRegistry"));
 	IAssetRegistry& AssetRegistry = AssetRegistryModule.Get();
 	AssetRegistryOnLoadCompleteHandle = AssetRegistry.OnFilesLoaded().AddRaw(this, &FNiagaraEditorModule::OnAssetRegistryLoadComplete);
+
+	TraversalBuilder = MakeShared<UE::Niagara::TraversalCache::FTraversalBuilder>();
+	TraversalBuilder->Initialize();
+	TraversalCache = MakeShared<UE::Niagara::TraversalCache::FTraversalCache>();
+	TraversalCache->Initialize();
 }
 
 void FNiagaraEditorModule::OnDeviceProfileManagerUpdated()
diff --git a/Engine/Plugins/FX/Niagara/Source/NiagaraEditor/Private/NiagaraGraphDataCache.cpp b/Engine/Plugins/FX/Niagara/Source/NiagaraEditor/Private/NiagaraGraphDataCache.cpp
--- a/Engine/Plugins/FX/Niagara/Source/NiagaraEditor/Private/NiagaraGraphDataCache.cpp
+++ b/Engine/Plugins/FX/Niagara/Source/NiagaraEditor/Private/NiagaraGraphDataCache.cpp
@@ -153,6 +153,11 @@ void FNiagaraGraphDataCache::GetStackFunctionInputs(const UNiagaraNodeFunctionCa
 	}
 }
 
+void FNiagaraGraphDataCache::InvalidateCache()
+{
+	StackFunctionInputPinCache.Empty(GNiagaraGraphDataCacheSize);
+}
+
 FNiagaraGraphDataCache::FStackFunctionInputPinKey::FStackFunctionInputPinKey(const UNiagaraNodeFunctionCall& FunctionCallNode, const UNiagaraGraph* CalledGraph, TConstArrayView<FNiagaraVariable> StaticVariables, const FCompileConstantResolver& InConstantResolver, FNiagaraStackGraphUtilities::ENiagaraGetStackFunctionInputPinsOptions InOptions, bool bInIgnoreDisabled, bool bInFilterForCompilation)
 {
 	if (ensure(CalledGraph))
diff --git a/Engine/Plugins/FX/Niagara/Source/NiagaraEditor/Private/NiagaraNodeParameterMapGet.cpp b/Engine/Plugins/FX/Niagara/Source/NiagaraEditor/Private/NiagaraNodeParameterMapGet.cpp
--- a/Engine/Plugins/FX/Niagara/Source/NiagaraEditor/Private/NiagaraNodeParameterMapGet.cpp
+++ b/Engine/Plugins/FX/Niagara/Source/NiagaraEditor/Private/NiagaraNodeParameterMapGet.cpp
@@ -280,7 +280,7 @@ void UNiagaraNodeParameterMapGet::RemoveDynamicPin(UEdGraphPin* Pin)
 	Super::RemoveDynamicPin(Pin);
 }
 
-UEdGraphPin* UNiagaraNodeParameterMapGet::GetDefaultPin(UEdGraphPin* OutputPin) const
+UEdGraphPin* UNiagaraNodeParameterMapGet::GetDefaultPin(const UEdGraphPin* OutputPin) const
 {
 	FPinCollectorArray InputPins;
 	GetInputPins(InputPins);
diff --git a/Engine/Plugins/FX/Niagara/Source/NiagaraEditor/Private/NiagaraNodeParameterMapGet.h b/Engine/Plugins/FX/Niagara/Source/NiagaraEditor/Private/NiagaraNodeParameterMapGet.h
--- a/Engine/Plugins/FX/Niagara/Source/NiagaraEditor/Private/NiagaraNodeParameterMapGet.h
+++ b/Engine/Plugins/FX/Niagara/Source/NiagaraEditor/Private/NiagaraNodeParameterMapGet.h
@@ -37,7 +37,7 @@ public:
 	void AddOrphanedPinPairGuids(UEdGraphPin* OutputPin, UEdGraphPin* DefaultPin);
 	
 	/** Get the default value input pin for one of the output pins specified.*/
-	UEdGraphPin* GetDefaultPin(UEdGraphPin* OutputPin) const;
+	UEdGraphPin* GetDefaultPin(const UEdGraphPin* OutputPin) const;
 
 	void SynchronizeDefaultPins();
 	
diff --git a/Engine/Plugins/FX/Niagara/Source/NiagaraEditor/Private/NiagaraNodeWithDynamicPins.cpp b/Engine/Plugins/FX/Niagara/Source/NiagaraEditor/Private/NiagaraNodeWithDynamicPins.cpp
--- a/Engine/Plugins/FX/Niagara/Source/NiagaraEditor/Private/NiagaraNodeWithDynamicPins.cpp
+++ b/Engine/Plugins/FX/Niagara/Source/NiagaraEditor/Private/NiagaraNodeWithDynamicPins.cpp
@@ -137,7 +137,7 @@ void UNiagaraNodeWithDynamicPins::CreateAddPin(EEdGraphPinDirection Direction)
 	CreatePin(Direction, FEdGraphPinType(UEdGraphSchema_Niagara::PinCategoryMisc, AddPinSubCategory, nullptr, EPinContainerType::None, false, FEdGraphTerminalType()), TEXT("Add"));
 }
 
-bool UNiagaraNodeWithDynamicPins::IsAddPin(const UEdGraphPin* Pin) const
+bool UNiagaraNodeWithDynamicPins::IsAddPin(const UEdGraphPin* Pin)
 {
 	return Pin->PinType.PinCategory == UEdGraphSchema_Niagara::PinCategoryMisc && 
 		Pin->PinType.PinSubCategory == AddPinSubCategory;
diff --git a/Engine/Plugins/FX/Niagara/Source/NiagaraEditor/Private/NiagaraScriptMergeManager.cpp b/Engine/Plugins/FX/Niagara/Source/NiagaraEditor/Private/NiagaraScriptMergeManager.cpp
--- a/Engine/Plugins/FX/Niagara/Source/NiagaraEditor/Private/NiagaraScriptMergeManager.cpp
+++ b/Engine/Plugins/FX/Niagara/Source/NiagaraEditor/Private/NiagaraScriptMergeManager.cpp
@@ -271,7 +271,7 @@ FNiagaraStackFunctionMergeAdapter::FNiagaraStackFunctionMergeAdapter(const FVers
 		TArray<FName> FunctionInputVariableNames;
 
 		{
-			FCompileConstantResolver Resolver(InOwningEmitter, ENiagaraScriptUsage::Function);
+			FCompileConstantResolver Resolver(InOwningEmitter, InOwningScript.GetUsage());
 			TArray<FNiagaraVariable> AllFunctionInputVariables;
 			GetStackFunctionInputs(*FunctionCallNode, AllFunctionInputVariables, Resolver, FNiagaraStackGraphUtilities::ENiagaraGetStackFunctionInputPinsOptions::ModuleInputsOnly, false);
 
diff --git a/Engine/Plugins/FX/Niagara/Source/NiagaraEditor/Private/TraversalCache/StaticContextBuilder.cpp b/Engine/Plugins/FX/Niagara/Source/NiagaraEditor/Private/TraversalCache/StaticContextBuilder.cpp
new file mode 100644
--- /dev/null
+++ b/Engine/Plugins/FX/Niagara/Source/NiagaraEditor/Private/TraversalCache/StaticContextBuilder.cpp
@@ -0,0 +1,164 @@
+// Copyright Epic Games, Inc. All Rights Reserved.
+
+#include "TraversalCache/StaticContextBuilder.h"
+
+#include "INiagaraEditorTypeUtilities.h"
+#include "NiagaraCommon.h"
+#include "NiagaraConstants.h"
+#include "NiagaraEditorModule.h"
+#include "NiagaraEmitter.h"
+#include "NiagaraGraph.h"
+#include "NiagaraNodeFunctionCall.h"
+#include "NiagaraScript.h"
+#include "NiagaraScriptSource.h"
+#include "NiagaraSimulationStageBase.h"
+#include "NiagaraSystem.h"
+#include "TraversalCache/Traversal.h"
+#include "TraversalCache/TraversalBuilder.h"
+#include "TraversalCache/TraversalCache.h"
+#include "ViewModels/Stack/NiagaraParameterHandle.h"
+
+namespace UE::Niagara::TraversalCache
+{
+
+void FStaticContextBuilder::CreateTopLevelScriptContext(
+	const UNiagaraSystem& System,
+	const FVersionedNiagaraEmitterData* EmitterData,
+	const UNiagaraScript& TargetScript,
+	FTopLevelScriptStaticContext& OutTopLevelScriptContext,
+	FTraversalDebugData* TraversalDebugData)
+{
+	ENiagaraScriptUsage TargetScriptUsage = TargetScript.GetUsage();
+	bool bIsSystemScript = UNiagaraScript::IsSystemScript(TargetScriptUsage);
+	bool bIsEmitterScript = UNiagaraScript::IsEmitterScript(TargetScriptUsage);
+	bool bIsParticleScript = UNiagaraScript::IsParticleScript(TargetScriptUsage);
+	if (ensureMsgf((bIsEmitterScript == false && bIsParticleScript == false) || EmitterData != nullptr, TEXT("Must supply emitter data when creating a context for an emitter or particle script.")) == false)
+	{
+		return;
+	}
+
+	// Gather context from owning objects and scripts that are evaluated before the target script.
+	if (bIsSystemScript)
+	{
+		// When handling target scripts in the system, we only have to gather the static writes from system spawn,
+		// when the target script is system update.
+		if (TargetScript.IsEquivalentUsage(ENiagaraScriptUsage::SystemUpdateScript))
+		{
+			GatherStaticParametersWrittenByScript(OutTopLevelScriptContext.GlobalContext, *System.GetSystemSpawnScript(), TraversalDebugData);
+		}
+	}
+	else if (bIsEmitterScript || bIsParticleScript)
+	{
+		// All emitter and particle scripts need both the system spawn and system update static writes, as well as static values
+		// gathered from the emitter object itself.
+		GatherStaticParametersWrittenByScript(OutTopLevelScriptContext.GlobalContext, *System.GetSystemSpawnScript(), TraversalDebugData);
+		GatherStaticParametersWrittenByScript(OutTopLevelScriptContext.GlobalContext, *System.GetSystemUpdateScript(), TraversalDebugData);
+		GatherContextFromEmitterObject(OutTopLevelScriptContext.GlobalContext, *EmitterData);
+
+		if (bIsEmitterScript)
+		{
+			// When handling target scripts in the emitter, the only additional static writes needed are emitter spawn static write
+			// when the target script is emitter update.
+			if (TargetScript.IsEquivalentUsage(ENiagaraScriptUsage::EmitterUpdateScript))
+			{
+				GatherStaticParametersWrittenByScript(OutTopLevelScriptContext.GlobalContext, *EmitterData->EmitterSpawnScriptProps.Script, TraversalDebugData);
+			}
+		}
+		else // bIsParticleScript
+		{
+			// Particle scripts need both the emitter spawn and emitter udpate static writes.
+			GatherStaticParametersWrittenByScript(OutTopLevelScriptContext.GlobalContext, *EmitterData->EmitterSpawnScriptProps.Script, TraversalDebugData);
+			GatherStaticParametersWrittenByScript(OutTopLevelScriptContext.GlobalContext, *EmitterData->EmitterUpdateScriptProps.Script, TraversalDebugData);
+
+			if (TargetScript.IsEquivalentUsage(ENiagaraScriptUsage::ParticleUpdateScript))
+			{
+				// Particle update scripts need static writes from particle spawn.
+				GatherStaticParametersWrittenByScript(OutTopLevelScriptContext.GlobalContext, *EmitterData->SpawnScriptProps.Script, TraversalDebugData);
+			}
+			else if (TargetScript.IsEquivalentUsage(ENiagaraScriptUsage::ParticleEventScript) || TargetScript.IsEquivalentUsage(ENiagaraScriptUsage::ParticleSimulationStageScript))
+			{
+				// Event scripts and simulation stage scripts need static writes from particle spawn and particle update.
+				GatherStaticParametersWrittenByScript(OutTopLevelScriptContext.GlobalContext, *EmitterData->SpawnScriptProps.Script, TraversalDebugData);
+				GatherStaticParametersWrittenByScript(OutTopLevelScriptContext.GlobalContext, *EmitterData->UpdateScriptProps.Script, TraversalDebugData);
+			
+				if (TargetScript.IsEquivalentUsage(ENiagaraScriptUsage::ParticleSimulationStageScript))
+				{
+					// Simulation stage scripts need static writes from any stage that runs before them.
+					for (UNiagaraSimulationStageBase* SimulationStageBase : EmitterData->GetSimulationStages())
+					{
+						if (SimulationStageBase->Script == &TargetScript)
+						{
+							break;
+						}
+						else
+						{
+							GatherStaticParametersWrittenByScript(OutTopLevelScriptContext.GlobalContext, *SimulationStageBase->Script, TraversalDebugData);
+						}
+					}
+				}
+			}
+		}
+	}
+
+	// Lastly gather the context from the target script object.
+	GatherContextFromScriptObject(OutTopLevelScriptContext.GlobalContext, OutTopLevelScriptContext.FunctionInputContexts, TargetScript, TraversalDebugData);
+}
+
+void FStaticContextBuilder::AddBoolValue(FStaticContext& TargetContext, ESelectKeySource KeySource, const FNiagaraVariableBase& BoolVariable, FName NamespaceModifier, bool bValue)
+{
+	FSelectKey BoolKey(KeySource, BoolVariable, NamespaceModifier);
+	FSelectValue BoolValue = bValue ? FSelectValue::GetBoolTrue() : FSelectValue::GetBoolFalse();
+	TargetContext.Add(BoolKey, BoolValue);
+}
+
+void FStaticContextBuilder::AddEnumValue(FStaticContext& TargetContext, ESelectKeySource KeySource, const FNiagaraVariableBase& EnumVariable, FName NamespaceModifier, int32 EnumValue)
+{
+	FSelectKey EnumKey(KeySource, EnumVariable, NamespaceModifier);
+	FSelectValue EnumSelectValue = FTraversalBuilder::CreateSelectValue(EnumVariable.GetType().GetEnum(), EnumValue);
+	TargetContext.Add(EnumKey, EnumSelectValue);
+}
+
+void FStaticContextBuilder::GatherContextFromEmitterObject(FStaticContext& StaticContext, const FVersionedNiagaraEmitterData& EmitterData)
+{
+	AddBoolValue(StaticContext, ESelectKeySource::ExternalConstant, SYS_PARAM_EMITTER_LOCALSPACE, NAME_None, EmitterData.bLocalSpace);
+	AddBoolValue(StaticContext, ESelectKeySource::ExternalConstant, SYS_PARAM_EMITTER_DETERMINISM, NAME_None, EmitterData.bDeterminism);
+	AddEnumValue(StaticContext, ESelectKeySource::ExternalConstant, SYS_PARAM_EMITTER_SIMULATION_TARGET, NAME_None, (int32)EmitterData.SimTarget);
+}
+
+void FStaticContextBuilder::GatherContextFromScriptObject(FStaticContext& StaticContext, TMap<FGuid, FStaticContext>& FunctionInputContexts, const UNiagaraScript& Script, FTraversalDebugData* TraversalDebugData)
+{
+	AddEnumValue(StaticContext, ESelectKeySource::ExternalConstant, SYS_PARAM_SCRIPT_USAGE,
+		NAME_None, (int32)FNiagaraUtilities::ConvertScriptUsageToStaticSwitchUsage(Script.Usage));
+	AddEnumValue(StaticContext, ESelectKeySource::ExternalConstant, SYS_PARAM_SCRIPT_CONTEXT,
+		NAME_None, (int32)FNiagaraUtilities::ConvertScriptUsageToStaticSwitchContext(Script.Usage));
+
+	FTraversalCache::GetTopLevelScriptFunctionInputContexts(Script, FunctionInputContexts, TraversalDebugData);
+}
+
+void FStaticContextBuilder::GatherStaticParametersWrittenByScript(FStaticContext& StaticContext, const UNiagaraScript& SourceScript, FTraversalDebugData* TraversalDebugData)
+{
+	FStaticContext ScriptGlobalContext = StaticContext;
+	TMap<FGuid, FStaticContext> ScriptFunctionInputContexts;
+
+	GatherContextFromScriptObject(ScriptGlobalContext, ScriptFunctionInputContexts, SourceScript, TraversalDebugData);
+
+	FTraversalCallingContext CallingContext(ScriptGlobalContext, ScriptFunctionInputContexts);
+	CallingContext.ConnectionTraversalMode = EConnectionTraversalMode::MatchingOnly;
+
+	TSharedPtr<const FTraversal> SourceScriptTraversal = FTraversalCache::GetTopLevelScriptTraversal(SourceScript, TraversalDebugData);
+	if (SourceScriptTraversal.IsValid() && SourceScriptTraversal->CanWriteStaticAttributes(TraversalDebugData))
+	{
+		FStaticContext StaticWrites;
+		SourceScriptTraversal->Traverse(CallingContext, StaticWrites, TraversalDebugData);
+
+		for (const TPair<FSelectKey, FSelectValue>& StaticWrite : StaticWrites)
+		{
+			if (StaticWrite.Key.Source == ESelectKeySource::Attribute)
+			{
+				StaticContext.Add(StaticWrite);
+			}
+		}
+	}
+}
+
+} // UE::Niagara::TraversalCache
diff --git a/Engine/Plugins/FX/Niagara/Source/NiagaraEditor/Private/TraversalCache/StaticContextBuilder.h b/Engine/Plugins/FX/Niagara/Source/NiagaraEditor/Private/TraversalCache/StaticContextBuilder.h
new file mode 100644
--- /dev/null
+++ b/Engine/Plugins/FX/Niagara/Source/NiagaraEditor/Private/TraversalCache/StaticContextBuilder.h
@@ -0,0 +1,37 @@
+// Copyright Epic Games, Inc. All Rights Reserved.
+
+#pragma once
+
+#include "NiagaraVariant.h"
+#include "TraversalCache/TraversalNode.h"
+#include "TraversalCache/TraversalShared.h"
+
+struct FNiagaraVariableBase;
+class INiagaraEditorTypeUtilities;
+class UEnum;
+class UNiagaraScript;
+class UNiagaraSystem;
+struct FVersionedNiagaraEmitterData;
+
+namespace UE::Niagara::TraversalCache
+{
+
+class FStaticContextBuilder
+{
+public:
+	static NIAGARAEDITOR_API void CreateTopLevelScriptContext(
+		const UNiagaraSystem& System,
+		const FVersionedNiagaraEmitterData* EmitterData,
+		const UNiagaraScript& TargetScript,
+		FTopLevelScriptStaticContext& OutTopLevelScriptContext,
+		FTraversalDebugData* TraversalDebugData = nullptr);
+
+private:
+	static void AddBoolValue(FStaticContext& TargetContext, ESelectKeySource KeySource, const FNiagaraVariableBase& BoolVariable, FName NamespaceModifier, bool bValue);
+	static void AddEnumValue(FStaticContext& TargetContext, ESelectKeySource KeySource, const FNiagaraVariableBase& EnumVariable, FName NamespaceModifier, int32 EnumValue);
+	static void GatherContextFromEmitterObject(FStaticContext& StaticContext, const FVersionedNiagaraEmitterData& EmitterData);
+	static void GatherContextFromScriptObject(FStaticContext& StaticContext, TMap<FGuid, FStaticContext>& FunctionInputContexts,const UNiagaraScript& Script, FTraversalDebugData* TraversalDebugData);
+	static void GatherStaticParametersWrittenByScript(FStaticContext& StaticContext, const UNiagaraScript& SourceScript, FTraversalDebugData* TraversalDebugData);
+};
+
+} // UE::Niagara::TraversalCache
\ No newline at end of file
diff --git a/Engine/Plugins/FX/Niagara/Source/NiagaraEditor/Private/TraversalCache/Traversal.cpp b/Engine/Plugins/FX/Niagara/Source/NiagaraEditor/Private/TraversalCache/Traversal.cpp
new file mode 100644
--- /dev/null
+++ b/Engine/Plugins/FX/Niagara/Source/NiagaraEditor/Private/TraversalCache/Traversal.cpp
@@ -0,0 +1,832 @@
+// Copyright Epic Games, Inc. All Rights Reserved.
+
+#include "TraversalCache/Traversal.h"
+
+
+#include "NiagaraConstants.h"
+#include "NiagaraEditorCommon.h"
+#include "NiagaraModule.h"
+#include "TraversalCache/TraversalCache.h"
+#include "TraversalCache/TraversalNode.h"
+
+namespace UE::Niagara::TraversalCache
+{
+
+void FTraversal::Traverse(const FTraversalCallingContext& CallingContext, FStaticContext& OutLocalStaticContext, FTraversalDebugData* InOutDebugData) const
+{
+	TraverseInternal(CallingContext, OutLocalStaticContext, InOutDebugData);
+}
+
+TSharedRef<const FTraversalData> FTraversal::TraverseWithVisitor(const FTraversalCallingContext& CallingContext, const ITraversalVisitor& Visitor, FTraversalDebugData* InOutDebugData) const
+{
+	FStaticContext LocalStaticContext;
+	return TraverseWithVisitorInternal(CallingContext, Visitor, LocalStaticContext, InOutDebugData);
+}
+
+bool FTraversal::CanWriteStatics(FTraversalDebugData* InOutDebugData) const
+{
+	if (bCanWriteStaticsCache.IsSet() == false)
+	{
+		TSet<const FTraversal*> CheckedTraversals;
+		bCanWriteStaticsCache = CanWriteStaticsInternal(*this, CheckedTraversals, InOutDebugData);
+	}
+	return bCanWriteStaticsCache.GetValue();
+}
+
+bool FTraversal::CanWriteStaticAttributes(FTraversalDebugData* InOutDebugData) const
+{
+	if (bCanWriteStaticAttributesCache.IsSet() == false)
+	{
+		TSet<const FTraversal*> CheckedTraversals;
+		bCanWriteStaticAttributesCache = CanWriteStaticAttributesInternal(*this, CheckedTraversals, InOutDebugData);
+	}
+	return bCanWriteStaticAttributesCache.GetValue();
+}
+
+void FTraversal::ResetCachedData() const
+{
+	HashToTraversalDataCache.Empty();
+	bCanWriteStaticsCache.Reset();
+	bCanWriteStaticAttributesCache.Reset();
+}
+
+void FTraversal::TraverseInternal(const FTraversalCallingContext& CallingContext, FStaticContext& OutLocalStaticContext, FTraversalDebugData* InOutDebugData) const
+{
+	FTraversalCallingContext::FContextHash CallingContextHash = CallingContext.GenerateHash();
+	FCachedTraversalData* CachedTraversalData = HashToTraversalDataCache.Find(CallingContextHash);
+	if (CachedTraversalData != nullptr)
+	{
+		OutLocalStaticContext = CachedTraversalData->LocalStaticContext;
+		return;
+	}
+
+	FTraversalLocalContext LocalContext;
+	LocalContext.TraversalState = CallingContext.TraversalState;
+	TraverseNode(CallingContext, LocalContext, TraversalRoot.Get(), nullptr, nullptr, InOutDebugData);
+
+	CachedTraversalData = &HashToTraversalDataCache.Add(CallingContextHash);
+	CachedTraversalData->LocalStaticContext = LocalContext.StaticContext;
+	OutLocalStaticContext = LocalContext.StaticContext;
+}
+
+TSharedRef<const FTraversalData> FTraversal::TraverseWithVisitorInternal(const FTraversalCallingContext& CallingContext, const ITraversalVisitor& Visitor,
+	FStaticContext& OutLocalStaticContext, FTraversalDebugData* InOutDebugData) const
+{
+	FTraversalCallingContext::FContextHash CallingContextHash = CallingContext.GenerateHash();
+	FCachedTraversalData* CachedTraversalData = HashToTraversalDataCache.Find(CallingContextHash);
+	FGuid VisitorId = Visitor.GetVisitorId();
+	if (CachedTraversalData != nullptr)
+	{
+		TSharedRef<const FTraversalData>* CachedVisitorData = CachedTraversalData->VisitorIdToTraversalData.Find(VisitorId);
+		if (CachedVisitorData != nullptr)
+		{
+			OutLocalStaticContext = CachedTraversalData->LocalStaticContext;
+			return *CachedVisitorData;
+		}
+	}
+	 
+	TSharedRef<FTraversalData> TraversalData = Visitor.CreateTraversalData();
+	FTraversalLocalContext LocalContext;
+	LocalContext.TraversalState = CallingContext.TraversalState;
+	TraverseNode(CallingContext, LocalContext, TraversalRoot.Get(), &Visitor, &TraversalData.Get(), InOutDebugData);
+
+	FTraversalLocalContext UnconnectedLocalContext;
+	UnconnectedLocalContext.TraversalState = SetFlag(CallingContext.TraversalState, ETraversalStateFlags::UnconnectedRoot);
+	for (const TSharedRef<FTraversalNode>& UnconnectedTraversalRoot : UnconnectedTraversalRoots)
+	{
+		TraverseNode(CallingContext, UnconnectedLocalContext, UnconnectedTraversalRoot.Get(), &Visitor, &TraversalData.Get(), InOutDebugData);
+	}
+
+	if (InOutDebugData == nullptr || InOutDebugData->HasData() == false)
+	{
+		// We only cache the results when either no debug data was supplied, or when the debug data is empty.  When
+		// the debug data HasData returns true, errors were encountered while traversing or building.
+		if (CachedTraversalData == nullptr)
+		{
+			CachedTraversalData = &HashToTraversalDataCache.Add(CallingContextHash);
+		}
+		CachedTraversalData->VisitorIdToTraversalData.Add(VisitorId, TraversalData);
+		CachedTraversalData->LocalStaticContext = LocalContext.StaticContext;
+	}
+	OutLocalStaticContext = LocalContext.StaticContext;
+
+	return TraversalData;
+}
+
+void FTraversal::TraverseNode(
+	const FTraversalCallingContext& CallingContext, FTraversalLocalContext& LocalContext, const FTraversalNode& Node,
+	const ITraversalVisitor* Visitor, FTraversalData* TraversalData, FTraversalDebugData* DebugData) const
+{
+	bool bInSet = false;
+	LocalContext.TraversedNodesWithState.Add(TPair<const FTraversalNode*, ETraversalStateFlags>(&Node,  LocalContext.TraversalState), &bInSet);
+	if (bInSet)
+	{
+		return;
+	}
+
+	TOptional<FSelectConnectionData> SelectConnectionData;
+	if (Node.SelectData.IsSet())
+	{
+		SelectConnectionData = FSelectConnectionData();
+		ResolveSelectConnectionData(CallingContext, LocalContext, Node, SelectConnectionData.GetValue(), DebugData);
+	}
+
+	if (Node.Connections.Num() > 0)
+	{
+		TraverseConnections(CallingContext, LocalContext, Node, SelectConnectionData, Visitor, TraversalData, DebugData);
+	}
+
+	if (Node.FunctionCallData.IsSet())
+	{
+		TraverseFunction(CallingContext, LocalContext, Node, Visitor, TraversalData, DebugData);
+	}
+
+	if (Node.ParameterData.IsSet())
+	{
+		TraverseParameters(CallingContext, LocalContext, Node, DebugData);
+	}
+
+	if (Visitor != nullptr && TraversalData != nullptr)
+	{
+		Visitor->VisitNode(LocalContext.TraversalState, Node, *TraversalData);
+	}
+}
+
+void FTraversal::TraverseConnections(
+	const FTraversalCallingContext& CallingContext, FTraversalLocalContext& LocalContext, const FTraversalNode& Node, 
+	const TOptional<FSelectConnectionData>& SelectConnectionData, const ITraversalVisitor* Visitor, FTraversalData* TraversalData, FTraversalDebugData* DebugData) const
+{
+	for (const FTraversalNode::FConnection& Connection : Node.Connections)
+	{
+		bool bIsSelectableConnection = SelectConnectionData.IsSet() && Connection.PinId != SelectConnectionData->SelectConnectionPinId;
+		bool bIsSelectedConnection = bIsSelectableConnection && Connection.PinId == SelectConnectionData->SelectedConnectionPinId;
+
+		if (SelectConnectionData.IsSet() == false ||
+			CallingContext.ConnectionTraversalMode == EConnectionTraversalMode::All ||
+			(CallingContext.ConnectionTraversalMode == EConnectionTraversalMode::MatchingOnly && bIsSelectedConnection))
+		{
+			ETraversalStateFlags ConnectionStateFlags = bIsSelectableConnection && (bIsSelectedConnection == false)
+				? SetFlag(LocalContext.TraversalState, ETraversalStateFlags::CulledBySwitch)
+				: LocalContext.TraversalState;
+			FScopedTraversalState ConnectionScopedTraversalState(LocalContext, ConnectionStateFlags);
+			TraverseNode(CallingContext, LocalContext, Connection.GetNode(), Visitor, TraversalData, DebugData);
+		}
+	}
+}
+
+void FTraversal::TraverseFunction(
+	const FTraversalCallingContext& CallingContext, FTraversalLocalContext& LocalContext, const FTraversalNode& Node,
+	const ITraversalVisitor* Visitor, FTraversalData* TraversalData, FTraversalDebugData* DebugData) const
+{
+	if (Node.FunctionCallData->FunctionScriptReference.Path.IsValid() == false ||
+		HasFlag(LocalContext.TraversalState, ETraversalStateFlags::UnconnectedRoot))
+	{
+		return;
+	}
+
+	TSharedPtr<const FTraversal> CalledFunctionTraversal = FTraversalCache::GetScriptAssetTraversal(
+		Node.FunctionCallData->FunctionScriptReference.Path, Node.FunctionCallData->FunctionScriptReference.Version, DebugData);
+	if (CalledFunctionTraversal.IsValid() == false)
+	{
+		return;
+	}
+
+	const TMap<FGuid, FStaticContext> ExternalFunctionInputStaticContexts;
+	FTraversalCallingContext FunctionCallContext(CallingContext.GlobalStaticContext, ExternalFunctionInputStaticContexts);
+	SetupFunctionCallTraversalContext(CallingContext, LocalContext, Node, FunctionCallContext);
+	SetupFunctionCallStaticContext(CallingContext, LocalContext, Node, FunctionCallContext.FunctionCallStaticContext, DebugData);
+
+	FStaticContext TraversedFunctionLocalStaticContext;
+	if (Visitor != nullptr && TraversalData != nullptr)
+	{
+		TSharedRef<const FTraversalData> FunctionTraversalData = CalledFunctionTraversal->TraverseWithVisitorInternal(FunctionCallContext, *Visitor, TraversedFunctionLocalStaticContext, DebugData);
+		TraversalData->CalledFunctionTraversalData.Add(FFunctionCallTraversalData(Node.SourceNodeGuid, FunctionTraversalData));
+	}
+	else if (CalledFunctionTraversal->CanWriteStatics(DebugData) && Node.bSourceNodeEnabled)
+	{
+		CalledFunctionTraversal->TraverseInternal(FunctionCallContext, TraversedFunctionLocalStaticContext, DebugData);
+	}
+
+	if (Node.bSourceNodeEnabled && HasFlag(LocalContext.TraversalState, ETraversalStateFlags::CulledBySwitch) == false)
+	{
+		UpdateLocalStaticContextFromFunctionStaticWrites(TraversedFunctionLocalStaticContext, Node.FunctionCallData.GetValue(), LocalContext.StaticContext);
+	}
+}
+
+void FTraversal::TraverseParameters(const FTraversalCallingContext& CallingContext, FTraversalLocalContext& LocalContext, const FTraversalNode& Node, FTraversalDebugData* DebugData) const
+{
+	if (Node.ParameterData->WriteParameterReferences.Num() > 0 &&
+		Node.bSourceNodeEnabled &&
+		HasFlag(LocalContext.TraversalState, ETraversalStateFlags::CulledBySwitch) == false &&
+		HasFlag(LocalContext.TraversalState, ETraversalStateFlags::UnconnectedRoot) == false)
+	{
+		UpdateLocalContextFromParameterData(CallingContext, LocalContext, Node, DebugData);
+	}
+}
+
+void FTraversal::SetupFunctionCallStaticContext(
+	const FTraversalCallingContext& CallingContext, const FTraversalLocalContext& LocalContext, const FTraversalNode& Node,
+	FStaticContext& FunctionCallStaticContext, FTraversalDebugData* DebugData)
+{
+	const FFunctionCallData& FunctionCallData = Node.FunctionCallData.GetValue();
+
+	// Initialize the function context with matching external function input context which will contain static inputs set externally.
+	const FStaticContext* FunctionCallInputContext = CallingContext.ExternalFunctionInputStaticContexts.Find(Node.SourceNodeGuid);
+	if (FunctionCallInputContext != nullptr)
+	{
+		FunctionCallStaticContext.Append(*FunctionCallInputContext);
+	}
+
+	// Add any matching local per-function input context which will contain static inputs set in the calling context.
+	const FStaticContext* LocalFunctionCallInputContext = LocalContext.FunctionInputStaticContexts.Find(Node.SourceNodeGuid);
+	if (LocalFunctionCallInputContext != nullptr)
+	{
+		FunctionCallStaticContext.Append(*LocalFunctionCallInputContext);
+	}
+
+	// Add static values from inputs on the function call node, either from static switches or static inputs set directly, static inputs
+	// connected to other nodes, or from propagated static switches.
+	for (const FFunctionInputSelectValue& InputSelectValue : FunctionCallData.InputSelectValues)
+	{
+		if (InputSelectValue.ConnectionPinId.IsSet())
+		{
+			const FTraversalNode* ConnectedNode = Node.GetConnectedNodeByPinId(InputSelectValue.ConnectionPinId.GetValue());
+			if (ConnectedNode != nullptr)
+			{
+				TSet<const FTraversalNode*> ResolveTraversedNodes;
+				TOptional<FSelectValue> ConnectedInputValue = ResolveSelectValueForNode(CallingContext, LocalContext, *ConnectedNode, ResolveTraversedNodes, DebugData);
+				if (ConnectedInputValue.IsSet())
+				{
+					FunctionCallStaticContext.Add(InputSelectValue.InputSelectKey, ConnectedInputValue.GetValue());
+				}
+			}
+		}
+		else if (InputSelectValue.LocalValue.IsSet())
+		{
+			FunctionCallStaticContext.Add(InputSelectValue.InputSelectKey, InputSelectValue.LocalValue.GetValue());
+		}
+		else
+		{
+			// If the select value wasn't supplied in the function reference or in a connected node, it's using switch propagation, so check the calling
+			// context to see if it has a value, and if so, use that.
+			FSelectKey CurrentKey;
+			if (InputSelectValue.OptionalPropagatedNameOverride.IsSet() == false)
+			{
+				CurrentKey = InputSelectValue.InputSelectKey;
+			}
+			else
+			{
+				CurrentKey = FSelectKey(
+					ESelectKeySource::FunctionCallNode,
+					FNiagaraVariableBase(InputSelectValue.InputSelectKey.Variable.GetType(), InputSelectValue.OptionalPropagatedNameOverride.GetValue()),
+					NAME_None);
+			}
+
+			TOptional<FSelectValue> CurrentValue = CallingContext.FunctionCallStaticContext.GetSelectValue(CurrentKey);
+			if (CurrentValue.IsSet())
+			{
+				FunctionCallStaticContext.Add(InputSelectValue.InputSelectKey, CurrentValue.GetValue());
+			}
+		}
+	}
+
+	// Lastly set the debug state based on the function reference, if it hasn't been set externally.
+	FSelectKey DebugStateKey(ESelectKeySource::ExternalConstant, SYS_PARAM_FUNCTION_DEBUG_STATE, NAME_None);
+	if (CallingContext.GlobalStaticContext.Contains(DebugStateKey) == false)
+	{
+		if (FunctionCallData.DebugState == ENiagaraFunctionDebugState::NoDebug)
+		{
+			FunctionCallStaticContext.Add(DebugStateKey, FSelectValue::GetDebugStateNoDebug());
+		}
+		else if (FunctionCallData.DebugState == ENiagaraFunctionDebugState::Basic)
+		{
+			FunctionCallStaticContext.Add(DebugStateKey, FSelectValue::GetDebugStateBasic());
+		}
+		else
+		{
+			ensureMsgf(false, TEXT("Traversal encountered an unknown debug state."));
+		}
+	}
+}
+
+void FTraversal::SetupFunctionCallTraversalContext(const FTraversalCallingContext& CallingContext, const FTraversalLocalContext& LocalContext,
+	const FTraversalNode& Node, FTraversalCallingContext& FunctionCallContext)
+{
+	FunctionCallContext.FunctionCallStack = CallingContext.FunctionCallStack;
+	FunctionCallContext.FunctionCallStack.Add(Node.SourceNodeGuid);
+
+	FunctionCallContext.CallingStaticContext = CallingContext.CallingStaticContext;
+	FunctionCallContext.CallingStaticContext.Append(LocalContext.StaticContext);
+
+	FunctionCallContext.TraversalState = Node.bSourceNodeEnabled
+		? LocalContext.TraversalState
+		: SetFlag(LocalContext.TraversalState, ETraversalStateFlags::CallerDisabled);
+}
+
+void FTraversal::UpdateLocalStaticContextFromFunctionStaticWrites(const FStaticContext& FunctionCallLocalStaticContext, const FFunctionCallData& FunctionCallData, FStaticContext& LocalStaticContext)
+{
+	for (const TPair<FSelectKey, FSelectValue>& FunctionCallLocalStaticPair : FunctionCallLocalStaticContext)
+	{
+		if (FunctionCallLocalStaticPair.Key.NamespaceModifier == NAME_None)
+		{
+			if (FunctionCallLocalStaticPair.Key.Source == ESelectKeySource::Attribute)
+			{
+				LocalStaticContext.Add(FunctionCallLocalStaticPair);
+			}
+		}
+		else
+		{
+			if (FunctionCallLocalStaticPair.Key.Source == ESelectKeySource::Attribute ||
+				FunctionCallLocalStaticPair.Key.Source == ESelectKeySource::ModuleOutput)
+			{
+				FSelectKey UpdatedKey(FunctionCallLocalStaticPair.Key.Source, FunctionCallLocalStaticPair.Key.Variable, FunctionCallData.FunctionCallName);
+				LocalStaticContext.Add(UpdatedKey, FunctionCallLocalStaticPair.Value);
+			}
+		}
+	}
+}
+
+void FTraversal::UpdateLocalContextFromParameterData(const FTraversalCallingContext& CallingContext, FTraversalLocalContext& LocalContext, const FTraversalNode& Node, FTraversalDebugData* DebugData)
+{
+	for (const FParameterWrite& WriteParameterReference : Node.ParameterData->WriteParameterReferences)
+	{
+		if (WriteParameterReference.Parameter.GetType().IsStatic())
+		{
+			TOptional<FSelectValue> WriteValue;
+			if (WriteParameterReference.OptionalLocalSelectValue.IsSet())
+			{
+				WriteValue = WriteParameterReference.OptionalLocalSelectValue;
+			}
+			else if (WriteParameterReference.OptionalConnectionPinId.IsSet())
+			{
+				const FTraversalNode* WriteConnectedNode = Node.GetConnectedNodeByPinId(WriteParameterReference.OptionalConnectionPinId.GetValue());
+				if (WriteConnectedNode != nullptr)
+				{
+					TSet<const FTraversalNode*> ResolveTraversedNodes;
+					WriteValue = ResolveSelectValueForNode(CallingContext, LocalContext, *WriteConnectedNode, ResolveTraversedNodes, DebugData);
+				} 
+			}
+			if (WriteValue.IsSet())
+			{
+				if (WriteParameterReference.OptionalTargetFunctionCallNodeGuid.IsSet())
+				{
+					LocalContext.FunctionInputStaticContexts.FindOrAdd(WriteParameterReference.OptionalTargetFunctionCallNodeGuid.GetValue())
+						.Add(FSelectKey(ESelectKeySource::ModuleInput, WriteParameterReference.Parameter, NAME_None), WriteValue.GetValue());
+				}
+				else
+				{
+					ESelectKeySource SelectKeySource = ParameterFlagsToSelectKeySource(WriteParameterReference.Flags);
+					FName NamespaceModifier = WriteParameterReference.OptionalNamespaceModifier.IsSet() 
+						? WriteParameterReference.OptionalNamespaceModifier.GetValue()
+						: NAME_None;
+					LocalContext.StaticContext.Add(FSelectKey(SelectKeySource, WriteParameterReference.Parameter, NamespaceModifier), WriteValue.GetValue());
+				}
+			}
+		}
+	}
+}
+
+void FTraversal::ResolveSelectConnectionData(const FTraversalCallingContext& CallingContext, FTraversalLocalContext& LocalContext, const FTraversalNode& Node, FTraversal::FSelectConnectionData& OutSelectConnectionData, FTraversalDebugData* DebugData)
+{
+	OutSelectConnectionData.SelectConnectionPinId = Node.SelectData->SelectConnectionPinId;
+	if (Node.SelectData->SelectMode == ESelectMode::None)
+	{
+		return;
+	}
+
+	TSet<const FTraversalNode*> ResolveTraversedNodes;
+	TOptional<FSelectValue> SelectValue = ResolveSelectValueForSelectData(CallingContext, LocalContext, Node, ResolveTraversedNodes, DebugData);
+	if (SelectValue.IsSet() == false)
+	{
+		if (DebugData != nullptr &&	HasFlag(LocalContext.TraversalState, ETraversalStateFlags::CulledBySwitch) == false)
+		{
+			DebugData->AddUnresolvedSelect(CallingContext.FunctionCallStack, Node.SourceNodeGuid, Node.SelectData->SelectKey);
+		}
+		return;
+	}
+
+	const FSelectInputData* SelectInputData = Node.SelectData->FindInputDataForSelectValue(SelectValue.GetValue());
+	if (SelectInputData == nullptr)
+	{
+		if (DebugData != nullptr)
+		{
+			DebugData->AddUnresolvedSelectInput(CallingContext.FunctionCallStack, Node.SourceNodeGuid, Node.SelectData->SelectKey, SelectValue.GetValue());
+		}
+		return;
+	}
+
+	OutSelectConnectionData.SelectedConnectionPinId = SelectInputData->ConnectionPinId.IsSet() ? SelectInputData->ConnectionPinId.GetValue() : FGuid();
+}
+
+TOptional<FSelectValue> FTraversal::ResolveSelectValueForSelectData(const FTraversalCallingContext& CallingContext, const FTraversalLocalContext& LocalContext,
+	const FTraversalNode& Node, TSet<const FTraversalNode*>& ResolveTraversedNodes, FTraversalDebugData* DebugData)
+{
+	TOptional<FSelectValue> SelectValue;
+	if (Node.SelectData->SelectMode == ESelectMode::Connection)
+	{
+		const FTraversalNode* ConnectedNode = Node.GetConnectedNodeByPinId(Node.SelectData->SelectConnectionPinId);
+		if (ConnectedNode != nullptr)
+		{
+			SelectValue = ResolveSelectValueForNode(CallingContext, LocalContext, *ConnectedNode, ResolveTraversedNodes, DebugData);
+		}
+	}
+	else if (Node.SelectData->SelectMode == ESelectMode::Value)
+	{
+		SelectValue = LocalContext.GetSelectValue(CallingContext, Node.SelectData->SelectKey);
+	}
+	return SelectValue;
+}
+
+TOptional<FSelectValue> FTraversal::ResolveSelectValueForNode(const FTraversalCallingContext& CallingContext, const FTraversalLocalContext& LocalContext,
+	const FTraversalNode& Node, TSet<const FTraversalNode*>& ResolveTraversedNodes, FTraversalDebugData* DebugData)
+{
+	if (ResolveTraversedNodes.Contains(&Node))
+	{
+		return TOptional<FSelectValue>();
+	}
+	TScopedAddToSet<FTraversalNode> ScopedAddTraversedNode(ResolveTraversedNodes, Node);
+
+	if (Node.FunctionInputData.IsSet())
+	{
+		return ResolveSelectValueForFunctionInputNode(CallingContext, LocalContext, Node, ResolveTraversedNodes, DebugData);
+	}
+	else if (Node.SelectData.IsSet())
+	{
+		return ResolveSelectValueForSelectNode(CallingContext, LocalContext, Node, ResolveTraversedNodes, DebugData);
+	}
+	else if (Node.ParameterData.IsSet())
+	{
+		return ResolveSelectValueForParameterNode(CallingContext, LocalContext, Node, ResolveTraversedNodes, DebugData);
+	}
+	else if (Node.StaticOpData.IsSet())
+	{
+		return ResolveSelectValueForOpNode(CallingContext, LocalContext, Node, ResolveTraversedNodes, DebugData);
+	}
+	else if (Node.FunctionCallData.IsSet())
+	{
+		return ResolveSelectValueForFunctionCallNode(CallingContext, LocalContext, Node, ResolveTraversedNodes, DebugData);
+	}
+
+	return TOptional<FSelectValue>();
+}
+
+TOptional<FSelectValue> FTraversal::ResolveSelectValueForFunctionInputNode(const FTraversalCallingContext& CallingContext, const FTraversalLocalContext& LocalContext,
+	const FTraversalNode& Node, TSet<const FTraversalNode*>& ResolveTraversedNodes, FTraversalDebugData* DebugData)
+{
+	TOptional<FSelectValue> InputValue;
+	InputValue = LocalContext.GetSelectValue(CallingContext, Node.FunctionInputData->InputSelectKey);
+	if (InputValue.IsSet() == false && Node.FunctionInputData->LocalValue.IsSet())
+	{
+		InputValue = Node.FunctionInputData->LocalValue;
+	}
+
+	if (InputValue.IsSet() == false && DebugData != nullptr)
+	{
+		DebugData->AddUnresolvedFunctionInput(CallingContext.FunctionCallStack, Node.SourceNodeGuid, Node.FunctionInputData->InputSelectKey.Variable);
+	}
+	return InputValue;
+}
+
+TOptional<FSelectValue> FTraversal::ResolveSelectValueForSelectNode(const FTraversalCallingContext& CallingContext, const FTraversalLocalContext& LocalContext,
+	const FTraversalNode& Node, TSet<const FTraversalNode*>& ResolveTraversedNodes, FTraversalDebugData* DebugData)
+{
+	TOptional<FSelectValue> SelectDataValue = ResolveSelectValueForSelectData(CallingContext, LocalContext, Node, ResolveTraversedNodes, DebugData);
+	if (SelectDataValue.IsSet() == false)
+	{
+		return TOptional<FSelectValue>();
+	}
+
+	const FSelectInputData* SelectInputData = Node.SelectData->FindInputDataForSelectValue(SelectDataValue.GetValue());
+	if (SelectInputData == nullptr ||
+		(SelectInputData->LocalValue.IsSet() == false && SelectInputData->ConnectionPinId.IsSet() == false))
+	{
+		return TOptional<FSelectValue>();
+	}
+
+	if (SelectInputData->LocalValue.IsSet())
+	{
+		return SelectInputData->LocalValue;
+	}
+
+	const FTraversalNode* ConnectedNode = Node.GetConnectedNodeByPinId(SelectInputData->ConnectionPinId.GetValue());
+	return ConnectedNode != nullptr
+		? ResolveSelectValueForNode(CallingContext, LocalContext, *ConnectedNode, ResolveTraversedNodes, DebugData)
+		: TOptional<FSelectValue>();
+}
+
+TOptional<FSelectValue> FTraversal::ResolveSelectValueForParameterNode(const FTraversalCallingContext& CallingContext, const FTraversalLocalContext& LocalContext,
+	const FTraversalNode& Node, TSet<const FTraversalNode*>& ResolveTraversedNodes, FTraversalDebugData* DebugData)
+{
+	TOptional<FSelectValue> ReadValue;
+	TArray<FParameterRead> NonDiscoverReads;
+	for (const FParameterRead& ReadParameterReference : Node.ParameterData->ReadParameterReferences)
+	{
+		if (ReadParameterReference.bIsDiscoverRead == false)
+		{
+			NonDiscoverReads.Add(ReadParameterReference);
+		}
+	}
+
+	if (NonDiscoverReads.Num() == 1 && NonDiscoverReads[0].Parameter.GetType().IsStatic())
+	{
+		const FParameterRead& NonDiscoverRead = NonDiscoverReads[0];
+
+		ESelectKeySource KeySource = ParameterFlagsToSelectKeySource(NonDiscoverRead.Flags);
+		FName NamespaceModifier = NonDiscoverRead.OptionalNamespaceModifier.IsSet()
+			? NonDiscoverRead.OptionalNamespaceModifier.GetValue()
+			: NAME_None;
+		FSelectKey StaticValueSelectKey(KeySource, NonDiscoverRead.Parameter, NamespaceModifier);
+		ReadValue = LocalContext.GetSelectValue(CallingContext, StaticValueSelectKey);
+
+		if (ReadValue.IsSet() == false && NonDiscoverRead.OptionalLocalSelectValue.IsSet())
+		{
+			ReadValue = NonDiscoverRead.OptionalLocalSelectValue;
+		}
+		if (ReadValue.IsSet() == false && NonDiscoverRead.OptionalConnectionPinId.IsSet())
+		{
+			const FTraversalNode* ReadDefaultNode = Node.GetConnectedNodeByPinId(NonDiscoverRead.OptionalConnectionPinId.GetValue());
+			if (ReadDefaultNode != nullptr)
+			{
+				ReadValue = ResolveSelectValueForNode(CallingContext, LocalContext, *ReadDefaultNode, ResolveTraversedNodes, DebugData);
+			}
+		}
+		if (ReadValue.IsSet() == false)
+		{
+			const FTraversalNode* ExecutionConnectionNode = Node.GetConnectedNodeByPinId(Node.ParameterData->ExecutionConnectionPinId);
+			if (ExecutionConnectionNode != nullptr)
+			{
+				ReadValue = ResolveStaticReadValueForNode(CallingContext, LocalContext, *ExecutionConnectionNode, StaticValueSelectKey, ResolveTraversedNodes, DebugData);
+			}
+		}
+
+		if (ReadValue.IsSet() == false && DebugData != nullptr)
+		{
+			DebugData->AddUnresolvedRead(CallingContext.FunctionCallStack, Node.SourceNodeGuid, NonDiscoverRead.Parameter);
+		}
+
+		return ReadValue;
+	}
+	return TOptional<FSelectValue>();
+}
+
+TOptional<FSelectValue> FTraversal::ResolveSelectValueForOpNode(const FTraversalCallingContext& CallingContext, const FTraversalLocalContext& LocalContext,
+	const FTraversalNode& Node, TSet<const FTraversalNode*>& ResolveTraversedNodes, FTraversalDebugData* DebugData)
+{
+	TArray<int32> OpInputs;
+	bool bInputUnresolved = false;
+	int32 UnresolvedPinIndex = INDEX_NONE;
+	for (int32 InputIndex = 0; InputIndex < Node.StaticOpData->InputData.Num() && bInputUnresolved == false; ++InputIndex)
+	{
+		const FStaticOpInputData& InputData = Node.StaticOpData->InputData[InputIndex];
+		TOptional<int32> InputValue;
+		if (InputData.LocalValue.IsSet())
+		{
+			InputValue = InputData.LocalValue->NumericValue;
+		}
+		else if (InputData.ConnectionPinId.IsSet())
+		{
+			const FTraversalNode* ConnectedNode = Node.GetConnectedNodeByPinId(InputData.ConnectionPinId.GetValue());
+			if (ConnectedNode != nullptr)
+			{
+				TOptional<FSelectValue> ResolvedInputValue = ResolveSelectValueForNode(CallingContext, LocalContext, *ConnectedNode, ResolveTraversedNodes, DebugData);
+				if (ResolvedInputValue.IsSet())
+				{
+					InputValue = ResolvedInputValue->NumericValue;
+				}
+			}
+		}
+
+		if (InputValue.IsSet() == false)
+		{
+			bInputUnresolved = true;
+			UnresolvedPinIndex = InputIndex;
+		}
+		else
+		{
+			OpInputs.Add(InputValue.GetValue());
+		}
+	}
+
+	if (bInputUnresolved == false)
+	{
+		const FNiagaraOpInfo* OpInfo = FNiagaraOpInfo::GetOpInfo(Node.StaticOpData->OpName);
+		int32 OpResult = OpInfo->StaticVariableResolveFunction.Execute(OpInputs);
+		return FSelectValue(OpResult, "OpResult");
+	}
+	else
+	{
+		if (DebugData != nullptr)
+		{
+			DebugData->AddUnresolvedStaticOp(CallingContext.FunctionCallStack, Node.SourceNodeGuid, Node.StaticOpData->OpName, UnresolvedPinIndex);
+		}
+		return TOptional<FSelectValue>();
+	}
+}
+
+TOptional<FSelectValue> FTraversal::ResolveSelectValueForFunctionCallNode(const FTraversalCallingContext& CallingContext, const FTraversalLocalContext& LocalContext,
+	const FTraversalNode& Node, TSet<const FTraversalNode*>& ResolveTraversedNodes, FTraversalDebugData* DebugData)
+{
+	if (Node.FunctionCallData->FunctionScriptReference.Path.IsValid() == false)
+	{
+		return TOptional<FSelectValue>();
+	}
+
+	TSharedPtr<const FTraversal> ResolveFunctionTraversal = FTraversalCache::GetScriptAssetTraversal(Node.FunctionCallData->FunctionScriptReference.Path,
+		Node.FunctionCallData->FunctionScriptReference.Version, DebugData);
+	if (ResolveFunctionTraversal.IsValid() == false)
+	{
+		return TOptional<FSelectValue>();
+	}
+
+	const TMap<FGuid, FStaticContext> ExternalFunctionInputStaticContexts;
+	FTraversalCallingContext FunctionCallContext(CallingContext.GlobalStaticContext, ExternalFunctionInputStaticContexts);
+	SetupFunctionCallTraversalContext(CallingContext, LocalContext, Node, FunctionCallContext);
+	SetupFunctionCallStaticContext(CallingContext, LocalContext, Node, FunctionCallContext.FunctionCallStaticContext, DebugData);
+
+	FTraversalLocalContext FunctionLocalContext;
+	return ResolveFunctionTraversal->ResolveSelectValueForFunctionTraversal(FunctionCallContext, DebugData);
+}
+
+TOptional<FSelectValue> FTraversal::ResolveSelectValueForFunctionTraversal(const FTraversalCallingContext& FunctionCallingContext, FTraversalDebugData* DebugData) const
+{
+	if (TraversalRoot->Connections.Num() == 1)
+	{
+		TSet<const FTraversalNode*> ResolveTraversedNodes;
+		FTraversalLocalContext FunctionLocalContext;
+		return ResolveSelectValueForNode(FunctionCallingContext, FunctionLocalContext, TraversalRoot->Connections[0].GetNode(), ResolveTraversedNodes, DebugData);
+	}
+	return TOptional<FSelectValue>();
+}
+
+TOptional<FSelectValue> FTraversal::ResolveStaticReadValueForNode(const FTraversalCallingContext& CallingContext, const FTraversalLocalContext& LocalContext,
+	const FTraversalNode& Node, FSelectKey& ReadSelectKey, TSet<const FTraversalNode*>& ResolveTraversedNodes, FTraversalDebugData* DebugData)
+{
+	if (ResolveTraversedNodes.Contains(&Node))
+	{
+		return TOptional<FSelectValue>();
+	}
+	TScopedAddToSet<FTraversalNode> ScopedAddTraversedNode(ResolveTraversedNodes, Node);
+
+	if (Node.SelectData.IsSet())
+	{
+		return ResolveStaticReadValueForSelectNode(CallingContext, LocalContext, Node, ReadSelectKey, ResolveTraversedNodes, DebugData);
+	}
+	else if (Node.ParameterData.IsSet())
+	{
+		return ResolveStaticReadValueForParameterNode(CallingContext, LocalContext, Node, ReadSelectKey, ResolveTraversedNodes, DebugData);
+	}
+	else if (Node.FunctionCallData.IsSet())
+	{
+		return ResolveStaticReadValueForFunctionCallNode(CallingContext, LocalContext, Node, ReadSelectKey, ResolveTraversedNodes, DebugData);
+	}
+
+	return TOptional<FSelectValue>();
+}
+
+TOptional<FSelectValue> FTraversal::ResolveStaticReadValueForSelectNode(const FTraversalCallingContext& CallingContext, const FTraversalLocalContext& LocalContext,
+	const FTraversalNode& Node, FSelectKey& ReadSelectKey, TSet<const FTraversalNode*>& ResolveTraversedNodes, FTraversalDebugData* DebugData)
+{
+	TOptional<FSelectValue> SelectDataValue = ResolveSelectValueForSelectData(CallingContext, LocalContext, Node, ResolveTraversedNodes, DebugData);
+	if (SelectDataValue.IsSet() == false)
+	{
+		return TOptional<FSelectValue>();
+	}
+
+	const FSelectInputData* SelectInputData = Node.SelectData->FindInputDataForSelectValue(SelectDataValue.GetValue());
+	if (SelectInputData == nullptr || SelectInputData->ConnectionPinId.IsSet() == false)
+	{
+		return TOptional<FSelectValue>();
+	}
+	
+	const FTraversalNode* ConnectedNode = Node.GetConnectedNodeByPinId(SelectInputData->ConnectionPinId.GetValue());
+	return ConnectedNode != nullptr 
+		? ResolveStaticReadValueForNode(CallingContext, LocalContext, *ConnectedNode, ReadSelectKey, ResolveTraversedNodes, DebugData)
+		: TOptional<FSelectValue>();
+}
+
+TOptional<FSelectValue> FTraversal::ResolveStaticReadValueForParameterNode(const FTraversalCallingContext& CallingContext, const FTraversalLocalContext& LocalContext,
+	const FTraversalNode& Node, FSelectKey& ReadSelectKey, TSet<const FTraversalNode*>& ResolveTraversedNodes, FTraversalDebugData* DebugData)
+{
+	const FParameterWrite* MatchingWriteParameterReference = Node.ParameterData->WriteParameterReferences.FindByPredicate([&ReadSelectKey]
+		(const FParameterWrite& WriteParameterReference) { return WriteParameterReference.Parameter == ReadSelectKey.Variable; });
+	if (MatchingWriteParameterReference != nullptr)
+	{
+		if (MatchingWriteParameterReference->OptionalLocalSelectValue.IsSet())
+		{
+			return MatchingWriteParameterReference->OptionalLocalSelectValue;
+		}
+		if (MatchingWriteParameterReference->OptionalConnectionPinId.IsSet())
+		{
+			const FTraversalNode* WriteValueNode = Node.GetConnectedNodeByPinId(MatchingWriteParameterReference->OptionalConnectionPinId.GetValue());
+			if (WriteValueNode != nullptr)
+			{
+				return ResolveSelectValueForNode(CallingContext, LocalContext, *WriteValueNode, ResolveTraversedNodes, DebugData);
+			}
+		}
+	}
+
+	const FTraversalNode* ExecutionConnectionNode = Node.GetConnectedNodeByPinId(Node.ParameterData->ExecutionConnectionPinId);
+	return ExecutionConnectionNode != nullptr
+		? ResolveStaticReadValueForNode(CallingContext, LocalContext, *ExecutionConnectionNode, ReadSelectKey, ResolveTraversedNodes, DebugData)
+		: TOptional<FSelectValue>();
+}
+
+TOptional<FSelectValue> FTraversal::ResolveStaticReadValueForFunctionCallNode(const FTraversalCallingContext& CallingContext, const FTraversalLocalContext& LocalContext,
+	const FTraversalNode& Node, FSelectKey& ReadSelectKey, TSet<const FTraversalNode*>& ResolveTraversedNodes, FTraversalDebugData* DebugData)
+{
+	const FTraversalNode* ExecutionConnectionNode = Node.GetConnectedNodeByPinId(Node.FunctionCallData->ExecutionConnectionPinId);
+	return ExecutionConnectionNode != nullptr
+		? ResolveStaticReadValueForNode(CallingContext, LocalContext, *ExecutionConnectionNode, ReadSelectKey, ResolveTraversedNodes, DebugData)
+		: TOptional<FSelectValue>();
+}
+
+bool FTraversal::CanWriteStaticsInternal(const FTraversal& Traversal, TSet<const FTraversal*>& CheckedTraversals, FTraversalDebugData* DebugData)
+{
+	bool bInSet = false;
+	CheckedTraversals.Add(&Traversal, &bInSet);
+	if (bInSet)
+	{
+		return false;
+	}
+
+	if (Traversal.StaticVariableWrites.Num() > 0)
+	{
+		return true;
+	}
+
+	for (const FScriptReference& ExternalReference : Traversal.ExternalReferences)
+	{
+		TSharedPtr<const FTraversal> ReferencedTraversal = FTraversalCache::GetScriptAssetTraversal(ExternalReference.Path, ExternalReference.Version, DebugData);
+		if (ReferencedTraversal.IsValid() && CanWriteStaticsInternal(*ReferencedTraversal.Get(), CheckedTraversals, DebugData))
+		{
+			return true;
+		}
+	}
+	return false;
+}
+
+bool FTraversal::CanWriteStaticAttributesInternal(const FTraversal& Traversal, TSet<const FTraversal*>& CheckedTraversals, FTraversalDebugData* DebugData)
+{
+	bool bInSet = false;
+	CheckedTraversals.Add(&Traversal, &bInSet);
+	if (bInSet)
+	{
+		return false;
+	}
+
+	if (Traversal.StaticVariableWritesToAttributes.Num() > 0)
+	{
+		return true;
+	}
+
+	for (const FScriptReference& ExternalReference : Traversal.ExternalReferences)
+	{
+		TSharedPtr<const FTraversal> ReferencedTraversal = FTraversalCache::GetScriptAssetTraversal(ExternalReference.Path, ExternalReference.Version, DebugData);
+		if (ReferencedTraversal.IsValid() && CanWriteStaticAttributesInternal(*ReferencedTraversal.Get(), CheckedTraversals, DebugData))
+		{
+			return true;
+		}
+	}
+	return false;
+}
+
+void FTraversal::GetAllConnectedNodes(const FTraversalNode& Node, TSet<const FTraversalNode*>& OutNodes)
+{
+	bool bInSet = false;
+	OutNodes.Add(&Node, &bInSet);
+	if (bInSet == false)
+	{
+		for (const FTraversalNode::FConnection& Connection : Node.Connections)
+		{
+			GetAllConnectedNodes(Connection.GetNode(), OutNodes);
+		}
+	}
+}
+
+ESelectKeySource FTraversal::ParameterFlagsToSelectKeySource(EParameterFlags ParameterFlags)
+{
+	ESelectKeySource SelectKeySource;
+	if (HasFlag(ParameterFlags, EParameterFlags::Attribute))
+	{
+		SelectKeySource = ESelectKeySource::Attribute;
+	}
+	else if (HasFlag(ParameterFlags, EParameterFlags::ModuleInput))
+	{
+		SelectKeySource = ESelectKeySource::ModuleInput;
+	}
+	else if (HasFlag(ParameterFlags, EParameterFlags::ModuleLocal))
+	{
+		SelectKeySource = ESelectKeySource::ModuleLocal;
+	}
+	else if (HasFlag(ParameterFlags, EParameterFlags::ModuleOutput))
+	{
+		SelectKeySource = ESelectKeySource::ModuleOutput;
+	}
+	else
+	{
+		SelectKeySource = ESelectKeySource::None;
+	}
+	return SelectKeySource;
+}
+
+} // UE::Niagara::TraversalCache
\ No newline at end of file
diff --git a/Engine/Plugins/FX/Niagara/Source/NiagaraEditor/Private/TraversalCache/Traversal.h b/Engine/Plugins/FX/Niagara/Source/NiagaraEditor/Private/TraversalCache/Traversal.h
new file mode 100644
--- /dev/null
+++ b/Engine/Plugins/FX/Niagara/Source/NiagaraEditor/Private/TraversalCache/Traversal.h
@@ -0,0 +1,180 @@
+// Copyright Epic Games, Inc. All Rights Reserved.
+
+#pragma once
+
+#include "Delegates/DelegateCombinations.h"
+#include "TraversalCache/TraversalShared.h"
+
+namespace UE::Niagara::TraversalCache
+{
+
+struct FFunctionCallData;
+struct FTraversalDebugData;
+class  FTraversalNode;
+
+class FTraversal
+{
+	friend class FTraversalBuilder;
+
+public:
+	FTraversal(TSharedRef<FTraversalNode> InTraversalRoot, TArray<TSharedRef<FTraversalNode>>& InUnconnectedTraversalRoots)
+		: TraversalRoot(InTraversalRoot)
+		, UnconnectedTraversalRoots(InUnconnectedTraversalRoots)
+	{
+	}
+
+	void Traverse(const FTraversalCallingContext&, FStaticContext& OutLocalStaticContext, FTraversalDebugData* InOutDebugData = nullptr) const;
+
+	TSharedRef<const FTraversalData> TraverseWithVisitor(const FTraversalCallingContext& CallingContext, const ITraversalVisitor& Visitor, FTraversalDebugData* InOutDebugData = nullptr) const;
+
+	const TMap<FName, FGuid>& GetFunctionNameToNodeGuidMap() const { return FunctionNameToNodeGuidMap; }
+
+	bool CanWriteStatics(FTraversalDebugData* InOutDebugData = nullptr) const;
+
+	bool CanWriteStaticAttributes(FTraversalDebugData* InOutDebugData = nullptr) const;
+
+	const TSet<FScriptReference>& GetExternalReferences() const { return ExternalReferences; }
+
+	void ResetCachedData() const;
+
+private:
+	struct FScopedTraversalState
+	{
+		FScopedTraversalState(FTraversalLocalContext& InLocalContext, ETraversalStateFlags InRequestedState)
+			: LocalContext(InLocalContext)
+			, OriginalState(LocalContext.TraversalState)
+			, RequestedState(InRequestedState)
+		{
+			LocalContext.TraversalState = RequestedState;
+		}
+
+		~FScopedTraversalState()
+		{
+			LocalContext.TraversalState = OriginalState;
+		}
+
+	private:
+		FTraversalLocalContext& LocalContext;
+		ETraversalStateFlags OriginalState;
+		ETraversalStateFlags RequestedState;
+	};
+
+	template<typename TItem>
+	class TScopedAddToSet
+	{
+	public:
+		TScopedAddToSet(TSet<const TItem*>& InItemSet, const TItem& InItemToAdd)
+			: ItemSet(InItemSet)
+			, AddItem(InItemToAdd)
+		{
+			ItemSet.Add(&AddItem);
+		}
+
+		~TScopedAddToSet()
+		{
+			ItemSet.Remove(&AddItem);
+		}
+
+	private:
+		TSet<const TItem*>& ItemSet;
+		const TItem& AddItem;
+	};
+
+	struct FSelectConnectionData
+	{
+		TOptional<FGuid> SelectConnectionPinId;
+		FGuid SelectedConnectionPinId;
+	};
+
+	void TraverseInternal(const FTraversalCallingContext& CallingContext, FStaticContext& OutLocalStaticContext, FTraversalDebugData* InOutDebugData) const;
+
+	TSharedRef<const FTraversalData> TraverseWithVisitorInternal(const FTraversalCallingContext& CallingContext, const ITraversalVisitor& Visitor,
+		FStaticContext& OutLocalStaticContext, FTraversalDebugData* InOutDebugData) const;
+
+	void TraverseNode(const FTraversalCallingContext& CallingContext, FTraversalLocalContext& LocalContext, const FTraversalNode& Node,
+		const ITraversalVisitor* Visitor, FTraversalData* TraversalData, FTraversalDebugData* DebugData) const;
+
+	void TraverseConnections(const FTraversalCallingContext& CallingContext, FTraversalLocalContext& LocalContext, const FTraversalNode& Node,
+		const TOptional<FSelectConnectionData>& SelectConnectionData, const ITraversalVisitor* Visitor, FTraversalData* TraversalData, FTraversalDebugData* DebugData) const;
+
+	void TraverseFunction(const FTraversalCallingContext& CallingContext, FTraversalLocalContext& LocalContext, const FTraversalNode& Node,
+		const ITraversalVisitor* Visitor, FTraversalData* TraversalData, FTraversalDebugData* DebugData) const;
+
+	void TraverseParameters(const FTraversalCallingContext& CallingContext, FTraversalLocalContext& LocalContext, const FTraversalNode& Node, FTraversalDebugData* DebugData) const;
+
+	static void SetupFunctionCallStaticContext(const FTraversalCallingContext& CallingContext, const FTraversalLocalContext& LocalContext,
+		const FTraversalNode& Node, FStaticContext& OutFunctionCallStaticContext, FTraversalDebugData* DebugData);
+
+	static void SetupFunctionCallTraversalContext(const FTraversalCallingContext& CallingContext, const FTraversalLocalContext& LocalContext,
+		const FTraversalNode& Node, FTraversalCallingContext& OutFunctionCallContext);
+	
+	static void UpdateLocalStaticContextFromFunctionStaticWrites(const FStaticContext& FunctionCallLocalStaticContext, const FFunctionCallData& FunctionCallData, FStaticContext& LocalStaticContext);
+
+	static void UpdateLocalContextFromParameterData(const FTraversalCallingContext& CallingContext, FTraversalLocalContext& LocalContext, const FTraversalNode& Node, FTraversalDebugData* DebugData);
+
+	static void ResolveSelectConnectionData(const FTraversalCallingContext& CallingContext, FTraversalLocalContext& LocalContext, const FTraversalNode& Node, FTraversal::FSelectConnectionData& OutSelectConnectionData, FTraversalDebugData* DebugData);
+
+	static TOptional<FSelectValue> ResolveSelectValueForSelectData(const FTraversalCallingContext& CallingContext, const FTraversalLocalContext& LocalContext,
+		const FTraversalNode& Node, TSet<const FTraversalNode*>& ResolveTraversedNodes, FTraversalDebugData* DebugData);
+
+	static TOptional<FSelectValue> ResolveSelectValueForNode(const FTraversalCallingContext& CallingContext, const FTraversalLocalContext& LocalContext,
+		const FTraversalNode& Node,	TSet<const FTraversalNode*>& ResolveTraversedNodes, FTraversalDebugData* DebugData);
+
+	static TOptional<FSelectValue> ResolveSelectValueForFunctionInputNode(const FTraversalCallingContext& CallingContext, const FTraversalLocalContext& LocalContext,
+		const FTraversalNode& Node, TSet<const FTraversalNode*>& ResolveTraversedNodes, FTraversalDebugData* DebugData);
+
+	static TOptional<FSelectValue> ResolveSelectValueForSelectNode(const FTraversalCallingContext& CallingContext, const FTraversalLocalContext& LocalContext,
+		const FTraversalNode& Node, TSet<const FTraversalNode*>& ResolveTraversedNodes, FTraversalDebugData* DebugData);
+
+	static TOptional<FSelectValue> ResolveSelectValueForParameterNode(const FTraversalCallingContext& CallingContext, const FTraversalLocalContext& LocalContext,
+		const FTraversalNode& Node, TSet<const FTraversalNode*>& ResolveTraversedNodes, FTraversalDebugData* DebugData);
+
+	static TOptional<FSelectValue> ResolveSelectValueForOpNode(const FTraversalCallingContext& CallingContext, const FTraversalLocalContext& LocalContext,
+		const FTraversalNode& Node, TSet<const FTraversalNode*>& ResolveTraversedNodes, FTraversalDebugData* DebugData);
+
+	static TOptional<FSelectValue> ResolveSelectValueForFunctionCallNode(const FTraversalCallingContext& CallingContext, const FTraversalLocalContext& LocalContext,
+		const FTraversalNode& Node, TSet<const FTraversalNode*>& ResolveTraversedNodes, FTraversalDebugData* DebugData);
+
+	TOptional<FSelectValue> ResolveSelectValueForFunctionTraversal(const FTraversalCallingContext& FunctionCallingContext, FTraversalDebugData* DebugData) const;
+
+	static TOptional<FSelectValue> ResolveStaticReadValueForNode(const FTraversalCallingContext& CallingContext, const FTraversalLocalContext& LocalContext,
+		const FTraversalNode& Node,	FSelectKey& ReadSelectKey, TSet<const FTraversalNode*>& ResolveTraversedNodes, FTraversalDebugData* DebugData);
+
+	static TOptional<FSelectValue> ResolveStaticReadValueForSelectNode(const FTraversalCallingContext& CallingContext, const FTraversalLocalContext& LocalContext,
+		const FTraversalNode& Node, FSelectKey& ReadSelectKey, TSet<const FTraversalNode*>& ResolveTraversedNodes, FTraversalDebugData* DebugData);
+
+	static TOptional<FSelectValue> ResolveStaticReadValueForParameterNode(const FTraversalCallingContext& CallingContext, const FTraversalLocalContext& LocalContext,
+		const FTraversalNode& Node, FSelectKey& ReadSelectKey, TSet<const FTraversalNode*>& ResolveTraversedNodes, FTraversalDebugData* DebugData);
+
+	static TOptional<FSelectValue> ResolveStaticReadValueForFunctionCallNode(const FTraversalCallingContext& CallingContext, const FTraversalLocalContext& LocalContext,
+		const FTraversalNode& Node, FSelectKey& ReadSelectKey, TSet<const FTraversalNode*>& ResolveTraversedNodes, FTraversalDebugData* DebugData);
+
+	static bool CanWriteStaticsInternal(const FTraversal& Traversal, TSet<const FTraversal*>& CheckedTraversals, FTraversalDebugData* DebugData);
+
+	static bool CanWriteStaticAttributesInternal(const FTraversal& Traversal, TSet<const FTraversal*>& CheckedTraversals, FTraversalDebugData* DebugData);
+
+	static void GetAllConnectedNodes(const FTraversalNode& Node, TSet<const FTraversalNode*>& OutNodes);
+
+	static ESelectKeySource ParameterFlagsToSelectKeySource(EParameterFlags ParameterFlags);
+
+private:
+	TSharedRef<FTraversalNode> TraversalRoot;
+	TArray<TSharedRef<FTraversalNode>> UnconnectedTraversalRoots;
+	TSet<FScriptReference> ExternalReferences;
+	TSet<FNiagaraVariableBase> StaticVariableReads;
+	TSet<FNiagaraVariableBase> StaticVariableWrites;
+	TSet<FNiagaraVariableBase> StaticVariableWritesToAttributes;
+	TMap<FName, FGuid> FunctionNameToNodeGuidMap;
+
+	struct FCachedTraversalData
+	{
+		TMap<FGuid, TSharedRef<const FTraversalData>> VisitorIdToTraversalData;
+		FStaticContext LocalStaticContext;
+	};
+
+	mutable TMap<FTraversalCallingContext::FContextHash, FCachedTraversalData> HashToTraversalDataCache;
+	mutable TOptional<bool> bCanWriteStaticsCache;
+	mutable TOptional<bool> bCanWriteStaticAttributesCache;
+};
+
+} // UE::Niagara::TraversalCache
\ No newline at end of file
diff --git a/Engine/Plugins/FX/Niagara/Source/NiagaraEditor/Private/TraversalCache/TraversalBuilder.cpp b/Engine/Plugins/FX/Niagara/Source/NiagaraEditor/Private/TraversalCache/TraversalBuilder.cpp
new file mode 100644
--- /dev/null
+++ b/Engine/Plugins/FX/Niagara/Source/NiagaraEditor/Private/TraversalCache/TraversalBuilder.cpp
@@ -0,0 +1,1213 @@
+// Copyright Epic Games, Inc. All Rights Reserved.
+
+#include "TraversalCache/TraversalBuilder.h"
+
+#include "EdGraph/EdGraphPin.h"
+#include "EdGraphSchema_Niagara.h"
+#include "INiagaraEditorTypeUtilities.h"
+#include "NiagaraConstants.h"
+#include "NiagaraEditorModule.h"
+#include "NiagaraEditorUtilities.h"
+#include "NiagaraGraph.h"
+#include "NiagaraNode.h"
+#include "NiagaraNodeAssignment.h"
+#include "NiagaraNodeFunctionCall.h"
+#include "NiagaraNodeInput.h"
+#include "NiagaraNodeOp.h"
+#include "NiagaraNodeOutput.h"
+#include "NiagaraNodeParameterMapGet.h"
+#include "NiagaraNodeParameterMapSet.h"
+#include "NiagaraNodeReroute.h"
+#include "NiagaraNodeStaticSwitch.h"
+#include "NiagaraNodeWithDynamicPins.h"
+#include "NiagaraScriptVariable.h"
+#include "NiagaraTypes.h"
+#include "TraversalCache/Traversal.h"
+#include "TraversalCache/TraversalNode.h"
+
+namespace UE::Niagara::TraversalCache
+{
+
+namespace BuilderHelpers
+{
+	void SortAndFilterInputs(TArray<const UEdGraphPin*>& InputPins)
+	{
+		TArray<const UEdGraphPin*> ParameterMapInputPins;
+		TArray<const UEdGraphPin*> OtherNonAddPins;
+		for (const UEdGraphPin* InputPin : InputPins)
+		{
+			if (InputPin->PinType.PinCategory == UEdGraphSchema_Niagara::PinCategoryMisc &&
+				InputPin->PinType.PinSubCategory == UNiagaraNodeWithDynamicPins::AddPinSubCategory)
+			{
+				continue;
+			}
+
+			FNiagaraTypeDefinition PinTypeDefinition = UEdGraphSchema_Niagara::PinToTypeDefinition(InputPin);
+			if (PinTypeDefinition.IsValid())
+			{
+				if (PinTypeDefinition == FNiagaraTypeDefinition::GetParameterMapDef())
+				{
+					ParameterMapInputPins.Add(InputPin);
+				}
+				else
+				{
+					OtherNonAddPins.Add(InputPin);
+				}
+			}
+		}
+
+		InputPins.Empty(ParameterMapInputPins.Num() + OtherNonAddPins.Num());
+		InputPins.Append(ParameterMapInputPins);
+		InputPins.Append(OtherNonAddPins);
+	}
+
+	void ResolveNamespaceModifier(FParameterReference& ParameterReference)
+	{
+		TArray<FString> ParameterNameParts;
+		ParameterReference.Parameter.GetName().ToString().ParseIntoArray(ParameterNameParts, TEXT("."));
+
+		FName NamespaceModifier = *ParameterNameParts[1];
+		ParameterNameParts[1] = "__MODIFIER__";
+
+		FName ResolvedParameterName = *FString::Join(ParameterNameParts, TEXT("."));
+		ParameterReference.Parameter = FNiagaraVariableBase(ParameterReference.Parameter.GetType(), ResolvedParameterName);
+		ParameterReference.OptionalNamespaceModifier = NamespaceModifier;
+	}
+
+	// The static traversal needs a slightly different behavior here where a value is required when the pin
+	// is not connected and the default value isn't ignored.  The default schema behavior is that an empty default
+	// string will result in an unallocated variable, but we need the type default, and if needs value is true
+	// all the time, then connected pins and pins that don't support defaults will always allocate data which is
+	// also undesired.
+	FNiagaraVariable TraversalPinToNiagaraVariable(const UEdGraphPin& InPin)
+	{
+		bool bNeedsValue = InPin.Direction == EGPD_Input && InPin.bDefaultValueIsIgnored == false && InPin.LinkedTo.Num() == 0;
+		FNiagaraVariable Variable = UEdGraphSchema_Niagara::PinToNiagaraVariable(&InPin, bNeedsValue);
+		return Variable;
+	}
+
+	class FInputHandler : public FTraversalBuilder::FGraphNodeHandler
+	{
+		virtual bool CanProvideFunctionInput() const override { return true; }
+
+		virtual void GetFunctionInputData(
+			const FTraversalBuilder::FGraphNodeAndOutputPin& GraphNodeAndOutputPin,
+			TOptional<const FFunctionInputData>& OutInputData,
+			FTraversalBuilderDebugData* DebugData) const override
+		{
+			if (GraphNodeAndOutputPin.OutputPin == nullptr)
+			{
+				return;
+			}
+
+			const UNiagaraNodeInput* InputNode = CastChecked<UNiagaraNodeInput>(&GraphNodeAndOutputPin.Node);
+			if (InputNode->Input.GetType().IsStatic() == false)
+			{
+				return;
+			}
+
+			FFunctionInputData InputData;
+			InputData.InputSelectKey = FSelectKey(ESelectKeySource::FunctionCallNode, InputNode->Input, NAME_None);
+			if (InputNode->Input.IsDataAllocated())
+			{
+				InputData.LocalValue = FTraversalBuilder::CreateSelectValue(InputNode->Input);
+			}
+
+			OutInputData = InputData;
+		}
+	};
+
+	class FMapGetHandler : public FTraversalBuilder::FGraphNodeHandler
+	{
+	public:
+		virtual bool CanAccessParameters() const override { return true; }
+
+		virtual void GetParameterData(
+			const FTraversalBuilder::FGraphNodeAndOutputPin& GraphNodeAndOutputPin,
+			TOptional<const FParameterData>& OutParameterData,
+			TOptional<TArray<FGuid>>& OutFilteredConnectedPinIds,
+			FTraversalBuilderDebugData* DebugData) const override
+		{
+			if (GraphNodeAndOutputPin.OutputPin == nullptr)
+			{
+				return;
+			}
+
+			FParameterData ParameterData;
+			OutFilteredConnectedPinIds = TArray<FGuid>();
+
+			const UNiagaraNodeParameterMapGet* GetNode = CastChecked<UNiagaraNodeParameterMapGet>(&GraphNodeAndOutputPin.Node);
+
+			const UEdGraphPin& ReadPin = *GraphNodeAndOutputPin.OutputPin;
+			FNiagaraVariable ReadVariable = TraversalPinToNiagaraVariable(ReadPin);
+			EParameterFlags ReadFlags = FTraversalBuilder::ExtractFlagsFromParameterName(ReadVariable.GetName());
+			FParameterRead& ReadParameterReference = ParameterData.ReadParameterReferences.Add_GetRef(
+				FParameterRead(ReadVariable, ReadFlags));
+
+			if (HasFlag(ReadFlags, EParameterFlags::NamespaceModifier))
+			{
+				ResolveNamespaceModifier(ReadParameterReference);
+			}
+
+			GetReadDefaultValues(*GetNode, ReadPin, ReadVariable, ReadParameterReference, OutFilteredConnectedPinIds, DebugData);
+
+			// Handle reads that are "discovered" but not actually read in this traversal path.
+			TArray<const UEdGraphPin*> OutputPins;
+			GetNode->GetOutputPins(OutputPins);
+			for (const UEdGraphPin* OutputPin : OutputPins)
+			{
+				if (GetNode->IsAddPin(OutputPin) || OutputPin == &ReadPin)
+				{
+					continue;
+				}
+
+				FNiagaraVariableBase DiscoveredVariable = TraversalPinToNiagaraVariable(*OutputPin);
+				if (DiscoveredVariable.IsValid())
+				{
+					EParameterFlags DiscoveredFlags = FTraversalBuilder::ExtractFlagsFromParameterName(DiscoveredVariable.GetName());
+					FParameterRead& DiscoveredParameterReference = ParameterData.ReadParameterReferences.Add_GetRef(
+						FParameterRead(DiscoveredVariable, DiscoveredFlags));
+					DiscoveredParameterReference.bIsDiscoverRead = true;
+				}
+			}
+
+			const UEdGraphPin* ParameterMapInputPin = FNiagaraStackGraphUtilities::GetParameterMapInputPin(*GetNode);
+			if (ParameterMapInputPin != nullptr)
+			{
+				ParameterData.ExecutionConnectionPinId = ParameterMapInputPin->PinId;
+				OutFilteredConnectedPinIds->Add(ParameterMapInputPin->PinId);
+			}
+
+			OutParameterData = ParameterData;
+		}
+
+		static void GetReadDefaultValues(
+			const UNiagaraNodeParameterMapGet& GetNode,
+			const UEdGraphPin& InOutputPin,
+			const FNiagaraVariable& ReadVariable,
+			FParameterRead& ReadParameterReference,
+			TOptional<TArray<FGuid>>& OutFilteredConnectedPinIds,
+			FTraversalBuilderDebugData* DebugData)
+		{
+			if (ReadVariable.IsDataInterface())
+			{
+				return;
+			}
+
+			const UNiagaraScriptVariable* ReadScriptVariable = GetNode.GetNiagaraGraph()->GetScriptVariable(ReadVariable);
+			if (ReadScriptVariable == nullptr)
+			{
+				return;
+			}
+
+			if (ReadScriptVariable->DefaultMode == ENiagaraDefaultMode::Custom)
+			{
+				TOptional<FNiagaraVariant> PinDefaultValue;
+				const UEdGraphPin* DefaultPin = GetNode.GetDefaultPin(&InOutputPin);
+				if (DefaultPin != nullptr)
+				{
+					if (DefaultPin->LinkedTo.Num() == 0)
+					{
+						FNiagaraVariable ReadDefaultVariable = TraversalPinToNiagaraVariable(*DefaultPin);
+						if (ReadDefaultVariable.IsDataAllocated())
+						{
+							TArray<uint8> ReadDefaultVariableValue;
+							ReadDefaultVariableValue.Append(ReadDefaultVariable.GetData(), ReadDefaultVariable.GetSizeInBytes());
+							ReadParameterReference.OptionalDefaultValue = FNiagaraVariant(ReadDefaultVariableValue);
+						}
+					}
+					else if (DefaultPin->LinkedTo.Num() == 1)
+					{
+						ReadParameterReference.OptionalConnectionPinId = DefaultPin->PinId;
+						OutFilteredConnectedPinIds->Add(DefaultPin->PinId);
+					}
+				}
+			}
+			else if (ReadScriptVariable->DefaultMode == ENiagaraDefaultMode::Value)
+			{
+				if (ReadScriptVariable->Variable.IsDataAllocated())
+				{
+					ReadParameterReference.OptionalDefaultValue = FNiagaraVariant(ReadScriptVariable->Variable.GetData(), ReadScriptVariable->Variable.GetSizeInBytes());
+				}
+				else if (ReadScriptVariable->GetDefaultValueVariant().IsValid())
+				{
+					ReadParameterReference.OptionalDefaultValue = ReadScriptVariable->GetDefaultValueVariant();
+				}
+				else if (DebugData != nullptr)
+				{
+					FNiagaraVariable ReadVariableCopy = ReadScriptVariable->Variable;
+					FNiagaraEditorUtilities::ResetVariableToDefaultValue(ReadVariableCopy);
+					FNiagaraVariant ReadVariableVariant(ReadVariableCopy.GetData(), ReadVariableCopy.GetSizeInBytes());
+					ReadParameterReference.OptionalDefaultValue = ReadVariableVariant;
+				}
+			}
+			else if (ReadScriptVariable->DefaultMode == ENiagaraDefaultMode::Binding)
+			{
+				ReadParameterReference.OptionalDefaultBinding = ReadScriptVariable->DefaultBinding.GetName();
+			}
+
+			if (ReadParameterReference.OptionalDefaultValue.IsSet() && FTraversalBuilder::IsValidSelectValueType(ReadVariable.GetType()))
+			{
+				FNiagaraVariable Temp = ReadVariable;
+				Temp.SetData(ReadParameterReference.OptionalDefaultValue->GetBytes());
+				ReadParameterReference.OptionalLocalSelectValue = FTraversalBuilder::CreateSelectValue(Temp);
+			}
+		}
+	};
+
+	class FMapSetHandler : public FTraversalBuilder::FGraphNodeHandler
+	{
+		virtual bool CanAccessParameters() const override { return true; }
+
+		virtual void GetParameterData(
+			const FTraversalBuilder::FGraphNodeAndOutputPin& GraphNodeAndOutputPin,
+			TOptional<const FParameterData>& OutParameterData,
+			TOptional<TArray<FGuid>>& OutFilteredConnectedPinIds,
+			FTraversalBuilderDebugData* DebugData) const override
+		{
+			FParameterData ParameterData;
+
+			const UNiagaraNodeParameterMapSet* SetNode = CastChecked<UNiagaraNodeParameterMapSet>(&GraphNodeAndOutputPin.Node);
+			TArray<const UEdGraphPin*> InputPins;
+			SetNode->GetInputPins(InputPins);
+			for (const UEdGraphPin* InputPin : InputPins)
+			{
+				if (SetNode->IsAddPin(InputPin))
+				{
+					continue;
+				}
+
+				FNiagaraVariable InputVariable = TraversalPinToNiagaraVariable(*InputPin);
+				if (InputVariable.GetType() == FNiagaraTypeDefinition::GetParameterMapDef())
+				{
+					continue;
+				}
+
+				EParameterFlags WriteFlags = FTraversalBuilder::ExtractFlagsFromParameterName(InputVariable.GetName());
+				FParameterWrite& WriteParameterReference = ParameterData.WriteParameterReferences.Add_GetRef(
+					FParameterWrite(InputVariable, WriteFlags));
+
+				if (HasFlag(WriteFlags, EParameterFlags::NamespaceModifier))
+				{
+					ResolveNamespaceModifier(WriteParameterReference);
+				}
+
+				TOptional<FSelectValue> WriteSelectValue;
+				if (InputVariable.GetType().IsStatic() && InputPin->LinkedTo.Num() == 0 && InputVariable.IsDataAllocated())
+				{
+					WriteSelectValue = FTraversalBuilder::CreateSelectValue(InputVariable);
+				}
+
+
+				WriteParameterReference.OptionalConnectionPinId = InputPin->PinId;
+				WriteParameterReference.OptionalLocalSelectValue = WriteSelectValue;
+			}
+
+			const UEdGraphPin* ParameterMapInputPin = FNiagaraStackGraphUtilities::GetParameterMapInputPin(*SetNode);
+			ParameterData.ExecutionConnectionPinId = ParameterMapInputPin != nullptr ? ParameterMapInputPin->PinId : FGuid();
+			OutParameterData = ParameterData;
+		}
+	};
+
+	class FFunctionCallHandler : public FTraversalBuilder::FGraphNodeHandler
+	{
+	public:
+		virtual bool CanCallFunctionScript() const override { return true; }
+		virtual void GetFunctionCallData(
+			const FTraversalBuilder::FGraphNodeAndOutputPin& GraphNodeAndOutputPin,
+			TOptional<const FFunctionCallData>& OutFunctionCallData,
+			TOptional<TArray<FGuid>>& OutFilteredConnectedPinIds,
+			FTraversalBuilderDebugData* DebugData) const override
+		{ 
+			const UNiagaraNodeFunctionCall* FunctionCallNode = CastChecked<UNiagaraNodeFunctionCall>(&GraphNodeAndOutputPin.Node);
+
+			FFunctionCallData FunctionCallData;
+			FunctionCallData.FunctionScriptReference = FScriptReference(FSoftObjectPath(FunctionCallNode->FunctionScript), FunctionCallNode->SelectedScriptVersion);
+			FunctionCallData.FunctionCallName = *FunctionCallNode->GetFunctionName();
+			FunctionCallData.DebugState = FunctionCallNode->DebugState;
+
+			UNiagaraGraph* CalledGraph = FunctionCallNode->GetCalledGraph();
+			if (CalledGraph != nullptr)
+			{
+				TArray<UEdGraphPin*> InputPins;
+				FunctionCallNode->GetInputPins(InputPins);
+				for (UEdGraphPin* InputPin : InputPins)
+				{
+					FNiagaraVariable InputVariable = TraversalPinToNiagaraVariable(*InputPin);
+					UNiagaraScriptVariable* ScriptVariable = CalledGraph->GetScriptVariable(InputPin->PersistentGuid);
+					if (ScriptVariable == nullptr)
+					{
+						ScriptVariable = CalledGraph->GetScriptVariable(InputVariable);
+					}
+					if (InputVariable.GetType().IsStatic() == false && (ScriptVariable == nullptr || ScriptVariable->GetIsStaticSwitch() == false))
+					{
+						continue;
+					}
+
+					FFunctionInputSelectValue& InputSelectValue = FunctionCallData.InputSelectValues.AddDefaulted_GetRef();
+					InputSelectValue.InputSelectKey.Source = ESelectKeySource::FunctionCallNode;
+					InputSelectValue.InputSelectKey.Variable = InputVariable;
+
+					if (InputPin->LinkedTo.Num() == 1)
+					{
+						InputSelectValue.ConnectionPinId = InputPin->PinId;
+					}
+					else
+					{
+						TOptional<FSelectValue> InputVariableValue;
+						if (InputVariable.IsDataAllocated() && FTraversalBuilder::IsValidSelectValueType(InputVariable.GetType()))
+						{
+							InputVariableValue = FTraversalBuilder::CreateSelectValue(InputVariable);
+						}
+
+						if (InputVariableValue.IsSet())
+						{
+							InputSelectValue.LocalValue = InputVariableValue;
+						}
+						else if (ScriptVariable != nullptr && ScriptVariable->GetIsStaticSwitch())
+						{
+							// Check for propagated switches.
+							const FNiagaraPropagatedVariable* PropagatedSwitch = FunctionCallNode->PropagatedStaticSwitchParameters.FindByPredicate(
+								[InputVariable](const FNiagaraPropagatedVariable& PropagatedSwitch) { return PropagatedSwitch.SwitchParameter == InputVariable; });
+							if (PropagatedSwitch != nullptr && PropagatedSwitch->PropagatedName.IsEmpty() == false)
+							{
+								InputSelectValue.OptionalPropagatedNameOverride = *PropagatedSwitch->PropagatedName;
+							}
+						}
+					}
+				}
+			}
+
+			const UEdGraphPin* ParameterMapInputPin = FNiagaraStackGraphUtilities::GetParameterMapInputPin(*FunctionCallNode);
+			FunctionCallData.ExecutionConnectionPinId = ParameterMapInputPin != nullptr ? ParameterMapInputPin->PinId : FGuid();
+			OutFunctionCallData = FunctionCallData;
+		}
+	};
+
+	class FStaticSwitchHandler : public FTraversalBuilder::FGraphNodeHandler
+	{
+	public:
+		virtual bool CanSelectInputPin() const override { return true; }
+
+		virtual void GetSelectData(
+			const FTraversalBuilder::FGraphNodeAndOutputPin& GraphNodeAndOutputPin,
+			TOptional<const FSelectData>& OutSelectData,
+			TOptional<TArray<FGuid>>& OutFilteredConnectedPinIds,
+			FTraversalBuilderDebugData* DebugData) const override
+		{
+			FSelectData SelectData;
+			OutFilteredConnectedPinIds = TArray<FGuid>();
+
+			const UNiagaraNodeStaticSwitch* StaticSwitchNode = CastChecked<UNiagaraNodeStaticSwitch>(&GraphNodeAndOutputPin.Node);
+			SelectData.SelectMode = StaticSwitchNode->IsSetByPin() ? ESelectMode::Connection : ESelectMode::Value;
+
+			if (SelectData.SelectMode == ESelectMode::Value)
+			{
+				SelectData.SelectKey = GetSelectKey(*StaticSwitchNode);
+			}
+			else if (SelectData.SelectMode == ESelectMode::Connection && StaticSwitchNode->IsSetByPin())
+			{
+				SelectData.SelectConnectionPinId = StaticSwitchNode->GetSelectorPin()->PinId;
+				OutFilteredConnectedPinIds->Add(SelectData.SelectConnectionPinId);
+			}
+
+			if (GraphNodeAndOutputPin.OutputPin != nullptr)
+			{
+				GetSelectInputData(*StaticSwitchNode, *GraphNodeAndOutputPin.OutputPin, SelectData.InputData, DebugData);
+				TArray<FGuid> SelectConnectedPinIds;
+				for (const FSelectInputData& SelectInputData : SelectData.InputData)
+				{
+					if (SelectInputData.ConnectionPinId.IsSet())
+					{
+						OutFilteredConnectedPinIds->Add(SelectInputData.ConnectionPinId.GetValue());
+					}
+				}
+			}
+			OutSelectData = SelectData;
+		}
+
+	private:
+		FSelectKey GetSelectKey(const UNiagaraNodeStaticSwitch& StaticSwitchNode) const
+		{
+			if (StaticSwitchNode.SwitchTypeData.SwitchConstant != NAME_None)
+			{
+				const FNiagaraVariable* StaticSwitchVariable = FNiagaraConstants::FindStaticSwitchConstant(StaticSwitchNode.SwitchTypeData.SwitchConstant);
+				return StaticSwitchVariable != nullptr
+					? FSelectKey(ESelectKeySource::ExternalConstant, *StaticSwitchVariable, NAME_None)
+					: FSelectKey(
+						ESelectKeySource::ExternalConstant,
+						FNiagaraVariableBase(StaticSwitchNode.GetInputType(), StaticSwitchNode.SwitchTypeData.SwitchConstant),
+						NAME_None);
+			}
+			else
+			{
+				FNiagaraVariableBase StaticSwitchVariable(StaticSwitchNode.GetInputType(), StaticSwitchNode.InputParameterName);
+				return FSelectKey(ESelectKeySource::FunctionCallNode, StaticSwitchVariable, NAME_None);
+			}
+		}
+
+		void GetSelectInputData(const UNiagaraNodeStaticSwitch& StaticSwitchNode, const UEdGraphPin& InOutputPin, TArray<FSelectInputData>& OutInputData, FTraversalBuilderDebugData* DebugData) const
+		{
+			FNiagaraVariable OutputVariable = TraversalPinToNiagaraVariable(InOutputPin);
+			int32 OutputIndex = StaticSwitchNode.OutputVars.IndexOfByPredicate([OutputVariable](const FNiagaraVariable& OutputVar)
+			{
+				// Numeric pins can have their type changed, so it's still a match if the names match but the types don't, and
+				// the OutputVar is numeric.
+				return OutputVar.GetName() == OutputVariable.GetName() &&
+					(OutputVar.GetType() == OutputVariable.GetType() || OutputVar.GetType() == FNiagaraTypeDefinition::GetGenericNumericDef());
+			});
+
+			if (OutputIndex == INDEX_NONE)
+			{
+				// If we couldn't find a variable that matched by both name and type, try to match by name, but only if there is only 1 match.
+				TArray<int32> OutputVarIndicesMatchedByName;
+				for (int32 OutputVarIndex = 0; OutputVarIndex < StaticSwitchNode.OutputVars.Num(); OutputVarIndex++)
+				{
+					if (StaticSwitchNode.OutputVars[OutputVarIndex].GetName() == OutputVariable.GetName())
+					{
+						OutputVarIndicesMatchedByName.Add(OutputVarIndex);
+					}
+				}
+				if (OutputVarIndicesMatchedByName.Num() == 1)
+				{
+					OutputIndex = OutputVarIndicesMatchedByName[0];
+				}
+			}
+
+			if (OutputIndex == INDEX_NONE)
+			{
+				if (DebugData != nullptr)
+				{
+					DebugData->AddUnresolvedSelectOutput(StaticSwitchNode.NodeGuid, InOutputPin.PinId);
+				}
+				return;
+			}
+
+			int32 OutputNum = StaticSwitchNode.OutputVars.Num();
+			TArray<int32> Options = StaticSwitchNode.GetOptionValues();
+			TArray<UEdGraphPin*> InputPins;
+			StaticSwitchNode.GetInputPins(InputPins);
+			int32 SetPins = StaticSwitchNode.IsSetByPin() ? 1 : 0;
+			if (InputPins.Num() == (StaticSwitchNode.OutputVars.Num() * Options.Num() + SetPins))
+			{
+				for (int32 OptionIndex = 0; OptionIndex < Options.Num(); OptionIndex++)
+				{
+					FSelectInputData& InputData = OutInputData.AddDefaulted_GetRef();
+					InputData.SelectValue = FTraversalBuilder::CreateSelectValue(StaticSwitchNode.GetInputType(), Options[OptionIndex]);
+
+					int32 InputPinIndex = (OptionIndex * OutputNum) + OutputIndex;
+					const UEdGraphPin* InputPin = InputPins[InputPinIndex];
+					if (InputPin->LinkedTo.Num() == 0)
+					{
+						FNiagaraVariable PinVariable = TraversalPinToNiagaraVariable(*InputPin);
+						if (PinVariable.GetType().IsStatic() && PinVariable.IsDataAllocated())
+						{
+							InputData.LocalValue = FTraversalBuilder::CreateSelectValue(PinVariable);
+						}
+					}
+					else if (InputPin->LinkedTo.Num() == 1)
+					{
+						InputData.ConnectionPinId = InputPin->PinId;
+					}
+				}
+			}
+		}
+	};
+
+	class FOpNodeHandler : public FTraversalBuilder::FGraphNodeHandler
+	{
+	public:
+		virtual bool CanEvaluateStaticValues() const override { return true; }
+		virtual void GetStaticOpData(
+			const FTraversalBuilder::FGraphNodeAndOutputPin& GraphNodeAndOutputPin,
+			TOptional<const FStaticOpData>& OutStaticOpData,
+			TOptional<TArray<FGuid>>& OutFilteredConnectedPinIds,
+			FTraversalBuilderDebugData* DebugData) const override
+		{
+			const UNiagaraNodeOp* OpNode = CastChecked<UNiagaraNodeOp>(&GraphNodeAndOutputPin.Node);
+			const FNiagaraOpInfo* OpInfo = FNiagaraOpInfo::GetOpInfo(OpNode->OpName);
+			if (OpInfo == nullptr || OpInfo->bSupportsStaticResolution == false || OpInfo->StaticVariableResolveFunction.IsBound() == false)
+			{
+				return;
+			}
+
+			FStaticOpData StaticOpData;
+			StaticOpData.OpName = OpNode->OpName;
+
+			TArray<const UEdGraphPin*> InputPins;
+			OpNode->GetInputPins(InputPins);
+			for (const UEdGraphPin* InputPin : InputPins)
+			{
+				if (OpNode->IsAddPin(InputPin))
+				{
+					continue;
+				}
+
+				FNiagaraVariable InputVariable = TraversalPinToNiagaraVariable(*InputPin);
+				if (InputVariable.IsValid() == false)
+				{
+					continue;
+				}
+
+				FStaticOpInputData& OpInput = StaticOpData.InputData.AddDefaulted_GetRef();
+				if (InputVariable.GetType().IsStatic())
+				{
+					if (InputPin->LinkedTo.Num() == 0)
+					{
+						if (InputVariable.IsDataAllocated())
+						{
+							OpInput.LocalValue = FTraversalBuilder::CreateSelectValue(InputVariable);
+						}
+					}
+					else if (InputPin->LinkedTo.Num() == 1)
+					{
+						OpInput.ConnectionPinId = InputPin->PinId;
+					}
+				}
+			}
+
+			OutStaticOpData = StaticOpData;
+		}
+	};
+
+	class FRerouteNodeHandler : public FTraversalBuilder::FGraphNodeHandler
+	{
+	public:
+		virtual bool IsNoop() const override { return true; }
+	};
+} // BuilderHelpers
+
+FTraversalBuilder* FTraversalBuilder::Instance = nullptr;
+
+FCriticalSection FTraversalBuilder::IsValidSelectValueTypeCacheCS;
+FCriticalSection FTraversalBuilder::SelectValueCacheCS;
+FCriticalSection FTraversalBuilder::ExtractedParameterFlagCacheCS;
+
+void FTraversalBuilder::Initialize()
+{
+	if (bInitialized == false)
+	{
+		checkf(Instance == nullptr, TEXT("FTraversalBuilder initialized more than once."));
+		Instance = this;
+		GraphNodeHandlers.Add(UNiagaraNodeInput::StaticClass(), MakeShared<BuilderHelpers::FInputHandler>());
+		GraphNodeHandlers.Add(UNiagaraNodeParameterMapGet::StaticClass(), MakeShared<BuilderHelpers::FMapGetHandler>());
+		GraphNodeHandlers.Add(UNiagaraNodeParameterMapSet::StaticClass(), MakeShared<BuilderHelpers::FMapSetHandler>());
+		GraphNodeHandlers.Add(UNiagaraNodeFunctionCall::StaticClass(), MakeShared<BuilderHelpers::FFunctionCallHandler>());
+		GraphNodeHandlers.Add(UNiagaraNodeStaticSwitch::StaticClass(), MakeShared<BuilderHelpers::FStaticSwitchHandler>());
+		GraphNodeHandlers.Add(UNiagaraNodeAssignment::StaticClass(), MakeShared<BuilderHelpers::FFunctionCallHandler>());
+		GraphNodeHandlers.Add(UNiagaraNodeOp::StaticClass(), MakeShared<BuilderHelpers::FOpNodeHandler>());
+		GraphNodeHandlers.Add(UNiagaraNodeReroute::StaticClass(), MakeShared<BuilderHelpers::FRerouteNodeHandler>());
+		bInitialized = true;
+	}
+}
+
+TSharedRef<const FTraversal> FTraversalBuilder::BuildTraversal(
+	const UNiagaraNodeOutput& OutputNode, FTraversalBuilderDebugData* DebugData)
+{
+	return GetInstance().BuildTraversalInternal(OutputNode, DebugData);
+}
+
+void FTraversalBuilder::ResolveFunctionCallStackNames(const UEdGraph& TopLevelGraph, const TArray<FGuid>& FunctionCallStack, TArray<FString>& OutFunctionCallNames)
+{
+	const UEdGraph* CurrentGraph = &TopLevelGraph;
+	for (int32 CurrentStackIndex = 0; CurrentStackIndex < FunctionCallStack.Num(); ++CurrentStackIndex)
+	{
+		if (CurrentGraph == nullptr)
+		{
+			break;
+		}
+
+		const FGuid& FunctionCallNodeGuid = FunctionCallStack[CurrentStackIndex];
+		const TObjectPtr<UEdGraphNode>* FunctionCallNodePtr = CurrentGraph->Nodes.FindByPredicate([&FunctionCallNodeGuid]
+			(const UEdGraphNode* GraphNode) { return GraphNode->NodeGuid == FunctionCallNodeGuid; });
+		if (FunctionCallNodePtr == nullptr)
+		{
+			break;
+		}
+
+		const UNiagaraNodeFunctionCall* FunctionCallNode = Cast<UNiagaraNodeFunctionCall>(FunctionCallNodePtr->Get());
+		if (FunctionCallNode == nullptr)
+		{
+			break;
+		}
+
+		OutFunctionCallNames.Add(FunctionCallNode->GetFunctionName());
+		CurrentGraph = FunctionCallNode->GetCalledGraph();
+	}
+};
+
+bool FTraversalBuilder::IsValidSelectValueType(const FNiagaraTypeDefinition& InValueType)
+{
+	return GetInstance().IsValidSelectValueTypeInternal(InValueType);
+}
+
+FSelectValue FTraversalBuilder::CreateSelectValue(const FNiagaraTypeDefinition& InValueType, int32 InSelectNumericValue)
+{
+	return GetInstance().CreateSelectValueInternal(InValueType, InSelectNumericValue);
+}
+
+FSelectValue FTraversalBuilder::CreateSelectValue(const FNiagaraVariable& InNiagaraVariable)
+{
+	return GetInstance().CreateSelectValueInternal(InNiagaraVariable);
+}
+
+FSelectValue FTraversalBuilder::CreateSelectValue(bool bInBoolValue)
+{
+	FNiagaraVariable Temp(FNiagaraTypeDefinition::GetBoolDef(), NAME_None);
+	Temp.SetValue(FNiagaraBool(bInBoolValue));
+	return CreateSelectValue(Temp);
+}
+
+FSelectValue FTraversalBuilder::CreateSelectValue(UEnum* Enum, int32 EnumValue)
+{
+	FNiagaraVariable Temp(FNiagaraTypeDefinition(Enum), NAME_None);
+	FNiagaraInt32 IntValue;
+	IntValue.Value = EnumValue;
+	Temp.SetValue(IntValue);
+	return CreateSelectValue(Temp);
+}
+
+EParameterFlags FTraversalBuilder::ExtractFlagsFromParameterName(FName ParameterName)
+{
+	return GetInstance().ExtractFlagsFromParameterNameInternal(ParameterName);
+}
+
+FTraversalBuilder& FTraversalBuilder::GetInstance()
+{
+	checkf(Instance != nullptr, TEXT("FTraversalBuilder was not initialized."));
+	return *Instance;
+}
+
+TSharedRef<const FTraversal> FTraversalBuilder::BuildTraversalInternal(
+	const UNiagaraNodeOutput& OutputNode, FTraversalBuilderDebugData* DebugData) const
+{
+	TMap<FGraphNodeAndOutputPin, TSharedRef<FTraversalNode>> TraversedOutputPins;
+	TSet<const FTraversalNode*> TrimTraversedNodes;
+	TOptional<TSet<TWeakPtr<FTraversalNode>>> TraversedNoopNodesWeak;
+	TMap<FName, FGuid> ModuleNameToFunctionCallNodeGuid;
+	TSet<const FTraversalNode*> ResolveModuleInputsTraversedNodes;
+
+	if (DebugData != nullptr)
+	{
+		TraversedNoopNodesWeak = TSet<TWeakPtr<FTraversalNode>>();
+	}
+
+	auto HandleRootGraphNode = [&](const FGraphNodeAndOutputPin& RootGraphNodeAndOutputPin)
+	{
+		TSharedPtr<FTraversalNode> RootNode = TraverseGraphNodeFromOutputPin(RootGraphNodeAndOutputPin, TraversedOutputPins, DebugData);
+		ResolveModuleInputWrites(*RootNode.Get(), ModuleNameToFunctionCallNodeGuid, ResolveModuleInputsTraversedNodes);
+		return RootNode.ToSharedRef();
+	};
+
+	TSharedRef<FTraversalNode> TraversalRoot = HandleRootGraphNode(FGraphNodeAndOutputPin(OutputNode, nullptr));
+
+	TArray<TSharedRef<FTraversalNode>> UnconnectedTraversalRoots;
+	TArray<FGraphNodeAndOutputPin> UnconnectedRoots;
+	GetUnconnectedRoots(OutputNode, TraversedOutputPins, UnconnectedRoots);
+	for (const FGraphNodeAndOutputPin& UnconnectedRoot : UnconnectedRoots)
+	{
+		UnconnectedTraversalRoots.Add(HandleRootGraphNode(UnconnectedRoot));
+	}
+
+	// Trim noops after all roots are handled.
+	TArray<TSharedRef<FTraversalNode>> RootsToTrim;
+	RootsToTrim.Add(TraversalRoot);
+	RootsToTrim.Append(UnconnectedTraversalRoots);
+	for (TSharedRef<FTraversalNode>& RootToTrim : RootsToTrim)
+	{
+		TrimNoops(RootToTrim.Get(), TrimTraversedNodes, TraversedNoopNodesWeak.GetPtrOrNull());
+	}
+
+	if (DebugData != nullptr && TraversedNoopNodesWeak.IsSet())
+	{
+		// Empty this map since it will have references to the noops.
+		TraversedOutputPins.Empty();
+		for (TWeakPtr<FTraversalNode>& TraversedNoopNodeWeak : TraversedNoopNodesWeak.GetValue())
+		{
+			if (TraversedNoopNodeWeak.IsValid())
+			{ 
+				TSharedPtr<FTraversalNode> TraversedNoopNode = TraversedNoopNodeWeak.Pin();
+				DebugData->AddUntrimmedNoop(TraversedNoopNode->SourceNodeGuid, TraversedNoopNode->SourceNodeTypeName);
+			}
+		}
+	}
+
+	TSharedRef<FTraversal> Traversal = MakeShared<FTraversal>(TraversalRoot, UnconnectedTraversalRoots);
+	CollectAdditionalTraversalData(Traversal.Get());
+
+	return Traversal;
+}
+
+TSharedPtr<FTraversalNode> FTraversalBuilder::TraverseGraphNodeFromOutputPin(
+	const FGraphNodeAndOutputPin& GraphNodeAndOutputPin,
+	TMap<FGraphNodeAndOutputPin, TSharedRef<FTraversalNode>>& TraversedOutputPins,
+	FTraversalBuilderDebugData* DebugData) const
+{
+	if (ensureMsgf(GraphNodeAndOutputPin.OutputPin == nullptr || GraphNodeAndOutputPin.OutputPin->Direction == EGPD_Output, TEXT("OutputPin must be using the output direction.")) == false)
+	{
+		return nullptr;
+	}
+
+	TSharedRef<FTraversalNode>* ExistingTraversalNode = TraversedOutputPins.Find(GraphNodeAndOutputPin);
+	if (ExistingTraversalNode != nullptr)
+	{
+		return *ExistingTraversalNode;
+	}
+
+	const UEdGraphNode& GraphNode = GraphNodeAndOutputPin.Node;
+	TSharedRef<FTraversalNode> NewTraversalNode = MakeShared<FTraversalNode>();
+	NewTraversalNode->SourceNodeGuid = GraphNode.NodeGuid;
+	NewTraversalNode->SourceNodeTypeName = GraphNode.GetClass()->GetFName();
+	NewTraversalNode->bSourceNodeEnabled = GraphNode.GetDesiredEnabledState() == ENodeEnabledState::Enabled;
+
+	TOptional<TArray<FGuid>> FilteredConnectedPinIds;
+	TraversedOutputPins.Add(GraphNodeAndOutputPin, NewTraversalNode);
+	GatherNodeDataAndFilteredConnectedPinIds(GraphNodeAndOutputPin, NewTraversalNode, FilteredConnectedPinIds, DebugData);
+
+	TArray<const UEdGraphPin*> InputPins;
+	for (const UEdGraphPin* Pin : GraphNode.Pins)
+	{
+		if (Pin->Direction == EGPD_Input)
+		{
+			InputPins.Add(Pin);
+		}
+	}
+	BuilderHelpers::SortAndFilterInputs(InputPins);
+	
+	for (const UEdGraphPin* InputPin : InputPins)
+	{
+		if (InputPin->LinkedTo.Num() > 1 && DebugData != nullptr)
+		{
+			TArray<TPair<FGuid, FGuid>> ConnectedNodeGuidsAndPinIds;
+			for (const UEdGraphPin* LinkedPin : InputPin->LinkedTo)
+			{
+				TPair<FGuid, FGuid>& ConnectedNodeGuidAndPinId = ConnectedNodeGuidsAndPinIds.AddDefaulted_GetRef();
+				if (LinkedPin != nullptr && LinkedPin->GetOwningNode() != nullptr)
+				{
+					ConnectedNodeGuidAndPinId.Key = LinkedPin->GetOwningNode()->NodeGuid;
+					ConnectedNodeGuidAndPinId.Value = LinkedPin->PinId;
+				}
+			}
+
+			DebugData->AddMultipleInputConnection(TPair<FGuid, FGuid>(InputPin->GetOwningNode()->NodeGuid, InputPin->PinId), ConnectedNodeGuidsAndPinIds);
+		}
+
+		if ((FilteredConnectedPinIds.IsSet() && FilteredConnectedPinIds->Contains(InputPin->PinId) == false) ||
+			InputPin->LinkedTo.Num() == 0 ||
+			InputPin->LinkedTo[0] == nullptr ||
+			InputPin->LinkedTo[0]->GetOwningNode() == nullptr)
+		{
+			continue;
+		}
+
+		const UEdGraphPin* LinkedOutputPin = InputPin->LinkedTo[0];
+		TSharedPtr<FTraversalNode> TraversalNodeForPin = TraverseGraphNodeFromOutputPin(
+			FGraphNodeAndOutputPin(*LinkedOutputPin->GetOwningNode(), LinkedOutputPin), TraversedOutputPins, DebugData);
+		if (TraversalNodeForPin.IsValid())
+		{
+			NewTraversalNode->Connections.Add(FTraversalNode::FConnection(InputPin->PinId, TraversalNodeForPin.ToSharedRef()));
+		}
+	}
+
+	return NewTraversalNode;
+}
+
+void FTraversalBuilder::GatherNodeDataAndFilteredConnectedPinIds(const FGraphNodeAndOutputPin& GraphNodeAndOutputPin,
+	TSharedRef<FTraversalNode> NewTraversalNode, TOptional<TArray<FGuid>>& OutFilteredConnectedPinIds, FTraversalBuilderDebugData* DebugData) const
+{
+	const UEdGraphNode& GraphNode = GraphNodeAndOutputPin.Node;
+	const TSharedRef<FTraversalBuilder::IGraphNodeHandler>* NodeHandlerPtr = GraphNodeHandlers.Find(GraphNode.GetClass());
+	if (NodeHandlerPtr != nullptr)
+	{
+		TSharedRef<FTraversalBuilder::IGraphNodeHandler> NodeHandler = *NodeHandlerPtr;
+		if (NodeHandler->IsNoop())
+		{
+			TArray<const UEdGraphPin*> InputPins;
+			CastChecked<UNiagaraNode>(&GraphNode)->GetInputPins(InputPins);
+
+			if (InputPins.Num() == 1)
+			{
+				NewTraversalNode->bIsNoop = true;
+				OutFilteredConnectedPinIds = TArray<FGuid>();
+				OutFilteredConnectedPinIds->Add(InputPins[0]->PinId);
+			}
+		}
+		else
+		{
+			if (NodeHandler->CanProvideFunctionInput())
+			{
+				NodeHandler->GetFunctionInputData(GraphNodeAndOutputPin, NewTraversalNode->FunctionInputData, DebugData);
+			}
+			if (NodeHandler->CanAccessParameters())
+			{
+				NodeHandler->GetParameterData(GraphNodeAndOutputPin, NewTraversalNode->ParameterData, OutFilteredConnectedPinIds, DebugData);
+			}
+			if (NodeHandler->CanCallFunctionScript())
+			{
+				NodeHandler->GetFunctionCallData(GraphNodeAndOutputPin, NewTraversalNode->FunctionCallData, OutFilteredConnectedPinIds, DebugData);
+			}
+			if (NodeHandler->CanEvaluateStaticValues())
+			{
+				NodeHandler->GetStaticOpData(GraphNodeAndOutputPin, NewTraversalNode->StaticOpData, OutFilteredConnectedPinIds, DebugData);
+			}
+			if (NodeHandler->CanSelectInputPin())
+			{
+				NodeHandler->GetSelectData(GraphNodeAndOutputPin, NewTraversalNode->SelectData, OutFilteredConnectedPinIds, DebugData);
+			}
+		}
+	}
+}
+
+void FTraversalBuilder::GetUnconnectedRoots(const UNiagaraNodeOutput& OutputNode, const TMap<FGraphNodeAndOutputPin, TSharedRef<FTraversalNode>>& TraversedOutputPins, TArray<FGraphNodeAndOutputPin>& OutUnconnectedRoots)
+{
+	TSet<const UEdGraphNode*> TraversedNodes;
+	for (const TPair<FGraphNodeAndOutputPin, TSharedRef<FTraversalNode>>& TraversedOutputPin : TraversedOutputPins)
+	{
+		TraversedNodes.Add(&TraversedOutputPin.Key.Node);
+	}
+
+	const UEdGraph* OwningGraph = OutputNode.GetGraph();
+	for (const UEdGraphNode* Node : OwningGraph->Nodes)
+	{
+		if (Node->IsA<UNiagaraNodeOutput>() == false && TraversedNodes.Contains(Node) == false)
+		{
+			bool bHasConnectedOutput = false;
+			TArray<const UEdGraphPin*> OutputPins;
+			for (const UEdGraphPin* Pin : Node->Pins)
+			{
+				if (Pin->Direction == EGPD_Output && UNiagaraNodeWithDynamicPins::IsAddPin(Pin) == false)
+				{
+					if (Pin->LinkedTo.Num() > 0)
+					{
+						bHasConnectedOutput = true;
+						break;
+					}
+					OutputPins.Add(Pin);
+				}
+			}
+			if (bHasConnectedOutput == false)
+			{
+				if (OutputPins.Num() == 0)
+				{
+					OutUnconnectedRoots.Add(FGraphNodeAndOutputPin(*Node, nullptr));
+				}
+				else
+				{
+					for (const UEdGraphPin* OutputPin : OutputPins)
+					{
+						OutUnconnectedRoots.Add(FGraphNodeAndOutputPin(*Node, OutputPin));
+					}
+				}
+			}
+		}
+	}
+}
+
+void FTraversalBuilder::TrimNoops(FTraversalNode& Node, TSet<const FTraversalNode*>& TraversedNodes, TSet<TWeakPtr<FTraversalNode>>* TraversedNoopNodesWeak)
+{
+	bool bAlreadyInSet = false;
+	TraversedNodes.Add(&Node, &bAlreadyInSet);
+	if (bAlreadyInSet)
+	{
+		return;
+	}
+
+	TArray<FGuid> InvalidConnectionPinIds;
+	for (FTraversalNode::FConnection& Connection : Node.Connections)
+	{
+		// Try to find a valid non-noop target connection node.
+		TSharedPtr<FTraversalNode> TargetConnectionNode;
+		TSharedPtr<FTraversalNode> CurrentNode = Connection.Node;
+		while (CurrentNode.IsValid() && TargetConnectionNode.IsValid() == false)
+		{
+			if (CurrentNode->bIsNoop)
+			{
+				if (TraversedNoopNodesWeak != nullptr)
+				{
+					// If this is not null, we're tracking noops for debugging.
+					TraversedNoopNodesWeak->Add(TWeakPtr<FTraversalNode>(CurrentNode));
+				}
+
+				if (CurrentNode->Connections.Num() == 0)
+				{
+					// No incoming connections so this noop is a dead end.
+					CurrentNode.Reset();
+				}
+				else if (CurrentNode->Connections.Num() == 1)
+				{
+					// Noop has single input, so keep looking for a non-noop.
+					CurrentNode = CurrentNode->Connections[0].Node;
+				}
+				else
+				{
+					// Multiple input connections are not supported, so this noop can not be trimmed.
+					TargetConnectionNode = CurrentNode;
+				}
+			}
+			else
+			{
+				// Non-noop found.
+				TargetConnectionNode = CurrentNode;
+			}
+		}
+
+		if (TargetConnectionNode.IsValid())
+		{
+			// If a valid target connection node was found and it's not the current connection node, replace it to trim
+			// the noops.  Continue trimming from the target connection node.
+			if (&Connection.Node.Get() != TargetConnectionNode.Get())
+			{
+				Connection.Node = TargetConnectionNode.ToSharedRef();
+			}
+			TrimNoops(Connection.Node.Get(), TraversedNodes, TraversedNoopNodesWeak);
+		}
+		else
+		{
+			// A valid target node was not found, so mark this connection for removal.
+			InvalidConnectionPinIds.Add(Connection.PinId);
+		}
+	}
+
+	if (InvalidConnectionPinIds.Num() > 0)
+	{
+		Node.Connections.RemoveAll([&InvalidConnectionPinIds](FTraversalNode::FConnection& Connection)
+			{ return InvalidConnectionPinIds.Contains(Connection.PinId); });
+	}
+}
+
+void FTraversalBuilder::ResolveModuleInputWrites(FTraversalNode& Node, TMap<FName, FGuid>& ModuleNameToFunctionCallNodeGuid, TSet<const FTraversalNode*>& TraversedNodes)
+{
+	bool bAlreadyInSet = false;
+	TraversedNodes.Add(&Node, &bAlreadyInSet);
+	if (bAlreadyInSet)
+	{
+		return;
+	}
+
+	if (Node.FunctionCallData.IsSet())
+	{
+		ModuleNameToFunctionCallNodeGuid.Add(Node.FunctionCallData->FunctionCallName, Node.SourceNodeGuid);
+	}
+
+	if (Node.ParameterData.IsSet() && Node.ParameterData->WriteParameterReferences.Num() > 0)
+	{
+		TArray<FParameterWrite> ResolvedWrites;
+		bool bAnyParameterResolved = false;
+		for (const FParameterWrite& WriteParameterReference : Node.ParameterData->WriteParameterReferences)
+		{
+			// Module input writes take the form [ModuleName].[InputName] so the initial traversal will flag the parameter as having an unknown namespace.
+			bool bParameterResolved = false;
+			if (HasFlag(WriteParameterReference.Flags, EParameterFlags::NamespaceUnknown))
+			{
+				FNiagaraParameterHandle WriteHandle(WriteParameterReference.Parameter.GetName());
+				FGuid* NodeGuid = ModuleNameToFunctionCallNodeGuid.Find(WriteHandle.GetNamespace());
+				if (NodeGuid != nullptr)
+				{
+					FName ResolvedParameterName = *(FNiagaraConstants::ModuleNamespaceString + TEXT(".") + WriteHandle.GetName().ToString());
+					FNiagaraVariableBase ResolvedParameter(WriteParameterReference.Parameter.GetType(), ResolvedParameterName);
+
+					EParameterFlags ResolvedFlags = WriteParameterReference.Flags;
+					ResolvedFlags = ClearFlag(ResolvedFlags, EParameterFlags::NamespaceUnknown);
+					ResolvedFlags = SetFlag(ResolvedFlags, EParameterFlags::ModuleInput);
+
+					FParameterWrite& ResolvedParameterWrite = ResolvedWrites.AddDefaulted_GetRef();
+					ResolvedParameterWrite = WriteParameterReference;
+					ResolvedParameterWrite.Parameter = ResolvedParameter;
+					ResolvedParameterWrite.Flags = ResolvedFlags;
+					ResolvedParameterWrite.OptionalTargetFunctionCallNodeGuid = *NodeGuid;
+
+					bParameterResolved = true;
+					bAnyParameterResolved = true;
+				}
+			}
+
+			if (bParameterResolved == false)
+			{
+				ResolvedWrites.Add(WriteParameterReference);
+			}
+		}
+
+		if (bAnyParameterResolved)
+		{
+			FParameterData ResolvedParameterData = Node.ParameterData.GetValue();
+			ResolvedParameterData.WriteParameterReferences = ResolvedWrites;
+			Node.ParameterData = ResolvedParameterData;
+		}
+	}
+
+	for (const FTraversalNode::FConnection& Connection : Node.Connections)
+	{
+		ResolveModuleInputWrites(Connection.Node.Get(), ModuleNameToFunctionCallNodeGuid, TraversedNodes);
+	}
+}
+
+void FTraversalBuilder::CollectAdditionalTraversalData(FTraversal& Traversal)
+{
+	TSet<const FTraversalNode*> AllNodes;
+	FTraversal::GetAllConnectedNodes(Traversal.TraversalRoot.Get(), AllNodes);
+
+	for (const FTraversalNode* Node : AllNodes)
+	{
+		if (Node->FunctionCallData.IsSet())
+		{
+			Traversal.FunctionNameToNodeGuidMap.Add(Node->FunctionCallData->FunctionCallName, Node->SourceNodeGuid);
+			if (Node->FunctionCallData->FunctionScriptReference.Path.IsValid())
+			{
+				Traversal.ExternalReferences.Add(Node->FunctionCallData->FunctionScriptReference);
+			}
+		}
+
+		if (Node->ParameterData.IsSet())
+		{
+			for (const FParameterRead& ReadParameterReference : Node->ParameterData->ReadParameterReferences)
+			{
+				if (ReadParameterReference.Parameter.GetType().IsStatic())
+				{
+					Traversal.StaticVariableReads.Add(ReadParameterReference.Parameter);
+				}
+			}
+			for (const FParameterWrite& WriteParameterReference : Node->ParameterData->WriteParameterReferences)
+			{
+				if (WriteParameterReference.Parameter.GetType().IsStatic())
+				{
+					Traversal.StaticVariableWrites.Add(WriteParameterReference.Parameter);
+					if (HasFlag(WriteParameterReference.Flags, EParameterFlags::Attribute))
+					{
+						Traversal.StaticVariableWritesToAttributes.Add(WriteParameterReference.Parameter);
+					}
+				}
+			}
+		}
+	}
+}
+
+bool FTraversalBuilder::IsValidSelectValueTypeInternal(const FNiagaraTypeDefinition& InValueType) const
+{
+	FScopeLock ScopeLock(&IsValidSelectValueTypeCacheCS);
+
+	bool* bIsValidType = IsValidSelectValueTypeCache.Find(InValueType);
+	if (bIsValidType == nullptr)
+	{
+		TSharedPtr<INiagaraEditorTypeUtilities, ESPMode::ThreadSafe> ValueTypeUtilities = FNiagaraEditorModule::Get().GetTypeUtilities(InValueType);
+		bIsValidType = &IsValidSelectValueTypeCache.Add(InValueType, ValueTypeUtilities.IsValid() && ValueTypeUtilities->CanBeSelectValue());
+	}
+	return *bIsValidType;
+}
+
+FSelectValue FTraversalBuilder::CreateSelectValueInternal(const FNiagaraVariable& InVariableValue) const
+{
+	FScopeLock ScopeLock(&SelectValueCacheCS);
+
+	const FNiagaraTypeDefinition& VariableType = InVariableValue.GetType();
+	UObject* TypeObject = VariableType.IsEnum() ? (UObject*)VariableType.GetEnum() : (UObject*)VariableType.GetStruct();
+	if (TypeObject == nullptr)
+	{
+		static FSelectValue Invalid(INDEX_NONE, NAME_None);
+		return Invalid;
+	}
+
+	FSelectValueCacheKey CacheKey;
+	CacheKey.TypeObjectKey = FObjectKey(TypeObject);
+	CacheKey.VariableData.Append(InVariableValue.GetData(), InVariableValue.GetSizeInBytes());
+
+	FSelectValue* SelectValue = SelectValueCache.Find(CacheKey);
+	if (SelectValue == nullptr)
+	{
+		TSharedPtr<INiagaraEditorTypeUtilities, ESPMode::ThreadSafe> ValueTypeUtilities = FNiagaraEditorModule::Get().GetTypeUtilities(VariableType);
+		int32 SelectNumericValue = ValueTypeUtilities->VariableToSelectNumericValue(InVariableValue);
+		FName SelectValueDebugName = ValueTypeUtilities->GetDebugNameForSelectValue(VariableType, SelectNumericValue);
+		SelectValue = &SelectValueCache.Add(CacheKey, FSelectValue(SelectNumericValue, SelectValueDebugName));
+	}
+	return *SelectValue;
+}
+
+FSelectValue FTraversalBuilder::CreateSelectValueInternal(const FNiagaraTypeDefinition& ValueType, int32 SelectNumericValue) const
+{
+	FScopeLock ScopeLock(&SelectValueCacheCS);
+
+	UObject* TypeObject = ValueType.IsEnum() ? (UObject*)ValueType.GetEnum() : (UObject*)ValueType.GetStruct();
+	if (TypeObject == nullptr)
+	{
+		static FSelectValue Invalid(INDEX_NONE, NAME_None);
+		return Invalid;
+	}
+
+	FSelectValueCacheKey CacheKey;
+	CacheKey.TypeObjectKey = FObjectKey(TypeObject);
+	CacheKey.SelectNumericValue = SelectNumericValue;
+
+	FSelectValue* SelectValue = SelectValueCache.Find(CacheKey);
+	if (SelectValue == nullptr)
+	{
+		TSharedPtr<INiagaraEditorTypeUtilities, ESPMode::ThreadSafe> ValueTypeUtilities = FNiagaraEditorModule::Get().GetTypeUtilities(ValueType);
+		FName SelectValueDebugName = ValueTypeUtilities->GetDebugNameForSelectValue(ValueType, SelectNumericValue);
+		SelectValue = &SelectValueCache.Add(CacheKey, FSelectValue(SelectNumericValue, SelectValueDebugName));
+	}
+	return *SelectValue;
+}
+
+EParameterFlags FTraversalBuilder::ExtractFlagsFromParameterNameInternal(FName ParameterName) const
+{
+	FScopeLock ScopeLock(&ExtractedParameterFlagCacheCS);
+
+	EParameterFlags* CachedFlags = ExtractedParameterFlagCache.Find(ParameterName);
+	if (CachedFlags != nullptr)
+	{
+		return *CachedFlags;
+	}
+
+	EParameterFlags Flags = EParameterFlags::None;
+	FNiagaraParameterHandle ParameterHandle(ParameterName);
+	if (ParameterHandle.GetNamespace() == NAME_None)
+	{
+		Flags = EParameterFlags::InvalidParameterName;
+	}
+	else if (ParameterHandle.IsModuleHandle())
+	{
+		Flags = EParameterFlags::ModuleInput;
+	}
+	else if (ParameterHandle.IsLocalHandle())
+	{
+		Flags = EParameterFlags::ModuleLocal;
+	}
+	else if (ParameterHandle.IsOutputHandle())
+	{
+		FNiagaraParameterHandle ParameterSubHandle(ParameterHandle.GetName());
+		Flags = ParameterSubHandle.GetNamespace() != NAME_None
+			? EParameterFlags::ModuleOutput | EParameterFlags::NamespaceModifier
+			: EParameterFlags::InvalidParameterName;
+	}
+	else if (ParameterHandle.IsSystemHandle() || ParameterHandle.IsEmitterHandle() || ParameterHandle.IsParticleAttributeHandle() || ParameterHandle.IsStackContextHandle())
+	{
+		Flags = EParameterFlags::Attribute;
+		FNiagaraParameterHandle ParameterSubHandle(ParameterHandle.GetName());
+		FName NamespaceModifier = ParameterSubHandle.GetNamespace();
+		if (NamespaceModifier != NAME_None &&
+			NamespaceModifier != FNiagaraConstants::PreviousNamespace &&
+			NamespaceModifier != FNiagaraConstants::InitialNamespace &&
+			NamespaceModifier != FNiagaraConstants::OwnerNamespace)
+		{
+			// If the attribute has a namespace modifier that's unknown, or is the module namespace
+			// it has to be handled differently to support attribute module outputs correctly.
+			Flags = Flags | EParameterFlags::NamespaceModifier;
+			if (NamespaceModifier == FNiagaraConstants::ModuleNamespace)
+			{
+				Flags = Flags | EParameterFlags::ModuleOutput;
+			}
+		}
+	}
+	else if (ParameterHandle.IsEngineHandle() || ParameterHandle.IsDataInstanceHandle() ||
+		ParameterHandle.IsParameterCollectionHandle() || ParameterHandle.IsUserHandle())
+	{
+		Flags = EParameterFlags::External;
+	}
+	else
+	{
+		Flags = EParameterFlags::NamespaceUnknown;
+	}
+
+	ExtractedParameterFlagCache.Add(ParameterName, Flags);
+	return Flags;
+}
+
+} // UE::Niagara::TraversalCache
\ No newline at end of file
diff --git a/Engine/Plugins/FX/Niagara/Source/NiagaraEditor/Private/TraversalCache/TraversalBuilder.h b/Engine/Plugins/FX/Niagara/Source/NiagaraEditor/Private/TraversalCache/TraversalBuilder.h
new file mode 100644
--- /dev/null
+++ b/Engine/Plugins/FX/Niagara/Source/NiagaraEditor/Private/TraversalCache/TraversalBuilder.h
@@ -0,0 +1,221 @@
+// Copyright Epic Games, Inc. All Rights Reserved.
+
+#pragma once
+
+#include "EdGraph/EdGraphNode.h"
+#include "NiagaraVariant.h"
+#include "NiagaraTypes.h"
+#include "Templates/SharedPointer.h"
+#include "Templates/SubclassOf.h"
+#include "TraversalCache/TraversalShared.h"
+
+class UClass;
+class UEdGraphNode;
+class UEdGraphPin;
+class UNiagaraNode;
+class UNiagaraNodeOutput;
+
+namespace UE::Niagara::TraversalCache
+{
+
+struct FFunctionCallData;
+struct FFunctionInputData;
+struct FParameterData;
+struct FSelectData;
+struct FStaticOpData;
+class  FTraversal;
+class  FTraversalNode;
+struct FTraversalBuilderDebugData;
+
+class FTraversalBuilder
+{
+public:
+	struct FGraphNodeAndOutputPin
+	{
+		FGraphNodeAndOutputPin(const UEdGraphNode& InNode, const UEdGraphPin* InOutputPin)
+			: Node(InNode)
+			, OutputPin(InOutputPin)
+		{
+		}
+
+		FORCEINLINE bool operator==(const FGraphNodeAndOutputPin& Other) const
+		{
+			return &Node == &Other.Node && OutputPin == Other.OutputPin;
+		}
+
+		friend FORCEINLINE uint32 GetTypeHash(const FGraphNodeAndOutputPin& NodeAndPin)
+		{
+			return HashCombine(GetTypeHash(&NodeAndPin.Node), GetTypeHash(NodeAndPin.OutputPin));
+		}
+
+		const UEdGraphNode& Node;
+		const UEdGraphPin* OutputPin;
+	};
+
+	class IGraphNodeHandler
+	{
+	public:
+		virtual ~IGraphNodeHandler() { }
+
+		virtual bool IsNoop() const = 0;
+
+		virtual bool CanProvideFunctionInput() const = 0;
+		virtual void GetFunctionInputData(
+			const FGraphNodeAndOutputPin& GraphNodeAndOutputPin,
+			TOptional<const FFunctionInputData>& OutInputData,
+			FTraversalBuilderDebugData* DebugData) const = 0;
+
+		virtual bool CanAccessParameters() const = 0;
+		virtual void GetParameterData(
+			const FGraphNodeAndOutputPin& GraphNodeAndOutputPin,
+			TOptional<const FParameterData>& OutParameterData,
+			TOptional<TArray<FGuid>>& OutFilteredConnectedPinIds,
+			FTraversalBuilderDebugData* DebugData) const = 0;
+
+		virtual bool CanSelectInputPin() const = 0;
+		virtual void GetSelectData(
+			const FGraphNodeAndOutputPin& GraphNodeAndOutputPin,
+			TOptional<const FSelectData>& OutSelectData,
+			TOptional<TArray<FGuid>>& OutFilteredConnectedPinIds,
+			FTraversalBuilderDebugData* DebugData) const = 0;
+
+		virtual bool CanCallFunctionScript() const = 0;
+		virtual void GetFunctionCallData(
+			const FGraphNodeAndOutputPin& GraphNodeAndOutputPin,
+			TOptional<const FFunctionCallData>& OutFunctionCallData,
+			TOptional<TArray<FGuid>>& OutFilteredConnectedPinIds,
+			FTraversalBuilderDebugData* DebugData) const = 0;
+
+		virtual bool CanEvaluateStaticValues() const = 0;
+		virtual void GetStaticOpData(
+			const FGraphNodeAndOutputPin& GraphNodeAndOutputPin,
+			TOptional<const FStaticOpData>& OutStaticOpData,
+			TOptional<TArray<FGuid>>& OutFilteredConnectedPinIds,
+			FTraversalBuilderDebugData* DebugData) const = 0;
+	};
+
+	class FGraphNodeHandler : public IGraphNodeHandler
+	{
+	public:
+		virtual bool IsNoop() const override { return false; }
+
+		virtual bool CanProvideFunctionInput() const override { return false; }
+		virtual void GetFunctionInputData(
+			const FGraphNodeAndOutputPin& GraphNodeAndOutputPin,
+			TOptional<const FFunctionInputData>& OutInputData,
+			FTraversalBuilderDebugData* DebugData) const override { unimplemented(); }
+
+		virtual bool CanAccessParameters() const override { return false; }
+		virtual void GetParameterData(
+			const FGraphNodeAndOutputPin& GraphNodeAndOutputPin,
+			TOptional<const FParameterData>& OutParameterData,
+			TOptional<TArray<FGuid>>& OutFilteredConnectedPinIds,
+			FTraversalBuilderDebugData* DebugData) const override { unimplemented(); }
+
+		virtual bool CanSelectInputPin() const override { return false; }
+		virtual void GetSelectData(
+			const FGraphNodeAndOutputPin& GraphNodeAndOutputPin,
+			TOptional<const FSelectData>& OutSelectData,
+			TOptional<TArray<FGuid>>& OutFilteredConnectedPinIds,
+			FTraversalBuilderDebugData* DebugData) const override { unimplemented(); }
+
+		virtual bool CanCallFunctionScript() const override { return false; }
+		virtual void GetFunctionCallData(
+			const FGraphNodeAndOutputPin& GraphNodeAndOutputPin,
+			TOptional<const FFunctionCallData>& OutFunctionCallData,
+			TOptional<TArray<FGuid>>& OutFilteredConnectedPinIds,
+			FTraversalBuilderDebugData* DebugData) const override { unimplemented(); }
+
+		virtual bool CanEvaluateStaticValues() const override { return false; }
+		virtual void GetStaticOpData(
+			const FGraphNodeAndOutputPin& GraphNodeAndOutputPin,
+			TOptional<const FStaticOpData>& OutStaticOpData,
+			TOptional<TArray<FGuid>>& OutFilteredConnectedPinIds,
+			FTraversalBuilderDebugData* DebugData) const override { unimplemented(); }
+	};
+
+public:
+	void Initialize();
+
+	static TSharedRef<const FTraversal> BuildTraversal(const UNiagaraNodeOutput& OutputNode, FTraversalBuilderDebugData* DebugData = nullptr);
+
+	static NIAGARAEDITOR_API void ResolveFunctionCallStackNames(const UEdGraph& TopLevelGraph, const TArray<FGuid>& FunctionCallStack, TArray<FString>& OutFunctionCallNames);
+
+	static bool IsValidSelectValueType(const FNiagaraTypeDefinition& InValueType);
+
+	static FSelectValue CreateSelectValue(const FNiagaraTypeDefinition& InValueType, int32 InSelectNumericValue);
+
+	static FSelectValue CreateSelectValue(const FNiagaraVariable& InNiagaraVariable);
+
+	static FSelectValue CreateSelectValue(bool bInBoolValue);
+
+	static FSelectValue CreateSelectValue(UEnum* Enum, int32 EnumValue);
+
+	static EParameterFlags ExtractFlagsFromParameterName(FName ParameterName);
+
+private:
+	static FTraversalBuilder& GetInstance();
+
+	TSharedRef<const FTraversal> BuildTraversalInternal(const UNiagaraNodeOutput& OutputNode, FTraversalBuilderDebugData* DebugData) const;
+
+	TSharedPtr<FTraversalNode> TraverseGraphNodeFromOutputPin(
+		const FGraphNodeAndOutputPin& GraphNodeAndOutputPin,
+		TMap<FGraphNodeAndOutputPin, TSharedRef<FTraversalNode>>& TraversedOutputPins,
+		FTraversalBuilderDebugData* DebugData) const;
+
+	void GatherNodeDataAndFilteredConnectedPinIds(
+		const FGraphNodeAndOutputPin& GraphNodeAndOutputPin,
+		TSharedRef<FTraversalNode> NewTraversalNode,
+		TOptional<TArray<FGuid>>& OutFilteredConnectedPinIds,
+		FTraversalBuilderDebugData* DebugData) const;
+
+	static void GetUnconnectedRoots(const UNiagaraNodeOutput& OutputNode, const TMap<FGraphNodeAndOutputPin, TSharedRef<FTraversalNode>>& TraversedOutputPins, TArray<FGraphNodeAndOutputPin>& OutUnconnectedRoots);
+
+	static void TrimNoops(FTraversalNode& Node, TSet<const FTraversalNode*>& TraversedNodes, TSet<TWeakPtr<FTraversalNode>>* TraversedNoopNodesWeak);
+
+	static void ResolveModuleInputWrites(FTraversalNode& Node, TMap<FName, FGuid>& ModuleNameToFunctionCallNodeGuid, TSet<const FTraversalNode*>& TraversedNodes);
+
+	bool IsValidSelectValueTypeInternal(const FNiagaraTypeDefinition& InValueType) const;
+
+	FSelectValue CreateSelectValueInternal(const FNiagaraVariable& InVariableValue) const;
+	
+	FSelectValue CreateSelectValueInternal(const FNiagaraTypeDefinition& ValueType, int32 SelectNumericValue) const;
+
+	static void CollectAdditionalTraversalData(FTraversal& Traversal);
+
+	struct FSelectValueCacheKey
+	{
+		FORCEINLINE bool operator==(const FSelectValueCacheKey& Other) const
+		{
+			return TypeObjectKey == Other.TypeObjectKey && VariableData == Other.VariableData && SelectNumericValue == Other.SelectNumericValue;
+		}
+
+		friend FORCEINLINE uint32 GetTypeHash(const FSelectValueCacheKey& Key)
+		{
+			int32 Hash = HashCombine(GetTypeHash(Key.TypeObjectKey), GetTypeHash(Key.VariableData));
+			return HashCombine(Hash, GetTypeHash(Key.SelectNumericValue));
+		}
+
+		FObjectKey TypeObjectKey;
+		TArray<uint8> VariableData;
+		int32 SelectNumericValue = INDEX_NONE;
+	};
+
+	EParameterFlags ExtractFlagsFromParameterNameInternal(FName ParameterName) const;
+
+private:
+	static FTraversalBuilder* Instance;
+	bool bInitialized = false;
+
+	TMap<TSubclassOf<UEdGraphNode>, TSharedRef<IGraphNodeHandler>> GraphNodeHandlers;
+	
+	static FCriticalSection IsValidSelectValueTypeCacheCS;
+	static FCriticalSection SelectValueCacheCS;
+	static FCriticalSection ExtractedParameterFlagCacheCS;
+
+	mutable TMap<FNiagaraTypeDefinition, bool> IsValidSelectValueTypeCache;
+	mutable TMap<FSelectValueCacheKey, FSelectValue> SelectValueCache;
+	mutable TMap<FName, EParameterFlags> ExtractedParameterFlagCache;
+};
+
+} // UE::Niagara::TraversalCache
\ No newline at end of file
diff --git a/Engine/Plugins/FX/Niagara/Source/NiagaraEditor/Private/TraversalCache/TraversalCache.cpp b/Engine/Plugins/FX/Niagara/Source/NiagaraEditor/Private/TraversalCache/TraversalCache.cpp
new file mode 100644
--- /dev/null
+++ b/Engine/Plugins/FX/Niagara/Source/NiagaraEditor/Private/TraversalCache/TraversalCache.cpp
@@ -0,0 +1,723 @@
+// Copyright Epic Games, Inc. All Rights Reserved.
+
+#include "TraversalCache/TraversalCache.h"
+
+#include "EdGraphSchema_Niagara.h"
+#include "NiagaraCommon.h"
+#include "NiagaraConstants.h"
+#include "NiagaraEditorModule.h"
+#include "NiagaraNodeFunctionCall.h"
+#include "NiagaraNodeOutput.h"
+#include "NiagaraScriptSource.h"
+#include "NiagaraScriptVariable.h"
+#include "String/ParseTokens.h"
+#include "TraversalCache/Traversal.h"
+#include "TraversalCache/TraversalBuilder.h"
+#include "TraversalCache/TraversalNode.h"
+#include "TraversalCache/TraversalShared.h"
+
+namespace UE::Niagara::TraversalCache
+{
+
+FTraversalCache* FTraversalCache::Instance = nullptr;
+
+void FTraversalCache::Initialize()
+{
+	if (bInitialized == false)
+	{
+		checkf(Instance == nullptr, TEXT("FTraversalCache initialized more than once."));
+		Instance = this;
+		FNiagaraEditorModule::Get().OnScriptApplied().AddSP(this, &FTraversalCache::OnScriptApplied);
+		bInitialized = true;
+	}
+}
+
+TSharedPtr<const FTraversal> FTraversalCache::GetScriptAssetTraversal(const FSoftObjectPath& ScriptSoftObjectPath, const FGuid& ScriptVersion, FTraversalBuilderDebugData* BuilderDebugData)
+{
+	return GetInstance().GetScriptAssetTraversalInternal(ScriptSoftObjectPath, ScriptVersion, BuilderDebugData);
+}
+
+TSharedPtr<const FTraversal> FTraversalCache::GetScriptAssetTraversal(const FSoftObjectPath& ScriptSoftObjectPath, const FGuid& ScriptVersion, FTraversalDebugData* TraversalDebugData)
+{
+	TOptional<FTraversalBuilderDebugData> BuilderDebugData = TraversalDebugData != nullptr ? FTraversalBuilderDebugData() : TOptional<FTraversalBuilderDebugData>();
+	TSharedPtr<const FTraversal> Traversal = FTraversalCache::GetScriptAssetTraversal(ScriptSoftObjectPath, ScriptVersion, BuilderDebugData.GetPtrOrNull());
+	if (TraversalDebugData != nullptr && BuilderDebugData->HasData())
+	{
+		TraversalDebugData->AddBuilderDebugOutput(ScriptSoftObjectPath, ENiagaraScriptUsage::Function, FGuid(), ScriptVersion, BuilderDebugData.GetValue());
+	}
+	return Traversal;
+}
+
+TSharedPtr<const FTraversal> FTraversalCache::GetTopLevelScriptTraversal(const UNiagaraScript& Script, FTraversalBuilderDebugData* BuilderDebugData)
+{
+	return GetInstance().GetTopLevelScriptTraversalInternal(Script, BuilderDebugData);
+}
+
+TSharedPtr<const FTraversal> FTraversalCache::GetTopLevelScriptTraversal(const UNiagaraScript& Script, FTraversalDebugData* TraversalDebugData)
+{
+	TOptional<FTraversalBuilderDebugData> BuilderDebugData = TraversalDebugData != nullptr ? FTraversalBuilderDebugData() : TOptional<FTraversalBuilderDebugData>();
+	TSharedPtr<const FTraversal> Traversal = FTraversalCache::GetTopLevelScriptTraversal(Script, BuilderDebugData.GetPtrOrNull());
+	if (TraversalDebugData != nullptr && BuilderDebugData->HasData())
+	{
+		const UNiagaraScriptSource* ScriptSource = CastChecked<UNiagaraScriptSource>(Script.GetSource(FGuid()));
+		TraversalDebugData->AddBuilderDebugOutput(ScriptSource->NodeGraph, Script.GetUsage(), Script.GetUsageId(), FGuid(), BuilderDebugData.GetValue());
+	}
+	return Traversal;
+}
+
+void FTraversalCache::GetTopLevelScriptFunctionInputContexts(const UNiagaraScript& Script, TMap<FGuid, FStaticContext>& OutFunctionInputContexts, FTraversalDebugData* TraversalDebugData)
+{
+	GetInstance().GetTopLevelScriptFunctionInputContextsInternal(Script, OutFunctionInputContexts, TraversalDebugData);
+}
+
+void FTraversalCache::GetStackFunctionReads(
+	const FTopLevelScriptStaticContext& ScriptStaticContext,
+	const UNiagaraNodeFunctionCall& FunctionCallNode,
+	TArray<FNiagaraVariable>& OutReads,
+	TSet<FNiagaraVariable>& OutHiddenReads,
+	EStackFunctionReadFilterFlags FilterFlags,
+	FTraversalDebugData* TraversalDebugData)
+{
+	TOptional<FTraversalBuilderDebugData> BuilderDebugData = TraversalDebugData != nullptr ? FTraversalBuilderDebugData() : TOptional<FTraversalBuilderDebugData>();
+	GetInstance().GetStackFunctionReadsInternal(ScriptStaticContext, FunctionCallNode, OutReads, OutHiddenReads, FilterFlags, TraversalDebugData, BuilderDebugData.GetPtrOrNull());
+	if (TraversalDebugData != nullptr && BuilderDebugData->HasData())
+	{
+		TraversalDebugData->AddBuilderDebugOutput(FunctionCallNode.GetCalledGraph(), ENiagaraScriptUsage::Function, FGuid(), FunctionCallNode.SelectedScriptVersion, BuilderDebugData.GetValue());
+	}
+}
+
+void FTraversalCache::GetStackFunctionStaticInputs(
+	const FTopLevelScriptStaticContext& ScriptStaticContext,
+	const UNiagaraNodeFunctionCall& FunctionCallNode,
+	TArray<FNiagaraVariable>& OutStaticInputs,
+	TSet<FNiagaraVariable>& OutHiddenStaticInputs,
+	FTraversalDebugData* TraversalDebugData)
+{
+	TOptional<FTraversalBuilderDebugData> BuilderDebugData = TraversalDebugData != nullptr ? FTraversalBuilderDebugData() : TOptional<FTraversalBuilderDebugData>();
+	GetInstance().GetStackFunctionStaticInputsInternal(ScriptStaticContext, FunctionCallNode, OutStaticInputs, OutHiddenStaticInputs, TraversalDebugData, BuilderDebugData.GetPtrOrNull());
+	if (TraversalDebugData != nullptr && BuilderDebugData->HasData())
+	{
+		TraversalDebugData->AddBuilderDebugOutput(FunctionCallNode.GetCalledGraph(), ENiagaraScriptUsage::Function, FGuid(), FunctionCallNode.SelectedScriptVersion, BuilderDebugData.GetValue());
+	}
+}
+
+void FTraversalCache::GetStackFunctionStaticInputPins(
+	const FTopLevelScriptStaticContext& ScriptStaticContext,
+	const UNiagaraNodeFunctionCall& FunctionCallNode,
+	TArray<UEdGraphPin*>& OutStaticInputPins,
+	TSet<UEdGraphPin*>& OutHiddenStaticInputPins,
+	FTraversalDebugData* TraversalDebugData)
+{
+	GetInstance().GetStackFunctionStaticInputPinsInternal(ScriptStaticContext, FunctionCallNode, OutStaticInputPins, OutHiddenStaticInputPins, TraversalDebugData);
+}
+
+void FTraversalCache::InvalidateCache()
+{
+	GetInstance().InvalidateCacheInternal();
+}
+
+FTraversalCache& FTraversalCache::GetInstance()
+{
+	checkf(Instance != nullptr, TEXT("FTraversalCache was not initialized."));
+	return *Instance;
+}
+
+const UNiagaraScript* FTraversalCache::GetScriptFromPath(const FSoftObjectPath& Path) const
+{
+	const UNiagaraScript* Script = nullptr;
+	TWeakObjectPtr<const UNiagaraScript>* CachedScriptPtr = SoftObjectPathToScriptCache.Find(Path);
+	if (CachedScriptPtr == nullptr)
+	{
+		Script = Cast<UNiagaraScript>(Path.ResolveObject());
+		SoftObjectPathToScriptCache.Add(Path, TWeakObjectPtr(Script));
+	}
+	else
+	{
+		Script = CachedScriptPtr->Get();
+	}
+	return Script;
+}
+
+void FTraversalCache::OnScriptApplied(UNiagaraScript* InScript, FGuid InVersionGuid)
+{
+	TMap<FScriptTraversalCacheKey, TArray<FScriptTraversalCacheKey>> CacheKeyToReferencerCacheKeys;
+	for (const TPair<FScriptTraversalCacheKey, FScriptTraversalCacheValue>& CachedTraversal : CachedTraversals)
+	{
+		for (const FScriptReference& ExternalReference : CachedTraversal.Value.ScriptTraversal->GetExternalReferences())
+		{
+			const UNiagaraScript* Script = GetScriptFromPath(ExternalReference.Path);
+			if (Script == nullptr)
+			{
+				continue;
+			}
+
+			FScriptTraversalCacheKey ExternalReferenceKey = FScriptTraversalCacheKey::FromScriptAsset(FObjectKey(Script), ExternalReference.Version);
+			TArray<FScriptTraversalCacheKey>& ReferencerCacheKeys = CacheKeyToReferencerCacheKeys.FindOrAdd(ExternalReferenceKey);
+			ReferencerCacheKeys.Add(CachedTraversal.Key);
+		}
+	}
+
+	TArray<FScriptTraversalCacheKey> CacheKeysToCheck = { FScriptTraversalCacheKey::FromScriptAsset(FObjectKey(InScript), InVersionGuid) };
+	int32 CacheKeysToCheckIndex = 0;
+
+	TSet<FScriptTraversalCacheKey> CacheKeysToRemove;
+	TSet<FScriptTraversalCacheKey> SeenDependents;
+	while (CacheKeysToCheckIndex < CacheKeysToCheck.Num())
+	{
+		const FScriptTraversalCacheKey& CurrentCacheKey = CacheKeysToCheck[CacheKeysToCheckIndex];
+		CacheKeysToCheckIndex++;
+
+		const TArray<FScriptTraversalCacheKey>* ReferencerCacheKeys = CacheKeyToReferencerCacheKeys.Find(CurrentCacheKey);
+		if (ReferencerCacheKeys == nullptr || ReferencerCacheKeys->Num() == 0)
+		{
+			continue;
+		}
+
+		for (const FScriptTraversalCacheKey& ReferencerCacheKey : *ReferencerCacheKeys)
+		{
+			bool bAlreadyAdded = false;
+			CacheKeysToRemove.Add(ReferencerCacheKey, &bAlreadyAdded);
+			if (bAlreadyAdded == false)
+			{
+				CacheKeysToCheck.Add(ReferencerCacheKey);
+			}
+		}
+	}
+
+	for (const FScriptTraversalCacheKey& CacheKeyToRemove : CacheKeysToRemove)
+	{
+		CachedTraversals.Remove(CacheKeyToRemove);
+	}
+}
+
+TSharedPtr<const FTraversal> FTraversalCache::GetScriptAssetTraversalInternal(const FSoftObjectPath& ScriptSoftObjectPath, const FGuid& ScriptVersion, FTraversalBuilderDebugData* BuilderDebugData) const
+{
+	const UNiagaraScript* Script = GetScriptFromPath(ScriptSoftObjectPath);
+	if (Script == nullptr)
+	{
+		return TSharedPtr<const FTraversal>();
+	}
+
+	const FVersionedNiagaraScriptData* ScriptData = Script->GetScriptData(ScriptVersion);
+	if (ScriptData == nullptr)
+	{
+		return TSharedPtr<const FTraversal>();
+	}
+
+	FScriptTraversalCacheKey CacheKey = FScriptTraversalCacheKey::FromScriptAsset(FObjectKey(Script), ScriptVersion);
+	FGuid ChangeId = Script->GetBaseChangeID(ScriptVersion);
+
+	FScriptTraversalCacheValue* CachedValue = CachedTraversals.Find(CacheKey);
+	if (CachedValue != nullptr && CachedValue->GraphChangeId == ChangeId)
+	{
+		return CachedValue->ScriptTraversal;
+	}
+
+	const UNiagaraScriptSource* ScriptSource = CastChecked<UNiagaraScriptSource>(Script->GetSource(ScriptVersion));
+	UNiagaraNodeOutput* OutputNode = ScriptSource->NodeGraph->FindEquivalentOutputNode(Script->GetUsage(), Script->GetUsageId());
+	if (OutputNode == nullptr)
+	{
+		return TSharedPtr<const FTraversal>();
+	}
+
+
+	TSharedRef<const FTraversal> Traversal = FTraversalBuilder::BuildTraversal(*OutputNode, BuilderDebugData);
+	if (BuilderDebugData == nullptr || BuilderDebugData->HasData() == false)
+	{
+		CachedTraversals.Add(CacheKey, FScriptTraversalCacheValue(ChangeId, Traversal));
+	}
+	return Traversal;
+}
+
+TSharedPtr<const FTraversal> FTraversalCache::GetTopLevelScriptTraversalInternal(const UNiagaraScript& Script, FTraversalBuilderDebugData* BuilderDebugData) const
+{
+	const UNiagaraScriptSource* ScriptSource = CastChecked<UNiagaraScriptSource>(Script.GetSource(FGuid()));
+
+	FScriptTraversalCacheKey CacheKey = FScriptTraversalCacheKey::FromTopLevelScript(FObjectKey(ScriptSource->NodeGraph), Script.GetUsage(), Script.GetUsageId());
+	FGuid ChangeId = Script.GetBaseChangeID(FGuid());
+
+	FScriptTraversalCacheValue* CachedValue = CachedTraversals.Find(CacheKey);
+	if (CachedValue != nullptr && CachedValue->GraphChangeId == ChangeId)
+	{
+		return CachedValue->ScriptTraversal;
+	}
+
+	UNiagaraNodeOutput* OutputNode = ScriptSource->NodeGraph->FindEquivalentOutputNode(Script.GetUsage(), Script.GetUsageId());
+	if (OutputNode == nullptr)
+	{
+		return TSharedPtr<const FTraversal>();
+	}
+
+	TSharedRef<const FTraversal> Traversal = FTraversalBuilder::BuildTraversal(*OutputNode, BuilderDebugData);
+	if (BuilderDebugData == nullptr || BuilderDebugData->HasData() == false)
+	{
+		// We only cache the results when either no debug data was supplied, or when the debug data was empty.  When
+		// the debug data HasData returns true, errors were encountered while building.
+		CachedTraversals.Add(CacheKey, FScriptTraversalCacheValue(ChangeId, Traversal));
+	}
+	return Traversal;
+}
+
+TSharedPtr<const FTraversal> FTraversalCache::GetTopLevelScriptTraversalInternal(const UNiagaraNode& StackNode, FTraversalBuilderDebugData* BuilderDebugData) const
+{
+	const UNiagaraGraph* OwningGraph = StackNode.GetNiagaraGraph();
+	if (OwningGraph == nullptr)
+	{
+		return TSharedPtr<const FTraversal>();
+	}
+
+	const UNiagaraNodeOutput* OutputNode = FNiagaraStackGraphUtilities::GetEmitterOutputNodeForStackNode(StackNode);
+	if (OutputNode == nullptr)
+	{
+		return TSharedPtr<const FTraversal>();
+	}
+
+	FScriptTraversalCacheKey CacheKey = FScriptTraversalCacheKey::FromTopLevelScript(FObjectKey(OwningGraph), OutputNode->GetUsage(), OutputNode->GetUsageId());
+	FScriptTraversalCacheValue* CachedValue = CachedTraversals.Find(CacheKey);
+	if (CachedValue != nullptr && CachedValue->GraphChangeId == OwningGraph->GetChangeID())
+	{
+		return CachedValue->ScriptTraversal;
+	}
+
+	TSharedRef<const FTraversal> Traversal = FTraversalBuilder::BuildTraversal(*OutputNode, BuilderDebugData);
+	if (BuilderDebugData == nullptr || BuilderDebugData->HasData() == false)
+	{
+		// We only cache the results when either no debug data was supplied, or when the debug data was empty.  When
+		// the debug data HasData returns true, errors were encountered while building.
+		CachedTraversals.Add(CacheKey, FScriptTraversalCacheValue(OwningGraph->GetChangeID(), Traversal));
+	}
+	return Traversal;
+}
+
+void FTraversalCache::GetTopLevelScriptFunctionInputContextsInternal(const UNiagaraScript& Script, TMap<FGuid, FStaticContext>& OutFunctionInputContexts, FTraversalDebugData* TraversalDebugData) const
+{
+	// Static inputs are stored as rapid iteration parameters, and their names are mangled to keep them unique. Modules
+	// will be expecting these to be in the format Module.[InputName] so the names must be unmangled, and disambiguated by
+	// module node guid so that they can be gathered when creating function contexts.  Emitter rapid iteration parameters
+	// may also have additional mangling for the emitter name for module inputs of the form Module.Emitter.[InputName].
+	const UNiagaraScriptSource* ScriptSource = CastChecked<UNiagaraScriptSource>(Script.GetSource(FGuid()));
+	FScriptTraversalCacheKey CacheKey = FScriptTraversalCacheKey::FromTopLevelScript(FObjectKey(ScriptSource->NodeGraph), Script.GetUsage(), Script.GetUsageId());
+	FParsedRapidIterationParameterData* ParsedRapidIterationParameterData = CachedParsedRapidIterationParameterData.Find(CacheKey);
+	if (ParsedRapidIterationParameterData == nullptr || ParsedRapidIterationParameterData->ParameterStoreLayoutVersion != Script.RapidIterationParameters.GetLayoutVersion())
+	{
+		if (ParsedRapidIterationParameterData == nullptr)
+		{
+			ParsedRapidIterationParameterData = &CachedParsedRapidIterationParameterData.Add(CacheKey);
+		}
+		ParsedRapidIterationParameterData->ParameterStoreLayoutVersion = Script.RapidIterationParameters.GetLayoutVersion();
+		ParsedRapidIterationParameterData->StaticFunctionInputs.Empty();
+
+		TSharedPtr<const FTraversal> SourceScriptTraversal = GetTopLevelScriptTraversal(Script, TraversalDebugData);
+		if (SourceScriptTraversal.IsValid() == false)
+		{
+			return;
+		}
+
+		const TMap<FName, FGuid>& FunctionNameToNodeGuidMap = SourceScriptTraversal->GetFunctionNameToNodeGuidMap();
+		if (FunctionNameToNodeGuidMap.IsEmpty())
+		{
+			return;
+		}
+
+		for (const FNiagaraVariableWithOffset& VariableWithOffset : Script.RapidIterationParameters.ReadParameterVariables())
+		{
+			const FNiagaraTypeDefinition& VariableType = VariableWithOffset.GetType();
+			if (VariableType.IsStatic() == false)
+			{
+				continue;
+			}
+
+			FNiagaraTypeDefinition NonStaticVariableType = VariableType.RemoveStaticDef();
+			if (FTraversalBuilder::IsValidSelectValueType(NonStaticVariableType) == false)
+			{
+				continue;
+			}
+
+			FNameBuilder VariableNameBuilder(VariableWithOffset.GetName());
+			TArray<FStringView> NameParts;
+			UE::String::ParseTokens(VariableNameBuilder.ToView(), TEXT("."), NameParts);
+			if (NameParts[0] != FNiagaraConstants::RapidIterationParametersNamespaceString)
+			{
+				continue;
+			}
+
+			// For system scripts we need to skip the rapid iteration namespace, but for emitter and particle scripts we need to skip the rapid iteration and the emitter namespaces. 
+			int32 ModuleNameIndex;
+			int32 EmitterNameIndex;
+			if (Script.GetUsage() == ENiagaraScriptUsage::SystemSpawnScript || Script.GetUsage() == ENiagaraScriptUsage::SystemUpdateScript)
+			{
+				EmitterNameIndex = INDEX_NONE;
+				ModuleNameIndex = 1;
+			}
+			else
+			{
+				EmitterNameIndex = 1;
+				ModuleNameIndex = 2;
+			}
+
+			if (NameParts.Num() - ModuleNameIndex <= 1)
+			{
+				continue;
+			}
+
+			FName ModuleName = FName(NameParts[ModuleNameIndex]);
+			const FGuid* ModuleGuid = FunctionNameToNodeGuidMap.Find(ModuleName);
+			if (ModuleGuid == nullptr)
+			{
+				continue;
+			}
+
+			TArray<FStringView> ModuleInputNameParts = { FNiagaraConstants::ModuleNamespaceString };
+			ModuleInputNameParts.Append(TArrayView<FStringView>(NameParts).Slice(ModuleNameIndex + 1, NameParts.Num() - (ModuleNameIndex + 1)));
+			if (ModuleInputNameParts.Num() > 2 &&
+				(Script.GetUsage() == ENiagaraScriptUsage::EmitterSpawnScript || Script.GetUsage() == ENiagaraScriptUsage::EmitterUpdateScript) &&
+				ModuleInputNameParts[1] == NameParts[EmitterNameIndex])
+			{
+				// This handles the special case of module inputs which have the form Module.Emitter.[InputName] in the emitter script.  In this case
+				// both 'Module' and 'Emitter' will have been replaced with the module name and emitter name respectively, so both parts of the name
+				// must be replaced with the generic versions.
+				ModuleInputNameParts[1] = FNiagaraConstants::EmitterNamespaceString;
+			}
+
+			FName ModuleInputName = *FString::Join(ModuleInputNameParts, TEXT("."));
+			FNiagaraVariableBase ModuleInputVariable(VariableType, ModuleInputName);
+
+			ParsedRapidIterationParameterData->StaticFunctionInputs.Add(FParsedStaticFunctionInput(*ModuleGuid, ModuleInputVariable, VariableWithOffset.Offset));
+		}
+	}
+
+	for (const FParsedStaticFunctionInput& ParsedStaticFunctionInput : ParsedRapidIterationParameterData->StaticFunctionInputs)
+	{
+		FStaticContext& FunctionInputContext = OutFunctionInputContexts.FindOrAdd(ParsedStaticFunctionInput.FunctionCallNodeGuid);
+		FSelectKey SelectKey(ESelectKeySource::ModuleInput, ParsedStaticFunctionInput.InputVariable, NAME_None);
+		if (FunctionInputContext.Contains(SelectKey) == false)
+		{
+			if (ensureMsgf(ParsedStaticFunctionInput.DataOffset + ParsedStaticFunctionInput.InputVariable.GetSizeInBytes() <= Script.RapidIterationParameters.GetParameterDataArray().Num(), TEXT("Cached parameter offset not valid for parameter store.")))
+			{
+				FNiagaraVariable Temp = ParsedStaticFunctionInput.InputVariable;
+				Temp.SetData(Script.RapidIterationParameters.GetParameterData(ParsedStaticFunctionInput.DataOffset));
+				FSelectValue SelectValue = FTraversalBuilder::CreateSelectValue(Temp);
+				FunctionInputContext.Add(SelectKey, SelectValue);
+			}
+		}
+	}
+}
+
+struct FTraversedParameter
+{
+	FTraversedParameter()
+	{
+	}
+
+	FTraversedParameter(const FNiagaraVariableBase& InParameter, EParameterFlags InParameterReferenceFlags, bool bInParameterNodeEnabled, ETraversalStateFlags InTraversalStateFlags)
+		: Parameter(InParameter)
+		, ParameterReferenceFlags(InParameterReferenceFlags)
+		, bParameterNodeEnabled(bInParameterNodeEnabled)
+		, TraversalStateFlags(InTraversalStateFlags)
+	{
+	}
+
+	FORCEINLINE bool operator==(const FTraversedParameter& Other) const
+	{
+		return
+			Parameter == Other.Parameter &&
+			ParameterReferenceFlags == Other.ParameterReferenceFlags &&
+			bParameterNodeEnabled == Other.bParameterNodeEnabled &&
+			TraversalStateFlags == Other.TraversalStateFlags;
+	}
+
+	FNiagaraVariableBase Parameter;
+	EParameterFlags ParameterReferenceFlags = EParameterFlags::None;
+	bool bParameterNodeEnabled = false;
+	ETraversalStateFlags TraversalStateFlags = ETraversalStateFlags::None;
+};
+
+class FTraversalReadAndWriteParameters : public FTraversalData
+{
+public:
+	TArray<FTraversedParameter> ReadParameters;
+	TArray<FTraversedParameter> WriteParameters;
+};
+
+class FCollectReadAndWriteParameters : public ITraversalVisitor
+{
+public:
+	FCollectReadAndWriteParameters()
+	{
+	}
+
+	virtual const FGuid GetVisitorId() const override
+	{
+		return FGuid(0x8C395293, 0xB41843A2, 0x942ABE00, 0x9EE1BA93);
+	}
+
+	virtual TSharedRef<FTraversalData> CreateTraversalData() const override
+	{
+		return MakeShared<FTraversalReadAndWriteParameters>();
+	}
+
+	virtual void VisitNode(ETraversalStateFlags TraversalState, const FTraversalNode& Node, FTraversalData& TraversalData) const override
+	{
+		if (Node.ParameterData.IsSet() == false)
+		{
+			return;
+		}
+
+		FTraversalReadAndWriteParameters* ParameterData = static_cast<FTraversalReadAndWriteParameters*>(&TraversalData);
+		for (const FParameterRead& ReadParameter : Node.ParameterData->ReadParameterReferences)
+		{
+			ParameterData->ReadParameters.AddUnique(FTraversedParameter(ReadParameter.Parameter, ReadParameter.Flags, Node.bSourceNodeEnabled, TraversalState));
+		}
+		for (const FParameterWrite& WriteParameter : Node.ParameterData->WriteParameterReferences)
+		{
+			ParameterData->WriteParameters.AddUnique(FTraversedParameter(WriteParameter.Parameter, WriteParameter.Flags, Node.bSourceNodeEnabled, TraversalState));
+		}
+	}
+};
+
+void FTraversalCache::GetStackFunctionReadsInternal(
+	const FTopLevelScriptStaticContext& ScriptStaticContext,
+	const UNiagaraNodeFunctionCall& FunctionCallNode,
+	TArray<FNiagaraVariable>& OutReads,
+	TSet<FNiagaraVariable>& OutHiddenReads,
+	EStackFunctionReadFilterFlags FilterFlags,
+	FTraversalDebugData* TraversalDebugData,
+	FTraversalBuilderDebugData* BuilderDebugData) const
+{
+	TSharedPtr<const FTraversal> ScriptTraversal = GetTopLevelScriptTraversalInternal(FunctionCallNode, BuilderDebugData);
+	if (ScriptTraversal.IsValid() == false)
+	{
+		return;
+	}
+
+	FTraversalCallingContext CallingContext(ScriptStaticContext.GlobalContext, ScriptStaticContext.FunctionInputContexts);
+	CallingContext.ConnectionTraversalMode = EConnectionTraversalMode::All;
+
+	FCollectReadAndWriteParameters Visitor;
+	TSharedRef<const FTraversalData> TraversalData = ScriptTraversal->TraverseWithVisitor(CallingContext, Visitor, TraversalDebugData);
+
+	const FTraversalData* FunctionCallTraversalData = TraversalData->GetCalledFunctionTraversalData(FunctionCallNode.NodeGuid);
+	if (FunctionCallTraversalData == nullptr)
+	{
+		return;
+	}
+
+	const FTraversalReadAndWriteParameters* TraversalParameters = static_cast<const FTraversalReadAndWriteParameters*>(FunctionCallTraversalData);
+	TMap<FNiagaraVariable, TArray<ETraversalStateFlags>> ReadVariableToFlagList;
+	for (const FTraversedParameter& ReadParameter : TraversalParameters->ReadParameters)
+	{
+		if (HasFlag(ReadParameter.TraversalStateFlags, ETraversalStateFlags::UnconnectedRoot))
+		{
+			continue;
+		}
+
+		bool bIsInput = HasFlag(ReadParameter.ParameterReferenceFlags, EParameterFlags::ModuleInput);
+		bool bIsEnabled = HasFlag(ReadParameter.TraversalStateFlags, ETraversalStateFlags::CallerDisabled) == false && ReadParameter.bParameterNodeEnabled;
+		if ((HasFlag(FilterFlags, EStackFunctionReadFilterFlags::InputsOnly) == false || bIsInput) &&
+			(HasFlag(FilterFlags, EStackFunctionReadFilterFlags::EnabledOnly) == false || bIsEnabled))
+		{
+			OutReads.AddUnique(ReadParameter.Parameter);
+			TArray<ETraversalStateFlags>& ReadTraversalStateFlags = ReadVariableToFlagList.FindOrAdd(ReadParameter.Parameter);
+			ReadTraversalStateFlags.Add(ReadParameter.TraversalStateFlags);
+		}
+	}
+
+	for (FNiagaraVariable& Read : OutReads)
+	{
+		TArray<ETraversalStateFlags> ReadFlags = ReadVariableToFlagList[Read];
+		bool bIsVisible = ReadFlags.ContainsByPredicate([](const ETraversalStateFlags& Flags)
+		{
+			return HasFlag(Flags, ETraversalStateFlags::CulledBySwitch) == false;
+		});
+
+		if (bIsVisible == false)
+		{
+			OutHiddenReads.Add(Read);
+		}
+	}
+}
+
+struct FTraversedStaticSwitchInput
+{
+	FTraversedStaticSwitchInput(const FSelectKey& InSelectKey, ETraversalStateFlags InTraversalStateFlags)
+		: SelectKey(InSelectKey)
+		, TraversalStateFlags(InTraversalStateFlags)
+	{
+	}
+
+	FORCEINLINE bool operator==(const FTraversedStaticSwitchInput& Other) const
+	{
+		return SelectKey == Other.SelectKey && TraversalStateFlags == Other.TraversalStateFlags;
+	}
+
+	FSelectKey SelectKey;
+	ETraversalStateFlags TraversalStateFlags = ETraversalStateFlags::None;
+};
+
+class FTraversalStaticSwitchInputs : public FTraversalData
+{
+public:
+	TArray<FTraversedStaticSwitchInput> StaticSwitchInputs;
+};
+
+class FCollectStaticSwitchInputs : public ITraversalVisitor
+{
+public:
+	FCollectStaticSwitchInputs()
+	{
+	}
+
+	virtual const FGuid GetVisitorId() const override
+	{
+		return FGuid(0x92B6A255, 0x484B4323, 0x96B888C9, 0x39B6B8C4);
+	}
+
+	virtual TSharedRef<FTraversalData> CreateTraversalData() const override
+	{
+		return MakeShared<FTraversalStaticSwitchInputs>();
+	}
+
+	virtual void VisitNode(ETraversalStateFlags TraversalState, const FTraversalNode& Node, FTraversalData& TraversalData) const override
+	{
+		FTraversalStaticSwitchInputs* SwitchInputData = static_cast<FTraversalStaticSwitchInputs*>(&TraversalData);
+		if (Node.SelectData.IsSet())
+		{
+			if (Node.SelectData->SelectKey.IsValid() && Node.SelectData->SelectKey.Source == ESelectKeySource::FunctionCallNode)
+			{
+				SwitchInputData->StaticSwitchInputs.AddUnique(FTraversedStaticSwitchInput(Node.SelectData->SelectKey, TraversalState));
+			}
+		}
+		else if (Node.FunctionCallData.IsSet())
+		{
+			for (const FFunctionInputSelectValue& InputSelectValue : Node.FunctionCallData->InputSelectValues)
+			{
+				if (InputSelectValue.ConnectionPinId.IsSet() == false && InputSelectValue.LocalValue.IsSet() == false)
+				{
+					FSelectKey PropagatedSelectKey;
+					if (InputSelectValue.OptionalPropagatedNameOverride.IsSet())
+					{
+						FNiagaraVariableBase OverrideVariable(InputSelectValue.InputSelectKey.Variable.GetType(), InputSelectValue.OptionalPropagatedNameOverride.GetValue());
+						PropagatedSelectKey = FSelectKey(ESelectKeySource::FunctionCallNode, OverrideVariable, NAME_None);
+					}
+					else
+					{
+						PropagatedSelectKey = InputSelectValue.InputSelectKey;
+					}
+					SwitchInputData->StaticSwitchInputs.AddUnique(FTraversedStaticSwitchInput(PropagatedSelectKey, TraversalState));
+				}
+			}
+		}
+	}
+};
+
+void FTraversalCache::GetStackFunctionStaticInputsInternal(
+	const FTopLevelScriptStaticContext& ScriptStaticContext,
+	const UNiagaraNodeFunctionCall& FunctionCallNode,
+	TArray<FNiagaraVariable>& OutStaticInputs,
+	TSet<FNiagaraVariable>& OutHiddenStaticInputs,
+	FTraversalDebugData* TraversalDebugData,
+	FTraversalBuilderDebugData* BuilderDebugData) const
+{
+	TSharedPtr<const FTraversal> ScriptTraversal = GetTopLevelScriptTraversalInternal(FunctionCallNode, BuilderDebugData);
+	if (ScriptTraversal.IsValid() == false)
+	{
+		return;
+	}
+
+	FTraversalCallingContext CallingContext(ScriptStaticContext.GlobalContext, ScriptStaticContext.FunctionInputContexts);
+	CallingContext.ConnectionTraversalMode = EConnectionTraversalMode::All;
+
+	FCollectStaticSwitchInputs Visitor;
+	TSharedRef<const FTraversalData> TraversalData = ScriptTraversal->TraverseWithVisitor(CallingContext, Visitor, TraversalDebugData);
+
+	const FTraversalData* FunctionCallTraversalData = TraversalData->GetCalledFunctionTraversalData(FunctionCallNode.NodeGuid);
+	if (FunctionCallTraversalData == nullptr)
+	{
+		return;
+	}
+
+	const FTraversalStaticSwitchInputs* TraversalStaticSwitches = static_cast<const FTraversalStaticSwitchInputs*>(FunctionCallTraversalData);
+	TMap<FNiagaraVariable, TArray<ETraversalStateFlags>> StaticInputVariableToFlagList;
+	for (const FTraversedStaticSwitchInput& TraversedStaticInput : TraversalStaticSwitches->StaticSwitchInputs)
+	{
+		OutStaticInputs.AddUnique(TraversedStaticInput.SelectKey.Variable);
+		TArray<ETraversalStateFlags>& StaticInputTraversalStateFlags = StaticInputVariableToFlagList.FindOrAdd(TraversedStaticInput.SelectKey.Variable);
+		StaticInputTraversalStateFlags.Add(TraversedStaticInput.TraversalStateFlags);
+	}
+
+	for (FNiagaraVariable& StaticInput : OutStaticInputs)
+	{
+		TArray<ETraversalStateFlags> StaticInputFlags = StaticInputVariableToFlagList[StaticInput];
+		bool bIsVisible = StaticInputFlags.ContainsByPredicate([](const ETraversalStateFlags& Flags)
+		{
+			return HasFlag(Flags, ETraversalStateFlags::CulledBySwitch) == false && HasFlag(Flags, ETraversalStateFlags::UnconnectedRoot) == false;
+		});
+
+		if (bIsVisible == false)
+		{
+			OutHiddenStaticInputs.Add(StaticInput);
+		}
+	}
+}
+
+void FTraversalCache::GetStackFunctionStaticInputPinsInternal(
+	const FTopLevelScriptStaticContext& ScriptStaticContext,
+	const UNiagaraNodeFunctionCall& FunctionCallNode,
+	TArray<UEdGraphPin*>& OutStaticInputPins,
+	TSet<UEdGraphPin*>& OutHiddenStaticInputPins,
+	FTraversalDebugData* TraversalDebugData) const
+{
+	const UEdGraphSchema_Niagara* Schema = GetDefault<UEdGraphSchema_Niagara>();
+	TArray<FNiagaraVariable> StaticInputs;
+	TSet<FNiagaraVariable> HiddenStaticInputs;
+	GetStackFunctionStaticInputs(ScriptStaticContext, FunctionCallNode, StaticInputs, HiddenStaticInputs, TraversalDebugData);
+	for (FNiagaraVariable StaticInput : StaticInputs)
+	{
+		FEdGraphPinType PinType = Schema->TypeDefinitionToPinType(StaticInput.GetType());
+		for (UEdGraphPin* Pin : FunctionCallNode.Pins)
+		{
+			if (Pin->Direction != EEdGraphPinDirection::EGPD_Input)
+			{
+				continue;
+			}
+			if (Pin->PinName.IsEqual(StaticInput.GetName()) && Pin->PinType == PinType)
+			{
+				OutStaticInputPins.Add(Pin);
+				if (HiddenStaticInputs.Contains(StaticInput))
+				{
+					OutHiddenStaticInputPins.Add(Pin);
+				}
+				break;
+			}
+		}
+	}
+}
+
+void FTraversalCache::InvalidateCacheInternal() const
+{
+	CachedTraversals.Empty();
+	CachedParsedRapidIterationParameterData.Empty();
+	SoftObjectPathToScriptCache.Empty();
+}
+
+EStackFunctionReadFilterFlags operator|(
+	EStackFunctionReadFilterFlags FlagsA,
+	EStackFunctionReadFilterFlags FlagsB)
+{
+	return (EStackFunctionReadFilterFlags)((uint8)FlagsA | (uint8)(FlagsB));
+}
+
+EStackFunctionReadFilterFlags operator&(
+	EStackFunctionReadFilterFlags FlagsA,
+	EStackFunctionReadFilterFlags FlagsB)
+{
+	return (EStackFunctionReadFilterFlags)((uint8)FlagsA & (uint8)(FlagsB));
+}
+
+EStackFunctionReadFilterFlags operator~(
+	EStackFunctionReadFilterFlags Flags)
+{
+	return (EStackFunctionReadFilterFlags)~(uint8)Flags;
+}
+
+} // UE::Niagara::TraversalCache
\ No newline at end of file
diff --git a/Engine/Plugins/FX/Niagara/Source/NiagaraEditor/Private/TraversalCache/TraversalCache.h b/Engine/Plugins/FX/Niagara/Source/NiagaraEditor/Private/TraversalCache/TraversalCache.h
new file mode 100644
--- /dev/null
+++ b/Engine/Plugins/FX/Niagara/Source/NiagaraEditor/Private/TraversalCache/TraversalCache.h
@@ -0,0 +1,232 @@
+// Copyright Epic Games, Inc. All Rights Reserved.
+
+#pragma once
+
+#include "Misc/Guid.h"
+#include "NiagaraCommon.h"
+#include "NiagaraTypes.h"
+#include "TraversalCache/TraversalShared.h"
+#include "UObject/ObjectKey.h"
+
+class UEdGraphPin;
+class UNiagaraNode;
+class UNiagaraNodeFunctionCall;
+class UNiagaraNodeOutput;
+class UNiagaraScript;
+
+namespace UE::Niagara::TraversalCache
+{
+
+class  FTraversal;
+struct FTraversalBuilderDebugData;
+struct FTraversalDebugData;
+
+enum class EStackFunctionReadFilterFlags : uint8
+{
+	None = 0,
+	InputsOnly = 1,
+	EnabledOnly = 2
+};
+
+class FTraversalCache : public TSharedFromThis<FTraversalCache>
+{
+public:
+	void Initialize();
+
+	static TSharedPtr<const FTraversal> GetScriptAssetTraversal(const FSoftObjectPath& ScriptSoftObjectPath, const FGuid& ScriptVersion, FTraversalBuilderDebugData* BuilderDebugData = nullptr);
+
+	static TSharedPtr<const FTraversal> GetScriptAssetTraversal(const FSoftObjectPath& ScriptSoftObjectPath, const FGuid& ScriptVersion, FTraversalDebugData* TraversalDebugData = nullptr);
+
+	static TSharedPtr<const FTraversal> GetTopLevelScriptTraversal(const UNiagaraScript& Script, FTraversalBuilderDebugData* BuilderDebugData = nullptr);
+
+	static TSharedPtr<const FTraversal> GetTopLevelScriptTraversal(const UNiagaraScript& Script, FTraversalDebugData* TraversalDebugData = nullptr);
+
+	static void GetTopLevelScriptFunctionInputContexts(const UNiagaraScript& Script, TMap<FGuid, FStaticContext>& OutFunctionInputContexts, FTraversalDebugData* TraversalDebugData = nullptr);
+
+	NIAGARAEDITOR_API static void GetStackFunctionReads(
+		const FTopLevelScriptStaticContext& ScriptStaticContext,
+		const UNiagaraNodeFunctionCall& FunctionCallNode,
+		TArray<FNiagaraVariable>& OutReads,
+		TSet<FNiagaraVariable>& OutHiddenReads,
+		EStackFunctionReadFilterFlags FilterFlags = EStackFunctionReadFilterFlags::InputsOnly,
+		FTraversalDebugData* TraversalDebugData = nullptr);
+
+	NIAGARAEDITOR_API static void GetStackFunctionStaticInputs(
+		const FTopLevelScriptStaticContext& ScriptStaticContext,
+		const UNiagaraNodeFunctionCall& FunctionCallNode,
+		TArray<FNiagaraVariable>& OutStaticInputs,
+		TSet<FNiagaraVariable>& OutHiddenStaticInputs,
+		FTraversalDebugData* TraversalDebugData = nullptr);
+
+	NIAGARAEDITOR_API static void GetStackFunctionStaticInputPins(
+		const FTopLevelScriptStaticContext& ScriptStaticContext,
+		const UNiagaraNodeFunctionCall& FunctionCallNode,
+		TArray<UEdGraphPin*>& OutStaticInputPins,
+		TSet<UEdGraphPin*>& OutHiddenStaticInputPins,
+		FTraversalDebugData* TraversalDebugData = nullptr);
+
+	NIAGARAEDITOR_API static void InvalidateCache();
+
+private:
+	void OnScriptApplied(UNiagaraScript* InScript, FGuid InVersionGuid);
+	
+	static FTraversalCache& GetInstance();
+
+	const UNiagaraScript* GetScriptFromPath(const FSoftObjectPath& Path) const;
+
+	TSharedPtr<const FTraversal> GetScriptAssetTraversalInternal(const FSoftObjectPath& ScriptSoftObjectPath, const FGuid& ScriptVersion, FTraversalBuilderDebugData* BuilderDebugData) const;
+
+	TSharedPtr<const FTraversal> GetTopLevelScriptTraversalInternal(const UNiagaraScript& Script, FTraversalBuilderDebugData* BuilderDebugData) const;
+
+	TSharedPtr<const FTraversal> GetTopLevelScriptTraversalInternal(const UNiagaraNode& StackNode, FTraversalBuilderDebugData* BuilderDebugData) const;
+
+	void GetTopLevelScriptFunctionInputContextsInternal(const UNiagaraScript& Script, TMap<FGuid, FStaticContext>& OutFunctionInputContexts, FTraversalDebugData* TraversalDebugData) const;
+
+	void GetStackFunctionReadsInternal(
+		const FTopLevelScriptStaticContext& ScriptStaticContext,
+		const UNiagaraNodeFunctionCall& FunctionCallNode,
+		TArray<FNiagaraVariable>& OutInputs,
+		TSet<FNiagaraVariable>& OutHiddenInputs,
+		EStackFunctionReadFilterFlags FilterFlags,
+		FTraversalDebugData* TraversalDebugData,
+		FTraversalBuilderDebugData* BuilderDebugData) const; 
+
+	void GetStackFunctionStaticInputsInternal(
+		const FTopLevelScriptStaticContext& ScriptStaticContext,
+		const UNiagaraNodeFunctionCall& FunctionCallNode,
+		TArray<FNiagaraVariable>& OutStaticInputs,
+		TSet<FNiagaraVariable>& OutHiddenStaticInputs,
+		FTraversalDebugData* TraversalDebugData,
+		FTraversalBuilderDebugData* BuilderDebugData) const;
+
+	void GetStackFunctionStaticInputPinsInternal(
+		const FTopLevelScriptStaticContext& ScriptStaticContext,
+		const UNiagaraNodeFunctionCall& FunctionCallNode,
+		TArray<UEdGraphPin*>& OutStaticInputPins,
+		TSet<UEdGraphPin*>& OutHiddenStaticInputPins,
+		FTraversalDebugData* TraversalDebugData) const;
+
+	void InvalidateCacheInternal() const;
+
+private:
+	struct FScriptTraversalCacheKey
+	{
+		FScriptTraversalCacheKey()
+		{
+		}
+
+		static FScriptTraversalCacheKey FromScriptAsset(
+			const FObjectKey& InScriptObjectKey,
+			const FGuid& InScriptVersion)
+		{
+			return FScriptTraversalCacheKey(
+				InScriptObjectKey,
+				InScriptVersion,
+				ENiagaraScriptUsage::Function,
+				FGuid());
+		}
+
+		static FScriptTraversalCacheKey FromTopLevelScript(
+			const FObjectKey& InGraphObjectKey,
+			ENiagaraScriptUsage InScriptUsage,
+			const FGuid& InScriptUsageId)
+		{
+			return FScriptTraversalCacheKey(
+				InGraphObjectKey,
+				FGuid(),
+				InScriptUsage,
+				InScriptUsageId);
+		}
+
+		FORCEINLINE bool operator==(const FScriptTraversalCacheKey& Other) const
+		{
+			return 
+				OwningObjectKey == Other.OwningObjectKey &&
+				ScriptVersion == Other.ScriptVersion &&
+				ScriptUsage == Other.ScriptUsage &&
+				ScriptUsageId == Other.ScriptUsageId;
+		}
+
+		friend FORCEINLINE uint32 GetTypeHash(const FScriptTraversalCacheKey& Key)
+		{
+			return Key.Hash;
+		}
+
+		const FObjectKey OwningObjectKey;
+		const FGuid ScriptVersion;
+		const ENiagaraScriptUsage ScriptUsage = ENiagaraScriptUsage::Function;
+		const FGuid ScriptUsageId;
+		const int32 Hash = 0;
+
+	private:
+		FScriptTraversalCacheKey(const FObjectKey& InOwningObjectKey, const FGuid& InScriptVersion, ENiagaraScriptUsage InScriptUsage, const FGuid& InScriptUsageId)
+			: OwningObjectKey(InOwningObjectKey)
+			, ScriptVersion(InScriptVersion)
+			, ScriptUsage(InScriptUsage)
+			, ScriptUsageId(InScriptUsageId)
+			, Hash(ComputeHash(OwningObjectKey, ScriptVersion, ScriptUsage, ScriptUsageId))
+		{
+		}
+
+		static uint32 ComputeHash(const FObjectKey& OwningObjectKey, const FGuid& ScriptVersion, ENiagaraScriptUsage ScriptUsage, const FGuid& ScriptUsageId)
+		{
+			uint32 Hash = GetTypeHash(OwningObjectKey);
+			Hash = HashCombine(Hash, GetTypeHash(ScriptVersion));
+			Hash = HashCombine(Hash, GetTypeHash(ScriptUsage));
+			Hash = HashCombine(Hash, GetTypeHash(ScriptUsageId));
+			return Hash;
+		}
+	};
+
+	struct FScriptTraversalCacheValue
+	{
+		FScriptTraversalCacheValue(FGuid InGraphChangeId, TSharedRef<const FTraversal> InScriptTraversal)
+			: GraphChangeId(InGraphChangeId)
+			, ScriptTraversal(InScriptTraversal)
+		{
+		}
+
+		FGuid GraphChangeId;
+		TSharedRef<const FTraversal> ScriptTraversal;
+	};
+
+	struct FParsedStaticFunctionInput
+	{
+		FParsedStaticFunctionInput(const FGuid& InFunctionCallNodeGuid, const FNiagaraVariableBase& InInputVariable, int32 InDataOffset)
+			: FunctionCallNodeGuid(InFunctionCallNodeGuid)
+			, InputVariable(InInputVariable)
+			, DataOffset(InDataOffset)
+		{
+		}
+
+		FGuid FunctionCallNodeGuid;
+		FNiagaraVariableBase InputVariable;
+		int32 DataOffset = INDEX_NONE;
+	};
+
+	struct FParsedRapidIterationParameterData
+	{
+		int32 ParameterStoreLayoutVersion = INDEX_NONE;
+		TArray<FParsedStaticFunctionInput> StaticFunctionInputs;
+	};
+
+	static FTraversalCache* Instance;
+	bool bInitialized = false;
+
+	mutable TMap<FScriptTraversalCacheKey, FScriptTraversalCacheValue> CachedTraversals;
+	mutable TMap<FScriptTraversalCacheKey, FParsedRapidIterationParameterData> CachedParsedRapidIterationParameterData;
+	mutable TMap<FSoftObjectPath, TWeakObjectPtr<const UNiagaraScript>> SoftObjectPathToScriptCache;
+};
+
+EStackFunctionReadFilterFlags operator|(
+	EStackFunctionReadFilterFlags FlagsA,
+	EStackFunctionReadFilterFlags FlagsB);
+
+EStackFunctionReadFilterFlags operator&(
+	EStackFunctionReadFilterFlags FlagsA,
+	EStackFunctionReadFilterFlags FlagsB);
+
+EStackFunctionReadFilterFlags operator~(
+	EStackFunctionReadFilterFlags Flags);
+
+} // UE::Niagara::TraversalCache
\ No newline at end of file
diff --git a/Engine/Plugins/FX/Niagara/Source/NiagaraEditor/Private/TraversalCache/TraversalNode.cpp b/Engine/Plugins/FX/Niagara/Source/NiagaraEditor/Private/TraversalCache/TraversalNode.cpp
new file mode 100644
--- /dev/null
+++ b/Engine/Plugins/FX/Niagara/Source/NiagaraEditor/Private/TraversalCache/TraversalNode.cpp
@@ -0,0 +1,38 @@
+// Copyright Epic Games, Inc. All Rights Reserved.
+
+#include "TraversalCache/TraversalNode.h"
+
+namespace UE::Niagara::TraversalCache
+{
+
+const FSelectInputData* FSelectData::FindInputDataForSelectValue(const FSelectValue& SelectValue) const
+{
+	for (const FSelectInputData& SelectInputData : InputData)
+	{
+		if (SelectInputData.SelectValue == SelectValue)
+		{
+			return &SelectInputData;
+		}
+	}
+	return nullptr;
+}
+
+const FSelectInputData* FSelectData::FindInputDataForConnectionPinId(const FGuid& ConnectionPinId) const
+{
+	for (const FSelectInputData& SelectInputData : InputData)
+	{
+		if (SelectInputData.ConnectionPinId.IsSet() && SelectInputData.ConnectionPinId.GetValue() == ConnectionPinId)
+		{
+			return &SelectInputData;
+		}
+	}
+	return nullptr;
+}
+
+const FTraversalNode* FTraversalNode::GetConnectedNodeByPinId(const FGuid& PinId) const
+{
+	const FConnection* MatchingConnection = Connections.FindByPredicate([PinId](const FConnection& Connection) { return Connection.PinId == PinId; });
+	return MatchingConnection != nullptr ? &MatchingConnection->GetNode() : nullptr;
+}
+
+} // UE::Niagara::TraversalCache
\ No newline at end of file
diff --git a/Engine/Plugins/FX/Niagara/Source/NiagaraEditor/Private/TraversalCache/TraversalNode.h b/Engine/Plugins/FX/Niagara/Source/NiagaraEditor/Private/TraversalCache/TraversalNode.h
new file mode 100644
--- /dev/null
+++ b/Engine/Plugins/FX/Niagara/Source/NiagaraEditor/Private/TraversalCache/TraversalNode.h
@@ -0,0 +1,169 @@
+// Copyright Epic Games, Inc. All Rights Reserved.
+
+#pragma once
+
+#include "NiagaraTypes.h"
+#include "TraversalCache/TraversalShared.h"
+
+namespace UE::Niagara::TraversalCache
+{
+
+struct FFunctionInputData
+{
+	FSelectKey InputSelectKey;
+	TOptional<FSelectValue> LocalValue;
+};
+
+struct FSelectInputData
+{
+	FSelectValue SelectValue;
+	TOptional<FSelectValue> LocalValue;
+	TOptional<FGuid> ConnectionPinId;
+};
+
+struct FSelectData
+{
+	const FSelectInputData* FindInputDataForSelectValue(const FSelectValue& SelectValue) const;
+	const FSelectInputData* FindInputDataForConnectionPinId(const FGuid& ConnectionPinId) const;
+
+	ESelectMode SelectMode = ESelectMode::None;
+	FSelectKey SelectKey;
+	FGuid SelectConnectionPinId;
+	TArray<FSelectInputData> InputData;
+};
+
+struct FParameterReference
+{
+	FParameterReference()
+	{
+	}
+
+	FParameterReference(const FNiagaraVariableBase& InParameter, EParameterFlags InFlags)
+		: Parameter(InParameter)
+		, Flags(InFlags)
+	{
+	}
+
+	FNiagaraVariableBase Parameter;
+	EParameterFlags Flags = EParameterFlags::None;
+	TOptional<FGuid> OptionalConnectionPinId;
+	TOptional<FSelectValue> OptionalLocalSelectValue;
+	TOptional<FName> OptionalNamespaceModifier;
+};
+
+struct FParameterRead : public FParameterReference
+{
+	FParameterRead()
+	{
+	}
+
+	FParameterRead(const FNiagaraVariableBase& InParameter, EParameterFlags InFlags)
+		: FParameterReference(InParameter, InFlags)
+	{
+	}
+
+	bool bIsDiscoverRead = false;
+	TOptional<FNiagaraVariant> OptionalDefaultValue;
+	TOptional<FName> OptionalDefaultBinding;
+};
+
+struct FParameterWrite : public FParameterReference
+{
+	FParameterWrite()
+	{
+	}
+
+	FParameterWrite(const FNiagaraVariableBase& InParameter, EParameterFlags InFlags)
+		: FParameterReference(InParameter, InFlags)
+	{
+	}
+
+	TOptional<FGuid> OptionalTargetFunctionCallNodeGuid;
+};
+
+struct FParameterData
+{
+	TArray<FParameterRead> ReadParameterReferences;
+	TArray<FParameterWrite> WriteParameterReferences;
+	FGuid ExecutionConnectionPinId;
+};
+
+struct FFunctionInputSelectValue
+{
+	FSelectKey InputSelectKey;
+	TOptional<FSelectValue> LocalValue;
+	TOptional<FGuid> ConnectionPinId;
+	TOptional<FName> OptionalPropagatedNameOverride;
+};
+
+struct FFunctionCallData
+{
+	FScriptReference FunctionScriptReference;
+	FName FunctionCallName;
+	TArray<FFunctionInputSelectValue> InputSelectValues;
+	ENiagaraFunctionDebugState DebugState = ENiagaraFunctionDebugState::NoDebug;
+	FGuid ExecutionConnectionPinId;
+};
+
+struct FStaticOpInputData
+{
+	TOptional<FSelectValue> LocalValue;
+	TOptional<FGuid> ConnectionPinId;
+};
+
+struct FStaticOpData
+{
+	FStaticOpData()
+		: OpName(NAME_None)
+	{
+	}
+
+	FName OpName;
+	TArray<FStaticOpInputData> InputData;
+};
+
+class FTraversalNode
+{
+	friend class FTraversalBuilder;
+
+public:
+	struct FConnection
+	{
+		friend class FTraversalBuilder;
+
+		FConnection(const FGuid& InPinId, TSharedRef<FTraversalNode> InNode)
+			: PinId(InPinId)
+			, Node(InNode)
+		{
+		}
+
+		FORCEINLINE bool operator==(const FConnection& Other) const
+		{
+			return PinId == Other.PinId && Node == Other.Node;
+		}
+
+		const FTraversalNode& GetNode() const { return Node.Get(); }
+
+		FGuid PinId;
+
+	private:
+		TSharedRef<FTraversalNode> Node;
+	};
+
+	const FTraversalNode* GetConnectedNodeByPinId(const FGuid& PinId) const;
+
+	FGuid SourceNodeGuid = FGuid();
+	FName SourceNodeTypeName;
+	bool bSourceNodeEnabled = true;
+
+	bool bIsNoop = false;
+	TOptional<const FFunctionInputData> FunctionInputData;
+	TOptional<const FSelectData> SelectData;
+	TOptional<const FParameterData> ParameterData;
+	TOptional<const FFunctionCallData> FunctionCallData;
+	TOptional<const FStaticOpData> StaticOpData;
+
+	TArray<FConnection> Connections;
+};
+
+} // UE::Niagara::TraversalCache
\ No newline at end of file
diff --git a/Engine/Plugins/FX/Niagara/Source/NiagaraEditor/Private/TraversalCache/TraversalShared.cpp b/Engine/Plugins/FX/Niagara/Source/NiagaraEditor/Private/TraversalCache/TraversalShared.cpp
new file mode 100644
--- /dev/null
+++ b/Engine/Plugins/FX/Niagara/Source/NiagaraEditor/Private/TraversalCache/TraversalShared.cpp
@@ -0,0 +1,233 @@
+// Copyright Epic Games, Inc. All Rights Reserved.
+
+#pragma once
+
+#include "TraversalCache/TraversalShared.h"
+
+#include "TraversalCache/TraversalBuilder.h"
+
+namespace UE::Niagara::TraversalCache
+{
+
+const FSelectValue& FSelectValue::GetBoolTrue()
+{
+	static const FSelectValue BoolTrue = FTraversalBuilder::CreateSelectValue(true);
+	return BoolTrue;
+}
+
+const FSelectValue& FSelectValue::GetBoolFalse()
+{
+	static const FSelectValue BoolFalse = FTraversalBuilder::CreateSelectValue(false);
+	return BoolFalse;
+}
+
+const FSelectValue& FSelectValue::GetDebugStateNoDebug()
+{
+	static const FSelectValue DebugStateNoDebug = FTraversalBuilder::CreateSelectValue(FNiagaraTypeDefinition::GetFunctionDebugStateEnum(), (int32)ENiagaraFunctionDebugState::NoDebug);
+	return DebugStateNoDebug;
+}
+
+const FSelectValue& FSelectValue::GetDebugStateBasic()
+{
+	static const FSelectValue DebugStateBasic = FTraversalBuilder::CreateSelectValue(FNiagaraTypeDefinition::GetFunctionDebugStateEnum(), (int32)ENiagaraFunctionDebugState::Basic);
+	return DebugStateBasic;
+}
+
+TOptional<FSelectValue> FStaticContext::GetSelectValue(const FSelectKey& Key) const
+{
+	const FSelectValue* SelectValuePtr = Find(Key);
+	return SelectValuePtr != nullptr
+		? TOptional<FSelectValue>(*SelectValuePtr)
+		: TOptional<FSelectValue>();
+}
+
+TOptional<FSelectValue> FTraversalLocalContext::GetSelectValue(const FTraversalCallingContext& CallingContext, const FSelectKey& SelectKey) const
+{
+	TOptional<FSelectValue> SelectValue;
+	
+	StaticContext.GetSelectValue(SelectKey);
+	if (SelectValue.IsSet())
+	{
+		return SelectValue;
+	}
+
+	SelectValue = CallingContext.FunctionCallStaticContext.GetSelectValue(SelectKey);
+	if (SelectValue.IsSet())
+	{
+		return SelectValue;
+	}
+
+	SelectValue = CallingContext.CallingStaticContext.GetSelectValue(SelectKey);
+	if (SelectValue.IsSet())
+	{
+		return SelectValue;
+	}
+
+	return CallingContext.GlobalStaticContext.GetSelectValue(SelectKey);
+}
+
+const FTraversalData* FTraversalData::GetCalledFunctionTraversalData(FGuid FunctionCallNodeGuid) const
+{
+	for (const FFunctionCallTraversalData& FunctionCallTraversalData : CalledFunctionTraversalData)
+	{
+		if (FunctionCallTraversalData.FunctionCallNodeGuid == FunctionCallNodeGuid)
+		{
+			return &FunctionCallTraversalData.FunctionCallTraversalData.Get();
+		}
+	}
+	return nullptr;
+}
+
+EParameterFlags operator|(
+	EParameterFlags FlagsA,
+	EParameterFlags FlagsB)
+{
+	return (EParameterFlags)((uint16)FlagsA | (uint16)FlagsB);
+}
+
+EParameterFlags operator&(
+	EParameterFlags FlagsA,
+	EParameterFlags FlagsB)
+{
+	return (EParameterFlags)((uint16)FlagsA & (uint16)FlagsB);
+}
+
+EParameterFlags operator~(
+	EParameterFlags Flags)
+{
+	return (EParameterFlags)(~(uint16)(Flags));
+}
+
+ETraversalStateFlags operator|(
+	ETraversalStateFlags FlagsA,
+	ETraversalStateFlags FlagsB)
+{
+	return (ETraversalStateFlags)((uint8)FlagsA | (uint8)FlagsB);
+}
+
+ETraversalStateFlags operator&(
+	ETraversalStateFlags FlagsA,
+	ETraversalStateFlags FlagsB)
+{
+	return (ETraversalStateFlags)((uint8)FlagsA & (uint8)FlagsB);
+}
+
+ETraversalStateFlags operator~(
+	ETraversalStateFlags Flags)
+{
+	return (ETraversalStateFlags)(~(uint8)Flags);
+}
+
+FTraversalCallingContext::FContextHash FTraversalCallingContext::GenerateHash() const
+{
+	auto AddStaticContextToHash = [&](FMD5& Hash, const FStaticContext& StaticContext)
+	{
+		TArray<const TPair<FSelectKey, FSelectValue>*> KeysAndValues;
+		KeysAndValues.Reserve(StaticContext.Num());
+		for (const TPair<FSelectKey, FSelectValue>& Pair : StaticContext)
+		{
+			KeysAndValues.Add(&Pair);
+		}
+
+		KeysAndValues.Sort([](const TPair<FSelectKey, FSelectValue>& PairA, const TPair<FSelectKey, FSelectValue>& PairB)
+		{
+			return PairA.Key.Variable.GetName().FastLess(PairB.Key.Variable.GetName());
+		});
+
+		for (const TPair<FSelectKey, FSelectValue>* Pair : KeysAndValues)
+		{
+			uint64 KeyNameUnstableInt = Pair->Key.Variable.GetName().ToUnstableInt();
+			Hash.Update((const uint8*)&KeyNameUnstableInt, sizeof(uint64));
+			Hash.Update((const uint8*)&Pair->Value.NumericValue, sizeof(int32));
+		}
+
+		int32 PairCount = StaticContext.Num();
+		Hash.Update((uint8*)&PairCount, sizeof(int32));
+	};
+
+	auto AddFunctionStaticInputContextsToHash = [&](FMD5& Hash, const TMap<FGuid, FStaticContext>& FunctionStaticInputContexts)
+	{
+		TArray<const TPair<FGuid, FStaticContext>*> ContextPairs;
+		ContextPairs.Reserve(FunctionStaticInputContexts.Num());
+		for (const TPair<FGuid, FStaticContext>& ContextPair : FunctionStaticInputContexts)
+		{
+			ContextPairs.Add(&ContextPair);
+		}
+
+		ContextPairs.Sort([](const TPair<FGuid, FStaticContext>& PairA, const TPair<FGuid, FStaticContext>& PairB)
+		{
+			return PairA.Key < PairB.Key;
+		});
+
+		for (const TPair<FGuid, FStaticContext>* ContextPair : ContextPairs)
+		{
+			Hash.Update((const uint8*)&ContextPair->Key, sizeof(FGuid));
+			AddStaticContextToHash(Hash, ContextPair->Value);
+		}
+
+		int32 ContextPairCount = FunctionStaticInputContexts.Num();
+		Hash.Update((uint8*)&ContextPairCount, sizeof(int32));
+	};
+
+	FMD5 Hash;
+	AddStaticContextToHash(Hash, GlobalStaticContext);
+	AddFunctionStaticInputContextsToHash(Hash, ExternalFunctionInputStaticContexts);
+	AddStaticContextToHash(Hash, CallingStaticContext);
+	AddStaticContextToHash(Hash, FunctionCallStaticContext);
+	Hash.Update((uint8*)&ConnectionTraversalMode, sizeof(EConnectionTraversalMode));
+	Hash.Update((uint8*)&TraversalState, sizeof(ETraversalStateFlags));
+
+	FContextHash ContextHash;
+	Hash.Final((uint8*)&ContextHash.HashData);
+	return ContextHash;
+}
+
+void FTraversalBuilderDebugData::AddMultipleInputConnection(
+	const TPair<FGuid, FGuid>& TargetNodeGuidAndPinId,
+	const TArray<TPair<FGuid, FGuid>>& ConnectedNodeGuidsAndPinIds)
+{
+	MultipleInputConnections.AddUnique(FMultipleInputConnection(TargetNodeGuidAndPinId, ConnectedNodeGuidsAndPinIds));
+}
+
+void FTraversalBuilderDebugData::AddUntrimmedNoop(const FGuid& NodeGuid, FName SourceNodeTypeName)
+{
+	UntrimmedNoops.AddUnique(FUntrimmedNoop(NodeGuid, SourceNodeTypeName));
+}
+
+void FTraversalBuilderDebugData::AddUnresolvedSelectOutput(const FGuid& NodeGuid, const FGuid& PinGuid)
+{
+	UnresolvedSelectOutputs.AddUnique(FUnresolvedSelectOutput(NodeGuid, PinGuid));
+}
+
+void FTraversalDebugData::AddUnresolvedSelect(const TArray<FGuid>& InFunctionCallStack, const FGuid& InSourceNodeGuid, const FSelectKey& InSelectKey)
+{
+	UnresolvedSelects.AddUnique(FUnresolvedSelect(InFunctionCallStack, InSourceNodeGuid, InSelectKey));
+}
+
+void FTraversalDebugData::AddUnresolvedSelectInput(const TArray<FGuid>& InFunctionCallStack, const FGuid& InSourceNodeGuid, const FSelectKey& InSelectKey, const FSelectValue& InResolvedSelectValue)
+{
+	UnresolvedSelectInputs.AddUnique(FUnresolvedSelectInput(InFunctionCallStack, InSourceNodeGuid, InSelectKey, InResolvedSelectValue));
+}
+
+void FTraversalDebugData::AddUnresolvedRead(const TArray<FGuid>& InFunctionCallStack, const FGuid& InSourceNodeGuid, const FNiagaraVariableBase& InReadParameter)
+{
+	UnresolvedReads.AddUnique(FUnresolvedRead(InFunctionCallStack, InSourceNodeGuid, InReadParameter));
+}
+
+void FTraversalDebugData::AddUnresolvedStaticOp(const TArray<FGuid>& InFunctionCallStack, const FGuid& InSourceNodeGuid, FName InOpName, int32 InUnresolvedPinIndex)
+{
+	UnresolvedStaticOps.AddUnique(FUnresolvedStaticOp(InFunctionCallStack, InSourceNodeGuid, InOpName, InUnresolvedPinIndex));
+}
+
+void FTraversalDebugData::AddUnresolvedFunctionInput(const TArray<FGuid>& InFunctionCallStack, const FGuid& InSourceNodeGuid, const FNiagaraVariableBase& InFunctionInput)
+{
+	UnresolvedReads.AddUnique(FUnresolvedRead(InFunctionCallStack, InSourceNodeGuid, InFunctionInput));
+}
+
+void FTraversalDebugData::AddBuilderDebugOutput(const FSoftObjectPath InGraphPath, const ENiagaraScriptUsage InScriptUsage,
+	const FGuid& InScriptUsageId, const FGuid& InScriptVersion, const FTraversalBuilderDebugData& InBuilderDebugData)
+{
+	BuilderDebugOutputs.AddUnique(FBuilderDebugOutput(InGraphPath, InScriptUsage, InScriptUsageId, InScriptVersion, InBuilderDebugData));
+}
+
+}
\ No newline at end of file
diff --git a/Engine/Plugins/FX/Niagara/Source/NiagaraEditor/Private/TraversalCache/TraversalShared.h b/Engine/Plugins/FX/Niagara/Source/NiagaraEditor/Private/TraversalCache/TraversalShared.h
new file mode 100644
--- /dev/null
+++ b/Engine/Plugins/FX/Niagara/Source/NiagaraEditor/Private/TraversalCache/TraversalShared.h
@@ -0,0 +1,554 @@
+// Copyright Epic Games, Inc. All Rights Reserved.
+
+#pragma once
+
+#include "Misc/Optional.h"
+#include "NiagaraCommon.h"
+#include "NiagaraTypes.h"
+#include "NiagaraVariant.h"
+#include "UObject/SoftObjectPath.h"
+
+namespace UE::Niagara::TraversalCache
+{
+
+class FTraversalNode;
+
+enum class ESelectMode : uint8
+{
+	None,
+	Value,
+	Connection
+};
+
+enum class ESelectKeySource : uint8
+{
+	None,
+	ExternalConstant,
+	Attribute,
+	ModuleInput,
+	ModuleLocal,
+	ModuleOutput,
+	FunctionCallNode,
+};
+
+struct FSelectKey
+{
+	FSelectKey()
+		: Source(ESelectKeySource::None)
+	{
+	}
+
+	FSelectKey(const ESelectKeySource& InSource, const FNiagaraVariableBase& InVariable, FName InNamespaceModifier)
+		: Source(InSource)
+		, Variable(InVariable)
+		, NamespaceModifier(InNamespaceModifier)
+	{
+	}
+
+	bool IsValid() const { return Source != ESelectKeySource::None && Variable.IsValid(); }
+
+	FORCEINLINE bool operator==(const FSelectKey& Other) const
+	{
+		return Source == Other.Source && Variable == Other.Variable && NamespaceModifier == Other.NamespaceModifier;
+	}
+
+	friend FORCEINLINE uint32 GetTypeHash(const FSelectKey& Key)
+	{
+		uint32 Hash = GetTypeHash(Key.Source);
+		Hash = HashCombine(Hash, GetTypeHash(Key.Variable));
+		Hash = HashCombine(Hash, GetTypeHash(Key.NamespaceModifier));
+		return Hash;
+	}
+
+	ESelectKeySource Source;
+	FNiagaraVariableBase Variable;
+	FName NamespaceModifier;
+};
+
+struct FSelectValue
+{
+	FSelectValue()
+		: NumericValue(0)
+	{
+	}
+
+	FSelectValue(int32 InNumericValue, FName InDebugName)
+		: NumericValue(InNumericValue)
+		, DebugName(InDebugName)
+	{
+	}
+
+	FORCEINLINE bool operator==(const FSelectValue& Other) const
+	{
+		return NumericValue == Other.NumericValue;
+	}
+
+	static const FSelectValue& GetBoolTrue();
+	static const FSelectValue& GetBoolFalse();
+	static const FSelectValue& GetDebugStateNoDebug();
+	static const FSelectValue& GetDebugStateBasic();
+
+	int32 NumericValue;
+	FName DebugName;
+};
+
+struct FStaticContext : public TMap<FSelectKey, FSelectValue>
+{
+	TOptional<FSelectValue> GetSelectValue(const FSelectKey& Key) const;
+};
+
+struct FTopLevelScriptStaticContext
+{
+	FStaticContext GlobalContext;
+	TMap<FGuid, FStaticContext> FunctionInputContexts;
+};
+
+enum class EParameterFlags : uint16
+{
+	None = 0,
+	ModuleInput = 1,
+	ModuleLocal = 2,
+	ModuleOutput = 4,
+	Attribute = 8,
+	Transient = 16,
+	External = 32,
+	NamespaceUnknown = 64,
+	NamespaceModifier = 128,
+	InvalidParameterName = 4096
+};
+
+EParameterFlags operator|(
+	EParameterFlags FlagsA,
+	EParameterFlags FlagsB);
+
+EParameterFlags operator&(
+	EParameterFlags FlagsA,
+	EParameterFlags FlagsB);
+
+EParameterFlags operator~(
+	EParameterFlags Flags);
+
+enum class EConnectionTraversalMode
+{
+	All,
+	MatchingOnly
+};
+
+enum class ETraversalStateFlags : uint8
+{
+	None = 0,
+	CulledBySwitch = 1,
+	CallerDisabled = 2,
+	UnconnectedRoot = 4
+};
+
+ETraversalStateFlags operator|(
+	ETraversalStateFlags FlagsA,
+	ETraversalStateFlags FlagsB);
+
+ETraversalStateFlags operator&(
+	ETraversalStateFlags FlagsA,
+	ETraversalStateFlags FlagsB);
+
+ETraversalStateFlags operator~(
+	ETraversalStateFlags Flags);
+
+struct FTraversalCallingContext
+{
+	struct FContextHash
+	{
+		FORCEINLINE bool operator==(const FContextHash& Other) const
+		{
+			return
+				HashData[0] == Other.HashData[0] &&
+				HashData[1] == Other.HashData[1] &&
+				HashData[2] == Other.HashData[2] &&
+				HashData[3] == Other.HashData[3];
+		}
+
+		friend FORCEINLINE uint32 GetTypeHash(const FContextHash& Hash)
+		{
+			return GetArrayHash(Hash.HashData, 4);
+		}
+
+		uint32 HashData[4] = { 0, 0, 0, 0 };
+	};
+
+	FTraversalCallingContext(const FStaticContext& InGlobalStaticContext, const TMap<FGuid, FStaticContext>& InExternalFunctionInputStaticContexts)
+		: GlobalStaticContext(InGlobalStaticContext)
+		, ExternalFunctionInputStaticContexts(InExternalFunctionInputStaticContexts)
+	{
+	}
+
+	FContextHash GenerateHash() const;
+
+	// The static context collected from scripts processed before this traversal.
+	const FStaticContext& GlobalStaticContext;
+
+	// In a top level script traversal this contains per-function static contexts for function inputs, gathered from
+	// rapid iteration parameters.  This is empty for function traversals.
+	const TMap<FGuid, FStaticContext>& ExternalFunctionInputStaticContexts;
+
+	// The local static context from the traversal that called this traversal.  Empty for top level scripts traversals.
+	FStaticContext CallingStaticContext;
+
+	// In a function traversal, this contains static context which was gathered for the function inputs.  This is empty
+	// for top level script traversals.
+	FStaticContext FunctionCallStaticContext;
+
+	// The call stack of the current traversal represented by the function node ids.
+	TArray<FGuid> FunctionCallStack;
+
+	// The connection traversal mode.
+	EConnectionTraversalMode ConnectionTraversalMode = EConnectionTraversalMode::All;
+
+	// The traversal state from the calling traversal.
+	ETraversalStateFlags TraversalState = ETraversalStateFlags::None;
+};
+
+struct FTraversalLocalContext
+{
+	TOptional<FSelectValue> GetSelectValue(const FTraversalCallingContext& CallingContext, const FSelectKey& SelectKey) const;
+
+	FStaticContext StaticContext;
+	TMap<FGuid, FStaticContext> FunctionInputStaticContexts;
+	ETraversalStateFlags TraversalState = ETraversalStateFlags::None;
+	TSet<TPair<const FTraversalNode*, ETraversalStateFlags>> TraversedNodesWithState;
+};
+
+class FTraversalData;
+
+struct FFunctionCallTraversalData
+{
+	FFunctionCallTraversalData(const FGuid& InFunctionCallNodeGuid, TSharedRef<const FTraversalData> InFunctionCallTraversalData)
+		: FunctionCallNodeGuid(InFunctionCallNodeGuid)
+		, FunctionCallTraversalData(InFunctionCallTraversalData)
+	{
+	}
+		
+	FGuid FunctionCallNodeGuid;
+	TSharedRef<const FTraversalData> FunctionCallTraversalData;
+};
+
+class FTraversalData
+{
+public:
+	const FTraversalData* GetCalledFunctionTraversalData(FGuid FunctionCallNodeGuid) const;
+	TArray<FFunctionCallTraversalData> CalledFunctionTraversalData;
+};
+
+class ITraversalVisitor
+{
+public:
+	virtual ~ITraversalVisitor() { }
+	virtual const FGuid GetVisitorId() const = 0;
+	virtual TSharedRef<FTraversalData> CreateTraversalData() const = 0;
+	virtual void VisitNode(ETraversalStateFlags TraversalStateFlags, const FTraversalNode& Node, FTraversalData& TraversalData) const = 0;
+};
+
+struct FScriptReference
+{
+	FScriptReference() { }
+
+	FScriptReference(const FSoftObjectPath& InPath, const FGuid& InVersion)
+		: Path(InPath)
+		, Version(InVersion)
+	{
+	}
+
+	FORCEINLINE bool operator==(const FScriptReference& Other) const
+	{
+		return Path == Other.Path && Version == Other.Version;
+	}
+
+	friend FORCEINLINE uint32 GetTypeHash(const FScriptReference& ScriptReference)
+	{
+		return HashCombine(GetTypeHash(ScriptReference.Path), GetTypeHash(ScriptReference.Version));
+	}
+
+	FSoftObjectPath Path;
+	FGuid Version;
+};
+
+struct FTraversalBuilderDebugData
+{
+	struct FMultipleInputConnection
+	{
+		FMultipleInputConnection(const TPair<FGuid, FGuid> InTargetNodeGuidAndPinId, const TArray<TPair<FGuid, FGuid>>& InConnectedNodeGuidsAndPinIds)
+			: TargetNodeGuidAndPinId(InTargetNodeGuidAndPinId)
+			, ConnectedNodeGuidsAndPinIds(InConnectedNodeGuidsAndPinIds)
+		{
+		}
+
+		FORCEINLINE bool operator==(const FMultipleInputConnection& Other) const
+		{
+			return
+				TargetNodeGuidAndPinId == Other.TargetNodeGuidAndPinId &&
+				ConnectedNodeGuidsAndPinIds == Other.ConnectedNodeGuidsAndPinIds;
+		}
+
+		TPair<FGuid, FGuid> TargetNodeGuidAndPinId;
+		TArray<TPair<FGuid, FGuid>> ConnectedNodeGuidsAndPinIds;
+	};
+
+	struct FUntrimmedNoop
+	{
+		FUntrimmedNoop(const FGuid& InNodeGuid, FName InSourceNodeTypeName)
+			: NodeGuid(InNodeGuid)
+			, SourceNodeTypeName(InSourceNodeTypeName)
+		{
+		}
+
+		FORCEINLINE bool operator==(const FUntrimmedNoop& Other) const
+		{
+			return NodeGuid == Other.NodeGuid && SourceNodeTypeName == Other.SourceNodeTypeName;
+		}
+
+		FGuid NodeGuid;
+		FName SourceNodeTypeName;
+	};
+
+	struct FUnresolvedSelectOutput
+	{
+		FUnresolvedSelectOutput(const FGuid& InNodeGuid, const FGuid& InPinGuid)
+			: NodeGuid(InNodeGuid)
+			, PinGuid(InPinGuid)
+		{
+		}
+
+		FORCEINLINE bool operator==(const FUnresolvedSelectOutput& Other) const
+		{
+			return NodeGuid == Other.NodeGuid && PinGuid == Other.PinGuid;
+		}
+
+		FGuid NodeGuid;
+		FGuid PinGuid;
+	};
+
+	FORCEINLINE bool operator==(const FTraversalBuilderDebugData& Other) const
+	{
+		return MultipleInputConnections == Other.MultipleInputConnections &&
+			UntrimmedNoops == Other.UntrimmedNoops;
+	}
+
+	void AddMultipleInputConnection(
+		const TPair<FGuid, FGuid>& TargetNodeGuidAndPinId,
+		const TArray<TPair<FGuid, FGuid>>& ConnectedNodeGuidsAndPinIds);
+
+	void AddUntrimmedNoop(const FGuid& NodeGuid, FName SourceNodeTypeName);
+
+	void AddUnresolvedSelectOutput(const FGuid& NodeGuid, const FGuid& PinGuid);
+
+	bool HasData() const
+	{
+		return
+			MultipleInputConnections.Num() > 0 ||
+			UntrimmedNoops.Num() > 0 ||
+			UnresolvedSelectOutputs.Num() > 0;
+	}
+
+	const TArray<FMultipleInputConnection>& GetMultipleInputConnections() const { return MultipleInputConnections; }
+
+	const TArray<FUntrimmedNoop>& GetUntrimmedNoops() const { return UntrimmedNoops; }
+
+	const TArray<FUnresolvedSelectOutput>& GetUnresolvedSelectOutputs() const { return UnresolvedSelectOutputs; }
+
+private:
+	TArray<FMultipleInputConnection> MultipleInputConnections;
+	TArray<FUntrimmedNoop> UntrimmedNoops;
+	TArray<FUnresolvedSelectOutput> UnresolvedSelectOutputs;
+};
+
+struct FTraversalDebugData
+{
+	struct FUnresolvedBase
+	{
+		FUnresolvedBase(const TArray<FGuid>& InFunctionCallStack, const FGuid& InSourceNodeGuid)
+			: FunctionCallStack(InFunctionCallStack)
+			, SourceNodeGuid(InSourceNodeGuid)
+		{
+		}
+
+		FORCEINLINE bool operator==(const FUnresolvedBase& Other) const
+		{
+			return
+				FunctionCallStack == Other.FunctionCallStack &&
+				SourceNodeGuid == Other.SourceNodeGuid;
+		}
+
+		TArray<FGuid> FunctionCallStack;
+		FGuid SourceNodeGuid;
+	};
+
+	struct FUnresolvedSelect : public FUnresolvedBase
+	{
+		FUnresolvedSelect(const TArray<FGuid>& InFunctionCallStack, const FGuid& InSourceNodeGuid, const FSelectKey& InSelectKey)
+			: FUnresolvedBase(InFunctionCallStack, InSourceNodeGuid)
+			, SelectKey(InSelectKey)
+		{
+		}
+
+		FORCEINLINE bool operator==(const FUnresolvedSelect& Other) const
+		{
+			return (FUnresolvedBase&)*this == (FUnresolvedBase&)Other && SelectKey == Other.SelectKey;
+		}
+
+		FSelectKey SelectKey;
+	};
+
+	struct FUnresolvedSelectInput : public FUnresolvedBase
+	{
+		FUnresolvedSelectInput(const TArray<FGuid>& InFunctionCallStack, const FGuid& InSourceNodeGuid, const FSelectKey& InSelectKey, const FSelectValue& InResolvedSelectValue)
+			: FUnresolvedBase(InFunctionCallStack, InSourceNodeGuid)
+			, SelectKey(InSelectKey)
+			, ResolvedSelectValue(InResolvedSelectValue)
+		{
+		}
+
+		FORCEINLINE bool operator==(const FUnresolvedSelectInput& Other) const
+		{
+			return 
+				(FUnresolvedBase&)*this == (FUnresolvedBase&)Other &&
+				SelectKey == Other.SelectKey &&
+				ResolvedSelectValue == Other.ResolvedSelectValue;
+		}
+
+		FSelectKey SelectKey;
+		FSelectValue ResolvedSelectValue;
+	};
+
+	struct FUnresolvedRead : public FUnresolvedBase
+	{
+		FUnresolvedRead(const TArray<FGuid>& InFunctionCallStack, const FGuid& InSourceNodeGuid, const FNiagaraVariableBase& InReadParameter)
+			: FUnresolvedBase(InFunctionCallStack, InSourceNodeGuid)
+			, ReadParameter(InReadParameter)
+		{
+		}
+
+		FORCEINLINE bool operator==(const FUnresolvedRead& Other) const
+		{
+			return (FUnresolvedBase&)*this == (FUnresolvedBase&)Other && ReadParameter == Other.ReadParameter;
+		}
+
+		FNiagaraVariableBase ReadParameter;
+	};
+
+	struct FUnresolvedStaticOp : public FUnresolvedBase
+	{
+		FUnresolvedStaticOp(const TArray<FGuid>& InFunctionCallStack, const FGuid& InSourceNodeGuid, FName InOpName, int32 InUnresolvedPinIndex)
+			: FUnresolvedBase(InFunctionCallStack, InSourceNodeGuid)
+			, OpName(InOpName)
+			, UnresolvedPinIndex(InUnresolvedPinIndex)
+		{
+		}
+
+		FORCEINLINE bool operator==(const FUnresolvedStaticOp& Other) const
+		{
+
+			return
+				(FUnresolvedBase&)*this == (FUnresolvedBase&)Other &&
+				OpName == Other.OpName &&
+				UnresolvedPinIndex == Other.UnresolvedPinIndex;
+		}
+
+		FName OpName;
+		int32 UnresolvedPinIndex = INDEX_NONE;
+	};
+
+	struct FUnresolvedFunctionInput : public FUnresolvedBase
+	{
+		FUnresolvedFunctionInput(const TArray<FGuid>& InFunctionCallStack, const FGuid& InSourceNodeGuid, const FNiagaraVariableBase& InFunctionInputParameter)
+			: FUnresolvedBase(InFunctionCallStack, InSourceNodeGuid)
+			, FunctionInputParameter(InFunctionInputParameter)
+		{
+		}
+
+		FORCEINLINE bool operator==(const FUnresolvedFunctionInput& Other) const
+		{
+			return (FUnresolvedBase&)*this == (FUnresolvedBase&)Other && FunctionInputParameter == Other.FunctionInputParameter;
+		}
+
+		FNiagaraVariableBase FunctionInputParameter;
+	};
+
+	struct FBuilderDebugOutput
+	{
+		FBuilderDebugOutput(const FSoftObjectPath InScriptPath, const ENiagaraScriptUsage InScriptUsage,
+			const FGuid& InScriptUsageId, const FGuid& InScriptVersion, const FTraversalBuilderDebugData& InBuilderDebugData)
+			: ScriptPath(InScriptPath)
+			, ScriptUsage(InScriptUsage)
+			, ScriptUsageId(InScriptUsageId)
+			, ScriptVersion(InScriptVersion)
+			, BuilderDebugData(InBuilderDebugData)
+		{
+		}
+
+
+		FORCEINLINE bool operator==(const FBuilderDebugOutput& Other) const
+		{
+			return 
+				ScriptPath == Other.ScriptPath &&
+				ScriptUsage == Other.ScriptUsage && 
+				ScriptUsageId == Other.ScriptUsageId && 
+				ScriptVersion == Other.ScriptVersion &&
+				BuilderDebugData == Other.BuilderDebugData;
+		}
+		
+		FSoftObjectPath ScriptPath;
+		ENiagaraScriptUsage ScriptUsage;
+		FGuid ScriptUsageId;
+		FGuid ScriptVersion;
+		FTraversalBuilderDebugData BuilderDebugData;
+	};
+
+	void AddUnresolvedSelect(const TArray<FGuid>& InFunctionCallStack, const FGuid& InSourceNodeGuid, const FSelectKey& InSelectKey);
+	void AddUnresolvedSelectInput(const TArray<FGuid>& InFunctionCallStack, const FGuid& InSourceNodeGuid, const FSelectKey& InSelectKey, const FSelectValue& InResolvedSelectValue);
+	void AddUnresolvedRead(const TArray<FGuid>& InFunctionCallStack, const FGuid& InSourceNodeGuid, const FNiagaraVariableBase& InReadParameter);
+	void AddUnresolvedStaticOp(const TArray<FGuid>& InFunctionCallStack, const FGuid& InSourceNodeGuid, FName InOpName, int32 InUnresolvedPinIndex);
+	void AddUnresolvedFunctionInput(const TArray<FGuid>& InFunctionCallStack, const FGuid& InSourceNodeGuid, const FNiagaraVariableBase& InFunctionInput);
+	void AddBuilderDebugOutput(const FSoftObjectPath InGraphPath, const ENiagaraScriptUsage InScriptUsage,
+		const FGuid& InScriptUsageId, const FGuid& InScriptVersion, const FTraversalBuilderDebugData& InBuilderDebugData);
+
+	bool HasData() const { return
+		UnresolvedSelects.Num() > 0 ||
+		UnresolvedSelectInputs.Num() > 0 ||
+		UnresolvedReads.Num() > 0 ||
+		UnresolvedStaticOps.Num() > 0 ||
+		UnresolvedFunctionInputs.Num() > 0 ||
+		BuilderDebugOutputs.Num() > 0; }
+
+	const TArray<FUnresolvedSelect>& GetUnresolvedSelects() const { return UnresolvedSelects; }
+	const TArray<FUnresolvedSelectInput>& GetUnresolvedSelectInputs() const { return UnresolvedSelectInputs; }
+	const TArray<FUnresolvedRead>& GetUnresolvedReads() const { return UnresolvedReads; }
+	const TArray<FUnresolvedStaticOp>& GetUnresolvedStaticOps() const { return UnresolvedStaticOps; }
+	const TArray<FUnresolvedFunctionInput>& GetUnresolvedFunctionInputs() const { return UnresolvedFunctionInputs; }
+	const TArray<FBuilderDebugOutput> GetBuilderDebugOutputs() const { return BuilderDebugOutputs; }
+
+private:
+	TArray<FUnresolvedSelect> UnresolvedSelects;
+	TArray<FUnresolvedSelectInput> UnresolvedSelectInputs;
+	TArray<FUnresolvedRead> UnresolvedReads;
+	TArray<FUnresolvedStaticOp> UnresolvedStaticOps;
+	TArray<FUnresolvedFunctionInput> UnresolvedFunctionInputs;
+	TArray<FBuilderDebugOutput> BuilderDebugOutputs; 
+};
+
+template<typename TEnumFlag>
+static bool HasFlag(TEnumFlag Flags, TEnumFlag FlagToCheck)
+{
+	return (Flags & FlagToCheck) == FlagToCheck;
+}
+
+template<typename TEnumFlag>
+static TEnumFlag SetFlag(TEnumFlag Flags, TEnumFlag FlagToCheck)
+{
+	return Flags | FlagToCheck;
+}
+
+template<typename TEnumFlag>
+static TEnumFlag ClearFlag(TEnumFlag Flags, TEnumFlag FlagToCheck)
+{
+	return Flags & ~FlagToCheck;
+}
+
+} // UE::Niagara::TraversalCache
\ No newline at end of file
diff --git a/Engine/Plugins/FX/Niagara/Source/NiagaraEditor/Private/TypeEditorUtilities/NiagaraBoolTypeEditorUtilities.cpp b/Engine/Plugins/FX/Niagara/Source/NiagaraEditor/Private/TypeEditorUtilities/NiagaraBoolTypeEditorUtilities.cpp
--- a/Engine/Plugins/FX/Niagara/Source/NiagaraEditor/Private/TypeEditorUtilities/NiagaraBoolTypeEditorUtilities.cpp
+++ b/Engine/Plugins/FX/Niagara/Source/NiagaraEditor/Private/TypeEditorUtilities/NiagaraBoolTypeEditorUtilities.cpp
@@ -118,3 +118,30 @@ bool FNiagaraEditorBoolTypeUtilities::TryUpdateTypedValueFromClipboardPortableVa
 	}
 	return false;
 }
+
+int32 FNiagaraEditorBoolTypeUtilities::VariableToSelectNumericValue(const FNiagaraVariable& InVariableValue) const
+{
+	FNiagaraTypeDefinition BaseDefinition = InVariableValue.GetType().RemoveStaticDef();
+	if (ensureMsgf(BaseDefinition == FNiagaraTypeDefinition::GetBoolDef() && InVariableValue.IsDataAllocated(),
+		TEXT("InVariableValue must be of type FNiagaraBool and have its data allocated to convert to a select value.")) == false)
+	{
+		return INDEX_NONE;
+	}
+
+	// The select values for bools differ from the underlying variant value in that it uses 1 for true and not -1.
+	bool bBoolValue = InVariableValue.GetValue<FNiagaraBool>().GetValue();
+	return bBoolValue ? 1 : 0;
+}
+
+FName FNiagaraEditorBoolTypeUtilities::GetDebugNameForSelectValue(const FNiagaraTypeDefinition& ValueType, int32 SelectValue) const
+{
+	if (SelectValue == 1)
+	{
+		return *LexToString(true);
+	}
+	if (SelectValue == 0)
+	{
+		return *LexToString(false);
+	}
+	return NAME_None;
+}
diff --git a/Engine/Plugins/FX/Niagara/Source/NiagaraEditor/Private/TypeEditorUtilities/NiagaraBoolTypeEditorUtilities.h b/Engine/Plugins/FX/Niagara/Source/NiagaraEditor/Private/TypeEditorUtilities/NiagaraBoolTypeEditorUtilities.h
--- a/Engine/Plugins/FX/Niagara/Source/NiagaraEditor/Private/TypeEditorUtilities/NiagaraBoolTypeEditorUtilities.h
+++ b/Engine/Plugins/FX/Niagara/Source/NiagaraEditor/Private/TypeEditorUtilities/NiagaraBoolTypeEditorUtilities.h
@@ -19,4 +19,7 @@ public:
 	virtual bool SupportsClipboardPortableValues() const { return true; }
 	virtual bool TryUpdateClipboardPortableValueFromTypedValue(const FNiagaraTypeDefinition& InSourceType, const FNiagaraVariant& InSourceValue, FNiagaraClipboardPortableValue& InTargetClipboardPortableValue) const override;
 	virtual bool TryUpdateTypedValueFromClipboardPortableValue(const FNiagaraClipboardPortableValue& InSourceClipboardPortableValue, const FNiagaraTypeDefinition& InTargetType, FNiagaraVariant& InTargetValue) const override;
+	virtual bool CanBeSelectValue() const override { return true; }
+	virtual int32 VariableToSelectNumericValue(const FNiagaraVariable& InVariableValue) const override;
+	virtual FName GetDebugNameForSelectValue(const FNiagaraTypeDefinition& ValueType, int32 SelectValue) const override;
 };
\ No newline at end of file
diff --git a/Engine/Plugins/FX/Niagara/Source/NiagaraEditor/Private/TypeEditorUtilities/NiagaraEnumTypeEditorUtilities.cpp b/Engine/Plugins/FX/Niagara/Source/NiagaraEditor/Private/TypeEditorUtilities/NiagaraEnumTypeEditorUtilities.cpp
--- a/Engine/Plugins/FX/Niagara/Source/NiagaraEditor/Private/TypeEditorUtilities/NiagaraEnumTypeEditorUtilities.cpp
+++ b/Engine/Plugins/FX/Niagara/Source/NiagaraEditor/Private/TypeEditorUtilities/NiagaraEnumTypeEditorUtilities.cpp
@@ -269,3 +269,31 @@ bool FNiagaraEditorEnumTypeUtilities::TryUpdateTypedValueFromClipboardPortableVa
 	}
 	return false;
 }
+
+int32 FNiagaraEditorEnumTypeUtilities::VariableToSelectNumericValue(const FNiagaraVariable& InVariableValue) const
+{
+	if (ensureMsgf(InVariableValue.GetType().IsEnum() && InVariableValue.IsDataAllocated(),
+		TEXT("InVariableValue must be an enum type and must have its data allocated to convert to a select value.")) == false)
+	{
+		return INDEX_NONE;
+	}
+
+	// We store enums by their value in FNiagaraVariables and parameter store, but the select value is the index.
+	UEnum* Enum = InVariableValue.GetType().GetEnum();
+	return Enum->GetIndexByValue(InVariableValue.GetValue<FNiagaraInt32>().Value);
+}
+
+FName FNiagaraEditorEnumTypeUtilities::GetDebugNameForSelectValue(const FNiagaraTypeDefinition& ValueType, int32 SelectValue) const
+{
+	UEnum* Enum = ValueType.GetEnum();
+	FName Name = Enum->GetNameByIndex(SelectValue);
+	FText DisplayName = Enum->GetDisplayNameTextByIndex(SelectValue);
+	if (DisplayName.IsEmptyOrWhitespace())
+	{
+		return Name;
+	}
+
+	FString NameString = Name.ToString();
+	FString DisplayNameString = DisplayName.ToString();
+	return NameString.Equals(DisplayNameString)	? Name : *FString::Printf(TEXT("%s (%s)"), *DisplayNameString, *NameString);
+}
diff --git a/Engine/Plugins/FX/Niagara/Source/NiagaraEditor/Private/TypeEditorUtilities/NiagaraEnumTypeEditorUtilities.h b/Engine/Plugins/FX/Niagara/Source/NiagaraEditor/Private/TypeEditorUtilities/NiagaraEnumTypeEditorUtilities.h
--- a/Engine/Plugins/FX/Niagara/Source/NiagaraEditor/Private/TypeEditorUtilities/NiagaraEnumTypeEditorUtilities.h
+++ b/Engine/Plugins/FX/Niagara/Source/NiagaraEditor/Private/TypeEditorUtilities/NiagaraEnumTypeEditorUtilities.h
@@ -24,4 +24,7 @@ public:
 	virtual bool SupportsClipboardPortableValues() const { return true; }
 	virtual bool TryUpdateClipboardPortableValueFromTypedValue(const FNiagaraTypeDefinition& InSourceType, const FNiagaraVariant& InSourceValue, FNiagaraClipboardPortableValue& InTargetClipboardPortableValue) const override;
 	virtual bool TryUpdateTypedValueFromClipboardPortableValue(const FNiagaraClipboardPortableValue& InSourceClipboardPortableValue, const FNiagaraTypeDefinition& InTargetType, FNiagaraVariant& InTargetValue) const override;
+	virtual bool CanBeSelectValue() const override { return true; }
+	virtual int32 VariableToSelectNumericValue(const FNiagaraVariable& InVariableValue) const override;
+	virtual FName GetDebugNameForSelectValue(const FNiagaraTypeDefinition& ValueType, int32 SelectValue) const override;
 };
\ No newline at end of file
diff --git a/Engine/Plugins/FX/Niagara/Source/NiagaraEditor/Private/TypeEditorUtilities/NiagaraIntegerTypeEditorUtilities.cpp b/Engine/Plugins/FX/Niagara/Source/NiagaraEditor/Private/TypeEditorUtilities/NiagaraIntegerTypeEditorUtilities.cpp
--- a/Engine/Plugins/FX/Niagara/Source/NiagaraEditor/Private/TypeEditorUtilities/NiagaraIntegerTypeEditorUtilities.cpp
+++ b/Engine/Plugins/FX/Niagara/Source/NiagaraEditor/Private/TypeEditorUtilities/NiagaraIntegerTypeEditorUtilities.cpp
@@ -290,3 +290,19 @@ bool FNiagaraEditorIntegerTypeUtilities::TryUpdateTypedValueFromClipboardPortabl
 	}
 	return false;
 }
+
+int32 FNiagaraEditorIntegerTypeUtilities::VariableToSelectNumericValue(const FNiagaraVariable& InVariableValue) const
+{
+	FNiagaraTypeDefinition BaseDefinition = InVariableValue.GetType().RemoveStaticDef();
+	if (ensureMsgf(BaseDefinition == FNiagaraTypeDefinition::GetIntDef() && InVariableValue.IsDataAllocated(),
+		TEXT("InVariableValue must be an FNiagaraInt32 and must have its data allocated to convert to a select value.")) == false)
+	{
+		return INDEX_NONE;
+	}
+	return InVariableValue.GetValue<FNiagaraInt32>().Value;
+}
+
+FName FNiagaraEditorIntegerTypeUtilities::GetDebugNameForSelectValue(const FNiagaraTypeDefinition& ValueType, int32 SelectValue) const
+{
+	return *LexToString(SelectValue);
+}
diff --git a/Engine/Plugins/FX/Niagara/Source/NiagaraEditor/Private/TypeEditorUtilities/NiagaraIntegerTypeEditorUtilities.h b/Engine/Plugins/FX/Niagara/Source/NiagaraEditor/Private/TypeEditorUtilities/NiagaraIntegerTypeEditorUtilities.h
--- a/Engine/Plugins/FX/Niagara/Source/NiagaraEditor/Private/TypeEditorUtilities/NiagaraIntegerTypeEditorUtilities.h
+++ b/Engine/Plugins/FX/Niagara/Source/NiagaraEditor/Private/TypeEditorUtilities/NiagaraIntegerTypeEditorUtilities.h
@@ -18,6 +18,9 @@ public:
 	virtual bool SupportsClipboardPortableValues() const override { return true; }
 	virtual bool TryUpdateClipboardPortableValueFromTypedValue(const FNiagaraTypeDefinition& InSourceType, const FNiagaraVariant& InSourceValue, FNiagaraClipboardPortableValue& InTargetClipboardPortableValue) const override;
 	virtual bool TryUpdateTypedValueFromClipboardPortableValue(const FNiagaraClipboardPortableValue& InSourceClipboardPortableValue, const FNiagaraTypeDefinition& InTargetType, FNiagaraVariant& InTargetValue) const override;
+	virtual bool CanBeSelectValue() const override { return true; }
+	virtual int32 VariableToSelectNumericValue(const FNiagaraVariable& InVariableValue) const override;
+	virtual FName GetDebugNameForSelectValue(const FNiagaraTypeDefinition& ValueType, int32 SelectValue) const override;
 };
 
 class SNiagaraIntegerParameterEditor : public SNiagaraParameterEditor
diff --git a/Engine/Plugins/FX/Niagara/Source/NiagaraEditor/Private/ViewModels/NiagaraSystemViewModel.cpp b/Engine/Plugins/FX/Niagara/Source/NiagaraEditor/Private/ViewModels/NiagaraSystemViewModel.cpp
--- a/Engine/Plugins/FX/Niagara/Source/NiagaraEditor/Private/ViewModels/NiagaraSystemViewModel.cpp
+++ b/Engine/Plugins/FX/Niagara/Source/NiagaraEditor/Private/ViewModels/NiagaraSystemViewModel.cpp
@@ -1622,7 +1622,7 @@ TArray<FNiagaraGraphParameterReference> FNiagaraSystemViewModel::GetGraphParamet
 			FNiagaraParameterMapHistoryBuilder Builder;
 			FVersionedNiagaraEmitter GraphOwningEmitter = Graph->GetOwningEmitter();
 			FCompileConstantResolver ConstantResolver = GraphOwningEmitter.Emitter != nullptr
-				? FCompileConstantResolver(GraphOwningEmitter, ENiagaraScriptUsage::Function)
+				? FCompileConstantResolver(GraphOwningEmitter, OutputNode->GetUsage())
 				: FCompileConstantResolver();
 
 			Builder.SetIgnoreDisabled(bIgnoreDisabled);
diff --git a/Engine/Plugins/FX/Niagara/Source/NiagaraEditor/Private/ViewModels/Stack/NiagaraStackGraphUtilities.cpp b/Engine/Plugins/FX/Niagara/Source/NiagaraEditor/Private/ViewModels/Stack/NiagaraStackGraphUtilities.cpp
--- a/Engine/Plugins/FX/Niagara/Source/NiagaraEditor/Private/ViewModels/Stack/NiagaraStackGraphUtilities.cpp
+++ b/Engine/Plugins/FX/Niagara/Source/NiagaraEditor/Private/ViewModels/Stack/NiagaraStackGraphUtilities.cpp
@@ -50,12 +50,16 @@
 #include "NiagaraEmitter.h"
 #include "ViewModels/HierarchyEditor/NiagaraHierarchyViewModelBase.h"
 #include "ViewModels/HierarchyEditor/NiagaraSummaryViewViewModel.h"
+#include "TraversalCache/StaticContextBuilder.h"
+#include "TraversalCache/TraversalCache.h"
 
 DECLARE_CYCLE_STAT(TEXT("Niagara - StackGraphUtilities - RelayoutGraph"), STAT_NiagaraEditor_StackGraphUtilities_RelayoutGraph, STATGROUP_NiagaraEditor);
 DECLARE_CYCLE_STAT(TEXT("Niagara - StackGraphUtilities - GatherInputRelationsForStack"), STAT_NiagaraEditor_StackGraphUtilities_GatherInputRelationsForStack, STATGROUP_NiagaraEditor);
 DECLARE_CYCLE_STAT(TEXT("Niagara - StackGraphUtilities - FindInputData"), STAT_NiagaraEditor_StackGraphUtilities_FindInputData, STATGROUP_NiagaraEditor);
 DECLARE_CYCLE_STAT(TEXT("Niagara - StackGraphUtilities - FindModuleInputData"), STAT_NiagaraEditor_StackGraphUtilities_FindModuleInputData, STATGROUP_NiagaraEditor);
 DECLARE_CYCLE_STAT(TEXT("Niagara - StackGraphUtilities - FindAssignmentInputData"), STAT_NiagaraEditor_StackGraphUtilities_FindAssignmentInputData, STATGROUP_NiagaraEditor);
+DECLARE_CYCLE_STAT(TEXT("Niagara - StackGraphUtilities - GetStackFunctionInputs"), STAT_NiagaraEditor_StackGraphUtilities_GetStackFunctionInputs, STATGROUP_NiagaraEditor);
+DECLARE_CYCLE_STAT(TEXT("Niagara - StackGraphUtilities - GetStackFunctionStaticSwitchPins"), STAT_NiagaraEditor_StackGraphUtilities_GetStackFunctionStaticSwitchPins, STATGROUP_NiagaraEditor);
 
 #define LOCTEXT_NAMESPACE "NiagaraStackGraphUtilities"
 
@@ -274,6 +278,13 @@ UEdGraphPin* FNiagaraStackGraphUtilities::GetParameterMapInputPin(UNiagaraNode&
 	return GetParameterMapPin(InputPins);
 }
 
+const UEdGraphPin* FNiagaraStackGraphUtilities::GetParameterMapInputPin(const UNiagaraNode& Node)
+{
+	FPinCollectorArray InputPins;
+	Node.GetInputPins(InputPins);
+	return GetParameterMapPin(InputPins);
+}
+
 UEdGraphPin* FNiagaraStackGraphUtilities::GetParameterMapOutputPin(UNiagaraNode& Node)
 {
 	FPinCollectorArray OutputPins;
@@ -1019,7 +1030,7 @@ void ExtractInputPinsFromHistory(FNiagaraParameterMapHistory& History, UEdGraph*
 	}
 }
 
-void FNiagaraStackGraphUtilities::GetStackFunctionInputs(const UNiagaraNodeFunctionCall& FunctionCallNode, TArray<FNiagaraVariable>& OutInputVariables, TSet<FNiagaraVariable>& OutHiddenVariables, FCompileConstantResolver ConstantResolver, ENiagaraGetStackFunctionInputPinsOptions Options, bool bIgnoreDisabled)
+void FNiagaraStackGraphUtilities::GetStackFunctionInputsLegacy(const UNiagaraNodeFunctionCall& FunctionCallNode, TArray<FNiagaraVariable>& OutInputVariables, TSet<FNiagaraVariable>& OutHiddenVariables, FCompileConstantResolver ConstantResolver, ENiagaraGetStackFunctionInputPinsOptions Options, bool bIgnoreDisabled)
 {
 	TArray<FNiagaraVariable> StaticVars;
 	FNiagaraStackGraphUtilitiesImpl::GetFunctionStaticVariables(FunctionCallNode, StaticVars);
@@ -1027,7 +1038,7 @@ void FNiagaraStackGraphUtilities::GetStackFunctionInputs(const UNiagaraNodeFunct
 	FNiagaraEditorModule::Get().GetGraphDataCache().GetStackFunctionInputs(FunctionCallNode, StaticVars, OutInputVariables, OutHiddenVariables, ConstantResolver, Options, bIgnoreDisabled);
 }
 
-void FNiagaraStackGraphUtilities::GetStackFunctionInputs(const UNiagaraNodeFunctionCall& FunctionCallNode, TArray<FNiagaraVariable>& OutInputVariables, FCompileConstantResolver ConstantResolver, ENiagaraGetStackFunctionInputPinsOptions Options, bool bIgnoreDisabled)
+void FNiagaraStackGraphUtilities::GetStackFunctionInputsLegacy(const UNiagaraNodeFunctionCall& FunctionCallNode, TArray<FNiagaraVariable>& OutInputVariables, FCompileConstantResolver ConstantResolver, ENiagaraGetStackFunctionInputPinsOptions Options, bool bIgnoreDisabled)
 {
 	TArray<FNiagaraVariable> StaticVars;
 	FNiagaraStackGraphUtilitiesImpl::GetFunctionStaticVariables(FunctionCallNode, StaticVars);
@@ -1035,10 +1046,142 @@ void FNiagaraStackGraphUtilities::GetStackFunctionInputs(const UNiagaraNodeFunct
 	FNiagaraEditorModule::Get().GetGraphDataCache().GetStackFunctionInputs(FunctionCallNode, StaticVars, OutInputVariables, ConstantResolver, Options, bIgnoreDisabled);
 }
 
-void FNiagaraStackGraphUtilities::GetStackFunctionInputs(const UNiagaraNodeFunctionCall& FunctionCallNode, TArray<FNiagaraVariable>& OutInputVariables, ENiagaraGetStackFunctionInputPinsOptions Options /*= ENiagaraGetStackFunctionInputPinsOptions::AllInputs*/, bool bIgnoreDisabled /*= false*/)
+void FNiagaraStackGraphUtilities::GetStackFunctionInputsLegacy(const UNiagaraNodeFunctionCall& FunctionCallNode, TArray<FNiagaraVariable>& OutInputVariables, ENiagaraGetStackFunctionInputPinsOptions Options /*= ENiagaraGetStackFunctionInputPinsOptions::AllInputs*/, bool bIgnoreDisabled /*= false*/)
 {
 	FCompileConstantResolver EmptyResolver;
-	GetStackFunctionInputs(FunctionCallNode, OutInputVariables, EmptyResolver, Options, bIgnoreDisabled);
+	GetStackFunctionInputsLegacy(FunctionCallNode, OutInputVariables, EmptyResolver, Options, bIgnoreDisabled);
+}
+
+namespace UE::Niagara::TraversalCache::StackGraphUtilities
+{
+	void ConstantResolverToTopLevelScriptContext(const FCompileConstantResolver& ConstantResolver, const UNiagaraNodeFunctionCall& FunctionCallNode, FTopLevelScriptStaticContext& OutTopLevelScriptContext)
+	{
+		const UNiagaraSystem* System = nullptr;
+		const FVersionedNiagaraEmitterData* EmitterData = nullptr;
+		const UNiagaraScript* Script = nullptr;
+
+		// Get the system.
+		if (ConstantResolver.GetSystem() != nullptr)
+		{
+			System = ConstantResolver.GetSystem();
+		}
+		else if (ConstantResolver.GetEmitter() != nullptr)
+		{
+			System = ConstantResolver.GetEmitter()->GetTypedOuter<UNiagaraSystem>();
+		}
+		else
+		{
+			System = FunctionCallNode.GetTypedOuter<UNiagaraSystem>();
+		}
+
+		if (System == nullptr)
+		{
+			return;
+		}
+
+		// Get the emitter if available
+		bool bIsEmitterOrParticleScript = UNiagaraScript::IsEmitterScript(ConstantResolver.GetUsage()) || UNiagaraScript::IsParticleScript(ConstantResolver.GetUsage());
+		if (bIsEmitterOrParticleScript)
+		{
+			if (ConstantResolver.GetEmitter() != nullptr)
+			{
+				EmitterData = ConstantResolver.GetEmitter()->GetLatestEmitterData();
+			}
+			else
+			{
+				const UNiagaraEmitter* Emitter = FunctionCallNode.GetTypedOuter<UNiagaraEmitter>();
+				EmitterData = Emitter->GetLatestEmitterData();
+			}
+		}
+
+		if (bIsEmitterOrParticleScript && EmitterData == nullptr)
+		{
+			return;
+		}
+
+		// Get the script.
+		if (UNiagaraScript::IsSystemScript(ConstantResolver.GetUsage()))
+		{
+			Script = ConstantResolver.GetUsage() == ENiagaraScriptUsage::SystemSpawnScript
+				? System->GetSystemSpawnScript()
+				: System->GetSystemUpdateScript();
+		}
+		else if(bIsEmitterOrParticleScript)
+		{
+			const UNiagaraNodeOutput* OutputNode = FNiagaraStackGraphUtilities::GetEmitterOutputNodeForStackNode(FunctionCallNode);
+			Script = OutputNode != nullptr ? EmitterData->GetScript(OutputNode->GetUsage(), OutputNode->GetUsageId()) : nullptr;
+		}
+
+		if (Script == nullptr)
+		{
+			return;
+		}
+
+		FStaticContextBuilder::CreateTopLevelScriptContext(*System, EmitterData, *Script, OutTopLevelScriptContext);
+	}
+
+	EStackFunctionReadFilterFlags InputOptionsToReadFilterFlags(FNiagaraStackGraphUtilities::ENiagaraGetStackFunctionInputPinsOptions Options, bool bIgnoreDisabled)
+	{
+		bool bInputsOnly = Options == FNiagaraStackGraphUtilities::ENiagaraGetStackFunctionInputPinsOptions::ModuleInputsOnly;
+		return 
+			(bInputsOnly ? EStackFunctionReadFilterFlags::InputsOnly : EStackFunctionReadFilterFlags::None) |
+			(bIgnoreDisabled ? EStackFunctionReadFilterFlags::EnabledOnly : EStackFunctionReadFilterFlags::None);
+	}
+}
+
+void FNiagaraStackGraphUtilities::GetStackFunctionInputs(const UNiagaraNodeFunctionCall& FunctionCallNode, TArray<FNiagaraVariable>& OutInputVariables, TSet<FNiagaraVariable>& OutHiddenVariables, FCompileConstantResolver ConstantResolver, ENiagaraGetStackFunctionInputPinsOptions Options, bool bIgnoreDisabled)
+{
+	SCOPE_CYCLE_COUNTER(STAT_NiagaraEditor_StackGraphUtilities_GetStackFunctionInputs);
+	if (GbEnableTraversalCache == 0)
+	{
+		GetStackFunctionInputsLegacy(FunctionCallNode, OutInputVariables, OutHiddenVariables, ConstantResolver, Options, bIgnoreDisabled);
+	}
+	else
+	{
+		using namespace UE::Niagara::TraversalCache;
+		FTopLevelScriptStaticContext ScriptContext;
+		StackGraphUtilities::ConstantResolverToTopLevelScriptContext(ConstantResolver, FunctionCallNode, ScriptContext);
+		EStackFunctionReadFilterFlags FilterFlags = StackGraphUtilities::InputOptionsToReadFilterFlags(Options, bIgnoreDisabled);
+		FTraversalCache::GetStackFunctionReads(ScriptContext, FunctionCallNode, OutInputVariables, OutHiddenVariables, FilterFlags);
+	}
+}
+
+void FNiagaraStackGraphUtilities::GetStackFunctionInputs(const UNiagaraNodeFunctionCall& FunctionCallNode, TArray<FNiagaraVariable>& OutInputVariables, FCompileConstantResolver ConstantResolver, ENiagaraGetStackFunctionInputPinsOptions Options, bool bIgnoreDisabled)
+{
+	SCOPE_CYCLE_COUNTER(STAT_NiagaraEditor_StackGraphUtilities_GetStackFunctionInputs);
+	if (GbEnableTraversalCache == 0)
+	{
+		GetStackFunctionInputsLegacy(FunctionCallNode, OutInputVariables, ConstantResolver, Options, bIgnoreDisabled);
+	}
+	else
+	{
+		using namespace UE::Niagara::TraversalCache;
+		TSet<FNiagaraVariable> HiddenVariables;
+		FTopLevelScriptStaticContext ScriptContext;
+		StackGraphUtilities::ConstantResolverToTopLevelScriptContext(ConstantResolver, FunctionCallNode, ScriptContext);
+		EStackFunctionReadFilterFlags FilterFlags = StackGraphUtilities::InputOptionsToReadFilterFlags(Options, bIgnoreDisabled);
+		FTraversalCache::GetStackFunctionReads(ScriptContext, FunctionCallNode, OutInputVariables, HiddenVariables, FilterFlags);
+	}
+}
+
+void FNiagaraStackGraphUtilities::GetStackFunctionInputs(
+	const UNiagaraNodeFunctionCall& FunctionCallNode,
+	TArray<FNiagaraVariable>& OutInputVariables,
+	ENiagaraGetStackFunctionInputPinsOptions Options, bool bIgnoreDisabled)
+{
+	SCOPE_CYCLE_COUNTER(STAT_NiagaraEditor_StackGraphUtilities_GetStackFunctionInputs);
+	if (GbEnableTraversalCache == 0)
+	{
+		GetStackFunctionInputsLegacy(FunctionCallNode, OutInputVariables, Options, bIgnoreDisabled);
+	}
+	else
+	{
+		using namespace UE::Niagara::TraversalCache;
+		TSet<FNiagaraVariable> HiddenVariables;
+		FTopLevelScriptStaticContext ScriptContext;
+		EStackFunctionReadFilterFlags FilterFlags = StackGraphUtilities::InputOptionsToReadFilterFlags(Options, bIgnoreDisabled);
+		FTraversalCache::GetStackFunctionReads(ScriptContext, FunctionCallNode, OutInputVariables, HiddenVariables, FilterFlags);
+	}
 }
 
 void FNiagaraStackGraphUtilities::GatherInputRelationsForStack(FInputDataCollection& State, TSharedRef<FNiagaraEmitterViewModel> EmitterViewModel)
@@ -1506,7 +1649,7 @@ TArray<UEdGraphPin*> FNiagaraStackGraphUtilities::GetUnusedFunctionInputPins(con
 	return ResultPins;
 }
 
-void FNiagaraStackGraphUtilities::GetStackFunctionStaticSwitchPins(const UNiagaraNodeFunctionCall& FunctionCallNode, TArray<UEdGraphPin*>& OutInputPins, TSet<UEdGraphPin*>& OutHiddenPins,
+void FNiagaraStackGraphUtilities::GetStackFunctionStaticSwitchPinsLegacy(const UNiagaraNodeFunctionCall& FunctionCallNode, TArray<UEdGraphPin*>& OutInputPins, TSet<UEdGraphPin*>& OutHiddenPins,
 	FCompileConstantResolver& ConstantResolver)
 {
 	const UEdGraphSchema_Niagara* Schema = CastChecked<UEdGraphSchema_Niagara>(FunctionCallNode.GetSchema());
@@ -1582,6 +1725,22 @@ void FNiagaraStackGraphUtilities::GetStackFunctionStaticSwitchPins(const UNiagar
 	}
 }
 
+void FNiagaraStackGraphUtilities::GetStackFunctionStaticSwitchPins(const UNiagaraNodeFunctionCall& FunctionCallNode, TArray<UEdGraphPin*>& OutInputPins, TSet<UEdGraphPin*>& OutHiddenPins, FCompileConstantResolver& ConstantResolver)
+{
+	SCOPE_CYCLE_COUNTER(STAT_NiagaraEditor_StackGraphUtilities_GetStackFunctionStaticSwitchPins);
+	if (GbEnableTraversalCache == 0)
+	{
+		GetStackFunctionStaticSwitchPinsLegacy(FunctionCallNode, OutInputPins, OutHiddenPins, ConstantResolver);
+	}
+	else
+	{
+		using namespace UE::Niagara::TraversalCache;
+		FTopLevelScriptStaticContext ScriptContext;
+		StackGraphUtilities::ConstantResolverToTopLevelScriptContext(ConstantResolver, FunctionCallNode, ScriptContext);
+		FTraversalCache::GetStackFunctionStaticInputPins(ScriptContext, FunctionCallNode, OutInputPins, OutHiddenPins);
+	}
+}
+
 void FNiagaraStackGraphUtilities::GetStackFunctionOutputVariables(UNiagaraNodeFunctionCall& FunctionCallNode, FCompileConstantResolver ConstantResolver, TArray<FNiagaraVariable>& OutOutputVariables, TArray<FNiagaraVariable>& OutOutputVariablesWithOriginalAliasesIntact)
 {
 	FNiagaraParameterMapHistoryBuilder Builder;
@@ -3298,9 +3457,13 @@ void FNiagaraStackGraphUtilities::GatherRenamedStackFunctionOutputVariableNames(
 		return;
 	}
 
+	const UNiagaraNodeOutput* OutputNode = GetEmitterOutputNodeForStackNode(FunctionCallNode);
+
 	TArray<FNiagaraVariable> OutputVariables;
 	TArray<FNiagaraVariable> OutputVariablesWithOriginalAliasesIntact;
-	FCompileConstantResolver ConstantResolver(Emitter, ENiagaraScriptUsage::Function, FunctionCallNode.DebugState);
+	FCompileConstantResolver ConstantResolver = OutputNode != nullptr 
+		? FCompileConstantResolver(Emitter, OutputNode->GetUsage(), FunctionCallNode.DebugState)
+		: FCompileConstantResolver();
 	FNiagaraStackGraphUtilities::GetStackFunctionOutputVariables(FunctionCallNode, ConstantResolver, OutputVariables, OutputVariablesWithOriginalAliasesIntact);
 
 	for (FNiagaraVariable& OutputVariableWithOriginalAliasesIntact : OutputVariablesWithOriginalAliasesIntact)
@@ -3332,9 +3495,13 @@ void FNiagaraStackGraphUtilities::GatherRenamedStackFunctionInputAndOutputVariab
 		return;
 	}
 
+	const UNiagaraNodeOutput* OutputNode = GetEmitterOutputNodeForStackNode(FunctionCallNode);
+
 	TArray<FNiagaraVariable> Variables;
 	TArray<FNiagaraVariable> VariablesWithOriginalAliasesIntact;
-	FCompileConstantResolver ConstantResolver(Emitter, ENiagaraScriptUsage::Function, FunctionCallNode.DebugState);
+	FCompileConstantResolver ConstantResolver = OutputNode != nullptr
+		? FCompileConstantResolver(Emitter, OutputNode->GetUsage(), FunctionCallNode.DebugState)
+		: FCompileConstantResolver();
 	FNiagaraStackGraphUtilities::GetStackFunctionInputAndOutputVariables(FunctionCallNode, ConstantResolver, Variables, VariablesWithOriginalAliasesIntact);
 
 	for (FNiagaraVariable& Variable : VariablesWithOriginalAliasesIntact)
diff --git a/Engine/Plugins/FX/Niagara/Source/NiagaraEditor/Public/INiagaraEditorTypeUtilities.h b/Engine/Plugins/FX/Niagara/Source/NiagaraEditor/Public/INiagaraEditorTypeUtilities.h
--- a/Engine/Plugins/FX/Niagara/Source/NiagaraEditor/Public/INiagaraEditorTypeUtilities.h
+++ b/Engine/Plugins/FX/Niagara/Source/NiagaraEditor/Public/INiagaraEditorTypeUtilities.h
@@ -60,6 +60,12 @@ public:
 	virtual bool CanUpdateTypedValueFromClipboardPortableValue(const FNiagaraClipboardPortableValue& InSourceClipboardPortableValue, const FNiagaraTypeDefinition& InTargetType) const = 0;
 
 	virtual bool TryUpdateTypedValueFromClipboardPortableValue(const FNiagaraClipboardPortableValue& InSourceClipboardPortableValue, const FNiagaraTypeDefinition& InTargetType, FNiagaraVariant& InTargetValue) const = 0;
+
+	virtual bool CanBeSelectValue() const = 0;
+
+	virtual int32 VariableToSelectNumericValue(const FNiagaraVariable& InVariableValue) const = 0;
+
+	virtual FName GetDebugNameForSelectValue(const FNiagaraTypeDefinition& ValueType, int32 SelectValue) const = 0;
 };
 
 class FNiagaraEditorTypeUtilities : public INiagaraEditorTypeUtilities, public TSharedFromThis<FNiagaraEditorTypeUtilities, ESPMode::ThreadSafe>
@@ -93,6 +99,17 @@ public:
 		return TryUpdateTypedValueFromClipboardPortableValue(InSourceClipboardPortableValue, InTargetType, Unused);
 	}
 	virtual bool TryUpdateTypedValueFromClipboardPortableValue(const FNiagaraClipboardPortableValue& InSourceClipboardPortableValue, const FNiagaraTypeDefinition& InTargetType, FNiagaraVariant& InTargetValue) const override { return false; }
+	virtual bool CanBeSelectValue() const override { return false; }
+	virtual int32 VariableToSelectNumericValue(const FNiagaraVariable& InVariableValue) const override
+	{
+		unimplemented();
+		return INDEX_NONE;
+	}
+	virtual FName GetDebugNameForSelectValue(const FNiagaraTypeDefinition& ValueType, int32 SelectValue) const override
+	{
+		unimplemented();
+		return NAME_None;
+	}
 };
 
 class INiagaraEditorPropertyUtilities
diff --git a/Engine/Plugins/FX/Niagara/Source/NiagaraEditor/Public/NiagaraEditorModule.h b/Engine/Plugins/FX/Niagara/Source/NiagaraEditor/Public/NiagaraEditorModule.h
--- a/Engine/Plugins/FX/Niagara/Source/NiagaraEditor/Public/NiagaraEditorModule.h
+++ b/Engine/Plugins/FX/Niagara/Source/NiagaraEditor/Public/NiagaraEditorModule.h
@@ -60,12 +60,19 @@ class UNiagaraParameterCollection;
 struct FNiagaraSystemAsyncCompileResults;
 class ITargetPlatform;
 
+namespace UE::Niagara::TraversalCache
+{
+	class FTraversalBuilder;
+	class FTraversalCache;
+}
+
 DECLARE_STATS_GROUP(TEXT("Niagara Editor"), STATGROUP_NiagaraEditor, STATCAT_Niagara);
 
 extern NIAGARAEDITOR_API int32 GbShowNiagaraDeveloperWindows;
 extern NIAGARAEDITOR_API int32 GbPreloadSelectablePluginAssetsOnDemand;
 extern NIAGARAEDITOR_API int32 GbEnableExperimentalInlineDynamicInputs;
 extern NIAGARAEDITOR_API int32 GbEnableCustomInlineDynamicInputFormats;
+extern NIAGARAEDITOR_API int32 GbEnableTraversalCache;
 
 /* Defines methods for allowing external modules to supply widgets to the core editor module. */
 class INiagaraEditorWidgetProvider
@@ -299,6 +306,10 @@ public:
 
 	UPackage* GetTempPackage() { return TempPackage; }
 
+
+	UE::Niagara::TraversalCache::FTraversalBuilder& GetParameterTraversalBuilder() { return *TraversalBuilder.Get(); };
+	UE::Niagara::TraversalCache::FTraversalCache& GetParameterTraversalCache() { return *TraversalCache.Get(); }
+
 private:
 	class FDeferredDestructionContainerBase
 	{
@@ -540,4 +551,7 @@ private:
 	TSharedPtr<FNiagaraRecentAndFavoritesManager> RecentAndFavoritesManager;
 
 	UPackage* TempPackage = nullptr;
+
+	TSharedPtr<UE::Niagara::TraversalCache::FTraversalBuilder> TraversalBuilder;
+	TSharedPtr<UE::Niagara::TraversalCache::FTraversalCache> TraversalCache;
 };
diff --git a/Engine/Plugins/FX/Niagara/Source/NiagaraEditor/Public/NiagaraGraph.h b/Engine/Plugins/FX/Niagara/Source/NiagaraEditor/Public/NiagaraGraph.h
--- a/Engine/Plugins/FX/Niagara/Source/NiagaraEditor/Public/NiagaraGraph.h
+++ b/Engine/Plugins/FX/Niagara/Source/NiagaraEditor/Public/NiagaraGraph.h
@@ -236,7 +236,7 @@ public:
 	TArray<FNiagaraVariable> NIAGARAEDITOR_API FindStaticSwitchInputs(bool bReachableOnly = false, const TArray<FNiagaraVariable>& InStaticVars = TArray<FNiagaraVariable>()) const;
 
 	/** Get an in-order traversal of a graph by the specified target output script usage.*/
-	void BuildTraversal(TArray<class UNiagaraNode*>& OutNodesTraversed, ENiagaraScriptUsage TargetUsage, FGuid TargetUsageId, bool bEvaluateStaticSwitches = false) const;
+	NIAGARAEDITOR_API void BuildTraversal(TArray<class UNiagaraNode*>& OutNodesTraversed, ENiagaraScriptUsage TargetUsage, FGuid TargetUsageId, bool bEvaluateStaticSwitches = false) const;
 	static void BuildTraversal(TArray<class UNiagaraNode*>& OutNodesTraversed, UNiagaraNode* FinalNode, bool bEvaluateStaticSwitches = false);
 
 	/** Generates a list of unique input and output parameters for when this script is used as a function. */
diff --git a/Engine/Plugins/FX/Niagara/Source/NiagaraEditor/Public/NiagaraGraphDataCache.h b/Engine/Plugins/FX/Niagara/Source/NiagaraEditor/Public/NiagaraGraphDataCache.h
--- a/Engine/Plugins/FX/Niagara/Source/NiagaraEditor/Public/NiagaraGraphDataCache.h
+++ b/Engine/Plugins/FX/Niagara/Source/NiagaraEditor/Public/NiagaraGraphDataCache.h
@@ -21,6 +21,8 @@ public:
 	void GetStackFunctionInputs(const UNiagaraNodeFunctionCall& FunctionCallNode, TConstArrayView<FNiagaraVariable> StaticVars, TArray<FNiagaraVariable>& OutInputVariables, FCompileConstantResolver ConstantResolver, FNiagaraStackGraphUtilities::ENiagaraGetStackFunctionInputPinsOptions Options, bool bIgnoreDisabled);
 	void GetStackFunctionInputs(const UNiagaraNodeFunctionCall& FunctionCallNode, TConstArrayView<FNiagaraVariable> StaticVars, TArray<FNiagaraVariable>& OutInputVariables, TSet<FNiagaraVariable>& OutHiddenVariables, FCompileConstantResolver ConstantResolver, FNiagaraStackGraphUtilities::ENiagaraGetStackFunctionInputPinsOptions Options, bool bIgnoreDisabled);
 
+	void NIAGARAEDITOR_API InvalidateCache();
+
 private:
 	void GetStackFunctionInputsInternal(
 		const UNiagaraNodeFunctionCall& FunctionCallNode,
diff --git a/Engine/Plugins/FX/Niagara/Source/NiagaraEditor/Public/NiagaraNodeWithDynamicPins.h b/Engine/Plugins/FX/Niagara/Source/NiagaraEditor/Public/NiagaraNodeWithDynamicPins.h
--- a/Engine/Plugins/FX/Niagara/Source/NiagaraEditor/Public/NiagaraNodeWithDynamicPins.h
+++ b/Engine/Plugins/FX/Niagara/Source/NiagaraEditor/Public/NiagaraNodeWithDynamicPins.h
@@ -32,7 +32,7 @@ public:
 	UEdGraphPin* RequestNewTypedPin(EEdGraphPinDirection Direction, const FNiagaraTypeDefinition& Type);
 
 	/** Helper to identify if a pin is an Add pin.*/
-	bool IsAddPin(const UEdGraphPin* Pin) const;
+	static bool IsAddPin(const UEdGraphPin* Pin);
 
 	/** Helper to identify if a pin is an Add pin.*/
 	bool IsExecPin(const UEdGraphPin* Pin) const;
diff --git a/Engine/Plugins/FX/Niagara/Source/NiagaraEditor/Public/ViewModels/Stack/NiagaraStackGraphUtilities.h b/Engine/Plugins/FX/Niagara/Source/NiagaraEditor/Public/ViewModels/Stack/NiagaraStackGraphUtilities.h
--- a/Engine/Plugins/FX/Niagara/Source/NiagaraEditor/Public/ViewModels/Stack/NiagaraStackGraphUtilities.h
+++ b/Engine/Plugins/FX/Niagara/Source/NiagaraEditor/Public/ViewModels/Stack/NiagaraStackGraphUtilities.h
@@ -49,6 +49,8 @@ namespace FNiagaraStackGraphUtilities
 
 	UEdGraphPin* GetParameterMapInputPin(UNiagaraNode& Node);
 
+	const UEdGraphPin* GetParameterMapInputPin(const UNiagaraNode& Node);
+
 	UEdGraphPin* GetParameterMapOutputPin(UNiagaraNode& Node);
 
 	void GetOrderedModuleNodes(UNiagaraNodeOutput& OutputNode, TArray<UNiagaraNodeFunctionCall*>& ModuleNodes);
@@ -119,12 +121,19 @@ namespace FNiagaraStackGraphUtilities
 		ModuleInputsOnly
 	};
 
+	/* Returns the input pins for the given function call node. Try not to use this method if possible and use the version that accepts a FCompileConstantResolver parameter. This method should only be used if the current context is completely outside of any system or emitter (e.g. inside the graph editor itself) and constants cannot be resolved at all. */
+	void GetStackFunctionInputsLegacy(const UNiagaraNodeFunctionCall& FunctionCallNode, TArray<FNiagaraVariable>& OutInputVariables, ENiagaraGetStackFunctionInputPinsOptions Options, bool bIgnoreDisabled = false);
+
+	/* Returns the input pins for the given function call node. */
+	NIAGARAEDITOR_API void GetStackFunctionInputsLegacy(const UNiagaraNodeFunctionCall& FunctionCallNode, TArray<FNiagaraVariable>& OutInputVariables, TSet<FNiagaraVariable>& OutHiddenVariables, FCompileConstantResolver ConstantResolver, ENiagaraGetStackFunctionInputPinsOptions Options, bool bIgnoreDisabled = false);
+	NIAGARAEDITOR_API void GetStackFunctionInputsLegacy(const UNiagaraNodeFunctionCall& FunctionCallNode, TArray<FNiagaraVariable>& OutInputVariables, FCompileConstantResolver ConstantResolver, ENiagaraGetStackFunctionInputPinsOptions Options, bool bIgnoreDisabled = false);
+	
 	/* Returns the input pins for the given function call node. Try not to use this method if possible and use the version that accepts a FCompileConstantResolver parameter. This method should only be used if the current context is completely outside of any system or emitter (e.g. inside the graph editor itself) and constants cannot be resolved at all. */
 	void GetStackFunctionInputs(const UNiagaraNodeFunctionCall& FunctionCallNode, TArray<FNiagaraVariable>& OutInputVariables, ENiagaraGetStackFunctionInputPinsOptions Options, bool bIgnoreDisabled = false);
 
 	/* Returns the input pins for the given function call node. */
-	void GetStackFunctionInputs(const UNiagaraNodeFunctionCall& FunctionCallNode, TArray<FNiagaraVariable>& OutInputVariables, TSet<FNiagaraVariable>& OutHiddenVariables, FCompileConstantResolver ConstantResolver, ENiagaraGetStackFunctionInputPinsOptions Options, bool bIgnoreDisabled = false);
-	void GetStackFunctionInputs(const UNiagaraNodeFunctionCall& FunctionCallNode, TArray<FNiagaraVariable>& OutInputVariables, FCompileConstantResolver ConstantResolver, ENiagaraGetStackFunctionInputPinsOptions Options, bool bIgnoreDisabled = false);
+	NIAGARAEDITOR_API void GetStackFunctionInputs(const UNiagaraNodeFunctionCall& FunctionCallNode, TArray<FNiagaraVariable>& OutInputVariables, TSet<FNiagaraVariable>& OutHiddenVariables, FCompileConstantResolver ConstantResolver, ENiagaraGetStackFunctionInputPinsOptions Options, bool bIgnoreDisabled = false);
+	NIAGARAEDITOR_API void GetStackFunctionInputs(const UNiagaraNodeFunctionCall& FunctionCallNode, TArray<FNiagaraVariable>& OutInputVariables, FCompileConstantResolver ConstantResolver, ENiagaraGetStackFunctionInputPinsOptions Options, bool bIgnoreDisabled = false);
 
 	/* Returns the input pins for the given function call node.  Bypasses the module level caching of this data and generates it each call. */
 	void GetStackFunctionInputPinsWithoutCache(
@@ -188,7 +197,9 @@ namespace FNiagaraStackGraphUtilities
 	/* Module script calls do not have direct inputs, but rely on the parameter map being initialized correctly. This utility function resolves which of the module's parameters are reachable during compilation and returns a list of pins on the parameter map node that do not have to be compiled. */
 	TArray<UEdGraphPin*> GetUnusedFunctionInputPins(const UNiagaraNodeFunctionCall& FunctionCallNode, FCompileConstantResolver ConstantResolver);
 
-	void GetStackFunctionStaticSwitchPins(const UNiagaraNodeFunctionCall& FunctionCallNode, TArray<UEdGraphPin*>& OutInputPins, TSet<UEdGraphPin*>& OutHiddenPins, FCompileConstantResolver& ConstantResolver);
+	void GetStackFunctionStaticSwitchPinsLegacy(const UNiagaraNodeFunctionCall& FunctionCallNode, TArray<UEdGraphPin*>& OutInputPins, TSet<UEdGraphPin*>& OutHiddenPins, FCompileConstantResolver& ConstantResolver);
+
+	NIAGARAEDITOR_API void GetStackFunctionStaticSwitchPins(const UNiagaraNodeFunctionCall& FunctionCallNode, TArray<UEdGraphPin*>& OutInputPins, TSet<UEdGraphPin*>& OutHiddenPins, FCompileConstantResolver& ConstantResolver);
 
 	void GetStackFunctionOutputVariables(UNiagaraNodeFunctionCall& FunctionCallNode, FCompileConstantResolver ConstantResolver, TArray<FNiagaraVariable>& OutOutputVariables, TArray<FNiagaraVariable>& OutOutputVariablesWithOriginalAliasesIntact);
 
