- Remove Force Solo from being a UPROPERTY as we do not want to save the state, or transfer is when duplicating components
- Force Solo can still be called from BP / C++ to temporarily disable batching

diff --git a/Engine/Plugins/FX/Niagara/Source/Niagara/Public/NiagaraComponent.h b/Engine/Plugins/FX/Niagara/Source/Niagara/Public/NiagaraComponent.h
--- a/Engine/Plugins/FX/Niagara/Source/Niagara/Public/NiagaraComponent.h
+++ b/Engine/Plugins/FX/Niagara/Source/Niagara/Public/NiagaraComponent.h
@@ -125,7 +125,6 @@ private:
 	/**
 	When true, this component's system will be force to update via a slower "solo" path rather than the more optimal batched path with other instances of the same system.
 	*/
-	UPROPERTY(EditAnywhere, Category = Parameters)
 	uint32 bForceSolo : 1;
 
 	/** When true the GPU simulation debug display will enabled, allowing information used during simulation to be visualized. */
