Fix for digested NPC references with the compilation mode.  We no longer clear the database copies with PreGC as the db is the owner and the NPC are already persistent within the edtior.
#rnx

diff --git a/Engine/Plugins/FX/Niagara/Source/NiagaraEditor/Private/NiagaraCompilationTasks.cpp b/Engine/Plugins/FX/Niagara/Source/NiagaraEditor/Private/NiagaraCompilationTasks.cpp
--- a/Engine/Plugins/FX/Niagara/Source/NiagaraEditor/Private/NiagaraCompilationTasks.cpp
+++ b/Engine/Plugins/FX/Niagara/Source/NiagaraEditor/Private/NiagaraCompilationTasks.cpp
@@ -55,7 +55,8 @@ namespace NiagaraSystemCookStats
 
 namespace NiagaraCompilationTasksImpl
 {
-	static const FGuid UE_NIAGARA_ASYNC_TASK_COMPILER_VER = FGuid(0x3076A815, 0x9F4F4E23, 0x99C2573E, 0xAC120D18);
+	static const FGuid UE_NIAGARA_ASYNC_TASK_COMPILER_VER = FGuid(0xB2E44F66, 0x59164A73, 0x8D687867, 0xFBBBB30B);
+	
 	static UE::DerivedData::FCacheBucket NiagaraDDCBucket("NiagaraScript");
 
 	void GetUsagesToDuplicate(ENiagaraScriptUsage TargetUsage, TArray<ENiagaraScriptUsage>& DuplicateUsages)
diff --git a/Engine/Plugins/FX/Niagara/Source/NiagaraEditor/Private/NiagaraDigestDatabase.cpp b/Engine/Plugins/FX/Niagara/Source/NiagaraEditor/Private/NiagaraDigestDatabase.cpp
--- a/Engine/Plugins/FX/Niagara/Source/NiagaraEditor/Private/NiagaraDigestDatabase.cpp
+++ b/Engine/Plugins/FX/Niagara/Source/NiagaraEditor/Private/NiagaraDigestDatabase.cpp
@@ -45,10 +45,11 @@ void FNiagaraDigestDatabase::Shutdown()
 FNiagaraDigestDatabase::FNiagaraDigestDatabase()
 	: CompilationGraphCache(NiagaraGraphDigestDatabaseImpl::GDigestGraphCacheSize)
 {
-	// we use GC as an opportunity to clear the cache, preventing the database from holding onto references and
+	// we use GC as an opportunity to clear the graph cache, preventing the database from holding onto references and
 	// extending object lifetimes unnecessarily.  In particular the digested graphs contain duplicates of the DI
-	// which could be referencing real objects
-	FCoreUObjectDelegates::GetPreGarbageCollectDelegate().AddRaw(this, &FNiagaraDigestDatabase::ReleaseDatabase);
+	// which could be referencing real objects.  Note that the NPC cache isn't cleared as it is a duplicate owned by
+	// this database that mirrors the universal data (i.e. the NPC themselves aren't released when running in the editor)
+	FCoreUObjectDelegates::GetPreGarbageCollectDelegate().AddRaw(this, &FNiagaraDigestDatabase::ReleaseGraphCache);
 }
 
 FNiagaraDigestDatabase::~FNiagaraDigestDatabase()
@@ -58,6 +59,13 @@ FNiagaraDigestDatabase::~FNiagaraDigestDatabase()
 	FCoreUObjectDelegates::GetPreGarbageCollectDelegate().RemoveAll(this);
 }
 
+void FNiagaraDigestDatabase::ReleaseGraphCache()
+{
+	FWriteScopeLock WriteScope(DigestCacheLock);
+
+	CompilationGraphCache.Empty(NiagaraGraphDigestDatabaseImpl::GDigestGraphCacheSize);
+}
+
 void FNiagaraDigestDatabase::ReleaseDatabase()
 {
 	FWriteScopeLock WriteScope(DigestCacheLock);
diff --git a/Engine/Plugins/FX/Niagara/Source/NiagaraEditor/Private/NiagaraDigestDatabase.h b/Engine/Plugins/FX/Niagara/Source/NiagaraEditor/Private/NiagaraDigestDatabase.h
--- a/Engine/Plugins/FX/Niagara/Source/NiagaraEditor/Private/NiagaraDigestDatabase.h
+++ b/Engine/Plugins/FX/Niagara/Source/NiagaraEditor/Private/NiagaraDigestDatabase.h
@@ -120,6 +120,8 @@ public:
 	virtual FString GetReferencerName() const override;
 
 protected:
+	void ReleaseGraphCache();
+
 	uint32 GraphCacheHits = 0;
 	uint32 GraphCacheMisses = 0;
 	uint32 CollectionCacheHits = 0;
