- Add customization for stateless emitter properties
- Auto expands the data by default to match stateful
- Hides some properties and uses inline toggles to match stateful
- Adds a missing bool control to enable / disable loop delay

#rnx

diff --git a/Engine/Plugins/FX/Niagara/Source/Niagara/Internal/NiagaraSystemEmitterState.h b/Engine/Plugins/FX/Niagara/Source/Niagara/Internal/NiagaraSystemEmitterState.h
--- a/Engine/Plugins/FX/Niagara/Source/Niagara/Internal/NiagaraSystemEmitterState.h
+++ b/Engine/Plugins/FX/Niagara/Source/Niagara/Internal/NiagaraSystemEmitterState.h
@@ -35,6 +35,13 @@ enum class ENiagaraLoopBehavior : uint8
 	Once,
 };
 
+UENUM()
+enum class ENiagaraLoopDurationMode : uint8
+{
+	Fixed,
+	Infinite,
+};
+
 USTRUCT()
 struct FNiagaraSystemStateData
 {
@@ -49,13 +56,16 @@ struct FNiagaraSystemStateData
 	UPROPERTY(EditAnywhere, Category = "System State")
 	uint32 bIgnoreSystemState : 1 = true;
 
-	UPROPERTY(EditAnywhere, Category = "System State", meta = (DisplayAfter = "LoopCount", EditCondition = "LoopBehavior != ENiagaraLoopBehavior::Once", EditConditionHides))
+	UPROPERTY(EditAnywhere, Category = "System State", meta = (DisplayAfter = "LoopDuration", EditCondition = "LoopBehavior != ENiagaraLoopBehavior::Once", EditConditionHides))
 	uint32 bRecalculateDurationEachLoop : 1 = false;
 
-	UPROPERTY(EditAnywhere, Category = "System State", meta = (DisplayAfter = "LoopDelay", EditCondition = "LoopBehavior != ENiagaraLoopBehavior::Once", EditConditionHides))
+	UPROPERTY(EditAnywhere, Category = "System State", meta = (InlineEditConditionToggle))
+	uint32 bLoopDelayEnabled : 1 = false;
+
+	UPROPERTY(EditAnywhere, Category = "System State", meta = (DisplayAfter = "LoopDelay", EditCondition = "bLoopDelayEnabled && LoopBehavior != ENiagaraLoopBehavior::Once", EditConditionHides))
 	uint32 bDelayFirstLoopOnly : 1 = false;
 
-	UPROPERTY(EditAnywhere, Category = "System State", meta = (DisplayAfter = "bDelayFirstLoopOnly", EditCondition = "LoopBehavior != ENiagaraLoopBehavior::Once && !bDelayFirstLoopOnly", EditConditionHides))
+	UPROPERTY(EditAnywhere, Category = "System State", meta = (DisplayAfter = "bDelayFirstLoopOnly", EditCondition = "bLoopDelayEnabled && LoopBehavior != ENiagaraLoopBehavior::Once && !bDelayFirstLoopOnly", EditConditionHides))
 	uint32 bRecalculateDelayEachLoop : 1 = false;
 
 	UPROPERTY(EditAnywhere, Category = "System State")
@@ -70,7 +80,7 @@ struct FNiagaraSystemStateData
 	UPROPERTY(EditAnywhere, Category = "System State", meta = (ClampMin = "1", EditCondition = "LoopBehavior == ENiagaraLoopBehavior::Multiple", EditConditionHides))
 	int LoopCount = 1;
 
-	UPROPERTY(EditAnywhere, Category = "System State", meta = (ClampMin = "0.0", Units="s"))
+	UPROPERTY(EditAnywhere, Category = "System State", meta = (ClampMin = "0.0", Units="s", EditCondition = "bLoopDelayEnabled"))
 	FNiagaraDistributionRangeFloat LoopDelay = FNiagaraDistributionRangeFloat(0.0f);
 };
 
@@ -91,51 +101,72 @@ struct FNiagaraEmitterStateData
 	UPROPERTY(EditAnywhere, Category = "Emitter State", meta = (ClampMin = "1", EditCondition = "LoopBehavior == ENiagaraLoopBehavior::Multiple", EditConditionHides))
 	int32 LoopCount = 1;
 
-	UPROPERTY(EditAnywhere, Category = "Emitter State", meta = (ClampMin = "0.0", Units="s"))
+	UPROPERTY(EditAnywhere, Category = "Emitter State", meta = (EditCondition = "LoopBehavior == ENiagaraLoopBehavior::Once", EditConditionHides, SegmentedDisplay))
+	ENiagaraLoopDurationMode LoopDurationMode = ENiagaraLoopDurationMode::Fixed;
+
+	UPROPERTY(EditAnywhere, Category = "Emitter State", meta = (EditCondition = "LoopBehavior != ENiagaraLoopBehavior::Once || (LoopBehavior == ENiagaraLoopBehavior::Once && LoopDurationMode == ENiagaraLoopDurationMode::Fixed)", EditConditionHides, ClampMin = "0.0", Units="s"))
 	FNiagaraDistributionRangeFloat LoopDuration = FNiagaraDistributionRangeFloat(1.0f);
 
-	UPROPERTY(EditAnywhere, Category = "Emitter State", meta = (ClampMin = "0.0", Units="s"))
+	UPROPERTY(EditAnywhere, Category = "Emitter State", meta = (ClampMin = "0.0", Units="s", EditCondition="bLoopDelayEnabled"))
 	FNiagaraDistributionRangeFloat LoopDelay = FNiagaraDistributionRangeFloat(0.0f);
 
+	UPROPERTY(EditAnywhere, Category = "Emitter State", meta = (InlineEditConditionToggle))
+	uint32 bLoopDelayEnabled : 1 = false;
+
 	UPROPERTY(EditAnywhere, Category = "Emitter State", meta = (EditCondition = "LoopBehavior != ENiagaraLoopBehavior::Once", EditConditionHides, DisplayAfter = "LoopDuration"))
 	uint32 bRecalculateDurationEachLoop : 1 = false;
 
-	UPROPERTY(EditAnywhere, Category = "Emitter State", meta = (EditCondition = "LoopBehavior != ENiagaraLoopBehavior::Once", EditConditionHides))
+	UPROPERTY(EditAnywhere, Category = "Emitter State", meta = (EditCondition = "LoopBehavior != ENiagaraLoopBehavior::Once && bLoopDelayEnabled", EditConditionHides))
 	uint32 bDelayFirstLoopOnly : 1 = false;
 
-	UPROPERTY(EditAnywhere, Category = "Emitter State", meta = (EditCondition = "LoopBehavior != ENiagaraLoopBehavior::Once && !bDelayFirstLoopOnly", EditConditionHides))
+	UPROPERTY(EditAnywhere, Category = "Emitter State", meta = (EditCondition = "LoopBehavior != ENiagaraLoopBehavior::Once && bLoopDelayEnabled && !bDelayFirstLoopOnly", EditConditionHides))
 	uint32 bRecalculateDelayEachLoop : 1 = false;
 
-	UPROPERTY(EditAnywhere, Category = "Scalability")
+	UPROPERTY(EditAnywhere, Category = "Emitter Scalability")
 	uint32 bEnableDistanceCulling : 1 = false;
 
-	UPROPERTY(EditAnywhere, Category = "Scalability", meta = (DisplayAfter = "MaxDistanceReaction"))
+	UPROPERTY(EditAnywhere, Category = "Emitter Scalability", meta = (DisplayAfter = "MaxDistanceReaction"))
 	uint32 bEnableVisibilityCulling : 1 = false;
 
-	UPROPERTY(EditAnywhere, Category = "Scalability", meta = (EditCondition = "bEnableDistanceCulling", EditConditionHides))
+	UPROPERTY(EditAnywhere, Category = "Emitter Scalability", meta = (InlineEditConditionToggle))
 	uint32 bMinDistanceEnabled : 1 = false;
 
-	UPROPERTY(EditAnywhere, Category = "Scalability", meta = (EditCondition = "bEnableDistanceCulling", EditConditionHides))
+	UPROPERTY(EditAnywhere, Category = "Emitter Scalability", meta = (InlineEditConditionToggle))
 	uint32 bMaxDistanceEnabled : 1 = false;
 
-	UPROPERTY(EditAnywhere, Category = "Scalability", meta = (DisplayAfter = "VisibilityCullDelay"))
+	UPROPERTY(EditAnywhere, Category = "Emitter Scalability", meta = (DisplayAfter = "VisibilityCullDelay"))
 	uint32 bResetAgeOnAwaken : 1 = false;
 
-	UPROPERTY(EditAnywhere, Category = "Scalability", meta = (EditCondition = "bEnableDistanceCulling && bMinDistanceEnabled", EditConditionHides))
+	UPROPERTY(EditAnywhere, Category = "Emitter Scalability", meta = (DisplayAfter = "bEnableDistanceCulling", EditCondition = "bMinDistanceEnabled"))
 	float MinDistance = 0.0f;
 
-	UPROPERTY(EditAnywhere, Category = "Scalability", meta = (EditCondition = "bEnableDistanceCulling && bMinDistanceEnabled", EditConditionHides))
+	UPROPERTY(EditAnywhere, Category = "Emitter Scalability", meta = (DisplayAfter = "MinDistance"))
 	ENiagaraExecutionStateManagement MinDistanceReaction = ENiagaraExecutionStateManagement::Awaken;
 
-	UPROPERTY(EditAnywhere, Category = "Scalability", meta = (EditCondition = "bEnableDistanceCulling && bMaxDistanceEnabled", EditConditionHides))
+	UPROPERTY(EditAnywhere, Category = "Emitter Scalability", meta = (DisplayAfter = "bEnableDistanceCulling", EditCondition = "bMaxDistanceEnabled"))
 	float MaxDistance = 5000.0f;
 
-	UPROPERTY(EditAnywhere, Category = "Scalability", meta = (EditCondition = "bEnableDistanceCulling && bMaxDistanceEnabled", EditConditionHides))
+	UPROPERTY(EditAnywhere, Category = "Emitter Scalability", meta = (DisplayAfter = "MaxDistance"))
 	ENiagaraExecutionStateManagement MaxDistanceReaction = ENiagaraExecutionStateManagement::SleepAndLetParticlesFinish;
 
-	UPROPERTY(EditAnywhere, Category = "Scalability", meta = (EditCondition = "bEnableVisibilityCulling", EditConditionHides))
+	UPROPERTY(EditAnywhere, Category = "Emitter Scalability", meta = (EditCondition = "bEnableVisibilityCulling", EditConditionHides))
 	ENiagaraExecutionStateManagement VisibilityCullReaction = ENiagaraExecutionStateManagement::SleepAndLetParticlesFinish;
 
-	UPROPERTY(EditAnywhere, Category = "Scalability", meta = (EditCondition = "bEnableVisibilityCulling", EditConditionHides))
+	UPROPERTY(EditAnywhere, Category = "Emitter Scalability", meta = (EditCondition = "bEnableVisibilityCulling", EditConditionHides))
 	float VisibilityCullDelay = 1.0f;
+
+#if WITH_EDITORONLY_DATA
+	void PostSerialize(const FArchive& Ar);
+#endif
+};
+
+#if WITH_EDITORONLY_DATA
+template<>
+struct TStructOpsTypeTraits<FNiagaraEmitterStateData> : public TStructOpsTypeTraitsBase2<FNiagaraEmitterStateData>
+{
+	enum
+	{
+		WithPostSerialize = true,
+	};
 };
+#endif
diff --git a/Engine/Plugins/FX/Niagara/Source/Niagara/Internal/Stateless/NiagaraStatelessEmitter.h b/Engine/Plugins/FX/Niagara/Source/Niagara/Internal/Stateless/NiagaraStatelessEmitter.h
--- a/Engine/Plugins/FX/Niagara/Source/Niagara/Internal/Stateless/NiagaraStatelessEmitter.h
+++ b/Engine/Plugins/FX/Niagara/Source/Niagara/Internal/Stateless/NiagaraStatelessEmitter.h
@@ -133,10 +133,10 @@ protected:
 	UPROPERTY()
 	FString UniqueEmitterName;
 
-	UPROPERTY(EditAnywhere, Category = "General", meta = (AllowedClasses = "/Script/Niagara.NiagaraStatelessEmitterTemplate", HideInStack))
+	UPROPERTY(EditAnywhere, Category = "Emitter Properties", meta = (AllowedClasses = "/Script/Niagara.NiagaraStatelessEmitterTemplate", HideInStack))
 	TObjectPtr<UClass> EmitterTemplateClass;
 
-	UPROPERTY(EditAnywhere, Category = "General")
+	UPROPERTY(EditAnywhere, Category = "Emitter Properties")
 	uint32 bDeterministic : 1 = false;
 
 #if WITH_EDITORONLY_DATA
@@ -144,24 +144,24 @@ protected:
 	When enabled the emitter will output all available attributes.
 	You should not need to modify this with the exception of debugging / testing and as it will impact cooked performance and memory
 	*/
-	UPROPERTY(EditAnywhere, Category = "General", AdvancedDisplay)
+	UPROPERTY(EditAnywhere, Category = "Emitter Properties", AdvancedDisplay)
 	uint32 bForceOutputAllAttributes : 1 = false;
 
 	/**
 	When enabled the emitter will always include UniqueID in the output attributes.
 	You should not need to modify this with the exception of debugging / testing and as it will impact cooked performance and memory
 	*/
-	UPROPERTY(EditAnywhere, Category = "General", AdvancedDisplay, meta = (EditCondition = "!bForceOutputAllAttributes"))
+	UPROPERTY(EditAnywhere, Category = "Emitter Properties", AdvancedDisplay, meta = (EditCondition = "!bForceOutputAllAttributes"))
 	uint32 bForceOutputUniqueID : 1 = false;
 #endif
 
-	UPROPERTY(EditAnywhere, Category = "General", AdvancedDisplay, meta = (Bitmask, BitMaskEnum = "/Script/Niagara.ENiagaraStatelessFeatureMask"))
+	UPROPERTY(EditAnywhere, Category = "Emitter Properties", AdvancedDisplay, meta = (Bitmask, BitMaskEnum = "/Script/Niagara.ENiagaraStatelessFeatureMask"))
 	uint32 AllowedFeatureMask = uint32(ENiagaraStatelessFeatureMask::All);
 
-	UPROPERTY(EditAnywhere, Category = "General")
+	UPROPERTY(EditAnywhere, Category = "Emitter Properties")
 	int32 RandomSeed = 0;
 
-	UPROPERTY(EditAnywhere, Category = "General")
+	UPROPERTY(EditAnywhere, Category = "Emitter Properties")
 	FBox FixedBounds = FBox(FVector(-100), FVector(100));
 
 	UPROPERTY(EditAnywhere, Category = "Emitter State")
diff --git a/Engine/Plugins/FX/Niagara/Source/Niagara/Private/NiagaraSystemEmitterState.cpp b/Engine/Plugins/FX/Niagara/Source/Niagara/Private/NiagaraSystemEmitterState.cpp
new file mode 100644
--- /dev/null
+++ b/Engine/Plugins/FX/Niagara/Source/Niagara/Private/NiagaraSystemEmitterState.cpp
@@ -0,0 +1,18 @@
+// Copyright Epic Games, Inc. All Rights Reserved.
+
+#include "NiagaraSystemEmitterState.h"
+#include "NiagaraCustomVersion.h"
+
+#if WITH_EDITORONLY_DATA
+void FNiagaraEmitterStateData::PostSerialize(const FArchive& Ar)
+{
+	if (Ar.IsLoading() && Ar.CustomVer(FNiagaraCustomVersion::GUID) < FNiagaraCustomVersion::EmitterStateAddLoopDelayEnabled)
+	{
+		const FNiagaraStatelessRangeFloat LoopDelayRange = LoopDelay.CalculateRange();
+		if ( LoopDelayRange.ParameterOffset != INDEX_NONE || !FMath::IsNearlyZero(LoopDelayRange.Min) || !FMath::IsNearlyZero(LoopDelayRange.Max) )
+		{
+			bLoopDelayEnabled = true;
+		}
+	}
+}
+#endif
diff --git a/Engine/Plugins/FX/Niagara/Source/Niagara/Private/Stateless/NiagaraStatelessEmitterInstance.cpp b/Engine/Plugins/FX/Niagara/Source/Niagara/Private/Stateless/NiagaraStatelessEmitterInstance.cpp
--- a/Engine/Plugins/FX/Niagara/Source/Niagara/Private/Stateless/NiagaraStatelessEmitterInstance.cpp
+++ b/Engine/Plugins/FX/Niagara/Source/Niagara/Private/Stateless/NiagaraStatelessEmitterInstance.cpp
@@ -284,14 +284,29 @@ void FNiagaraStatelessEmitterInstance::InitEmitterState()
 {
 	using namespace NiagaraStatelessEmitterInstancePrivate;
 
+
 	const FNiagaraEmitterStateData& EmitterState = EmitterData->EmitterState;
 	LoopCount			= 0;
-	CurrentLoopDuration = EvaluateDistribution(EmitterState.LoopDuration, RandomStream, RendererBindings, DefaultLoopDuration);
-	CurrentLoopDuration = FMath::Max(CurrentLoopDuration, DefaultLoopDuration);
-	CurrentLoopDelay	= EvaluateDistribution(EmitterState.LoopDelay, RandomStream, RendererBindings, DefaultLoopDelay);
-	CurrentLoopDelay	= FMath::Max(CurrentLoopDelay, 0.0f);
+
+	CurrentLoopDelay = 0.0f;
+	if (EmitterState.bLoopDelayEnabled)
+	{
+		CurrentLoopDelay = EvaluateDistribution(EmitterState.LoopDelay, RandomStream, RendererBindings, DefaultLoopDelay);
+		CurrentLoopDelay = FMath::Max(CurrentLoopDelay, 0.0f);
+	}
 	CurrentLoopAgeStart	= 0.0f;
-	CurrentLoopAgeEnd	= CurrentLoopAgeStart + CurrentLoopDelay + CurrentLoopDuration;
+
+	if (EmitterState.LoopBehavior == ENiagaraLoopBehavior::Once && EmitterState.LoopDurationMode == ENiagaraLoopDurationMode::Infinite)
+	{
+		CurrentLoopDuration	= FLT_MAX;
+		CurrentLoopAgeEnd	= FLT_MAX;
+	}
+	else
+	{
+		CurrentLoopDuration = EvaluateDistribution(EmitterState.LoopDuration, RandomStream, RendererBindings, DefaultLoopDuration);
+		CurrentLoopDuration = FMath::Max(CurrentLoopDuration, DefaultLoopDuration);
+		CurrentLoopAgeEnd	= CurrentLoopAgeStart + CurrentLoopDelay + CurrentLoopDuration;
+	}
 }
 
 void FNiagaraStatelessEmitterInstance::TickEmitterState()
@@ -421,14 +436,17 @@ void FNiagaraStatelessEmitterInstance::TickEmitterState()
 					CurrentLoopDuration = FMath::Max(CurrentLoopDuration, DefaultLoopDuration);
 				}
 
-				if (EmitterState.bDelayFirstLoopOnly)
-				{
-					CurrentLoopDelay = 0.0f;
-				}
-				else if (EmitterState.bRecalculateDelayEachLoop)
+				if (EmitterState.bLoopDelayEnabled)
 				{
-					CurrentLoopDelay = EvaluateDistribution(EmitterState.LoopDelay, RandomStream, RendererBindings, DefaultLoopDelay);
-					CurrentLoopDelay = FMath::Max(CurrentLoopDelay, 0.0f);
+					if (EmitterState.bDelayFirstLoopOnly)
+					{
+						CurrentLoopDelay = 0.0f;
+					}
+					else if (EmitterState.bRecalculateDelayEachLoop)
+					{
+						CurrentLoopDelay = EvaluateDistribution(EmitterState.LoopDelay, RandomStream, RendererBindings, DefaultLoopDelay);
+						CurrentLoopDelay = FMath::Max(CurrentLoopDelay, 0.0f);
+					}
 				}
 
 				CurrentLoopAgeStart = CurrentLoopAgeEnd;
diff --git a/Engine/Plugins/FX/Niagara/Source/NiagaraCore/Public/NiagaraCustomVersion.h b/Engine/Plugins/FX/Niagara/Source/NiagaraCore/Public/NiagaraCustomVersion.h
--- a/Engine/Plugins/FX/Niagara/Source/NiagaraCore/Public/NiagaraCustomVersion.h
+++ b/Engine/Plugins/FX/Niagara/Source/NiagaraCore/Public/NiagaraCustomVersion.h
@@ -208,6 +208,8 @@ struct FNiagaraCustomVersion
 		StatelessInitialMeshOrientationV1,	// Update Initial Mesh Orientation Module
 
 		HierarchyEditorScriptSupport, // Hierarchy Editor was implemented
+
+		EmitterStateAddLoopDelayEnabled,	// Added loop delay enabled to emitter state
 		
 		// DO NOT ADD A NEW VERSION UNLESS YOU HAVE TALKED TO THE NIAGARA LEAD. Mismanagement of these versions can lead to data loss if it is adjusted in multiple streams simultaneously.
 		// -----<new versions can be added above this line>  -------------------------------------------------
diff --git a/Engine/Plugins/FX/Niagara/Source/NiagaraEditor/Private/Stateless/NiagaraEmitterStatePropertyCustomization.cpp b/Engine/Plugins/FX/Niagara/Source/NiagaraEditor/Private/Stateless/NiagaraEmitterStatePropertyCustomization.cpp
new file mode 100644
--- /dev/null
+++ b/Engine/Plugins/FX/Niagara/Source/NiagaraEditor/Private/Stateless/NiagaraEmitterStatePropertyCustomization.cpp
@@ -0,0 +1,94 @@
+// Copyright Epic Games, Inc. All Rights Reserved.
+
+#include "Stateless/NiagaraEmitterStatePropertyCustomization.h"
+#include "NiagaraSystemEmitterState.h"
+
+#include "DetailLayoutBuilder.h"
+#include "DetailCategoryBuilder.h"
+#include "IDetailChildrenBuilder.h"
+#include "IDetailGroup.h"
+#include "DetailWidgetRow.h"
+
+#define LOCTEXT_NAMESPACE "NiagaraEmitterStatePropertyCustomization"
+
+
+TSharedRef<IPropertyTypeCustomization> FNiagaraEmitterStatePropertyCustomization::MakeInstance()
+{
+	return MakeShared<FNiagaraEmitterStatePropertyCustomization>();
+}
+
+void FNiagaraEmitterStatePropertyCustomization::CustomizeHeader(TSharedRef<IPropertyHandle> PropertyHandle, FDetailWidgetRow& HeaderRow, IPropertyTypeCustomizationUtils& CustomizationUtils)
+{
+}
+
+void FNiagaraEmitterStatePropertyCustomization::CustomizeChildren(TSharedRef<IPropertyHandle> StructPropertyHandle, IDetailChildrenBuilder& ChildBuilder, IPropertyTypeCustomizationUtils& CustomizationUtils)
+{
+	EmitterStatePropertyHandle = StructPropertyHandle;
+	EmitterStateOwnerObject = nullptr;
+	{
+		TArray<UObject*> OuterObjects;
+		StructPropertyHandle->GetOuterObjects(OuterObjects);
+		EmitterStateOwnerObject = OuterObjects.Num() == 1 ? OuterObjects[0] : nullptr;
+	}
+
+	if (EmitterStateOwnerObject == nullptr)
+	{
+		return;
+	}
+
+	IDetailLayoutBuilder& LayoutBuilder = ChildBuilder.GetParentCategory().GetParentLayout();
+	LayoutBuilder.EditCategory(FName("Emitter Properties"));
+
+	uint32 NumChildren = 0;
+	StructPropertyHandle->GetNumChildren(NumChildren);
+	for (uint32 iChild=0; iChild < NumChildren; ++iChild)
+	{
+		TSharedPtr<IPropertyHandle> ChildHandle = StructPropertyHandle->GetChildHandle(iChild);
+		IDetailPropertyRow& PropertyRow = LayoutBuilder.AddPropertyToCategory(ChildHandle.ToSharedRef());
+
+		const FName PropertyName = ChildHandle->GetProperty()->GetFName();
+
+		if (PropertyName == GET_MEMBER_NAME_CHECKED(FNiagaraEmitterStateData, MinDistance) || PropertyName == GET_MEMBER_NAME_CHECKED(FNiagaraEmitterStateData, MaxDistance))
+		{
+			PropertyRow.Visibility(TAttribute<EVisibility>::CreateSP(this, &FNiagaraEmitterStatePropertyCustomization::GetEnableDistanceCullingVisibility));
+		}
+		else if (PropertyName == GET_MEMBER_NAME_CHECKED(FNiagaraEmitterStateData, MinDistanceReaction))
+		{
+			PropertyRow.Visibility(TAttribute<EVisibility>::CreateSP(this, &FNiagaraEmitterStatePropertyCustomization::GetMinDistanceVisibility));
+		}
+		else if (PropertyName == GET_MEMBER_NAME_CHECKED(FNiagaraEmitterStateData, MaxDistanceReaction))
+		{
+			PropertyRow.Visibility(TAttribute<EVisibility>::CreateSP(this, &FNiagaraEmitterStatePropertyCustomization::GetMaxDistanceVisibility));
+		}
+	}
+}
+
+const FNiagaraEmitterStateData* FNiagaraEmitterStatePropertyCustomization::GetEmitterState() const
+{
+	uint8* ObjectAddress = reinterpret_cast<uint8*>(EmitterStateOwnerObject.Get());
+	if (EmitterStatePropertyHandle && ObjectAddress)
+	{
+		return reinterpret_cast<const FNiagaraEmitterStateData*>(EmitterStatePropertyHandle->GetValueBaseAddress(ObjectAddress));
+	}
+	return nullptr;
+}
+
+EVisibility FNiagaraEmitterStatePropertyCustomization::GetEnableDistanceCullingVisibility() const
+{
+	const FNiagaraEmitterStateData* EmitterState = GetEmitterState();
+	return EmitterState && EmitterState->bEnableDistanceCulling ? EVisibility::Visible : EVisibility::Hidden;
+}
+
+EVisibility FNiagaraEmitterStatePropertyCustomization::GetMinDistanceVisibility() const
+{
+	const FNiagaraEmitterStateData* EmitterState = GetEmitterState();
+	return EmitterState && EmitterState->bEnableDistanceCulling && EmitterState->bMinDistanceEnabled ? EVisibility::Visible : EVisibility::Hidden;
+}
+
+EVisibility FNiagaraEmitterStatePropertyCustomization::GetMaxDistanceVisibility() const
+{
+	const FNiagaraEmitterStateData* EmitterState = GetEmitterState();
+	return EmitterState && EmitterState->bEnableDistanceCulling && EmitterState->bMaxDistanceEnabled ? EVisibility::Visible : EVisibility::Hidden;
+}
+
+#undef LOCTEXT_NAMESPACE
diff --git a/Engine/Plugins/FX/Niagara/Source/NiagaraEditor/Private/Stateless/NiagaraEmitterStatePropertyCustomization.h b/Engine/Plugins/FX/Niagara/Source/NiagaraEditor/Private/Stateless/NiagaraEmitterStatePropertyCustomization.h
new file mode 100644
--- /dev/null
+++ b/Engine/Plugins/FX/Niagara/Source/NiagaraEditor/Private/Stateless/NiagaraEmitterStatePropertyCustomization.h
@@ -0,0 +1,26 @@
+// Copyright Epic Games, Inc. All Rights Reserved.
+
+#pragma once
+
+#include "PropertyHandle.h"
+#include "IPropertyTypeCustomization.h"
+
+struct FNiagaraEmitterStateData;
+
+class FNiagaraEmitterStatePropertyCustomization : public IPropertyTypeCustomization
+{
+public:
+	static TSharedRef<IPropertyTypeCustomization> MakeInstance();
+	virtual void CustomizeHeader(TSharedRef<IPropertyHandle> PropertyHandle, FDetailWidgetRow& HeaderRow, IPropertyTypeCustomizationUtils& CustomizationUtils) override;
+	virtual void CustomizeChildren(TSharedRef<IPropertyHandle> PropertyHandle, IDetailChildrenBuilder& ChildBuilder, IPropertyTypeCustomizationUtils& CustomizationUtils) override;
+
+	const FNiagaraEmitterStateData* GetEmitterState() const;
+
+	EVisibility GetEnableDistanceCullingVisibility() const;
+	EVisibility GetMinDistanceVisibility() const;
+	EVisibility GetMaxDistanceVisibility() const;
+
+private:
+	TSharedPtr<IPropertyHandle> EmitterStatePropertyHandle;
+	TWeakObjectPtr<UObject>		EmitterStateOwnerObject;
+};
diff --git a/Engine/Plugins/FX/Niagara/Source/NiagaraEditor/Private/ViewModels/Stack/NiagaraStackStatelessEmitterGroup.cpp b/Engine/Plugins/FX/Niagara/Source/NiagaraEditor/Private/ViewModels/Stack/NiagaraStackStatelessEmitterGroup.cpp
--- a/Engine/Plugins/FX/Niagara/Source/NiagaraEditor/Private/ViewModels/Stack/NiagaraStackStatelessEmitterGroup.cpp
+++ b/Engine/Plugins/FX/Niagara/Source/NiagaraEditor/Private/ViewModels/Stack/NiagaraStackStatelessEmitterGroup.cpp
@@ -4,6 +4,7 @@
 
 #include "IDetailTreeNode.h"
 #include "PropertyHandle.h"
+#include "Stateless/NiagaraEmitterStatePropertyCustomization.h"
 #include "Stateless/NiagaraDistributionPropertyCustomization.h"
 #include "Stateless/NiagaraDistributionIntPropertyCustomization.h"
 #include "Stateless/NiagaraStatelessEmitter.h"
@@ -129,6 +130,7 @@ void UNiagaraStackStatelessEmitterObjectItem::RefreshChildrenInternal(const TArr
 			bool bHideTopLevelCategories = false;
 			StatelessEmitterStackObject->Initialize(CreateDefaultChildRequiredData(), StatelessEmitter, bIsTopLevelObject, bHideTopLevelCategories, GetStackEditorDataKey());
 			StatelessEmitterStackObject->SetOnFilterDetailNodes(OnFilterDetailNodes, UNiagaraStackObject::EDetailNodeFilterMode::FilterAllNodes);
+			StatelessEmitterStackObject->RegisterInstancedCustomPropertyTypeLayout(FNiagaraEmitterStateData::StaticStruct()->GetFName(), FOnGetPropertyTypeCustomizationInstance::CreateStatic(&FNiagaraEmitterStatePropertyCustomization::MakeInstance));
 			StatelessEmitterStackObject->RegisterInstancedCustomPropertyTypeLayout(FNiagaraDistributionFloat::StaticStruct()->GetFName(), FOnGetPropertyTypeCustomizationInstance::CreateStatic(&FNiagaraDistributionPropertyCustomization::MakeFloatInstance));
 			StatelessEmitterStackObject->RegisterInstancedCustomPropertyTypeLayout(FNiagaraDistributionVector2::StaticStruct()->GetFName(), FOnGetPropertyTypeCustomizationInstance::CreateStatic(&FNiagaraDistributionPropertyCustomization::MakeVector2Instance));
 			StatelessEmitterStackObject->RegisterInstancedCustomPropertyTypeLayout(FNiagaraDistributionVector3::StaticStruct()->GetFName(), FOnGetPropertyTypeCustomizationInstance::CreateStatic(&FNiagaraDistributionPropertyCustomization::MakeVector3Instance));
