A couple of DIs using the wrong context for Pre/Post stage.
Not an issue in current code but needs to be correct after upcoming change.

diff --git a/Engine/Plugins/FX/Niagara/Source/Niagara/Private/DataInterface/NiagaraDataInterfacePhysicsAsset.cpp b/Engine/Plugins/FX/Niagara/Source/Niagara/Private/DataInterface/NiagaraDataInterfacePhysicsAsset.cpp
--- a/Engine/Plugins/FX/Niagara/Source/Niagara/Private/DataInterface/NiagaraDataInterfacePhysicsAsset.cpp
+++ b/Engine/Plugins/FX/Niagara/Source/Niagara/Private/DataInterface/NiagaraDataInterfacePhysicsAsset.cpp
@@ -544,7 +544,7 @@ void FNDIPhysicsAssetProxy::DestroyPerInstanceData(const FNiagaraSystemInstanceI
 	SystemInstancesToProxyData.Remove(SystemInstance);
 }
 
-void FNDIPhysicsAssetProxy::PreStage(const FNDIGpuComputePostStageContext& Context)
+void FNDIPhysicsAssetProxy::PreStage(const FNDIGpuComputePreStageContext& Context)
 {
 	using namespace NDIPhysicsAssetLocal;
 
diff --git a/Engine/Plugins/FX/Niagara/Source/Niagara/Private/NiagaraDataInterfaceRigidMeshCollisionQuery.cpp b/Engine/Plugins/FX/Niagara/Source/Niagara/Private/NiagaraDataInterfaceRigidMeshCollisionQuery.cpp
--- a/Engine/Plugins/FX/Niagara/Source/Niagara/Private/NiagaraDataInterfaceRigidMeshCollisionQuery.cpp
+++ b/Engine/Plugins/FX/Niagara/Source/Niagara/Private/NiagaraDataInterfaceRigidMeshCollisionQuery.cpp
@@ -1056,7 +1056,7 @@ struct FNDIRigidMeshCollisionProxy : public FNiagaraDataInterfaceProxy
 	}
 
 	/** Launch all pre stage functions */
-	virtual void PreStage(const FNDIGpuComputePostStageContext& Context) override
+	virtual void PreStage(const FNDIGpuComputePreStageContext& Context) override
 	{
 		using namespace NDIRigidMeshCollisionLocal;
 
diff --git a/Engine/Plugins/FX/Niagara/Source/Niagara/Public/NiagaraDataInterfacePhysicsAsset.h b/Engine/Plugins/FX/Niagara/Source/Niagara/Public/NiagaraDataInterfacePhysicsAsset.h
--- a/Engine/Plugins/FX/Niagara/Source/Niagara/Public/NiagaraDataInterfacePhysicsAsset.h
+++ b/Engine/Plugins/FX/Niagara/Source/Niagara/Public/NiagaraDataInterfacePhysicsAsset.h
@@ -214,7 +214,7 @@ struct FNDIPhysicsAssetProxy : public FNiagaraDataInterfaceProxy
 	void DestroyPerInstanceData(const FNiagaraSystemInstanceID& SystemInstance);
 
 	/** Launch all pre stage functions */
-	virtual void PreStage(const FNDIGpuComputePostStageContext& Context) override;
+	virtual void PreStage(const FNDIGpuComputePreStageContext& Context) override;
 
 	/** List of proxy data for each system instances*/
 	TMap<FNiagaraSystemInstanceID, FNDIPhysicsAssetData> SystemInstancesToProxyData;
