- Stateless copy and paste support for spawn infos & modules

#rnx

diff --git a/Engine/Plugins/FX/Niagara/Source/NiagaraEditor/Private/ViewModels/Stack/NiagaraStackStatelessEmitterSimulateGroup.cpp b/Engine/Plugins/FX/Niagara/Source/NiagaraEditor/Private/ViewModels/Stack/NiagaraStackStatelessEmitterSimulateGroup.cpp
--- a/Engine/Plugins/FX/Niagara/Source/NiagaraEditor/Private/ViewModels/Stack/NiagaraStackStatelessEmitterSimulateGroup.cpp
+++ b/Engine/Plugins/FX/Niagara/Source/NiagaraEditor/Private/ViewModels/Stack/NiagaraStackStatelessEmitterSimulateGroup.cpp
@@ -3,6 +3,7 @@
 #include "ViewModels/Stack/NiagaraStackStatelessEmitterSimulateGroup.h"
 
 #include "IDetailTreeNode.h"
+#include "NiagaraClipboard.h"
 #include "NiagaraEditorStyle.h"
 #include "NiagaraStackEditorData.h"
 #include "PropertyHandle.h"
@@ -19,6 +20,65 @@
 
 #define LOCTEXT_NAMESPACE "NiagaraEmitterStatelessSimulateGroup"
 
+namespace NiagaraStackStatelessEmitterSimulateGroupPrivate
+{
+	bool TestCanPasteModules(UNiagaraStatelessEmitter* StatelessEmitter, const UNiagaraClipboardContent* ClipboardContent, FText& OutMessage)
+	{
+		if (StatelessEmitter)
+		{
+			for (const UObject* ClipboardModuleObject : ClipboardContent->StatelessModules)
+			{
+				const UNiagaraStatelessModule* ClipboardModule = Cast<const UNiagaraStatelessModule>(ClipboardModuleObject);
+				UNiagaraStatelessModule* StatelessModule = ClipboardModule ? StatelessEmitter->GetModule(ClipboardModule->GetClass()) : nullptr;
+				if (!StatelessModule)
+				{
+					continue;
+				}
+
+				OutMessage = LOCTEXT("CanPasteStatelessModules", "Paste module data, either adding or replacing existing module.");
+				return true;
+			}
+		}
+
+		OutMessage = LOCTEXT("CanPasteStatelessModuleUnsupported", "Incompatible or no data to paste.");
+		return false;
+	}
+
+	FText GetPasteModulesTransactionText(const UNiagaraClipboardContent* ClipboardContent)
+	{
+		return LOCTEXT("PasteStatelessModulesTransaction", "Paste modules(s).");
+	}
+
+	TArray<UObject*> PasteModules(UNiagaraStatelessEmitter* StatelessEmitter, const UNiagaraClipboardContent* ClipboardContent)
+	{
+		TArray<UObject*> ModifiedObjects;
+		if (StatelessEmitter)
+		{
+			for (const UObject* ClipboardModuleObject : ClipboardContent->StatelessModules)
+			{
+				const UNiagaraStatelessModule* ClipboardModule = Cast<const UNiagaraStatelessModule>(ClipboardModuleObject);
+				UNiagaraStatelessModule* StatelessModule = ClipboardModule ? StatelessEmitter->GetModule(ClipboardModule->GetClass()) : nullptr;
+				if (!StatelessModule)
+				{
+					continue;
+				}
+
+				StatelessModule->Modify();
+				UEngine::CopyPropertiesForUnrelatedObjects(const_cast<UNiagaraStatelessModule*>(ClipboardModule), StatelessModule);
+				if (StatelessModule->CanDisableModule())
+				{
+					StatelessModule->SetIsModuleEnabled(ClipboardModule->IsModuleEnabled());
+				}
+				StatelessModule->PostEditChange();
+
+				ModifiedObjects.Add(StatelessModule);
+			}
+		}
+
+		return ModifiedObjects;
+	}
+}
+
 class FNiagaraStatelessEmitterAddModuleAction : public INiagaraStackItemGroupAddAction
 {
 
@@ -110,6 +170,26 @@ const FSlateBrush* UNiagaraStackStatelessEmitterSimulateGroup::GetIconBrush() co
 	return FNiagaraEditorStyle::Get().GetBrush("NiagaraEditor.Stateless.UpdateIcon");
 }
 
+bool UNiagaraStackStatelessEmitterSimulateGroup::TestCanPasteWithMessage(const UNiagaraClipboardContent* ClipboardContent, FText& OutMessage) const
+{
+	return NiagaraStackStatelessEmitterSimulateGroupPrivate::TestCanPasteModules(GetStatelessEmitter(), ClipboardContent, OutMessage);
+}
+
+FText UNiagaraStackStatelessEmitterSimulateGroup::GetPasteTransactionText(const UNiagaraClipboardContent* ClipboardContent) const
+{
+	return NiagaraStackStatelessEmitterSimulateGroupPrivate::GetPasteModulesTransactionText(ClipboardContent);
+}
+
+void UNiagaraStackStatelessEmitterSimulateGroup::Paste(const UNiagaraClipboardContent* ClipboardContent, FText& OutPasteWarning)
+{
+	TArray<UObject*> ModifiedObjects = NiagaraStackStatelessEmitterSimulateGroupPrivate::PasteModules(GetStatelessEmitter(), ClipboardContent);
+	if (ModifiedObjects.Num() > 0)
+	{
+		OnDataObjectModified().Broadcast(ModifiedObjects, ENiagaraDataObjectChange::Changed);
+		RefreshChildren();
+	}
+}
+
 void UNiagaraStackStatelessEmitterSimulateGroup::RefreshChildrenInternal(const TArray<UNiagaraStackEntry*>& CurrentChildren, TArray<UNiagaraStackEntry*>& NewChildren, TArray<FStackIssue>& NewIssues)
 {
 	Super::RefreshChildrenInternal(CurrentChildren, NewChildren, NewIssues);
@@ -178,6 +258,49 @@ FText UNiagaraStackStatelessModuleItem::GetTooltipText() const
 	return Super::GetTooltipText();
 }
 
+bool UNiagaraStackStatelessModuleItem::TestCanCopyWithMessage(FText& OutMessage) const
+{
+	if (UNiagaraStatelessModule* StatelessModule = StatelessModuleWeak.Get())
+	{
+		OutMessage = LOCTEXT("CanCopyStatelessModule", "Copy module to the clipboard.");
+		return true;
+	}
+	OutMessage = LOCTEXT("CanCopyStatelessModuleUnsupported", "This module does not support copy.");
+	return false;
+}
+
+void UNiagaraStackStatelessModuleItem::Copy(UNiagaraClipboardContent* ClipboardContent) const
+{
+	UNiagaraStatelessModule* StatelessModule = StatelessModuleWeak.Get();
+	if (StatelessModule == nullptr)
+	{
+		return;
+	}
+
+	ClipboardContent->StatelessModules.Add(StaticDuplicateObject(StatelessModule, ClipboardContent));
+}
+
+bool UNiagaraStackStatelessModuleItem::TestCanPasteWithMessage(const UNiagaraClipboardContent* ClipboardContent, FText& OutMessage) const
+{
+	return NiagaraStackStatelessEmitterSimulateGroupPrivate::TestCanPasteModules(GetStatelessEmitter(), ClipboardContent, OutMessage);
+}
+
+FText UNiagaraStackStatelessModuleItem::GetPasteTransactionText(const UNiagaraClipboardContent* ClipboardContent) const
+{
+	return NiagaraStackStatelessEmitterSimulateGroupPrivate::GetPasteModulesTransactionText(ClipboardContent);
+}
+
+void UNiagaraStackStatelessModuleItem::Paste(const UNiagaraClipboardContent* ClipboardContent, FText& OutPasteWarning)
+{
+	TArray<UObject*> ModifiedObjects = NiagaraStackStatelessEmitterSimulateGroupPrivate::PasteModules(GetStatelessEmitter(), ClipboardContent);
+	if (ModifiedObjects.Num() > 0)
+	{
+		OnDataObjectModified().Broadcast(ModifiedObjects, ENiagaraDataObjectChange::Changed);
+		GetStackEditorData().Modify();
+		OnModifiedGroupItems().Broadcast();
+	}
+}
+
 bool UNiagaraStackStatelessModuleItem::TestCanDeleteWithMessage(FText& OutCanDeleteMessage) const
 {
 	UNiagaraStatelessModule* StatelessModule = StatelessModuleWeak.Get();
@@ -225,6 +348,12 @@ bool UNiagaraStackStatelessModuleItem::GetIsEnabled() const
 	return StatelessModule != nullptr && StatelessModule->IsModuleEnabled();
 }
 
+UNiagaraStatelessEmitter* UNiagaraStackStatelessModuleItem::GetStatelessEmitter() const
+{
+	UNiagaraStatelessModule* StatelessModule = GetStatelessModule();
+	return StatelessModule ? StatelessModule->GetTypedOuter<UNiagaraStatelessEmitter>() : nullptr;
+}
+
 void UNiagaraStackStatelessModuleItem::RefreshChildrenInternal(const TArray<UNiagaraStackEntry*>& CurrentChildren, TArray<UNiagaraStackEntry*>& NewChildren, TArray<FStackIssue>& NewIssues)
 {
 	Super::RefreshChildrenInternal(CurrentChildren, NewChildren, NewIssues);
diff --git a/Engine/Plugins/FX/Niagara/Source/NiagaraEditor/Private/ViewModels/Stack/NiagaraStackStatelessEmitterSpawnGroup.cpp b/Engine/Plugins/FX/Niagara/Source/NiagaraEditor/Private/ViewModels/Stack/NiagaraStackStatelessEmitterSpawnGroup.cpp
--- a/Engine/Plugins/FX/Niagara/Source/NiagaraEditor/Private/ViewModels/Stack/NiagaraStackStatelessEmitterSpawnGroup.cpp
+++ b/Engine/Plugins/FX/Niagara/Source/NiagaraEditor/Private/ViewModels/Stack/NiagaraStackStatelessEmitterSpawnGroup.cpp
@@ -15,9 +15,79 @@
 #include "ViewModels/NiagaraSystemViewModel.h"
 #include "ViewModels/Stack/NiagaraStackItemPropertyHeaderValueShared.h"
 #include "ViewModels/Stack/NiagaraStackObject.h"
+#include "NiagaraClipboard.h"
 
 #define LOCTEXT_NAMESPACE "NiagaraEmitterStatelessSpawnGroup"
 
+namespace NiagaraStackStatelessEmitterSpawnGroupPrivate
+{
+	TOptional<FNiagaraStatelessSpawnInfo> ConvertPortableValue(const FNiagaraClipboardPortableValue& PortableValue)
+	{
+		TOptional<FNiagaraStatelessSpawnInfo> ReturnValue;
+
+		FNiagaraStatelessSpawnInfo TempSpawnInfo;
+		if (PortableValue.TryUpdateStructValue(*FNiagaraStatelessSpawnInfo::StaticStruct(), reinterpret_cast<uint8*>(&TempSpawnInfo)))
+		{
+			ReturnValue = TempSpawnInfo;
+		}
+		return ReturnValue;
+	}
+
+	bool TestCanPaste(UNiagaraStatelessEmitter* StatelessEmitter, const UNiagaraClipboardContent* ClipboardContent, FText& OutMessage)
+	{
+		if (StatelessEmitter)
+		{
+			for (const FNiagaraClipboardPortableValue& PortableValue : ClipboardContent->PortableValues)
+			{
+				if (ConvertPortableValue(PortableValue).IsSet())
+				{
+					OutMessage = LOCTEXT("CanPasteSpawnInfo", "Paste spawn info(s).");
+					return true;
+				}
+			}
+		}
+
+		OutMessage = LOCTEXT("CanPasteSpawnInfoUnsupported", "Incompatible or no data to paste.");
+		return false;
+	}
+
+	FText GetPasteTransactionText(const UNiagaraClipboardContent* ClipboardContent)
+	{
+		return LOCTEXT("PasteSpawnInfoTransaction", "Paste spawn info(s).");
+	}
+
+	bool Paste(UNiagaraStatelessEmitter* StatelessEmitter, const UNiagaraClipboardContent* ClipboardContent)
+	{
+		bool bHasPastedValues = false;
+		if (StatelessEmitter)
+		{
+			for (const FNiagaraClipboardPortableValue& PortableValue : ClipboardContent->PortableValues)
+			{
+				TOptional<FNiagaraStatelessSpawnInfo> NewSpawnInfo = ConvertPortableValue(PortableValue);
+				if (!NewSpawnInfo.IsSet())
+				{
+					continue;
+				}
+
+				if (!bHasPastedValues)
+				{
+					bHasPastedValues = true;
+					StatelessEmitter->Modify();
+				}
+
+				StatelessEmitter->AddSpawnInfo() = NewSpawnInfo.GetValue();
+			}
+
+			if (bHasPastedValues)
+			{
+				StatelessEmitter->PostEditChange();
+			}
+		}
+
+		return bHasPastedValues;
+	}
+}
+
 class FNiagaraStatelessEmitterAddSpawnInfoAction : public INiagaraStackItemGroupAddAction
 {
 public:
@@ -155,6 +225,26 @@ void UNiagaraStackStatelessEmitterSpawnGroup::OnChildRequestDelete(FGuid DeleteI
 	}
 }
 
+bool UNiagaraStackStatelessEmitterSpawnGroup::TestCanPasteWithMessage(const UNiagaraClipboardContent* ClipboardContent, FText& OutMessage) const
+{
+	return NiagaraStackStatelessEmitterSpawnGroupPrivate::TestCanPaste(GetStatelessEmitter(), ClipboardContent, OutMessage);
+}
+
+FText UNiagaraStackStatelessEmitterSpawnGroup::GetPasteTransactionText(const UNiagaraClipboardContent* ClipboardContent) const
+{
+	return NiagaraStackStatelessEmitterSpawnGroupPrivate::GetPasteTransactionText(ClipboardContent);
+}
+
+void UNiagaraStackStatelessEmitterSpawnGroup::Paste(const UNiagaraClipboardContent* ClipboardContent, FText& OutPasteWarning)
+{
+	UNiagaraStatelessEmitter* StatelessEmitter = GetStatelessEmitter();
+	if (NiagaraStackStatelessEmitterSpawnGroupPrivate::Paste(StatelessEmitter, ClipboardContent))
+	{
+		OnDataObjectModified().Broadcast({ StatelessEmitter }, ENiagaraDataObjectChange::Changed);
+		RefreshChildren();
+	}
+}
+
 void UNiagaraStackStatelessEmitterSpawnItem::Initialize(FRequiredEntryData InRequiredEntryData, UNiagaraStatelessEmitter* InStatelessEmitter, int32 InIndex)
 {
 	Super::Initialize(InRequiredEntryData, FString::Printf(TEXT("StatelessEmitterSpawnItem-%i"), InIndex));
@@ -204,6 +294,47 @@ FGuid UNiagaraStackStatelessEmitterSpawnItem::GetSelectionId() const
 	return SourceId;
 }
 
+bool UNiagaraStackStatelessEmitterSpawnItem::TestCanCopyWithMessage(FText& OutMessage) const
+{
+	if (FNiagaraStatelessSpawnInfo* SpawnInfo = GetSpawnInfo())
+	{
+		OutMessage = LOCTEXT("CanCopyStatelessSpawnInfo", "Copy spawn info to the clipboard.");
+		return true;
+	}
+	OutMessage = LOCTEXT("CanCopyStatelessSpawnInfoUnsupported", "This spawn info does not support copy.");
+	return false;
+}
+
+void UNiagaraStackStatelessEmitterSpawnItem::Copy(UNiagaraClipboardContent* ClipboardContent) const
+{
+	if (FNiagaraStatelessSpawnInfo* SpawnInfo = GetSpawnInfo())
+	{
+		ClipboardContent->PortableValues.Emplace(
+			FNiagaraClipboardPortableValue::CreateFromStructValue(*FNiagaraStatelessSpawnInfo::StaticStruct(), reinterpret_cast<uint8*>(SpawnInfo))
+		);
+	}
+}
+
+bool UNiagaraStackStatelessEmitterSpawnItem::TestCanPasteWithMessage(const UNiagaraClipboardContent* ClipboardContent, FText& OutMessage) const
+{
+	return NiagaraStackStatelessEmitterSpawnGroupPrivate::TestCanPaste(GetStatelessEmitter(), ClipboardContent, OutMessage);
+}
+
+FText UNiagaraStackStatelessEmitterSpawnItem::GetPasteTransactionText(const UNiagaraClipboardContent* ClipboardContent) const
+{
+	return NiagaraStackStatelessEmitterSpawnGroupPrivate::GetPasteTransactionText(ClipboardContent);
+}
+
+void UNiagaraStackStatelessEmitterSpawnItem::Paste(const UNiagaraClipboardContent* ClipboardContent, FText& OutPasteWarning)
+{
+	UNiagaraStatelessEmitter* StatelessEmitter = GetStatelessEmitter();
+	if (NiagaraStackStatelessEmitterSpawnGroupPrivate::Paste(StatelessEmitter, ClipboardContent))
+	{
+		OnDataObjectModified().Broadcast({ StatelessEmitter }, ENiagaraDataObjectChange::Changed);
+		RefreshChildren();
+	}
+}
+
 bool UNiagaraStackStatelessEmitterSpawnItem::TestCanDeleteWithMessage(FText& OutCanDeleteMessage) const
 {
 	OutCanDeleteMessage = LOCTEXT("DeleteSpawnDataMessage", "Delete this spawn data.");
diff --git a/Engine/Plugins/FX/Niagara/Source/NiagaraEditor/Public/NiagaraClipboard.h b/Engine/Plugins/FX/Niagara/Source/NiagaraEditor/Public/NiagaraClipboard.h
--- a/Engine/Plugins/FX/Niagara/Source/NiagaraEditor/Public/NiagaraClipboard.h
+++ b/Engine/Plugins/FX/Niagara/Source/NiagaraEditor/Public/NiagaraClipboard.h
@@ -266,6 +266,9 @@ public:
 	UPROPERTY()
 	TArray<FNiagaraClipboardScriptVariable> ScriptVariables;
 
+	UPROPERTY()
+	TArray<TObjectPtr<const UObject>> StatelessModules;
+
 	/** We expect nodes to be exported into this string using FEdGraphUtilities::ExportNodesToText */
 	UPROPERTY()
 	FString ExportedNodes;
diff --git a/Engine/Plugins/FX/Niagara/Source/NiagaraEditor/Public/ViewModels/Stack/NiagaraStackStatelessEmitterSimulateGroup.h b/Engine/Plugins/FX/Niagara/Source/NiagaraEditor/Public/ViewModels/Stack/NiagaraStackStatelessEmitterSimulateGroup.h
--- a/Engine/Plugins/FX/Niagara/Source/NiagaraEditor/Public/ViewModels/Stack/NiagaraStackStatelessEmitterSimulateGroup.h
+++ b/Engine/Plugins/FX/Niagara/Source/NiagaraEditor/Public/ViewModels/Stack/NiagaraStackStatelessEmitterSimulateGroup.h
@@ -29,6 +29,11 @@ public:
 
 	UNiagaraStatelessEmitter* GetStatelessEmitter() const { return StatelessEmitterWeak.Get(); }
 
+	virtual bool SupportsPaste() const { return true; }
+	virtual bool TestCanPasteWithMessage(const UNiagaraClipboardContent* ClipboardContent, FText& OutMessage) const override;
+	virtual FText GetPasteTransactionText(const UNiagaraClipboardContent* ClipboardContent) const override;
+	virtual void Paste(const UNiagaraClipboardContent* ClipboardContent, FText& OutPasteWarning) override;
+
 protected:
 	virtual void RefreshChildrenInternal(const TArray<UNiagaraStackEntry*>& CurrentChildren, TArray<UNiagaraStackEntry*>& NewChildren, TArray<FStackIssue>& NewIssues) override;
 
@@ -54,6 +59,15 @@ public:
 	virtual FText GetDisplayName() const override { return DisplayName; }
 	virtual FText GetTooltipText() const override;
 
+	virtual bool SupportsCopy() const override { return true; }
+	virtual bool TestCanCopyWithMessage(FText& OutMessage) const override;
+	virtual void Copy(UNiagaraClipboardContent* ClipboardContent) const override;
+
+	virtual bool SupportsPaste() const { return true; }
+	virtual bool TestCanPasteWithMessage(const UNiagaraClipboardContent* ClipboardContent, FText& OutMessage) const override;
+	virtual FText GetPasteTransactionText(const UNiagaraClipboardContent* ClipboardContent) const override;
+	virtual void Paste(const UNiagaraClipboardContent* ClipboardContent, FText& OutPasteWarning) override;
+
 	virtual bool SupportsDelete() const override { return true; }
 	virtual bool TestCanDeleteWithMessage(FText& OutCanDeleteMessage) const override;
 	virtual FText GetDeleteTransactionText() const override;
@@ -66,6 +80,7 @@ public:
 	virtual void GetHeaderValueHandlers(TArray<TSharedRef<INiagaraStackItemHeaderValueHandler>>& OutHeaderValueHandlers) const override;
 
 	UNiagaraStatelessModule* GetStatelessModule() const { return StatelessModuleWeak.Get(); }
+	UNiagaraStatelessEmitter* GetStatelessEmitter() const;
 
 protected:
 	virtual void RefreshChildrenInternal(const TArray<UNiagaraStackEntry*>& CurrentChildren, TArray<UNiagaraStackEntry*>& NewChildren, TArray<FStackIssue>& NewIssues) override;
diff --git a/Engine/Plugins/FX/Niagara/Source/NiagaraEditor/Public/ViewModels/Stack/NiagaraStackStatelessEmitterSpawnGroup.h b/Engine/Plugins/FX/Niagara/Source/NiagaraEditor/Public/ViewModels/Stack/NiagaraStackStatelessEmitterSpawnGroup.h
--- a/Engine/Plugins/FX/Niagara/Source/NiagaraEditor/Public/ViewModels/Stack/NiagaraStackStatelessEmitterSpawnGroup.h
+++ b/Engine/Plugins/FX/Niagara/Source/NiagaraEditor/Public/ViewModels/Stack/NiagaraStackStatelessEmitterSpawnGroup.h
@@ -31,6 +31,11 @@ public:
 
 	UNiagaraStatelessEmitter* GetStatelessEmitter() const { return StatelessEmitterWeak.Get(); }
 
+	virtual bool SupportsPaste() const { return true; }
+	virtual bool TestCanPasteWithMessage(const UNiagaraClipboardContent* ClipboardContent, FText& OutMessage) const override;
+	virtual FText GetPasteTransactionText(const UNiagaraClipboardContent* ClipboardContent) const override;
+	virtual void Paste(const UNiagaraClipboardContent* ClipboardContent, FText& OutPasteWarning) override;
+
 protected:
 	virtual void RefreshChildrenInternal(const TArray<UNiagaraStackEntry*>& CurrentChildren, TArray<UNiagaraStackEntry*>& NewChildren, TArray<FStackIssue>& NewIssues) override;
 
@@ -58,6 +63,15 @@ public:
 	virtual FText GetTooltipText() const override;
 	virtual FGuid GetSelectionId() const override;
 
+	virtual bool SupportsCopy() const override { return true; }
+	virtual bool TestCanCopyWithMessage(FText& OutMessage) const override;
+	virtual void Copy(UNiagaraClipboardContent* ClipboardContent) const override;
+
+	virtual bool SupportsPaste() const { return true; }
+	virtual bool TestCanPasteWithMessage(const UNiagaraClipboardContent* ClipboardContent, FText& OutMessage) const override;
+	virtual FText GetPasteTransactionText(const UNiagaraClipboardContent* ClipboardContent) const override;
+	virtual void Paste(const UNiagaraClipboardContent* ClipboardContent, FText& OutPasteWarning) override;
+
 	virtual bool SupportsDelete() const override { return true; }
 	virtual bool TestCanDeleteWithMessage(FText& OutCanDeleteMessage) const override;
 	virtual FText GetDeleteTransactionText() const override;
