Additional protection for guarding against crashes when compiling (with the old compilation mode) Niagara systems and they are being edited at the same time.  Code was in place to catch emitters being deleted, but wouldn't catch emitters being added or reordered (which would cause crashes later in the precompile).

diff --git a/Engine/Plugins/FX/Niagara/Source/Niagara/Private/NiagaraAsyncCompile.cpp b/Engine/Plugins/FX/Niagara/Source/Niagara/Private/NiagaraAsyncCompile.cpp
--- a/Engine/Plugins/FX/Niagara/Source/Niagara/Private/NiagaraAsyncCompile.cpp
+++ b/Engine/Plugins/FX/Niagara/Source/Niagara/Private/NiagaraAsyncCompile.cpp
@@ -38,6 +38,13 @@ static FAutoConsoleVariableRef CVarNiagaraCompileDDCWaitTimeout(
 //////////////////////////////////////////////////////////////////////////
 
 #if WITH_EDITORONLY_DATA
+FNiagaraLazyPrecompileReference::FNiagaraLazyPrecompileReference(UNiagaraSystem* InSystem, TConstArrayView<UNiagaraScript*> InScripts)
+	: System(InSystem)
+	, Scripts(InScripts)
+{
+	GenerateEmitterScriptIndexMap(EmitterScriptIndex);
+}
+
 TSharedPtr<FNiagaraCompileRequestDataBase, ESPMode::ThreadSafe> FNiagaraLazyPrecompileReference::GetPrecompileData(UNiagaraScript* ForScript)
 {
 	if (SystemPrecompiledData == nullptr)
@@ -112,37 +119,49 @@ TSharedPtr<FNiagaraCompileRequestDuplicateDataBase, ESPMode::ThreadSafe> FNiagar
 	return PrecompileDuplicateData;
 }
 
+// annoyingly because we're driving some of the compilation process on live data (UNiagaraSystem/UNiagaraEmitter) which
+// might be changing in the scenario of a user continuing to modify their effect in the system editor, we need to at
+// least put some effort into validating that the system/emitters are still going to match up for the work we do in the
+// precompile.  So, we again generate the EmitterScriptIndexMap and make sure that they are the same.  This should handle
+// the common case of emitters being added, removed or reordered.
 bool FNiagaraLazyPrecompileReference::IsValidForPrecompile() const
 {
+	if (System)
+	{
+		FEmitterScriptIndexMap CurrentScriptIndexMap;
+		GenerateEmitterScriptIndexMap(CurrentScriptIndexMap);
+		
+		if (!CurrentScriptIndexMap.OrderIndependentCompareEqual(EmitterScriptIndex))
+		{
+			return false;
+		}
+	}
+
+	return true;
+}
+
+void FNiagaraLazyPrecompileReference::GenerateEmitterScriptIndexMap(FEmitterScriptIndexMap& OutIndexMap) const
+{
+	OutIndexMap.Reset();
+
 	if (System)
 	{
 		const TArray<FNiagaraEmitterHandle>& EmitterHandles = System->GetEmitterHandles();
-		for (TMap<UNiagaraScript*, int32>::TConstIterator It = EmitterScriptIndex.CreateConstIterator(); It; ++It)
+		for (int32 i = 0; i < EmitterHandles.Num(); i++)
 		{
-			if (EmitterHandles.IsValidIndex(It.Value()))
+			FNiagaraEmitterHandle Handle = EmitterHandles[i];
+			if (Handle.GetIsEnabled() && Handle.GetEmitterData())
 			{
-				const FNiagaraEmitterHandle& EmitterHandle = EmitterHandles[It.Value()];
-				const FVersionedNiagaraEmitterData* EmitterData = EmitterHandle.GetEmitterData();
-				if (!EmitterHandle.GetIsEnabled() || !EmitterData)
-				{
-					return false;
-				}
 				TArray<UNiagaraScript*> EmitterScripts;
-				EmitterData->GetScripts(EmitterScripts, false, true);
-
-				if (!EmitterScripts.Contains(It.Key()))
+				Handle.GetEmitterData()->GetScripts(EmitterScripts, false, true);
+				check(EmitterScripts.Num() > 0);
+				for (UNiagaraScript* EmitterScript : EmitterScripts)
 				{
-					return false;
+					OutIndexMap.Add(EmitterScript, i);
 				}
 			}
-			else
-			{
-				return false;
-			}
 		}
 	}
-
-	return true;
 }
 
 FNiagaraAsyncCompileTask::FNiagaraAsyncCompileTask(UNiagaraSystem* InOwningSystem, FString InAssetPath, const FEmitterCompiledScriptPair& InScriptPair)
@@ -1100,25 +1119,7 @@ bool FNiagaraActiveCompilationDefault::Launch(const FNiagaraCompilationOptions&
 		bEvaluateParametersPending = !UseRapidIterationParameters;
 
 		// prepare data for any precompile the ddc tasks need to do
-		TSharedPtr<FNiagaraLazyPrecompileReference, ESPMode::ThreadSafe> PrecompileReference = MakeShared<FNiagaraLazyPrecompileReference, ESPMode::ThreadSafe>();
-		PrecompileReference->System = Options.System;
-		PrecompileReference->Scripts = ScriptsNeedingCompile;
-
-		const TArray<FNiagaraEmitterHandle>& EmitterHandles = Options.System->GetEmitterHandles();
-		for (int32 i = 0; i < EmitterHandles.Num(); i++)
-		{
-			FNiagaraEmitterHandle Handle = EmitterHandles[i];
-			if (Handle.GetIsEnabled() && Handle.GetEmitterData())
-			{
-				TArray<UNiagaraScript*> EmitterScripts;
-				Handle.GetEmitterData()->GetScripts(EmitterScripts, false, true);
-				check(EmitterScripts.Num() > 0);
-				for (UNiagaraScript* EmitterScript : EmitterScripts)
-				{
-					PrecompileReference->EmitterScriptIndex.Add(EmitterScript, i);
-				}
-			}
-		}
+		TSharedPtr<FNiagaraLazyPrecompileReference, ESPMode::ThreadSafe> PrecompileReference = MakeShared<FNiagaraLazyPrecompileReference, ESPMode::ThreadSafe>(Options.System, ScriptsNeedingCompile);
 
 		for (UNiagaraScript* Script : ScriptsNeedingCompile)
 		{
diff --git a/Engine/Plugins/FX/Niagara/Source/Niagara/Private/NiagaraAsyncCompile.h b/Engine/Plugins/FX/Niagara/Source/Niagara/Private/NiagaraAsyncCompile.h
--- a/Engine/Plugins/FX/Niagara/Source/Niagara/Private/NiagaraAsyncCompile.h
+++ b/Engine/Plugins/FX/Niagara/Source/Niagara/Private/NiagaraAsyncCompile.h
@@ -40,6 +40,9 @@ enum class ENiagaraCompilationState : uint8
 
 struct FNiagaraLazyPrecompileReference
 {
+	FNiagaraLazyPrecompileReference(UNiagaraSystem* InSystem, TConstArrayView<UNiagaraScript*> InScripts);
+	FNiagaraLazyPrecompileReference() = delete;
+
 	TSharedPtr<FNiagaraCompileRequestDataBase, ESPMode::ThreadSafe> GetPrecompileData(UNiagaraScript* ForScript);
 	TSharedPtr<FNiagaraCompileRequestDuplicateDataBase, ESPMode::ThreadSafe> GetPrecompileDuplicateData(UNiagaraEmitter* OwningEmitter, UNiagaraScript* TargetScript);
 
@@ -47,10 +50,14 @@ struct FNiagaraLazyPrecompileReference
 	
 	UNiagaraSystem* System = nullptr;
 	TArray<UNiagaraScript*> Scripts;
-	TMap<UNiagaraScript*, int32> EmitterScriptIndex;
+
+	using FEmitterScriptIndexMap = TMap<UNiagaraScript*, int32>;
+	FEmitterScriptIndexMap EmitterScriptIndex;
 	TArray<TObjectPtr<UObject>> CompilationRootObjects;
 
 private:
+	void GenerateEmitterScriptIndexMap(FEmitterScriptIndexMap& OutIndexMap) const;
+
 	TSharedPtr<FNiagaraCompileRequestDataBase, ESPMode::ThreadSafe> SystemPrecompiledData;
 	TMap<UNiagaraScript*, TSharedPtr<FNiagaraCompileRequestDataBase, ESPMode::ThreadSafe>> EmitterMapping;
 	TArray<TSharedPtr<FNiagaraCompileRequestDuplicateDataBase, ESPMode::ThreadSafe>> PrecompileDuplicateDatas;
