- Fix memory corruption when capturing a sim cache and reading back array types that don't match between CPU <-> GPU

#rnx
#tests Run with stomp malloc and capture a sim cache from an emitter that contains a bool array on the GPU

diff --git a/Engine/Plugins/FX/Niagara/Source/Niagara/Classes/NiagaraDataInterfaceArrayImpl.h b/Engine/Plugins/FX/Niagara/Source/Niagara/Classes/NiagaraDataInterfaceArrayImpl.h
--- a/Engine/Plugins/FX/Niagara/Source/Niagara/Classes/NiagaraDataInterfaceArrayImpl.h
+++ b/Engine/Plugins/FX/Niagara/Source/Niagara/Classes/NiagaraDataInterfaceArrayImpl.h
@@ -1070,9 +1070,9 @@ struct FNDIArrayProxyImpl : public INDIArrayProxyBase
 							}
 							if ( FrameData.NumElements > 0 )
 							{
-								TArray<TArrayType> ArrayData;
+								TArray<TVMArrayType> ArrayData;
 								ArrayData.AddUninitialized(FrameData.NumElements);
-								FNDIArrayImplHelper<TArrayType>::CopyGpuToCpuMemory(ArrayData.GetData(), reinterpret_cast<const TVMArrayType*>(ReadbackData[0].Key), FrameData.NumElements);
+								FNDIArrayImplHelper<TArrayType>::CopyGpuToCpuMemory(ArrayData.GetData(), ReadbackData[0].Key, FrameData.NumElements);
 
 								FrameData.DataOffset = CacheData->FindOrAddData(
 									MakeArrayView(
