- Sim Cache Visualizer for Simple Counter

#rnx
#tests Create an emitter with a simple counter data interface, set the value / adjust the value, capture cache and observe the value, or open the cache test from ET

diff --git a/Engine/Plugins/FX/Niagara/Source/Niagara/Internal/DataInterface/NiagaraDataInterfaceSimpleCounter.h b/Engine/Plugins/FX/Niagara/Source/Niagara/Internal/DataInterface/NiagaraDataInterfaceSimpleCounter.h
new file mode 100644
--- /dev/null
+++ b/Engine/Plugins/FX/Niagara/Source/Niagara/Internal/DataInterface/NiagaraDataInterfaceSimpleCounter.h
@@ -0,0 +1,103 @@
+// Copyright Epic Games, Inc. All Rights Reserved.
+#pragma once
+
+#include "NiagaraCommon.h"
+#include "NiagaraShared.h"
+#include "NiagaraSimCacheCustomStorageInterface.h"
+
+#include "VectorVM.h"
+#include "NiagaraDataInterfaceRW.h"
+#include "HAL/PlatformAtomics.h"
+
+#include "NiagaraDataInterfaceSimpleCounter.generated.h"
+
+struct FNiagaraDataInterfaceGeneratedFunction;
+
+/**
+Thread safe counter starts at the initial value on start / reset.
+When operating between CPU & GPU ensure you set the appropriate sync mode.
+*/
+UCLASS(EditInlineNew, Category = "Counting", CollapseCategories, meta = (DisplayName = "Simple Counter"), MinimalAPI)
+class UNiagaraDataInterfaceSimpleCounter : public UNiagaraDataInterfaceRWBase, public INiagaraSimCacheCustomStorageInterface
+{
+	GENERATED_UCLASS_BODY()
+
+	BEGIN_SHADER_PARAMETER_STRUCT(FShaderParameters, )
+		SHADER_PARAMETER(int32,		CountOffset)
+	END_SHADER_PARAMETER_STRUCT();
+
+public:
+	// UObject Interface
+	NIAGARA_API virtual void PostInitProperties() override;
+	NIAGARA_API virtual void PostLoad() override;
+#if WITH_EDITOR
+	NIAGARA_API virtual void PostEditChangeProperty(struct FPropertyChangedEvent& PropertyChangedEvent) override;
+#endif
+	// UObject Interface End
+
+	// UNiagaraDataInterface Interface Begin
+	virtual bool CanExecuteOnTarget(ENiagaraSimTarget Target) const override { return true; }
+
+	NIAGARA_API virtual bool Equals(const UNiagaraDataInterface* Other) const override;
+	NIAGARA_API virtual bool CopyToInternal(UNiagaraDataInterface* Destination) const override;
+
+	NIAGARA_API virtual bool InitPerInstanceData(void* PerInstanceData, FNiagaraSystemInstance* SystemInstance) override;
+	NIAGARA_API virtual void DestroyPerInstanceData(void* PerInstanceData, FNiagaraSystemInstance* SystemInstance) override;
+	NIAGARA_API virtual int32 PerInstanceDataSize() const override;
+
+	NIAGARA_API virtual void GetVMExternalFunction(const FVMExternalFunctionBindingInfo& BindingInfo, void* InstanceData, FVMExternalFunction &OutFunc) override;
+#if WITH_EDITORONLY_DATA
+	NIAGARA_API virtual void GetParameterDefinitionHLSL(const FNiagaraDataInterfaceGPUParamInfo& ParamInfo, FString& OutHLSL) override;
+	NIAGARA_API virtual bool GetFunctionHLSL(const FNiagaraDataInterfaceGPUParamInfo& ParamInfo, const FNiagaraDataInterfaceGeneratedFunction& FunctionInfo, int FunctionInstanceIndex, FString& OutHLSL) override;
+	NIAGARA_API virtual bool AppendCompileHash(FNiagaraCompileHashVisitor* InVisitor) const override;
+#endif
+#if WITH_EDITOR
+	virtual bool GetGpuUseIndirectDispatch() const override { return true; }
+#endif
+	NIAGARA_API virtual void BuildShaderParameters(FNiagaraShaderParametersBuilder& ShaderParametersBuilder) const override;
+	NIAGARA_API virtual void SetShaderParameters(const FNiagaraDataInterfaceSetShaderParametersContext& Context) const override;
+	NIAGARA_API virtual void PushToRenderThreadImpl() override;
+	// UNiagaraDataInterface Interface End
+
+	//~ INiagaraSimCacheCustomStorageInterface interface BEGIN
+	virtual UObject* SimCacheBeginWrite(UObject* SimCache, FNiagaraSystemInstance* NiagaraSystemInstance, const void* OptionalPerInstanceData, FNiagaraSimCacheFeedbackContext& FeedbackContext) const override;
+	virtual bool SimCacheWriteFrame(UObject* StorageObject, int FrameIndex, FNiagaraSystemInstance* SystemInstance, const void* OptionalPerInstanceData, FNiagaraSimCacheFeedbackContext& FeedbackContext) const override;
+	virtual bool SimCacheReadFrame(UObject* StorageObject, int FrameA, int FrameB, float Interp, FNiagaraSystemInstance* SystemInstance, void* OptionalPerInstanceData) override;
+	virtual bool SimCacheCompareFrame(const UObject* LhsStorageObject, const UObject* RhsStorageObject, int FrameIndex, TOptional<float> Tolerance, FString& OutErrors) const override;
+	//~ UNiagaraDataInterface interface END
+
+	void UpdateDIProxy();
+
+	// VM functions
+	NIAGARA_API void VMGet(FVectorVMExternalFunctionContext& Context);
+	NIAGARA_API void VMSet(FVectorVMExternalFunctionContext& Context);
+	NIAGARA_API void VMExchange(FVectorVMExternalFunctionContext& Context);
+	NIAGARA_API void VMAdd(FVectorVMExternalFunctionContext& Context);
+	NIAGARA_API void VMIncrement(FVectorVMExternalFunctionContext& Context);
+	NIAGARA_API void VMDecrement(FVectorVMExternalFunctionContext& Context);
+
+	NIAGARA_API void GetNextValue_Deprecated(FVectorVMExternalFunctionContext& Context);
+
+	/** Select how we should synchronize the counter between Cpu & Gpu */
+	UPROPERTY(EditAnywhere, Category = "Simple Counter")
+	ENiagaraGpuSyncMode GpuSyncMode = ENiagaraGpuSyncMode::None;
+
+	/** This is the value the counter will have when the instance is reset / created. */
+	UPROPERTY(EditAnywhere, Category = "Simple Counter")
+	int32 InitialValue = 0;
+
+protected:
+#if WITH_EDITORONLY_DATA
+	NIAGARA_API virtual void GetFunctionsInternal(TArray<FNiagaraFunctionSignature>& OutFunctions) const override;
+#endif
+};
+
+UCLASS(MinimalAPI)
+class UNDISimpleCounterSimCacheData : public UObject
+{
+	GENERATED_BODY()
+
+public:
+	UPROPERTY()
+	TArray<int32> Values;
+};
diff --git a/Engine/Plugins/FX/Niagara/Source/Niagara/Private/DataInterface/NiagaraDataInterfaceSimpleCounter.cpp b/Engine/Plugins/FX/Niagara/Source/Niagara/Private/DataInterface/NiagaraDataInterfaceSimpleCounter.cpp
--- a/Engine/Plugins/FX/Niagara/Source/Niagara/Private/DataInterface/NiagaraDataInterfaceSimpleCounter.cpp
+++ b/Engine/Plugins/FX/Niagara/Source/Niagara/Private/DataInterface/NiagaraDataInterfaceSimpleCounter.cpp
@@ -1,6 +1,6 @@
 // Copyright Epic Games, Inc. All Rights Reserved.
 
-#include "NiagaraDataInterfaceSimpleCounter.h"
+#include "DataInterface/NiagaraDataInterfaceSimpleCounter.h"
 #include "NiagaraClearCounts.h"
 #include "NiagaraCompileHashVisitor.h"
 #include "NiagaraGpuComputeDispatchInterface.h"
diff --git a/Engine/Plugins/FX/Niagara/Source/Niagara/Private/DataInterface/NiagaraDataInterfaceSimpleCounter.h b/Engine/Plugins/FX/Niagara/Source/Niagara/Private/DataInterface/NiagaraDataInterfaceSimpleCounter.h
deleted file mode 100644
diff --git a/Engine/Plugins/FX/Niagara/Source/NiagaraEditor/Private/Customizations/SimCache/NiagaraSimpleCounterSimCacheVisualizer.cpp b/Engine/Plugins/FX/Niagara/Source/NiagaraEditor/Private/Customizations/SimCache/NiagaraSimpleCounterSimCacheVisualizer.cpp
new file mode 100644
--- /dev/null
+++ b/Engine/Plugins/FX/Niagara/Source/NiagaraEditor/Private/Customizations/SimCache/NiagaraSimpleCounterSimCacheVisualizer.cpp
@@ -0,0 +1,110 @@
+// Copyright Epic Games, Inc. All Rights Reserved.
+
+#include "NiagaraSimpleCounterSimCacheVisualizer.h"
+
+#include "DataInterface/NiagaraDataInterfaceSimpleCounter.h"
+
+#include "UObject/StrongObjectPtr.h"
+#include "ViewModels/NiagaraSimCacheViewModel.h"
+#include "Widgets/Input/SSpinBox.h"
+#include "Widgets/Text/STextBlock.h"
+#include "DetailLayoutBuilder.h"
+#include "SEnumCombo.h"
+
+#define LOCTEXT_NAMESPACE "NiagaraSimpleCounterSimCacheVisualizer"
+
+////////////////////////////////////////////////////////////////////////////////////////////////////////////////////
+
+namespace NDISimpleCounterSimCacheVisualizer
+{
+class SSimCacheView : public SCompoundWidget
+{
+public:
+	SLATE_BEGIN_ARGS(SSimCacheView)
+	{}
+	SLATE_END_ARGS()
+
+	virtual ~SSimCacheView() override
+	{
+		ViewModel->OnViewDataChanged().RemoveAll(this);
+	}
+
+	void Construct(const FArguments& InArgs, TSharedPtr<FNiagaraSimCacheViewModel> InViewModel, const UNDISimpleCounterSimCacheData* InCacheData)
+	{
+		ViewModel = InViewModel;
+		CacheData.Reset(InCacheData);
+
+		ChildSlot
+		[
+			SNew(SVerticalBox)
+			+ SVerticalBox::Slot()
+			.AutoHeight()
+			[
+				SNew(SHorizontalBox)
+				+SHorizontalBox::Slot()
+				[
+					SNew(STextBlock)
+					.Text(this, &SSimCacheView::GetCpuInformation)
+				]
+			]
+			+ SVerticalBox::Slot()
+			.AutoHeight()
+			[
+				SNew(SHorizontalBox)
+				+SHorizontalBox::Slot()
+				[
+					SNew(STextBlock)
+					.Text(this, &SSimCacheView::GetGpuInformation)
+				]
+			]
+		];
+	}
+
+	FText GetValueText(int32 DataOffset) const
+	{
+		FText DataValue;
+		if (CacheData->Values.IsValidIndex(DataOffset))
+		{
+			return FText::AsNumber(CacheData->Values[DataOffset]);
+		}
+		else
+		{
+			return LOCTEXT("DataInvalid", "Invalid");
+		}
+	}
+
+	FText GetCpuInformation() const
+	{
+		const int32 FrameIndex = ViewModel->GetFrameIndex();
+		const int32 DataOffset = FrameIndex * 2;
+		return FText::Format(LOCTEXT("CpuDataFormat", "CPU Visibile Value = {0}"), GetValueText(DataOffset + 0));
+	}
+
+	FText GetGpuInformation() const
+	{
+		const int32 FrameIndex = ViewModel->GetFrameIndex();
+		const int32 DataOffset = FrameIndex * 2;
+		return FText::Format(LOCTEXT("GpuDataFormat", "GPU Visibile Value = {0}"), GetValueText(DataOffset + 1));
+	}
+
+private:
+	TSharedPtr<FNiagaraSimCacheViewModel>					ViewModel;
+	TStrongObjectPtr<const UNDISimpleCounterSimCacheData>	CacheData;
+};
+
+} // NDISimpleCounterSimCacheVisualizer
+
+////////////////////////////////////////////////////////////////////////////////////////////////////////////////////
+
+TSharedPtr<SWidget> FNiagaraSimpleCounterSimCacheVisualizer::CreateWidgetFor(const UObject* InCachedData, TSharedPtr<FNiagaraSimCacheViewModel> ViewModel)
+{
+	using namespace NDISimpleCounterSimCacheVisualizer;
+
+	if (const UNDISimpleCounterSimCacheData* CachedData = Cast<const UNDISimpleCounterSimCacheData>(InCachedData))
+	{
+		return SNew(SSimCacheView, ViewModel, CachedData);
+	}
+	return TSharedPtr<SWidget>();
+}
+
+#undef LOCTEXT_NAMESPACE
diff --git a/Engine/Plugins/FX/Niagara/Source/NiagaraEditor/Private/Customizations/SimCache/NiagaraSimpleCounterSimCacheVisualizer.h b/Engine/Plugins/FX/Niagara/Source/NiagaraEditor/Private/Customizations/SimCache/NiagaraSimpleCounterSimCacheVisualizer.h
new file mode 100644
--- /dev/null
+++ b/Engine/Plugins/FX/Niagara/Source/NiagaraEditor/Private/Customizations/SimCache/NiagaraSimpleCounterSimCacheVisualizer.h
@@ -0,0 +1,18 @@
+// Copyright Epic Games, Inc. All Rights Reserved.
+
+#pragma once
+
+#include "NiagaraCommon.h"
+#include "Customizations/NiagaraDataInterfaceSimCacheVisualizer.h"
+
+/*
+ * Provides a custom widget to show the array DI data in the sim cache in a table.
+ */
+class FNiagaraSimpleCounterSimCacheVisualizer : public INiagaraDataInterfaceSimCacheVisualizer
+{
+public:
+	FNiagaraSimpleCounterSimCacheVisualizer() {}
+	virtual ~FNiagaraSimpleCounterSimCacheVisualizer() override = default;
+	
+	virtual TSharedPtr<SWidget> CreateWidgetFor(const UObject* CachedData, TSharedPtr<FNiagaraSimCacheViewModel> ViewModel) override;
+};
diff --git a/Engine/Plugins/FX/Niagara/Source/NiagaraEditor/Private/NiagaraEditorModule.cpp b/Engine/Plugins/FX/Niagara/Source/NiagaraEditor/Private/NiagaraEditorModule.cpp
--- a/Engine/Plugins/FX/Niagara/Source/NiagaraEditor/Private/NiagaraEditorModule.cpp
+++ b/Engine/Plugins/FX/Niagara/Source/NiagaraEditor/Private/NiagaraEditorModule.cpp
@@ -58,6 +58,7 @@
 #include "DataInterface/NiagaraDataInterfaceDataChannelRead.h"
 #include "DataInterface/NiagaraDataInterfaceDataChannelWrite.h"
 #include "DataInterface/NiagaraDataInterfaceMemoryBuffer.h"
+#include "DataInterface/NiagaraDataInterfaceSimpleCounter.h"
 #include "NiagaraDataInterfaceRenderTargetVolume.h"
 
 #include "ViewModels/NiagaraScriptViewModel.h"
@@ -122,6 +123,7 @@
 #include "Customizations/SimCache/FNiagaraDataChannelSimCacheVisualizer.h"
 #include "Customizations/SimCache/NiagaraMemoryBufferSimCacheVisualizer.h"
 #include "Customizations/SimCache/NiagaraRenderTargetVolumeSimCacheVisualizer.h"
+#include "Customizations/SimCache/NiagaraSimpleCounterSimCacheVisualizer.h"
 
 #include "NiagaraComponent.h"
 #include "NiagaraNodeStaticSwitch.h"
@@ -1560,6 +1562,7 @@ void FNiagaraEditorModule::StartupModule()
 	RegisterDataInterfaceCacheVisualizer(UNiagaraDataInterfaceDataChannelWrite::StaticClass(), MakeShared<FNiagaraDataChannelSimCacheVisualizer>());
 	RegisterDataInterfaceCacheVisualizer(UNiagaraDataInterfaceMemoryBuffer::StaticClass(), MakeShared<FNiagaraMemoryBufferSimCacheVisualizer>());
 	RegisterDataInterfaceCacheVisualizer(UNiagaraDataInterfaceRenderTargetVolume::StaticClass(), MakeShared<FNiagaraRenderTargetVolumeSimCacheVisualizer>());
+	RegisterDataInterfaceCacheVisualizer(UNiagaraDataInterfaceSimpleCounter::StaticClass(), MakeShared<FNiagaraSimpleCounterSimCacheVisualizer>());
 	for (TObjectIterator<UClass> It; It; ++It)
 	{
 		if (It->IsChildOf(UNiagaraDataInterfaceArray::StaticClass()))
