NDC GPU Writes

Niagara Data Channels can now be written to from Niagara GPU simulations.
This data can then be passed to other GPU simulations in the same or another Niagara System.
It can also optionally be sent back to the CPU for processing in Niagara or code/BP.

A more detailed list of the changes and supporting work. In no particular order.

Write DI:

- Write DI can now write from GPU sims.
HLSL gen and shader structure very similar to the ReadDI.
A single per NDC data proxy buffer is allocated for all systems operating on it in the frame.
During tick processing, proxies accumulate an allcation count from writing DIs.
All GPU bound data is written into this shared buffer.
Separate "ForCPU" buffer that is per DI. We send this buffer back to the CPU in PostStage.
- CPU writes now allocate a data buffer directly from the NDC Game Data on the GT. This is then allocated from and written to on worker threads.
- Moved CPU allocate and publish from Pre/Post Stage to Pre/Post Sim.
- New Dynamic Allocation mode for NDC Writes. Emitter scripts can allocate the NDC data directly via an "Allocate" function.
- Refactored DI Function definition code to match the read DI.

Read DI:

- Fixed Consume Logic and implemented for GPU.
- Introduced buffer transitions for GPU Reads as now data can be coming from previous GPU writes instead of just a CPU upload.
- Moved NDC Num/Count values to be indirect, taken from the instance count buffer, as now GPU writes can modify the count on the GPU side beyond the innitial CPU count.

GPU Dispatcher and NDC Proxies:
- Dispatcher now keeps array of NDC proxies in the scene.
- Proxy tick flow moved to be done via the dispatcher rather than an RT command from the NDC data.
- Dispatcher now accumulates an allocation count for each NDC proxy and we allocate one large buffer per NDC data into which all DIs write and read.
- Pending buffers from the CPU are kept until the next Dispatcher tick when they are shipped over to the GPU NDC Buffers.
- We can combine the incomming CPU data and newly written GPU data due to a new pass in the instance count manager to initialize some counts to specific values. In this case, the NDC element count from the CPU.
- All GPU emitters that use a read NDC DI reading the current frame's data are pushed to PostRenderOpaque. This is to help reduce issues with same frame reads until we can introduce a more robust dependany + ordering system.
- Removed some constness internal to the dispatcher to allow some non const operations inside DI Pre/Post Stage

Readback From GPU:

We can send data written on the GPU back to the CPU for read in CPU Systems and BP.
This is done via a new FNiagaraDataBufferReadback helper that we use from inside the NDC Proxy.
Once the data is on the CPU, it's published to the NDC like any other CPU side write.
Though we wait until the start of each frame to consume GPU data in an effort to reduce racy behavior.

NDC DI HLSL:
- DI HLSL Gen can now access and generate default hlsl for common types.
- Fixed Variadic signatures with both input and output params.
- Refactored varaidic funciton hlsl gen out of read DI into shared code for use with write DI too.
- Refactored hlsl and shader params etc to make better use of common code and be more generally usable and readable.

NDC General:
- Modified NDC Init flow to Initialize NDCs after world init.
- Modified NDC Cleanup flow to clean up all data on world tear down rather than on GC.
- Fixed issue that would allow NDC writing systems to execute and finish beyond the end of their original tick group.
- Fixed issue with NDC spawning using an incorrect count for num NDC elements.
- Made NDCData and NDC Game Data shared pointers threadsafe as they're now being used on the render and worker threads.

Misc:
- New pass in instance count manager to init instance counts to specific values rather than only zeroing.
- Adding separate function to allow allocating of data buffers from a DataSet.
- DataSets can optionally pre allocate a number of buffers.
- Adding some utility dummy UAVs to NDIDataChannelUtilities.

diff --git a/Engine/Plugins/FX/Niagara/Shaders/Private/DataChannel/NiagaraDataInterfaceDataChannelCommon.ush b/Engine/Plugins/FX/Niagara/Shaders/Private/DataChannel/NiagaraDataInterfaceDataChannelCommon.ush
--- a/Engine/Plugins/FX/Niagara/Shaders/Private/DataChannel/NiagaraDataInterfaceDataChannelCommon.ush
+++ b/Engine/Plugins/FX/Niagara/Shaders/Private/DataChannel/NiagaraDataInterfaceDataChannelCommon.ush
@@ -14,85 +14,168 @@ Shader code common to all instances and functions of Niagara Data Channel DI.
 //TOOD: Add support for N multiple buffers
 //TODO: Add support for bindless buffers on platforms that support it.
 
-void NiagaraDataChannelRead_float(Buffer<float> InBuffer, uint Stride, in out uint RegisterIdx, uint ParticleIdx, out float OutValue)
+void NiagaraDataChannelRead_Float(Buffer<float> InBuffer, uint Stride, in out uint RegisterIdx, uint ParticleIdx, out float OutValue)
 {
 	OutValue = InBuffer[RegisterIdx * Stride + ParticleIdx];
 	++RegisterIdx;
 }
 
-void NiagaraDataChannelRead_float2(Buffer<float> InBuffer, uint Stride, in out uint RegisterIdx, uint ParticleIdx, out float2 OutValue)
+void NiagaraDataChannelRead_Float(Buffer<float> InBuffer, uint Stride, in out uint RegisterIdx, uint ParticleIdx, out float2 OutValue)
 {
-	NiagaraDataChannelRead_float(InBuffer, Stride, RegisterIdx, ParticleIdx, OutValue.x);
-	NiagaraDataChannelRead_float(InBuffer, Stride, RegisterIdx, ParticleIdx, OutValue.y);
+	NiagaraDataChannelRead_Float(InBuffer, Stride, RegisterIdx, ParticleIdx, OutValue.x);
+	NiagaraDataChannelRead_Float(InBuffer, Stride, RegisterIdx, ParticleIdx, OutValue.y);
 }
 
-void NiagaraDataChannelRead_float3(Buffer<float> InBuffer, uint Stride, in out uint RegisterIdx, uint ParticleIdx, out float3 OutValue)
+void NiagaraDataChannelRead_Float(Buffer<float> InBuffer, uint Stride, in out uint RegisterIdx, uint ParticleIdx, out float3 OutValue)
 {
-	NiagaraDataChannelRead_float(InBuffer, Stride, RegisterIdx, ParticleIdx, OutValue.x);
-	NiagaraDataChannelRead_float(InBuffer, Stride, RegisterIdx, ParticleIdx, OutValue.y);
-	NiagaraDataChannelRead_float(InBuffer, Stride, RegisterIdx, ParticleIdx, OutValue.z);
+	NiagaraDataChannelRead_Float(InBuffer, Stride, RegisterIdx, ParticleIdx, OutValue.x);
+	NiagaraDataChannelRead_Float(InBuffer, Stride, RegisterIdx, ParticleIdx, OutValue.y);
+	NiagaraDataChannelRead_Float(InBuffer, Stride, RegisterIdx, ParticleIdx, OutValue.z);
 }
 
-void NiagaraDataChannelRead_float4(Buffer<float> InBuffer, uint Stride, in out uint RegisterIdx, uint ParticleIdx, out float4 OutValue)
+void NiagaraDataChannelRead_Float(Buffer<float> InBuffer, uint Stride, in out uint RegisterIdx, uint ParticleIdx, out float4 OutValue)
 {
-	NiagaraDataChannelRead_float(InBuffer, Stride, RegisterIdx, ParticleIdx, OutValue.x);
-	NiagaraDataChannelRead_float(InBuffer, Stride, RegisterIdx, ParticleIdx, OutValue.y);
-	NiagaraDataChannelRead_float(InBuffer, Stride, RegisterIdx, ParticleIdx, OutValue.z);
-	NiagaraDataChannelRead_float(InBuffer, Stride, RegisterIdx, ParticleIdx, OutValue.w);
+	NiagaraDataChannelRead_Float(InBuffer, Stride, RegisterIdx, ParticleIdx, OutValue.x);
+	NiagaraDataChannelRead_Float(InBuffer, Stride, RegisterIdx, ParticleIdx, OutValue.y);
+	NiagaraDataChannelRead_Float(InBuffer, Stride, RegisterIdx, ParticleIdx, OutValue.z);
+	NiagaraDataChannelRead_Float(InBuffer, Stride, RegisterIdx, ParticleIdx, OutValue.w);
 }
 
-void NiagaraDataChannelRead_int(Buffer<int> InBuffer, uint Stride, in out uint RegisterIdx, uint ParticleIdx, out int OutValue)
+void NiagaraDataChannelRead_Int32(Buffer<int> InBuffer, uint Stride, in out uint RegisterIdx, uint ParticleIdx, out int OutValue)
 {
 	OutValue = InBuffer[RegisterIdx * Stride + ParticleIdx];
 	++RegisterIdx;
 }
 
-void NiagaraDataChannelRead_int2(Buffer<int> InBuffer, uint Stride, in out uint RegisterIdx, uint ParticleIdx, out int2 OutValue)
+void NiagaraDataChannelRead_Int32(Buffer<int> InBuffer, uint Stride, in out uint RegisterIdx, uint ParticleIdx, out int2 OutValue)
 {
-	NiagaraDataChannelRead_int(InBuffer, Stride, RegisterIdx, ParticleIdx, OutValue.x);
-	NiagaraDataChannelRead_int(InBuffer, Stride, RegisterIdx, ParticleIdx, OutValue.y);
+	NiagaraDataChannelRead_Int32(InBuffer, Stride, RegisterIdx, ParticleIdx, OutValue.x);
+	NiagaraDataChannelRead_Int32(InBuffer, Stride, RegisterIdx, ParticleIdx, OutValue.y);
 }
 
-void NiagaraDataChannelRead_int3(Buffer<int> InBuffer, uint Stride, in out uint RegisterIdx, uint ParticleIdx, out int3 OutValue)
+void NiagaraDataChannelRead_Int32(Buffer<int> InBuffer, uint Stride, in out uint RegisterIdx, uint ParticleIdx, out int3 OutValue)
 {
-	NiagaraDataChannelRead_int(InBuffer, Stride, RegisterIdx, ParticleIdx, OutValue.x);
-	NiagaraDataChannelRead_int(InBuffer, Stride, RegisterIdx, ParticleIdx, OutValue.y);
-	NiagaraDataChannelRead_int(InBuffer, Stride, RegisterIdx, ParticleIdx, OutValue.z);
+	NiagaraDataChannelRead_Int32(InBuffer, Stride, RegisterIdx, ParticleIdx, OutValue.x);
+	NiagaraDataChannelRead_Int32(InBuffer, Stride, RegisterIdx, ParticleIdx, OutValue.y);
+	NiagaraDataChannelRead_Int32(InBuffer, Stride, RegisterIdx, ParticleIdx, OutValue.z);
 }
 
-void NiagaraDataChannelRead_int4(Buffer<int> InBuffer, uint Stride, in out uint RegisterIdx, uint ParticleIdx, out int4 OutValue)
+void NiagaraDataChannelRead_Int32(Buffer<int> InBuffer, uint Stride, in out uint RegisterIdx, uint ParticleIdx, out int4 OutValue)
 {
-	NiagaraDataChannelRead_int(InBuffer, Stride, RegisterIdx, ParticleIdx, OutValue.x);
-	NiagaraDataChannelRead_int(InBuffer, Stride, RegisterIdx, ParticleIdx, OutValue.y);
-	NiagaraDataChannelRead_int(InBuffer, Stride, RegisterIdx, ParticleIdx, OutValue.z);
-	NiagaraDataChannelRead_int(InBuffer, Stride, RegisterIdx, ParticleIdx, OutValue.w);
+	NiagaraDataChannelRead_Int32(InBuffer, Stride, RegisterIdx, ParticleIdx, OutValue.x);
+	NiagaraDataChannelRead_Int32(InBuffer, Stride, RegisterIdx, ParticleIdx, OutValue.y);
+	NiagaraDataChannelRead_Int32(InBuffer, Stride, RegisterIdx, ParticleIdx, OutValue.z);
+	NiagaraDataChannelRead_Int32(InBuffer, Stride, RegisterIdx, ParticleIdx, OutValue.w);
 }
 
-void NiagaraDataChannelRead_half(Buffer<half> InBuffer, uint Stride, in out uint RegisterIdx, uint ParticleIdx, out half OutValue)
+void NiagaraDataChannelRead_Half(Buffer<half> InBuffer, uint Stride, in out uint RegisterIdx, uint ParticleIdx, out half OutValue)
 {
 	OutValue = InBuffer[RegisterIdx * Stride + ParticleIdx];
 	++RegisterIdx;
 }
 
-void NiagaraDataChannelRead_half2(Buffer<half> InBuffer, uint Stride, in out uint RegisterIdx, uint ParticleIdx, out half2 OutValue)
+void NiagaraDataChannelRead_Half(Buffer<half> InBuffer, uint Stride, in out uint RegisterIdx, uint ParticleIdx, out half2 OutValue)
 {
-	NiagaraDataChannelRead_half(InBuffer, Stride, RegisterIdx, ParticleIdx, OutValue.x);
-	NiagaraDataChannelRead_half(InBuffer, Stride, RegisterIdx, ParticleIdx, OutValue.y);
+	NiagaraDataChannelRead_Half(InBuffer, Stride, RegisterIdx, ParticleIdx, OutValue.x);
+	NiagaraDataChannelRead_Half(InBuffer, Stride, RegisterIdx, ParticleIdx, OutValue.y);
 }
 
-void NiagaraDataChannelRead_half3(Buffer<half> InBuffer, uint Stride, in out uint RegisterIdx, uint ParticleIdx, out half3 OutValue)
+void NiagaraDataChannelRead_Half(Buffer<half> InBuffer, uint Stride, in out uint RegisterIdx, uint ParticleIdx, out half3 OutValue)
 {
-	NiagaraDataChannelRead_half(InBuffer, Stride, RegisterIdx, ParticleIdx, OutValue.x);
-	NiagaraDataChannelRead_half(InBuffer, Stride, RegisterIdx, ParticleIdx, OutValue.y);
-	NiagaraDataChannelRead_half(InBuffer, Stride, RegisterIdx, ParticleIdx, OutValue.z);
+	NiagaraDataChannelRead_Half(InBuffer, Stride, RegisterIdx, ParticleIdx, OutValue.x);
+	NiagaraDataChannelRead_Half(InBuffer, Stride, RegisterIdx, ParticleIdx, OutValue.y);
+	NiagaraDataChannelRead_Half(InBuffer, Stride, RegisterIdx, ParticleIdx, OutValue.z);
 }
 
-void NiagaraDataChannelRead_half4(Buffer<half> InBuffer, uint Stride, in out uint RegisterIdx, uint ParticleIdx, out half4 OutValue)
+void NiagaraDataChannelRead_Half(Buffer<half> InBuffer, uint Stride, in out uint RegisterIdx, uint ParticleIdx, out half4 OutValue)
 {
-	NiagaraDataChannelRead_half(InBuffer, Stride, RegisterIdx, ParticleIdx, OutValue.x);
-	NiagaraDataChannelRead_half(InBuffer, Stride, RegisterIdx, ParticleIdx, OutValue.y);
-	NiagaraDataChannelRead_half(InBuffer, Stride, RegisterIdx, ParticleIdx, OutValue.z);
-	NiagaraDataChannelRead_half(InBuffer, Stride, RegisterIdx, ParticleIdx, OutValue.w);
+	NiagaraDataChannelRead_Half(InBuffer, Stride, RegisterIdx, ParticleIdx, OutValue.x);
+	NiagaraDataChannelRead_Half(InBuffer, Stride, RegisterIdx, ParticleIdx, OutValue.y);
+	NiagaraDataChannelRead_Half(InBuffer, Stride, RegisterIdx, ParticleIdx, OutValue.z);
+	NiagaraDataChannelRead_Half(InBuffer, Stride, RegisterIdx, ParticleIdx, OutValue.w);
+}
+
+// Write accessors 
+
+void NiagaraDataChannelWrite_Float(RWBuffer<float> InBuffer, uint Stride, in out uint RegisterIdx, uint ParticleIdx, float InValue)
+{
+	InBuffer[RegisterIdx * Stride + ParticleIdx] = InValue;
+	++RegisterIdx;
+}
+
+void NiagaraDataChannelWrite_Float(RWBuffer<float> InBuffer, uint Stride, in out uint RegisterIdx, uint ParticleIdx, float2 InValue)
+{
+	NiagaraDataChannelWrite_Float(InBuffer, Stride, RegisterIdx, ParticleIdx, InValue.x);
+	NiagaraDataChannelWrite_Float(InBuffer, Stride, RegisterIdx, ParticleIdx, InValue.y);
+}
+
+void NiagaraDataChannelWrite_Float(RWBuffer<float> InBuffer, uint Stride, in out uint RegisterIdx, uint ParticleIdx, float3 InValue)
+{
+	NiagaraDataChannelWrite_Float(InBuffer, Stride, RegisterIdx, ParticleIdx, InValue.x);
+	NiagaraDataChannelWrite_Float(InBuffer, Stride, RegisterIdx, ParticleIdx, InValue.y);
+	NiagaraDataChannelWrite_Float(InBuffer, Stride, RegisterIdx, ParticleIdx, InValue.z);
+}
+
+void NiagaraDataChannelWrite_Float(RWBuffer<float> InBuffer, uint Stride, in out uint RegisterIdx, uint ParticleIdx, float4 InValue)
+{
+	NiagaraDataChannelWrite_Float(InBuffer, Stride, RegisterIdx, ParticleIdx, InValue.x);
+	NiagaraDataChannelWrite_Float(InBuffer, Stride, RegisterIdx, ParticleIdx, InValue.y);
+	NiagaraDataChannelWrite_Float(InBuffer, Stride, RegisterIdx, ParticleIdx, InValue.z);
+	NiagaraDataChannelWrite_Float(InBuffer, Stride, RegisterIdx, ParticleIdx, InValue.w);
+}
+
+void NiagaraDataChannelWrite_Int32(RWBuffer<int> InBuffer, uint Stride, in out uint RegisterIdx, uint ParticleIdx, int InValue)
+{
+	InBuffer[RegisterIdx * Stride + ParticleIdx] = InValue;
+	++RegisterIdx;
+}
+
+void NiagaraDataChannelWrite_Int32(RWBuffer<int> InBuffer, uint Stride, in out uint RegisterIdx, uint ParticleIdx, int2 InValue)
+{
+	NiagaraDataChannelWrite_Int32(InBuffer, Stride, RegisterIdx, ParticleIdx, InValue.x);
+	NiagaraDataChannelWrite_Int32(InBuffer, Stride, RegisterIdx, ParticleIdx, InValue.y);
+}
+
+void NiagaraDataChannelWrite_Int32(RWBuffer<int> InBuffer, uint Stride, in out uint RegisterIdx, uint ParticleIdx, int3 InValue)
+{
+	NiagaraDataChannelWrite_Int32(InBuffer, Stride, RegisterIdx, ParticleIdx, InValue.x);
+	NiagaraDataChannelWrite_Int32(InBuffer, Stride, RegisterIdx, ParticleIdx, InValue.y);
+	NiagaraDataChannelWrite_Int32(InBuffer, Stride, RegisterIdx, ParticleIdx, InValue.z);
+}
+
+void NiagaraDataChannelWrite_Int32(RWBuffer<int> InBuffer, uint Stride, in out uint RegisterIdx, uint ParticleIdx, int4 InValue)
+{
+	NiagaraDataChannelWrite_Int32(InBuffer, Stride, RegisterIdx, ParticleIdx, InValue.x);
+	NiagaraDataChannelWrite_Int32(InBuffer, Stride, RegisterIdx, ParticleIdx, InValue.y);
+	NiagaraDataChannelWrite_Int32(InBuffer, Stride, RegisterIdx, ParticleIdx, InValue.z);
+	NiagaraDataChannelWrite_Int32(InBuffer, Stride, RegisterIdx, ParticleIdx, InValue.w);
+}
+
+void NiagaraDataChannelWrite_Half(RWBuffer<half> InBuffer, uint Stride, in out uint RegisterIdx, uint ParticleIdx, half InValue)
+{
+	InBuffer[RegisterIdx * Stride + ParticleIdx] = InValue;
+	++RegisterIdx;
+}
+
+void NiagaraDataChannelWrite_Half(RWBuffer<half> InBuffer, uint Stride, in out uint RegisterIdx, uint ParticleIdx, half2 InValue)
+{
+	NiagaraDataChannelWrite_Half(InBuffer, Stride, RegisterIdx, ParticleIdx, InValue.x);
+	NiagaraDataChannelWrite_Half(InBuffer, Stride, RegisterIdx, ParticleIdx, InValue.y);
+}
+
+void NiagaraDataChannelWrite_Half(RWBuffer<half> InBuffer, uint Stride, in out uint RegisterIdx, uint ParticleIdx, half3 InValue)
+{
+	NiagaraDataChannelWrite_Half(InBuffer, Stride, RegisterIdx, ParticleIdx, InValue.x);
+	NiagaraDataChannelWrite_Half(InBuffer, Stride, RegisterIdx, ParticleIdx, InValue.y);
+	NiagaraDataChannelWrite_Half(InBuffer, Stride, RegisterIdx, ParticleIdx, InValue.z);
+}
+
+void NiagaraDataChannelWrite_Half(RWBuffer<half> InBuffer, uint Stride, in out uint RegisterIdx, uint ParticleIdx, half4 InValue)
+{
+	NiagaraDataChannelWrite_Half(InBuffer, Stride, RegisterIdx, ParticleIdx, InValue.x);
+	NiagaraDataChannelWrite_Half(InBuffer, Stride, RegisterIdx, ParticleIdx, InValue.y);
+	NiagaraDataChannelWrite_Half(InBuffer, Stride, RegisterIdx, ParticleIdx, InValue.z);
+	NiagaraDataChannelWrite_Half(InBuffer, Stride, RegisterIdx, ParticleIdx, InValue.w);
 }
 
 //TODO: More base type accessors.
diff --git a/Engine/Plugins/FX/Niagara/Shaders/Private/DataChannel/NiagaraDataInterfaceDataChannelTemplateCommon.ush b/Engine/Plugins/FX/Niagara/Shaders/Private/DataChannel/NiagaraDataInterfaceDataChannelTemplateCommon.ush
--- a/Engine/Plugins/FX/Niagara/Shaders/Private/DataChannel/NiagaraDataInterfaceDataChannelTemplateCommon.ush
+++ b/Engine/Plugins/FX/Niagara/Shaders/Private/DataChannel/NiagaraDataInterfaceDataChannelTemplateCommon.ush
@@ -1,9 +1,7 @@
 // Copyright Epic Games, Inc. All Rights Reserved.
 
 /**
-
 Template hlsl for Data Channel DI functions to apply to their own that is shared accross all functions for the DI.
-
 */
 
 //Offset of each parameter accessed by this script.
@@ -11,141 +9,11 @@ Template hlsl for Data Channel DI functions to apply to their own that is shared
 //TODO: Can probably pack this into a single global buffer rather than one per DI.
 Buffer<uint> {ParameterName}_ParamOffsetTable;
 
-Buffer<float> {ParameterName}_DataFloat;
-Buffer<int> {ParameterName}_DataInt32;
-Buffer<float> {ParameterName}_DataHalf;
-
 int {ParameterName}_ParameterOffsetTableIndex;
-int {ParameterName}_Num;
 int {ParameterName}_FloatStride;
 int {ParameterName}_Int32Stride;
 int {ParameterName}_HalfStride;
 
-//Buffer containing back mapping from spawn particle exec index to the NDC index that spawned that particle.
-//Arraged into power of two buckets so that we can avoid an entry per particle.
-Buffer<int> {ParameterName}_NDCSpawnDataBuffer;
-
-void Num_{ParameterName}(out int OutNum)
-{
-	OutNum = {ParameterName}_Num;
-}
-
-void GetNDCSpawnData_{ParameterName}(NiagaraEmitterID InEmitterID, int InExecIndex, out int OutNDCIndex, out int OutNDCSpawnIndex, out int OutNDCSpawnCount)
-{
-	/*
-	Look up the various bits of spawn data from the NDC that spawned this particle.
-	The first section of the NDCSpawnDataBuffer contains a SpawnCount for every NDC Item from which we can calcualte NDCSpawnCount.
-	The next section is a series of buckets containing all NDC indices that contain enough particles to spawn 1 << Bucket particles.
-	ie. if an NDC spawns 10 particles it will have an entry in the 8 and 2 buckets and none in any other bucket.
-	From this we can map the main ExecIndex back into an NDCIndex that generated it. However ofc this mapping is by bucket and so particles from one NDC may not be contiuous in the spawn script execution/particle data.
-	From the per NDCSpawnCount we can also determine if an NDC has an entry in each bucket and so calculate the SpawnIndex of each particle within the NDC Spawn.
-
-	This is a little convoluted but means we can access the spawning NDC Index when spawning in a single large SpawnInfo/Execution.
-	This power of 2 bucket system also means we can keep the indirection table size reasonable for large and small particle counts in exchange for some ALU work.
-
-	|| NDCSpawnDataBuffer |
-	|| Bucket Sizes	                                           || Spawn Counts for each NDC         || 2^15 Bucket   || 2^14 Bucket || ... || 2^1 Bucket || 2^0 Bucket   ||
-	|| Bucket 0 Size | Bucket 1 Size | ... | Bucket 15 Size    || NDC 2 | ... | NDC N               || 0 | 0 | 5 | 7 || 0 |.........|| ... || 0 | 1 | 5  ||              || 
-	
-	Above shows the layout for NDCSpawnDataBuffer.
-	All Emitters are packed into the same buffer. The first part of the buffer is a lookup table for each emitter (by emitter ID/Index) into the rest of the data.
-
-	As we proceed through each ExecIndex in the spawn execution, we iterate through each bucket and can know for example, we can allocate the first 2^15 particles to NDC 0 as that has an entry in the 2^15 bucket.
-	The 2^15 bucket is special in that it can have multiple entries for the same NDC as it's the final bucket. All other buckets have only one as higher particle counts just move the NDC into other buckets. 
-	Less fills up lower buckets, more with fill up this and lower buckets. 2x will push it into the bucket above. Except in the final bucket if we hit 2x or more 2^15 then we just have to add another entry to that bucket.
-	So in the above example the first and second 2^15 instances of the execution are given to NDC 0.
-	The next 2^15 are given to 5, then 7. Now we've processed the 2^15 bucket, further instances take from the next bucket and so on until we're at the 2^0 bucket and we've handled all instances.
-	*/
-
-	OutNDCIndex = -1;
-	OutNDCSpawnIndex = -1;
-	OutNDCSpawnCount = -1;
-
-	uint NumNDCEntries = {ParameterName}_Num;
-
-	//First part of the buffer is the per emitter offsets into the data.
-	//Grab the right base offset for this emitter.
-	uint NDCSpawnDataOffset = { ParameterName }_NDCSpawnDataBuffer[InEmitterID.ID];
-	if(NDCSpawnDataOffset == 0 || NumNDCEntries == 0)
-	{
-		OutNDCIndex = NDCSpawnDataOffset;
-		OutNDCSpawnIndex = InEmitterID.ID;
-		OutNDCSpawnCount = NumNDCEntries;
-
-		return;
-	}
-
-	uint NDCSpawnData_BucketSizeStart = NDCSpawnDataOffset;
-	uint NDCSpawnData_NDCCountsStart = NDCSpawnData_BucketSizeStart + 16;
-	uint NDCSPawnData_NDCIndexBucketsStart = NDCSpawnData_NDCCountsStart + NumNDCEntries;
-	
-	uint MaxBucketExecIndex = 0;
-	uint BucketEntryStart = NDCSPawnData_NDCIndexBucketsStart;
-
-	UNROLL
-	for(int BucketIdx = 0; BucketIdx < 16; ++BucketIdx)
-	{
-		uint BucketSize = (1U << 15U) >> BucketIdx;
-		uint NumEntriesInBucket = {ParameterName}_NDCSpawnDataBuffer[NDCSpawnData_BucketSizeStart + BucketIdx];
-		uint MinBucketExecIndex = MaxBucketExecIndex; 
-		MaxBucketExecIndex += BucketSize * NumEntriesInBucket;
-		
-		if(InExecIndex < MaxBucketExecIndex)
-		{
-			//We found our bucket.
-			//Now we need to find our NDCIndex Entry Index.
-			uint NDCIndexEntry = (InExecIndex - MinBucketExecIndex) / BucketSize;			
-			OutNDCIndex = {ParameterName}_NDCSpawnDataBuffer[BucketEntryStart + NDCIndexEntry];
-			break;
-		}
-		
-		BucketEntryStart += NumEntriesInBucket;
-	}
-	
-	if(OutNDCIndex >= 0 && OutNDCIndex < (int)NumNDCEntries)
-	{
-		//Pull the NDC Spawn count from the first part of the buffer.
-		OutNDCSpawnCount = {ParameterName}_NDCSpawnDataBuffer[NDCSpawnData_NDCCountsStart + OutNDCIndex];
-
-		//Do another pass to calculate our SpawnIndex for this NDC within the total count for this NDC.
-		OutNDCSpawnIndex = 0;
-		MaxBucketExecIndex = 0;
-		uint Count = OutNDCSpawnCount;
-		UNROLL
-		for(int BucketIdx = 0; BucketIdx < 16; ++BucketIdx)
-		{
-			uint BucketSize = (1U << 15U) >> BucketIdx;
-			uint Mask = (0xFFFFU >> (BucketIdx + 1U));
-			uint CountMasked = Count & ~Mask;
-			Count &= Mask;
-			uint NumNDCEntriesInBucket = CountMasked >> (15U - BucketIdx);
-			uint NumEntriesInBucket = {ParameterName}_NDCSpawnDataBuffer[NDCSpawnData_BucketSizeStart + BucketIdx];
-			uint NumNDCInstancesInBucket = NumNDCEntriesInBucket * BucketSize;
-
-			uint MinBucketExecIndex = MaxBucketExecIndex;
-			MaxBucketExecIndex += BucketSize * NumEntriesInBucket;
-			if (InExecIndex < MaxBucketExecIndex && NumNDCInstancesInBucket > 0)
-			{
-				uint NDCIndexEntry = (InExecIndex - MinBucketExecIndex) >> (15U - BucketIdx);
-				uint MinNDCBucketExecIndex = MinBucketExecIndex + (BucketSize * NDCIndexEntry);
-
-				OutNDCSpawnIndex += (InExecIndex - MinNDCBucketExecIndex);
-				break;
-			}
-			else
-			{
-				OutNDCSpawnIndex += NumNDCInstancesInBucket;
-			}
-		}
-	}
-	else
-	{
-		OutNDCIndex = -1;
-		OutNDCSpawnIndex = -1;
-		OutNDCSpawnCount = -1;
-	}
-}
-
 // Outputs the starting offset of each component in the channel data for a partiuclar parameter.
 bool GetParameterLayoutInfo_{ParameterName}(uint ParameterIndex, out uint OutFloatRegister, out uint OutInt32Register, out uint OutHalfRegister)
 {
@@ -167,4 +35,3 @@ bool GetParameterLayoutInfo_{ParameterName}(uint ParameterIndex, out uint OutFlo
 		return false;
 	}
 }
-
diff --git a/Engine/Plugins/FX/Niagara/Shaders/Private/DataChannel/NiagaraDataInterfaceDataChannelTemplateReadCommon.ush b/Engine/Plugins/FX/Niagara/Shaders/Private/DataChannel/NiagaraDataInterfaceDataChannelTemplateReadCommon.ush
new file mode 100644
--- /dev/null
+++ b/Engine/Plugins/FX/Niagara/Shaders/Private/DataChannel/NiagaraDataInterfaceDataChannelTemplateReadCommon.ush
@@ -0,0 +1,251 @@
+// Copyright Epic Games, Inc. All Rights Reserved.
+
+/**
+Template hlsl for Data Channel Read DI functions to apply to their own that is shared accross all functions for the DI.
+*/
+
+Buffer<float> {ParameterName}_DataFloat;
+Buffer<int> {ParameterName}_DataInt32;
+Buffer<float> {ParameterName}_DataHalf;
+
+//Instance count offset for the buffer we're reading.
+//Must be done via an instance count slot as data can come from CPU and GPU.
+int {ParameterName}_InstanceCountOffset;
+
+//Instance count offset for tracking consume calls.
+int {ParameterName}_ConsumeInstanceCountOffset;
+
+int {ParameterName}_BufferSize;
+
+int {ParameterName}_NDCElementCountAtSpawn;
+
+//Buffer containing back mapping from spawn particle exec index to the NDC index that spawned that particle.
+//Arraged into power of two buckets so that we can avoid an entry per particle.
+Buffer<int> {ParameterName}_NDCSpawnDataBuffer;
+
+//Helper struct with various info needed for reading that can be passed around more easily in generated code.
+struct FNDCAccessContext_{ParameterName}
+{
+	bool bValid;
+
+	//Index at which to read data.
+	int ReadIndex;
+		
+	//Component start offests for the current parameter.
+	int ParameterComponentStart_Float;
+	int ParameterComponentStart_Int32;
+	int ParameterComponentStart_Half;
+
+	//Current registe indices for each of the GPU write buffers
+	int Register_Float;
+	int Register_Int32;
+	int Register_Half;
+
+	int GetFloatStride() { return {ParameterName}_FloatStride; }
+	int GetIntStride() { return {ParameterName}_Int32Stride; }
+	int GetHalfStride() { return {ParameterName}_HalfStride; }
+		
+	Buffer<float> GetBuffer_Float() { return {ParameterName}_DataFloat; }
+	Buffer<int> GetBuffer_Int32() { return {ParameterName}_DataInt32; }
+	Buffer<float> GetBuffer_Half() { return {ParameterName}_DataHalf; }
+	
+	//Instance count used for the written number of elements in the NDC buffers.
+	int GetInstanceCountOffset() { return {ParameterName}_InstanceCountOffset; }
+
+	//Instance count used to get a read index when calling Consume().
+	int GetConsumeInstanceCountOffset() { return {ParameterName}_ConsumeInstanceCountOffset; }
+	
+	//Return the actual buffer size allocated.
+	int GetBufferSize() { return {ParameterName}_BufferSize; }
+
+	//Actual number of NDC items in the buffers. Must come from an indirect instance count instead of a uniform as the GPU could have written to this buffer.
+	int Num() { return RWInstanceCounts[ GetInstanceCountOffset() ]; }
+
+	bool IsValid() { return bValid; }
+
+	int GetConsumeIndex()
+	{
+		//TODO: Add GS and Wave op path to reduce atomic contention.
+		//TODO: Right now doing the simple dumb thing.
+
+		int OutIndex = -1;
+		if(GetConsumeInstanceCountOffset() != -1)
+		{
+			InterlockedAdd(RWInstanceCounts[GetConsumeInstanceCountOffset()], 1, OutIndex);
+			return OutIndex >= 0 && OutIndex < Num() && OutIndex < GetBufferSize() ? OutIndex : -1;
+		}
+
+		return -1;
+	}
+	
+	void InitDirect(int InReadIndex)
+	{
+		ReadIndex = InReadIndex;
+		bValid = ReadIndex >= 0 && ReadIndex < Num() && ReadIndex < GetBufferSize();
+	}
+
+	void InitConsume()
+	{
+		ReadIndex = GetConsumeIndex();
+		bValid = ReadIndex >= 0 && ReadIndex < Num() && ReadIndex < GetBufferSize();
+	}
+
+	void InitForParameter(int ParameterIndex)
+	{
+		bValid = bValid && GetParameterLayoutInfo_{ParameterName}(ParameterIndex, ParameterComponentStart_Float, ParameterComponentStart_Int32, ParameterComponentStart_Half);
+		
+		Register_Float = ParameterComponentStart_Float;
+		Register_Int32 = ParameterComponentStart_Int32;
+		Register_Half = ParameterComponentStart_Half;
+	}
+
+	void Read_Float(out float Value){ NiagaraDataChannelRead_Float(GetBuffer_Float(), GetFloatStride(), Register_Float, ReadIndex, Value); }
+	void Read_Float(out float2 Value){ NiagaraDataChannelRead_Float(GetBuffer_Float(), GetFloatStride(), Register_Float, ReadIndex, Value); }
+	void Read_Float(out float3 Value){ NiagaraDataChannelRead_Float(GetBuffer_Float(), GetFloatStride(), Register_Float, ReadIndex, Value); }
+	void Read_Float(out float4 Value){ NiagaraDataChannelRead_Float(GetBuffer_Float(), GetFloatStride(), Register_Float, ReadIndex, Value); }
+
+	void Read_Int32(out int Value){ NiagaraDataChannelRead_Int32(GetBuffer_Int32(), GetFloatStride(), Register_Int32, ReadIndex, Value); }
+	void Read_Int32(out int2 Value){ NiagaraDataChannelRead_Int32(GetBuffer_Int32(), GetFloatStride(), Register_Int32, ReadIndex, Value); }
+	void Read_Int32(out int3 Value){ NiagaraDataChannelRead_Int32(GetBuffer_Int32(), GetFloatStride(), Register_Int32, ReadIndex, Value); }
+	void Read_Int32(out int4 Value){ NiagaraDataChannelRead_Int32(GetBuffer_Int32(), GetFloatStride(), Register_Int32, ReadIndex, Value); }
+
+	void Read_Half(out half Value){ NiagaraDataChannelRead_Half(GetBuffer_Half(), GetFloatStride(), Register_Half, ReadIndex, Value); }
+	void Read_Half(out half2 Value){ NiagaraDataChannelRead_Half(GetBuffer_Half(), GetFloatStride(), Register_Half, ReadIndex, Value); }
+	void Read_Half(out half3 Value){ NiagaraDataChannelRead_Half(GetBuffer_Half(), GetFloatStride(), Register_Half, ReadIndex, Value); }
+	void Read_Half(out half4 Value){ NiagaraDataChannelRead_Half(GetBuffer_Half(), GetFloatStride(), Register_Half, ReadIndex, Value); }
+};
+
+void Num_{ParameterName}(out int OutNum)
+{
+	if({ParameterName}_InstanceCountOffset != -1)
+	{
+		OutNum = RWInstanceCounts[{ParameterName}_InstanceCountOffset];
+	}
+	else
+	{
+		OutNum = 0;
+	}
+}
+
+void GetNDCSpawnData_{ParameterName}(NiagaraEmitterID InEmitterID, int InExecIndex, out int OutNDCIndex, out int OutNDCSpawnIndex, out int OutNDCSpawnCount)
+{
+	/*
+	Look up the various bits of spawn data from the NDC that spawned this particle.
+	The first section of the NDCSpawnDataBuffer contains a SpawnCount for every NDC Item from which we can calcualte NDCSpawnCount.
+	The next section is a series of buckets containing all NDC indices that contain enough particles to spawn 1 << Bucket particles.
+	ie. if an NDC spawns 10 particles it will have an entry in the 8 and 2 buckets and none in any other bucket.
+	From this we can map the main ExecIndex back into an NDCIndex that generated it. However ofc this mapping is by bucket and so particles from one NDC may not be contiuous in the spawn script execution/particle data.
+	From the per NDCSpawnCount we can also determine if an NDC has an entry in each bucket and so calculate the SpawnIndex of each particle within the NDC Spawn.
+
+	This is a little convoluted but means we can access the spawning NDC Index when spawning in a single large SpawnInfo/Execution.
+	This power of 2 bucket system also means we can keep the indirection table size reasonable for large and small particle counts in exchange for some ALU work.
+
+	|| NDCSpawnDataBuffer |
+	|| Bucket Sizes	                                           || Spawn Counts for each NDC         || 2^15 Bucket   || 2^14 Bucket || ... || 2^1 Bucket || 2^0 Bucket   ||
+	|| Bucket 0 Size | Bucket 1 Size | ... | Bucket 15 Size    || NDC 2 | ... | NDC N               || 0 | 0 | 5 | 7 || 0 |.........|| ... || 0 | 1 | 5  ||              || 
+	
+	Above shows the layout for NDCSpawnDataBuffer.
+	All Emitters are packed into the same buffer. The first part of the buffer is a lookup table for each emitter (by emitter ID/Index) into the rest of the data.
+
+	As we proceed through each ExecIndex in the spawn execution, we iterate through each bucket and can know for example, we can allocate the first 2^15 particles to NDC 0 as that has an entry in the 2^15 bucket.
+	The 2^15 bucket is special in that it can have multiple entries for the same NDC as it's the final bucket. All other buckets have only one as higher particle counts just move the NDC into other buckets. 
+	Less fills up lower buckets, more with fill up this and lower buckets. 2x will push it into the bucket above. Except in the final bucket if we hit 2x or more 2^15 then we just have to add another entry to that bucket.
+	So in the above example the first and second 2^15 instances of the execution are given to NDC 0.
+	The next 2^15 are given to 5, then 7. Now we've processed the 2^15 bucket, further instances take from the next bucket and so on until we're at the 2^0 bucket and we've handled all instances.
+	*/
+	
+	OutNDCIndex = -1;
+	OutNDCSpawnIndex = -1;
+	OutNDCSpawnCount = -1;
+
+	if(InExecIndex < 0)
+	{
+		return;
+	}
+
+	uint ExecIdx = (uint)InExecIndex;
+
+	uint NumNDCEntries = {ParameterName}_NDCElementCountAtSpawn;
+
+	//First part of the buffer is the per emitter offsets into the data.
+	//Grab the right base offset for this emitter.
+	uint NDCSpawnDataOffset = { ParameterName }_NDCSpawnDataBuffer[InEmitterID.ID];
+	if(NDCSpawnDataOffset == 0 || NumNDCEntries == 0)
+	{
+		OutNDCIndex = -1;
+		OutNDCSpawnIndex = -1;
+		OutNDCSpawnCount = 0;
+
+		return;
+	}
+
+	uint NDCSpawnData_BucketSizeStart = NDCSpawnDataOffset;
+	uint NDCSpawnData_NDCCountsStart = NDCSpawnData_BucketSizeStart + 16;
+	uint NDCSPawnData_NDCIndexBucketsStart = NDCSpawnData_NDCCountsStart + NumNDCEntries;
+	
+	uint MaxBucketExecIndex = 0;
+	uint BucketEntryStart = NDCSPawnData_NDCIndexBucketsStart;
+
+	UNROLL
+	for(int BucketIdx = 0; BucketIdx < 16; ++BucketIdx)
+	{
+		uint BucketSize = (1U << 15U) >> BucketIdx;
+		uint NumEntriesInBucket = {ParameterName}_NDCSpawnDataBuffer[NDCSpawnData_BucketSizeStart + BucketIdx];
+		uint MinBucketExecIndex = MaxBucketExecIndex; 
+		MaxBucketExecIndex += BucketSize * NumEntriesInBucket;
+		
+		if(ExecIdx < MaxBucketExecIndex)
+		{
+			//We found our bucket.
+			//Now we need to find our NDCIndex Entry Index.
+			uint NDCIndexEntry = (ExecIdx - MinBucketExecIndex) / BucketSize;			
+			OutNDCIndex = {ParameterName}_NDCSpawnDataBuffer[BucketEntryStart + NDCIndexEntry];
+			break;
+		}
+		
+		BucketEntryStart += NumEntriesInBucket;
+	}
+	
+	if(OutNDCIndex >= 0 && OutNDCIndex < (int)NumNDCEntries)
+	{
+		//Pull the NDC Spawn count from the first part of the buffer.
+		OutNDCSpawnCount = {ParameterName}_NDCSpawnDataBuffer[NDCSpawnData_NDCCountsStart + OutNDCIndex];
+
+		//Do another pass to calculate our SpawnIndex for this NDC within the total count for this NDC.
+		OutNDCSpawnIndex = 0;
+		MaxBucketExecIndex = 0;
+		uint Count = OutNDCSpawnCount;
+		UNROLL
+		for(int BucketIdx = 0; BucketIdx < 16; ++BucketIdx)
+		{
+			uint BucketSize = (1U << 15U) >> BucketIdx;
+			uint Mask = (0xFFFFU >> (BucketIdx + 1U));
+			uint CountMasked = Count & ~Mask;
+			Count &= Mask;
+			uint NumNDCEntriesInBucket = CountMasked >> (15U - BucketIdx);
+			uint NumEntriesInBucket = {ParameterName}_NDCSpawnDataBuffer[NDCSpawnData_BucketSizeStart + BucketIdx];
+			uint NumNDCInstancesInBucket = NumNDCEntriesInBucket * BucketSize;
+
+			uint MinBucketExecIndex = MaxBucketExecIndex;
+			MaxBucketExecIndex += BucketSize * NumEntriesInBucket;
+			if (ExecIdx < MaxBucketExecIndex && NumNDCInstancesInBucket > 0)
+			{
+				uint NDCIndexEntry = (ExecIdx - MinBucketExecIndex) >> (15U - BucketIdx);
+				uint MinNDCBucketExecIndex = MinBucketExecIndex + (BucketSize * NDCIndexEntry);
+
+				OutNDCSpawnIndex += (ExecIdx - MinNDCBucketExecIndex);
+				break;
+			}
+			else
+			{
+				OutNDCSpawnIndex += NumNDCInstancesInBucket;
+			}
+		}
+	}
+	else
+	{
+		OutNDCIndex = -1;
+		OutNDCSpawnIndex = -1;
+		OutNDCSpawnCount = -1;
+	}
+}
diff --git a/Engine/Plugins/FX/Niagara/Shaders/Private/DataChannel/NiagaraDataInterfaceDataChannelTemplateWriteCommon.ush b/Engine/Plugins/FX/Niagara/Shaders/Private/DataChannel/NiagaraDataInterfaceDataChannelTemplateWriteCommon.ush
new file mode 100644
--- /dev/null
+++ b/Engine/Plugins/FX/Niagara/Shaders/Private/DataChannel/NiagaraDataInterfaceDataChannelTemplateWriteCommon.ush
@@ -0,0 +1,157 @@
+// Copyright Epic Games, Inc. All Rights Reserved.
+
+/**
+Template hlsl for Data Channel Write DI functions to apply to their own that is shared accross all functions for the DI.
+*/
+
+//Data bound for other readers on the GPU is written into one large shared buffer
+//We index into this using atomic ops on an instance count in RWCountsBuffer.
+RWBuffer<float> {ParameterName}_GPUBufferFloat;
+RWBuffer<int> {ParameterName}_GPUBufferInt32;
+RWBuffer<float> {ParameterName}_GPUBufferHalf;
+
+//Buffer and offset for our instance count value.
+//We'll atomically update this accordingly if we're write/appending.
+int {ParameterName}_GPUInstanceCountOffset;
+int {ParameterName}_GPUBufferSize;
+
+/////////////////////////////////////////
+
+//Data bounds back to the CPU for game and CPU sims is written into a separate specific buffer indexed directly by the thread index.
+RWBuffer<float> {ParameterName}_CPUBufferFloat;
+RWBuffer<int> {ParameterName}_CPUBufferInt32;
+RWBuffer<float> {ParameterName}_CPUBufferHalf;
+
+//Buffer and offset for our instance count value.
+//We'll atomically update this accordingly if we're write/appending.
+int {ParameterName}_CPUInstanceCountOffset;
+int {ParameterName}_CPUBufferSize;
+int {ParameterName}_CPUFloatStride;
+int {ParameterName}_CPUInt32Stride;
+int {ParameterName}_CPUHalfStride;
+
+//Helper struct with various info needed for writing that can be passed around more easily in generated code.
+struct FNDCAccessContext_{ParameterName}
+{
+	bool bValid;
+
+	//Indices we'll be writing to each buffer.
+	int WriteIndex;
+
+	int WriteIndexCPU;
+	int WriteIndexGPU;
+		
+	//Component start offests for the current parameter.
+	int ParameterComponentStart_Float;
+	int ParameterComponentStart_Int32;
+	int ParameterComponentStart_Half;
+
+	//Current register indices for each component buffer
+	int CurrRegister_Float;
+	int CurrRegister_Int32;
+	int CurrRegister_Half;
+
+	bool bReadGPUOrCPU;
+
+	int GetFloatStride() { return bReadGPUOrCPU ? {ParameterName}_FloatStride : {ParameterName}_CPUFloatStride; }
+	int GetInt32Stride() { return bReadGPUOrCPU ? {ParameterName}_Int32Stride : {ParameterName}_CPUInt32Stride; }
+	int GetHalfStride() { return bReadGPUOrCPU ? {ParameterName}_HalfStride : {ParameterName}_CPUHalfStride; }
+	
+	RWBuffer<float> GetGPUBuffer_Float() { return {ParameterName}_GPUBufferFloat; }
+	RWBuffer<int> GetGPUBuffer_Int32() { return {ParameterName}_GPUBufferInt32; }
+	RWBuffer<float> GetGPUBuffer_Half() { return {ParameterName}_GPUBufferHalf; }
+	int GetGPUInstanceCountOffset() { return {ParameterName}_GPUInstanceCountOffset; }
+	int GetGPUBufferSize() { return {ParameterName}_GPUBufferSize; }
+
+	RWBuffer<float> GetCPUBuffer_Float() { return {ParameterName}_CPUBufferFloat; }
+	RWBuffer<int> GetCPUBuffer_Int32() { return {ParameterName}_CPUBufferInt32; }
+	RWBuffer<float> GetCPUBuffer_Half() { return {ParameterName}_CPUBufferHalf; }
+	int GetCPUInstanceCountOffset() { return {ParameterName}_CPUInstanceCountOffset; }
+	int GetCPUBufferSize() { return {ParameterName}_CPUBufferSize; }
+
+	RWBuffer<float> GetCurrBuffer_Float() { return bReadGPUOrCPU ? GetGPUBuffer_Float() : GetCPUBuffer_Float(); }
+	RWBuffer<int> GetCurrBuffer_Int32() { return bReadGPUOrCPU ? GetGPUBuffer_Int32() : GetCPUBuffer_Int32(); }
+	RWBuffer<float> GetCurrBuffer_Half() { return bReadGPUOrCPU ? GetGPUBuffer_Half() : GetCPUBuffer_Half(); }
+
+	bool IsValid() { return bValid; }
+	bool ShouldWriteGPUData() { return bValid && GetGPUInstanceCountOffset() != -1; }
+	bool ShouldWriteCPUData() { return bValid && GetCPUInstanceCountOffset() != -1; }
+
+	int GetAppendIndex(int InstanceCountOffset, int BufferSize)
+	{
+		//TODO: Add GS and Wave op path to reduce atomic contention.
+		//TODO: Right now doing the simple dumb thing.
+
+		int OutIndex = -1;
+		if(InstanceCountOffset != -1)
+		{
+			//Note: This allows the index to temporarily exceed the max limits so is unsafe if we access this concurrently anywhere else without checking the limits.
+			//Note: However it does avoid a more expensive looping compare exchange.
+			InterlockedAdd(RWInstanceCounts[InstanceCountOffset], 1, OutIndex);
+			if(OutIndex >= 0 && OutIndex < BufferSize)
+			{
+				return OutIndex;
+			}
+		
+			OutIndex = -1;
+			InterlockedMin(RWInstanceCounts[InstanceCountOffset], BufferSize);
+		}
+
+		return OutIndex;
+	}
+	
+	//Initialize our append indices. Called at the start of the append funciton.
+	void InitAppend()
+	{
+		WriteIndexGPU = GetAppendIndex(GetGPUInstanceCountOffset(), GetGPUBufferSize());
+		WriteIndexCPU = GetAppendIndex(GetCPUInstanceCountOffset(), GetCPUBufferSize());
+	
+		bValid = (WriteIndexGPU >= 0 && WriteIndexGPU < GetGPUBufferSize()) ||
+				(WriteIndexCPU  >= 0 && WriteIndexCPU < GetCPUBufferSize());
+	}
+
+	//Initialize for the given parameter index.
+	void InitForParameter(int ParameterIndex)
+	{
+		bValid = bValid && GetParameterLayoutInfo_{ParameterName}(ParameterIndex, ParameterComponentStart_Float, ParameterComponentStart_Int32, ParameterComponentStart_Half);				
+	}
+	
+	//Initialize for wrting the current parameter to the GPU buffers
+	bool InitForGPUWrite()
+	{
+		bReadGPUOrCPU = true;
+		WriteIndex = WriteIndexGPU;	
+		CurrRegister_Float = ParameterComponentStart_Float;
+		CurrRegister_Int32 = ParameterComponentStart_Int32;
+		CurrRegister_Half = ParameterComponentStart_Half;
+		
+		return ShouldWriteGPUData();
+	}
+
+	//Initialize for wrting the current parameter to the "CPU" buffers
+	bool InitForCPUWrite()
+	{
+		bReadGPUOrCPU = false;
+		WriteIndex = WriteIndexCPU;	
+		CurrRegister_Float = ParameterComponentStart_Float;
+		CurrRegister_Int32 = ParameterComponentStart_Int32;
+		CurrRegister_Half = ParameterComponentStart_Half;
+		
+		return ShouldWriteCPUData();
+	}
+
+	void Write_Float(float Value){ NiagaraDataChannelWrite_Float(GetCurrBuffer_Float(), GetFloatStride(), CurrRegister_Float, WriteIndex, Value);	}
+	void Write_Float(float2 Value){	NiagaraDataChannelWrite_Float(GetCurrBuffer_Float(), GetFloatStride(), CurrRegister_Float, WriteIndex, Value);	}
+	void Write_Float(float3 Value){	NiagaraDataChannelWrite_Float(GetCurrBuffer_Float(), GetFloatStride(), CurrRegister_Float, WriteIndex, Value);	}
+	void Write_Float(float4 Value){	NiagaraDataChannelWrite_Float(GetCurrBuffer_Float(), GetFloatStride(), CurrRegister_Float, WriteIndex, Value);	}
+
+	void Write_Int32(int Value){ NiagaraDataChannelWrite_Int32(GetCurrBuffer_Int32(), GetInt32Stride(), CurrRegister_Int32, WriteIndex, Value);	}
+	void Write_Int32(int2 Value){ NiagaraDataChannelWrite_Int32(GetCurrBuffer_Int32(), GetInt32Stride(), CurrRegister_Int32, WriteIndex, Value);	}
+	void Write_Int32(int3 Value){ NiagaraDataChannelWrite_Int32(GetCurrBuffer_Int32(), GetInt32Stride(), CurrRegister_Int32, WriteIndex, Value);	}
+	void Write_Int32(int4 Value){ NiagaraDataChannelWrite_Int32(GetCurrBuffer_Int32(), GetInt32Stride(), CurrRegister_Int32, WriteIndex, Value);	}
+
+	void Write_Half(half Value){ NiagaraDataChannelWrite_Half(GetCurrBuffer_Half(), GetHalfStride(), CurrRegister_Half, WriteIndex, Value); }
+	void Write_Half(half2 Value){ NiagaraDataChannelWrite_Half(GetCurrBuffer_Half(), GetHalfStride(), CurrRegister_Half, WriteIndex, Value); }
+	void Write_Half(half3 Value){ NiagaraDataChannelWrite_Half(GetCurrBuffer_Half(), GetHalfStride(), CurrRegister_Half, WriteIndex, Value); }
+	void Write_Half(half4 Value){ NiagaraDataChannelWrite_Half(GetCurrBuffer_Half(), GetHalfStride(), CurrRegister_Half, WriteIndex, Value); }
+};
\ No newline at end of file
diff --git a/Engine/Plugins/FX/Niagara/Shaders/Private/DataChannel/NiagaraDataInterfaceDataChannelTemplate_Append.ush b/Engine/Plugins/FX/Niagara/Shaders/Private/DataChannel/NiagaraDataInterfaceDataChannelTemplate_Append.ush
--- a/Engine/Plugins/FX/Niagara/Shaders/Private/DataChannel/NiagaraDataInterfaceDataChannelTemplate_Append.ush
+++ b/Engine/Plugins/FX/Niagara/Shaders/Private/DataChannel/NiagaraDataInterfaceDataChannelTemplate_Append.ush
@@ -1,14 +1,26 @@
 // Copyright Epic Games, Inc. All Rights Reserved.
 
-#pragma once
-
 /**
  * Template shader code for the Append function for the NiagaraDataInterfaceDataChannel.
  * Contained in it's own shader file to allow reuse between different invocations with different params.
  */
 
-void Append_{ParameterName}_UEImpureCall(bool bEmit, out bool bOutSuccess, {FunctionParameters})
+{PerParameterFunctionDefinitions}
+
+void {FunctionSymbol}(bool bEmit, {FunctionInputParameters}, out bool bOutSuccess)
 {
-	//TODO;
 	bOutSuccess = false;
+	
+	if(bEmit)
+	{
+		FNDCAccessContext_{ParameterName} Context;
+		Context.InitAppend();
+		if(Context.IsValid())
+		{
+			bOutSuccess = true;
+	
+			//This is generated in the DI Cpp depending on the function parameters.
+{PerFunctionParameterShaderCode}
+		}
+	}
 }
\ No newline at end of file
diff --git a/Engine/Plugins/FX/Niagara/Shaders/Private/DataChannel/NiagaraDataInterfaceDataChannelTemplate_Consume.ush b/Engine/Plugins/FX/Niagara/Shaders/Private/DataChannel/NiagaraDataInterfaceDataChannelTemplate_Consume.ush
--- a/Engine/Plugins/FX/Niagara/Shaders/Private/DataChannel/NiagaraDataInterfaceDataChannelTemplate_Consume.ush
+++ b/Engine/Plugins/FX/Niagara/Shaders/Private/DataChannel/NiagaraDataInterfaceDataChannelTemplate_Consume.ush
@@ -1,14 +1,28 @@
 // Copyright Epic Games, Inc. All Rights Reserved.
 
-#pragma once
-
 /**
  * Template shader code for the Consume function for the NiagaraDataInterfaceDataChannel.
  * Contained in it's own shader file to allow reuse between different invocations with different params.
  */
 
-void Consume_{ParameterName}(int ElementIndex, out bool bOutSuccess, {FunctionParameters})
+{PerParameterFunctionDefinitions}
+
+void {FunctionSymbol}(bool bInConsume, out bool bOutSuccess, out int OutIndex, {FunctionOutputParameters})
 {
-	//TODO;
+	//Default outputs in case of failure.
+	{DefaultOutputsShaderCode}
+
 	bOutSuccess = false;
-}
\ No newline at end of file
+	OutIndex = -1;
+	
+	FNDCAccessContext_{ParameterName} Context;
+	Context.InitConsume();
+	if(Context.IsValid())
+	{
+		bOutSuccess = true;
+
+		//This is generated in the DI Cpp depending on the function parameters.
+{PerFunctionParameterShaderCode}
+	}
+}
+
diff --git a/Engine/Plugins/FX/Niagara/Shaders/Private/DataChannel/NiagaraDataInterfaceDataChannelTemplate_Read.ush b/Engine/Plugins/FX/Niagara/Shaders/Private/DataChannel/NiagaraDataInterfaceDataChannelTemplate_Read.ush
--- a/Engine/Plugins/FX/Niagara/Shaders/Private/DataChannel/NiagaraDataInterfaceDataChannelTemplate_Read.ush
+++ b/Engine/Plugins/FX/Niagara/Shaders/Private/DataChannel/NiagaraDataInterfaceDataChannelTemplate_Read.ush
@@ -5,26 +5,23 @@
  * Contained in it's own shader file to allow reuse between different invocations with different params.
  */
 
-void {FunctionSymbol}(int ElementIndex, out bool bOutSuccess, {FunctionParameters})
+{PerParameterFunctionDefinitions}
+
+void {FunctionSymbol}(int ElementIndex, out bool bOutSuccess, {FunctionOutputParameters})
 {
-	if(ElementIndex >= 0 && ElementIndex < {ParameterName}_Num)
+	//Default outputs in case of failure.
+	{DefaultOutputsShaderCode}
+
+	bOutSuccess = false;
+
+	FNDCAccessContext_{ParameterName} Context;
+	Context.InitDirect(ElementIndex);
+	if(Context.IsValid())
 	{
 		bOutSuccess = true;
 
-#if NIAGARA_DATACHANNEL_MULTIBUFFERS
-		//TODO: Allow reads from N fixed number of buffers.
-#elif NIAGARA_DATACHANNEL_BINDLESSBUFFERS
-		//TODO: Allow reading from any number of buffers when bindless buffers are supported.
-#else
-		//Reads from only a single buffer supported.
-		
 		//This is generated in the DI Cpp depending on the function parameters.
 {PerFunctionParameterShaderCode}
-#endif
 	}
-	//else
-	//{
-	//	//TODO: generate defaults hlsl.
-			//	bOutSuccess = true;
-	//}
-}
\ No newline at end of file
+}
+
diff --git a/Engine/Plugins/FX/Niagara/Shaders/Private/DataChannel/NiagaraDataInterfaceDataChannelTemplate_Write.ush b/Engine/Plugins/FX/Niagara/Shaders/Private/DataChannel/NiagaraDataInterfaceDataChannelTemplate_Write.ush
--- a/Engine/Plugins/FX/Niagara/Shaders/Private/DataChannel/NiagaraDataInterfaceDataChannelTemplate_Write.ush
+++ b/Engine/Plugins/FX/Niagara/Shaders/Private/DataChannel/NiagaraDataInterfaceDataChannelTemplate_Write.ush
@@ -1,14 +1,15 @@
 // Copyright Epic Games, Inc. All Rights Reserved.
 
-#pragma once
-
 /**
  * Template shader code for the Write function for the NiagaraDataInterfaceDataChannel.
  * Contained in it's own shader file to allow reuse between different invocations with different params.
  */
 
-void Write_{ParameterName}_UEImpureCall(bool bEmit, int ElementIndex, out bool bOutSuccess, {FunctionParameters})
+{PerParameterFunctionDefinitions}
+
+void {FunctionSymbol}(bool bEmit, int ElementIndex, {FunctionInputParameters}, out bool bOutSuccess)
 {
-	//TODO;
-	bOutSuccess = false;
+	//Write is deprecated on the GPU.
+	//All writes now done with append as a/ it's less useful, b/ it's brittle as multiple things could change internally, c/ Already can't work on GPU as we're wll writing concurrently to a single shared buffer.
+	//TODO: Remove
 }
\ No newline at end of file
diff --git a/Engine/Plugins/FX/Niagara/Shaders/Private/NiagaraDrawIndirectArgsGen.usf b/Engine/Plugins/FX/Niagara/Shaders/Private/NiagaraDrawIndirectArgsGen.usf
--- a/Engine/Plugins/FX/Niagara/Shaders/Private/NiagaraDrawIndirectArgsGen.usf
+++ b/Engine/Plugins/FX/Niagara/Shaders/Private/NiagaraDrawIndirectArgsGen.usf
@@ -102,4 +102,18 @@ void ResetCountsCS(uint TaskIndex : SV_DispatchThreadID)
 		const uint InfoOffset = (ArgGenTaskOffset + NumArgGenTasks) * NIAGARA_DRAW_INDIRECT_TASK_INFO_SIZE + (TaskIndex);
 		RWInstanceCounts[TaskInfos[InfoOffset]] = 0;
 	}
+}
+
+[numthreads(THREAD_COUNT,1,1)]
+void InitInstanceCountsCS(uint TaskIndex : SV_DispatchThreadID)
+{
+	const uint NumTasks = TaskCount.x;
+	if (TaskIndex < NumTasks)
+	{
+		const uint InfoOffset = TaskIndex * NIAGARA_INIT_INSTANCE_COUNT_TASK_INFO_SIZE;
+		
+		const uint CountOffset = TaskInfos[InfoOffset];
+		const uint CountValue = TaskInfos[InfoOffset + 1];
+		RWInstanceCounts[CountOffset] = CountValue;
+	}
 }
\ No newline at end of file
diff --git a/Engine/Plugins/FX/Niagara/Source/Niagara/Classes/NiagaraDataInterface.h b/Engine/Plugins/FX/Niagara/Source/Niagara/Classes/NiagaraDataInterface.h
--- a/Engine/Plugins/FX/Niagara/Source/Niagara/Classes/NiagaraDataInterface.h
+++ b/Engine/Plugins/FX/Niagara/Source/Niagara/Classes/NiagaraDataInterface.h
@@ -25,6 +25,10 @@ class FNiagaraGPUSystemTick;
 struct FNiagaraSimStageData;
 class FNiagaraSystemInstance;
 struct FNiagaraDataInterfaceHlslGenerationContext;
+class FNiagaraGPUInstanceCountManager;
+class FNiagaraGpuReadbackManager;
+
+namespace ERHIFeatureLevel { enum Type : int; }
 
 struct FNDITransformHandlerNoop
 {
@@ -289,7 +293,7 @@ protected:
 
 struct FNDIGpuComputeContext
 {
-	FNDIGpuComputeContext(FRDGBuilder& InGraphBuilder, const FNiagaraGpuComputeDispatchInterface& InComputeDispatchInterface)
+	FNDIGpuComputeContext(FRDGBuilder& InGraphBuilder, FNiagaraGpuComputeDispatchInterface& InComputeDispatchInterface)
 		: GraphBuilder(InGraphBuilder)
 		, ComputeDispatchInterface(InComputeDispatchInterface)
 	{
@@ -298,16 +302,18 @@ struct FNDIGpuComputeContext
 	FRDGBuilder& GetGraphBuilder() const { return GraphBuilder; }
 	NIAGARA_API FRDGExternalAccessQueue& GetRDGExternalAccessQueue() const;
 
+	// Expose only const direct access to the DispatchInterface.
+	// For any non-const operations this and child context structs should expose only necessary and safe non const operations for that context.
 	const FNiagaraGpuComputeDispatchInterface& GetComputeDispatchInterface() const { return ComputeDispatchInterface; }
 
 protected:
 	FRDGBuilder& GraphBuilder;
-	const FNiagaraGpuComputeDispatchInterface& ComputeDispatchInterface;
+	FNiagaraGpuComputeDispatchInterface& ComputeDispatchInterface;
 };
 
 struct FNDIGpuComputeResetContext : public FNDIGpuComputeContext
 {
-	FNDIGpuComputeResetContext(FRDGBuilder& InGraphBuilder, const FNiagaraGpuComputeDispatchInterface& InComputeDispatchInterface, FNiagaraSystemInstanceID InSystemInstanceID)
+	FNDIGpuComputeResetContext(FRDGBuilder& InGraphBuilder, FNiagaraGpuComputeDispatchInterface& InComputeDispatchInterface, FNiagaraSystemInstanceID InSystemInstanceID)
 		: FNDIGpuComputeContext(InGraphBuilder, InComputeDispatchInterface)
 		, SystemInstanceID(InSystemInstanceID)
 	{
@@ -321,7 +327,7 @@ protected:
 
 struct FNDIGpuComputePrePostStageContext : public FNDIGpuComputeContext
 {
-	FNDIGpuComputePrePostStageContext(FRDGBuilder& InGraphBuilder, const FNiagaraGpuComputeDispatchInterface& InComputeDispatchInterface, const FNiagaraGPUSystemTick& InSystemTick, const FNiagaraComputeInstanceData& InComputeInstanceData, const FNiagaraSimStageData& InSimStageData)
+	FNDIGpuComputePrePostStageContext(FRDGBuilder& InGraphBuilder, FNiagaraGpuComputeDispatchInterface& InComputeDispatchInterface, const FNiagaraGPUSystemTick& InSystemTick, const FNiagaraComputeInstanceData& InComputeInstanceData, const FNiagaraSimStageData& InSimStageData)
 		: FNDIGpuComputeContext(InGraphBuilder, InComputeDispatchInterface)
 		, SystemTick(InSystemTick)
 		, ComputeInstanceData(InComputeInstanceData)
@@ -340,6 +346,8 @@ struct FNDIGpuComputePrePostStageContext : public FNDIGpuComputeContext
 
 	void SetDataInterfaceProxy(FNiagaraDataInterfaceProxy* InDataInterfaceProxy) { DataInterfaceProxy = InDataInterfaceProxy; }
 
+	FNiagaraGPUInstanceCountManager& GetInstanceCountManager()const;
+
 protected:
 	const FNiagaraGPUSystemTick& SystemTick;
 	const FNiagaraComputeInstanceData& ComputeInstanceData;
@@ -347,12 +355,25 @@ protected:
 	FNiagaraDataInterfaceProxy* DataInterfaceProxy = nullptr;
 };
 
-using FNDIGpuComputePreStageContext = FNDIGpuComputePrePostStageContext;
-using FNDIGpuComputePostStageContext = FNDIGpuComputePrePostStageContext;
+struct FNDIGpuComputePreStageContext : public FNDIGpuComputePrePostStageContext
+{
+	FNDIGpuComputePreStageContext(FRDGBuilder& InGraphBuilder, FNiagaraGpuComputeDispatchInterface& InComputeDispatchInterface, const FNiagaraGPUSystemTick& InSystemTick, const FNiagaraComputeInstanceData& InComputeInstanceData, const FNiagaraSimStageData& InSimStageData)
+		: FNDIGpuComputePrePostStageContext(InGraphBuilder, InComputeDispatchInterface, InSystemTick, InComputeInstanceData, InSimStageData)
+	{
+	}
+};
+
+struct FNDIGpuComputePostStageContext : public FNDIGpuComputePrePostStageContext
+{
+	FNDIGpuComputePostStageContext(FRDGBuilder& InGraphBuilder, FNiagaraGpuComputeDispatchInterface& InComputeDispatchInterface, const FNiagaraGPUSystemTick& InSystemTick, const FNiagaraComputeInstanceData& InComputeInstanceData, const FNiagaraSimStageData& InSimStageData)
+		: FNDIGpuComputePrePostStageContext(InGraphBuilder, InComputeDispatchInterface, InSystemTick, InComputeInstanceData, InSimStageData)
+	{
+	}
+};
 
 struct FNDIGpuComputePostSimulateContext : public FNDIGpuComputeContext
 {
-	FNDIGpuComputePostSimulateContext(FRDGBuilder& InGraphBuilder, const FNiagaraGpuComputeDispatchInterface& InComputeDispatchInterface, FNiagaraSystemInstanceID InSystemInstanceID, bool InFinalPostSimulate)
+	FNDIGpuComputePostSimulateContext(FRDGBuilder& InGraphBuilder, FNiagaraGpuComputeDispatchInterface& InComputeDispatchInterface, FNiagaraSystemInstanceID InSystemInstanceID, bool InFinalPostSimulate)
 		: FNDIGpuComputeContext(InGraphBuilder, InComputeDispatchInterface)
 		, SystemInstanceID(InSystemInstanceID)
 		, bFinalPostSimulate(InFinalPostSimulate)
@@ -362,6 +383,8 @@ struct FNDIGpuComputePostSimulateContext : public FNDIGpuComputeContext
 	FNiagaraSystemInstanceID GetSystemInstanceID() const { return SystemInstanceID; }
 	bool IsFinalPostSimulate() const { return bFinalPostSimulate; }
 
+	FNiagaraGPUInstanceCountManager& GetInstanceCountManager()const;
+
 protected:
 	FNiagaraSystemInstanceID SystemInstanceID = FNiagaraSystemInstanceID();
 	bool bFinalPostSimulate = false;
@@ -699,6 +722,7 @@ public:
 	virtual bool RequiresDepthBuffer() const { return false; }
 	virtual bool RequiresEarlyViewData() const { return false; }
 	virtual bool RequiresRayTracingScene() const { return false; }
+	virtual bool RequiresCurrentFrameNDC() const { return false; }
 
 	virtual bool HasTickGroupPrereqs() const { return false; }
 	virtual ETickingGroup CalculateTickGroup(const void* PerInstanceData) const { return NiagaraFirstTickGroup; }
diff --git a/Engine/Plugins/FX/Niagara/Source/Niagara/Classes/NiagaraDataInterfaceRW.h b/Engine/Plugins/FX/Niagara/Source/Niagara/Classes/NiagaraDataInterfaceRW.h
--- a/Engine/Plugins/FX/Niagara/Source/Niagara/Classes/NiagaraDataInterfaceRW.h
+++ b/Engine/Plugins/FX/Niagara/Source/Niagara/Classes/NiagaraDataInterfaceRW.h
@@ -22,7 +22,7 @@ struct FNDIGpuComputeDispatchArgsGenContext : public FNDIGpuComputeContext
 	using FIndirectArgs = TPair<FRDGBuffer*, uint32>;
 	using FCreateIndirectCallback = TFunction<void(FRHICommandList&, FRDGBuffer*, uint32)>;
 
-	explicit FNDIGpuComputeDispatchArgsGenContext(FRDGBuilder& InGraphBuilder, const FNiagaraGpuComputeDispatchInterface& InComputeDispatchInterface)
+	explicit FNDIGpuComputeDispatchArgsGenContext(FRDGBuilder& InGraphBuilder, FNiagaraGpuComputeDispatchInterface& InComputeDispatchInterface)
 		: FNDIGpuComputeContext(InGraphBuilder, InComputeDispatchInterface)
 	{
 	}
diff --git a/Engine/Plugins/FX/Niagara/Source/Niagara/Classes/NiagaraDataSet.h b/Engine/Plugins/FX/Niagara/Source/Niagara/Classes/NiagaraDataSet.h
--- a/Engine/Plugins/FX/Niagara/Source/Niagara/Classes/NiagaraDataSet.h
+++ b/Engine/Plugins/FX/Niagara/Source/Niagara/Classes/NiagaraDataSet.h
@@ -54,11 +54,20 @@ public:
 		return ReadRefCount.CompareExchange(Expected, INDEX_NONE);
 	}
 
+	//Remove the write lock
 	FORCEINLINE void Unlock()
 	{
 		int32 Expected = INDEX_NONE;
 		ensureAlwaysMsgf(ReadRefCount.CompareExchange(Expected, 0), TEXT("Trying to release a write lock on a Niagara shared object that is not locked for write."));
 	}
+	
+	//Removes the write lock and moves directly into a read state.
+	FORCEINLINE TRefCountPtr<FNiagaraSharedObject> UnlockForRead()
+	{
+		int32 Expected = INDEX_NONE;
+		ensureAlwaysMsgf(ReadRefCount.CompareExchange(Expected, 1), TEXT("Trying to release a write lock on a Niagara shared object that is not locked for write."));
+		return TRefCountPtr<FNiagaraSharedObject>(this, false);
+	}
 
 protected:
 
@@ -84,6 +93,13 @@ protected:
 	NIAGARA_API virtual ~FNiagaraDataBuffer();
 
 public:
+	FORCEINLINE TRefCountPtr<FNiagaraDataBuffer> UnlockForRead()
+	{
+		int32 Expected = INDEX_NONE;
+		ensureAlwaysMsgf(ReadRefCount.CompareExchange(Expected, 1), TEXT("Trying to release a write lock on a Niagara shared object that is not locked for write."));
+		return TRefCountPtr<FNiagaraDataBuffer>(this, false);
+	}
+
 	NIAGARA_API FNiagaraDataBuffer(FNiagaraDataSet* InOwner);
 	NIAGARA_API void Allocate(uint32 NumInstances, bool bMaintainExisting = false);
 	NIAGARA_API void ReleaseCPU();
@@ -97,7 +113,7 @@ public:
 	NIAGARA_API void CopyTo(FNiagaraDataBuffer& DestBuffer, int32 SrcStartIdx, int32 DestStartIdx, int32 NumInstances) const;
 	NIAGARA_API void CopyToUnrelated(FNiagaraDataBuffer& DestBuffer, int32 SrcStartIdx, int32 DestStartIdx, int32 NumInstances) const;
 	NIAGARA_API void GPUCopyFrom(const float* GPUReadBackFloat, const int* GPUReadBackInt, const FFloat16* GPUReadBackHalf, int32 StartIdx, int32 NumInstances, uint32 InSrcFloatStride, uint32 InSrcIntStride, uint32 InSrcHalfStride);
-	NIAGARA_API void PushCPUBuffersToGPU(const TArray<FNiagaraDataBufferRef>& SourceBuffers, bool bReleaseRef, FRHICommandList& RHICmdList, ERHIFeatureLevel::Type FeatureLevel, const TCHAR* DebugSimName);
+	NIAGARA_API void PushCPUBuffersToGPU(const TArray<FNiagaraDataBufferRef>& SourceBuffers, bool bReleaseRef, FRHICommandList& RHICmdList, ERHIFeatureLevel::Type FeatureLevel, const TCHAR* DebugSimName, bool bAllocate=true);
 	NIAGARA_API void TransferGPUToCPUImmediate(FRHICommandListImmediate& RHICmdList, FNiagaraGpuComputeDispatchInterface* ComputeInterface, FNiagaraDataBuffer* CPUBuffer) const;
 	NIAGARA_API void Dump(int32 StartIndex, int32 NumInstances, const FString& Label, const FName& SortParameterKey = FName()) const;
 
@@ -254,11 +270,14 @@ public:
 	FNiagaraDataSet& operator=(const FNiagaraDataSet&) = delete;
 
 	/** Initialize the data set with the compiled data */
-	NIAGARA_API void Init(const FNiagaraDataSetCompiledData* InDataSetCompiledData);
+	NIAGARA_API void Init(const FNiagaraDataSetCompiledData* InDataSetCompiledData, int32 DefaultNumBuffers=0);
 
 	/** Resets current data but leaves variable/layout information etc intact. */
 	NIAGARA_API void ResetBuffers();
 
+	/** Allocates a new buffer from this data set, or reuses an unused one. */
+	NIAGARA_API FNiagaraDataBuffer& AllocateBuffer();
+
 	/** Begins a new simulation pass and grabs a destination buffer. Returns the new destination data buffer. */
 	NIAGARA_API FNiagaraDataBuffer& BeginSimulate(bool bResetDestinationData = true);
 
diff --git a/Engine/Plugins/FX/Niagara/Source/Niagara/Classes/NiagaraGPUInstanceCountManager.h b/Engine/Plugins/FX/Niagara/Source/Niagara/Classes/NiagaraGPUInstanceCountManager.h
--- a/Engine/Plugins/FX/Niagara/Source/Niagara/Classes/NiagaraGPUInstanceCountManager.h
+++ b/Engine/Plugins/FX/Niagara/Source/Niagara/Classes/NiagaraGPUInstanceCountManager.h
@@ -123,6 +123,9 @@ public:
 
 	void CopyToMultiViewCountBuffer(FRHICommandListImmediate& RHICmdList);
 
+	void ProcessInitInstanceCountTasks(FNiagaraGpuComputeDispatchInterface* ComputeDispatchInterface, FRHICommandList& RHICmdList);
+	void AddInstanceCountInitTask(uint32 Offset, uint32 Value);
+
 protected:
 	struct FIndirectArgsPoolEntry
 	{
@@ -180,6 +183,14 @@ protected:
 	/** The list of all instance count clear tasks that are to be run in UpdateDrawIndirectBuffers() */
 	TArray<uint32> InstanceCountClearTasks;
 
+	struct FInstanceCountInitTask
+	{
+		uint32 Offset = INDEX_NONE;
+		uint32 Value = INDEX_NONE;
+	};
+	/** The list of initialization tasks we process in ProcessInitializationTasks. These initialize instance counts with data from the CPU. */
+	TArray<uint32> InstanceCountInitTasks;
+
 	/** Buffers holding drawindirect data to render GPU emitter renderers. */
 	TArray<FIndirectArgsPoolEntryPtr> DrawIndirectPool;
 	uint32 DrawIndirectLowWaterFrames = 0;
diff --git a/Engine/Plugins/FX/Niagara/Source/Niagara/Internal/DataInterface/NiagaraDataInterfaceDataChannelCommon.h b/Engine/Plugins/FX/Niagara/Source/Niagara/Internal/DataInterface/NiagaraDataInterfaceDataChannelCommon.h
--- a/Engine/Plugins/FX/Niagara/Source/Niagara/Internal/DataInterface/NiagaraDataInterfaceDataChannelCommon.h
+++ b/Engine/Plugins/FX/Niagara/Source/Niagara/Internal/DataInterface/NiagaraDataInterfaceDataChannelCommon.h
@@ -5,6 +5,8 @@
 #include "NiagaraDataSetCompiledData.h"
 #include "NiagaraCompileHash.h"
 #include "NiagaraDataInterface.h"
+#include "RHIUtilities.h"
+#include "Containers/DynamicRHIResourceArray.h"
 #include "NiagaraDataInterfaceDataChannelCommon.generated.h"
 
 UENUM()
@@ -13,10 +15,8 @@ enum class ENiagaraDataChannelAllocationMode : uint8
 	/** Fixed number of elements available to write per frame. */
 	Static,
 
-	/** Allow N elements per instance, per frame. Per instance is context dependent meaning data written from particles scripts will allocate per particle. From emitter script will allocate per writing emitter etc. */
-	//TODO: For this we need a Pre Stage (+on CPU) from which we can allocate and a Post Stage from which we can publish the results.
-	//PerInstance,
-	//Dynamic?
+	/** Allocation count is determined by DI script calls to Allocate in Emitter Scripts. */
+	Dynamic
 };
 
 //TODO: Possible we may want to do reads and writes using data channels in a single system in future, avoiding the need to push data out to any manager class etc.
@@ -202,12 +202,43 @@ protected:
 	void GatherAccessInfo(UNiagaraSystem* System, UNiagaraDataInterface* Owner);
 };
 
+class FNDIDummyUAV : public FRenderResource
+{
+private:
+
+	EPixelFormat PixelFmt;
+	uint32 Size;
+
+public:
+	FNDIDummyUAV(EPixelFormat Fmt, uint32 InSize) :PixelFmt(Fmt), Size(InSize) {}
+
+	FRWBuffer Buffer;
+
+	virtual void InitRHI(FRHICommandListBase& RHICmdList) override
+	{
+		Buffer.Initialize(RHICmdList, TEXT("FNDIDummyUAV"), Size, 1, PixelFmt, BUF_Static);
+	}
+
+	virtual void ReleaseRHI() override
+	{
+		Buffer.Release();
+	}
+};
 
 namespace NDIDataChannelUtilities
 {
 	extern const FName GetNDCSpawnDataName;
 
+	const TGlobalResource<FNDIDummyUAV>& GetDummyUAVFloat();
+	const TGlobalResource<FNDIDummyUAV>& GetDummyUAVInt32();
+	const TGlobalResource<FNDIDummyUAV>& GetDummyUAVHalf();
+
 	void SortParameters(TArray<FNiagaraVariableBase>& Parameters);
+
+#if WITH_EDITORONLY_DATA
+	void GenerateDataChannelAccessHlsl(FNiagaraDataInterfaceHlslGenerationContext& HlslGenContext, TConstArrayView<FString> CommonTemplateShaderCode, const TMap<FName, FString>& TemplateShaderMap, FString& OutHLSL);
+#endif
+
 }
 
 
@@ -367,4 +398,60 @@ struct FNDIVariadicOutputHandler
 			}
 		}
 	}
+};
+
+
+struct FVariadicParameterGPUScriptInfo
+{
+	/**
+	Table of all parameter offsets used by each GPU script using this DI.
+	Each script has to have it's own section of this table as the offsets into this table are embedded in the hlsl.
+	At hlsl gen time we only have the context of each script individually to generate these indexes.
+	TODO: Can possible elevate this up to the LayoutManager and have a single layout buffer for all scripts
+	*/
+	TResourceArray<uint32> GPUScriptParameterOffsetTable;
+
+	/**
+	Offsets into the parameter table are embedded in the gpu script hlsl.
+	At hlsl gen time we can only know which parameters are accessed by each script individually so each script must have it's own parameter binding table.
+	We provide the offset into the above table via a shader param.
+	TODO: Can just as easily be an offset into a global buffer in the Layout manager.
+	*/
+	TMap<FNiagaraCompileHash, uint32> GPUScriptParameterTableOffsets;
+
+	bool bDirty = false;
+
+	void Init(const FNDIDataChannelCompiledData& DICompiledData, const FNiagaraDataSetCompiledData& GPUDataSetCompiledData)
+	{
+		bDirty = true;
+
+		//For every GPU script, we append it's parameter access info to the table.
+		GPUScriptParameterTableOffsets.Reset();
+		constexpr int32 ElemsPerParam = 3;
+		GPUScriptParameterOffsetTable.Reset(DICompiledData.GetTotalParams() * ElemsPerParam);
+		for (auto& GPUParameterAccessInfoPair : DICompiledData.GetGPUScriptParameterInfos())
+		{
+			const FNDIDataChannel_GPUScriptParameterAccessInfo& ParamAccessInfo = GPUParameterAccessInfoPair.Value;
+
+			//First get the offset for this script in the table.
+			GPUScriptParameterTableOffsets.FindOrAdd(GPUParameterAccessInfoPair.Key) = GPUScriptParameterOffsetTable.Num();
+
+			//Now fill the table for this script
+			for (const FNiagaraVariableBase& Param : ParamAccessInfo.SortedParameters)
+			{
+				if (const FNiagaraVariableLayoutInfo* LayoutInfo = GPUDataSetCompiledData.FindVariableLayoutInfo(Param))
+				{
+					GPUScriptParameterOffsetTable.Add(LayoutInfo->GetNumFloatComponents() > 0 ? LayoutInfo->GetFloatComponentStart() : INDEX_NONE);
+					GPUScriptParameterOffsetTable.Add(LayoutInfo->GetNumInt32Components() > 0 ? LayoutInfo->GetInt32ComponentStart() : INDEX_NONE);
+					GPUScriptParameterOffsetTable.Add(LayoutInfo->GetNumHalfComponents() > 0 ? LayoutInfo->GetHalfComponentStart() : INDEX_NONE);
+				}
+				else
+				{
+					GPUScriptParameterOffsetTable.Add(INDEX_NONE);
+					GPUScriptParameterOffsetTable.Add(INDEX_NONE);
+					GPUScriptParameterOffsetTable.Add(INDEX_NONE);
+				}
+			}
+		}
+	}
 };
\ No newline at end of file
diff --git a/Engine/Plugins/FX/Niagara/Source/Niagara/Internal/DataInterface/NiagaraDataInterfaceDataChannelRead.h b/Engine/Plugins/FX/Niagara/Source/Niagara/Internal/DataInterface/NiagaraDataInterfaceDataChannelRead.h
--- a/Engine/Plugins/FX/Niagara/Source/Niagara/Internal/DataInterface/NiagaraDataInterfaceDataChannelRead.h
+++ b/Engine/Plugins/FX/Niagara/Source/Niagara/Internal/DataInterface/NiagaraDataInterfaceDataChannelRead.h
@@ -26,7 +26,6 @@ There is working code for this already but the concept/API needs more fleshing o
 #include "NiagaraDataChannelPublic.h"
 #include "NiagaraDataInterfaceDataChannelCommon.h"
 #include "NiagaraDataInterfaceRW.h"
-#include "NiagaraDataSetAccessor.h"
 #include "NiagaraDataInterfaceDataChannelRead.generated.h"
 
 class UNiagaraDataInterfaceDataChannelWrite;
@@ -162,6 +161,9 @@ public:
 
 	//We cannot overlap frames as we must correctly sync up with the data channel manager on Begin/End frame etc.
 	virtual bool PostSimulateCanOverlapFrames() const { return false; }
+	virtual bool PostStageCanOverlapTickGroups() const { return false; }
+
+	virtual bool RequiresCurrentFrameNDC() const { return bReadCurrentFrame; }
 	//UNiagaraDataInterface Interface
 
 	//Functions usable anywhere.
@@ -276,6 +278,9 @@ struct FNDIDataChannelReadInstanceData
 	*/
 	std::atomic<int32> ConsumeIndex = 0;
 
+	// Num Elements in the NDC at the time we generated our NDCSpawnData.
+	int32 NDCElementCountAtSpawn = 0;
+
 	/** 
 	Instance data for each emitter using this DI.
 	*/
@@ -285,7 +290,7 @@ struct FNDIDataChannelReadInstanceData
 	FNiagaraSystemInstance* Owner = nullptr;
 
 	virtual ~FNDIDataChannelReadInstanceData();
-	FNiagaraDataBuffer* GetReadBufferCPU(bool bPrevFrame);
+	FNiagaraDataBuffer* GetReadBufferCPU(bool bPrevFrame)const;
 	bool Init(UNiagaraDataInterfaceDataChannelRead* Interface, FNiagaraSystemInstance* Instance);
 	bool Tick(UNiagaraDataInterfaceDataChannelRead* Interface, FNiagaraSystemInstance* Instance, bool bIsInit = false);	
 	bool PostTick(UNiagaraDataInterfaceDataChannelRead* Interface, FNiagaraSystemInstance* Instance);
@@ -309,6 +314,8 @@ struct FNiagaraDataInterfaceProxy_DataChannelRead : public FNiagaraDataInterface
 
 		bool bReadPrevFrame = false;
 
+		int32 NDCElementCountAtSpawn = 0;
+
 		/**
 		A buffer containing layout information needed to access parameters for each script using this DI.
 		*/
@@ -326,6 +333,8 @@ struct FNiagaraDataInterfaceProxy_DataChannelRead : public FNiagaraDataInterface
 		TArray<int32> NDCSpawnData;
 
 		FRDGBufferRef NDCSpawnDataBuffer;
+
+		uint32 ConsumeInstanceCountOffset = INDEX_NONE;
 	};
 
 	TMap<FNiagaraSystemInstanceID, FInstanceData> SystemInstancesToProxyData_RT;
diff --git a/Engine/Plugins/FX/Niagara/Source/Niagara/Internal/DataInterface/NiagaraDataInterfaceDataChannelWrite.h b/Engine/Plugins/FX/Niagara/Source/Niagara/Internal/DataInterface/NiagaraDataInterfaceDataChannelWrite.h
--- a/Engine/Plugins/FX/Niagara/Source/Niagara/Internal/DataInterface/NiagaraDataInterfaceDataChannelWrite.h
+++ b/Engine/Plugins/FX/Niagara/Source/Niagara/Internal/DataInterface/NiagaraDataInterfaceDataChannelWrite.h
@@ -21,6 +21,7 @@ This avoids cumbersome work in the graph to access data but requires special han
 #include "NiagaraDataChannelCommon.h"
 #include "NiagaraDataInterface.h"
 #include "NiagaraSimCacheCustomStorageInterface.h"
+#include "NiagaraDataInterfaceRW.h"
 #include "NiagaraDataInterfaceDataChannelWrite.generated.h"
 
 /** Additional compile time information used by the Write DI. */
@@ -99,7 +100,7 @@ public:
 	ENiagaraDataChannelAllocationMode AllocationMode = ENiagaraDataChannelAllocationMode::Static;
 
 	/** How many elements to allocate for writing per frame? Usage is defendant on AllocationMode. TODO: Allow allocation count to be controlled dynamically from script? */
-	UPROPERTY(EditAnywhere, Category = "Data Channel")
+	UPROPERTY(EditAnywhere, Category = "Data Channel", meta = (EditCondition="AllocationMode == ENiagaraDataChannelAllocationMode::Static"))
 	uint32 AllocationCount = 0;
 
 	/** Whether the data generated by this DI should be published to the world game data channel. This is require to allow game BP and C++ to read this data. */
@@ -156,12 +157,10 @@ public:
 	NIAGARA_API virtual void DestroyPerInstanceData(void* PerInstanceData, FNiagaraSystemInstance* SystemInstance) override;
 	NIAGARA_API virtual int32 PerInstanceDataSize() const override;
 	virtual bool HasPreSimulateTick() const override { return true; }
-	virtual bool HasPreStageTick(ENiagaraScriptUsage Usage) const override { return true; }
-	virtual bool HasPostStageTick(ENiagaraScriptUsage Usage) const override { return true; }
+	virtual bool HasPostSimulateTick() const override { return true; }
+
 	NIAGARA_API virtual bool PerInstanceTick(void* PerInstanceData, FNiagaraSystemInstance* SystemInstance, float DeltaSeconds) override;
 	NIAGARA_API virtual bool PerInstanceTickPostSimulate(void* PerInstanceData, FNiagaraSystemInstance* SystemInstance, float DeltaSeconds) override;
-	NIAGARA_API virtual void PreStageTick(FNDICpuPreStageContext& Context) override;
-	NIAGARA_API virtual void PostStageTick(FNDICpuPostStageContext& Context) override;
 	NIAGARA_API virtual void ProvidePerInstanceDataForRenderThread(void* DataForRenderThread, void* PerInstanceData, const FNiagaraSystemInstanceID& SystemInstance) override;
 
 	virtual bool HasTickGroupPostreqs() const override;
@@ -174,6 +173,7 @@ public:
 	//UNiagaraDataInterface Interface
 
 	NIAGARA_API void Num(FVectorVMExternalFunctionContext& Context);
+	NIAGARA_API void Allocate(FVectorVMExternalFunctionContext& Context);
 	NIAGARA_API void Write(FVectorVMExternalFunctionContext& Context, int32 FuncIdx);
 	NIAGARA_API void Append(FVectorVMExternalFunctionContext& Context, int32 FuncIdx);
 
@@ -196,3 +196,51 @@ protected:
 	UPROPERTY()
 	FNDIDataChannelWriteCompiledData CompiledData;
 };
+
+struct FNiagaraDataInterfaceProxy_DataChannelWrite : public FNiagaraDataInterfaceProxyRW
+{
+	FNiagaraDataInterfaceProxy_DataChannelWrite()
+	: FNiagaraDataInterfaceProxyRW()
+	{
+	}
+
+	virtual void ConsumePerInstanceDataFromGameThread(void* PerInstanceData, const FNiagaraSystemInstanceID& Instance) override;
+	virtual int32 PerInstanceDataPassedToRenderThreadSize() const override;
+
+	virtual void PreStage(const FNDIGpuComputePreStageContext& Context)override;
+	virtual void PostStage(const FNDIGpuComputePostStageContext& Context)override;
+
+	/** Persistent per instance data on the RT. Constructed when consuming data passed from GT->RT. */
+	struct FInstanceData
+	{
+		//GPU Dataset from the channel handler. We'll grab the current buffer from this on the RT.
+		//This must be grabbed fresh from the handler each frame as it's lifetime cannot be ensured.
+		FNiagaraDataChannelDataProxy* ChannelDataRTProxy = nullptr;
+
+		/**
+		A buffer containing layout information needed to access parameters for each script using this DI.
+		*/
+		FReadBuffer ParameterLayoutBuffer;
+
+		TResourceArray<uint32> ParameterLayoutData;
+
+		/**
+		Offsets into the parameter table are embedded in the gpu script hlsl.
+		At hlsl gen time we can only know which parameter are accessed by each script individually so each script must have it's own parameter binding table.
+		*/
+		TMap<FNiagaraCompileHash, uint32> GPUScriptParameterTableOffsets;
+
+		//This is the buffer we'll write into for shipping data back to the CPU & Game.
+		FNiagaraDataBufferRef BufferForCPU;
+
+		bool bPublishToGame = false;
+		bool bPublishToCPU = false;
+		bool bPublishToGPU = false;
+
+
+		int32 AllocationCount = 0;
+		FVector3f LwcTile;
+	};
+
+	TMap<FNiagaraSystemInstanceID, FInstanceData> SystemInstancesToProxyData_RT;
+};
diff --git a/Engine/Plugins/FX/Niagara/Source/Niagara/Internal/NiagaraDataChannelCommon.h b/Engine/Plugins/FX/Niagara/Source/Niagara/Internal/NiagaraDataChannelCommon.h
--- a/Engine/Plugins/FX/Niagara/Source/Niagara/Internal/NiagaraDataChannelCommon.h
+++ b/Engine/Plugins/FX/Niagara/Source/Niagara/Internal/NiagaraDataChannelCommon.h
@@ -48,12 +48,12 @@ struct FNiagaraDataChannelPublishRequest
 #endif
 
 	FNiagaraDataChannelPublishRequest() = default;
-	explicit FNiagaraDataChannelPublishRequest(FNiagaraDataBuffer* InData)
+	explicit FNiagaraDataChannelPublishRequest(FNiagaraDataBufferRef InData)
 		: Data(InData)
 	{
 	}
 
-	explicit FNiagaraDataChannelPublishRequest(FNiagaraDataBuffer* InData, bool bInVisibleToGame, bool bInVisibleToCPUSims, bool bInVisibleToGPUSims, FVector3f InLwcTile)
+	explicit FNiagaraDataChannelPublishRequest(FNiagaraDataBufferRef InData, bool bInVisibleToGame, bool bInVisibleToCPUSims, bool bInVisibleToGPUSims, FVector3f InLwcTile)
 	: Data(InData), bVisibleToGame(bInVisibleToGame), bVisibleToCPUSims(bInVisibleToCPUSims), bVisibleToGPUSims(bInVisibleToGPUSims), LwcTile(InLwcTile)
 	{
 	}
@@ -66,12 +66,13 @@ Some data channels will have many of these and can distribute them as needed to
 For example, some data channel handlers may subdivide the scene such that distant systems are not interacting.
 In this case, each subdivision would have it's own FNiagaraDataChannelData and distribute these to the relevant NiagaraSystems.
 */
-struct FNiagaraDataChannelData final
+struct FNiagaraDataChannelData final : public TSharedFromThis<FNiagaraDataChannelData, ESPMode::ThreadSafe>
 {
 	UE_NONCOPYABLE(FNiagaraDataChannelData)
-	NIAGARA_API explicit FNiagaraDataChannelData(UNiagaraDataChannelHandler* Owner);
+	NIAGARA_API explicit FNiagaraDataChannelData();
 	NIAGARA_API ~FNiagaraDataChannelData();
 
+	NIAGARA_API void Init(UNiagaraDataChannelHandler* Owner);
 	NIAGARA_API void Reset();
 
 	NIAGARA_API void BeginFrame(UNiagaraDataChannelHandler* Owner);
@@ -84,6 +85,8 @@ struct FNiagaraDataChannelData final
 	
 	/** Adds a request to publish some data into the channel on the next tick. */
 	NIAGARA_API void Publish(const FNiagaraDataChannelPublishRequest& Request);
+
+	NIAGARA_API void PublishFromGPU(const FNiagaraDataChannelPublishRequest& Request);
 		
 	/**
 	 *Removes all publish requests involving the given dataset.
@@ -97,6 +100,9 @@ struct FNiagaraDataChannelData final
 
 	void SetLwcTile(FVector3f InLwcTile){ LwcTile = InLwcTile; }
 	FVector3f GetLwcTile()const { return LwcTile; }
+
+	//This will get a buffer from the CPU dataset intended to be written to on the CPU.
+	FNiagaraDataBuffer* GetBufferForCPUWrite();
 private:
 
 	/** DataChannel data accessible from Game/BP. AoS Layout. LWC types. */
@@ -137,6 +143,9 @@ private:
 	/** Pending requests to publish data into this data channel. These requests are consumed at tick tick group. */
 	TArray<FNiagaraDataChannelPublishRequest> PublishRequests;
 
+	/** Pending requests to publish data into this data channel from the GPU. To alleviate data race behavior with data coming back from the GPU, we always consume GPU requests at the start of the frame only. */
+	TArray<FNiagaraDataChannelPublishRequest> PublishRequestsFromGPU;
+
 	FVector3f LwcTile = FVector3f::ZeroVector;
 
 	/** Critical section protecting shared state for multiple writers publishing from different threads. */
diff --git a/Engine/Plugins/FX/Niagara/Source/Niagara/Private/DataInterface/NiagaraDataInterfaceDataChannelCommon.cpp b/Engine/Plugins/FX/Niagara/Source/Niagara/Private/DataInterface/NiagaraDataInterfaceDataChannelCommon.cpp
--- a/Engine/Plugins/FX/Niagara/Source/Niagara/Private/DataInterface/NiagaraDataInterfaceDataChannelCommon.cpp
+++ b/Engine/Plugins/FX/Niagara/Source/Niagara/Private/DataInterface/NiagaraDataInterfaceDataChannelCommon.cpp
@@ -48,11 +48,11 @@ FNDIDataChannel_FunctionToDataSetBinding::FNDIDataChannel_FunctionToDataSetBindi
 			{
 				DataSetLayoutHash = 0;
 
-				#if !UE_BUILD_SHIPPING && !UE_BUILD_TEST
+#if !UE_BUILD_SHIPPING && !UE_BUILD_TEST
 				OutMissingParams.Emplace(Param);
-				#else
+#else
 				return;
-				#endif
+#endif
 			}
 		}
 	};
@@ -288,6 +288,24 @@ int32 FNDIDataChannelCompiledData::FindFunctionInfoIndex(FName Name, const TArra
 namespace NDIDataChannelUtilities
 {
 	const FName GetNDCSpawnDataName(TEXT("GetNDCSpawnData"));
+	const TGlobalResource<FNDIDummyUAV> DummyUAVFloat(PF_R32_FLOAT, sizeof(float));
+	const TGlobalResource<FNDIDummyUAV> DummyUAVInt32(PF_R32_SINT, sizeof(int32));
+	const TGlobalResource<FNDIDummyUAV> DummyUAVHalf(PF_R16F, sizeof(FFloat16));
+
+	const TGlobalResource<FNDIDummyUAV>& GetDummyUAVFloat()
+	{
+		return DummyUAVFloat;
+	}
+
+	const TGlobalResource<FNDIDummyUAV>& GetDummyUAVInt32()
+	{
+		return DummyUAVInt32;
+	}
+
+	const TGlobalResource<FNDIDummyUAV>& GetDummyUAVHalf()
+	{
+		return DummyUAVHalf;
+	}
 
 	void SortParameters(TArray<FNiagaraVariableBase>& Parameters)
 	{
@@ -301,6 +319,397 @@ namespace NDIDataChannelUtilities
 				return ComparisonDiff < 0;
 			});
 	}
+
+#if WITH_EDITORONLY_DATA
+
+	void GenerateDataChannelAccessHlsl(FNiagaraDataInterfaceHlslGenerationContext& HlslGenContext, TConstArrayView<FString> CommonTemplateShaderCode, const TMap<FName, FString>& FunctionTemplateMap, FString& OutHLSL)
+	{
+		//Desc:
+		// This function iterates over all functions called for this DI in each script and builds the correct hlsl.
+		// The main part of the complexity here is dealing with Variadic function parameters.
+		// We must interrogate the script to see what parameters are actually accessed and generate hlsl accordingly.
+		// Ideally at some future point we can refactor most of this out to a utility helper that will do most of the heavy lifting.
+		// Allowing users to simply provide some details and lambdas etc to specify what exactly they'd like to do in the function body with each variadic param.
+
+		//////////////////////////////////////////////////////////////////////////
+		//Initially we'll do some preamble, setting up various template strings and args etc.
+
+		//Map of all arguments for various pieces of template code.
+		//We have some common code that is shared by all functions.
+		//Some code is unique for each function.
+		//Some is unique per parameter to each function.
+		//Finally there is some that is unique for each sub component of each parameter until we've hit a base type. float/2/3/4 etc.
+		//We add to the map and take refs for  params we may want to change as we iterate over parameters etc.
+		
+		TMap<FString, FStringFormatArg> HlslTemplateArgs =
+		{
+			//Common args for all functions
+			{TEXT("ParameterName"),							HlslGenContext.ParameterInfo.DataInterfaceHLSLSymbol},
+
+			//Per function args. These will be changed with each function written
+			{TEXT("FunctionSymbol"),						FString(TEXT("FunctionSymbol"))},//Function symbol which will be a mangled form from the translator.
+
+			{TEXT("FunctionInputParameters"),				FString(TEXT("FunctionInputParameters"))},//Function input parameters written into the function signature.
+			{TEXT("FunctionOutputParameters"),				FString(TEXT("FunctionOutputParameters"))},//Function output parameters written into the function signature.
+
+			{TEXT("DefaultOutputsShaderCode"),				FString(TEXT("DefaultOutputsShaderCode"))},
+			{TEXT("PerParameterFunctionDefinitions"),		FString(TEXT("PerParameterFunctionDefinitions"))},//Per parameter helper/access functions
+			{TEXT("PerFunctionParameterShaderCode"),		FString(TEXT("PerFunctionParameterShaderCode"))},//Per paraemter shader code embeddded in the DI function body
+
+			//Per function parameter args. These will be changed with each parameter written
+			{TEXT("FunctionParameterIndex"),				TEXT("FunctionParameterIndex")},//Function parameter index allowing us to look up the layout information for the correct parameter to the function.
+			{TEXT("FunctionParameterName"),					FString(TEXT("FunctionParameterName"))},//Name of this function parameter.
+			{TEXT("FunctionParameterType"),					FString(TEXT("FunctionParameterType"))},//Type of this function parameter.
+			{TEXT("FuncParamShaderCode"),						FString(TEXT("PerParamRWCode"))},//Code that does the actual reading or writing to the data channel buffers.
+			//Per component/base type args. These will change with every base type we I/O from the Data Channel.
+			{TEXT("FunctionParameterComponentBufferType"),	FString(TEXT("FunctionParameterComponentBufferType"))},//The actual base data buffer type being accessed by a particular DataChannel access code line. Float, Int32, Half etc
+			{TEXT("FunctionParameterComponentName"),		FString(TEXT("FunctionParameterComponentName"))},//The name/symbol of the actual member of a parameter that we can I/O from the DataChannel via a standard getter/setter.
+			{TEXT("FunctionParameterComponentType"),		FString(TEXT("FunctionParameterComponentType"))},//The type of the actual member of a parameter that we can I/O from the DataChannel via a standard getter/setter.
+			{TEXT("FunctionParameterComponentDefault"),		FString(TEXT("FunctionParameterComponentDefault"))},
+		};
+
+		//Grab refs to per function args we'll change with each function written.
+		FStringFormatArg& FunctionSymbol = HlslTemplateArgs.FindChecked(TEXT("FunctionSymbol"));
+		
+		FStringFormatArg& FunctionInputParameters = HlslTemplateArgs.FindChecked(TEXT("FunctionInputParameters"));
+		FStringFormatArg& FunctionOutputParameters = HlslTemplateArgs.FindChecked(TEXT("FunctionOutputParameters"));
+
+		FStringFormatArg& DefaultOutputsShaderCode = HlslTemplateArgs.FindChecked(TEXT("DefaultOutputsShaderCode"));
+		FStringFormatArg& PerParameterFunctionDefinitions = HlslTemplateArgs.FindChecked(TEXT("PerParameterFunctionDefinitions"));
+		FStringFormatArg& PerFunctionParameterShaderCode = HlslTemplateArgs.FindChecked(TEXT("PerFunctionParameterShaderCode"));
+
+		//Grab refs to per function args we'll change with each parameter written.
+		FStringFormatArg& FunctionParameterIndex = HlslTemplateArgs.FindChecked(TEXT("FunctionParameterIndex"));
+		FStringFormatArg& FunctionParameterName = HlslTemplateArgs.FindChecked(TEXT("FunctionParameterName"));
+		FStringFormatArg& FunctionParameterType = HlslTemplateArgs.FindChecked(TEXT("FunctionParameterType"));
+		FStringFormatArg& FuncParamShaderCode = HlslTemplateArgs.FindChecked(TEXT("FuncParamShaderCode"));
+
+		//Grab refs to per component/base type args we'll change for every parameter component we access in the Data Channel.
+		FStringFormatArg& FunctionParameterComponentBufferType = HlslTemplateArgs.FindChecked(TEXT("FunctionParameterComponentBufferType"));
+		FStringFormatArg& FunctionParameterComponentName = HlslTemplateArgs.FindChecked(TEXT("FunctionParameterComponentName"));
+		FStringFormatArg& FunctionParameterComponentType = HlslTemplateArgs.FindChecked(TEXT("FunctionParameterComponentType"));
+		FStringFormatArg& FunctionParameterComponentDefault = HlslTemplateArgs.FindChecked(TEXT("FunctionParameterComponentDefault"));
+
+
+		//////////////////////////////////////////////////////////////////////////
+		//Next we'll define some template code blocks and helper lambdas that we'll use with the above args map to build various pieces of shader code.
+		
+		//Template code for writing handling code for each of the function parameters.
+		//Some preamble and an replacement arg into which we write all the actual I/O with the DataChannel buffers.
+		FString PerParameterReadTemplate = TEXT("\n\
+void Read_{FunctionParameterName}_{ParameterName}(FNDCAccessContext_{ParameterName} Context, out {FunctionParameterType} {FunctionParameterName})\n\
+{\n\
+	Context.InitForParameter({FunctionParameterIndex});\n\
+	if(Context.IsValid())\n\
+	{\n\
+		{FuncParamShaderCode}\
+	}\n\
+	else\n\
+	{\n\
+		{FunctionParameterName} = 0;\n\
+	}\n\
+}\n");
+
+		FString PerParameterReadCallTemplate = TEXT("Read_{FunctionParameterName}_{ParameterName}(Context, {FunctionParameterName});\n");
+
+		FString PerParameterWriteTemplate = TEXT("\n\
+void Write_{FunctionParameterName}_{ParameterName}(FNDCAccessContext_{ParameterName} Context, {FunctionParameterType} {FunctionParameterName})\n\
+{\n\
+	Context.InitForParameter({FunctionParameterIndex});\n\
+	if(Context.IsValid())\n\
+	{\n\
+		if(Context.InitForGPUWrite())\n\
+		{\n\
+			{FuncParamShaderCode}\n\
+		}\n\
+		if (Context.InitForCPUWrite())\n\
+		{\n\
+			{FuncParamShaderCode}\n\
+		}\n\
+	}\n\
+}\n");
+		FString PerParameterWriteCallTemplate = TEXT("Write_{FunctionParameterName}_{ParameterName}(Context, {FunctionParameterName});\n");
+
+		//Template code for accessing data from the Data Channel's buffers.
+		static const FString ReadDataTemplate = TEXT("Context.Read_{FunctionParameterComponentBufferType}({FunctionParameterComponentName});\n");
+		static const FString WriteDataTemplate = TEXT("Context.Write_{FunctionParameterComponentBufferType}({FunctionParameterComponentName});\n");
+		static const FString DefaultsDataTemplate = TEXT("{FunctionParameterComponentName} = {FunctionParameterComponentDefault};\n");
+
+		//Function that will recurse down a parameter's type and generate the appropriate IO code for all of it's members.
+		TFunction<void(bool, UScriptStruct*, FString&, FString&)> GeneratePerParamShaderCode = [&](bool bRead, UScriptStruct* Struct, FString& OutCode, FString& OutDefaultsCode)
+		{
+			//Intercept positions and replace with FVector3fs
+			if (Struct == FNiagaraPosition::StaticStruct())
+			{
+				Struct = FNiagaraTypeDefinition::GetVec3Struct();
+			}
+
+			if (Struct == FNiagaraTypeDefinition::GetFloatStruct() || Struct == FNiagaraTypeDefinition::GetVec2Struct() || Struct == FNiagaraTypeDefinition::GetVec3Struct()
+				|| Struct == FNiagaraTypeDefinition::GetVec4Struct() || Struct == FNiagaraTypeDefinition::GetColorStruct() || Struct == FNiagaraTypeDefinition::GetQuatStruct())
+			{
+				FunctionParameterComponentBufferType = TEXT("Float");
+				FNiagaraTypeDefinition Type = FNiagaraTypeDefinition(Struct);
+
+				FunctionParameterComponentType = HlslGenContext.GetStructHlslTypeName(Type);
+				OutCode += FString::Format(bRead ? *ReadDataTemplate : *WriteDataTemplate, HlslTemplateArgs);
+
+				FunctionParameterComponentDefault = HlslGenContext.GetHlslDefaultForType(Type);
+				OutDefaultsCode += FString::Format(*DefaultsDataTemplate, HlslTemplateArgs);
+			}
+			else if (Struct == FNiagaraTypeDefinition::GetIntStruct())
+			{
+				FunctionParameterComponentBufferType = TEXT("Int32");
+				FNiagaraTypeDefinition Type = FNiagaraTypeDefinition(Struct);
+
+				FunctionParameterComponentType = HlslGenContext.GetStructHlslTypeName(Type);
+				OutCode += FString::Format(bRead ? *ReadDataTemplate : *WriteDataTemplate, HlslTemplateArgs);
+
+				FunctionParameterComponentDefault = HlslGenContext.GetHlslDefaultForType(Type);
+				OutDefaultsCode += FString::Format(*DefaultsDataTemplate, HlslTemplateArgs);
+			}
+			else if (Struct == FNiagaraTypeDefinition::GetBoolStruct())
+			{
+				FunctionParameterComponentBufferType = TEXT("Int32");
+				FunctionParameterComponentType = HlslGenContext.GetStructHlslTypeName(FNiagaraTypeDefinition::GetIntDef());
+				OutCode += FString::Format(bRead ? *ReadDataTemplate : *WriteDataTemplate, HlslTemplateArgs);
+
+				FunctionParameterComponentDefault = HlslGenContext.GetHlslDefaultForType(FNiagaraTypeDefinition::GetIntDef());
+				OutDefaultsCode += FString::Format(*DefaultsDataTemplate, HlslTemplateArgs);
+			}
+			else
+			{
+				FString PropertyBaseName = FunctionParameterComponentName.StringValue;
+
+				for (TFieldIterator<FProperty> PropertyIt(Struct, EFieldIteratorFlags::IncludeSuper); PropertyIt; ++PropertyIt)
+				{
+					FunctionParameterComponentName.StringValue = PropertyBaseName;
+					const FProperty* Property = *PropertyIt;
+
+					if (Property->IsA(FFloatProperty::StaticClass()))
+					{
+						FunctionParameterComponentBufferType = TEXT("Float");
+						FunctionParameterComponentType = HlslGenContext.GetStructHlslTypeName(FNiagaraTypeDefinition::GetFloatDef());
+						FunctionParameterComponentName.StringValue += FString::Printf(TEXT(".%s"), *HlslGenContext.GetSanitizedSymbolName(Property->GetName()));
+						OutCode += FString::Format(bRead ? *ReadDataTemplate : *WriteDataTemplate, HlslTemplateArgs);
+
+						FunctionParameterComponentDefault = HlslGenContext.GetHlslDefaultForType(FNiagaraTypeDefinition::GetFloatDef());
+						OutDefaultsCode += FString::Format(*DefaultsDataTemplate, HlslTemplateArgs);
+					}
+					else if (Property->IsA(FIntProperty::StaticClass()))
+					{
+						FunctionParameterComponentBufferType = TEXT("Int32");
+						FunctionParameterComponentType = HlslGenContext.GetStructHlslTypeName(FNiagaraTypeDefinition::GetIntDef());
+						FunctionParameterComponentName.StringValue += FString::Printf(TEXT(".%s"), *HlslGenContext.GetSanitizedSymbolName(Property->GetName()));
+						OutCode += FString::Format(bRead ? *ReadDataTemplate : *WriteDataTemplate, HlslTemplateArgs);
+
+						FunctionParameterComponentDefault = HlslGenContext.GetHlslDefaultForType(FNiagaraTypeDefinition::GetIntDef());
+						OutDefaultsCode += FString::Format(*DefaultsDataTemplate, HlslTemplateArgs);
+					}
+					else if (const FStructProperty* StructProp = CastFieldChecked<const FStructProperty>(Property))
+					{
+						FunctionParameterComponentName.StringValue += FString::Printf(TEXT(".%s"), *HlslGenContext.GetSanitizedSymbolName(Property->GetName()));
+						FunctionParameterComponentType = HlslGenContext.GetPropertyHlslTypeName(Property);
+						GeneratePerParamShaderCode(bRead, StructProp->Struct, OutCode, OutDefaultsCode);
+					}
+					else
+					{
+						check(false);
+						OutCode += FString::Printf(TEXT("Error! - DataChannel Interface encountered an invalid type when generating it's function hlsl. %s"), *Property->GetName());
+					}
+				}
+			}
+		};
+
+		//////////////////////////////////////////////////////////////////////////
+
+		//We may call the same function multiple times so avoid duplicating the same function impl.
+		TMap<uint32, FString> FunctionHlslMap;
+
+		//Now lets build some hlsl!
+		// 
+		//First build the shader code common to all functions.
+		for(const FString& TemplateShader : CommonTemplateShaderCode)
+		{
+			OutHLSL += FString::Format(*TemplateShader, HlslTemplateArgs);
+		}
+
+		auto GetSignatureHash = [](const FNiagaraFunctionSignature& Sig)
+		{
+			uint32 Ret = GetTypeHash(Sig.Name);
+			for (const FNiagaraVariable& Input : Sig.Inputs)
+			{
+				Ret = HashCombine(Ret, GetTypeHash(Input));
+			}
+			for (const FNiagaraVariable& Output : Sig.Outputs)
+			{
+				Ret = HashCombine(Ret, GetTypeHash(Output));
+			}
+			return Ret;
+		};
+
+		struct FParamAccessInfo
+		{			
+			bool bRead;
+			bool bWrite;
+			int32 SortedOffset;
+		};
+
+		TMap<FNiagaraVariableBase, FParamAccessInfo> ParametersAccessed;
+		TArray<FNiagaraVariableBase> SortedParameters;
+
+		//First iterate over the generated functions to gather all used parameters so we can generate the sorted parameter list for all functions called by this script.
+		for (int32 FuncIdx = 0; FuncIdx < HlslGenContext.ParameterInfo.GeneratedFunctions.Num(); ++FuncIdx)
+		{
+			const FNiagaraDataInterfaceGeneratedFunction& Func = HlslGenContext.ParameterInfo.GeneratedFunctions[FuncIdx];
+			const FNiagaraFunctionSignature& Signature = HlslGenContext.Signatures[FuncIdx];
+
+			if(Signature.VariadicInput())
+			{
+				int32 StartInput = Signature.VariadicInputStartIndex();
+				for (int32 InputIdx = StartInput; InputIdx < Signature.Inputs.Num(); ++InputIdx)
+				{
+					const FNiagaraVariable& InputParam = Signature.Inputs[InputIdx];
+					FParamAccessInfo& Access = ParametersAccessed.FindOrAdd(InputParam);
+					Access.bWrite = true;
+				}
+			}			
+
+			if(Signature.VariadicOutput())
+			{
+				int32 StartOutput = Signature.VariadicOutputStartIndex();
+				for (int32 OutputIdx = StartOutput; OutputIdx < Signature.Outputs.Num(); ++OutputIdx)
+				{
+					const FNiagaraVariable& OutputParam = Signature.Outputs[OutputIdx];
+					FParamAccessInfo& Access = ParametersAccessed.FindOrAdd(OutputParam);
+					Access.bRead = true;
+				}
+			}			
+		}
+		
+		//Sort the parameters so that that generated hlsl can exactly match the runtime ordering of parameters.
+		ParametersAccessed.GenerateKeyArray(SortedParameters);
+		NDIDataChannelUtilities::SortParameters(SortedParameters);
+		
+		DefaultOutputsShaderCode.StringValue.Reset();
+
+		PerParameterFunctionDefinitions.StringValue.Reset();
+		FunctionInputParameters.StringValue.Reset();
+		FunctionOutputParameters.StringValue.Reset();
+
+		//Next we iterate over each of the accessed parameters and generate specific read or write functions for that parameter.
+		//This is placed in the template arg {PerParameterAccessFunctions}
+		//Each of these functions performing any reads or writes needed for each parameter.
+		for (int32 ParamIdx = 0; ParamIdx < SortedParameters.Num(); ++ParamIdx)
+		{
+			const FNiagaraVariableBase& Param = SortedParameters[ParamIdx];
+			FParamAccessInfo& Access = 	ParametersAccessed.FindChecked(Param);
+			Access.SortedOffset = ParamIdx;
+
+			FunctionParameterIndex = LexToString(ParamIdx);
+			FunctionParameterName = Param.GetName().ToString();
+			FunctionParameterType = HlslGenContext.GetStructHlslTypeName(Param.GetType());
+
+			FunctionParameterComponentName = HlslGenContext.GetSanitizedSymbolName(Param.GetName().ToString());
+
+			//Generate read function for this parameter if needed.
+			if(Access.bRead)
+			{
+				FuncParamShaderCode.StringValue.Reset();
+				FunctionParameterComponentType.StringValue.Reset();
+				FunctionParameterComponentBufferType.StringValue.Reset();
+				//Generate the per component/base type access code that will be used in the subsequent per parameters shader code template.
+				GeneratePerParamShaderCode(true, Param.GetType().GetScriptStruct(), FuncParamShaderCode.StringValue, DefaultOutputsShaderCode.StringValue);
+				PerParameterFunctionDefinitions.StringValue += FString::Format(*PerParameterReadTemplate, HlslTemplateArgs);
+			}
+
+			//Generate write function for this parameter if needed.
+			if (Access.bWrite)
+			{
+				FuncParamShaderCode.StringValue.Reset();
+				FunctionParameterComponentType.StringValue.Reset();
+				FunctionParameterComponentBufferType.StringValue.Reset();
+				//Generate the per component/base type access code that will be used in the subsequent per parameters shader code template.
+				GeneratePerParamShaderCode(false, Param.GetType().GetScriptStruct(), FuncParamShaderCode.StringValue, DefaultOutputsShaderCode.StringValue);
+				PerParameterFunctionDefinitions.StringValue += FString::Format(*PerParameterWriteTemplate, HlslTemplateArgs);
+			}
+		}
+
+		//Now to iterate on the functions and build the hlsl for each as needed.
+		static FString FirstParamPrefix = TEXT("");
+		static FString ParamPrefix = TEXT(", ");
+		static FString OutputParamPrefix = TEXT("out");
+		static FString InputParamPrefix = TEXT("in");
+		uint32 CurrParamIdx = 0;
+		for (int32 FuncIdx = 0; FuncIdx < HlslGenContext.ParameterInfo.GeneratedFunctions.Num(); ++FuncIdx)
+		{
+			const FNiagaraDataInterfaceGeneratedFunction& Func = HlslGenContext.ParameterInfo.GeneratedFunctions[FuncIdx];
+			const FNiagaraFunctionSignature& Signature = HlslGenContext.Signatures[FuncIdx];
+
+			//Init/Reset our per function hlsl template args.
+			FunctionSymbol.StringValue = HlslGenContext.GetFunctionSignatureSymbol(Signature);
+			PerFunctionParameterShaderCode.StringValue.Reset(); //Reset function parameters ready to rebuild when we iterate over the parameters.
+
+			const FString* FunctionTemplate = FunctionTemplateMap.Find(Signature.Name);
+
+			//Generate the function hlsl if we've not done so already for this signature.
+			uint32 FuncHash = GetSignatureHash(Signature);
+			FString& FunctionHlsl = FunctionHlslMap.FindOrAdd(FuncHash);
+			if (FunctionHlsl.IsEmpty() && FunctionTemplate != nullptr)
+			{
+				if (Signature.VariadicInput())
+				{
+					int32 StartInput = Signature.VariadicInputStartIndex();
+					for (int32 InputIdx = StartInput; InputIdx < Signature.Inputs.Num(); ++InputIdx)
+					{
+						const FNiagaraVariable& InputParam = Signature.Inputs[InputIdx];
+						FParamAccessInfo& Access = ParametersAccessed.FindChecked(InputParam);
+
+						FunctionParameterName = InputParam.GetName().ToString();
+						FunctionParameterType = HlslGenContext.GetStructHlslTypeName(InputParam.GetType());
+						PerFunctionParameterShaderCode.StringValue += FString::Format(*PerParameterWriteCallTemplate, HlslTemplateArgs);
+
+
+						bool bFirstParam = CurrParamIdx++ == 0;
+
+						FunctionInputParameters.StringValue += FString::Printf(TEXT("%s%s %s %s"),
+							bFirstParam ? *FirstParamPrefix : *ParamPrefix,
+							*InputParamPrefix,
+							*FunctionParameterType.StringValue, 
+							*FunctionParameterName.StringValue);
+					}
+				}
+
+				if (Signature.VariadicOutput())
+				{
+					int32 StartOutput = Signature.VariadicOutputStartIndex();
+					for (int32 OutputIdx = StartOutput; OutputIdx < Signature.Outputs.Num(); ++OutputIdx)
+					{
+						const FNiagaraVariable& OutputParam = Signature.Outputs[OutputIdx];
+						FParamAccessInfo& Access = ParametersAccessed.FindChecked(OutputParam);
+						
+						FunctionParameterName = OutputParam.GetName().ToString();
+						FunctionParameterType = HlslGenContext.GetStructHlslTypeName(OutputParam.GetType());
+						PerFunctionParameterShaderCode.StringValue += FString::Format(*PerParameterReadCallTemplate, HlslTemplateArgs);
+
+
+						bool bFirstParam = CurrParamIdx++ == 0;
+
+						FunctionOutputParameters.StringValue += FString::Printf(TEXT("%s%s %s %s"),
+							bFirstParam ? *FirstParamPrefix : *ParamPrefix,
+							*OutputParamPrefix,
+							*FunctionParameterType.StringValue, 
+							*FunctionParameterName.StringValue);
+					}
+				}
+
+
+				//Finally generate the final code for this function and add it to the final hlsl.
+				OutHLSL += FString::Format(**FunctionTemplate, HlslTemplateArgs);
+				OutHLSL += TEXT("\n");
+			}
+		}
+	}
+#endif //WITH_EDITORONLY_DATA
 }
 
 #undef LOCTEXT_NAMESPACE
diff --git a/Engine/Plugins/FX/Niagara/Source/Niagara/Private/DataInterface/NiagaraDataInterfaceDataChannelRead.cpp b/Engine/Plugins/FX/Niagara/Source/Niagara/Private/DataInterface/NiagaraDataInterfaceDataChannelRead.cpp
--- a/Engine/Plugins/FX/Niagara/Source/Niagara/Private/DataInterface/NiagaraDataInterfaceDataChannelRead.cpp
+++ b/Engine/Plugins/FX/Niagara/Source/Niagara/Private/DataInterface/NiagaraDataInterfaceDataChannelRead.cpp
@@ -25,6 +25,7 @@
 #include "NiagaraGPUSystemTick.h"
 
 #include "NiagaraDataInterfaceUtilities.h"
+#include "NiagaraGpuComputeDispatchInterface.h"
 
 #if WITH_EDITOR
 #include "INiagaraEditorOnlyDataUtlities.h"
@@ -85,6 +86,7 @@ namespace NDIDataChannelReadLocal
 {
 	static const TCHAR* CommonShaderFile = TEXT("/Plugin/FX/Niagara/Private/DataChannel/NiagaraDataInterfaceDataChannelCommon.ush");
 	static const TCHAR* TemplateShaderFile_Common = TEXT("/Plugin/FX/Niagara/Private/DataChannel/NiagaraDataInterfaceDataChannelTemplateCommon.ush");
+	static const TCHAR* TemplateShaderFile_ReadCommon = TEXT("/Plugin/FX/Niagara/Private/DataChannel/NiagaraDataInterfaceDataChannelTemplateReadCommon.ush");
 	static const TCHAR* TemplateShaderFile_Read = TEXT("/Plugin/FX/Niagara/Private/DataChannel/NiagaraDataInterfaceDataChannelTemplate_Read.ush");
 	static const TCHAR* TemplateShaderFile_Consume = TEXT("/Plugin/FX/Niagara/Private/DataChannel/NiagaraDataInterfaceDataChannelTemplate_Consume.ush");
 
@@ -268,24 +270,30 @@ namespace NDIDataChannelReadLocal
 	// Function definitions END
 	//////////////////////////////////////////////////////////////////////////
 
-	const TCHAR* GetFunctionTemplate(FName FunctionName)
+	void BuildFunctionTemplateMap(TArray<FString>& OutCommonTemplateShaders, TMap<FName, FString>& OutMap)
 	{
-		if (FunctionName == GetFunctionSig_Read().Name) return TemplateShaderFile_Read;
-		if (FunctionName == GetFunctionSig_Consume().Name) return TemplateShaderFile_Consume;
+		LoadShaderSourceFile(TemplateShaderFile_Common, EShaderPlatform::SP_PCD3D_SM5, &OutCommonTemplateShaders.AddDefaulted_GetRef(), nullptr);
+		LoadShaderSourceFile(TemplateShaderFile_ReadCommon, EShaderPlatform::SP_PCD3D_SM5, &OutCommonTemplateShaders.AddDefaulted_GetRef(), nullptr);
 
-		return nullptr;
+		LoadShaderSourceFile(TemplateShaderFile_Read, EShaderPlatform::SP_PCD3D_SM5, &OutMap.Add(GetFunctionSig_Read().Name), nullptr);
+		LoadShaderSourceFile(TemplateShaderFile_Consume, EShaderPlatform::SP_PCD3D_SM5, &OutMap.Add(GetFunctionSig_Consume().Name), nullptr);
 	}
 
 	BEGIN_SHADER_PARAMETER_STRUCT(FShaderParameters, )
 		SHADER_PARAMETER_SRV(Buffer<uint>, ParamOffsetTable)
-		SHADER_PARAMETER_SRV(Buffer<float>, DataFloat)
-		SHADER_PARAMETER_SRV(Buffer<int>, DataInt32)
-		SHADER_PARAMETER_SRV(Buffer<float>, DataHalf)
 		SHADER_PARAMETER(int32, ParameterOffsetTableIndex)
-		SHADER_PARAMETER(int32, Num)
 		SHADER_PARAMETER(int32, FloatStride)
 		SHADER_PARAMETER(int32, Int32Stride)
 		SHADER_PARAMETER(int32, HalfStride)
+
+		SHADER_PARAMETER_SRV(Buffer<float>, DataFloat)
+		SHADER_PARAMETER_SRV(Buffer<int>, DataInt32)
+		SHADER_PARAMETER_SRV(Buffer<float>, DataHalf)
+		SHADER_PARAMETER(int32, InstanceCountOffset)
+		SHADER_PARAMETER(int32, ConsumeInstanceCountOffset)
+		SHADER_PARAMETER(int32, BufferSize)
+
+		SHADER_PARAMETER(int32, NDCElementCountAtSpawn)
 		SHADER_PARAMETER_RDG_BUFFER_SRV(Buffer<int32>, NDCSpawnDataBuffer)
 	END_SHADER_PARAMETER_STRUCT()
 }
@@ -297,27 +305,15 @@ struct FNDIDataChannelReadInstanceData_RT
 	FNiagaraDataChannelDataProxy* ChannelDataRTProxy = nullptr;
 
 	bool bReadPrevFrame = false;
-	/**
-	Table of all parameter offsets used by each GPU script using this DI.
-	Each script has to have it's own section of this table as the offsets into this table are embedded in the hlsl.
-	At hlsl gen time we only have the context of each script individually to generate these indexes.
-	TODO: Can possible elevate this up to the LayoutManager and have a single layout buffer for all scripts
-	*/
-	TResourceArray<uint32> GPUScriptParameterOffsetTable;
-
-	/**
-	Offsets into the parameter table are embedded in the gpu script hlsl.
-	At hlsl gen time we can only know which parameters are accessed by each script individually so each script must have it's own parameter binding table.
-	We provide the offset into the above table via a shader param.
-	TODO: Can just as easily be an offset into a global buffer in the Layout manager.
-	*/
-	TMap<FNiagaraCompileHash, uint32> GPUScriptParameterTableOffsets;
-
-	/** Signal we have updated function binding data from the GT this frame we should process. */
-	bool bHasFunctionBindingUpdate = false;
+
+	/** Parameter mapping info for every function in every script used by this DI. */
+	FVariadicParameterGPUScriptInfo ScriptParamInfo;
 
 	/** Buffer containing packed data for all emitters NDC spawning data for use on the GPU. */
 	TArray<int32> NDCSpawnData;
+
+	/** Number of NDC elements at the point of spawning. More could have been added after this. */
+	int32 NDCElementCountAtSpawn = 0;
 };
 
 //////////////////////////////////////////////////////////////////////////
@@ -327,7 +323,7 @@ FNDIDataChannelReadInstanceData::~FNDIDataChannelReadInstanceData()
 
 }
 
-FNiagaraDataBuffer* FNDIDataChannelReadInstanceData::GetReadBufferCPU(bool bPrevFrame)
+FNiagaraDataBuffer* FNDIDataChannelReadInstanceData::GetReadBufferCPU(bool bPrevFrame)const
 {
 	if (DataChannelData)
 	{
@@ -686,6 +682,8 @@ void UNiagaraDataInterfaceDataChannelRead::PostStageTick(FNDICpuPostStageContext
 	check(InstanceData);
 	check(Context.Usage == ENiagaraScriptUsage::EmitterUpdateScript || Context.Usage == ENiagaraScriptUsage::SystemUpdateScript);
 
+	InstanceData->NDCElementCountAtSpawn = 0;
+	
 	for (auto& EmitterInstanceDataPair : InstanceData->EmitterInstanceData)
 	{
 		FNDIDataChannelRead_EmitterInstanceData& EmitterInstData = EmitterInstanceDataPair.Value;
@@ -695,6 +693,12 @@ void UNiagaraDataInterfaceDataChannelRead::PostStageTick(FNDICpuPostStageContext
 			EmitterInstData.Reset();
 			continue;
 		}
+		
+		if (InstanceData->NDCElementCountAtSpawn == 0)
+		{
+			InstanceData->NDCElementCountAtSpawn = EmitterInstData.NDCSpawnCounts.Num();
+		}
+		check(EmitterInstData.NDCSpawnCounts.Num() == InstanceData->NDCElementCountAtSpawn);
 
 		TArray<uint32> PerNDCSpawnCounts;
 		PerNDCSpawnCounts.Reserve(EmitterInstData.NDCSpawnCounts.Num());
@@ -703,7 +707,7 @@ void UNiagaraDataInterfaceDataChannelRead::PostStageTick(FNDICpuPostStageContext
 			PerNDCSpawnCounts.Add(Info.Get());
 		}
 
-		EmitterInstData.Reset();
+		EmitterInstData.NDCSpawnData.Reset();
 
 		//-TODO:Stateless:
 		if (FNiagaraEmitterInstanceImpl* StatefulEmitter = TargetEmitter->AsStateful())
@@ -803,113 +807,87 @@ void UNiagaraDataInterfaceDataChannelRead::ProvidePerInstanceDataForRenderThread
 
 	bool bReadPrevFrame = bReadCurrentFrame == false || GNDCReadForcePrevFrame;
 	TargetData->bReadPrevFrame = bReadPrevFrame;
+	TargetData->NDCElementCountAtSpawn = SourceData.NDCElementCountAtSpawn;
 
 	if (SourceData.bUpdateFunctionBindingRTData && INiagaraModule::DataChannelsEnabled())
 	{
 		SourceData.bUpdateFunctionBindingRTData = false;
-
-		TargetData->bHasFunctionBindingUpdate = true;
 		
 		const FNiagaraDataSetCompiledData& GPUCompiledData = SourceData.DataChannel->GetDataChannel()->GetCompiledData(ENiagaraSimTarget::GPUComputeSim);
-
-		//For every GPU script, we append it's parameter access info to the table.
-		TargetData->GPUScriptParameterTableOffsets.Reset();
-		constexpr int32 ElemsPerParam = 3;
-		TargetData->GPUScriptParameterOffsetTable.Reset(CompiledData.GetTotalParams() * ElemsPerParam);
-		for (auto& GPUParameterAccessInfoPair : CompiledData.GetGPUScriptParameterInfos())
-		{
-			const FNDIDataChannel_GPUScriptParameterAccessInfo& ParamAccessInfo = GPUParameterAccessInfoPair.Value;
-
-			//First get the offset for this script in the table.
-			TargetData->GPUScriptParameterTableOffsets.FindOrAdd(GPUParameterAccessInfoPair.Key) = TargetData->GPUScriptParameterOffsetTable.Num();
-
-			//Now fill the table for this script
-			for (const FNiagaraVariableBase& Param : ParamAccessInfo.SortedParameters)
-			{
-				if (const FNiagaraVariableLayoutInfo* LayoutInfo = GPUCompiledData.FindVariableLayoutInfo(Param))
-				{
-					TargetData->GPUScriptParameterOffsetTable.Add(LayoutInfo->GetNumFloatComponents() > 0 ? LayoutInfo->GetFloatComponentStart() : INDEX_NONE);
-					TargetData->GPUScriptParameterOffsetTable.Add(LayoutInfo->GetNumInt32Components() > 0 ? LayoutInfo->GetInt32ComponentStart() : INDEX_NONE);
-					TargetData->GPUScriptParameterOffsetTable.Add(LayoutInfo->GetNumHalfComponents() > 0 ? LayoutInfo->GetHalfComponentStart() : INDEX_NONE);
-				}
-				else
-				{
-					TargetData->GPUScriptParameterOffsetTable.Add(INDEX_NONE);
-					TargetData->GPUScriptParameterOffsetTable.Add(INDEX_NONE);
-					TargetData->GPUScriptParameterOffsetTable.Add(INDEX_NONE);
-				}
-			}
-		}
+		TargetData->ScriptParamInfo.Init(CompiledData, GPUCompiledData);
 	}
 
-	//Always need to fill in the NDCSpawnData array as it will change every frame and be pushed into an RDG buffer.
+	//NDC SpawnData Handling. TODO: Refactor into function
+	{
+		//Always need to fill in the NDCSpawnData array as it will change every frame and be pushed into an RDG buffer.
 
-	//New buffer is every emitter continuous NDCSpawnDataArray. We need to store an offset that we pass in as a uniform.
-	//The buckets come first, then the per NDC SpawnCounts, Then the bucket back ptrs.
+		//New buffer is every emitter continuous NDCSpawnDataArray. We need to store an offset that we pass in as a uniform.
+		//The buckets come first, then the per NDC SpawnCounts, Then the bucket back ptrs.
 
-	//Do one pass to calculate size.
-	auto GetEmitterNDCSpawnDataSize = [](const FNDIDataChannelRead_EmitterInstanceData& EmitterInstData)
-	{
-		return 16 + EmitterInstData.NDCSpawnCounts.Num() + EmitterInstData.NDCSpawnData.NDCSpawnData.Num();
-	};
+		//Do one pass to calculate size.
+		auto GetEmitterNDCSpawnDataSize = [](const FNDIDataChannelRead_EmitterInstanceData& EmitterInstData)
+		{
+			return 16 + EmitterInstData.NDCSpawnCounts.Num() + EmitterInstData.NDCSpawnData.NDCSpawnData.Num();
+		};
 
-	uint32 NumEmitters = SourceData.EmitterInstanceData.Num();
-	uint32 TotalPacckedNDCSpawnDataSize = 0;
-	int32 MaxEmitterIndex = 0;
-	for (const TPair<FNiagaraEmitterInstance*, FNDIDataChannelRead_EmitterInstanceData>& EmitterInstDataPair : SourceData.EmitterInstanceData)
-	{
-		if (const FNiagaraEmitterInstance* EmitterInst = EmitterInstDataPair.Key)
+		uint32 NumEmitters = SourceData.EmitterInstanceData.Num();
+		uint32 TotalPacckedNDCSpawnDataSize = 0;
+		int32 MaxEmitterIndex = 0;
+		for (const TPair<FNiagaraEmitterInstance*, FNDIDataChannelRead_EmitterInstanceData>& EmitterInstDataPair : SourceData.EmitterInstanceData)
 		{
-			const FNDIDataChannelRead_EmitterInstanceData& EmitterInstData = EmitterInstDataPair.Value;
+			if (const FNiagaraEmitterInstance* EmitterInst = EmitterInstDataPair.Key)
+			{
+				const FNDIDataChannelRead_EmitterInstanceData& EmitterInstData = EmitterInstDataPair.Value;
 
-			TotalPacckedNDCSpawnDataSize += GetEmitterNDCSpawnDataSize(EmitterInstData);
-			FNiagaraEmitterID ID = EmitterInst->GetEmitterID();
-			MaxEmitterIndex = FMath::Max(MaxEmitterIndex, ID.ID);
+				TotalPacckedNDCSpawnDataSize += GetEmitterNDCSpawnDataSize(EmitterInstData);
+				FNiagaraEmitterID ID = EmitterInst->GetEmitterID();
+				MaxEmitterIndex = FMath::Max(MaxEmitterIndex, ID.ID);
+			}
 		}
-	}
 
-	//First section of the NDCSpawnDataBuffer is an offset into the buffer for each emitter.
-	TotalPacckedNDCSpawnDataSize += MaxEmitterIndex;
+		//First section of the NDCSpawnDataBuffer is an offset into the buffer for each emitter.
+		TotalPacckedNDCSpawnDataSize += MaxEmitterIndex;
 
-	TargetData->NDCSpawnData.Reset(TotalPacckedNDCSpawnDataSize);
-	TArray<int32>& TargetNDCSpawnData = TargetData->NDCSpawnData;
+		TargetData->NDCSpawnData.Reset(TotalPacckedNDCSpawnDataSize);
+		TArray<int32>& TargetNDCSpawnData = TargetData->NDCSpawnData;
 
-	//First grab space for the per emitter offset table. We'll fill this in as we go.
-	TargetNDCSpawnData.AddZeroed(MaxEmitterIndex + 1);
+		//First grab space for the per emitter offset table. We'll fill this in as we go.
+		TargetNDCSpawnData.AddZeroed(MaxEmitterIndex + 1);
 
-	uint32 CurrentSpawnDataOffset = TargetNDCSpawnData.Num();
+		uint32 CurrentSpawnDataOffset = TargetNDCSpawnData.Num();
 
-	for (const TPair<FNiagaraEmitterInstance*, FNDIDataChannelRead_EmitterInstanceData>& EmitterInstDataPair : SourceData.EmitterInstanceData)
-	{
-		if (const FNiagaraEmitterInstance* EmitterInst = EmitterInstDataPair.Key)
+		for (const TPair<FNiagaraEmitterInstance*, FNDIDataChannelRead_EmitterInstanceData>& EmitterInstDataPair : SourceData.EmitterInstanceData)
 		{
-			const FNDIDataChannelRead_EmitterInstanceData& EmitterInstData = EmitterInstDataPair.Value;
-			if(EmitterInst->GetGPUContext() == nullptr)
+			if (const FNiagaraEmitterInstance* EmitterInst = EmitterInstDataPair.Key)
 			{
-				continue;
-			}
+				const FNDIDataChannelRead_EmitterInstanceData& EmitterInstData = EmitterInstDataPair.Value;
+				if (EmitterInst->GetGPUContext() == nullptr)
+				{
+					continue;
+				}
 
 
-			uint32 EmitterNDCSpawnDataSize = GetEmitterNDCSpawnDataSize(EmitterInstData);
+				uint32 EmitterNDCSpawnDataSize = GetEmitterNDCSpawnDataSize(EmitterInstData);
 
-			//First fill in the current offset for this emitter.
-			FNiagaraEmitterID EmitterID = EmitterInst->GetEmitterID();
-			TargetNDCSpawnData[EmitterID.ID] = CurrentSpawnDataOffset;
+				//First fill in the current offset for this emitter.
+				FNiagaraEmitterID EmitterID = EmitterInst->GetEmitterID();
+				TargetNDCSpawnData[EmitterID.ID] = CurrentSpawnDataOffset;
 
-			CurrentSpawnDataOffset += EmitterNDCSpawnDataSize;
+				CurrentSpawnDataOffset += EmitterNDCSpawnDataSize;
 
-			//Next fill in bucket counts
-			for (int32 i = 0; i < 16; ++i)
-			{
-				TargetNDCSpawnData.Add(EmitterInstData.NDCSpawnData.NDCSpawnDataBuckets[i]);
-			}
-			//Next the per NDC Spawn Counts
-			for(int32 i=0; i < EmitterInstData.NDCSpawnCounts.Num(); ++i)
-			{
-				TargetNDCSpawnData.Add(EmitterInstData.NDCSpawnCounts[i].Get());
+				//Next fill in bucket counts
+				for (int32 i = 0; i < 16; ++i)
+				{
+					TargetNDCSpawnData.Add(EmitterInstData.NDCSpawnData.NDCSpawnDataBuckets[i]);
+				}
+				//Next the per NDC Spawn Counts
+				for (int32 i = 0; i < EmitterInstData.NDCSpawnCounts.Num(); ++i)
+				{
+					TargetNDCSpawnData.Add(EmitterInstData.NDCSpawnCounts[i].Get());
+				}
+				//Finally the exec index to NDC index mapping table
+				TargetNDCSpawnData.Append(EmitterInstData.NDCSpawnData.NDCSpawnData);
 			}
-			//Finally the exec index to NDC index mapping table
-			TargetNDCSpawnData.Append(EmitterInstData.NDCSpawnData.NDCSpawnData);
 		}
 	}
 }
@@ -1358,12 +1336,11 @@ void UNiagaraDataInterfaceDataChannelRead::GetNDCSpawnData(FVectorVMExternalFunc
 
 	if(Buffer && INiagaraModule::DataChannelsEnabled())
 	{
-		uint32 NumNDCEntries = Buffer->GetNumInstances();
-
 		auto CalculateNDCSpawnInfo = [&](const FNDIDataChannelRead_EmitterInstanceData& EmitterInstanceData)
 		{
 			const int32* NDCSpawnBukets = EmitterInstanceData.NDCSpawnData.NDCSpawnDataBuckets;
 			TConstArrayView<int32> NDCSpawnData(EmitterInstanceData.NDCSpawnData.NDCSpawnData);
+			uint32 NumNDCEntries = EmitterInstanceData.NDCSpawnCounts.Num();
 
 			uint32 ExecIndex = InExecIndex.GetAndAdvance();
 			uint32 NDCIndex = INDEX_NONE;
@@ -1580,34 +1557,41 @@ void UNiagaraDataInterfaceDataChannelRead::Consume(FVectorVMExternalFunctionCont
 			bool bConsume = InConsume.GetAndAdvance();
 
 			bool bSuccess = false;
-			int32 Index = InstData->ConsumeIndex;
+			int32 Index = INDEX_NONE;
 			if (bConsume)
 			{
-				while (Index < (int32)Data->GetNumInstances() && InstData->ConsumeIndex.compare_exchange_strong(Index, Index + 1) == false)
+				//Increment counter and enforce max if the result is over acceptable values.
+				//Note: This allows the index to temporarily exceed the max limits so is unsafe if we access this concurrently anywhere else without checking the limits.
+				//Note: However it does avoid a more expensive looping compare exchange.
+				Index = InstData->ConsumeIndex++;
+				if(Index >= 0 && Index < (int32)Data->GetNumInstances())
 				{
-					Index = InstData->ConsumeIndex;
-				}
+					//TODO: Wrap/clamp modes etc
 
-				bConsume &= (uint32)Index < Data->GetNumInstances();
-
-				auto FloatFunc = [Data, Index](const FNDIDataChannelRegisterBinding& VMBinding, VectorVM::FExternalFuncRegisterHandler<float>& FloatData)
-				{
-					if (VMBinding.DataSetRegisterIndex != INDEX_NONE)
-						*FloatData.GetDestAndAdvance() = *Data->GetInstancePtrFloat(VMBinding.DataSetRegisterIndex, (uint32)Index);
-				};
-				auto IntFunc = [Data, Index](const FNDIDataChannelRegisterBinding& VMBinding, VectorVM::FExternalFuncRegisterHandler<int32>& IntData)
-				{
-					if (VMBinding.DataSetRegisterIndex != INDEX_NONE)
-						*IntData.GetDestAndAdvance() = *Data->GetInstancePtrInt32(VMBinding.DataSetRegisterIndex, (uint32)Index);
-				};
-				auto HalfFunc = [Data, Index](const FNDIDataChannelRegisterBinding& VMBinding, VectorVM::FExternalFuncRegisterHandler<FFloat16>& HalfData)
+					auto FloatFunc = [Data, Index](const FNDIDataChannelRegisterBinding& VMBinding, VectorVM::FExternalFuncRegisterHandler<float>& FloatData)
+					{
+						if (VMBinding.DataSetRegisterIndex != INDEX_NONE)
+							*FloatData.GetDestAndAdvance() = *Data->GetInstancePtrFloat(VMBinding.DataSetRegisterIndex, (uint32)Index);
+					};
+					auto IntFunc = [Data, Index](const FNDIDataChannelRegisterBinding& VMBinding, VectorVM::FExternalFuncRegisterHandler<int32>& IntData)
+					{
+						if (VMBinding.DataSetRegisterIndex != INDEX_NONE)
+							*IntData.GetDestAndAdvance() = *Data->GetInstancePtrInt32(VMBinding.DataSetRegisterIndex, (uint32)Index);
+					};
+					auto HalfFunc = [Data, Index](const FNDIDataChannelRegisterBinding& VMBinding, VectorVM::FExternalFuncRegisterHandler<FFloat16>& HalfData)
+					{
+						if (VMBinding.DataSetRegisterIndex != INDEX_NONE)
+							*HalfData.GetDestAndAdvance() = *Data->GetInstancePtrHalf(VMBinding.DataSetRegisterIndex, (uint32)Index);
+					};
+					bSuccess = VariadicOutputs.Process(bConsume, 1, BindingInfo, FloatFunc, IntFunc, HalfFunc);
+				}
+				else
 				{
-					if (VMBinding.DataSetRegisterIndex != INDEX_NONE)
-						*HalfData.GetDestAndAdvance() = *Data->GetInstancePtrHalf(VMBinding.DataSetRegisterIndex, (uint32)Index);
-				};
-				bSuccess = VariadicOutputs.Process(bConsume, 1, BindingInfo, FloatFunc, IntFunc, HalfFunc);
+					InstData->ConsumeIndex = Data->GetNumInstances();
+				}
 			}
-			else
+		
+			if(bSuccess == false)
 			{
 				VariadicOutputs.Fallback(1);
 			}
@@ -1930,7 +1914,8 @@ bool UNiagaraDataInterfaceDataChannelRead::AppendCompileHash(FNiagaraCompileHash
 {
 	bool bSuccess = Super::AppendCompileHash(InVisitor);
 	bSuccess &= InVisitor->UpdateString(TEXT("UNiagaraDataInterfaceDataChannelCommon"), GetShaderFileHash(NDIDataChannelReadLocal::CommonShaderFile, EShaderPlatform::SP_PCD3D_SM5).ToString());
-	bSuccess &= InVisitor->UpdateString(TEXT("UNiagaraDataInterfaceDataChannelRead_Common"), GetShaderFileHash(NDIDataChannelReadLocal::TemplateShaderFile_Common, EShaderPlatform::SP_PCD3D_SM5).ToString());
+	bSuccess &= InVisitor->UpdateString(TEXT("UNiagaraDataInterfaceDataChannelTemplateCommon"), GetShaderFileHash(NDIDataChannelReadLocal::TemplateShaderFile_Common, EShaderPlatform::SP_PCD3D_SM5).ToString());
+	bSuccess &= InVisitor->UpdateString(TEXT("UNiagaraDataInterfaceDataChannelRead_Common"), GetShaderFileHash(NDIDataChannelReadLocal::TemplateShaderFile_ReadCommon, EShaderPlatform::SP_PCD3D_SM5).ToString());
 	bSuccess &= InVisitor->UpdateString(TEXT("UNiagaraDataInterfaceDataChannelRead_Read"), GetShaderFileHash(NDIDataChannelReadLocal::TemplateShaderFile_Read, EShaderPlatform::SP_PCD3D_SM5).ToString());
 	bSuccess &= InVisitor->UpdateString(TEXT("UNiagaraDataInterfaceDataChannelRead_Consume"), GetShaderFileHash(NDIDataChannelReadLocal::TemplateShaderFile_Consume, EShaderPlatform::SP_PCD3D_SM5).ToString());
 
@@ -1954,279 +1939,13 @@ bool UNiagaraDataInterfaceDataChannelRead::GetFunctionHLSL(FNiagaraDataInterface
 
 void UNiagaraDataInterfaceDataChannelRead::GetParameterDefinitionHLSL(FNiagaraDataInterfaceHlslGenerationContext& HlslGenContext, FString& OutHLSL)
 {
-	//Desc:
-	// This function iterates over all functions called for this DI in each script and builds the correct hlsl.
-	// The main part of the complexity here is dealing with Variadic function parameters.
-	// We must interrogate the script to see what parameters are actually accessed and generate hlsl accordingly.
-	// Ideally at some future point we can refactor most of this out to a utility helper that will do most of the heavy lifting.
-	// Allowing users to simply provide some details and lambdas etc to specify what exactly they'd like to do in the function body with each variadic param.
-
-	//Initially we'll do some preamble, setting up various template strings and args etc.
-
-	//Map of all arguments for various pieces of template code.
-	//We have some common code that is shared by all functions.
-	//Some code is unique for each function.
-	//Some is unique per parameter to each function.
-	//Finally there is some that is unique for each sub component of each parameter until we've hit a base type. float/2/3/4 etc.
-	TMap<FString, FStringFormatArg> HlslTemplateArgs =
-	{
-		//Common args for all functions
-		{TEXT("ParameterName"),							HlslGenContext.ParameterInfo.DataInterfaceHLSLSymbol},
-		
-		//Per function args. These will be changed with each function written
-		{TEXT("FunctionSymbol"),						FString(TEXT("FunctionSymbol"))},//Function symbol which will be a mangled form from the translator.
-		{TEXT("FunctionParameters"),					FString(TEXT("FunctionParameters"))},//Function parameters written into the function signature.
-		{TEXT("PerFunctionParameterShaderCode"),		FString(TEXT("PerFunctionParameterShaderCode"))},//Combined string of all code dealing with each of the function parameters.
-		//Per function parameter args. These will be changed with each parameter written
-		{TEXT("FunctionParameterIndex"),				TEXT("FunctionParameterIndex")},//Function parameter index allowing us to look up the layout information for the correct parameter to the function.
-		{TEXT("FunctionParameterName"),					FString(TEXT("FunctionParameterName"))},//Name of this function parameter.
-		{TEXT("FunctionParameterType"),					FString(TEXT("FunctionParameterType"))},//Type of this function parameter.
-		{TEXT("PerParamRWCode"),						FString(TEXT("PerParamRWCode"))},//Code that does the actual reading or writing to the data channel buffers.
-		//Per component/base type args. These will change with every base type we I/O from the Data Channel.
-		{TEXT("FunctionParameterComponentBufferType"),	FString(TEXT("FunctionParameterComponentBufferType"))},//The actual base data buffer type being accessed by a particular DataChannel access code line. Float, Int32, Half etc
-		{TEXT("FunctionParameterComponentName"),		FString(TEXT("FunctionParameterComponentName"))},//The name/symbol of the actual member of a parameter that we can I/O from the DataChannel via a standard getter/setter.
-		{TEXT("FunctionParameterComponentType"),		FString(TEXT("FunctionParameterComponentType"))},//The type of the actual member of a parameter that we can I/O from the DataChannel via a standard getter/setter.
-	};
-
-	//Grab refs to per function args we'll change with each function written.
-	FStringFormatArg& FunctionSymbol = HlslTemplateArgs.FindChecked(TEXT("FunctionSymbol"));
-	FStringFormatArg& FunctionParameters = HlslTemplateArgs.FindChecked(TEXT("FunctionParameters"));
-	FStringFormatArg& PerFunctionParameterShaderCode = HlslTemplateArgs.FindChecked(TEXT("PerFunctionParameterShaderCode"));
-
-	//Grab refs to per function args we'll change with each parameter written.
-	FStringFormatArg& FunctionParameterIndex = HlslTemplateArgs.FindChecked(TEXT("FunctionParameterIndex"));
-	FStringFormatArg& FuncParamName = HlslTemplateArgs.FindChecked(TEXT("FunctionParameterName"));
-	FStringFormatArg& FuncParamType = HlslTemplateArgs.FindChecked(TEXT("FunctionParameterType"));
-	FStringFormatArg& PerParamRWCode = HlslTemplateArgs.FindChecked(TEXT("PerParamRWCode"));
-
-	//Grab refs to per component/base type args we'll change for every parameter component we access in the Data Channel.
-	FStringFormatArg& FunctionParameterComponentBufferType = HlslTemplateArgs.FindChecked(TEXT("FunctionParameterComponentBufferType"));
-	FStringFormatArg& FunctionParameterComponentName = HlslTemplateArgs.FindChecked(TEXT("FunctionParameterComponentName"));
-	FStringFormatArg& FunctionParameterComponentType = HlslTemplateArgs.FindChecked(TEXT("FunctionParameterComponentType"));
-
-	//Template code for writing handling code for each of the function parameters.
-	//Some preamble and an replacement arg into which we write all the actual I/O with the DataChannel buffers.
-	FString PerFunctionParameterTemplate = TEXT("\n\
-		//Handling for function parameter {FunctionParameterName}\n\
-		int FloatRegisterIndex_{FunctionParameterName};\n\
-		int Int32RegisterIndex_{FunctionParameterName};\n\
-		int HalfRegisterIndex_{FunctionParameterName};\n\
-		if(GetParameterLayoutInfo_{ParameterName}({FunctionParameterIndex}, FloatRegisterIndex_{FunctionParameterName}, Int32RegisterIndex_{FunctionParameterName}, HalfRegisterIndex_{FunctionParameterName}))\n\
-		{\n\
-{PerParamRWCode}\n\
-		}\n\
-		else\n\
-		{\n\
-			bOutSuccess = false;\n\
-		}\n");
-
-	//Template code for accessing data from the Data Channel's buffers.
-	static const FString ReadDataTemplate = TEXT("			NiagaraDataChannelRead_{FunctionParameterComponentType}({ParameterName}_Data{FunctionParameterComponentBufferType}, {ParameterName}_{FunctionParameterComponentBufferType}Stride, {FunctionParameterComponentBufferType}RegisterIndex_{FunctionParameterName}, ElementIndex, {FunctionParameterComponentName});\n\t\t");
-	static const FString WriteDataTemplate = TEXT("			NiagaraDataChannelWrite_{FunctionParameterComponentType}({ParameterName}_Data{FunctionParameterComponentBufferType}, {ParameterName}_{FunctionParameterComponentBufferType}Stride, {FunctionParameterComponentBufferType}RegisterIndex_{FunctionParameterName}, ElementIndex, {FunctionParameterComponentName});\n\t\t");
-
-	//Map to store template shader code for each DI function as each is in it's own ush file. We'll load those as needed and place in here for ease of access.
-	TMap<FName, FString> FunctionTemplateMap;
-
-	//We may call the same function multiple times so avoid duplicating the same function impl.
-	TMap<uint32, FString> FunctionHlslMap;
-
-	//Now lets build some hlsl!
-	// 
-	//First build the shader code common to all functions.
-	FString CommonTemplateFile;
-	LoadShaderSourceFile(NDIDataChannelReadLocal::TemplateShaderFile_Common, EShaderPlatform::SP_PCD3D_SM5, &CommonTemplateFile, nullptr);
-	OutHLSL += FString::Format(*CommonTemplateFile, HlslTemplateArgs);
-
-	auto GetSignatureHash = [](const FNiagaraFunctionSignature& Sig)
-	{
-		uint32 Ret = GetTypeHash(Sig.Name);
-		for (const FNiagaraVariable& Input : Sig.Inputs)
-		{
-			Ret = HashCombine(Ret, GetTypeHash(Input));
-		}
-		for (const FNiagaraVariableBase& Output : Sig.Outputs)
-		{
-			Ret = HashCombine(Ret, GetTypeHash(Output));
-		}
-		return Ret;
-	};
-
-	TArray<FNiagaraVariableBase> ParametersAccessed;
-	//First iterate over the generated functions to gather all used parameters so we can generate the sorted parameter list for all functions called by this script.
-	for (int32 FuncIdx = 0; FuncIdx < HlslGenContext.ParameterInfo.GeneratedFunctions.Num(); ++FuncIdx)
-	{
-		const FNiagaraDataInterfaceGeneratedFunction& Func = HlslGenContext.ParameterInfo.GeneratedFunctions[FuncIdx];
-		const FNiagaraFunctionSignature& Signature = HlslGenContext.Signatures[FuncIdx];
-
-		for (int32 InputIdx = Signature.NumRequiredInputs(); InputIdx < Signature.Inputs.Num(); ++InputIdx)
-		{
-			const FNiagaraVariable& InputParam = Signature.Inputs[InputIdx];
-			ParametersAccessed.AddUnique(InputParam);
-		}
+	Super::GetParameterDefinitionHLSL(HlslGenContext, OutHLSL);
 
-		for (int32 OutputIdx = Signature.NumRequiredOutputs(); OutputIdx < Signature.Outputs.Num(); ++OutputIdx)
-		{
-			const FNiagaraVariableBase& OutputParam = Signature.Outputs[OutputIdx];
-			ParametersAccessed.AddUnique(OutputParam);
-		}
-	}
-	//Sort the parameters so that that generated hlsl can exactly match the runtime ordering of parameters.
-	NDIDataChannelUtilities::SortParameters(ParametersAccessed);
-
-	//Utility map for easy access of the right param offset during hlsl gen.
-	TMap<FNiagaraVariableBase, uint32> ParamOffsets;
-	for (int32 ParamIdx = 0; ParamIdx < ParametersAccessed.Num(); ++ParamIdx)
-	{
-		ParamOffsets.FindOrAdd(ParametersAccessed[ParamIdx]) = ParamIdx;
-	}
-
-	//Now to iterate on the functions and build the hlsl for each as needed.
-	for (int32 FuncIdx = 0; FuncIdx < HlslGenContext.ParameterInfo.GeneratedFunctions.Num(); ++FuncIdx)
-	{
-		const FNiagaraDataInterfaceGeneratedFunction& Func = HlslGenContext.ParameterInfo.GeneratedFunctions[FuncIdx];
-		const FNiagaraFunctionSignature& Signature = HlslGenContext.Signatures[FuncIdx];
-
-		const TCHAR* FunctionTemplateFilename = NDIDataChannelReadLocal::GetFunctionTemplate(Signature.Name);
-
-		//Init/Reset our per function hlsl template args.
-		FunctionSymbol.StringValue = HlslGenContext.GetFunctionSignatureSymbol(Signature);
-		FunctionParameters.StringValue.Reset();//Reset function parameters ready to rebuild when we iterate over the parameters.
-		PerFunctionParameterShaderCode.StringValue.Reset(); //Reset function parameters ready to rebuild when we iterate over the parameters.
-
-		//Generate the function hlsl if we've not done so already for this signature. We can also skip for functions that have no template file as they don't need special per signature hlsl generation.
-		uint32 FuncHash = GetSignatureHash(Signature);
-		FString& FunctionHlsl = FunctionHlslMap.FindOrAdd(FuncHash);
-		if (FunctionHlsl.IsEmpty() && FunctionTemplateFilename != nullptr)
-		{
-			//Function that will recurse down a parameter's type and generate the appropriate IO code for all of it's members.
-			TFunction<void(bool, UScriptStruct*, FString&)> GenerateRWParameterCode = [&](bool bRead, UScriptStruct* Struct, FString& OutCode)
-			{
-				//Intercept positions and replace with FVector3fs
-				if(Struct == FNiagaraPosition::StaticStruct())
-				{
-					Struct = FNiagaraTypeDefinition::GetVec3Struct();
-				}
-
-				if (Struct == FNiagaraTypeDefinition::GetFloatStruct() || Struct == FNiagaraTypeDefinition::GetVec2Struct() || Struct == FNiagaraTypeDefinition::GetVec3Struct()
-					|| Struct == FNiagaraTypeDefinition::GetVec4Struct() || Struct == FNiagaraTypeDefinition::GetColorStruct() || Struct == FNiagaraTypeDefinition::GetQuatStruct())
-				{
-					FunctionParameterComponentBufferType = TEXT("Float");
-					FunctionParameterComponentType = HlslGenContext.GetStructHlslTypeName(FNiagaraTypeDefinition(Struct));
-					OutCode += FString::Format(bRead ? *ReadDataTemplate : *WriteDataTemplate, HlslTemplateArgs);
-				}
-				else if (Struct == FNiagaraTypeDefinition::GetIntStruct())
-				{
-					FunctionParameterComponentBufferType = TEXT("Int32");
-					FunctionParameterComponentType = HlslGenContext.GetStructHlslTypeName(FNiagaraTypeDefinition(Struct));
-					OutCode += FString::Format(bRead ? *ReadDataTemplate : *WriteDataTemplate, HlslTemplateArgs);
-				}
-				else if (Struct == FNiagaraTypeDefinition::GetBoolStruct())
-				{
-					FunctionParameterComponentBufferType = TEXT("Int32");
-					FunctionParameterComponentType = HlslGenContext.GetStructHlslTypeName(FNiagaraTypeDefinition::GetIntDef());
-					OutCode += FString::Format(bRead ? *ReadDataTemplate : *WriteDataTemplate, HlslTemplateArgs);
-				}
-				else
-				{
-					FString PropertyBaseName = FunctionParameterComponentName.StringValue;
-
-					for (TFieldIterator<FProperty> PropertyIt(Struct, EFieldIteratorFlags::IncludeSuper); PropertyIt; ++PropertyIt)
-					{
-						FunctionParameterComponentName.StringValue = PropertyBaseName;
-						const FProperty* Property = *PropertyIt;
-
-						if (Property->IsA(FFloatProperty::StaticClass()))
-						{
-							FunctionParameterComponentBufferType = TEXT("Float");
-							FunctionParameterComponentType = HlslGenContext.GetStructHlslTypeName(FNiagaraTypeDefinition::GetFloatDef());
-							FunctionParameterComponentName.StringValue += FString::Printf(TEXT(".%s"), *HlslGenContext.GetSanitizedSymbolName(Property->GetName()));
-							OutCode += FString::Format(bRead ? *ReadDataTemplate : *WriteDataTemplate, HlslTemplateArgs);
-						}
-						else if (Property->IsA(FIntProperty::StaticClass()))
-						{
-							FunctionParameterComponentBufferType = TEXT("Int32");
-							FunctionParameterComponentType = HlslGenContext.GetStructHlslTypeName(FNiagaraTypeDefinition::GetIntDef());
-							FunctionParameterComponentName.StringValue += FString::Printf(TEXT(".%s"), *HlslGenContext.GetSanitizedSymbolName(Property->GetName()));
-							OutCode += FString::Format(bRead ? *ReadDataTemplate : *WriteDataTemplate, HlslTemplateArgs);
-						}
-						else if (const FStructProperty* StructProp = CastFieldChecked<const FStructProperty>(Property))
-						{
-							FunctionParameterComponentName.StringValue += FString::Printf(TEXT(".%s"), *HlslGenContext.GetSanitizedSymbolName(Property->GetName()));
-							FunctionParameterComponentType = HlslGenContext.GetPropertyHlslTypeName(Property);
-							GenerateRWParameterCode(bRead, StructProp->Struct, OutCode);
-						}
-						else
-						{
-							check(false);
-							OutCode += FString::Printf(TEXT("Error! - DataChannel Interface encountered an invalid type when generating it's function hlsl. %s"), *Property->GetName());
-						}
-					}
-				}
-			};
-
-			int32 CurrParamIdx = 0;
-			int32 CurrVariadicParamIdx = 0;//Track current variadic param index so the shader code can tell which param's layout info to lookup.
-
-			//Iterate through the variadic inputs and output parameters building the function signature and the per parameter function body as we go.
-			static FString FirstParamPrefix = TEXT("");
-			static FString ParamPrefix = TEXT(", ");
-			static FString OutputParamPrefix = TEXT("out");
-			static FString InputParamPrefix = TEXT("in");
-			auto HandleParameters = [&](TConstArrayView<FNiagaraVariableBase> Parameters, int32 StartParamIdx, bool bOutput)
-			{
-				for (int32 ParamIdx = StartParamIdx; ParamIdx < Parameters.Num(); ++ParamIdx)
-				{
-					const FNiagaraVariableBase& Param = Parameters[ParamIdx];
-
-					//Reset or init per parameter hlsl template args.
-
-					//Grab the offset for this parameter and embed in the hlsl directly.
-					//This will correspond to an entry in the ParameterLayoutBuffer passed to the shader.
-					int32 ParamOffset = ParamOffsets.FindChecked(Param);
-					FunctionParameterIndex = LexToString(ParamOffset);
-					FuncParamName = Param.GetName().ToString();
-					FuncParamType = HlslGenContext.GetStructHlslTypeName(Param.GetType());
-					PerParamRWCode.StringValue.Reset();
-
-					bool bFirstParam = CurrParamIdx++ == 0;
-					FunctionParameters.StringValue += FString::Printf(TEXT("%s%s %s %s"),
-						bFirstParam ? *FirstParamPrefix : *ParamPrefix,
-						bOutput ? *OutputParamPrefix : *InputParamPrefix,
-						*FuncParamType.StringValue, *FuncParamName.StringValue);
-
-					//Reset per component args and build the code for per component/base type Data Channel access.
-					FunctionParameterComponentName.StringValue.Reset();
-					FunctionParameterComponentType.StringValue.Reset();
-					FunctionParameterComponentBufferType.StringValue.Reset();
-
-					//Generate the per component/base type IO code that will be used in the subsequent per paramters shader code template.
-					FunctionParameterComponentName = HlslGenContext.GetSanitizedSymbolName(Param.GetName().ToString());
-					GenerateRWParameterCode(bOutput, Param.GetType().GetScriptStruct(), PerParamRWCode.StringValue);
-
-					//Generate the whole code for this parameter and append it to the function's per parameter code.
-					//This will be used at the end in the per function template.
-					PerFunctionParameterShaderCode.StringValue += FString::Format(*PerFunctionParameterTemplate, HlslTemplateArgs);
-
-					++CurrVariadicParamIdx;
-				}
-			};
-
-			HandleParameters(Signature.GetInputs(), Signature.NumRequiredInputs(), false);
-			HandleParameters(Signature.Outputs, Signature.NumRequiredOutputs(), true);
-
-			//Grab the template file for the function.
-			FString& FunctionTemplateFile = FunctionTemplateMap.FindOrAdd(Signature.Name);
-			if (FunctionTemplateFile.IsEmpty())
-			{
-				//Load it if we have not previously.
-				LoadShaderSourceFile(FunctionTemplateFilename, EShaderPlatform::SP_PCD3D_SM5, &FunctionTemplateFile, nullptr);
-			}
-
-			//Finally generate the final code for this function and add it to the final hlsl.
-			OutHLSL += FString::Format(*FunctionTemplateFile, HlslTemplateArgs);
-			OutHLSL += TEXT("\n");
-		}
-	}
+	TArray<FString> CommonTemplateShaders;
+	TMap<FName, FString> TemplateShaderMap;
+	NDIDataChannelReadLocal::BuildFunctionTemplateMap(CommonTemplateShaders, TemplateShaderMap);
+	
+	NDIDataChannelUtilities::GenerateDataChannelAccessHlsl(HlslGenContext, CommonTemplateShaders, TemplateShaderMap, OutHLSL);
 }
 
 bool UNiagaraDataInterfaceDataChannelRead::UpgradeFunctionCall(FNiagaraFunctionSignature& FunctionSignature)
@@ -2287,7 +2006,7 @@ void UNiagaraDataInterfaceDataChannelRead::SetShaderParameters(const FNiagaraDat
 
 		if (InstanceData->ChannelDataRTProxy && ParameterOffsetTableIndex != INDEX_NONE)
 		{
-			FNiagaraDataBuffer* Data = InstanceData->bReadPrevFrame ? InstanceData->ChannelDataRTProxy->PrevFrameData.GetReference() : InstanceData->ChannelDataRTProxy->GPUDataSet->GetCurrentData();
+			FNiagaraDataBuffer* Data = InstanceData->bReadPrevFrame ? InstanceData->ChannelDataRTProxy->GetPrevFrameData().GetReference() : InstanceData->ChannelDataRTProxy->GetCurrentData().GetReference();
 			if (Data)
 			{
 				const FReadBuffer& ParameterLayoutBuffer = InstanceData->ParameterLayoutBuffer;
@@ -2295,17 +2014,30 @@ void UNiagaraDataInterfaceDataChannelRead::SetShaderParameters(const FNiagaraDat
 				FRDGBufferSRVRef NDCSpawnDataBufferSRV = Context.GetGraphBuilder().CreateSRV(InstanceData->NDCSpawnDataBuffer, PF_R32_SINT);
 				if (NDCSpawnDataBufferSRV && ParameterLayoutBuffer.SRV.IsValid() && ParameterLayoutBuffer.NumBytes > 0)
 				{
+					TArray<FRHITransitionInfo, TInlineAllocator<3>> Transitions;
+					Transitions.Reserve(3);
+					Transitions.Emplace(Data->GetGPUBufferFloat().UAV, ERHIAccess::Unknown, ERHIAccess::SRVCompute);
+					Transitions.Emplace(Data->GetGPUBufferInt().UAV, ERHIAccess::Unknown, ERHIAccess::SRVCompute);
+					Transitions.Emplace(Data->GetGPUBufferHalf().UAV, ERHIAccess::Unknown, ERHIAccess::SRVCompute);
+					Context.GetGraphBuilder().RHICmdList.Transition(Transitions);
+
 					InstParameters->ParamOffsetTable = ParameterLayoutBuffer.SRV.IsValid() ? ParameterLayoutBuffer.SRV.GetReference() : FNiagaraRenderer::GetDummyUIntBuffer();
-					InstParameters->DataFloat = Data->GetGPUBufferFloat().SRV.IsValid() ? Data->GetGPUBufferFloat().SRV.GetReference() : FNiagaraRenderer::GetDummyFloatBuffer();
-					InstParameters->DataInt32 = Data->GetGPUBufferInt().SRV.IsValid() ? Data->GetGPUBufferInt().SRV.GetReference() : FNiagaraRenderer::GetDummyIntBuffer();
-					InstParameters->DataHalf = Data->GetGPUBufferHalf().SRV.IsValid() ? Data->GetGPUBufferHalf().SRV.GetReference() : FNiagaraRenderer::GetDummyHalfBuffer();
 					InstParameters->ParameterOffsetTableIndex = ParameterOffsetTableIndex;
-					InstParameters->Num = Data->GetNumInstances();
+
 					InstParameters->FloatStride = Data->GetFloatStride() / sizeof(float);
 					InstParameters->Int32Stride = Data->GetInt32Stride() / sizeof(int32);
 					InstParameters->HalfStride = Data->GetHalfStride() / sizeof(FFloat16);
 
+					InstParameters->DataFloat = Data->GetGPUBufferFloat().SRV.IsValid() ? Data->GetGPUBufferFloat().SRV.GetReference() : FNiagaraRenderer::GetDummyFloatBuffer();
+					InstParameters->DataInt32 = Data->GetGPUBufferInt().SRV.IsValid() ? Data->GetGPUBufferInt().SRV.GetReference() : FNiagaraRenderer::GetDummyIntBuffer();
+					InstParameters->DataHalf = Data->GetGPUBufferHalf().SRV.IsValid() ? Data->GetGPUBufferHalf().SRV.GetReference() : FNiagaraRenderer::GetDummyHalfBuffer();
+					InstParameters->InstanceCountOffset = Data->GetGPUInstanceCountBufferOffset();
+					InstParameters->ConsumeInstanceCountOffset = InstanceData->ConsumeInstanceCountOffset;
+					InstParameters->BufferSize = Data ? Data->GetNumInstancesAllocated() : INDEX_NONE;
+
+					InstParameters->NDCElementCountAtSpawn = InstanceData->NDCElementCountAtSpawn;
 					InstParameters->NDCSpawnDataBuffer = NDCSpawnDataBufferSRV;
+
 					bSuccess = true;
 				}
 			}
@@ -2315,15 +2047,21 @@ void UNiagaraDataInterfaceDataChannelRead::SetShaderParameters(const FNiagaraDat
 	if (bSuccess == false)
 	{
 		InstParameters->ParamOffsetTable = FNiagaraRenderer::GetDummyUIntBuffer();
-		InstParameters->DataFloat = FNiagaraRenderer::GetDummyFloatBuffer();
-		InstParameters->DataInt32 = FNiagaraRenderer::GetDummyIntBuffer();
-		InstParameters->DataHalf = FNiagaraRenderer::GetDummyHalfBuffer();
 		InstParameters->ParameterOffsetTableIndex = INDEX_NONE;
-		InstParameters->Num = 0;
+
 		InstParameters->FloatStride = 0;
 		InstParameters->Int32Stride = 0;
 		InstParameters->HalfStride = 0;
+
+		InstParameters->DataFloat = FNiagaraRenderer::GetDummyFloatBuffer();
+		InstParameters->DataInt32 = FNiagaraRenderer::GetDummyIntBuffer();
+		InstParameters->DataHalf = FNiagaraRenderer::GetDummyHalfBuffer();
+
+		InstParameters->InstanceCountOffset = INDEX_NONE;
+		InstParameters->ConsumeInstanceCountOffset = INDEX_NONE;
+		InstParameters->BufferSize = INDEX_NONE;
 		
+		InstParameters->NDCElementCountAtSpawn = 0;
 		FRDGBufferRef DummyBuffer = GSystemTextures.GetDefaultBuffer(Context.GetGraphBuilder(), 4, 0u);
 		InstParameters->NDCSpawnDataBuffer = Context.GetGraphBuilder().CreateSRV(DummyBuffer, PF_R32_SINT);		
 	}
@@ -2333,12 +2071,22 @@ void FNiagaraDataInterfaceProxy_DataChannelRead::PreStage(const FNDIGpuComputePr
 {
 	FNiagaraDataInterfaceProxy_DataChannelRead::FInstanceData* InstanceData = SystemInstancesToProxyData_RT.Find(Context.GetSystemInstanceID());
 
-	if(InstanceData && InstanceData->NDCSpawnDataBuffer == nullptr)
+	if(InstanceData)
 	{
-		InstanceData->NDCSpawnDataBuffer = CreateUploadBuffer<int32>(
-			Context.GetGraphBuilder(),
-			TEXT("Niagara_NDCReadDI_NDCSpawnData"),
-			InstanceData->NDCSpawnData);
+		//TODO: Should grab just one for the whole frame...
+		//TODO: Add some wrap behavior...		
+		if(InstanceData->ConsumeInstanceCountOffset == INDEX_NONE)
+		{
+			InstanceData->ConsumeInstanceCountOffset = Context.GetInstanceCountManager().AcquireOrAllocateEntry(Context.GetGraphBuilder().RHICmdList);
+		}
+
+		if(InstanceData->NDCSpawnDataBuffer == nullptr)
+		{
+			InstanceData->NDCSpawnDataBuffer = CreateUploadBuffer<int32>(
+				Context.GetGraphBuilder(),
+				TEXT("Niagara_NDCReadDI_NDCSpawnData"),
+				InstanceData->NDCSpawnData);
+		}
 	}
 }
 
@@ -2349,6 +2097,8 @@ void FNiagaraDataInterfaceProxy_DataChannelRead::PostSimulate(const FNDIGpuCompu
 		if(FNiagaraDataInterfaceProxy_DataChannelRead::FInstanceData* InstanceData = SystemInstancesToProxyData_RT.Find(Context.GetSystemInstanceID()))
 		{
 			InstanceData->NDCSpawnDataBuffer = nullptr;
+
+			Context.GetInstanceCountManager().FreeEntry(InstanceData->ConsumeInstanceCountOffset);
 		}
 	}
 }
@@ -2362,16 +2112,17 @@ void FNiagaraDataInterfaceProxy_DataChannelRead::ConsumePerInstanceDataFromGameT
 
 	InstData.ChannelDataRTProxy = SourceData.ChannelDataRTProxy;
 	InstData.bReadPrevFrame = SourceData.bReadPrevFrame;
+	InstData.NDCElementCountAtSpawn = SourceData.NDCElementCountAtSpawn;
 
 	InstData.NDCSpawnData = SourceData.NDCSpawnData;
 	InstData.NDCSpawnDataBuffer = nullptr;//Clear the RDG buffer ready for re-up to the GPU.
 
-	if (SourceData.bHasFunctionBindingUpdate) 
+	if (SourceData.ScriptParamInfo.bDirty) 
 	{
 		//Take the offset map from the source data.
 		//This maps from GPU script to that scripts offset into the ParameterLayoutBuffer.
 		//Allows us to look up and pass in at SetShaderParameters time.
-		InstData.GPUScriptParameterTableOffsets = MoveTemp(SourceData.GPUScriptParameterTableOffsets);
+		InstData.GPUScriptParameterTableOffsets = MoveTemp(SourceData.ScriptParamInfo.GPUScriptParameterTableOffsets);
 
 		//Now generate the ParameterLayoutBuffer
 		//This contains a table of all parameters used by each GPU script that uses this DI.
@@ -2382,10 +2133,10 @@ void FNiagaraDataInterfaceProxy_DataChannelRead::ConsumePerInstanceDataFromGameT
 				InstData.ParameterLayoutBuffer.Release();
 			}
 
-			if(SourceData.GPUScriptParameterOffsetTable.Num() > 0)
+			if(SourceData.ScriptParamInfo.GPUScriptParameterOffsetTable.Num() > 0)
 			{
-				InstData.ParameterLayoutData = SourceData.GPUScriptParameterOffsetTable;
-				InstData.ParameterLayoutBuffer.Initialize(RHICmdList, TEXT("NDIDataChannel_ParameterLayoutBuffer"), sizeof(uint32), SourceData.GPUScriptParameterOffsetTable.Num(), EPixelFormat::PF_R32_UINT, BUF_Static, &InstData.ParameterLayoutData);
+				InstData.ParameterLayoutData = SourceData.ScriptParamInfo.GPUScriptParameterOffsetTable;
+				InstData.ParameterLayoutBuffer.Initialize(RHICmdList, TEXT("NDIDataChannel_ParameterLayoutBuffer"), sizeof(uint32), SourceData.ScriptParamInfo.GPUScriptParameterOffsetTable.Num(), EPixelFormat::PF_R32_UINT, BUF_Static, &InstData.ParameterLayoutData);
 			}
 		}
 	}
@@ -2403,10 +2154,13 @@ void FNiagaraDataInterfaceProxy_DataChannelRead::GetDispatchArgs(const FNDIGpuCo
 	const FNiagaraDataInterfaceProxy_DataChannelRead::FInstanceData* InstanceData = SystemInstancesToProxyData_RT.Find(Context.GetSystemInstanceID());
 	if (InstanceData && InstanceData->ChannelDataRTProxy)
 	{
-		FNiagaraDataBuffer* Data = InstanceData->bReadPrevFrame ? InstanceData->ChannelDataRTProxy->PrevFrameData.GetReference() : InstanceData->ChannelDataRTProxy->GPUDataSet->GetCurrentData();
+		FNiagaraDataBuffer* Data = InstanceData->bReadPrevFrame ? InstanceData->ChannelDataRTProxy->PrevFrameData.GetReference() : InstanceData->ChannelDataRTProxy->GetCurrentData().GetReference();
 		if (Data)
 		{
-			Context.SetDirect(Data->GetNumInstances());
+			//Indirect args via the instance count buffer is not working. TODO.
+			//Running for all allocated elements will execute more than needed but should allow things to work.
+			//Context.CreateIndirect(Data->GetGPUInstanceCountBufferOffset());
+			Context.SetDirect(Data->GetNumInstancesAllocated());
 		}
 	}
 }
diff --git a/Engine/Plugins/FX/Niagara/Source/Niagara/Private/DataInterface/NiagaraDataInterfaceDataChannelWrite.cpp b/Engine/Plugins/FX/Niagara/Source/Niagara/Private/DataInterface/NiagaraDataInterfaceDataChannelWrite.cpp
--- a/Engine/Plugins/FX/Niagara/Source/Niagara/Private/DataInterface/NiagaraDataInterfaceDataChannelWrite.cpp
+++ b/Engine/Plugins/FX/Niagara/Source/Niagara/Private/DataInterface/NiagaraDataInterfaceDataChannelWrite.cpp
@@ -12,10 +12,23 @@
 #include "NiagaraWorldManager.h"
 #include "NiagaraSystemInstance.h"
 
+#include "NiagaraCompileHashVisitor.h"
+#include "ShaderCompilerCore.h"
+#include "SystemTextures.h"
+#include "RenderGraphUtils.h"
+#include "NiagaraModule.h"
+#include "NiagaraShaderParametersBuilder.h"
+
 #include "NiagaraDataChannel.h"
 #include "NiagaraDataChannelHandler.h"
 #include "NiagaraDataChannelManager.h"
 
+#include "NiagaraRenderer.h"
+#include "NiagaraGPUSystemTick.h"
+#include "NiagaraGpuComputeDispatchInterface.h"
+
+#include "NiagaraDataInterfaceUtilities.h"
+
 #if WITH_EDITOR
 #include "INiagaraEditorOnlyDataUtlities.h"
 #include "Modules/ModuleManager.h"
@@ -40,26 +53,162 @@ static FAutoConsoleVariableRef CVarDebugDumpWriterDI(
 
 namespace NDIDataChannelWriteLocal
 {
-	static const TCHAR* CommonShaderFile = TEXT("/Plugin/FX/Niagara/Private/DataChannel/NiagaraDataInterfaceDataChannelCommon.ush");
+	static const TCHAR* CommonShaderFile = TEXT("/Plugin/FX/Niagara/Private/DataChannel/NiagaraDataInterfaceDataChannelCommon.ush");	
 	static const TCHAR* TemplateShaderFile_Common = TEXT("/Plugin/FX/Niagara/Private/DataChannel/NiagaraDataInterfaceDataChannelTemplateCommon.ush");
+	static const TCHAR* TemplateShaderFile_WriteCommon = TEXT("/Plugin/FX/Niagara/Private/DataChannel/NiagaraDataInterfaceDataChannelTemplateWriteCommon.ush");
 	static const TCHAR* TemplateShaderFile_Write = TEXT("/Plugin/FX/Niagara/Private/DataChannel/NiagaraDataInterfaceDataChannelTemplate_Write.ush");
 	static const TCHAR* TemplateShaderFile_Append = TEXT("/Plugin/FX/Niagara/Private/DataChannel/NiagaraDataInterfaceDataChannelTemplate_Append.ush");
 
-	static const FName NumName(TEXT("Num"));
-	static const FName WriteName(TEXT("Write"));
-	static const FName AppendName(TEXT("Append"));
+	//////////////////////////////////////////////////////////////////////////
+	//Function definitions
+
+	/////
+	// NOTE: *any* changes to function inputs or outputs here needs to be included in FWriteNDCModel::GenerateNewModuleContent()
+	/////
 
-	static const FName SpawnName(TEXT("Spawn"));
+	const FNiagaraFunctionSignature& GetFunctionSig_Num()
+	{
+		static FNiagaraFunctionSignature Sig;
+		if(!Sig.IsValid())
+		{
+			Sig.Name = TEXT("Num");
+#if WITH_EDITORONLY_DATA
+			NIAGARA_ADD_FUNCTION_SOURCE_INFO(Sig)
+			Sig.Description = LOCTEXT("NumFunctionDescription", "Returns the number of instances allocated for writing in the bound NDC buffer. Writes at an index beyond this will fail.");
+#endif
+			Sig.bMemberFunction = true;
+			Sig.bExperimental = true;
+			Sig.AddInput(FNiagaraVariable(FNiagaraTypeDefinition(UNiagaraDataInterfaceDataChannelWrite::StaticClass()), TEXT("DataChannel interface")));
+			Sig.AddOutput(FNiagaraVariable(FNiagaraTypeDefinition::GetIntDef(), TEXT("Num")));
+		}
+		return Sig;
+	}
 
-	const TCHAR* GetFunctionTemplate(FName FunctionName)
+	const FNiagaraFunctionSignature& GetFunctionSig_Allocate()
 	{
-		if (FunctionName == WriteName) return TemplateShaderFile_Write;
-		if (FunctionName == AppendName) return TemplateShaderFile_Append;
+		static FNiagaraFunctionSignature Sig;
+		if (!Sig.IsValid())
+		{
+			Sig.Name = TEXT("Allocate");
+#if WITH_EDITORONLY_DATA
+			NIAGARA_ADD_FUNCTION_SOURCE_INFO(Sig)
+			Sig.Description = LOCTEXT("AllocateFunctionDescription", "Adds an amount to allocated into the bound NDC data for the given emitter to write into.");
+#endif
+			Sig.bMemberFunction = true;
+			Sig.bExperimental = true;
+			Sig.bRequiresExecPin = true;
+			Sig.ModuleUsageBitmask = ENiagaraScriptUsageMask::Emitter | ENiagaraScriptUsageMask::System;
+			Sig.AddInput(FNiagaraVariable(FNiagaraTypeDefinition(UNiagaraDataInterfaceDataChannelWrite::StaticClass()), TEXT("DataChannel interface")));
+			Sig.AddInputWithoutDefault(FNiagaraVariable(FNiagaraTypeDefinition(FNiagaraEmitterID::StaticStruct()), TEXT("Emitter ID")), LOCTEXT("EmitterIDDesc", "ID of the emitter that will be writing into this allocated space."));
+			Sig.AddInputWithDefault(FNiagaraVariable(FNiagaraTypeDefinition::GetIntDef(), TEXT("Allocation Count")), 0, LOCTEXT("AllocationCountDesc", "The number of elements to allocate in the NDC for writing from the given emitter."));
+		}
+		return Sig;
+	}
 
-		return nullptr;
+	const FNiagaraFunctionSignature& GetFunctionSig_Write()
+	{
+		static FNiagaraFunctionSignature Sig;
+		if (!Sig.IsValid())
+		{
+			static FNiagaraVariable EmitVar(FNiagaraTypeDefinition::GetBoolDef(), TEXT("Emit"));
+			EmitVar.SetValue(FNiagaraBool(true));
+			Sig.Name = TEXT("Write");
+#if WITH_EDITORONLY_DATA
+			NIAGARA_ADD_FUNCTION_SOURCE_INFO(Sig)
+			Sig.Description = LOCTEXT("WriteFunctionDescription", "Writes DataChannel data at a specific index.  Values in the DataChannel that are not written here are set to their defaults. Returns success if an DataChannel was written to.");
+#endif
+			Sig.bMemberFunction = true;
+			Sig.bRequiresExecPin = true;
+			Sig.bExperimental = true;
+			Sig.bSupportsGPU = false;//Cannot use direct index writes on GPU as we write into one shared buffer with all DIs using the same NDC data.
+			Sig.AddInput(FNiagaraVariable(FNiagaraTypeDefinition(UNiagaraDataInterfaceDataChannelWrite::StaticClass()), TEXT("DataChannel interface")));
+			Sig.AddInput(EmitVar, LOCTEXT("ExecuteAppendFlagTooltip", "If true then the append is executed, if false then the append call is skipped"));
+			Sig.AddInput(FNiagaraVariable(FNiagaraTypeDefinition::GetIntDef(), TEXT("Index")));
+			Sig.AddOutput(FNiagaraVariable(FNiagaraTypeDefinition::GetBoolDef(), TEXT("Success")));
+			Sig.RequiredInputs = IntCastChecked<int16>(Sig.Inputs.Num());//The user defines what we write in the graph.
+		}
+		return Sig;
 	}
+
+	const FNiagaraFunctionSignature& GetFunctionSig_Append()
+	{
+		static FNiagaraFunctionSignature Sig;
+		if (!Sig.IsValid())
+		{
+			static FNiagaraVariable EmitVar(FNiagaraTypeDefinition::GetBoolDef(), TEXT("Emit"));
+			EmitVar.SetValue(FNiagaraBool(true));
+
+			Sig.Name = TEXT("Append");
+#if WITH_EDITORONLY_DATA
+			NIAGARA_ADD_FUNCTION_SOURCE_INFO(Sig)
+			Sig.Description = LOCTEXT("AppendFunctionDescription", "Appends a new DataChannel to the end of the DataChannel array and writes the specified values. Values in the DataChannel that are not written here are set to their defaults. Returns success if an DataChannel was successfully pushed.");
+#endif
+			Sig.bMemberFunction = true;
+			Sig.bRequiresExecPin = true;
+			Sig.bExperimental = true;
+			Sig.AddInput(FNiagaraVariable(FNiagaraTypeDefinition(UNiagaraDataInterfaceDataChannelWrite::StaticClass()), TEXT("DataChannel interface")));
+			Sig.AddInput(EmitVar, LOCTEXT("ExecuteAppendFlagTooltip", "If true then the append is executed, if false then the append call is skipped"));
+			Sig.AddOutput(FNiagaraVariable(FNiagaraTypeDefinition::GetBoolDef(), TEXT("Success")));
+			Sig.RequiredInputs = IntCastChecked<int16>(Sig.Inputs.Num());//The user defines what we write in the graph.
+		}
+		return Sig;
+	}
+
+	void BuildFunctionTemplateMap(TArray<FString>& OutCommonTemplateShaders, TMap<FName, FString>& OutMap)
+	{
+		//Add common template shaders 
+		LoadShaderSourceFile(TemplateShaderFile_Common, EShaderPlatform::SP_PCD3D_SM5, &OutCommonTemplateShaders.AddDefaulted_GetRef(), nullptr);
+		LoadShaderSourceFile(TemplateShaderFile_WriteCommon, EShaderPlatform::SP_PCD3D_SM5, &OutCommonTemplateShaders.AddDefaulted_GetRef(), nullptr);
+
+		//Add per function template shaders
+//		LoadShaderSourceFile(TemplateShaderFile_Write, EShaderPlatform::SP_PCD3D_SM5, &OutMap.Add(GetFunctionSig_Write().Name), nullptr);//Wite is not supported on GPU
+		LoadShaderSourceFile(TemplateShaderFile_Append, EShaderPlatform::SP_PCD3D_SM5, &OutMap.Add(GetFunctionSig_Append().Name), nullptr);
+	}
+
+
+	BEGIN_SHADER_PARAMETER_STRUCT(FShaderParameters, )
+		SHADER_PARAMETER_SRV(Buffer<uint>, ParamOffsetTable)
+		SHADER_PARAMETER(int32, ParameterOffsetTableIndex)
+		SHADER_PARAMETER(int32, FloatStride)
+		SHADER_PARAMETER(int32, Int32Stride)
+		SHADER_PARAMETER(int32, HalfStride)
+
+		SHADER_PARAMETER_UAV(RWBuffer<float>, GPUBufferFloat)
+		SHADER_PARAMETER_UAV(RWBuffer<int>, GPUBufferInt32)
+		SHADER_PARAMETER_UAV(RWBuffer<float>, GPUBufferHalf)
+		SHADER_PARAMETER(int32, GPUInstanceCountOffset)
+		SHADER_PARAMETER(int32, GPUBufferSize)
+
+		SHADER_PARAMETER_UAV(RWBuffer<float>, CPUBufferFloat)
+		SHADER_PARAMETER_UAV(RWBuffer<int>, CPUBufferInt32)
+		SHADER_PARAMETER_UAV(RWBuffer<float>, CPUBufferHalf)
+		SHADER_PARAMETER(int32, CPUInstanceCountOffset)
+		SHADER_PARAMETER(int32, CPUBufferSize)
+		SHADER_PARAMETER(int32, CPUFloatStride)
+		SHADER_PARAMETER(int32, CPUInt32Stride)
+		SHADER_PARAMETER(int32, CPUHalfStride)
+	END_SHADER_PARAMETER_STRUCT()
 }
 
+/** Render thread copy of current instance data. */
+struct FNDIDataChannelWriteInstanceData_RT
+{
+	//RT proxy for game channel data from which we're reading.
+	FNiagaraDataChannelDataProxy* ChannelDataRTProxy = nullptr;
+
+	/** Parameter mapping info for every function in every script used by this DI. */
+	FVariadicParameterGPUScriptInfo ScriptParamInfo;
+
+	/** How many instances should we allocate in the NDC for this DI. */
+	uint32 AllocationCount = 0;
+
+	bool bPublishToGame = false;
+	bool bPublishToCPU = false;
+	bool bPublishToGPU = false;
+
+	FVector3f LwcTile;
+};
+
 //////////////////////////////////////////////////////////////////////////
 //FNDIDataChannelWriteCompiledData
 
@@ -108,43 +257,30 @@ struct FNDIDataChannelWriteInstanceData
 	//Shared pointer to the actual data we'll be pushing into for this data channel.
 	FNiagaraDataChannelDataPtr DataChannelData;
 
-	/** Local dataset we write into. 
-	TODO: alternate write modes.
-		- Crit sec access to the data channel buffer and write direct?
-		- 
-	*/
-	FNiagaraDataSet* Data = nullptr;
+	/** Cached hash to check if the layout of our source data has changed. */
+	uint64 ChachedDataSetLayoutHash = INDEX_NONE;
 
 	TArray<FNDIDataChannel_FuncToDataSetBindingPtr, TInlineAllocator<8>> FunctionToDatSetBindingInfo;
 
 	//Atomic uint for tracking num instances of the target data buffer when writing from multiple threads in the VM.
 	std::atomic<uint32> AtomicNumInstances;
 
+	/** When true we should update our function binding info on the RT next tick. */
+	mutable bool bUpdateFunctionBindingRTData = false;
+
+	int32 DynamicAllocationCount = 0;
+
+	FVector3f LwcTile;
+
+	//Buffer we're currently writing into this frame.
+	FNiagaraDataBuffer* DestinationData = nullptr;
+
 	~FNDIDataChannelWriteInstanceData()
 	{
-		if (Data && DataChannelData.IsValid())
-		{
-			DataChannelData->RemovePublishRequests(Data);
-		}
-
-		/** We defer the deletion of the dataset to the RT to be sure all in-flight RT commands have finished using it.*/
-		ENQUEUE_RENDER_COMMAND(FDeleteContextCommand)(
-			[DataChannelDataSet = Data](FRHICommandListImmediate& RHICmdList)
-			{
-				if (DataChannelDataSet != nullptr)
-				{
-					delete DataChannelDataSet;
-				}
-			}
-		);
-		Data = nullptr;
 	}
 
 	bool Init(UNiagaraDataInterfaceDataChannelWrite* Interface, FNiagaraSystemInstance* Instance)
 	{
-		Data = new FNiagaraDataSet();
-		Data->Init(&Interface->GetCompiledData().DataLayout);
-
 		//In non test/shipping builds we gather and log and missing parameters that cause us to fail to find correct bindings.
 		TArray<FNiagaraVariableBase> MissingParams;
 
@@ -153,7 +289,7 @@ struct FNDIDataChannelWriteInstanceData
 		FunctionToDatSetBindingInfo.Reset(CompiledData.GetFunctionInfo().Num());
 		for (const FNDIDataChannelFunctionInfo& FuncInfo : CompiledData.GetFunctionInfo())
 		{
-			FunctionToDatSetBindingInfo.Add(FNDIDataChannelLayoutManager::Get().GetLayoutInfo(FuncInfo, Data->GetCompiledData(), MissingParams));			
+			FunctionToDatSetBindingInfo.Add(FNDIDataChannelLayoutManager::Get().GetLayoutInfo(FuncInfo, Interface->GetCompiledData().DataLayout, MissingParams));			
 		}
 
 		return true;
@@ -161,6 +297,10 @@ struct FNDIDataChannelWriteInstanceData
 
 	bool Tick(UNiagaraDataInterfaceDataChannelWrite* Interface, FNiagaraSystemInstance* Instance)
 	{
+		DynamicAllocationCount = 0;
+		AtomicNumInstances = 0;
+		
+		LwcTile = Instance->GetLWCTile();
 		if (Interface->ShouldPublish())
 		{
 			UNiagaraDataChannelHandler* DataChannelPtr = DataChannel.Get();
@@ -209,11 +349,57 @@ struct FNDIDataChannelWriteInstanceData
 
 			if (DataChannelPtr)
 			{
+				//In non test/shipping builds we gather and log and missing parameters that cause us to fail to find correct bindings.
+				TArray<FNiagaraVariableBase> MissingParams;
+
 				if (DataChannelData == nullptr || Interface->bUpdateDestinationDataEveryTick)
 				{
 					FNiagaraDataChannelSearchParameters SearchParams(Instance->GetAttachComponent());
 					DataChannelData = DataChannelPtr->FindData(SearchParams, ENiagaraResourceAccess::WriteOnly);
 				}
+
+				const FNDIDataChannelCompiledData& CompiledData = Interface->GetCompiledData();
+
+				if(DataChannelData)
+				{
+					if(CompiledData.UsedByCPU())
+					{
+						//Grab the buffer we'll be writing into for cpu sims. Must be done on the GT but actual buffer alloc can be done concurrently.
+						DestinationData = DataChannelData->GetBufferForCPUWrite();
+						if(DestinationData && Interface->AllocationMode == ENiagaraDataChannelAllocationMode::Static)
+						{
+							DestinationData->Allocate(Interface->AllocationCount);
+						}
+					}
+				}
+
+				const FNiagaraDataSetCompiledData& CPUSourceDataCompiledData = DataChannelPtr->GetDataChannel()->GetCompiledData(ENiagaraSimTarget::CPUSim);
+				const FNiagaraDataSetCompiledData& GPUSourceDataCompiledData = DataChannelPtr->GetDataChannel()->GetCompiledData(ENiagaraSimTarget::GPUComputeSim);
+				check(CPUSourceDataCompiledData.GetLayoutHash() && CPUSourceDataCompiledData.GetLayoutHash() == GPUSourceDataCompiledData.GetLayoutHash());
+				uint64 SourceDataLayoutHash = CPUSourceDataCompiledData.GetLayoutHash();
+				bool bChanged = SourceDataLayoutHash != ChachedDataSetLayoutHash;
+
+				//If our CPU or GPU source data has changed then regenerate our binding info.
+				//TODO: Multi-source buffer support. 
+				//TODO: Variable input layout support. i.e. allow source systems to publish their particle buffers without the need for a separate write.
+				if (bChanged)
+				{
+					ChachedDataSetLayoutHash = SourceDataLayoutHash;
+
+					//We can likely be more targeted here.
+					//Could probably only update the RT when the GPU data changes and only update the bindings if the function hashes change etc.
+					bUpdateFunctionBindingRTData = CompiledData.UsedByGPU();
+					int32 NumFuncs = CompiledData.GetFunctionInfo().Num();
+					FunctionToDatSetBindingInfo.SetNum(NumFuncs);
+					//FuncToDataSetLayoutKeys.SetNumZeroed(NumFuncs);
+					for (int32 BindingIdx = 0; BindingIdx < NumFuncs; ++BindingIdx)
+					{
+						const FNDIDataChannelFunctionInfo& FuncInfo = CompiledData.GetFunctionInfo()[BindingIdx];
+
+						FNDIDataChannel_FuncToDataSetBindingPtr& BindingPtr = FunctionToDatSetBindingInfo[BindingIdx];
+						BindingPtr = FNDIDataChannelLayoutManager::Get().GetLayoutInfo(FuncInfo, CPUSourceDataCompiledData, MissingParams);
+					}
+				}
 			}
 		}		
 
@@ -235,63 +421,38 @@ struct FNDIDataChannelWriteInstanceData
 
 		return true;
 	}
-
-	void PreStageTick(UNiagaraDataInterfaceDataChannelWrite* Interface, FNiagaraSystemInstance* Instance)
-	{
-		if (FNiagaraDataBuffer* CurrBuff = Data->GetCurrentData())
-		{
-			FNiagaraDataBuffer& DestBuff = Data->BeginSimulate(true);
-		}
-
-		//TODO: Currently allocating for every stage using this DI whether it's a write or not. We should limit this to write functions.
-		AtomicNumInstances = 0;
-		if (Interface->AllocationMode == ENiagaraDataChannelAllocationMode::Static)
-		{
-			Data->GetDestinationData()->Allocate(Interface->AllocationCount);
-		}
-		//else if (Interface->AllocationMode == ENiagaraDataChannelAllocationMode::PerInstance)
-		//{
-			//TODO: 
-			//Have to count up the current size of all users of this DI.
-			//Can this be done in DI tick? Or do we need a pre stage API once we do sim stages on CPU?
-			//Also adds extra wrinkle if we have multiple stages or emitters using the same writer DI. 
-			//Maybe need a pre pass to gather all users and their allocation sizes this frame.
-		//}
-		else
-		{
-			check(0);
-		}
-	}
-
-	void PostStageTick(UNiagaraDataInterfaceDataChannelWrite* Interface, FNiagaraSystemInstance* Instance)
+	
+	bool PostSimTick(UNiagaraDataInterfaceDataChannelWrite* Interface, FNiagaraSystemInstance* Instance)
 	{
-		if (Data && Data->GetDestinationData())
+		if (DestinationData)
 		{
 			//The count here can overrun the num allocated but we should never actually write beyond the max allocated.
 			uint32 WrittenInstances = AtomicNumInstances.load(std::memory_order_seq_cst);
-			WrittenInstances = FMath::Min(WrittenInstances, Data->GetDestinationData()->GetNumInstancesAllocated());
-			Data->GetDestinationData()->SetNumInstances(WrittenInstances);
-			Data->EndSimulate();
+			WrittenInstances = FMath::Min(WrittenInstances, DestinationData->GetNumInstancesAllocated());
+			DestinationData->SetNumInstances(WrittenInstances);
 
 			if (GbDebugDumpWriter)
 			{
-				FNiagaraDataBuffer& Buffer = Data->GetCurrentDataChecked();
-				Buffer.Dump(0, Buffer.GetNumInstances(), FString::Printf(TEXT("=== Data Channle Write: %d Elements --> %s ==="), Buffer.GetNumInstances(), *Interface->Channel.GetName()));
+				DestinationData->Dump(0, DestinationData->GetNumInstances(), FString::Printf(TEXT("=== Data Channle Write: %d Elements --> %s ==="), DestinationData->GetNumInstances(), *Interface->Channel.GetName()));
 			}
 
-			if (DataChannelData && Interface->ShouldPublish() && Data->GetCurrentData()->GetNumInstances() > 0)
+			if (DataChannelData && Interface->ShouldPublish() && DestinationData->GetNumInstances() > 0)
 			{
-				FNiagaraDataChannelPublishRequest PublishRequest(Data->GetCurrentData());
+				FNiagaraDataChannelPublishRequest PublishRequest(DestinationData->UnlockForRead());
 				PublishRequest.bVisibleToGame = Interface->bPublishToGame;
 				PublishRequest.bVisibleToCPUSims = Interface->bPublishToCPU;
 				PublishRequest.bVisibleToGPUSims = Interface->bPublishToGPU;
 				PublishRequest.LwcTile = Instance->GetLWCTile();
 #if WITH_NIAGARA_DEBUGGER
-				PublishRequest.DebugSource = FString::Format(TEXT("{0} ({1})"), {Instance->GetSystem()->GetName(), GetPathNameSafe(Interface)});
+				PublishRequest.DebugSource = FString::Format(TEXT("{0} ({1})"), { Instance->GetSystem()->GetName(), GetPathNameSafe(Interface) });
 #endif
 				DataChannelData->Publish(PublishRequest);
 			}
+
+			DestinationData = nullptr;
+			AtomicNumInstances = 0;
 		}
+		return true;
 	}
 };
 
@@ -318,21 +479,20 @@ void UNiagaraDataInterfaceDataChannelWrite::DestroyPerInstanceData(void* PerInst
 	FNDIDataChannelWriteInstanceData* InstanceData = static_cast<FNDIDataChannelWriteInstanceData*>(PerInstanceData);
 	InstanceData->~FNDIDataChannelWriteInstanceData();
 
-	// 	ENQUEUE_RENDER_COMMAND(RemoveProxy)
-	// 	(
-	// 		[RT_Proxy = GetProxyAs<FNDIDataChannelProxy>(), InstanceID = SystemInstance->GetId()](FRHICommandListImmediate& CmdList)
-	// 		{
-	// 			RT_Proxy->SystemInstancesToInstanceData_RT.Remove(InstanceID);
-	// 		}
-	// 	);
+	ENQUEUE_RENDER_COMMAND(RemoveProxy)
+	(
+		[RT_Proxy = GetProxyAs<FNiagaraDataInterfaceProxy_DataChannelWrite>(), InstanceID = SystemInstance->GetId()](FRHICommandListImmediate& CmdList)
+		{
+			RT_Proxy->SystemInstancesToProxyData_RT.Remove(InstanceID);
+		}
+	);
 }
 
 
 UNiagaraDataInterfaceDataChannelWrite::UNiagaraDataInterfaceDataChannelWrite(FObjectInitializer const& ObjectInitializer)
 	: Super(ObjectInitializer)
 {
-	//TODO:
-	//Proxy.Reset(new FNDIDataChannelProxy());
+	Proxy.Reset(new FNiagaraDataInterfaceProxy_DataChannelWrite());
 }
 
 void UNiagaraDataInterfaceDataChannelWrite::PostInitProperties()
@@ -377,49 +537,50 @@ bool UNiagaraDataInterfaceDataChannelWrite::PerInstanceTick(void* PerInstanceDat
 
 bool UNiagaraDataInterfaceDataChannelWrite::PerInstanceTickPostSimulate(void* PerInstanceData, FNiagaraSystemInstance* SystemInstance, float DeltaSeconds)
 {
- 	return false;
-}
-
-void UNiagaraDataInterfaceDataChannelWrite::PreStageTick(FNDICpuPostStageContext& Context)
-{
-	if (INiagaraModule::DataChannelsEnabled() == false)
+	FNDIDataChannelWriteInstanceData* InstanceData = static_cast<FNDIDataChannelWriteInstanceData*>(PerInstanceData);
+	if (!InstanceData)
 	{
-		return;
+		return true;
 	}
 
-	SCOPE_CYCLE_COUNTER(STAT_NDIDataChannelWrite_PreStageTick);
-	check(Context.SystemInstance);
-	FNDIDataChannelWriteInstanceData* InstanceData = static_cast<FNDIDataChannelWriteInstanceData*>(Context.PerInstanceData);
-	if (!InstanceData)
+	if (InstanceData->PostSimTick(this, SystemInstance) == false)
 	{
-		return;
+		return true;
 	}
 
-	InstanceData->PreStageTick(this, Context.SystemInstance);
+ 	return false;
 }
 
-void UNiagaraDataInterfaceDataChannelWrite::PostStageTick(FNDICpuPostStageContext& Context)
+void UNiagaraDataInterfaceDataChannelWrite::ProvidePerInstanceDataForRenderThread(void* DataForRenderThread, void* PerInstanceData, const FNiagaraSystemInstanceID& SystemInstance)
 {
-	if (INiagaraModule::DataChannelsEnabled() == false)
+	const FNDIDataChannelWriteInstanceData& SourceData = *reinterpret_cast<const FNDIDataChannelWriteInstanceData*>(PerInstanceData);
+	FNDIDataChannelWriteInstanceData_RT* TargetData = new(DataForRenderThread) FNDIDataChannelWriteInstanceData_RT();
+
+	//Always update the dataset, this may change without triggering a full update if it's layout is the same.
+	TargetData->ChannelDataRTProxy = SourceData.DataChannelData ? SourceData.DataChannelData->GetRTProxy() : nullptr;
+
+	if (SourceData.bUpdateFunctionBindingRTData && INiagaraModule::DataChannelsEnabled())
 	{
-		return;
+		SourceData.bUpdateFunctionBindingRTData = false;
+
+		const FNiagaraDataSetCompiledData& GPUCompiledData = SourceData.DataChannel->GetDataChannel()->GetCompiledData(ENiagaraSimTarget::GPUComputeSim);
+		TargetData->ScriptParamInfo.Init(CompiledData, GPUCompiledData);
 	}
 
-	SCOPE_CYCLE_COUNTER(STAT_NDIDataChannelWrite_PostStageTick);
-	check(Context.PerInstanceData);
-	FNDIDataChannelWriteInstanceData* InstanceData = static_cast<FNDIDataChannelWriteInstanceData*>(Context.PerInstanceData);
-	if (!Context.SystemInstance)
+	TargetData->AllocationCount = 0;
+	if(AllocationMode == ENiagaraDataChannelAllocationMode::Static)
 	{
-		return;
+		TargetData->AllocationCount = AllocationCount;
+	}
+	else if(AllocationMode == ENiagaraDataChannelAllocationMode::Dynamic)
+	{
+		TargetData->AllocationCount = SourceData.DynamicAllocationCount;
 	}
-	
-	InstanceData->PostStageTick(this, Context.SystemInstance);
-}
 
-void UNiagaraDataInterfaceDataChannelWrite::ProvidePerInstanceDataForRenderThread(void* DataForRenderThread, void* PerInstanceData, const FNiagaraSystemInstanceID& SystemInstance)
-{
-	//TODO:
-	//FNDIDataChannelProxy::ProvidePerInstanceDataForRenderThread(DataForRenderThread, PerInstanceData, SystemInstance);
+	TargetData->bPublishToGame = bPublishToGame;
+	TargetData->bPublishToCPU = bPublishToCPU;
+	TargetData->bPublishToGPU = bPublishToGPU;
+	TargetData->LwcTile = SourceData.LwcTile;
 }
 
 bool UNiagaraDataInterfaceDataChannelWrite::HasTickGroupPostreqs() const
@@ -471,7 +632,7 @@ void UNiagaraDataInterfaceDataChannelWrite::GetFeedback(UNiagaraSystem* InAsset,
 
 	Super::GetFeedback(InAsset, InComponent, OutErrors, OutWarnings, OutInfo);
 
-	if (Channel == nullptr)
+	if (Channel == nullptr || RuntimeDI->Channel == nullptr)
 	{
 		OutErrors.Emplace(LOCTEXT("DataChannelMissingFmt", "Data Channel Interface has no valid Data Channel."),
 			LOCTEXT("DataChannelMissingErrorSummaryFmt", "Missing Data Channel."),
@@ -608,11 +769,11 @@ bool UNiagaraDataInterfaceDataChannelWrite::SimCacheWriteFrame(UObject* StorageO
 		Storage->DataChannelReference = Channel.Get();
 		FNDIDataChannelWriteSimCacheFrame& FrameData = Storage->FrameData.AddDefaulted_GetRef();
 		
-		if (InstanceData->DataChannelData && ShouldPublish() && InstanceData->Data && InstanceData->Data->GetCurrentData() && InstanceData->Data->GetCurrentData()->GetNumInstances() > 0)
+		if (InstanceData->DataChannelData && ShouldPublish() && InstanceData->DestinationData && InstanceData->DestinationData->GetNumInstances() > 0)
 		{
 			FNiagaraDataChannelGameData GameData;
 			GameData.Init(Channel->Get());
-			GameData.AppendFromDataSet(InstanceData->Data->GetCurrentData(), SystemInstance->GetLWCTile());
+			GameData.AppendFromDataSet(InstanceData->DestinationData, SystemInstance->GetLWCTile());
 			
 			FrameData.NumElements = GameData.Num();
 			for (const FNiagaraDataChannelVariableBuffer& VarBuffer : GameData.GetVariableBuffers())
@@ -774,71 +935,32 @@ bool UNiagaraDataInterfaceDataChannelWrite::CopyToInternal(UNiagaraDataInterface
 #if WITH_EDITORONLY_DATA
 void UNiagaraDataInterfaceDataChannelWrite::GetFunctionsInternal(TArray<FNiagaraFunctionSignature>& OutFunctions) const
 {
-
-	/////
-	// NOTE: *any* changes to function inputs or outputs here needs to be included in FWriteNDCModel::GenerateNewModuleContent()
-	/////
-	{
-		FNiagaraFunctionSignature Sig;
-		Sig.Name = NDIDataChannelWriteLocal::NumName;
-		NIAGARA_ADD_FUNCTION_SOURCE_INFO(Sig)
-		Sig.Description = LOCTEXT("NumFunctionDescription", "Returns the current number of DataChannel accessible by this interface.");
-		Sig.bMemberFunction = true;
-		Sig.bExperimental = true;
-		Sig.AddInput(FNiagaraVariable(FNiagaraTypeDefinition(GetClass()), TEXT("DataChannel interface")));
-		Sig.AddOutput(FNiagaraVariable(FNiagaraTypeDefinition::GetIntDef(), TEXT("Num")));
-		OutFunctions.Add(Sig);
-	}
-
-	static FNiagaraVariable EmitVar(FNiagaraTypeDefinition::GetBoolDef(), TEXT("Emit"));
-	EmitVar.SetValue(FNiagaraBool(true));
-
-	{
-		FNiagaraFunctionSignature Sig;
-		Sig.Name = NDIDataChannelWriteLocal::WriteName;
-		NIAGARA_ADD_FUNCTION_SOURCE_INFO(Sig)
-		Sig.Description = LOCTEXT("WriteFunctionDescription", "Writes DataChannel data at a specific index.  Values in the DataChannel that are not written here are set to their defaults. Returns success if an DataChannel was written to.");
-		Sig.bMemberFunction = true;
-		Sig.bRequiresExecPin = true;
-		Sig.bExperimental = true;
-		Sig.AddInput(FNiagaraVariable(FNiagaraTypeDefinition(GetClass()), TEXT("DataChannel interface")));
-		Sig.AddInput(EmitVar, LOCTEXT("ExecuteFlagTooltip", "If true then the write is executed, if false then the whole write is skipped"));
-		Sig.AddInput(FNiagaraVariable(FNiagaraTypeDefinition::GetIntDef(), TEXT("Index")));
-		Sig.AddOutput(FNiagaraVariable(FNiagaraTypeDefinition::GetBoolDef(), TEXT("Success")));
-		Sig.RequiredInputs = IntCastChecked<int16>(Sig.Inputs.Num());//The user defines what we write in the graph.
-		OutFunctions.Add(Sig);
-	}
-	{
-		FNiagaraFunctionSignature Sig;
-		Sig.Name = NDIDataChannelWriteLocal::AppendName;
-		NIAGARA_ADD_FUNCTION_SOURCE_INFO(Sig)
-		Sig.Description = LOCTEXT("AppendFunctionDescription", "Appends a new DataChannel to the end of the DataChannel array and writes the specified values. Values in the DataChannel that are not written here are set to their defaults. Returns success if an DataChannel was successfully pushed.");
-		Sig.bMemberFunction = true;
-		Sig.bRequiresExecPin = true;
-		Sig.bExperimental = true;
-		Sig.AddInput(FNiagaraVariable(FNiagaraTypeDefinition(GetClass()), TEXT("DataChannel interface")));
-		Sig.AddInput(EmitVar, LOCTEXT("ExecuteAppendFlagTooltip", "If true then the append is executed, if false then the append call is skipped"));
-		Sig.AddOutput(FNiagaraVariable(FNiagaraTypeDefinition::GetBoolDef(), TEXT("Success")));
-		Sig.RequiredInputs = IntCastChecked<int16>(Sig.Inputs.Num());//The user defines what we write in the graph.
-		OutFunctions.Add(Sig);
-	}
+	OutFunctions.Add(NDIDataChannelWriteLocal::GetFunctionSig_Num());
+	OutFunctions.Add(NDIDataChannelWriteLocal::GetFunctionSig_Allocate());
+	OutFunctions.Add(NDIDataChannelWriteLocal::GetFunctionSig_Write());
+	OutFunctions.Add(NDIDataChannelWriteLocal::GetFunctionSig_Append());
 }
 #endif
 
 void UNiagaraDataInterfaceDataChannelWrite::GetVMExternalFunction(const FVMExternalFunctionBindingInfo& BindingInfo, void* InstanceData, FVMExternalFunction& OutFunc)
 {
-	if (BindingInfo.Name == NDIDataChannelWriteLocal::NumName)
+	if (BindingInfo.Name == NDIDataChannelWriteLocal::GetFunctionSig_Num().Name)
 	{
 		OutFunc = FVMExternalFunction::CreateLambda([this](FVectorVMExternalFunctionContext& Context) { this->Num(Context); });
 	}
+	else if (BindingInfo.Name == NDIDataChannelWriteLocal::GetFunctionSig_Allocate().Name)
+	{
+		OutFunc = FVMExternalFunction::CreateLambda([this](FVectorVMExternalFunctionContext& Context) { this->Allocate(Context); });
+	}
 	else
 	{
 		int32 FuncIndex = CompiledData.FindFunctionInfoIndex(BindingInfo.Name, BindingInfo.VariadicInputs, BindingInfo.VariadicOutputs);
-		if (BindingInfo.Name == NDIDataChannelWriteLocal::WriteName)
+		if (BindingInfo.Name == NDIDataChannelWriteLocal::GetFunctionSig_Write().Name)
 		{
 			OutFunc = FVMExternalFunction::CreateLambda([this, FuncIndex](FVectorVMExternalFunctionContext& Context) { this->Write(Context, FuncIndex); });
 		}
-		else if (BindingInfo.Name == NDIDataChannelWriteLocal::AppendName)
+		else 
+		if (BindingInfo.Name == NDIDataChannelWriteLocal::GetFunctionSig_Append().Name)
 		{
 			OutFunc = FVMExternalFunction::CreateLambda([this, FuncIndex](FVectorVMExternalFunctionContext& Context) { this->Append(Context, FuncIndex); });
 		}
@@ -855,11 +977,11 @@ void UNiagaraDataInterfaceDataChannelWrite::Num(FVectorVMExternalFunctionContext
 
 	FNDIOutputParam<int32> OutNum(Context);
 
-	FNiagaraDataBuffer* Buffer = InstData->Data->GetDestinationData();	
+	FNiagaraDataBuffer* Buffer = InstData->DestinationData;	
 	int32 Num = 0;
 	if (Buffer && INiagaraModule::DataChannelsEnabled())
 	{
-		Num = (int32)Buffer->GetNumInstances();
+		Num = (int32)Buffer->GetNumInstancesAllocated();
 	}
 
 	for (int32 i = 0; i < Context.GetNumInstances(); ++i)
@@ -868,6 +990,27 @@ void UNiagaraDataInterfaceDataChannelWrite::Num(FVectorVMExternalFunctionContext
 	}
 }
 
+void UNiagaraDataInterfaceDataChannelWrite::Allocate(FVectorVMExternalFunctionContext& Context)
+{
+	VectorVM::FUserPtrHandler<FNDIDataChannelWriteInstanceData> InstData(Context);
+
+	FNDIInputParam<FNiagaraEmitterID> InEmitterID(Context);
+	FNDIInputParam<int32> InAllocationCount(Context);
+
+	check(Context.GetNumInstances() == 1);
+
+	uint32 Count = InAllocationCount.GetAndAdvance();
+	
+	//Store the count so we can pass this to the GPU for allocating space in the main GPU write buffer.
+	InstData->DynamicAllocationCount = Count;
+	
+	//If we have a CPU write buffer, allocate that now. Do this here so the emitter/system script itself can write data if it wants to.
+	if(InstData->DestinationData)
+	{
+		InstData->DestinationData->Allocate(Count, true);//On the off chance we alloc->write->alloc, re-alloc and keep existing data.
+	}
+}
+
 void UNiagaraDataInterfaceDataChannelWrite::Write(FVectorVMExternalFunctionContext& Context, int32 FuncIdx)
 {
 	SCOPE_CYCLE_COUNTER(STAT_NDIDataChannelWrite_Write);
@@ -883,9 +1026,9 @@ void UNiagaraDataInterfaceDataChannelWrite::Write(FVectorVMExternalFunctionConte
 	std::atomic<uint32>& AtomicNumInstances = InstData->AtomicNumInstances;
 
 	bool bAllFailedFallback = true;
-	if (InstData->Data && BindingInfo && INiagaraModule::DataChannelsEnabled())
+	if (InstData->DestinationData && BindingInfo && INiagaraModule::DataChannelsEnabled())
 	{
-		if(FNiagaraDataBuffer* Data = InstData->Data->GetDestinationData())
+		if(FNiagaraDataBuffer* Data = InstData->DestinationData)
 		{			
 			bAllFailedFallback = false;
 			uint32 MaxLocalIndex = 0;
@@ -961,9 +1104,9 @@ void UNiagaraDataInterfaceDataChannelWrite::Append(FVectorVMExternalFunctionCont
 	std::atomic<uint32>& AtomicNumInstances = InstData->AtomicNumInstances;
 
 	bool bAllFailedFallback = true;
-	if (InstData->Data && BindingInfo &&  INiagaraModule::DataChannelsEnabled())
+	if (InstData->DestinationData && BindingInfo &&  INiagaraModule::DataChannelsEnabled())
 	{
-		if(FNiagaraDataBuffer* Data = InstData->Data->GetDestinationData())
+		if(FNiagaraDataBuffer* Data = InstData->DestinationData)
 		{
 			//Get the total number to emit.
 			//Allows going via a faster write path if we're emiting every instance.
@@ -1122,56 +1265,250 @@ void UNiagaraDataInterfaceDataChannelWrite::Append(FVectorVMExternalFunctionCont
 
 #if WITH_EDITORONLY_DATA
 bool UNiagaraDataInterfaceDataChannelWrite::AppendCompileHash(FNiagaraCompileHashVisitor* InVisitor) const
-{
-	//TODO: GPU Writes.
-	// 
-// 	bool bSuccess = Super::AppendCompileHash(InVisitor);
-// 	bSuccess &= InVisitor->UpdateString(TEXT("UNiagaraDataInterfaceDataChannelCommon"), GetShaderFileHash(NDIDataChannelWriteLocal::CommonShaderFile, EShaderPlatform::SP_PCD3D_SM5).ToString());
-// 	bSuccess &= InVisitor->UpdateString(TEXT("UNiagaraDataInterfaceDataChannelWrite_Common"), GetShaderFileHash(NDIDataChannelWriteLocal::TemplateShaderFile_Common, EShaderPlatform::SP_PCD3D_SM5).ToString());
-// 	bSuccess &= InVisitor->UpdateString(TEXT("UNiagaraDataInterfaceDataChannelWrite_Write"), GetShaderFileHash(NDIDataChannelWriteLocal::TemplateShaderFile_Write, EShaderPlatform::SP_PCD3D_SM5).ToString());
-// 	bSuccess &= InVisitor->UpdateString(TEXT("UNiagaraDataInterfaceDataChannelWrite_Append"), GetShaderFileHash(NDIDataChannelWriteLocal::TemplateShaderFile_Append, EShaderPlatform::SP_PCD3D_SM5).ToString());
-// 
-// 	bSuccess &= InVisitor->UpdateShaderParameters<NDIDataChannelWriteLocal::FShaderParameters>();
-// 	return bSuccess;
-
-	return true;
+{ 
+	bool bSuccess = Super::AppendCompileHash(InVisitor);
+ 	bSuccess &= InVisitor->UpdateString(TEXT("UNiagaraDataInterfaceDataChannelCommon"), GetShaderFileHash(NDIDataChannelWriteLocal::CommonShaderFile, EShaderPlatform::SP_PCD3D_SM5).ToString());
+	bSuccess &= InVisitor->UpdateString(TEXT("UNiagaraDataInterfaceDataChannelTemplateCommon"), GetShaderFileHash(NDIDataChannelWriteLocal::TemplateShaderFile_Common, EShaderPlatform::SP_PCD3D_SM5).ToString());
+	bSuccess &= InVisitor->UpdateString(TEXT("UNiagaraDataInterfaceDataChannelWrite_WriteCommon"), GetShaderFileHash(NDIDataChannelWriteLocal::TemplateShaderFile_WriteCommon, EShaderPlatform::SP_PCD3D_SM5).ToString());
+ 	//bSuccess &= InVisitor->UpdateString(TEXT("UNiagaraDataInterfaceDataChannelWrite_Write"), GetShaderFileHash(NDIDataChannelWriteLocal::TemplateShaderFile_Write, EShaderPlatform::SP_PCD3D_SM5).ToString());
+ 	bSuccess &= InVisitor->UpdateString(TEXT("UNiagaraDataInterfaceDataChannelWrite_Append"), GetShaderFileHash(NDIDataChannelWriteLocal::TemplateShaderFile_Append, EShaderPlatform::SP_PCD3D_SM5).ToString());
+ 
+ 	bSuccess &= InVisitor->UpdateShaderParameters<NDIDataChannelWriteLocal::FShaderParameters>();
+ 	return bSuccess;
 }
 
 void UNiagaraDataInterfaceDataChannelWrite::GetCommonHLSL(FString& OutHLSL)
 {
-	//TODO: GPU Writes.
-	// 
-// 	Super::GetCommonHLSL(OutHLSL);
-// 	OutHLSL.Append(TEXT("\n//Niagara Data Channel Write Interface Common Code.\n"));
-// 	OutHLSL.Appendf(TEXT("#include \"%s\"\n"), NDIDataChannelWriteLocal::CommonShaderFile);
+ 	Super::GetCommonHLSL(OutHLSL);
+ 	OutHLSL.Appendf(TEXT("#include \"%s\"\n"), NDIDataChannelWriteLocal::CommonShaderFile);
 }
 
 bool UNiagaraDataInterfaceDataChannelWrite::GetFunctionHLSL(FNiagaraDataInterfaceHlslGenerationContext& HlslGenContext, FString& OutHLSL)
 {
-	//TODO: GPU Writes.
-	// 
-// 	return	HlslGenContext.GetFunctionInfo().DefinitionName == GET_FUNCTION_NAME_CHECKED(UNiagaraDataInterfaceDataChannelWrite, Num) ||
-// 		HlslGenContext.GetFunctionInfo().DefinitionName == GET_FUNCTION_NAME_CHECKED(UNiagaraDataInterfaceDataChannelWrite, Write) ||
-// 		HlslGenContext.GetFunctionInfo().DefinitionName == GET_FUNCTION_NAME_CHECKED(UNiagaraDataInterfaceDataChannelWrite, Append);
-
-	return false;
+ 	return	HlslGenContext.GetFunctionInfo().DefinitionName == GET_FUNCTION_NAME_CHECKED(UNiagaraDataInterfaceDataChannelWrite, Num) ||
+ 		//HlslGenContext.GetFunctionInfo().DefinitionName == GET_FUNCTION_NAME_CHECKED(UNiagaraDataInterfaceDataChannelWrite, Write) ||
+ 		HlslGenContext.GetFunctionInfo().DefinitionName == GET_FUNCTION_NAME_CHECKED(UNiagaraDataInterfaceDataChannelWrite, Append);
 }
 
 void UNiagaraDataInterfaceDataChannelWrite::GetParameterDefinitionHLSL(FNiagaraDataInterfaceHlslGenerationContext& HlslGenContext, FString& OutHLSL)
 {
 	Super::GetParameterDefinitionHLSL(HlslGenContext, OutHLSL);
 
-	//TODO: GPU Writes. 
-	//TODO: ADD VARIADIC PARAM HANDLING SIMILAR TO READ DI
+	TArray<FString> CommonTemplateShaders;
+	TMap<FName, FString> TemplateShaderMap;
+	NDIDataChannelWriteLocal::BuildFunctionTemplateMap(CommonTemplateShaders, TemplateShaderMap);
+
+	NDIDataChannelUtilities::GenerateDataChannelAccessHlsl(HlslGenContext, CommonTemplateShaders, TemplateShaderMap, OutHLSL);
 }
 
 #endif
 void UNiagaraDataInterfaceDataChannelWrite::BuildShaderParameters(FNiagaraShaderParametersBuilder& ShaderParametersBuilder) const
 {
-	//ShaderParametersBuilder.AddNestedStruct<FShaderParameters>();
+	ShaderParametersBuilder.AddNestedStruct<NDIDataChannelWriteLocal::FShaderParameters>();
 }
+
 void UNiagaraDataInterfaceDataChannelWrite::SetShaderParameters(const FNiagaraDataInterfaceSetShaderParametersContext& Context) const
 {
+	FNiagaraDataInterfaceProxy_DataChannelWrite& DataInterfaceProxy = Context.GetProxy<FNiagaraDataInterfaceProxy_DataChannelWrite>();
+	FNiagaraDataInterfaceProxy_DataChannelWrite::FInstanceData* InstanceData = DataInterfaceProxy.SystemInstancesToProxyData_RT.Find(Context.GetSystemInstanceID());
+
+	NDIDataChannelWriteLocal::FShaderParameters* InstParameters = Context.GetParameterNestedStruct<NDIDataChannelWriteLocal::FShaderParameters>();
+
+	bool bSuccess = false;
+	if (InstanceData)
+	{
+		//Find the start offset in the parameter table for this script.
+		const FNiagaraCompileHash& ScriptCompileHash = Context.GetComputeInstanceData().Context->GPUScript_RT->GetBaseCompileHash();
+		uint32 ParameterOffsetTableIndex = INDEX_NONE;
+		if (uint32* ParameterOffsetTableIndexPtr = InstanceData->GPUScriptParameterTableOffsets.Find(ScriptCompileHash))
+		{
+			ParameterOffsetTableIndex = *ParameterOffsetTableIndexPtr;
+		}
+
+		if (InstanceData->ChannelDataRTProxy && ParameterOffsetTableIndex != INDEX_NONE)
+		{
+			FNiagaraDataBuffer* GPUBuffer = InstanceData->bPublishToGPU ? InstanceData->ChannelDataRTProxy->GetCurrentData() : nullptr;
+			FNiagaraDataBuffer* BufferForCPU = InstanceData->bPublishToCPU || InstanceData->bPublishToGame ? InstanceData->BufferForCPU.GetReference() : nullptr;
+			if (GPUBuffer || BufferForCPU)
+			{
+				TArray<FRHITransitionInfo, TInlineAllocator<3>> Transitions;
+				Transitions.Reserve(6);
+				if(GPUBuffer)
+				{
+					Transitions.Emplace(GPUBuffer->GetGPUBufferFloat().UAV, ERHIAccess::Unknown, ERHIAccess::UAVCompute);
+					Transitions.Emplace(GPUBuffer->GetGPUBufferInt().UAV, ERHIAccess::Unknown, ERHIAccess::UAVCompute);
+					Transitions.Emplace(GPUBuffer->GetGPUBufferHalf().UAV, ERHIAccess::Unknown, ERHIAccess::UAVCompute);
+				}
+				if(BufferForCPU)
+				{
+					Transitions.Emplace(BufferForCPU->GetGPUBufferFloat().UAV, ERHIAccess::Unknown, ERHIAccess::UAVCompute);
+					Transitions.Emplace(BufferForCPU->GetGPUBufferInt().UAV, ERHIAccess::Unknown, ERHIAccess::UAVCompute);
+					Transitions.Emplace(BufferForCPU->GetGPUBufferHalf().UAV, ERHIAccess::Unknown, ERHIAccess::UAVCompute);
+				}
+				Context.GetGraphBuilder().RHICmdList.Transition(Transitions);
+
+				const FReadBuffer& ParameterLayoutBuffer = InstanceData->ParameterLayoutBuffer;
+
+				if (ParameterLayoutBuffer.SRV.IsValid() && ParameterLayoutBuffer.NumBytes > 0)
+				{
+					InstParameters->ParamOffsetTable = ParameterLayoutBuffer.SRV.IsValid() ? ParameterLayoutBuffer.SRV.GetReference() : FNiagaraRenderer::GetDummyUIntBuffer();
+					InstParameters->ParameterOffsetTableIndex = ParameterOffsetTableIndex;
+					InstParameters->FloatStride = (GPUBuffer ? GPUBuffer->GetFloatStride() : sizeof(float)) / sizeof(float);
+					InstParameters->Int32Stride = (GPUBuffer ? GPUBuffer->GetInt32Stride() : sizeof(int32)) / sizeof(int32);
+					InstParameters->HalfStride = (GPUBuffer ? GPUBuffer->GetHalfStride() : sizeof(FFloat16)) / sizeof(FFloat16);
+					
+					InstParameters->GPUBufferFloat = GPUBuffer && GPUBuffer->GetGPUBufferFloat().UAV.IsValid() ? GPUBuffer->GetGPUBufferFloat().UAV : NDIDataChannelUtilities::GetDummyUAVFloat().Buffer.UAV;
+					InstParameters->GPUBufferInt32 = GPUBuffer && GPUBuffer->GetGPUBufferInt().UAV.IsValid() ? GPUBuffer->GetGPUBufferInt().UAV : NDIDataChannelUtilities::GetDummyUAVInt32().Buffer.UAV;
+					InstParameters->GPUBufferHalf = GPUBuffer && GPUBuffer->GetGPUBufferHalf().UAV.IsValid() ? GPUBuffer->GetGPUBufferHalf().UAV : NDIDataChannelUtilities::GetDummyUAVHalf().Buffer.UAV;
+					InstParameters->GPUInstanceCountOffset = GPUBuffer ? GPUBuffer->GetGPUInstanceCountBufferOffset() : INDEX_NONE;
+					InstParameters->GPUBufferSize = GPUBuffer ? GPUBuffer->GetNumInstancesAllocated() : INDEX_NONE;
+
+
+					InstParameters->CPUBufferFloat = BufferForCPU && BufferForCPU->GetGPUBufferFloat().UAV.IsValid() ? BufferForCPU->GetGPUBufferFloat().UAV : NDIDataChannelUtilities::GetDummyUAVFloat().Buffer.UAV;
+					InstParameters->CPUBufferInt32 = BufferForCPU && BufferForCPU->GetGPUBufferInt().UAV.IsValid() ? BufferForCPU->GetGPUBufferInt().UAV : NDIDataChannelUtilities::GetDummyUAVInt32().Buffer.UAV;
+					InstParameters->CPUBufferHalf = BufferForCPU && BufferForCPU->GetGPUBufferHalf().UAV.IsValid() ? BufferForCPU->GetGPUBufferHalf().UAV : NDIDataChannelUtilities::GetDummyUAVHalf().Buffer.UAV;
+					InstParameters->CPUInstanceCountOffset = BufferForCPU ? BufferForCPU->GetGPUInstanceCountBufferOffset() : INDEX_NONE;
+					InstParameters->CPUBufferSize = BufferForCPU ? BufferForCPU->GetNumInstancesAllocated() : INDEX_NONE;
+
+					InstParameters->CPUFloatStride = (BufferForCPU ? BufferForCPU->GetFloatStride() : sizeof(float)) / sizeof(float);
+					InstParameters->CPUInt32Stride = (BufferForCPU ? BufferForCPU->GetInt32Stride() : sizeof(int32)) / sizeof(int32);
+					InstParameters->CPUHalfStride = (BufferForCPU ? BufferForCPU->GetHalfStride() : sizeof(FFloat16)) / sizeof(FFloat16);
+
+					bSuccess = true;
+				}
+			}
+		}
+	}
+
+	if (bSuccess == false)
+	{
+		InstParameters->ParamOffsetTable = FNiagaraRenderer::GetDummyUIntBuffer();
+		InstParameters->ParameterOffsetTableIndex = INDEX_NONE;
+		InstParameters->FloatStride = 0;
+		InstParameters->Int32Stride = 0;
+		InstParameters->HalfStride = 0;
+
+		InstParameters->GPUBufferFloat = NDIDataChannelUtilities::GetDummyUAVFloat().Buffer.UAV;
+		InstParameters->GPUBufferInt32 = NDIDataChannelUtilities::GetDummyUAVInt32().Buffer.UAV;
+		InstParameters->GPUBufferHalf = NDIDataChannelUtilities::GetDummyUAVHalf().Buffer.UAV;
+		InstParameters->GPUInstanceCountOffset = INDEX_NONE;
+		InstParameters->GPUBufferSize = INDEX_NONE;
+
+
+		InstParameters->CPUBufferFloat = NDIDataChannelUtilities::GetDummyUAVFloat().Buffer.UAV;
+		InstParameters->CPUBufferInt32 = NDIDataChannelUtilities::GetDummyUAVInt32().Buffer.UAV;
+		InstParameters->CPUBufferHalf = NDIDataChannelUtilities::GetDummyUAVHalf().Buffer.UAV;
+		InstParameters->CPUInstanceCountOffset = INDEX_NONE;
+		InstParameters->CPUBufferSize = INDEX_NONE;
+	}
+}
+
+void FNiagaraDataInterfaceProxy_DataChannelWrite::ConsumePerInstanceDataFromGameThread(void* PerInstanceData, const FNiagaraSystemInstanceID& Instance)
+{
+	FNDIDataChannelWriteInstanceData_RT& SourceData = *reinterpret_cast<FNDIDataChannelWriteInstanceData_RT*>(PerInstanceData);
+	FInstanceData& InstData = SystemInstancesToProxyData_RT.FindOrAdd(Instance);
+
+	FRHICommandListImmediate& RHICmdList = FRHICommandListImmediate::Get();
+
+	if(SourceData.ChannelDataRTProxy)
+	{
+		InstData.ChannelDataRTProxy = SourceData.ChannelDataRTProxy;
+		
+		if(SourceData.bPublishToGPU)
+		{
+			//Allocate space in the GPU buffers for writing...
+			InstData.ChannelDataRTProxy->AddGPUAllocationForNextTick(SourceData.AllocationCount);
+		}
+
+		InstData.AllocationCount = SourceData.AllocationCount;
+		InstData.bPublishToGame = SourceData.bPublishToGame;
+		InstData.bPublishToCPU = SourceData.bPublishToCPU;
+		InstData.bPublishToGPU = SourceData.bPublishToGPU;
+		InstData.LwcTile = SourceData.LwcTile;
+	}
+
+	if (SourceData.ScriptParamInfo.bDirty)
+	{
+		SourceData.ScriptParamInfo.bDirty = false;
+
+		//Take the offset map from the source data.
+		//This maps from GPU script to that scripts offset into the ParameterLayoutBuffer.
+		//Allows us to look up and pass in at SetShaderParameters time.
+		InstData.GPUScriptParameterTableOffsets = MoveTemp(SourceData.ScriptParamInfo.GPUScriptParameterTableOffsets);
+
+		//Now generate the ParameterLayoutBuffer
+		//This contains a table of all parameters used by each GPU script that uses this DI.
+		//TODO: This buffer can likely be shared among many instances and stored in the layout manager or in the DI proxy.
+		{
+			if (InstData.ParameterLayoutBuffer.NumBytes > 0)
+			{
+				InstData.ParameterLayoutBuffer.Release();
+			}
+
+			if (SourceData.ScriptParamInfo.GPUScriptParameterOffsetTable.Num() > 0)
+			{
+				InstData.ParameterLayoutData = SourceData.ScriptParamInfo.GPUScriptParameterOffsetTable;
+				InstData.ParameterLayoutBuffer.Initialize(RHICmdList, TEXT("NDIDataChannel_ParameterLayoutBuffer"), sizeof(uint32), SourceData.ScriptParamInfo.GPUScriptParameterOffsetTable.Num(), EPixelFormat::PF_R32_UINT, BUF_Static, &InstData.ParameterLayoutData);
+			}
+		}
+	}
+
+	SourceData.~FNDIDataChannelWriteInstanceData_RT();
+}
+
+int32 FNiagaraDataInterfaceProxy_DataChannelWrite::PerInstanceDataPassedToRenderThreadSize() const
+{
+	return sizeof(FNDIDataChannelWriteInstanceData_RT);
+}
+
+void FNiagaraDataInterfaceProxy_DataChannelWrite::PreStage(const FNDIGpuComputePreStageContext& Context)
+{
+	FNiagaraDataInterfaceProxy_DataChannelWrite::FInstanceData* InstanceData = SystemInstancesToProxyData_RT.Find(Context.GetSystemInstanceID());
+
+	if (InstanceData && InstanceData->ChannelDataRTProxy)
+	{
+		if(InstanceData->bPublishToGame || InstanceData->bPublishToCPU)
+		{
+			//Allocate a separate buffer that we will write into and ship back to the CPU.			 
+			InstanceData->BufferForCPU = InstanceData->ChannelDataRTProxy->AllocateBufferForCPU(Context.GetGraphBuilder().RHICmdList, Context.GetComputeDispatchInterface().GetFeatureLevel(), InstanceData->AllocationCount);
+
+			//Get a new instance count.
+			uint32 Offset = InstanceData->BufferForCPU->GetGPUInstanceCountBufferOffset();
+			Context.GetInstanceCountManager().FreeEntry(Offset);
+			InstanceData->BufferForCPU->SetGPUInstanceCountBufferOffset(Context.GetInstanceCountManager().AcquireOrAllocateEntry(Context.GetGraphBuilder().RHICmdList));
+		}
+	}
+}
+
+void FNiagaraDataInterfaceProxy_DataChannelWrite::PostStage(const FNDIGpuComputePostStageContext& Context)
+{
+	FNiagaraDataInterfaceProxy_DataChannelWrite::FInstanceData* InstanceData = SystemInstancesToProxyData_RT.Find(Context.GetSystemInstanceID());
+
+	if (InstanceData && InstanceData->ChannelDataRTProxy)
+	{
+		if (InstanceData->BufferForCPU)
+		{
+			Context.GetGraphBuilder().AddPass(
+				RDG_EVENT_NAME("NDC Write Enqueue CPU Readback"),
+				ERDGPassFlags::None,
+				[PASS_ReadbackManager = Context.GetComputeDispatchInterface().GetGpuReadbackManager(), 
+				PASS_InstCountManager = &Context.GetInstanceCountManager(), 
+				PASS_Buffer = InstanceData->BufferForCPU,
+				PASS_Proxy = InstanceData->ChannelDataRTProxy,
+				PASS_bPublishToGame = InstanceData->bPublishToGame,
+				PASS_bPublishToCPU = InstanceData->bPublishToCPU,
+				PASS_LwcTile = InstanceData->LwcTile]
+				(FRHICommandList& RHICmdList)
+				{
+					PASS_Proxy->EnqueueReadbackForCPUBuffer(RHICmdList, PASS_Buffer, PASS_ReadbackManager, *PASS_InstCountManager, PASS_bPublishToGame, PASS_bPublishToCPU, PASS_LwcTile);
+					uint32 Offset = PASS_Buffer->GetGPUInstanceCountBufferOffset();
+					PASS_InstCountManager->FreeEntry(Offset);
+					PASS_Buffer->SetGPUInstanceCountBufferOffset(INDEX_NONE);
+				});
+		
+			InstanceData->BufferForCPU = nullptr;
+		}
+	}
 }
 
 #undef LOCTEXT_NAMESPACE
diff --git a/Engine/Plugins/FX/Niagara/Source/Niagara/Private/NiagaraCommon.cpp b/Engine/Plugins/FX/Niagara/Source/Niagara/Private/NiagaraCommon.cpp
--- a/Engine/Plugins/FX/Niagara/Source/Niagara/Private/NiagaraCommon.cpp
+++ b/Engine/Plugins/FX/Niagara/Source/Niagara/Private/NiagaraCommon.cpp
@@ -1696,19 +1696,24 @@ EPSCPoolMethod ToPSCPoolMethod(ENCPoolMethod PoolingMethod)
 //////////////////////////////////////////////////////////////////////////
 
 
-void FNiagaraFunctionSignature::GetVariadicInputs(TArray<FNiagaraVariableBase>& OutVariadicInputs, bool bStripNonExecution/* =true */)const
+void FNiagaraFunctionSignature::GetVariadicInputs(TArray<FNiagaraVariableBase>& OutVariadicInputs)const
 {
+	if(!VariadicInput())
+	{
+		return;
+	}
+
 	static const FNiagaraVariableBase InstDataVar(FNiagaraTypeDefinition::GetIntDef(), TEXT("InstanceData"));
 	OutVariadicInputs.Reset(NumOptionalInputs());
 	int32 NumInputs = 0;
 	for (const FNiagaraVariableBase& Param : Inputs)
 	{
-		if (bStripNonExecution && (Param.GetType() == FNiagaraTypeDefinition::GetParameterMapDef() || Param == InstDataVar))
+		if (Param == InstDataVar)
 		{
 			continue;
 		}
 
-		if (NumInputs++ < NumRequiredInputs())
+		if (NumInputs++ < VariadicInputStartIndex())
 		{
 			continue;
 		}
@@ -1717,18 +1722,18 @@ void FNiagaraFunctionSignature::GetVariadicInputs(TArray<FNiagaraVariableBase>&
 	}
 }
 
-void FNiagaraFunctionSignature::GetVariadicOutputs(TArray<FNiagaraVariableBase>& OutVariadicOutputs, bool bStripNonExecution/* =true */)const
+void FNiagaraFunctionSignature::GetVariadicOutputs(TArray<FNiagaraVariableBase>& OutVariadicOutputs)const
 {
+	if(!VariadicOutput())
+	{
+		return;
+	}
+
 	OutVariadicOutputs.Reset(NumOptionalOutputs());
 	int32 NumOutputs = 0;
 	for (const FNiagaraVariableBase& Param : Outputs)
 	{
-		if (bStripNonExecution && (Param.GetType() == FNiagaraTypeDefinition::GetParameterMapDef()))
-		{
-			continue;
-		}
-
-		if (NumOutputs++ < NumRequiredOutputs())
+		if (NumOutputs++ < VariadicOutputStartIndex())
 		{
 			continue;
 		}
diff --git a/Engine/Plugins/FX/Niagara/Source/Niagara/Private/NiagaraDataChannel.cpp b/Engine/Plugins/FX/Niagara/Source/Niagara/Private/NiagaraDataChannel.cpp
--- a/Engine/Plugins/FX/Niagara/Source/Niagara/Private/NiagaraDataChannel.cpp
+++ b/Engine/Plugins/FX/Niagara/Source/Niagara/Private/NiagaraDataChannel.cpp
@@ -17,6 +17,7 @@
 #include "Serialization/JsonSerializer.h"
 #include "Dom/JsonObject.h"
 #endif
+#include "NiagaraDataSetReadback.h"
 
 #include UE_INLINE_GENERATED_CPP_BY_NAME(NiagaraDataChannel)
 
@@ -455,70 +456,131 @@ void FNiagaraDataChannelGameData::SetFromSimCache(const FNiagaraVariableBase& So
 
 //////////////////////////////////////////////////////////////////////////
 
-void FNiagaraDataChannelDataProxy::BeginFrame(bool bKeepPreviousFrameData)
+FNiagaraDataChannelDataProxy::~FNiagaraDataChannelDataProxy()
 {
-	check(GPUDataSet);
-	check(GPUDataSet->GetSimTarget() == ENiagaraSimTarget::GPUComputeSim);
+	check(DispatchInterface);
+	DispatchInterface->RemoveNDCDataProxy(this);
+}
+
+void FNiagaraDataChannelDataProxy::Reset()
+{
+	PrevFrameData = nullptr;
+}
 
+void FNiagaraDataChannelDataProxy::Init()
+{
+	check(DispatchInterface);
+	DispatchInterface->AddNDCDataProxy(this);
 }
 
-void FNiagaraDataChannelDataProxy::EndFrame(FNiagaraGpuComputeDispatchInterface* DispathInterface, FRHICommandListImmediate& CmdList, const TArray<FNiagaraDataBufferRef>& BuffersForGPU)
+void FNiagaraDataChannelDataProxy::BeginFrame(FRHICommandListImmediate& RHICmdList)
 {
 	check(GPUDataSet);
 	check(GPUDataSet->GetSimTarget() == ENiagaraSimTarget::GPUComputeSim);
 
-	PrevFrameData = GPUDataSet->GetCurrentData();
+	FNiagaraGPUInstanceCountManager& InstCountManager = DispatchInterface->GetGPUInstanceCounterManager();	
+	if(this->bNeedsPrevFrameData)
+	{
+		Swap(PrevFrameData, CurrFrameData);
+	}
 
-	uint32 NumInstance = 0;
-	for (auto& Buffer : BuffersForGPU)
+	//The base num instances for writing from the GPU.
+	uint32 NumInstanceFromCPU = 0;
+	for (auto& Buffer : PendingCPUBuffers)
 	{
-		NumInstance += Buffer->GetNumInstances();
+		NumInstanceFromCPU += Buffer->GetNumInstances();
 	}
 
-	GPUDataSet->BeginSimulate();
-	FNiagaraDataBuffer* DestBuffer = GPUDataSet->GetDestinationData();	
-	DestBuffer->PushCPUBuffersToGPU(BuffersForGPU, true, CmdList, DispathInterface->GetFeatureLevel(), GetDebugName());
-	GPUDataSet->EndSimulate();
+	//Allocate our GPU buffers we'll write into. Combine accumulated counts from writing DIs and the data coming in from the CPU.
+	uint32 InstancesToAllocate = NumInstanceFromCPU + PendingGPUAllocations;
+	PendingGPUAllocations = 0;
 
-	//For now we need not deal with the instance count manager but when we do GPU->GPU writes we will
-	//DestBuffer->SetGPUInstanceCountBufferOffset()
+	if(InstancesToAllocate == 0)
+	{
+		return;
+	}
+
+	if(CurrFrameData == nullptr)
+	{		
+		FNiagaraDataBuffer& NewBuffer = GPUDataSet->AllocateBuffer();
+		CurrFrameData = NewBuffer.UnlockForRead();
+	}
+
+	uint32 InstanceCountOffset = CurrFrameData->GetGPUInstanceCountBufferOffset();
+	InstCountManager.FreeEntry(InstanceCountOffset);
+
+	CurrFrameData->AllocateGPU(RHICmdList, InstancesToAllocate, DispatchInterface->GetFeatureLevel(), GetDebugName());
+	CurrFrameData->PushCPUBuffersToGPU(PendingCPUBuffers, true, RHICmdList, DispatchInterface->GetFeatureLevel(), GetDebugName(), false);
+	CurrFrameData->SetNumInstances(InstancesToAllocate);
+	PendingCPUBuffers.Reset();
+
+	InstanceCountOffset = InstCountManager.AcquireOrAllocateEntry(RHICmdList);	
+	CurrFrameData->SetGPUInstanceCountBufferOffset(InstanceCountOffset);
+
+	//Init the instance count to the value from the CPU before it is added to by the GPU writes.
+	if(NumInstanceFromCPU > 0)
+	{
+		InstCountManager.AddInstanceCountInitTask(InstanceCountOffset, NumInstanceFromCPU);
+	}
 }
 
-void FNiagaraDataChannelDataProxy::Reset()
+void FNiagaraDataChannelDataProxy::EndFrame(FRHICommandListImmediate& RHICmdList)
 {
-	PrevFrameData = nullptr;
+	check(GPUDataSet);
+	check(GPUDataSet->GetSimTarget() == ENiagaraSimTarget::GPUComputeSim);
 }
 
-//////////////////////////////////////////////////////////////////////////
+FNiagaraDataBufferRef FNiagaraDataChannelDataProxy::AllocateBufferForCPU(FRHICommandListImmediate& RHICmdList, ERHIFeatureLevel::Type FeatureLevel, int32 AllocationSize)
+{
+	FNiagaraDataBuffer* Ret = &GPUDataSet->AllocateBuffer();
+	Ret->AllocateGPU(RHICmdList, AllocationSize, FeatureLevel, TEXT("NDC GPU Buffers for CPU"));	
+	return Ret->UnlockForRead();
+}
 
-FNiagaraDataChannelData::FNiagaraDataChannelData(UNiagaraDataChannelHandler* Owner)
+void FNiagaraDataChannelDataProxy::EnqueueReadbackForCPUBuffer(FRHICommandList& RHICmdList, FNiagaraDataBufferRef Buffer, FNiagaraGpuReadbackManager* ReadbackManager, FNiagaraGPUInstanceCountManager& InstanceCountManager, bool bPublishToGame, bool bPublishToCPU, FVector3f LWCTile)
 {
-	const UNiagaraDataChannel* DataChannel = Owner->GetDataChannel();
-	check(DataChannel);
+	auto PublishOnCPU = [=, LocalOwner = Owner](FNiagaraDataBufferRef ReadbackBuffer)
+	{
+		if (FNiagaraDataChannelData* NDCData = LocalOwner.Pin().Get())
+		{
+			if (ReadbackBuffer && ReadbackBuffer->GetNumInstances())
+			{
+				FNiagaraDataChannelPublishRequest PublishRequest(ReadbackBuffer);
+				PublishRequest.bVisibleToGame = bPublishToGame;
+				PublishRequest.bVisibleToCPUSims = bPublishToCPU;
+				PublishRequest.bVisibleToGPUSims = false;//Don't want to ping pong back to the GPU
+				PublishRequest.LwcTile = LWCTile;
+#if WITH_NIAGARA_DEBUGGER
+				PublishRequest.DebugSource = TEXT("NDC GPU Readback");//TODO: Feed in real source names also
+#endif
+				NDCData->PublishFromGPU(PublishRequest);
+			}
+		}
+	};
 
-	GameData = DataChannel->CreateGameData();
+	//We enqueue a readback for the data if we're wanting to publish it to the CPU or Game
+	check(bPublishToCPU || bPublishToGame);//We should be doing one or the other or we should not be calling this function.
 
-	CPUSimData = new FNiagaraDataSet();
-	GPUSimData = new FNiagaraDataSet();
+	TSharedPtr<FNiagaraDataBufferReadback, ESPMode::ThreadSafe> NewReadback = MakeShared<FNiagaraDataBufferReadback, ESPMode::ThreadSafe>();
+	NewReadback->GetOnReadbackComplete().BindLambda(PublishOnCPU);
+	NewReadback->EnqueueReadback(RHICmdList, Buffer, ReadbackManager, InstanceCountManager);
+}
 
+void FNiagaraDataChannelDataProxy::AddBuffersFromCPU(const TArray<FNiagaraDataBufferRef>& BuffersFromCPU)
+{
+	PendingCPUBuffers.Append(BuffersFromCPU);
+}
 
-	const FNiagaraDataSetCompiledData& CompiledData = DataChannel->GetCompiledData(ENiagaraSimTarget::CPUSim);
-	const FNiagaraDataSetCompiledData& CompiledDataGPU = DataChannel->GetCompiledData(ENiagaraSimTarget::GPUComputeSim);
-	CPUSimData->Init(&CompiledData);
-	GPUSimData->Init(&CompiledDataGPU);
+void FNiagaraDataChannelDataProxy::AddGPUAllocationForNextTick(int32 AllocationCount)
+{
+	PendingGPUAllocations += AllocationCount;
+}
 
-	//TODO: Send game data to GPU direct without staging.
-	GameDataStaging = new FNiagaraDataSet();
-	GameDataStaging->Init(&CompiledData);
+//////////////////////////////////////////////////////////////////////////
 
-	CPUSimData->BeginSimulate();
-	CPUSimData->EndSimulate();
+FNiagaraDataChannelData::FNiagaraDataChannelData()
+{
 
-	RTProxy.Reset(new FNiagaraDataChannelDataProxy());
-	RTProxy->GPUDataSet = GPUSimData;
-	#if !UE_BUILD_SHIPPING
-	RTProxy->DebugName = FString::Printf(TEXT("%s__GPUData"), *DataChannel->GetName());
-	#endif
 }
 
 FNiagaraDataChannelData::~FNiagaraDataChannelData()
@@ -547,6 +609,45 @@ FNiagaraDataChannelData::~FNiagaraDataChannelData()
 	GPUSimData = nullptr;
 }
 
+void FNiagaraDataChannelData::Init(UNiagaraDataChannelHandler* Owner)
+{
+	check(Owner);
+	const UNiagaraDataChannel* DataChannel = Owner->GetDataChannel();
+	check(DataChannel);
+
+	GameData = DataChannel->CreateGameData();
+
+	CPUSimData = new FNiagaraDataSet();
+	GPUSimData = new FNiagaraDataSet();
+
+
+	const FNiagaraDataSetCompiledData& CompiledData = DataChannel->GetCompiledData(ENiagaraSimTarget::CPUSim);
+	const FNiagaraDataSetCompiledData& CompiledDataGPU = DataChannel->GetCompiledData(ENiagaraSimTarget::GPUComputeSim);
+	CPUSimData->Init(&CompiledData, 1);
+	GPUSimData->Init(&CompiledDataGPU, 2);
+
+	//TODO: Send game data to GPU direct without staging.
+	GameDataStaging = new FNiagaraDataSet();
+	GameDataStaging->Init(&CompiledData);
+
+	if (FNiagaraGpuComputeDispatchInterface* ComputeDispatchInterface = FNiagaraGpuComputeDispatchInterface::Get(Owner->GetWorld()))
+	{
+		RTProxy.Reset(new FNiagaraDataChannelDataProxy());
+		RTProxy->GPUDataSet = GPUSimData;
+		RTProxy->bNeedsPrevFrameData = DataChannel->KeepPreviousFrameData();
+		RTProxy->DispatchInterface = FNiagaraGpuComputeDispatchInterface::Get(Owner->GetWorld());
+		RTProxy->Owner = this->AsWeak();
+#if !UE_BUILD_SHIPPING
+		RTProxy->DebugName = FString::Printf(TEXT("%s__GPUData"), *DataChannel->GetName());
+#endif
+		ENQUEUE_RENDER_COMMAND(FNiagaraDataChannelDataProxyInit) (
+			[RT_Proxy = RTProxy.Get()](FRHICommandListImmediate& CmdList)
+			{
+				RT_Proxy->Init();
+			});
+	}
+}
+
 void FNiagaraDataChannelData::Reset()
 {
 	FScopeLock Lock(&PublishCritSec);
@@ -583,16 +684,10 @@ void FNiagaraDataChannelData::BeginFrame(UNiagaraDataChannelHandler* Owner)
 	//Grab a new buffer to store the CPU data.
 	CPUSimData->BeginSimulate();
 	CPUSimData->EndSimulate();
-	
-	check(RTProxy);
-	if (FNiagaraGpuComputeDispatchInterface* DispathInterface = FNiagaraGpuComputeDispatchInterface::Get(Owner->GetWorld()))
-	{
-		ENQUEUE_RENDER_COMMAND(FDataChannelProxyBeginFrame) (
-			[RT_Proxy = RTProxy.Get(), bRequirePreviousData](FRHICommandListImmediate& CmdList)
-		{
-			RT_Proxy->BeginFrame(bRequirePreviousData);
-		});
-	}
+
+	//Pull in our publish requests from the GPU to be processed in the first tick group.
+	PublishRequests.Append(PublishRequestsFromGPU);
+	PublishRequestsFromGPU.Reset();
 }
 
 void FNiagaraDataChannelData::EndFrame(UNiagaraDataChannelHandler* Owner)
@@ -600,14 +695,13 @@ void FNiagaraDataChannelData::EndFrame(UNiagaraDataChannelHandler* Owner)
 	//We must do one final tick to process any final items generated by the last things to tick this frame.
 	ConsumePublishRequests(Owner, TG_LastDemotable);
 
-	check(RTProxy);
-	if (FNiagaraGpuComputeDispatchInterface* DispathInterface = FNiagaraGpuComputeDispatchInterface::Get(Owner->GetWorld()))
+	if(RTProxy)
 	{
 		ENQUEUE_RENDER_COMMAND(FDataChannelProxyEndFrame) (
-			[DispathInterface, RT_Proxy = RTProxy.Get(), RT_BuffersForGPU = MoveTemp(BuffersForGPU)](FRHICommandListImmediate& CmdList)
-		{
-			RT_Proxy->EndFrame(DispathInterface, CmdList, RT_BuffersForGPU);
-		});
+			[RT_Proxy = RTProxy.Get(), RT_BuffersForGPU = MoveTemp(BuffersForGPU)](FRHICommandListImmediate& CmdList)
+			{
+				RT_Proxy->AddBuffersFromCPU(RT_BuffersForGPU);
+			});
 	}
 	BuffersForGPU.Reset();
 }
@@ -720,25 +814,13 @@ int32 FNiagaraDataChannelData::ConsumePublishRequests(UNiagaraDataChannelHandler
 
 		if (PublishRequest.Data)
 		{
-			const ENiagaraSimTarget SimTarget = PublishRequest.Data->GetOwner()->GetSimTarget();
-			if (SimTarget == ENiagaraSimTarget::CPUSim)
-			{
-				if (PublishRequest.bVisibleToGame)
-				{
-					GameData->AppendFromDataSet(PublishRequest.Data, PublishRequest.LwcTile);
-				}
-				if (PublishRequest.bVisibleToCPUSims)
-				{
-					PublishRequest.Data->CopyToUnrelated(CPUSimData->GetDestinationDataChecked(), 0, CPUSimData->GetDestinationDataChecked().GetNumInstances(), PublishRequest.Data->GetNumInstances());
-				}
-			}
-			else if (SimTarget == ENiagaraSimTarget::GPUComputeSim)
+			if (PublishRequest.bVisibleToGame)
 			{
-				//TODO: GPU->GPU handling will be done all RT side. GPU->CPU may be done here in future but we'll have to pull the data from the GPU on the RT and pass beck to GT here.
+				GameData->AppendFromDataSet(PublishRequest.Data, PublishRequest.LwcTile);
 			}
-			else
+			if (PublishRequest.bVisibleToCPUSims)
 			{
-				check(0);
+				PublishRequest.Data->CopyToUnrelated(CPUSimData->GetDestinationDataChecked(), 0, CPUSimData->GetDestinationDataChecked().GetNumInstances(), PublishRequest.Data->GetNumInstances());
 			}
 		}
 #if WITH_NIAGARA_DEBUGGER
@@ -791,6 +873,12 @@ void FNiagaraDataChannelData::Publish(const FNiagaraDataChannelPublishRequest& R
 	PublishRequests.Add(Request);
 }
 
+void FNiagaraDataChannelData::PublishFromGPU(const FNiagaraDataChannelPublishRequest& Request)
+{
+	check(IsInGameThread());
+	PublishRequestsFromGPU.Add(Request);
+}
+
 void FNiagaraDataChannelData::RemovePublishRequests(const FNiagaraDataSet* DataSet)
 {
 	FScopeLock Lock(&PublishCritSec);
@@ -822,6 +910,16 @@ const FNiagaraDataSetCompiledData& FNiagaraDataChannelData::GetCompiledData(ENia
 	return SimTarget == ENiagaraSimTarget::CPUSim ? CPUSimData->GetCompiledData() : GPUSimData->GetCompiledData();
 }
 
+FNiagaraDataBuffer* FNiagaraDataChannelData::GetBufferForCPUWrite()
+{
+	check(IsInGameThread());
+	if(CPUSimData)
+	{
+		return &CPUSimData->AllocateBuffer();
+	}
+	return nullptr;
+}
+
 //////////////////////////////////////////////////////////////////////////
 
 void UNiagaraDataChannel::PostInitProperties()
diff --git a/Engine/Plugins/FX/Niagara/Source/Niagara/Private/NiagaraDataChannelAccessor.cpp b/Engine/Plugins/FX/Niagara/Source/Niagara/Private/NiagaraDataChannelAccessor.cpp
--- a/Engine/Plugins/FX/Niagara/Source/Niagara/Private/NiagaraDataChannelAccessor.cpp
+++ b/Engine/Plugins/FX/Niagara/Source/Niagara/Private/NiagaraDataChannelAccessor.cpp
@@ -14,6 +14,11 @@
 
 //////////////////////////////////////////////////////////////////////////
 
+void UNiagaraDataChannelReader::Cleanup()
+{
+	Data = nullptr;
+}
+
 bool UNiagaraDataChannelReader::InitAccess(FNiagaraDataChannelSearchParameters SearchParams, bool bReadPreviousFrameData)
 {
 	Data = nullptr;
@@ -130,6 +135,11 @@ FNiagaraSpawnInfo UNiagaraDataChannelReader::ReadSpawnInfo(FName VarName, int32
 
 //////////////////////////////////////////////////////////////////////////
 
+void UNiagaraDataChannelWriter::Cleanup()
+{
+	Data = nullptr;
+}
+
 bool UNiagaraDataChannelWriter::InitWrite(FNiagaraDataChannelSearchParameters SearchParams, int32 Count, bool bVisibleToGame, bool bVisibleToCPU, bool bVisibleToGPU, const FString& DebugSource)
 {
 	if (Count == 0)
diff --git a/Engine/Plugins/FX/Niagara/Source/Niagara/Private/NiagaraDataChannelHandler.cpp b/Engine/Plugins/FX/Niagara/Source/Niagara/Private/NiagaraDataChannelHandler.cpp
--- a/Engine/Plugins/FX/Niagara/Source/Niagara/Private/NiagaraDataChannelHandler.cpp
+++ b/Engine/Plugins/FX/Niagara/Source/Niagara/Private/NiagaraDataChannelHandler.cpp
@@ -4,8 +4,10 @@
 #include "NiagaraDataChannelAccessor.h"
 #include "NiagaraDataChannelCommon.h"
 
-UNiagaraDataChannelHandler::~UNiagaraDataChannelHandler()
+void UNiagaraDataChannelHandler::BeginDestroy()
 {
+	Super::BeginDestroy();
+	Cleanup();
 }
 
 void UNiagaraDataChannelHandler::Init(const UNiagaraDataChannel* InChannel)
@@ -13,6 +15,21 @@ void UNiagaraDataChannelHandler::Init(const UNiagaraDataChannel* InChannel)
 	DataChannel = InChannel;
 }
 
+void UNiagaraDataChannelHandler::Cleanup()
+{
+	if(Reader)
+	{
+		Reader->Cleanup();
+		Reader = nullptr;
+	}
+	
+	if(Writer)
+	{
+		Writer->Cleanup();
+		Writer = nullptr;
+	}
+}
+
 void UNiagaraDataChannelHandler::BeginFrame(float DeltaTime, FNiagaraWorldManager* OwningWorld)
 {
 	CurrentTG = TG_PrePhysics;
@@ -50,5 +67,7 @@ UNiagaraDataChannelReader* UNiagaraDataChannelHandler::GetDataChannelReader()
 
 FNiagaraDataChannelDataPtr UNiagaraDataChannelHandler::CreateData()
 {
-	return MakeShared<FNiagaraDataChannelData>(this);
+	FNiagaraDataChannelDataPtr Ret = MakeShared<FNiagaraDataChannelData>();
+	Ret->Init(this);
+	return Ret;
 }
\ No newline at end of file
diff --git a/Engine/Plugins/FX/Niagara/Source/Niagara/Private/NiagaraDataChannelManager.cpp b/Engine/Plugins/FX/Niagara/Source/Niagara/Private/NiagaraDataChannelManager.cpp
--- a/Engine/Plugins/FX/Niagara/Source/Niagara/Private/NiagaraDataChannelManager.cpp
+++ b/Engine/Plugins/FX/Niagara/Source/Niagara/Private/NiagaraDataChannelManager.cpp
@@ -56,6 +56,10 @@ void FNiagaraDataChannelManager::Init()
 
 void FNiagaraDataChannelManager::Cleanup()
 {
+	for (auto& ChannelPair : Channels)
+	{
+		ChannelPair.Value->Cleanup();
+	}
 	Channels.Empty();
 	bIsCleanedUp = true;
 }
diff --git a/Engine/Plugins/FX/Niagara/Source/Niagara/Private/NiagaraDataChannel_Global.cpp b/Engine/Plugins/FX/Niagara/Source/Niagara/Private/NiagaraDataChannel_Global.cpp
--- a/Engine/Plugins/FX/Niagara/Source/Niagara/Private/NiagaraDataChannel_Global.cpp
+++ b/Engine/Plugins/FX/Niagara/Source/Niagara/Private/NiagaraDataChannel_Global.cpp
@@ -36,12 +36,6 @@ UNiagaraDataChannelHandler_Global::UNiagaraDataChannelHandler_Global(FObjectInit
 {
 }
 
-void UNiagaraDataChannelHandler_Global::BeginDestroy()
-{
-	Super::BeginDestroy();
-	Data.Reset();
-}
-
 void UNiagaraDataChannelHandler_Global::Init(const UNiagaraDataChannel* InChannel)
 {
 	check(InChannel);
@@ -49,6 +43,12 @@ void UNiagaraDataChannelHandler_Global::Init(const UNiagaraDataChannel* InChanne
 	Data = CreateData();
 }
 
+void UNiagaraDataChannelHandler_Global::Cleanup()
+{
+	Super::Cleanup();
+	Data.Reset();
+}
+
 void UNiagaraDataChannelHandler_Global::BeginFrame(float DeltaTime, FNiagaraWorldManager* OwningWorld)
 {
 	Super::BeginFrame(DeltaTime, OwningWorld);
diff --git a/Engine/Plugins/FX/Niagara/Source/Niagara/Private/NiagaraDataChannel_Islands.cpp b/Engine/Plugins/FX/Niagara/Source/Niagara/Private/NiagaraDataChannel_Islands.cpp
--- a/Engine/Plugins/FX/Niagara/Source/Niagara/Private/NiagaraDataChannel_Islands.cpp
+++ b/Engine/Plugins/FX/Niagara/Source/Niagara/Private/NiagaraDataChannel_Islands.cpp
@@ -110,14 +110,6 @@ UNiagaraDataChannelHandler_Islands::UNiagaraDataChannelHandler_Islands(FObjectIn
 {
 }
 
-void UNiagaraDataChannelHandler_Islands::BeginDestroy()
-{
-	Super::BeginDestroy();
-	ActiveIslands.Empty();
-	FreeIslands.Empty();
-	IslandPool.Empty();
-}
-
 void UNiagaraDataChannelHandler_Islands::Init(const UNiagaraDataChannel* InChannel)
 {
 	Super::Init(InChannel);
@@ -137,6 +129,14 @@ void UNiagaraDataChannelHandler_Islands::Init(const UNiagaraDataChannel* InChann
 	}
 }
 
+void UNiagaraDataChannelHandler_Islands::Cleanup()
+{
+	Super::Cleanup();
+	ActiveIslands.Empty();
+	FreeIslands.Empty();
+	IslandPool.Empty();
+}
+
 void UNiagaraDataChannelHandler_Islands::BeginFrame(float DeltaTime, FNiagaraWorldManager* OwningWorld)
 {
 	Super::BeginFrame(DeltaTime, OwningWorld);
diff --git a/Engine/Plugins/FX/Niagara/Source/Niagara/Private/NiagaraDataInterface.cpp b/Engine/Plugins/FX/Niagara/Source/Niagara/Private/NiagaraDataInterface.cpp
--- a/Engine/Plugins/FX/Niagara/Source/Niagara/Private/NiagaraDataInterface.cpp
+++ b/Engine/Plugins/FX/Niagara/Source/Niagara/Private/NiagaraDataInterface.cpp
@@ -326,11 +326,17 @@ FSimpleMulticastDelegate& UNiagaraDataInterface::OnErrorsRefreshed()
 
 #endif
 
+//////////////////////////////////////////////////////////////////////////
+//FNDIGpuComputeContext
+
 FRDGExternalAccessQueue& FNDIGpuComputeContext::GetRDGExternalAccessQueue() const
 {
 	return static_cast<const FNiagaraGpuComputeDispatch&>(ComputeDispatchInterface).GetCurrentPassExternalAccessQueue();
 }
 
+//////////////////////////////////////////////////////////////////////////
+//FNDIGpuComputePrePostStageContext
+
 FNiagaraSystemInstanceID FNDIGpuComputePrePostStageContext::GetSystemInstanceID() const
 {
 	return SystemTick.SystemInstanceID;
@@ -359,6 +365,21 @@ bool FNDIGpuComputePrePostStageContext::IsIterationStage() const
 	return ComputeInstanceData.IsIterationStage(DataInterfaceProxy, SimStageData.StageIndex);
 }
 
+FNiagaraGPUInstanceCountManager& FNDIGpuComputePrePostStageContext::GetInstanceCountManager() const
+{
+	return ComputeDispatchInterface.GetGPUInstanceCounterManager();
+}
+
+//////////////////////////////////////////////////////////////////////////
+//FNDIGpuComputePostSimulateContext
+
+FNiagaraGPUInstanceCountManager& FNDIGpuComputePostSimulateContext::GetInstanceCountManager() const
+{
+	return ComputeDispatchInterface.GetGPUInstanceCounterManager();
+}
+
+//////////////////////////////////////////////////////////////////////////
+
 FRDGExternalAccessQueue& FNiagaraDataInterfaceSetShaderParametersContext::GetRDGExternalAccessQueue() const
 {
 	return static_cast<const FNiagaraGpuComputeDispatch&>(ComputeDispatchInterface).GetCurrentPassExternalAccessQueue();
diff --git a/Engine/Plugins/FX/Niagara/Source/Niagara/Private/NiagaraDataSet.cpp b/Engine/Plugins/FX/Niagara/Source/Niagara/Private/NiagaraDataSet.cpp
--- a/Engine/Plugins/FX/Niagara/Source/Niagara/Private/NiagaraDataSet.cpp
+++ b/Engine/Plugins/FX/Niagara/Source/Niagara/Private/NiagaraDataSet.cpp
@@ -149,9 +149,15 @@ FNiagaraDataSet::~FNiagaraDataSet()
 	GPUNumAllocatedIDs = 0;
 }
 
-void FNiagaraDataSet::Init(const FNiagaraDataSetCompiledData* InDataSetCompiledData)
+void FNiagaraDataSet::Init(const FNiagaraDataSetCompiledData* InDataSetCompiledData, int32 DefaultNumBuffers)
 {
 	CompiledData.Init(InDataSetCompiledData != nullptr ? InDataSetCompiledData : &FNiagaraDataSetCompiledData::DummyCompiledData);
+
+	while(Data.Num() < DefaultNumBuffers)
+	{
+		Data.Add(new FNiagaraDataBuffer(this));
+	}
+
 	if (bInitialized)
 	{
 		Reset();
@@ -243,11 +249,8 @@ void FNiagaraDataSet::ResetBuffersInternal()
 #endif
 }
 
-FNiagaraDataBuffer& FNiagaraDataSet::BeginSimulate(bool bResetDestinationData)
+FNiagaraDataBuffer& FNiagaraDataSet::AllocateBuffer()
 {
-	//CheckCorrectThread();
-	check(DestinationData == nullptr);
-
 	//Find a free buffer we can write into.
 	//Linear search but there should only be 2 or three entries.
 	for (FNiagaraDataBuffer* Buffer : Data)
@@ -255,18 +258,22 @@ FNiagaraDataBuffer& FNiagaraDataSet::BeginSimulate(bool bResetDestinationData)
 		check(Buffer);
 		if (Buffer != CurrentData && Buffer->TryLock())
 		{
-			DestinationData = Buffer;
-			break;
+			return *Buffer;
 		}
 	}
 
-	if (DestinationData == nullptr)
-	{
-		Data.Add(new FNiagaraDataBuffer(this));
-		DestinationData = Data.Last();
-		verifySlow(DestinationData->TryLock());
-		checkSlow(DestinationData->IsBeingWritten());
-	}
+	FNiagaraDataBuffer* Buffer = Data.Add_GetRef(new FNiagaraDataBuffer(this));
+	verify(Buffer->TryLock());
+	check(Buffer->IsBeingWritten());
+	return *Buffer;
+}
+
+FNiagaraDataBuffer& FNiagaraDataSet::BeginSimulate(bool bResetDestinationData)
+{
+	//CheckCorrectThread();
+	check(DestinationData == nullptr);
+
+	DestinationData = &AllocateBuffer();
 
 	if (bResetDestinationData)
 	{
@@ -774,6 +781,11 @@ void FNiagaraDataBuffer::Allocate(uint32 InNumInstances, bool bMaintainExisting)
 		HalfStride = NewHalfStride;
 		BuildRegisterTable();
 	}
+
+	check(GetNumInstances() <= NumInstancesAllocated);
+	check(FloatData.Num() >= NewFloatNum);
+	check(Int32Data.Num() >= NewInt32Num);
+	check(HalfData.Num() >= NewHalfNum);
 }
 
 void FNiagaraDataBuffer::ReleaseCPU()
@@ -1035,8 +1047,6 @@ void FNiagaraDataBuffer::KillInstance(uint32 InstanceIdx)
 
 void FNiagaraDataBuffer::CopyToUnrelated(FNiagaraDataBuffer& DestBuffer, int32 StartIdx, int32 DestStartIdx, int32 InNumInstances)const
 {
-	CheckUsage(false);
-
 	if (StartIdx < 0 || (uint32)StartIdx >= NumInstances)
 	{
 		return;
@@ -1225,14 +1235,14 @@ void FNiagaraDataBuffer::GPUCopyFrom(const float* GPUReadBackFloat, const int* G
 {
 	//CheckUsage(false); //Have to disable this as in this specific case we write to a "CPUSim" from the RT.
 
+	Allocate(InNumInstances);
+	SetNumInstances(InNumInstances);
+
 	if (InNumInstances <= 0)
 	{
 		return;
 	}
 
-	Allocate(InNumInstances);
-	SetNumInstances(InNumInstances);
-
 	if (GPUReadBackFloat)
 	{
 		uint32 FloatComponents = Owner->GetNumFloatComponents();
@@ -1277,7 +1287,7 @@ void FNiagaraDataBuffer::GPUCopyFrom(const float* GPUReadBackFloat, const int* G
 	}
 }
 
-void FNiagaraDataBuffer::PushCPUBuffersToGPU(const TArray<FNiagaraDataBufferRef>& SourceBuffers, bool bReleaseRef, FRHICommandList& RHICmdList, ERHIFeatureLevel::Type FeatureLevel, const TCHAR* DebugSimName)
+void FNiagaraDataBuffer::PushCPUBuffersToGPU(const TArray<FNiagaraDataBufferRef>& SourceBuffers, bool bReleaseRef, FRHICommandList& RHICmdList, ERHIFeatureLevel::Type FeatureLevel, const TCHAR* DebugSimName, bool bAllocate)
 {
 	uint32 NewCount = 0;
 	check(GetOwner()->GetSimTarget() == ENiagaraSimTarget::GPUComputeSim);
@@ -1293,7 +1303,10 @@ void FNiagaraDataBuffer::PushCPUBuffersToGPU(const TArray<FNiagaraDataBufferRef>
 		}
 	}
 
-	AllocateGPU(RHICmdList, NewCount, FeatureLevel, DebugSimName);
+	if(bAllocate)
+	{
+		AllocateGPU(RHICmdList, NewCount, FeatureLevel, DebugSimName);
+	}
 	SetNumInstances(NewCount);
 
 	if (NewCount > 0)
diff --git a/Engine/Plugins/FX/Niagara/Source/Niagara/Private/NiagaraDataSetReadback.cpp b/Engine/Plugins/FX/Niagara/Source/Niagara/Private/NiagaraDataSetReadback.cpp
--- a/Engine/Plugins/FX/Niagara/Source/Niagara/Private/NiagaraDataSetReadback.cpp
+++ b/Engine/Plugins/FX/Niagara/Source/Niagara/Private/NiagaraDataSetReadback.cpp
@@ -193,3 +193,145 @@ void FNiagaraDataSetReadback::GPUReadbackInternal(FRHICommandListImmediate& RHIC
 	}
 	RHICmdList.Transition(Transitions);
 }
+
+//////////////////////////////////////////////////////////////////////////
+
+void FNiagaraDataBufferReadback::EnqueueReadback(FRHICommandList& RHICmdList, FNiagaraDataBufferRef InDataBuffer, FNiagaraGpuReadbackManager* ReadbackManager, FNiagaraGPUInstanceCountManager& InstanceCountManager)
+{
+	check(InDataBuffer);
+	check(IsInRenderingThread());
+
+	FNiagaraDataBuffer* DestBuffer = &InDataBuffer->GetOwner()->AllocateBuffer();
+	bool bHasGPUData = InDataBuffer->GetGPUBufferFloat().NumBytes > 0 || InDataBuffer->GetGPUBufferInt().NumBytes > 0 || InDataBuffer->GetGPUBufferHalf().NumBytes > 0;
+	++PendingReadbacks;
+	if (bHasGPUData)
+	{
+		GPUReadbackInternal(RHICmdList, ReadbackManager, InstanceCountManager, InDataBuffer, DestBuffer);
+	}
+	else
+	{
+		InDataBuffer->CopyTo(*DestBuffer, 0, 0, InDataBuffer->GetNumInstances());
+		ReadbackCompleteInternal(DestBuffer);
+	}
+}
+
+void FNiagaraDataBufferReadback::ReadbackCompleteInternal(FNiagaraDataBuffer* CompleteDataBuffer)
+{	
+	if (OnReadbackComplete.IsBound() && CompleteDataBuffer)
+	{
+		AsyncTask(
+			ENamedThreads::GameThread,
+			[Readback = AsShared(),
+			 CompleteBuffer_GT = CompleteDataBuffer->UnlockForRead()]()
+			{
+				Readback->PendingReadbacks--;
+				Readback->OnReadbackComplete.Execute(CompleteBuffer_GT);
+			}
+		);
+	}
+	else
+	{
+		PendingReadbacks--;
+	}
+}
+
+void FNiagaraDataBufferReadback::GPUReadbackInternal(FRHICommandList& RHICmdList, FNiagaraGpuReadbackManager* ReadbackManager, FNiagaraGPUInstanceCountManager& InstanceCountManager, FNiagaraDataBufferRef SrcDataBuffer, FNiagaraDataBuffer* DestDataBuffer)
+{
+	if (SrcDataBuffer == nullptr || DestDataBuffer == nullptr || ReadbackManager == nullptr)
+	{
+		ReadbackCompleteInternal(nullptr);
+		return;
+	}
+
+	const uint32 CountOffset = SrcDataBuffer->GetGPUInstanceCountBufferOffset();
+	if (CountOffset == INDEX_NONE)
+	{
+		ReadbackCompleteInternal(nullptr);
+		return;
+	}
+
+	FRWBuffer& FloatBuffer = SrcDataBuffer->GetGPUBufferFloat();
+	FRWBuffer& HalfBuffer = SrcDataBuffer->GetGPUBufferHalf();
+	FRWBuffer& IntBuffer = SrcDataBuffer->GetGPUBufferInt();
+	FRWBuffer& IDtoIndexBuffer = SrcDataBuffer->GetGPUIDToIndexTable();
+
+	constexpr int32 NumReadbackBuffers = 5;
+	TArray<FNiagaraGpuReadbackManager::FBufferRequest, TInlineAllocator<NumReadbackBuffers>> ReadbackBuffers;
+	const int32 CountBufferIndex = ReadbackBuffers.Emplace(InstanceCountManager.GetInstanceCountBuffer().Buffer, sizeof(uint32) * CountOffset, sizeof(uint32));
+	const int32 FloatBufferIndex = FloatBuffer.NumBytes == 0 ? INDEX_NONE : ReadbackBuffers.Emplace(FloatBuffer.Buffer, 0, FloatBuffer.NumBytes);
+	const int32 HalfBufferIndex = HalfBuffer.NumBytes == 0 ? INDEX_NONE : ReadbackBuffers.Emplace(HalfBuffer.Buffer, 0, HalfBuffer.NumBytes);
+	const int32 IntBufferIndex = IntBuffer.NumBytes == 0 ? INDEX_NONE : ReadbackBuffers.Emplace(IntBuffer.Buffer, 0, IntBuffer.NumBytes);
+	const int32 IDtoIndexBufferIndex = IDtoIndexBuffer.NumBytes == 0 ? INDEX_NONE : ReadbackBuffers.Emplace(IDtoIndexBuffer.Buffer, 0, IDtoIndexBuffer.NumBytes);
+
+	const int32 FloatBufferStride = SrcDataBuffer->GetFloatStride();
+	const int32 HalfBufferStride = SrcDataBuffer->GetHalfStride();
+	const int32 IntBufferStride = SrcDataBuffer->GetInt32Stride();
+
+	int32 SourceNumInstancesAllocated = (int32)SrcDataBuffer->GetNumInstancesAllocated();
+
+	check(SrcDataBuffer->GetOwner() == DestDataBuffer->GetOwner());
+
+	// Transition buffers to copy
+	TArray<FRHITransitionInfo, TInlineAllocator<NumReadbackBuffers>> Transitions;
+	Transitions.Emplace(ReadbackBuffers[0].Buffer, FNiagaraGPUInstanceCountManager::kCountBufferDefaultState, ERHIAccess::CopySrc);
+	for (int32 i = 1; i < ReadbackBuffers.Num(); ++i)
+	{
+		Transitions.Emplace(
+			ReadbackBuffers[i].Buffer,
+			(i == IDtoIndexBufferIndex) ? ERHIAccess::SRVCompute : ERHIAccess::SRVMask,
+			ERHIAccess::CopySrc
+		);
+	}
+	RHICmdList.Transition(Transitions);
+
+	// Enqueue readback
+	ReadbackManager->EnqueueReadbacks(
+		RHICmdList,
+		MakeArrayView(ReadbackBuffers),
+		[=, Readback = AsShared()](TConstArrayView<TPair<void*, uint32>> BufferData)
+		{
+			int32 InstanceCount = reinterpret_cast<int32*>(BufferData[CountBufferIndex].Key)[0];
+
+			ensure(SourceNumInstancesAllocated >= InstanceCount);
+			InstanceCount = FMath::Min(InstanceCount, SourceNumInstancesAllocated);
+
+			// Copy dataset databuffer
+			const float* FloatDataBuffer = FloatBufferIndex == INDEX_NONE ? nullptr : reinterpret_cast<float*>(BufferData[FloatBufferIndex].Key);
+			const FFloat16* HalfDataBuffer = HalfBufferIndex == INDEX_NONE ? nullptr : reinterpret_cast<FFloat16*>(BufferData[HalfBufferIndex].Key);
+			const int32* IntDataBuffer = IntBufferIndex == INDEX_NONE ? nullptr : reinterpret_cast<int32*>(BufferData[IntBufferIndex].Key);
+			DestDataBuffer->GPUCopyFrom(FloatDataBuffer, IntDataBuffer, HalfDataBuffer, 0, InstanceCount, FloatBufferStride, IntBufferStride, HalfBufferStride);
+
+			check(DestDataBuffer->GetNumInstances() <= DestDataBuffer->GetNumInstancesAllocated());
+
+//			DestDataBuffer->Dump(0, DestDataBuffer->GetNumInstances(), TEXT("GPU Readback dump"));
+
+			// Copy ID to Index table
+			if (SrcDataBuffer)
+			{
+				TArray<int32>& IDTable = SrcDataBuffer->GetIDTable();
+
+				const int32* IDtoIndexBuffer = IDtoIndexBufferIndex == INDEX_NONE ? nullptr : reinterpret_cast<int32*>(BufferData[IDtoIndexBufferIndex].Key);
+				if (IDtoIndexBuffer != nullptr)
+				{
+					const int32 NumIDs = BufferData[IDtoIndexBufferIndex].Value / sizeof(int32);
+					check(NumIDs >= InstanceCount);
+					IDTable.SetNumUninitialized(NumIDs);
+					FMemory::Memcpy(IDTable.GetData(), IDtoIndexBuffer, NumIDs * sizeof(int32));
+				}
+				else
+				{
+					IDTable.Empty();
+				}
+			}
+
+			Readback->ReadbackCompleteInternal(DestDataBuffer);
+		}
+	);
+
+	// Transition buffers from copy
+	for (FRHITransitionInfo& Transition : Transitions)
+	{
+		Swap(Transition.AccessBefore, Transition.AccessAfter);
+	}
+	RHICmdList.Transition(Transitions);
+}
diff --git a/Engine/Plugins/FX/Niagara/Source/Niagara/Private/NiagaraDebugHud.cpp b/Engine/Plugins/FX/Niagara/Source/Niagara/Private/NiagaraDebugHud.cpp
--- a/Engine/Plugins/FX/Niagara/Source/Niagara/Private/NiagaraDebugHud.cpp
+++ b/Engine/Plugins/FX/Niagara/Source/Niagara/Private/NiagaraDebugHud.cpp
@@ -734,6 +734,10 @@ namespace NiagaraDebugLocal
 			{
 				GpuFeaturesBuilder.Append(TEXT(" RayTracingScene"));
 			}
+			if (SystemInstance->RequiresCurrentFrameNDC())
+			{
+				GpuFeaturesBuilder.Append(TEXT(" CurrentFrameNDC"));
+			}
 			if (GpuFeaturesBuilder.Len() > 0)
 			{
 				Output.Appendf(TEXT("GpuFeatures -%s\n"), *GpuFeaturesBuilder);
diff --git a/Engine/Plugins/FX/Niagara/Source/Niagara/Private/NiagaraDebuggerClient.cpp b/Engine/Plugins/FX/Niagara/Source/Niagara/Private/NiagaraDebuggerClient.cpp
--- a/Engine/Plugins/FX/Niagara/Source/Niagara/Private/NiagaraDebuggerClient.cpp
+++ b/Engine/Plugins/FX/Niagara/Source/Niagara/Private/NiagaraDebuggerClient.cpp
@@ -456,6 +456,7 @@ void FNiagaraDebuggerClient::CaptureOutlinerData(FNiagaraOutlinerData& OutlinerD
 				InstData.bRequiresEarlyViewData = Inst->RequiresEarlyViewData();
 				InstData.bRequiresViewUniformBuffer = Inst->RequiresViewUniformBuffer();
 				InstData.bRequiresRayTracingScene = Inst->RequiresRayTracingScene();
+				InstData.bRequiresCurrentFrameNDC = Inst->RequiresCurrentFrameNDC();
 
 				InstData.Emitters.Reserve(Inst->GetEmitters().Num());
 				for (const FNiagaraEmitterInstanceRef& EmitterInst : Inst->GetEmitters())
diff --git a/Engine/Plugins/FX/Niagara/Source/Niagara/Private/NiagaraGPUInstanceCountManager.cpp b/Engine/Plugins/FX/Niagara/Source/Niagara/Private/NiagaraGPUInstanceCountManager.cpp
--- a/Engine/Plugins/FX/Niagara/Source/Niagara/Private/NiagaraGPUInstanceCountManager.cpp
+++ b/Engine/Plugins/FX/Niagara/Source/Niagara/Private/NiagaraGPUInstanceCountManager.cpp
@@ -719,3 +719,73 @@ void FNiagaraGPUInstanceCountManager::CopyToMultiViewCountBuffer(FRHICommandList
 	}
 }
 
+
+void FNiagaraGPUInstanceCountManager::ProcessInitInstanceCountTasks(FNiagaraGpuComputeDispatchInterface* ComputeDispatchInterface, FRHICommandList& RHICmdList)
+{
+	if(InstanceCountInitTasks.NumBytes() == 0)
+	{
+		return;
+	}
+
+	SCOPED_DRAW_EVENT(RHICmdList, NiagaraProcessInitInstanceCountTasks);
+
+	const int32 NumTasks = InstanceCountInitTasks.Num() / 2;
+
+	// Allocate task buffer
+	FReadBuffer TaskInfosBuffer;
+	{
+		const uint32 TaskBufferSize = InstanceCountInitTasks.Num() * sizeof(uint32);
+		TaskInfosBuffer.Initialize(RHICmdList, TEXT("NiagaraInitCountsTaskInfosBuffer"), sizeof(uint32), InstanceCountInitTasks.Num(), EPixelFormat::PF_R32_UINT, BUF_Volatile);
+
+		uint8* TaskBufferData = (uint8*)RHICmdList.LockBuffer(TaskInfosBuffer.Buffer, 0, TaskBufferSize, RLM_WriteOnly);
+		FMemory::Memcpy(TaskBufferData, InstanceCountInitTasks.GetData(), TaskBufferSize);
+		RHICmdList.UnlockBuffer(TaskInfosBuffer.Buffer);
+	}
+
+	FNiagaraEmptyUAVPoolScopedAccess UAVPoolAccessScope(ComputeDispatchInterface->GetEmptyUAVPool());
+	TArray<FRHITransitionInfo, TInlineAllocator<10>> Transitions;
+	Transitions.Reserve(1);
+	FRWBuffer& CurrentCountBuffer = CountBuffer;
+
+	// Get counts buffer
+	FUnorderedAccessViewRHIRef CountsUAV = nullptr;
+	const bool bCountBufferIsValid = CurrentCountBuffer.UAV.IsValid();
+	if (bCountBufferIsValid)
+	{
+		// treat the incoming UAV as being unknown to be sure a barrier is inserted in the case
+		// where the preceding dispatch wrote to the counts buffer
+		Transitions.Emplace(CurrentCountBuffer.UAV, ERHIAccess::Unknown, ERHIAccess::UAVCompute);
+		CountsUAV = CurrentCountBuffer.UAV;
+	}
+	else
+	{
+		// This can happen if there are no InstanceCountClearTasks and all DrawIndirectArgGenTasks_PreOpaque are using culled counts
+		CountsUAV = ComputeDispatchInterface->GetEmptyUAVFromPool(RHICmdList, PF_R32_UINT, ENiagaraEmptyUAVType::Buffer);
+	}
+
+	RHICmdList.Transition(Transitions);
+
+
+	FNiagaraInstanceCountsInitCS::FParameters InitCountParameters;
+	InitCountParameters.TaskInfos = TaskInfosBuffer.SRV;
+	InitCountParameters.RWInstanceCounts = CountsUAV;
+	InitCountParameters.TaskCount.X = NumTasks;
+	
+
+	FNiagaraInstanceCountsInitCS::FPermutationDomain PermutationVectorResetCounts;
+	TShaderMapRef<FNiagaraInstanceCountsInitCS> InitCountsCS(GetGlobalShaderMap(FeatureLevel), PermutationVectorResetCounts);
+	FComputeShaderUtils::Dispatch(RHICmdList, InitCountsCS, InitCountParameters, FIntVector(FMath::DivideAndRoundUp(NumTasks, NIAGARA_DRAW_INDIRECT_ARGS_GEN_THREAD_COUNT), 1, 1));
+
+	// Generate and execute transitions
+	Transitions.Reset();
+	Transitions.Emplace(CurrentCountBuffer.UAV, ERHIAccess::UAVCompute, kCountBufferDefaultState);
+	RHICmdList.Transition(Transitions);
+
+	InstanceCountInitTasks.Reset();
+}
+
+void FNiagaraGPUInstanceCountManager::AddInstanceCountInitTask(uint32 Offset, uint32 Value)
+{
+ 	InstanceCountInitTasks.Emplace(Offset);
+ 	InstanceCountInitTasks.Emplace(Value);
+}
\ No newline at end of file
diff --git a/Engine/Plugins/FX/Niagara/Source/Niagara/Private/NiagaraGpuComputeDispatch.cpp b/Engine/Plugins/FX/Niagara/Source/Niagara/Private/NiagaraGpuComputeDispatch.cpp
--- a/Engine/Plugins/FX/Niagara/Source/Niagara/Private/NiagaraGpuComputeDispatch.cpp
+++ b/Engine/Plugins/FX/Niagara/Source/Niagara/Private/NiagaraGpuComputeDispatch.cpp
@@ -37,6 +37,8 @@
 #include "Rendering/RenderCommandPipes.h"
 #include "GPUSkinCache.h"
 
+#include "NiagaraDataChannel.h"
+
 DECLARE_CYCLE_STAT(TEXT("GPU Dispatch Setup [RT]"), STAT_NiagaraGPUDispatchSetup_RT, STATGROUP_Niagara);
 DECLARE_CYCLE_STAT(TEXT("GPU Emitter Dispatch [RT]"), STAT_NiagaraGPUSimTick_RT, STATGROUP_Niagara);
 DECLARE_CYCLE_STAT(TEXT("GPU Data Readback [RT]"), STAT_NiagaraGPUReadback_RT, STATGROUP_Niagara);
@@ -252,6 +254,7 @@ void FNiagaraGpuComputeDispatch::AddGpuComputeProxy(FNiagaraSystemGpuComputeProx
 	NumProxiesThatRequireDepthBuffer			+= ComputeProxy->RequiresDepthBuffer() ? 1 : 0;
 	NumProxiesThatRequireEarlyViewData			+= ComputeProxy->RequiresEarlyViewData() ? 1 : 0;
 	NumProxiesThatRequireRayTracingScene		+= ComputeProxy->RequiresRayTracingScene() ? 1 : 0;
+	NumProxiesThatRequireCurrentFrameNDC		+= ComputeProxy->RequiresCurrentFrameNDC() ? 1 : 0;
 }
 
 void FNiagaraGpuComputeDispatch::RemoveGpuComputeProxy(FNiagaraSystemGpuComputeProxy* ComputeProxy)
@@ -273,6 +276,7 @@ void FNiagaraGpuComputeDispatch::RemoveGpuComputeProxy(FNiagaraSystemGpuComputeP
 	NumProxiesThatRequireDepthBuffer			-= ComputeProxy->RequiresDepthBuffer() ? 1 : 0;
 	NumProxiesThatRequireEarlyViewData			-= ComputeProxy->RequiresEarlyViewData() ? 1 : 0;
 	NumProxiesThatRequireRayTracingScene		-= ComputeProxy->RequiresRayTracingScene() ? 1 : 0;
+	NumProxiesThatRequireCurrentFrameNDC		-= ComputeProxy->RequiresCurrentFrameNDC() ? 1 : 0;
 
 #if NIAGARA_COMPUTEDEBUG_ENABLED
 	if (FNiagaraGpuComputeDebug* GpuComputeDebug = GpuComputeDebugPtr.Get())
@@ -296,6 +300,18 @@ void FNiagaraGpuComputeDispatch::RemoveGpuComputeProxy(FNiagaraSystemGpuComputeP
 #endif
 }
 
+void FNiagaraGpuComputeDispatch::AddNDCDataProxy(FNiagaraDataChannelDataProxy* NDCDataProxy)
+{
+	check(NDCDataProxy->DispatchInterface == this);
+	NDCDataProxies.Add(NDCDataProxy);	
+}
+
+void FNiagaraGpuComputeDispatch::RemoveNDCDataProxy(FNiagaraDataChannelDataProxy* NDCDataProxy)
+{
+	check(NDCDataProxy->DispatchInterface == this);
+	NDCDataProxies.RemoveSwap(NDCDataProxy);
+}
+
 void FNiagaraGpuComputeDispatch::Tick(UWorld* World, float DeltaTime)
 {
 	check(IsInGameThread());
@@ -714,6 +730,11 @@ void FNiagaraGpuComputeDispatch::UpdateInstanceCountManager(FRHICommandListImmed
 			GPUInstanceCounterManager.ReleaseGPUReadback();
 		}
 	}
+
+	//Run instance count Initialization task. Some instance counts must be initialized with values from the CPU before use in GPU sims.
+	{
+		GPUInstanceCounterManager.ProcessInitInstanceCountTasks(this, RHICmdList);
+	}
 }
 
 void FNiagaraGpuComputeDispatch::PrepareTicksForProxy(FRHICommandListImmediate& RHICmdList, FNiagaraSystemGpuComputeProxy* ComputeProxy, FNiagaraGpuDispatchList& GpuDispatchList)
@@ -1921,6 +1942,11 @@ void FNiagaraGpuComputeDispatch::PreInitViews(FRDGBuilder& GraphBuilder, bool bA
 		{
 			FramesBeforeTickFlush = 0;
 
+			for(FNiagaraDataChannelDataProxy* NDCProxy : NDCDataProxies)
+			{
+				NDCProxy->BeginFrame(GraphBuilder.RHICmdList);
+			}
+
 			UpdateInstanceCountManager(GraphBuilder.RHICmdList);
 			PrepareAllTicks(GraphBuilder.RHICmdList);
 
@@ -2081,6 +2107,11 @@ void FNiagaraGpuComputeDispatch::PostRenderOpaque(FRDGBuilder& GraphBuilder, TCo
 			if (bAllowGPUParticleUpdate)
 			{
 				FinishDispatches();
+
+				for (FNiagaraDataChannelDataProxy* NDCProxy : NDCDataProxies)
+				{
+					NDCProxy->EndFrame(RHICmdList);
+				}
 			}
 		}
 	);
@@ -2329,7 +2360,7 @@ FNiagaraAsyncGpuTraceHelper& FNiagaraGpuComputeDispatch::GetAsyncGpuTraceHelper(
 	return *AsyncGpuTraceHelper.Get();
 }
 
-void FNiagaraGpuComputeDispatch::ResetDataInterfaces(FRDGBuilder& GraphBuilder, const FNiagaraGPUSystemTick& Tick, const FNiagaraComputeInstanceData& InstanceData) const
+void FNiagaraGpuComputeDispatch::ResetDataInterfaces(FRDGBuilder& GraphBuilder, const FNiagaraGPUSystemTick& Tick, const FNiagaraComputeInstanceData& InstanceData) 
 {
 	const int32 NumDataInterfaces = InstanceData.DataInterfaceProxies.Num();
 	if (NumDataInterfaces == 0)
@@ -2352,7 +2383,7 @@ void FNiagaraGpuComputeDispatch::ResetDataInterfaces(FRDGBuilder& GraphBuilder,
 	}
 }
 
-void FNiagaraGpuComputeDispatch::SetDataInterfaceParameters(FRDGBuilder& GraphBuilder, const FNiagaraGPUSystemTick& Tick, const FNiagaraComputeInstanceData& InstanceData, const FNiagaraShaderRef& ComputeShader, const FNiagaraSimStageData& SimStageData, const FNiagaraShaderScriptParametersMetadata& NiagaraShaderParametersMetadata, uint8* ParametersStructure) const
+void FNiagaraGpuComputeDispatch::SetDataInterfaceParameters(FRDGBuilder& GraphBuilder, const FNiagaraGPUSystemTick& Tick, const FNiagaraComputeInstanceData& InstanceData, const FNiagaraShaderRef& ComputeShader, const FNiagaraSimStageData& SimStageData, const FNiagaraShaderScriptParametersMetadata& NiagaraShaderParametersMetadata, uint8* ParametersStructure) 
 {
 	const int32 NumDataInterfaces = InstanceData.DataInterfaceProxies.Num();
 	if (NumDataInterfaces == 0)
@@ -2377,7 +2408,7 @@ void FNiagaraGpuComputeDispatch::SetDataInterfaceParameters(FRDGBuilder& GraphBu
 	}
 }
 
-void FNiagaraGpuComputeDispatch::PreStageInterface(FRDGBuilder& GraphBuilder, const FNiagaraGPUSystemTick& Tick, const FNiagaraComputeInstanceData& InstanceData, const FNiagaraSimStageData& SimStageData, TSet<FNiagaraDataInterfaceProxy*>& ProxiesToFinalize) const
+void FNiagaraGpuComputeDispatch::PreStageInterface(FRDGBuilder& GraphBuilder, const FNiagaraGPUSystemTick& Tick, const FNiagaraComputeInstanceData& InstanceData, const FNiagaraSimStageData& SimStageData, TSet<FNiagaraDataInterfaceProxy*>& ProxiesToFinalize) 
 {
 	const int32 NumDataInterfaces = InstanceData.DataInterfaceProxies.Num();
 	if (NumDataInterfaces == 0)
@@ -2405,7 +2436,7 @@ void FNiagaraGpuComputeDispatch::PreStageInterface(FRDGBuilder& GraphBuilder, co
 	}
 }
 
-void FNiagaraGpuComputeDispatch::PostStageInterface(FRDGBuilder& GraphBuilder, const FNiagaraGPUSystemTick& Tick, const FNiagaraComputeInstanceData& InstanceData, const FNiagaraSimStageData& SimStageData, TSet<FNiagaraDataInterfaceProxy*>& ProxiesToFinalize) const
+void FNiagaraGpuComputeDispatch::PostStageInterface(FRDGBuilder& GraphBuilder, const FNiagaraGPUSystemTick& Tick, const FNiagaraComputeInstanceData& InstanceData, const FNiagaraSimStageData& SimStageData, TSet<FNiagaraDataInterfaceProxy*>& ProxiesToFinalize) 
 {
 	const int32 NumDataInterfaces = InstanceData.DataInterfaceProxies.Num();
 	if (NumDataInterfaces == 0)
@@ -2433,7 +2464,7 @@ void FNiagaraGpuComputeDispatch::PostStageInterface(FRDGBuilder& GraphBuilder, c
 	}
 }
 
-void FNiagaraGpuComputeDispatch::PostSimulateInterface(FRDGBuilder& GraphBuilder, const FNiagaraGPUSystemTick& Tick, const FNiagaraComputeInstanceData& InstanceData, const FNiagaraSimStageData& SimStageData) const
+void FNiagaraGpuComputeDispatch::PostSimulateInterface(FRDGBuilder& GraphBuilder, const FNiagaraGPUSystemTick& Tick, const FNiagaraComputeInstanceData& InstanceData, const FNiagaraSimStageData& SimStageData) 
 {
 	const int32 NumDataInterfaces = InstanceData.DataInterfaceProxies.Num();
 	if (NumDataInterfaces == 0)
diff --git a/Engine/Plugins/FX/Niagara/Source/Niagara/Private/NiagaraGpuComputeDispatch.h b/Engine/Plugins/FX/Niagara/Source/Niagara/Private/NiagaraGpuComputeDispatch.h
--- a/Engine/Plugins/FX/Niagara/Source/Niagara/Private/NiagaraGpuComputeDispatch.h
+++ b/Engine/Plugins/FX/Niagara/Source/Niagara/Private/NiagaraGpuComputeDispatch.h
@@ -98,6 +98,9 @@ public:
 	/** Remove system instance proxy from the batcher. */
 	virtual void RemoveGpuComputeProxy(FNiagaraSystemGpuComputeProxy* ComputeProxy) override;
 
+	virtual void AddNDCDataProxy(FNiagaraDataChannelDataProxy* NDCDataProxy) override;
+	virtual void RemoveNDCDataProxy(FNiagaraDataChannelDataProxy* NDCDataProxy) override;
+
 #if WITH_EDITOR
 	virtual void Suspend() override {}
 	virtual void Resume() override {}
@@ -144,11 +147,11 @@ public:
 	virtual bool AddSortedGPUSimulation(FRHICommandListBase& RHICmdList, FNiagaraGPUSortInfo& SortInfo) override;
 	virtual const FGlobalDistanceFieldParameterData* GetGlobalDistanceFieldData() const override;
 
-	void ResetDataInterfaces(FRDGBuilder& GraphBuilder, const FNiagaraGPUSystemTick& Tick, const FNiagaraComputeInstanceData& InstanceData) const;
-	void SetDataInterfaceParameters(FRDGBuilder& GraphBuilder, const FNiagaraGPUSystemTick& Tick, const FNiagaraComputeInstanceData& InstanceData, const FNiagaraShaderRef& ComputeShader, const FNiagaraSimStageData& SimStageData, const FNiagaraShaderScriptParametersMetadata& NiagaraShaderParametersMetadata, uint8* ParametersStructure) const;
-	void PreStageInterface(FRDGBuilder& GraphBuilder, const FNiagaraGPUSystemTick& Tick, const FNiagaraComputeInstanceData& InstanceData, const FNiagaraSimStageData& SimStageData, TSet<FNiagaraDataInterfaceProxy*>& ProxiesToFinalize) const;
-	void PostStageInterface(FRDGBuilder& GraphBuilder, const FNiagaraGPUSystemTick& Tick, const FNiagaraComputeInstanceData& InstanceData, const FNiagaraSimStageData& SimStageData, TSet<FNiagaraDataInterfaceProxy*>& ProxiesToFinalize) const;
-	void PostSimulateInterface(FRDGBuilder& GraphBuilder, const FNiagaraGPUSystemTick& Tick, const FNiagaraComputeInstanceData& InstanceData, const FNiagaraSimStageData& SimStageData) const;
+	void ResetDataInterfaces(FRDGBuilder& GraphBuilder, const FNiagaraGPUSystemTick& Tick, const FNiagaraComputeInstanceData& InstanceData);
+	void SetDataInterfaceParameters(FRDGBuilder& GraphBuilder, const FNiagaraGPUSystemTick& Tick, const FNiagaraComputeInstanceData& InstanceData, const FNiagaraShaderRef& ComputeShader, const FNiagaraSimStageData& SimStageData, const FNiagaraShaderScriptParametersMetadata& NiagaraShaderParametersMetadata, uint8* ParametersStructure);
+	void PreStageInterface(FRDGBuilder& GraphBuilder, const FNiagaraGPUSystemTick& Tick, const FNiagaraComputeInstanceData& InstanceData, const FNiagaraSimStageData& SimStageData, TSet<FNiagaraDataInterfaceProxy*>& ProxiesToFinalize);
+	void PostStageInterface(FRDGBuilder& GraphBuilder, const FNiagaraGPUSystemTick& Tick, const FNiagaraComputeInstanceData& InstanceData, const FNiagaraSimStageData& SimStageData, TSet<FNiagaraDataInterfaceProxy*>& ProxiesToFinalize);
+	void PostSimulateInterface(FRDGBuilder& GraphBuilder, const FNiagaraGPUSystemTick& Tick, const FNiagaraComputeInstanceData& InstanceData, const FNiagaraSimStageData& SimStageData);
 
 	/** Given a shader stage index, find the corresponding data interface */
 	FNiagaraDataInterfaceProxyRW* FindIterationInterface(FNiagaraComputeInstanceData* Instance, const uint32 SimulationStageIndex) const;
@@ -235,6 +238,7 @@ private:
 	uint32 NumProxiesThatRequireDepthBuffer = 0;
 	uint32 NumProxiesThatRequireEarlyViewData = 0;
 	uint32 NumProxiesThatRequireRayTracingScene = 0;
+	uint32 NumProxiesThatRequireCurrentFrameNDC = 0;
 
 	int32 TotalDispatchesThisFrame = 0;
 
@@ -244,6 +248,8 @@ private:
 
 	bool bRequiresReadback = false;
 	TArray<FNiagaraSystemGpuComputeProxy*> ProxiesPerStage[ENiagaraGpuComputeTickStage::Max];
+	
+	TArray<FNiagaraDataChannelDataProxy*> NDCDataProxies;
 
 	FNiagaraGpuDispatchList DispatchListPerStage[ENiagaraGpuComputeTickStage::Max];
 
diff --git a/Engine/Plugins/FX/Niagara/Source/Niagara/Private/NiagaraSystemGpuComputeProxy.cpp b/Engine/Plugins/FX/Niagara/Source/Niagara/Private/NiagaraSystemGpuComputeProxy.cpp
--- a/Engine/Plugins/FX/Niagara/Source/Niagara/Private/NiagaraSystemGpuComputeProxy.cpp
+++ b/Engine/Plugins/FX/Niagara/Source/Niagara/Private/NiagaraSystemGpuComputeProxy.cpp
@@ -20,6 +20,7 @@ FNiagaraSystemGpuComputeProxy::FNiagaraSystemGpuComputeProxy(FNiagaraSystemInsta
 	bRequiresEarlyViewData = OwnerInstance->RequiresEarlyViewData();
 	bRequiresViewUniformBuffer = OwnerInstance->RequiresViewUniformBuffer();
 	bRequiresRayTracingScene = OwnerInstance->RequiresRayTracingScene();
+	bRequiresCurrentFrameNDC = OwnerInstance->RequiresCurrentFrameNDC();
 
 	// Gather all emitter compute contexts
 	for ( auto& Emitter : OwnerInstance->GetEmitters() )
@@ -31,7 +32,7 @@ FNiagaraSystemGpuComputeProxy::FNiagaraSystemGpuComputeProxy(FNiagaraSystemInsta
 	}
 
 	// Calculate Tick Stage
-	if (bRequiresGlobalDistanceField || bRequiresDepthBuffer || bRequiresRayTracingScene)
+	if (bRequiresGlobalDistanceField || bRequiresDepthBuffer || bRequiresRayTracingScene || bRequiresCurrentFrameNDC)
 	{
 		ComputeTickStage = ENiagaraGpuComputeTickStage::PostOpaqueRender;
 	}
@@ -166,8 +167,11 @@ void FNiagaraSystemGpuComputeProxy::QueueTick(const FNiagaraGPUSystemTick& Tick)
 		for (auto& Pair : Tick.DIInstanceData->InterfaceProxiesToOffsets)
 		{
 			FNiagaraDataInterfaceProxy* Proxy = Pair.Key;
-			uint8* InstanceDataPtr = BasePointer + Pair.Value;
-			Proxy->ConsumePerInstanceDataFromGameThread(InstanceDataPtr, Tick.SystemInstanceID);
+			//if(Proxy->bConsumeCPUDataEarly)//TODO: Add path to optionally allow DIs to consume data in "Frames" rather than here when we consume the ticks.
+			//{
+				uint8* InstanceDataPtr = BasePointer + Pair.Value;
+				Proxy->ConsumePerInstanceDataFromGameThread(InstanceDataPtr, Tick.SystemInstanceID);
+			//}
 		}
 	}
 }
diff --git a/Engine/Plugins/FX/Niagara/Source/Niagara/Private/NiagaraSystemInstance.cpp b/Engine/Plugins/FX/Niagara/Source/Niagara/Private/NiagaraSystemInstance.cpp
--- a/Engine/Plugins/FX/Niagara/Source/Niagara/Private/NiagaraSystemInstance.cpp
+++ b/Engine/Plugins/FX/Niagara/Source/Niagara/Private/NiagaraSystemInstance.cpp
@@ -1407,6 +1407,31 @@ bool FNiagaraSystemInstance::RequiresRayTracingScene() const
 	return false;
 }
 
+bool FNiagaraSystemInstance::RequiresCurrentFrameNDC() const
+{
+	if (!bHasGPUEmitters)
+	{
+		return false;
+	}
+
+	for (const FNiagaraEmitterInstanceRef& Emitter : Emitters)
+	{
+		FNiagaraComputeExecutionContext* GPUContext = Emitter->GetGPUContext();
+		if (GPUContext)
+		{
+			for (UNiagaraDataInterface* DataInterface : GPUContext->CombinedParamStore.GetDataInterfaces())
+			{
+				if (DataInterface && DataInterface->RequiresCurrentFrameNDC())
+				{
+					return true;
+				}
+			}
+		}
+	}
+
+	return false;
+}
+
 FNDIStageTickHandler* FNiagaraSystemInstance::GetSystemDIStageTickHandler(ENiagaraScriptUsage Usage)
 {
 	if(Usage == ENiagaraScriptUsage::SystemSpawnScript || Usage == ENiagaraScriptUsage::EmitterSpawnScript)
diff --git a/Engine/Plugins/FX/Niagara/Source/Niagara/Private/NiagaraWorldManager.cpp b/Engine/Plugins/FX/Niagara/Source/Niagara/Private/NiagaraWorldManager.cpp
--- a/Engine/Plugins/FX/Niagara/Source/Niagara/Private/NiagaraWorldManager.cpp
+++ b/Engine/Plugins/FX/Niagara/Source/Niagara/Private/NiagaraWorldManager.cpp
@@ -27,6 +27,7 @@
 #include "Engine/LocalPlayer.h"
 #include "SceneView.h"
 #include "UObject/UObjectIterator.h"
+#include "Logging/StructuredLog.h"
 
 #include UE_INLINE_GENERATED_CPP_BY_NAME(NiagaraWorldManager)
 
@@ -250,6 +251,7 @@ FAutoConsoleCommandWithWorldAndArgs GCmdNiagaraScalabilityCullingMode(
 );
 
 FDelegateHandle FNiagaraWorldManager::OnWorldInitHandle;
+FDelegateHandle FNiagaraWorldManager::OnPostWorldInitHandle;
 FDelegateHandle FNiagaraWorldManager::OnWorldCleanupHandle;
 FDelegateHandle FNiagaraWorldManager::OnPostWorldCleanupHandle;
 FDelegateHandle FNiagaraWorldManager::OnPreWorldFinishDestroyHandle;
@@ -399,8 +401,6 @@ void FNiagaraWorldManager::Init(UWorld* InWorld)
 	//Possibly a later hook we can use.
 	//PrimePoolForAllSystems();
 
-	DataChannelManager->Init();
-
 #if WITH_NIAGARA_DEBUGGER
 	NiagaraDebugHud.Reset(new FNiagaraDebugHud(World));
 #endif
@@ -415,6 +415,11 @@ void FNiagaraWorldManager::Init(UWorld* InWorld)
 	FNiagaraComponentSettings::UpdateSettings();
 }
 
+void FNiagaraWorldManager::OnPostWorldInit()
+{
+	DataChannelManager->Init();
+}
+
 FNiagaraWorldManager::~FNiagaraWorldManager()
 {
 	OnWorldCleanup(true, true);
@@ -436,6 +441,7 @@ FNiagaraWorldManager* FNiagaraWorldManager::Get(const UWorld* World)
 void FNiagaraWorldManager::OnStartup()
 {
 	OnWorldInitHandle = FWorldDelegates::OnPreWorldInitialization.AddStatic(&FNiagaraWorldManager::OnWorldInit);
+	OnPostWorldInitHandle = FWorldDelegates::OnPostWorldInitialization.AddStatic(&FNiagaraWorldManager::OnPostWorldInit);
 	OnWorldCleanupHandle = FWorldDelegates::OnWorldCleanup.AddStatic(&FNiagaraWorldManager::OnWorldCleanup);
 	OnPostWorldCleanupHandle = FWorldDelegates::OnPostWorldCleanup.AddStatic(&FNiagaraWorldManager::OnPostWorldCleanup);
 	OnPreWorldFinishDestroyHandle = FWorldDelegates::OnPreWorldFinishDestroy.AddStatic(&FNiagaraWorldManager::OnPreWorldFinishDestroy);
@@ -472,6 +478,7 @@ void FNiagaraWorldManager::OnStartup()
 void FNiagaraWorldManager::OnShutdown()
 {
 	FWorldDelegates::OnPreWorldInitialization.Remove(OnWorldInitHandle);
+	FWorldDelegates::OnPostWorldInitialization.Remove(OnPostWorldInitHandle);
 	FWorldDelegates::OnWorldCleanup.Remove(OnWorldCleanupHandle);
 	FWorldDelegates::OnPostWorldCleanup.Remove(OnPostWorldCleanupHandle);
 	FWorldDelegates::OnPreWorldFinishDestroy.Remove(OnPreWorldFinishDestroyHandle);
@@ -647,7 +654,7 @@ FNiagaraSystemSimulationRef FNiagaraWorldManager::GetSystemSimulation(ETickingGr
 	//The async work for this system can then properly be chained to the tick task so that everything completes on time.
 	//Most systems most of the time will no need this and we'll clear this requirement as soon as all systems no longer need it.
 	//TODO: This should go away when we move over to the new task graph system and rejig our dependency tracking.
-	if(System->AsyncWorkCanOverlapTickGroups())
+	if(!System->AsyncWorkCanOverlapTickGroups())
 	{
 		TickFunctions[ActualTickGroup].EndTickGroup = (ETickingGroup)ActualTickGroup;
 	}
@@ -898,6 +905,19 @@ void FNiagaraWorldManager::OnWorldInit(UWorld* World, const UWorld::Initializati
 	NewManager->Init(World);
 }
 
+void FNiagaraWorldManager::OnPostWorldInit(UWorld* World, const UWorld::InitializationValues IVS)
+{
+	FNiagaraWorldManager** Manager = WorldManagers.Find(World);
+	if (Manager)
+	{
+		FNiagaraWorldManager* WorldMan = *Manager;
+		if (ensure(WorldMan))
+		{
+			WorldMan->OnPostWorldInit();
+		}
+	}
+}
+
 void FNiagaraWorldManager::OnWorldCleanup(UWorld* World, bool bSessionEnded, bool bCleanupResources)
 {
 	//Cleanup world manager contents but not the manager itself.
@@ -1461,6 +1481,10 @@ void FNiagaraWorldManager::Tick(ETickingGroup TickGroup, float DeltaSeconds, ELe
 		//Set here but this wont take effect until next tick.
 		TickFunctions[ActualTickGroup].EndTickGroup = TG_LastDemotable;
 	}
+	else
+	{
+		TickFunctions[ActualTickGroup].EndTickGroup = (ETickingGroup)ActualTickGroup;
+	}
 
 	ActiveNiagaraTickGroup = -1;
 
diff --git a/Engine/Plugins/FX/Niagara/Source/Niagara/Public/NiagaraCommon.h b/Engine/Plugins/FX/Niagara/Source/Niagara/Public/NiagaraCommon.h
--- a/Engine/Plugins/FX/Niagara/Source/Niagara/Public/NiagaraCommon.h
+++ b/Engine/Plugins/FX/Niagara/Source/Niagara/Public/NiagaraCommon.h
@@ -706,6 +706,9 @@ struct FNiagaraFunctionSignature
 	bool VariadicInput()const { return RequiredInputs != INDEX_NONE; }
 	bool VariadicOutput()const { return RequiredOutputs != INDEX_NONE; }
 
+	int32 VariadicInputStartIndex()const { return VariadicInput() ? NumRequiredInputs() + (bRequiresExecPin ? 1 : 0) : INDEX_NONE; }
+	int32 VariadicOutputStartIndex()const { return VariadicOutput() ? NumRequiredOutputs() + (bRequiresExecPin ? 1 : 0) : INDEX_NONE; }
+
 	int32 NumRequiredInputs()const { return RequiredInputs == INDEX_NONE ? Inputs.Num() : RequiredInputs; }
 	int32 NumOptionalInputs()const { return Inputs.Num() - NumRequiredInputs(); }
 
@@ -721,9 +724,9 @@ struct FNiagaraFunctionSignature
 	}
 	TConstArrayView<FNiagaraVariableBase> GetOutputs() const { return Outputs; }
 #endif
-	
-	NIAGARA_API void GetVariadicInputs(TArray<FNiagaraVariableBase>& OutVariadicInputs, bool bStripNonExecution = true)const;
-	NIAGARA_API void GetVariadicOutputs(TArray<FNiagaraVariableBase>& OutVariadicOutputs, bool bStripNonExecution = true)const;
+
+	NIAGARA_API void GetVariadicInputs(TArray<FNiagaraVariableBase>& OutVariadicInputs)const;
+	NIAGARA_API void GetVariadicOutputs(TArray<FNiagaraVariableBase>& OutVariadicOutputs)const;
 };
 
 USTRUCT()
diff --git a/Engine/Plugins/FX/Niagara/Source/Niagara/Public/NiagaraDataChannel.h b/Engine/Plugins/FX/Niagara/Source/Niagara/Public/NiagaraDataChannel.h
--- a/Engine/Plugins/FX/Niagara/Source/Niagara/Public/NiagaraDataChannel.h
+++ b/Engine/Plugins/FX/Niagara/Source/Niagara/Public/NiagaraDataChannel.h
@@ -58,14 +58,28 @@ class UNiagaraDataChannelWriter;
 class UNiagaraDataChannelReader;
 struct FNiagaraDataChannelPublishRequest;
 struct FNiagaraDataChannelGameDataLayout;
+class FNiagaraGpuReadbackManager;
 
 //////////////////////////////////////////////////////////////////////////
 
 /** Render thread proxy of FNiagaraDataChannelData. */
 struct FNiagaraDataChannelDataProxy
 {
+	~FNiagaraDataChannelDataProxy();
+
+	TWeakPtr<FNiagaraDataChannelData> Owner;
 	FNiagaraDataSet* GPUDataSet = nullptr;
+	FNiagaraDataBufferRef CurrFrameData = nullptr;
 	FNiagaraDataBufferRef PrevFrameData = nullptr;
+	bool bNeedsPrevFrameData = false;
+
+	//Buffers coming from the CPU that we're going to copy up for reading on the GPU
+	TArray<FNiagaraDataBufferRef> PendingCPUBuffers;
+	
+	//Users that need space in this NDC Data add to this for each tick via AddGPUAllocationForNextTick().
+	int32 PendingGPUAllocations = 0;
+
+	FNiagaraGpuComputeDispatchInterface* DispatchInterface = nullptr;
 
 	#if !UE_BUILD_SHIPPING
 	FString DebugName;
@@ -74,9 +88,18 @@ struct FNiagaraDataChannelDataProxy
 	const TCHAR* GetDebugName()const{return nullptr;}
 	#endif
 
-	void BeginFrame(bool bKeepPreviousFrameData);
-	void EndFrame(FNiagaraGpuComputeDispatchInterface* DispathInterface, FRHICommandListImmediate& CmdList, const TArray<FNiagaraDataBufferRef>& BuffersForGPU);
+	void Init();
+	void BeginFrame(FRHICommandListImmediate& RHICmdList);
+	void EndFrame(FRHICommandListImmediate& RHICmdList);
 	void Reset();
+
+	FNiagaraDataBufferRef AllocateBufferForCPU(FRHICommandListImmediate& RHICmdList, ERHIFeatureLevel::Type FeatureLevel, int32 AllocationSize);
+	void EnqueueReadbackForCPUBuffer(FRHICommandList& RHICmdList, FNiagaraDataBufferRef Buffer, FNiagaraGpuReadbackManager* ReadbackManager, FNiagaraGPUInstanceCountManager& InstanceCountManager, bool bPublishToGame, bool bPublishToCPU, FVector3f LWCTile);
+	void AddBuffersFromCPU(const TArray<FNiagaraDataBufferRef>& BuffersFromCPU);	
+	void AddGPUAllocationForNextTick(int32 AllocationCount);
+
+	FNiagaraDataBufferRef GetCurrentData()const { return CurrFrameData; }
+	FNiagaraDataBufferRef GetPrevFrameData()const { return PrevFrameData; }
 };
 
 DECLARE_MULTICAST_DELEGATE_OneParam(FOnDataChannelCreated, const UNiagaraDataChannel*);
diff --git a/Engine/Plugins/FX/Niagara/Source/Niagara/Public/NiagaraDataChannelAccessor.h b/Engine/Plugins/FX/Niagara/Source/Niagara/Public/NiagaraDataChannelAccessor.h
--- a/Engine/Plugins/FX/Niagara/Source/Niagara/Public/NiagaraDataChannelAccessor.h
+++ b/Engine/Plugins/FX/Niagara/Source/Niagara/Public/NiagaraDataChannelAccessor.h
@@ -30,6 +30,8 @@ private:
 
 public:
 
+	void Cleanup();
+
 	UPROPERTY()
 	TObjectPtr<UNiagaraDataChannelHandler> Owner;
 	
@@ -100,6 +102,8 @@ public:
 		}
 	}
 
+	void Cleanup();
+
 	UPROPERTY()
 	TObjectPtr<UNiagaraDataChannelHandler> Owner;
 	
diff --git a/Engine/Plugins/FX/Niagara/Source/Niagara/Public/NiagaraDataChannelHandler.h b/Engine/Plugins/FX/Niagara/Source/Niagara/Public/NiagaraDataChannelHandler.h
--- a/Engine/Plugins/FX/Niagara/Source/Niagara/Public/NiagaraDataChannelHandler.h
+++ b/Engine/Plugins/FX/Niagara/Source/Niagara/Public/NiagaraDataChannelHandler.h
@@ -20,11 +20,15 @@ class UNiagaraDataChannelHandler : public UObject
 public:
 
 	GENERATED_BODY()
-
-	NIAGARA_API virtual ~UNiagaraDataChannelHandler();
+	
+	//UObject Interface
+	NIAGARA_API virtual void BeginDestroy()override;
+	//UObject Interface END
 
 	NIAGARA_API virtual void Init(const UNiagaraDataChannel* InChannel);
 
+	NIAGARA_API virtual void Cleanup();
+
 	NIAGARA_API virtual void BeginFrame(float DeltaTime, FNiagaraWorldManager* OwningWorld);
 
 	NIAGARA_API virtual void EndFrame(float DeltaTime, FNiagaraWorldManager* OwningWorld);
diff --git a/Engine/Plugins/FX/Niagara/Source/Niagara/Public/NiagaraDataChannel_Global.h b/Engine/Plugins/FX/Niagara/Source/Niagara/Public/NiagaraDataChannel_Global.h
--- a/Engine/Plugins/FX/Niagara/Source/Niagara/Public/NiagaraDataChannel_Global.h
+++ b/Engine/Plugins/FX/Niagara/Source/Niagara/Public/NiagaraDataChannel_Global.h
@@ -27,11 +27,8 @@ class UNiagaraDataChannelHandler_Global : public UNiagaraDataChannelHandler
 
 	FNiagaraDataChannelDataPtr Data;
 
-	//UObject Interface
-	NIAGARA_API virtual void BeginDestroy()override;
-	//UObject Interface End
-
 	NIAGARA_API virtual void Init(const UNiagaraDataChannel* InChannel) override;
+	NIAGARA_API virtual void Cleanup() override;
 	NIAGARA_API virtual void BeginFrame(float DeltaTime, FNiagaraWorldManager* OwningWorld)override;
 	NIAGARA_API virtual void EndFrame(float DeltaTime, FNiagaraWorldManager* OwningWorld)override;
 	NIAGARA_API virtual void Tick(float DeltaTime, ETickingGroup TickGroup, FNiagaraWorldManager* OwningWorld) override;
diff --git a/Engine/Plugins/FX/Niagara/Source/Niagara/Public/NiagaraDataChannel_Islands.h b/Engine/Plugins/FX/Niagara/Source/Niagara/Public/NiagaraDataChannel_Islands.h
--- a/Engine/Plugins/FX/Niagara/Source/Niagara/Public/NiagaraDataChannel_Islands.h
+++ b/Engine/Plugins/FX/Niagara/Source/Niagara/Public/NiagaraDataChannel_Islands.h
@@ -176,7 +176,7 @@ protected:
 
 	/** How many pre-allocated islands to keep in the pool. Higher values will incur a larger standing memory cost but will reduce activation times for new islands. */
 	UPROPERTY(EditAnywhere, Category = "Islands")
-	int32 IslandPoolSize = 16;
+	int32 IslandPoolSize = 4;
 
 	UPROPERTY(EditAnywhere, Category = "Debug Rendering")
 	FNDCIslandDebugDrawSettings DebugDrawSettings;
@@ -194,11 +194,8 @@ class UNiagaraDataChannelHandler_Islands : public UNiagaraDataChannelHandler
 {
 	GENERATED_UCLASS_BODY()
 
-	//UObject Interface
-	NIAGARA_API virtual void BeginDestroy()override;
-	//UObject Interface End
-
 	NIAGARA_API virtual void Init(const UNiagaraDataChannel* InChannel) override;
+	NIAGARA_API virtual void Cleanup() override;
 	NIAGARA_API virtual void BeginFrame(float DeltaTime, FNiagaraWorldManager* OwningWorld)override;
 	NIAGARA_API virtual void EndFrame(float DeltaTime, FNiagaraWorldManager* OwningWorld)override;
 	NIAGARA_API virtual void Tick(float DeltaTime, ETickingGroup TickGroup, FNiagaraWorldManager* OwningWorld) override;
diff --git a/Engine/Plugins/FX/Niagara/Source/Niagara/Public/NiagaraDataInterfaceUtilities.h b/Engine/Plugins/FX/Niagara/Source/Niagara/Public/NiagaraDataInterfaceUtilities.h
--- a/Engine/Plugins/FX/Niagara/Source/Niagara/Public/NiagaraDataInterfaceUtilities.h
+++ b/Engine/Plugins/FX/Niagara/Source/Niagara/Public/NiagaraDataInterfaceUtilities.h
@@ -57,6 +57,7 @@ struct FNiagaraDataInterfaceHlslGenerationContext
 	DECLARE_DELEGATE_RetVal_OneParam(FString, FGetStructHlslTypeName, const FNiagaraTypeDefinition& /*Type*/)
 	DECLARE_DELEGATE_RetVal_OneParam(FString, FGetPropertyHlslTypeName, const FProperty* /*Property*/)
 	DECLARE_DELEGATE_RetVal_TwoParams(FString, FGetSanitizedSymbolName, FStringView /*SymbolName*/, bool /*bCollapsNamespaces*/)
+	DECLARE_DELEGATE_RetVal_OneParam(FString, FGetHlslDefaultForType, const FNiagaraTypeDefinition& /*Type*/)
 
 	FNiagaraDataInterfaceHlslGenerationContext(const FNiagaraDataInterfaceGPUParamInfo& InParameterInfo, TArrayView<const FNiagaraFunctionSignature> InSignatures)
 		: ParameterInfo(InParameterInfo), Signatures(InSignatures)
@@ -71,6 +72,7 @@ struct FNiagaraDataInterfaceHlslGenerationContext
 	FGetStructHlslTypeName GetStructHlslTypeNameDelegate;
 	FGetPropertyHlslTypeName GetPropertyHlslTypeNameDelegate;
 	FGetSanitizedSymbolName GetSanitizedSymbolNameDelegate;
+	FGetHlslDefaultForType GetHlslDefaultForTypeDelegate;
 
 	const FNiagaraDataInterfaceGeneratedFunction& GetFunctionInfo() const { return ParameterInfo.GeneratedFunctions[FunctionInstanceIndex]; }
 
@@ -78,6 +80,7 @@ struct FNiagaraDataInterfaceHlslGenerationContext
 	FString GetStructHlslTypeName(const FNiagaraTypeDefinition& Type) { return GetStructHlslTypeNameDelegate.Execute(Type); }
 	FString GetPropertyHlslTypeName(const FProperty* Property) { return GetPropertyHlslTypeNameDelegate.Execute(Property); }
 	FString GetSanitizedSymbolName(FStringView SymbolName, bool bCollapsNamespaces = false) { return GetSanitizedSymbolNameDelegate.Execute(SymbolName, bCollapsNamespaces); }
+	FString GetHlslDefaultForType(const FNiagaraTypeDefinition& Type) { return GetHlslDefaultForTypeDelegate.Execute(Type); }
 };
 
 #endif
\ No newline at end of file
diff --git a/Engine/Plugins/FX/Niagara/Source/Niagara/Public/NiagaraDataSetReadback.h b/Engine/Plugins/FX/Niagara/Source/Niagara/Public/NiagaraDataSetReadback.h
--- a/Engine/Plugins/FX/Niagara/Source/Niagara/Public/NiagaraDataSetReadback.h
+++ b/Engine/Plugins/FX/Niagara/Source/Niagara/Public/NiagaraDataSetReadback.h
@@ -9,6 +9,7 @@
 struct FNiagaraComputeExecutionContext;
 class FNiagaraEmitterInstance;
 class FNiagaraGpuComputeDispatchInterface;
+class FNiagaraGpuReadbackManager;
 
 class FNiagaraDataSetReadback : public TSharedFromThis<FNiagaraDataSetReadback, ESPMode::ThreadSafe>
 {
@@ -42,3 +43,25 @@ private:
 
 	FOnReadbackReady		OnReadbackReady;
 };
+
+
+
+class FNiagaraDataBufferReadback : public TSharedFromThis<FNiagaraDataBufferReadback, ESPMode::ThreadSafe>
+{
+public:
+	DECLARE_DELEGATE_OneParam(FOnReadbackComplete, FNiagaraDataBufferRef);
+
+public:
+	FOnReadbackComplete& GetOnReadbackComplete() { return OnReadbackComplete; }
+
+	void EnqueueReadback(FRHICommandList& RHICmdList, FNiagaraDataBufferRef InDataBuffer, FNiagaraGpuReadbackManager* ReadbackManager, FNiagaraGPUInstanceCountManager& InstanceCountManager);
+
+private:
+	void ReadbackCompleteInternal(FNiagaraDataBuffer* CompleteDataBuffer);
+	void GPUReadbackInternal(FRHICommandList& RHICmdList, FNiagaraGpuReadbackManager* ReadbackManager, FNiagaraGPUInstanceCountManager& InstanceCountManager, FNiagaraDataBufferRef SrcDataBuffer, FNiagaraDataBuffer* DestDataBuffer);
+
+private:
+	std::atomic<int>		PendingReadbacks;
+
+	FOnReadbackComplete		OnReadbackComplete;
+};
diff --git a/Engine/Plugins/FX/Niagara/Source/Niagara/Public/NiagaraDebuggerCommon.h b/Engine/Plugins/FX/Niagara/Source/Niagara/Public/NiagaraDebuggerCommon.h
--- a/Engine/Plugins/FX/Niagara/Source/Niagara/Public/NiagaraDebuggerCommon.h
+++ b/Engine/Plugins/FX/Niagara/Source/Niagara/Public/NiagaraDebuggerCommon.h
@@ -126,6 +126,9 @@ struct FNiagaraOutlinerSystemInstanceData
 	UPROPERTY(VisibleAnywhere, Category = "Gpu")
 	uint32 bRequiresRayTracingScene : 1;
 
+	UPROPERTY(VisibleAnywhere, Category = "Gpu")
+	uint32 bRequiresCurrentFrameNDC : 1;
+
 	FNiagaraOutlinerSystemInstanceData()
 		: bPendingKill(false)
 		, bUsingCullProxy(false)
@@ -137,6 +140,7 @@ struct FNiagaraOutlinerSystemInstanceData
 		, bRequiresEarlyViewData(false)
 		, bRequiresViewUniformBuffer(false)
 		, bRequiresRayTracingScene(false)
+		, bRequiresCurrentFrameNDC(false)
 	{}
 };
 
diff --git a/Engine/Plugins/FX/Niagara/Source/Niagara/Public/NiagaraGpuComputeDispatchInterface.h b/Engine/Plugins/FX/Niagara/Source/Niagara/Public/NiagaraGpuComputeDispatchInterface.h
--- a/Engine/Plugins/FX/Niagara/Source/Niagara/Public/NiagaraGpuComputeDispatchInterface.h
+++ b/Engine/Plugins/FX/Niagara/Source/Niagara/Public/NiagaraGpuComputeDispatchInterface.h
@@ -23,6 +23,7 @@ class FNiagaraGpuReadbackManager;
 class FNiagaraRayTracingHelper;
 struct FNiagaraScriptDebuggerInfo;
 class FNiagaraSystemGpuComputeProxy;
+struct FNiagaraDataChannelDataProxy;
 
 // Public API for Niagara's Compute Dispatcher
 // This is generally used with DataInterfaces or Custom Renderers
@@ -49,6 +50,11 @@ public:
 	/** Remove system instance proxy from the batcher. */
 	virtual void RemoveGpuComputeProxy(FNiagaraSystemGpuComputeProxy* ComputeProxy) = 0;
 
+	/** Add NDC Data to the batcher for tracking */
+	virtual void AddNDCDataProxy(FNiagaraDataChannelDataProxy* NDCDataProxy) = 0;
+	/** Add NDC Data to the batcher for tracking */
+	virtual void RemoveNDCDataProxy(FNiagaraDataChannelDataProxy* NDCDataProxy) = 0;
+
 	/**
 	 * Register work for GPU sorting (using the GPUSortManager).
 	 * The constraints of the sort request are defined in SortInfo.SortFlags.
diff --git a/Engine/Plugins/FX/Niagara/Source/Niagara/Public/NiagaraSystemGpuComputeProxy.h b/Engine/Plugins/FX/Niagara/Source/Niagara/Public/NiagaraSystemGpuComputeProxy.h
--- a/Engine/Plugins/FX/Niagara/Source/Niagara/Public/NiagaraSystemGpuComputeProxy.h
+++ b/Engine/Plugins/FX/Niagara/Source/Niagara/Public/NiagaraSystemGpuComputeProxy.h
@@ -33,6 +33,7 @@ public:
 	bool RequiresEarlyViewData() const { return bRequiresEarlyViewData; }
 	bool RequiresViewUniformBuffer() const { return bRequiresViewUniformBuffer; }
 	bool RequiresRayTracingScene() const { return bRequiresRayTracingScene; }
+	bool RequiresCurrentFrameNDC() const { return bRequiresCurrentFrameNDC; }
 	FVector3f GetSystemLWCTile() const { return SystemLWCTile; }
 
 private:
@@ -48,6 +49,7 @@ private:
 	uint32										bRequiresEarlyViewData : 1;
 	uint32										bRequiresViewUniformBuffer : 1;
 	uint32										bRequiresRayTracingScene : 1;
+	uint32										bRequiresCurrentFrameNDC: 1;
 
 	FShaderResourceViewRHIRef					StaticFloatBuffer;
 
diff --git a/Engine/Plugins/FX/Niagara/Source/Niagara/Public/NiagaraSystemInstance.h b/Engine/Plugins/FX/Niagara/Source/Niagara/Public/NiagaraSystemInstance.h
--- a/Engine/Plugins/FX/Niagara/Source/Niagara/Public/NiagaraSystemInstance.h
+++ b/Engine/Plugins/FX/Niagara/Source/Niagara/Public/NiagaraSystemInstance.h
@@ -188,6 +188,7 @@ public:
 	NIAGARA_API bool RequiresEarlyViewData() const;
 	NIAGARA_API bool RequiresViewUniformBuffer() const;
 	NIAGARA_API bool RequiresRayTracingScene() const;
+	NIAGARA_API bool RequiresCurrentFrameNDC() const;
 
 	/** Requests the the simulation be reset on the next tick. */
 	NIAGARA_API void Reset(EResetMode Mode);
diff --git a/Engine/Plugins/FX/Niagara/Source/Niagara/Public/NiagaraWorldManager.h b/Engine/Plugins/FX/Niagara/Source/Niagara/Public/NiagaraWorldManager.h
--- a/Engine/Plugins/FX/Niagara/Source/Niagara/Public/NiagaraWorldManager.h
+++ b/Engine/Plugins/FX/Niagara/Source/Niagara/Public/NiagaraWorldManager.h
@@ -92,6 +92,8 @@ public:
 
 	void Init(UWorld* InWorld);
 
+	void OnPostWorldInit();
+
 	static NIAGARA_API FNiagaraWorldManager* Get(const UWorld* World);
 	static void OnStartup();
 	static void OnShutdown();
@@ -297,6 +299,9 @@ public:
 private:
 	// Callback function registered with global world delegates to instantiate world manager when a game world is created
 	static void OnWorldInit(UWorld* World, const UWorld::InitializationValues IVS);
+	
+	// Callback function registered with global world delegates to instantiate world manager after a game world is initialized
+	static void OnPostWorldInit(UWorld* World, const UWorld::InitializationValues IVS);
 
 	// Callback function registered with global world delegates to cleanup world manager contents
 	static void OnWorldCleanup(UWorld* World, bool bSessionEnded, bool bCleanupResources);
@@ -350,6 +355,7 @@ private:
 	void HandleCSVStats(float DeltaSeconds);
 
 	static FDelegateHandle OnWorldInitHandle;
+	static FDelegateHandle OnPostWorldInitHandle;
 	static FDelegateHandle OnWorldCleanupHandle;
 	static FDelegateHandle OnPostWorldCleanupHandle;
 	static FDelegateHandle OnPreWorldFinishDestroyHandle;
diff --git a/Engine/Plugins/FX/Niagara/Source/NiagaraEditor/Private/NiagaraHlslTranslator.cpp b/Engine/Plugins/FX/Niagara/Source/NiagaraEditor/Private/NiagaraHlslTranslator.cpp
--- a/Engine/Plugins/FX/Niagara/Source/NiagaraEditor/Private/NiagaraHlslTranslator.cpp
+++ b/Engine/Plugins/FX/Niagara/Source/NiagaraEditor/Private/NiagaraHlslTranslator.cpp
@@ -3083,6 +3083,7 @@ void FNiagaraHlslTranslator::DefineDataInterfaceHLSL(FString& InHlslOutput)
 			DIHlslGenContext.GetStructHlslTypeNameDelegate.BindStatic(&FNiagaraHlslTranslator::GetStructHlslTypeName);
 			DIHlslGenContext.GetPropertyHlslTypeNameDelegate.BindStatic(&FNiagaraHlslTranslator::GetPropertyHlslTypeName);
 			DIHlslGenContext.GetSanitizedSymbolNameDelegate.BindStatic(&FNiagaraHlslTranslator::GetSanitizedSymbolName);
+			DIHlslGenContext.GetHlslDefaultForTypeDelegate.BindStatic(&FNiagaraHlslTranslator::GetHlslDefaultForType);
 			CDO->GetParameterDefinitionHLSL(DIHlslGenContext, InterfaceUniformHLSL);
 
 			// Ask the DI to generate HLSL.
@@ -9315,7 +9316,6 @@ FString FNiagaraHlslTranslator::GetFunctionSignatureSymbol(const FNiagaraFunctio
 	TArray<FNiagaraVariableBase> VariadicParams;
 	VariadicParams.Reserve(Sig.NumOptionalInputs() + Sig.NumOptionalOutputs());
 	Sig.GetVariadicInputs(VariadicParams);
-	AddVarsToSig(VariadicParams);
 	Sig.GetVariadicOutputs(VariadicParams);
 	AddVarsToSig(VariadicParams);
 
diff --git a/Engine/Plugins/FX/Niagara/Source/NiagaraEditor/Private/Widgets/SNiagaraSystemViewport.cpp b/Engine/Plugins/FX/Niagara/Source/NiagaraEditor/Private/Widgets/SNiagaraSystemViewport.cpp
--- a/Engine/Plugins/FX/Niagara/Source/NiagaraEditor/Private/Widgets/SNiagaraSystemViewport.cpp
+++ b/Engine/Plugins/FX/Niagara/Source/NiagaraEditor/Private/Widgets/SNiagaraSystemViewport.cpp
@@ -523,6 +523,11 @@ void FNiagaraSystemViewportClient::DrawGpuTickInformation(UNiagaraComponent* Com
 			Canvas->DrawShadowedString(CurrentX + 5.0f, CurrentY, TEXT("RequiresRayTracingScene"), Font, FLinearColor::White);
 			CurrentY += FontHeight;
 		}
+		if (SystemInstance->RequiresCurrentFrameNDC())
+		{
+			Canvas->DrawShadowedString(CurrentX + 5.0f, CurrentY, TEXT("RequiresCurrentFrameNDC"), Font, FLinearColor::White);
+			CurrentY += FontHeight;
+		}
 	}
 }
 
diff --git a/Engine/Plugins/FX/Niagara/Source/NiagaraVertexFactories/Private/NiagaraDrawIndirect.cpp b/Engine/Plugins/FX/Niagara/Source/NiagaraVertexFactories/Private/NiagaraDrawIndirect.cpp
--- a/Engine/Plugins/FX/Niagara/Source/NiagaraVertexFactories/Private/NiagaraDrawIndirect.cpp
+++ b/Engine/Plugins/FX/Niagara/Source/NiagaraVertexFactories/Private/NiagaraDrawIndirect.cpp
@@ -10,6 +10,7 @@ NiagaraDrawIndirect.cpp : Niagara shader to generate the draw indirect args for
 
 IMPLEMENT_GLOBAL_SHADER(FNiagaraDrawIndirectArgsGenCS, "/Plugin/FX/Niagara/Private/NiagaraDrawIndirectArgsGen.usf", "MainCS", SF_Compute);
 IMPLEMENT_GLOBAL_SHADER(FNiagaraDrawIndirectResetCountsCS, "/Plugin/FX/Niagara/Private/NiagaraDrawIndirectArgsGen.usf", "ResetCountsCS", SF_Compute);
+IMPLEMENT_GLOBAL_SHADER(FNiagaraInstanceCountsInitCS, "/Plugin/FX/Niagara/Private/NiagaraDrawIndirectArgsGen.usf", "InitInstanceCountsCS", SF_Compute);
 
 void FNiagaraDrawIndirectArgsGenCS::ModifyCompilationEnvironment(const FGlobalShaderPermutationParameters& Parameters, FShaderCompilerEnvironment& OutEnvironment)
 {
@@ -17,6 +18,7 @@ void FNiagaraDrawIndirectArgsGenCS::ModifyCompilationEnvironment(const FGlobalSh
 	OutEnvironment.SetDefine(TEXT("THREAD_COUNT"), NIAGARA_DRAW_INDIRECT_ARGS_GEN_THREAD_COUNT);
 	OutEnvironment.SetDefine(TEXT("NIAGARA_DRAW_INDIRECT_ARGS_SIZE"), NIAGARA_DRAW_INDIRECT_ARGS_SIZE);
 	OutEnvironment.SetDefine(TEXT("NIAGARA_DRAW_INDIRECT_TASK_INFO_SIZE"), NIAGARA_DRAW_INDIRECT_TASK_INFO_SIZE);
+	OutEnvironment.SetDefine(TEXT("NIAGARA_INIT_INSTANCE_COUNT_TASK_INFO_SIZE"), NIAGARA_INIT_INSTANCE_COUNT_TASK_INFO_SIZE);
 }
 
 void FNiagaraDrawIndirectResetCountsCS::ModifyCompilationEnvironment(const FGlobalShaderPermutationParameters& Parameters, FShaderCompilerEnvironment& OutEnvironment)
@@ -25,4 +27,14 @@ void FNiagaraDrawIndirectResetCountsCS::ModifyCompilationEnvironment(const FGlob
 	OutEnvironment.SetDefine(TEXT("THREAD_COUNT"), NIAGARA_DRAW_INDIRECT_ARGS_GEN_THREAD_COUNT);
 	OutEnvironment.SetDefine(TEXT("NIAGARA_DRAW_INDIRECT_ARGS_SIZE"), NIAGARA_DRAW_INDIRECT_ARGS_SIZE);
 	OutEnvironment.SetDefine(TEXT("NIAGARA_DRAW_INDIRECT_TASK_INFO_SIZE"), NIAGARA_DRAW_INDIRECT_TASK_INFO_SIZE);
+	OutEnvironment.SetDefine(TEXT("NIAGARA_INIT_INSTANCE_COUNT_TASK_INFO_SIZE"), NIAGARA_INIT_INSTANCE_COUNT_TASK_INFO_SIZE);
+}
+
+void FNiagaraInstanceCountsInitCS::ModifyCompilationEnvironment(const FGlobalShaderPermutationParameters& Parameters, FShaderCompilerEnvironment& OutEnvironment)
+{
+	FGlobalShader::ModifyCompilationEnvironment(Parameters, OutEnvironment);
+	OutEnvironment.SetDefine(TEXT("THREAD_COUNT"), NIAGARA_DRAW_INDIRECT_ARGS_GEN_THREAD_COUNT);
+	OutEnvironment.SetDefine(TEXT("NIAGARA_DRAW_INDIRECT_ARGS_SIZE"), NIAGARA_DRAW_INDIRECT_ARGS_SIZE);
+	OutEnvironment.SetDefine(TEXT("NIAGARA_DRAW_INDIRECT_TASK_INFO_SIZE"), NIAGARA_DRAW_INDIRECT_TASK_INFO_SIZE);
+	OutEnvironment.SetDefine(TEXT("NIAGARA_INIT_INSTANCE_COUNT_TASK_INFO_SIZE"), NIAGARA_INIT_INSTANCE_COUNT_TASK_INFO_SIZE);
 }
diff --git a/Engine/Plugins/FX/Niagara/Source/NiagaraVertexFactories/Public/NiagaraDrawIndirect.h b/Engine/Plugins/FX/Niagara/Source/NiagaraVertexFactories/Public/NiagaraDrawIndirect.h
--- a/Engine/Plugins/FX/Niagara/Source/NiagaraVertexFactories/Public/NiagaraDrawIndirect.h
+++ b/Engine/Plugins/FX/Niagara/Source/NiagaraVertexFactories/Public/NiagaraDrawIndirect.h
@@ -15,6 +15,7 @@ NiagaraDrawIndirect.h: Niagara shader to generate the draw indirect args for Nia
 #define NIAGARA_DRAW_INDIRECT_ARGS_GEN_THREAD_COUNT 64
 #define NIAGARA_DRAW_INDIRECT_ARGS_SIZE 5
 #define NIAGARA_DRAW_INDIRECT_TASK_INFO_SIZE 5
+#define NIAGARA_INIT_INSTANCE_COUNT_TASK_INFO_SIZE 2
 
 // #define NIAGARA_COPY_BUFFER_THREAD_COUNT 64
 // #define NIAGARA_COPY_BUFFER_BUFFER_COUNT 3
@@ -103,4 +104,24 @@ public:
 	END_SHADER_PARAMETER_STRUCT()
 
 	static void ModifyCompilationEnvironment(const FGlobalShaderPermutationParameters& Parameters, FShaderCompilerEnvironment& OutEnvironment);
-};
\ No newline at end of file
+};
+
+/**
+ * Compute shader used to initialize instance count entries to specific values. Used if the platform doesn't support RW texture buffers
+ */
+class FNiagaraInstanceCountsInitCS : public FGlobalShader
+{
+	DECLARE_EXPORTED_GLOBAL_SHADER(FNiagaraInstanceCountsInitCS, NIAGARAVERTEXFACTORIES_API);
+	SHADER_USE_PARAMETER_STRUCT(FNiagaraInstanceCountsInitCS, FGlobalShader);
+
+public:
+	using FPermutationDomain = TShaderPermutationDomain<>;
+
+	BEGIN_SHADER_PARAMETER_STRUCT(FParameters, NIAGARAVERTEXFACTORIES_API)
+		SHADER_PARAMETER_SRV(Buffer<uint>,		TaskInfos)
+		SHADER_PARAMETER_UAV(RWBuffer<uint>,	RWInstanceCounts)
+		SHADER_PARAMETER(FUintVector4,			TaskCount)
+	END_SHADER_PARAMETER_STRUCT()
+
+	static void ModifyCompilationEnvironment(const FGlobalShaderPermutationParameters& Parameters, FShaderCompilerEnvironment& OutEnvironment);
+};
