Linux fix - Changes loop var from FNiagaraVariable to FNiagaraVariable base to match conflicting changes in another CL.

diff --git a/Engine/Plugins/FX/Niagara/Source/Niagara/Private/DataInterface/NiagaraDataInterfaceDataChannelCommon.cpp b/Engine/Plugins/FX/Niagara/Source/Niagara/Private/DataInterface/NiagaraDataInterfaceDataChannelCommon.cpp
--- a/Engine/Plugins/FX/Niagara/Source/Niagara/Private/DataInterface/NiagaraDataInterfaceDataChannelCommon.cpp
+++ b/Engine/Plugins/FX/Niagara/Source/Niagara/Private/DataInterface/NiagaraDataInterfaceDataChannelCommon.cpp
@@ -540,7 +540,7 @@ void Write_{FunctionParameterName}_{ParameterName}(FNDCAccessContext_{ParameterN
 			{
 				Ret = HashCombine(Ret, GetTypeHash(Input));
 			}
-			for (const FNiagaraVariable& Output : Sig.Outputs)
+			for (const FNiagaraVariableBase& Output : Sig.Outputs)
 			{
 				Ret = HashCombine(Ret, GetTypeHash(Output));
 			}
