Properly sanitising parameter names in hlsl gen.

diff --git a/Engine/Plugins/FX/Niagara/Source/Niagara/Private/DataInterface/NiagaraDataInterfaceDataChannelCommon.cpp b/Engine/Plugins/FX/Niagara/Source/Niagara/Private/DataInterface/NiagaraDataInterfaceDataChannelCommon.cpp
--- a/Engine/Plugins/FX/Niagara/Source/Niagara/Private/DataInterface/NiagaraDataInterfaceDataChannelCommon.cpp
+++ b/Engine/Plugins/FX/Niagara/Source/Niagara/Private/DataInterface/NiagaraDataInterfaceDataChannelCommon.cpp
@@ -606,7 +606,7 @@ void Write_{FunctionParameterName}_{ParameterName}(FNDCAccessContext_{ParameterN
 			Access.SortedOffset = ParamIdx;
 
 			FunctionParameterIndex = LexToString(ParamIdx);
-			FunctionParameterName = Param.GetName().ToString();
+			FunctionParameterName = HlslGenContext.GetSanitizedSymbolName(Param.GetName().ToString());
 			FunctionParameterType = HlslGenContext.GetStructHlslTypeName(Param.GetType());
 
 			FunctionParameterComponentName = HlslGenContext.GetSanitizedSymbolName(Param.GetName().ToString());
@@ -664,7 +664,7 @@ void Write_{FunctionParameterName}_{ParameterName}(FNDCAccessContext_{ParameterN
 						const FNiagaraVariable& InputParam = Signature.Inputs[InputIdx];
 						FParamAccessInfo& Access = ParametersAccessed.FindChecked(InputParam);
 
-						FunctionParameterName = InputParam.GetName().ToString();
+						FunctionParameterName = HlslGenContext.GetSanitizedSymbolName(InputParam.GetName().ToString());
 						FunctionParameterType = HlslGenContext.GetStructHlslTypeName(InputParam.GetType());
 						PerFunctionParameterShaderCode.StringValue += FString::Format(*PerParameterWriteCallTemplate, HlslTemplateArgs);
 
@@ -687,7 +687,7 @@ void Write_{FunctionParameterName}_{ParameterName}(FNDCAccessContext_{ParameterN
 						const FNiagaraVariable& OutputParam = Signature.Outputs[OutputIdx];
 						FParamAccessInfo& Access = ParametersAccessed.FindChecked(OutputParam);
 						
-						FunctionParameterName = OutputParam.GetName().ToString();
+						FunctionParameterName = HlslGenContext.GetSanitizedSymbolName(OutputParam.GetName().ToString());
 						FunctionParameterType = HlslGenContext.GetStructHlslTypeName(OutputParam.GetType());
 						PerFunctionParameterShaderCode.StringValue += FString::Format(*PerParameterReadCallTemplate, HlslTemplateArgs);
 
