- Change ensure to log if the SM becomes invalid
- There are some cases due to compilation / building where this can occur but doesn't cause a problem
- Also handle a case if the material remap table was not built, I don't know if this can happen in reality but forcing the SM to null it can

#rnx
#tests Forced the SM to be null in code to see what paths failed and ensure we don't crash

diff --git a/Engine/Plugins/FX/Niagara/Source/Niagara/Private/NiagaraMeshRendererProperties.cpp b/Engine/Plugins/FX/Niagara/Source/Niagara/Private/NiagaraMeshRendererProperties.cpp
--- a/Engine/Plugins/FX/Niagara/Source/Niagara/Private/NiagaraMeshRendererProperties.cpp
+++ b/Engine/Plugins/FX/Niagara/Source/Niagara/Private/NiagaraMeshRendererProperties.cpp
@@ -154,8 +154,9 @@ public:
 	void GetUsedMaterials(TArray<UMaterialInterface*>& OutMaterials) const override
 	{
 		const UStaticMesh* StaticMesh = WeakStaticMesh.Get();
-		if (!ensure(StaticMesh))
+		if (StaticMesh == nullptr)
 		{
+			UE_LOG(LogNiagara, Log, TEXT("FNiagaraRenderableStaticMesh - StaticMesh is no longer valid"));
 			return;
 		}
 
diff --git a/Engine/Plugins/FX/Niagara/Source/Niagara/Private/NiagaraRendererMeshes.cpp b/Engine/Plugins/FX/Niagara/Source/Niagara/Private/NiagaraRendererMeshes.cpp
--- a/Engine/Plugins/FX/Niagara/Source/Niagara/Private/NiagaraRendererMeshes.cpp
+++ b/Engine/Plugins/FX/Niagara/Source/Niagara/Private/NiagaraRendererMeshes.cpp
@@ -1534,7 +1534,7 @@ void FNiagaraRendererMeshes::GetDynamicMeshElements(const TArray<const FSceneVie
 					for (int32 SectionIndex = 0; SectionIndex < LODModel.Sections.Num(); SectionIndex++)
 					{
 						const FStaticMeshSection& Section = LODModel.Sections[SectionIndex];
-						const uint32 RemappedMaterialIndex = MeshData.MaterialRemapTable[Section.MaterialIndex];
+						const uint32 RemappedMaterialIndex = MeshData.MaterialRemapTable.IsValidIndex(Section.MaterialIndex) ? MeshData.MaterialRemapTable[Section.MaterialIndex] : INDEX_NONE;
 						FMaterialRenderProxy* MaterialProxy = ParticleMeshRenderData.DynamicDataMesh->Materials.IsValidIndex(RemappedMaterialIndex) ? ParticleMeshRenderData.DynamicDataMesh->Materials[RemappedMaterialIndex] : UMaterial::GetDefaultMaterial(MD_Surface)->GetRenderProxy();
 						if (Section.NumTriangles == 0 || MaterialProxy == nullptr)
 						{
