- Add missing transitions to indirect arg generation when using a count from the count manager

#rnx
#tests Ran Indirect ET before / after change with validation enabled

diff --git a/Engine/Plugins/FX/Niagara/Source/Niagara/Private/NiagaraDataInterfaceRW.cpp b/Engine/Plugins/FX/Niagara/Source/Niagara/Private/NiagaraDataInterfaceRW.cpp
--- a/Engine/Plugins/FX/Niagara/Source/Niagara/Private/NiagaraDataInterfaceRW.cpp
+++ b/Engine/Plugins/FX/Niagara/Source/Niagara/Private/NiagaraDataInterfaceRW.cpp
@@ -122,31 +122,43 @@ void FNDIGpuComputeDispatchArgsGenContext::FlushPass() const
 
 	// Create Args Buffer and Upload into it
 	FNiagaraDispatchIndirectArgsGenCS::FParameters* IndirectArgsGenParameters = &GetPassParameters()->IndirectArgsGenParameters;
+	TArray<FRHITransitionInfo, TInlineAllocator<1>> TransitionsBefore;
+	TArray<FRHITransitionInfo, TInlineAllocator<1>> TransitionsAfter;
 	if (NumIndirectCounterGenArgs > 0)
 	{
 		const uint32 DispatchInfosBufferSize = sizeof(FNiagaraDispatchIndirectInfoCS) * NumIndirectCounterGenArgs;
 		FRDGBufferRef DispatchInfosBuffer = GraphBuilder.CreateBuffer(FRDGBufferDesc::CreateBufferDesc(sizeof(uint32), DispatchInfosBufferSize / sizeof(uint32)), TEXT("NiagaraIndirectGenArgs"));
 		GraphBuilder.QueueBufferUpload(DispatchInfosBuffer, IndirectCounterGenArgs, DispatchInfosBufferSize, ERDGInitialDataFlags::NoCopy);
 
+		const FRWBuffer& InstanceCountBuffer = ComputeDispatchInterface.GetGPUInstanceCounterManager().GetInstanceCountBuffer();
+		TransitionsBefore.Emplace(InstanceCountBuffer.Buffer, ERHIAccess::UAVCompute, ERHIAccess::SRVCompute);
+		TransitionsAfter.Emplace(InstanceCountBuffer.Buffer, ERHIAccess::SRVCompute, ERHIAccess::UAVCompute);
+
 		IndirectArgsGenParameters->DispatchInfos = GraphBuilder.CreateSRV(DispatchInfosBuffer, PF_R32_UINT);
 		IndirectArgsGenParameters->NumDispatchInfos = NumIndirectCounterGenArgs;
 		IndirectArgsGenParameters->MaxGroupsPerDimension = FUintVector3(GRHIMaxDispatchThreadGroupsPerDimension);
 
-		IndirectArgsGenParameters->InstanceCounts = ComputeDispatchInterface.GetGPUInstanceCounterManager().GetInstanceCountBuffer().SRV;
+		IndirectArgsGenParameters->InstanceCounts = InstanceCountBuffer.SRV;
 		IndirectArgsGenParameters->RWDispatchIndirectArgs = GraphBuilder.CreateUAV(IndirectBuffer, PF_R32_UINT);
 	}
 
+	const FRWBuffer& CountBuffer = ComputeDispatchInterface.GetGPUInstanceCounterManager().GetInstanceCountBuffer();
+
 	// Kick pass to genreate args
 	GraphBuilder.AddPass(
 		RDG_EVENT_NAME("Niagara::ExecuteTicks::DispatchGroupPre"),
 		GetPassParameters(),
 		ERDGPassFlags::Compute | ERDGPassFlags::NeverCull,
-		[PassParameters=GetPassParameters(), IndirectArgsGenParameters, CreateCallbacks=MoveTemp(IndirectCallbacks)](FRHICommandListImmediate& RHICmdList)
+		[PassParameters=GetPassParameters(), IndirectArgsGenParameters, CreateCallbacks=MoveTemp(IndirectCallbacks), TransitionsBefore, TransitionsAfter](FRHICommandListImmediate& RHICmdList)
 		{
 			if (IndirectArgsGenParameters != nullptr)
 			{
+				RHICmdList.Transition(TransitionsBefore);
+
 				TShaderMapRef<FNiagaraDispatchIndirectArgsGenCS> DispatchIndirectArgsGenCS(GetGlobalShaderMap(GMaxRHIFeatureLevel));
 				FComputeShaderUtils::Dispatch(RHICmdList, DispatchIndirectArgsGenCS, *IndirectArgsGenParameters, FIntVector(1, 1, 1));
+
+				RHICmdList.Transition(TransitionsAfter);
 			}
 			for (const auto& CreateCallback : CreateCallbacks)
 			{
