- Fix dangling pointer to dispatch interface from data channel proxy
- Destruction order in not guaranteed to be DC Proxy then compute dispatch

#rnx
#tests PIE in / out of level a few times to get the crash

diff --git a/Engine/Plugins/FX/Niagara/Source/Niagara/Internal/NiagaraDataChannelCommon.h b/Engine/Plugins/FX/Niagara/Source/Niagara/Internal/NiagaraDataChannelCommon.h
--- a/Engine/Plugins/FX/Niagara/Source/Niagara/Internal/NiagaraDataChannelCommon.h
+++ b/Engine/Plugins/FX/Niagara/Source/Niagara/Internal/NiagaraDataChannelCommon.h
@@ -133,7 +133,6 @@ private:
 	/** DataChannel data accessible to Niagara GPU sims. SoA layout. Non LWC types. */
 	FNiagaraDataSet* GPUSimData = nullptr;
 
-
 	/** Data buffers we'll be passing to the RT proxy for uploading to the GPU */
 	TArray<FNiagaraDataBufferRef> BuffersForGPU;
 
@@ -146,6 +145,9 @@ private:
 	/** Pending requests to publish data into this data channel from the GPU. To alleviate data race behavior with data coming back from the GPU, we always consume GPU requests at the start of the frame only. */
 	TArray<FNiagaraDataChannelPublishRequest> PublishRequestsFromGPU;
 
+	/** The world we were initialized with, used to get the compute interface. */
+	TWeakObjectPtr<UWorld> WeakOwnerWorld;
+
 	FVector3f LwcTile = FVector3f::ZeroVector;
 
 	/** Critical section protecting shared state for multiple writers publishing from different threads. */
diff --git a/Engine/Plugins/FX/Niagara/Source/Niagara/Private/NiagaraDataChannel.cpp b/Engine/Plugins/FX/Niagara/Source/Niagara/Private/NiagaraDataChannel.cpp
--- a/Engine/Plugins/FX/Niagara/Source/Niagara/Private/NiagaraDataChannel.cpp
+++ b/Engine/Plugins/FX/Niagara/Source/Niagara/Private/NiagaraDataChannel.cpp
@@ -458,8 +458,9 @@ void FNiagaraDataChannelGameData::SetFromSimCache(const FNiagaraVariableBase& So
 
 FNiagaraDataChannelDataProxy::~FNiagaraDataChannelDataProxy()
 {
-	check(DispatchInterface);
-	DispatchInterface->RemoveNDCDataProxy(this);
+#if !UE_BUILD_SHIPPING
+	ensureMsgf(DispatchInterfaceForDebuggingOnly == nullptr, TEXT("FNiagaraDataChannelDataProxy - ComputeDispatchInterface was not cleared this could result in a crash"));
+#endif
 }
 
 void FNiagaraDataChannelDataProxy::Reset()
@@ -467,16 +468,13 @@ void FNiagaraDataChannelDataProxy::Reset()
 	PrevFrameData = nullptr;
 }
 
-void FNiagaraDataChannelDataProxy::Init()
-{
-	check(DispatchInterface);
-	DispatchInterface->AddNDCDataProxy(this);
-}
-
-void FNiagaraDataChannelDataProxy::BeginFrame(FRHICommandListImmediate& RHICmdList)
+void FNiagaraDataChannelDataProxy::BeginFrame(FNiagaraGpuComputeDispatchInterface* DispatchInterface, FRHICommandListImmediate& RHICmdList)
 {
 	check(GPUDataSet);
 	check(GPUDataSet->GetSimTarget() == ENiagaraSimTarget::GPUComputeSim);
+#if !UE_BUILD_SHIPPING
+	check(DispatchInterfaceForDebuggingOnly == DispatchInterface);
+#endif
 
 	FNiagaraGPUInstanceCountManager& InstCountManager = DispatchInterface->GetGPUInstanceCounterManager();	
 	if(this->bNeedsPrevFrameData)
@@ -587,9 +585,12 @@ FNiagaraDataChannelData::~FNiagaraDataChannelData()
 {
 	GameData.Reset();
 
+	UWorld* OwnerWorld = WeakOwnerWorld.Get();
+	FNiagaraGpuComputeDispatchInterface* ComputeDispatchInterface = FNiagaraGpuComputeDispatchInterface::Get(OwnerWorld);
+
 	/** We defer the deletion of the dataset to the RT to be sure all in-flight RT commands have finished using it.*/
 	ENQUEUE_RENDER_COMMAND(FDeleteContextCommand)(
-		[CPUDataChannelDataSet = CPUSimData, GPUDataChannelDataSet = GPUSimData, ReleasedRTProxy = RTProxy.Release()](FRHICommandListImmediate& RHICmdList)
+		[ComputeDispatchInterface, CPUDataChannelDataSet = CPUSimData, GPUDataChannelDataSet = GPUSimData, ReleasedRTProxy = RTProxy.Release()](FRHICommandListImmediate& RHICmdList)
 	{
 		if (CPUDataChannelDataSet != nullptr)
 		{
@@ -601,6 +602,16 @@ FNiagaraDataChannelData::~FNiagaraDataChannelData()
 		}
 		if (ReleasedRTProxy)
 		{
+			if (ComputeDispatchInterface)
+			{
+			#if !UE_BUILD_SHIPPING
+				check(ComputeDispatchInterface == ReleasedRTProxy->DispatchInterfaceForDebuggingOnly);
+			#endif
+				ComputeDispatchInterface->RemoveNDCDataProxy(ReleasedRTProxy);
+			}
+		#if !UE_BUILD_SHIPPING
+			ReleasedRTProxy->DispatchInterfaceForDebuggingOnly = nullptr;
+		#endif
 			delete ReleasedRTProxy;
 		}
 	}
@@ -620,7 +631,6 @@ void FNiagaraDataChannelData::Init(UNiagaraDataChannelHandler* Owner)
 	CPUSimData = new FNiagaraDataSet();
 	GPUSimData = new FNiagaraDataSet();
 
-
 	const FNiagaraDataSetCompiledData& CompiledData = DataChannel->GetCompiledData(ENiagaraSimTarget::CPUSim);
 	const FNiagaraDataSetCompiledData& CompiledDataGPU = DataChannel->GetCompiledData(ENiagaraSimTarget::GPUComputeSim);
 	CPUSimData->Init(&CompiledData, 1);
@@ -630,21 +640,25 @@ void FNiagaraDataChannelData::Init(UNiagaraDataChannelHandler* Owner)
 	GameDataStaging = new FNiagaraDataSet();
 	GameDataStaging->Init(&CompiledData);
 
-	if (FNiagaraGpuComputeDispatchInterface* ComputeDispatchInterface = FNiagaraGpuComputeDispatchInterface::Get(Owner->GetWorld()))
+	UWorld* OwnerWorld = Owner->GetWorld();
+	WeakOwnerWorld = OwnerWorld;
+
+	if (FNiagaraGpuComputeDispatchInterface* ComputeDispatchInterface = FNiagaraGpuComputeDispatchInterface::Get(OwnerWorld))
 	{
 		RTProxy.Reset(new FNiagaraDataChannelDataProxy());
 		RTProxy->GPUDataSet = GPUSimData;
 		RTProxy->bNeedsPrevFrameData = DataChannel->KeepPreviousFrameData();
-		RTProxy->DispatchInterface = FNiagaraGpuComputeDispatchInterface::Get(Owner->GetWorld());
 		RTProxy->Owner = this->AsWeak();
 #if !UE_BUILD_SHIPPING
 		RTProxy->DebugName = FString::Printf(TEXT("%s__GPUData"), *DataChannel->GetName());
+		RTProxy->DispatchInterfaceForDebuggingOnly = ComputeDispatchInterface;
 #endif
 		ENQUEUE_RENDER_COMMAND(FNiagaraDataChannelDataProxyInit) (
-			[RT_Proxy = RTProxy.Get()](FRHICommandListImmediate& CmdList)
+			[ComputeDispatchInterface, RT_Proxy=RTProxy.Get()](FRHICommandListImmediate& CmdList)
 			{
-				RT_Proxy->Init();
-			});
+				ComputeDispatchInterface->AddNDCDataProxy(RT_Proxy);
+			}
+		);
 	}
 }
 
diff --git a/Engine/Plugins/FX/Niagara/Source/Niagara/Private/NiagaraGpuComputeDispatch.cpp b/Engine/Plugins/FX/Niagara/Source/Niagara/Private/NiagaraGpuComputeDispatch.cpp
--- a/Engine/Plugins/FX/Niagara/Source/Niagara/Private/NiagaraGpuComputeDispatch.cpp
+++ b/Engine/Plugins/FX/Niagara/Source/Niagara/Private/NiagaraGpuComputeDispatch.cpp
@@ -302,13 +302,17 @@ void FNiagaraGpuComputeDispatch::RemoveGpuComputeProxy(FNiagaraSystemGpuComputeP
 
 void FNiagaraGpuComputeDispatch::AddNDCDataProxy(FNiagaraDataChannelDataProxy* NDCDataProxy)
 {
-	check(NDCDataProxy->DispatchInterface == this);
-	NDCDataProxies.Add(NDCDataProxy);	
+#if !UE_BUILD_SHIPPING
+	check(NDCDataProxy->DispatchInterfaceForDebuggingOnly == this);
+#endif
+	NDCDataProxies.Add(NDCDataProxy);
 }
 
 void FNiagaraGpuComputeDispatch::RemoveNDCDataProxy(FNiagaraDataChannelDataProxy* NDCDataProxy)
 {
-	check(NDCDataProxy->DispatchInterface == this);
+#if !UE_BUILD_SHIPPING
+	check(NDCDataProxy->DispatchInterfaceForDebuggingOnly == this);
+#endif
 	NDCDataProxies.RemoveSwap(NDCDataProxy);
 }
 
@@ -1944,7 +1948,7 @@ void FNiagaraGpuComputeDispatch::PreInitViews(FRDGBuilder& GraphBuilder, bool bA
 
 			for(FNiagaraDataChannelDataProxy* NDCProxy : NDCDataProxies)
 			{
-				NDCProxy->BeginFrame(GraphBuilder.RHICmdList);
+				NDCProxy->BeginFrame(this, GraphBuilder.RHICmdList);
 			}
 
 			UpdateInstanceCountManager(GraphBuilder.RHICmdList);
diff --git a/Engine/Plugins/FX/Niagara/Source/Niagara/Public/NiagaraDataChannel.h b/Engine/Plugins/FX/Niagara/Source/Niagara/Public/NiagaraDataChannel.h
--- a/Engine/Plugins/FX/Niagara/Source/Niagara/Public/NiagaraDataChannel.h
+++ b/Engine/Plugins/FX/Niagara/Source/Niagara/Public/NiagaraDataChannel.h
@@ -79,17 +79,16 @@ struct FNiagaraDataChannelDataProxy
 	//Users that need space in this NDC Data add to this for each tick via AddGPUAllocationForNextTick().
 	int32 PendingGPUAllocations = 0;
 
-	FNiagaraGpuComputeDispatchInterface* DispatchInterface = nullptr;
-
 	#if !UE_BUILD_SHIPPING
+	FNiagaraGpuComputeDispatchInterface* DispatchInterfaceForDebuggingOnly = nullptr;
+	
 	FString DebugName;
 	const TCHAR* GetDebugName()const{return *DebugName;}
 	#else
 	const TCHAR* GetDebugName()const{return nullptr;}
 	#endif
 
-	void Init();
-	void BeginFrame(FRHICommandListImmediate& RHICmdList);
+	void BeginFrame(FNiagaraGpuComputeDispatchInterface* DispatchInterface, FRHICommandListImmediate& RHICmdList);
 	void EndFrame(FRHICommandListImmediate& RHICmdList);
 	void Reset();
 
