Remove Experimental for NDC.

Also remove some old unused accessor helpers that were used previously but no more. Reduce surface area of new things to minimum.

diff --git a/Engine/Plugins/FX/Niagara/Source/Niagara/Internal/DataInterface/NiagaraDataInterfaceDataChannelRead.h b/Engine/Plugins/FX/Niagara/Source/Niagara/Internal/DataInterface/NiagaraDataInterfaceDataChannelRead.h
--- a/Engine/Plugins/FX/Niagara/Source/Niagara/Internal/DataInterface/NiagaraDataInterfaceDataChannelRead.h
+++ b/Engine/Plugins/FX/Niagara/Source/Niagara/Internal/DataInterface/NiagaraDataInterfaceDataChannelRead.h
@@ -64,7 +64,7 @@ enum class ENDIDataChannelSpawnScaleMode
 	Max UMETA(Hidden),
 };
 
-UCLASS(Experimental, EditInlineNew, Category = "Data Channels", CollapseCategories, meta = (DisplayName = "Data Channel Reader"), MinimalAPI)
+UCLASS(EditInlineNew, Category = "Data Channels", CollapseCategories, meta = (DisplayName = "Data Channel Reader"), MinimalAPI)
 class UNiagaraDataInterfaceDataChannelRead : public UNiagaraDataInterfaceRWBase
 {
 	GENERATED_UCLASS_BODY()
diff --git a/Engine/Plugins/FX/Niagara/Source/Niagara/Internal/DataInterface/NiagaraDataInterfaceDataChannelWrite.h b/Engine/Plugins/FX/Niagara/Source/Niagara/Internal/DataInterface/NiagaraDataInterfaceDataChannelWrite.h
--- a/Engine/Plugins/FX/Niagara/Source/Niagara/Internal/DataInterface/NiagaraDataInterfaceDataChannelWrite.h
+++ b/Engine/Plugins/FX/Niagara/Source/Niagara/Internal/DataInterface/NiagaraDataInterfaceDataChannelWrite.h
@@ -88,7 +88,7 @@ public:
 };
 
 
-UCLASS(Experimental, EditInlineNew, Category = "Data Channels", CollapseCategories, meta = (DisplayName = "Data Channel Writer"), MinimalAPI)
+UCLASS(EditInlineNew, Category = "Data Channels", CollapseCategories, meta = (DisplayName = "Data Channel Writer"), MinimalAPI)
 class UNiagaraDataInterfaceDataChannelWrite : public UNiagaraDataInterface, public INiagaraSimCacheCustomStorageInterface
 {
 	GENERATED_UCLASS_BODY()
diff --git a/Engine/Plugins/FX/Niagara/Source/Niagara/Private/DataInterface/NiagaraDataInterfaceDataChannelRead.cpp b/Engine/Plugins/FX/Niagara/Source/Niagara/Private/DataInterface/NiagaraDataInterfaceDataChannelRead.cpp
--- a/Engine/Plugins/FX/Niagara/Source/Niagara/Private/DataInterface/NiagaraDataInterfaceDataChannelRead.cpp
+++ b/Engine/Plugins/FX/Niagara/Source/Niagara/Private/DataInterface/NiagaraDataInterfaceDataChannelRead.cpp
@@ -103,7 +103,6 @@ namespace NDIDataChannelReadLocal
 			NIAGARA_ADD_FUNCTION_SOURCE_INFO(Sig)
 #endif
 			Sig.bMemberFunction = true;
-			Sig.bExperimental = true;
 			Sig.AddInput(FNiagaraVariable(FNiagaraTypeDefinition(UNiagaraDataInterfaceDataChannelRead::StaticClass()), TEXT("DataChannel interface")));
 			Sig.AddOutput(FNiagaraVariable(FNiagaraTypeDefinition::GetIntDef(), TEXT("Num")));
 		}		
@@ -121,7 +120,6 @@ namespace NDIDataChannelReadLocal
 			NIAGARA_ADD_FUNCTION_SOURCE_INFO(Sig)
 #endif		
 			Sig.bMemberFunction = true;
-			Sig.bExperimental = true;
 			Sig.ModuleUsageBitmask = ENiagaraScriptUsageMask::Particle;
 			Sig.AddInput(FNiagaraVariable(FNiagaraTypeDefinition(UNiagaraDataInterfaceDataChannelRead::StaticClass()), TEXT("DataChannel interface")));
 			Sig.AddInputWithoutDefault(FNiagaraVariable(FNiagaraTypeDefinition(FNiagaraEmitterID::StaticStruct()), TEXT("Emitter ID")), LOCTEXT("EmitterIDDesc", "ID of the emitter we'd like to spawn into. This can be obtained from Engine.Emitter.ID."));
@@ -144,7 +142,6 @@ namespace NDIDataChannelReadLocal
 			NIAGARA_ADD_FUNCTION_SOURCE_INFO(Sig)
 #endif
 			Sig.bMemberFunction = true;
-			Sig.bExperimental = true;
 			Sig.AddInput(FNiagaraVariable(FNiagaraTypeDefinition(UNiagaraDataInterfaceDataChannelRead::StaticClass()), TEXT("DataChannel interface")));
 			Sig.AddInputWithDefault(FNiagaraVariable(FNiagaraTypeDefinition::GetIntDef(), TEXT("Index")), 0, LOCTEXT("ConsumeIndexInputDesc", "The index to read."));
 			Sig.AddOutput(FNiagaraVariable(FNiagaraTypeDefinition::GetBoolDef(), TEXT("Success")), LOCTEXT("ConsumeSuccessOutputDesc", "True if all reads succeeded."));
@@ -164,7 +161,6 @@ namespace NDIDataChannelReadLocal
 			NIAGARA_ADD_FUNCTION_SOURCE_INFO(Sig)
 #endif
 			Sig.bMemberFunction = true;
-			Sig.bExperimental = true;
 			Sig.AddInput(FNiagaraVariable(FNiagaraTypeDefinition(UNiagaraDataInterfaceDataChannelRead::StaticClass()), TEXT("DataChannel interface")));
 			Sig.AddInputWithDefault(FNiagaraVariable(FNiagaraTypeDefinition::GetBoolDef(), TEXT("Consume")), FNiagaraBool(true), LOCTEXT("ConsumeInputDesc", "True if this instance (particle/emitter etc) should consume data from the data channel in this call."));
 			Sig.AddOutput(FNiagaraVariable(FNiagaraTypeDefinition::GetBoolDef(), TEXT("Success")), LOCTEXT("ConsumeSuccessOutputDesc", "True if all reads succeeded."));
@@ -197,7 +193,6 @@ namespace NDIDataChannelReadLocal
 			Sig.FunctionVersion = static_cast<uint32>(FunctionVersion_SpawnConditional::EmitterIDParameter);
 #endif
 			Sig.bMemberFunction = true;
-			Sig.bExperimental = true;
 			Sig.bRequiresExecPin = true;
 			Sig.ModuleUsageBitmask = ENiagaraScriptUsageMask::Emitter | ENiagaraScriptUsageMask::System;
 			Sig.AddInput(FNiagaraVariable(FNiagaraTypeDefinition(UNiagaraDataInterfaceDataChannelRead::StaticClass()), TEXT("DataChannel interface")));
@@ -225,7 +220,6 @@ namespace NDIDataChannelReadLocal
 			Sig.FunctionSpecifiers.Add(VarNameKey);
 			Sig.FunctionSpecifiers.Add(VarTypeKey);
 			Sig.bMemberFunction = true;
-			Sig.bExperimental = true;
 			Sig.bRequiresExecPin = true;
 			Sig.ModuleUsageBitmask = ENiagaraScriptUsageMask::Emitter | ENiagaraScriptUsageMask::System;
 			Sig.AddInput(FNiagaraVariable(FNiagaraTypeDefinition(UNiagaraDataInterfaceDataChannelRead::StaticClass()), TEXT("DataChannel interface")));
@@ -253,7 +247,6 @@ namespace NDIDataChannelReadLocal
 			Sig.FunctionSpecifiers.Add(VarNameKey);
 			Sig.FunctionSpecifiers.Add(VarTypeKey);
 			Sig.bMemberFunction = true;
-			Sig.bExperimental = true;
 			Sig.bRequiresExecPin = true;
 			Sig.ModuleUsageBitmask = ENiagaraScriptUsageMask::Emitter | ENiagaraScriptUsageMask::System;
 			Sig.AddInput(FNiagaraVariable(FNiagaraTypeDefinition(UNiagaraDataInterfaceDataChannelRead::StaticClass()), TEXT("DataChannel interface")));
diff --git a/Engine/Plugins/FX/Niagara/Source/Niagara/Private/DataInterface/NiagaraDataInterfaceDataChannelWrite.cpp b/Engine/Plugins/FX/Niagara/Source/Niagara/Private/DataInterface/NiagaraDataInterfaceDataChannelWrite.cpp
--- a/Engine/Plugins/FX/Niagara/Source/Niagara/Private/DataInterface/NiagaraDataInterfaceDataChannelWrite.cpp
+++ b/Engine/Plugins/FX/Niagara/Source/Niagara/Private/DataInterface/NiagaraDataInterfaceDataChannelWrite.cpp
@@ -77,7 +77,6 @@ namespace NDIDataChannelWriteLocal
 			Sig.Description = LOCTEXT("NumFunctionDescription", "Returns the number of instances allocated for writing in the bound NDC buffer. Writes at an index beyond this will fail.");
 #endif
 			Sig.bMemberFunction = true;
-			Sig.bExperimental = true;
 			Sig.AddInput(FNiagaraVariable(FNiagaraTypeDefinition(UNiagaraDataInterfaceDataChannelWrite::StaticClass()), TEXT("DataChannel interface")));
 			Sig.AddOutput(FNiagaraVariable(FNiagaraTypeDefinition::GetIntDef(), TEXT("Num")));
 		}
@@ -95,7 +94,6 @@ namespace NDIDataChannelWriteLocal
 			Sig.Description = LOCTEXT("AllocateFunctionDescription", "Adds an amount to allocated into the bound NDC data for the given emitter to write into.");
 #endif
 			Sig.bMemberFunction = true;
-			Sig.bExperimental = true;
 			Sig.bRequiresExecPin = true;
 			Sig.ModuleUsageBitmask = ENiagaraScriptUsageMask::Emitter | ENiagaraScriptUsageMask::System;
 			Sig.AddInput(FNiagaraVariable(FNiagaraTypeDefinition(UNiagaraDataInterfaceDataChannelWrite::StaticClass()), TEXT("DataChannel interface")));
@@ -119,7 +117,6 @@ namespace NDIDataChannelWriteLocal
 #endif
 			Sig.bMemberFunction = true;
 			Sig.bRequiresExecPin = true;
-			Sig.bExperimental = true;
 			Sig.bSupportsGPU = false;//Cannot use direct index writes on GPU as we write into one shared buffer with all DIs using the same NDC data.
 			Sig.AddInput(FNiagaraVariable(FNiagaraTypeDefinition(UNiagaraDataInterfaceDataChannelWrite::StaticClass()), TEXT("DataChannel interface")));
 			Sig.AddInput(EmitVar, LOCTEXT("ExecuteAppendFlagTooltip", "If true then the append is executed, if false then the append call is skipped"));
@@ -145,7 +142,6 @@ namespace NDIDataChannelWriteLocal
 #endif
 			Sig.bMemberFunction = true;
 			Sig.bRequiresExecPin = true;
-			Sig.bExperimental = true;
 			Sig.AddInput(FNiagaraVariable(FNiagaraTypeDefinition(UNiagaraDataInterfaceDataChannelWrite::StaticClass()), TEXT("DataChannel interface")));
 			Sig.AddInput(EmitVar, LOCTEXT("ExecuteAppendFlagTooltip", "If true then the append is executed, if false then the append call is skipped"));
 			Sig.AddOutput(FNiagaraVariable(FNiagaraTypeDefinition::GetBoolDef(), TEXT("Success")));
diff --git a/Engine/Plugins/FX/Niagara/Source/Niagara/Private/NiagaraDataChannelAccessor.cpp b/Engine/Plugins/FX/Niagara/Source/Niagara/Private/NiagaraDataChannelAccessor.cpp
--- a/Engine/Plugins/FX/Niagara/Source/Niagara/Private/NiagaraDataChannelAccessor.cpp
+++ b/Engine/Plugins/FX/Niagara/Source/Niagara/Private/NiagaraDataChannelAccessor.cpp
@@ -239,72 +239,3 @@ void UNiagaraDataChannelWriter::WriteID(FName VarName, int32 Index, FNiagaraID I
 {
 	WriteData(FNiagaraVariableBase(FNiagaraTypeDefinition::GetIDDef(), VarName), Index, InData);
 }
-
-
-//////////////////////////////////////////////////////////////////////////
-
-void FNiagaraDataChannelGameDataWriterBase::BeginWrite()
-{
-	if(Data)
-	{
-		Data->BeginFrame();
-	}
-}
-
-void FNiagaraDataChannelGameDataWriterBase::Publish(FNiagaraDataChannelDataPtr& Destination)
-{
-	if (Data && Data->Num() > 0)
-	{
-		FNiagaraDataChannelPublishRequest PublishRequest;
-		PublishRequest.bVisibleToCPUSims = true;
-		PublishRequest.bVisibleToGPUSims = true;
-		PublishRequest.GameData = Data;
-		Destination->Publish(PublishRequest);
-	}
-}
-
-void FNiagaraDataChannelGameDataWriterBase::SetNum(int32 Num)
-{
-	if (Data)
-	{
-		Data->SetNum(Num);
-	}
-}
-
-void FNiagaraDataChannelGameDataWriterBase::Reserve(int32 Num)
-{
-	if (Data)
-	{
-		Data->Reserve(Num);
-	}
-}
-
-int32 FNiagaraDataChannelGameDataWriterBase::Add(int32 Count)
-{
-	if (Data)
-	{
-		return Data->Add(Count);
-	}
-	return INDEX_NONE;
-}
-
-//////////////////////////////////////////////////////////////////////////
-
-FNiagaraDataChannelDataPtr FNiagaraDataChannelGameDataGroupedWriterBase::FindDataChannelData(UWorld* World, const UNiagaraDataChannel* NDC, const FNiagaraDataChannelSearchParameters& ItemSearchParams)
-{
-	if(World == nullptr || NDC == nullptr)
-	{
-		return nullptr;
-	}
-
-	if (FNiagaraWorldManager* WorldMan = FNiagaraWorldManager::Get(World))
-	{
-		if (UNiagaraDataChannelHandler* NDCHandler = WorldMan->FindDataChannelHandler(NDC))
-		{
-			return NDCHandler->FindData(ItemSearchParams, ENiagaraResourceAccess::WriteOnly);
-		}
-	}
-	return nullptr;
-}
-
-//////////////////////////////////////////////////////////////////////////
diff --git a/Engine/Plugins/FX/Niagara/Source/Niagara/Public/NiagaraDataChannel.h b/Engine/Plugins/FX/Niagara/Source/Niagara/Public/NiagaraDataChannel.h
--- a/Engine/Plugins/FX/Niagara/Source/Niagara/Public/NiagaraDataChannel.h
+++ b/Engine/Plugins/FX/Niagara/Source/Niagara/Public/NiagaraDataChannel.h
@@ -15,10 +15,6 @@ Each of these writes optionally being made visible to Game, CPU and/or GPU Syste
 At the "Game" level, all data is held in LWC compatible types in AoS format.
 When making this data available to Niagara Systems it is converted to SWC, SoA layout that is compatible with Niagara simulation.
 
-EXPERIMENTAL:
-Data Channels are currently experimental and undergoing heavy development.
-Anything and everything can change, including content breaking changes.
-
 Some Current limitations:
 
 Tick Ordering:
@@ -27,13 +23,6 @@ Reading from the current frame allows zero latency but introduces a frame depend
 This frame dependency needs work to be more robust and less error prone.
 Reading the previous frames data introduces a frame of latency but removes the need to tick later than the writer. Also means you're sure to get a complete frame worth of data.
 
-GPU Support:
-Currently GPU support is very limited.
-Only Game->GPU and CPUSim->GPU are supported.
-Only the Read function of the DI is supported.
-GPU simulations always use the current frame's CPU data as this is all pushed to the at the end of the frame.
-When GPU->GPU is supported, the frame dependency issue have more meaning for GPU systems.
-
 ==============================================================================*/
 
 #pragma once
@@ -103,7 +92,7 @@ struct FNiagaraDataChannelDataProxy
 
 DECLARE_MULTICAST_DELEGATE_OneParam(FOnDataChannelCreated, const UNiagaraDataChannel*);
 
-UCLASS(Experimental, abstract, EditInlineNew, MinimalAPI, prioritizeCategories=("Data Channel"))
+UCLASS(abstract, EditInlineNew, MinimalAPI, prioritizeCategories=("Data Channel"))
 class UNiagaraDataChannel : public UObject
 {
 public:
@@ -225,7 +214,7 @@ enum class ENiagartaDataChannelReadResult : uint8
 /**
 * A C++ and Blueprint accessible library of utility functions for accessing Niagara DataChannel
 */
-UCLASS(Experimental)
+UCLASS()
 class NIAGARA_API UNiagaraDataChannelLibrary : public UBlueprintFunctionLibrary
 {
 	GENERATED_UCLASS_BODY()
diff --git a/Engine/Plugins/FX/Niagara/Source/Niagara/Public/NiagaraDataChannelAccessor.h b/Engine/Plugins/FX/Niagara/Source/Niagara/Public/NiagaraDataChannelAccessor.h
--- a/Engine/Plugins/FX/Niagara/Source/Niagara/Public/NiagaraDataChannelAccessor.h
+++ b/Engine/Plugins/FX/Niagara/Source/Niagara/Public/NiagaraDataChannelAccessor.h
@@ -16,7 +16,7 @@ Likely to be replaced in the near future with a custom BP node and a helper stru
 */
 
 
-UCLASS(Experimental, BlueprintType, MinimalAPI)
+UCLASS(BlueprintType, MinimalAPI)
 class UNiagaraDataChannelReader : public UObject
 {
 	GENERATED_BODY()
@@ -79,7 +79,7 @@ public:
 	NIAGARA_API FNiagaraSpawnInfo ReadSpawnInfo(FName VarName, int32 Index, bool& IsValid)const;
 };
 
-UCLASS(Experimental, BlueprintType, MinimalAPI)
+UCLASS(BlueprintType, MinimalAPI)
 class UNiagaraDataChannelWriter : public UObject
 {
 	GENERATED_BODY()
@@ -153,268 +153,3 @@ public:
 
 };
 
-
-
-/** Base clas for template data channel writer objects. Provides common, non-template functionality. */
-struct FNiagaraDataChannelGameDataWriterBase
-{
-	FNiagaraDataChannelGameDataWriterBase(FNiagaraDataChannelGameDataPtr InData)
-	: Data(InData)
-	{
-	}
-
-	/** Call before we begin writing to this data. */
-	NIAGARA_API void BeginWrite();
-
-	/** Publish the contents of this writer's Data buffer to the given data channel data. */
-	NIAGARA_API void Publish(FNiagaraDataChannelDataPtr& Destination);
-
-	//Binds a specific variable to an index.
-	NIAGARA_API bool BindVariable(int32 VariableIndex, const FNiagaraVariableBase& Var);
-
-	// Set the size of buffers for all variables.
-	NIAGARA_API void SetNum(int32 Num);
-
-	// Reserve space in buffers for all variables.
-	NIAGARA_API void Reserve(int32 Num);
-
-	/** Adds count elements to all variable buffers. */
-	NIAGARA_API int32 Add(int32 Count = 1);
-
-	FNiagaraDataChannelGameDataPtr& GetData(){ return Data; }
-
-protected:
-	FNiagaraDataChannelGameDataPtr Data = nullptr;
-};
-
-//Templated utility writer class.
-//Provides an interface for pushing specific types easily into Niagara Data Channels.
-template<typename T>
-struct FNiagaraDataChannelGameDataWriter : public FNiagaraDataChannelGameDataWriterBase
-{
-	FNiagaraDataChannelGameDataWriter(FNiagaraDataChannelGameDataPtr InData)
-		: FNiagaraDataChannelGameDataWriterBase(InData)
-	{
-	}
-
-	/** Write data to a specific variable index. */
-	bool Write(int32 Index, const T& InData)
-	{
-		//static_assert(false, "All implementations must provide a write function");
-		return false;
-	}
-};
-
-/**
-* Below is an example implementation of FNiagaraDataChannelGameDataWriter.
-* It allows pushing of a simple set of data into a Niagara Data Channel.
-struct FExampleNDCData
-{
-	FVector Position;
-	FVector PrevPosition;
-	FVector Velocity;
-	FNiagaraID ID;
-};
-
-template<typename T>
-struct NIAGARA_API FNiagaraDataChannelGameDataWriter<FExampleNDCData> : public FNiagaraDataChannelGameDataWriterBase
-{
-	FNiagaraDataChannelGameDataWriter(FNiagaraDataChannelGameDataPtr InData)
-		: FNiagaraDataChannelGameDataWriterBase(InData)
-	{
-		if (FNiagaraDataChannelGameData* DataPtr = Data.Get())
-		{
-			PositionBuffer = DataPtr->FindVariableBuffer(FNiagaraVariable(FNiagaraTypeDefinition::GetPositionDef(), TEXT("Position")));
-			PrevPositionBuffer = = DataPtr->FindVariableBuffer(FNiagaraVariable(FNiagaraTypeDefinition::GetPositionDef(), TEXT("PrevPosition")));
-			VelocityBuffer = = DataPtr->FindVariableBuffer(FNiagaraVariable(FNiagaraTypeHelper::GetVectorDef(), TEXT("Velocity")));
-			IDBuffer = = DataPtr->FindVariableBuffer(FNiagaraVariable(FNiagaraTypeDefinition::GetIDDef(), TEXT("ID")));
-
-			if(PositionBuffer == nullptr || PreviousPositionBuffer == nullptr || VelocityBuffer == nullptr || IDBuffer == nullptr)
-			{
-				Data.Reset();
-			}
-		}
-	}
-
-	bool Write(int32 Index, const FLWProjectileNDCData& InData);
-	{
-		if(Data)
-		{
-			check(PositionBuffer && PreviousPositionBuffer && VelocityBuffer && IDBuffer);
-
-			PositionBuffer->Write(InData.Position);
-			PrevPositionBuffer->Write(InData.PrevPosition);
-			VelocityBuffer->Write(InData.Velocity);
-			IDBuffer->Write(InData.ID);
-		}
-	}
-
-	const FNiagaraDataChannelVariableBuffer* PositionBuffer = nullptr;
-	const FNiagaraDataChannelVariableBuffer* PrevPositionBuffer = nullptr;
-	const FNiagaraDataChannelVariableBuffer* VelocityBuffer = nullptr;
-	const FNiagaraDataChannelVariableBuffer* IDBuffer = nullptr;
-};
-
-*/
-
-//////////////////////////////////////////////////////////////////////////
-// FNiagaraDataChannelGameDataGroupedWriter 
-// A utility class allowing for many persistent items to be written each frame to one or more Data Channels in efficient batches.
-// Items are grouped together for allocation and publishing to the NDC by the destination NDC data they'll write to based on the FNiagaraDataChannelSearchParameters.
-// Must use a type the has an implementation of FNiagaraDataChannelGameDataWriter.
-// Usage:
-// 
-// Construct:
-// Create the FNiagaraDataChannelGameDataGroupedWriter with desired type.
-// 
-// Add/Remove
-// When new Items you want to track are created call AddItem(). This will return a GroupIndex that you should store with each Item.
-// Items can be added to be written into different Niagara Data Channels if needed. Items are grouped by the specific NDC data they write to.
-// Call RemoveItem when the items you're tracking are destroyed.
-// 
-// Writing:
-// Ideally write once per frame.
-// To write, construct an FScopedWriter from your FNiagaraDataChannelGameDataGroupedWriter.
-// This new scoped writer will allocate space to write as many items as have been added in each group via AddItem.
-// Writing more than this will not cause any major harm but will cause additional allocations.
-// Results will be published to the NDC when the scoped writer is destroyed.
-// Call Append(), passing each item's GroupIndex and the data to write.
-
-struct FNiagaraDataChannelGameDataGroupedWriterBase
-{
-	NIAGARA_API FNiagaraDataChannelDataPtr FindDataChannelData(UWorld* World, const UNiagaraDataChannel* NDC, const FNiagaraDataChannelSearchParameters& ItemSearchParams);
-};
-
-template<typename T>
-struct FNiagaraDataChannelGameDataGroupedWriter : public FNiagaraDataChannelGameDataGroupedWriterBase
-{
-private:
-	struct FGroupInfo
-	{
-		//A writer helper for each group;
-		FNiagaraDataChannelGameDataWriter<T> Writer;
-				
-		int32 NumItems = 0;
-
-		FGroupInfo(const UNiagaraDataChannel* DataChannel, FNiagaraDataChannelGameDataGroupedWriter& Owner)
-		: Writer(DataChannel ? DataChannel->CreateGameData() : nullptr)
-		{
-		}
-	};
-
-	TMap<FNiagaraDataChannelDataPtr, int32> GroupMap;
-	TArray<FGroupInfo> Groups;
-
-public:
-
-	/**
-	 Scoped utility writer class for appending items to a FNiagaraDataChannelGameDataGroupedWriter.
-	*/
-	struct FScopedWriter
-	{
-	public:
-		FScopedWriter(FNiagaraDataChannelGameDataGroupedWriter<T>& InOwner)	: Owner(InOwner)
-		{
-			Owner.BeginWrite();
-		}
-
-		~FScopedWriter()
-		{
-			Owner.EndWrite();
-		}
-
-		void Append(int32 GroupIdx, const T& InData)
-		{
-			FNiagaraDataChannelGameDataWriter<T>& Writer = Owner.GetWriter(GroupIdx);
-			int32 DataIndex = Writer.Add();
-			Writer.Write(DataIndex, InData);
-		}
-
-	private:
-		FNiagaraDataChannelGameDataGroupedWriter<T>& Owner;
-	};
-
-	FScopedWriter CreateScopedWriter(){ return FScopedWriter(*this); }
-
-	void BeginWrite()
-	{
-		for (auto It = GroupMap.CreateIterator(); It; ++It)
-		{
-			FNiagaraDataChannelDataPtr& Data = It.Key();
-			int32 Index = It.Value();
-
-			if (Groups.IsValidIndex(Index) && Groups[Index].NumItems > 0)
-			{
-				FGroupInfo& Group = Groups[Index];
-				Group.Writer.BeginWrite();
-				Group.Writer.Reserve(Group.NumItems);
-			}
-			else
-			{
-				It.RemoveCurrent();
-			}
-		}
-	}
-
-	void EndWrite()
-	{
-		for (auto& GroupPair : GroupMap)
-		{
-			FNiagaraDataChannelDataPtr& Data = GroupPair.Key;
-			int32 Index = GroupPair.Value;
-			if (Groups.IsValidIndex(Index))
-			{
-				Groups[Index].Writer.Publish(Data);
-			}
-		}
-	}
-
-	/** 
-	Adds a new item to be written to the correct data channel group based on the search params.
-	Returns the group index this item now belongs to.
-	*/
-	int32 AddItem(UWorld* World, const UNiagaraDataChannel* NDC, const FNiagaraDataChannelSearchParameters& ItemSearchParams)
-	{
-		if(FNiagaraDataChannelDataPtr DataChannelData = FindDataChannelData(World, NDC, ItemSearchParams))
-		{
-			int32 LocalGroupIndex = INDEX_NONE;
-			if (int32* FoundGroup = GroupMap.Find(DataChannelData))
-			{
-				LocalGroupIndex = *FoundGroup;
-			}
-			else
-			{
-				//We don't have an group for this data. Create a new group.
-				//Try to reuse an old group first.
-				LocalGroupIndex = Groups.IndexOfByPredicate([](const auto& Info) { return Info.NumItems == 0; });
-				if (LocalGroupIndex == INDEX_NONE)
-				{
-					LocalGroupIndex = Groups.Num();
-					Groups.Emplace(NDC, *this);
-				}
-				GroupMap.Add(DataChannelData) = LocalGroupIndex;
-			}
-
-			++Groups[LocalGroupIndex].NumItems;
-			return LocalGroupIndex;
-		}
-		return INDEX_NONE;
-	}
-
-	void RemoveItem(int32 GroupIndex)
-	{
-		--Groups[GroupIndex].NumItems;
-	}
-
-	/** Returns the writer object for the given group index. */
-	FNiagaraDataChannelGameDataWriter<T>& GetWriter(int32 GroupIndex)
-	{
-		return Groups[GroupIndex].Writer;
-	}
-
-	/** Returns the number of items in a given group. */
-	int32 GetNumItems(int32 GroupIndex)const
-	{
-		return Groups[GroupIndex].NumItems;
-	}
-};
diff --git a/Engine/Plugins/FX/Niagara/Source/Niagara/Public/NiagaraDataChannelHandler.h b/Engine/Plugins/FX/Niagara/Source/Niagara/Public/NiagaraDataChannelHandler.h
--- a/Engine/Plugins/FX/Niagara/Source/Niagara/Public/NiagaraDataChannelHandler.h
+++ b/Engine/Plugins/FX/Niagara/Source/Niagara/Public/NiagaraDataChannelHandler.h
@@ -14,7 +14,7 @@ Some more complex handlers may want to divide up the scene in various different
 #include "NiagaraDataChannelPublic.h"
 #include "NiagaraDataChannelHandler.generated.h"
 
-UCLASS(Experimental, abstract, BlueprintType, MinimalAPI)
+UCLASS(abstract, BlueprintType, MinimalAPI)
 class UNiagaraDataChannelHandler : public UObject
 {
 public:
diff --git a/Engine/Plugins/FX/Niagara/Source/Niagara/Public/NiagaraDataChannelPublic.h b/Engine/Plugins/FX/Niagara/Source/Niagara/Public/NiagaraDataChannelPublic.h
--- a/Engine/Plugins/FX/Niagara/Source/Niagara/Public/NiagaraDataChannelPublic.h
+++ b/Engine/Plugins/FX/Niagara/Source/Niagara/Public/NiagaraDataChannelPublic.h
@@ -22,10 +22,8 @@ Data channel assets define the payload as well as some transfer settings.
 Niagara Systems can read from and write to data channels via data interfaces.
 Blueprint and C++ code can also read from and write to data channels using its API functions.
 
-EXPERIMENTAL:
-Data Channels are currently experimental and undergoing heavy development.
  */
-UCLASS(Experimental, BlueprintType, DisplayName = "Niagara Data Channel", MinimalAPI)
+UCLASS(BlueprintType, DisplayName = "Niagara Data Channel", MinimalAPI)
 class UNiagaraDataChannelAsset : public UObject
 {
 	GENERATED_BODY()
diff --git a/Engine/Plugins/FX/Niagara/Source/Niagara/Public/NiagaraDataChannel_Global.h b/Engine/Plugins/FX/Niagara/Source/Niagara/Public/NiagaraDataChannel_Global.h
--- a/Engine/Plugins/FX/Niagara/Source/Niagara/Public/NiagaraDataChannel_Global.h
+++ b/Engine/Plugins/FX/Niagara/Source/Niagara/Public/NiagaraDataChannel_Global.h
@@ -9,7 +9,7 @@
 /**
 Simple DataChannel handler that makes all data visible globally.
 */
-UCLASS(Experimental, MinimalAPI)
+UCLASS(MinimalAPI)
 class UNiagaraDataChannel_Global : public UNiagaraDataChannel
 {
 	GENERATED_BODY()
@@ -20,7 +20,7 @@ class UNiagaraDataChannel_Global : public UNiagaraDataChannel
 /**
 Basic DataChannel handler that makes all data visible globally.
 */
-UCLASS(Experimental, BlueprintType, MinimalAPI)
+UCLASS(BlueprintType, MinimalAPI)
 class UNiagaraDataChannelHandler_Global : public UNiagaraDataChannelHandler
 {
 	GENERATED_UCLASS_BODY()
diff --git a/Engine/Plugins/FX/Niagara/Source/Niagara/Public/NiagaraDataChannel_Islands.h b/Engine/Plugins/FX/Niagara/Source/Niagara/Public/NiagaraDataChannel_Islands.h
--- a/Engine/Plugins/FX/Niagara/Source/Niagara/Public/NiagaraDataChannel_Islands.h
+++ b/Engine/Plugins/FX/Niagara/Source/Niagara/Public/NiagaraDataChannel_Islands.h
@@ -111,7 +111,7 @@ struct FNDCIslandDebugDrawSettings
 /**
 Data channel that will automatically sub-divide the world into discreet "islands" based on location.
 */
-UCLASS(Experimental, MinimalAPI)
+UCLASS(MinimalAPI)
 class UNiagaraDataChannel_Islands : public UNiagaraDataChannel
 {
 	GENERATED_BODY()
@@ -189,7 +189,7 @@ protected:
 	mutable TSharedPtr<FStreamableHandle> AsyncLoadHandle;
 };
 
-UCLASS(Experimental, BlueprintType, MinimalAPI)
+UCLASS(BlueprintType, MinimalAPI)
 class UNiagaraDataChannelHandler_Islands : public UNiagaraDataChannelHandler
 {
 	GENERATED_UCLASS_BODY()
