- Allow Parameter Collections to be ticked per tick group rather than once per frame
- Controlled by console variable fx.Niagara.WorldManager.TickParameterCollectionsOncePerFrame
- Not enabled by default until performance implications are understood well

#rnx

diff --git a/Engine/Plugins/FX/Niagara/Source/Niagara/Private/NiagaraWorldManager.cpp b/Engine/Plugins/FX/Niagara/Source/Niagara/Private/NiagaraWorldManager.cpp
--- a/Engine/Plugins/FX/Niagara/Source/Niagara/Private/NiagaraWorldManager.cpp
+++ b/Engine/Plugins/FX/Niagara/Source/Niagara/Private/NiagaraWorldManager.cpp
@@ -316,6 +316,16 @@ namespace NiagaraWorldManagerInternal
 		),
 		ECVF_Default
 	); 
+
+	static bool GTickParameterCollectionsOncePerFrame = true;
+	static FAutoConsoleVariableRef CVarTickParameterCollectionsOncePerFrame(
+		TEXT("fx.Niagara.WorldManager.TickParameterCollectionsOncePerFrame"),
+		GTickParameterCollectionsOncePerFrame,
+		TEXT("When true parameter collection updates will occur once per frame in the first tick group.\n")
+		TEXT("When false updates will occur every tick group.\n")
+		TEXT("Updating constantly can cost performance but can fix issues with NPC values updated then used from a system in a later tick group, this would be a frame out without updateing each tick group."),
+		ECVF_Default | ECVF_Scalability
+	);
 }
 
 //////////////////////////////////////////////////////////////////////////
@@ -1098,6 +1108,12 @@ void FNiagaraWorldManager::PostActorTick(float DeltaSeconds)
 	// - Instances that were spawned and we need to ensure the async tick is complete
 	if (SimulationsWithPostActorWork.Num() > 0)
 	{
+		// Tick Parameter Collections to ensure they are up to date
+		if (NiagaraWorldManagerInternal::GTickParameterCollectionsOncePerFrame == false)
+		{
+			TickParameterCollections();
+		}
+
 		// Ensure completion of all systems
 		for (int32 i = 0; i < SimulationsWithPostActorWork.Num(); ++i)
 		{
@@ -1432,6 +1448,14 @@ void FNiagaraWorldManager::Tick(ETickingGroup TickGroup, float DeltaSeconds, ELe
 
 		TickParameterCollections();
 	}
+	else
+	{
+		// Tick Parameter Collections to ensure they are up to date
+		if (NiagaraWorldManagerInternal::GTickParameterCollectionsOncePerFrame == false)
+		{
+			TickParameterCollections();
+		}
+	}
 
 	// If we are in paused don't do anything
 	if ( DebugPlaybackMode == ENiagaraDebugPlaybackMode::Paused )
