- Use scale when calculating per particle LODs
- This matches static mesh LODs when using uniform scale
- When using non uniform scale the max scale axis will not be accurate, to address this we would need to compose the full transform and apply to the bbox, this is quite costly

#rnx

#tests Custom level with various LODs & ET

diff --git a/Engine/Plugins/FX/Niagara/Shaders/Private/NiagaraUpdateMeshGPUSceneInstances.usf b/Engine/Plugins/FX/Niagara/Shaders/Private/NiagaraUpdateMeshGPUSceneInstances.usf
--- a/Engine/Plugins/FX/Niagara/Shaders/Private/NiagaraUpdateMeshGPUSceneInstances.usf
+++ b/Engine/Plugins/FX/Niagara/Shaders/Private/NiagaraUpdateMeshGPUSceneInstances.usf
@@ -178,8 +178,13 @@ void UpdateMeshInstancesCS(uint3 GroupId : SV_GroupID, uint GroupIndex : SV_Grou
 			CullInstance(InstanceId);
 			return;
 		}
+
+		const float3 ParticleScale = NiagaraSafeGetVec3(Params.ScaleDataOffset, Params.ParticleIndex, Params.DefaultScale);
+		const float3 FinalScale = ParticleScale * MeshScale * PrimitiveData.NonUniformScale;
+		const float MaxScaleAxis = max3(abs(FinalScale.x), abs(FinalScale.y), abs(FinalScale.z));
+		const float ScaledCamDstSq = MaxScaleAxis > 0.0f ? CamDistSq / (MaxScaleAxis*MaxScaleAxis) : CamDistSq;
 		
-		const float ScreenRadius = saturate(LODScreenSize.z / max(1.0f, CamDistSq * LODScreenSize.w));
+		const float ScreenRadius = saturate(LODScreenSize.z / max(1.0f, ScaledCamDstSq * LODScreenSize.w));
 		if ( ScreenRadius < LODScreenSize.x || ScreenRadius >= LODScreenSize.y )
 		{
 			CullInstance(InstanceId);
