Fix for crash when trying to read from a user parameter NDC DI, the read will now silently fail (like the write currently do)
#rnx

diff --git a/Engine/Plugins/FX/Niagara/Source/Niagara/Private/DataInterface/NiagaraDataInterfaceDataChannelRead.cpp b/Engine/Plugins/FX/Niagara/Source/Niagara/Private/DataInterface/NiagaraDataInterfaceDataChannelRead.cpp
--- a/Engine/Plugins/FX/Niagara/Source/Niagara/Private/DataInterface/NiagaraDataInterfaceDataChannelRead.cpp
+++ b/Engine/Plugins/FX/Niagara/Source/Niagara/Private/DataInterface/NiagaraDataInterfaceDataChannelRead.cpp
@@ -1627,7 +1627,13 @@ void UNiagaraDataInterfaceDataChannelRead::SpawnConditional(FVectorVMExternalFun
 	VectorVM::FUserPtrHandler<FNDIDataChannelReadInstanceData> InstData(Context);
 
 	FNiagaraSystemInstance* SystemInstance = InstData->Owner;
-	check(SystemInstance);
+
+	// SystemInstance could be null because of the DI being a user parameter, in which case it is not fully initilaized
+	// because it doesn't have access to all the compile time data
+	if (!SystemInstance)
+	{
+		return;
+	}
 
 	//Binding info can be null here as we can be spawning without any conditions, i.e. no variadic parameters to the function.
 	const FNDIDataChannel_FunctionToDataSetBinding* BindingInfo = InstData->FuncToDataSetBindingInfo.IsValidIndex(FuncIdx) ? InstData->FuncToDataSetBindingInfo[FuncIdx].Get() : nullptr;
