Harden's NiagaraStackScriptHierarchyRoot to stack inputs including orphaned pins, which can happen with renamed parameters
#rnx

diff --git a/Engine/Plugins/FX/Niagara/Source/NiagaraEditor/Private/ViewModels/Stack/NiagaraStackScriptHierarchyRoot.cpp b/Engine/Plugins/FX/Niagara/Source/NiagaraEditor/Private/ViewModels/Stack/NiagaraStackScriptHierarchyRoot.cpp
--- a/Engine/Plugins/FX/Niagara/Source/NiagaraEditor/Private/ViewModels/Stack/NiagaraStackScriptHierarchyRoot.cpp
+++ b/Engine/Plugins/FX/Niagara/Source/NiagaraEditor/Private/ViewModels/Stack/NiagaraStackScriptHierarchyRoot.cpp
@@ -132,20 +132,27 @@ void UNiagaraStackScriptHierarchyRoot::RefreshInstanceData()
 	UNiagaraGraph* NiagaraGraph = OwningFunctionCallNode->GetFunctionScriptSource()->NodeGraph;
 	const TMap<FNiagaraVariable, TObjectPtr<UNiagaraScriptVariable>>& ScriptVariableMap = NiagaraGraph->GetAllMetaData();
 
-	for(auto It(UsedInputs.CreateIterator()); It; ++It)
+	auto IsInputValid = [&ScriptVariableMap](const FNiagaraVariable& InputVariable) -> bool
 	{
-		UNiagaraScriptVariable* ScriptVariable = ScriptVariableMap[*It];
-		
-		if(It->GetType().IsValid() == false)
+		if (InputVariable.GetType().IsValid() == false)
 		{
-			It.RemoveCurrent();
+			return false;
 		}
 
-		if(ScriptVariable->Metadata.bInlineEditConditionToggle)
+		if (const UNiagaraScriptVariable* MatchingScriptVariable = ScriptVariableMap.FindRef(InputVariable))
 		{
-			It.RemoveCurrent();
+			return MatchingScriptVariable->Metadata.bInlineEditConditionToggle == false;
 		}
 
+		return false;
+	};
+
+	for(auto It(UsedInputs.CreateIterator()); It; ++It)
+	{
+		if (!IsInputValid(*It))
+		{
+			It.RemoveCurrent();
+		}
 	}
 	
 	for(const FNiagaraVariable& InputVariable : UsedInputs)
