Fixed GPU calls to GetNDCSpawnData.

diff --git a/Engine/Plugins/FX/Niagara/Shaders/Private/DataChannel/NiagaraDataInterfaceDataChannelTemplateReadCommon.ush b/Engine/Plugins/FX/Niagara/Shaders/Private/DataChannel/NiagaraDataInterfaceDataChannelTemplateReadCommon.ush
--- a/Engine/Plugins/FX/Niagara/Shaders/Private/DataChannel/NiagaraDataInterfaceDataChannelTemplateReadCommon.ush
+++ b/Engine/Plugins/FX/Niagara/Shaders/Private/DataChannel/NiagaraDataInterfaceDataChannelTemplateReadCommon.ush
@@ -104,15 +104,15 @@ struct FNDCAccessContext_{ParameterName}
 	void Read_Float(out float3 Value){ NiagaraDataChannelRead_Float(GetBuffer_Float(), GetFloatStride(), Register_Float, ReadIndex, Value); }
 	void Read_Float(out float4 Value){ NiagaraDataChannelRead_Float(GetBuffer_Float(), GetFloatStride(), Register_Float, ReadIndex, Value); }
 
-	void Read_Int32(out int Value){ NiagaraDataChannelRead_Int32(GetBuffer_Int32(), GetFloatStride(), Register_Int32, ReadIndex, Value); }
-	void Read_Int32(out int2 Value){ NiagaraDataChannelRead_Int32(GetBuffer_Int32(), GetFloatStride(), Register_Int32, ReadIndex, Value); }
-	void Read_Int32(out int3 Value){ NiagaraDataChannelRead_Int32(GetBuffer_Int32(), GetFloatStride(), Register_Int32, ReadIndex, Value); }
-	void Read_Int32(out int4 Value){ NiagaraDataChannelRead_Int32(GetBuffer_Int32(), GetFloatStride(), Register_Int32, ReadIndex, Value); }
-
-	void Read_Half(out half Value){ NiagaraDataChannelRead_Half(GetBuffer_Half(), GetFloatStride(), Register_Half, ReadIndex, Value); }
-	void Read_Half(out half2 Value){ NiagaraDataChannelRead_Half(GetBuffer_Half(), GetFloatStride(), Register_Half, ReadIndex, Value); }
-	void Read_Half(out half3 Value){ NiagaraDataChannelRead_Half(GetBuffer_Half(), GetFloatStride(), Register_Half, ReadIndex, Value); }
-	void Read_Half(out half4 Value){ NiagaraDataChannelRead_Half(GetBuffer_Half(), GetFloatStride(), Register_Half, ReadIndex, Value); }
+	void Read_Int32(out int Value){ NiagaraDataChannelRead_Int32(GetBuffer_Int32(), GetIntStride(), Register_Int32, ReadIndex, Value); }
+	void Read_Int32(out int2 Value){ NiagaraDataChannelRead_Int32(GetBuffer_Int32(), GetIntStride(), Register_Int32, ReadIndex, Value); }
+	void Read_Int32(out int3 Value){ NiagaraDataChannelRead_Int32(GetBuffer_Int32(), GetIntStride(), Register_Int32, ReadIndex, Value); }
+	void Read_Int32(out int4 Value){ NiagaraDataChannelRead_Int32(GetBuffer_Int32(), GetIntStride(), Register_Int32, ReadIndex, Value); }
+
+	void Read_Half(out half Value){ NiagaraDataChannelRead_Half(GetBuffer_Half(), GetHalfStride(), Register_Half, ReadIndex, Value); }
+	void Read_Half(out half2 Value){ NiagaraDataChannelRead_Half(GetBuffer_Half(), GetHalfStride(), Register_Half, ReadIndex, Value); }
+	void Read_Half(out half3 Value){ NiagaraDataChannelRead_Half(GetBuffer_Half(), GetHalfStride(), Register_Half, ReadIndex, Value); }
+	void Read_Half(out half4 Value){ NiagaraDataChannelRead_Half(GetBuffer_Half(), GetHalfStride(), Register_Half, ReadIndex, Value); }
 };
 
 void Num_{ParameterName}(out int OutNum)
@@ -165,12 +165,12 @@ void GetNDCSpawnData_{ParameterName}(NiagaraEmitterID InEmitterID, int InExecInd
 
 	uint ExecIdx = (uint)InExecIndex;
 
-	uint NumNDCEntries = {ParameterName}_NDCElementCountAtSpawn;
+	int NumNDCEntries = {ParameterName}_NDCElementCountAtSpawn;
 
 	//First part of the buffer is the per emitter offsets into the data.
 	//Grab the right base offset for this emitter.
 	uint NDCSpawnDataOffset = { ParameterName }_NDCSpawnDataBuffer[InEmitterID.ID];
-	if(NDCSpawnDataOffset == 0 || NumNDCEntries == 0)
+	if(NDCSpawnDataOffset == 0 || NumNDCEntries <= 0)
 	{
 		OutNDCIndex = -1;
 		OutNDCSpawnIndex = -1;
@@ -197,8 +197,9 @@ void GetNDCSpawnData_{ParameterName}(NiagaraEmitterID InEmitterID, int InExecInd
 		if(ExecIdx < MaxBucketExecIndex)
 		{
 			//We found our bucket.
-			//Now we need to find our NDCIndex Entry Index.
-			uint NDCIndexEntry = (ExecIdx - MinBucketExecIndex) / BucketSize;			
+			//Find our NDC entry. There is one entry for each bucket load of instances. So we divide our current adjusted exec index in this bucket by the bucket size.
+			//As buckets are power of 2 we can do this faster by just shifting down.
+			uint NDCIndexEntry = (ExecIdx - MinBucketExecIndex) >> (15U - BucketIdx);			
 			OutNDCIndex = {ParameterName}_NDCSpawnDataBuffer[BucketEntryStart + NDCIndexEntry];
 			break;
 		}
diff --git a/Engine/Plugins/FX/Niagara/Source/Niagara/Private/DataInterface/NiagaraDataInterfaceDataChannelRead.cpp b/Engine/Plugins/FX/Niagara/Source/Niagara/Private/DataInterface/NiagaraDataInterfaceDataChannelRead.cpp
--- a/Engine/Plugins/FX/Niagara/Source/Niagara/Private/DataInterface/NiagaraDataInterfaceDataChannelRead.cpp
+++ b/Engine/Plugins/FX/Niagara/Source/Niagara/Private/DataInterface/NiagaraDataInterfaceDataChannelRead.cpp
@@ -675,7 +675,7 @@ void UNiagaraDataInterfaceDataChannelRead::PostStageTick(FNDICpuPostStageContext
 	check(InstanceData);
 	check(Context.Usage == ENiagaraScriptUsage::EmitterUpdateScript || Context.Usage == ENiagaraScriptUsage::SystemUpdateScript);
 
-	InstanceData->NDCElementCountAtSpawn = 0;
+	InstanceData->NDCElementCountAtSpawn = INDEX_NONE;
 	
 	for (auto& EmitterInstanceDataPair : InstanceData->EmitterInstanceData)
 	{
@@ -687,7 +687,7 @@ void UNiagaraDataInterfaceDataChannelRead::PostStageTick(FNDICpuPostStageContext
 			continue;
 		}
 		
-		if (InstanceData->NDCElementCountAtSpawn == 0)
+		if (InstanceData->NDCElementCountAtSpawn == INDEX_NONE)
 		{
 			InstanceData->NDCElementCountAtSpawn = EmitterInstData.NDCSpawnCounts.Num();
 		}
@@ -700,7 +700,7 @@ void UNiagaraDataInterfaceDataChannelRead::PostStageTick(FNDICpuPostStageContext
 			PerNDCSpawnCounts.Add(Info.Get());
 		}
 
-		EmitterInstData.NDCSpawnData.Reset();
+		EmitterInstData.Reset();
 
 		//-TODO:Stateless:
 		if (FNiagaraEmitterInstanceImpl* StatefulEmitter = TargetEmitter->AsStateful())
@@ -820,7 +820,7 @@ void UNiagaraDataInterfaceDataChannelRead::ProvidePerInstanceDataForRenderThread
 		//Do one pass to calculate size.
 		auto GetEmitterNDCSpawnDataSize = [](const FNDIDataChannelRead_EmitterInstanceData& EmitterInstData)
 		{
-			return 16 + EmitterInstData.NDCSpawnCounts.Num() + EmitterInstData.NDCSpawnData.NDCSpawnData.Num();
+			return 16 + EmitterInstData.NDCSpawnData.NDCSpawnData.Num();
 		};
 
 		uint32 NumEmitters = SourceData.EmitterInstanceData.Num();
@@ -839,7 +839,7 @@ void UNiagaraDataInterfaceDataChannelRead::ProvidePerInstanceDataForRenderThread
 		}
 
 		//First section of the NDCSpawnDataBuffer is an offset into the buffer for each emitter.
-		TotalPacckedNDCSpawnDataSize += MaxEmitterIndex;
+		TotalPacckedNDCSpawnDataSize += (MaxEmitterIndex + 1);
 
 		TargetData->NDCSpawnData.Reset(TotalPacckedNDCSpawnDataSize);
 		TArray<int32>& TargetNDCSpawnData = TargetData->NDCSpawnData;
@@ -873,12 +873,6 @@ void UNiagaraDataInterfaceDataChannelRead::ProvidePerInstanceDataForRenderThread
 				{
 					TargetNDCSpawnData.Add(EmitterInstData.NDCSpawnData.NDCSpawnDataBuckets[i]);
 				}
-				//Next the per NDC Spawn Counts
-				for (int32 i = 0; i < EmitterInstData.NDCSpawnCounts.Num(); ++i)
-				{
-					TargetNDCSpawnData.Add(EmitterInstData.NDCSpawnCounts[i].Get());
-				}
-				//Finally the exec index to NDC index mapping table
 				TargetNDCSpawnData.Append(EmitterInstData.NDCSpawnData.NDCSpawnData);
 			}
 		}
@@ -1333,7 +1327,7 @@ void UNiagaraDataInterfaceDataChannelRead::GetNDCSpawnData(FVectorVMExternalFunc
 		{
 			const int32* NDCSpawnBukets = EmitterInstanceData.NDCSpawnData.NDCSpawnDataBuckets;
 			TConstArrayView<int32> NDCSpawnData(EmitterInstanceData.NDCSpawnData.NDCSpawnData);
-			uint32 NumNDCEntries = EmitterInstanceData.NDCSpawnCounts.Num();
+			uint32 NumNDCEntries = InstData->NDCElementCountAtSpawn;
 
 			uint32 ExecIndex = InExecIndex.GetAndAdvance();
 			uint32 NDCIndex = INDEX_NONE;
@@ -1387,6 +1381,8 @@ void UNiagaraDataInterfaceDataChannelRead::GetNDCSpawnData(FVectorVMExternalFunc
 						MaxBucketExecIndex += BucketSize * NumEntriesInBucket;
 						if (ExecIndex < MaxBucketExecIndex && NumNDCInstancesInBucket > 0)
 						{
+							//Find our NDC entry. There is one entry for each bucket load of instances. So we divide our current adjusted exec index in this bucket by the bucket size.
+							//As buckets are power of 2 we can do this faster by just shifting down.
 							uint32 NDCIndexEntry = (ExecIndex - MinBucketExecIndex) >> (15 - BucketIdx);
 
 							uint32 MinNDCBucketExecIndex = MinBucketExecIndex + (BucketSize * NDCIndexEntry);
