Force NDC GPU readback pass to happen on the RT.

diff --git a/Engine/Plugins/FX/Niagara/Source/Niagara/Private/DataInterface/NiagaraDataInterfaceDataChannelWrite.cpp b/Engine/Plugins/FX/Niagara/Source/Niagara/Private/DataInterface/NiagaraDataInterfaceDataChannelWrite.cpp
--- a/Engine/Plugins/FX/Niagara/Source/Niagara/Private/DataInterface/NiagaraDataInterfaceDataChannelWrite.cpp
+++ b/Engine/Plugins/FX/Niagara/Source/Niagara/Private/DataInterface/NiagaraDataInterfaceDataChannelWrite.cpp
@@ -1494,7 +1494,7 @@ void FNiagaraDataInterfaceProxy_DataChannelWrite::PostStage(const FNDIGpuCompute
 				PASS_bPublishToGame = InstanceData->bPublishToGame,
 				PASS_bPublishToCPU = InstanceData->bPublishToCPU,
 				PASS_LwcTile = InstanceData->LwcTile]
-				(FRHICommandList& RHICmdList)
+				(FRHICommandListImmediate& RHICmdList)
 				{
 					PASS_Proxy->EnqueueReadbackForCPUBuffer(RHICmdList, PASS_Buffer, PASS_ReadbackManager, *PASS_InstCountManager, PASS_bPublishToGame, PASS_bPublishToCPU, PASS_LwcTile);
 					uint32 Offset = PASS_Buffer->GetGPUInstanceCountBufferOffset();
