- Handle adding / removing sockets in the reader DI while in the editor

#rnx

diff --git a/Engine/Plugins/FX/Niagara/Source/Niagara/Private/DataInterface/NiagaraDataInterfaceSocketReader.cpp b/Engine/Plugins/FX/Niagara/Source/Niagara/Private/DataInterface/NiagaraDataInterfaceSocketReader.cpp
--- a/Engine/Plugins/FX/Niagara/Source/Niagara/Private/DataInterface/NiagaraDataInterfaceSocketReader.cpp
+++ b/Engine/Plugins/FX/Niagara/Source/Niagara/Private/DataInterface/NiagaraDataInterfaceSocketReader.cpp
@@ -286,6 +286,30 @@ namespace NDISocketReaderLocal
 		InstanceData->bIsDataValid			= ResolvedObject != nullptr;
 		InstanceData->bNeedsSocketRecache  |= ResolvedObject != InstanceData->ResolvedObject;
 		InstanceData->ResolvedObject		= ResolvedObject;
+
+	#if WITH_EDITOR
+		// When in the editor the socket counts can change so we might need to recache the socket list
+		if (!InstanceData->bNeedsSocketRecache && ResolvedObject)
+		{
+			int32 NewSocketCount = InstanceData->NumSockets;
+			if (USceneComponent* SceneComponent = Cast<USceneComponent>(ResolvedObject))
+			{
+				TArray<FComponentSocketDescription> SocketList;
+				SocketList.Reserve(InstanceData->NumSockets);
+				SceneComponent->QuerySupportedSockets(SocketList);
+				NewSocketCount = SocketList.Num();
+			}
+			else if (UStaticMesh* StaticMesh = Cast<UStaticMesh>(ResolvedObject))
+			{
+				NewSocketCount = StaticMesh->Sockets.Num();
+			}
+			else if (USkeletalMesh* SkeletalMesh = Cast<USkeletalMesh>(ResolvedObject))
+			{
+				NewSocketCount = SkeletalMesh->NumSockets();
+			}
+			InstanceData->bNeedsSocketRecache |= InstanceData->NumSockets != NewSocketCount;
+		}
+	#endif
 	}
 
 	void GetSocketNames(UObject* ResolvedObject, TArray<FName>& OutSocketNames)
