* Added wizard to spawn particles from data channel entries
* Allow existing data channel wizards to define the target namespace
* Improved usage bitmask for read and write modules
* Added auto-transformation of position values that are read or written with data channel functions in localspace emitters
* Improve default values

diff --git a/Engine/Plugins/FX/Niagara/Source/Niagara/Internal/DataInterface/NiagaraDataInterfaceDataChannelWrite.h b/Engine/Plugins/FX/Niagara/Source/Niagara/Internal/DataInterface/NiagaraDataInterfaceDataChannelWrite.h
--- a/Engine/Plugins/FX/Niagara/Source/Niagara/Internal/DataInterface/NiagaraDataInterfaceDataChannelWrite.h
+++ b/Engine/Plugins/FX/Niagara/Source/Niagara/Internal/DataInterface/NiagaraDataInterfaceDataChannelWrite.h
@@ -99,7 +99,7 @@ public:
 	UPROPERTY(EditAnywhere, Category = "Data Channel")
 	ENiagaraDataChannelAllocationMode AllocationMode = ENiagaraDataChannelAllocationMode::Static;
 
-	/** How many elements to allocate for writing per frame? Usage is defendant on AllocationMode. TODO: Allow allocation count to be controlled dynamically from script? */
+	/** How many elements to allocate for writing per frame? Usage depends on AllocationMode. TODO: Allow allocation count to be controlled dynamically from script? */
 	UPROPERTY(EditAnywhere, Category = "Data Channel", meta = (EditCondition="AllocationMode == ENiagaraDataChannelAllocationMode::Static"))
 	uint32 AllocationCount = 0;
 
diff --git a/Engine/Plugins/FX/Niagara/Source/NiagaraEditor/Private/Widgets/DataChannel/NiagaraDataChannelWizard.cpp b/Engine/Plugins/FX/Niagara/Source/NiagaraEditor/Private/Widgets/DataChannel/NiagaraDataChannelWizard.cpp
--- a/Engine/Plugins/FX/Niagara/Source/NiagaraEditor/Private/Widgets/DataChannel/NiagaraDataChannelWizard.cpp
+++ b/Engine/Plugins/FX/Niagara/Source/NiagaraEditor/Private/Widgets/DataChannel/NiagaraDataChannelWizard.cpp
@@ -16,10 +16,15 @@
 #include "ViewModels/NiagaraScratchPadScriptViewModel.h"
 #include "ViewModels/NiagaraScriptGraphViewModel.h"
 #include "IDetailsView.h"
+#include "NiagaraNodeOp.h"
+#include "NiagaraScriptSource.h"
+#include "ViewModels/NiagaraEmitterViewModel.h"
+#include "Widgets/SNiagaraParameterName.h"
 #include "Widgets/Views/SListView.h"
 #include "Widgets/Layout/SSeparator.h"
 #include "Widgets/Text/STextBlock.h"
 #include "Widgets/Input/SEditableTextBox.h"
+#include "Widgets/Input/SComboBox.h"
 #include "Widgets/Input/SCheckBox.h"
 
 
@@ -76,6 +81,11 @@ namespace UE::Niagara::Wizard::DataChannel
 		explicit FSelectVariablesPageBase(FSelectAssetPageBase* InPreviousPage) : PreviousPage(InPreviousPage)
 		{
 			Name = LOCTEXT("VariablesPageName", "Select variables");
+			SupportedNamespaces.Add(MakeShared<FString>("StackContext.Module"));
+			SupportedNamespaces.Add(MakeShared<FString>("Output.Module"));
+			SupportedNamespaces.Add(MakeShared<FString>("StackContext"));
+			SupportedNamespaces.Add(MakeShared<FString>("Transient"));
+			TargetNamespace = SupportedNamespaces[0];
 		}
 
 		virtual ~FSelectVariablesPageBase() override = default;
@@ -142,6 +152,46 @@ namespace UE::Niagara::Wizard::DataChannel
 						.SelectionMode(ESelectionMode::Single)
 				]
 				+SVerticalBox::Slot()
+				.AutoHeight()
+				[
+					SNew(SHorizontalBox)
+					+ SHorizontalBox::Slot()
+					.AutoWidth()
+					.VAlign(VAlign_Center)
+					[
+						SNew(STextBlock)
+						.Text(LOCTEXT("TargetNamespaceNameText", "Target Namespace: "))
+					]
+					+ SHorizontalBox::Slot()
+					.AutoWidth()
+					.VAlign(VAlign_Center)
+					.Padding(4.0f, 0)
+					[
+						SNew(SComboBox<TSharedPtr<FString>>)
+						.OptionsSource(&SupportedNamespaces)
+						.ContentPadding(2.0f)
+						.InitiallySelectedItem(TargetNamespace)
+						.ToolTipText(LOCTEXT("TargetNamespaceTooltip", "Select the namespace where the variables should be written to. The StackContext namespace changes depending on the script context it is used in (system, emitter, particle)."))
+						.OnGenerateWidget_Lambda([](TSharedPtr<FString> Item)
+						{
+							return SNew(SNiagaraParameterName)
+							.ParameterName(FName(*Item + ". "))
+							.IsReadOnly(true)
+							.SingleNameDisplayMode(SNiagaraParameterName::ESingleNameDisplayMode::Namespace);
+						})
+						.OnSelectionChanged(this, &FSelectVariablesPageBase::HandleNamespaceSelectionChanged)
+						[
+							SNew(SNiagaraParameterName)
+							.ParameterName_Lambda([this]()
+							{
+								return FName(GetTargetNamespace() + ". ");
+							})
+							.IsReadOnly(true)
+							.SingleNameDisplayMode(SNiagaraParameterName::ESingleNameDisplayMode::Namespace)
+						]
+					]
+				]
+				+SVerticalBox::Slot()
 				.Padding(0, 10)
 				.AutoHeight()
 				[
@@ -160,6 +210,7 @@ namespace UE::Niagara::Wizard::DataChannel
 					[
 						SNew( SEditableTextBox )
 						.MinDesiredWidth(200)
+						.Padding(2.0f)
 						.Text(this, &FSelectVariablesPageBase::GetModuleNameText)
 						.SelectAllTextWhenFocused(true)
 						.ClearKeyboardFocusOnCommit(false)
@@ -168,6 +219,11 @@ namespace UE::Niagara::Wizard::DataChannel
 				];
 		}
 
+		void HandleNamespaceSelectionChanged(TSharedPtr<FString> InItem, ESelectInfo::Type)
+		{
+			TargetNamespace = InItem;
+		}
+
 		FText GetModuleNameText() const
 		{
 			return ModuleName;
@@ -238,12 +294,191 @@ namespace UE::Niagara::Wizard::DataChannel
 		{
 			return VariablesToProcess.Contains(Var.Version) ? ECheckBoxState::Checked : ECheckBoxState::Unchecked;
 		}
+		
+		FString GetTargetNamespace() const
+		{
+			return *TargetNamespace.Get();
+		}
 
 		TArray<TSharedPtr<FNiagaraDataChannelVariable>> AllVariables;
 		TSet<FGuid> VariablesToProcess;
 		FSelectAssetPageBase* PreviousPage;
 		FObjectKey LastDataChannelRef;
 		FText ModuleName;
+		TSharedPtr<FString> TargetNamespace;
+
+		TSharedPtr<SListView<TSharedPtr<FNiagaraDataChannelVariable>>> VarListView;
+		TArray<TSharedPtr<FString>> SupportedNamespaces;
+	};
+
+	struct FSelectSpawnAssetPage : FSelectAssetPageBase
+	{
+		virtual UNiagaraDataChannelAsset* GetAsset() const override
+		{
+			if (UNiagaraDataChannelSpawnModuleData* ModuleData = Data.Get())
+			{
+				return ModuleData->DataChannel;
+			}
+			return nullptr;
+		}
+		
+		virtual TSharedRef<SWidget> GetContent() override
+		{
+			Data.Reset(NewObject<UNiagaraDataChannelSpawnModuleData>());
+			return GetDetailsViewContent(Data.Get());
+		}
+		
+		TStrongObjectPtr<UNiagaraDataChannelSpawnModuleData> Data;
+	};
+	
+	struct FSpawnConditionPage : FModuleWizardPage
+	{
+		explicit FSpawnConditionPage(FSelectSpawnAssetPage* InPreviousPage) : PreviousPage(InPreviousPage)
+		{
+			Name = LOCTEXT("SpawnConditionPageName", "Spawn conditions");
+		}
+
+		virtual ~FSpawnConditionPage() override = default;
+
+		virtual bool CanGoToNextPage() const override
+		{
+			if (PreviousPage->Data->SpawnMode == ENiagaraDataChanneSpawnModuleMode::DirectSpawn)
+			{
+				return ConditionVariables.Num() > 0;
+			}
+			return true;
+		};
+		virtual bool CanCompleteWizard() const override { return CanGoToNextPage(); };
+
+		virtual FText GetHeaderLabel() const
+		{
+			if (PreviousPage->Data->SpawnMode == ENiagaraDataChanneSpawnModuleMode::ConditionalSpawn)
+			{
+				return LOCTEXT("SpawnConditionalPageHeader", "Please select which data channel variables should be used as conditions to spawn particles. This is optional!\nFor example, if you select a vector and integer parameter in the data channel they will be compared against the corresponding module inputs.\nResult = (InputA == ChannelValue.A) && (InputB == ChannelValue.B)");
+			}
+			return LOCTEXT("SpawnDirectPageHeader", "Please select which data channel variable should be used as particle spawn count. This needs to be an integer parameter in the data channel.");
+		}
+
+		virtual void RefreshContent() override
+		{
+			FObjectKey NewDataChannelRef;
+			TArray<FNiagaraDataChannelVariable> DataChannelVariables;
+			if (UNiagaraDataChannelAsset* ChannelAsset = PreviousPage->GetAsset())
+			{
+				NewDataChannelRef = ChannelAsset;
+				DataChannelVariables = ChannelAsset->Get()->GetVariables();
+			}
+			if (NewDataChannelRef != LastDataChannelRef)
+			{
+				ConditionVariables.Empty();
+			} 
+
+			AllVariables.Empty(DataChannelVariables.Num());
+			for (const FNiagaraDataChannelVariable& Var : DataChannelVariables)
+			{
+				*AllVariables.Add_GetRef(MakeShared<FNiagaraDataChannelVariable>()).Get() = Var;
+			}
+
+			if (VarListView.IsValid())
+			{
+				VarListView->RebuildList();
+			}
+			LastDataChannelRef = NewDataChannelRef;
+		}
+
+		virtual TSharedRef<SWidget> GetContent() override
+		{
+			return SNew(SVerticalBox)
+				+SVerticalBox::Slot()
+				.Padding(2)
+				.AutoHeight()
+				[
+					SNew(SSeparator)
+					.Orientation(Orient_Horizontal)
+				]
+				+SVerticalBox::Slot()
+				.Padding(15)
+				.AutoHeight()
+				[
+					SNew(STextBlock)
+					.Text(this, &FSpawnConditionPage::GetHeaderLabel)
+				]
+				+SVerticalBox::Slot()
+				[
+					SAssignNew(VarListView, SListView<TSharedPtr<FNiagaraDataChannelVariable>>)
+						.ListItemsSource(&AllVariables)
+						.OnGenerateRow(this, &FSpawnConditionPage::GenerateRow)
+						.SelectionMode(ESelectionMode::Single)
+				];
+		}
+
+		TSharedRef<ITableRow> GenerateRow(const TSharedPtr<FNiagaraDataChannelVariable> Var, const TSharedRef<STableViewBase>& OwnerTable)
+		{
+			FLinearColor TypeColor = UEdGraphSchema_Niagara::GetTypeColor(Var->GetType());
+			return SNew(STableRow<TSharedPtr<FString>>, OwnerTable)
+				.IsEnabled(this, &FSpawnConditionPage::IsRowEnabled, *Var.Get())
+				.Padding(FMargin(5, 0))
+				[
+					SNew(SCheckBox)
+					.OnCheckStateChanged(this, &FSpawnConditionPage::OnCheckStateChanged, *Var.Get())
+					.IsChecked(this, &FSpawnConditionPage::OnGetCheckState, *Var.Get())
+					.ToolTipText(FText::Format(LOCTEXT("ConditionSelectionTooltipFmt", "Name: {0}\nType: {1}"), FText::FromName(Var->GetName()), Var->GetType().GetNameText()))
+					[
+						SNew(SHorizontalBox)
+						+ SHorizontalBox::Slot()
+						.VAlign(VAlign_Center)
+						.HAlign(HAlign_Center)
+						.AutoWidth()
+						[
+							SNew(SImage)
+							.ColorAndOpacity(TypeColor)
+							.Image(FNiagaraEditorStyle::Get().GetBrush("NiagaraEditor.Module.TypeIconPill"))
+						]
+						+ SHorizontalBox::Slot()
+						.Padding(4, 2, 2, 2)
+						[
+							SNew(STextBlock)
+							.MinDesiredWidth(150)
+							.Text(FText::FromName(Var->GetName()))
+						]
+					]
+				];
+		}
+		
+		bool IsRowEnabled(FNiagaraDataChannelVariable Var) const
+		{
+			if (PreviousPage->Data->SpawnMode == ENiagaraDataChanneSpawnModuleMode::DirectSpawn)
+			{
+				return Var.GetType() == FNiagaraTypeDefinition::GetIntDef();
+			}
+			return true;
+		}
+
+		void OnCheckStateChanged(const ECheckBoxState NewState, FNiagaraDataChannelVariable Var)
+		{
+			if (PreviousPage->Data->SpawnMode == ENiagaraDataChanneSpawnModuleMode::DirectSpawn)
+			{
+				ConditionVariables.Empty();
+			}
+			if (NewState == ECheckBoxState::Checked)
+			{
+				ConditionVariables.Add(Var.Version);
+			}
+			else if (NewState == ECheckBoxState::Unchecked)
+			{
+				ConditionVariables.Remove(Var.Version);
+			}
+		}
+
+		ECheckBoxState OnGetCheckState(FNiagaraDataChannelVariable Var) const
+		{
+			return ConditionVariables.Contains(Var.Version) ? ECheckBoxState::Checked : ECheckBoxState::Unchecked;
+		}
+
+		TArray<TSharedPtr<FNiagaraDataChannelVariable>> AllVariables;
+		TSet<FGuid> ConditionVariables;
+		FSelectSpawnAssetPage* PreviousPage;
+		FObjectKey LastDataChannelRef;
 
 		TSharedPtr<SListView<TSharedPtr<FNiagaraDataChannelVariable>>> VarListView;
 	};
@@ -304,6 +539,7 @@ TSharedRef<FModuleWizardModel> DataChannel::CreateReadNDCModuleWizardModel()
 		{
 			FText ScriptName = VariablesPage->ModuleName;
 			ScratchPadScriptViewModel->SetScriptName(ScriptName.IsEmptyOrWhitespace() ? VariablesPage->CreateNewModuleName() : ScriptName);
+			ScratchPadScriptViewModel->GetEditScript().GetScriptData()->ModuleUsageBitmask = ENiagaraScriptUsageMask::System | ENiagaraScriptUsageMask::Emitter | ENiagaraScriptUsageMask::Particle;
 			
 			UNiagaraDataChannel* Channel = AssetPage->GetDataChannel();
 			UNiagaraGraph* Graph = ScratchPadScriptViewModel->GetGraphViewModel()->GetGraph();
@@ -329,7 +565,7 @@ TSharedRef<FModuleWizardModel> DataChannel::CreateReadNDCModuleWizardModel()
 					}
 
 					// create and connect read success output pin
-					UEdGraphPin* SuccessVarPin = Utilities::AddWriteParameterPin(FNiagaraTypeDefinition::GetBoolDef(), FName("Output.Module.ReadSuccess"), MapSetNode);
+					UEdGraphPin* SuccessVarPin = Utilities::AddWriteParameterPin(FNiagaraTypeDefinition::GetBoolDef(), FName(VariablesPage->GetTargetNamespace() + ".ReadSuccess"), MapSetNode);
 					UEdGraphPin* SuccessOutPin = ReadFunction->GetOutputPin(0);
 					if (SuccessOutPin && SuccessOutPin->GetName() == TEXT("Success"))
 					{
@@ -353,10 +589,22 @@ TSharedRef<FModuleWizardModel> DataChannel::CreateReadNDCModuleWizardModel()
 						UEdGraphPin* ReadParamPin = ReadFunction->AddParameterPin(SWCVar, EGPD_Output);
 
 						// add matching node on map set and connect them
-						UEdGraphPin* SetVarPin = Utilities::AddWriteParameterPin(SwcType, FName(TEXT("Output.Module.") + Var.GetName().ToString()), MapSetNode);
+						UEdGraphPin* SetVarPin = Utilities::AddWriteParameterPin(SwcType, FName(VariablesPage->GetTargetNamespace() + TEXT(".") + Var.GetName().ToString()), MapSetNode);
 						if (ReadParamPin && SetVarPin)
 						{
-							GraphSchema->TryCreateConnection(ReadParamPin, SetVarPin);
+							if (SwcType == FNiagaraTypeDefinition::GetPositionDef() && AssetPage->Data->bAutoTransformPositionData)
+							{
+								// transform position if necessary
+								if (UNiagaraNodeFunctionCall* TransformNode = Utilities::CreateFunctionCallNode(LoadObject<UNiagaraScript>(nullptr, TEXT("/Niagara/Functions/Localspace/TransformPosition.TransformPosition")), Graph))
+								{
+									GraphSchema->TryCreateConnection(ReadParamPin, TransformNode->FindPin(FName("Position"), EGPD_Input));
+									GraphSchema->TryCreateConnection(SetVarPin, TransformNode->FindPin(FName("Position"), EGPD_Output));
+								}								
+							}
+							else
+							{
+								GraphSchema->TryCreateConnection(ReadParamPin, SetVarPin);
+							}
 						}
 					}
 				}
@@ -383,14 +631,6 @@ TSharedRef<FModuleWizardModel> DataChannel::CreateReadNDCModuleWizardModel()
 							DataInterface->bUpdateSourceDataEveryTick = AssetPage->Data->bUpdateSourceDataEveryTick;
 						}
 					}
-
-					// bind index input for particle scripts. System and emitter scripts default to 0
-					if (FunctionInput->InputType == FNiagaraTypeDefinition::GetIntDef() && FNiagaraUtilities::ConvertScriptUsageToStaticSwitchContext(TargetUsage) == ENiagaraScriptContextStaticSwitch::Particle)
-					{
-						UNiagaraClipboardFunctionInput* EditableInput = const_cast<UNiagaraClipboardFunctionInput*>(FunctionInput);
-						EditableInput->ValueMode = ENiagaraClipboardFunctionInputValueMode::Linked;
-						EditableInput->Linked = SYS_PARAM_PARTICLES_UNIQUE_ID.GetName();
-					}
 				}
 				return true;
 			}
@@ -458,6 +698,7 @@ TSharedRef<FModuleWizardModel> DataChannel::CreateWriteNDCModuleWizardModel()
 		{
 			FText ScriptName = VariablesPage->ModuleName;
 			ScratchPadScriptViewModel->SetScriptName(ScriptName.IsEmptyOrWhitespace() ? VariablesPage->CreateNewModuleName() : ScriptName);
+			ScratchPadScriptViewModel->GetEditScript().GetScriptData()->ModuleUsageBitmask = ENiagaraScriptUsageMask::System | ENiagaraScriptUsageMask::Emitter | ENiagaraScriptUsageMask::Particle;
 			
 			UNiagaraDataChannel* Channel = AssetPage->GetDataChannel();
 			UNiagaraGraph* Graph = ScratchPadScriptViewModel->GetGraphViewModel()->GetGraph();
@@ -474,7 +715,7 @@ TSharedRef<FModuleWizardModel> DataChannel::CreateWriteNDCModuleWizardModel()
 			// Add inputs
 			UEdGraphPin* DIPin = Utilities::AddReadParameterPin(FNiagaraTypeDefinition(UNiagaraDataInterfaceDataChannelWrite::StaticClass()), FName("Data Channel"), MapGetNode);
 			UEdGraphPin* ExecWritePin = Utilities::AddReadParameterPin(FNiagaraTypeDefinition::GetBoolDef(), FName("Execute Write"), MapGetNode);
-			SetBoolDefaultValue(Graph, ExecWritePin->PinName, true);
+			Utilities::SetDefaultValue(Graph, ExecWritePin->PinName, FNiagaraTypeDefinition::GetBoolDef(), true);
 			UEdGraphPin* IndexPin = nullptr;
 			if (WriteMode == ENiagaraDataChanneWriteModuleMode::WriteToExistingElement)
 			{
@@ -503,7 +744,7 @@ TSharedRef<FModuleWizardModel> DataChannel::CreateWriteNDCModuleWizardModel()
 				}
 
 				// create and connect write success output pin
-				UEdGraphPin* SuccessVarPin = Utilities::AddWriteParameterPin(FNiagaraTypeDefinition::GetBoolDef(), FName("Output.Module.WriteSuccess"), MapSetNode);
+				UEdGraphPin* SuccessVarPin = Utilities::AddWriteParameterPin(FNiagaraTypeDefinition::GetBoolDef(), FName(VariablesPage->GetTargetNamespace() + ".WriteSuccess"), MapSetNode);
 				UEdGraphPin* SuccessOutPin = WriteFunction->GetOutputPin(1);
 				if (SuccessOutPin && SuccessOutPin->GetName() == TEXT("Success"))
 				{
@@ -530,7 +771,21 @@ TSharedRef<FModuleWizardModel> DataChannel::CreateWriteNDCModuleWizardModel()
 					UEdGraphPin* SetVarPin = Utilities::AddReadParameterPin(SwcType, FName(TEXT("Module.") + Var.GetName().ToString()), MapGetNode);
 					if (WriteParamPin && SetVarPin)
 					{
-						GraphSchema->TryCreateConnection(WriteParamPin, SetVarPin);
+						if (SwcType == FNiagaraTypeDefinition::GetPositionDef() && AssetPage->Data->bAutoTransformPositionData)
+						{
+							// transform position if necessary
+							if (UNiagaraNodeFunctionCall* TransformNode = Utilities::CreateFunctionCallNode(LoadObject<UNiagaraScript>(nullptr, TEXT("/Niagara/Functions/Localspace/TransformPosition.TransformPosition")), Graph))
+							{
+								GraphSchema->TryCreateConnection(WriteParamPin, TransformNode->FindPin(FName("Position"), EGPD_Output));
+								GraphSchema->TryCreateConnection(SetVarPin, TransformNode->FindPin(FName("Position"), EGPD_Input));
+								TransformNode->FindPin(FName("Source Space"), EGPD_Input)->DefaultValue = TEXT("Simulation");
+								TransformNode->FindPin(FName("Destination Space"), EGPD_Input)->DefaultValue = TEXT("World");
+							}								
+						}
+						else
+						{
+							GraphSchema->TryCreateConnection(WriteParamPin, SetVarPin);
+						}
 					}
 				}
 				
@@ -560,36 +815,373 @@ TSharedRef<FModuleWizardModel> DataChannel::CreateWriteNDCModuleWizardModel()
 							DataInterface->bUpdateDestinationDataEveryTick = AssetPage->Data->bUpdateDestinationDataEveryTick;
 						}
 					}
+				}
+				return true;
+			}
+			return false;
+		}
 
-					// bind index input for particle scripts. System and emitter scripts default to 0
-					if (FunctionInput->InputType == FNiagaraTypeDefinition::GetIntDef() && FunctionInput->InputName == FName("Write Index") && FNiagaraUtilities::ConvertScriptUsageToStaticSwitchContext(TargetUsage) == ENiagaraScriptContextStaticSwitch::Particle)
+		TSharedPtr<FSelectAssetPage> AssetPage;
+		TSharedPtr<FSelectVariablesPage> VariablesPage;
+	};
+		
+	return MakeShared<FWriteNDCModel>();
+}
+
+TSharedRef<FModuleWizardModel> DataChannel::CreateSpawnNDCModuleWizardModel()
+{
+	struct FSelectVariablesPage : FSelectVariablesPageBase
+	{
+		explicit FSelectVariablesPage(FSelectSpawnAssetPage* InPreviousPage) : FSelectVariablesPageBase(InPreviousPage)
+		{}
+		virtual ~FSelectVariablesPage() override = default;
+
+		virtual FText GetHeaderLabel() override
+		{
+			return LOCTEXT("VariablesPageLabel", "Please select which variables should be read into particle attributes when spawning.");
+		}
+		
+		virtual FText GetFormattedModuleName(const FText& AssetName) const override
+		{
+			return FText::Format(LOCTEXT("SpawnModuleNameFmt", "Spawn From {0}"), AssetName);
+		}
+	};
+	
+	struct FSpawnNDCModel : FModuleWizardModel
+	{
+		FSpawnNDCModel()
+		{
+			AssetPage = MakeShared<FSelectSpawnAssetPage>();
+			ConditionPage = MakeShared<FSpawnConditionPage>(AssetPage.Get());
+			VariablesPage = MakeShared<FSelectVariablesPage>(AssetPage.Get());
+			Pages.Add(AssetPage.ToSharedRef());
+			Pages.Add(ConditionPage.ToSharedRef());
+			Pages.Add(VariablesPage.ToSharedRef());
+		}
+		virtual ~FSpawnNDCModel() override = default;
+
+		virtual TArray<FModuleCreationEntry> GetModulesToCreate(UNiagaraNodeOutput* ProvidedOutputNode, int32 ProvidedTargetIndex, TSharedPtr<FNiagaraSystemViewModel> SystemModel, TSharedPtr<FNiagaraEmitterViewModel> EmitterViewModel) override
+		{
+			TArray<FModuleCreationEntry> Result;
+			if (UNiagaraScriptSource* EmitterGraphSource = Cast<UNiagaraScriptSource>(EmitterViewModel->GetEmitter().GetEmitterData()->GraphSource))
+			{
+				UNiagaraNodeOutput* SpawnScriptNode = EmitterGraphSource->NodeGraph->FindEquivalentOutputNode(ENiagaraScriptUsage::EmitterSpawnScript, FGuid());
+				Result.Add({SpawnScriptNode, INDEX_NONE}); // this is the module in emitter spawn to set up the common data channel parameter
+			}
+			Result.Add({ProvidedOutputNode, ProvidedTargetIndex}); // this is the spawn module in emitter update
+			if (UNiagaraScriptSource* EmitterGraphSource = Cast<UNiagaraScriptSource>(EmitterViewModel->GetEmitter().GetEmitterData()->GraphSource))
+			{
+				UNiagaraNodeOutput* SpawnScriptNode = EmitterGraphSource->NodeGraph->FindEquivalentOutputNode(ENiagaraScriptUsage::ParticleSpawnScript, FGuid());
+				Result.Add({SpawnScriptNode, 1}); // this is the module in particle spawn to write the particle data from the ndc
+			}
+			return Result;
+		}
+		
+		virtual void GenerateNewModuleContent(TSharedPtr<FNiagaraScratchPadScriptViewModel> ScratchPadScriptViewModel, const TArray<const UNiagaraNodeFunctionCall*>& PreviousModules) override
+		{
+			if (PreviousModules.Num() == 0)
+			{
+				// the emitter spawn module sets up the ndc parameter used by the other two modules
+				GenerateEmitterSpawnModule(ScratchPadScriptViewModel);
+			}
+			if (PreviousModules.Num() == 1)
+			{
+				// the emitter update module spawns calls the spawn functions
+				GenerateEmitterUpdateModule(ScratchPadScriptViewModel);
+			}
+			else if (PreviousModules.Num() == 2)
+			{
+				// the particle spawn module reads the data from the data channel row that spawned each particle
+				GenerateParticleSpawnModule(ScratchPadScriptViewModel);
+			}
+		}
+
+		void GenerateEmitterSpawnModule(TSharedPtr<FNiagaraScratchPadScriptViewModel> ScratchPadScriptViewModel)
+		{
+			ScratchPadScriptViewModel->SetScriptName(FText::FromString("Init data channel"));
+			
+			UNiagaraDataChannel* Channel = AssetPage->GetDataChannel();
+			UNiagaraGraph* Graph = ScratchPadScriptViewModel->GetGraphViewModel()->GetGraph();
+			if (Channel && Graph)
+			{
+				const UEdGraphSchema_Niagara* GraphSchema = Graph->GetNiagaraSchema();
+				UNiagaraNodeParameterMapGet* MapGetNode = Utilities::FindSingleNodeChecked<UNiagaraNodeParameterMapGet>(Graph);
+				UNiagaraNodeParameterMapSet* MapSetNode = Utilities::FindSingleNodeChecked<UNiagaraNodeParameterMapSet>(Graph);
+				UEdGraphPin* DIInputPin = Utilities::AddReadParameterPin(FNiagaraTypeDefinition(UNiagaraDataInterfaceDataChannelRead::StaticClass()), FName("Data Channel"), MapGetNode);
+				UEdGraphPin* DIVarPin = Utilities::AddWriteParameterPin(UNiagaraDataInterfaceDataChannelRead::StaticClass(), FName("Emitter.SpawnDataChannel"), MapSetNode);
+				GraphSchema->TryCreateConnection(DIInputPin, DIVarPin);
+				
+				FNiagaraStackGraphUtilities::RelayoutGraph(*Graph);
+                ScratchPadScriptViewModel->ApplyChanges();
+			}
+		}
+
+		void GenerateEmitterUpdateModule(TSharedPtr<FNiagaraScratchPadScriptViewModel> ScratchPadScriptViewModel)
+		{
+			FText ScriptName = VariablesPage->ModuleName;
+			ScratchPadScriptViewModel->SetScriptName(ScriptName.IsEmptyOrWhitespace() ? VariablesPage->CreateNewModuleName() : ScriptName);
+			
+			UNiagaraDataChannel* Channel = AssetPage->GetDataChannel();
+			UNiagaraGraph* Graph = ScratchPadScriptViewModel->GetGraphViewModel()->GetGraph();
+			if (Channel && Graph)
+			{
+				const UEdGraphSchema_Niagara* GraphSchema = Graph->GetNiagaraSchema();
+				UNiagaraNodeParameterMapGet* MapGetNode = Utilities::FindSingleNodeChecked<UNiagaraNodeParameterMapGet>(Graph);
+				UNiagaraNodeParameterMapSet* MapSetNode = Utilities::FindSingleNodeChecked<UNiagaraNodeParameterMapSet>(Graph);
+				Graph->RemoveNode(MapSetNode);
+				UNiagaraNodeInput* InputNode = Utilities::FindSingleNodeChecked<UNiagaraNodeInput>(Graph);
+				UNiagaraNodeOutput* OutputNode = Utilities::FindSingleNodeChecked<UNiagaraNodeOutput>(Graph);
+
+				// Call spawn function
+				ENiagaraDataChanneSpawnModuleMode SpawnMode = AssetPage->Data->SpawnMode;
+				if (UNiagaraNodeFunctionCall* SpawnFunction = Utilities::CreateDataInterfaceFunctionNode(UNiagaraDataInterfaceDataChannelRead::StaticClass(), FName(SpawnMode == ENiagaraDataChanneSpawnModuleMode::ConditionalSpawn ? "SpawnConditional" : "SpawnDirect"), Graph))
+				{
+					// connect base pins of the function call
+					UEdGraphPin* DIPin = Utilities::AddReadParameterPin(FNiagaraTypeDefinition(UNiagaraDataInterfaceDataChannelRead::StaticClass()), FName("Data Channel"), MapGetNode);
+					SpawnFunction->AutowireNewNode(DIPin);
+					GraphSchema->TryCreateConnection(InputNode->GetOutputPin(0), SpawnFunction->GetInputPin(0));
+					GraphSchema->TryCreateConnection(SpawnFunction->GetOutputPin(0), OutputNode->GetInputPin(0));
+					Utilities::SetDefaultBinding(Graph, DIPin->PinName, FName("Emitter.SpawnDataChannel"));
+
+					// Add module inputs
+					if (UEdGraphPin* EnableInput = SpawnFunction->FindPin(FName("Enable"), EGPD_Input))
+					{
+						UEdGraphPin* EnablePin = Utilities::AddReadParameterPin(FNiagaraTypeDefinition::GetBoolDef(), FName("Spawn Enabled"), MapGetNode);
+						Utilities::SetDefaultValue(Graph, EnablePin->PinName, FNiagaraTypeDefinition::GetBoolDef(), true);
+						GraphSchema->TryCreateConnection(EnablePin, EnableInput);
+					}
+					if (UEdGraphPin* EmitterIDInput = SpawnFunction->FindPin(FName("Emitter ID"), EGPD_Input))
+					{
+						UEdGraphPin* EmitterIDPin = Utilities::AddReadParameterPin(FNiagaraTypeDefinition(FNiagaraEmitterID::StaticStruct()), FName("Emitter ID"), MapGetNode);
+						GraphSchema->TryCreateConnection(EmitterIDInput, EmitterIDPin);
+						Utilities::SetDefaultBinding(Graph, EmitterIDPin->PinName, SYS_PARAM_ENGINE_EMITTER_ID.GetName());
+					}
+					if (UEdGraphPin* ModeInput = SpawnFunction->FindPin(FName("Mode"), EGPD_Input))
+					{
+						UEdGraphPin* SpawnModePin = Utilities::AddReadParameterPin(FNiagaraTypeDefinition(StaticEnum<ENDIDataChannelSpawnMode>()), FName("Spawn Mode"), MapGetNode);
+						GraphSchema->TryCreateConnection(ModeInput, SpawnModePin);
+					}
+					if (UEdGraphPin* OperatorInput = SpawnFunction->FindPin(FName("Operator"), EGPD_Input))
+					{
+						UEdGraphPin* OperatorPin = Utilities::AddReadParameterPin(FNiagaraTypeDefinition(StaticEnum<ENiagaraConditionalOperator>()), FName("Comparison Operator"), MapGetNode);
+						GraphSchema->TryCreateConnection(OperatorInput, OperatorPin);
+					}
+					if (UEdGraphPin* MinInput = SpawnFunction->FindPin(FName(SpawnMode == ENiagaraDataChanneSpawnModuleMode::ConditionalSpawn ? "Min Spawn Count" : "ClampMin"), EGPD_Input))
+					{
+						UEdGraphPin* SpawnMinPin = Utilities::AddReadParameterPin(FNiagaraTypeDefinition::GetIntDef(), FName("Min Count"), MapGetNode);
+						Utilities::SetDefaultValue(Graph, SpawnMinPin->PinName, FNiagaraTypeDefinition::GetIntDef(), 1);
+						GraphSchema->TryCreateConnection(MinInput, SpawnMinPin);
+					}
+					if (UEdGraphPin* MaxInput = SpawnFunction->FindPin(FName(SpawnMode == ENiagaraDataChanneSpawnModuleMode::ConditionalSpawn ? "Max Spawn Count" : "ClampMax"), EGPD_Input))
 					{
-						UNiagaraClipboardFunctionInput* EditableInput = const_cast<UNiagaraClipboardFunctionInput*>(FunctionInput);
-						EditableInput->ValueMode = ENiagaraClipboardFunctionInputValueMode::Linked;
-						EditableInput->Linked = SYS_PARAM_PARTICLES_UNIQUE_ID.GetName();
+						UEdGraphPin* SpawnMaxPin = Utilities::AddReadParameterPin(FNiagaraTypeDefinition::GetIntDef(), FName("Max Count"), MapGetNode);
+						Utilities::SetDefaultValue(Graph, SpawnMaxPin->PinName, FNiagaraTypeDefinition::GetIntDef(), 1);
+						GraphSchema->TryCreateConnection(MaxInput, SpawnMaxPin);
+					}
+					if (SpawnMode == ENiagaraDataChanneSpawnModuleMode::DirectSpawn)
+					{
+						if (UEdGraphPin* ScaleMinInput = SpawnFunction->FindPin(FName("RandomScaleMin"), EGPD_Input))
+						{
+							UEdGraphPin* ScaleMinPin = Utilities::AddReadParameterPin(FNiagaraTypeDefinition::GetFloatDef(), FName("Random Scale Min"), MapGetNode);
+							Utilities::SetDefaultValue(Graph, ScaleMinPin->PinName, FNiagaraTypeDefinition::GetFloatDef(), 1.0f);
+							GraphSchema->TryCreateConnection(ScaleMinInput, ScaleMinPin);
+						}
+						if (UEdGraphPin* ScaleMaxInput = SpawnFunction->FindPin(FName("RandomScaleMax"), EGPD_Input))
+						{
+							UEdGraphPin* ScaleMaxPin = Utilities::AddReadParameterPin(FNiagaraTypeDefinition::GetFloatDef(), FName("Random Scale Max"), MapGetNode);
+							Utilities::SetDefaultValue(Graph, ScaleMaxPin->PinName, FNiagaraTypeDefinition::GetFloatDef(), 1.0f);
+							GraphSchema->TryCreateConnection(ScaleMaxInput, ScaleMaxPin);
+						}
+
+						// Encode spawn direct condition variable
+						if (ConditionPage->ConditionVariables.Num() == 1)
+						{
+							for (const FNiagaraDataChannelVariable& Var : Channel->GetVariables())
+							{
+								if (ConditionPage->ConditionVariables.Contains(Var.Version))
+								{
+									const static FName VarNameSpecifierKey(TEXT("VarName"));
+									const static FName VarTypeSpecifierKey(TEXT("VarType"));
+									FString TypeStr;
+									UScriptStruct* TypeStruct = FNiagaraTypeDefinition::StaticStruct();
+									FNiagaraTypeDefinition TypeDef = FNiagaraTypeDefinition::GetIntDef();
+									TypeStruct->ExportText(TypeStr, &TypeDef, nullptr, nullptr, PPF_None, nullptr);
+									SpawnFunction->SetFunctionSpecifier(VarNameSpecifierKey, Var.GetName());
+									SpawnFunction->SetFunctionSpecifier(VarTypeSpecifierKey, *TypeStr);
+									break;
+								}
+							}
+						}
+					}
+					
+					if (SpawnMode == ENiagaraDataChanneSpawnModuleMode::ConditionalSpawn)
+					{
+						for (const FNiagaraDataChannelVariable& Var : Channel->GetVariables())
+						{
+							if (!ConditionPage->ConditionVariables.Contains(Var.Version))
+							{
+								continue;
+							}
+							FNiagaraTypeDefinition SwcType = Var.GetType();
+							if (SwcType.IsEnum() == false)
+							{
+								SwcType = FNiagaraTypeDefinition(FNiagaraTypeHelper::GetSWCStruct(Var.GetType().GetScriptStruct()));
+							}
+							FNiagaraVariable SWCVar(SwcType, Var.GetName());
+							UEdGraphPin* ConditionParamPin = SpawnFunction->AddParameterPin(SWCVar, EGPD_Input);
+
+							// add matching node on map get and connect them
+							UEdGraphPin* SetVarPin = Utilities::AddReadParameterPin(SwcType, FName(TEXT("Module.") + Var.GetName().ToString() + " Condition"), MapGetNode);
+							if (ConditionParamPin && SetVarPin)
+							{
+								if (SwcType == FNiagaraTypeDefinition::GetPositionDef() && AssetPage->Data->bAutoTransformPositionData)
+								{
+									// transform position if necessary
+									if (UNiagaraNodeFunctionCall* TransformNode = Utilities::CreateFunctionCallNode(LoadObject<UNiagaraScript>(nullptr, TEXT("/Niagara/Functions/Localspace/TransformPosition.TransformPosition")), Graph))
+									{
+										GraphSchema->TryCreateConnection(ConditionParamPin, TransformNode->FindPin(FName("Position"), EGPD_Output));
+										GraphSchema->TryCreateConnection(SetVarPin, TransformNode->FindPin(FName("Position"), EGPD_Input));
+										TransformNode->FindPin(FName("Source Space"), EGPD_Input)->DefaultValue = TEXT("Simulation");
+										TransformNode->FindPin(FName("Destination Space"), EGPD_Input)->DefaultValue = TEXT("World");
+									}								
+								}
+								else
+								{
+									GraphSchema->TryCreateConnection(ConditionParamPin, SetVarPin);
+								}
+							}
+						}
 					}
 				}
-				return true;
+				
+				FNiagaraStackGraphUtilities::RelayoutGraph(*Graph);
+				ScratchPadScriptViewModel->ApplyChanges();
+			}
+		}
+
+		void GenerateParticleSpawnModule(TSharedPtr<FNiagaraScratchPadScriptViewModel> ScratchPadScriptViewModel)
+		{
+			ScratchPadScriptViewModel->SetScriptName(LOCTEXT("SpawnParticleModuleName", "Init Particle From NDC"));
+			
+			UNiagaraDataChannel* Channel = AssetPage->GetDataChannel();
+			UNiagaraGraph* Graph = ScratchPadScriptViewModel->GetGraphViewModel()->GetGraph();
+			if (Channel && Graph)
+			{
+				const UEdGraphSchema_Niagara* GraphSchema = Graph->GetNiagaraSchema();
+				UNiagaraNodeParameterMapGet* MapGetNode = Utilities::FindSingleNodeChecked<UNiagaraNodeParameterMapGet>(Graph);
+				UNiagaraNodeParameterMapSet* MapSetNode = Utilities::FindSingleNodeChecked<UNiagaraNodeParameterMapSet>(Graph);
+
+				// Call read functions
+				UNiagaraNodeFunctionCall* SpawnDataFunction = Utilities::CreateDataInterfaceFunctionNode(UNiagaraDataInterfaceDataChannelRead::StaticClass(), FName("GetNDCSpawnData"), Graph);
+				UNiagaraNodeFunctionCall* ReadFunction = Utilities::CreateDataInterfaceFunctionNode(UNiagaraDataInterfaceDataChannelRead::StaticClass(), FName("Read"), Graph);
+				if (SpawnDataFunction && ReadFunction)
+				{
+					// Add module inputs
+					UEdGraphPin* DIPin = Utilities::AddReadParameterPin(FNiagaraTypeDefinition(UNiagaraDataInterfaceDataChannelRead::StaticClass()), FName("Data Channel"), MapGetNode);
+					SpawnDataFunction->AutowireNewNode(DIPin);
+					ReadFunction->AutowireNewNode(DIPin);
+					Utilities::SetDefaultBinding(Graph, DIPin->PinName, FName("Emitter.SpawnDataChannel"));
+					
+					if (UEdGraphPin* EmitterIDInput = SpawnDataFunction->FindPin(FName("Emitter ID"), EGPD_Input))
+					{
+						UEdGraphPin* EmitterIDPin = Utilities::AddReadParameterPin(FNiagaraTypeDefinition(FNiagaraEmitterID::StaticStruct()), FName("Emitter ID"), MapGetNode);
+						GraphSchema->TryCreateConnection(EmitterIDInput, EmitterIDPin);
+						Utilities::SetDefaultBinding(Graph, EmitterIDPin->PinName, SYS_PARAM_ENGINE_EMITTER_ID.GetName());
+					}
+					
+					// create and connect exec index node
+					UNiagaraNodeOp* ExecIndexNode = NewObject<UNiagaraNodeOp>(Graph);
+					ExecIndexNode->OpName = FName("Util::ExecIndex");
+					ExecIndexNode->SetFlags(RF_Transactional);
+					Graph->AddNode(ExecIndexNode, false, false);
+					ExecIndexNode->AllocateDefaultPins();
+					if (UEdGraphPin* EmitterIDInput = SpawnDataFunction->FindPin(FName("Spawned Particle Exec Index"), EGPD_Input))
+					{
+						GraphSchema->TryCreateConnection(EmitterIDInput, ExecIndexNode->Pins[0]);
+					}
+
+					// connect index pins
+					GraphSchema->TryCreateConnection(SpawnDataFunction->GetOutputPin(0), ReadFunction->GetInputPin(1));
+
+					// create and connect read success output pin
+					UEdGraphPin* SuccessVarPin = Utilities::AddWriteParameterPin(FNiagaraTypeDefinition::GetBoolDef(), FName(VariablesPage->GetTargetNamespace() + ".ReadSuccess"), MapSetNode);
+					UEdGraphPin* SuccessOutPin = ReadFunction->GetOutputPin(0);
+					if (SuccessOutPin && SuccessOutPin->GetName() == TEXT("Success"))
+					{
+						GraphSchema->TryCreateConnection(SuccessOutPin, SuccessVarPin);
+					}
+
+					// add channel variable pins to read node
+					for (const FNiagaraDataChannelVariable& Var : Channel->GetVariables())
+					{
+						if (!VariablesPage->VariablesToProcess.Contains(Var.Version))
+						{
+							continue;
+						}
+						
+						FNiagaraTypeDefinition SwcType = Var.GetType();
+						if (SwcType.IsEnum() == false)
+						{
+							SwcType = FNiagaraTypeDefinition(FNiagaraTypeHelper::GetSWCStruct(Var.GetType().GetScriptStruct()));
+						}
+						FNiagaraVariable SWCVar(SwcType, Var.GetName());
+						UEdGraphPin* ReadParamPin = ReadFunction->AddParameterPin(SWCVar, EGPD_Output);
+
+						// add matching node on map set and connect them
+						UEdGraphPin* SetVarPin = Utilities::AddWriteParameterPin(SwcType, FName(VariablesPage->GetTargetNamespace() + TEXT(".") + Var.GetName().ToString()), MapSetNode);
+						if (ReadParamPin && SetVarPin)
+						{
+							if (SwcType == FNiagaraTypeDefinition::GetPositionDef() && AssetPage->Data->bAutoTransformPositionData)
+							{
+								// transform position if necessary
+								if (UNiagaraNodeFunctionCall* TransformNode = Utilities::CreateFunctionCallNode(LoadObject<UNiagaraScript>(nullptr, TEXT("/Niagara/Functions/Localspace/TransformPosition.TransformPosition")), Graph))
+								{
+									GraphSchema->TryCreateConnection(ReadParamPin, TransformNode->FindPin(FName("Position"), EGPD_Input));
+									GraphSchema->TryCreateConnection(SetVarPin, TransformNode->FindPin(FName("Position"), EGPD_Output));
+								}								
+							}
+							else
+							{
+								GraphSchema->TryCreateConnection(ReadParamPin, SetVarPin);
+							}
+						}
+					}
+				}
+				
+				FNiagaraStackGraphUtilities::RelayoutGraph(*Graph);
+				ScratchPadScriptViewModel->ApplyChanges();
 			}
-			return false;
 		}
 
-		void SetBoolDefaultValue(UNiagaraGraph* Graph, const FName& VarName, bool Value) const
+		virtual bool UpdateModuleInputs(UNiagaraClipboardContent* NewModule, const TArray<const UNiagaraNodeFunctionCall*>& PreviousModules) override
 		{
-			if (UNiagaraScriptVariable* ScriptVariable = Graph->GetScriptVariable(VarName))
+			if (UNiagaraDataChannelAsset* Channel = AssetPage->GetAsset())
 			{
-				FNiagaraVariable Var(FNiagaraTypeDefinition::GetBoolDef(), FName("Var"));
-				Var.SetValue(Value);
-				ScriptVariable->SetDefaultValueData(Var.GetData());
-				Graph->ScriptVariableChanged(ScriptVariable->Variable);
+				TArray<TObjectPtr<const UNiagaraClipboardFunctionInput>> FunctionInputs = NewModule->FunctionInputs;
+				for (const UNiagaraClipboardFunctionInput* FunctionInput : FunctionInputs)
+				{
+					if (PreviousModules.Num() == 0 && FunctionInput->InputType == FNiagaraTypeDefinition(UNiagaraDataInterfaceDataChannelRead::StaticClass()))
+					{
+						// set data interface module input
+						if (UNiagaraDataInterfaceDataChannelRead* DataInterface = Cast<UNiagaraDataInterfaceDataChannelRead>(FunctionInput->Data))
+						{
+							DataInterface->Channel = Channel;
+							DataInterface->bReadCurrentFrame = AssetPage->Data->bReadCurrentFrame;
+							DataInterface->bUpdateSourceDataEveryTick = AssetPage->Data->bUpdateSourceDataEveryTick;
+						}
+					}
+				}
+				return true;
 			}
+			return false;
 		}
 
-		TSharedPtr<FSelectAssetPage> AssetPage;
+		TSharedPtr<FSelectSpawnAssetPage> AssetPage;
+		TSharedPtr<FSpawnConditionPage> ConditionPage;
 		TSharedPtr<FSelectVariablesPage> VariablesPage;
 	};
-		
-	return MakeShared<FWriteNDCModel>();
+	
+	return MakeShared<FSpawnNDCModel>();
 }
 
 TSharedRef<FModuleWizardGenerator> DataChannel::CreateNDCWizardGenerator()
@@ -605,19 +1197,21 @@ TSharedRef<FModuleWizardGenerator> DataChannel::CreateNDCWizardGenerator()
 			ReadAction.DisplayName = LOCTEXT("NewReadNDCModuleName", "Read From Data Channel...");
 			ReadAction.Description = LOCTEXT("NewReadNDCModuleDescription", "Description: Create a new scratch pad module to read attributes from a data channel");
 			ReadAction.Keywords = LOCTEXT("NewReadNDCModuleKeywords", "ndc reader datachannel get external");
-			ReadAction.bSuggestedAction = true;
 			ReadAction.WizardModel = CreateReadNDCModuleWizardModel();
 
 			FAction& WriteAction = WizardActions.AddDefaulted_GetRef();
-			WriteAction. DisplayName = LOCTEXT("NewWriteNDCModuleName", "Write To Data Channel...");
+			WriteAction.DisplayName = LOCTEXT("NewWriteNDCModuleName", "Write To Data Channel...");
 			WriteAction.Description = LOCTEXT("NewWriteNDCModuleDescription", "Description: Create a new scratch pad module to write attributes to a data channel");
 			WriteAction.Keywords = LOCTEXT("NewWriteNDCModuleKeywords", "ndc writer datachannel save append external");
-			WriteAction.bSuggestedAction = true;
 			WriteAction.WizardModel = CreateWriteNDCModuleWizardModel();
 
 			if (Usage == ENiagaraScriptUsage::EmitterUpdateScript)
 			{
-				//TODO: add spawn from ndc wizard
+				FAction& SpawnAction = WizardActions.AddDefaulted_GetRef();
+				SpawnAction.DisplayName = LOCTEXT("NewSpawnNDCModuleName", "Spawn From Data Channel...");
+				SpawnAction.Description = LOCTEXT("NewSpawnNDCModuleDescription", "Description: Create a new scratch pad module to spawn particles from data channel entries. Every time an entry is added to the data channel, it will burst spawn new particles.");
+				SpawnAction.Keywords = LOCTEXT("NewSpawnNDCModuleKeywords", "ndc spawner datachannel particles burst external");
+				SpawnAction.WizardModel = CreateSpawnNDCModuleWizardModel();
 			}
 			
 			return WizardActions;
diff --git a/Engine/Plugins/FX/Niagara/Source/NiagaraEditor/Private/Widgets/DataChannel/NiagaraDataChannelWizard.h b/Engine/Plugins/FX/Niagara/Source/NiagaraEditor/Private/Widgets/DataChannel/NiagaraDataChannelWizard.h
--- a/Engine/Plugins/FX/Niagara/Source/NiagaraEditor/Private/Widgets/DataChannel/NiagaraDataChannelWizard.h
+++ b/Engine/Plugins/FX/Niagara/Source/NiagaraEditor/Private/Widgets/DataChannel/NiagaraDataChannelWizard.h
@@ -42,6 +42,72 @@ public:
 	*/
 	UPROPERTY(EditAnywhere, Category = "Advanced")
 	bool bUpdateSourceDataEveryTick = true;
+
+	/**
+	If true then position inputs are automatically transformed from world space to simulation space, so the read works correctly for localspace emitters.
+	*/
+	UPROPERTY(EditAnywhere, Category = "Advanced")
+	bool bAutoTransformPositionData = true;
+};
+
+UENUM()
+enum class ENiagaraDataChanneSpawnModuleMode : uint8
+{
+	/**
+	 Spawn particles for each entry in the data channel. Optionally checks if certain conditions are met, for example if
+	 the entry has the correct material attribute or if a bool attribute is set to true.
+	 */
+	ConditionalSpawn,
+	
+	/**
+	 The number of particles to spawn is read directly from an attribute in the data channel. 
+	 */
+	DirectSpawn
+};
+
+/**
+ * Helper class to display properties in the read data channel module wizard
+ */
+UCLASS()
+class UNiagaraDataChannelSpawnModuleData : public UObject
+{
+	GENERATED_BODY()
+
+public:
+	/** The source asset to spawn from */
+	UPROPERTY(EditAnywhere, Category = "Source")
+	TObjectPtr<UNiagaraDataChannelAsset> DataChannel;
+
+	/** Determines how new particles should be spawned.
+	 *
+	 * ConditionalSpawn - Always spawn particles in a when a data channel entry fulfills the (optional) conditions. 
+	 *
+	 * DirectSpawn - The number of particles to spawn is read directly from an attribute in the data channel.
+	 */
+	UPROPERTY(EditAnywhere, Category = "Source")
+	ENiagaraDataChanneSpawnModuleMode SpawnMode = ENiagaraDataChanneSpawnModuleMode::ConditionalSpawn;
+
+
+	/** True if this reader will read the current frame's data. If false, we read the previous frame.
+	* Reading the previous frame's data introduces a frame of latency but ensures we never miss any data as we have access to the whole frame.
+	*/
+	UPROPERTY(EditAnywhere, Category = "Advanced")
+	bool bReadCurrentFrame = false;
+
+	/**
+	Whether this DI should request updated source data from the Data Channel each tick.
+	Some Data Channels have multiple separate source data elements for things such as spatial subdivision. 
+	Each DI will request the correct one for it's owning system instance from the data channel. 
+	Depending on the data channel this could be an expensive search so we should avoid doing this every tick if possible.
+	*/
+	UPROPERTY(EditAnywhere, Category = "Advanced")
+	bool bUpdateSourceDataEveryTick = true;
+
+	/**
+	If true then position inputs are automatically transformed from world space to simulation space, so the spawning works correctly for localspace emitters.
+	*/
+	UPROPERTY(EditAnywhere, Category = "Advanced")
+	bool bAutoTransformPositionData = true;
 };
 
 UENUM()
@@ -84,9 +150,9 @@ public:
 	UPROPERTY(EditAnywhere, Category = "Advanced")
 	ENiagaraDataChannelAllocationMode AllocationMode = ENiagaraDataChannelAllocationMode::Static;
 
-	/** How many elements to allocate for writing per frame? Usage is defendant on AllocationMode. */
+	/** How many elements to allocate for writing per frame? Usage depends on AllocationMode. */
 	UPROPERTY(EditAnywhere, Category = "Advanced")
-	uint32 AllocationCount = 0;
+	uint32 AllocationCount = 1;
 
 	/**
 	Whether this DI should request updated destination data from the Data Channel each tick.
@@ -94,11 +160,18 @@ public:
 	*/
 	UPROPERTY(EditAnywhere, Category = "Advanced")
 	bool bUpdateDestinationDataEveryTick = true;
+
+	/**
+	If true then position inputs are automatically transformed from simulation space to world space, so the write works correctly for localspace emitters.
+	*/
+	UPROPERTY(EditAnywhere, Category = "Advanced")
+	bool bAutoTransformPositionData = true;
 };
 
 namespace UE::Niagara::Wizard::DataChannel
 {
 	TSharedRef<FModuleWizardModel> CreateReadNDCModuleWizardModel();
 	TSharedRef<FModuleWizardModel> CreateWriteNDCModuleWizardModel();
+	TSharedRef<FModuleWizardModel> CreateSpawnNDCModuleWizardModel();
 	TSharedRef<FModuleWizardGenerator> CreateNDCWizardGenerator();
 };
diff --git a/Engine/Plugins/FX/Niagara/Source/NiagaraEditor/Private/Widgets/Wizard/SNiagaraModuleWizard.cpp b/Engine/Plugins/FX/Niagara/Source/NiagaraEditor/Private/Widgets/Wizard/SNiagaraModuleWizard.cpp
--- a/Engine/Plugins/FX/Niagara/Source/NiagaraEditor/Private/Widgets/Wizard/SNiagaraModuleWizard.cpp
+++ b/Engine/Plugins/FX/Niagara/Source/NiagaraEditor/Private/Widgets/Wizard/SNiagaraModuleWizard.cpp
@@ -214,4 +214,32 @@ UNiagaraNodeFunctionCall* Utilities::CreateDataInterfaceFunctionNode(const TSubc
 	return nullptr;
 }
 
+UNiagaraNodeFunctionCall* Utilities::CreateFunctionCallNode(UNiagaraScript* FunctionScript, UNiagaraGraph* Graph)
+{
+	if (FunctionScript)
+	{
+		UNiagaraNodeFunctionCall* FuncNode = NewObject<UNiagaraNodeFunctionCall>(Graph);
+		FuncNode->FunctionScript = FunctionScript;
+		FuncNode->SelectedScriptVersion = FunctionScript->GetExposedVersion().VersionGuid;
+		FuncNode->SetFlags(RF_Transactional);
+		Graph->AddNode(FuncNode, false, false);
+
+		FuncNode->CreateNewGuid();
+		FuncNode->PostPlacedNewNode();
+		FuncNode->AllocateDefaultPins();
+		return FuncNode;
+	}
+	return nullptr;
+}
+
+void Utilities::SetDefaultBinding(UNiagaraGraph* Graph, const FName& VarName, const FName& DefaultBinding)
+{
+	if (UNiagaraScriptVariable* ScriptVariable = Graph->GetScriptVariable(VarName))
+	{
+		ScriptVariable->DefaultMode = ENiagaraDefaultMode::Binding;
+		ScriptVariable->DefaultBinding.SetName(DefaultBinding);
+		Graph->ScriptVariableChanged(ScriptVariable->Variable);
+	}
+}
+
 #undef LOCTEXT_NAMESPACE
diff --git a/Engine/Plugins/FX/Niagara/Source/NiagaraEditor/Public/Widgets/Wizard/SNiagaraModuleWizard.h b/Engine/Plugins/FX/Niagara/Source/NiagaraEditor/Public/Widgets/Wizard/SNiagaraModuleWizard.h
--- a/Engine/Plugins/FX/Niagara/Source/NiagaraEditor/Public/Widgets/Wizard/SNiagaraModuleWizard.h
+++ b/Engine/Plugins/FX/Niagara/Source/NiagaraEditor/Public/Widgets/Wizard/SNiagaraModuleWizard.h
@@ -119,5 +119,20 @@ namespace UE::Niagara::Wizard
 		NIAGARAEDITOR_API UEdGraphPin* AddWriteParameterPin(const FNiagaraTypeDefinition& Type, const FName& Name, UNiagaraNodeParameterMapSet* MapSetNode);
 		NIAGARAEDITOR_API TSharedRef<IDetailsView> CreateDetailsView();
 		NIAGARAEDITOR_API UNiagaraNodeFunctionCall* CreateDataInterfaceFunctionNode(const TSubclassOf<UNiagaraDataInterface>& DataInterfaceClass, const FName& FunctionName, UNiagaraGraph* Graph);
+		NIAGARAEDITOR_API UNiagaraNodeFunctionCall* CreateFunctionCallNode(UNiagaraScript* FunctionScript, UNiagaraGraph* Graph);
+
+		template <typename T>
+		NIAGARAEDITOR_API void SetDefaultValue(UNiagaraGraph* Graph, const FName& VarName, const FNiagaraTypeDefinition& TypeDef, T Value)
+		{
+			if (UNiagaraScriptVariable* ScriptVariable = Graph->GetScriptVariable(VarName))
+			{
+				// create temp var for data storage
+				FNiagaraVariable Var(TypeDef, FName("Var"));
+				Var.SetValue(Value);
+				ScriptVariable->SetDefaultValueData(Var.GetData());
+				Graph->ScriptVariableChanged(ScriptVariable->Variable);
+			}
+		}
+		NIAGARAEDITOR_API void SetDefaultBinding(UNiagaraGraph* Graph, const FName& VarName, const FName& DefaultBinding);
 	}
 }
