Removing stubbed out support for half float types in Niagara Data Channel shaders.
Half variables cannot be used currently anyway without additional work.
Stubeed out support in shaders was causing warnings in iOS cooks and using additional UAV slots when not needed.

diff --git a/Engine/Plugins/FX/Niagara/Shaders/Private/DataChannel/NiagaraDataInterfaceDataChannelCommon.ush b/Engine/Plugins/FX/Niagara/Shaders/Private/DataChannel/NiagaraDataInterfaceDataChannelCommon.ush
--- a/Engine/Plugins/FX/Niagara/Shaders/Private/DataChannel/NiagaraDataInterfaceDataChannelCommon.ush
+++ b/Engine/Plugins/FX/Niagara/Shaders/Private/DataChannel/NiagaraDataInterfaceDataChannelCommon.ush
@@ -68,6 +68,8 @@ void NiagaraDataChannelRead_Int32(Buffer<int> InBuffer, uint Stride, in out uint
 	NiagaraDataChannelRead_Int32(InBuffer, Stride, RegisterIdx, ParticleIdx, OutValue.w);
 }
 
+/*
+//TODO: Half Support | 
 void NiagaraDataChannelRead_Half(Buffer<half> InBuffer, uint Stride, in out uint RegisterIdx, uint ParticleIdx, out half OutValue)
 {
 	OutValue = InBuffer[RegisterIdx * Stride + ParticleIdx];
@@ -94,6 +96,7 @@ void NiagaraDataChannelRead_Half(Buffer<half> InBuffer, uint Stride, in out uint
 	NiagaraDataChannelRead_Half(InBuffer, Stride, RegisterIdx, ParticleIdx, OutValue.z);
 	NiagaraDataChannelRead_Half(InBuffer, Stride, RegisterIdx, ParticleIdx, OutValue.w);
 }
+*/
 
 // Write accessors 
 
@@ -151,6 +154,8 @@ void NiagaraDataChannelWrite_Int32(RWBuffer<int> InBuffer, uint Stride, in out u
 	NiagaraDataChannelWrite_Int32(InBuffer, Stride, RegisterIdx, ParticleIdx, InValue.w);
 }
 
+/*
+//TODO: Half Support | 
 void NiagaraDataChannelWrite_Half(RWBuffer<half> InBuffer, uint Stride, in out uint RegisterIdx, uint ParticleIdx, half InValue)
 {
 	InBuffer[RegisterIdx * Stride + ParticleIdx] = InValue;
@@ -177,6 +182,7 @@ void NiagaraDataChannelWrite_Half(RWBuffer<half> InBuffer, uint Stride, in out u
 	NiagaraDataChannelWrite_Half(InBuffer, Stride, RegisterIdx, ParticleIdx, InValue.z);
 	NiagaraDataChannelWrite_Half(InBuffer, Stride, RegisterIdx, ParticleIdx, InValue.w);
 }
+*/
 
 //TODO: More base type accessors.
 
diff --git a/Engine/Plugins/FX/Niagara/Shaders/Private/DataChannel/NiagaraDataInterfaceDataChannelTemplateCommon.ush b/Engine/Plugins/FX/Niagara/Shaders/Private/DataChannel/NiagaraDataInterfaceDataChannelTemplateCommon.ush
--- a/Engine/Plugins/FX/Niagara/Shaders/Private/DataChannel/NiagaraDataInterfaceDataChannelTemplateCommon.ush
+++ b/Engine/Plugins/FX/Niagara/Shaders/Private/DataChannel/NiagaraDataInterfaceDataChannelTemplateCommon.ush
@@ -12,26 +12,26 @@ Buffer<uint> {ParameterName}_ParamOffsetTable;
 int {ParameterName}_ParameterOffsetTableIndex;
 int {ParameterName}_FloatStride;
 int {ParameterName}_Int32Stride;
-int {ParameterName}_HalfStride;
+//TODO: Half Support | int {ParameterName}_HalfStride;
 
 // Outputs the starting offset of each component in the channel data for a partiuclar parameter.
-bool GetParameterLayoutInfo_{ParameterName}(uint ParameterIndex, out uint OutFloatRegister, out uint OutInt32Register, out uint OutHalfRegister)
+bool GetParameterLayoutInfo_{ParameterName}(uint ParameterIndex, out uint OutFloatRegister, out uint OutInt32Register)//TODO: Half Support | , out uint OutHalfRegister)
 {
 	int ParameterTableStart = {ParameterName}_ParameterOffsetTableIndex;
 	if(ParameterTableStart != -1)
 	{
-		const int NumComponents = 3;
+		const int NumComponents = 2;
 		uint ParamLayoutStart = ParameterTableStart + NumComponents * ParameterIndex;
 		OutFloatRegister = {ParameterName}_ParamOffsetTable[ParamLayoutStart];
 		OutInt32Register = {ParameterName}_ParamOffsetTable[ParamLayoutStart + 1];
-		OutHalfRegister = {ParameterName}_ParamOffsetTable[ParamLayoutStart + 2];
+		//TODO: Half Support | OutHalfRegister = {ParameterName}_ParamOffsetTable[ParamLayoutStart + 2];
 		return true;
 	}
 	else
 	{
 		OutFloatRegister = -1;
 		OutInt32Register = -1;
-		OutHalfRegister = -1;
+		//TODO: Half Support | OutHalfRegister = -1;
 		return false;
 	}
 }
diff --git a/Engine/Plugins/FX/Niagara/Shaders/Private/DataChannel/NiagaraDataInterfaceDataChannelTemplateReadCommon.ush b/Engine/Plugins/FX/Niagara/Shaders/Private/DataChannel/NiagaraDataInterfaceDataChannelTemplateReadCommon.ush
--- a/Engine/Plugins/FX/Niagara/Shaders/Private/DataChannel/NiagaraDataInterfaceDataChannelTemplateReadCommon.ush
+++ b/Engine/Plugins/FX/Niagara/Shaders/Private/DataChannel/NiagaraDataInterfaceDataChannelTemplateReadCommon.ush
@@ -6,7 +6,7 @@ Template hlsl for Data Channel Read DI functions to apply to their own that is s
 
 Buffer<float> {ParameterName}_DataFloat;
 Buffer<int> {ParameterName}_DataInt32;
-Buffer<float> {ParameterName}_DataHalf;
+//TODO: Half Support | Buffer<half> {ParameterName}_DataHalf;
 
 //Instance count offset for the buffer we're reading.
 //Must be done via an instance count slot as data can come from CPU and GPU.
@@ -34,20 +34,20 @@ struct FNDCAccessContext_{ParameterName}
 	//Component start offests for the current parameter.
 	int ParameterComponentStart_Float;
 	int ParameterComponentStart_Int32;
-	int ParameterComponentStart_Half;
+	//TODO: Half Support | int ParameterComponentStart_Half;
 
 	//Current registe indices for each of the GPU write buffers
 	int Register_Float;
 	int Register_Int32;
-	int Register_Half;
+	//TODO: Half Support | int Register_Half;
 
 	int GetFloatStride() { return {ParameterName}_FloatStride; }
 	int GetIntStride() { return {ParameterName}_Int32Stride; }
-	int GetHalfStride() { return {ParameterName}_HalfStride; }
+	//TODO: Half Support | int GetHalfStride() { return {ParameterName}_HalfStride; }
 		
 	Buffer<float> GetBuffer_Float() { return {ParameterName}_DataFloat; }
 	Buffer<int> GetBuffer_Int32() { return {ParameterName}_DataInt32; }
-	Buffer<float> GetBuffer_Half() { return {ParameterName}_DataHalf; }
+	//TODO: Half Support | Buffer<half> GetBuffer_Half() { return {ParameterName}_DataHalf; }
 	
 	//Instance count used for the written number of elements in the NDC buffers.
 	int GetInstanceCountOffset() { return {ParameterName}_InstanceCountOffset; }
@@ -92,11 +92,11 @@ struct FNDCAccessContext_{ParameterName}
 
 	void InitForParameter(int ParameterIndex)
 	{
-		bValid = bValid && GetParameterLayoutInfo_{ParameterName}(ParameterIndex, ParameterComponentStart_Float, ParameterComponentStart_Int32, ParameterComponentStart_Half);
+		bValid = bValid && GetParameterLayoutInfo_{ParameterName}(ParameterIndex, ParameterComponentStart_Float, ParameterComponentStart_Int32);//TODO: Half Support | ParameterComponentStart_Half);
 		
 		Register_Float = ParameterComponentStart_Float;
 		Register_Int32 = ParameterComponentStart_Int32;
-		Register_Half = ParameterComponentStart_Half;
+		//TODO: Half Support | Register_Half = ParameterComponentStart_Half;
 	}
 
 	void Read_Float(out float Value){ NiagaraDataChannelRead_Float(GetBuffer_Float(), GetFloatStride(), Register_Float, ReadIndex, Value); }
@@ -109,10 +109,10 @@ struct FNDCAccessContext_{ParameterName}
 	void Read_Int32(out int3 Value){ NiagaraDataChannelRead_Int32(GetBuffer_Int32(), GetIntStride(), Register_Int32, ReadIndex, Value); }
 	void Read_Int32(out int4 Value){ NiagaraDataChannelRead_Int32(GetBuffer_Int32(), GetIntStride(), Register_Int32, ReadIndex, Value); }
 
-	void Read_Half(out half Value){ NiagaraDataChannelRead_Half(GetBuffer_Half(), GetHalfStride(), Register_Half, ReadIndex, Value); }
-	void Read_Half(out half2 Value){ NiagaraDataChannelRead_Half(GetBuffer_Half(), GetHalfStride(), Register_Half, ReadIndex, Value); }
-	void Read_Half(out half3 Value){ NiagaraDataChannelRead_Half(GetBuffer_Half(), GetHalfStride(), Register_Half, ReadIndex, Value); }
-	void Read_Half(out half4 Value){ NiagaraDataChannelRead_Half(GetBuffer_Half(), GetHalfStride(), Register_Half, ReadIndex, Value); }
+	//TODO: Half Support | void Read_Half(out half Value){ NiagaraDataChannelRead_Half(GetBuffer_Half(), GetHalfStride(), Register_Half, ReadIndex, Value); }
+	//TODO: Half Support | void Read_Half(out half2 Value){ NiagaraDataChannelRead_Half(GetBuffer_Half(), GetHalfStride(), Register_Half, ReadIndex, Value); }
+	//TODO: Half Support | void Read_Half(out half3 Value){ NiagaraDataChannelRead_Half(GetBuffer_Half(), GetHalfStride(), Register_Half, ReadIndex, Value); }
+	//TODO: Half Support | void Read_Half(out half4 Value){ NiagaraDataChannelRead_Half(GetBuffer_Half(), GetHalfStride(), Register_Half, ReadIndex, Value); }
 };
 
 void Num_{ParameterName}(out int OutNum)
diff --git a/Engine/Plugins/FX/Niagara/Shaders/Private/DataChannel/NiagaraDataInterfaceDataChannelTemplateWriteCommon.ush b/Engine/Plugins/FX/Niagara/Shaders/Private/DataChannel/NiagaraDataInterfaceDataChannelTemplateWriteCommon.ush
--- a/Engine/Plugins/FX/Niagara/Shaders/Private/DataChannel/NiagaraDataInterfaceDataChannelTemplateWriteCommon.ush
+++ b/Engine/Plugins/FX/Niagara/Shaders/Private/DataChannel/NiagaraDataInterfaceDataChannelTemplateWriteCommon.ush
@@ -8,7 +8,7 @@ Template hlsl for Data Channel Write DI functions to apply to their own that is
 //We index into this using atomic ops on an instance count in RWCountsBuffer.
 RWBuffer<float> {ParameterName}_GPUBufferFloat;
 RWBuffer<int> {ParameterName}_GPUBufferInt32;
-RWBuffer<float> {ParameterName}_GPUBufferHalf;
+//TODO: Half Support | RWBuffer<half> {ParameterName}_GPUBufferHalf;
 
 //Buffer and offset for our instance count value.
 //We'll atomically update this accordingly if we're write/appending.
@@ -20,7 +20,7 @@ int {ParameterName}_GPUBufferSize;
 //Data bounds back to the CPU for game and CPU sims is written into a separate specific buffer indexed directly by the thread index.
 RWBuffer<float> {ParameterName}_CPUBufferFloat;
 RWBuffer<int> {ParameterName}_CPUBufferInt32;
-RWBuffer<float> {ParameterName}_CPUBufferHalf;
+//TODO: Half Support | RWBuffer<half> {ParameterName}_CPUBufferHalf;
 
 //Buffer and offset for our instance count value.
 //We'll atomically update this accordingly if we're write/appending.
@@ -28,7 +28,7 @@ int {ParameterName}_CPUInstanceCountOffset;
 int {ParameterName}_CPUBufferSize;
 int {ParameterName}_CPUFloatStride;
 int {ParameterName}_CPUInt32Stride;
-int {ParameterName}_CPUHalfStride;
+//TODO: Half Support | int {ParameterName}_CPUHalfStride;
 
 //Helper struct with various info needed for writing that can be passed around more easily in generated code.
 struct FNDCAccessContext_{ParameterName}
@@ -44,34 +44,34 @@ struct FNDCAccessContext_{ParameterName}
 	//Component start offests for the current parameter.
 	int ParameterComponentStart_Float;
 	int ParameterComponentStart_Int32;
-	int ParameterComponentStart_Half;
+	//TODO: Half Support | int ParameterComponentStart_Half;
 
 	//Current register indices for each component buffer
 	int CurrRegister_Float;
 	int CurrRegister_Int32;
-	int CurrRegister_Half;
+	//TODO: Half Support | int CurrRegister_Half;
 
 	bool bReadGPUOrCPU;
 
 	int GetFloatStride() { return bReadGPUOrCPU ? {ParameterName}_FloatStride : {ParameterName}_CPUFloatStride; }
 	int GetInt32Stride() { return bReadGPUOrCPU ? {ParameterName}_Int32Stride : {ParameterName}_CPUInt32Stride; }
-	int GetHalfStride() { return bReadGPUOrCPU ? {ParameterName}_HalfStride : {ParameterName}_CPUHalfStride; }
+	//TODO: Half Support | int GetHalfStride() { return bReadGPUOrCPU ? {ParameterName}_HalfStride : {ParameterName}_CPUHalfStride; }
 	
 	RWBuffer<float> GetGPUBuffer_Float() { return {ParameterName}_GPUBufferFloat; }
 	RWBuffer<int> GetGPUBuffer_Int32() { return {ParameterName}_GPUBufferInt32; }
-	RWBuffer<float> GetGPUBuffer_Half() { return {ParameterName}_GPUBufferHalf; }
+	//TODO: Half Support | RWBuffer<half> GetGPUBuffer_Half() { return {ParameterName}_GPUBufferHalf; }
 	int GetGPUInstanceCountOffset() { return {ParameterName}_GPUInstanceCountOffset; }
 	int GetGPUBufferSize() { return {ParameterName}_GPUBufferSize; }
 
 	RWBuffer<float> GetCPUBuffer_Float() { return {ParameterName}_CPUBufferFloat; }
 	RWBuffer<int> GetCPUBuffer_Int32() { return {ParameterName}_CPUBufferInt32; }
-	RWBuffer<float> GetCPUBuffer_Half() { return {ParameterName}_CPUBufferHalf; }
+	//TODO: Half Support | RWBuffer<half> GetCPUBuffer_Half() { return {ParameterName}_CPUBufferHalf; }
 	int GetCPUInstanceCountOffset() { return {ParameterName}_CPUInstanceCountOffset; }
 	int GetCPUBufferSize() { return {ParameterName}_CPUBufferSize; }
 
 	RWBuffer<float> GetCurrBuffer_Float() { return bReadGPUOrCPU ? GetGPUBuffer_Float() : GetCPUBuffer_Float(); }
 	RWBuffer<int> GetCurrBuffer_Int32() { return bReadGPUOrCPU ? GetGPUBuffer_Int32() : GetCPUBuffer_Int32(); }
-	RWBuffer<float> GetCurrBuffer_Half() { return bReadGPUOrCPU ? GetGPUBuffer_Half() : GetCPUBuffer_Half(); }
+	//TODO: Half Support | RWBuffer<half> GetCurrBuffer_Half() { return bReadGPUOrCPU ? GetGPUBuffer_Half() : GetCPUBuffer_Half(); }
 
 	bool IsValid() { return bValid; }
 	bool ShouldWriteGPUData() { return bValid && GetGPUInstanceCountOffset() != -1; }
@@ -113,7 +113,7 @@ struct FNDCAccessContext_{ParameterName}
 	//Initialize for the given parameter index.
 	void InitForParameter(int ParameterIndex)
 	{
-		bValid = bValid && GetParameterLayoutInfo_{ParameterName}(ParameterIndex, ParameterComponentStart_Float, ParameterComponentStart_Int32, ParameterComponentStart_Half);				
+		bValid = bValid && GetParameterLayoutInfo_{ParameterName}(ParameterIndex, ParameterComponentStart_Float, ParameterComponentStart_Int32);//TODO: Half Support | , ParameterComponentStart_Half);				
 	}
 	
 	//Initialize for wrting the current parameter to the GPU buffers
@@ -123,7 +123,7 @@ struct FNDCAccessContext_{ParameterName}
 		WriteIndex = WriteIndexGPU;	
 		CurrRegister_Float = ParameterComponentStart_Float;
 		CurrRegister_Int32 = ParameterComponentStart_Int32;
-		CurrRegister_Half = ParameterComponentStart_Half;
+		//TODO: Half Support | CurrRegister_Half = ParameterComponentStart_Half;
 		
 		return ShouldWriteGPUData();
 	}
@@ -135,7 +135,7 @@ struct FNDCAccessContext_{ParameterName}
 		WriteIndex = WriteIndexCPU;	
 		CurrRegister_Float = ParameterComponentStart_Float;
 		CurrRegister_Int32 = ParameterComponentStart_Int32;
-		CurrRegister_Half = ParameterComponentStart_Half;
+		//TODO: Half Support | CurrRegister_Half = ParameterComponentStart_Half;
 		
 		return ShouldWriteCPUData();
 	}
@@ -150,8 +150,8 @@ struct FNDCAccessContext_{ParameterName}
 	void Write_Int32(int3 Value){ NiagaraDataChannelWrite_Int32(GetCurrBuffer_Int32(), GetInt32Stride(), CurrRegister_Int32, WriteIndex, Value);	}
 	void Write_Int32(int4 Value){ NiagaraDataChannelWrite_Int32(GetCurrBuffer_Int32(), GetInt32Stride(), CurrRegister_Int32, WriteIndex, Value);	}
 
-	void Write_Half(half Value){ NiagaraDataChannelWrite_Half(GetCurrBuffer_Half(), GetHalfStride(), CurrRegister_Half, WriteIndex, Value); }
-	void Write_Half(half2 Value){ NiagaraDataChannelWrite_Half(GetCurrBuffer_Half(), GetHalfStride(), CurrRegister_Half, WriteIndex, Value); }
-	void Write_Half(half3 Value){ NiagaraDataChannelWrite_Half(GetCurrBuffer_Half(), GetHalfStride(), CurrRegister_Half, WriteIndex, Value); }
-	void Write_Half(half4 Value){ NiagaraDataChannelWrite_Half(GetCurrBuffer_Half(), GetHalfStride(), CurrRegister_Half, WriteIndex, Value); }
+	//TODO: Half Support | void Write_Half(half Value){ NiagaraDataChannelWrite_Half(GetCurrBuffer_Half(), GetHalfStride(), CurrRegister_Half, WriteIndex, Value); }
+	//TODO: Half Support | void Write_Half(half2 Value){ NiagaraDataChannelWrite_Half(GetCurrBuffer_Half(), GetHalfStride(), CurrRegister_Half, WriteIndex, Value); }
+	//TODO: Half Support | void Write_Half(half3 Value){ NiagaraDataChannelWrite_Half(GetCurrBuffer_Half(), GetHalfStride(), CurrRegister_Half, WriteIndex, Value); }
+	//TODO: Half Support | void Write_Half(half4 Value){ NiagaraDataChannelWrite_Half(GetCurrBuffer_Half(), GetHalfStride(), CurrRegister_Half, WriteIndex, Value); }
 };
\ No newline at end of file
diff --git a/Engine/Plugins/FX/Niagara/Source/Niagara/Internal/DataInterface/NiagaraDataInterfaceDataChannelCommon.h b/Engine/Plugins/FX/Niagara/Source/Niagara/Internal/DataInterface/NiagaraDataInterfaceDataChannelCommon.h
--- a/Engine/Plugins/FX/Niagara/Source/Niagara/Internal/DataInterface/NiagaraDataInterfaceDataChannelCommon.h
+++ b/Engine/Plugins/FX/Niagara/Source/Niagara/Internal/DataInterface/NiagaraDataInterfaceDataChannelCommon.h
@@ -443,13 +443,13 @@ struct FVariadicParameterGPUScriptInfo
 				{
 					GPUScriptParameterOffsetTable.Add(LayoutInfo->GetNumFloatComponents() > 0 ? LayoutInfo->GetFloatComponentStart() : INDEX_NONE);
 					GPUScriptParameterOffsetTable.Add(LayoutInfo->GetNumInt32Components() > 0 ? LayoutInfo->GetInt32ComponentStart() : INDEX_NONE);
-					GPUScriptParameterOffsetTable.Add(LayoutInfo->GetNumHalfComponents() > 0 ? LayoutInfo->GetHalfComponentStart() : INDEX_NONE);
+					//TODO: Half Support | GPUScriptParameterOffsetTable.Add(LayoutInfo->GetNumHalfComponents() > 0 ? LayoutInfo->GetHalfComponentStart() : INDEX_NONE);
 				}
 				else
 				{
 					GPUScriptParameterOffsetTable.Add(INDEX_NONE);
 					GPUScriptParameterOffsetTable.Add(INDEX_NONE);
-					GPUScriptParameterOffsetTable.Add(INDEX_NONE);
+					//TODO: Half Support | GPUScriptParameterOffsetTable.Add(INDEX_NONE);
 				}
 			}
 		}
diff --git a/Engine/Plugins/FX/Niagara/Source/Niagara/Private/DataInterface/NiagaraDataInterfaceDataChannelRead.cpp b/Engine/Plugins/FX/Niagara/Source/Niagara/Private/DataInterface/NiagaraDataInterfaceDataChannelRead.cpp
--- a/Engine/Plugins/FX/Niagara/Source/Niagara/Private/DataInterface/NiagaraDataInterfaceDataChannelRead.cpp
+++ b/Engine/Plugins/FX/Niagara/Source/Niagara/Private/DataInterface/NiagaraDataInterfaceDataChannelRead.cpp
@@ -277,11 +277,11 @@ namespace NDIDataChannelReadLocal
 		SHADER_PARAMETER(int32, ParameterOffsetTableIndex)
 		SHADER_PARAMETER(int32, FloatStride)
 		SHADER_PARAMETER(int32, Int32Stride)
-		SHADER_PARAMETER(int32, HalfStride)
+		//TODO: Half Support | SHADER_PARAMETER(int32, HalfStride)
 
 		SHADER_PARAMETER_SRV(Buffer<float>, DataFloat)
 		SHADER_PARAMETER_SRV(Buffer<int>, DataInt32)
-		SHADER_PARAMETER_SRV(Buffer<float>, DataHalf)
+		//TODO: Half Support | SHADER_PARAMETER_SRV(Buffer<float>, DataHalf)
 		SHADER_PARAMETER(int32, InstanceCountOffset)
 		SHADER_PARAMETER(int32, ConsumeInstanceCountOffset)
 		SHADER_PARAMETER(int32, BufferSize)
@@ -2013,7 +2013,7 @@ void UNiagaraDataInterfaceDataChannelRead::SetShaderParameters(const FNiagaraDat
 					Transitions.Reserve(3);
 					Transitions.Emplace(Data->GetGPUBufferFloat().UAV, ERHIAccess::Unknown, ERHIAccess::SRVCompute);
 					Transitions.Emplace(Data->GetGPUBufferInt().UAV, ERHIAccess::Unknown, ERHIAccess::SRVCompute);
-					Transitions.Emplace(Data->GetGPUBufferHalf().UAV, ERHIAccess::Unknown, ERHIAccess::SRVCompute);
+					//TODO: Half Support | Transitions.Emplace(Data->GetGPUBufferHalf().UAV, ERHIAccess::Unknown, ERHIAccess::SRVCompute);
 					Context.GetGraphBuilder().RHICmdList.Transition(Transitions);
 
 					InstParameters->ParamOffsetTable = ParameterLayoutBuffer.SRV.IsValid() ? ParameterLayoutBuffer.SRV.GetReference() : FNiagaraRenderer::GetDummyUIntBuffer();
@@ -2021,11 +2021,11 @@ void UNiagaraDataInterfaceDataChannelRead::SetShaderParameters(const FNiagaraDat
 
 					InstParameters->FloatStride = Data->GetFloatStride() / sizeof(float);
 					InstParameters->Int32Stride = Data->GetInt32Stride() / sizeof(int32);
-					InstParameters->HalfStride = Data->GetHalfStride() / sizeof(FFloat16);
+					//TODO: Half Support | InstParameters->HalfStride = Data->GetHalfStride() / sizeof(FFloat16);
 
 					InstParameters->DataFloat = Data->GetGPUBufferFloat().SRV.IsValid() ? Data->GetGPUBufferFloat().SRV.GetReference() : FNiagaraRenderer::GetDummyFloatBuffer();
 					InstParameters->DataInt32 = Data->GetGPUBufferInt().SRV.IsValid() ? Data->GetGPUBufferInt().SRV.GetReference() : FNiagaraRenderer::GetDummyIntBuffer();
-					InstParameters->DataHalf = Data->GetGPUBufferHalf().SRV.IsValid() ? Data->GetGPUBufferHalf().SRV.GetReference() : FNiagaraRenderer::GetDummyHalfBuffer();
+					//TODO: Half Support | InstParameters->DataHalf = Data->GetGPUBufferHalf().SRV.IsValid() ? Data->GetGPUBufferHalf().SRV.GetReference() : FNiagaraRenderer::GetDummyHalfBuffer();
 					InstParameters->InstanceCountOffset = Data->GetGPUInstanceCountBufferOffset();
 					InstParameters->ConsumeInstanceCountOffset = InstanceData->ConsumeInstanceCountOffset;
 					InstParameters->BufferSize = Data ? Data->GetNumInstancesAllocated() : INDEX_NONE;
@@ -2046,11 +2046,11 @@ void UNiagaraDataInterfaceDataChannelRead::SetShaderParameters(const FNiagaraDat
 
 		InstParameters->FloatStride = 0;
 		InstParameters->Int32Stride = 0;
-		InstParameters->HalfStride = 0;
+		//TODO: Half Support | InstParameters->HalfStride = 0;
 
 		InstParameters->DataFloat = FNiagaraRenderer::GetDummyFloatBuffer();
 		InstParameters->DataInt32 = FNiagaraRenderer::GetDummyIntBuffer();
-		InstParameters->DataHalf = FNiagaraRenderer::GetDummyHalfBuffer();
+		//TODO: Half Support | InstParameters->DataHalf = FNiagaraRenderer::GetDummyHalfBuffer();
 
 		InstParameters->InstanceCountOffset = INDEX_NONE;
 		InstParameters->ConsumeInstanceCountOffset = INDEX_NONE;
diff --git a/Engine/Plugins/FX/Niagara/Source/Niagara/Private/DataInterface/NiagaraDataInterfaceDataChannelWrite.cpp b/Engine/Plugins/FX/Niagara/Source/Niagara/Private/DataInterface/NiagaraDataInterfaceDataChannelWrite.cpp
--- a/Engine/Plugins/FX/Niagara/Source/Niagara/Private/DataInterface/NiagaraDataInterfaceDataChannelWrite.cpp
+++ b/Engine/Plugins/FX/Niagara/Source/Niagara/Private/DataInterface/NiagaraDataInterfaceDataChannelWrite.cpp
@@ -167,22 +167,22 @@ namespace NDIDataChannelWriteLocal
 		SHADER_PARAMETER(int32, ParameterOffsetTableIndex)
 		SHADER_PARAMETER(int32, FloatStride)
 		SHADER_PARAMETER(int32, Int32Stride)
-		SHADER_PARAMETER(int32, HalfStride)
+		//TODO: Half Support | SHADER_PARAMETER(int32, HalfStride)
 
 		SHADER_PARAMETER_UAV(RWBuffer<float>, GPUBufferFloat)
 		SHADER_PARAMETER_UAV(RWBuffer<int>, GPUBufferInt32)
-		SHADER_PARAMETER_UAV(RWBuffer<float>, GPUBufferHalf)
+		//TODO: Half Support | SHADER_PARAMETER_UAV(RWBuffer<float>, GPUBufferHalf)
 		SHADER_PARAMETER(int32, GPUInstanceCountOffset)
 		SHADER_PARAMETER(int32, GPUBufferSize)
 
 		SHADER_PARAMETER_UAV(RWBuffer<float>, CPUBufferFloat)
 		SHADER_PARAMETER_UAV(RWBuffer<int>, CPUBufferInt32)
-		SHADER_PARAMETER_UAV(RWBuffer<float>, CPUBufferHalf)
+		//TODO: Half Support | SHADER_PARAMETER_UAV(RWBuffer<float>, CPUBufferHalf)
 		SHADER_PARAMETER(int32, CPUInstanceCountOffset)
 		SHADER_PARAMETER(int32, CPUBufferSize)
 		SHADER_PARAMETER(int32, CPUFloatStride)
 		SHADER_PARAMETER(int32, CPUInt32Stride)
-		SHADER_PARAMETER(int32, CPUHalfStride)
+		//TODO: Half Support | SHADER_PARAMETER(int32, CPUHalfStride)
 	END_SHADER_PARAMETER_STRUCT()
 }
 
@@ -1333,13 +1333,13 @@ void UNiagaraDataInterfaceDataChannelWrite::SetShaderParameters(const FNiagaraDa
 				{
 					Transitions.Emplace(GPUBuffer->GetGPUBufferFloat().UAV, ERHIAccess::Unknown, ERHIAccess::UAVCompute);
 					Transitions.Emplace(GPUBuffer->GetGPUBufferInt().UAV, ERHIAccess::Unknown, ERHIAccess::UAVCompute);
-					Transitions.Emplace(GPUBuffer->GetGPUBufferHalf().UAV, ERHIAccess::Unknown, ERHIAccess::UAVCompute);
+					//TODO: Half Support | Transitions.Emplace(GPUBuffer->GetGPUBufferHalf().UAV, ERHIAccess::Unknown, ERHIAccess::UAVCompute);
 				}
 				if(BufferForCPU)
 				{
 					Transitions.Emplace(BufferForCPU->GetGPUBufferFloat().UAV, ERHIAccess::Unknown, ERHIAccess::UAVCompute);
 					Transitions.Emplace(BufferForCPU->GetGPUBufferInt().UAV, ERHIAccess::Unknown, ERHIAccess::UAVCompute);
-					Transitions.Emplace(BufferForCPU->GetGPUBufferHalf().UAV, ERHIAccess::Unknown, ERHIAccess::UAVCompute);
+					//TODO: Half Support | Transitions.Emplace(BufferForCPU->GetGPUBufferHalf().UAV, ERHIAccess::Unknown, ERHIAccess::UAVCompute);
 				}
 				Context.GetGraphBuilder().RHICmdList.Transition(Transitions);
 
@@ -1351,24 +1351,24 @@ void UNiagaraDataInterfaceDataChannelWrite::SetShaderParameters(const FNiagaraDa
 					InstParameters->ParameterOffsetTableIndex = ParameterOffsetTableIndex;
 					InstParameters->FloatStride = (GPUBuffer ? GPUBuffer->GetFloatStride() : sizeof(float)) / sizeof(float);
 					InstParameters->Int32Stride = (GPUBuffer ? GPUBuffer->GetInt32Stride() : sizeof(int32)) / sizeof(int32);
-					InstParameters->HalfStride = (GPUBuffer ? GPUBuffer->GetHalfStride() : sizeof(FFloat16)) / sizeof(FFloat16);
+					//TODO: Half Support | InstParameters->HalfStride = (GPUBuffer ? GPUBuffer->GetHalfStride() : sizeof(FFloat16)) / sizeof(FFloat16);
 					
 					InstParameters->GPUBufferFloat = GPUBuffer && GPUBuffer->GetGPUBufferFloat().UAV.IsValid() ? GPUBuffer->GetGPUBufferFloat().UAV : NDIDataChannelUtilities::GetDummyUAVFloat().Buffer.UAV;
 					InstParameters->GPUBufferInt32 = GPUBuffer && GPUBuffer->GetGPUBufferInt().UAV.IsValid() ? GPUBuffer->GetGPUBufferInt().UAV : NDIDataChannelUtilities::GetDummyUAVInt32().Buffer.UAV;
-					InstParameters->GPUBufferHalf = GPUBuffer && GPUBuffer->GetGPUBufferHalf().UAV.IsValid() ? GPUBuffer->GetGPUBufferHalf().UAV : NDIDataChannelUtilities::GetDummyUAVHalf().Buffer.UAV;
+					//TODO: Half Support | InstParameters->GPUBufferHalf = GPUBuffer && GPUBuffer->GetGPUBufferHalf().UAV.IsValid() ? GPUBuffer->GetGPUBufferHalf().UAV : NDIDataChannelUtilities::GetDummyUAVHalf().Buffer.UAV;
 					InstParameters->GPUInstanceCountOffset = GPUBuffer ? GPUBuffer->GetGPUInstanceCountBufferOffset() : INDEX_NONE;
 					InstParameters->GPUBufferSize = GPUBuffer ? GPUBuffer->GetNumInstancesAllocated() : INDEX_NONE;
 
 
 					InstParameters->CPUBufferFloat = BufferForCPU && BufferForCPU->GetGPUBufferFloat().UAV.IsValid() ? BufferForCPU->GetGPUBufferFloat().UAV : NDIDataChannelUtilities::GetDummyUAVFloat().Buffer.UAV;
 					InstParameters->CPUBufferInt32 = BufferForCPU && BufferForCPU->GetGPUBufferInt().UAV.IsValid() ? BufferForCPU->GetGPUBufferInt().UAV : NDIDataChannelUtilities::GetDummyUAVInt32().Buffer.UAV;
-					InstParameters->CPUBufferHalf = BufferForCPU && BufferForCPU->GetGPUBufferHalf().UAV.IsValid() ? BufferForCPU->GetGPUBufferHalf().UAV : NDIDataChannelUtilities::GetDummyUAVHalf().Buffer.UAV;
+					//TODO: Half Support | InstParameters->CPUBufferHalf = BufferForCPU && BufferForCPU->GetGPUBufferHalf().UAV.IsValid() ? BufferForCPU->GetGPUBufferHalf().UAV : NDIDataChannelUtilities::GetDummyUAVHalf().Buffer.UAV;
 					InstParameters->CPUInstanceCountOffset = BufferForCPU ? BufferForCPU->GetGPUInstanceCountBufferOffset() : INDEX_NONE;
 					InstParameters->CPUBufferSize = BufferForCPU ? BufferForCPU->GetNumInstancesAllocated() : INDEX_NONE;
 
 					InstParameters->CPUFloatStride = (BufferForCPU ? BufferForCPU->GetFloatStride() : sizeof(float)) / sizeof(float);
 					InstParameters->CPUInt32Stride = (BufferForCPU ? BufferForCPU->GetInt32Stride() : sizeof(int32)) / sizeof(int32);
-					InstParameters->CPUHalfStride = (BufferForCPU ? BufferForCPU->GetHalfStride() : sizeof(FFloat16)) / sizeof(FFloat16);
+					//TODO: Half Support | InstParameters->CPUHalfStride = (BufferForCPU ? BufferForCPU->GetHalfStride() : sizeof(FFloat16)) / sizeof(FFloat16);
 
 					bSuccess = true;
 				}
@@ -1382,18 +1382,18 @@ void UNiagaraDataInterfaceDataChannelWrite::SetShaderParameters(const FNiagaraDa
 		InstParameters->ParameterOffsetTableIndex = INDEX_NONE;
 		InstParameters->FloatStride = 0;
 		InstParameters->Int32Stride = 0;
-		InstParameters->HalfStride = 0;
+		//TODO: Half Support | InstParameters->HalfStride = 0;
 
 		InstParameters->GPUBufferFloat = NDIDataChannelUtilities::GetDummyUAVFloat().Buffer.UAV;
 		InstParameters->GPUBufferInt32 = NDIDataChannelUtilities::GetDummyUAVInt32().Buffer.UAV;
-		InstParameters->GPUBufferHalf = NDIDataChannelUtilities::GetDummyUAVHalf().Buffer.UAV;
+		//TODO: Half Support | InstParameters->GPUBufferHalf = NDIDataChannelUtilities::GetDummyUAVHalf().Buffer.UAV;
 		InstParameters->GPUInstanceCountOffset = INDEX_NONE;
 		InstParameters->GPUBufferSize = INDEX_NONE;
 
 
 		InstParameters->CPUBufferFloat = NDIDataChannelUtilities::GetDummyUAVFloat().Buffer.UAV;
 		InstParameters->CPUBufferInt32 = NDIDataChannelUtilities::GetDummyUAVInt32().Buffer.UAV;
-		InstParameters->CPUBufferHalf = NDIDataChannelUtilities::GetDummyUAVHalf().Buffer.UAV;
+		//TODO: Half Support | InstParameters->CPUBufferHalf = NDIDataChannelUtilities::GetDummyUAVHalf().Buffer.UAV;
 		InstParameters->CPUInstanceCountOffset = INDEX_NONE;
 		InstParameters->CPUBufferSize = INDEX_NONE;
 	}
