Fix duplicate LOCTEXT macro key

#rnx

diff --git a/Engine/Plugins/FX/Niagara/Source/NiagaraEditor/Private/Widgets/DataChannel/NiagaraDataChannelWizard.cpp b/Engine/Plugins/FX/Niagara/Source/NiagaraEditor/Private/Widgets/DataChannel/NiagaraDataChannelWizard.cpp
--- a/Engine/Plugins/FX/Niagara/Source/NiagaraEditor/Private/Widgets/DataChannel/NiagaraDataChannelWizard.cpp
+++ b/Engine/Plugins/FX/Niagara/Source/NiagaraEditor/Private/Widgets/DataChannel/NiagaraDataChannelWizard.cpp
@@ -838,7 +838,7 @@ TSharedRef<FModuleWizardModel> DataChannel::CreateSpawnNDCModuleWizardModel()
 
 		virtual FText GetHeaderLabel() override
 		{
-			return LOCTEXT("VariablesPageLabel", "Please select which variables should be read into particle attributes when spawning.");
+			return LOCTEXT("VariablesPageSpawnLabel", "Please select which variables should be read into particle attributes when spawning.");
 		}
 		
 		virtual FText GetFormattedModuleName(const FText& AssetName) const override
