- Fix for ensure incorrectly triggering during PostEditChangeProperty
- For LWC types we leverage this function to copy data into the SWC format, in this case we can have > 1 instance in the per instance data.

#rnx

diff --git a/Engine/Plugins/FX/Niagara/Source/Niagara/Classes/NiagaraDataInterfaceArrayImpl.h b/Engine/Plugins/FX/Niagara/Source/Niagara/Classes/NiagaraDataInterfaceArrayImpl.h
--- a/Engine/Plugins/FX/Niagara/Source/Niagara/Classes/NiagaraDataInterfaceArrayImpl.h
+++ b/Engine/Plugins/FX/Niagara/Source/Niagara/Classes/NiagaraDataInterfaceArrayImpl.h
@@ -566,12 +566,19 @@ struct FNDIArrayProxyImpl : public INDIArrayProxyBase
 		}
 		else
 		{
-			ensure(PerInstanceData_GameThread.Num() == 1);
-			FNDIArrayInstanceData_GameThread<TArrayType>* InstanceData = PerInstanceData_GameThread.CreateConstIterator().Value();
-			FWriteScopeLock	ScopeLock(InstanceData->ArrayRWGuard);
-			InstanceData->bIsModified = false;
-			InstanceData->bIsRenderDirty |= bShouldSyncToGpu;
-			InstanceData->ArrayData.Empty();
+			// Note: Can be called from PostEditChangeProperty where we can expect to the have > 1 data interface.  For BP user parameters we can expect this to be 1 exactly.
+		#if WITH_EDITORONLY_DATA
+			ensure(GIsTransacting || PerInstanceData_GameThread.Num() == 1);
+		#endif
+
+			for (auto It=PerInstanceData_GameThread.CreateConstIterator(); ++It; It)
+			{
+				FNDIArrayInstanceData_GameThread<TArrayType>* InstanceData = It.Value();
+				FWriteScopeLock	ScopeLock(InstanceData->ArrayRWGuard);
+				InstanceData->bIsModified = false;
+				InstanceData->bIsRenderDirty |= bShouldSyncToGpu;
+				InstanceData->ArrayData.Empty();
+			}
 			Owner->GetArrayReference().SetNum(InArrayData.Num());
 			FNDIArrayImplHelper<TArrayType>::CopyCpuToCpuMemory(Owner->GetArrayReference().GetData(), InArrayData.GetData(), InArrayData.Num());
 		}
