- Fix for ++It vs It

#rnx

diff --git a/Engine/Plugins/FX/Niagara/Source/Niagara/Classes/NiagaraDataInterfaceArrayImpl.h b/Engine/Plugins/FX/Niagara/Source/Niagara/Classes/NiagaraDataInterfaceArrayImpl.h
--- a/Engine/Plugins/FX/Niagara/Source/Niagara/Classes/NiagaraDataInterfaceArrayImpl.h
+++ b/Engine/Plugins/FX/Niagara/Source/Niagara/Classes/NiagaraDataInterfaceArrayImpl.h
@@ -571,7 +571,7 @@ struct FNDIArrayProxyImpl : public INDIArrayProxyBase
 			ensure(GIsTransacting || PerInstanceData_GameThread.Num() == 1);
 		#endif
 
-			for (auto It=PerInstanceData_GameThread.CreateConstIterator(); ++It; It)
+			for (auto It=PerInstanceData_GameThread.CreateConstIterator(); It; ++It)
 			{
 				FNDIArrayInstanceData_GameThread<TArrayType>* InstanceData = It.Value();
 				FWriteScopeLock	ScopeLock(InstanceData->ArrayRWGuard);
