- Removed 'Interface' from socket reader DI display name to follow the pattern of other DIs

#rnx

diff --git a/Engine/Plugins/FX/Niagara/Source/Niagara/Internal/DataInterface/NiagaraDataInterfaceSocketReader.h b/Engine/Plugins/FX/Niagara/Source/Niagara/Internal/DataInterface/NiagaraDataInterfaceSocketReader.h
--- a/Engine/Plugins/FX/Niagara/Source/Niagara/Internal/DataInterface/NiagaraDataInterfaceSocketReader.h
+++ b/Engine/Plugins/FX/Niagara/Source/Niagara/Internal/DataInterface/NiagaraDataInterfaceSocketReader.h
@@ -28,7 +28,7 @@ enum class ENDISocketReaderSourceMode : uint8
 Data interface for reading sockets from various sources.
 This can be from a live component in the scene or from a static / skeletal mesh asset.
 */
-UCLASS(EditInlineNew, Category = "Actor", CollapseCategories, meta = (DisplayName = "Socket Reader Interface"), MinimalAPI)
+UCLASS(EditInlineNew, Category = "Actor", CollapseCategories, meta = (DisplayName = "Socket Reader"), MinimalAPI)
 class UNiagaraDataInterfaceSocketReader : public UNiagaraDataInterface
 {
 	GENERATED_UCLASS_BODY()
