Fix for MeshRenderer crash when mesh is invalid and doesn't have a valid LODResource
#tests none

diff --git a/Engine/Plugins/FX/Niagara/Source/Niagara/Private/NiagaraMeshRendererProperties.cpp b/Engine/Plugins/FX/Niagara/Source/Niagara/Private/NiagaraMeshRendererProperties.cpp
--- a/Engine/Plugins/FX/Niagara/Source/Niagara/Private/NiagaraMeshRendererProperties.cpp
+++ b/Engine/Plugins/FX/Niagara/Source/Niagara/Private/NiagaraMeshRendererProperties.cpp
@@ -1356,9 +1356,12 @@ void UNiagaraMeshRendererProperties::CheckMaterialUsage()
 {
 	for (const auto& MeshProperties : Meshes)
 	{
-		if (MeshProperties.Mesh && MeshProperties.Mesh->GetRenderData())
+		const FStaticMeshRenderData* MeshRenderData = MeshProperties.Mesh ? MeshProperties.Mesh->GetRenderData() : nullptr;
+
+		constexpr int32 LodIndexForCheck = 0;
+		if (MeshRenderData && MeshRenderData->LODResources.IsValidIndex(LodIndexForCheck))
 		{
-			const FStaticMeshLODResources& LODModel = MeshProperties.Mesh->GetRenderData()->LODResources[0];
+			const FStaticMeshLODResources& LODModel = MeshRenderData->LODResources[LodIndexForCheck];
 			for (int32 SectionIndex = 0; SectionIndex < LODModel.Sections.Num(); SectionIndex++)
 			{
 				const FStaticMeshSection& Section = LODModel.Sections[SectionIndex];
