Fix a few issues around NDC proxy lifetimes/cleanup.

Adding path to recreate NDC proxies when the world's scene changes.
Free instance count entries on NDC proxy cleanup.
Ensuring the worlds cleanup their FXSystem when recreating the scene.
Ensuring NDC handlers are cleaned up properly and their lifetime extends safely to allow NDC proxies to free correctly on the RT first.

diff --git a/Engine/Plugins/FX/Niagara/Source/Niagara/Internal/NiagaraDataChannelCommon.h b/Engine/Plugins/FX/Niagara/Source/Niagara/Internal/NiagaraDataChannelCommon.h
--- a/Engine/Plugins/FX/Niagara/Source/Niagara/Internal/NiagaraDataChannelCommon.h
+++ b/Engine/Plugins/FX/Niagara/Source/Niagara/Internal/NiagaraDataChannelCommon.h
@@ -103,8 +103,13 @@ struct FNiagaraDataChannelData final : public TSharedFromThis<FNiagaraDataChanne
 
 	//This will get a buffer from the CPU dataset intended to be written to on the CPU.
 	FNiagaraDataBuffer* GetBufferForCPUWrite();
+
+	void DestroyRenderThreadProxy(FNiagaraGpuComputeDispatchInterface* ComputeDispatchInterface);
+
 private:
 
+	void CreateRenderThreadProxy(UNiagaraDataChannelHandler* Owner);
+
 	/** DataChannel data accessible from Game/BP. AoS Layout. LWC types. */
 	FNiagaraDataChannelGameDataPtr GameData;
 
diff --git a/Engine/Plugins/FX/Niagara/Source/Niagara/Private/NiagaraDataChannel.cpp b/Engine/Plugins/FX/Niagara/Source/Niagara/Private/NiagaraDataChannel.cpp
--- a/Engine/Plugins/FX/Niagara/Source/Niagara/Private/NiagaraDataChannel.cpp
+++ b/Engine/Plugins/FX/Niagara/Source/Niagara/Private/NiagaraDataChannel.cpp
@@ -463,6 +463,53 @@ FNiagaraDataChannelDataProxy::~FNiagaraDataChannelDataProxy()
 #endif
 }
 
+void FNiagaraDataChannelDataProxy::Cleanup(FNiagaraGpuComputeDispatchInterface* ComputeDispatchInterface)
+{
+#if !UE_BUILD_SHIPPING
+	ensureMsgf(DispatchInterfaceForDebuggingOnly == ComputeDispatchInterface || ComputeDispatchInterface == nullptr, TEXT("FNiagaraDataChannelDataProxy - ComputeDispatchInterface has been modified without proper cleanup. This could cause a crash."));
+#endif
+	
+	if(ComputeDispatchInterface)
+	{
+		//Remove ourselves from the dispatcher and free our instance counts from the instance count manager.
+		ComputeDispatchInterface->RemoveNDCDataProxy(this);
+	
+		FNiagaraGPUInstanceCountManager& InstCountManager = ComputeDispatchInterface->GetGPUInstanceCounterManager();
+		
+		if(CurrFrameData)
+		{
+			uint32 InstanceCountOffset = CurrFrameData->GetGPUInstanceCountBufferOffset();
+			InstCountManager.FreeEntry(InstanceCountOffset);
+			CurrFrameData->SetGPUInstanceCountBufferOffset(INDEX_NONE);
+		}
+
+		if (PrevFrameData)
+		{
+			uint32 InstanceCountOffset = PrevFrameData->GetGPUInstanceCountBufferOffset();
+			InstCountManager.FreeEntry(InstanceCountOffset);
+			PrevFrameData->SetGPUInstanceCountBufferOffset(INDEX_NONE);
+		}
+	}
+	else
+	{
+		if (CurrFrameData)
+		{
+			CurrFrameData->SetGPUInstanceCountBufferOffset(INDEX_NONE);
+		}
+
+		if (PrevFrameData)
+		{
+			PrevFrameData->SetGPUInstanceCountBufferOffset(INDEX_NONE);
+		}
+	}
+
+#if !UE_BUILD_SHIPPING
+	DispatchInterfaceForDebuggingOnly = nullptr;
+#endif
+	CurrFrameData = nullptr;
+	PrevFrameData = nullptr;
+}
+
 void FNiagaraDataChannelDataProxy::Reset()
 {
 	PrevFrameData = nullptr;
@@ -592,6 +639,11 @@ FNiagaraDataChannelData::~FNiagaraDataChannelData()
 	ENQUEUE_RENDER_COMMAND(FDeleteContextCommand)(
 		[ComputeDispatchInterface, CPUDataChannelDataSet = CPUSimData, GPUDataChannelDataSet = GPUSimData, ReleasedRTProxy = RTProxy.Release()](FRHICommandListImmediate& RHICmdList)
 	{
+		if (ReleasedRTProxy)
+		{
+			ReleasedRTProxy->Cleanup(ComputeDispatchInterface);
+			delete ReleasedRTProxy;
+		}
 		if (CPUDataChannelDataSet != nullptr)
 		{
 			delete CPUDataChannelDataSet;
@@ -600,24 +652,11 @@ FNiagaraDataChannelData::~FNiagaraDataChannelData()
 		{
 			delete GPUDataChannelDataSet;
 		}
-		if (ReleasedRTProxy)
-		{
-			if (ComputeDispatchInterface)
-			{
-			#if !UE_BUILD_SHIPPING
-				check(ComputeDispatchInterface == ReleasedRTProxy->DispatchInterfaceForDebuggingOnly);
-			#endif
-				ComputeDispatchInterface->RemoveNDCDataProxy(ReleasedRTProxy);
-			}
-		#if !UE_BUILD_SHIPPING
-			ReleasedRTProxy->DispatchInterfaceForDebuggingOnly = nullptr;
-		#endif
-			delete ReleasedRTProxy;
-		}
 	}
 	);
 	CPUSimData = nullptr;
 	GPUSimData = nullptr;
+	check(RTProxy == nullptr);
 }
 
 void FNiagaraDataChannelData::Init(UNiagaraDataChannelHandler* Owner)
@@ -643,23 +682,7 @@ void FNiagaraDataChannelData::Init(UNiagaraDataChannelHandler* Owner)
 	UWorld* OwnerWorld = Owner->GetWorld();
 	WeakOwnerWorld = OwnerWorld;
 
-	if (FNiagaraGpuComputeDispatchInterface* ComputeDispatchInterface = FNiagaraGpuComputeDispatchInterface::Get(OwnerWorld))
-	{
-		RTProxy.Reset(new FNiagaraDataChannelDataProxy());
-		RTProxy->GPUDataSet = GPUSimData;
-		RTProxy->bNeedsPrevFrameData = DataChannel->KeepPreviousFrameData();
-		RTProxy->Owner = this->AsWeak();
-#if !UE_BUILD_SHIPPING
-		RTProxy->DebugName = FString::Printf(TEXT("%s__GPUData"), *DataChannel->GetName());
-		RTProxy->DispatchInterfaceForDebuggingOnly = ComputeDispatchInterface;
-#endif
-		ENQUEUE_RENDER_COMMAND(FNiagaraDataChannelDataProxyInit) (
-			[ComputeDispatchInterface, RT_Proxy=RTProxy.Get()](FRHICommandListImmediate& CmdList)
-			{
-				ComputeDispatchInterface->AddNDCDataProxy(RT_Proxy);
-			}
-		);
-	}
+	CreateRenderThreadProxy(Owner);
 }
 
 void FNiagaraDataChannelData::Reset()
@@ -687,6 +710,12 @@ void FNiagaraDataChannelData::Reset()
 
 void FNiagaraDataChannelData::BeginFrame(UNiagaraDataChannelHandler* Owner)
 {
+	//Lazy re-init our proxy data in cases where we've had to recreate our render state while running due to an FScene change etc.
+	if(RTProxy == nullptr)
+	{
+		CreateRenderThreadProxy(Owner);
+	}
+
 	GameData->BeginFrame();
 
 	bool bRequirePreviousData = Owner->GetDataChannel()->KeepPreviousFrameData();
@@ -934,6 +963,47 @@ FNiagaraDataBuffer* FNiagaraDataChannelData::GetBufferForCPUWrite()
 	return nullptr;
 }
 
+void FNiagaraDataChannelData::DestroyRenderThreadProxy(FNiagaraGpuComputeDispatchInterface* ComputeDispatchInterface)
+{
+	ENQUEUE_RENDER_COMMAND(FDestroyRTProxyCommand)(
+		[ComputeDispatchInterface, ReleasedRTProxy = RTProxy.Release()](FRHICommandListImmediate& RHICmdList)
+		{
+			if (ReleasedRTProxy)
+			{
+				ReleasedRTProxy->Cleanup(ComputeDispatchInterface);
+				delete ReleasedRTProxy;
+			}
+		}
+	);
+	check(RTProxy == nullptr);
+}
+
+void FNiagaraDataChannelData::CreateRenderThreadProxy(UNiagaraDataChannelHandler* Owner)
+{
+	check(Owner);
+	UWorld* OwnerWorld = Owner->GetWorld();	
+	const UNiagaraDataChannel* DataChannel = Owner->GetDataChannel();
+	check(DataChannel);
+
+	if (FNiagaraGpuComputeDispatchInterface* ComputeDispatchInterface = FNiagaraGpuComputeDispatchInterface::Get(OwnerWorld))
+	{
+		RTProxy.Reset(new FNiagaraDataChannelDataProxy());
+		RTProxy->GPUDataSet = GPUSimData;
+		RTProxy->bNeedsPrevFrameData = DataChannel->KeepPreviousFrameData();
+		RTProxy->Owner = this->AsWeak();
+#if !UE_BUILD_SHIPPING
+		RTProxy->DebugName = FString::Printf(TEXT("%s__GPUData"), *DataChannel->GetName());
+		RTProxy->DispatchInterfaceForDebuggingOnly = ComputeDispatchInterface;
+#endif
+		ENQUEUE_RENDER_COMMAND(FNiagaraDataChannelDataProxyInit) (
+			[ComputeDispatchInterface, RT_Proxy = RTProxy.Get()](FRHICommandListImmediate& CmdList)
+			{
+				ComputeDispatchInterface->AddNDCDataProxy(RT_Proxy);
+			}
+		);
+	}
+}
+
 //////////////////////////////////////////////////////////////////////////
 
 void UNiagaraDataChannel::PostInitProperties()
@@ -1003,6 +1073,13 @@ void UNiagaraDataChannel::BeginDestroy()
 {
 	Super::BeginDestroy();
 	INiagaraModule::RequestRefreshDataChannels();
+
+	RTFence.BeginFence();
+}
+
+bool UNiagaraDataChannel::IsReadyForFinishDestroy()
+{
+	return RTFence.IsFenceComplete() && Super::IsReadyForFinishDestroy();
 }
 
 #if WITH_EDITOR
diff --git a/Engine/Plugins/FX/Niagara/Source/Niagara/Private/NiagaraDataChannelHandler.cpp b/Engine/Plugins/FX/Niagara/Source/Niagara/Private/NiagaraDataChannelHandler.cpp
--- a/Engine/Plugins/FX/Niagara/Source/Niagara/Private/NiagaraDataChannelHandler.cpp
+++ b/Engine/Plugins/FX/Niagara/Source/Niagara/Private/NiagaraDataChannelHandler.cpp
@@ -3,11 +3,20 @@
 #include "NiagaraDataChannelHandler.h"
 #include "NiagaraDataChannelAccessor.h"
 #include "NiagaraDataChannelCommon.h"
+#include "Logging/StructuredLog.h"
+#include "NiagaraGpuComputeDispatchInterface.h"
 
 void UNiagaraDataChannelHandler::BeginDestroy()
 {
 	Super::BeginDestroy();
 	Cleanup();
+	
+	RTFence.BeginFence();
+}
+
+bool UNiagaraDataChannelHandler::IsReadyForFinishDestroy()
+{
+	return RTFence.IsFenceComplete() && Super::IsReadyForFinishDestroy();
 }
 
 void UNiagaraDataChannelHandler::Init(const UNiagaraDataChannel* InChannel)
@@ -33,6 +42,14 @@ void UNiagaraDataChannelHandler::Cleanup()
 void UNiagaraDataChannelHandler::BeginFrame(float DeltaTime, FNiagaraWorldManager* OwningWorld)
 {
 	CurrentTG = TG_PrePhysics;
+
+	for(auto It = WeakDataArray.CreateIterator(); It; ++It)
+	{
+		if(It->IsValid() == false)
+		{
+			It.RemoveCurrentSwap();
+		}
+	}
 }
 
 void UNiagaraDataChannelHandler::EndFrame(float DeltaTime, FNiagaraWorldManager* OwningWorld)
@@ -68,6 +85,18 @@ UNiagaraDataChannelReader* UNiagaraDataChannelHandler::GetDataChannelReader()
 FNiagaraDataChannelDataPtr UNiagaraDataChannelHandler::CreateData()
 {
 	FNiagaraDataChannelDataPtr Ret = MakeShared<FNiagaraDataChannelData>();
+	WeakDataArray.Add(Ret);
 	Ret->Init(this);
 	return Ret;
+}
+
+void UNiagaraDataChannelHandler::OnComputeDispatchInterfaceDestroyed(FNiagaraGpuComputeDispatchInterface* InComputeDispatchInterface)
+{
+	//Destroy all RT proxies when the dispatcher is destroyed.
+	//In cases where this is done on a running world, we'll do a lazy reinit next frame.
+	ForEachNDCData([InComputeDispatchInterface](FNiagaraDataChannelDataPtr& NDCData)
+	{
+		check(NDCData);
+		NDCData->DestroyRenderThreadProxy(InComputeDispatchInterface);
+	});
 }
\ No newline at end of file
diff --git a/Engine/Plugins/FX/Niagara/Source/Niagara/Private/NiagaraDataChannelManager.cpp b/Engine/Plugins/FX/Niagara/Source/Niagara/Private/NiagaraDataChannelManager.cpp
--- a/Engine/Plugins/FX/Niagara/Source/Niagara/Private/NiagaraDataChannelManager.cpp
+++ b/Engine/Plugins/FX/Niagara/Source/Niagara/Private/NiagaraDataChannelManager.cpp
@@ -6,6 +6,7 @@
 #include "NiagaraWorldManager.h"
 #include "NiagaraDataChannel.h"
 #include "NiagaraDataChannelHandler.h"
+#include "NiagaraGpuComputeDispatchInterface.h"
 
 DECLARE_CYCLE_STAT(TEXT("FNiagaraDataChannelManager::BeginFrame"), STAT_DataChannelManager_BeginFrame, STATGROUP_NiagaraDataChannels);
 DECLARE_CYCLE_STAT(TEXT("FNiagaraDataChannelManager::EndFrame"), STAT_DataChannelManager_EndFrame, STATGROUP_NiagaraDataChannels);
@@ -81,6 +82,7 @@ void FNiagaraDataChannelManager::BeginFrame(float DeltaSeconds)
 			}
 			else
 			{
+				It.Value()->Cleanup();
 				It.RemoveCurrent();
 			}
 		}
@@ -119,6 +121,10 @@ void FNiagaraDataChannelManager::Tick(float DeltaSeconds, ETickingGroup TickGrou
 	}
 	else
 	{
+		for (auto& ChannelPair : Channels)
+		{
+			ChannelPair.Value->Cleanup();
+		}
 		Channels.Empty();
 	}
 }
@@ -157,6 +163,10 @@ UNiagaraDataChannelHandler* FNiagaraDataChannelManager::InitDataChannel(const UN
 
 		if (bForce || Handler == nullptr)
 		{
+			if(Handler)
+			{
+				Handler->Cleanup();
+			}
 			Handler = InChannel->CreateHandler(WorldMan->GetWorld());
 		}
 		return Handler;
@@ -174,3 +184,12 @@ UWorld* FNiagaraDataChannelManager::GetWorld()const
 {
 	return WorldMan->GetWorld();
 }
+
+void FNiagaraDataChannelManager::OnComputeDispatchInterfaceDestroyed(FNiagaraGpuComputeDispatchInterface* InComputeDispatchInterface)
+{
+	check(IsInGameThread());
+	for (auto& ChannelPair : Channels)
+	{
+		ChannelPair.Value->OnComputeDispatchInterfaceDestroyed(InComputeDispatchInterface);
+	}
+}
\ No newline at end of file
diff --git a/Engine/Plugins/FX/Niagara/Source/Niagara/Private/NiagaraDataChannelManager.h b/Engine/Plugins/FX/Niagara/Source/Niagara/Private/NiagaraDataChannelManager.h
--- a/Engine/Plugins/FX/Niagara/Source/Niagara/Private/NiagaraDataChannelManager.h
+++ b/Engine/Plugins/FX/Niagara/Source/Niagara/Private/NiagaraDataChannelManager.h
@@ -8,6 +8,7 @@
 class FNiagaraWorldManager;
 class UNiagaraDataChannel;
 class UNiagaraDataChannelHandler;
+class FNiagaraGpuComputeDispatchInterface;
 enum ETickingGroup : int;
 
 /**
@@ -42,6 +43,8 @@ public:
 
 	UWorld* GetWorld()const;
 
+	void OnComputeDispatchInterfaceDestroyed(FNiagaraGpuComputeDispatchInterface* InComputeDispatchInterface);
+
 private:
 	FNiagaraWorldManager* WorldMan = nullptr;
 
diff --git a/Engine/Plugins/FX/Niagara/Source/Niagara/Private/NiagaraWorldManager.cpp b/Engine/Plugins/FX/Niagara/Source/Niagara/Private/NiagaraWorldManager.cpp
--- a/Engine/Plugins/FX/Niagara/Source/Niagara/Private/NiagaraWorldManager.cpp
+++ b/Engine/Plugins/FX/Niagara/Source/Niagara/Private/NiagaraWorldManager.cpp
@@ -727,12 +727,11 @@ void FNiagaraWorldManager::OnSystemPostChange(UNiagaraSystem* System)
 
 void FNiagaraWorldManager::OnComputeDispatchInterfaceDestroyed_Internal(FNiagaraGpuComputeDispatchInterface* InComputeDispatchInterface)
 {
+	GetDataChannelManager().OnComputeDispatchInterfaceDestroyed(InComputeDispatchInterface);
+
 	// Process the deferred deletion queue before deleting the ComputeDispatchInterface of this world.
 	// This is required because the ComputeDispatchInterface is accessed in FNiagaraEmitterInstance::~FNiagaraEmitterInstance
-	if (FNiagaraGpuComputeDispatchInterface::Get(World) == InComputeDispatchInterface)
-	{
-		DeferredDeletionQueue.Empty();
-	}
+	DeferredDeletionQueue.Empty();
 }
 
 void FNiagaraWorldManager::OnWorldBeginTearDown()
@@ -998,7 +997,11 @@ void FNiagaraWorldManager::OnComputeDispatchInterfaceDestroyed(FNiagaraGpuComput
 {
 	for (TPair<UWorld*, FNiagaraWorldManager*>& Pair : WorldManagers)
 	{
-		Pair.Value->OnComputeDispatchInterfaceDestroyed_Internal(InComputeDispatchInterface);
+		UWorld* World = Pair.Key;
+		if (FNiagaraGpuComputeDispatchInterface::Get(World) == InComputeDispatchInterface)
+		{
+			Pair.Value->OnComputeDispatchInterfaceDestroyed_Internal(InComputeDispatchInterface);
+		}
 	}
 }
 
diff --git a/Engine/Plugins/FX/Niagara/Source/Niagara/Public/NiagaraDataChannel.h b/Engine/Plugins/FX/Niagara/Source/Niagara/Public/NiagaraDataChannel.h
--- a/Engine/Plugins/FX/Niagara/Source/Niagara/Public/NiagaraDataChannel.h
+++ b/Engine/Plugins/FX/Niagara/Source/Niagara/Public/NiagaraDataChannel.h
@@ -33,6 +33,7 @@ Reading the previous frames data introduces a frame of latency but removes the n
 #include "NiagaraCommon.h"
 #include "NiagaraDataSet.h"
 #include "UObject/UObjectIterator.h"
+#include "RenderCommandFence.h"
 #include "NiagaraDataChannel.generated.h"
 
 DECLARE_STATS_GROUP(TEXT("Niagara Data Channels"), STATGROUP_NiagaraDataChannels, STATCAT_Niagara);
@@ -77,6 +78,7 @@ struct FNiagaraDataChannelDataProxy
 	const TCHAR* GetDebugName()const{return nullptr;}
 	#endif
 
+	void Cleanup(FNiagaraGpuComputeDispatchInterface* ComputeDispatchInterface);
 	void BeginFrame(FNiagaraGpuComputeDispatchInterface* DispatchInterface, FRHICommandListImmediate& RHICmdList);
 	void EndFrame(FRHICommandListImmediate& RHICmdList);
 	void Reset();
@@ -102,6 +104,7 @@ public:
 	NIAGARA_API virtual void PostInitProperties() override;
 	NIAGARA_API virtual void PostLoad() override;
 	NIAGARA_API virtual void BeginDestroy() override;
+	NIAGARA_API virtual bool IsReadyForFinishDestroy() override;
 #if WITH_EDITOR
 	NIAGARA_API virtual void PreEditChange(FProperty* PropertyAboutToChange) override;
 	NIAGARA_API virtual void PostEditChangeProperty(FPropertyChangedEvent& PropertyChangedDataChannel) override;
@@ -187,6 +190,8 @@ private:
 	#if WITH_NIAGARA_DEBUGGER
 	mutable bool bVerboseLogging = false;
 	#endif
+
+	FRenderCommandFence RTFence;
 };
 
 template<typename TAction>
diff --git a/Engine/Plugins/FX/Niagara/Source/Niagara/Public/NiagaraDataChannelHandler.h b/Engine/Plugins/FX/Niagara/Source/Niagara/Public/NiagaraDataChannelHandler.h
--- a/Engine/Plugins/FX/Niagara/Source/Niagara/Public/NiagaraDataChannelHandler.h
+++ b/Engine/Plugins/FX/Niagara/Source/Niagara/Public/NiagaraDataChannelHandler.h
@@ -12,8 +12,12 @@ Some more complex handlers may want to divide up the scene in various different
 */
 
 #include "NiagaraDataChannelPublic.h"
+#include "RenderCommandFence.h"
 #include "NiagaraDataChannelHandler.generated.h"
 
+class FNiagaraGpuComputeDispatchInterface;
+struct FNiagaraDataChannelData;
+
 UCLASS(abstract, BlueprintType, MinimalAPI)
 class UNiagaraDataChannelHandler : public UObject
 {
@@ -23,6 +27,7 @@ public:
 	
 	//UObject Interface
 	NIAGARA_API virtual void BeginDestroy()override;
+	NIAGARA_API virtual bool IsReadyForFinishDestroy()override;
 	//UObject Interface END
 
 	NIAGARA_API virtual void Init(const UNiagaraDataChannel* InChannel);
@@ -62,6 +67,11 @@ public:
 	/** Returns the tick group we're currently being processed or most recently processed. */
 	ETickingGroup GetCurrentTickGroup() { return CurrentTG; }
 
+	void OnComputeDispatchInterfaceDestroyed(FNiagaraGpuComputeDispatchInterface* InComputeDispatchInterface);
+
+	template<typename TAction>
+	void ForEachNDCData(TAction Func);
+
 protected:
 
 	UPROPERTY()
@@ -76,4 +86,23 @@ protected:
 	TObjectPtr<UNiagaraDataChannelReader> Reader;
 
 	ETickingGroup CurrentTG = ETickingGroup::TG_PrePhysics;
-};
\ No newline at end of file
+
+	//Weak refs to all NDC data created for this handler.
+	//Allows us to perform book keeping and other operations on all data when needed.
+	TArray<TWeakPtr<FNiagaraDataChannelData>> WeakDataArray;
+
+	FRenderCommandFence RTFence;
+};
+
+template<typename TAction>
+void UNiagaraDataChannelHandler::ForEachNDCData(TAction Func)
+{
+	check(IsInGameThread());
+	for(auto& WeakData : WeakDataArray)
+	{
+		if(FNiagaraDataChannelDataPtr DataPtr = WeakData.Pin())
+		{
+			Func(DataPtr);
+		}
+	}
+}
\ No newline at end of file
diff --git a/Engine/Source/Runtime/Engine/Private/World.cpp b/Engine/Source/Runtime/Engine/Private/World.cpp
--- a/Engine/Source/Runtime/Engine/Private/World.cpp
+++ b/Engine/Source/Runtime/Engine/Private/World.cpp
@@ -9166,6 +9166,16 @@ void UWorld::RecreateScene(ERHIFeatureLevel::Type InFeatureLevel, bool bBroadcas
 			Level->ReleaseRenderingResources();
 		}
 
+		//Ensure we've destroyed our FXSystem before we change Scene on the world.
+		bool bCreateFXSystem = false;
+		if (FXSystem)
+		{
+			bCreateFXSystem = true;
+			FFXSystemInterface::Destroy(FXSystem);
+			FXSystem = nullptr;
+			Scene->SetFXSystem(nullptr);
+		}
+
 		Scene->Release();
 		IRendererModule& RendererModule = GetRendererModule();
 		RendererModule.RemoveScene(Scene);
@@ -9175,7 +9185,7 @@ void UWorld::RecreateScene(ERHIFeatureLevel::Type InFeatureLevel, bool bBroadcas
 			FRenderResource::ChangeFeatureLevel(InFeatureLevel);
 		}
 
-		RendererModule.AllocateScene(this, bRequiresHitProxies, FXSystem != nullptr, InFeatureLevel);
+		RendererModule.AllocateScene(this, bRequiresHitProxies, bCreateFXSystem, InFeatureLevel);
 
 		for (ULevel* Level : Levels)
 		{
