- Fix for user parameter variable names accumulating
- For example if you renamed MyFloat to Float you could not rename back to MyFloat

#rnx
#tests Tested renaming to / from the same name, and renaming into a duplicate parameter name

diff --git a/Engine/Plugins/FX/Niagara/Source/NiagaraEditor/Private/Customizations/UserParameters/NiagaraSystemUserParametersBuilder.cpp b/Engine/Plugins/FX/Niagara/Source/NiagaraEditor/Private/Customizations/UserParameters/NiagaraSystemUserParametersBuilder.cpp
--- a/Engine/Plugins/FX/Niagara/Source/NiagaraEditor/Private/Customizations/UserParameters/NiagaraSystemUserParametersBuilder.cpp
+++ b/Engine/Plugins/FX/Niagara/Source/NiagaraEditor/Private/Customizations/UserParameters/NiagaraSystemUserParametersBuilder.cpp
@@ -436,7 +436,7 @@ void FNiagaraSystemUserParameterBuilder::RequestRename(FNiagaraVariable UserPara
 	}));
 }
 
-void FNiagaraSystemUserParameterBuilder::RenameParameter(FNiagaraVariable UserParameter, FName NewName) const
+void FNiagaraSystemUserParameterBuilder::RenameParameter(FNiagaraVariable UserParameter, FName NewName)
 {
 	if(UserParameter.GetName() == NewName)
 	{
@@ -448,6 +448,8 @@ void FNiagaraSystemUserParameterBuilder::RenameParameter(FNiagaraVariable UserPa
 		SystemViewModel.Pin()->RenameParameter(UserParameter, NewName, ENiagaraGetGraphParameterReferencesMode::AllGraphs);
 	}
 
+	UserParamToWidgetMap.Empty();
+
 	Rebuild();
 }
 
diff --git a/Engine/Plugins/FX/Niagara/Source/NiagaraEditor/Private/Customizations/UserParameters/NiagaraSystemUserParametersBuilder.h b/Engine/Plugins/FX/Niagara/Source/NiagaraEditor/Private/Customizations/UserParameters/NiagaraSystemUserParametersBuilder.h
--- a/Engine/Plugins/FX/Niagara/Source/NiagaraEditor/Private/Customizations/UserParameters/NiagaraSystemUserParametersBuilder.h
+++ b/Engine/Plugins/FX/Niagara/Source/NiagaraEditor/Private/Customizations/UserParameters/NiagaraSystemUserParametersBuilder.h
@@ -51,7 +51,7 @@ private:
 	void SelectAllSection();
 	
 	void RequestRename(FNiagaraVariable UserParameter);
-	void RenameParameter(FNiagaraVariable UserParameter, FName NewName) const;
+	void RenameParameter(FNiagaraVariable UserParameter, FName NewName);
 
 	void RequestDuplication(FNiagaraVariable NiagaraVariable);
 
