- Prevent pasting renderers into emitters that don't support that type
- Moved stateless support to be meta data tag on the renderer
- The tag is used to populate the add renderer menu and also when pasting the emitter

#rnx

diff --git a/Engine/Plugins/FX/Niagara/Source/Niagara/Public/NiagaraMeshRendererProperties.h b/Engine/Plugins/FX/Niagara/Source/Niagara/Public/NiagaraMeshRendererProperties.h
--- a/Engine/Plugins/FX/Niagara/Source/Niagara/Public/NiagaraMeshRendererProperties.h
+++ b/Engine/Plugins/FX/Niagara/Source/Niagara/Public/NiagaraMeshRendererProperties.h
@@ -242,7 +242,7 @@ struct FNiagaraMeshRendererMeshProperties
 	NIAGARA_API bool HasValidRenderableMesh() const;
 };
 
-UCLASS(editinlinenew, meta = (DisplayName = "Mesh Renderer"), MinimalAPI)
+UCLASS(editinlinenew, meta = (DisplayName = "Mesh Renderer", SupportsStateless), MinimalAPI)
 class UNiagaraMeshRendererProperties : public UNiagaraRendererProperties
 {
 public:
diff --git a/Engine/Plugins/FX/Niagara/Source/Niagara/Public/NiagaraRibbonRendererProperties.h b/Engine/Plugins/FX/Niagara/Source/Niagara/Public/NiagaraRibbonRendererProperties.h
--- a/Engine/Plugins/FX/Niagara/Source/Niagara/Public/NiagaraRibbonRendererProperties.h
+++ b/Engine/Plugins/FX/Niagara/Source/Niagara/Public/NiagaraRibbonRendererProperties.h
@@ -189,7 +189,7 @@ namespace ENiagaraRibbonVFLayout
 	};
 };
 
-UCLASS(editinlinenew, meta = (DisplayName = "Ribbon Renderer"), MinimalAPI)
+UCLASS(editinlinenew, meta = (DisplayName = "Ribbon Renderer", SupportsStateless), MinimalAPI)
 class UNiagaraRibbonRendererProperties : public UNiagaraRendererProperties
 {
 public:
diff --git a/Engine/Plugins/FX/Niagara/Source/Niagara/Public/NiagaraSpriteRendererProperties.h b/Engine/Plugins/FX/Niagara/Source/Niagara/Public/NiagaraSpriteRendererProperties.h
--- a/Engine/Plugins/FX/Niagara/Source/Niagara/Public/NiagaraSpriteRendererProperties.h
+++ b/Engine/Plugins/FX/Niagara/Source/Niagara/Public/NiagaraSpriteRendererProperties.h
@@ -106,7 +106,7 @@ namespace ENiagaraSpriteVFLayout
 class FAssetThumbnailPool;
 class SWidget;
 
-UCLASS(editinlinenew, meta = (DisplayName = "Sprite Renderer"), MinimalAPI)
+UCLASS(editinlinenew, meta = (DisplayName = "Sprite Renderer", SupportsStateless), MinimalAPI)
 class UNiagaraSpriteRendererProperties : public UNiagaraRendererProperties
 {
 public:
diff --git a/Engine/Plugins/FX/Niagara/Source/NiagaraEditor/Private/NiagaraEditorModule.cpp b/Engine/Plugins/FX/Niagara/Source/NiagaraEditor/Private/NiagaraEditorModule.cpp
--- a/Engine/Plugins/FX/Niagara/Source/NiagaraEditor/Private/NiagaraEditorModule.cpp
+++ b/Engine/Plugins/FX/Niagara/Source/NiagaraEditor/Private/NiagaraEditorModule.cpp
@@ -2164,11 +2164,9 @@ void FNiagaraEditorModule::GetTargetSystemAndEmitterForDataInterface(UNiagaraDat
 
 void FNiagaraEditorModule::RegisterDefaultRendererFactories()
 {
-	bool bIsSupportedByStateless = true;
 	RegisterRendererCreationInfo(FNiagaraRendererCreationInfo(
 		UNiagaraMeshRendererProperties::StaticClass()->GetDisplayNameText(),
 		FText::FromString(UNiagaraMeshRendererProperties::StaticClass()->GetDescription()),
-		bIsSupportedByStateless,
 		UNiagaraMeshRendererProperties::StaticClass()->GetClassPathName(),
 		FNiagaraRendererCreationInfo::FRendererFactory::CreateLambda([](UObject* OuterEmitter)
 		{
@@ -2185,7 +2183,6 @@ void FNiagaraEditorModule::RegisterDefaultRendererFactories()
 	RegisterRendererCreationInfo(FNiagaraRendererCreationInfo(
 		UNiagaraSpriteRendererProperties::StaticClass()->GetDisplayNameText(),
 		FText::FromString(UNiagaraSpriteRendererProperties::StaticClass()->GetDescription()),
-		bIsSupportedByStateless,
 		UNiagaraSpriteRendererProperties::StaticClass()->GetClassPathName(),
 		FNiagaraRendererCreationInfo::FRendererFactory::CreateLambda([](UObject* OuterEmitter)
 		{
@@ -2198,7 +2195,6 @@ void FNiagaraEditorModule::RegisterDefaultRendererFactories()
 	RegisterRendererCreationInfo(FNiagaraRendererCreationInfo(
 		UNiagaraRibbonRendererProperties::StaticClass()->GetDisplayNameText(),
 		FText::FromString(UNiagaraRibbonRendererProperties::StaticClass()->GetDescription()),
-		bIsSupportedByStateless,
 		UNiagaraRibbonRendererProperties::StaticClass()->GetClassPathName(),
 		FNiagaraRendererCreationInfo::FRendererFactory::CreateLambda([](UObject* OuterEmitter)
 		{
diff --git a/Engine/Plugins/FX/Niagara/Source/NiagaraEditor/Private/ViewModels/Stack/NiagaraStackRenderItemGroup.cpp b/Engine/Plugins/FX/Niagara/Source/NiagaraEditor/Private/ViewModels/Stack/NiagaraStackRenderItemGroup.cpp
--- a/Engine/Plugins/FX/Niagara/Source/NiagaraEditor/Private/ViewModels/Stack/NiagaraStackRenderItemGroup.cpp
+++ b/Engine/Plugins/FX/Niagara/Source/NiagaraEditor/Private/ViewModels/Stack/NiagaraStackRenderItemGroup.cpp
@@ -79,7 +79,8 @@ public:
 		for (const FNiagaraRendererCreationInfo& RendererCreationInfo : RendererCreationInfos)
 		{
 			FSoftClassPath SoftClassPath(RendererCreationInfo.RendererClassPath.ToString());
-			if (NiagaraEditorSettings->IsVisibleClass(SoftClassPath.ResolveClass()) && RenderersOwner->IsRenderCreationInfoSupported(RendererCreationInfo))
+			UClass* RendererClass = SoftClassPath.ResolveClass();
+			if (NiagaraEditorSettings->IsVisibleClass(RendererClass) && RenderersOwner->SupportsRendererClass(RendererClass))
 			{
 				OutAddActions.Add(MakeShared<FRenderItemGroupAddAction>(RendererCreationInfo));
 			}
@@ -160,14 +161,21 @@ void UNiagaraStackRenderItemGroup::Paste(const UNiagaraClipboardContent* Clipboa
 	{
 		for (const UNiagaraClipboardRenderer* ClipboardRenderer : ClipboardContent->Renderers)
 		{
-			if (ClipboardRenderer != nullptr && ClipboardRenderer->RendererProperties != nullptr)
+			if (!ClipboardRenderer || !ClipboardRenderer->RendererProperties)
 			{
-				UNiagaraRendererProperties* NewRenderer = ClipboardRenderer->RendererProperties->StaticDuplicateWithNewMergeId(RenderersOwner->GetOwnerObject());
-				RenderersOwner->AddRenderer(NewRenderer);
-				if(ClipboardRenderer->StackNoteData.IsValid())
-				{
-					GetStackEditorData().AddOrReplaceStackNote(FNiagaraStackGraphUtilities::StackKeys::GenerateStackRendererEditorDataKey(*NewRenderer), ClipboardRenderer->StackNoteData);
-				}
+				continue;
+			}
+
+			if (!RenderersOwner->SupportsRendererClass(ClipboardRenderer->RendererProperties->GetClass()))
+			{
+				continue;
+			}
+
+			UNiagaraRendererProperties* NewRenderer = ClipboardRenderer->RendererProperties->StaticDuplicateWithNewMergeId(RenderersOwner->GetOwnerObject());
+			RenderersOwner->AddRenderer(NewRenderer);
+			if(ClipboardRenderer->StackNoteData.IsValid())
+			{
+				GetStackEditorData().AddOrReplaceStackNote(FNiagaraStackGraphUtilities::StackKeys::GenerateStackRendererEditorDataKey(*NewRenderer), ClipboardRenderer->StackNoteData);
 			}
 		}
 	}
diff --git a/Engine/Plugins/FX/Niagara/Source/NiagaraEditor/Private/ViewModels/Stack/NiagaraStackRenderersOwner.cpp b/Engine/Plugins/FX/Niagara/Source/NiagaraEditor/Private/ViewModels/Stack/NiagaraStackRenderersOwner.cpp
--- a/Engine/Plugins/FX/Niagara/Source/NiagaraEditor/Private/ViewModels/Stack/NiagaraStackRenderersOwner.cpp
+++ b/Engine/Plugins/FX/Niagara/Source/NiagaraEditor/Private/ViewModels/Stack/NiagaraStackRenderersOwner.cpp
@@ -182,9 +182,10 @@ void FNiagaraStackRenderersOwnerStateless::GetRenderers(TArray<UNiagaraRendererP
 	}
 }
 
-bool FNiagaraStackRenderersOwnerStateless::IsRenderCreationInfoSupported(const FNiagaraRendererCreationInfo& RendererCreationInfo) const
+bool FNiagaraStackRenderersOwnerStateless::SupportsRendererClass(const UClass* RendererClass) const
 {
-	return RendererCreationInfo.bIsSupportedByStateless;
+	static const FName NAME_SupportsStateless("SupportsStateless");
+	return RendererClass->HasMetaData(NAME_SupportsStateless);
 }
 
 void FNiagaraStackRenderersOwnerStateless::AddRenderer(UNiagaraRendererProperties* RendererToAdd)
diff --git a/Engine/Plugins/FX/Niagara/Source/NiagaraEditor/Private/ViewModels/Stack/NiagaraStackRenderersOwner.h b/Engine/Plugins/FX/Niagara/Source/NiagaraEditor/Private/ViewModels/Stack/NiagaraStackRenderersOwner.h
--- a/Engine/Plugins/FX/Niagara/Source/NiagaraEditor/Private/ViewModels/Stack/NiagaraStackRenderersOwner.h
+++ b/Engine/Plugins/FX/Niagara/Source/NiagaraEditor/Private/ViewModels/Stack/NiagaraStackRenderersOwner.h
@@ -22,7 +22,7 @@ public:
 	virtual bool IsValid() const = 0;
 	virtual UObject* GetOwnerObject() const = 0;
 	virtual void GetRenderers(TArray<UNiagaraRendererProperties*>& OutRenderers) const = 0;
-	virtual bool IsRenderCreationInfoSupported(const FNiagaraRendererCreationInfo& RendererCreationInfo) const = 0;
+	virtual bool SupportsRendererClass(const UClass* RendererClass) const = 0;
 
 	virtual void AddRenderer(UNiagaraRendererProperties* RendererToAdd) = 0;
 	virtual void RemoveRenderer(UNiagaraRendererProperties* Renderer) = 0;
@@ -43,7 +43,7 @@ public:
 class FNiagaraStackRenderersOwner : public INiagaraStackRenderersOwner
 {
 public:
-	virtual bool IsRenderCreationInfoSupported(const FNiagaraRendererCreationInfo& RendererCreationInfo) const override { return true; }
+	virtual bool SupportsRendererClass(const UClass* RendererClass) const override { return true; }
 
 	virtual bool HasBaseRenderer(UNiagaraRendererProperties* Renderer) const override { return false; }
 	virtual bool IsRendererDifferentFromBase(UNiagaraRendererProperties* Renderer) const override { return false; }
@@ -96,7 +96,7 @@ public:
 	virtual bool IsValid() const override;
 	virtual UObject* GetOwnerObject() const override;
 	virtual void GetRenderers(TArray<UNiagaraRendererProperties*>& OutRenderers) const override;
-	virtual bool IsRenderCreationInfoSupported(const FNiagaraRendererCreationInfo& RendererCreationInfo) const override;
+	virtual bool SupportsRendererClass(const UClass* RendererClass) const override;
 	virtual void AddRenderer(UNiagaraRendererProperties* RendererToAdd) override;
 	virtual void RemoveRenderer(UNiagaraRendererProperties* Renderer) override;
 	virtual void MoveRenderer(UNiagaraRendererProperties* Renderer, int32 NewIndex) override;
diff --git a/Engine/Plugins/FX/Niagara/Source/NiagaraEditor/Public/NiagaraEditorModule.h b/Engine/Plugins/FX/Niagara/Source/NiagaraEditor/Public/NiagaraEditorModule.h
--- a/Engine/Plugins/FX/Niagara/Source/NiagaraEditor/Public/NiagaraEditorModule.h
+++ b/Engine/Plugins/FX/Niagara/Source/NiagaraEditor/Public/NiagaraEditorModule.h
@@ -131,23 +131,12 @@ struct FNiagaraRendererCreationInfo
 	FNiagaraRendererCreationInfo(FText InDisplayName, FText InDescription, const FTopLevelAssetPath& InRendererClassPath, FRendererFactory InFactory) : DisplayName(InDisplayName), Description(InDescription), RendererClassPath(InRendererClassPath), RendererFactory(InFactory)
 	{}
 
-	FNiagaraRendererCreationInfo(FText InDisplayName, bool bInIsSupportedByStateless, const FTopLevelAssetPath& InRendererClassPath, FRendererFactory InFactory) 
-		: DisplayName(InDisplayName), bIsSupportedByStateless(bInIsSupportedByStateless), RendererClassPath(InRendererClassPath), RendererFactory(InFactory)
-	{}
-
-	FNiagaraRendererCreationInfo(FText InDisplayName, FText InDescription, bool bInIsSupportedByStateless, const FTopLevelAssetPath& InRendererClassPath, FRendererFactory InFactory) 
-		: DisplayName(InDisplayName), Description(InDescription), bIsSupportedByStateless(bInIsSupportedByStateless), RendererClassPath(InRendererClassPath), RendererFactory(InFactory)
-	{}
-
 	UPROPERTY()
 	FText DisplayName;
 
 	UPROPERTY()
 	FText Description;
 
-	UPROPERTY()
-	bool bIsSupportedByStateless = false;
-
 	UPROPERTY()
 	FTopLevelAssetPath RendererClassPath;
 	
