Fix RHI Validation Issues.

- Reworked NDC buffer transitions to occur in pre/post stage when the first reading or writing NDC gets access to the buffers.
- Adding mechanism to prevent write and read access to the same NDC buffers within the same dispatch group. Same frame rw is still possible as reads are pushed back to PostRenderOpaque. As long as writes happen before this we're good. TODO: better dependency handling.
- Reworked CPU Readback buffer to allocate and transition on first prestage and be processed for readback in EndFrame. This primarily avoids transition issues with the instance count buffer.
- Marked NDC DI functions as Read/Write
- Removed Allocate DI function until it can get some more dev time for 5.6

diff --git a/Engine/Plugins/FX/Niagara/Source/Niagara/Internal/DataInterface/NiagaraDataInterfaceDataChannelRead.h b/Engine/Plugins/FX/Niagara/Source/Niagara/Internal/DataInterface/NiagaraDataInterfaceDataChannelRead.h
--- a/Engine/Plugins/FX/Niagara/Source/Niagara/Internal/DataInterface/NiagaraDataInterfaceDataChannelRead.h
+++ b/Engine/Plugins/FX/Niagara/Source/Niagara/Internal/DataInterface/NiagaraDataInterfaceDataChannelRead.h
@@ -335,6 +335,8 @@ struct FNiagaraDataInterfaceProxy_DataChannelRead : public FNiagaraDataInterface
 		FRDGBufferRef NDCSpawnDataBuffer;
 
 		uint32 ConsumeInstanceCountOffset = INDEX_NONE;
+
+		FNiagaraDataBufferRef GPUBuffer;
 	};
 
 	TMap<FNiagaraSystemInstanceID, FInstanceData> SystemInstancesToProxyData_RT;
diff --git a/Engine/Plugins/FX/Niagara/Source/Niagara/Internal/DataInterface/NiagaraDataInterfaceDataChannelWrite.h b/Engine/Plugins/FX/Niagara/Source/Niagara/Internal/DataInterface/NiagaraDataInterfaceDataChannelWrite.h
--- a/Engine/Plugins/FX/Niagara/Source/Niagara/Internal/DataInterface/NiagaraDataInterfaceDataChannelWrite.h
+++ b/Engine/Plugins/FX/Niagara/Source/Niagara/Internal/DataInterface/NiagaraDataInterfaceDataChannelWrite.h
@@ -209,6 +209,7 @@ struct FNiagaraDataInterfaceProxy_DataChannelWrite : public FNiagaraDataInterfac
 
 	virtual void PreStage(const FNDIGpuComputePreStageContext& Context)override;
 	virtual void PostStage(const FNDIGpuComputePostStageContext& Context)override;
+	virtual void PostSimulate(const FNDIGpuComputePostSimulateContext& Context)override;
 
 	/** Persistent per instance data on the RT. Constructed when consuming data passed from GT->RT. */
 	struct FInstanceData
@@ -233,6 +234,9 @@ struct FNiagaraDataInterfaceProxy_DataChannelWrite : public FNiagaraDataInterfac
 		//This is the buffer we'll write into for shipping data back to the CPU & Game.
 		FNiagaraDataBufferRef BufferForCPU;
 
+		//This is the main GPU buffer we'll be writing into.
+		FNiagaraDataBufferRef GPUBuffer;
+
 		bool bPublishToGame = false;
 		bool bPublishToCPU = false;
 		bool bPublishToGPU = false;
diff --git a/Engine/Plugins/FX/Niagara/Source/Niagara/Private/DataInterface/NiagaraDataInterfaceDataChannelRead.cpp b/Engine/Plugins/FX/Niagara/Source/Niagara/Private/DataInterface/NiagaraDataInterfaceDataChannelRead.cpp
--- a/Engine/Plugins/FX/Niagara/Source/Niagara/Private/DataInterface/NiagaraDataInterfaceDataChannelRead.cpp
+++ b/Engine/Plugins/FX/Niagara/Source/Niagara/Private/DataInterface/NiagaraDataInterfaceDataChannelRead.cpp
@@ -142,6 +142,7 @@ namespace NDIDataChannelReadLocal
 			NIAGARA_ADD_FUNCTION_SOURCE_INFO(Sig)
 #endif
 			Sig.bMemberFunction = true;
+			Sig.bReadFunction = true;
 			Sig.AddInput(FNiagaraVariable(FNiagaraTypeDefinition(UNiagaraDataInterfaceDataChannelRead::StaticClass()), TEXT("DataChannel interface")));
 			Sig.AddInputWithDefault(FNiagaraVariable(FNiagaraTypeDefinition::GetIntDef(), TEXT("Index")), 0, LOCTEXT("ConsumeIndexInputDesc", "The index to read."));
 			Sig.AddOutput(FNiagaraVariable(FNiagaraTypeDefinition::GetBoolDef(), TEXT("Success")), LOCTEXT("ConsumeSuccessOutputDesc", "True if all reads succeeded."));
@@ -161,6 +162,7 @@ namespace NDIDataChannelReadLocal
 			NIAGARA_ADD_FUNCTION_SOURCE_INFO(Sig)
 #endif
 			Sig.bMemberFunction = true;
+			Sig.bReadFunction = true;
 			Sig.AddInput(FNiagaraVariable(FNiagaraTypeDefinition(UNiagaraDataInterfaceDataChannelRead::StaticClass()), TEXT("DataChannel interface")));
 			Sig.AddInputWithDefault(FNiagaraVariable(FNiagaraTypeDefinition::GetBoolDef(), TEXT("Consume")), FNiagaraBool(true), LOCTEXT("ConsumeInputDesc", "True if this instance (particle/emitter etc) should consume data from the data channel in this call."));
 			Sig.AddOutput(FNiagaraVariable(FNiagaraTypeDefinition::GetBoolDef(), TEXT("Success")), LOCTEXT("ConsumeSuccessOutputDesc", "True if all reads succeeded."));
@@ -2001,7 +2003,7 @@ void UNiagaraDataInterfaceDataChannelRead::SetShaderParameters(const FNiagaraDat
 
 		if (InstanceData->ChannelDataRTProxy && ParameterOffsetTableIndex != INDEX_NONE)
 		{
-			FNiagaraDataBuffer* Data = InstanceData->bReadPrevFrame ? InstanceData->ChannelDataRTProxy->GetPrevFrameData().GetReference() : InstanceData->ChannelDataRTProxy->GetCurrentData().GetReference();
+			FNiagaraDataBuffer* Data = InstanceData->GPUBuffer;
 			if (Data)
 			{
 				const FReadBuffer& ParameterLayoutBuffer = InstanceData->ParameterLayoutBuffer;
@@ -2009,13 +2011,6 @@ void UNiagaraDataInterfaceDataChannelRead::SetShaderParameters(const FNiagaraDat
 				FRDGBufferSRVRef NDCSpawnDataBufferSRV = Context.GetGraphBuilder().CreateSRV(InstanceData->NDCSpawnDataBuffer, PF_R32_SINT);
 				if (NDCSpawnDataBufferSRV && ParameterLayoutBuffer.SRV.IsValid() && ParameterLayoutBuffer.NumBytes > 0)
 				{
-					TArray<FRHITransitionInfo, TInlineAllocator<3>> Transitions;
-					Transitions.Reserve(3);
-					Transitions.Emplace(Data->GetGPUBufferFloat().UAV, ERHIAccess::Unknown, ERHIAccess::SRVCompute);
-					Transitions.Emplace(Data->GetGPUBufferInt().UAV, ERHIAccess::Unknown, ERHIAccess::SRVCompute);
-					//TODO: Half Support | Transitions.Emplace(Data->GetGPUBufferHalf().UAV, ERHIAccess::Unknown, ERHIAccess::SRVCompute);
-					Context.GetGraphBuilder().RHICmdList.Transition(Transitions);
-
 					InstParameters->ParamOffsetTable = ParameterLayoutBuffer.SRV.IsValid() ? ParameterLayoutBuffer.SRV.GetReference() : FNiagaraRenderer::GetDummyUIntBuffer();
 					InstParameters->ParameterOffsetTableIndex = ParameterOffsetTableIndex;
 
@@ -2068,6 +2063,8 @@ void FNiagaraDataInterfaceProxy_DataChannelRead::PreStage(const FNDIGpuComputePr
 
 	if(InstanceData)
 	{
+		InstanceData->GPUBuffer = InstanceData->ChannelDataRTProxy->PrepareForReadAccess(Context.GetGraphBuilder(), InstanceData->bReadPrevFrame == false);
+
 		//TODO: Should grab just one for the whole frame...
 		//TODO: Add some wrap behavior...		
 		if(InstanceData->ConsumeInstanceCountOffset == INDEX_NONE)
@@ -2087,12 +2084,15 @@ void FNiagaraDataInterfaceProxy_DataChannelRead::PreStage(const FNDIGpuComputePr
 
 void FNiagaraDataInterfaceProxy_DataChannelRead::PostSimulate(const FNDIGpuComputePostSimulateContext& Context)
 {
-	if (Context.IsFinalPostSimulate())
+	FNiagaraDataInterfaceProxy_DataChannelRead::FInstanceData* InstanceData = SystemInstancesToProxyData_RT.Find(Context.GetSystemInstanceID());
+	if(InstanceData && InstanceData->ChannelDataRTProxy)
 	{
-		if(FNiagaraDataInterfaceProxy_DataChannelRead::FInstanceData* InstanceData = SystemInstancesToProxyData_RT.Find(Context.GetSystemInstanceID()))
+		InstanceData->GPUBuffer = nullptr;
+		InstanceData->ChannelDataRTProxy->EndReadAccess(Context.GetGraphBuilder(), InstanceData->bReadPrevFrame == false);
+
+		if (Context.IsFinalPostSimulate())
 		{
 			InstanceData->NDCSpawnDataBuffer = nullptr;
-
 			Context.GetInstanceCountManager().FreeEntry(InstanceData->ConsumeInstanceCountOffset);
 		}
 	}
diff --git a/Engine/Plugins/FX/Niagara/Source/Niagara/Private/DataInterface/NiagaraDataInterfaceDataChannelWrite.cpp b/Engine/Plugins/FX/Niagara/Source/Niagara/Private/DataInterface/NiagaraDataInterfaceDataChannelWrite.cpp
--- a/Engine/Plugins/FX/Niagara/Source/Niagara/Private/DataInterface/NiagaraDataInterfaceDataChannelWrite.cpp
+++ b/Engine/Plugins/FX/Niagara/Source/Niagara/Private/DataInterface/NiagaraDataInterfaceDataChannelWrite.cpp
@@ -83,25 +83,28 @@ namespace NDIDataChannelWriteLocal
 		return Sig;
 	}
 
-	const FNiagaraFunctionSignature& GetFunctionSig_Allocate()
-	{
-		static FNiagaraFunctionSignature Sig;
-		if (!Sig.IsValid())
-		{
-			Sig.Name = TEXT("Allocate");
-#if WITH_EDITORONLY_DATA
-			NIAGARA_ADD_FUNCTION_SOURCE_INFO(Sig)
-			Sig.Description = LOCTEXT("AllocateFunctionDescription", "Adds an amount to allocated into the bound NDC data for the given emitter to write into.");
-#endif
-			Sig.bMemberFunction = true;
-			Sig.bRequiresExecPin = true;
-			Sig.ModuleUsageBitmask = ENiagaraScriptUsageMask::Emitter | ENiagaraScriptUsageMask::System;
-			Sig.AddInput(FNiagaraVariable(FNiagaraTypeDefinition(UNiagaraDataInterfaceDataChannelWrite::StaticClass()), TEXT("DataChannel interface")));
-			Sig.AddInputWithoutDefault(FNiagaraVariable(FNiagaraTypeDefinition(FNiagaraEmitterID::StaticStruct()), TEXT("Emitter ID")), LOCTEXT("EmitterIDDesc", "ID of the emitter that will be writing into this allocated space."));
-			Sig.AddInputWithDefault(FNiagaraVariable(FNiagaraTypeDefinition::GetIntDef(), TEXT("Allocation Count")), 0, LOCTEXT("AllocationCountDesc", "The number of elements to allocate in the NDC for writing from the given emitter."));
-		}
-		return Sig;
-	}
+	//For now, disable the allocate function as we don't have time to test it thoroughly.
+	// TODO: Make experimental?
+// 	const FNiagaraFunctionSignature& GetFunctionSig_Allocate()
+// 	{
+// 		static FNiagaraFunctionSignature Sig;
+// 		if (!Sig.IsValid())
+// 		{
+// 			Sig.Name = TEXT("Allocate");
+// #if WITH_EDITORONLY_DATA
+// 			NIAGARA_ADD_FUNCTION_SOURCE_INFO(Sig)
+// 			Sig.Description = LOCTEXT("AllocateFunctionDescription", "Adds an amount to allocated into the bound NDC data for the given emitter to write into.");
+// #endif
+// 			Sig.bMemberFunction = true;
+// 			Sig.bRequiresExecPin = true;
+// 			Sig.bWriteFunction = true;
+// 			Sig.ModuleUsageBitmask = ENiagaraScriptUsageMask::Emitter | ENiagaraScriptUsageMask::System;
+// 			Sig.AddInput(FNiagaraVariable(FNiagaraTypeDefinition(UNiagaraDataInterfaceDataChannelWrite::StaticClass()), TEXT("DataChannel interface")));
+// 			Sig.AddInputWithoutDefault(FNiagaraVariable(FNiagaraTypeDefinition(FNiagaraEmitterID::StaticStruct()), TEXT("Emitter ID")), LOCTEXT("EmitterIDDesc", "ID of the emitter that will be writing into this allocated space."));
+// 			Sig.AddInputWithDefault(FNiagaraVariable(FNiagaraTypeDefinition::GetIntDef(), TEXT("Allocation Count")), 0, LOCTEXT("AllocationCountDesc", "The number of elements to allocate in the NDC for writing from the given emitter."));
+// 		}
+// 		return Sig;
+// 	}
 
 	const FNiagaraFunctionSignature& GetFunctionSig_Write()
 	{
@@ -117,6 +120,7 @@ namespace NDIDataChannelWriteLocal
 #endif
 			Sig.bMemberFunction = true;
 			Sig.bRequiresExecPin = true;
+			Sig.bWriteFunction = true;
 			Sig.bSupportsGPU = false;//Cannot use direct index writes on GPU as we write into one shared buffer with all DIs using the same NDC data.
 			Sig.AddInput(FNiagaraVariable(FNiagaraTypeDefinition(UNiagaraDataInterfaceDataChannelWrite::StaticClass()), TEXT("DataChannel interface")));
 			Sig.AddInput(EmitVar, LOCTEXT("ExecuteAppendFlagTooltip", "If true then the append is executed, if false then the append call is skipped"));
@@ -142,6 +146,7 @@ namespace NDIDataChannelWriteLocal
 #endif
 			Sig.bMemberFunction = true;
 			Sig.bRequiresExecPin = true;
+			Sig.bWriteFunction = true;
 			Sig.AddInput(FNiagaraVariable(FNiagaraTypeDefinition(UNiagaraDataInterfaceDataChannelWrite::StaticClass()), TEXT("DataChannel interface")));
 			Sig.AddInput(EmitVar, LOCTEXT("ExecuteAppendFlagTooltip", "If true then the append is executed, if false then the append call is skipped"));
 			Sig.AddOutput(FNiagaraVariable(FNiagaraTypeDefinition::GetBoolDef(), TEXT("Success")));
@@ -932,7 +937,7 @@ bool UNiagaraDataInterfaceDataChannelWrite::CopyToInternal(UNiagaraDataInterface
 void UNiagaraDataInterfaceDataChannelWrite::GetFunctionsInternal(TArray<FNiagaraFunctionSignature>& OutFunctions) const
 {
 	OutFunctions.Add(NDIDataChannelWriteLocal::GetFunctionSig_Num());
-	OutFunctions.Add(NDIDataChannelWriteLocal::GetFunctionSig_Allocate());
+	//OutFunctions.Add(NDIDataChannelWriteLocal::GetFunctionSig_Allocate());
 	OutFunctions.Add(NDIDataChannelWriteLocal::GetFunctionSig_Write());
 	OutFunctions.Add(NDIDataChannelWriteLocal::GetFunctionSig_Append());
 }
@@ -944,10 +949,10 @@ void UNiagaraDataInterfaceDataChannelWrite::GetVMExternalFunction(const FVMExter
 	{
 		OutFunc = FVMExternalFunction::CreateLambda([this](FVectorVMExternalFunctionContext& Context) { this->Num(Context); });
 	}
-	else if (BindingInfo.Name == NDIDataChannelWriteLocal::GetFunctionSig_Allocate().Name)
-	{
-		OutFunc = FVMExternalFunction::CreateLambda([this](FVectorVMExternalFunctionContext& Context) { this->Allocate(Context); });
-	}
+// 	else if (BindingInfo.Name == NDIDataChannelWriteLocal::GetFunctionSig_Allocate().Name)
+// 	{
+// 		OutFunc = FVMExternalFunction::CreateLambda([this](FVectorVMExternalFunctionContext& Context) { this->Allocate(Context); });
+// 	}
 	else
 	{
 		int32 FuncIndex = CompiledData.FindFunctionInfoIndex(BindingInfo.Name, BindingInfo.VariadicInputs, BindingInfo.VariadicOutputs);
@@ -998,7 +1003,7 @@ void UNiagaraDataInterfaceDataChannelWrite::Allocate(FVectorVMExternalFunctionCo
 	uint32 Count = InAllocationCount.GetAndAdvance();
 	
 	//Store the count so we can pass this to the GPU for allocating space in the main GPU write buffer.
-	InstData->DynamicAllocationCount = Count;
+	InstData->DynamicAllocationCount += Count;
 	
 	//If we have a CPU write buffer, allocate that now. Do this here so the emitter/system script itself can write data if it wants to.
 	if(InstData->DestinationData)
@@ -1323,26 +1328,10 @@ void UNiagaraDataInterfaceDataChannelWrite::SetShaderParameters(const FNiagaraDa
 
 		if (InstanceData->ChannelDataRTProxy && ParameterOffsetTableIndex != INDEX_NONE)
 		{
-			FNiagaraDataBuffer* GPUBuffer = InstanceData->bPublishToGPU ? InstanceData->ChannelDataRTProxy->GetCurrentData() : nullptr;
-			FNiagaraDataBuffer* BufferForCPU = InstanceData->bPublishToCPU || InstanceData->bPublishToGame ? InstanceData->BufferForCPU.GetReference() : nullptr;
+			FNiagaraDataBuffer* GPUBuffer = InstanceData->GPUBuffer.GetReference();
+			FNiagaraDataBuffer* BufferForCPU = InstanceData->BufferForCPU.GetReference();
 			if (GPUBuffer || BufferForCPU)
 			{
-				TArray<FRHITransitionInfo, TInlineAllocator<3>> Transitions;
-				Transitions.Reserve(6);
-				if(GPUBuffer)
-				{
-					Transitions.Emplace(GPUBuffer->GetGPUBufferFloat().UAV, ERHIAccess::Unknown, ERHIAccess::UAVCompute);
-					Transitions.Emplace(GPUBuffer->GetGPUBufferInt().UAV, ERHIAccess::Unknown, ERHIAccess::UAVCompute);
-					//TODO: Half Support | Transitions.Emplace(GPUBuffer->GetGPUBufferHalf().UAV, ERHIAccess::Unknown, ERHIAccess::UAVCompute);
-				}
-				if(BufferForCPU)
-				{
-					Transitions.Emplace(BufferForCPU->GetGPUBufferFloat().UAV, ERHIAccess::Unknown, ERHIAccess::UAVCompute);
-					Transitions.Emplace(BufferForCPU->GetGPUBufferInt().UAV, ERHIAccess::Unknown, ERHIAccess::UAVCompute);
-					//TODO: Half Support | Transitions.Emplace(BufferForCPU->GetGPUBufferHalf().UAV, ERHIAccess::Unknown, ERHIAccess::UAVCompute);
-				}
-				Context.GetGraphBuilder().RHICmdList.Transition(Transitions);
-
 				const FReadBuffer& ParameterLayoutBuffer = InstanceData->ParameterLayoutBuffer;
 
 				if (ParameterLayoutBuffer.SRV.IsValid() && ParameterLayoutBuffer.NumBytes > 0)
@@ -1463,45 +1452,49 @@ void FNiagaraDataInterfaceProxy_DataChannelWrite::PreStage(const FNDIGpuComputeP
 
 	if (InstanceData && InstanceData->ChannelDataRTProxy)
 	{
-		if(InstanceData->bPublishToGame || InstanceData->bPublishToCPU)
+		if( InstanceData->BufferForCPU == nullptr && (InstanceData->bPublishToGame || InstanceData->bPublishToCPU) )
 		{
 			//Allocate a separate buffer that we will write into and ship back to the CPU.			 
-			InstanceData->BufferForCPU = InstanceData->ChannelDataRTProxy->AllocateBufferForCPU(Context.GetGraphBuilder().RHICmdList, Context.GetComputeDispatchInterface().GetFeatureLevel(), InstanceData->AllocationCount);
+			InstanceData->BufferForCPU = InstanceData->ChannelDataRTProxy->AllocateBufferForCPU(Context.GetGraphBuilder(), Context.GetComputeDispatchInterface().GetFeatureLevel(), InstanceData->AllocationCount, InstanceData->bPublishToGame, InstanceData->bPublishToCPU, InstanceData->LwcTile);
 
 			//Get a new instance count.
 			uint32 Offset = InstanceData->BufferForCPU->GetGPUInstanceCountBufferOffset();
 			Context.GetInstanceCountManager().FreeEntry(Offset);
 			InstanceData->BufferForCPU->SetGPUInstanceCountBufferOffset(Context.GetInstanceCountManager().AcquireOrAllocateEntry(Context.GetGraphBuilder().RHICmdList));
 		}
+
+		if(InstanceData->bPublishToGPU)
+		{
+			check(InstanceData->GPUBuffer == nullptr);
+			InstanceData->GPUBuffer = InstanceData->ChannelDataRTProxy->PrepareForWriteAccess(Context.GetGraphBuilder());
+		}
 	}
 }
 
+
 void FNiagaraDataInterfaceProxy_DataChannelWrite::PostStage(const FNDIGpuComputePostStageContext& Context)
 {
 	FNiagaraDataInterfaceProxy_DataChannelWrite::FInstanceData* InstanceData = SystemInstancesToProxyData_RT.Find(Context.GetSystemInstanceID());
 
-	if (InstanceData && InstanceData->ChannelDataRTProxy)
+	if(InstanceData && InstanceData->ChannelDataRTProxy)
 	{
-		if (InstanceData->BufferForCPU)
+		if (InstanceData->GPUBuffer)
 		{
-			Context.GetGraphBuilder().AddPass(
-				RDG_EVENT_NAME("NDC Write Enqueue CPU Readback"),
-				ERDGPassFlags::None,
-				[PASS_ReadbackManager = Context.GetComputeDispatchInterface().GetGpuReadbackManager(), 
-				PASS_InstCountManager = &Context.GetInstanceCountManager(), 
-				PASS_Buffer = InstanceData->BufferForCPU,
-				PASS_Proxy = InstanceData->ChannelDataRTProxy,
-				PASS_bPublishToGame = InstanceData->bPublishToGame,
-				PASS_bPublishToCPU = InstanceData->bPublishToCPU,
-				PASS_LwcTile = InstanceData->LwcTile]
-				(FRHICommandListImmediate& RHICmdList)
-				{
-					PASS_Proxy->EnqueueReadbackForCPUBuffer(RHICmdList, PASS_Buffer, PASS_ReadbackManager, *PASS_InstCountManager, PASS_bPublishToGame, PASS_bPublishToCPU, PASS_LwcTile);
-					uint32 Offset = PASS_Buffer->GetGPUInstanceCountBufferOffset();
-					PASS_InstCountManager->FreeEntry(Offset);
-					PASS_Buffer->SetGPUInstanceCountBufferOffset(INDEX_NONE);
-				});
-		
+			InstanceData->ChannelDataRTProxy->EndWriteAccess(Context.GetGraphBuilder());
+			InstanceData->GPUBuffer = nullptr;
+		}
+	}
+}
+
+void FNiagaraDataInterfaceProxy_DataChannelWrite::PostSimulate(const FNDIGpuComputePostSimulateContext& Context)
+{
+	FNiagaraDataInterfaceProxy_DataChannelWrite::FInstanceData* InstanceData = SystemInstancesToProxyData_RT.Find(Context.GetSystemInstanceID());
+	
+	if (InstanceData && InstanceData->ChannelDataRTProxy && Context.IsFinalPostSimulate())
+	{
+		if(InstanceData->BufferForCPU)
+		{
+			InstanceData->ChannelDataRTProxy->AddTransition(Context.GetGraphBuilder(), ERHIAccess::UAVCompute, ERHIAccess::SRVMask, InstanceData->BufferForCPU.GetReference());
 			InstanceData->BufferForCPU = nullptr;
 		}
 	}
diff --git a/Engine/Plugins/FX/Niagara/Source/Niagara/Private/NiagaraDataChannel.cpp b/Engine/Plugins/FX/Niagara/Source/Niagara/Private/NiagaraDataChannel.cpp
--- a/Engine/Plugins/FX/Niagara/Source/Niagara/Private/NiagaraDataChannel.cpp
+++ b/Engine/Plugins/FX/Niagara/Source/Niagara/Private/NiagaraDataChannel.cpp
@@ -19,6 +19,8 @@
 #endif
 #include "NiagaraDataSetReadback.h"
 
+#include "RenderGraphUtils.h"
+
 #include UE_INLINE_GENERATED_CPP_BY_NAME(NiagaraDataChannel)
 
 #define LOCTEXT_NAMESPACE "NiagaraDataChannels"
@@ -569,46 +571,180 @@ void FNiagaraDataChannelDataProxy::BeginFrame(FNiagaraGpuComputeDispatchInterfac
 	}
 }
 
-void FNiagaraDataChannelDataProxy::EndFrame(FRHICommandListImmediate& RHICmdList)
+void FNiagaraDataChannelDataProxy::EndFrame(FNiagaraGpuComputeDispatchInterface* DispatchInterface, FRHICommandListImmediate& RHICmdList)
 {
 	check(GPUDataSet);
 	check(GPUDataSet->GetSimTarget() == ENiagaraSimTarget::GPUComputeSim);
+
+	check(CurrBufferAccessCounts == 0);
+
+	//Handle pending readbacks for GPU->CPU data.
+	{
+		for (FNDCGpuReadbackInfo& ReadbackInfo : PendingGPUReadbackBuffers)
+		{
+			//We enqueue a readback for the data if we're wanting to publish it to the CPU or Game
+			check(ReadbackInfo.bPublishToCPU || ReadbackInfo.bPublishToGame);//We should be doing one or the other or we should not get here.
+
+			auto PublishOnCPU = [
+				LocalOwner = Owner,
+				PASS_PublishCPU = ReadbackInfo.bPublishToCPU,
+				PASS_PublishGame = ReadbackInfo.bPublishToGame,
+				PASS_LWCTile = ReadbackInfo.LWCTile](FNiagaraDataBufferRef ReadbackBuffer)
+				{
+					if (FNiagaraDataChannelData* NDCData = LocalOwner.Pin().Get())
+					{
+						if (ReadbackBuffer && ReadbackBuffer->GetNumInstances())
+						{
+							FNiagaraDataChannelPublishRequest PublishRequest(ReadbackBuffer);
+							PublishRequest.bVisibleToGame = PASS_PublishGame;
+							PublishRequest.bVisibleToCPUSims = PASS_PublishCPU;
+							PublishRequest.bVisibleToGPUSims = false;//Don't want to ping pong back to the GPU
+							PublishRequest.LwcTile = PASS_LWCTile;
+#if WITH_NIAGARA_DEBUGGER
+							PublishRequest.DebugSource = TEXT("NDC GPU Readback");//TODO: Feed in real source names also
+#endif
+							NDCData->PublishFromGPU(PublishRequest);
+						}
+					}
+				};
+
+			TSharedPtr<FNiagaraDataBufferReadback, ESPMode::ThreadSafe> NewReadback = MakeShared<FNiagaraDataBufferReadback, ESPMode::ThreadSafe>();
+			NewReadback->GetOnReadbackComplete().BindLambda(PublishOnCPU);
+			NewReadback->EnqueueReadback(RHICmdList, ReadbackInfo.Buffer, DispatchInterface->GetGpuReadbackManager(), DispatchInterface->GetGPUInstanceCounterManager());
+
+			uint32 Offset = ReadbackInfo.Buffer->GetGPUInstanceCountBufferOffset();
+			DispatchInterface->GetGPUInstanceCounterManager().FreeEntry(Offset);
+			ReadbackInfo.Buffer->SetGPUInstanceCountBufferOffset(INDEX_NONE);
+		}
+		PendingGPUReadbackBuffers.Reset();
+	}
+}
+
+FNiagaraDataBufferRef FNiagaraDataChannelDataProxy::PrepareForWriteAccess(FRDGBuilder& GraphBuilder)
+{
+	check(IsInRenderingThread());
+
+	if(CurrFrameData)
+	{
+		if (CurrBufferAccessCounts > 0)
+		{
+#if !UE_BUILD_SHIPPING
+			UE_LOG(LogNiagara, Warning, TEXT("Attempting to write to a Niagara Data Channel in the same stage in which it's being read. {%s}"), *DebugName);
+#endif
+			return nullptr;
+		}
+
+		//If we're prepping the first writer, transition to UAV
+		if (CurrBufferAccessCounts == 0)
+		{
+			AddTransition(GraphBuilder, ERHIAccess::SRVMask, ERHIAccess::UAVCompute, CurrFrameData.GetReference());
+		}
+
+		--CurrBufferAccessCounts;
+	}
+
+	return CurrFrameData;
 }
 
-FNiagaraDataBufferRef FNiagaraDataChannelDataProxy::AllocateBufferForCPU(FRHICommandListImmediate& RHICmdList, ERHIFeatureLevel::Type FeatureLevel, int32 AllocationSize)
+void FNiagaraDataChannelDataProxy::EndWriteAccess(FRDGBuilder& GraphBuilder)
 {
-	FNiagaraDataBuffer* Ret = &GPUDataSet->AllocateBuffer();
-	Ret->AllocateGPU(RHICmdList, AllocationSize, FeatureLevel, TEXT("NDC GPU Buffers for CPU"));	
-	return Ret->UnlockForRead();
+	check(IsInRenderingThread());
+
+	check(CurrBufferAccessCounts < 0);
+
+	++CurrBufferAccessCounts;
+
+	//If we reach 0 writers, transition back to SRV.
+	if (CurrBufferAccessCounts == 0)
+	{
+		AddTransition(GraphBuilder, ERHIAccess::UAVCompute, ERHIAccess::SRVMask, CurrFrameData.GetReference());
+	}
 }
 
-void FNiagaraDataChannelDataProxy::EnqueueReadbackForCPUBuffer(FRHICommandList& RHICmdList, FNiagaraDataBufferRef Buffer, FNiagaraGpuReadbackManager* ReadbackManager, FNiagaraGPUInstanceCountManager& InstanceCountManager, bool bPublishToGame, bool bPublishToCPU, FVector3f LWCTile)
+FNiagaraDataBufferRef FNiagaraDataChannelDataProxy::PrepareForReadAccess(FRDGBuilder& GraphBuilder, bool bCurrentFrame)
 {
-	auto PublishOnCPU = [=, LocalOwner = Owner](FNiagaraDataBufferRef ReadbackBuffer)
+	check(IsInRenderingThread());
+
+	if(bCurrentFrame)
 	{
-		if (FNiagaraDataChannelData* NDCData = LocalOwner.Pin().Get())
+		if (CurrBufferAccessCounts < 0)
 		{
-			if (ReadbackBuffer && ReadbackBuffer->GetNumInstances())
-			{
-				FNiagaraDataChannelPublishRequest PublishRequest(ReadbackBuffer);
-				PublishRequest.bVisibleToGame = bPublishToGame;
-				PublishRequest.bVisibleToCPUSims = bPublishToCPU;
-				PublishRequest.bVisibleToGPUSims = false;//Don't want to ping pong back to the GPU
-				PublishRequest.LwcTile = LWCTile;
-#if WITH_NIAGARA_DEBUGGER
-				PublishRequest.DebugSource = TEXT("NDC GPU Readback");//TODO: Feed in real source names also
+#if !UE_BUILD_SHIPPING
+			UE_LOG(LogNiagara, Warning, TEXT("Attempting to write to a Niagara Data Channel in the same stage in which it's being read. {%s}"), *DebugName);
 #endif
-				NDCData->PublishFromGPU(PublishRequest);
-			}
+			return nullptr;
 		}
-	};
+		
+		++CurrBufferAccessCounts;
+
+		return CurrFrameData;
+	}
+	else
+	{
+		return PrevFrameData;
+	}
+}
+
+void FNiagaraDataChannelDataProxy::EndReadAccess(FRDGBuilder& GraphBuilder, bool bCurrentFrame)
+{
+	check(IsInRenderingThread());
+	if(bCurrentFrame)
+	{
+		check(CurrBufferAccessCounts > 0);
+		--CurrBufferAccessCounts;
+	}
+}
+
+void FNiagaraDataChannelDataProxy::AddTransition(FRDGBuilder& GraphBuilder, ERHIAccess AccessBefore, ERHIAccess AccessAfter, FNiagaraDataBuffer* Buffer)
+{
+	if (Buffer)
+	{
+		TArray<FRHITransitionInfo, TInlineAllocator<6>> Transitions;
+		Transitions.Reserve(3);
+
+		Transitions.Emplace(Buffer->GetGPUBufferFloat().UAV, AccessBefore, AccessAfter);
+		Transitions.Emplace(Buffer->GetGPUBufferInt().UAV, AccessBefore, AccessAfter);
+		//TODO: Half Support | Transitions.Emplace(Buffer->GetGPUBufferHalf().UAV, AccessBefore, AccessAfter);
+
+		TArray<FRHIUnorderedAccessView*> UAVsToOverlap;
+		UAVsToOverlap.Emplace(Buffer->GetGPUBufferFloat().UAV);
+		UAVsToOverlap.Emplace(Buffer->GetGPUBufferInt().UAV);
+		//TODO: Half Support | UAVsToOverlap.Emplace(Buffer->GetGPUBufferHalf().UAV);
+
+	GraphBuilder.AddPass(
+		RDG_EVENT_NAME("NDC Proxy - Transition Buffers"),
+		ERDGPassFlags::None,
+		[Transitions, AccessAfter, UAVsToOverlap](FRHICommandListImmediate& RHICmdList)
+		{	
+			RHICmdList.Transition(Transitions);		
+			
+			//We may have multiple overlapping dispatches accessing the same NDC buffers.
+			if(AccessAfter == ERHIAccess::UAVCompute)
+			{
+				RHICmdList.BeginUAVOverlap(UAVsToOverlap);
+			}
+			else
+			{
+				RHICmdList.EndUAVOverlap(UAVsToOverlap);
+			}
+		});
+	}
+}
+
+FNiagaraDataBufferRef FNiagaraDataChannelDataProxy::AllocateBufferForCPU(FRDGBuilder& GraphBuilder, ERHIFeatureLevel::Type FeatureLevel, int32 AllocationSize, bool bPublishToGame, bool bPublishToCPU, FVector3f LWCTile)
+{
+	FNiagaraDataBuffer* Buffer = &GPUDataSet->AllocateBuffer();
+	Buffer->AllocateGPU(GraphBuilder.RHICmdList, AllocationSize, FeatureLevel, TEXT("NDC GPU Buffers for CPU"));
+
+	FNDCGpuReadbackInfo& NewReadback = PendingGPUReadbackBuffers.AddDefaulted_GetRef();
+	NewReadback.Buffer = Buffer->UnlockForRead();
+	NewReadback.bPublishToCPU = bPublishToCPU;
+	NewReadback.bPublishToGame = bPublishToGame;
+	NewReadback.LWCTile = LWCTile;
 
-	//We enqueue a readback for the data if we're wanting to publish it to the CPU or Game
-	check(bPublishToCPU || bPublishToGame);//We should be doing one or the other or we should not be calling this function.
+	AddTransition(GraphBuilder, ERHIAccess::SRVMask, ERHIAccess::UAVCompute, Buffer);
 
-	TSharedPtr<FNiagaraDataBufferReadback, ESPMode::ThreadSafe> NewReadback = MakeShared<FNiagaraDataBufferReadback, ESPMode::ThreadSafe>();
-	NewReadback->GetOnReadbackComplete().BindLambda(PublishOnCPU);
-	NewReadback->EnqueueReadback(RHICmdList, Buffer, ReadbackManager, InstanceCountManager);
+	return NewReadback.Buffer;
 }
 
 void FNiagaraDataChannelDataProxy::AddBuffersFromCPU(const TArray<FNiagaraDataBufferRef>& BuffersFromCPU)
diff --git a/Engine/Plugins/FX/Niagara/Source/Niagara/Private/NiagaraGpuComputeDispatch.cpp b/Engine/Plugins/FX/Niagara/Source/Niagara/Private/NiagaraGpuComputeDispatch.cpp
--- a/Engine/Plugins/FX/Niagara/Source/Niagara/Private/NiagaraGpuComputeDispatch.cpp
+++ b/Engine/Plugins/FX/Niagara/Source/Niagara/Private/NiagaraGpuComputeDispatch.cpp
@@ -2114,7 +2114,7 @@ void FNiagaraGpuComputeDispatch::PostRenderOpaque(FRDGBuilder& GraphBuilder, TCo
 
 				for (FNiagaraDataChannelDataProxy* NDCProxy : NDCDataProxies)
 				{
-					NDCProxy->EndFrame(RHICmdList);
+					NDCProxy->EndFrame(this, RHICmdList);
 				}
 			}
 		}
diff --git a/Engine/Plugins/FX/Niagara/Source/Niagara/Public/NiagaraDataChannel.h b/Engine/Plugins/FX/Niagara/Source/Niagara/Public/NiagaraDataChannel.h
--- a/Engine/Plugins/FX/Niagara/Source/Niagara/Public/NiagaraDataChannel.h
+++ b/Engine/Plugins/FX/Niagara/Source/Niagara/Public/NiagaraDataChannel.h
@@ -49,9 +49,18 @@ class UNiagaraDataChannelReader;
 struct FNiagaraDataChannelPublishRequest;
 struct FNiagaraDataChannelGameDataLayout;
 class FNiagaraGpuReadbackManager;
+class FRDGBuilder;
 
 //////////////////////////////////////////////////////////////////////////
 
+struct FNDCGpuReadbackInfo
+{
+	FNiagaraDataBufferRef Buffer;
+	bool bPublishToCPU = false;
+	bool bPublishToGame = false;
+	FVector3f LWCTile;
+};
+
 /** Render thread proxy of FNiagaraDataChannelData. */
 struct FNiagaraDataChannelDataProxy
 {
@@ -65,10 +74,16 @@ struct FNiagaraDataChannelDataProxy
 
 	//Buffers coming from the CPU that we're going to copy up for reading on the GPU
 	TArray<FNiagaraDataBufferRef> PendingCPUBuffers;
+
+	//Buffers written from the GPU that we must send back to the CPU.
+	TArray<FNDCGpuReadbackInfo> PendingGPUReadbackBuffers;
 	
 	//Users that need space in this NDC Data add to this for each tick via AddGPUAllocationForNextTick().
 	int32 PendingGPUAllocations = 0;
 
+	//Track current read/write counts +ve for readers, -ve for writers. We cannot mix readers and writers in the same buffer in the same stage.
+	int32 CurrBufferAccessCounts = 0;
+
 	#if !UE_BUILD_SHIPPING
 	FNiagaraGpuComputeDispatchInterface* DispatchInterfaceForDebuggingOnly = nullptr;
 	
@@ -80,16 +95,24 @@ struct FNiagaraDataChannelDataProxy
 
 	void Cleanup(FNiagaraGpuComputeDispatchInterface* ComputeDispatchInterface);
 	void BeginFrame(FNiagaraGpuComputeDispatchInterface* DispatchInterface, FRHICommandListImmediate& RHICmdList);
-	void EndFrame(FRHICommandListImmediate& RHICmdList);
+	void EndFrame(FNiagaraGpuComputeDispatchInterface* DispatchInterface, FRHICommandListImmediate& RHICmdList);
 	void Reset();
 
-	FNiagaraDataBufferRef AllocateBufferForCPU(FRHICommandListImmediate& RHICmdList, ERHIFeatureLevel::Type FeatureLevel, int32 AllocationSize);
-	void EnqueueReadbackForCPUBuffer(FRHICommandList& RHICmdList, FNiagaraDataBufferRef Buffer, FNiagaraGpuReadbackManager* ReadbackManager, FNiagaraGPUInstanceCountManager& InstanceCountManager, bool bPublishToGame, bool bPublishToCPU, FVector3f LWCTile);
+	FNiagaraDataBufferRef PrepareForWriteAccess(FRDGBuilder& GraphBuilder);
+	void EndWriteAccess(FRDGBuilder& GraphBuilder);
+	
+	FNiagaraDataBufferRef PrepareForReadAccess(FRDGBuilder& GraphBuilder, bool bCurrentFrame);
+	void EndReadAccess(FRDGBuilder& GraphBuilder, bool bCurrentFrame);
+
+
+	FNiagaraDataBufferRef AllocateBufferForCPU(FRDGBuilder& GraphBuilder, ERHIFeatureLevel::Type FeatureLevel, int32 AllocationSize, bool bPublishToGame, bool bPublishToCPU, FVector3f LWCTile);
 	void AddBuffersFromCPU(const TArray<FNiagaraDataBufferRef>& BuffersFromCPU);	
 	void AddGPUAllocationForNextTick(int32 AllocationCount);
 
 	FNiagaraDataBufferRef GetCurrentData()const { return CurrFrameData; }
 	FNiagaraDataBufferRef GetPrevFrameData()const { return PrevFrameData; }
+	
+	void AddTransition(FRDGBuilder& GraphBuilder, ERHIAccess AccessBefore, ERHIAccess AccessAfter, FNiagaraDataBuffer* Buffer);
 };
 
 DECLARE_MULTICAST_DELEGATE_OneParam(FOnDataChannelCreated, const UNiagaraDataChannel*);
