Ensures that regions captured by the Landscape DI actually point to valid collision components and also adds in optional limits for the number of regions to search/capture

diff --git a/Engine/Plugins/FX/Niagara/Source/Niagara/Private/NiagaraDataInterfaceLandscape.cpp b/Engine/Plugins/FX/Niagara/Source/Niagara/Private/NiagaraDataInterfaceLandscape.cpp
--- a/Engine/Plugins/FX/Niagara/Source/Niagara/Private/NiagaraDataInterfaceLandscape.cpp
+++ b/Engine/Plugins/FX/Niagara/Source/Niagara/Private/NiagaraDataInterfaceLandscape.cpp
@@ -35,6 +35,26 @@ struct FNDILandscapeData_GameThread;
 
 namespace NiagaraDataInterfaceLandscape
 {
+	// This controls the maximum number of regions that will be evaluated for capture in FNDI_Landscape_GeneratedData.
+	// A value of -1 means no limit.
+	static int32 GMaxRegionSearchCount = -1;
+	static FAutoConsoleVariableRef CVarMaxRegionSearchCount(
+		TEXT("fx.Niagara.Landscape.MaxRegionSearchCount"),
+		GMaxRegionSearchCount,
+		TEXT("The maximum number of collision components that will be evaluated for capture by the Landscape DI."),
+		ECVF_Default
+	);
+
+	// This controls the maximum number of regions that will be captured in FNDI_Landscape_GeneratedData
+	// A value of -1 means no limit.
+	static int32 GMaxRegionCaptureCount = -1;
+	static FAutoConsoleVariableRef CVarMaxRegionCaptureCount(
+		TEXT("fx.Niagara.Landscape.MaxRegionCaptureCount"),
+		GMaxRegionCaptureCount,
+		TEXT("The maximum number of collision components that will be captured by the Landscape DI."),
+		ECVF_Default
+	);
+
 	enum Version
 	{
 		InitialVersion = 0,
@@ -835,41 +855,36 @@ void FNDI_Landscape_SharedResource::Initialize()
 
 		for (const FIntPoint& Region : ResourceKey.CapturedRegions)
 		{
-			auto FoundCollisionComponent = LandscapeInfo->XYtoCollisionComponentMap.Find(Region);
-			check(FoundCollisionComponent);
-
-			if (FoundCollisionComponent)
+			const ULandscapeHeightfieldCollisionComponent* CollisionComponent = LandscapeInfo->XYtoCollisionComponentMap.FindRef(Region);
+			if (ensure(CollisionComponent != nullptr))
 			{
-				if (const ULandscapeHeightfieldCollisionComponent* CollisionComponent = *FoundCollisionComponent)
+				if (HeightValues)
+				{
+					const FIntPoint SectionBase = (Region - ResourceKey.MinCaptureRegion) * ComponentQuadCount;
+					CollisionComponent->FillHeightTile(*HeightValues, SectionBase.X + SectionBase.Y * CaptureVertexSpan.X, CaptureVertexSpan.X);
+				}
+
+				if (PhysMatValues)
 				{
-					if (HeightValues)
+					const FIntPoint SectionBase = (Region - ResourceKey.MinCaptureRegion) * ComponentQuadCount;
+					CollisionComponent->FillMaterialIndexTile(*PhysMatValues, SectionBase.X + SectionBase.Y * CaptureVertexSpan.X, CaptureVertexSpan.X);
+
+					// remap the material index to the list we have on the DI
+					TArray<uint8> PhysMatRemap;
+					for (const UPhysicalMaterial* ComponentMaterial : CollisionComponent->CookedPhysicalMaterials)
 					{
-						const FIntPoint SectionBase = (Region - ResourceKey.MinCaptureRegion) * ComponentQuadCount;
-						CollisionComponent->FillHeightTile(*HeightValues, SectionBase.X + SectionBase.Y * CaptureVertexSpan.X, CaptureVertexSpan.X);
+						const int32 RemapIndex = ResourceKey.PhysicalMaterials.IndexOfByKey(ComponentMaterial);
+						ensure(RemapIndex <= TNumericLimits<uint8>::Max());
+						PhysMatRemap.Emplace(uint8(RemapIndex));
 					}
 
-					if (PhysMatValues)
+					for (int32 Y = 0; Y < ComponentQuadCount; ++Y)
 					{
-						const FIntPoint SectionBase = (Region - ResourceKey.MinCaptureRegion) * ComponentQuadCount;
-						CollisionComponent->FillMaterialIndexTile(*PhysMatValues, SectionBase.X + SectionBase.Y * CaptureVertexSpan.X, CaptureVertexSpan.X);
-
-						// remap the material index to the list we have on the DI
-						TArray<uint8> PhysMatRemap;
-						for (const UPhysicalMaterial* ComponentMaterial : CollisionComponent->CookedPhysicalMaterials)
+						for (int32 X = 0; X < ComponentQuadCount; ++X)
 						{
-							const int32 RemapIndex = ResourceKey.PhysicalMaterials.IndexOfByKey(ComponentMaterial);
-							ensure(RemapIndex <= TNumericLimits<uint8>::Max());
-							PhysMatRemap.Emplace(uint8(RemapIndex));
-						}
-
-						for (int32 Y = 0; Y < ComponentQuadCount; ++Y)
-						{
-							for (int32 X = 0; X < ComponentQuadCount; ++X)
-							{
-								const int32 WriteIndex = SectionBase.X + X + (SectionBase.Y + Y) * CaptureVertexSpan.X;
-								uint8& PhysMatIndex = (*PhysMatValues)[WriteIndex];
-								PhysMatIndex = PhysMatRemap.IsValidIndex(PhysMatIndex) ? PhysMatRemap[PhysMatIndex] : INDEX_NONE;
-							}
+							const int32 WriteIndex = SectionBase.X + X + (SectionBase.Y + Y) * CaptureVertexSpan.X;
+							uint8& PhysMatIndex = (*PhysMatValues)[WriteIndex];
+							PhysMatIndex = PhysMatRemap.IsValidIndex(PhysMatIndex) ? PhysMatRemap[PhysMatIndex] : INDEX_NONE;
 						}
 					}
 				}
@@ -1007,6 +1022,8 @@ FNDI_GeneratedData::TypeHash FNDI_Landscape_GeneratedData::GetTypeHash()
 
 FNDI_Landscape_SharedResourceHandle FNDI_Landscape_GeneratedData::GetLandscapeData(const UNiagaraDataInterfaceLandscape& LandscapeDI, const FNiagaraSystemInstance& SystemInstance, const FNDILandscapeData_GameThread& InstanceData, FNDI_SharedResourceUsage Usage, bool bNeedsDataImmediately)
 {
+	using namespace NiagaraDataInterfaceLandscape;
+
 	check(IsInGameThread());
 
 	const ALandscape* Landscape = InstanceData.Landscape.Get();
@@ -1074,37 +1091,86 @@ FNDI_Landscape_SharedResourceHandle FNDI_Landscape_GeneratedData::GetLandscapeDa
 		Key.PhysicalMaterials.Emplace(Material);
 	}
 
-	auto AddRegion = [&](const FIntPoint& Region)
+	ensureMsgf(GMaxRegionSearchCount < 0 || MaxRegionCount <= GMaxRegionSearchCount, TEXT("FNDI_Landscape_GeneratedData exceeded search count (%d:%d vs %d) for NiagaraSystem %s"),
+		MaxSystemRegionCount,
+		MaxLandscapeRegionCount,
+		GMaxRegionSearchCount,
+		*GetNameSafe(SystemInstance.GetSystem()));
+
+	auto CaptureRegion = [&](const FIntPoint& Region) -> bool
 	{
+		if (GMaxRegionCaptureCount >= 0 && Key.CapturedRegions.Num() >= GMaxRegionCaptureCount)
+		{
+			return false;
+		}
+
 		Key.CapturedRegions.Add(Region);
 		Key.MinCaptureRegion = Key.MinCaptureRegion.ComponentMin(Region);
 		Key.MaxCaptureRegion = Key.MaxCaptureRegion.ComponentMax(Region);
+		return true;
 	};
 
-	if (MaxSystemRegionCount > MaxLandscapeRegionCount)
+	bool bFailedToCaptureRegion = false;
+	int32 RegionSearchCount = 0;
+
+	auto ConditionalCaptureByComponent = [&]() -> void
 	{
 		for (const auto& LandscapeComponent : LandscapeInfo->XYtoCollisionComponentMap)
 		{
-			if (SystemRect.Contains(LandscapeComponent.Key))
+			if (GMaxRegionSearchCount >= 0 && RegionSearchCount >= GMaxRegionSearchCount)
+			{
+				return;
+			}
+
+			++RegionSearchCount;
+			if (SystemRect.Contains(LandscapeComponent.Key) && LandscapeComponent.Value)
 			{
-				AddRegion(LandscapeComponent.Key);
+				if (!CaptureRegion(LandscapeComponent.Key))
+				{
+					bFailedToCaptureRegion = true;
+					return;
+				}
 			}
 		}
-	}
-	else
+	};
+
+	auto ConditionalCaptureByRect = [&]() -> void
 	{
 		for (int32 GridY = SystemRect.Min.Y; GridY < SystemRect.Max.Y; ++GridY)
 		{
 			for (int32 GridX = SystemRect.Min.X; GridX < SystemRect.Max.X; ++GridX)
 			{
+				if (GMaxRegionSearchCount >= 0 && RegionSearchCount >= GMaxRegionSearchCount)
+				{
+					return;
+				}
+
+				++RegionSearchCount;
 				const FIntPoint CurrentRegion(GridX, GridY);
-				if (LandscapeInfo->XYtoCollisionComponentMap.Contains(CurrentRegion))
+				if (LandscapeInfo->XYtoCollisionComponentMap.FindRef(CurrentRegion) != nullptr)
 				{
-					AddRegion(CurrentRegion);
+					if (!CaptureRegion(CurrentRegion))
+					{
+						bFailedToCaptureRegion = true;
+						return;
+					}
 				}
 			}
 		}
+	};
+
+	if (MaxSystemRegionCount > MaxLandscapeRegionCount)
+	{
+		ConditionalCaptureByComponent();
 	}
+	else
+	{
+		ConditionalCaptureByRect();
+	}
+
+	ensureMsgf(!bFailedToCaptureRegion, TEXT("FNDI_Landscape_GeneratedData exceeded maximum capture count (%d) for NiagaraSystem %s"),
+		GMaxRegionCaptureCount,
+		*GetNameSafe(SystemInstance.GetSystem()));
 
 	if (!Key.CapturedRegions.Num())
 	{
