- When pasting renderers into stateless reset the bindings
- FNiagaraAttributeBinding doesn't support stateless right now so it's impossible to recover modified bindings or emitter mode
- This change also handles disabling the emitter if we have no attributes to process

#rnx

diff --git a/Engine/Plugins/FX/Niagara/Source/Niagara/Private/Stateless/NiagaraStatelessComputeManager.cpp b/Engine/Plugins/FX/Niagara/Source/Niagara/Private/Stateless/NiagaraStatelessComputeManager.cpp
--- a/Engine/Plugins/FX/Niagara/Source/Niagara/Private/Stateless/NiagaraStatelessComputeManager.cpp
+++ b/Engine/Plugins/FX/Niagara/Source/Niagara/Private/Stateless/NiagaraStatelessComputeManager.cpp
@@ -94,6 +94,7 @@ namespace NiagaraStatelessComputeManagerPrivate
 		}
 
 		FNiagaraEmptyUAVPoolScopedAccess UAVPoolAccessScope(ComputeInterface->GetEmptyUAVPool());
+		FRHIUnorderedAccessView* EmptyFloatBufferUAV = ComputeInterface->GetEmptyUAVFromPool(RHICmdList, PF_R32_FLOAT, ENiagaraEmptyUAVType::Buffer);
 		FRHIUnorderedAccessView* EmptyIntBufferUAV = ComputeInterface->GetEmptyUAVFromPool(RHICmdList, PF_R32_SINT, ENiagaraEmptyUAVType::Buffer);
 
 		// Execute Simulations
@@ -124,7 +125,7 @@ namespace NiagaraStatelessComputeManagerPrivate
 			ShaderParameters->Common_SimulationInvDeltaTime = EmitterInstance->DeltaTime > 0.0f ? (1.0f / EmitterInstance->DeltaTime) : 0.0f;
 			ShaderParameters->Common_OutputBufferStride = DestinationData->GetFloatStride() / sizeof(float);
 			ShaderParameters->Common_GPUCountBufferOffset = DestinationData->GetGPUInstanceCountBufferOffset();
-			ShaderParameters->Common_FloatOutputBuffer = DestinationData->GetGPUBufferFloat().UAV;
+			ShaderParameters->Common_FloatOutputBuffer = DestinationData->GetGPUBufferFloat().UAV.IsValid() ? DestinationData->GetGPUBufferFloat().UAV.GetReference() : EmptyFloatBufferUAV;
 			//ShaderParameters->Common_HalfOutputBuffer		= DestinationData->GetGPUBufferHalf().UAV;
 			ShaderParameters->Common_IntOutputBuffer = DestinationData->GetGPUBufferInt().UAV.IsValid() ? DestinationData->GetGPUBufferInt().UAV.GetReference() : EmptyIntBufferUAV;
 			ShaderParameters->Common_GPUCountBuffer = CountBufferUAV;
diff --git a/Engine/Plugins/FX/Niagara/Source/Niagara/Private/Stateless/NiagaraStatelessEmitter.cpp b/Engine/Plugins/FX/Niagara/Source/Niagara/Private/Stateless/NiagaraStatelessEmitter.cpp
--- a/Engine/Plugins/FX/Niagara/Source/Niagara/Private/Stateless/NiagaraStatelessEmitter.cpp
+++ b/Engine/Plugins/FX/Niagara/Source/Niagara/Private/Stateless/NiagaraStatelessEmitter.cpp
@@ -64,6 +64,12 @@ namespace NiagaraStatelessInternal
 			return false;
 		}
 
+		// Validate we have output components if we don't there's no point simulating the emitter
+		if (!EmitterData.ParticleDataSetCompiledData.Get() || !EmitterData.ParticleDataSetCompiledData->Variables.Num())
+		{
+			return false;
+		}
+
 		// Validate the shader is correct
 		const FShaderParametersMetadata* ShaderParametersMetadata = EmitterData.GetShaderParametersMetadata();
 		if (!EmitterData.GetShader().IsValid() || !ShaderParametersMetadata)
@@ -632,14 +638,28 @@ void UNiagaraStatelessEmitter::AddRenderer(UNiagaraRendererProperties* Renderer,
 
 	Modify();
 	Renderer->OuterEmitterVersion = EmitterVersion;
-//	FVersionedNiagaraEmitterData* EmitterData = GetEmitterData(EmitterVersion);
 	RendererProperties.Add(Renderer);
 #if WITH_EDITOR
-//	Renderer->OnChanged().AddUObject(this, &UNiagaraEmitter::RendererChanged);
-//	UpdateChangeId(TEXT("Renderer added"));
+	// When pasting a renderer from stateful they can come with bindings which are not supported for stateless
+	// temporarily we reset them to default values.  We will need to upgrade all these calls to take in an adapter to be able to
+	// interop between different emitter types, or introduce a base emitter type.
+	if (UNiagaraRendererProperties* RendererCDO = Renderer->GetClass()->GetDefaultObject<UNiagaraRendererProperties>())
+	{
+		for (TFieldIterator<FStructProperty> PropIt(Renderer->GetClass()); PropIt; ++PropIt)
+		{
+			FStructProperty* StructProp = *PropIt;
+			if (!StructProp || !StructProp->Struct || !StructProp->Struct->IsChildOf(FNiagaraVariableAttributeBinding::StaticStruct()))
+			{
+				continue;
+			}
+			FNiagaraVariableAttributeBinding* ParameterBinding = StructProp->ContainerPtrToValuePtr<FNiagaraVariableAttributeBinding>(Renderer);
+			FNiagaraVariableAttributeBinding* ParameterBindingDefault = StructProp->ContainerPtrToValuePtr<FNiagaraVariableAttributeBinding>(RendererCDO);
+			*ParameterBinding = *ParameterBindingDefault;
+		}
+	}
+
 	OnRenderersChangedDelegate.Broadcast();
 #endif
-//	EmitterData->RebuildRendererBindings(*this);
 }
 
 void UNiagaraStatelessEmitter::RemoveRenderer(UNiagaraRendererProperties* Renderer, FGuid EmitterVersion)
