Changed NDC preload to only load exposed NDC assets.
This fixes an issue where unexposed NDC assets were loaded, causing other dependent assets (niagara systems etc.) to load, which might not be mounted

#rnx

diff --git a/Engine/Plugins/FX/Niagara/Source/NiagaraEditor/Private/NiagaraEditorModule.cpp b/Engine/Plugins/FX/Niagara/Source/NiagaraEditor/Private/NiagaraEditorModule.cpp
--- a/Engine/Plugins/FX/Niagara/Source/NiagaraEditor/Private/NiagaraEditorModule.cpp
+++ b/Engine/Plugins/FX/Niagara/Source/NiagaraEditor/Private/NiagaraEditorModule.cpp
@@ -2379,14 +2379,16 @@ void FNiagaraEditorModule::OnAssetRegistryLoadComplete()
 
 	check(AssetRegistry.IsLoadingAssets() == false);
 
+	IAssetTools& AssetTools = IAssetTools::Get();
+
 	//Ensure All Data Channel Assets are loaded and available for use in editor.
 	TArray<FAssetData> AllDataChannels;
 	AssetRegistry.GetAssetsByClass(UNiagaraDataChannelAsset::StaticClass()->GetClassPathName(), AllDataChannels);
-	for (FAssetData& DataChannelAsset : AllDataChannels)
+	for (const FAssetData& DataChannelAsset : AllDataChannels)
 	{
-		if (FPackageName::GetPackageMountPoint(DataChannelAsset.PackageName.ToString()) != NAME_None)
+		if (AssetTools.IsAssetVisible(DataChannelAsset, true))
 		{
-			UNiagaraDataChannelAsset* NewAsset = Cast<UNiagaraDataChannelAsset>(DataChannelAsset.GetAsset());
+			DataChannelAsset.GetAsset();
 		}
 	}
 }
