Move read bufffer handling to PostStage vs PostSimulate to properly handle SimStages.

diff --git a/Engine/Plugins/FX/Niagara/Source/Niagara/Internal/DataInterface/NiagaraDataInterfaceDataChannelRead.h b/Engine/Plugins/FX/Niagara/Source/Niagara/Internal/DataInterface/NiagaraDataInterfaceDataChannelRead.h
--- a/Engine/Plugins/FX/Niagara/Source/Niagara/Internal/DataInterface/NiagaraDataInterfaceDataChannelRead.h
+++ b/Engine/Plugins/FX/Niagara/Source/Niagara/Internal/DataInterface/NiagaraDataInterfaceDataChannelRead.h
@@ -303,6 +303,7 @@ struct FNiagaraDataInterfaceProxy_DataChannelRead : public FNiagaraDataInterface
 	virtual void GetDispatchArgs(const FNDIGpuComputeDispatchArgsGenContext& Context) override;
 
 	virtual void PreStage(const FNDIGpuComputePreStageContext& Context)override;
+	virtual void PostStage(const FNDIGpuComputePostStageContext& Context)override;
 	virtual void PostSimulate(const FNDIGpuComputePostSimulateContext& Context)override;
 
 	/** Persistent per instance data on the RT. Constructed when consuming data passed from GT->RT. */
diff --git a/Engine/Plugins/FX/Niagara/Source/Niagara/Private/DataInterface/NiagaraDataInterfaceDataChannelRead.cpp b/Engine/Plugins/FX/Niagara/Source/Niagara/Private/DataInterface/NiagaraDataInterfaceDataChannelRead.cpp
--- a/Engine/Plugins/FX/Niagara/Source/Niagara/Private/DataInterface/NiagaraDataInterfaceDataChannelRead.cpp
+++ b/Engine/Plugins/FX/Niagara/Source/Niagara/Private/DataInterface/NiagaraDataInterfaceDataChannelRead.cpp
@@ -2082,13 +2082,22 @@ void FNiagaraDataInterfaceProxy_DataChannelRead::PreStage(const FNDIGpuComputePr
 	}
 }
 
-void FNiagaraDataInterfaceProxy_DataChannelRead::PostSimulate(const FNDIGpuComputePostSimulateContext& Context)
+void FNiagaraDataInterfaceProxy_DataChannelRead::PostStage(const FNDIGpuComputePostStageContext& Context)
 {
 	FNiagaraDataInterfaceProxy_DataChannelRead::FInstanceData* InstanceData = SystemInstancesToProxyData_RT.Find(Context.GetSystemInstanceID());
-	if(InstanceData && InstanceData->ChannelDataRTProxy)
+	if (InstanceData && InstanceData->ChannelDataRTProxy)
 	{
 		InstanceData->GPUBuffer = nullptr;
 		InstanceData->ChannelDataRTProxy->EndReadAccess(Context.GetGraphBuilder(), InstanceData->bReadPrevFrame == false);
+	}
+}
+
+void FNiagaraDataInterfaceProxy_DataChannelRead::PostSimulate(const FNDIGpuComputePostSimulateContext& Context)
+{
+	FNiagaraDataInterfaceProxy_DataChannelRead::FInstanceData* InstanceData = SystemInstancesToProxyData_RT.Find(Context.GetSystemInstanceID());
+	if(InstanceData && InstanceData->ChannelDataRTProxy)
+	{
+		check(InstanceData->GPUBuffer == nullptr);
 
 		if (Context.IsFinalPostSimulate())
 		{
