- Fix crash when running both CPU & GPU paths for stateless emitters in a single frame and bindings are used
- The GPU path would release the parameter data so it wasn't available for the CPU path
- We now keep the CPU data rather than release it when changed

#rnx

diff --git a/Engine/Plugins/FX/Niagara/Source/Niagara/Internal/Stateless/NiagaraStatelessEmitterInstance.h b/Engine/Plugins/FX/Niagara/Source/Niagara/Internal/Stateless/NiagaraStatelessEmitterInstance.h
--- a/Engine/Plugins/FX/Niagara/Source/Niagara/Internal/Stateless/NiagaraStatelessEmitterInstance.h
+++ b/Engine/Plugins/FX/Niagara/Source/Niagara/Internal/Stateless/NiagaraStatelessEmitterInstance.h
@@ -35,7 +35,8 @@ namespace NiagaraStateless
 		TArray<FNiagaraStatelessRuntimeSpawnInfo>				SpawnInfos;
 		TUniquePtr<NiagaraStateless::FCommonShaderParameters, FShaderParametersDeletor>	ShaderParameters;
 
-		mutable TOptional<TArray<uint8>>						BindingBufferData;
+		mutable bool											bBindingBufferDirty = true;
+		mutable TArray<uint8>									BindingBufferData;
 		mutable FReadBuffer										BindingBuffer;
 
 		// Begin: INiagaraComputeDataBufferInterface
diff --git a/Engine/Plugins/FX/Niagara/Source/Niagara/Private/Stateless/NiagaraStatelessComputeManager.cpp b/Engine/Plugins/FX/Niagara/Source/Niagara/Private/Stateless/NiagaraStatelessComputeManager.cpp
--- a/Engine/Plugins/FX/Niagara/Source/Niagara/Private/Stateless/NiagaraStatelessComputeManager.cpp
+++ b/Engine/Plugins/FX/Niagara/Source/Niagara/Private/Stateless/NiagaraStatelessComputeManager.cpp
@@ -108,14 +108,18 @@ namespace NiagaraStatelessComputeManagerPrivate
 			FNiagaraDataBuffer* DestinationData = GenerationRequest.DestinationData;
 
 			// Do we need to update the parameter buffer?
-			if (EmitterInstance->BindingBufferData.IsSet())
+			if (EmitterInstance->bBindingBufferDirty)
 			{
+				EmitterInstance->bBindingBufferDirty = false;
 				EmitterInstance->BindingBuffer.Release();
-				EmitterInstance->BindingBuffer.Initialize(RHICmdList, TEXT("FNiagaraStatelessEmitterInstance::BindingBuffer"), sizeof(uint32), EmitterInstance->BindingBufferData->Num() / sizeof(uint32), EPixelFormat::PF_R32_UINT, EBufferUsageFlags::Static);
-				void* LockedBuffer = RHICmdList.LockBuffer(EmitterInstance->BindingBuffer.Buffer, 0, EmitterInstance->BindingBuffer.NumBytes, RLM_WriteOnly);
-				FMemory::Memcpy(LockedBuffer, EmitterInstance->BindingBufferData->GetData(), EmitterInstance->BindingBuffer.NumBytes);
-				RHICmdList.UnlockBuffer(EmitterInstance->BindingBuffer.Buffer);
-				EmitterInstance->BindingBufferData.Reset();
+
+				if (EmitterInstance->BindingBufferData.Num())
+				{
+					EmitterInstance->BindingBuffer.Initialize(RHICmdList, TEXT("FNiagaraStatelessEmitterInstance::BindingBuffer"), sizeof(uint32), EmitterInstance->BindingBufferData.Num() / sizeof(uint32), EPixelFormat::PF_R32_UINT, EBufferUsageFlags::Static);
+					void* LockedBuffer = RHICmdList.LockBuffer(EmitterInstance->BindingBuffer.Buffer, 0, EmitterInstance->BindingBuffer.NumBytes, RLM_WriteOnly);
+					FMemory::Memcpy(LockedBuffer, EmitterInstance->BindingBufferData.GetData(), EmitterInstance->BindingBuffer.NumBytes);
+					RHICmdList.UnlockBuffer(EmitterInstance->BindingBuffer.Buffer);
+				}
 			}
 
 			// Update parameters for this compute invocation
@@ -237,7 +241,7 @@ FNiagaraDataBuffer* FNiagaraStatelessComputeManager::GetDataBuffer(FRHICommandLi
 	{
 		case EComputeExecutionPath::CPU:
 		{
-			FParticleSimulationContext ParticleSimulation(EmitterData, EmitterInstance->ShaderParameters.Get(), EmitterInstance->BindingBufferData.Get(TArray<uint8>()));
+			FParticleSimulationContext ParticleSimulation(EmitterData, EmitterInstance->ShaderParameters.Get(), EmitterInstance->BindingBufferData);
 			ParticleSimulation.SimulateGPU(RHICmdList, EmitterInstance->RandomSeed, EmitterInstance->Age, EmitterInstance->DeltaTime, EmitterInstance->SpawnInfos, CacheData->DataBuffer);
 			if (ParticleSimulation.GetNumInstances() == 0)
 			{
@@ -300,7 +304,7 @@ void FNiagaraStatelessComputeManager::GenerateDataBufferForDebugging(FRHICommand
 	{
 		case EComputeExecutionPath::CPU:
 		{
-			FParticleSimulationContext ParticleSimulation(EmitterData, EmitterInstance->ShaderParameters.Get(), EmitterInstance->BindingBufferData.Get(TArray<uint8>()));
+			FParticleSimulationContext ParticleSimulation(EmitterData, EmitterInstance->ShaderParameters.Get(), EmitterInstance->BindingBufferData);
 			ParticleSimulation.Simulate(EmitterInstance->RandomSeed, EmitterInstance->Age, EmitterInstance->DeltaTime, EmitterInstance->SpawnInfos, DataBuffer);
 			break;
 		}
diff --git a/Engine/Plugins/FX/Niagara/Source/Niagara/Private/Stateless/NiagaraStatelessEmitterInstance.cpp b/Engine/Plugins/FX/Niagara/Source/Niagara/Private/Stateless/NiagaraStatelessEmitterInstance.cpp
--- a/Engine/Plugins/FX/Niagara/Source/Niagara/Private/Stateless/NiagaraStatelessEmitterInstance.cpp
+++ b/Engine/Plugins/FX/Niagara/Source/Niagara/Private/Stateless/NiagaraStatelessEmitterInstance.cpp
@@ -546,6 +546,7 @@ void FNiagaraStatelessEmitterInstance::SendRenderData()
 
 			if (EmitterData.bHasBindingBufferData)
 			{
+				RenderThreadData->bBindingBufferDirty = true;
 				RenderThreadData->BindingBufferData = MoveTemp(EmitterData.BindingBufferData);
 			}
 
