Fix for new compilation mode not properly handling partial failures in GPU compute shaders (where a subset of the permutations failed).
#rnx
#tests Compile asset (Grid2D_SW_UpdateCollisions[REDACTED_HANDLE]) with the new compilation mode

diff --git a/Engine/Plugins/FX/Niagara/Source/NiagaraEditor/Private/NiagaraCompilationTasks.cpp b/Engine/Plugins/FX/Niagara/Source/NiagaraEditor/Private/NiagaraCompilationTasks.cpp
--- a/Engine/Plugins/FX/Niagara/Source/NiagaraEditor/Private/NiagaraCompilationTasks.cpp
+++ b/Engine/Plugins/FX/Niagara/Source/NiagaraEditor/Private/NiagaraCompilationTasks.cpp
@@ -55,7 +55,7 @@ namespace NiagaraSystemCookStats
 
 namespace NiagaraCompilationTasksImpl
 {
-	static const FGuid UE_NIAGARA_ASYNC_TASK_COMPILER_VER = FGuid(0xB2E44F66, 0x59164A73, 0x8D687867, 0xFBBBB30B);
+	static const FGuid UE_NIAGARA_ASYNC_TASK_COMPILER_VER = FGuid(0x0E657251, 0xB787406B, 0x8DE09B90, 0x817ABFB9);
 	
 	static UE::DerivedData::FCacheBucket NiagaraDDCBucket("NiagaraScript");
 
diff --git a/Engine/Plugins/FX/Niagara/Source/NiagaraEditor/Private/NiagaraCompiler.cpp b/Engine/Plugins/FX/Niagara/Source/NiagaraEditor/Private/NiagaraCompiler.cpp
--- a/Engine/Plugins/FX/Niagara/Source/NiagaraEditor/Private/NiagaraCompiler.cpp
+++ b/Engine/Plugins/FX/Niagara/Source/NiagaraEditor/Private/NiagaraCompiler.cpp
@@ -2487,11 +2487,18 @@ bool FNiagaraShaderMapCompiler::ProcessCompileResults(bool bWait)
 
 	for (TArray<FActiveCompilation>::TIterator CompileIt = ActiveCompilations.CreateIterator(); CompileIt; ++CompileIt)
 	{
-		// make sure that all of the shader compile jobs have been released and finalized
-		const bool bReadyToProcess = !CompileIt->ShaderCompileJobs.ContainsByPredicate([](const FShaderCommonCompileJobPtr& CompileJob) -> bool
+		auto IsCompileJobIncomplete = [](const FShaderCommonCompileJobPtr& CompileJob) -> bool
+		{
+			return !CompileJob.IsValid() || !CompileJob->bReleased || !CompileJob->bFinalized;
+		};
+
+		auto IsCompileJobError = [](const FShaderCommonCompileJobPtr& CompileJob) -> bool
 		{
-			return !CompileJob->bReleased || !CompileJob->bFinalized;
-		});
+			return !CompileJob->bSucceeded;
+		};
+
+		// make sure that all of the shader compile jobs have been released and finalized
+		const bool bReadyToProcess = !CompileIt->ShaderCompileJobs.ContainsByPredicate(IsCompileJobIncomplete);
 
 		if (!bReadyToProcess)
 		{
@@ -2504,19 +2511,26 @@ bool FNiagaraShaderMapCompiler::ProcessCompileResults(bool bWait)
 		FActiveCompilation& CurrentCompilation = *CompileIt;
 		FCompletedCompilation& CompletedCompilation = CompletedCompilations.AddDefaulted_GetRef();
 
-		// for now we'll process all shaders at once (need to measure the cost here)
-		for (const FShaderCommonCompileJobPtr& ShaderCompileJob : CurrentCompilation.ShaderCompileJobs)
+		// do a first pass over all of the ShaderCompileJobs to see if any of them failed.  If it did, then we don't need
+		// to worry about pushing out our incomplete ShaderMap and we should just report the errors
+		const bool bSuccessfulCompilation = !CurrentCompilation.ShaderCompileJobs.ContainsByPredicate(IsCompileJobError);
+
+		if (bSuccessfulCompilation)
 		{
-			if (ShaderCompileJob.IsValid() && ShaderCompileJob->bSucceeded)
+			// for now we'll process all shaders at once (need to measure the cost here)
+			for (const FShaderCommonCompileJobPtr& ShaderCompileJob : CurrentCompilation.ShaderCompileJobs)
 			{
 				CurrentCompilation.ShaderMap->ProcessAndFinalizeShaderCompileJob(ShaderCompileJob);
 			}
-			else
-			{
-				CurrentCompilation.ShaderMap->SetCompiledSuccessfully(false);
-			}
+		}
+		else
+		{
+			CurrentCompilation.ShaderMap->SetCompiledSuccessfully(false);
+		}
 
-			// pass on error/warning info
+		// pass on error/warning info
+		for (const FShaderCommonCompileJobPtr& ShaderCompileJob : CurrentCompilation.ShaderCompileJobs)
+		{
 			if (const FShaderCompileJob* SingleShaderJob = ShaderCompileJob->GetSingleShaderJob())
 			{
 				CompletedCompilation.CompilationErrors.Append(SingleShaderJob->Output.Errors);
