- Fix for not assigning a new Guid when pasting spawn infos
- Fix stack not refreshing when pasting with a spawn info selected

#rnx

diff --git a/Engine/Plugins/FX/Niagara/Source/NiagaraEditor/Private/ViewModels/Stack/NiagaraStackStatelessEmitterSpawnGroup.cpp b/Engine/Plugins/FX/Niagara/Source/NiagaraEditor/Private/ViewModels/Stack/NiagaraStackStatelessEmitterSpawnGroup.cpp
--- a/Engine/Plugins/FX/Niagara/Source/NiagaraEditor/Private/ViewModels/Stack/NiagaraStackStatelessEmitterSpawnGroup.cpp
+++ b/Engine/Plugins/FX/Niagara/Source/NiagaraEditor/Private/ViewModels/Stack/NiagaraStackStatelessEmitterSpawnGroup.cpp
@@ -75,7 +75,9 @@ namespace NiagaraStackStatelessEmitterSpawnGroupPrivate
 					StatelessEmitter->Modify();
 				}
 
-				StatelessEmitter->AddSpawnInfo() = NewSpawnInfo.GetValue();
+				FNiagaraStatelessSpawnInfo& SpawnInfo = StatelessEmitter->AddSpawnInfo();
+				SpawnInfo = NewSpawnInfo.GetValue();
+				SpawnInfo.SourceId = FGuid::NewGuid();
 			}
 
 			if (bHasPastedValues)
@@ -330,8 +332,11 @@ void UNiagaraStackStatelessEmitterSpawnItem::Paste(const UNiagaraClipboardConten
 	UNiagaraStatelessEmitter* StatelessEmitter = GetStatelessEmitter();
 	if (NiagaraStackStatelessEmitterSpawnGroupPrivate::Paste(StatelessEmitter, ClipboardContent))
 	{
-		OnDataObjectModified().Broadcast({ StatelessEmitter }, ENiagaraDataObjectChange::Changed);
-		RefreshChildren();
+		if (UNiagaraStackStatelessEmitterSpawnGroup* SpawnGroup = GetTypedOuter<UNiagaraStackStatelessEmitterSpawnGroup>())
+		{
+			SpawnGroup->OnDataObjectModified().Broadcast({ StatelessEmitter }, ENiagaraDataObjectChange::Changed);
+			SpawnGroup->RefreshChildren();
+		}
 	}
 }
 
