- Warn about materials not having the correct 'used with' flags for sprite / mesh / ribbon renderers
- Auto setting used with flags is not guaranteed and it can be confusing as to why the default material shows vs the selected material, this will give the user some information
- Note that there is not a 'fix' button as cooked editor builds are not guaranteed to be able to fix, and detecting is not always possible

#rnx

diff --git a/Engine/Plugins/FX/Niagara/Source/Niagara/Private/NiagaraMeshRendererProperties.cpp b/Engine/Plugins/FX/Niagara/Source/Niagara/Private/NiagaraMeshRendererProperties.cpp
--- a/Engine/Plugins/FX/Niagara/Source/Niagara/Private/NiagaraMeshRendererProperties.cpp
+++ b/Engine/Plugins/FX/Niagara/Source/Niagara/Private/NiagaraMeshRendererProperties.cpp
@@ -1122,11 +1122,12 @@ void UNiagaraMeshRendererProperties::GetRendererTooltipWidgets(const FNiagaraEmi
 	}
 }
 
-
 void UNiagaraMeshRendererProperties::GetRendererFeedback(const FVersionedNiagaraEmitter& InEmitter, TArray<FNiagaraRendererFeedback>& OutErrors, TArray<FNiagaraRendererFeedback>& OutWarnings, TArray<FNiagaraRendererFeedback>& OutInfo) const
 {
 	Super::GetRendererFeedback(InEmitter, OutErrors, OutWarnings, OutInfo);
 
+	GetMaterialUsageFeedback(MATUSAGE_NiagaraMeshParticles, OutWarnings);
+
 	if (MaterialParameters.HasAnyBindings())
 	{
 		TArray<UMaterialInterface*> Materials;
diff --git a/Engine/Plugins/FX/Niagara/Source/Niagara/Private/NiagaraRendererProperties.cpp b/Engine/Plugins/FX/Niagara/Source/Niagara/Private/NiagaraRendererProperties.cpp
--- a/Engine/Plugins/FX/Niagara/Source/Niagara/Private/NiagaraRendererProperties.cpp
+++ b/Engine/Plugins/FX/Niagara/Source/Niagara/Private/NiagaraRendererProperties.cpp
@@ -753,6 +753,33 @@ void UNiagaraRendererProperties::GetRendererFeedback(const FVersionedNiagaraEmit
 	}
 }
 
+void UNiagaraRendererProperties::GetMaterialUsageFeedback(EMaterialUsage Usage, TArray<FNiagaraRendererFeedback>& OutFeedback) const
+{
+	TArray<UMaterialInterface*> Materials;
+	GetUsedMaterials(nullptr, Materials);
+
+	FString FailedMaterialNames;
+	for (UMaterialInterface* Material : Materials)
+	{
+		if (Material && !Material->CheckMaterialUsage(Usage))
+		{
+			if (FailedMaterialNames.Len() > 0)
+			{
+				FailedMaterialNames.Append(TEXT(", "));
+			}
+			FailedMaterialNames.Append(Material->GetName());
+		}
+	}
+
+	if (FailedMaterialNames.Len() > 0)
+	{
+		OutFeedback.Emplace(
+			FText::Format(LOCTEXT("InvalidMaterialUsage_Desc", "Some materials '{0}' do not have the correct usage flags set, and will use the default material."), FText::FromString(FailedMaterialNames)),
+			LOCTEXT("InvalidMaterialUsage_Summary", "Materials might not render correctly.")
+		);
+	}
+}
+
 const FSlateBrush* UNiagaraRendererProperties::GetStackIcon() const
 {
 	return FSlateIconFinder::FindIconBrushForClass(GetClass());
diff --git a/Engine/Plugins/FX/Niagara/Source/Niagara/Private/NiagaraRibbonRendererProperties.cpp b/Engine/Plugins/FX/Niagara/Source/Niagara/Private/NiagaraRibbonRendererProperties.cpp
--- a/Engine/Plugins/FX/Niagara/Source/Niagara/Private/NiagaraRibbonRendererProperties.cpp
+++ b/Engine/Plugins/FX/Niagara/Source/Niagara/Private/NiagaraRibbonRendererProperties.cpp
@@ -656,6 +656,8 @@ void UNiagaraRibbonRendererProperties::GetRendererFeedback(const FVersionedNiaga
 {
 	Super::GetRendererFeedback(InEmitter, OutErrors, OutWarnings, OutInfo);
 
+	GetMaterialUsageFeedback(MATUSAGE_NiagaraRibbons, OutWarnings);
+
 	const FVersionedNiagaraEmitterData* EmitterData = InEmitter.GetEmitterData();
 
 	// If this is a renderer made before we used UniqueID give the user an option to switch into the new mode
diff --git a/Engine/Plugins/FX/Niagara/Source/Niagara/Private/NiagaraSpriteRendererProperties.cpp b/Engine/Plugins/FX/Niagara/Source/Niagara/Private/NiagaraSpriteRendererProperties.cpp
--- a/Engine/Plugins/FX/Niagara/Source/Niagara/Private/NiagaraSpriteRendererProperties.cpp
+++ b/Engine/Plugins/FX/Niagara/Source/Niagara/Private/NiagaraSpriteRendererProperties.cpp
@@ -618,6 +618,7 @@ void UNiagaraSpriteRendererProperties::GetRendererWidgets(const FNiagaraEmitterI
 void UNiagaraSpriteRendererProperties::GetRendererFeedback(const FVersionedNiagaraEmitter& InEmitter, TArray<FText>& OutErrors, TArray<FText>& OutWarnings, TArray<FText>& OutInfo) const
 {
 	Super::GetRendererFeedback(InEmitter, OutErrors, OutWarnings, OutInfo);
+
 	if (InEmitter.GetEmitterData()->SpawnScriptProps.Script->GetVMExecutableData().IsValid())
 	{
 		if (bUseMaterialCutoutTexture || CutoutTexture)
@@ -639,6 +640,8 @@ void UNiagaraSpriteRendererProperties::GetRendererFeedback(const FVersionedNiaga
 {
 	Super::GetRendererFeedback(InEmitter, OutErrors, OutWarnings, OutInfo);
 
+	GetMaterialUsageFeedback(MATUSAGE_NiagaraSprites, OutWarnings);
+
 	if (MaterialParameters.HasAnyBindings())
 	{
 		TArray<UMaterialInterface*> Materials;
diff --git a/Engine/Plugins/FX/Niagara/Source/Niagara/Public/NiagaraRendererProperties.h b/Engine/Plugins/FX/Niagara/Source/Niagara/Public/NiagaraRendererProperties.h
--- a/Engine/Plugins/FX/Niagara/Source/Niagara/Public/NiagaraRendererProperties.h
+++ b/Engine/Plugins/FX/Niagara/Source/Niagara/Public/NiagaraRendererProperties.h
@@ -27,6 +27,7 @@ struct FNiagaraDataSetCompiledData;
 struct FSlateBrush;
 struct FStaticParameterSet;
 struct FStreamingRenderAssetPrimitiveInfo;
+enum EMaterialUsage : int;
 
 #if WITH_EDITOR
 // Helper class for GUI error handling
@@ -371,6 +372,8 @@ public:
 	virtual void GetRendererFeedback(const FVersionedNiagaraEmitter& InEmitter, TArray<FText>& OutErrors, TArray<FText>& OutWarnings, TArray<FText>& OutInfo) const {}
 	NIAGARA_API virtual void GetRendererFeedback(const FVersionedNiagaraEmitter& InEmitter, TArray<FNiagaraRendererFeedback>& OutErrors, TArray<FNiagaraRendererFeedback>& OutWarnings, TArray<FNiagaraRendererFeedback>& OutInfo) const;
 
+	void GetMaterialUsageFeedback(EMaterialUsage Usage, TArray<FNiagaraRendererFeedback>& OutFeedback) const;
+
 	// The icon to display in the niagara stack widget under the renderer section
 	NIAGARA_API virtual const FSlateBrush* GetStackIcon() const;
 
