- Add some stats for stateless, i.e. number of jobs, memory usage, etc.

#rnx
#tests use 'stat niagarastatess' to see the data in the dev build

diff --git a/Engine/Plugins/FX/Niagara/Source/Niagara/Classes/NiagaraDataSet.h b/Engine/Plugins/FX/Niagara/Source/Niagara/Classes/NiagaraDataSet.h
--- a/Engine/Plugins/FX/Niagara/Source/Niagara/Classes/NiagaraDataSet.h
+++ b/Engine/Plugins/FX/Niagara/Source/Niagara/Classes/NiagaraDataSet.h
@@ -188,6 +188,11 @@ public:
 
 	NIAGARA_API void BuildRegisterTable();
 
+#if NIAGARA_MEMORY_TRACKING
+	int32 GetAllocationSizeBytes() const { return AllocationSizeBytes; }
+#endif
+
+
 private:
 	NIAGARA_API FORCEINLINE void CheckUsage(bool bReadOnly)const;
 
@@ -355,8 +360,8 @@ public:
 	const FNiagaraDataSetCompiledData& GetCompiledData() const { check(CompiledData.Get() != nullptr); return *CompiledData.Get(); }
 
 	int NumSpawnedIDs;
-private:
 
+private:
 	NIAGARA_API void Reset();
 
 	NIAGARA_API void ResetBuffersInternal();
diff --git a/Engine/Plugins/FX/Niagara/Source/Niagara/Private/Stateless/NiagaraStatelessComputeManager.cpp b/Engine/Plugins/FX/Niagara/Source/Niagara/Private/Stateless/NiagaraStatelessComputeManager.cpp
--- a/Engine/Plugins/FX/Niagara/Source/Niagara/Private/Stateless/NiagaraStatelessComputeManager.cpp
+++ b/Engine/Plugins/FX/Niagara/Source/Niagara/Private/Stateless/NiagaraStatelessComputeManager.cpp
@@ -16,6 +16,16 @@
 #include "RenderGraphUtils.h"
 #include "SceneView.h"
 
+DECLARE_STATS_GROUP(TEXT("NiagaraStateless"), STATGROUP_NiagaraStateless, STATCAT_NiagaraStateless);
+
+DECLARE_CYCLE_STAT(TEXT("GetDataBuffer"), STAT_NiagaraStateless_GetDataBuffer, STATGROUP_NiagaraStateless);
+DECLARE_CYCLE_STAT(TEXT("GenerateGPUData"), STAT_NiagaraStateless_GenerateGPUData, STATGROUP_NiagaraStateless);
+DECLARE_DWORD_COUNTER_STAT(TEXT("CPU Simulate"), STAT_NiagaraStateless_CPUSimulate, STATGROUP_NiagaraStateless);
+DECLARE_DWORD_COUNTER_STAT(TEXT("GPU Simulate"), STAT_NiagaraStateless_GPUSimulate, STATGROUP_NiagaraStateless);
+DECLARE_DWORD_COUNTER_STAT(TEXT("CPU Particle Count"), STAT_NiagaraStateless_CPUParticleCount, STATGROUP_NiagaraStateless);
+DECLARE_DWORD_COUNTER_STAT(TEXT("GPU Particle Estimate"), STAT_NiagaraStateless_GPUParticleEstimate, STATGROUP_NiagaraStateless);
+DECLARE_MEMORY_STAT(TEXT("Buffer Memory"), STAT_NiagaraStateless_BufferMemory, STATGROUP_NiagaraStateless);
+
 namespace NiagaraStatelessComputeManagerPrivate
 {
 	enum class EComputeExecutionPath
@@ -59,6 +69,8 @@ namespace NiagaraStatelessComputeManagerPrivate
 	void GenerateGPUData(FRHICommandList& RHICmdList, FNiagaraGpuComputeDispatchInterface* ComputeInterface, TConstArrayView<FNiagaraStatelessComputeManager::FStatelessDataGenerationRequest> GenerationRequests)
 	{
 		const int32 NumJobs = GenerationRequests.Num();
+		INC_DWORD_STAT_BY(STAT_NiagaraStateless_GPUSimulate, NumJobs);
+		SCOPE_CYCLE_COUNTER(STAT_NiagaraStateless_GenerateGPUData);
 
 		// Get Count Buffer
 		FNiagaraGPUInstanceCountManager& CountManager = ComputeInterface->GetGPUInstanceCounterManager();
@@ -160,13 +172,22 @@ FNiagaraStatelessComputeManager::FNiagaraStatelessComputeManager(FNiagaraGpuComp
 
 FNiagaraStatelessComputeManager::~FNiagaraStatelessComputeManager()
 {
+#if NIAGARA_MEMORY_TRACKING && STATS
+	for (const TUniquePtr<FStatelessDataCache>& CacheData : FreeData)
+	{
+		DEC_MEMORY_STAT_BY(STAT_NiagaraStateless_BufferMemory, CacheData->DataBuffer->GetAllocationSizeBytes());
+	}
+	ensure(UsedData.Num() == 0);
+#endif
 }
 
 FNiagaraDataBuffer* FNiagaraStatelessComputeManager::GetDataBuffer(FRHICommandListBase& RHICmdList, uintptr_t EmitterKey, const NiagaraStateless::FEmitterInstance_RT* EmitterInstance)
 {
 	using namespace NiagaraStateless;
 	using namespace NiagaraStatelessComputeManagerPrivate;
-	
+
+	SCOPE_CYCLE_COUNTER(STAT_NiagaraStateless_GetDataBuffer);
+
 	//-OPT: This lock is very conservative, ideally we only have it around the relevant parts
 	UE::TScopeLock ScopeLock(GetDataBufferGuard);
 
@@ -229,8 +250,16 @@ FNiagaraDataBuffer* FNiagaraStatelessComputeManager::GetDataBuffer(FRHICommandLi
 		CacheData->DataBuffer = new FNiagaraDataBuffer(&CacheData->DataSet);
 	}
 
+#if NIAGARA_MEMORY_TRACKING && STATS
+	DEC_MEMORY_STAT_BY(STAT_NiagaraStateless_BufferMemory, CacheData->DataBuffer->GetAllocationSizeBytes());
+#endif
+
 	CacheData->DataBuffer->AllocateGPU(RHICmdList, ActiveParticles, ComputeInterface->GetFeatureLevel(), TEXT("StatelessSimBuffer"));
 
+#if NIAGARA_MEMORY_TRACKING && STATS
+	INC_MEMORY_STAT_BY(STAT_NiagaraStateless_BufferMemory, CacheData->DataBuffer->GetAllocationSizeBytes());
+#endif
+
 	// Until we add an extension to the render to notify about GDME start / end we can not allow GPU generation requests outside of PreInitViews / PreRender
 	// The shadow rendering, for example, will call GDME outside of this and it will result in crashes
 	// Therefore we fallback to CPU generation in these cases, if available
@@ -241,6 +270,8 @@ FNiagaraDataBuffer* FNiagaraStatelessComputeManager::GetDataBuffer(FRHICommandLi
 	{
 		case EComputeExecutionPath::CPU:
 		{
+			INC_DWORD_STAT_BY(STAT_NiagaraStateless_CPUSimulate, 1);
+
 			FParticleSimulationContext ParticleSimulation(EmitterData, EmitterInstance->ShaderParameters.Get(), EmitterInstance->BindingBufferData);
 			ParticleSimulation.SimulateGPU(RHICmdList, EmitterInstance->RandomSeed, EmitterInstance->Age, EmitterInstance->DeltaTime, EmitterInstance->SpawnInfos, CacheData->DataBuffer);
 			if (ParticleSimulation.GetNumInstances() == 0)
@@ -250,6 +281,7 @@ FNiagaraDataBuffer* FNiagaraStatelessComputeManager::GetDataBuffer(FRHICommandLi
 				UsedData.Emplace(EmitterKey, CacheData);
 				return nullptr;
 			}
+			INC_DWORD_STAT_BY(STAT_NiagaraStateless_CPUParticleCount, ParticleSimulation.GetNumInstances());
 			break;
 		}
 
@@ -263,6 +295,8 @@ FNiagaraDataBuffer* FNiagaraStatelessComputeManager::GetDataBuffer(FRHICommandLi
 			CountsToRelease.Add(CountOffset);
 
 			GPUGenerationRequests.Emplace(CacheData->DataBuffer, EmitterInstance, ActiveParticles);
+
+			INC_DWORD_STAT_BY(STAT_NiagaraStateless_GPUParticleEstimate, ActiveParticles);
 			break;
 		}
 
@@ -304,6 +338,8 @@ void FNiagaraStatelessComputeManager::GenerateDataBufferForDebugging(FRHICommand
 	{
 		case EComputeExecutionPath::CPU:
 		{
+			INC_DWORD_STAT_BY(STAT_NiagaraStateless_CPUSimulate, 1);
+
 			FParticleSimulationContext ParticleSimulation(EmitterData, EmitterInstance->ShaderParameters.Get(), EmitterInstance->BindingBufferData);
 			ParticleSimulation.Simulate(EmitterInstance->RandomSeed, EmitterInstance->Age, EmitterInstance->DeltaTime, EmitterInstance->SpawnInfos, DataBuffer);
 			break;
@@ -400,6 +436,13 @@ void FNiagaraStatelessComputeManager::OnPostPostRender(FRDGBuilder& GraphBuilder
 		RDG_EVENT_NAME("FNiagaraStatelessComputeManager::OnPostPostRender"),
 		[this](FRHICommandListImmediate& RHICmdList)
 		{
+		#if NIAGARA_MEMORY_TRACKING && STATS
+			for (const TUniquePtr<FStatelessDataCache>& CacheData : FreeData)
+			{
+				DEC_MEMORY_STAT_BY(STAT_NiagaraStateless_BufferMemory, CacheData->DataBuffer->GetAllocationSizeBytes());
+			}
+		#endif
+
 			FreeData.Empty(UsedData.Num());
 			for (auto it=UsedData.CreateIterator(); it; ++it)
 			{
diff --git a/Engine/Plugins/FX/Niagara/Source/Niagara/Private/Stateless/NiagaraStatelessComputeManager.h b/Engine/Plugins/FX/Niagara/Source/Niagara/Private/Stateless/NiagaraStatelessComputeManager.h
--- a/Engine/Plugins/FX/Niagara/Source/Niagara/Private/Stateless/NiagaraStatelessComputeManager.h
+++ b/Engine/Plugins/FX/Niagara/Source/Niagara/Private/Stateless/NiagaraStatelessComputeManager.h
@@ -63,8 +63,6 @@ private:
 	void OnPostPostRender(FRDGBuilder& GraphBuilder);
 
 private:
-	FNiagaraRendererLayout								SpriteStatelessLayout;
-
 	TMap<uintptr_t, TUniquePtr<FStatelessDataCache>>	UsedData;
 	TArray<TUniquePtr<FStatelessDataCache>>				FreeData;
 	TArray<uint32>										CountsToRelease;
