Taking a more belt and braces approach to NDC layout and proxy lifetimes to account for all edge cases.

UNiagaraDataChannel
Moved some layout information that is referenced by DataSets in handlers and DIs out to a utility struct that is alloacted and distributed as a shared ptr.
Avoids issue with lifetime management when the data channel object is destroyed but we still have in flight DIs and handlers.

FNiagaraDataChannelDataProxy
Converting these to be allocated and distributed as a shared ptr also.
This allows us to ensure it's lifetime in all cases with dispatcher/FXsystem/world distruction as well as NDC and DI destrution.

Potental issues this avoids -
1. NDC data is destroyed.
Both the dispatcher and DI Proxies hold refs to the NDCDataProxy.
We can submit an RT command to destroy the proxy and remove it from the dispatcher but cannot immediately and safely remove the dangling ptrs from potentially in flight DI proxies.
These will refresh next frame due to the GT side data having been destroyed but it leaves these dangling pointers temproarily in the DI Proxy and the potential for use.
I don't think we've hit this in practice yet but it's possible.

2. Simplifies and solidifies when the FScene changes and the dispatcher is destroyed.
UWorld cleanup flows can be janky. The FScene/Dispatcher on UWorlds can be swapped early etc.
We can't rely on getting a valid dispatcher ptr from the UWorld in all cases.
Using a shared ptr for the proxy allows the GT to just forget it's proxy and create a new one.
The old proxy can be destroyed from the dispatcher or the DI's RT proxy.
Any cleanup needing the dispatcher ptr can be done from the dispatcher itself which removes the previous possibility of mis-matched dispatchers when changing FScene.

The NDCDataProxy now also fully owns it's FNiagaraDataSet. Prevously this was partly managed by the GT side data which could lead to issues.
Now it's created along with the proxy but used and destroyed solely on the RT.
NDCDataProxy also keeps a reference to the layout info it needs from the main UNiagaraDataChannel.

FNiagaraDataBufferReadback

Changed helper class for reading back data from the CPU to be safer wrt data set lifetimes.
It's not safe to pass niagara buffers owned by the RT data set back to the GT.
Modified this class to pass the buffers back indivudially which are then coped over to a GT owned data buffer.

- Misc
Fixed FNiagaraDataChannelManager::RemoveDataChannel() to properly clean up the removed handler.
Island NDCs are clearing their NDCData when the islands are cleaned up. Currently recreating these anyway on acquire so this just holds onto memory longer than needed.

diff --git a/Engine/Plugins/FX/Niagara/Source/Niagara/Internal/DataInterface/NiagaraDataInterfaceDataChannelRead.h b/Engine/Plugins/FX/Niagara/Source/Niagara/Internal/DataInterface/NiagaraDataInterfaceDataChannelRead.h
--- a/Engine/Plugins/FX/Niagara/Source/Niagara/Internal/DataInterface/NiagaraDataInterfaceDataChannelRead.h
+++ b/Engine/Plugins/FX/Niagara/Source/Niagara/Internal/DataInterface/NiagaraDataInterfaceDataChannelRead.h
@@ -33,7 +33,7 @@ struct FNDIDataChannelWriteInstanceData;
 class UNiagaraDataChannelHandler;
 class FNiagaraDataBuffer;
 
-struct FNiagaraDataChannelDataProxy;
+using FNiagaraDataChannelDataProxyPtr = TSharedPtr<struct FNiagaraDataChannelDataProxy>;
 
 /** Mode controlling the behavior of the Spawn function for Niagara's Data Channel Read Data Interface.*/
 UENUM()
@@ -311,7 +311,7 @@ struct FNiagaraDataInterfaceProxy_DataChannelRead : public FNiagaraDataInterface
 	{
 		//GPU Dataset from the channel handler. We'll grab the current buffer from this on the RT.
 		//This must be grabbed fresh from the handler each frame as it's lifetime cannot be ensured.
-		FNiagaraDataChannelDataProxy* ChannelDataRTProxy = nullptr;
+		FNiagaraDataChannelDataProxyPtr ChannelDataRTProxy = nullptr;
 
 		bool bReadPrevFrame = false;
 
diff --git a/Engine/Plugins/FX/Niagara/Source/Niagara/Internal/DataInterface/NiagaraDataInterfaceDataChannelWrite.h b/Engine/Plugins/FX/Niagara/Source/Niagara/Internal/DataInterface/NiagaraDataInterfaceDataChannelWrite.h
--- a/Engine/Plugins/FX/Niagara/Source/Niagara/Internal/DataInterface/NiagaraDataInterfaceDataChannelWrite.h
+++ b/Engine/Plugins/FX/Niagara/Source/Niagara/Internal/DataInterface/NiagaraDataInterfaceDataChannelWrite.h
@@ -216,7 +216,7 @@ struct FNiagaraDataInterfaceProxy_DataChannelWrite : public FNiagaraDataInterfac
 	{
 		//GPU Dataset from the channel handler. We'll grab the current buffer from this on the RT.
 		//This must be grabbed fresh from the handler each frame as it's lifetime cannot be ensured.
-		FNiagaraDataChannelDataProxy* ChannelDataRTProxy = nullptr;
+		FNiagaraDataChannelDataProxyPtr ChannelDataRTProxy = nullptr;
 
 		/**
 		A buffer containing layout information needed to access parameters for each script using this DI.
diff --git a/Engine/Plugins/FX/Niagara/Source/Niagara/Internal/NiagaraDataChannelCommon.h b/Engine/Plugins/FX/Niagara/Source/Niagara/Internal/NiagaraDataChannelCommon.h
--- a/Engine/Plugins/FX/Niagara/Source/Niagara/Internal/NiagaraDataChannelCommon.h
+++ b/Engine/Plugins/FX/Niagara/Source/Niagara/Internal/NiagaraDataChannelCommon.h
@@ -16,7 +16,8 @@ class UNiagaraDataChannel;
 class UNiagaraDataChannelHandler;
 class UNiagaraDataInterfaceDataChannelWrite;
 class UNiagaraDataInterfaceDataChannelRead;
-struct FNiagaraDataChannelDataProxy;
+
+using FNiagaraDataChannelDataProxyPtr = TSharedPtr<struct FNiagaraDataChannelDataProxy>;
 
 /** A request to publish data into a Niagara Data Channel.  */
 struct FNiagaraDataChannelPublishRequest
@@ -81,7 +82,7 @@ struct FNiagaraDataChannelData final : public TSharedFromThis<FNiagaraDataChanne
 
 	NIAGARA_API FNiagaraDataChannelGameData* GetGameData();
 	NIAGARA_API FNiagaraDataBufferRef GetCPUData(bool bPreviousFrame);
-	FNiagaraDataChannelDataProxy* GetRTProxy(){ return RTProxy.Get(); }
+	FNiagaraDataChannelDataProxyPtr GetRTProxy(){ return RTProxy; }
 	
 	/** Adds a request to publish some data into the channel on the next tick. */
 	NIAGARA_API void Publish(const FNiagaraDataChannelPublishRequest& Request);
@@ -113,12 +114,6 @@ private:
 	/** DataChannel data accessible from Game/BP. AoS Layout. LWC types. */
 	FNiagaraDataChannelGameDataPtr GameData;
 
-	//		▲	CPU Sim Data can optionally be made visible to the Game Data.
-	//		|
-	//		|
-	//		|
-	//		▼	Game/BP Data can optionally be made visible to CPU sims.
-
 	/** DataChannel data accessible to Niagara CPU sims. SoA layout. Non LWC types. */
 	FNiagaraDataSet* CPUSimData = nullptr;
 
@@ -129,20 +124,11 @@ private:
 	/** Dataset we use for staging game data for the consumption by RT/GPU sims. */
 	FNiagaraDataSet* GameDataStaging = nullptr;
 
-	//		▲	GPU Sim Data can optionally be made visible to CPU Sims and Game Data.
-	//		|
-	//		|
-	//		|
-	//		▼	CPU Sim and Game Data can optionally be made visible to the GPU.
-
-	/** DataChannel data accessible to Niagara GPU sims. SoA layout. Non LWC types. */
-	FNiagaraDataSet* GPUSimData = nullptr;
-
 	/** Data buffers we'll be passing to the RT proxy for uploading to the GPU */
 	TArray<FNiagaraDataBufferRef> BuffersForGPU;
 
 	/** Render thread proxy for this data. Owns all RT side data meant for GPU simulations. */
-	TUniquePtr<FNiagaraDataChannelDataProxy> RTProxy;
+	FNiagaraDataChannelDataProxyPtr RTProxy;
 
 	/** Pending requests to publish data into this data channel. These requests are consumed at tick tick group. */
 	TArray<FNiagaraDataChannelPublishRequest> PublishRequests;
@@ -157,4 +143,7 @@ private:
 
 	/** Critical section protecting shared state for multiple writers publishing from different threads. */
 	FCriticalSection PublishCritSec;
+
+	//Keep reference to the layout this data was built with.
+	FNiagaraDataChannelLayoutInfoPtr LayoutInfo;
 };
diff --git a/Engine/Plugins/FX/Niagara/Source/Niagara/Private/DataInterface/NiagaraDataInterfaceDataChannelRead.cpp b/Engine/Plugins/FX/Niagara/Source/Niagara/Private/DataInterface/NiagaraDataInterfaceDataChannelRead.cpp
--- a/Engine/Plugins/FX/Niagara/Source/Niagara/Private/DataInterface/NiagaraDataInterfaceDataChannelRead.cpp
+++ b/Engine/Plugins/FX/Niagara/Source/Niagara/Private/DataInterface/NiagaraDataInterfaceDataChannelRead.cpp
@@ -297,7 +297,7 @@ namespace NDIDataChannelReadLocal
 struct FNDIDataChannelReadInstanceData_RT
 {
 	//RT proxy for game channel data from which we're reading.
-	FNiagaraDataChannelDataProxy* ChannelDataRTProxy = nullptr;
+	FNiagaraDataChannelDataProxyPtr ChannelDataRTProxy = nullptr;
 
 	bool bReadPrevFrame = false;
 
@@ -470,8 +470,8 @@ bool FNDIDataChannelReadInstanceData::Tick(UNiagaraDataInterfaceDataChannelRead*
 				}
 			}
 
-			const FNiagaraDataSetCompiledData& CPUSourceDataCompiledData = DataChannelPtr->GetDataChannel()->GetCompiledData(ENiagaraSimTarget::CPUSim);
-			const FNiagaraDataSetCompiledData& GPUSourceDataCompiledData = DataChannelPtr->GetDataChannel()->GetCompiledData(ENiagaraSimTarget::GPUComputeSim);
+			const FNiagaraDataSetCompiledData& CPUSourceDataCompiledData = DataChannelPtr->GetDataChannel()->GetLayoutInfo()->GetDataSetCompiledData();
+			const FNiagaraDataSetCompiledData& GPUSourceDataCompiledData = DataChannelPtr->GetDataChannel()->GetLayoutInfo()->GetDataSetCompiledDataGPU();
 			check(CPUSourceDataCompiledData.GetLayoutHash() && CPUSourceDataCompiledData.GetLayoutHash() == GPUSourceDataCompiledData.GetLayoutHash());			
 			uint64 SourceDataLayoutHash = CPUSourceDataCompiledData.GetLayoutHash();
 			bool bChanged = SourceDataLayoutHash != ChachedDataSetLayoutHash;
@@ -808,7 +808,7 @@ void UNiagaraDataInterfaceDataChannelRead::ProvidePerInstanceDataForRenderThread
 	{
 		SourceData.bUpdateFunctionBindingRTData = false;
 		
-		const FNiagaraDataSetCompiledData& GPUCompiledData = SourceData.DataChannel->GetDataChannel()->GetCompiledData(ENiagaraSimTarget::GPUComputeSim);
+		const FNiagaraDataSetCompiledData& GPUCompiledData = SourceData.DataChannel->GetDataChannel()->GetLayoutInfo()->GetDataSetCompiledDataGPU();
 		TargetData->ScriptParamInfo.Init(CompiledData, GPUCompiledData);
 	}
 
@@ -2063,7 +2063,10 @@ void FNiagaraDataInterfaceProxy_DataChannelRead::PreStage(const FNDIGpuComputePr
 
 	if(InstanceData)
 	{
-		InstanceData->GPUBuffer = InstanceData->ChannelDataRTProxy->PrepareForReadAccess(Context.GetGraphBuilder(), InstanceData->bReadPrevFrame == false);
+		if(InstanceData->ChannelDataRTProxy)
+		{
+			InstanceData->GPUBuffer = InstanceData->ChannelDataRTProxy->PrepareForReadAccess(Context.GetGraphBuilder(), InstanceData->bReadPrevFrame == false);
+		}
 
 		//TODO: Should grab just one for the whole frame...
 		//TODO: Add some wrap behavior...		
@@ -2103,6 +2106,7 @@ void FNiagaraDataInterfaceProxy_DataChannelRead::PostSimulate(const FNDIGpuCompu
 		{
 			InstanceData->NDCSpawnDataBuffer = nullptr;
 			Context.GetInstanceCountManager().FreeEntry(InstanceData->ConsumeInstanceCountOffset);
+			InstanceData->ConsumeInstanceCountOffset = INDEX_NONE;//This should already be done inside FreeEntry but just to be sure.
 		}
 	}
 }
diff --git a/Engine/Plugins/FX/Niagara/Source/Niagara/Private/DataInterface/NiagaraDataInterfaceDataChannelWrite.cpp b/Engine/Plugins/FX/Niagara/Source/Niagara/Private/DataInterface/NiagaraDataInterfaceDataChannelWrite.cpp
--- a/Engine/Plugins/FX/Niagara/Source/Niagara/Private/DataInterface/NiagaraDataInterfaceDataChannelWrite.cpp
+++ b/Engine/Plugins/FX/Niagara/Source/Niagara/Private/DataInterface/NiagaraDataInterfaceDataChannelWrite.cpp
@@ -195,7 +195,7 @@ namespace NDIDataChannelWriteLocal
 struct FNDIDataChannelWriteInstanceData_RT
 {
 	//RT proxy for game channel data from which we're reading.
-	FNiagaraDataChannelDataProxy* ChannelDataRTProxy = nullptr;
+	FNiagaraDataChannelDataProxyPtr ChannelDataRTProxy = nullptr;
 
 	/** Parameter mapping info for every function in every script used by this DI. */
 	FVariadicParameterGPUScriptInfo ScriptParamInfo;
@@ -321,7 +321,7 @@ struct FNDIDataChannelWriteInstanceData
 						const FNDIDataChannelCompiledData& CompiledData = Interface->GetCompiledData();
 						for (const FNDIDataChannelFunctionInfo& FuncInfo : CompiledData.GetFunctionInfo())
 						{
-							FNDIDataChannelLayoutManager::Get().GetLayoutInfo(FuncInfo, DataChannelPtr->GetDataChannel()->GetCompiledData(ENiagaraSimTarget::CPUSim), MissingParams);
+							FNDIDataChannelLayoutManager::Get().GetLayoutInfo(FuncInfo, DataChannelPtr->GetDataChannel()->GetLayoutInfo()->GetDataSetCompiledData(), MissingParams);
 						}
 
 						if (MissingParams.Num() > 0)
@@ -374,8 +374,8 @@ struct FNDIDataChannelWriteInstanceData
 					}
 				}
 
-				const FNiagaraDataSetCompiledData& CPUSourceDataCompiledData = DataChannelPtr->GetDataChannel()->GetCompiledData(ENiagaraSimTarget::CPUSim);
-				const FNiagaraDataSetCompiledData& GPUSourceDataCompiledData = DataChannelPtr->GetDataChannel()->GetCompiledData(ENiagaraSimTarget::GPUComputeSim);
+				const FNiagaraDataSetCompiledData& CPUSourceDataCompiledData = DataChannelPtr->GetDataChannel()->GetLayoutInfo()->GetDataSetCompiledData();
+				const FNiagaraDataSetCompiledData& GPUSourceDataCompiledData = DataChannelPtr->GetDataChannel()->GetLayoutInfo()->GetDataSetCompiledDataGPU();
 				check(CPUSourceDataCompiledData.GetLayoutHash() && CPUSourceDataCompiledData.GetLayoutHash() == GPUSourceDataCompiledData.GetLayoutHash());
 				uint64 SourceDataLayoutHash = CPUSourceDataCompiledData.GetLayoutHash();
 				bool bChanged = SourceDataLayoutHash != ChachedDataSetLayoutHash;
@@ -564,7 +564,7 @@ void UNiagaraDataInterfaceDataChannelWrite::ProvidePerInstanceDataForRenderThrea
 	{
 		SourceData.bUpdateFunctionBindingRTData = false;
 
-		const FNiagaraDataSetCompiledData& GPUCompiledData = SourceData.DataChannel->GetDataChannel()->GetCompiledData(ENiagaraSimTarget::GPUComputeSim);
+		const FNiagaraDataSetCompiledData& GPUCompiledData = SourceData.DataChannel->GetDataChannel()->GetLayoutInfo()->GetDataSetCompiledDataGPU();
 		TargetData->ScriptParamInfo.Init(CompiledData, GPUCompiledData);
 	}
 
@@ -783,7 +783,7 @@ bool UNiagaraDataInterfaceDataChannelWrite::SimCacheWriteFrame(UObject* StorageO
 				FrameBuffer.Size = VarBuffer.Size;
 				FrameBuffer.Data = VarBuffer.Data;
 			}
-			const FNiagaraDataChannelGameDataLayout& Layout = Channel->Get()->GetGameDataLayout();
+			const FNiagaraDataChannelGameDataLayout& Layout = Channel->Get()->GetLayoutInfo()->GetGameDataLayout();
 			for (const TPair<FNiagaraVariableBase, int32>& VarPair : Layout.VariableIndices)
 			{
 				FrameData.VariableData[VarPair.Value].SourceVar = VarPair.Key;
@@ -1411,6 +1411,10 @@ void FNiagaraDataInterfaceProxy_DataChannelWrite::ConsumePerInstanceDataFromGame
 		InstData.bPublishToGPU = SourceData.bPublishToGPU;
 		InstData.LwcTile = SourceData.LwcTile;
 	}
+	else
+	{
+		InstData.ChannelDataRTProxy = nullptr;
+	}
 
 	if (SourceData.ScriptParamInfo.bDirty)
 	{
@@ -1457,7 +1461,7 @@ void FNiagaraDataInterfaceProxy_DataChannelWrite::PreStage(const FNDIGpuComputeP
 			//Allocate a separate buffer that we will write into and ship back to the CPU.			 
 			InstanceData->BufferForCPU = InstanceData->ChannelDataRTProxy->AllocateBufferForCPU(Context.GetGraphBuilder(), Context.GetComputeDispatchInterface().GetFeatureLevel(), InstanceData->AllocationCount, InstanceData->bPublishToGame, InstanceData->bPublishToCPU, InstanceData->LwcTile);
 
-			//Get a new instance count.
+			//Get a new instance count. This is later released by the ndc proxy
 			uint32 Offset = InstanceData->BufferForCPU->GetGPUInstanceCountBufferOffset();
 			Context.GetInstanceCountManager().FreeEntry(Offset);
 			InstanceData->BufferForCPU->SetGPUInstanceCountBufferOffset(Context.GetInstanceCountManager().AcquireOrAllocateEntry(Context.GetGraphBuilder().RHICmdList));
diff --git a/Engine/Plugins/FX/Niagara/Source/Niagara/Private/NiagaraDataChannel.cpp b/Engine/Plugins/FX/Niagara/Source/Niagara/Private/NiagaraDataChannel.cpp
--- a/Engine/Plugins/FX/Niagara/Source/Niagara/Private/NiagaraDataChannel.cpp
+++ b/Engine/Plugins/FX/Niagara/Source/Niagara/Private/NiagaraDataChannel.cpp
@@ -47,7 +47,7 @@ namespace NDCCVars
 
 //////////////////////////////////////////////////////////////////////////
 
-void FNiagaraDataChannelGameDataLayout::Init(const TArray<FNiagaraDataChannelVariable>& Variables)
+void FNiagaraDataChannelGameDataLayout::Init(TConstArrayView<FNiagaraDataChannelVariable> Variables)
 {
 	VariableIndices.Reset();
 	LwcConverters.Reserve(Variables.Num());
@@ -72,8 +72,9 @@ void FNiagaraDataChannelGameData::Init(const UNiagaraDataChannel* InDataChannel)
 {
 	check(InDataChannel);
 	DataChannel = InDataChannel;
+	LayoutInfo = InDataChannel->GetLayoutInfo();
 
-	const FNiagaraDataChannelGameDataLayout& Layout = DataChannel->GetGameDataLayout();
+	const FNiagaraDataChannelGameDataLayout& Layout = LayoutInfo->GetGameDataLayout();
 	VariableData.Empty();
 	VariableData.SetNum(Layout.VariableIndices.Num());
 	for(const auto& Pair : Layout.VariableIndices)
@@ -123,7 +124,7 @@ void FNiagaraDataChannelGameData::Reserve(int32 NewNum)
 
 FNiagaraDataChannelVariableBuffer* FNiagaraDataChannelGameData::FindVariableBuffer(const FNiagaraVariableBase& Var)
 {
-	const FNiagaraDataChannelGameDataLayout& Layout = DataChannel->GetGameDataLayout();
+	const FNiagaraDataChannelGameDataLayout& Layout = LayoutInfo->GetGameDataLayout();
 	const FNiagaraTypeDefinition& VarType = Var.GetType();
 	for(auto& Pair : Layout.VariableIndices)
 	{
@@ -158,7 +159,7 @@ void FNiagaraDataChannelGameData::WriteToDataSet(FNiagaraDataBuffer* DestBuffer,
 	
 	DestBuffer->SetNumInstances(DestStartIdx + NumInstances);
 
-	const FNiagaraDataChannelGameDataLayout& Layout = DataChannel->GetGameDataLayout();
+	const FNiagaraDataChannelGameDataLayout& Layout = LayoutInfo->GetGameDataLayout();
 
 	for (const TPair<FNiagaraVariableBase, int32>& VarIndexPair : Layout.VariableIndices)
 	{
@@ -316,7 +317,7 @@ void FNiagaraDataChannelGameData::AppendFromDataSet(const FNiagaraDataBuffer* Sr
 	int32 StartIndex = NumElements;
 	NumElements += NumInstances;
 
-	const FNiagaraDataChannelGameDataLayout& Layout = DataChannel->GetGameDataLayout();
+	const FNiagaraDataChannelGameDataLayout& Layout = LayoutInfo->GetGameDataLayout();
 	for (const TPair<FNiagaraVariableBase, int32>& VarIndexPair : Layout.VariableIndices)
 	{
 		FNiagaraVariableBase Var = VarIndexPair.Key;
@@ -446,7 +447,7 @@ void FNiagaraDataChannelGameData::AppendFromDataSet(const FNiagaraDataBuffer* Sr
 
 void FNiagaraDataChannelGameData::SetFromSimCache(const FNiagaraVariableBase& SourceVar, TConstArrayView<uint8> Data, int32 Size)
 {
-	const FNiagaraDataChannelGameDataLayout& Layout = DataChannel->GetGameDataLayout();
+	const FNiagaraDataChannelGameDataLayout& Layout = LayoutInfo->GetGameDataLayout();
 	if (const int* Index = Layout.VariableIndices.Find(SourceVar))
 	{
 		if (VariableData[*Index].Size == Size)
@@ -459,57 +460,58 @@ void FNiagaraDataChannelGameData::SetFromSimCache(const FNiagaraVariableBase& So
 //////////////////////////////////////////////////////////////////////////
 
 FNiagaraDataChannelDataProxy::~FNiagaraDataChannelDataProxy()
+{
+	if(GPUDataSet)
+	{
+		delete GPUDataSet;
+	}
+	GPUDataSet = nullptr;
+
+	Owner = nullptr;
+	CurrFrameData = nullptr;
+	PrevFrameData = nullptr;
+	bNeedsPrevFrameData = false;
+	PendingCPUBuffers.Empty();
+	PendingGPUAllocations = 0;
+}
+
+
+void FNiagaraDataChannelDataProxy::OnAddedToDispatcher(FNiagaraGpuComputeDispatchInterface* ComputeDispatchInterface)
 {
 #if !UE_BUILD_SHIPPING
-	ensureMsgf(DispatchInterfaceForDebuggingOnly == nullptr, TEXT("FNiagaraDataChannelDataProxy - ComputeDispatchInterface was not cleared this could result in a crash"));
+	DispatchInterfaceForDebuggingOnly = ComputeDispatchInterface;
 #endif
 }
 
-void FNiagaraDataChannelDataProxy::Cleanup(FNiagaraGpuComputeDispatchInterface* ComputeDispatchInterface)
+void FNiagaraDataChannelDataProxy::OnRemovedFromDispatcher(FNiagaraGpuComputeDispatchInterface* ComputeDispatchInterface)
 {
 #if !UE_BUILD_SHIPPING
-	ensureMsgf(DispatchInterfaceForDebuggingOnly == ComputeDispatchInterface || ComputeDispatchInterface == nullptr, TEXT("FNiagaraDataChannelDataProxy - ComputeDispatchInterface has been modified without proper cleanup. This could cause a crash."));
+	check(DispatchInterfaceForDebuggingOnly == ComputeDispatchInterface);
 #endif
-	
-	if(ComputeDispatchInterface)
-	{
-		//Remove ourselves from the dispatcher and free our instance counts from the instance count manager.
-		ComputeDispatchInterface->RemoveNDCDataProxy(this);
-	
-		FNiagaraGPUInstanceCountManager& InstCountManager = ComputeDispatchInterface->GetGPUInstanceCounterManager();
-		
-		if(CurrFrameData)
-		{
-			uint32 InstanceCountOffset = CurrFrameData->GetGPUInstanceCountBufferOffset();
-			InstCountManager.FreeEntry(InstanceCountOffset);
-			CurrFrameData->SetGPUInstanceCountBufferOffset(INDEX_NONE);
-		}
 
-		if (PrevFrameData)
-		{
-			uint32 InstanceCountOffset = PrevFrameData->GetGPUInstanceCountBufferOffset();
-			InstCountManager.FreeEntry(InstanceCountOffset);
-			PrevFrameData->SetGPUInstanceCountBufferOffset(INDEX_NONE);
-		}
-	}
-	else
+	FNiagaraGPUInstanceCountManager& InstCountManager = ComputeDispatchInterface->GetGPUInstanceCounterManager();
+
+	if (CurrFrameData)
 	{
-		if (CurrFrameData)
-		{
-			CurrFrameData->SetGPUInstanceCountBufferOffset(INDEX_NONE);
-		}
+		uint32 InstanceCountOffset = CurrFrameData->GetGPUInstanceCountBufferOffset();
+		InstCountManager.FreeEntry(InstanceCountOffset);
+		CurrFrameData->SetGPUInstanceCountBufferOffset(INDEX_NONE);
+	}
 
-		if (PrevFrameData)
-		{
-			PrevFrameData->SetGPUInstanceCountBufferOffset(INDEX_NONE);
-		}
+	if (PrevFrameData)
+	{
+		uint32 InstanceCountOffset = PrevFrameData->GetGPUInstanceCountBufferOffset();
+		InstCountManager.FreeEntry(InstanceCountOffset);
+		PrevFrameData->SetGPUInstanceCountBufferOffset(INDEX_NONE);
 	}
 
-#if !UE_BUILD_SHIPPING
-	DispatchInterfaceForDebuggingOnly = nullptr;
-#endif
-	CurrFrameData = nullptr;
-	PrevFrameData = nullptr;
+	for (FNDCGpuReadbackInfo& ReadbackInfo : PendingGPUReadbackBuffers)
+	{
+		check(ReadbackInfo.Buffer);
+		uint32 Offset = ReadbackInfo.Buffer->GetGPUInstanceCountBufferOffset();
+		ComputeDispatchInterface->GetGPUInstanceCounterManager().FreeEntry(Offset);
+		ReadbackInfo.Buffer->SetGPUInstanceCountBufferOffset(INDEX_NONE);
+	}
 }
 
 void FNiagaraDataChannelDataProxy::Reset()
@@ -576,6 +578,10 @@ void FNiagaraDataChannelDataProxy::EndFrame(FNiagaraGpuComputeDispatchInterface*
 	check(GPUDataSet);
 	check(GPUDataSet->GetSimTarget() == ENiagaraSimTarget::GPUComputeSim);
 
+#if !UE_BUILD_SHIPPING
+	check(DispatchInterfaceForDebuggingOnly == DispatchInterface);
+#endif
+
 	check(CurrBufferAccessCounts == 0);
 
 	//Handle pending readbacks for GPU->CPU data.
@@ -589,17 +595,26 @@ void FNiagaraDataChannelDataProxy::EndFrame(FNiagaraGpuComputeDispatchInterface*
 				LocalOwner = Owner,
 				PASS_PublishCPU = ReadbackInfo.bPublishToCPU,
 				PASS_PublishGame = ReadbackInfo.bPublishToGame,
-				PASS_LWCTile = ReadbackInfo.LWCTile](FNiagaraDataBufferRef ReadbackBuffer)
+				PASS_LWCTile = ReadbackInfo.LWCTile,
+				PASS_DataSetLayoutHash = LayoutInfo->GetDataSetCompiledData().GetLayoutHash()](TSharedRef<FNiagaraDataBufferReadback> Readback)
 				{
 					if (FNiagaraDataChannelData* NDCData = LocalOwner.Pin().Get())
 					{
-						if (ReadbackBuffer && ReadbackBuffer->GetNumInstances())
+						const FNiagaraDataSetCompiledData& CompiledData = NDCData->GetCompiledData(ENiagaraSimTarget::CPUSim);
+						if(PASS_DataSetLayoutHash == CompiledData.GetLayoutHash())
 						{
-							FNiagaraDataChannelPublishRequest PublishRequest(ReadbackBuffer);
+							//Only consume the readback it the layouts match.
+							//It's possible the layout has been changed in the time it took for the readback to get here.
+							FNiagaraDataChannelPublishRequest PublishRequest;
 							PublishRequest.bVisibleToGame = PASS_PublishGame;
 							PublishRequest.bVisibleToCPUSims = PASS_PublishCPU;
 							PublishRequest.bVisibleToGPUSims = false;//Don't want to ping pong back to the GPU
 							PublishRequest.LwcTile = PASS_LWCTile;
+
+							FNiagaraDataBuffer* ReadbackBuffer = NDCData->GetBufferForCPUWrite();
+							Readback->ReadResultsToDataBuffer(ReadbackBuffer);
+							PublishRequest.Data = ReadbackBuffer->UnlockForRead();
+
 #if WITH_NIAGARA_DEBUGGER
 							PublishRequest.DebugSource = TEXT("NDC GPU Readback");//TODO: Feed in real source names also
 #endif
@@ -768,31 +783,18 @@ FNiagaraDataChannelData::~FNiagaraDataChannelData()
 {
 	GameData.Reset();
 
-	UWorld* OwnerWorld = WeakOwnerWorld.Get();
-	FNiagaraGpuComputeDispatchInterface* ComputeDispatchInterface = FNiagaraGpuComputeDispatchInterface::Get(OwnerWorld);
-
 	/** We defer the deletion of the dataset to the RT to be sure all in-flight RT commands have finished using it.*/
 	ENQUEUE_RENDER_COMMAND(FDeleteContextCommand)(
-		[ComputeDispatchInterface, CPUDataChannelDataSet = CPUSimData, GPUDataChannelDataSet = GPUSimData, ReleasedRTProxy = RTProxy.Release()](FRHICommandListImmediate& RHICmdList)
+		[CPUDataChannelDataSet = CPUSimData](FRHICommandListImmediate& RHICmdList)
 	{
-		if (ReleasedRTProxy)
-		{
-			ReleasedRTProxy->Cleanup(ComputeDispatchInterface);
-			delete ReleasedRTProxy;
-		}
 		if (CPUDataChannelDataSet != nullptr)
 		{
 			delete CPUDataChannelDataSet;
 		}
-		if (GPUDataChannelDataSet != nullptr)
-		{
-			delete GPUDataChannelDataSet;
-		}
 	}
 	);
 	CPUSimData = nullptr;
-	GPUSimData = nullptr;
-	check(RTProxy == nullptr);
+	RTProxy = nullptr;
 }
 
 void FNiagaraDataChannelData::Init(UNiagaraDataChannelHandler* Owner)
@@ -800,20 +802,18 @@ void FNiagaraDataChannelData::Init(UNiagaraDataChannelHandler* Owner)
 	check(Owner);
 	const UNiagaraDataChannel* DataChannel = Owner->GetDataChannel();
 	check(DataChannel);
-
+	
+	LayoutInfo = DataChannel->GetLayoutInfo();
+	
 	GameData = DataChannel->CreateGameData();
 
 	CPUSimData = new FNiagaraDataSet();
-	GPUSimData = new FNiagaraDataSet();
 
-	const FNiagaraDataSetCompiledData& CompiledData = DataChannel->GetCompiledData(ENiagaraSimTarget::CPUSim);
-	const FNiagaraDataSetCompiledData& CompiledDataGPU = DataChannel->GetCompiledData(ENiagaraSimTarget::GPUComputeSim);
-	CPUSimData->Init(&CompiledData, 1);
-	GPUSimData->Init(&CompiledDataGPU, 2);
+	CPUSimData->Init(&LayoutInfo->GetDataSetCompiledData(), 1);
 
 	//TODO: Send game data to GPU direct without staging.
 	GameDataStaging = new FNiagaraDataSet();
-	GameDataStaging->Init(&CompiledData);
+	GameDataStaging->Init(&LayoutInfo->GetDataSetCompiledData());
 
 	UWorld* OwnerWorld = Owner->GetWorld();
 	WeakOwnerWorld = OwnerWorld;
@@ -831,10 +831,9 @@ void FNiagaraDataChannelData::Reset()
 
 	GameData->Empty();
 	CPUSimData->ResetBuffers();
-	GPUSimData->ResetBuffers();
 	GameDataStaging->ResetBuffers();
 	ENQUEUE_RENDER_COMMAND(FResetProxyCommand)(
-		[RT_Proxy = RTProxy.Get()](FRHICommandListImmediate& RHICmdList)
+		[RT_Proxy = RTProxy](FRHICommandListImmediate& RHICmdList)
 	{
 		if (RT_Proxy)
 		{
@@ -877,7 +876,7 @@ void FNiagaraDataChannelData::EndFrame(UNiagaraDataChannelHandler* Owner)
 	if(RTProxy)
 	{
 		ENQUEUE_RENDER_COMMAND(FDataChannelProxyEndFrame) (
-			[RT_Proxy = RTProxy.Get(), RT_BuffersForGPU = MoveTemp(BuffersForGPU)](FRHICommandListImmediate& CmdList)
+			[RT_Proxy = RTProxy, RT_BuffersForGPU = MoveTemp(BuffersForGPU)](FRHICommandListImmediate& CmdList)
 			{
 				RT_Proxy->AddBuffersFromCPU(RT_BuffersForGPU);
 			});
@@ -1085,8 +1084,7 @@ void FNiagaraDataChannelData::RemovePublishRequests(const FNiagaraDataSet* DataS
 
 const FNiagaraDataSetCompiledData& FNiagaraDataChannelData::GetCompiledData(ENiagaraSimTarget SimTarget)
 {
-	check(CPUSimData && GPUSimData)
-	return SimTarget == ENiagaraSimTarget::CPUSim ? CPUSimData->GetCompiledData() : GPUSimData->GetCompiledData();
+	return SimTarget == ENiagaraSimTarget::CPUSim ? LayoutInfo->GetDataSetCompiledData() : LayoutInfo->GetDataSetCompiledDataGPU();
 }
 
 FNiagaraDataBuffer* FNiagaraDataChannelData::GetBufferForCPUWrite()
@@ -1101,17 +1099,8 @@ FNiagaraDataBuffer* FNiagaraDataChannelData::GetBufferForCPUWrite()
 
 void FNiagaraDataChannelData::DestroyRenderThreadProxy(FNiagaraGpuComputeDispatchInterface* ComputeDispatchInterface)
 {
-	ENQUEUE_RENDER_COMMAND(FDestroyRTProxyCommand)(
-		[ComputeDispatchInterface, ReleasedRTProxy = RTProxy.Release()](FRHICommandListImmediate& RHICmdList)
-		{
-			if (ReleasedRTProxy)
-			{
-				ReleasedRTProxy->Cleanup(ComputeDispatchInterface);
-				delete ReleasedRTProxy;
-			}
-		}
-	);
-	check(RTProxy == nullptr);
+	//We can leave the proxy to be destroyed by the dispatcher ( or in some edge cases perhaps a DI proxy ).
+	RTProxy = nullptr;
 }
 
 void FNiagaraDataChannelData::CreateRenderThreadProxy(UNiagaraDataChannelHandler* Owner)
@@ -1123,16 +1112,18 @@ void FNiagaraDataChannelData::CreateRenderThreadProxy(UNiagaraDataChannelHandler
 
 	if (FNiagaraGpuComputeDispatchInterface* ComputeDispatchInterface = FNiagaraGpuComputeDispatchInterface::Get(OwnerWorld))
 	{
-		RTProxy.Reset(new FNiagaraDataChannelDataProxy());
-		RTProxy->GPUDataSet = GPUSimData;
+		RTProxy = MakeShared<FNiagaraDataChannelDataProxy>();
+		RTProxy->GPUDataSet = new FNiagaraDataSet();
+		RTProxy->GPUDataSet->Init(&LayoutInfo->GetDataSetCompiledDataGPU());
 		RTProxy->bNeedsPrevFrameData = DataChannel->KeepPreviousFrameData();
 		RTProxy->Owner = this->AsWeak();
+		RTProxy->LayoutInfo = LayoutInfo;//Keep our own ref to the layout info as the NDCData may die on the GT before the proxy.
 #if !UE_BUILD_SHIPPING
 		RTProxy->DebugName = FString::Printf(TEXT("%s__GPUData"), *DataChannel->GetName());
 		RTProxy->DispatchInterfaceForDebuggingOnly = ComputeDispatchInterface;
 #endif
 		ENQUEUE_RENDER_COMMAND(FNiagaraDataChannelDataProxyInit) (
-			[ComputeDispatchInterface, RT_Proxy = RTProxy.Get()](FRHICommandListImmediate& CmdList)
+			[ComputeDispatchInterface, RT_Proxy = RTProxy](FRHICommandListImmediate& CmdList)
 			{
 				ComputeDispatchInterface->AddNDCDataProxy(RT_Proxy);
 			}
@@ -1142,6 +1133,35 @@ void FNiagaraDataChannelData::CreateRenderThreadProxy(UNiagaraDataChannelHandler
 
 //////////////////////////////////////////////////////////////////////////
 
+FNiagaraDataChannelLayoutInfo::FNiagaraDataChannelLayoutInfo(const UNiagaraDataChannel* DataChannel)
+{
+	check(DataChannel);
+	GameDataLayout.Init(DataChannel->GetVariables());
+
+	CompiledData.SimTarget = ENiagaraSimTarget::CPUSim;
+	CompiledDataGPU.SimTarget = ENiagaraSimTarget::GPUComputeSim;
+	for (const FNiagaraDataChannelVariable& NDCVar : DataChannel->GetVariables())
+	{
+		FNiagaraVariableBase Var = NDCVar;
+		if (Var.GetType().IsEnum() == false)
+		{
+			Var.SetType(FNiagaraTypeDefinition(FNiagaraTypeHelper::FindNiagaraFriendlyTopLevelStruct(Var.GetType().GetScriptStruct(), ENiagaraStructConversion::Simulation)));
+		}
+		CompiledData.Variables.Add(Var);
+		CompiledDataGPU.Variables.Add(Var);
+	}
+	CompiledData.BuildLayout();
+	CompiledDataGPU.BuildLayout();
+}
+
+FNiagaraDataChannelLayoutInfo::~FNiagaraDataChannelLayoutInfo()
+{
+	CompiledData.Empty();
+	CompiledDataGPU.Empty();
+}
+
+//////////////////////////////////////////////////////////////////////////
+
 void UNiagaraDataChannel::PostInitProperties()
 {
 	Super::PostInitProperties();
@@ -1195,12 +1215,7 @@ void UNiagaraDataChannel::PostLoad()
 	}
 #endif
 
-	//Init compiled data. These are not currently serialized as we have no mechanism to rebuild then on internal data format changes like those in scripts.
-	GetCompiledData(ENiagaraSimTarget::CPUSim);
-	GetCompiledData(ENiagaraSimTarget::GPUComputeSim);
-
-	//TODO: Can serialize?
-	GameDataLayout.Init(ChannelVariables);
+	GetLayoutInfo();
 
 	INiagaraModule::RequestRefreshDataChannels();
 }
@@ -1269,59 +1284,21 @@ void UNiagaraDataChannel::PostEditChangeProperty(FPropertyChangedEvent& Property
 	check(IsInGameThread());
 
 	//Refresh compiled data
-	CompiledData.Empty();
-	CompiledDataGPU.Empty();
-	GetCompiledData(ENiagaraSimTarget::CPUSim);
-	GetCompiledData(ENiagaraSimTarget::GPUComputeSim);
-
-	GameDataLayout.Init(ChannelVariables);
+	LayoutInfo = nullptr;
+	GetLayoutInfo();
 
 	INiagaraModule::RequestRefreshDataChannels();
 }
 
 #endif//WITH_EDITOR
 
-const FNiagaraDataSetCompiledData& UNiagaraDataChannel::GetCompiledData(ENiagaraSimTarget SimTarget)const
+const FNiagaraDataChannelLayoutInfoPtr UNiagaraDataChannel::GetLayoutInfo()const
 {
-	if(SimTarget == ENiagaraSimTarget::CPUSim)
+	if(LayoutInfo == nullptr)
 	{
-		if (CompiledData.Variables.Num() != ChannelVariables.Num())
-		{
-			//Build the compiled data from the current variables but convert to Simulation types.
-			CompiledData.Empty();
-			CompiledData.SimTarget = ENiagaraSimTarget::CPUSim;
-			for (FNiagaraVariableBase Var : ChannelVariables)
-			{
-				if(Var.GetType().IsEnum() == false)
-				{
-					Var.SetType(FNiagaraTypeDefinition(FNiagaraTypeHelper::FindNiagaraFriendlyTopLevelStruct(Var.GetType().GetScriptStruct(), ENiagaraStructConversion::Simulation)));
-				}
-				CompiledData.Variables.Add(Var);
-			}
-			CompiledData.BuildLayout();
-		}
-		return CompiledData;
-	}
-	else
-	{
-		if (CompiledDataGPU.Variables.Num() != ChannelVariables.Num())
-		{
-			check(SimTarget == ENiagaraSimTarget::GPUComputeSim);
-			//Build the compiled data from the current variables but convert to Simulation types.
-			CompiledDataGPU.Empty();
-			CompiledDataGPU.SimTarget = ENiagaraSimTarget::GPUComputeSim;
-			for (FNiagaraVariableBase Var : ChannelVariables)
-			{
-				if (Var.GetType().IsEnum() == false)
-				{
-					Var.SetType(FNiagaraTypeDefinition(FNiagaraTypeHelper::FindNiagaraFriendlyTopLevelStruct(Var.GetType().GetScriptStruct(), ENiagaraStructConversion::Simulation)));
-				}
-				CompiledDataGPU.Variables.Add(Var);
-			}
-			CompiledDataGPU.BuildLayout();
-		}
-		return CompiledDataGPU;
+		LayoutInfo =  MakeShared<FNiagaraDataChannelLayoutInfo>(this);
 	}
+	return LayoutInfo;
 }
 
 FNiagaraDataChannelGameDataPtr UNiagaraDataChannel::CreateGameData()const
@@ -1333,7 +1310,7 @@ FNiagaraDataChannelGameDataPtr UNiagaraDataChannel::CreateGameData()const
 
 bool UNiagaraDataChannel::IsValid()const
 {
-	return ChannelVariables.Num() > 0 && CompiledData.Variables.Num() == ChannelVariables.Num() && CompiledDataGPU.Variables.Num() == ChannelVariables.Num();
+	return LayoutInfo && ChannelVariables.Num() > 0 && LayoutInfo->GetDataSetCompiledData().Variables.Num() == ChannelVariables.Num() && LayoutInfo->GetDataSetCompiledDataGPU().Variables.Num() == ChannelVariables.Num();
 }
 
 //////////////////////////////////////////////////////////////////////////
diff --git a/Engine/Plugins/FX/Niagara/Source/Niagara/Private/NiagaraDataChannelManager.cpp b/Engine/Plugins/FX/Niagara/Source/Niagara/Private/NiagaraDataChannelManager.cpp
--- a/Engine/Plugins/FX/Niagara/Source/Niagara/Private/NiagaraDataChannelManager.cpp
+++ b/Engine/Plugins/FX/Niagara/Source/Niagara/Private/NiagaraDataChannelManager.cpp
@@ -177,7 +177,15 @@ UNiagaraDataChannelHandler* FNiagaraDataChannelManager::InitDataChannel(const UN
 void FNiagaraDataChannelManager::RemoveDataChannel(const UNiagaraDataChannel* InChannel)
 {
 	check(IsInGameThread());
-	Channels.Remove(InChannel);
+
+	TObjectPtr<UNiagaraDataChannelHandler> RemovedHandler;
+	if(Channels.RemoveAndCopyValue(InChannel, RemovedHandler))
+	{
+		if(RemovedHandler)
+		{
+			RemovedHandler->Cleanup();
+		}
+	}
 }
 
 UWorld* FNiagaraDataChannelManager::GetWorld()const
diff --git a/Engine/Plugins/FX/Niagara/Source/Niagara/Private/NiagaraDataChannel_Islands.cpp b/Engine/Plugins/FX/Niagara/Source/Niagara/Private/NiagaraDataChannel_Islands.cpp
--- a/Engine/Plugins/FX/Niagara/Source/Niagara/Private/NiagaraDataChannel_Islands.cpp
+++ b/Engine/Plugins/FX/Niagara/Source/Niagara/Private/NiagaraDataChannel_Islands.cpp
@@ -398,7 +398,7 @@ void FNDCIsland::OnAcquired(FVector Location)
 
 void FNDCIsland::OnReleased()
 {
-	Data->Reset();
+	Data = nullptr;
 	for (UNiagaraComponent* Comp : NiagaraSystems)
 	{
 		Comp->ReleaseToPool();
diff --git a/Engine/Plugins/FX/Niagara/Source/Niagara/Private/NiagaraDataSetReadback.cpp b/Engine/Plugins/FX/Niagara/Source/Niagara/Private/NiagaraDataSetReadback.cpp
--- a/Engine/Plugins/FX/Niagara/Source/Niagara/Private/NiagaraDataSetReadback.cpp
+++ b/Engine/Plugins/FX/Niagara/Source/Niagara/Private/NiagaraDataSetReadback.cpp
@@ -201,31 +201,39 @@ void FNiagaraDataBufferReadback::EnqueueReadback(FRHICommandList& RHICmdList, FN
 	check(InDataBuffer);
 	check(IsInRenderingThread());
 
-	FNiagaraDataBuffer* DestBuffer = &InDataBuffer->GetOwner()->AllocateBuffer();
 	bool bHasGPUData = InDataBuffer->GetGPUBufferFloat().NumBytes > 0 || InDataBuffer->GetGPUBufferInt().NumBytes > 0 || InDataBuffer->GetGPUBufferHalf().NumBytes > 0;
 	++PendingReadbacks;
 	if (bHasGPUData)
 	{
-		GPUReadbackInternal(RHICmdList, ReadbackManager, InstanceCountManager, InDataBuffer, DestBuffer);
+		GPUReadbackInternal(RHICmdList, ReadbackManager, InstanceCountManager, InDataBuffer);
 	}
 	else
 	{
-		InDataBuffer->CopyTo(*DestBuffer, 0, 0, InDataBuffer->GetNumInstances());
-		ReadbackCompleteInternal(DestBuffer);
+		GatherResults(	InDataBuffer->GetNumInstances(), 
+						(const float*)InDataBuffer->GetFloatBuffer().GetData(), 
+						(const int32*)InDataBuffer->GetInt32Buffer().GetData(), 
+						(const FFloat16*)InDataBuffer->GetHalfBuffer().GetData(), 
+						InDataBuffer->GetFloatStride(), 
+						InDataBuffer->GetInt32Stride(), 
+						InDataBuffer->GetHalfStride(),
+						InDataBuffer->GetOwner()->GetNumFloatComponents(),
+						InDataBuffer->GetOwner()->GetNumInt32Components(),
+						InDataBuffer->GetOwner()->GetNumHalfComponents());
+
+		ReadbackCompleteInternal();
 	}
 }
 
-void FNiagaraDataBufferReadback::ReadbackCompleteInternal(FNiagaraDataBuffer* CompleteDataBuffer)
+void FNiagaraDataBufferReadback::ReadbackCompleteInternal()
 {	
-	if (OnReadbackComplete.IsBound() && CompleteDataBuffer)
+	if (OnReadbackComplete.IsBound())
 	{
 		AsyncTask(
 			ENamedThreads::GameThread,
-			[Readback = AsShared(),
-			 CompleteBuffer_GT = CompleteDataBuffer->UnlockForRead()]()
+			[Readback = AsShared()]()
 			{
 				Readback->PendingReadbacks--;
-				Readback->OnReadbackComplete.Execute(CompleteBuffer_GT);
+				Readback->OnReadbackComplete.Execute(Readback);
 			}
 		);
 	}
@@ -235,18 +243,29 @@ void FNiagaraDataBufferReadback::ReadbackCompleteInternal(FNiagaraDataBuffer* Co
 	}
 }
 
-void FNiagaraDataBufferReadback::GPUReadbackInternal(FRHICommandList& RHICmdList, FNiagaraGpuReadbackManager* ReadbackManager, FNiagaraGPUInstanceCountManager& InstanceCountManager, FNiagaraDataBufferRef SrcDataBuffer, FNiagaraDataBuffer* DestDataBuffer)
+void FNiagaraDataBufferReadback::ReadResultsToDataBuffer(FNiagaraDataBuffer* DestBuffer)const
+{
+	check(DestBuffer);
+
+	check(ReadbackData_NumFloatComponents == DestBuffer->GetOwner()->GetNumFloatComponents());
+	check(ReadbackData_NumInt32Components == DestBuffer->GetOwner()->GetNumInt32Components());
+	check(ReadbackData_NumHalfComponents == DestBuffer->GetOwner()->GetNumHalfComponents());
+
+	DestBuffer->GPUCopyFrom(ReadbackData_Float.GetData(), ReadbackData_Int32.GetData(), ReadbackData_Half.GetData(), 0, ReadbackData_Count, ReadbackData_FloatStride, ReadbackData_Int32Stride, ReadbackData_HalfStride);
+}
+
+void FNiagaraDataBufferReadback::GPUReadbackInternal(FRHICommandList& RHICmdList, FNiagaraGpuReadbackManager* ReadbackManager, FNiagaraGPUInstanceCountManager& InstanceCountManager, FNiagaraDataBufferRef SrcDataBuffer)
 {
-	if (SrcDataBuffer == nullptr || DestDataBuffer == nullptr || ReadbackManager == nullptr)
+	if (SrcDataBuffer == nullptr || ReadbackManager == nullptr)
 	{
-		ReadbackCompleteInternal(nullptr);
+		ReadbackCompleteInternal();
 		return;
 	}
 
 	const uint32 CountOffset = SrcDataBuffer->GetGPUInstanceCountBufferOffset();
 	if (CountOffset == INDEX_NONE)
 	{
-		ReadbackCompleteInternal(nullptr);
+		ReadbackCompleteInternal();
 		return;
 	}
 
@@ -267,9 +286,11 @@ void FNiagaraDataBufferReadback::GPUReadbackInternal(FRHICommandList& RHICmdList
 	const int32 HalfBufferStride = SrcDataBuffer->GetHalfStride();
 	const int32 IntBufferStride = SrcDataBuffer->GetInt32Stride();
 
-	int32 SourceNumInstancesAllocated = (int32)SrcDataBuffer->GetNumInstancesAllocated();
+	const int32 NumFloatComponents = SrcDataBuffer->GetOwner()->GetNumFloatComponents();
+	const int32 NumInt32Components = SrcDataBuffer->GetOwner()->GetNumInt32Components();
+	const int32 NumHalfComponents = SrcDataBuffer->GetOwner()->GetNumHalfComponents();
 
-	check(SrcDataBuffer->GetOwner() == DestDataBuffer->GetOwner());
+	int32 SourceNumInstancesAllocated = (int32)SrcDataBuffer->GetNumInstancesAllocated();
 
 	// Transition buffers to copy
 	TArray<FRHITransitionInfo, TInlineAllocator<NumReadbackBuffers>> Transitions;
@@ -299,32 +320,30 @@ void FNiagaraDataBufferReadback::GPUReadbackInternal(FRHICommandList& RHICmdList
 			const float* FloatDataBuffer = FloatBufferIndex == INDEX_NONE ? nullptr : reinterpret_cast<float*>(BufferData[FloatBufferIndex].Key);
 			const FFloat16* HalfDataBuffer = HalfBufferIndex == INDEX_NONE ? nullptr : reinterpret_cast<FFloat16*>(BufferData[HalfBufferIndex].Key);
 			const int32* IntDataBuffer = IntBufferIndex == INDEX_NONE ? nullptr : reinterpret_cast<int32*>(BufferData[IntBufferIndex].Key);
-			DestDataBuffer->GPUCopyFrom(FloatDataBuffer, IntDataBuffer, HalfDataBuffer, 0, InstanceCount, FloatBufferStride, IntBufferStride, HalfBufferStride);
-
-			check(DestDataBuffer->GetNumInstances() <= DestDataBuffer->GetNumInstancesAllocated());
-
-//			DestDataBuffer->Dump(0, DestDataBuffer->GetNumInstances(), TEXT("GPU Readback dump"));
+			
+			Readback->GatherResults(InstanceCount, FloatDataBuffer, IntDataBuffer, HalfDataBuffer, FloatBufferStride, IntBufferStride, HalfBufferStride, NumFloatComponents, NumInt32Components, NumHalfComponents);
 
+			//TODO: Safely copy IDToIndex Table.
 			// Copy ID to Index table
-			if (SrcDataBuffer)
-			{
-				TArray<int32>& IDTable = SrcDataBuffer->GetIDTable();
+// 			if (SrcDataBuffer)
+// 			{
+// 				TArray<int32>& IDTable = SrcDataBuffer->GetIDTable();
+// 
+// 				const int32* IDtoIndexBuffer = IDtoIndexBufferIndex == INDEX_NONE ? nullptr : reinterpret_cast<int32*>(BufferData[IDtoIndexBufferIndex].Key);
+// 				if (IDtoIndexBuffer != nullptr)
+// 				{
+// 					const int32 NumIDs = BufferData[IDtoIndexBufferIndex].Value / sizeof(int32);
+// 					check(NumIDs >= InstanceCount);
+// 					IDTable.SetNumUninitialized(NumIDs);
+// 					FMemory::Memcpy(IDTable.GetData(), IDtoIndexBuffer, NumIDs * sizeof(int32));
+// 				}
+// 				else
+// 				{
+// 					IDTable.Empty();
+// 				}
+// 			}
 
-				const int32* IDtoIndexBuffer = IDtoIndexBufferIndex == INDEX_NONE ? nullptr : reinterpret_cast<int32*>(BufferData[IDtoIndexBufferIndex].Key);
-				if (IDtoIndexBuffer != nullptr)
-				{
-					const int32 NumIDs = BufferData[IDtoIndexBufferIndex].Value / sizeof(int32);
-					check(NumIDs >= InstanceCount);
-					IDTable.SetNumUninitialized(NumIDs);
-					FMemory::Memcpy(IDTable.GetData(), IDtoIndexBuffer, NumIDs * sizeof(int32));
-				}
-				else
-				{
-					IDTable.Empty();
-				}
-			}
-
-			Readback->ReadbackCompleteInternal(DestDataBuffer);
+			Readback->ReadbackCompleteInternal();
 		}
 	);
 
@@ -335,3 +354,40 @@ void FNiagaraDataBufferReadback::GPUReadbackInternal(FRHICommandList& RHICmdList
 	}
 	RHICmdList.Transition(Transitions);
 }
+
+void FNiagaraDataBufferReadback::GatherResults(	int32 InstanceCount, const float* FloatData, const int32* Int32Data, const FFloat16* HalfData, 
+												int32 FloatStride, int32 Int32Stride, int32 HalfStride,
+												int32 NumFloatComponents, int32 NumInt32Components, int32 NumHalfComponents)
+{
+	//Copy over data to local temp buffers for use to pass over the to GT.
+	//It is not safe for us to copy directly into an RT owned data buffer at this point.
+	//We must wait until we're on the GT and copy into a GT owned data buffer.
+	if(FloatData)
+	{
+		int32 BufferSize = NumFloatComponents * FloatStride;
+		ReadbackData_Float.SetNumUninitialized(BufferSize);
+		FMemory::Memcpy(ReadbackData_Float.GetData(), FloatData, BufferSize);
+	}
+
+	if(Int32Data)
+	{
+		int32 BufferSize = NumInt32Components * Int32Stride;
+		ReadbackData_Int32.SetNumUninitialized(BufferSize);
+		FMemory::Memcpy(ReadbackData_Int32.GetData(), Int32Data, BufferSize);
+	}
+	
+	if(HalfData)
+	{
+		int32 BufferSize = NumHalfComponents * HalfStride;
+		ReadbackData_Half.SetNumUninitialized(BufferSize);
+		FMemory::Memcpy(ReadbackData_Half.GetData(), HalfData, BufferSize);
+	}
+
+	ReadbackData_Count = InstanceCount;
+	ReadbackData_FloatStride = FloatStride;
+	ReadbackData_Int32Stride = Int32Stride;
+	ReadbackData_HalfStride = HalfStride;
+	ReadbackData_NumFloatComponents = NumFloatComponents;
+	ReadbackData_NumInt32Components = NumInt32Components;
+	ReadbackData_NumHalfComponents = NumHalfComponents;
+}
\ No newline at end of file
diff --git a/Engine/Plugins/FX/Niagara/Source/Niagara/Private/NiagaraGpuComputeDispatch.cpp b/Engine/Plugins/FX/Niagara/Source/Niagara/Private/NiagaraGpuComputeDispatch.cpp
--- a/Engine/Plugins/FX/Niagara/Source/Niagara/Private/NiagaraGpuComputeDispatch.cpp
+++ b/Engine/Plugins/FX/Niagara/Source/Niagara/Private/NiagaraGpuComputeDispatch.cpp
@@ -240,6 +240,13 @@ FNiagaraGpuComputeDispatch::~FNiagaraGpuComputeDispatch()
 	FinishDispatches();
 
 	AsyncGpuTraceHelper->Reset();
+
+	for (auto It = NDCDataProxies.CreateIterator(); It; ++It)
+	{
+		FNiagaraDataChannelDataProxyPtr& NDCProxy = *It;
+		NDCProxy->OnRemovedFromDispatcher(this);
+	}
+	NDCDataProxies.Empty();
 }
 
 void FNiagaraGpuComputeDispatch::AddGpuComputeProxy(FNiagaraSystemGpuComputeProxy* ComputeProxy)
@@ -300,19 +307,15 @@ void FNiagaraGpuComputeDispatch::RemoveGpuComputeProxy(FNiagaraSystemGpuComputeP
 #endif
 }
 
-void FNiagaraGpuComputeDispatch::AddNDCDataProxy(FNiagaraDataChannelDataProxy* NDCDataProxy)
+void FNiagaraGpuComputeDispatch::AddNDCDataProxy(FNiagaraDataChannelDataProxyPtr NDCDataProxy)
 {
-#if !UE_BUILD_SHIPPING
-	check(NDCDataProxy->DispatchInterfaceForDebuggingOnly == this);
-#endif
+	NDCDataProxy->OnAddedToDispatcher(this);
 	NDCDataProxies.Add(NDCDataProxy);
 }
 
-void FNiagaraGpuComputeDispatch::RemoveNDCDataProxy(FNiagaraDataChannelDataProxy* NDCDataProxy)
-{
-#if !UE_BUILD_SHIPPING
-	check(NDCDataProxy->DispatchInterfaceForDebuggingOnly == this);
-#endif
+void FNiagaraGpuComputeDispatch::RemoveNDCDataProxy(FNiagaraDataChannelDataProxyPtr NDCDataProxy)
+{	
+	NDCDataProxy->OnRemovedFromDispatcher(this);
 	NDCDataProxies.RemoveSwap(NDCDataProxy);
 }
 
@@ -1946,9 +1949,18 @@ void FNiagaraGpuComputeDispatch::PreInitViews(FRDGBuilder& GraphBuilder, bool bA
 		{
 			FramesBeforeTickFlush = 0;
 
-			for(FNiagaraDataChannelDataProxy* NDCProxy : NDCDataProxies)
+			for (auto It = NDCDataProxies.CreateIterator(); It; ++It)
 			{
-				NDCProxy->BeginFrame(this, GraphBuilder.RHICmdList);
+				FNiagaraDataChannelDataProxyPtr& NDCProxy = *It;
+				if (NDCProxy.IsUnique() == false)
+				{
+					NDCProxy->BeginFrame(this, GraphBuilder.RHICmdList);
+				}
+				else
+				{
+					NDCProxy->OnRemovedFromDispatcher(this);
+					It.RemoveCurrentSwap();
+				}
 			}
 
 			UpdateInstanceCountManager(GraphBuilder.RHICmdList);
@@ -2112,9 +2124,18 @@ void FNiagaraGpuComputeDispatch::PostRenderOpaque(FRDGBuilder& GraphBuilder, TCo
 			{
 				FinishDispatches();
 
-				for (FNiagaraDataChannelDataProxy* NDCProxy : NDCDataProxies)
+				for (auto It = NDCDataProxies.CreateIterator(); It; ++It)
 				{
-					NDCProxy->EndFrame(this, RHICmdList);
+					FNiagaraDataChannelDataProxyPtr& NDCProxy = *It;
+					if(NDCProxy.IsUnique() == false)
+					{
+						NDCProxy->EndFrame(this, RHICmdList);
+					}
+					else
+					{
+						NDCProxy->OnRemovedFromDispatcher(this);
+						It.RemoveCurrentSwap();
+					}
 				}
 			}
 		}
diff --git a/Engine/Plugins/FX/Niagara/Source/Niagara/Private/NiagaraGpuComputeDispatch.h b/Engine/Plugins/FX/Niagara/Source/Niagara/Private/NiagaraGpuComputeDispatch.h
--- a/Engine/Plugins/FX/Niagara/Source/Niagara/Private/NiagaraGpuComputeDispatch.h
+++ b/Engine/Plugins/FX/Niagara/Source/Niagara/Private/NiagaraGpuComputeDispatch.h
@@ -98,8 +98,8 @@ public:
 	/** Remove system instance proxy from the batcher. */
 	virtual void RemoveGpuComputeProxy(FNiagaraSystemGpuComputeProxy* ComputeProxy) override;
 
-	virtual void AddNDCDataProxy(FNiagaraDataChannelDataProxy* NDCDataProxy) override;
-	virtual void RemoveNDCDataProxy(FNiagaraDataChannelDataProxy* NDCDataProxy) override;
+	virtual void AddNDCDataProxy(FNiagaraDataChannelDataProxyPtr NDCDataProxy) override;
+	virtual void RemoveNDCDataProxy(FNiagaraDataChannelDataProxyPtr NDCDataProxy) override;
 
 #if WITH_EDITOR
 	virtual void Suspend() override {}
@@ -249,7 +249,7 @@ private:
 	bool bRequiresReadback = false;
 	TArray<FNiagaraSystemGpuComputeProxy*> ProxiesPerStage[ENiagaraGpuComputeTickStage::Max];
 	
-	TArray<FNiagaraDataChannelDataProxy*> NDCDataProxies;
+	TArray<FNiagaraDataChannelDataProxyPtr> NDCDataProxies;
 
 	FNiagaraGpuDispatchList DispatchListPerStage[ENiagaraGpuComputeTickStage::Max];
 
diff --git a/Engine/Plugins/FX/Niagara/Source/Niagara/Public/NiagaraDataChannel.h b/Engine/Plugins/FX/Niagara/Source/Niagara/Public/NiagaraDataChannel.h
--- a/Engine/Plugins/FX/Niagara/Source/Niagara/Public/NiagaraDataChannel.h
+++ b/Engine/Plugins/FX/Niagara/Source/Niagara/Public/NiagaraDataChannel.h
@@ -50,6 +50,7 @@ struct FNiagaraDataChannelPublishRequest;
 struct FNiagaraDataChannelGameDataLayout;
 class FNiagaraGpuReadbackManager;
 class FRDGBuilder;
+class UNiagaraDataChannel;
 
 //////////////////////////////////////////////////////////////////////////
 
@@ -61,8 +62,10 @@ struct FNDCGpuReadbackInfo
 	FVector3f LWCTile;
 };
 
+using FNiagaraDataChannelDataProxyPtr = TSharedPtr<struct FNiagaraDataChannelDataProxy>;
+
 /** Render thread proxy of FNiagaraDataChannelData. */
-struct FNiagaraDataChannelDataProxy
+struct FNiagaraDataChannelDataProxy : public TSharedFromThis<FNiagaraDataChannelDataProxy>
 {
 	~FNiagaraDataChannelDataProxy();
 
@@ -72,6 +75,9 @@ struct FNiagaraDataChannelDataProxy
 	FNiagaraDataBufferRef PrevFrameData = nullptr;
 	bool bNeedsPrevFrameData = false;
 
+	//Keeping layout info ref to ensure lifetime for GPUDataSet.
+	FNiagaraDataChannelLayoutInfoPtr LayoutInfo;
+
 	//Buffers coming from the CPU that we're going to copy up for reading on the GPU
 	TArray<FNiagaraDataBufferRef> PendingCPUBuffers;
 
@@ -93,7 +99,6 @@ struct FNiagaraDataChannelDataProxy
 	const TCHAR* GetDebugName()const{return nullptr;}
 	#endif
 
-	void Cleanup(FNiagaraGpuComputeDispatchInterface* ComputeDispatchInterface);
 	void BeginFrame(FNiagaraGpuComputeDispatchInterface* DispatchInterface, FRHICommandListImmediate& RHICmdList);
 	void EndFrame(FNiagaraGpuComputeDispatchInterface* DispatchInterface, FRHICommandListImmediate& RHICmdList);
 	void Reset();
@@ -113,6 +118,36 @@ struct FNiagaraDataChannelDataProxy
 	FNiagaraDataBufferRef GetPrevFrameData()const { return PrevFrameData; }
 	
 	void AddTransition(FRDGBuilder& GraphBuilder, ERHIAccess AccessBefore, ERHIAccess AccessAfter, FNiagaraDataBuffer* Buffer);
+
+	//Perform and bookkeeping required when we remove a proxy from a dispatcher.
+	void OnAddedToDispatcher(FNiagaraGpuComputeDispatchInterface* ComputeDispatchInterface);
+	void OnRemovedFromDispatcher(FNiagaraGpuComputeDispatchInterface* ComputeDispatchInterface);
+};
+
+using FNiagaraDataChannelLayoutInfoPtr = TSharedPtr<FNiagaraDataChannelLayoutInfo>;
+
+
+/** Data describing the layout of Niagara Data channel buffers that is used in multiple places and must live beyond it's owning Data Channel. */
+struct FNiagaraDataChannelLayoutInfo : public TSharedFromThis<FNiagaraDataChannelLayoutInfo>
+{
+	FNiagaraDataChannelLayoutInfo(const UNiagaraDataChannel* DataChannel);
+	~FNiagaraDataChannelLayoutInfo();
+
+	const FNiagaraDataSetCompiledData& GetDataSetCompiledData()const{ return CompiledData; }
+	const FNiagaraDataSetCompiledData& GetDataSetCompiledDataGPU()const { return CompiledDataGPU; }
+	const FNiagaraDataChannelGameDataLayout& GetGameDataLayout()const { return GameDataLayout; }
+
+private:
+
+	/**
+	Data layout for payloads in Niagara datasets.
+	*/
+	FNiagaraDataSetCompiledData CompiledData;
+
+	FNiagaraDataSetCompiledData CompiledDataGPU;
+
+	/** Layout information for any data stored at the "Game" level. i.e. From game code/BP. AoS layout and LWC types. */
+	FNiagaraDataChannelGameDataLayout GameDataLayout;
 };
 
 DECLARE_MULTICAST_DELEGATE_OneParam(FOnDataChannelCreated, const UNiagaraDataChannel*);
@@ -140,13 +175,10 @@ public:
 	/** If true, we keep our previous frame's data. Some users will prefer a frame of latency to tick dependency. */
 	bool KeepPreviousFrameData() const { return bKeepPreviousFrameData; }
 
-	/** Returns the compiled data describing the data layout for DataChannels in this channel. */
-	NIAGARA_API const FNiagaraDataSetCompiledData& GetCompiledData(ENiagaraSimTarget SimTarget) const;
-
 	/** Create the appropriate handler object for this data channel. */
 	NIAGARA_API virtual UNiagaraDataChannelHandler* CreateHandler(UWorld* OwningWorld) const PURE_VIRTUAL(UNiagaraDataChannel::CreateHandler, {return nullptr;} );
 	
-	const FNiagaraDataChannelGameDataLayout& GetGameDataLayout() const { return GameDataLayout; }
+	const FNiagaraDataChannelLayoutInfoPtr GetLayoutInfo()const;
 
 	NIAGARA_API FNiagaraDataChannelGameDataPtr CreateGameData() const;
 
@@ -201,14 +233,7 @@ private:
 	/**
 	Data layout for payloads in Niagara datasets.
 	*/
-	UPROPERTY(Transient)
-	mutable FNiagaraDataSetCompiledData CompiledData;
-
-	UPROPERTY(Transient)
-	mutable FNiagaraDataSetCompiledData CompiledDataGPU;
-
-	/** Layout information for any data stored at the "Game" level. i.e. From game code/BP. AoS layout and LWC types. */
-	FNiagaraDataChannelGameDataLayout GameDataLayout;
+	mutable FNiagaraDataChannelLayoutInfoPtr LayoutInfo;
 	
 	#if WITH_NIAGARA_DEBUGGER
 	mutable bool bVerboseLogging = false;
diff --git a/Engine/Plugins/FX/Niagara/Source/Niagara/Public/NiagaraDataChannelPublic.h b/Engine/Plugins/FX/Niagara/Source/Niagara/Public/NiagaraDataChannelPublic.h
--- a/Engine/Plugins/FX/Niagara/Source/Niagara/Public/NiagaraDataChannelPublic.h
+++ b/Engine/Plugins/FX/Niagara/Source/Niagara/Public/NiagaraDataChannelPublic.h
@@ -13,8 +13,10 @@ class UNiagaraDataChannel;
 class UNiagaraDataChannelHandler;
 struct FNiagaraDataChannelGameData;
 struct FNiagaraDataChannelData;
+struct FNiagaraDataChannelLayoutInfo;
 using FNiagaraDataChannelGameDataPtr = TSharedPtr<FNiagaraDataChannelGameData>;
 using FNiagaraDataChannelDataPtr = TSharedPtr<FNiagaraDataChannelData>;
+using FNiagaraDataChannelLayoutInfoPtr = TSharedPtr<FNiagaraDataChannelLayoutInfo>;
 
 /** Niagara Data Channels are a system for communication between Niagara Systems and with game code/Blueprint.
 
@@ -145,7 +147,7 @@ struct FNiagaraDataChannelGameDataLayout
 	UPROPERTY()
 	TArray<FNiagaraLwcStructConverter> LwcConverters;
 
-	void Init(const TArray<FNiagaraDataChannelVariable>& Variables);
+	void Init(TConstArrayView<FNiagaraDataChannelVariable> Variables);
 };
 
 
@@ -289,6 +291,9 @@ private:
 
 	TWeakObjectPtr<const UNiagaraDataChannel> DataChannel;
 
+	//Keep reference to the layout this data was built with.
+	FNiagaraDataChannelLayoutInfoPtr LayoutInfo;
+
 public:
 
 	NIAGARA_API void Init(const UNiagaraDataChannel* InDataChannel);
diff --git a/Engine/Plugins/FX/Niagara/Source/Niagara/Public/NiagaraDataSetReadback.h b/Engine/Plugins/FX/Niagara/Source/Niagara/Public/NiagaraDataSetReadback.h
--- a/Engine/Plugins/FX/Niagara/Source/Niagara/Public/NiagaraDataSetReadback.h
+++ b/Engine/Plugins/FX/Niagara/Source/Niagara/Public/NiagaraDataSetReadback.h
@@ -49,19 +49,33 @@ private:
 class FNiagaraDataBufferReadback : public TSharedFromThis<FNiagaraDataBufferReadback, ESPMode::ThreadSafe>
 {
 public:
-	DECLARE_DELEGATE_OneParam(FOnReadbackComplete, FNiagaraDataBufferRef);
+	DECLARE_DELEGATE_OneParam(FOnReadbackComplete, TSharedRef<FNiagaraDataBufferReadback>);
 
 public:
 	FOnReadbackComplete& GetOnReadbackComplete() { return OnReadbackComplete; }
 
 	void EnqueueReadback(FRHICommandList& RHICmdList, FNiagaraDataBufferRef InDataBuffer, FNiagaraGpuReadbackManager* ReadbackManager, FNiagaraGPUInstanceCountManager& InstanceCountManager);
 
-private:
-	void ReadbackCompleteInternal(FNiagaraDataBuffer* CompleteDataBuffer);
-	void GPUReadbackInternal(FRHICommandList& RHICmdList, FNiagaraGpuReadbackManager* ReadbackManager, FNiagaraGPUInstanceCountManager& InstanceCountManager, FNiagaraDataBufferRef SrcDataBuffer, FNiagaraDataBuffer* DestDataBuffer);
+	void ReadResultsToDataBuffer(FNiagaraDataBuffer* DestBuffer)const;
 
 private:
+	void ReadbackCompleteInternal();
+	void GPUReadbackInternal(FRHICommandList& RHICmdList, FNiagaraGpuReadbackManager* ReadbackManager, FNiagaraGPUInstanceCountManager& InstanceCountManager, FNiagaraDataBufferRef SrcDataBuffer);
+
+	void GatherResults(int32 InstanceCount, const float* FloatData, const int32* Int32Data, const FFloat16* HalfData, int32 FloatStride, int32 Int32Stride, int32 HalfStride, int32 NumFloatComponents, int32 NumInt32Components, int32 NumHalfComponents);
+
 	std::atomic<int>		PendingReadbacks;
 
 	FOnReadbackComplete		OnReadbackComplete;
+
+	TArray<float> ReadbackData_Float;
+	TArray<int32> ReadbackData_Int32;
+	TArray<FFloat16> ReadbackData_Half;
+	int32 ReadbackData_Count = 0;
+	int32 ReadbackData_FloatStride = 0;
+	int32 ReadbackData_Int32Stride = 0;
+	int32 ReadbackData_HalfStride = 0;
+	int32 ReadbackData_NumFloatComponents = 0;
+	int32 ReadbackData_NumInt32Components = 0;
+	int32 ReadbackData_NumHalfComponents = 0;
 };
diff --git a/Engine/Plugins/FX/Niagara/Source/Niagara/Public/NiagaraGpuComputeDispatchInterface.h b/Engine/Plugins/FX/Niagara/Source/Niagara/Public/NiagaraGpuComputeDispatchInterface.h
--- a/Engine/Plugins/FX/Niagara/Source/Niagara/Public/NiagaraGpuComputeDispatchInterface.h
+++ b/Engine/Plugins/FX/Niagara/Source/Niagara/Public/NiagaraGpuComputeDispatchInterface.h
@@ -23,7 +23,8 @@ class FNiagaraGpuReadbackManager;
 class FNiagaraRayTracingHelper;
 struct FNiagaraScriptDebuggerInfo;
 class FNiagaraSystemGpuComputeProxy;
-struct FNiagaraDataChannelDataProxy;
+
+using FNiagaraDataChannelDataProxyPtr = TSharedPtr<struct FNiagaraDataChannelDataProxy>;
 
 // Public API for Niagara's Compute Dispatcher
 // This is generally used with DataInterfaces or Custom Renderers
@@ -51,9 +52,9 @@ public:
 	virtual void RemoveGpuComputeProxy(FNiagaraSystemGpuComputeProxy* ComputeProxy) = 0;
 
 	/** Add NDC Data to the batcher for tracking */
-	virtual void AddNDCDataProxy(FNiagaraDataChannelDataProxy* NDCDataProxy) = 0;
+	virtual void AddNDCDataProxy(FNiagaraDataChannelDataProxyPtr NDCDataProxy) = 0;
 	/** Add NDC Data to the batcher for tracking */
-	virtual void RemoveNDCDataProxy(FNiagaraDataChannelDataProxy* NDCDataProxy) = 0;
+	virtual void RemoveNDCDataProxy(FNiagaraDataChannelDataProxyPtr NDCDataProxy) = 0;
 
 	/**
 	 * Register work for GPU sorting (using the GPUSortManager).
