Fixed emitter assets not keeping their collapsed state

#tests saved emitter, reopened it in both states

diff --git a/Engine/Plugins/FX/Niagara/Source/NiagaraEditor/Private/NiagaraEditorUtilities.cpp b/Engine/Plugins/FX/Niagara/Source/NiagaraEditor/Private/NiagaraEditorUtilities.cpp
--- a/Engine/Plugins/FX/Niagara/Source/NiagaraEditor/Private/NiagaraEditorUtilities.cpp
+++ b/Engine/Plugins/FX/Niagara/Source/NiagaraEditor/Private/NiagaraEditorUtilities.cpp
@@ -2123,7 +2123,12 @@ const FGuid FNiagaraEditorUtilities::AddEmitterToSystem(UNiagaraSystem& InSystem
 		InSystem.AddEmitterHandleDirect(EmitterHandle);		
 	}
 
-	Cast<UNiagaraEmitterEditorData>(EmitterHandle.GetEmitterData()->GetEditorData())->SetShowSummaryView(InEmitterToAdd.GetEmitterData(EmitterVersion)->AddEmitterDefaultViewState == ENiagaraEmitterDefaultSummaryState::Summary ? true : false);
+	// We only want to use the default view state when actually adding an emitter to a system asset, not when opening up an emitter which adds said emitter to a transient system
+	if (SystemEditorData->GetOwningSystemIsPlaceholder() == false)
+	{
+		Cast<UNiagaraEmitterEditorData>(EmitterHandle.GetEmitterData()->GetEditorData())->SetShowSummaryView(InEmitterToAdd.GetEmitterData(EmitterVersion)->AddEmitterDefaultViewState == ENiagaraEmitterDefaultSummaryState::Summary ? true : false);
+	}
+	
 	FNiagaraStackGraphUtilities::RebuildEmitterNodes(InSystem);
 	SystemEditorData->SynchronizeOverviewGraphWithSystem(InSystem);
 
diff --git a/Engine/Plugins/FX/Niagara/Source/NiagaraEditor/Private/ViewModels/NiagaraSystemViewModel.cpp b/Engine/Plugins/FX/Niagara/Source/NiagaraEditor/Private/ViewModels/NiagaraSystemViewModel.cpp
--- a/Engine/Plugins/FX/Niagara/Source/NiagaraEditor/Private/ViewModels/NiagaraSystemViewModel.cpp
+++ b/Engine/Plugins/FX/Niagara/Source/NiagaraEditor/Private/ViewModels/NiagaraSystemViewModel.cpp
@@ -560,8 +560,12 @@ TSharedPtr<FNiagaraEmitterHandleViewModel> FNiagaraSystemViewModel::AddEmitter(U
 	}
 
 	RefreshAll();
-	
-	GetEmitterHandleViewModelById(NewEmitterHandleId)->GetEmitterViewModel()->GetEditorData().SetShowSummaryView(Emitter.GetEmitterData(EmitterVersion)->AddEmitterDefaultViewState == ENiagaraEmitterDefaultSummaryState::Summary ? true : false);
+
+	if (bSystemIsPlaceholder == false)
+	{
+		GetEmitterHandleViewModelById(NewEmitterHandleId)->GetEmitterViewModel()->GetEditorData().SetShowSummaryView(Emitter.GetEmitterData(EmitterVersion)->AddEmitterDefaultViewState == ENiagaraEmitterDefaultSummaryState::Summary ? true : false);
+	}
+
 	
 	TRange<float> SystemPlaybackRange = GetEditorData().GetPlaybackRange();
 	TRange<float> EmitterPlaybackRange = GetEmitterHandleViewModelById(NewEmitterHandleId)->GetEmitterViewModel()->GetEditorData().GetPlaybackRange();
