Added option for static switch output pins to move

#tests added multiple options, tried out the moving options, checked if input pins behaved correctly

diff --git a/Engine/Plugins/FX/Niagara/Source/NiagaraEditor/Private/NiagaraNodeStaticSwitch.cpp b/Engine/Plugins/FX/Niagara/Source/NiagaraEditor/Private/NiagaraNodeStaticSwitch.cpp
--- a/Engine/Plugins/FX/Niagara/Source/NiagaraEditor/Private/NiagaraNodeStaticSwitch.cpp
+++ b/Engine/Plugins/FX/Niagara/Source/NiagaraEditor/Private/NiagaraNodeStaticSwitch.cpp
@@ -268,6 +268,38 @@ bool UNiagaraNodeStaticSwitch::CanModifyPin(const UEdGraphPin* Pin) const
 	return true;
 }
 
+bool UNiagaraNodeStaticSwitch::CanMovePin(const UEdGraphPin* Pin, int32 DirectionToMove) const
+{
+	auto FindPredicate = [=](const FGuid& Guid) { return Guid == Pin->PersistentGuid; };
+	int32 FoundIndex = OutputVarGuids.IndexOfByPredicate(FindPredicate);
+	if (FoundIndex != INDEX_NONE && OutputVars.IsValidIndex(FoundIndex + DirectionToMove))
+	{
+		return Pin->Direction == EGPD_Output && Pin->bOrphanedPin == false;
+	}
+
+	return false;
+}
+
+void UNiagaraNodeStaticSwitch::MoveDynamicPin(UEdGraphPin* Pin, int32 MoveAmount)
+{
+	auto FindPredicate = [=](const FGuid& Guid) { return Guid == Pin->PersistentGuid; };
+	int32 FoundIndex = OutputVarGuids.IndexOfByPredicate(FindPredicate);
+	if (FoundIndex != INDEX_NONE && OutputVars.IsValidIndex(FoundIndex + MoveAmount))
+	{
+		this->Modify();
+		
+		FNiagaraVariable TmpVar = OutputVars[FoundIndex];
+		OutputVars[FoundIndex] = OutputVars[FoundIndex + MoveAmount];
+		OutputVars[FoundIndex + MoveAmount] = TmpVar;
+
+		FGuid TmpGuid = OutputVarGuids[FoundIndex];
+		OutputVarGuids[FoundIndex] = OutputVarGuids[FoundIndex + MoveAmount];
+		OutputVarGuids[FoundIndex + MoveAmount] = TmpGuid;
+
+		ReallocatePins();
+	}
+}
+
 void UNiagaraNodeStaticSwitch::PreChange(const UUserDefinedEnum* Changed, FEnumEditorUtils::EEnumEditorChangeInfo ChangedType)
 {
 	// do nothing
diff --git a/Engine/Plugins/FX/Niagara/Source/NiagaraEditor/Private/NiagaraNodeStaticSwitch.h b/Engine/Plugins/FX/Niagara/Source/NiagaraEditor/Private/NiagaraNodeStaticSwitch.h
--- a/Engine/Plugins/FX/Niagara/Source/NiagaraEditor/Private/NiagaraNodeStaticSwitch.h
+++ b/Engine/Plugins/FX/Niagara/Source/NiagaraEditor/Private/NiagaraNodeStaticSwitch.h
@@ -124,6 +124,8 @@ protected:
 	//~ End UNiagaraNodeUsageSelector Interface
 
 	virtual bool CanModifyPin(const UEdGraphPin* Pin) const override;
+	virtual bool CanMovePin(const UEdGraphPin* Pin, int32 DirectionToMove) const override;
+	virtual void MoveDynamicPin(UEdGraphPin* Pin, int32 MoveAmount) override;
 private:
 	/** INotifyOnEnumChanged interface */
 	virtual void PreChange(const UUserDefinedEnum* Changed, FEnumEditorUtils::EEnumEditorChangeInfo ChangedType) override;
