Fix off by one error in NDC Write function.
Causing us to write 1 less value than intended.

diff --git a/Engine/Plugins/FX/Niagara/Source/Niagara/Private/DataInterface/NiagaraDataInterfaceDataChannelWrite.cpp b/Engine/Plugins/FX/Niagara/Source/Niagara/Private/DataInterface/NiagaraDataInterfaceDataChannelWrite.cpp
--- a/Engine/Plugins/FX/Niagara/Source/Niagara/Private/DataInterface/NiagaraDataInterfaceDataChannelWrite.cpp
+++ b/Engine/Plugins/FX/Niagara/Source/Niagara/Private/DataInterface/NiagaraDataInterfaceDataChannelWrite.cpp
@@ -1072,7 +1072,8 @@ void UNiagaraDataInterfaceDataChannelWrite::Write(FVectorVMExternalFunctionConte
 
 			//Update the shared instance count with an updated max.
 			uint32 CurrNumInstances = AtomicNumInstances;
-			while(CurrNumInstances < MaxLocalIndex && !AtomicNumInstances.compare_exchange_weak(CurrNumInstances, MaxLocalIndex))
+			uint32 MaxLocalNumInstances = MaxLocalIndex + 1;
+			while(CurrNumInstances < MaxLocalNumInstances && !AtomicNumInstances.compare_exchange_weak(CurrNumInstances, MaxLocalNumInstances))
 			{
 				CurrNumInstances = AtomicNumInstances;
 			}
