Ensure cpu data used for spawning into GPU emitters is made available to the GPU.
We now auto upload all CPU buffers for NDC data that are used for spawning into GPU emitters.
This avoids a possible mis-match between the data on the GPU and what's expected by the spawning particles.
Can be disabled via "fx.Niagara.DataChannels.AutoUploadGPUSpawnData"

#tests PS5 Game, Content examples editor workflows

diff --git a/Engine/Plugins/FX/Niagara/Source/Niagara/Internal/DataInterface/NiagaraDataInterfaceDataChannelCommon.h b/Engine/Plugins/FX/Niagara/Source/Niagara/Internal/DataInterface/NiagaraDataInterfaceDataChannelCommon.h
--- a/Engine/Plugins/FX/Niagara/Source/Niagara/Internal/DataInterface/NiagaraDataInterfaceDataChannelCommon.h
+++ b/Engine/Plugins/FX/Niagara/Source/Niagara/Internal/DataInterface/NiagaraDataInterfaceDataChannelCommon.h
@@ -167,6 +167,7 @@ struct FNDIDataChannelCompiledData
 	bool UsedByCPU()const{ return bUsedByCPU; }
 	bool UsedByGPU()const{ return bUsedByGPU; }
 	bool NeedSpawnDataTable()const { return bNeedsSpawnDataTable; }
+	bool SpawnsParticles()const { return bSpawnsParticles; }
 	int32 GetTotalParams()const{ return TotalParams; }
 
 protected:
@@ -198,6 +199,9 @@ protected:
 	UPROPERTY()
 	bool bNeedsSpawnDataTable = true;
 
+	UPROPERTY()
+	bool bSpawnsParticles = false;
+
 	/** Iterates over all scripts for the owning system and gathers all functions and parameters accessing this DI. Building the FunctionInfoTable and GPUScriptParameterInfos map.  */
 	void GatherAccessInfo(UNiagaraSystem* System, UNiagaraDataInterface* Owner);
 };
@@ -228,6 +232,8 @@ public:
 namespace NDIDataChannelUtilities
 {
 	extern const FName GetNDCSpawnDataName;
+	extern const FName SpawnConditionalName;
+	extern const FName SpawnDirectName;
 
 	const TGlobalResource<FNDIDummyUAV>& GetDummyUAVFloat();
 	const TGlobalResource<FNDIDummyUAV>& GetDummyUAVInt32();
diff --git a/Engine/Plugins/FX/Niagara/Source/Niagara/Internal/DataInterface/NiagaraDataInterfaceDataChannelRead.h b/Engine/Plugins/FX/Niagara/Source/Niagara/Internal/DataInterface/NiagaraDataInterfaceDataChannelRead.h
--- a/Engine/Plugins/FX/Niagara/Source/Niagara/Internal/DataInterface/NiagaraDataInterfaceDataChannelRead.h
+++ b/Engine/Plugins/FX/Niagara/Source/Niagara/Internal/DataInterface/NiagaraDataInterfaceDataChannelRead.h
@@ -292,8 +292,13 @@ struct FNDIDataChannelReadInstanceData
 	virtual ~FNDIDataChannelReadInstanceData();
 	FNiagaraDataBuffer* GetReadBufferCPU(bool bPrevFrame)const;
 	bool Init(UNiagaraDataInterfaceDataChannelRead* Interface, FNiagaraSystemInstance* Instance);
+	void Cleanup(UNiagaraDataInterfaceDataChannelRead* Interface, FNiagaraSystemInstance* Instance);
 	bool Tick(UNiagaraDataInterfaceDataChannelRead* Interface, FNiagaraSystemInstance* Instance, bool bIsInit = false);	
 	bool PostTick(UNiagaraDataInterfaceDataChannelRead* Interface, FNiagaraSystemInstance* Instance);
+
+	//Set the current NDC Data. In some cases we must unregister ourselves from the current data and register with the new.
+	//TODO: move to a more handle based approach that will make this bookkeeping more robust.
+	void SetDataChannelData(FNiagaraDataChannelDataPtr NewData, UNiagaraDataInterfaceDataChannelRead* Interface);
 };
 
 struct FNiagaraDataInterfaceProxy_DataChannelRead : public FNiagaraDataInterfaceProxyRW
diff --git a/Engine/Plugins/FX/Niagara/Source/Niagara/Internal/NiagaraDataChannelCommon.h b/Engine/Plugins/FX/Niagara/Source/Niagara/Internal/NiagaraDataChannelCommon.h
--- a/Engine/Plugins/FX/Niagara/Source/Niagara/Internal/NiagaraDataChannelCommon.h
+++ b/Engine/Plugins/FX/Niagara/Source/Niagara/Internal/NiagaraDataChannelCommon.h
@@ -88,14 +88,6 @@ struct FNiagaraDataChannelData final : public TSharedFromThis<FNiagaraDataChanne
 	NIAGARA_API void Publish(const FNiagaraDataChannelPublishRequest& Request);
 
 	NIAGARA_API void PublishFromGPU(const FNiagaraDataChannelPublishRequest& Request);
-		
-	/**
-	 *Removes all publish requests involving the given dataset.
-	 *TODO: REMOVE
-	 *This is a hack to get around lifetime issues wrt data buffers/datasets and their compiled data.
-	 *We should rework things such that data buffers can exist beyond their owning dataset, including the compiled data detailing their layout.
-	 **/
-	NIAGARA_API void RemovePublishRequests(const FNiagaraDataSet* DataSet);
 
 	NIAGARA_API const FNiagaraDataSetCompiledData& GetCompiledData(ENiagaraSimTarget SimTarget);
 
@@ -107,6 +99,9 @@ struct FNiagaraDataChannelData final : public TSharedFromThis<FNiagaraDataChanne
 
 	void DestroyRenderThreadProxy(FNiagaraGpuComputeDispatchInterface* ComputeDispatchInterface);
 
+	void RegisterGPUSpawningReader() { ++NumGPUSpawningReaders; }
+	void UnregisterGPUSpawningReader() { --NumGPUSpawningReaders; }
+	int32 NumRegisteredGPUSpawningReaders()const{ return NumGPUSpawningReaders; }
 private:
 
 	void CreateRenderThreadProxy(UNiagaraDataChannelHandler* Owner);
@@ -125,7 +120,7 @@ private:
 	FNiagaraDataSet* GameDataStaging = nullptr;
 
 	/** Data buffers we'll be passing to the RT proxy for uploading to the GPU */
-	TArray<FNiagaraDataBufferRef> BuffersForGPU;
+	TArray<FNiagaraDataChannelPublishRequest> PublishRequestsForGPU;
 
 	/** Render thread proxy for this data. Owns all RT side data meant for GPU simulations. */
 	FNiagaraDataChannelDataProxyPtr RTProxy;
@@ -146,4 +141,11 @@ private:
 
 	//Keep reference to the layout this data was built with.
 	FNiagaraDataChannelLayoutInfoPtr LayoutInfo;
+
+	/** 
+	Track number of explicitly registered readers that spawn GPU particles from this data.
+	If we're spawning GPU particles using the CPU data (Spawn Conditional etc) then we have to send all CPU data to the GPU every frame.
+	Can possibly extend this to be a more automatic, registration based approach to shipping NDC data around rather than exmplicit flags on write.
+	*/
+	std::atomic<int32> NumGPUSpawningReaders;
 };
diff --git a/Engine/Plugins/FX/Niagara/Source/Niagara/Private/DataInterface/NiagaraDataInterfaceDataChannelCommon.cpp b/Engine/Plugins/FX/Niagara/Source/Niagara/Private/DataInterface/NiagaraDataInterfaceDataChannelCommon.cpp
--- a/Engine/Plugins/FX/Niagara/Source/Niagara/Private/DataInterface/NiagaraDataInterfaceDataChannelCommon.cpp
+++ b/Engine/Plugins/FX/Niagara/Source/Niagara/Private/DataInterface/NiagaraDataInterfaceDataChannelCommon.cpp
@@ -211,12 +211,16 @@ void FNDIDataChannelCompiledData::GatherAccessInfo(UNiagaraSystem* System, UNiag
 		{
 			bNeedsSpawnDataTable = true;
 		}
+		
+		if (BindingInfo.Name == NDIDataChannelUtilities::SpawnConditionalName || BindingInfo.Name == NDIDataChannelUtilities::SpawnDirectName)
+		{
+			bSpawnsParticles = true;			
+		}
 
 		return true;
 	};
 	FNiagaraDataInterfaceUtilities::ForEachVMFunction(Owner, System, HandleVMFunc);
 
-
 	//For every GPU script we iterate over the functions it calls and add each of them to the mapping.
 	//This will then be placed in a buffer for the RT to pass to the GPU so that each script can look up the correct function layout info.
 	GPUScriptParameterInfos.Empty();
@@ -288,6 +292,8 @@ int32 FNDIDataChannelCompiledData::FindFunctionInfoIndex(FName Name, const TArra
 namespace NDIDataChannelUtilities
 {
 	const FName GetNDCSpawnDataName(TEXT("GetNDCSpawnData"));
+	const FName SpawnConditionalName(TEXT("SpawnConditional"));
+	const FName SpawnDirectName(TEXT("SpawnConditional"));
 	const TGlobalResource<FNDIDummyUAV> DummyUAVFloat(PF_R32_FLOAT, sizeof(float));
 	const TGlobalResource<FNDIDummyUAV> DummyUAVInt32(PF_R32_SINT, sizeof(int32));
 	const TGlobalResource<FNDIDummyUAV> DummyUAVHalf(PF_R16F, sizeof(FFloat16));
diff --git a/Engine/Plugins/FX/Niagara/Source/Niagara/Private/DataInterface/NiagaraDataInterfaceDataChannelRead.cpp b/Engine/Plugins/FX/Niagara/Source/Niagara/Private/DataInterface/NiagaraDataInterfaceDataChannelRead.cpp
--- a/Engine/Plugins/FX/Niagara/Source/Niagara/Private/DataInterface/NiagaraDataInterfaceDataChannelRead.cpp
+++ b/Engine/Plugins/FX/Niagara/Source/Niagara/Private/DataInterface/NiagaraDataInterfaceDataChannelRead.cpp
@@ -315,7 +315,8 @@ struct FNDIDataChannelReadInstanceData_RT
 
 FNDIDataChannelReadInstanceData::~FNDIDataChannelReadInstanceData()
 {
-
+	//We must have cleared this by now in Cleanup so that we can unregister if needed.
+	check(DataChannelData == nullptr);
 }
 
 FNiagaraDataBuffer* FNDIDataChannelReadInstanceData::GetReadBufferCPU(bool bPrevFrame)const
@@ -344,6 +345,34 @@ bool FNDIDataChannelReadInstanceData::Init(UNiagaraDataInterfaceDataChannelRead*
 	return bSuccess;
 }
 
+void FNDIDataChannelReadInstanceData::Cleanup(UNiagaraDataInterfaceDataChannelRead* Interface, FNiagaraSystemInstance* Instance)
+{
+	SetDataChannelData(nullptr, Interface);
+}
+
+void FNDIDataChannelReadInstanceData::SetDataChannelData(FNiagaraDataChannelDataPtr NewData, UNiagaraDataInterfaceDataChannelRead* Interface)
+{
+	FNDIDataChannelCompiledData& CompiledData = Interface->GetCompiledData();
+	//If this interface spawns into a GPU emitter then we need to inform the NDC data so that it can automatically upload all CPU data to the GPU.
+	//Otherwise we get mis-matching data used for spawn and the subsequent reads on the GPU to init particles.
+	//It can create some very confusing janky behavior.
+	//Unfortunately we can't detect if this is directly spawning into a GPU emitter without some additional compiler/translator work.
+	//For now we can make do with checking if it's used to spawn particles and if any GPU emitters read the NDC data.
+	if(CompiledData.UsedByGPU() && CompiledData.SpawnsParticles())
+	{
+		if(DataChannelData)
+		{
+			DataChannelData->UnregisterGPUSpawningReader();
+		}
+		if(NewData)
+		{
+			NewData->RegisterGPUSpawningReader();
+		}
+	}
+
+	DataChannelData = NewData;
+}
+
 bool FNDIDataChannelReadInstanceData::Tick(UNiagaraDataInterfaceDataChannelRead* Interface, FNiagaraSystemInstance* Instance, bool bIsInit)
 {
 	ConsumeIndex = 0;
@@ -426,7 +455,7 @@ bool FNDIDataChannelReadInstanceData::Tick(UNiagaraDataInterfaceDataChannelRead*
 		UNiagaraDataChannelHandler* DataChannelPtr = DataChannel.Get();
 		if (DataChannelPtr == nullptr)
 		{
-			DataChannelData = nullptr;
+			SetDataChannelData(nullptr, Interface);
 			ChachedDataSetLayoutHash = INDEX_NONE;
 			UWorld* World = Instance->GetWorld();
 			if (FNiagaraWorldManager* WorldMan = FNiagaraWorldManager::Get(World))
@@ -451,7 +480,7 @@ bool FNDIDataChannelReadInstanceData::Tick(UNiagaraDataInterfaceDataChannelRead*
 			{
 				//TODO: Automatically modify tick group if we have DIs that require current frame info?
 				FNiagaraDataChannelSearchParameters SearchParams(Instance->GetAttachComponent());
-				DataChannelData = DataChannelPtr->FindData(SearchParams, ENiagaraResourceAccess::ReadOnly);//TODO: Maybe should have two paths, one for system instances and another for SceneComponents...
+				SetDataChannelData(DataChannelPtr->FindData(SearchParams, ENiagaraResourceAccess::ReadOnly), Interface);//TODO: Maybe should have two paths, one for system instances and another for SceneComponents...
 			}	
 
 			if(const UNiagaraDataChannel* ChannelPtr = DataChannelPtr->GetDataChannel())
@@ -606,6 +635,7 @@ bool UNiagaraDataInterfaceDataChannelRead::InitPerInstanceData(void* PerInstance
 void UNiagaraDataInterfaceDataChannelRead::DestroyPerInstanceData(void* PerInstanceData, FNiagaraSystemInstance* SystemInstance)
 {
 	FNDIDataChannelReadInstanceData* InstanceData = static_cast<FNDIDataChannelReadInstanceData*>(PerInstanceData);
+	InstanceData->Cleanup(this, SystemInstance);
 	InstanceData->~FNDIDataChannelReadInstanceData();
 
 	ENQUEUE_RENDER_COMMAND(RemoveProxy)
diff --git a/Engine/Plugins/FX/Niagara/Source/Niagara/Private/NiagaraDataChannel.cpp b/Engine/Plugins/FX/Niagara/Source/Niagara/Private/NiagaraDataChannel.cpp
--- a/Engine/Plugins/FX/Niagara/Source/Niagara/Private/NiagaraDataChannel.cpp
+++ b/Engine/Plugins/FX/Niagara/Source/Niagara/Private/NiagaraDataChannel.cpp
@@ -43,6 +43,11 @@ namespace NDCCVars
 		FConsoleCommandDelegate::CreateStatic(FNiagaraDataChannelDebugUtilities::DumpAllWritesToLog)
 	);
 #endif
+
+
+	bool AutoUploadGPUSpawnData = true;
+	static FAutoConsoleVariableRef CVarAutoUploadGPUSpawnData(TEXT("fx.Niagara.DataChannels.AutoUploadGPUSpawnData"), AutoUploadGPUSpawnData, TEXT("When true we will automatically upload any CPU NDC data to the GPU if it has been used to spawn GPU particles."), ECVF_Default);
+
 };
 
 //////////////////////////////////////////////////////////////////////////
@@ -825,7 +830,8 @@ void FNiagaraDataChannelData::Reset()
 {
 	FScopeLock Lock(&PublishCritSec);
 	PublishRequests.Reset();
-	BuffersForGPU.Reset();
+	PublishRequestsFromGPU.Reset();
+	PublishRequestsForGPU.Reset();
 	
 	PrevCPUSimData = nullptr;	
 
@@ -875,13 +881,62 @@ void FNiagaraDataChannelData::EndFrame(UNiagaraDataChannelHandler* Owner)
 
 	if(RTProxy)
 	{
-		ENQUEUE_RENDER_COMMAND(FDataChannelProxyEndFrame) (
-			[RT_Proxy = RTProxy, RT_BuffersForGPU = MoveTemp(BuffersForGPU)](FRHICommandListImmediate& CmdList)
+		int32 TotalGPUInstances = 0;
+		//Pass over any data bound for the GPU to the render thread.
+		TArray<FNiagaraDataBufferRef> BuffersForGPU;
+		BuffersForGPU.Reserve(PublishRequestsForGPU.Num());
+		
+		//If we've used our CPU data to spawn into a GPU sim this frame then we need to pass over all our CPU data regardless of if it was marked for GPU or not.		
+		if(NumGPUSpawningReaders > 0 && NDCCVars::AutoUploadGPUSpawnData)
+		{
+			if(CPUSimData->GetCurrentData() && CPUSimData->GetCurrentData()->GetNumInstances() > 0)
+			{
+				//First add our entire CPU sim data buffer.
+				BuffersForGPU.Emplace(CPUSimData->GetCurrentData());
+				TotalGPUInstances += CPUSimData->GetCurrentData()->GetNumInstances();
+			}
+
+			//Next add existing buffers we've kept around to pass to the GPU but only those that don't already exist in the main CPU buffers.
+			for(auto It = PublishRequestsForGPU.CreateIterator() ; It; ++It)
+			{
+				FNiagaraDataChannelPublishRequest& PublishRequest = *It;
+				check(PublishRequest.bVisibleToGPUSims);
+				if(!PublishRequest.bVisibleToCPUSims)
+				{
+					if(PublishRequest.Data->GetNumInstances() > 0)
+					{
+						BuffersForGPU.Emplace(PublishRequest.Data);
+						TotalGPUInstances += PublishRequest.Data->GetNumInstances();
+					}
+				}
+			}
+		}
+		else
+		{
+			//Only add data that was explicitly marked for GPU.
+			for (auto It = PublishRequestsForGPU.CreateIterator(); It; ++It)
 			{
-				RT_Proxy->AddBuffersFromCPU(RT_BuffersForGPU);
-			});
+				FNiagaraDataChannelPublishRequest& PublishRequest = *It;
+				check(PublishRequest.bVisibleToGPUSims);
+				if (PublishRequest.Data->GetNumInstances() > 0)
+				{
+					BuffersForGPU.Emplace(PublishRequest.Data);
+					TotalGPUInstances += PublishRequest.Data->GetNumInstances();
+				}
+			}
+		}
+
+		if(TotalGPUInstances > 0)
+		{
+			ENQUEUE_RENDER_COMMAND(FDataChannelProxyEndFrame) (
+				[RT_Proxy = RTProxy, RT_BuffersForGPU = MoveTemp(BuffersForGPU)](FRHICommandListImmediate& CmdList)
+				{
+
+					RT_Proxy->AddBuffersFromCPU(RT_BuffersForGPU);
+				});
+		}		
 	}
-	BuffersForGPU.Reset();
+	PublishRequestsForGPU.Reset();
 }
 
 int32 FNiagaraDataChannelData::ConsumePublishRequests(UNiagaraDataChannelHandler* Owner, const ETickingGroup& TickGroup)
@@ -925,32 +980,43 @@ int32 FNiagaraDataChannelData::ConsumePublishRequests(UNiagaraDataChannelHandler
 
 	//Each DI that generates DataChannel can control whether it's pushed to Game/CPU/GPU.
 	int32 RequestCount = PublishRequests.Num();
-	BuffersForGPU.Reserve(BuffersForGPU.Num() + RequestCount);
+	PublishRequestsForGPU.Reserve(PublishRequestsForGPU.Num() + RequestCount);
 
 	for (FNiagaraDataChannelPublishRequest& PublishRequest : PublishRequests)
 	{
 		uint32 NumInsts = 0;
+		//Don't bother uploading to the GPU separately if we already know we're going to be sending it via the CPU data as a whole.
+		bool bAutoUploadToGPU = NDCCVars::AutoUploadGPUSpawnData && NumGPUSpawningReaders > 0 && PublishRequest.bVisibleToCPUSims;
+		bool bShouldSendToGPU = PublishRequest.bVisibleToGPUSims && !bAutoUploadToGPU;
 
 		if (FNiagaraDataChannelGameData* RequestGameData = PublishRequest.GameData.Get())
 		{
 			NumInsts = RequestGameData->Num();
-			if (RequestGameData->Num() > 0 && PublishRequest.bVisibleToGPUSims)
+			if (RequestGameData->Num() > 0 && bShouldSendToGPU)
 			{
-				//We stage the game data into a data set to facilitate easier copy over to the GPU. TODO: Send and copy the game data directly.			
+				//We stage the game data into a data set to facilitate easier copy over to the GPU. TODO: Send and copy the game data directly.	
 				FNiagaraDataBuffer* StagingBuf = &GameDataStaging->BeginSimulate();
-				StagingBuf->Allocate(NumInsts);				
+				StagingBuf->Allocate(NumInsts);
 				RequestGameData->WriteToDataSet(StagingBuf, 0, LwcTile);
-				GameDataStaging->EndSimulate();				
-				BuffersForGPU.Emplace(StagingBuf);	
+				GameDataStaging->EndSimulate();
+				FNiagaraDataChannelPublishRequest& NewGPUReq = PublishRequestsForGPU.AddDefaulted_GetRef();
+				NewGPUReq.Data = StagingBuf;
+				NewGPUReq.bVisibleToCPUSims = PublishRequest.bVisibleToCPUSims;
+				NewGPUReq.bVisibleToGPUSims = PublishRequest.bVisibleToGPUSims;
+				NewGPUReq.LwcTile = PublishRequest.LwcTile;
 			}
 		}
 		else if (ensure(PublishRequest.Data))
 		{
 			NumInsts = PublishRequest.Data->GetNumInstances();
 
-			if (NumInsts > 0 && PublishRequest.bVisibleToGPUSims)
+			if (NumInsts > 0 && bShouldSendToGPU)
 			{
-				BuffersForGPU.Add(PublishRequest.Data);
+				FNiagaraDataChannelPublishRequest& NewGPUReq = PublishRequestsForGPU.AddDefaulted_GetRef();
+				NewGPUReq.Data = PublishRequest.Data;
+				NewGPUReq.bVisibleToCPUSims = PublishRequest.bVisibleToCPUSims;
+				NewGPUReq.bVisibleToGPUSims = PublishRequest.bVisibleToGPUSims;
+				NewGPUReq.LwcTile = PublishRequest.LwcTile;
 			}
 		}
 
@@ -1057,31 +1123,6 @@ void FNiagaraDataChannelData::PublishFromGPU(const FNiagaraDataChannelPublishReq
 	PublishRequestsFromGPU.Add(Request);
 }
 
-void FNiagaraDataChannelData::RemovePublishRequests(const FNiagaraDataSet* DataSet)
-{
-	FScopeLock Lock(&PublishCritSec);
-	//TODO: Have to ensure lifetime of anything we're holding
-	for (auto It = PublishRequests.CreateIterator(); It; ++It)
-	{
-		const FNiagaraDataChannelPublishRequest& PublishRequest = *It;
-
-		//First remove any buffers we've queued up to push to the GPU
-		for (auto BufferIt = BuffersForGPU.CreateIterator(); BufferIt; ++BufferIt)
-		{
-			FNiagaraDataBuffer* Buffer = *BufferIt;
-			if (Buffer == nullptr || Buffer->GetOwner() == DataSet)
-			{
-				BufferIt.RemoveCurrentSwap();
-			}
-		}
-
-		if (PublishRequest.Data && PublishRequest.Data->GetOwner() == DataSet)
-		{
-			It.RemoveCurrentSwap();
-		}
-	}
-}
-
 const FNiagaraDataSetCompiledData& FNiagaraDataChannelData::GetCompiledData(ENiagaraSimTarget SimTarget)
 {
 	return SimTarget == ENiagaraSimTarget::CPUSim ? LayoutInfo->GetDataSetCompiledData() : LayoutInfo->GetDataSetCompiledDataGPU();
