Fixed Pasting of modules not correctly forwarding the copied data of children inputs to the new module.
This is because the previous pasting function for inputs "SetValueFromClipboardInput" was only responsible for setting a single value, and wasn't recursively handling the data on children inputs
The new function "PasteFunctionInput" calls the above, but recursively goes through children too.

#tests Copy Pasted modules with increasing complexity and nested inputs

diff --git a/Engine/Plugins/FX/Niagara/Source/NiagaraEditor/Private/ViewModels/Stack/NiagaraStackScriptHierarchyRoot.cpp b/Engine/Plugins/FX/Niagara/Source/NiagaraEditor/Private/ViewModels/Stack/NiagaraStackScriptHierarchyRoot.cpp
--- a/Engine/Plugins/FX/Niagara/Source/NiagaraEditor/Private/ViewModels/Stack/NiagaraStackScriptHierarchyRoot.cpp
+++ b/Engine/Plugins/FX/Niagara/Source/NiagaraEditor/Private/ViewModels/Stack/NiagaraStackScriptHierarchyRoot.cpp
@@ -413,7 +413,7 @@ void UNiagaraStackScriptHierarchyRoot::SetValuesFromClipboardFunctionInputs(cons
 				}
 				else
 				{
-					StackFunctionInput->SetValueFromClipboardFunctionInput(*ClipboardFunctionInput);
+					StackFunctionInput->PasteFunctionInput(ClipboardFunctionInput);
 				}
 
 				bInputSetAsSwitch = true;
@@ -444,7 +444,7 @@ void UNiagaraStackScriptHierarchyRoot::SetValuesFromClipboardFunctionInputs(cons
 				}
 				else
 				{
-					StackFunctionInput->SetValueFromClipboardFunctionInput(*ClipboardFunctionInput);
+					StackFunctionInput->PasteFunctionInput(ClipboardFunctionInput);
 				}
 			}
 		}
