Fixed crash when undoing creation of scratch pad

#tests Undo/Redo of creation of scratch pad multiple times

diff --git a/Engine/Plugins/FX/Niagara/Source/NiagaraEditor/Private/ViewModels/HierarchyEditor/NiagaraHierarchyScriptParametersViewModel.cpp b/Engine/Plugins/FX/Niagara/Source/NiagaraEditor/Private/ViewModels/HierarchyEditor/NiagaraHierarchyScriptParametersViewModel.cpp
--- a/Engine/Plugins/FX/Niagara/Source/NiagaraEditor/Private/ViewModels/HierarchyEditor/NiagaraHierarchyScriptParametersViewModel.cpp
+++ b/Engine/Plugins/FX/Niagara/Source/NiagaraEditor/Private/ViewModels/HierarchyEditor/NiagaraHierarchyScriptParametersViewModel.cpp
@@ -141,7 +141,11 @@ void UNiagaraHierarchyScriptParametersViewModel::FinalizeInternal()
 {
 	if(ScriptViewModelWeak.IsValid())
 	{
-		Cast<UNiagaraScriptSource>(ScriptViewModelWeak.Pin()->GetStandaloneScript().GetScriptData()->GetSource())->NodeGraph->OnParametersChanged().RemoveAll(this);
+		// If this is called during Undo, it's possible the Graph does no longer exist
+		if(UNiagaraGraph* Graph = Cast<UNiagaraScriptSource>(ScriptViewModelWeak.Pin()->GetStandaloneScript().GetScriptData()->GetSource())->NodeGraph)
+		{
+			Graph->OnParametersChanged().RemoveAll(this);
+		}
 	}
 
 	Super::FinalizeInternal();
