Fixed crash when Undo/Redoing deleting an emitter while in Scalability Mode

#tests Deleted emitter in scalability mode and Undo/Redo multiple times, quickly

diff --git a/Engine/Plugins/FX/Niagara/Source/NiagaraEditor/Private/Widgets/Scalability/SNiagaraSystemResolvedScalabilitySettings.cpp b/Engine/Plugins/FX/Niagara/Source/NiagaraEditor/Private/Widgets/Scalability/SNiagaraSystemResolvedScalabilitySettings.cpp
--- a/Engine/Plugins/FX/Niagara/Source/NiagaraEditor/Private/Widgets/Scalability/SNiagaraSystemResolvedScalabilitySettings.cpp
+++ b/Engine/Plugins/FX/Niagara/Source/NiagaraEditor/Private/Widgets/Scalability/SNiagaraSystemResolvedScalabilitySettings.cpp
@@ -358,6 +358,11 @@ void SNiagaraSystemResolvedScalabilitySettings::RebuildWidget()
 		
 		for(const TSharedRef<FNiagaraEmitterHandleViewModel>& EmitterHandleViewModel : SystemViewModel->GetEmitterHandleViewModels())
 		{
+			if(EmitterHandleViewModel->IsValid() == false)
+			{
+				continue;
+			}
+			
 			FVersionedNiagaraEmitter VersionedEmitter = EmitterHandleViewModel->GetEmitterHandle()->GetInstance();
 			VersionedEmitter.Emitter->OnPropertiesChanged().RemoveAll(this);
 			VersionedEmitter.Emitter->OnPropertiesChanged().AddSP(this, &SNiagaraSystemResolvedScalabilitySettings::RebuildWidget);
