- Call Destroy() on the Stateless data buffer
- FNiagaraDataBuffer ref counts don't auto destroy when no references remain

#rnx

diff --git a/Engine/Plugins/FX/Niagara/Source/Niagara/Private/Stateless/NiagaraStatelessComputeManager.cpp b/Engine/Plugins/FX/Niagara/Source/Niagara/Private/Stateless/NiagaraStatelessComputeManager.cpp
--- a/Engine/Plugins/FX/Niagara/Source/Niagara/Private/Stateless/NiagaraStatelessComputeManager.cpp
+++ b/Engine/Plugins/FX/Niagara/Source/Niagara/Private/Stateless/NiagaraStatelessComputeManager.cpp
@@ -162,6 +162,15 @@ namespace NiagaraStatelessComputeManagerPrivate
 	}
 }
 
+FNiagaraStatelessComputeManager::FStatelessDataCache::~FStatelessDataCache()
+{
+	if (DataBuffer)
+	{
+		DataBuffer->Destroy();
+		DataBuffer = nullptr;
+	}
+}
+
 FNiagaraStatelessComputeManager::FNiagaraStatelessComputeManager(FNiagaraGpuComputeDispatchInterface* InOwnerInterface)
 	: FNiagaraGpuComputeDataManager(InOwnerInterface)
 {
diff --git a/Engine/Plugins/FX/Niagara/Source/Niagara/Private/Stateless/NiagaraStatelessComputeManager.h b/Engine/Plugins/FX/Niagara/Source/Niagara/Private/Stateless/NiagaraStatelessComputeManager.h
--- a/Engine/Plugins/FX/Niagara/Source/Niagara/Private/Stateless/NiagaraStatelessComputeManager.h
+++ b/Engine/Plugins/FX/Niagara/Source/Niagara/Private/Stateless/NiagaraStatelessComputeManager.h
@@ -21,6 +21,8 @@ class FNiagaraStatelessComputeManager final : public FNiagaraGpuComputeDataManag
 public:
 	struct FStatelessDataCache
 	{
+		~FStatelessDataCache();
+
 		uint32										DataSetLayoutHash = 0;
 		TSharedPtr<FNiagaraDataSetCompiledData>		DataSetCompiledData;
 		FNiagaraDataSet								DataSet;
