- Fix issue when count buffer is reallocated we could get incorrect readbacks, etc
- Added a way to reserve up front the number of required counts for emitters
- In one situation we could be running GPU sim, request a readback and it could get the wrong value killing particles from GPU emitters incorrectly.

#rnx

diff --git a/Engine/Plugins/FX/Niagara/Source/Niagara/Classes/NiagaraEmitterInstance.h b/Engine/Plugins/FX/Niagara/Source/Niagara/Classes/NiagaraEmitterInstance.h
--- a/Engine/Plugins/FX/Niagara/Source/Niagara/Classes/NiagaraEmitterInstance.h
+++ b/Engine/Plugins/FX/Niagara/Source/Niagara/Classes/NiagaraEmitterInstance.h
@@ -70,6 +70,8 @@ public:
 	NIAGARA_API virtual int32 GetNumParticles() const;
 	int32 GetTotalSpawnedParticles() const { return TotalSpawnedParticles; }
 
+	virtual uint32 GetGpuCountBufferEstimate() const { return 0; }
+
 	bool AreBoundsDynamic() const { return bCachedBoundsDynamic; }
 	[[nodiscard]] FBox GetBounds() const { return CachedBounds; }
 
diff --git a/Engine/Plugins/FX/Niagara/Source/Niagara/Classes/NiagaraGPUInstanceCountManager.h b/Engine/Plugins/FX/Niagara/Source/Niagara/Classes/NiagaraGPUInstanceCountManager.h
--- a/Engine/Plugins/FX/Niagara/Source/Niagara/Classes/NiagaraGPUInstanceCountManager.h
+++ b/Engine/Plugins/FX/Niagara/Source/Niagara/Classes/NiagaraGPUInstanceCountManager.h
@@ -74,14 +74,6 @@ public:
 	/** Acquire an entry, this will either come from the free list or reallocate the buffer. */
 	uint32 AcquireOrAllocateEntry(FRHICommandListImmediate& RHICmdList);
 
-	/**
-	Allocate a count entry that might return an existing cleared count, or a count waiting to be allocated.
-	This method allows you to allocated counts concurrently, then commit them by calling AllocateDeferredCounts once you join back on the render thread.
-	*/
-	uint32 AllocateDeferredEntry();
-	/** Commit any deferred count allocations */
-	void AllocateDeferredCounts(FRHICommandListImmediate& RHICmdList);
-
 	/** Free the entry and reset it to INDEX_NONE if valid. */
 	void FreeEntry(uint32& BufferOffset);
 	/** Free and array of entries, you are expected to reset or change to INDEX_NONE. */
@@ -195,6 +187,5 @@ protected:
 	TArray<FIndirectArgsPoolEntryPtr> DrawIndirectPool;
 	uint32 DrawIndirectLowWaterFrames = 0;
 
-	UE::FMutex	DeferredCountAllocationGuard;
-	int			DeferredCountAllocations = 0;
+	UE::FMutex	AcquireEntryGuard;
 };
diff --git a/Engine/Plugins/FX/Niagara/Source/Niagara/Internal/Stateless/NiagaraStatelessEmitterInstance.h b/Engine/Plugins/FX/Niagara/Source/Niagara/Internal/Stateless/NiagaraStatelessEmitterInstance.h
--- a/Engine/Plugins/FX/Niagara/Source/Niagara/Internal/Stateless/NiagaraStatelessEmitterInstance.h
+++ b/Engine/Plugins/FX/Niagara/Source/Niagara/Internal/Stateless/NiagaraStatelessEmitterInstance.h
@@ -71,6 +71,7 @@ public:
 	virtual void OnPooledReuse() override {}
 	virtual bool HandleCompletion(bool bForce) override;
 	virtual int32 GetNumParticles() const override;
+	virtual uint32 GetGpuCountBufferEstimate() const { return 1; }
 	virtual FNiagaraStatelessEmitterInstance* AsStateless() override { return this; }
 	virtual TConstArrayView<UNiagaraRendererProperties*> GetRenderers() const override;
 	virtual void BindParameters(bool bExternalOnly) override;
diff --git a/Engine/Plugins/FX/Niagara/Source/Niagara/Private/NiagaraGPUInstanceCountManager.cpp b/Engine/Plugins/FX/Niagara/Source/Niagara/Private/NiagaraGPUInstanceCountManager.cpp
--- a/Engine/Plugins/FX/Niagara/Source/Niagara/Private/NiagaraGPUInstanceCountManager.cpp
+++ b/Engine/Plugins/FX/Niagara/Source/Niagara/Private/NiagaraGPUInstanceCountManager.cpp
@@ -148,7 +148,9 @@ void FNiagaraGPUInstanceCountManager::ReleaseCounts()
 
 uint32 FNiagaraGPUInstanceCountManager::AcquireEntry()
 {
-	check(IsInRenderingThread());
+	check(IsInParallelRenderingThread());
+
+	UE::TScopeLock LockGuard(AcquireEntryGuard);
 
 	if (FreeEntries.Num())
 	{
@@ -161,9 +163,6 @@ uint32 FNiagaraGPUInstanceCountManager::AcquireEntry()
 	}
 	else
 	{
-		// @TODO : add realloc the buffer and copy the current content to it. Might require reallocating the readback in FNiagaraGPUInstanceCountManager::EnqueueGPUReadback()
-		ensure(UsedInstanceCounts < AllocatedInstanceCounts);
-		//UE_LOG(LogNiagara, Error, TEXT("Niagara.MinGPUInstanceCount too small. UsedInstanceCounts: %d < AllocatedInstanceCounts: %d"), UsedInstanceCounts, AllocatedInstanceCounts);
 		return INDEX_NONE;
 	}
 }
@@ -189,35 +188,6 @@ uint32 FNiagaraGPUInstanceCountManager::AcquireOrAllocateEntry(FRHICommandListIm
 	return UsedInstanceCounts++;
 }
 
-uint32 FNiagaraGPUInstanceCountManager::AllocateDeferredEntry()
-{
-	check(IsInParallelRenderingThread());
-
-	UE::TScopeLock LockGuard(DeferredCountAllocationGuard);
-
-	if (FreeEntries.Num())
-	{
-		return FreeEntries.Pop();
-	}
-	else if (UsedInstanceCounts < AllocatedInstanceCounts)
-	{
-		return UsedInstanceCounts++;
-	}
-
-	return AllocatedInstanceCounts + DeferredCountAllocations++;
-}
-
-void FNiagaraGPUInstanceCountManager::AllocateDeferredCounts(FRHICommandListImmediate& RHICmdList)
-{
-	if (DeferredCountAllocations == 0 )
-	{
-		return;
-	}
-	ResizeBuffers(RHICmdList, AllocatedInstanceCounts + DeferredCountAllocations);
-	UsedInstanceCounts += DeferredCountAllocations;
-	DeferredCountAllocations = 0;
-}
-
 void FNiagaraGPUInstanceCountManager::FreeEntry(uint32& BufferOffset)
 {
 	check(IsInRenderingThread());
diff --git a/Engine/Plugins/FX/Niagara/Source/Niagara/Private/NiagaraGpuComputeDispatch.cpp b/Engine/Plugins/FX/Niagara/Source/Niagara/Private/NiagaraGpuComputeDispatch.cpp
--- a/Engine/Plugins/FX/Niagara/Source/Niagara/Private/NiagaraGpuComputeDispatch.cpp
+++ b/Engine/Plugins/FX/Niagara/Source/Niagara/Private/NiagaraGpuComputeDispatch.cpp
@@ -262,6 +262,7 @@ void FNiagaraGpuComputeDispatch::AddGpuComputeProxy(FNiagaraSystemGpuComputeProx
 	NumProxiesThatRequireEarlyViewData			+= ComputeProxy->RequiresEarlyViewData() ? 1 : 0;
 	NumProxiesThatRequireRayTracingScene		+= ComputeProxy->RequiresRayTracingScene() ? 1 : 0;
 	NumProxiesThatRequireCurrentFrameNDC		+= ComputeProxy->RequiresCurrentFrameNDC() ? 1 : 0;
+	ProxyGpuCountBufferEstimate					+= ComputeProxy->GetGpuCountBufferEstimate();
 }
 
 void FNiagaraGpuComputeDispatch::RemoveGpuComputeProxy(FNiagaraSystemGpuComputeProxy* ComputeProxy)
@@ -283,7 +284,17 @@ void FNiagaraGpuComputeDispatch::RemoveGpuComputeProxy(FNiagaraSystemGpuComputeP
 	NumProxiesThatRequireDepthBuffer			-= ComputeProxy->RequiresDepthBuffer() ? 1 : 0;
 	NumProxiesThatRequireEarlyViewData			-= ComputeProxy->RequiresEarlyViewData() ? 1 : 0;
 	NumProxiesThatRequireRayTracingScene		-= ComputeProxy->RequiresRayTracingScene() ? 1 : 0;
-	NumProxiesThatRequireCurrentFrameNDC		-= ComputeProxy->RequiresCurrentFrameNDC() ? 1 : 0;
+	NumProxiesThatRequireCurrentFrameNDC		-= ComputeProxy->RequiresCurrentFrameNDC() ? 1 : 0;	
+
+	const uint32 GpuCountEstimate = ComputeProxy->GetGpuCountBufferEstimate();
+	if (ensure(GpuCountEstimate <= ProxyGpuCountBufferEstimate))
+	{
+		ProxyGpuCountBufferEstimate				-= GpuCountEstimate;
+	}
+	else
+	{
+		ProxyGpuCountBufferEstimate				= 0;
+	}
 
 #if NIAGARA_COMPUTEDEBUG_ENABLED
 	if (FNiagaraGpuComputeDebug* GpuComputeDebug = GpuComputeDebugPtr.Get())
@@ -693,7 +704,8 @@ void FNiagaraGpuComputeDispatch::UpdateInstanceCountManager(FRHICommandListImmed
 				}
 			}
 		}
-		GPUInstanceCounterManager.ResizeBuffers(RHICmdList, TotalDispatchCount);
+
+		GPUInstanceCounterManager.ResizeBuffers(RHICmdList, TotalDispatchCount + ProxyGpuCountBufferEstimate);
 	}
 
 	// Consume any pending readbacks that are ready
diff --git a/Engine/Plugins/FX/Niagara/Source/Niagara/Private/NiagaraGpuComputeDispatch.h b/Engine/Plugins/FX/Niagara/Source/Niagara/Private/NiagaraGpuComputeDispatch.h
--- a/Engine/Plugins/FX/Niagara/Source/Niagara/Private/NiagaraGpuComputeDispatch.h
+++ b/Engine/Plugins/FX/Niagara/Source/Niagara/Private/NiagaraGpuComputeDispatch.h
@@ -240,6 +240,8 @@ private:
 	uint32 NumProxiesThatRequireRayTracingScene = 0;
 	uint32 NumProxiesThatRequireCurrentFrameNDC = 0;
 
+	uint32 ProxyGpuCountBufferEstimate = 0;
+
 	int32 TotalDispatchesThisFrame = 0;
 
 	int32 MaxTicksToFlush = TNumericLimits<int32>::Max();
diff --git a/Engine/Plugins/FX/Niagara/Source/Niagara/Private/NiagaraSystemGpuComputeProxy.cpp b/Engine/Plugins/FX/Niagara/Source/Niagara/Private/NiagaraSystemGpuComputeProxy.cpp
--- a/Engine/Plugins/FX/Niagara/Source/Niagara/Private/NiagaraSystemGpuComputeProxy.cpp
+++ b/Engine/Plugins/FX/Niagara/Source/Niagara/Private/NiagaraSystemGpuComputeProxy.cpp
@@ -29,6 +29,7 @@ FNiagaraSystemGpuComputeProxy::FNiagaraSystemGpuComputeProxy(FNiagaraSystemInsta
 		{
 			ComputeContexts.Emplace(ComputeContext);
 		}
+		GpuCountBufferEstimate += Emitter->GetGpuCountBufferEstimate();
 	}
 
 	// Calculate Tick Stage
diff --git a/Engine/Plugins/FX/Niagara/Source/Niagara/Private/Stateless/NiagaraStatelessComputeManager.cpp b/Engine/Plugins/FX/Niagara/Source/Niagara/Private/Stateless/NiagaraStatelessComputeManager.cpp
--- a/Engine/Plugins/FX/Niagara/Source/Niagara/Private/Stateless/NiagaraStatelessComputeManager.cpp
+++ b/Engine/Plugins/FX/Niagara/Source/Niagara/Private/Stateless/NiagaraStatelessComputeManager.cpp
@@ -66,6 +66,34 @@ namespace NiagaraStatelessComputeManagerPrivate
 		return EComputeExecutionPath::None;
 	}
 
+	bool GenerateCPUDataForCPUSim(const NiagaraStateless::FEmitterInstance_RT* EmitterInstance, FNiagaraDataBuffer* DestinationBuffer)
+	{
+		using namespace NiagaraStateless;
+
+		const FNiagaraStatelessEmitterData* EmitterData = EmitterInstance->EmitterData.Get();
+
+		FParticleSimulationContext ParticleSimulation(EmitterData, EmitterInstance->ShaderParameters.Get(), EmitterInstance->BindingBufferData);
+		ParticleSimulation.Simulate(EmitterInstance->RandomSeed, EmitterInstance->Age, EmitterInstance->DeltaTime, EmitterInstance->SpawnInfos, DestinationBuffer);
+
+		INC_DWORD_STAT_BY(STAT_NiagaraStateless_CPUSimulate, 1);
+		INC_DWORD_STAT_BY(STAT_NiagaraStateless_CPUParticleCount, ParticleSimulation.GetNumInstances());
+		return ParticleSimulation.GetNumInstances() > 0;
+	}
+
+	bool GenerateCPUDataForGPUSim(FRHICommandListBase& RHICmdList, const NiagaraStateless::FEmitterInstance_RT* EmitterInstance, FNiagaraDataBuffer* DestinationBuffer)
+	{
+		using namespace NiagaraStateless;
+
+		const FNiagaraStatelessEmitterData* EmitterData = EmitterInstance->EmitterData.Get();
+
+		FParticleSimulationContext ParticleSimulation(EmitterData, EmitterInstance->ShaderParameters.Get(), EmitterInstance->BindingBufferData);
+		ParticleSimulation.SimulateGPU(RHICmdList, EmitterInstance->RandomSeed, EmitterInstance->Age, EmitterInstance->DeltaTime, EmitterInstance->SpawnInfos, DestinationBuffer);
+
+		INC_DWORD_STAT_BY(STAT_NiagaraStateless_CPUSimulate, 1);
+		INC_DWORD_STAT_BY(STAT_NiagaraStateless_CPUParticleCount, ParticleSimulation.GetNumInstances());
+		return ParticleSimulation.GetNumInstances() > 0;
+	}
+
 	void GenerateGPUData(FRHICommandList& RHICmdList, FNiagaraGpuComputeDispatchInterface* ComputeInterface, TConstArrayView<FNiagaraStatelessComputeManager::FStatelessDataGenerationRequest> GenerationRequests)
 	{
 		const int32 NumJobs = GenerationRequests.Num();
@@ -275,37 +303,45 @@ FNiagaraDataBuffer* FNiagaraStatelessComputeManager::GetDataBuffer(FRHICommandLi
 	const bool bAllowGPUGeneration = bAllowDeferredGeneration;
 
 	const EComputeExecutionPath ComputeExecutionPath = DetermineComputeExecutionPath(EmitterData, ActiveParticles, bAllowGPUGeneration);
+
 	switch (ComputeExecutionPath)
 	{
 		case EComputeExecutionPath::CPU:
 		{
-			INC_DWORD_STAT_BY(STAT_NiagaraStateless_CPUSimulate, 1);
-
-			FParticleSimulationContext ParticleSimulation(EmitterData, EmitterInstance->ShaderParameters.Get(), EmitterInstance->BindingBufferData);
-			ParticleSimulation.SimulateGPU(RHICmdList, EmitterInstance->RandomSeed, EmitterInstance->Age, EmitterInstance->DeltaTime, EmitterInstance->SpawnInfos, CacheData->DataBuffer);
-			if (ParticleSimulation.GetNumInstances() == 0)
+			if (!GenerateCPUDataForGPUSim(RHICmdList, EmitterInstance, CacheData->DataBuffer))
 			{
 				// Note: We can not free the data when in paralell GDME mode as multiple lock / unlock operations could result in the wrong one providing the final data
 				// The plus side to adding back into the cache data is that multiple calls to get the same emitter's data will resolve to no active particles.
 				UsedData.Emplace(EmitterKey, CacheData);
 				return nullptr;
 			}
-			INC_DWORD_STAT_BY(STAT_NiagaraStateless_CPUParticleCount, ParticleSimulation.GetNumInstances());
 			break;
 		}
 
 		case EComputeExecutionPath::GPU:
 		{
-			CacheData->DataBuffer->SetNumInstances(ActiveParticles);
-
 			FNiagaraGPUInstanceCountManager& CountManager = ComputeInterface->GetGPUInstanceCounterManager();
-			const uint32 CountOffset = CountManager.AllocateDeferredEntry();
-			CacheData->DataBuffer->SetGPUInstanceCountBufferOffset(CountOffset);
-			CountsToRelease.Add(CountOffset);
-
-			GPUGenerationRequests.Emplace(CacheData->DataBuffer, EmitterInstance, ActiveParticles);
-
-			INC_DWORD_STAT_BY(STAT_NiagaraStateless_GPUParticleEstimate, ActiveParticles);
+			const uint32 CountOffset = CountManager.AcquireEntry();
+			if (ensure(CountOffset != INDEX_NONE))
+			{
+				CacheData->DataBuffer->SetNumInstances(ActiveParticles);
+				CacheData->DataBuffer->SetGPUInstanceCountBufferOffset(CountOffset);
+				GPUGenerationRequests.Emplace(CacheData->DataBuffer, EmitterInstance, ActiveParticles);
+				CountsToRelease.Add(CountOffset);
+			}
+			// If we failed to alocate a count we will need to go through the CPU path (if available)
+			// This should never happen as we reserve a count up front via the compute proxy
+			// If it does occur this means some other system has used a count slot but not reserved one
+			else
+			{
+				if (!GenerateCPUDataForGPUSim(RHICmdList, EmitterInstance, CacheData->DataBuffer))
+				{
+					// Note: We can not free the data when in paralell GDME mode as multiple lock / unlock operations could result in the wrong one providing the final data
+					// The plus side to adding back into the cache data is that multiple calls to get the same emitter's data will resolve to no active particles.
+					UsedData.Emplace(EmitterKey, CacheData);
+					return nullptr;
+				}
+			}
 			break;
 		}
 
@@ -347,10 +383,7 @@ void FNiagaraStatelessComputeManager::GenerateDataBufferForDebugging(FRHICommand
 	{
 		case EComputeExecutionPath::CPU:
 		{
-			INC_DWORD_STAT_BY(STAT_NiagaraStateless_CPUSimulate, 1);
-
-			FParticleSimulationContext ParticleSimulation(EmitterData, EmitterInstance->ShaderParameters.Get(), EmitterInstance->BindingBufferData);
-			ParticleSimulation.Simulate(EmitterInstance->RandomSeed, EmitterInstance->Age, EmitterInstance->DeltaTime, EmitterInstance->SpawnInfos, DataBuffer);
+			GenerateCPUDataForCPUSim(EmitterInstance, DataBuffer);
 			break;
 		}
 
@@ -407,18 +440,11 @@ void FNiagaraStatelessComputeManager::OnPreRender(FRDGBuilder& GraphBuilder)
 	RDG_CSV_STAT_EXCLUSIVE_SCOPE(GraphBuilder, NiagaraStateless);
 	RDG_GPU_MASK_SCOPE(GraphBuilder, FRHIGPUMask::All());
 
-	// Ensure we allocate any deferred counts that we need
-	FNiagaraGpuComputeDispatchInterface* ComputeInterface = GetOwnerInterface();
-	{
-		FNiagaraGPUInstanceCountManager& CountManager = ComputeInterface->GetGPUInstanceCounterManager();
-		CountManager.AllocateDeferredCounts(GraphBuilder.RHICmdList);
-	}
-
 	// Execute dispatches
 	AddPass(
 		GraphBuilder,
 		RDG_EVENT_NAME("FNiagaraStatelessComputeManager::OnPreRender"),
-		[GPUGenerationRequests_RDG=MoveTemp(GPUGenerationRequests), ComputeInterface](FRHICommandListImmediate& RHICmdList)
+		[GPUGenerationRequests_RDG=MoveTemp(GPUGenerationRequests), ComputeInterface=GetOwnerInterface()](FRHICommandListImmediate& RHICmdList)
 		{
 			SCOPED_DRAW_EVENT(RHICmdList, FNiagaraStatelessComputeManager_OnPreRender);
 
diff --git a/Engine/Plugins/FX/Niagara/Source/Niagara/Public/NiagaraSystemGpuComputeProxy.h b/Engine/Plugins/FX/Niagara/Source/Niagara/Public/NiagaraSystemGpuComputeProxy.h
--- a/Engine/Plugins/FX/Niagara/Source/Niagara/Public/NiagaraSystemGpuComputeProxy.h
+++ b/Engine/Plugins/FX/Niagara/Source/Niagara/Public/NiagaraSystemGpuComputeProxy.h
@@ -35,6 +35,7 @@ public:
 	bool RequiresRayTracingScene() const { return bRequiresRayTracingScene; }
 	bool RequiresCurrentFrameNDC() const { return bRequiresCurrentFrameNDC; }
 	FVector3f GetSystemLWCTile() const { return SystemLWCTile; }
+	uint32 GetGpuCountBufferEstimate() const { return GpuCountBufferEstimate; }
 
 private:
 	FNiagaraSystemInstance*						DebugOwnerInstance = nullptr;
@@ -50,6 +51,7 @@ private:
 	uint32										bRequiresViewUniformBuffer : 1;
 	uint32										bRequiresRayTracingScene : 1;
 	uint32										bRequiresCurrentFrameNDC: 1;
+	uint32										GpuCountBufferEstimate = 0;
 
 	FShaderResourceViewRHIRef					StaticFloatBuffer;
 
