- Fix crash when using data channel wizard on channel with no attributes

#rnx

diff --git a/Engine/Plugins/FX/Niagara/Source/NiagaraEditor/Private/Widgets/DataChannel/NiagaraDataChannelWizard.cpp b/Engine/Plugins/FX/Niagara/Source/NiagaraEditor/Private/Widgets/DataChannel/NiagaraDataChannelWizard.cpp
--- a/Engine/Plugins/FX/Niagara/Source/NiagaraEditor/Private/Widgets/DataChannel/NiagaraDataChannelWizard.cpp
+++ b/Engine/Plugins/FX/Niagara/Source/NiagaraEditor/Private/Widgets/DataChannel/NiagaraDataChannelWizard.cpp
@@ -102,7 +102,10 @@ namespace UE::Niagara::Wizard::DataChannel
 			if (UNiagaraDataChannelAsset* ChannelAsset = PreviousPage->GetAsset())
 			{
 				NewDataChannelRef = ChannelAsset;
-				DataChannelVariables = ChannelAsset->Get()->GetVariables();
+				if (UNiagaraDataChannel* DataChannel = ChannelAsset->Get())
+				{
+					DataChannelVariables = DataChannel->GetVariables();
+				}
 			}
 			if (NewDataChannelRef != LastDataChannelRef)
 			{
@@ -366,7 +369,10 @@ namespace UE::Niagara::Wizard::DataChannel
 			if (UNiagaraDataChannelAsset* ChannelAsset = PreviousPage->GetAsset())
 			{
 				NewDataChannelRef = ChannelAsset;
-				DataChannelVariables = ChannelAsset->Get()->GetVariables();
+				if (UNiagaraDataChannel* DataChannel = ChannelAsset->Get())
+				{
+					DataChannelVariables = DataChannel->GetVariables();
+				}
 			}
 			if (NewDataChannelRef != LastDataChannelRef)
 			{
