- Fix debugger widget not refreshing when enabled / disabling the overview

#rnx

diff --git a/Engine/Plugins/FX/Niagara/Source/Niagara/Public/NiagaraDebuggerCommon.h b/Engine/Plugins/FX/Niagara/Source/Niagara/Public/NiagaraDebuggerCommon.h
--- a/Engine/Plugins/FX/Niagara/Source/Niagara/Public/NiagaraDebuggerCommon.h
+++ b/Engine/Plugins/FX/Niagara/Source/Niagara/Public/NiagaraDebuggerCommon.h
@@ -475,18 +475,18 @@ struct FNiagaraDebugHUDSettingsData
 	UPROPERTY()
 	bool bOverviewEnabled = false;
 
-	UPROPERTY(EditAnywhere, Category = "Debug Overview", meta = (DisplayName = "Debug Overview Mode", EditCondition = "bOverviewEnabled"))
+	UPROPERTY(EditAnywhere, Category = "Debug Overview", meta = (DisplayName = "Debug Overview Mode"))
 	ENiagaraDebugHUDOverviewMode OverviewMode = ENiagaraDebugHUDOverviewMode::Overview;
 
 	// How to sort the displayed list of systems
-	UPROPERTY(EditAnywhere, Category = "Debug Overview", meta = (DisplayName = "Debug Overview Sort Mode", EditCondition = "bOverviewEnabled"))
+	UPROPERTY(EditAnywhere, Category = "Debug Overview", meta = (DisplayName = "Debug Overview Sort Mode"))
 	ENiagaraDebugHUDDOverviewSort OverviewSortMode = ENiagaraDebugHUDDOverviewSort::Name;
 
 	/** When enabled the overview display will include cascade FX. */
-	UPROPERTY(EditAnywhere, Category = "Debug Overview", AdvancedDisplay, meta = (EditCondition = "bOverviewEnabled"))
+	UPROPERTY(EditAnywhere, Category = "Debug Overview", AdvancedDisplay)
 	bool bIncludeCascade = true;
 
-	UPROPERTY(EditAnywhere, Category = "Debug Overview", meta = (EditCondition = "bOverviewEnabled && OverviewMode == ENiagaraDebugHUDOverviewMode::Overview"))
+	UPROPERTY(EditAnywhere, Category = "Debug Overview")
 	bool bShowRegisteredComponents = false;
 	
 	/** When enabled the overview will only show the filter system information. */
diff --git a/Engine/Plugins/FX/Niagara/Source/NiagaraEditor/Private/Customizations/NiagaraDebugHUDCustomization.cpp b/Engine/Plugins/FX/Niagara/Source/NiagaraEditor/Private/Customizations/NiagaraDebugHUDCustomization.cpp
--- a/Engine/Plugins/FX/Niagara/Source/NiagaraEditor/Private/Customizations/NiagaraDebugHUDCustomization.cpp
+++ b/Engine/Plugins/FX/Niagara/Source/NiagaraEditor/Private/Customizations/NiagaraDebugHUDCustomization.cpp
@@ -243,6 +243,28 @@ FNiagaraDebugHUDSettingsDetailsCustomization::FNiagaraDebugHUDSettingsDetailsCus
 
 void FNiagaraDebugHUDSettingsDetailsCustomization::CustomizeDetails(IDetailLayoutBuilder& DetailBuilder)
 {
+	{
+		IDetailCategoryBuilder& OverviewCategory = DetailBuilder.EditCategory("Debug Overview");
+
+		TArray<TSharedRef<IPropertyHandle>> PropertyHandles;
+		OverviewCategory.GetDefaultProperties(PropertyHandles);
+		for (const TSharedRef<IPropertyHandle>& PropertyHandle : PropertyHandles)
+		{
+			if (PropertyHandle->GetProperty()->GetFName() == GET_MEMBER_NAME_CHECKED(FNiagaraDebugHUDSettingsData, bShowRegisteredComponents))
+			{
+				OverviewCategory.AddProperty(PropertyHandle).IsEnabled(
+					TAttribute<bool>::CreateLambda([&]() -> bool { const UNiagaraDebugHUDSettings* Settings = WeakSettings.Get(); return Settings->Data.bOverviewEnabled && Settings->Data.OverviewMode == ENiagaraDebugHUDOverviewMode::Overview; })
+				);
+			}
+			else
+			{
+				OverviewCategory.AddProperty(PropertyHandle).IsEnabled(
+					TAttribute<bool>::CreateLambda([&]() -> bool { return WeakSettings.Get()->Data.bOverviewEnabled; })
+				);
+			}
+		}
+	}
+
 	// Customize Filters
 	{
 		IDetailCategoryBuilder& FilterCategory = DetailBuilder.EditCategory("Debug Filter");
