Remove use of AcquireOrAllocate from NDC DIs.
Allow DIs to specify an instance count estimate for pre-allocation in the count manager.

#tests Console game, Content examples, engine tests

diff --git a/Engine/Plugins/FX/Niagara/Source/Niagara/Classes/NiagaraDataInterface.h b/Engine/Plugins/FX/Niagara/Source/Niagara/Classes/NiagaraDataInterface.h
--- a/Engine/Plugins/FX/Niagara/Source/Niagara/Classes/NiagaraDataInterface.h
+++ b/Engine/Plugins/FX/Niagara/Source/Niagara/Classes/NiagaraDataInterface.h
@@ -724,6 +724,9 @@ public:
 	virtual bool RequiresRayTracingScene() const { return false; }
 	virtual bool RequiresCurrentFrameNDC() const { return false; }
 
+	//Returns an estimate allocation count for the GPU instance count manager. This is experimental and will likely be modified in the future.
+	virtual uint32 GetGpuCountBufferEstimate() const { return 0; }
+
 	virtual bool HasTickGroupPrereqs() const { return false; }
 	virtual ETickingGroup CalculateTickGroup(const void* PerInstanceData) const { return NiagaraFirstTickGroup; }
 	virtual bool HasTickGroupPostreqs() const { return false; }
diff --git a/Engine/Plugins/FX/Niagara/Source/Niagara/Internal/DataInterface/NiagaraDataInterfaceDataChannelRead.h b/Engine/Plugins/FX/Niagara/Source/Niagara/Internal/DataInterface/NiagaraDataInterfaceDataChannelRead.h
--- a/Engine/Plugins/FX/Niagara/Source/Niagara/Internal/DataInterface/NiagaraDataInterfaceDataChannelRead.h
+++ b/Engine/Plugins/FX/Niagara/Source/Niagara/Internal/DataInterface/NiagaraDataInterfaceDataChannelRead.h
@@ -164,6 +164,8 @@ public:
 	virtual bool PostStageCanOverlapTickGroups() const { return false; }
 
 	virtual bool RequiresCurrentFrameNDC() const { return bReadCurrentFrame; }
+
+	virtual uint32 GetGpuCountBufferEstimate() const { return 1; }
 	//UNiagaraDataInterface Interface
 
 	//Functions usable anywhere.
diff --git a/Engine/Plugins/FX/Niagara/Source/Niagara/Internal/DataInterface/NiagaraDataInterfaceDataChannelWrite.h b/Engine/Plugins/FX/Niagara/Source/Niagara/Internal/DataInterface/NiagaraDataInterfaceDataChannelWrite.h
--- a/Engine/Plugins/FX/Niagara/Source/Niagara/Internal/DataInterface/NiagaraDataInterfaceDataChannelWrite.h
+++ b/Engine/Plugins/FX/Niagara/Source/Niagara/Internal/DataInterface/NiagaraDataInterfaceDataChannelWrite.h
@@ -170,6 +170,8 @@ public:
 	virtual bool PostSimulateCanOverlapFrames() const override { return false; }
 	//We cannot have post stage overlap tick groups so that the write DI can publish it's contents to the data channel at the correct time to allow same frame reads.
 	virtual bool PostStageCanOverlapTickGroups() const override { return false; }
+
+	virtual uint32 GetGpuCountBufferEstimate() const { return 1; }
 	//UNiagaraDataInterface Interface
 
 	NIAGARA_API void Num(FVectorVMExternalFunctionContext& Context);
diff --git a/Engine/Plugins/FX/Niagara/Source/Niagara/Private/DataInterface/NiagaraDataInterfaceDataChannelRead.cpp b/Engine/Plugins/FX/Niagara/Source/Niagara/Private/DataInterface/NiagaraDataInterfaceDataChannelRead.cpp
--- a/Engine/Plugins/FX/Niagara/Source/Niagara/Private/DataInterface/NiagaraDataInterfaceDataChannelRead.cpp
+++ b/Engine/Plugins/FX/Niagara/Source/Niagara/Private/DataInterface/NiagaraDataInterfaceDataChannelRead.cpp
@@ -2102,7 +2102,7 @@ void FNiagaraDataInterfaceProxy_DataChannelRead::PreStage(const FNDIGpuComputePr
 		//TODO: Add some wrap behavior...		
 		if(InstanceData->ConsumeInstanceCountOffset == INDEX_NONE)
 		{
-			InstanceData->ConsumeInstanceCountOffset = Context.GetInstanceCountManager().AcquireOrAllocateEntry(Context.GetGraphBuilder().RHICmdList);
+			InstanceData->ConsumeInstanceCountOffset = Context.GetInstanceCountManager().AcquireEntry();
 		}
 
 		if(InstanceData->NDCSpawnDataBuffer == nullptr)
diff --git a/Engine/Plugins/FX/Niagara/Source/Niagara/Private/DataInterface/NiagaraDataInterfaceDataChannelWrite.cpp b/Engine/Plugins/FX/Niagara/Source/Niagara/Private/DataInterface/NiagaraDataInterfaceDataChannelWrite.cpp
--- a/Engine/Plugins/FX/Niagara/Source/Niagara/Private/DataInterface/NiagaraDataInterfaceDataChannelWrite.cpp
+++ b/Engine/Plugins/FX/Niagara/Source/Niagara/Private/DataInterface/NiagaraDataInterfaceDataChannelWrite.cpp
@@ -1465,7 +1465,7 @@ void FNiagaraDataInterfaceProxy_DataChannelWrite::PreStage(const FNDIGpuComputeP
 			//Get a new instance count. This is later released by the ndc proxy
 			uint32 Offset = InstanceData->BufferForCPU->GetGPUInstanceCountBufferOffset();
 			Context.GetInstanceCountManager().FreeEntry(Offset);
-			InstanceData->BufferForCPU->SetGPUInstanceCountBufferOffset(Context.GetInstanceCountManager().AcquireOrAllocateEntry(Context.GetGraphBuilder().RHICmdList));
+			InstanceData->BufferForCPU->SetGPUInstanceCountBufferOffset(Context.GetInstanceCountManager().AcquireEntry());
 		}
 
 		if(InstanceData->bPublishToGPU)
diff --git a/Engine/Plugins/FX/Niagara/Source/Niagara/Private/NiagaraSystemGpuComputeProxy.cpp b/Engine/Plugins/FX/Niagara/Source/Niagara/Private/NiagaraSystemGpuComputeProxy.cpp
--- a/Engine/Plugins/FX/Niagara/Source/Niagara/Private/NiagaraSystemGpuComputeProxy.cpp
+++ b/Engine/Plugins/FX/Niagara/Source/Niagara/Private/NiagaraSystemGpuComputeProxy.cpp
@@ -31,6 +31,16 @@ FNiagaraSystemGpuComputeProxy::FNiagaraSystemGpuComputeProxy(FNiagaraSystemInsta
 		}
 		GpuCountBufferEstimate += Emitter->GetGpuCountBufferEstimate();
 	}
+	
+	for (const auto& Pair : OwnerInstance->GPUDataInterfaces)
+	{
+		UNiagaraDataInterface* Interface = Pair.Key.Get();
+		if (Interface == nullptr)
+		{
+			continue;
+		}
+		GpuCountBufferEstimate += Interface->GetGpuCountBufferEstimate();
+	}
 
 	// Calculate Tick Stage
 	if (bRequiresGlobalDistanceField || bRequiresDepthBuffer || bRequiresRayTracingScene || bRequiresCurrentFrameNDC)
