Fixed assert in NDC read buffer access due to mis-matched pre/post stage.(UE-227619)

Fixed Sim Cache writes for CPU and engine tests reliant on it. (GPU sim cache write will have to wait given current timing) (UE-225994)

Fixed shader error due to redundant defaulting shader code. (UE-227814)

Fixed editor interactions with NDC assets (UE-227810)
- Inserting into variables array breaks existing entries
- Undo/Redo doesn't work in many cases
- Edits to NDC assets would not take effect for running systems until GC or re-init.

Made GPU/CPU Reads and Writes behave consistently in response to missing NDC data or OOB access. (UE-227886)
- Missing data returns defaulted values. (previously could be uninitialized)
- Valid params are correctly read/written even if there are missing params. Previously one missing parameter could prevent other access functioning.
- Overall success returns false if there are any missing params or access is OOB.

diff --git a/Engine/Plugins/FX/Niagara/Shaders/Private/DataChannel/NiagaraDataInterfaceDataChannelCommon.ush b/Engine/Plugins/FX/Niagara/Shaders/Private/DataChannel/NiagaraDataInterfaceDataChannelCommon.ush
--- a/Engine/Plugins/FX/Niagara/Shaders/Private/DataChannel/NiagaraDataInterfaceDataChannelCommon.ush
+++ b/Engine/Plugins/FX/Niagara/Shaders/Private/DataChannel/NiagaraDataInterfaceDataChannelCommon.ush
@@ -16,8 +16,15 @@ Shader code common to all instances and functions of Niagara Data Channel DI.
 
 void NiagaraDataChannelRead_Float(Buffer<float> InBuffer, uint Stride, in out uint RegisterIdx, uint ParticleIdx, out float OutValue)
 {
-	OutValue = InBuffer[RegisterIdx * Stride + ParticleIdx];
-	++RegisterIdx;
+	if(RegisterIdx != -1)
+	{
+		OutValue = InBuffer[RegisterIdx * Stride + ParticleIdx];
+		++RegisterIdx;
+	}
+	else
+	{
+		OutValue = 0.0;
+	}
 }
 
 void NiagaraDataChannelRead_Float(Buffer<float> InBuffer, uint Stride, in out uint RegisterIdx, uint ParticleIdx, out float2 OutValue)
@@ -43,8 +50,15 @@ void NiagaraDataChannelRead_Float(Buffer<float> InBuffer, uint Stride, in out ui
 
 void NiagaraDataChannelRead_Int32(Buffer<int> InBuffer, uint Stride, in out uint RegisterIdx, uint ParticleIdx, out int OutValue)
 {
-	OutValue = InBuffer[RegisterIdx * Stride + ParticleIdx];
-	++RegisterIdx;
+	if(RegisterIdx != -1)
+	{
+		OutValue = InBuffer[RegisterIdx * Stride + ParticleIdx];
+		++RegisterIdx;
+	}
+	else
+	{
+		OutValue = 0;
+	}
 }
 
 void NiagaraDataChannelRead_Int32(Buffer<int> InBuffer, uint Stride, in out uint RegisterIdx, uint ParticleIdx, out int2 OutValue)
@@ -72,8 +86,15 @@ void NiagaraDataChannelRead_Int32(Buffer<int> InBuffer, uint Stride, in out uint
 //TODO: Half Support | 
 void NiagaraDataChannelRead_Half(Buffer<half> InBuffer, uint Stride, in out uint RegisterIdx, uint ParticleIdx, out half OutValue)
 {
-	OutValue = InBuffer[RegisterIdx * Stride + ParticleIdx];
-	++RegisterIdx;
+	if(RegisterIdx != -1)
+	{
+		OutValue = InBuffer[RegisterIdx * Stride + ParticleIdx];
+		++RegisterIdx;
+	}
+	else
+	{
+		OutValue = 0.0;
+	}
 }
 
 void NiagaraDataChannelRead_Half(Buffer<half> InBuffer, uint Stride, in out uint RegisterIdx, uint ParticleIdx, out half2 OutValue)
@@ -102,8 +123,11 @@ void NiagaraDataChannelRead_Half(Buffer<half> InBuffer, uint Stride, in out uint
 
 void NiagaraDataChannelWrite_Float(RWBuffer<float> InBuffer, uint Stride, in out uint RegisterIdx, uint ParticleIdx, float InValue)
 {
-	InBuffer[RegisterIdx * Stride + ParticleIdx] = InValue;
-	++RegisterIdx;
+	if(RegisterIdx != -1)
+	{
+		InBuffer[RegisterIdx * Stride + ParticleIdx] = InValue;
+		++RegisterIdx;
+	}
 }
 
 void NiagaraDataChannelWrite_Float(RWBuffer<float> InBuffer, uint Stride, in out uint RegisterIdx, uint ParticleIdx, float2 InValue)
@@ -129,8 +153,11 @@ void NiagaraDataChannelWrite_Float(RWBuffer<float> InBuffer, uint Stride, in out
 
 void NiagaraDataChannelWrite_Int32(RWBuffer<int> InBuffer, uint Stride, in out uint RegisterIdx, uint ParticleIdx, int InValue)
 {
-	InBuffer[RegisterIdx * Stride + ParticleIdx] = InValue;
-	++RegisterIdx;
+	if(RegisterIdx != -1)
+	{
+		InBuffer[RegisterIdx * Stride + ParticleIdx] = InValue;
+		++RegisterIdx;
+	}
 }
 
 void NiagaraDataChannelWrite_Int32(RWBuffer<int> InBuffer, uint Stride, in out uint RegisterIdx, uint ParticleIdx, int2 InValue)
@@ -158,8 +185,11 @@ void NiagaraDataChannelWrite_Int32(RWBuffer<int> InBuffer, uint Stride, in out u
 //TODO: Half Support | 
 void NiagaraDataChannelWrite_Half(RWBuffer<half> InBuffer, uint Stride, in out uint RegisterIdx, uint ParticleIdx, half InValue)
 {
-	InBuffer[RegisterIdx * Stride + ParticleIdx] = InValue;
-	++RegisterIdx;
+	if(RegisterIdx != -1)
+	{
+		InBuffer[RegisterIdx * Stride + ParticleIdx] = InValue;
+		++RegisterIdx;
+	}
 }
 
 void NiagaraDataChannelWrite_Half(RWBuffer<half> InBuffer, uint Stride, in out uint RegisterIdx, uint ParticleIdx, half2 InValue)
diff --git a/Engine/Plugins/FX/Niagara/Shaders/Private/DataChannel/NiagaraDataInterfaceDataChannelTemplateCommon.ush b/Engine/Plugins/FX/Niagara/Shaders/Private/DataChannel/NiagaraDataInterfaceDataChannelTemplateCommon.ush
--- a/Engine/Plugins/FX/Niagara/Shaders/Private/DataChannel/NiagaraDataInterfaceDataChannelTemplateCommon.ush
+++ b/Engine/Plugins/FX/Niagara/Shaders/Private/DataChannel/NiagaraDataInterfaceDataChannelTemplateCommon.ush
@@ -25,7 +25,7 @@ bool GetParameterLayoutInfo_{ParameterName}(uint ParameterIndex, out uint OutFlo
 		OutFloatRegister = {ParameterName}_ParamOffsetTable[ParamLayoutStart];
 		OutInt32Register = {ParameterName}_ParamOffsetTable[ParamLayoutStart + 1];
 		//TODO: Half Support | OutHalfRegister = {ParameterName}_ParamOffsetTable[ParamLayoutStart + 2];
-		return true;
+		return OutFloatRegister != -1 || OutInt32Register != -1;//If we have no valid mappings then this parameter is not valid.
 	}
 	else
 	{
diff --git a/Engine/Plugins/FX/Niagara/Shaders/Private/DataChannel/NiagaraDataInterfaceDataChannelTemplateReadCommon.ush b/Engine/Plugins/FX/Niagara/Shaders/Private/DataChannel/NiagaraDataInterfaceDataChannelTemplateReadCommon.ush
--- a/Engine/Plugins/FX/Niagara/Shaders/Private/DataChannel/NiagaraDataInterfaceDataChannelTemplateReadCommon.ush
+++ b/Engine/Plugins/FX/Niagara/Shaders/Private/DataChannel/NiagaraDataInterfaceDataChannelTemplateReadCommon.ush
@@ -26,7 +26,8 @@ Buffer<int> {ParameterName}_NDCSpawnDataBuffer;
 //Helper struct with various info needed for reading that can be passed around more easily in generated code.
 struct FNDCAccessContext_{ParameterName}
 {
-	bool bValid;
+	bool bValidParameter;
+	bool bValidReadIndex;
 
 	//Index at which to read data.
 	int ReadIndex;
@@ -61,8 +62,6 @@ struct FNDCAccessContext_{ParameterName}
 	//Actual number of NDC items in the buffers. Must come from an indirect instance count instead of a uniform as the GPU could have written to this buffer.
 	int Num() { return RWInstanceCounts[ GetInstanceCountOffset() ]; }
 
-	bool IsValid() { return bValid; }
-
 	int GetConsumeIndex()
 	{
 		//TODO: Add GS and Wave op path to reduce atomic contention.
@@ -71,32 +70,46 @@ struct FNDCAccessContext_{ParameterName}
 		int OutIndex = -1;
 		if(GetConsumeInstanceCountOffset() != -1)
 		{
-			InterlockedAdd(RWInstanceCounts[GetConsumeInstanceCountOffset()], 1, OutIndex);
-			return OutIndex >= 0 && OutIndex < Num() && OutIndex < GetBufferSize() ? OutIndex : -1;
+			uint OrigVal = 0;
+			InterlockedAdd(RWInstanceCounts[GetConsumeInstanceCountOffset()], 1U, OrigVal);
+			int MaxSize = min(Num(), GetBufferSize());
+			if(OrigVal >= 0 && OrigVal < (uint)MaxSize)
+			{
+				OutIndex = (int)OrigVal;
+			}
+			else
+			{
+				OutIndex = -1;
+				InterlockedMin(RWInstanceCounts[GetConsumeInstanceCountOffset()], (uint)MaxSize);
+			}
 		}
 
-		return -1;
+		return OutIndex;
 	}
 	
-	void InitDirect(int InReadIndex)
+	bool InitDirect(int InReadIndex)
 	{
 		ReadIndex = InReadIndex;
-		bValid = ReadIndex >= 0 && ReadIndex < Num() && ReadIndex < GetBufferSize();
+		bValidReadIndex = ReadIndex >= 0 && ReadIndex < Num() && ReadIndex < GetBufferSize();
+		return bValidReadIndex;
 	}
 
-	void InitConsume()
+	bool InitConsume()
 	{
 		ReadIndex = GetConsumeIndex();
-		bValid = ReadIndex >= 0 && ReadIndex < Num() && ReadIndex < GetBufferSize();
+		bValidReadIndex = ReadIndex >= 0 && ReadIndex < Num() && ReadIndex < GetBufferSize();
+		return bValidReadIndex;
 	}
 
-	void InitForParameter(int ParameterIndex)
+	bool InitForParameter(int ParameterIndex)
 	{
-		bValid = bValid && GetParameterLayoutInfo_{ParameterName}(ParameterIndex, ParameterComponentStart_Float, ParameterComponentStart_Int32);//TODO: Half Support | ParameterComponentStart_Half);
+		bValidParameter = GetParameterLayoutInfo_{ParameterName}(ParameterIndex, ParameterComponentStart_Float, ParameterComponentStart_Int32);//TODO: Half Support | ParameterComponentStart_Half);
 		
 		Register_Float = ParameterComponentStart_Float;
 		Register_Int32 = ParameterComponentStart_Int32;
 		//TODO: Half Support | Register_Half = ParameterComponentStart_Half;
+	
+		return bValidParameter;
 	}
 
 	void Read_Float(out float Value){ NiagaraDataChannelRead_Float(GetBuffer_Float(), GetFloatStride(), Register_Float, ReadIndex, Value); }
diff --git a/Engine/Plugins/FX/Niagara/Shaders/Private/DataChannel/NiagaraDataInterfaceDataChannelTemplateWriteCommon.ush b/Engine/Plugins/FX/Niagara/Shaders/Private/DataChannel/NiagaraDataInterfaceDataChannelTemplateWriteCommon.ush
--- a/Engine/Plugins/FX/Niagara/Shaders/Private/DataChannel/NiagaraDataInterfaceDataChannelTemplateWriteCommon.ush
+++ b/Engine/Plugins/FX/Niagara/Shaders/Private/DataChannel/NiagaraDataInterfaceDataChannelTemplateWriteCommon.ush
@@ -33,7 +33,8 @@ int {ParameterName}_CPUInt32Stride;
 //Helper struct with various info needed for writing that can be passed around more easily in generated code.
 struct FNDCAccessContext_{ParameterName}
 {
-	bool bValid;
+	bool bValidWriteIndex;
+	bool bValidParameter;
 
 	//Indices we'll be writing to each buffer.
 	int WriteIndex;
@@ -72,10 +73,9 @@ struct FNDCAccessContext_{ParameterName}
 	RWBuffer<float> GetCurrBuffer_Float() { return bReadGPUOrCPU ? GetGPUBuffer_Float() : GetCPUBuffer_Float(); }
 	RWBuffer<int> GetCurrBuffer_Int32() { return bReadGPUOrCPU ? GetGPUBuffer_Int32() : GetCPUBuffer_Int32(); }
 	//TODO: Half Support | RWBuffer<half> GetCurrBuffer_Half() { return bReadGPUOrCPU ? GetGPUBuffer_Half() : GetCPUBuffer_Half(); }
-
-	bool IsValid() { return bValid; }
-	bool ShouldWriteGPUData() { return bValid && GetGPUInstanceCountOffset() != -1; }
-	bool ShouldWriteCPUData() { return bValid && GetCPUInstanceCountOffset() != -1; }
+	
+	bool ShouldWriteGPUData() { return bValidParameter && bValidWriteIndex && GetGPUInstanceCountOffset() != -1; }
+	bool ShouldWriteCPUData() { return bValidParameter && bValidWriteIndex && GetCPUInstanceCountOffset() != -1; }
 
 	int GetAppendIndex(int InstanceCountOffset, int BufferSize)
 	{
@@ -87,33 +87,39 @@ struct FNDCAccessContext_{ParameterName}
 		{
 			//Note: This allows the index to temporarily exceed the max limits so is unsafe if we access this concurrently anywhere else without checking the limits.
 			//Note: However it does avoid a more expensive looping compare exchange.
-			InterlockedAdd(RWInstanceCounts[InstanceCountOffset], 1, OutIndex);
-			if(OutIndex >= 0 && OutIndex < BufferSize)
+			uint OrigVal = 0;
+			InterlockedAdd(RWInstanceCounts[InstanceCountOffset], 1U, OrigVal);
+			if(OrigVal >= 0 && OrigVal < (uint)BufferSize)
+			{ 
+				OutIndex = (int)OrigVal;
+			}
+			else
 			{
-				return OutIndex;
+				OutIndex = -1;
+				InterlockedMin(RWInstanceCounts[InstanceCountOffset], (uint)BufferSize);
 			}
-		
-			OutIndex = -1;
-			InterlockedMin(RWInstanceCounts[InstanceCountOffset], BufferSize);
 		}
 
 		return OutIndex;
 	}
 	
 	//Initialize our append indices. Called at the start of the append funciton.
-	void InitAppend()
+	bool InitAppend()
 	{
 		WriteIndexGPU = GetAppendIndex(GetGPUInstanceCountOffset(), GetGPUBufferSize());
 		WriteIndexCPU = GetAppendIndex(GetCPUInstanceCountOffset(), GetCPUBufferSize());
 	
-		bValid = (WriteIndexGPU >= 0 && WriteIndexGPU < GetGPUBufferSize()) ||
+		bValidWriteIndex = (WriteIndexGPU >= 0 && WriteIndexGPU < GetGPUBufferSize()) ||
 				(WriteIndexCPU  >= 0 && WriteIndexCPU < GetCPUBufferSize());
+	
+		return bValidWriteIndex;
 	}
 
 	//Initialize for the given parameter index.
-	void InitForParameter(int ParameterIndex)
+	bool InitForParameter(int ParameterIndex)
 	{
-		bValid = bValid && GetParameterLayoutInfo_{ParameterName}(ParameterIndex, ParameterComponentStart_Float, ParameterComponentStart_Int32);//TODO: Half Support | , ParameterComponentStart_Half);				
+		bValidParameter = GetParameterLayoutInfo_{ParameterName}(ParameterIndex, ParameterComponentStart_Float, ParameterComponentStart_Int32);//TODO: Half Support | , ParameterComponentStart_Half);
+		return bValidParameter;
 	}
 	
 	//Initialize for wrting the current parameter to the GPU buffers
diff --git a/Engine/Plugins/FX/Niagara/Shaders/Private/DataChannel/NiagaraDataInterfaceDataChannelTemplate_Append.ush b/Engine/Plugins/FX/Niagara/Shaders/Private/DataChannel/NiagaraDataInterfaceDataChannelTemplate_Append.ush
--- a/Engine/Plugins/FX/Niagara/Shaders/Private/DataChannel/NiagaraDataInterfaceDataChannelTemplate_Append.ush
+++ b/Engine/Plugins/FX/Niagara/Shaders/Private/DataChannel/NiagaraDataInterfaceDataChannelTemplate_Append.ush
@@ -14,8 +14,7 @@ void {FunctionSymbol}(bool bEmit, {FunctionInputParameters}, out bool bOutSucces
 	if(bEmit)
 	{
 		FNDCAccessContext_{ParameterName} Context;
-		Context.InitAppend();
-		if(Context.IsValid())
+		if(Context.InitAppend())
 		{
 			bOutSuccess = true;
 	
diff --git a/Engine/Plugins/FX/Niagara/Shaders/Private/DataChannel/NiagaraDataInterfaceDataChannelTemplate_Consume.ush b/Engine/Plugins/FX/Niagara/Shaders/Private/DataChannel/NiagaraDataInterfaceDataChannelTemplate_Consume.ush
--- a/Engine/Plugins/FX/Niagara/Shaders/Private/DataChannel/NiagaraDataInterfaceDataChannelTemplate_Consume.ush
+++ b/Engine/Plugins/FX/Niagara/Shaders/Private/DataChannel/NiagaraDataInterfaceDataChannelTemplate_Consume.ush
@@ -15,14 +15,17 @@ void {FunctionSymbol}(bool bInConsume, out bool bOutSuccess, out int OutIndex, {
 	bOutSuccess = false;
 	OutIndex = -1;
 	
-	FNDCAccessContext_{ParameterName} Context;
-	Context.InitConsume();
-	if(Context.IsValid())
+	if(bInConsume)
 	{
-		bOutSuccess = true;
+		FNDCAccessContext_{ParameterName} Context;
+		if(Context.InitConsume())
+		{
+			bOutSuccess = true;
+			OutIndex = Context.ReadIndex;
 
-		//This is generated in the DI Cpp depending on the function parameters.
-{PerFunctionParameterShaderCode}
+			//This is generated in the DI Cpp depending on the function parameters.
+			{PerFunctionParameterShaderCode}
+		}
 	}
 }
 
diff --git a/Engine/Plugins/FX/Niagara/Shaders/Private/DataChannel/NiagaraDataInterfaceDataChannelTemplate_Read.ush b/Engine/Plugins/FX/Niagara/Shaders/Private/DataChannel/NiagaraDataInterfaceDataChannelTemplate_Read.ush
--- a/Engine/Plugins/FX/Niagara/Shaders/Private/DataChannel/NiagaraDataInterfaceDataChannelTemplate_Read.ush
+++ b/Engine/Plugins/FX/Niagara/Shaders/Private/DataChannel/NiagaraDataInterfaceDataChannelTemplate_Read.ush
@@ -15,8 +15,7 @@ void {FunctionSymbol}(int ElementIndex, out bool bOutSuccess, {FunctionOutputPar
 	bOutSuccess = false;
 
 	FNDCAccessContext_{ParameterName} Context;
-	Context.InitDirect(ElementIndex);
-	if(Context.IsValid())
+	if(Context.InitDirect(ElementIndex))
 	{
 		bOutSuccess = true;
 
diff --git a/Engine/Plugins/FX/Niagara/Source/Niagara/Internal/DataInterface/NiagaraDataInterfaceDataChannelCommon.h b/Engine/Plugins/FX/Niagara/Source/Niagara/Internal/DataInterface/NiagaraDataInterfaceDataChannelCommon.h
--- a/Engine/Plugins/FX/Niagara/Source/Niagara/Internal/DataInterface/NiagaraDataInterfaceDataChannelCommon.h
+++ b/Engine/Plugins/FX/Niagara/Source/Niagara/Internal/DataInterface/NiagaraDataInterfaceDataChannelCommon.h
@@ -61,18 +61,24 @@ uint32 GetTypeHash(const FNDIDataChannelFunctionInfo& FuncInfo);
 struct FNDIDataChannelRegisterBinding
 {
 	static const uint32 RegisterBits = 30;
-	static const uint32 DataTypeBitst = 2;
+	static const uint32 DataTypeBits = 2;
 	FNDIDataChannelRegisterBinding(uint32 InFunctionRegisterIndex, uint32 InDataSetRegisterIndex, ENiagaraBaseTypes InDataType)
-	: FunctionRegisterIndex(InFunctionRegisterIndex)
-	, DataSetRegisterIndex(InDataSetRegisterIndex)
+	: DataSetRegisterIndex(InDataSetRegisterIndex)
+	, FunctionRegisterIndex(InFunctionRegisterIndex) 
 	, DataType((uint32)InDataType)
 	{
-		check(InDataSetRegisterIndex <= (1u << RegisterBits) - 1);
-		check((uint32)InDataType <= (1u << DataTypeBitst) - 1);
+		check(FunctionRegisterIndex <= (1u << RegisterBits) - 1);
+		check((uint32)InDataType <= (1u << DataTypeBits) - 1);
 	}
-	uint32 FunctionRegisterIndex;
-	uint32 DataSetRegisterIndex : RegisterBits;
-	uint32 DataType : DataTypeBitst;
+
+	uint32 GetDataSetRegisterIndex()const { return DataSetRegisterIndex; }
+	uint32 GetFunctionRegisterIndex()const { return FunctionRegisterIndex; }
+	ENiagaraBaseTypes GetDataType()const { return (ENiagaraBaseTypes)DataType; }
+
+private:
+	uint32 DataSetRegisterIndex;
+	uint32 FunctionRegisterIndex : RegisterBits;
+	uint32 DataType : DataTypeBits;
 };
 
 
@@ -266,12 +272,12 @@ struct FNDIVariadicInputHandler
 			HalfInputs.Reserve(BindingPtr->NumHalfComponents);
 			for (const FNDIDataChannelRegisterBinding& VMBinding : BindingPtr->VMRegisterBindings)
 			{
-				switch (VMBinding.DataType)
+				switch (VMBinding.GetDataType())
 				{
-				case (int32)ENiagaraBaseTypes::Float: FloatInputs.Emplace(Context); break;
-				case (int32)ENiagaraBaseTypes::Int32: IntInputs.Emplace(Context); break;
-				case (int32)ENiagaraBaseTypes::Bool: IntInputs.Emplace(Context); break;
-				case (int32)ENiagaraBaseTypes::Half: HalfInputs.Emplace(Context); break;
+				case ENiagaraBaseTypes::Float: FloatInputs.Emplace(Context); break;
+				case ENiagaraBaseTypes::Int32: IntInputs.Emplace(Context); break;
+				case ENiagaraBaseTypes::Bool: IntInputs.Emplace(Context); break;
+				case ENiagaraBaseTypes::Half: HalfInputs.Emplace(Context); break;
 				default: check(0);
 				};
 			}
@@ -300,12 +306,12 @@ struct FNDIVariadicInputHandler
 			//TODO: Optimize for long runs of writes to reduce binding/lookup overhead.
 			for (const FNDIDataChannelRegisterBinding& VMBinding : BindingInfo->VMRegisterBindings)
 			{
-				switch(VMBinding.DataType)
+				switch(VMBinding.GetDataType())
 				{
-					case (int32)ENiagaraBaseTypes::Float: FloatFunc(VMBinding, FloatInputs[VMBinding.FunctionRegisterIndex]); break;
-					case (int32)ENiagaraBaseTypes::Int32: IntFunc(VMBinding, IntInputs[VMBinding.FunctionRegisterIndex]); break;
-					case (int32)ENiagaraBaseTypes::Bool: IntFunc(VMBinding, IntInputs[VMBinding.FunctionRegisterIndex]); break;
-					case (int32)ENiagaraBaseTypes::Half: HalfFunc(VMBinding, HalfInputs[VMBinding.FunctionRegisterIndex]); break;
+					case ENiagaraBaseTypes::Float: FloatFunc(VMBinding, FloatInputs[VMBinding.GetFunctionRegisterIndex()]); break;
+					case ENiagaraBaseTypes::Int32: IntFunc(VMBinding, IntInputs[VMBinding.GetFunctionRegisterIndex()]); break;
+					case ENiagaraBaseTypes::Bool: IntFunc(VMBinding, IntInputs[VMBinding.GetFunctionRegisterIndex()]); break;
+					case ENiagaraBaseTypes::Half: HalfFunc(VMBinding, HalfInputs[VMBinding.GetFunctionRegisterIndex()]); break;
 					default: check(0);
 				};
 			}
@@ -335,12 +341,12 @@ struct FNDIVariadicOutputHandler
 			HalfOutputs.Reserve(BindingPtr->NumHalfComponents);
 			for (const FNDIDataChannelRegisterBinding& VMBinding : BindingPtr->VMRegisterBindings)
 			{
-				switch (VMBinding.DataType)
+				switch (VMBinding.GetDataType())
 				{
-				case (int32)ENiagaraBaseTypes::Float: FloatOutputs.Emplace(Context); break;
-				case (int32)ENiagaraBaseTypes::Int32: IntOutputs.Emplace(Context); break;
-				case (int32)ENiagaraBaseTypes::Bool: IntOutputs.Emplace(Context); break;
-				case (int32)ENiagaraBaseTypes::Half: HalfOutputs.Emplace(Context); break;
+				case ENiagaraBaseTypes::Float: FloatOutputs.Emplace(Context); break;
+				case ENiagaraBaseTypes::Int32: IntOutputs.Emplace(Context); break;
+				case ENiagaraBaseTypes::Bool: IntOutputs.Emplace(Context); break;
+				case ENiagaraBaseTypes::Half: HalfOutputs.Emplace(Context); break;
 				default: check(0);
 				};
 			}
@@ -355,12 +361,12 @@ struct FNDIVariadicOutputHandler
 			//TODO: Optimize for long runs of writes to reduce binding/lookup overhead.
 			for (const FNDIDataChannelRegisterBinding& VMBinding : BindingInfo->VMRegisterBindings)
 			{
-				switch (VMBinding.DataType)
+				switch (VMBinding.GetDataType())
 				{
-				case (int32)ENiagaraBaseTypes::Float: FloatFunc(VMBinding, FloatOutputs[VMBinding.FunctionRegisterIndex]); break;
-				case (int32)ENiagaraBaseTypes::Int32: IntFunc(VMBinding, IntOutputs[VMBinding.FunctionRegisterIndex]); break;
-				case (int32)ENiagaraBaseTypes::Bool: IntFunc(VMBinding, IntOutputs[VMBinding.FunctionRegisterIndex]); break;
-				case (int32)ENiagaraBaseTypes::Half: HalfFunc(VMBinding, HalfOutputs[VMBinding.FunctionRegisterIndex]); break;
+				case ENiagaraBaseTypes::Float: FloatFunc(VMBinding, FloatOutputs[VMBinding.GetFunctionRegisterIndex()]); break;
+				case ENiagaraBaseTypes::Int32: IntFunc(VMBinding, IntOutputs[VMBinding.GetFunctionRegisterIndex()]); break;
+				case ENiagaraBaseTypes::Bool: IntFunc(VMBinding, IntOutputs[VMBinding.GetFunctionRegisterIndex()]); break;
+				case ENiagaraBaseTypes::Half: HalfFunc(VMBinding, HalfOutputs[VMBinding.GetFunctionRegisterIndex()]); break;
 				default: check(0);
 				};
 			}
diff --git a/Engine/Plugins/FX/Niagara/Source/Niagara/Internal/DataInterface/NiagaraDataInterfaceDataChannelRead.h b/Engine/Plugins/FX/Niagara/Source/Niagara/Internal/DataInterface/NiagaraDataInterfaceDataChannelRead.h
--- a/Engine/Plugins/FX/Niagara/Source/Niagara/Internal/DataInterface/NiagaraDataInterfaceDataChannelRead.h
+++ b/Engine/Plugins/FX/Niagara/Source/Niagara/Internal/DataInterface/NiagaraDataInterfaceDataChannelRead.h
@@ -288,7 +288,6 @@ struct FNDIDataChannelReadInstanceData
 	*/
 	TMap<FNiagaraEmitterInstance*,FNDIDataChannelRead_EmitterInstanceData> EmitterInstanceData;
 
-	uint32 CachedLayoutHash = INDEX_NONE;
 	FNiagaraSystemInstance* Owner = nullptr;
 
 	virtual ~FNDIDataChannelReadInstanceData();
diff --git a/Engine/Plugins/FX/Niagara/Source/Niagara/Internal/NiagaraDataChannelCommon.h b/Engine/Plugins/FX/Niagara/Source/Niagara/Internal/NiagaraDataChannelCommon.h
--- a/Engine/Plugins/FX/Niagara/Source/Niagara/Internal/NiagaraDataChannelCommon.h
+++ b/Engine/Plugins/FX/Niagara/Source/Niagara/Internal/NiagaraDataChannelCommon.h
@@ -102,6 +102,9 @@ struct FNiagaraDataChannelData final : public TSharedFromThis<FNiagaraDataChanne
 	void RegisterGPUSpawningReader() { ++NumGPUSpawningReaders; }
 	void UnregisterGPUSpawningReader() { --NumGPUSpawningReaders; }
 	int32 NumRegisteredGPUSpawningReaders()const{ return NumGPUSpawningReaders; }
+
+	//Returns if this data is still valid. This can return false in cases where the owning data channel has been modified for example.
+	bool IsLayoutValid(UNiagaraDataChannelHandler* Owner)const;
 private:
 
 	void CreateRenderThreadProxy(UNiagaraDataChannelHandler* Owner);
diff --git a/Engine/Plugins/FX/Niagara/Source/Niagara/Private/DataInterface/NiagaraDataInterfaceDataChannelCommon.cpp b/Engine/Plugins/FX/Niagara/Source/Niagara/Private/DataInterface/NiagaraDataInterfaceDataChannelCommon.cpp
--- a/Engine/Plugins/FX/Niagara/Source/Niagara/Private/DataInterface/NiagaraDataInterfaceDataChannelCommon.cpp
+++ b/Engine/Plugins/FX/Niagara/Source/Niagara/Private/DataInterface/NiagaraDataInterfaceDataChannelCommon.cpp
@@ -36,24 +36,23 @@ FNDIDataChannel_FunctionToDataSetBinding::FNDIDataChannel_FunctionToDataSetBindi
 		for (int32 ParamIdx = 0; ParamIdx < Parameters.Num(); ++ParamIdx)
 		{
 			const FNiagaraVariableBase& Param = Parameters[ParamIdx];
+			uint32 DataSetFloatRegister = INDEX_NONE;
+			uint32 DataSetIntRegister = INDEX_NONE;
+			uint32 DataSetHalfRegister = INDEX_NONE;
 			if (const FNiagaraVariableLayoutInfo* DataSetVariableLayout = DataSetLayout.FindVariableLayoutInfo(Param))
 			{
-				uint32 DataSetFloatRegister = DataSetVariableLayout->GetFloatComponentStart();
-				uint32 DataSetIntRegister = DataSetVariableLayout->GetInt32ComponentStart();
-				uint32 DataSetHalfRegister = DataSetVariableLayout->GetHalfComponentStart();
-
-				GenVMBindings(Param, Param.GetType().GetStruct(), NumFloatComponents, NumInt32Components, NumHalfComponents, DataSetFloatRegister, DataSetIntRegister, DataSetHalfRegister);
+				DataSetFloatRegister = DataSetVariableLayout->GetFloatComponentStart();
+				DataSetIntRegister = DataSetVariableLayout->GetInt32ComponentStart();
+				DataSetHalfRegister = DataSetVariableLayout->GetHalfComponentStart();
 			}
 			else
 			{
-				DataSetLayoutHash = 0;
-
 #if !UE_BUILD_SHIPPING && !UE_BUILD_TEST
 				OutMissingParams.Emplace(Param);
-#else
-				return;
 #endif
 			}
+
+			GenVMBindings(Param, Param.GetType().GetStruct(), NumFloatComponents, NumInt32Components, NumHalfComponents, DataSetFloatRegister, DataSetIntRegister, DataSetHalfRegister);
 		}
 	};
 
@@ -403,26 +402,24 @@ namespace NDIDataChannelUtilities
 		//Template code for writing handling code for each of the function parameters.
 		//Some preamble and an replacement arg into which we write all the actual I/O with the DataChannel buffers.
 		FString PerParameterReadTemplate = TEXT("\n\
-void Read_{FunctionParameterName}_{ParameterName}(FNDCAccessContext_{ParameterName} Context, out {FunctionParameterType} {FunctionParameterName})\n\
+void Read_{FunctionParameterName}_{ParameterName}(FNDCAccessContext_{ParameterName} Context, inout bool bSuccess, inout {FunctionParameterType} {FunctionParameterName})\n\
 {\n\
-	Context.InitForParameter({FunctionParameterIndex});\n\
-	if(Context.IsValid())\n\
+	if(Context.InitForParameter({FunctionParameterIndex}))\n\
 	{\n\
 		{FuncParamShaderCode}\
 	}\n\
 	else\n\
 	{\n\
-		{FunctionParameterName} = 0;\n\
+		bSuccess = false;\n\
 	}\n\
 }\n");
 
-		FString PerParameterReadCallTemplate = TEXT("Read_{FunctionParameterName}_{ParameterName}(Context, {FunctionParameterName});\n");
+		FString PerParameterReadCallTemplate = TEXT("Read_{FunctionParameterName}_{ParameterName}(Context, bOutSuccess, {FunctionParameterName});\n");
 
 		FString PerParameterWriteTemplate = TEXT("\n\
-void Write_{FunctionParameterName}_{ParameterName}(FNDCAccessContext_{ParameterName} Context, {FunctionParameterType} {FunctionParameterName})\n\
+void Write_{FunctionParameterName}_{ParameterName}(FNDCAccessContext_{ParameterName} Context, inout bool bSuccess, {FunctionParameterType} {FunctionParameterName})\n\
 {\n\
-	Context.InitForParameter({FunctionParameterIndex});\n\
-	if(Context.IsValid())\n\
+	if(Context.InitForParameter({FunctionParameterIndex}))\n\
 	{\n\
 		if(Context.InitForGPUWrite())\n\
 		{\n\
@@ -433,8 +430,12 @@ void Write_{FunctionParameterName}_{ParameterName}(FNDCAccessContext_{ParameterN
 			{FuncParamShaderCode}\n\
 		}\n\
 	}\n\
+	else\n\
+	{\n\
+		bSuccess = false;\n\
+	}\n\
 }\n");
-		FString PerParameterWriteCallTemplate = TEXT("Write_{FunctionParameterName}_{ParameterName}(Context, {FunctionParameterName});\n");
+		FString PerParameterWriteCallTemplate = TEXT("Write_{FunctionParameterName}_{ParameterName}(Context, bOutSuccess, {FunctionParameterName});\n");
 
 		//Template code for accessing data from the Data Channel's buffers.
 		static const FString ReadDataTemplate = TEXT("Context.Read_{FunctionParameterComponentBufferType}({FunctionParameterComponentName});\n");
diff --git a/Engine/Plugins/FX/Niagara/Source/Niagara/Private/DataInterface/NiagaraDataInterfaceDataChannelRead.cpp b/Engine/Plugins/FX/Niagara/Source/Niagara/Private/DataInterface/NiagaraDataInterfaceDataChannelRead.cpp
--- a/Engine/Plugins/FX/Niagara/Source/Niagara/Private/DataInterface/NiagaraDataInterfaceDataChannelRead.cpp
+++ b/Engine/Plugins/FX/Niagara/Source/Niagara/Private/DataInterface/NiagaraDataInterfaceDataChannelRead.cpp
@@ -476,7 +476,8 @@ bool FNDIDataChannelReadInstanceData::Tick(UNiagaraDataInterfaceDataChannelRead*
 		//Grab the world DataChannel data if we're reading from there.	
 		if (DataChannelPtr)
 		{
-			if(DataChannelData == nullptr || Interface->bUpdateSourceDataEveryTick)
+			bool bNDCDataIsValid = DataChannelData && DataChannelData->IsLayoutValid(DataChannelPtr);
+			if(bNDCDataIsValid == false || Interface->bUpdateSourceDataEveryTick)
 			{
 				//TODO: Automatically modify tick group if we have DIs that require current frame info?
 				FNiagaraDataChannelSearchParameters SearchParams(Instance->GetAttachComponent());
@@ -1514,33 +1515,51 @@ void UNiagaraDataInterfaceDataChannelRead::Read(FVectorVMExternalFunctionContext
 			int32 Index = InIndex.GetAndAdvance();
 
 			bool bProcess = (uint32)Index < Data->GetNumInstances();
+			bool bAllReadsSuccess = true;
 
-			auto FloatFunc = [Data, Index](const FNDIDataChannelRegisterBinding& VMBinding, VectorVM::FExternalFuncRegisterHandler<float>& FloatData)
+			auto FloatFunc = [Data, Index, &bAllReadsSuccess](const FNDIDataChannelRegisterBinding& VMBinding, VectorVM::FExternalFuncRegisterHandler<float>& FloatData)
 			{
-				if (VMBinding.DataSetRegisterIndex != INDEX_NONE)
+				if (VMBinding.GetDataSetRegisterIndex() != INDEX_NONE)
 				{
-					float* Src = Data->GetInstancePtrFloat(VMBinding.DataSetRegisterIndex, (uint32)Index);
+					float* Src = Data->GetInstancePtrFloat(VMBinding.GetDataSetRegisterIndex(), (uint32)Index);
 					*FloatData.GetDestAndAdvance() = *Src;
 				}
+				else
+				{
+					bAllReadsSuccess = false;
+					*FloatData.GetDestAndAdvance() = 0.0f;
+				}
 			};
-			auto IntFunc = [Data, Index](const FNDIDataChannelRegisterBinding& VMBinding, VectorVM::FExternalFuncRegisterHandler<int32>& IntData)
+			auto IntFunc = [Data, Index, &bAllReadsSuccess](const FNDIDataChannelRegisterBinding& VMBinding, VectorVM::FExternalFuncRegisterHandler<int32>& IntData)
 			{
-				if (VMBinding.DataSetRegisterIndex != INDEX_NONE)
+				if (VMBinding.GetDataSetRegisterIndex() != INDEX_NONE)
 				{
-					int32* Src = Data->GetInstancePtrInt32(VMBinding.DataSetRegisterIndex, (uint32)Index);
+					int32* Src = Data->GetInstancePtrInt32(VMBinding.GetDataSetRegisterIndex(), (uint32)Index);
 					*IntData.GetDestAndAdvance() = *Src;
 				}
+				else
+				{
+					bAllReadsSuccess = false;
+					*IntData.GetDestAndAdvance() = 0;
+				}
 			};
-			auto HalfFunc = [Data, Index](const FNDIDataChannelRegisterBinding& VMBinding, VectorVM::FExternalFuncRegisterHandler<FFloat16>& HalfData)
+			auto HalfFunc = [Data, Index, &bAllReadsSuccess](const FNDIDataChannelRegisterBinding& VMBinding, VectorVM::FExternalFuncRegisterHandler<FFloat16>& HalfData)
 			{
-				if (VMBinding.DataSetRegisterIndex != INDEX_NONE)
-					*HalfData.GetDestAndAdvance() = *Data->GetInstancePtrHalf(VMBinding.DataSetRegisterIndex, (uint32)Index);
+				if (VMBinding.GetDataSetRegisterIndex() != INDEX_NONE)
+				{
+					*HalfData.GetDestAndAdvance() = *Data->GetInstancePtrHalf(VMBinding.GetDataSetRegisterIndex(), (uint32)Index);
+				}
+				else
+				{
+					bAllReadsSuccess = false;
+					*HalfData.GetDestAndAdvance() = 0.0f;
+				}				
 			};
-			bool bSuccess = VariadicOutputs.Process(bProcess, 1, BindingInfo, FloatFunc, IntFunc, HalfFunc);
+			bool bFinalSuccess = VariadicOutputs.Process(bProcess, 1, BindingInfo, FloatFunc, IntFunc, HalfFunc) && bAllReadsSuccess;
 
 			if (OutSuccess.IsValid())
 			{
-				OutSuccess.SetAndAdvance(bSuccess);
+				OutSuccess.SetAndAdvance(bFinalSuccess);
 			}
 		}
 	}
@@ -1578,41 +1597,70 @@ void UNiagaraDataInterfaceDataChannelRead::Consume(FVectorVMExternalFunctionCont
 			bool bConsume = InConsume.GetAndAdvance();
 
 			bool bSuccess = false;
+			bool bNeedsFallback = true;
 			int32 Index = INDEX_NONE;
+
 			if (bConsume)
 			{
 				//Increment counter and enforce max if the result is over acceptable values.
 				//Note: This allows the index to temporarily exceed the max limits so is unsafe if we access this concurrently anywhere else without checking the limits.
 				//Note: However it does avoid a more expensive looping compare exchange.
 				Index = InstData->ConsumeIndex++;
+				bool bAllReadsSuccess = true;
+
 				if(Index >= 0 && Index < (int32)Data->GetNumInstances())
 				{
 					//TODO: Wrap/clamp modes etc
 
-					auto FloatFunc = [Data, Index](const FNDIDataChannelRegisterBinding& VMBinding, VectorVM::FExternalFuncRegisterHandler<float>& FloatData)
+					auto FloatFunc = [Data, Index, &bAllReadsSuccess](const FNDIDataChannelRegisterBinding& VMBinding, VectorVM::FExternalFuncRegisterHandler<float>& FloatData)
 					{
-						if (VMBinding.DataSetRegisterIndex != INDEX_NONE)
-							*FloatData.GetDestAndAdvance() = *Data->GetInstancePtrFloat(VMBinding.DataSetRegisterIndex, (uint32)Index);
+						if (VMBinding.GetDataSetRegisterIndex() != INDEX_NONE)
+						{
+							*FloatData.GetDestAndAdvance() = *Data->GetInstancePtrFloat(VMBinding.GetDataSetRegisterIndex(), (uint32)Index);
+						}
+						else
+						{
+							bAllReadsSuccess = false;
+							*FloatData.GetDestAndAdvance() = 0.0f;
+						}
 					};
-					auto IntFunc = [Data, Index](const FNDIDataChannelRegisterBinding& VMBinding, VectorVM::FExternalFuncRegisterHandler<int32>& IntData)
+					auto IntFunc = [Data, Index, &bAllReadsSuccess](const FNDIDataChannelRegisterBinding& VMBinding, VectorVM::FExternalFuncRegisterHandler<int32>& IntData)
 					{
-						if (VMBinding.DataSetRegisterIndex != INDEX_NONE)
-							*IntData.GetDestAndAdvance() = *Data->GetInstancePtrInt32(VMBinding.DataSetRegisterIndex, (uint32)Index);
+						if (VMBinding.GetDataSetRegisterIndex() != INDEX_NONE)
+						{
+							*IntData.GetDestAndAdvance() = *Data->GetInstancePtrInt32(VMBinding.GetDataSetRegisterIndex(), (uint32)Index);
+						}
+						else
+						{
+							bAllReadsSuccess = false;
+							*IntData.GetDestAndAdvance() = 0;
+						}
 					};
-					auto HalfFunc = [Data, Index](const FNDIDataChannelRegisterBinding& VMBinding, VectorVM::FExternalFuncRegisterHandler<FFloat16>& HalfData)
+					auto HalfFunc = [Data, Index, &bAllReadsSuccess](const FNDIDataChannelRegisterBinding& VMBinding, VectorVM::FExternalFuncRegisterHandler<FFloat16>& HalfData)
 					{
-						if (VMBinding.DataSetRegisterIndex != INDEX_NONE)
-							*HalfData.GetDestAndAdvance() = *Data->GetInstancePtrHalf(VMBinding.DataSetRegisterIndex, (uint32)Index);
+						if (VMBinding.GetDataSetRegisterIndex() != INDEX_NONE)
+						{
+							*HalfData.GetDestAndAdvance() = *Data->GetInstancePtrHalf(VMBinding.GetDataSetRegisterIndex(), (uint32)Index);
+						}
+						else
+						{
+							bAllReadsSuccess = false;
+							*HalfData.GetDestAndAdvance() = 0.0f;
+						}
 					};
-					bSuccess = VariadicOutputs.Process(bConsume, 1, BindingInfo, FloatFunc, IntFunc, HalfFunc);
+					
+					bNeedsFallback = VariadicOutputs.Process(bConsume, 1, BindingInfo, FloatFunc, IntFunc, HalfFunc) == false;
+					bSuccess = !bNeedsFallback && bAllReadsSuccess;
 				}
 				else
 				{
+					Index = INDEX_NONE;
+					bSuccess = false;
 					InstData->ConsumeIndex = Data->GetNumInstances();
 				}
 			}
-		
-			if(bSuccess == false)
+			
+			if(bNeedsFallback)
 			{
 				VariadicOutputs.Fallback(1);
 			}
@@ -1623,7 +1671,7 @@ void UNiagaraDataInterfaceDataChannelRead::Consume(FVectorVMExternalFunctionCont
 			}
 			if (OutIndex.IsValid())
 			{
-				OutIndex.SetAndAdvance(bSuccess ? Index : INDEX_NONE);
+				OutIndex.SetAndAdvance(Index);
 			}
 		}
 	}
@@ -1713,18 +1761,37 @@ void UNiagaraDataInterfaceDataChannelRead::SpawnConditional(FVectorVMExternalFun
 			bool bConditionsPass = true;
 			auto FloatFunc = [&bConditionsPass, Op, Data, DataChannelIdx](const FNDIDataChannelRegisterBinding& VMBinding, FNDIInputParam<float>& FloatData)
 			{
-				if (VMBinding.DataSetRegisterIndex != INDEX_NONE)
-					bConditionsPass &= EvalConditional(Op, FloatData.GetAndAdvance(), *Data->GetInstancePtrFloat(VMBinding.DataSetRegisterIndex, (uint32)DataChannelIdx));
+				if (VMBinding.GetDataSetRegisterIndex() != INDEX_NONE)
+				{
+					bConditionsPass &= EvalConditional(Op, FloatData.GetAndAdvance(), *Data->GetInstancePtrFloat(VMBinding.GetDataSetRegisterIndex(), (uint32)DataChannelIdx));
+				}
+				else
+				{
+					bConditionsPass = false;
+				}
+				
 			};
 			auto IntFunc = [&bConditionsPass, Op, Data, DataChannelIdx](const FNDIDataChannelRegisterBinding& VMBinding, FNDIInputParam<int32>& IntData)
 			{
-				if (VMBinding.DataSetRegisterIndex != INDEX_NONE)
-					bConditionsPass &= EvalConditional(Op, IntData.GetAndAdvance(), *Data->GetInstancePtrInt32(VMBinding.DataSetRegisterIndex, (uint32)DataChannelIdx));
+				if (VMBinding.GetDataSetRegisterIndex() != INDEX_NONE)
+				{
+					bConditionsPass &= EvalConditional(Op, IntData.GetAndAdvance(), *Data->GetInstancePtrInt32(VMBinding.GetDataSetRegisterIndex(), (uint32)DataChannelIdx));
+				}
+				else
+				{
+					bConditionsPass = false;
+				}
 			};
 			auto HalfFunc = [&bConditionsPass, Op, Data, DataChannelIdx](const FNDIDataChannelRegisterBinding& VMBinding, FNDIInputParam<FFloat16>& HalfData)
 			{
-				if (VMBinding.DataSetRegisterIndex != INDEX_NONE)
-					bConditionsPass &= EvalConditional(Op, HalfData.GetAndAdvance(), *Data->GetInstancePtrHalf(VMBinding.DataSetRegisterIndex, (uint32)DataChannelIdx));
+				if (VMBinding.GetDataSetRegisterIndex() != INDEX_NONE)
+				{
+					bConditionsPass &= EvalConditional(Op, HalfData.GetAndAdvance(), *Data->GetInstancePtrHalf(VMBinding.GetDataSetRegisterIndex(), (uint32)DataChannelIdx));
+				}
+				else
+				{
+					bConditionsPass = false;
+				}
 			};
 			VariadicInputs.Process(true, 1, BindingInfo, FloatFunc, IntFunc, HalfFunc);
 			VariadicInputs.Reset();
@@ -1842,21 +1909,24 @@ void UNiagaraDataInterfaceDataChannelRead::SpawnDirect(FVectorVMExternalFunction
 
 		const auto ValueData = FNiagaraDataSetAccessor<TSimType>::CreateReader(Data, NDCVarName);
 
-		for (int32 DataChannelIndex = 0; DataChannelIndex < NumDataChannelInstances; DataChannelIndex++)
+		if (ValueData.IsValid())
 		{
-			TSimType NDCValue = ValueData.GetSafe(DataChannelIndex, NDCValueDefault<TSimType>());
-			double VarSize = NDCValueSize(NDCValue);
-			float Scale = RandHelper.RandRange(DataChannelIndex, RandMinScale, RandMaxScale);
-			int32 ScaledCount = FMath::TruncToInt32(VarSize * Scale);
-			int32 FinalCount = FMath::Clamp(ScaledCount, ClampMin, ClampMax);
-
-			if (Mode == ENDIDataChannelSpawnMode::Accumulate)
-			{
-				EmitterConditionalSpawns[DataChannelIndex].Append(FinalCount);
-			}
-			else if (Mode == ENDIDataChannelSpawnMode::Override)
+			for (int32 DataChannelIndex = 0; DataChannelIndex < NumDataChannelInstances; DataChannelIndex++)
 			{
-				EmitterConditionalSpawns[DataChannelIndex].SetCount(FinalCount);
+				TSimType NDCValue = ValueData.GetSafe(DataChannelIndex, NDCValueDefault<TSimType>());
+				double VarSize = NDCValueSize(NDCValue);
+				float Scale = RandHelper.RandRange(DataChannelIndex, RandMinScale, RandMaxScale);
+				int32 ScaledCount = FMath::TruncToInt32(VarSize * Scale);
+				int32 FinalCount = FMath::Clamp(ScaledCount, ClampMin, ClampMax);
+
+				if (Mode == ENDIDataChannelSpawnMode::Accumulate)
+				{
+					EmitterConditionalSpawns[DataChannelIndex].Append(FinalCount);
+				}
+				else if (Mode == ENDIDataChannelSpawnMode::Override)
+				{
+					EmitterConditionalSpawns[DataChannelIndex].SetCount(FinalCount);
+				}
 			}
 		}
 	}
@@ -1915,22 +1985,25 @@ void UNiagaraDataInterfaceDataChannelRead::ScaleSpawnCount(FVectorVMExternalFunc
 
 		const auto ValueData = FNiagaraDataSetAccessor<TSimType>::CreateReader(Data, NDCVarName);
 
-		for(int32 DataChannelIndex = 0; DataChannelIndex < NumDataChannelInstances; DataChannelIndex++)
+		if(ValueData.IsValid())
 		{
-			TSimType NDCValue = ValueData.GetSafe(DataChannelIndex, NDCValueDefault<TSimType>());
-			float VarSize = NDCValueSize(NDCValue);
-			float Scale = RandHelper.RandRange(DataChannelIndex, RandMinScale, RandMaxScale);			
-			float FinalScale = VarSize * Scale;
-			FinalScale = FMath::Clamp(FinalScale, ClampMin, ClampMax);
-
-			//TODO: Either change this to a float/double or add a separate scale value applied at the end so that multiple scales will combine correctly.
-			if(bOverrideScale)
-			{
-				EmitterConditionalSpawns[DataChannelIndex].SetScale(FinalScale);
-			}
-			else
+			for (int32 DataChannelIndex = 0; DataChannelIndex < NumDataChannelInstances; DataChannelIndex++)
 			{
-				EmitterConditionalSpawns[DataChannelIndex].ApplyScale(FinalScale);
+				TSimType NDCValue = ValueData.GetSafe(DataChannelIndex, NDCValueDefault<TSimType>());
+				float VarSize = NDCValueSize(NDCValue);
+				float Scale = RandHelper.RandRange(DataChannelIndex, RandMinScale, RandMaxScale);
+				float FinalScale = VarSize * Scale;
+				FinalScale = FMath::Clamp(FinalScale, ClampMin, ClampMax);
+
+				//TODO: Either change this to a float/double or add a separate scale value applied at the end so that multiple scales will combine correctly.
+				if (bOverrideScale)
+				{
+					EmitterConditionalSpawns[DataChannelIndex].SetScale(FinalScale);
+				}
+				else
+				{
+					EmitterConditionalSpawns[DataChannelIndex].ApplyScale(FinalScale);
+				}
 			}
 		}
 	}
@@ -2093,7 +2166,7 @@ void FNiagaraDataInterfaceProxy_DataChannelRead::PreStage(const FNDIGpuComputePr
 
 	if(InstanceData)
 	{
-		if(InstanceData->ChannelDataRTProxy)
+		if(InstanceData->ChannelDataRTProxy && InstanceData->GPUBuffer == nullptr)
 		{
 			InstanceData->GPUBuffer = InstanceData->ChannelDataRTProxy->PrepareForReadAccess(Context.GetGraphBuilder(), InstanceData->bReadPrevFrame == false);
 		}
@@ -2120,8 +2193,11 @@ void FNiagaraDataInterfaceProxy_DataChannelRead::PostStage(const FNDIGpuComputeP
 	FNiagaraDataInterfaceProxy_DataChannelRead::FInstanceData* InstanceData = SystemInstancesToProxyData_RT.Find(Context.GetSystemInstanceID());
 	if (InstanceData && InstanceData->ChannelDataRTProxy)
 	{
-		InstanceData->GPUBuffer = nullptr;
-		InstanceData->ChannelDataRTProxy->EndReadAccess(Context.GetGraphBuilder(), InstanceData->bReadPrevFrame == false);
+		if (InstanceData->GPUBuffer)
+		{
+			InstanceData->ChannelDataRTProxy->EndReadAccess(Context.GetGraphBuilder(), InstanceData->bReadPrevFrame == false);
+			InstanceData->GPUBuffer = nullptr;
+		}	
 	}
 }
 
diff --git a/Engine/Plugins/FX/Niagara/Source/Niagara/Private/DataInterface/NiagaraDataInterfaceDataChannelWrite.cpp b/Engine/Plugins/FX/Niagara/Source/Niagara/Private/DataInterface/NiagaraDataInterfaceDataChannelWrite.cpp
--- a/Engine/Plugins/FX/Niagara/Source/Niagara/Private/DataInterface/NiagaraDataInterfaceDataChannelWrite.cpp
+++ b/Engine/Plugins/FX/Niagara/Source/Niagara/Private/DataInterface/NiagaraDataInterfaceDataChannelWrite.cpp
@@ -300,6 +300,7 @@ struct FNDIDataChannelWriteInstanceData
 	{
 		DynamicAllocationCount = 0;
 		AtomicNumInstances = 0;
+		DestinationData = nullptr;
 		
 		LwcTile = Instance->GetLWCTile();
 		if (Interface->ShouldPublish())
@@ -353,7 +354,8 @@ struct FNDIDataChannelWriteInstanceData
 				//In non test/shipping builds we gather and log and missing parameters that cause us to fail to find correct bindings.
 				TArray<FNiagaraVariableBase> MissingParams;
 
-				if (DataChannelData == nullptr || Interface->bUpdateDestinationDataEveryTick)
+				bool bNDCDataIsValid = DataChannelData && DataChannelData->IsLayoutValid(DataChannelPtr);
+				if (bNDCDataIsValid == false || Interface->bUpdateDestinationDataEveryTick)
 				{
 					FNiagaraDataChannelSearchParameters SearchParams(Instance->GetAttachComponent());
 					DataChannelData = DataChannelPtr->FindData(SearchParams, ENiagaraResourceAccess::WriteOnly);
@@ -449,8 +451,7 @@ struct FNDIDataChannelWriteInstanceData
 #endif
 				DataChannelData->Publish(PublishRequest);
 			}
-
-			DestinationData = nullptr;
+			
 			AtomicNumInstances = 0;
 		}
 		return true;
@@ -1042,31 +1043,49 @@ void UNiagaraDataInterfaceDataChannelWrite::Write(FVectorVMExternalFunctionConte
 				uint32 Index = static_cast<uint32>(RawIndex);
 				MaxLocalIndex = bEmit ? FMath::Max(Index, MaxLocalIndex) : MaxLocalIndex;
 
-				bool bSuccess = false;
+				bool bAllWritesSuccess = true;
 
 				//TODO: Optimize case where emit is constant
 				//TODO: Optimize for runs of sequential true emits.
-				auto FloatFunc = [Data, Index](const FNDIDataChannelRegisterBinding& VMBinding, FNDIInputParam<float>& FloatData)
+				auto FloatFunc = [Data, Index, &bAllWritesSuccess](const FNDIDataChannelRegisterBinding& VMBinding, FNDIInputParam<float>& FloatData)
 				{
-					if (VMBinding.DataSetRegisterIndex != INDEX_NONE)
-						*Data->GetInstancePtrFloat(VMBinding.DataSetRegisterIndex, Index) = FloatData.GetAndAdvance();
+					if (VMBinding.GetDataSetRegisterIndex() != INDEX_NONE)
+					{
+						*Data->GetInstancePtrFloat(VMBinding.GetDataSetRegisterIndex(), Index) = FloatData.GetAndAdvance();
+					}
+					else
+					{
+						bAllWritesSuccess = false;
+					}
 				};
-				auto IntFunc = [Data, Index](const FNDIDataChannelRegisterBinding& VMBinding, FNDIInputParam<int32>& IntData)
+				auto IntFunc = [Data, Index, &bAllWritesSuccess](const FNDIDataChannelRegisterBinding& VMBinding, FNDIInputParam<int32>& IntData)
 				{
-					if (VMBinding.DataSetRegisterIndex != INDEX_NONE)
-						*Data->GetInstancePtrInt32(VMBinding.DataSetRegisterIndex, Index) = IntData.GetAndAdvance();
+					if (VMBinding.GetDataSetRegisterIndex() != INDEX_NONE)
+					{
+						*Data->GetInstancePtrInt32(VMBinding.GetDataSetRegisterIndex(), Index) = IntData.GetAndAdvance();
+					}
+					else
+					{
+						bAllWritesSuccess = false;
+					}
 				};
-				auto HalfFunc = [Data, Index](const FNDIDataChannelRegisterBinding& VMBinding, FNDIInputParam<FFloat16>& HalfData)
+				auto HalfFunc = [Data, Index, &bAllWritesSuccess](const FNDIDataChannelRegisterBinding& VMBinding, FNDIInputParam<FFloat16>& HalfData)
 				{
-					if (VMBinding.DataSetRegisterIndex != INDEX_NONE)
-						*Data->GetInstancePtrHalf(VMBinding.DataSetRegisterIndex, Index) = HalfData.GetAndAdvance();
+					if (VMBinding.GetDataSetRegisterIndex() != INDEX_NONE)
+					{
+						*Data->GetInstancePtrHalf(VMBinding.GetDataSetRegisterIndex(), Index) = HalfData.GetAndAdvance();
+					}
+					else
+					{
+						bAllWritesSuccess = false;
+					}
 				};
 
-				bSuccess = VariadicInputs.Process(bEmit, 1, BindingInfo, FloatFunc, IntFunc, HalfFunc);
+				bool bFinalSuccess = VariadicInputs.Process(bEmit, 1, BindingInfo, FloatFunc, IntFunc, HalfFunc) && bAllWritesSuccess;
 
 				if (OutSuccess.IsValid())
 				{
-					OutSuccess.SetAndAdvance(bSuccess);
+					OutSuccess.SetAndAdvance(bFinalSuccess);
 				}
 			}
 
@@ -1149,13 +1168,13 @@ void UNiagaraDataInterfaceDataChannelWrite::Append(FVectorVMExternalFunctionCont
 					LocalNumToEmit = FMath::Min(LocalNumToEmit, MaxWriteCount);
 
 					//If we're writing all instances then we can do a memcpy instead of slower loop copies.
-					bool bSuccess = false;
+					bool bAllWritesSuccess = true;
 					uint32 Index = CurrNumInstances;
-					auto FloatFunc = [Data, Index, LocalNumToEmit](const FNDIDataChannelRegisterBinding& VMBinding, FNDIInputParam<float>& FloatData)
+					auto FloatFunc = [Data, Index, LocalNumToEmit, &bAllWritesSuccess](const FNDIDataChannelRegisterBinding& VMBinding, FNDIInputParam<float>& FloatData)
 					{
-						if (VMBinding.DataSetRegisterIndex != INDEX_NONE)
+						if (VMBinding.GetDataSetRegisterIndex() != INDEX_NONE)
 						{
-							float* Dest = Data->GetInstancePtrFloat(VMBinding.DataSetRegisterIndex, Index);
+							float* Dest = Data->GetInstancePtrFloat(VMBinding.GetDataSetRegisterIndex(), Index);
 							if (FloatData.IsConstant())
 							{
 								float Value = FloatData.GetAndAdvance();
@@ -1167,12 +1186,16 @@ void UNiagaraDataInterfaceDataChannelWrite::Append(FVectorVMExternalFunctionCont
 								 FMemory::Memcpy(Dest, Src, LocalNumToEmit * sizeof(float));
 							}
 						}
+						else
+						{
+							bAllWritesSuccess = false;
+						}
 					};
-					auto IntFunc = [Data, Index, LocalNumToEmit](const FNDIDataChannelRegisterBinding& VMBinding, FNDIInputParam<int32>& IntData)
+					auto IntFunc = [Data, Index, LocalNumToEmit, &bAllWritesSuccess](const FNDIDataChannelRegisterBinding& VMBinding, FNDIInputParam<int32>& IntData)
 					{
-						if (VMBinding.DataSetRegisterIndex != INDEX_NONE)
+						if (VMBinding.GetDataSetRegisterIndex() != INDEX_NONE)
 						{
-							int32* Dest = Data->GetInstancePtrInt32(VMBinding.DataSetRegisterIndex, Index);
+							int32* Dest = Data->GetInstancePtrInt32(VMBinding.GetDataSetRegisterIndex(), Index);
 							if (IntData.IsConstant())
 							{
 								int32 Value = IntData.GetAndAdvance();
@@ -1184,12 +1207,16 @@ void UNiagaraDataInterfaceDataChannelWrite::Append(FVectorVMExternalFunctionCont
 								FMemory::Memcpy(Dest, Src, LocalNumToEmit * sizeof(int32));
 							}
 						}
+						else
+						{
+							bAllWritesSuccess = false;
+						}
 					};
-					auto HalfFunc = [Data, Index, LocalNumToEmit](const FNDIDataChannelRegisterBinding& VMBinding, FNDIInputParam<FFloat16>& HalfData)
+					auto HalfFunc = [Data, Index, LocalNumToEmit, &bAllWritesSuccess](const FNDIDataChannelRegisterBinding& VMBinding, FNDIInputParam<FFloat16>& HalfData)
 					{
-						if (VMBinding.DataSetRegisterIndex != INDEX_NONE)
+						if (VMBinding.GetDataSetRegisterIndex() != INDEX_NONE)
 						{
-							FFloat16* Dest = Data->GetInstancePtrHalf(VMBinding.DataSetRegisterIndex, Index);
+							FFloat16* Dest = Data->GetInstancePtrHalf(VMBinding.GetDataSetRegisterIndex(), Index);
 							if (HalfData.IsConstant())
 							{
 								FFloat16 Value = HalfData.GetAndAdvance();
@@ -1201,14 +1228,19 @@ void UNiagaraDataInterfaceDataChannelWrite::Append(FVectorVMExternalFunctionCont
 								FMemory::Memcpy(Dest, Src, LocalNumToEmit * sizeof(FFloat16));
 							}
 						}
+						else
+						{
+							bAllWritesSuccess = false;
+						}
 					};
-					bSuccess = VariadicInputs.Process(true, Context.GetNumInstances(), BindingInfo, FloatFunc, IntFunc, HalfFunc);
+
+					bool bFinalSuccess = VariadicInputs.Process(true, Context.GetNumInstances(), BindingInfo, FloatFunc, IntFunc, HalfFunc) && bAllWritesSuccess;
 
 					if (OutSuccess.IsValid())
 					{
 						for (int32 i = 0; i < Context.GetNumInstances(); ++i)
 						{
-							OutSuccess.SetAndAdvance(bSuccess);
+							OutSuccess.SetAndAdvance(bFinalSuccess);
 						}
 					}
 				}
@@ -1219,33 +1251,52 @@ void UNiagaraDataInterfaceDataChannelWrite::Append(FVectorVMExternalFunctionCont
 						uint32 Index = CurrNumInstances;
 
 						bool bEmit = InEmit.GetAndAdvance();
-						bool bSuccess = false;
+						bool bAllWritesSuccess = true;
 
 						if(bEmit)
 						{
 							++CurrNumInstances;
 						}
 
-						auto FloatFunc = [Data, Index](const FNDIDataChannelRegisterBinding& VMBinding, FNDIInputParam<float>& FloatData)
+						auto FloatFunc = [Data, Index, &bAllWritesSuccess](const FNDIDataChannelRegisterBinding& VMBinding, FNDIInputParam<float>& FloatData)
 						{
-							if (VMBinding.DataSetRegisterIndex != INDEX_NONE)
-								*Data->GetInstancePtrFloat(VMBinding.DataSetRegisterIndex, Index) = FloatData.GetAndAdvance();
+							if (VMBinding.GetDataSetRegisterIndex() != INDEX_NONE)
+							{
+								*Data->GetInstancePtrFloat(VMBinding.GetDataSetRegisterIndex(), Index) = FloatData.GetAndAdvance();
+							}
+							else
+							{
+								bAllWritesSuccess = false;
+							}
 						};
-						auto IntFunc = [Data, Index](const FNDIDataChannelRegisterBinding& VMBinding, FNDIInputParam<int32>& IntData)
+						auto IntFunc = [Data, Index, &bAllWritesSuccess](const FNDIDataChannelRegisterBinding& VMBinding, FNDIInputParam<int32>& IntData)
 						{
-							if (VMBinding.DataSetRegisterIndex != INDEX_NONE)
-								*Data->GetInstancePtrInt32(VMBinding.DataSetRegisterIndex, Index) = IntData.GetAndAdvance();
+							if (VMBinding.GetDataSetRegisterIndex() != INDEX_NONE)
+							{
+								*Data->GetInstancePtrInt32(VMBinding.GetDataSetRegisterIndex(), Index) = IntData.GetAndAdvance();
+							}
+							else
+							{
+								bAllWritesSuccess = false;
+							}
 						};
-						auto HalfFunc = [Data, Index](const FNDIDataChannelRegisterBinding& VMBinding, FNDIInputParam<FFloat16>& HalfData)
+						auto HalfFunc = [Data, Index, &bAllWritesSuccess](const FNDIDataChannelRegisterBinding& VMBinding, FNDIInputParam<FFloat16>& HalfData)
 						{
-							if (VMBinding.DataSetRegisterIndex != INDEX_NONE)
-								*Data->GetInstancePtrHalf(VMBinding.DataSetRegisterIndex, Index) = HalfData.GetAndAdvance();
+							if (VMBinding.GetDataSetRegisterIndex() != INDEX_NONE)
+							{
+								*Data->GetInstancePtrHalf(VMBinding.GetDataSetRegisterIndex(), Index) = HalfData.GetAndAdvance();
+							}
+							else
+							{
+								bAllWritesSuccess = false;
+							}
 						};
-						bSuccess = VariadicInputs.Process(bEmit, 1, BindingInfo, FloatFunc, IntFunc, HalfFunc);
+
+						bool bFinalSuccess = VariadicInputs.Process(bEmit, 1, BindingInfo, FloatFunc, IntFunc, HalfFunc) && bAllWritesSuccess;
 
 						if (OutSuccess.IsValid())
 						{
-							OutSuccess.SetAndAdvance(bSuccess);
+							OutSuccess.SetAndAdvance(bFinalSuccess);
 						}
 					}
 				}
diff --git a/Engine/Plugins/FX/Niagara/Source/Niagara/Private/NiagaraDataChannel.cpp b/Engine/Plugins/FX/Niagara/Source/Niagara/Private/NiagaraDataChannel.cpp
--- a/Engine/Plugins/FX/Niagara/Source/Niagara/Private/NiagaraDataChannel.cpp
+++ b/Engine/Plugins/FX/Niagara/Source/Niagara/Private/NiagaraDataChannel.cpp
@@ -687,17 +687,21 @@ FNiagaraDataBufferRef FNiagaraDataChannelDataProxy::PrepareForReadAccess(FRDGBui
 
 	if(bCurrentFrame)
 	{
-		if (CurrBufferAccessCounts < 0)
+		if(CurrFrameData)
 		{
+			if (CurrBufferAccessCounts < 0)
+			{
 #if !UE_BUILD_SHIPPING
-			UE_LOG(LogNiagara, Warning, TEXT("Attempting to write to a Niagara Data Channel in the same stage in which it's being read. {%s}"), *DebugName);
+				UE_LOG(LogNiagara, Warning, TEXT("Attempting to read from a Niagara Data Channel in the same stage in which it's being written. {%s}"), *DebugName);
 #endif
-			return nullptr;
-		}
-		
-		++CurrBufferAccessCounts;
+				return nullptr;
+			}
 
-		return CurrFrameData;
+			++CurrBufferAccessCounts;
+
+			return CurrFrameData;
+		}
+		return nullptr;
 	}
 	else
 	{
@@ -1138,6 +1142,15 @@ FNiagaraDataBuffer* FNiagaraDataChannelData::GetBufferForCPUWrite()
 	return nullptr;
 }
 
+bool FNiagaraDataChannelData::IsLayoutValid(UNiagaraDataChannelHandler* Owner)const
+{
+	if(Owner && LayoutInfo)
+	{
+		return Owner->GetDataChannel()->GetLayoutInfo() == LayoutInfo;//If our layout has been modified then this NDC data is no longer valid.
+	}
+	return false;
+}
+
 void FNiagaraDataChannelData::DestroyRenderThreadProxy(FNiagaraGpuComputeDispatchInterface* ComputeDispatchInterface)
 {
 	//We can leave the proxy to be destroyed by the dispatcher ( or in some edge cases perhaps a DI proxy ).
@@ -1278,6 +1291,8 @@ bool UNiagaraDataChannel::IsReadyForFinishDestroy()
 
 void UNiagaraDataChannel::PreEditChange(FProperty* PropertyAboutToChange)
 {
+	Super::PreEditChange(PropertyAboutToChange);
+
 	FNiagaraWorldManager::ForAllWorldManagers(
 		[DataChannel = this](FNiagaraWorldManager& WorldMan)
 		{
@@ -1290,13 +1305,17 @@ void UNiagaraDataChannel::PostEditChangeProperty(FPropertyChangedEvent& Property
 	FName VariablesMemberName = GET_MEMBER_NAME_CHECKED(UNiagaraDataChannel, ChannelVariables);
 	if (PropertyChangedEvent.ChangeType == EPropertyChangeType::ArrayAdd && PropertyChangedEvent.GetPropertyName() == VariablesMemberName)
 	{
-		TSet<FName> ExistingNames;
-		for (const FNiagaraDataChannelVariable& Var : ChannelVariables)
+		int32 ArrayIndex = PropertyChangedEvent.GetArrayIndex(VariablesMemberName.ToString());
+		if(ChannelVariables.IsValidIndex(ArrayIndex))
 		{
-			ExistingNames.Add(Var.GetName());
+			TSet<FName> ExistingNames;
+			for (const FNiagaraDataChannelVariable& Var : ChannelVariables)
+			{
+				ExistingNames.Add(Var.GetName());
+			}
+			FName UniqueName = FNiagaraUtilities::GetUniqueName(FName("MyNewVar"), ExistingNames);
+			ChannelVariables[ArrayIndex].SetName(UniqueName);
 		}
-		FName UniqueName = FNiagaraUtilities::GetUniqueName(FName("MyNewVar"), ExistingNames);
-		ChannelVariables.Last().SetName(UniqueName);
 	}
 	if (PropertyChangedEvent.ChangeType == EPropertyChangeType::Duplicate && PropertyChangedEvent.GetPropertyName() == VariablesMemberName)
 	{
diff --git a/Engine/Plugins/FX/Niagara/Source/Niagara/Private/NiagaraDataChannelHandler.cpp b/Engine/Plugins/FX/Niagara/Source/Niagara/Private/NiagaraDataChannelHandler.cpp
--- a/Engine/Plugins/FX/Niagara/Source/Niagara/Private/NiagaraDataChannelHandler.cpp
+++ b/Engine/Plugins/FX/Niagara/Source/Niagara/Private/NiagaraDataChannelHandler.cpp
@@ -37,6 +37,12 @@ void UNiagaraDataChannelHandler::Cleanup()
 		Writer->Cleanup();
 		Writer = nullptr;
 	}
+
+	if (!HasAnyFlags(RF_ClassDefaultObject))
+	{
+		//Mark this handler as garbage so any reading DIs will know to stop using it.
+		MarkAsGarbage();
+	}
 }
 
 void UNiagaraDataChannelHandler::BeginFrame(float DeltaTime, FNiagaraWorldManager* OwningWorld)
diff --git a/Engine/Plugins/FX/Niagara/Source/Niagara/Private/NiagaraDataChannel_Global.cpp b/Engine/Plugins/FX/Niagara/Source/Niagara/Private/NiagaraDataChannel_Global.cpp
--- a/Engine/Plugins/FX/Niagara/Source/Niagara/Private/NiagaraDataChannel_Global.cpp
+++ b/Engine/Plugins/FX/Niagara/Source/Niagara/Private/NiagaraDataChannel_Global.cpp
@@ -72,6 +72,10 @@ void UNiagaraDataChannelHandler_Global::Tick(float DeltaSeconds, ETickingGroup T
 
 FNiagaraDataChannelDataPtr UNiagaraDataChannelHandler_Global::FindData(FNiagaraDataChannelSearchParameters SearchParams, ENiagaraResourceAccess AccessType)
 {
+	if(!Data)
+	{
+		Data = CreateData();
+	}
 	return Data;
 	//For more complicated channels we could check the location + bounds of the system instance etc to return some spatially localized data.
 }
diff --git a/Engine/Plugins/FX/Niagara/Source/Niagara/Private/NiagaraDataSet.cpp b/Engine/Plugins/FX/Niagara/Source/Niagara/Private/NiagaraDataSet.cpp
--- a/Engine/Plugins/FX/Niagara/Source/Niagara/Private/NiagaraDataSet.cpp
+++ b/Engine/Plugins/FX/Niagara/Source/Niagara/Private/NiagaraDataSet.cpp
@@ -1073,13 +1073,13 @@ void FNiagaraDataBuffer::CopyToUnrelated(FNiagaraDataBuffer& DestBuffer, int32 S
 		const FNiagaraDataSetCompiledData& DestCompiledData = DestBuffer.GetOwner()->GetCompiledData();
 		const FNiagaraDataSetCompiledData& CompiledData = GetOwner()->GetCompiledData();
 
-		for(int32 VarIdx = 0; VarIdx < DestCompiledData.Variables.Num(); ++VarIdx)
+		for (int32 VarIdx = 0; VarIdx < DestCompiledData.Variables.Num(); ++VarIdx)
 		{
 			const FNiagaraVariableBase& DestVar = DestCompiledData.Variables[VarIdx];
 			const FNiagaraVariableLayoutInfo& DestVarLayout = DestCompiledData.VariableLayouts[VarIdx];
 
 			const int32 SrcVarIdx = CompiledData.Variables.IndexOfByKey(DestVar);
-			if(SrcVarIdx != INDEX_NONE)
+			if (SrcVarIdx != INDEX_NONE)
 			{
 				//Found a the variable to copy over.
 				const FNiagaraVariableLayoutInfo& SrcVarLayout = CompiledData.VariableLayouts[SrcVarIdx];
@@ -1144,7 +1144,32 @@ void FNiagaraDataBuffer::CopyToUnrelated(FNiagaraDataBuffer& DestBuffer, int32 S
 						}
 					}
 				}
+			}
+			else
+			{
+				const uint32 FloatComponents = DestVarLayout.GetNumFloatComponents();
+				for (uint32 CompIdx = 0; CompIdx < FloatComponents; ++CompIdx)
+				{
+					const int32 DestCompOffest = DestVarLayout.GetFloatComponentStart() + CompIdx;
+					float* Dst = DestBuffer.GetInstancePtrFloat(DestCompOffest, DestStartIdx);
+					FMemory::Memzero(Dst, InstancesToCopy * sizeof(float));
+				}
 
+				const uint32 IntComponents = DestVarLayout.GetNumInt32Components();
+				for (uint32 CompIdx = 0; CompIdx < IntComponents; ++CompIdx)
+				{
+					const int32 DestCompOffest = DestVarLayout.GetInt32ComponentStart() + CompIdx;
+					int32* Dst = DestBuffer.GetInstancePtrInt32(DestCompOffest, DestStartIdx);
+					FMemory::Memzero(Dst, InstancesToCopy * sizeof(int32));
+				}
+
+				const uint32 HalfComponents = DestVarLayout.GetNumHalfComponents();
+				for (uint32 CompIdx = 0; CompIdx < HalfComponents; ++CompIdx)
+				{
+					const int32 DestCompOffest = DestVarLayout.GetHalfComponentStart() + CompIdx;
+					FFloat16* Dst = DestBuffer.GetInstancePtrHalf(DestCompOffest, DestStartIdx);
+					FMemory::Memzero(Dst, InstancesToCopy * sizeof(FFloat16));
+				}
 			}
 		}
 	}
diff --git a/Engine/Source/Runtime/Core/Private/UObject/UE5ReleaseStreamObjectVersion.cpp b/Engine/Source/Runtime/Core/Private/UObject/UE5ReleaseStreamObjectVersion.cpp
--- a/Engine/Source/Runtime/Core/Private/UObject/UE5ReleaseStreamObjectVersion.cpp
+++ b/Engine/Source/Runtime/Core/Private/UObject/UE5ReleaseStreamObjectVersion.cpp
@@ -11,8 +11,8 @@ TMap<FGuid, FGuid> FUE5ReleaseStreamObjectVersion::GetSystemGuids()
 	SystemGuids.Add(DevGuids.GLOBALSHADERMAP_DERIVEDDATA_VER, FGuid("D0BF3452816D46908073DFDD4B855AE5"));
 	SystemGuids.Add(DevGuids.MATERIALSHADERMAP_DERIVEDDATA_VER, FGuid("79090A66F9D94B5BB285D49E5D39468E"));
 	SystemGuids.Add(DevGuids.NANITE_DERIVEDDATA_VER, FGuid("AC88CFBDDA614A9C8488F64D2DAF699D"));
-	SystemGuids.Add(DevGuids.NIAGARASHADERMAP_DERIVEDDATA_VER, FGuid("6BE35B6FACB34568970120F9BC8DAB80"));
-	SystemGuids.Add(DevGuids.Niagara_LatestScriptCompileVersion, FGuid("43A777C54EB24E9894525B8D04529F23"));
+	SystemGuids.Add(DevGuids.NIAGARASHADERMAP_DERIVEDDATA_VER, FGuid("7BCE96ABAA6E46D689306FE673DED4DE"));
+	SystemGuids.Add(DevGuids.Niagara_LatestScriptCompileVersion, FGuid("85D2347062C644AAAD31E3E6CBA1FDA0"));
 	SystemGuids.Add(DevGuids.SkeletalMeshDerivedDataVersion, FGuid("BA91E71F642E4813A1B74B80D0448928"));
 	SystemGuids.Add(DevGuids.STATICMESH_DERIVEDDATA_VER, FGuid("7A556175518D458AA8786ED69AB8DDE7"));
 
