IncrementalCook: Niagara: Remove unnecessary resolves of TObjectPtr; resolving TObjectPtr creates a transitive build dependency for incremental cooks (unless opted out), so we should avoid it when possible to reduce the cost of incremental cooks.
#rnx

diff --git a/Engine/Plugins/FX/Niagara/Source/Niagara/Classes/NiagaraEmitterInstance.h b/Engine/Plugins/FX/Niagara/Source/Niagara/Classes/NiagaraEmitterInstance.h
--- a/Engine/Plugins/FX/Niagara/Source/Niagara/Classes/NiagaraEmitterInstance.h
+++ b/Engine/Plugins/FX/Niagara/Source/Niagara/Classes/NiagaraEmitterInstance.h
@@ -94,7 +94,7 @@ public:
 
 	NIAGARA_API bool GetBoundRendererValue_GT(const FNiagaraVariableBase& InBaseVar, const FNiagaraVariableBase& InSubVar, void* OutValueData) const;
 
-	[[nodiscard]] NIAGARA_API UObject* FindBinding(const FNiagaraVariable& InVariable) const;
+	[[nodiscard]] NIAGARA_API TObjectPtr<UObject> FindBinding(const FNiagaraVariable& InVariable) const;
 	[[nodiscard]] NIAGARA_API UNiagaraDataInterface* FindDataInterface(const FNiagaraVariable& InVariable) const;
 
 	[[nodiscard]] NIAGARA_API FNiagaraEmitterID GetEmitterID() const{ return FNiagaraEmitterID(EmitterIndex); }
diff --git a/Engine/Plugins/FX/Niagara/Source/Niagara/Private/NiagaraCommon.cpp b/Engine/Plugins/FX/Niagara/Source/Niagara/Private/NiagaraCommon.cpp
--- a/Engine/Plugins/FX/Niagara/Source/Niagara/Private/NiagaraCommon.cpp
+++ b/Engine/Plugins/FX/Niagara/Source/Niagara/Private/NiagaraCommon.cpp
@@ -4,6 +4,8 @@
 
 #include "DataDrivenShaderPlatformInfo.h"
 #include "RenderUtils.h"
+#include "Misc/PackageAccessTracking.h"
+#include "Misc/PackageAccessTrackingOps.h"
 #include "Misc/StringBuilder.h"
 #include "NiagaraComponent.h"
 #include "NiagaraComponentSettings.h"
@@ -77,6 +79,10 @@ static FAutoConsoleVariableRef CVarNiagaraLogVerboseWarnings(
 
 FNiagaraSystemUpdateContext::~FNiagaraSystemUpdateContext()
 {
+	// IncrementalCook: Add a ResetContext scope so that dereferences of the Components we gathered do not
+	// record build dependencies from our caller to those components. We are providing an update event
+	// to the components we gathered; we are not using them for any modifications of our caller.
+	UE_TRACK_REFERENCING_OPNAME_SCOPED(PackageAccessTrackingOps::NAME_ResetContext);
 	CommitUpdate();
 }
 
@@ -179,70 +185,66 @@ void FNiagaraSystemUpdateContext::AddSoloComponent(UNiagaraComponent* Component,
 
 void FNiagaraSystemUpdateContext::Add(const UNiagaraSystem* System, bool bReInit)
 {
-	for (TObjectIterator<UNiagaraComponent> It; It; ++It)
-	{
-		UNiagaraComponent* Comp = *It;
-		check(Comp);
-		if (Comp->GetAsset() == System)
+	AddComponentsWithMatchingSystem(bReInit,
+		[System](const UNiagaraSystem* SystemOnComponent)
 		{
-			bool bAllowDestroySystemSim = true;
-			AddInternal(Comp, bReInit, bAllowDestroySystemSim);
-		}
-	}
+			return SystemOnComponent == System;
+		});
 }
 #if WITH_EDITORONLY_DATA
 
 void FNiagaraSystemUpdateContext::Add(const FVersionedNiagaraEmitter& Emitter, bool bReInit)
 {
-	for (TObjectIterator<UNiagaraComponent> It; It; ++It)
-	{
-		UNiagaraComponent* Comp = *It;
-		check(Comp);
-		UNiagaraSystem* System = Comp->GetAsset();
-
-		// only worry about systems that are fully loaded (and don't force it to be fully loaded just because it shows up in the iterator
-		if (System && System->bFullyLoaded && System->UsesEmitter(Emitter))
+	AddComponentsWithMatchingSystem(bReInit,
+		[&Emitter](const UNiagaraSystem* System)
 		{
-			bool bAllowDestroySystemSim = true;
-			AddInternal(Comp, bReInit, bAllowDestroySystemSim);
-		}
-	}
+			// only worry about systems that are fully loaded (and don't force it to be fully loaded just because it shows up in the iterator
+			return System->bFullyLoaded && System->UsesEmitter(Emitter);
+		});
 }
 
 void FNiagaraSystemUpdateContext::Add(const UNiagaraScript* Script, bool bReInit)
 {
-	for (TObjectIterator<UNiagaraComponent> It; It; ++It)
-	{
-		UNiagaraComponent* Comp = *It;
-		check(Comp);
-
-		// only worry about systems that are fully loaded (and don't force it to be fully loaded just because it shows up in the iterator
-		UNiagaraSystem* System = Comp->GetAsset();
-		if (System && System->bFullyLoaded && System->UsesScript(Script))
+	AddComponentsWithMatchingSystem(bReInit,
+		[Script](const UNiagaraSystem* System)
 		{
-			bool bAllowDestroySystemSim = true;
-			AddInternal(Comp, bReInit, bAllowDestroySystemSim);
-		}
-	}
+			// only worry about systems that are fully loaded (and don't force it to be fully loaded just because it shows up in the iterator)
+			return System->bFullyLoaded && System->UsesScript(Script);
+		});
 }
 
 void FNiagaraSystemUpdateContext::Add(const UNiagaraParameterCollection* Collection, bool bReInit)
 {
+	AddComponentsWithMatchingSystem(bReInit,
+		[Collection](const UNiagaraSystem* System)
+		{
+			// only worry about systems that are fully loaded (and don't force it to be fully loaded just because it shows up in the iterator)
+			return System->bFullyLoaded && System->UsesCollection(Collection);
+		});
+}
+
+#endif
+
+template <typename SystemMatchesFunc>
+void FNiagaraSystemUpdateContext::AddComponentsWithMatchingSystem(bool bReInit, SystemMatchesFunc&& SystemMatches)
+{
+
+	// IncrementalCook: Add a ResetContext scope so that dereferences of the Components we gather do not
+	// record build dependencies from our caller to those components. We are providing an update event
+	// to the components we gather; we are not using them for any modifications of our caller.
+	UE_TRACK_REFERENCING_OPNAME_SCOPED(PackageAccessTrackingOps::NAME_ResetContext);
 	for (TObjectIterator<UNiagaraComponent> It; It; ++It)
 	{
 		UNiagaraComponent* Comp = *It;
 		check(Comp);
 		UNiagaraSystem* System = Comp->GetAsset();
-
-		// only worry about systems that are fully loaded (and don't force it to be fully loaded just because it shows up in the iterator
-		if (System && System->bFullyLoaded && System->UsesCollection(Collection))
+		if (System && SystemMatches(System))
 		{
 			bool bAllowDestroySystemSim = true;
 			AddInternal(Comp, bReInit, bAllowDestroySystemSim);
 		}
 	}
 }
-#endif
 
 void FNiagaraSystemUpdateContext::AddInternal(UNiagaraComponent* Comp, bool bReInit, bool bAllowDestroySystemSim)
 {
diff --git a/Engine/Plugins/FX/Niagara/Source/Niagara/Private/NiagaraComponent.cpp b/Engine/Plugins/FX/Niagara/Source/Niagara/Private/NiagaraComponent.cpp
--- a/Engine/Plugins/FX/Niagara/Source/Niagara/Private/NiagaraComponent.cpp
+++ b/Engine/Plugins/FX/Niagara/Source/Niagara/Private/NiagaraComponent.cpp
@@ -2884,7 +2884,7 @@ void FixInvalidUserParameters(FNiagaraUserRedirectionParameterStore& ParameterSt
 			// If the correctly named parameter is also in the collection then both parameters need to be removed and the 
 			// correct one re-added.  First we need to cache the value of the parameter so that it's not lost on removal.
 			UNiagaraDataInterface* DataInterfaceValue = nullptr;
-			UObject* ObjectValue = nullptr;
+			TObjectPtr<UObject> ObjectValue = nullptr;
 			TArray<uint8> DataValue;
 			if (IncorrectlyNamedParameter.IsDataInterface())
 			{
@@ -3329,7 +3329,7 @@ void UNiagaraComponent::EnsureOverrideParametersConsistent() const
 			}
 			else if (Key.IsUObject())
 			{
-				const UObject* ActualObj = OverrideParameters.GetUObject(Key);
+				TObjectPtr<const UObject> ActualObj = OverrideParameters.GetUObject(Key);
 				if (ActualObj != nullptr)
 				{
 					ensureAlways(OverrideValue.GetUObject() == ActualObj);
@@ -3472,7 +3472,7 @@ void UNiagaraComponent::CopyParametersFromAsset(bool bResetExistingOverrideParam
 			}
 			else if (ExistingOverrideVar.IsUObject())
 			{
-				UObject* ExistingUObject = ExistingOverrideParameters.GetUObject(ExistingOverrideVar.Offset);
+				TObjectPtr<UObject> ExistingUObject = ExistingOverrideParameters.GetUObject(ExistingOverrideVar.Offset);
 				OverrideParameters.SetUObject(ExistingUObject, ExistingOverrideVar);
 			}
 			else
diff --git a/Engine/Plugins/FX/Niagara/Source/Niagara/Private/NiagaraEmitterInstance.cpp b/Engine/Plugins/FX/Niagara/Source/Niagara/Private/NiagaraEmitterInstance.cpp
--- a/Engine/Plugins/FX/Niagara/Source/Niagara/Private/NiagaraEmitterInstance.cpp
+++ b/Engine/Plugins/FX/Niagara/Source/Niagara/Private/NiagaraEmitterInstance.cpp
@@ -143,7 +143,7 @@ bool FNiagaraEmitterInstance::GetBoundRendererValue_GT(const FNiagaraVariableBas
 	return false;
 }
 
-UObject* FNiagaraEmitterInstance::FindBinding(const FNiagaraVariable& InVariable) const
+TObjectPtr<UObject> FNiagaraEmitterInstance::FindBinding(const FNiagaraVariable& InVariable) const
 {
 	if (!InVariable.IsValid())
 	{
diff --git a/Engine/Plugins/FX/Niagara/Source/Niagara/Private/NiagaraMeshRendererProperties.cpp b/Engine/Plugins/FX/Niagara/Source/Niagara/Private/NiagaraMeshRendererProperties.cpp
--- a/Engine/Plugins/FX/Niagara/Source/Niagara/Private/NiagaraMeshRendererProperties.cpp
+++ b/Engine/Plugins/FX/Niagara/Source/Niagara/Private/NiagaraMeshRendererProperties.cpp
@@ -797,7 +797,7 @@ void UNiagaraMeshRendererProperties::GetUsedMaterials(const FNiagaraEmitterInsta
 			ApplyMaterialOverrides(EmitterInstance, OrderedMeshMaterials);
 
 			OutMaterials.Reserve(OutMaterials.Num() + OrderedMeshMaterials.Num());
-			for (UMaterialInterface* MaterialInterface : OrderedMeshMaterials)
+			for (TObjectPtr<UMaterialInterface> MaterialInterface : OrderedMeshMaterials)
 			{
 				if (MaterialInterface)
 				{
diff --git a/Engine/Plugins/FX/Niagara/Source/Niagara/Private/NiagaraModule.cpp b/Engine/Plugins/FX/Niagara/Source/Niagara/Private/NiagaraModule.cpp
--- a/Engine/Plugins/FX/Niagara/Source/Niagara/Private/NiagaraModule.cpp
+++ b/Engine/Plugins/FX/Niagara/Source/Niagara/Private/NiagaraModule.cpp
@@ -1546,7 +1546,7 @@ bool FNiagaraTypeDefinition::Serialize(FArchive& Ar)
 				StructToResetOnExit = StructDef;
 				ClassStructOrEnum = FNiagaraTypeHelper::GetLWCStruct(StructDef);
 				check(ClassStructOrEnum);
-				check(ClassStructOrEnum->GetOutermost() != GetTransientPackage());
+				check(ClassStructOrEnum.GetPackage() != GetTransientPackage());
 			}
 		}
 	}
diff --git a/Engine/Plugins/FX/Niagara/Source/Niagara/Private/NiagaraSystemInstanceController.cpp b/Engine/Plugins/FX/Niagara/Source/Niagara/Private/NiagaraSystemInstanceController.cpp
--- a/Engine/Plugins/FX/Niagara/Source/Niagara/Private/NiagaraSystemInstanceController.cpp
+++ b/Engine/Plugins/FX/Niagara/Source/Niagara/Private/NiagaraSystemInstanceController.cpp
@@ -468,7 +468,7 @@ void FNiagaraSystemInstanceController::SetVariable(FName InVariableName, TWeakOb
 	if (OverrideParameters)
 	{
 		const FNiagaraVariable VariableDesc(FNiagaraTypeDefinition::GetUMaterialDef(), InVariableName);
-		UObject* CurrentValue = OverrideParameters->GetUObject(VariableDesc);
+		TObjectPtr<UObject> CurrentValue = OverrideParameters->GetUObject(VariableDesc);
 		UMaterialInterface* NewValue = InValue.Get();
 		OverrideParameters->SetUObject(NewValue, VariableDesc);
 		if (CurrentValue != NewValue)
@@ -483,7 +483,7 @@ void FNiagaraSystemInstanceController::SetVariable(FName InVariableName, TWeakOb
 	if (OverrideParameters)
 	{
 		const FNiagaraVariable VariableDesc(FNiagaraTypeDefinition::GetUStaticMeshDef(), InVariableName);
-		UObject* CurrentValue = OverrideParameters->GetUObject(VariableDesc);
+		TObjectPtr<UObject> CurrentValue = OverrideParameters->GetUObject(VariableDesc);
 		UStaticMesh* NewValue = InValue.Get();
 		OverrideParameters->SetUObject(NewValue, VariableDesc);
 		if (CurrentValue != NewValue)
diff --git a/Engine/Plugins/FX/Niagara/Source/Niagara/Private/NiagaraVariant.cpp b/Engine/Plugins/FX/Niagara/Source/Niagara/Private/NiagaraVariant.cpp
--- a/Engine/Plugins/FX/Niagara/Source/Niagara/Private/NiagaraVariant.cpp
+++ b/Engine/Plugins/FX/Niagara/Source/Niagara/Private/NiagaraVariant.cpp
@@ -13,7 +13,7 @@ FNiagaraVariant::FNiagaraVariant()
 	DataInterface = nullptr;
 }
 
-FNiagaraVariant::FNiagaraVariant(UObject* InObject)
+FNiagaraVariant::FNiagaraVariant(TObjectPtr<UObject> InObject)
 {
 	CurrentMode = ENiagaraVariantMode::Object;
 	Object = InObject;
@@ -43,13 +43,13 @@ FNiagaraVariant::FNiagaraVariant(const void* InBytes, int32 Size)
 	Object = nullptr;
 }
 
-UObject* FNiagaraVariant::GetUObject() const 
+TObjectPtr<UObject> FNiagaraVariant::GetUObject() const
 {
 	ensure(CurrentMode == ENiagaraVariantMode::Object);
 	return Object;
 }
 
-void FNiagaraVariant::SetUObject(UObject* InObject)
+void FNiagaraVariant::SetUObject(TObjectPtr<UObject> InObject)
 {
 	ensure(CurrentMode == ENiagaraVariantMode::None || CurrentMode == ENiagaraVariantMode::Object);
 
diff --git a/Engine/Plugins/FX/Niagara/Source/Niagara/Public/NiagaraCommon.h b/Engine/Plugins/FX/Niagara/Source/Niagara/Public/NiagaraCommon.h
--- a/Engine/Plugins/FX/Niagara/Source/Niagara/Public/NiagaraCommon.h
+++ b/Engine/Plugins/FX/Niagara/Source/Niagara/Public/NiagaraCommon.h
@@ -1033,12 +1033,13 @@ struct FNiagaraSystemUpdateContext
 {
 	GENERATED_BODY()
 
-	FNiagaraSystemUpdateContext(const UNiagaraSystem* System, bool bReInit, bool bInDestroyOnAdd = false, bool bInOnlyActive = false, bool bInDestroySystemSim = true) :bDestroyOnAdd(bInDestroyOnAdd), bOnlyActive(bInOnlyActive), bDestroySystemSim(bInDestroySystemSim) { Add(System, bReInit); }
-#if WITH_EDITORONLY_DATA
-	FNiagaraSystemUpdateContext(const FVersionedNiagaraEmitter& Emitter, bool bReInit, bool bInDestroyOnAdd = false, bool bInOnlyActive = false, bool bInDestroySystemSim = true) :bDestroyOnAdd(bInDestroyOnAdd), bOnlyActive(bInOnlyActive), bDestroySystemSim(bInDestroySystemSim) { Add(Emitter, bReInit); }
-	FNiagaraSystemUpdateContext(const UNiagaraScript* Script, bool bReInit, bool bInDestroyOnAdd = false, bool bInOnlyActive = false, bool bInDestroySystemSim = true) :bDestroyOnAdd(bInDestroyOnAdd), bOnlyActive(bInOnlyActive), bDestroySystemSim(bInDestroySystemSim) { Add(Script, bReInit); }
-	FNiagaraSystemUpdateContext(const UNiagaraParameterCollection* Collection, bool bReInit, bool bInDestroyOnAdd = false, bool bInOnlyActive = false, bool bInDestroySystemSim = true) :bDestroyOnAdd(bInDestroyOnAdd), bOnlyActive(bInOnlyActive), bDestroySystemSim(bInDestroySystemSim) { Add(Collection, bReInit); }
-#endif
+	template <typename AddableType> // Must be handled by Add function
+	FNiagaraSystemUpdateContext(AddableType&& Addable, bool bReInit,
+		bool bInDestroyOnAdd = false, bool bInOnlyActive = false, bool bInDestroySystemSim = true)
+		: bDestroyOnAdd(bInDestroyOnAdd), bOnlyActive(bInOnlyActive), bDestroySystemSim(bInDestroySystemSim)
+	{
+		Add(Forward<AddableType>(Addable), bReInit);
+	}
 	FNiagaraSystemUpdateContext() :bDestroyOnAdd(false), bOnlyActive(false), bDestroySystemSim(true) { }
 	FNiagaraSystemUpdateContext(FNiagaraSystemUpdateContext& Other) = delete;
 
@@ -1067,7 +1068,9 @@ struct FNiagaraSystemUpdateContext
 	FCustomWorkDelegate& GetPostWork() { return PostWork; }
 
 private:
-	NIAGARA_API void AddInternal(class UNiagaraComponent* Comp, bool bReInit, bool bAllowDestroySystemSim);
+	template <typename SystemMatchesFunc>
+	void AddComponentsWithMatchingSystem(bool bReinit, SystemMatchesFunc&& SystemMatches);
+	void AddInternal(class UNiagaraComponent* Comp, bool bReInit, bool bAllowDestroySystemSim);
 
 	UPROPERTY(transient)
 	TArray<TObjectPtr<UNiagaraComponent>> ComponentsToReset;
diff --git a/Engine/Plugins/FX/Niagara/Source/Niagara/Public/NiagaraParameterStore.h b/Engine/Plugins/FX/Niagara/Source/Niagara/Public/NiagaraParameterStore.h
--- a/Engine/Plugins/FX/Niagara/Source/Niagara/Public/NiagaraParameterStore.h
+++ b/Engine/Plugins/FX/Niagara/Source/Niagara/Public/NiagaraParameterStore.h
@@ -362,7 +362,7 @@ public:
 
 	FORCEINLINE TArray<FNiagaraParameterStore*>& GetSourceParameterStores() { return SourceStores; }
 
-	FORCEINLINE const TArray<UObject*>& GetUObjects()const { return UObjects; }
+	FORCEINLINE const TArray<TObjectPtr<UObject>>& GetUObjects()const { return UObjects; }
 	FORCEINLINE const TArray<UNiagaraDataInterface*>& GetDataInterfaces()const { return DataInterfaces; }
 	FORCEINLINE const TArray<uint8>& GetParameterDataArray()const { return ParameterData; }
 
@@ -481,7 +481,7 @@ public:
 	NIAGARA_API void SortParameters();
 
 	/** Returns the UObject at the passed offset. */
-	FORCEINLINE UObject* GetUObject(int32 Offset)const
+	FORCEINLINE TObjectPtr<UObject> GetUObject(int32 Offset)const
 	{
 		if (UObjects.IsValidIndex(Offset))
 		{
@@ -491,11 +491,11 @@ public:
 		return nullptr;
 	}
 
-	FORCEINLINE UObject* GetUObject(const FNiagaraVariable& Parameter)const
+	FORCEINLINE TObjectPtr<UObject> GetUObject(const FNiagaraVariable& Parameter)const
 	{
 		int32 Offset = IndexOf(Parameter);
-		UObject* Obj = GetUObject(Offset);
-		checkSlow(!Obj || Obj->IsA(Parameter.GetType().GetClass()));
+		TObjectPtr<UObject> Obj = GetUObject(Offset);
+		checkSlow(!Obj || Obj.IsA(Parameter.GetType().GetClass()));
 		return Obj;
 	}
 
@@ -596,13 +596,13 @@ public:
 		}
 	}
 
-	FORCEINLINE_DEBUGGABLE void SetUObject(UObject* InObject, int32 Offset)
+	FORCEINLINE_DEBUGGABLE void SetUObject(TObjectPtr<UObject> InObject, int32 Offset)
 	{
 		UObjects[Offset] = InObject;
 		OnUObjectChange();
 	}
 
-	FORCEINLINE_DEBUGGABLE void SetUObject(UObject* InObject, const FNiagaraVariable& Parameter)
+	FORCEINLINE_DEBUGGABLE void SetUObject(TObjectPtr<UObject> InObject, const FNiagaraVariable& Parameter)
 	{
 		int32 Offset = IndexOf(Parameter);
 		if (Offset != INDEX_NONE)
@@ -888,8 +888,8 @@ FORCEINLINE_DEBUGGABLE void FNiagaraParameterStoreBinding::Dump(const FNiagaraPa
 			UE_LOG(LogNiagara, Log, TEXT("| UObject | %s %s: Src:%d - Dest:%d\n"), *Param.GetType().GetName(), *Param.GetName().ToString(), Binding.SrcOffset, Binding.DestOffset);
 
 			//Also ensure the param has been pushed correctly.
-			const UObject* SrcData = SrcStore->GetUObjects()[Binding.SrcOffset];
-			const UObject* DestData = DestStore->GetUObjects()[Binding.DestOffset];
+			TObjectPtr<const UObject> SrcData = SrcStore->GetUObjects()[Binding.SrcOffset];
+			TObjectPtr<const UObject> DestData = DestStore->GetUObjects()[Binding.DestOffset];
 			if (!ensure(SrcData == DestData))
 			{
 				UE_LOG(LogNiagara, Log, TEXT("Data interface parameter in dest store is incorrect!\n"));
@@ -1244,7 +1244,7 @@ struct FNiagaraParameterDirectBinding<UObject*>
 
 	FORCEINLINE bool IsBound()const { return UObjectOffset != INDEX_NONE; }
 
-	UObject* Init(FNiagaraParameterStore& InStore, const FNiagaraVariable& DestVariable)
+	TObjectPtr<UObject> Init(FNiagaraParameterStore& InStore, const FNiagaraVariable& DestVariable)
 	{
 		if (DestVariable.IsValid())
 		{
@@ -1254,7 +1254,7 @@ struct FNiagaraParameterDirectBinding<UObject*>
 
 			check(BoundVariable.GetType().IsUObject());
 			UObjectOffset = BoundStore->IndexOf(DestVariable);
-			UObject* Ret = BoundStore->GetUObject(UObjectOffset);
+			TObjectPtr<UObject> Ret = BoundStore->GetUObject(UObjectOffset);
 			return Ret;
 		}
 		return nullptr;
@@ -1271,7 +1271,7 @@ struct FNiagaraParameterDirectBinding<UObject*>
 		}
 	}
 
-	FORCEINLINE UObject* GetValue() const
+	FORCEINLINE TObjectPtr<UObject> GetValue() const
 	{
 		if (UObjectOffset != INDEX_NONE)
 		{
diff --git a/Engine/Plugins/FX/Niagara/Source/Niagara/Public/NiagaraVariant.h b/Engine/Plugins/FX/Niagara/Source/Niagara/Public/NiagaraVariant.h
--- a/Engine/Plugins/FX/Niagara/Source/Niagara/Public/NiagaraVariant.h
+++ b/Engine/Plugins/FX/Niagara/Source/Niagara/Public/NiagaraVariant.h
@@ -28,12 +28,12 @@ struct FNiagaraVariant
 	FNiagaraVariant(const FNiagaraVariant&) = default;
 	FNiagaraVariant& operator=(const FNiagaraVariant&) = default;
 	NIAGARA_API explicit FNiagaraVariant(UNiagaraDataInterface* InDataInterface);
-	NIAGARA_API explicit FNiagaraVariant(UObject* InObject);
+	NIAGARA_API explicit FNiagaraVariant(TObjectPtr<UObject> InObject);
 	NIAGARA_API explicit FNiagaraVariant(const TArray<uint8>& InBytes);
 	NIAGARA_API FNiagaraVariant(const void* InBytes, int32 Size);
 
-	NIAGARA_API UObject* GetUObject() const;
-	NIAGARA_API void SetUObject(UObject* InObject);
+	NIAGARA_API TObjectPtr<UObject> GetUObject() const;
+	NIAGARA_API void SetUObject(TObjectPtr<UObject> InObject);
 
 	NIAGARA_API UNiagaraDataInterface* GetDataInterface() const;
 	NIAGARA_API void SetDataInterface(UNiagaraDataInterface* InDataInterface);
diff --git a/Engine/Plugins/FX/Niagara/Source/NiagaraEditor/Private/Customizations/NiagaraParameterCollectionCustomNodeBuilder.h b/Engine/Plugins/FX/Niagara/Source/NiagaraEditor/Private/Customizations/NiagaraParameterCollectionCustomNodeBuilder.h
--- a/Engine/Plugins/FX/Niagara/Source/NiagaraEditor/Private/Customizations/NiagaraParameterCollectionCustomNodeBuilder.h
+++ b/Engine/Plugins/FX/Niagara/Source/NiagaraEditor/Private/Customizations/NiagaraParameterCollectionCustomNodeBuilder.h
@@ -141,8 +141,8 @@ public:
 								.ObjectPath_Lambda(
 									[Parameter]() -> FString
 									{
-										UObject* DefaultValueObject = Parameter->GetDefaultValueObjectAsset();
-										return DefaultValueObject ? DefaultValueObject->GetPathName() : FString();
+										TObjectPtr<UObject> DefaultValueObject = Parameter->GetDefaultValueObjectAsset();
+										return DefaultValueObject ? DefaultValueObject.GetPathName() : FString();
 									}
 								)
 								.DisplayBrowse(true)
diff --git a/Engine/Plugins/FX/Niagara/Source/NiagaraEditor/Private/Customizations/UserParameters/NiagaraComponentUserParametersBuilder.cpp b/Engine/Plugins/FX/Niagara/Source/NiagaraEditor/Private/Customizations/UserParameters/NiagaraComponentUserParametersBuilder.cpp
--- a/Engine/Plugins/FX/Niagara/Source/NiagaraEditor/Private/Customizations/UserParameters/NiagaraComponentUserParametersBuilder.cpp
+++ b/Engine/Plugins/FX/Niagara/Source/NiagaraEditor/Private/Customizations/UserParameters/NiagaraComponentUserParametersBuilder.cpp
@@ -204,7 +204,7 @@ public:
 
 	FString GetCurrentAssetPath() const
 	{
-		UObject* CurrentObject = nullptr;
+		TObjectPtr<UObject> CurrentObject = nullptr;
 		
 		TWeakObjectPtr<UNiagaraComponent> RawComponent = GetComponent();
 
@@ -227,7 +227,7 @@ public:
 			}
 		}
 
-		return CurrentObject != nullptr ? CurrentObject->GetPathName() : FString();
+		return CurrentObject != nullptr ? CurrentObject.GetPathName() : FString();
 	}
 
 	void SetParameterEditor(TSharedRef<SNiagaraParameterEditor> InParameterEditor)
diff --git a/Engine/Plugins/FX/Niagara/Source/NiagaraEditor/Private/Customizations/UserParameters/NiagaraSystemUserParametersBuilder.cpp b/Engine/Plugins/FX/Niagara/Source/NiagaraEditor/Private/Customizations/UserParameters/NiagaraSystemUserParametersBuilder.cpp
--- a/Engine/Plugins/FX/Niagara/Source/NiagaraEditor/Private/Customizations/UserParameters/NiagaraSystemUserParametersBuilder.cpp
+++ b/Engine/Plugins/FX/Niagara/Source/NiagaraEditor/Private/Customizations/UserParameters/NiagaraSystemUserParametersBuilder.cpp
@@ -545,14 +545,14 @@ FString FNiagaraSystemUserParameterBuilder::GetObjectAssetPathForUserParameter(F
 	ensure(UserParameter.GetType().IsUObject() && UserParameter.GetType().IsDataInterface() == false);
 	FNiagaraVariant Variant = GetCurrentParameterValue(UserParameter);
 
-	if(Variant.IsValid() == false)
+	if (Variant.IsValid() == false)
 	{
 		return FString();
 	}
 	
-	if(UObject* Object = Variant.GetUObject())
+	if (TObjectPtr<UObject> Object = Variant.GetUObject())
 	{
-		return Object->GetPathName();
+		return Object.GetPathName();
 	}
 	
 	return FString();
diff --git a/Engine/Plugins/FX/Niagara/Source/NiagaraEditor/Private/Customizations/UserParameters/NiagaraUserParametersBuilderBase.cpp b/Engine/Plugins/FX/Niagara/Source/NiagaraEditor/Private/Customizations/UserParameters/NiagaraUserParametersBuilderBase.cpp
--- a/Engine/Plugins/FX/Niagara/Source/NiagaraEditor/Private/Customizations/UserParameters/NiagaraUserParametersBuilderBase.cpp
+++ b/Engine/Plugins/FX/Niagara/Source/NiagaraEditor/Private/Customizations/UserParameters/NiagaraUserParametersBuilderBase.cpp
@@ -385,7 +385,7 @@ void FNiagaraUserParameterNodeBuilder::AddCustomMenuActionsForParameter(FDetailW
 				return;
 			}
 
-			FPropertyEditorClipboard::ClipboardCopy(*CurrentParameterValue.GetUObject()->GetPathName());
+			FPropertyEditorClipboard::ClipboardCopy(*CurrentParameterValue.GetUObject().GetPathName());
 		});
 		WidgetRow.CopyAction(CopyAction);
 
diff --git a/Engine/Plugins/FX/Niagara/Source/NiagaraEditor/Private/NiagaraCompilationGraphBridgeImpl.cpp b/Engine/Plugins/FX/Niagara/Source/NiagaraEditor/Private/NiagaraCompilationGraphBridgeImpl.cpp
--- a/Engine/Plugins/FX/Niagara/Source/NiagaraEditor/Private/NiagaraCompilationGraphBridgeImpl.cpp
+++ b/Engine/Plugins/FX/Niagara/Source/NiagaraEditor/Private/NiagaraCompilationGraphBridgeImpl.cpp
@@ -511,7 +511,7 @@ UNiagaraDataInterface* FNiagaraCompilationGraphBridge::GetCollectionDataInterfac
 	return Collection->GetDefaultInstance()->GetParameterStore().GetDataInterface(Variable);
 }
 
-UObject* FNiagaraCompilationGraphBridge::GetCollectionUObject(FParameterCollection Collection, const FNiagaraVariable& Variable)
+TObjectPtr<UObject> FNiagaraCompilationGraphBridge::GetCollectionUObject(FParameterCollection Collection, const FNiagaraVariable& Variable)
 {
 	return Collection->GetDefaultInstance()->GetParameterStore().GetUObject(Variable);
 }
diff --git a/Engine/Plugins/FX/Niagara/Source/NiagaraEditor/Private/NiagaraNodeFunctionCall.cpp b/Engine/Plugins/FX/Niagara/Source/NiagaraEditor/Private/NiagaraNodeFunctionCall.cpp
--- a/Engine/Plugins/FX/Niagara/Source/NiagaraEditor/Private/NiagaraNodeFunctionCall.cpp
+++ b/Engine/Plugins/FX/Niagara/Source/NiagaraEditor/Private/NiagaraNodeFunctionCall.cpp
@@ -660,7 +660,7 @@ static FText GetFormattedDeprecationMessage(const FVersionedNiagaraScriptData* S
 
 	if (ScriptData->DeprecationRecommendation != nullptr)
 	{
-		Args.Add(TEXT("Recommendation"), FText::FromString(ScriptData->DeprecationRecommendation->GetPathName()));
+		Args.Add(TEXT("Recommendation"), FText::FromString(ScriptData->DeprecationRecommendation.GetPathName()));
 	}
 
 	if (ScriptData->DeprecationMessage.IsEmptyOrWhitespace() == false)
diff --git a/Engine/Plugins/FX/Niagara/Source/NiagaraEditor/Private/NiagaraNodeInput.cpp b/Engine/Plugins/FX/Niagara/Source/NiagaraEditor/Private/NiagaraNodeInput.cpp
--- a/Engine/Plugins/FX/Niagara/Source/NiagaraEditor/Private/NiagaraNodeInput.cpp
+++ b/Engine/Plugins/FX/Niagara/Source/NiagaraEditor/Private/NiagaraNodeInput.cpp
@@ -517,12 +517,12 @@ void UNiagaraNodeInput::SetDataInterface(UNiagaraDataInterface* InDataInterface)
 	DataInterfaceChanged();
 }
 
-UObject* UNiagaraNodeInput::GetObjectAsset() const
+TObjectPtr<UObject> UNiagaraNodeInput::GetObjectAsset() const
 {
 	return ObjectAsset;
 }
 
-void UNiagaraNodeInput::SetObjectAsset(UObject* InObject)
+void UNiagaraNodeInput::SetObjectAsset(TObjectPtr<UObject> InObject)
 {
 	ObjectAsset = InObject;
 }
diff --git a/Engine/Plugins/FX/Niagara/Source/NiagaraEditor/Private/ViewModels/NiagaraCollectionParameterViewModel.cpp b/Engine/Plugins/FX/Niagara/Source/NiagaraEditor/Private/ViewModels/NiagaraCollectionParameterViewModel.cpp
--- a/Engine/Plugins/FX/Niagara/Source/NiagaraEditor/Private/ViewModels/NiagaraCollectionParameterViewModel.cpp
+++ b/Engine/Plugins/FX/Niagara/Source/NiagaraEditor/Private/ViewModels/NiagaraCollectionParameterViewModel.cpp
@@ -188,12 +188,12 @@ UNiagaraDataInterface* FNiagaraCollectionParameterViewModel::GetDefaultValueData
 	return CollectionInst->GetParameterStore().GetDataInterface(Parameter);
 }
 
-UObject* FNiagaraCollectionParameterViewModel::GetDefaultValueObjectAsset()
+TObjectPtr<UObject> FNiagaraCollectionParameterViewModel::GetDefaultValueObjectAsset()
 {
 	return CollectionInst->GetParameterStore().GetUObject(Parameter);
 }
 
-void FNiagaraCollectionParameterViewModel::SetDefaultValueObjectAsset(UObject* Object)
+void FNiagaraCollectionParameterViewModel::SetDefaultValueObjectAsset(TObjectPtr<UObject> Object)
 {
 	CollectionInst->GetParameterStore().SetUObject(Object, Parameter);
 }
diff --git a/Engine/Plugins/FX/Niagara/Source/NiagaraEditor/Private/ViewModels/NiagaraCollectionParameterViewModel.h b/Engine/Plugins/FX/Niagara/Source/NiagaraEditor/Private/ViewModels/NiagaraCollectionParameterViewModel.h
--- a/Engine/Plugins/FX/Niagara/Source/NiagaraEditor/Private/ViewModels/NiagaraCollectionParameterViewModel.h
+++ b/Engine/Plugins/FX/Niagara/Source/NiagaraEditor/Private/ViewModels/NiagaraCollectionParameterViewModel.h
@@ -28,8 +28,8 @@ public:
 	virtual EDefaultValueType GetDefaultValueType() override;
 	virtual TSharedRef<FStructOnScope> GetDefaultValueStruct() override;
 	virtual UNiagaraDataInterface* GetDefaultValueDataInterface() override;
-	virtual UObject* GetDefaultValueObjectAsset() override;
-	virtual void SetDefaultValueObjectAsset(UObject*) override;
+	virtual TObjectPtr<UObject> GetDefaultValueObjectAsset() override;
+	virtual void SetDefaultValueObjectAsset(TObjectPtr<UObject>) override;
 	virtual void NotifyDefaultValuePropertyChanged(const FPropertyChangedEvent& PropertyChangedEvent) override;
 	virtual void NotifyBeginDefaultValueChange() override;
 	virtual void NotifyEndDefaultValueChange() override;
diff --git a/Engine/Plugins/FX/Niagara/Source/NiagaraEditor/Private/ViewModels/NiagaraParameterViewModel.h b/Engine/Plugins/FX/Niagara/Source/NiagaraEditor/Private/ViewModels/NiagaraParameterViewModel.h
--- a/Engine/Plugins/FX/Niagara/Source/NiagaraEditor/Private/ViewModels/NiagaraParameterViewModel.h
+++ b/Engine/Plugins/FX/Niagara/Source/NiagaraEditor/Private/ViewModels/NiagaraParameterViewModel.h
@@ -73,9 +73,9 @@ public:
 	virtual UNiagaraDataInterface* GetDefaultValueDataInterface() = 0;
 
 	/** Gets the object reference representing the default value for the parameter. */
-	virtual UObject* GetDefaultValueObjectAsset() = 0;
+	virtual TObjectPtr<UObject> GetDefaultValueObjectAsset() = 0;
 	/** Sets the object reference representing the default value for the parameter. */
-	virtual void SetDefaultValueObjectAsset(UObject*) = 0;
+	virtual void SetDefaultValueObjectAsset(TObjectPtr<UObject>) = 0;
 
 	/** Called to notify the parameter view model that a property on the default value has been changed by the UI. */
 	virtual void NotifyDefaultValuePropertyChanged(const FPropertyChangedEvent& PropertyChangedEvent) = 0;
diff --git a/Engine/Plugins/FX/Niagara/Source/NiagaraEditor/Private/ViewModels/NiagaraScriptInputCollectionViewModel.cpp b/Engine/Plugins/FX/Niagara/Source/NiagaraEditor/Private/ViewModels/NiagaraScriptInputCollectionViewModel.cpp
--- a/Engine/Plugins/FX/Niagara/Source/NiagaraEditor/Private/ViewModels/NiagaraScriptInputCollectionViewModel.cpp
+++ b/Engine/Plugins/FX/Niagara/Source/NiagaraEditor/Private/ViewModels/NiagaraScriptInputCollectionViewModel.cpp
@@ -546,7 +546,7 @@ void FNiagaraScriptInputCollectionViewModel::OnParameterValueChangedInternal(TSh
 		}
 		else if (ChangedParameter->GetDefaultValueType() == INiagaraParameterViewModel::EDefaultValueType::ObjectAsset)
 		{
-			UObject* NewObject = ChangedParameter->GetDefaultValueObjectAsset();
+			TObjectPtr<UObject> NewObject = ChangedParameter->GetDefaultValueObjectAsset();
 
 			for (UNiagaraNodeInput* InputNodeToUpdate : InputNodesToUpdate)
 			{
diff --git a/Engine/Plugins/FX/Niagara/Source/NiagaraEditor/Private/ViewModels/NiagaraScriptParameterViewModel.cpp b/Engine/Plugins/FX/Niagara/Source/NiagaraEditor/Private/ViewModels/NiagaraScriptParameterViewModel.cpp
--- a/Engine/Plugins/FX/Niagara/Source/NiagaraEditor/Private/ViewModels/NiagaraScriptParameterViewModel.cpp
+++ b/Engine/Plugins/FX/Niagara/Source/NiagaraEditor/Private/ViewModels/NiagaraScriptParameterViewModel.cpp
@@ -166,12 +166,12 @@ UNiagaraDataInterface* FNiagaraScriptParameterViewModel::GetDefaultValueDataInte
 	return ValueDataInterface;
 }
 
-UObject* FNiagaraScriptParameterViewModel::GetDefaultValueObjectAsset()
+TObjectPtr<UObject> FNiagaraScriptParameterViewModel::GetDefaultValueObjectAsset()
 {
 	return ValueObjectAsset;
 }
 
-void FNiagaraScriptParameterViewModel::SetDefaultValueObjectAsset(UObject* Object)
+void FNiagaraScriptParameterViewModel::SetDefaultValueObjectAsset(TObjectPtr<UObject> Object)
 {
 	ValueObjectAsset = Object;
 	OnDefaultValueChangedDelegate.Broadcast();
diff --git a/Engine/Plugins/FX/Niagara/Source/NiagaraEditor/Private/ViewModels/NiagaraScriptParameterViewModel.h b/Engine/Plugins/FX/Niagara/Source/NiagaraEditor/Private/ViewModels/NiagaraScriptParameterViewModel.h
--- a/Engine/Plugins/FX/Niagara/Source/NiagaraEditor/Private/ViewModels/NiagaraScriptParameterViewModel.h
+++ b/Engine/Plugins/FX/Niagara/Source/NiagaraEditor/Private/ViewModels/NiagaraScriptParameterViewModel.h
@@ -46,8 +46,8 @@ public:
 	virtual EDefaultValueType GetDefaultValueType() override;
 	virtual TSharedRef<FStructOnScope> GetDefaultValueStruct() override;
 	virtual UNiagaraDataInterface* GetDefaultValueDataInterface() override;
-	virtual UObject* GetDefaultValueObjectAsset() override;
-	virtual void SetDefaultValueObjectAsset(UObject*) override;
+	virtual TObjectPtr<UObject> GetDefaultValueObjectAsset() override;
+	virtual void SetDefaultValueObjectAsset(TObjectPtr<UObject>) override;
 	virtual void NotifyDefaultValuePropertyChanged(const FPropertyChangedEvent& PropertyChangedEvent) override;
 	virtual void NotifyBeginDefaultValueChange() override;
 	virtual void NotifyEndDefaultValueChange() override;
@@ -88,7 +88,7 @@ private:
 	UNiagaraDataInterface* ValueDataInterface = nullptr;
 
 	/** The object asset which is providing the parameter value. */
-	UObject* ValueObjectAsset = nullptr;
+	TObjectPtr<UObject> ValueObjectAsset = nullptr;
 
 	/** The type of default value this parameter is providing. */
 	EDefaultValueType DefaultValueType;
diff --git a/Engine/Plugins/FX/Niagara/Source/NiagaraEditor/Private/Widgets/SNiagaraParameterCollection.cpp b/Engine/Plugins/FX/Niagara/Source/NiagaraEditor/Private/Widgets/SNiagaraParameterCollection.cpp
--- a/Engine/Plugins/FX/Niagara/Source/NiagaraEditor/Private/Widgets/SNiagaraParameterCollection.cpp
+++ b/Engine/Plugins/FX/Niagara/Source/NiagaraEditor/Private/Widgets/SNiagaraParameterCollection.cpp
@@ -557,8 +557,8 @@ TSharedRef<ITableRow> SNiagaraParameterCollection::OnGenerateRowForParameter(TSh
 			.ObjectPath_Lambda(
 				[Item]() -> FString
 				{
-					UObject* DefaultValue = Item->GetDefaultValueObjectAsset();
-					return DefaultValue ? DefaultValue->GetPathName() : FString();
+					TObjectPtr<UObject> DefaultValue = Item->GetDefaultValueObjectAsset();
+					return DefaultValue ? DefaultValue.GetPathName() : FString();
 				}
 			)
 			.DisplayBrowse(true)
diff --git a/Engine/Plugins/FX/Niagara/Source/NiagaraEditor/Public/NiagaraCompilationBridge.h b/Engine/Plugins/FX/Niagara/Source/NiagaraEditor/Public/NiagaraCompilationBridge.h
--- a/Engine/Plugins/FX/Niagara/Source/NiagaraEditor/Public/NiagaraCompilationBridge.h
+++ b/Engine/Plugins/FX/Niagara/Source/NiagaraEditor/Public/NiagaraCompilationBridge.h
@@ -208,7 +208,7 @@ struct FNiagaraCompilationGraphBridge
 	static FString GetCollectionFullName(FParameterCollection Collection);
 	static bool IsCollectionValid(FParameterCollection Collection);
 	static UNiagaraDataInterface* GetCollectionDataInterface(FParameterCollection Collection, const FNiagaraVariable& Variable);
-	static UObject* GetCollectionUObject(FParameterCollection Collection, const FNiagaraVariable& Variable);
+	static TObjectPtr<UObject> GetCollectionUObject(FParameterCollection Collection, const FNiagaraVariable& Variable);
 
 	static const FOutputNode* AsOutputNode(const FNode* Node);
 	static bool IsParameterMapGet(const FNode* Node);
diff --git a/Engine/Plugins/FX/Niagara/Source/NiagaraEditor/Public/NiagaraNodeInput.h b/Engine/Plugins/FX/Niagara/Source/NiagaraEditor/Public/NiagaraNodeInput.h
--- a/Engine/Plugins/FX/Niagara/Source/NiagaraEditor/Public/NiagaraNodeInput.h
+++ b/Engine/Plugins/FX/Niagara/Source/NiagaraEditor/Public/NiagaraNodeInput.h
@@ -113,8 +113,8 @@ public:
 	UNiagaraDataInterface* GetDataInterface() const;
 	void SetDataInterface(UNiagaraDataInterface* InDataInterface);
 	
-	UObject* GetObjectAsset() const;
-	void SetObjectAsset(UObject* Object);
+	TObjectPtr<UObject> GetObjectAsset() const;
+	void SetObjectAsset(TObjectPtr<UObject> Object);
 
 private:
 	void DataInterfaceChanged();
