- Fix crash when toggling "Start with Tick Enabled"
- We manage our ticks so this can result in the tick being enabled on batched instances

#rnx

diff --git a/Engine/Plugins/FX/Niagara/Source/Niagara/Private/NiagaraComponent.cpp b/Engine/Plugins/FX/Niagara/Source/Niagara/Private/NiagaraComponent.cpp
--- a/Engine/Plugins/FX/Niagara/Source/Niagara/Private/NiagaraComponent.cpp
+++ b/Engine/Plugins/FX/Niagara/Source/Niagara/Private/NiagaraComponent.cpp
@@ -1212,7 +1212,7 @@ void UNiagaraComponent::ActivateInternal(bool bReset /* = false */, bool bIsScal
 	struct FScopedComponentTickEnabled
 	{
 		FScopedComponentTickEnabled(UNiagaraComponent* InComponent) : Component(InComponent) {}
-		~FScopedComponentTickEnabled() { Component->SetComponentTickEnabled(bTickEnabled); }
+		~FScopedComponentTickEnabled() { Component->PrimaryComponentTick.bStartWithTickEnabled &= bTickEnabled; Component->SetComponentTickEnabled(bTickEnabled); }
 		UNiagaraComponent* Component = nullptr;
 		bool bTickEnabled = false;
 	};
