Fix SpawnDirect compile flag that could cause issues when spawning into GPU simulations.

Also Fixed a few minor issues.
Default for SpawnDirect wizard is clamping all spawning to 1. This is confusing and will trip people up immediately. Change default to -1 aka no clamp.
Skip shared ptr churn in SetDataChannelData() if there is not change.
Hide a currently disabled allocation mode.

diff --git a/Engine/Plugins/FX/Niagara/Source/Niagara/Internal/DataInterface/NiagaraDataInterfaceDataChannelCommon.h b/Engine/Plugins/FX/Niagara/Source/Niagara/Internal/DataInterface/NiagaraDataInterfaceDataChannelCommon.h
--- a/Engine/Plugins/FX/Niagara/Source/Niagara/Internal/DataInterface/NiagaraDataInterfaceDataChannelCommon.h
+++ b/Engine/Plugins/FX/Niagara/Source/Niagara/Internal/DataInterface/NiagaraDataInterfaceDataChannelCommon.h
@@ -16,7 +16,7 @@ enum class ENiagaraDataChannelAllocationMode : uint8
 	Static,
 
 	/** Allocation count is determined by DI script calls to Allocate in Emitter Scripts. */
-	Dynamic
+	Dynamic UMETA(Hidden)
 };
 
 //TODO: Possible we may want to do reads and writes using data channels in a single system in future, avoiding the need to push data out to any manager class etc.
diff --git a/Engine/Plugins/FX/Niagara/Source/Niagara/Private/DataInterface/NiagaraDataInterfaceDataChannelCommon.cpp b/Engine/Plugins/FX/Niagara/Source/Niagara/Private/DataInterface/NiagaraDataInterfaceDataChannelCommon.cpp
--- a/Engine/Plugins/FX/Niagara/Source/Niagara/Private/DataInterface/NiagaraDataInterfaceDataChannelCommon.cpp
+++ b/Engine/Plugins/FX/Niagara/Source/Niagara/Private/DataInterface/NiagaraDataInterfaceDataChannelCommon.cpp
@@ -292,7 +292,7 @@ namespace NDIDataChannelUtilities
 {
 	const FName GetNDCSpawnDataName(TEXT("GetNDCSpawnData"));
 	const FName SpawnConditionalName(TEXT("SpawnConditional"));
-	const FName SpawnDirectName(TEXT("SpawnConditional"));
+	const FName SpawnDirectName(TEXT("SpawnDirect"));
 	const TGlobalResource<FNDIDummyUAV> DummyUAVFloat(PF_R32_FLOAT, sizeof(float));
 	const TGlobalResource<FNDIDummyUAV> DummyUAVInt32(PF_R32_SINT, sizeof(int32));
 	const TGlobalResource<FNDIDummyUAV> DummyUAVHalf(PF_R16F, sizeof(FFloat16));
diff --git a/Engine/Plugins/FX/Niagara/Source/Niagara/Private/DataInterface/NiagaraDataInterfaceDataChannelRead.cpp b/Engine/Plugins/FX/Niagara/Source/Niagara/Private/DataInterface/NiagaraDataInterfaceDataChannelRead.cpp
--- a/Engine/Plugins/FX/Niagara/Source/Niagara/Private/DataInterface/NiagaraDataInterfaceDataChannelRead.cpp
+++ b/Engine/Plugins/FX/Niagara/Source/Niagara/Private/DataInterface/NiagaraDataInterfaceDataChannelRead.cpp
@@ -352,6 +352,11 @@ void FNDIDataChannelReadInstanceData::Cleanup(UNiagaraDataInterfaceDataChannelRe
 
 void FNDIDataChannelReadInstanceData::SetDataChannelData(FNiagaraDataChannelDataPtr NewData, UNiagaraDataInterfaceDataChannelRead* Interface)
 {
+	if(DataChannelData == NewData)
+	{
+		return;
+	}
+
 	FNDIDataChannelCompiledData& CompiledData = Interface->GetCompiledData();
 	//If this interface spawns into a GPU emitter then we need to inform the NDC data so that it can automatically upload all CPU data to the GPU.
 	//Otherwise we get mis-matching data used for spawn and the subsequent reads on the GPU to init particles.
diff --git a/Engine/Plugins/FX/Niagara/Source/NiagaraEditor/Private/Widgets/DataChannel/NiagaraDataChannelWizard.cpp b/Engine/Plugins/FX/Niagara/Source/NiagaraEditor/Private/Widgets/DataChannel/NiagaraDataChannelWizard.cpp
--- a/Engine/Plugins/FX/Niagara/Source/NiagaraEditor/Private/Widgets/DataChannel/NiagaraDataChannelWizard.cpp
+++ b/Engine/Plugins/FX/Niagara/Source/NiagaraEditor/Private/Widgets/DataChannel/NiagaraDataChannelWizard.cpp
@@ -975,13 +975,15 @@ TSharedRef<FModuleWizardModel> DataChannel::CreateSpawnNDCModuleWizardModel()
 					if (UEdGraphPin* MinInput = SpawnFunction->FindPin(FName(SpawnMode == ENiagaraDataChanneSpawnModuleMode::ConditionalSpawn ? "Min Spawn Count" : "ClampMin"), EGPD_Input))
 					{
 						UEdGraphPin* SpawnMinPin = Utilities::AddReadParameterPin(FNiagaraTypeDefinition::GetIntDef(), FName("Min Count"), MapGetNode);
-						Utilities::SetDefaultValue(Graph, SpawnMinPin->PinName, FNiagaraTypeDefinition::GetIntDef(), 1);
+						int32 Default = SpawnMode == ENiagaraDataChanneSpawnModuleMode::ConditionalSpawn ? 1 : -1;//Default to -1 for spawn direct as this is a clamp. aka, no clamp by default.
+						Utilities::SetDefaultValue(Graph, SpawnMinPin->PinName, FNiagaraTypeDefinition::GetIntDef(), Default);
 						GraphSchema->TryCreateConnection(MinInput, SpawnMinPin);
 					}
 					if (UEdGraphPin* MaxInput = SpawnFunction->FindPin(FName(SpawnMode == ENiagaraDataChanneSpawnModuleMode::ConditionalSpawn ? "Max Spawn Count" : "ClampMax"), EGPD_Input))
 					{
 						UEdGraphPin* SpawnMaxPin = Utilities::AddReadParameterPin(FNiagaraTypeDefinition::GetIntDef(), FName("Max Count"), MapGetNode);
-						Utilities::SetDefaultValue(Graph, SpawnMaxPin->PinName, FNiagaraTypeDefinition::GetIntDef(), 1);
+						int32 Default = SpawnMode == ENiagaraDataChanneSpawnModuleMode::ConditionalSpawn ? 1 : -1;//Default to -1 for spawn direct as this is a clamp. aka, no clamp by default.
+						Utilities::SetDefaultValue(Graph, SpawnMaxPin->PinName, FNiagaraTypeDefinition::GetIntDef(), Default);
 						GraphSchema->TryCreateConnection(MaxInput, SpawnMaxPin);
 					}
 					if (SpawnMode == ENiagaraDataChanneSpawnModuleMode::DirectSpawn)
