Fixed issue with Emit flag being false for NDC Write calls.
Also removed over eager check that would fire in some valid cases. Now handling multiple pre stage calls per DI proxy for when it's used in multiple emitters.

diff --git a/Engine/Plugins/FX/Niagara/Source/Niagara/Private/DataInterface/NiagaraDataInterfaceDataChannelWrite.cpp b/Engine/Plugins/FX/Niagara/Source/Niagara/Private/DataInterface/NiagaraDataInterfaceDataChannelWrite.cpp
--- a/Engine/Plugins/FX/Niagara/Source/Niagara/Private/DataInterface/NiagaraDataInterfaceDataChannelWrite.cpp
+++ b/Engine/Plugins/FX/Niagara/Source/Niagara/Private/DataInterface/NiagaraDataInterfaceDataChannelWrite.cpp
@@ -1033,15 +1033,23 @@ void UNiagaraDataInterfaceDataChannelWrite::Write(FVectorVMExternalFunctionConte
 		if(FNiagaraDataBuffer* Data = InstData->DestinationData)
 		{			
 			bAllFailedFallback = false;
-			uint32 MaxLocalIndex = 0;
+			int32 MaxLocalIndex = INDEX_NONE;
 			int32 NumAllocated = IntCastChecked<int32>(Data->GetNumInstancesAllocated());
 			for (int32 i = 0; i < Context.GetNumInstances(); ++i)
 			{
-				int32 RawIndex = InIndex.GetAndAdvance();
-				bool bEmit = InEmit.GetAndAdvance() && RawIndex >= 0 && RawIndex < NumAllocated;
+				int32 Index = InIndex.GetAndAdvance();
+				bool bEmit = InEmit.GetAndAdvance() && Index >= 0 && Index < NumAllocated;
 
-				uint32 Index = static_cast<uint32>(RawIndex);
-				MaxLocalIndex = bEmit ? FMath::Max(Index, MaxLocalIndex) : MaxLocalIndex;
+				if(!bEmit)
+				{
+					if (OutSuccess.IsValid())
+					{
+						OutSuccess.SetAndAdvance(false);
+					}
+					continue;
+				}
+
+				MaxLocalIndex = FMath::Max(Index, MaxLocalIndex);
 
 				bool bAllWritesSuccess = true;
 
@@ -1089,12 +1097,15 @@ void UNiagaraDataInterfaceDataChannelWrite::Write(FVectorVMExternalFunctionConte
 				}
 			}
 
-			//Update the shared instance count with an updated max.
-			uint32 CurrNumInstances = AtomicNumInstances;
-			uint32 MaxLocalNumInstances = MaxLocalIndex + 1;
-			while(CurrNumInstances < MaxLocalNumInstances && !AtomicNumInstances.compare_exchange_weak(CurrNumInstances, MaxLocalNumInstances))
+			if(MaxLocalIndex != INDEX_NONE)
 			{
-				CurrNumInstances = AtomicNumInstances;
+				//Update the shared instance count with an updated max.
+				uint32 CurrNumInstances = AtomicNumInstances;
+				uint32 MaxLocalNumInstances = MaxLocalIndex + 1;
+				while (CurrNumInstances < MaxLocalNumInstances && !AtomicNumInstances.compare_exchange_weak(CurrNumInstances, MaxLocalNumInstances))
+				{
+					CurrNumInstances = AtomicNumInstances;
+				}
 			}
 		}
 	}
@@ -1519,9 +1530,8 @@ void FNiagaraDataInterfaceProxy_DataChannelWrite::PreStage(const FNDIGpuComputeP
 			InstanceData->BufferForCPU->SetGPUInstanceCountBufferOffset(Context.GetInstanceCountManager().AcquireEntry());
 		}
 
-		if(InstanceData->bPublishToGPU)
+		if(InstanceData->bPublishToGPU && InstanceData->GPUBuffer == nullptr)
 		{
-			check(InstanceData->GPUBuffer == nullptr);
 			InstanceData->GPUBuffer = InstanceData->ChannelDataRTProxy->PrepareForWriteAccess(Context.GetGraphBuilder());
 		}
 	}
