Fix for handling migration of HierarchyRoot data from the deprecated parameters on the script data.  Also adds some safety nets to ensure that we don't run into the case where StackItem children have duplicate entries.
Serialization issue was caused by GetLinkerCustomVersion() not necessarily being a valid source of the data version when it comes to being called from PostLoad() of a duplicated object, which is a common occurrence in emitter merging.
#rnx

diff --git a/Engine/Plugins/FX/Niagara/Source/Niagara/Classes/NiagaraScript.h b/Engine/Plugins/FX/Niagara/Source/Niagara/Classes/NiagaraScript.h
--- a/Engine/Plugins/FX/Niagara/Source/Niagara/Classes/NiagaraScript.h
+++ b/Engine/Plugins/FX/Niagara/Source/Niagara/Classes/NiagaraScript.h
@@ -1470,6 +1470,11 @@ private:
 
 	/* Flag set on load based on whether we're loading from a cooked package. */
 	bool IsCooked = false;
+
+	/* Transient flag marking that the script will need to rebuild the hierarchy root based on the custom version encountered during serialize.
+	*  Note that this is required because GetLinkerCustomVersion() is not reliable when assets are being duplicated, something that happens
+	*  often to our scripts during emitter merging, for example */
+	bool bMigrateParameterDataToHierarchyRoot = false;
 #endif
 
 	/** Compiled VM bytecode and data necessary to run this script.*/
diff --git a/Engine/Plugins/FX/Niagara/Source/Niagara/Classes/NiagaraScriptSourceBase.h b/Engine/Plugins/FX/Niagara/Source/Niagara/Classes/NiagaraScriptSourceBase.h
--- a/Engine/Plugins/FX/Niagara/Source/Niagara/Classes/NiagaraScriptSourceBase.h
+++ b/Engine/Plugins/FX/Niagara/Source/Niagara/Classes/NiagaraScriptSourceBase.h
@@ -155,7 +155,7 @@ class UNiagaraScriptSourceBase : public UObject
 	/** Adds a module if it isn't already in the graph. If the module isn't found bOutFoundModule will be false. If it is found and it did need to be added, the function returns true. If it already exists, it returns false. */
 	virtual bool AddModuleIfMissing(FString ModulePath, ENiagaraScriptUsage Usage, bool& bOutFoundModule) { bOutFoundModule = false; return false; }
 
-	virtual void PostLoadFromOwner(FVersionedNiagaraScriptData& OwnerData) {}
+	virtual void MigrateParameterDataToHierarchyRoot(FVersionedNiagaraScriptData& OwnerData) {}
 #if WITH_EDITOR
 	virtual void CleanUpOldAndInitializeNewRapidIterationParameters(const FVersionedNiagaraEmitter& Emitter, ENiagaraScriptUsage ScriptUsage, FGuid ScriptUsageId, FNiagaraParameterStore& RapidIterationParameters) const { checkf(false, TEXT("Not implemented")); }
 
diff --git a/Engine/Plugins/FX/Niagara/Source/Niagara/Private/NiagaraScript.cpp b/Engine/Plugins/FX/Niagara/Source/Niagara/Private/NiagaraScript.cpp
--- a/Engine/Plugins/FX/Niagara/Source/Niagara/Private/NiagaraScript.cpp
+++ b/Engine/Plugins/FX/Niagara/Source/Niagara/Private/NiagaraScript.cpp
@@ -2072,6 +2072,11 @@ void UNiagaraScript::Serialize(FArchive& Ar)
 	{
 		SaveShaderStableKeys(Ar.CookingTarget());
 	}
+
+	if (Ar.IsLoading())
+	{
+		bMigrateParameterDataToHierarchyRoot = NiagaraVer < FNiagaraCustomVersion::HierarchyEditorScriptSupport;
+	}
 #endif
 
 	SerializeNiagaraShaderMaps(Ar, NiagaraVer, IsValidShaderScript);
@@ -2358,7 +2363,11 @@ void UNiagaraScript::PostLoad()
 		{
 			Source->ConditionalPostLoad();
 
-			Source->PostLoadFromOwner(Data);
+			if (bMigrateParameterDataToHierarchyRoot)
+			{
+				Source->MigrateParameterDataToHierarchyRoot(Data);
+				bMigrateParameterDataToHierarchyRoot = false;
+			}
 
 			// Synchronize with Definitions after source scripts have been postloaded.
 			FVersionedNiagaraScript& VersionedScriptAdapter = VersionedScriptAdapters.Emplace_GetRef(this, Data.Version.VersionGuid);
diff --git a/Engine/Plugins/FX/Niagara/Source/NiagaraEditor/Private/NiagaraGraph.cpp b/Engine/Plugins/FX/Niagara/Source/NiagaraEditor/Private/NiagaraGraph.cpp
--- a/Engine/Plugins/FX/Niagara/Source/NiagaraEditor/Private/NiagaraGraph.cpp
+++ b/Engine/Plugins/FX/Niagara/Source/NiagaraEditor/Private/NiagaraGraph.cpp
@@ -4339,5 +4339,214 @@ bool UNiagaraGraph::ReferencesStaticVariable(FNiagaraStaticVariableSearchContext
 	return false;
 }
 
+/** Migrates the deprecated data from the provided OwnerData and populates the ParameterHierarchyRoot */
+void UNiagaraGraph::MigrateParameterScriptDataToHierarchyRoot(FVersionedNiagaraScriptData& OwnerData)
+{
+	UNiagaraHierarchyRoot* ParameterRoot = GetScriptParameterHierarchyRoot();
+
+	TMap<FName, UNiagaraHierarchySection*> HierarchySections;
+	TMap<FName, UNiagaraHierarchyCategory*> IdentityCategoryMap;
+	TMap<FName, FName> SectionCategoryMapping;
+		
+	// First we transfer all sections and create contained categories
+	for(const FNiagaraStackSection& StackSection : OwnerData.InputSections_DEPRECATED)
+	{
+		FText SectionDisplayName = StackSection.SectionDisplayName.IsEmptyOrWhitespace() ? FText::FromName(StackSection.SectionIdentifier) : StackSection.SectionDisplayName;
+		UNiagaraHierarchySection* HierarchySection = ParameterRoot->AddSection(SectionDisplayName);
+		HierarchySection->SetTooltip(StackSection.Tooltip);
+
+		FNiagaraHierarchyIdentity SectionIdentity;
+		SectionIdentity.Names.Add(StackSection.SectionIdentifier);
+		HierarchySection->SetIdentity(SectionIdentity);
+			
+		HierarchySections.Add(StackSection.SectionIdentifier, HierarchySection);
+
+		for(const FText& CategoryIdentity : StackSection.Categories)
+		{
+			// We skip categories that don't have a proper identifier setup
+			if(CategoryIdentity.IsEmptyOrWhitespace())
+			{
+				continue;
+			}
+				
+			FName CategoryName = FName(CategoryIdentity.ToString());
+			UNiagaraHierarchyScriptCategory* HierarchyCategory = ParameterRoot->AddChild<UNiagaraHierarchyScriptCategory>();
+			HierarchyCategory->SetSection(HierarchySection);
+			HierarchyCategory->SetCategoryName(CategoryName);
+			IdentityCategoryMap.Add(CategoryName, HierarchyCategory);
+			SectionCategoryMapping.Add(CategoryName, HierarchySection->GetSectionName());
+		}
+	}
+
+	OwnerData.InputSections_DEPRECATED.Empty();
+		
+	TMap<FNiagaraVariable, TObjectPtr<UNiagaraScriptVariable>> ScriptVariablesMap = GetAllMetaData();
+
+	TArray<UNiagaraScriptVariable*> ParentLevelInputs;
+	TMap<UNiagaraScriptVariable*, TArray<UNiagaraScriptVariable*>> ParentChildInputMap;
+	// We determine all parent & child level module inputs & static switches
+	for(const auto& ScriptVariablePair : ScriptVariablesMap)
+	{
+		if(ScriptVariablePair.Key.IsInNameSpace(FNiagaraConstants::ModuleNamespace) == false && ScriptVariablePair.Value->GetIsStaticSwitch() == false)
+		{
+			continue;
+		}
+			
+		UNiagaraScriptVariable* ScriptVariable = ScriptVariablePair.Value;
+
+		if(ScriptVariable->Metadata.GetParentAttribute_DEPRECATED().IsNone())
+		{
+			ParentLevelInputs.Add(ScriptVariable);
+			ParentChildInputMap.Add(ScriptVariable, GetChildScriptVariablesForInput_Deprecated(ScriptVariable->Metadata.GetVariableGuid()));
+		}
+	}
+
+	// Now we create the actual hierarchy elements
+	for(const auto& ParentChildPair : ParentChildInputMap)
+	{
+		// By default, we add to the root directly, unless a category was specified
+		UNiagaraHierarchyItemBase* Owner = ParameterRoot;
+			
+		FText SpecifiedCategoryName = ParentChildPair.Key->Metadata.GetCategoryName_DEPRECATED();
+		if(SpecifiedCategoryName.IsEmptyOrWhitespace() == false)
+		{
+			// Categories might already exist due to section info
+			if(IdentityCategoryMap.Contains(FName(SpecifiedCategoryName.ToString())))
+			{
+				Owner = IdentityCategoryMap[FName(SpecifiedCategoryName.ToString())];
+			}
+			// But if not we create a new category
+			else
+			{
+				UNiagaraHierarchyScriptCategory* Category = ParameterRoot->AddChild<UNiagaraHierarchyScriptCategory>();
+				Category->SetCategoryName(FName(SpecifiedCategoryName.ToString()));
+				IdentityCategoryMap.Add(Category->GetCategoryName(), Category);
+
+				Owner = Category;
+			}
+		}
+
+		// We add the parent input to either the root, or a given category
+		UNiagaraHierarchyScriptParameter* ParentHierarchyInput = Owner->AddChild<UNiagaraHierarchyScriptParameter>();
+		ParentHierarchyInput->Initialize(*ParentChildPair.Key);
+
+		// Children inputs can be added directly to the parent input and are already sorted
+		for(UNiagaraScriptVariable* ChildInput : ParentChildPair.Value)
+		{
+			UNiagaraHierarchyScriptParameter* ChildHierarchyInput = ParentHierarchyInput->AddChild<UNiagaraHierarchyScriptParameter>();
+			ChildHierarchyInput->Initialize(*ChildInput);
+		}
+	}
+
+	// now that we have created all inputs, categories and sections, we fix up the order
+	TArray<UNiagaraHierarchyCategory*> TopLevelCategories;
+	ParameterRoot->GetChildrenOfType(TopLevelCategories, false);
+		
+	// inputs are generally sorted by editor sort priority, but within a given context (within a category etc.)
+	// child inputs also already are sorted when constructed
+	auto SortInputsPredicate = ([&](UNiagaraHierarchyItemBase& ItemA, UNiagaraHierarchyItemBase& ItemB)
+	{
+		UNiagaraHierarchyScriptParameter* ScriptParameterA = Cast<UNiagaraHierarchyScriptParameter>(&ItemA);
+		UNiagaraHierarchyScriptParameter* ScriptParameterB = Cast<UNiagaraHierarchyScriptParameter>(&ItemB);
+
+		if(ScriptParameterA == nullptr || ScriptParameterB == nullptr)
+		{
+			return false;
+		}
+
+		const FNiagaraVariableMetaData& MetaDataA = ScriptParameterA->GetScriptVariable()->Metadata;
+		const FNiagaraVariableMetaData& MetaDataB = ScriptParameterB->GetScriptVariable()->Metadata;
+
+		if(MetaDataA.bAdvancedDisplay && !MetaDataB.bAdvancedDisplay)
+		{
+			return false;
+		}
+		else if(!MetaDataA.bAdvancedDisplay && MetaDataB.bAdvancedDisplay)
+		{
+			return true;
+		}
+			
+		if(MetaDataA.GetEditorSortPriority_DEPRECATED() != MetaDataB.GetEditorSortPriority_DEPRECATED())
+		{
+			return MetaDataA.GetEditorSortPriority_DEPRECATED() < MetaDataB.GetEditorSortPriority_DEPRECATED();
+		}
+			
+		return ItemA.ToString().Compare(ItemB.ToString(), ESearchCase::Type::IgnoreCase) < 0;
+	});
+
+	// We now sort inputs, and keep track of the minimum sort order per category
+	// This is because the stack determined category order by input sort priority (the category with the parameter of the lowest priority value would come first)
+	TMap<UNiagaraHierarchyItemBase*, int32> CategoryMinimumInputSortOrderMap;
+	for(UNiagaraHierarchyCategory* Category : TopLevelCategories)
+	{
+		Category->SortChildren(SortInputsPredicate, false);
+		if(Category->GetChildren().Num() > 0)
+		{
+			// since we have sorted the children by sort priority, child 0 has the lowest sort priority number (= highest priority)
+			UNiagaraHierarchyScriptParameter* MinimumItem = CastChecked<UNiagaraHierarchyScriptParameter>(Category->GetChildren()[0]);
+			int32 MinimumCategorySortOrder = MinimumItem->GetScriptVariable()->Metadata.GetEditorSortPriority_DEPRECATED();
+			
+			CategoryMinimumInputSortOrderMap.Add(Category, MinimumCategorySortOrder);
+		}
+		else
+		{
+			CategoryMinimumInputSortOrderMap.Add(Category, 0);
+		}
+	}
+
+	auto SortRootChildrenPredicate = ([&](UNiagaraHierarchyItemBase& ItemA, UNiagaraHierarchyItemBase& ItemB)
+	{
+		if(ItemA.IsA<UNiagaraHierarchyCategory>() && ItemB.IsA<UNiagaraHierarchyCategory>())
+		{
+			if(CategoryMinimumInputSortOrderMap[&ItemA] != CategoryMinimumInputSortOrderMap[&ItemB])
+			{
+				return CategoryMinimumInputSortOrderMap[&ItemA] < CategoryMinimumInputSortOrderMap[&ItemB];
+			}
+
+			return ItemA.ToString().Compare(ItemB.ToString(), ESearchCase::Type::IgnoreCase) < 0;
+		}
+
+		// Categories always come after items to ensure consistency with old behavior
+		if(ItemA.IsA<UNiagaraHierarchyCategory>() && ItemB.IsA<UNiagaraHierarchyItem>())
+		{
+			return false;
+		}
+			
+		if(ItemB.IsA<UNiagaraHierarchyCategory>() && ItemA.IsA<UNiagaraHierarchyItem>())
+		{
+			return true;
+		}
+			
+		if(ItemA.IsA<UNiagaraHierarchyScriptParameter>() && ItemB.IsA<UNiagaraHierarchyScriptParameter>())
+		{
+			UNiagaraHierarchyScriptParameter* ScriptParameterA = Cast<UNiagaraHierarchyScriptParameter>(&ItemA);
+			UNiagaraHierarchyScriptParameter* ScriptParameterB = Cast<UNiagaraHierarchyScriptParameter>(&ItemB);
+
+			const FNiagaraVariableMetaData& MetaDataA = ScriptParameterA->GetScriptVariable()->Metadata;
+			const FNiagaraVariableMetaData& MetaDataB = ScriptParameterB->GetScriptVariable()->Metadata;
+
+			if(MetaDataA.bAdvancedDisplay && !MetaDataB.bAdvancedDisplay)
+			{
+				return false;
+			}
+			else if(!MetaDataA.bAdvancedDisplay && MetaDataB.bAdvancedDisplay)
+			{
+				return true;
+			}
+			
+			if(MetaDataA.GetEditorSortPriority_DEPRECATED() != MetaDataB.GetEditorSortPriority_DEPRECATED())
+			{
+				return MetaDataA.GetEditorSortPriority_DEPRECATED() < MetaDataB.GetEditorSortPriority_DEPRECATED();
+			}
+			
+			return ItemA.ToString().Compare(ItemB.ToString(), ESearchCase::Type::IgnoreCase) < 0;
+		}
+			
+		return false;
+	});
+
+	ParameterRoot->SortChildren(SortRootChildrenPredicate, false);
+}
+
 #undef NIAGARA_SCOPE_CYCLE_COUNTER
 #undef LOCTEXT_NAMESPACE
diff --git a/Engine/Plugins/FX/Niagara/Source/NiagaraEditor/Private/NiagaraScriptSource.cpp b/Engine/Plugins/FX/Niagara/Source/NiagaraEditor/Private/NiagaraScriptSource.cpp
--- a/Engine/Plugins/FX/Niagara/Source/NiagaraEditor/Private/NiagaraScriptSource.cpp
+++ b/Engine/Plugins/FX/Niagara/Source/NiagaraEditor/Private/NiagaraScriptSource.cpp
@@ -370,9 +370,12 @@ bool UNiagaraScriptSource::AddModuleIfMissing(FString ModulePath, ENiagaraScript
 	return false;
 }
 
-void UNiagaraScriptSource::PostLoadFromOwner(FVersionedNiagaraScriptData& OwnerData)
+void UNiagaraScriptSource::MigrateParameterDataToHierarchyRoot(FVersionedNiagaraScriptData& OwnerData)
 {
-	PostLoadTransferScriptHierarchyData(OwnerData);
+	if (NodeGraph)
+	{
+		NodeGraph->MigrateParameterScriptDataToHierarchyRoot(OwnerData);
+	}
 }
 
 void UNiagaraScriptSource::FixupRenamedParameters(UNiagaraNodeFunctionCall* FunctionCallNode, TConstArrayView<FNiagaraVariable> ModuleInputVariables, FNiagaraParameterStore& RapidIterationParameters, const TArray<FNiagaraVariable>& OldRapidIterationVariables, const FVersionedNiagaraEmitter& VersionedEmitter, ENiagaraScriptUsage ScriptUsage) const
@@ -669,227 +672,6 @@ void UNiagaraScriptSource::OnGraphDataInterfaceChanged()
 	OnChangedDelegate.Broadcast();
 }
 
-void UNiagaraScriptSource::PostLoadTransferScriptHierarchyData(FVersionedNiagaraScriptData& OwnerData)
-{
-	UNiagaraGraph* Graph = Cast<UNiagaraScriptSource>(OwnerData.GetSource())->NodeGraph;
-
-	if(Graph == nullptr)
-	{
-		UE_LOG(LogNiagaraEditor, Verbose, TEXT("Failed post loading script hierarchy of versioned script data. No graph found."));
-		return;
-	}
-	
-	const int32 CustomVersion = Graph->GetLinkerCustomVersion(FNiagaraCustomVersion::GUID);
-
-	if(CustomVersion < FNiagaraCustomVersion::HierarchyEditorScriptSupport)
-	{
-		UNiagaraHierarchyRoot* ParameterRoot = Graph->GetScriptParameterHierarchyRoot();
-
-		TMap<FName, UNiagaraHierarchySection*> HierarchySections;
-		TMap<FName, UNiagaraHierarchyCategory*> IdentityCategoryMap;
-		TMap<FName, FName> SectionCategoryMapping;
-		
-		// First we transfer all sections and create contained categories
-		for(const FNiagaraStackSection& StackSection : OwnerData.InputSections_DEPRECATED)
-		{
-			FText SectionDisplayName = StackSection.SectionDisplayName.IsEmptyOrWhitespace() ? FText::FromName(StackSection.SectionIdentifier) : StackSection.SectionDisplayName;
-			UNiagaraHierarchySection* HierarchySection = ParameterRoot->AddSection(SectionDisplayName);
-			HierarchySection->SetTooltip(StackSection.Tooltip);
-
-			FNiagaraHierarchyIdentity SectionIdentity;
-			SectionIdentity.Names.Add(StackSection.SectionIdentifier);
-			HierarchySection->SetIdentity(SectionIdentity);
-			
-			HierarchySections.Add(StackSection.SectionIdentifier, HierarchySection);
-
-			for(const FText& CategoryIdentity : StackSection.Categories)
-			{
-				// We skip categories that don't have a proper identifier setup
-				if(CategoryIdentity.IsEmptyOrWhitespace())
-				{
-					continue;
-				}
-				
-				FName CategoryName = FName(CategoryIdentity.ToString());
-				UNiagaraHierarchyScriptCategory* HierarchyCategory = ParameterRoot->AddChild<UNiagaraHierarchyScriptCategory>();
-				HierarchyCategory->SetSection(HierarchySection);
-				HierarchyCategory->SetCategoryName(CategoryName);
-				IdentityCategoryMap.Add(CategoryName, HierarchyCategory);
-				SectionCategoryMapping.Add(CategoryName, HierarchySection->GetSectionName());
-			}
-		}
-
-		OwnerData.InputSections_DEPRECATED.Empty();
-		
-		TMap<FNiagaraVariable, TObjectPtr<UNiagaraScriptVariable>> ScriptVariablesMap = Graph->GetAllMetaData();
-
-		TArray<UNiagaraScriptVariable*> ParentLevelInputs;
-		TMap<UNiagaraScriptVariable*, TArray<UNiagaraScriptVariable*>> ParentChildInputMap;
-		// We determine all parent & child level module inputs & static switches
-		for(const auto& ScriptVariablePair : ScriptVariablesMap)
-		{
-			if(ScriptVariablePair.Key.IsInNameSpace(FNiagaraConstants::ModuleNamespace) == false && ScriptVariablePair.Value->GetIsStaticSwitch() == false)
-			{
-				continue;
-			}
-			
-			UNiagaraScriptVariable* ScriptVariable = ScriptVariablePair.Value;
-
-			if(ScriptVariable->Metadata.GetParentAttribute_DEPRECATED().IsNone())
-			{
-				ParentLevelInputs.Add(ScriptVariable);
-				ParentChildInputMap.Add(ScriptVariable, Graph->GetChildScriptVariablesForInput_Deprecated(ScriptVariable->Metadata.GetVariableGuid()));
-			}
-		}
-
-		// Now we create the actual hierarchy elements
-		for(const auto& ParentChildPair : ParentChildInputMap)
-		{
-			// By default, we add to the root directly, unless a category was specified
-			UNiagaraHierarchyItemBase* Owner = ParameterRoot;
-			
-			FText SpecifiedCategoryName = ParentChildPair.Key->Metadata.GetCategoryName_DEPRECATED();
-			if(SpecifiedCategoryName.IsEmptyOrWhitespace() == false)
-			{
-				// Categories might already exist due to section info
-				if(IdentityCategoryMap.Contains(FName(SpecifiedCategoryName.ToString())))
-				{
-					Owner = IdentityCategoryMap[FName(SpecifiedCategoryName.ToString())];
-				}
-				// But if not we create a new category
-				else
-				{
-					UNiagaraHierarchyScriptCategory* Category = ParameterRoot->AddChild<UNiagaraHierarchyScriptCategory>();
-					Category->SetCategoryName(FName(SpecifiedCategoryName.ToString()));
-					IdentityCategoryMap.Add(Category->GetCategoryName(), Category);
-
-					Owner = Category;
-				}
-			}
-
-			// We add the parent input to either the root, or a given category
-			UNiagaraHierarchyScriptParameter* ParentHierarchyInput = Owner->AddChild<UNiagaraHierarchyScriptParameter>();
-			ParentHierarchyInput->Initialize(*ParentChildPair.Key);
-
-			// Children inputs can be added directly to the parent input and are already sorted
-			for(UNiagaraScriptVariable* ChildInput : ParentChildPair.Value)
-			{
-				UNiagaraHierarchyScriptParameter* ChildHierarchyInput = ParentHierarchyInput->AddChild<UNiagaraHierarchyScriptParameter>();
-				ChildHierarchyInput->Initialize(*ChildInput);
-			}
-		}
-
-		// now that we have created all inputs, categories and sections, we fix up the order
-		TArray<UNiagaraHierarchyCategory*> TopLevelCategories;
-		ParameterRoot->GetChildrenOfType(TopLevelCategories, false);
-		
-		// inputs are generally sorted by editor sort priority, but within a given context (within a category etc.)
-		// child inputs also already are sorted when constructed
-		auto SortInputsPredicate = ([&](UNiagaraHierarchyItemBase& ItemA, UNiagaraHierarchyItemBase& ItemB)
-		{
-			UNiagaraHierarchyScriptParameter* ScriptParameterA = Cast<UNiagaraHierarchyScriptParameter>(&ItemA);
-			UNiagaraHierarchyScriptParameter* ScriptParameterB = Cast<UNiagaraHierarchyScriptParameter>(&ItemB);
-
-			if(ScriptParameterA == nullptr || ScriptParameterB == nullptr)
-			{
-				return false;
-			}
-
-			const FNiagaraVariableMetaData& MetaDataA = ScriptParameterA->GetScriptVariable()->Metadata;
-			const FNiagaraVariableMetaData& MetaDataB = ScriptParameterB->GetScriptVariable()->Metadata;
-
-			if(MetaDataA.bAdvancedDisplay && !MetaDataB.bAdvancedDisplay)
-			{
-				return false;
-			}
-			else if(!MetaDataA.bAdvancedDisplay && MetaDataB.bAdvancedDisplay)
-			{
-				return true;
-			}
-			
-			if(MetaDataA.GetEditorSortPriority_DEPRECATED() != MetaDataB.GetEditorSortPriority_DEPRECATED())
-			{
-				return MetaDataA.GetEditorSortPriority_DEPRECATED() < MetaDataB.GetEditorSortPriority_DEPRECATED();
-			}
-			
-			return ItemA.ToString().Compare(ItemB.ToString(), ESearchCase::Type::IgnoreCase) < 0;
-		});
-
-		// We now sort inputs, and keep track of the minimum sort order per category
-		// This is because the stack determined category order by input sort priority (the category with the parameter of the lowest priority value would come first)
-		TMap<UNiagaraHierarchyItemBase*, int32> CategoryMinimumInputSortOrderMap;
-		for(UNiagaraHierarchyCategory* Category : TopLevelCategories)
-		{
-			Category->SortChildren(SortInputsPredicate, false);
-			if(Category->GetChildren().Num() > 0)
-			{
-				// since we have sorted the children by sort priority, child 0 has the lowest sort priority number (= highest priority)
-				UNiagaraHierarchyScriptParameter* MinimumItem = CastChecked<UNiagaraHierarchyScriptParameter>(Category->GetChildren()[0]);
-				int32 MinimumCategorySortOrder = MinimumItem->GetScriptVariable()->Metadata.GetEditorSortPriority_DEPRECATED();
-			
-				CategoryMinimumInputSortOrderMap.Add(Category, MinimumCategorySortOrder);
-			}
-			else
-			{
-				CategoryMinimumInputSortOrderMap.Add(Category, 0);
-			}
-		}
-
-		auto SortRootChildrenPredicate = ([&](UNiagaraHierarchyItemBase& ItemA, UNiagaraHierarchyItemBase& ItemB)
-		{
-			if(ItemA.IsA<UNiagaraHierarchyCategory>() && ItemB.IsA<UNiagaraHierarchyCategory>())
-			{
-				if(CategoryMinimumInputSortOrderMap[&ItemA] != CategoryMinimumInputSortOrderMap[&ItemB])
-				{
-					return CategoryMinimumInputSortOrderMap[&ItemA] < CategoryMinimumInputSortOrderMap[&ItemB];
-				}
-
-				return ItemA.ToString().Compare(ItemB.ToString(), ESearchCase::Type::IgnoreCase) < 0;
-			}
-
-			// Categories always come after items to ensure consistency with old behavior
-			if(ItemA.IsA<UNiagaraHierarchyCategory>() && ItemB.IsA<UNiagaraHierarchyItem>())
-			{
-				return false;
-			}
-			
-			if(ItemB.IsA<UNiagaraHierarchyCategory>() && ItemA.IsA<UNiagaraHierarchyItem>())
-			{
-				return true;
-			}
-			
-			if(ItemA.IsA<UNiagaraHierarchyScriptParameter>() && ItemB.IsA<UNiagaraHierarchyScriptParameter>())
-			{
-				UNiagaraHierarchyScriptParameter* ScriptParameterA = Cast<UNiagaraHierarchyScriptParameter>(&ItemA);
-				UNiagaraHierarchyScriptParameter* ScriptParameterB = Cast<UNiagaraHierarchyScriptParameter>(&ItemB);
-
-				const FNiagaraVariableMetaData& MetaDataA = ScriptParameterA->GetScriptVariable()->Metadata;
-				const FNiagaraVariableMetaData& MetaDataB = ScriptParameterB->GetScriptVariable()->Metadata;
-
-				if(MetaDataA.bAdvancedDisplay && !MetaDataB.bAdvancedDisplay)
-				{
-					return false;
-				}
-				else if(!MetaDataA.bAdvancedDisplay && MetaDataB.bAdvancedDisplay)
-				{
-					return true;
-				}
-			
-				if(MetaDataA.GetEditorSortPriority_DEPRECATED() != MetaDataB.GetEditorSortPriority_DEPRECATED())
-				{
-					return MetaDataA.GetEditorSortPriority_DEPRECATED() < MetaDataB.GetEditorSortPriority_DEPRECATED();
-				}
-			
-				return ItemA.ToString().Compare(ItemB.ToString(), ESearchCase::Type::IgnoreCase) < 0;
-			}
-			
-			return false;
-		});
-
-		ParameterRoot->SortChildren(SortRootChildrenPredicate, false);
-	}
-}
-
 FGuid UNiagaraScriptSource::GetChangeID() 
 { 
 	return NodeGraph->GetChangeID(); 
diff --git a/Engine/Plugins/FX/Niagara/Source/NiagaraEditor/Private/ViewModels/Stack/NiagaraStackEntry.cpp b/Engine/Plugins/FX/Niagara/Source/NiagaraEditor/Private/ViewModels/Stack/NiagaraStackEntry.cpp
--- a/Engine/Plugins/FX/Niagara/Source/NiagaraEditor/Private/ViewModels/Stack/NiagaraStackEntry.cpp
+++ b/Engine/Plugins/FX/Niagara/Source/NiagaraEditor/Private/ViewModels/Stack/NiagaraStackEntry.cpp
@@ -777,9 +777,19 @@ void UNiagaraStackEntry::RefreshChildren()
 		}
 	}
 
-	Children.Empty();	
-	Children.Append(NewChildren);
-	
+	Children.Empty(NewChildren.Num());
+	TSet<UNiagaraStackEntry*> HandledChildren;
+	for (UNiagaraStackEntry* NewChild : NewChildren)
+	{
+		bool bAlreadyHandled = false;
+		HandledChildren.Add(NewChild, &bAlreadyHandled);
+		if (ensureMsgf(bAlreadyHandled == false, TEXT("Duplicate child %s (%s) in entry %s (%s)."),
+			*NewChild->GetDisplayName().ToString(), *NewChild->GetClass()->GetName(), *GetDisplayName().ToString(), *GetClass()->GetName()))
+		{
+			Children.Add(NewChild);
+		}
+	}
+
 	for (UNiagaraStackEntry* Child : Children)
 	{
 		UNiagaraStackEntry* OuterOwner = Cast<UNiagaraStackEntry>(Child->GetOuter());
diff --git a/Engine/Plugins/FX/Niagara/Source/NiagaraEditor/Private/ViewModels/Stack/NiagaraStackScriptHierarchyRoot.cpp b/Engine/Plugins/FX/Niagara/Source/NiagaraEditor/Private/ViewModels/Stack/NiagaraStackScriptHierarchyRoot.cpp
--- a/Engine/Plugins/FX/Niagara/Source/NiagaraEditor/Private/ViewModels/Stack/NiagaraStackScriptHierarchyRoot.cpp
+++ b/Engine/Plugins/FX/Niagara/Source/NiagaraEditor/Private/ViewModels/Stack/NiagaraStackScriptHierarchyRoot.cpp
@@ -224,7 +224,7 @@ void UNiagaraStackScriptHierarchyRoot::RefreshChildrenInternal(const TArray<UNia
 			FGuid VariableGuid = HierarchyParameter->GetScriptVariable()->Metadata.GetVariableGuid();
 			InputChild->SetIsHidden(ScriptInstanceData.PerInputInstanceData[VariableGuid].bIsHidden);
 			
-			NewChildren.Add(InputChild);
+			NewChildren.AddUnique(InputChild);
 		}
 		
 		if(const UNiagaraHierarchyCategory* HierarchyCategory = Cast<UNiagaraHierarchyCategory>(ChildHierarchyItem))
@@ -246,7 +246,7 @@ void UNiagaraStackScriptHierarchyRoot::RefreshChildrenInternal(const TArray<UNia
 			
 			StackCategory->SetScriptInstanceData(ScriptInstanceData);
 
-			NewChildren.Add(StackCategory);
+			NewChildren.AddUnique(StackCategory);
 		}		
 	}
 
@@ -277,7 +277,7 @@ void UNiagaraStackScriptHierarchyRoot::RefreshChildrenInternal(const TArray<UNia
 		{
 			InputChild->SetIsHidden(ScriptInstanceData.PerInputInstanceData[VariableGuid].bIsHidden);
 		}
-		NewChildren.Add(InputChild);
+		NewChildren.AddUnique(InputChild);
 	}
 }
 
diff --git a/Engine/Plugins/FX/Niagara/Source/NiagaraEditor/Public/NiagaraGraph.h b/Engine/Plugins/FX/Niagara/Source/NiagaraEditor/Public/NiagaraGraph.h
--- a/Engine/Plugins/FX/Niagara/Source/NiagaraEditor/Public/NiagaraGraph.h
+++ b/Engine/Plugins/FX/Niagara/Source/NiagaraEditor/Public/NiagaraGraph.h
@@ -461,6 +461,9 @@ public:
 
 	bool ReferencesStaticVariable(FNiagaraStaticVariableSearchContext& SearchContext) const;
 
+	/** Migrates the deprecated data from the provided OwnerData and populates the ParameterHierarchyRoot */
+	void MigrateParameterScriptDataToHierarchyRoot(FVersionedNiagaraScriptData& OwnerData);
+
 protected:
 	void RebuildNumericCache();
 	bool bNeedNumericCacheRebuilt;
diff --git a/Engine/Plugins/FX/Niagara/Source/NiagaraEditor/Public/NiagaraScriptSource.h b/Engine/Plugins/FX/Niagara/Source/NiagaraEditor/Public/NiagaraScriptSource.h
--- a/Engine/Plugins/FX/Niagara/Source/NiagaraEditor/Public/NiagaraScriptSource.h
+++ b/Engine/Plugins/FX/Niagara/Source/NiagaraEditor/Public/NiagaraScriptSource.h
@@ -60,7 +60,7 @@ public:
 
 	NIAGARAEDITOR_API virtual bool AddModuleIfMissing(FString ModulePath, ENiagaraScriptUsage Usage, bool& bOutFoundModule)override;
 
-	virtual void PostLoadFromOwner(FVersionedNiagaraScriptData& OwnerData) override;
+	virtual void MigrateParameterDataToHierarchyRoot(FVersionedNiagaraScriptData& OwnerData) override;
 
 	void FixupRenamedParameters(UNiagaraNodeFunctionCall* FunctionCallNode, TConstArrayView<FNiagaraVariable> ModuleInputVariables, FNiagaraParameterStore& RapidIterationParameters, const TArray<FNiagaraVariable>& OldRapidIterationVariables, const FVersionedNiagaraEmitter& Emitter, ENiagaraScriptUsage ScriptUsage) const;
 	void InitializeNewParameters(UNiagaraNodeFunctionCall* FunctionCallNode, TConstArrayView<FNiagaraVariable> ModuleInputVariables, FNiagaraParameterStore& RapidIterationParameters, const FVersionedNiagaraEmitter& VersionedEmitter, ENiagaraScriptUsage ScriptUsage, TSet<FNiagaraVariableBase>& ValidRapidIterationParameters) const;
@@ -98,6 +98,4 @@ public:
 private:
 	void OnGraphChanged(const FEdGraphEditAction &Action);
 	void OnGraphDataInterfaceChanged();
-
-	void PostLoadTransferScriptHierarchyData(FVersionedNiagaraScriptData& OwnerData);
 };
