Fix platform set DI.
GPU proxy was grabbing the value too early meaning we were caching the default value rather than actual loaded value of the asset.

diff --git a/Engine/Plugins/FX/Niagara/Source/Niagara/Classes/NiagaraDataInterfacePlatformSet.h b/Engine/Plugins/FX/Niagara/Source/Niagara/Classes/NiagaraDataInterfacePlatformSet.h
--- a/Engine/Plugins/FX/Niagara/Source/Niagara/Classes/NiagaraDataInterfacePlatformSet.h
+++ b/Engine/Plugins/FX/Niagara/Source/Niagara/Classes/NiagaraDataInterfacePlatformSet.h
@@ -20,7 +20,12 @@ public:
 
 	//UObject Interface
 	NIAGARA_API virtual void PostInitProperties() override;
-	//UObject Interface End
+	virtual void PostLoad()override;
+	
+#if WITH_EDITOR
+	virtual void PostEditChangeProperty(FPropertyChangedEvent& PropertyChangedEvent) override;
+#endif
+//UObject Interface End
 
 	// UNiagaraDataInterface interface
 	NIAGARA_API virtual void GetVMExternalFunction(const FVMExternalFunctionBindingInfo& BindingInfo, void* InstanceData, FVMExternalFunction& OutFunc) override;
diff --git a/Engine/Plugins/FX/Niagara/Source/Niagara/Private/NiagaraDataInterfacePlatformSet.cpp b/Engine/Plugins/FX/Niagara/Source/Niagara/Private/NiagaraDataInterfacePlatformSet.cpp
--- a/Engine/Plugins/FX/Niagara/Source/Niagara/Private/NiagaraDataInterfacePlatformSet.cpp
+++ b/Engine/Plugins/FX/Niagara/Source/Niagara/Private/NiagaraDataInterfacePlatformSet.cpp
@@ -57,12 +57,35 @@ void UNiagaraDataInterfacePlatformSet::PostInitProperties()
 		ENiagaraTypeRegistryFlags Flags = ENiagaraTypeRegistryFlags::AllowAnyVariable | ENiagaraTypeRegistryFlags::AllowParameter;
 		FNiagaraTypeRegistry::Register(FNiagaraTypeDefinition(GetClass()), Flags);
 	}
-	else
+}
+
+void UNiagaraDataInterfacePlatformSet::PostLoad()
+{
+	using namespace NDIPlatformSetPrivate;
+	Super::PostLoad();
+	if (!HasAnyFlags(RF_ClassDefaultObject))
 	{
 		GetProxyAs<FNDIPlatformSetProxy>()->bIsActive = Platforms.IsActive();
 	}
 }
 
+#if WITH_EDITOR
+void UNiagaraDataInterfacePlatformSet::PostEditChangeProperty(FPropertyChangedEvent& PropertyChangedEvent)
+{
+	Super::PostEditChangeProperty(PropertyChangedEvent);
+
+	using namespace NDIPlatformSetPrivate;
+	ENQUEUE_RENDER_COMMAND(UpdateProxyState)
+		(
+			[RT_Proxy = GetProxyAs<FNDIPlatformSetProxy>(), bActive = Platforms.IsActive()](FRHICommandListImmediate& CmdList)
+			{
+				RT_Proxy->bIsActive = bActive;
+			}
+		);
+}
+
+#endif
+
 #if WITH_EDITORONLY_DATA
 void UNiagaraDataInterfacePlatformSet::GetFunctionsInternal(TArray<FNiagaraFunctionSignature>& OutFunctions) const
 {
