- Export parts of stateless to allow for modules / templates to be built in plugins
- Expose template selection and remove alternative template option

#rnx

diff --git a/Engine/Plugins/FX/Niagara/Source/Niagara/Internal/Stateless/Modules/NiagaraStatelessModule_CalculateAccurateVelocity.h b/Engine/Plugins/FX/Niagara/Source/Niagara/Internal/Stateless/Modules/NiagaraStatelessModule_CalculateAccurateVelocity.h
--- a/Engine/Plugins/FX/Niagara/Source/Niagara/Internal/Stateless/Modules/NiagaraStatelessModule_CalculateAccurateVelocity.h
+++ b/Engine/Plugins/FX/Niagara/Source/Niagara/Internal/Stateless/Modules/NiagaraStatelessModule_CalculateAccurateVelocity.h
@@ -7,7 +7,7 @@
 
 #include "NiagaraStatelessModule_CalculateAccurateVelocity.generated.h"
 
-UCLASS(MinimalAPI, EditInlineNew, meta = (DisplayName = "Rotate Around Point"))
+UCLASS(MinimalAPI, EditInlineNew, meta = (DisplayName = "Calculate Accurate Velocity"))
 class UNiagaraStatelessModule_CalculateAccurateVelocity : public UNiagaraStatelessModule
 {
 	GENERATED_BODY()
diff --git a/Engine/Plugins/FX/Niagara/Source/Niagara/Internal/Stateless/Modules/NiagaraStatelessModule_RotateAroundPoint.h b/Engine/Plugins/FX/Niagara/Source/Niagara/Internal/Stateless/Modules/NiagaraStatelessModule_RotateAroundPoint.h
--- a/Engine/Plugins/FX/Niagara/Source/Niagara/Internal/Stateless/Modules/NiagaraStatelessModule_RotateAroundPoint.h
+++ b/Engine/Plugins/FX/Niagara/Source/Niagara/Internal/Stateless/Modules/NiagaraStatelessModule_RotateAroundPoint.h
@@ -3,6 +3,7 @@
 #pragma once
 
 #include "Stateless/NiagaraStatelessModule.h"
+#include "Stateless/NiagaraStatelessEmitterDataBuildContext.h"
 #include "Stateless/NiagaraStatelessModuleShaderParameters.h"
 
 #include "NiagaraStatelessModule_RotateAroundPoint.generated.h"
@@ -12,27 +13,41 @@ class UNiagaraStatelessModule_RotateAroundPoint : public UNiagaraStatelessModule
 {
 	GENERATED_BODY()
 
+	struct FModuleBuiltData
+	{
+		FNiagaraStatelessRangeFloat	Rate			= FNiagaraStatelessRangeFloat(0.0f);
+		FNiagaraStatelessRangeFloat	Radius			= FNiagaraStatelessRangeFloat(0.0f);
+		FNiagaraStatelessRangeFloat	InitialPhase	= FNiagaraStatelessRangeFloat(0.0f);
+	};
+
 public:
 	using FParameters = NiagaraStateless::FRotateAroundPointModule_ShaderParameters;
 
-	UPROPERTY(EditAnywhere, Category = "Parameters", meta=(Units="deg/s"))
-	float RateMin = 360.f;
-	UPROPERTY(EditAnywhere, Category = "Parameters", meta=(Units="deg/s"))
-	float RateMax = 360.f;
+	UPROPERTY(EditAnywhere, Category = "Parameters", meta=(DisableBindingDistribution, Units="deg/s"))
+	FNiagaraDistributionRangeFloat Rate = FNiagaraDistributionRangeFloat(360.0f);
 
-	UPROPERTY(EditAnywhere, Category = "Parameters")
-	float RadiusMin = 50.f;
-	UPROPERTY(EditAnywhere, Category = "Parameters")
-	float RadiusMax = 50.f;
+	UPROPERTY(EditAnywhere, Category = "Parameters", meta = (DisableBindingDistribution))
+	FNiagaraDistributionRangeFloat Radius = FNiagaraDistributionRangeFloat(100.0f);
 
-	UPROPERTY(EditAnywhere, Category = "Parameters")
-	float InitialPhaseMin = 0.f;
-	UPROPERTY(EditAnywhere, Category = "Parameters")
-	float InitialPhaseMax = 0.f;
+	UPROPERTY(EditAnywhere, Category = "Parameters", meta = (DisableBindingDistribution))
+	FNiagaraDistributionRangeFloat InitialPhase = FNiagaraDistributionRangeFloat(0.0f);
 
 	//-TODO: Add support for GPU once we settle on a feature set for this module
 	virtual ENiagaraStatelessFeatureMask GetFeatureMask() const override { return ENiagaraStatelessFeatureMask::ExecuteGPU; }
 
+	virtual void BuildEmitterData(const FNiagaraStatelessEmitterDataBuildContext& BuildContext) const override
+	{
+		FModuleBuiltData* BuiltData = BuildContext.AllocateBuiltData<FModuleBuiltData>();
+		if (!IsModuleEnabled())
+		{
+			return;
+		}
+
+		BuiltData->Rate			= Rate.CalculateRange();
+		BuiltData->Radius		= Radius.CalculateRange();
+		BuiltData->InitialPhase	= InitialPhase.CalculateRange();
+	}
+
 	virtual void BuildShaderParameters(FNiagaraStatelessShaderParametersBuilder& ShaderParametersBuilder) const override
 	{
 		ShaderParametersBuilder.AddParameterNestedStruct<FParameters>();
@@ -41,24 +56,14 @@ public:
 	virtual void SetShaderParameters(const FNiagaraStatelessSetShaderParameterContext& SetShaderParameterContext) const override
 	{
 		FParameters* Parameters = SetShaderParameterContext.GetParameterNestedStruct<FParameters>();
-		if (IsModuleEnabled())
-		{
-			Parameters->RotateAroundPoint_RateScale			= FMath::DegreesToRadians(RateMax - RateMin);
-			Parameters->RotateAroundPoint_RateBias			= FMath::DegreesToRadians(RateMin);
-			Parameters->RotateAroundPoint_RadiusScale		= RadiusMax - RadiusMin;
-			Parameters->RotateAroundPoint_RadiusBias		= RadiusMin;
-			Parameters->RotateAroundPoint_InitialPhaseScale	= InitialPhaseMax - InitialPhaseMin;
-			Parameters->RotateAroundPoint_InitialPhaseBias	= InitialPhaseMin;
-		}
-		else
-		{
-			Parameters->RotateAroundPoint_RateScale			= 0.0f;
-			Parameters->RotateAroundPoint_RateBias			= 0.0f;
-			Parameters->RotateAroundPoint_RadiusScale		= 0.0f;
-			Parameters->RotateAroundPoint_RadiusBias		= 0.0f;
-			Parameters->RotateAroundPoint_InitialPhaseScale	= 0.0f;
-			Parameters->RotateAroundPoint_InitialPhaseBias	= 0.0f;
-		}
+		const FModuleBuiltData* BuiltData = SetShaderParameterContext.ReadBuiltData<FModuleBuiltData>();
+
+		Parameters->RotateAroundPoint_RateScale			= FMath::DegreesToRadians(BuiltData->Rate.GetScale());
+		Parameters->RotateAroundPoint_RateBias			= FMath::DegreesToRadians(BuiltData->Rate.Min);
+		Parameters->RotateAroundPoint_RadiusScale		= BuiltData->Radius.GetScale();
+		Parameters->RotateAroundPoint_RadiusBias		= BuiltData->Radius.Min;
+		Parameters->RotateAroundPoint_InitialPhaseScale	= BuiltData->InitialPhase.GetScale();
+		Parameters->RotateAroundPoint_InitialPhaseBias	= BuiltData->InitialPhase.Min;
 	}
 
 #if WITH_EDITOR
diff --git a/Engine/Plugins/FX/Niagara/Source/Niagara/Internal/Stateless/NiagaraStatelessCommon.h b/Engine/Plugins/FX/Niagara/Source/Niagara/Internal/Stateless/NiagaraStatelessCommon.h
--- a/Engine/Plugins/FX/Niagara/Source/Niagara/Internal/Stateless/NiagaraStatelessCommon.h
+++ b/Engine/Plugins/FX/Niagara/Source/Niagara/Internal/Stateless/NiagaraStatelessCommon.h
@@ -121,7 +121,7 @@ struct FNiagaraStatelessGlobals
 	inline static FVector2f		GetDefaultSpriteSizeValue() { return FVector2f(10.0f); }
 	inline static float			GetDefaultSpriteRotationValue() { return 0.0f; }
 
-	static const FNiagaraStatelessGlobals& Get();
+	NIAGARA_API static const FNiagaraStatelessGlobals& Get();
 };
 
 namespace NiagaraStatelessCommon
diff --git a/Engine/Plugins/FX/Niagara/Source/Niagara/Internal/Stateless/NiagaraStatelessDrawDebugContext.h b/Engine/Plugins/FX/Niagara/Source/Niagara/Internal/Stateless/NiagaraStatelessDrawDebugContext.h
--- a/Engine/Plugins/FX/Niagara/Source/Niagara/Internal/Stateless/NiagaraStatelessDrawDebugContext.h
+++ b/Engine/Plugins/FX/Niagara/Source/Niagara/Internal/Stateless/NiagaraStatelessDrawDebugContext.h
@@ -11,11 +11,11 @@ struct FNiagaraStatelessDrawDebugContext
 	FTransform	LocalToWorldTransform = FTransform::Identity;
 	FTransform	WorldTransform = FTransform::Identity;
 
-	void DrawArrow(const FVector& Origin, const FVector& DirectionWithLength, const FColor& Color = FColor::White) const;
-	void DrawBox(const FVector& Center, const FVector& Extent, const FColor& Color = FColor::White, const FQuat& Rotation = FQuat::Identity) const;
-	void DrawCone(const FVector& Origin, const FQuat& Rotation, float Angle, float Length, const FColor& Color = FColor::White) const;
-	void DrawCylinder(float CylinderHeight, float CylinderRadius, float CylinderHeightMidpoint, const FColor& Color = FColor::White) const;
-	void DrawCircle(const FVector& Center, const float Radius, const FColor& Color = FColor::White) const;
-	void DrawSphere(const FVector& Center, const float Radius, const FColor& Color = FColor::White) const;
+	NIAGARA_API void DrawArrow(const FVector& Origin, const FVector& DirectionWithLength, const FColor& Color = FColor::White) const;
+	NIAGARA_API void DrawBox(const FVector& Center, const FVector& Extent, const FColor& Color = FColor::White, const FQuat& Rotation = FQuat::Identity) const;
+	NIAGARA_API void DrawCone(const FVector& Origin, const FQuat& Rotation, float Angle, float Length, const FColor& Color = FColor::White) const;
+	NIAGARA_API void DrawCylinder(float CylinderHeight, float CylinderRadius, float CylinderHeightMidpoint, const FColor& Color = FColor::White) const;
+	NIAGARA_API void DrawCircle(const FVector& Center, const float Radius, const FColor& Color = FColor::White) const;
+	NIAGARA_API void DrawSphere(const FVector& Center, const float Radius, const FColor& Color = FColor::White) const;
 };
 #endif
diff --git a/Engine/Plugins/FX/Niagara/Source/Niagara/Internal/Stateless/NiagaraStatelessEmitter.h b/Engine/Plugins/FX/Niagara/Source/Niagara/Internal/Stateless/NiagaraStatelessEmitter.h
--- a/Engine/Plugins/FX/Niagara/Source/Niagara/Internal/Stateless/NiagaraStatelessEmitter.h
+++ b/Engine/Plugins/FX/Niagara/Source/Niagara/Internal/Stateless/NiagaraStatelessEmitter.h
@@ -34,6 +34,9 @@ class UNiagaraStatelessEmitter : public UObject//, public INiagaraParameterDefin
 {
 	GENERATED_BODY()
 
+public:
+	DECLARE_MULTICAST_DELEGATE(FOnTemplateChanged);
+
 public:
 	//~Begin UObject Interface
 	virtual void PostLoad() override;
@@ -44,9 +47,14 @@ public:
 	//~End UObject Interface
 
 #if WITH_EDITOR
-	void OnEmitterTemplateChanged();
-	void OnCacheParameterCollectionReferences();
+	FOnTemplateChanged& GetOnTemplateChanged() { return OnTemplateChanged; }
+
+protected:
+	void PostTemplateChanged();
+	void CacheParameterCollectionReferences();
 #endif
+
+public:
 	NIAGARA_API const UNiagaraStatelessEmitterTemplate* GetEmitterTemplate() const;
 
 	bool UsesCollection(const UNiagaraParameterCollection* Collection) const;
@@ -133,7 +141,7 @@ protected:
 	UPROPERTY()
 	FString UniqueEmitterName;
 
-	UPROPERTY(EditAnywhere, Category = "Emitter Properties", meta = (AllowedClasses = "/Script/Niagara.NiagaraStatelessEmitterTemplate", HideInStack))
+	UPROPERTY(EditAnywhere, Category = "Emitter Properties", NoClear, meta = (DisplayName = "Template", DisallowCreateNew, AllowedClasses = "/Script/Niagara.NiagaraStatelessEmitterTemplate"))
 	TObjectPtr<UClass> EmitterTemplateClass;
 
 	UPROPERTY(EditAnywhere, Category = "Emitter Properties")
@@ -194,4 +202,7 @@ protected:
 #if WITH_EDITORONLY_DATA
 	FSimpleMulticastDelegate OnRenderersChangedDelegate;
 #endif
+#if WITH_EDITOR
+	FOnTemplateChanged OnTemplateChanged;
+#endif
 };
diff --git a/Engine/Plugins/FX/Niagara/Source/Niagara/Internal/Stateless/NiagaraStatelessEmitterTemplate.h b/Engine/Plugins/FX/Niagara/Source/Niagara/Internal/Stateless/NiagaraStatelessEmitterTemplate.h
--- a/Engine/Plugins/FX/Niagara/Source/Niagara/Internal/Stateless/NiagaraStatelessEmitterTemplate.h
+++ b/Engine/Plugins/FX/Niagara/Source/Niagara/Internal/Stateless/NiagaraStatelessEmitterTemplate.h
@@ -17,7 +17,7 @@ class UNiagaraStatelessEmitterTemplate : public UObject
 
 public:
 	//~Begin: UObject Interface
-	virtual void PostInitProperties() override;
+	NIAGARA_API virtual void PostInitProperties() override;
 	//~End: UObject Interface
 
 	static void InitCDOPropertiesAfterModuleStartup();
@@ -56,18 +56,3 @@ class UNiagaraStatelessEmitterDefault : public UNiagaraStatelessEmitterTemplate
 	virtual void SetShaderParameters(uint8* ShaderParametersBase, TConstArrayView<int32> ComponentOffsets) const override;
 	//-TODO: Add a way to set them directly, we should know that the final struct is a series of ints in the order of the provided variables
 };
-
-UCLASS(Category = "Statless", DisplayName = "Emitter Example 1")
-class UNiagaraStatelessEmitterExample1 : public UNiagaraStatelessEmitterTemplate
-{
-	GENERATED_BODY()
-
-	virtual void InitModulesAndAttributes() override;
-
-	virtual const FShaderParametersMetadata* GetShaderParametersMetadata() const override;
-	virtual TShaderRef<NiagaraStateless::FSimulationShader> GetSimulationShader() const override;
-
-	//-TODO: Add a way to set them directly, we should know that the final struct is a series of ints in the order of the provided variables
-	virtual void SetShaderParameters(uint8* ShaderParametersBase, TConstArrayView<int32> ComponentOffsets) const override;
-	//-TODO: Add a way to set them directly, we should know that the final struct is a series of ints in the order of the provided variables
-};
diff --git a/Engine/Plugins/FX/Niagara/Source/Niagara/Internal/Stateless/NiagaraStatelessModule.h b/Engine/Plugins/FX/Niagara/Source/Niagara/Internal/Stateless/NiagaraStatelessModule.h
--- a/Engine/Plugins/FX/Niagara/Source/Niagara/Internal/Stateless/NiagaraStatelessModule.h
+++ b/Engine/Plugins/FX/Niagara/Source/Niagara/Internal/Stateless/NiagaraStatelessModule.h
@@ -41,28 +41,28 @@ public:
 	};
 #endif
 
-	virtual ENiagaraStatelessFeatureMask GetFeatureMask() const { return ENiagaraStatelessFeatureMask::All; }
+	NIAGARA_API virtual ENiagaraStatelessFeatureMask GetFeatureMask() const { return ENiagaraStatelessFeatureMask::All; }
 
-	virtual void BuildEmitterData(const FNiagaraStatelessEmitterDataBuildContext& BuildContext) const {}
-	virtual void BuildShaderParameters(FNiagaraStatelessShaderParametersBuilder& ShaderParametersBuilder) const {}
-	virtual void SetShaderParameters(const FNiagaraStatelessSetShaderParameterContext& SetShaderParameterContext) const {}
+	NIAGARA_API virtual void BuildEmitterData(const FNiagaraStatelessEmitterDataBuildContext& BuildContext) const {}
+	NIAGARA_API virtual void BuildShaderParameters(FNiagaraStatelessShaderParametersBuilder& ShaderParametersBuilder) const {}
+	NIAGARA_API virtual void SetShaderParameters(const FNiagaraStatelessSetShaderParameterContext& SetShaderParameterContext) const {}
 
 #if WITH_EDITOR
-	virtual bool CanDisableModule() const { return false; }
+	NIAGARA_API virtual bool CanDisableModule() const { return false; }
 	void SetIsModuleEnabled(bool bInIsEnabled) { bModuleEnabled = bInIsEnabled; }
 
-	virtual bool CanDebugDraw() const { return false; }
-	virtual void DrawDebug(const FNiagaraStatelessDrawDebugContext& DrawDebugContext) const {}
+	NIAGARA_API virtual bool CanDebugDraw() const { return false; }
+	NIAGARA_API virtual void DrawDebug(const FNiagaraStatelessDrawDebugContext& DrawDebugContext) const {}
 #endif
 
 #if WITH_EDITORONLY_DATA
-	virtual void GetOutputVariables(TArray<FNiagaraVariableBase>& OutVariables) const {}
+	NIAGARA_API virtual void GetOutputVariables(TArray<FNiagaraVariableBase>& OutVariables) const {}
 #endif
 
 	//~UObject interface Begin
 #if WITH_EDITOR
-	virtual bool CanEditChange(const FProperty* InProperty) const override;
-	virtual void PostEditChangeProperty(struct FPropertyChangedEvent& PropertyChangedEvent) override;
+	NIAGARA_API virtual bool CanEditChange(const FProperty* InProperty) const override;
+	NIAGARA_API virtual void PostEditChangeProperty(struct FPropertyChangedEvent& PropertyChangedEvent) override;
 #endif //WITH_EDITOR
 	//~UObject interface End
 };
diff --git a/Engine/Plugins/FX/Niagara/Source/Niagara/Internal/Stateless/NiagaraStatelessSetShaderParameterContext.h b/Engine/Plugins/FX/Niagara/Source/Niagara/Internal/Stateless/NiagaraStatelessSetShaderParameterContext.h
--- a/Engine/Plugins/FX/Niagara/Source/Niagara/Internal/Stateless/NiagaraStatelessSetShaderParameterContext.h
+++ b/Engine/Plugins/FX/Niagara/Source/Niagara/Internal/Stateless/NiagaraStatelessSetShaderParameterContext.h
@@ -69,7 +69,7 @@ public:
 
 protected:
 #if DO_CHECK
-	void ValidateIncludeStructType(uint32 StructOffset, const FShaderParametersMetadata* StructMetaData) const;
+	NIAGARA_API void ValidateIncludeStructType(uint32 StructOffset, const FShaderParametersMetadata* StructMetaData) const;
 #endif
 
 private:
diff --git a/Engine/Plugins/FX/Niagara/Source/Niagara/Private/Stateless/NiagaraStatelessEmitter.cpp b/Engine/Plugins/FX/Niagara/Source/Niagara/Private/Stateless/NiagaraStatelessEmitter.cpp
--- a/Engine/Plugins/FX/Niagara/Source/Niagara/Private/Stateless/NiagaraStatelessEmitter.cpp
+++ b/Engine/Plugins/FX/Niagara/Source/Niagara/Private/Stateless/NiagaraStatelessEmitter.cpp
@@ -105,8 +105,8 @@ void UNiagaraStatelessEmitter::PostLoad()
 
 #if WITH_EDITOR
 	// Ensure our module list is up to date
-	OnEmitterTemplateChanged();
-	//OnCacheParameterCollectionReferences();
+	PostTemplateChanged();
+	//CacheParameterCollectionReferences();
 #endif
 }
 
@@ -162,15 +162,16 @@ void UNiagaraStatelessEmitter::PostEditChangeProperty(FPropertyChangedEvent& Pro
 		FNiagaraSystemUpdateContext UpdateContext(NiagaraSystem, true);
 
 		// Ensure our template is up to date, right not it's easy to remove things by accident and cause a crash
-		//if (PropertyChangedEvent.GetPropertyName() == GET_MEMBER_NAME_CHECKED(UNiagaraStatelessEmitter, EmitterTemplateClass))
+		if (PropertyChangedEvent.GetPropertyName() == GET_MEMBER_NAME_CHECKED(UNiagaraStatelessEmitter, EmitterTemplateClass))
 		{
-			OnEmitterTemplateChanged();
-			OnCacheParameterCollectionReferences();
+			PostTemplateChanged();
+			OnTemplateChanged.Broadcast();
 		}
+		CacheParameterCollectionReferences();
 	}
 }
 
-void UNiagaraStatelessEmitter::OnEmitterTemplateChanged()
+void UNiagaraStatelessEmitter::PostTemplateChanged()
 {
 	// No template nothing to do
 	const UNiagaraStatelessEmitterTemplate* EmitterTemplate = GetEmitterTemplate();
@@ -208,7 +209,7 @@ void UNiagaraStatelessEmitter::OnEmitterTemplateChanged()
 	Modules = MoveTemp(NewModules);
 }
 
-void UNiagaraStatelessEmitter::OnCacheParameterCollectionReferences()
+void UNiagaraStatelessEmitter::CacheParameterCollectionReferences()
 {
 	TArray<TObjectPtr<UNiagaraParameterCollection>> OriginalReferences;
 	Swap(OriginalReferences, CachedParameterCollectionReferences);
@@ -386,7 +387,7 @@ void UNiagaraStatelessEmitter::CacheFromCompiledData()
 	#if WITH_EDITOR
 		// UNiagaraStatelessEmitter does not get a notification when bindings change.
 		//-OPT: Shift this to be when distributions change mode in / out of binding + binding changed
-		OnCacheParameterCollectionReferences();
+		CacheParameterCollectionReferences();
 	#endif
 		StatelessEmitterData->BoundParameterCollections.Reset(CachedParameterCollectionReferences.Num());
 		for (UNiagaraParameterCollection* ParameterCollection : CachedParameterCollectionReferences)
@@ -624,7 +625,7 @@ void UNiagaraStatelessEmitter::SetEmitterTemplateClass(UClass* TemplateClass)
 	}
 
 	EmitterTemplateClass = TemplateClass;
-	OnEmitterTemplateChanged();
+	PostTemplateChanged();
 }
 
 void UNiagaraStatelessEmitter::AddRenderer(UNiagaraRendererProperties* Renderer, FGuid EmitterVersion)
diff --git a/Engine/Plugins/FX/Niagara/Source/Niagara/Private/Stateless/NiagaraStatelessEmitterTemplate.cpp b/Engine/Plugins/FX/Niagara/Source/Niagara/Private/Stateless/NiagaraStatelessEmitterTemplate.cpp
--- a/Engine/Plugins/FX/Niagara/Source/Niagara/Private/Stateless/NiagaraStatelessEmitterTemplate.cpp
+++ b/Engine/Plugins/FX/Niagara/Source/Niagara/Private/Stateless/NiagaraStatelessEmitterTemplate.cpp
@@ -182,73 +182,3 @@ void UNiagaraStatelessEmitterDefault::SetShaderParameters(uint8* ShaderParameter
 }
 //-TODO: Add a way to set them directly, we should know that the final struct is a series of ints in the order of the provided variables
 
-//////////////////////////////////////////////////////////////////////////////////////////////////////////////////////////////////////////////////////////
-
-void UNiagaraStatelessEmitterExample1::InitModulesAndAttributes()
-{
-#if WITH_EDITORONLY_DATA
-	Modules =
-	{
-		UNiagaraStatelessModule_InitializeParticle::StaticClass(),
-		UNiagaraStatelessModule_ShapeLocation::StaticClass(),
-		UNiagaraStatelessModule_ScaleColor::StaticClass(),
-		UNiagaraStatelessModule_ScaleSpriteSize::StaticClass(),
-		UNiagaraStatelessModule_RotateAroundPoint::StaticClass(),
-		UNiagaraStatelessModule_CalculateAccurateVelocity::StaticClass(),
-	};
-
-	const FNiagaraStatelessGlobals& StatelessGlobals = FNiagaraStatelessGlobals::Get();
-	OutputComponents =
-	{
-		StatelessGlobals.UniqueIDVariable,
-		StatelessGlobals.PositionVariable,
-		StatelessGlobals.ColorVariable,
-		StatelessGlobals.RibbonWidthVariable,
-		StatelessGlobals.ScaleVariable,
-		StatelessGlobals.SpriteSizeVariable,
-		StatelessGlobals.SpriteRotationVariable,
-		StatelessGlobals.VelocityVariable,
-		StatelessGlobals.PreviousPositionVariable,
-		StatelessGlobals.PreviousRibbonWidthVariable,
-		StatelessGlobals.PreviousScaleVariable,
-		StatelessGlobals.PreviousSpriteSizeVariable,
-		StatelessGlobals.PreviousSpriteRotationVariable,
-		StatelessGlobals.PreviousVelocityVariable,
-	};
-#endif //WITH_EDITORONLY_DATA
-}
-
-const FShaderParametersMetadata* UNiagaraStatelessEmitterExample1::GetShaderParametersMetadata() const
-{
-	using namespace NiagaraStateless;
-	return FSimulationShaderExample1CS::FParameters::FTypeInfo::GetStructMetadata();
-}
-
-TShaderRef<NiagaraStateless::FSimulationShader> UNiagaraStatelessEmitterExample1::GetSimulationShader() const
-{
-	using namespace NiagaraStateless;
-	return TShaderMapRef<FSimulationShaderExample1CS>(GetGlobalShaderMap(GMaxRHIFeatureLevel));
-}
-
-//-TODO: Add a way to set them directly, we should know that the final struct is a series of ints in the order of the provided variables
-void UNiagaraStatelessEmitterExample1::SetShaderParameters(uint8* ShaderParametersBase, TConstArrayView<int32> ComponentOffsets) const
-{
-	using namespace NiagaraStateless;
-	FSimulationShaderExample1CS::FParameters* ShaderParameters		= reinterpret_cast<FSimulationShaderExample1CS::FParameters*>(ShaderParametersBase);
-	int iComponent = 0;
-	ShaderParameters->Permutation_UniqueIDComponent					= ComponentOffsets[iComponent++];
-	ShaderParameters->Permutation_PositionComponent					= ComponentOffsets[iComponent++];
-	ShaderParameters->Permutation_ColorComponent					= ComponentOffsets[iComponent++];
-	ShaderParameters->Permutation_RibbonWidthComponent				= ComponentOffsets[iComponent++];
-	ShaderParameters->Permutation_ScaleComponent					= ComponentOffsets[iComponent++];
-	ShaderParameters->Permutation_SpriteSizeComponent				= ComponentOffsets[iComponent++];
-	ShaderParameters->Permutation_SpriteRotationComponent			= ComponentOffsets[iComponent++];
-	ShaderParameters->Permutation_VelocityComponent					= ComponentOffsets[iComponent++];
-	ShaderParameters->Permutation_PreviousPositionComponent			= ComponentOffsets[iComponent++];
-	ShaderParameters->Permutation_PreviousRibbonWidthComponent		= ComponentOffsets[iComponent++];
-	ShaderParameters->Permutation_PreviousScaleComponent			= ComponentOffsets[iComponent++];
-	ShaderParameters->Permutation_PreviousSpriteSizeComponent		= ComponentOffsets[iComponent++];
-	ShaderParameters->Permutation_PreviousSpriteRotationComponent	= ComponentOffsets[iComponent++];
-	ShaderParameters->Permutation_PreviousVelocityComponent			= ComponentOffsets[iComponent++];
-}
-//-TODO: Add a way to set them directly, we should know that the final struct is a series of ints in the order of the provided variables
diff --git a/Engine/Plugins/FX/Niagara/Source/NiagaraEditor/Private/ViewModels/Stack/NiagaraStackStatelessEmitterSimulateGroup.cpp b/Engine/Plugins/FX/Niagara/Source/NiagaraEditor/Private/ViewModels/Stack/NiagaraStackStatelessEmitterSimulateGroup.cpp
--- a/Engine/Plugins/FX/Niagara/Source/NiagaraEditor/Private/ViewModels/Stack/NiagaraStackStatelessEmitterSimulateGroup.cpp
+++ b/Engine/Plugins/FX/Niagara/Source/NiagaraEditor/Private/ViewModels/Stack/NiagaraStackStatelessEmitterSimulateGroup.cpp
@@ -163,6 +163,8 @@ void UNiagaraStackStatelessEmitterSimulateGroup::Initialize(FRequiredEntryData I
 		LOCTEXT("EmitterStatelessSimulateGroupToolTip", "Data related to the simulation of the particles"),
 		AddUtilities.Get());
 	StatelessEmitterWeak = InStatelessEmitter;
+
+	InStatelessEmitter->GetOnTemplateChanged().AddUObject(this, &UNiagaraStackStatelessEmitterSimulateGroup::OnTemplateChanged);
 }
 
 const FSlateBrush* UNiagaraStackStatelessEmitterSimulateGroup::GetIconBrush() const
@@ -190,6 +192,11 @@ void UNiagaraStackStatelessEmitterSimulateGroup::Paste(const UNiagaraClipboardCo
 	}
 }
 
+void UNiagaraStackStatelessEmitterSimulateGroup::OnTemplateChanged()
+{
+	RefreshChildren();
+}
+
 void UNiagaraStackStatelessEmitterSimulateGroup::RefreshChildrenInternal(const TArray<UNiagaraStackEntry*>& CurrentChildren, TArray<UNiagaraStackEntry*>& NewChildren, TArray<FStackIssue>& NewIssues)
 {
 	Super::RefreshChildrenInternal(CurrentChildren, NewChildren, NewIssues);
diff --git a/Engine/Plugins/FX/Niagara/Source/NiagaraEditor/Public/ViewModels/Stack/NiagaraStackStatelessEmitterSimulateGroup.h b/Engine/Plugins/FX/Niagara/Source/NiagaraEditor/Public/ViewModels/Stack/NiagaraStackStatelessEmitterSimulateGroup.h
--- a/Engine/Plugins/FX/Niagara/Source/NiagaraEditor/Public/ViewModels/Stack/NiagaraStackStatelessEmitterSimulateGroup.h
+++ b/Engine/Plugins/FX/Niagara/Source/NiagaraEditor/Public/ViewModels/Stack/NiagaraStackStatelessEmitterSimulateGroup.h
@@ -35,6 +35,7 @@ public:
 	virtual void Paste(const UNiagaraClipboardContent* ClipboardContent, FText& OutPasteWarning) override;
 
 protected:
+	void OnTemplateChanged();
 	virtual void RefreshChildrenInternal(const TArray<UNiagaraStackEntry*>& CurrentChildren, TArray<UNiagaraStackEntry*>& NewChildren, TArray<FStackIssue>& NewIssues) override;
 
 private:
diff --git a/Engine/Plugins/FX/Niagara/Source/NiagaraShader/Internal/Stateless/NiagaraStatelessEmitterShaders.h b/Engine/Plugins/FX/Niagara/Source/NiagaraShader/Internal/Stateless/NiagaraStatelessEmitterShaders.h
--- a/Engine/Plugins/FX/Niagara/Source/NiagaraShader/Internal/Stateless/NiagaraStatelessEmitterShaders.h
+++ b/Engine/Plugins/FX/Niagara/Source/NiagaraShader/Internal/Stateless/NiagaraStatelessEmitterShaders.h
@@ -75,38 +75,6 @@ namespace NiagaraStateless
 		END_SHADER_PARAMETER_STRUCT()
 	};
 
-	class FSimulationShaderExample1CS : public FSimulationShader
-	{
-	public:
-		DECLARE_EXPORTED_GLOBAL_SHADER(FSimulationShaderExample1CS, NIAGARASHADER_API);
-		SHADER_USE_PARAMETER_STRUCT(FSimulationShaderExample1CS, FSimulationShader);
-
-		BEGIN_SHADER_PARAMETER_STRUCT(FParameters, NIAGARASHADER_API)
-			SHADER_PARAMETER_STRUCT_INCLUDE(FCommonShaderParameters, CommonParameters)
-			ADD_STATELESS_MODULE(FInitializeParticleModule_ShaderParameters)
-			ADD_STATELESS_MODULE(FShapeLocationModule_ShaderParameters)
-			ADD_STATELESS_MODULE(FScaleColorModule_ShaderParameters)
-			ADD_STATELESS_MODULE(FScaleSpriteSizeModule_ShaderParameters)
-			ADD_STATELESS_MODULE(FRotateAroundPointModule_ShaderParameters)
-			ADD_STATELESS_MODULE(FSolveVelocitiesAndForcesModule_ShaderParameters)
-
-			SHADER_PARAMETER(int, Permutation_UniqueIDComponent)
-			SHADER_PARAMETER(int, Permutation_PositionComponent)
-			SHADER_PARAMETER(int, Permutation_ColorComponent)
-			SHADER_PARAMETER(int, Permutation_ScaleComponent)
-			SHADER_PARAMETER(int, Permutation_RibbonWidthComponent)
-			SHADER_PARAMETER(int, Permutation_SpriteSizeComponent)
-			SHADER_PARAMETER(int, Permutation_SpriteRotationComponent)
-			SHADER_PARAMETER(int, Permutation_VelocityComponent)
-			SHADER_PARAMETER(int, Permutation_PreviousPositionComponent)
-			SHADER_PARAMETER(int, Permutation_PreviousRibbonWidthComponent)
-			SHADER_PARAMETER(int, Permutation_PreviousScaleComponent)
-			SHADER_PARAMETER(int, Permutation_PreviousSpriteSizeComponent)
-			SHADER_PARAMETER(int, Permutation_PreviousSpriteRotationComponent)
-			SHADER_PARAMETER(int, Permutation_PreviousVelocityComponent)
-		END_SHADER_PARAMETER_STRUCT()
-	};
-
 	#undef ADD_STATELESS_MODULE
 	#undef STATELESS_MODULE_COMBINE
 	#undef STATELESS_MODULE_COMBINE2
diff --git a/Engine/Plugins/FX/Niagara/Source/NiagaraShader/Private/Stateless/NiagaraStatelessEmitterShaders.cpp b/Engine/Plugins/FX/Niagara/Source/NiagaraShader/Private/Stateless/NiagaraStatelessEmitterShaders.cpp
--- a/Engine/Plugins/FX/Niagara/Source/NiagaraShader/Private/Stateless/NiagaraStatelessEmitterShaders.cpp
+++ b/Engine/Plugins/FX/Niagara/Source/NiagaraShader/Private/Stateless/NiagaraStatelessEmitterShaders.cpp
@@ -5,5 +5,4 @@
 namespace NiagaraStateless
 {
 	IMPLEMENT_GLOBAL_SHADER(FSimulationShaderDefaultCS, "/Plugin/FX/Niagara/Private/Stateless/NiagaraStatelessSimulationDefault.usf", "StatelessMain", SF_Compute);
-	IMPLEMENT_GLOBAL_SHADER(FSimulationShaderExample1CS, "/Plugin/FX/Niagara/Private/Stateless/NiagaraStatelessSimulationExample1.usf", "StatelessMain", SF_Compute);
 }
