Made Module Wizard modal and fixed stack refresh issue that was happening due to invalidation of previous stack entries when RefreshAll was called.

#tests Tested adding Spawn/Read from/Write to Data Channel actions

diff --git a/Engine/Plugins/FX/Niagara/Source/NiagaraEditor/Private/ViewModels/Stack/NiagaraStackScriptItemGroup.cpp b/Engine/Plugins/FX/Niagara/Source/NiagaraEditor/Private/ViewModels/Stack/NiagaraStackScriptItemGroup.cpp
--- a/Engine/Plugins/FX/Niagara/Source/NiagaraEditor/Private/ViewModels/Stack/NiagaraStackScriptItemGroup.cpp
+++ b/Engine/Plugins/FX/Niagara/Source/NiagaraEditor/Private/ViewModels/Stack/NiagaraStackScriptItemGroup.cpp
@@ -32,6 +32,7 @@
 #include "ViewModels/Stack/NiagaraStackModuleItem.h"
 #include "Widgets/Notifications/SNotificationList.h"
 #include "Toolkits/SystemToolkitModes/NiagaraSystemToolkitModeBase.h"
+#include "ViewModels/Stack/NiagaraStackViewModel.h"
 #include "Widgets/DataChannel/NiagaraDataChannelWizard.h"
 #include "Widgets/Wizard/SNiagaraModuleWizard.h"
 
@@ -361,7 +362,18 @@ public:
 					{
 						SystemModel->GetSelectionViewModel()->AddEntryToSelectionByDisplayedObjectDeferred(ModuleNode);
 					}
-					SystemModel->RefreshAll();
+
+					if(EmitterViewModel.IsValid())
+					{
+						if(TSharedPtr<FNiagaraEmitterHandleViewModel> EmitterHandleViewModel = SystemModel->GetEmitterHandleViewModelForEmitter(EmitterViewModel.Pin()->GetEmitter()))
+						{
+							EmitterHandleViewModel->GetEmitterStackViewModel()->RequestRefreshDeferred();
+						}
+					}
+					else
+					{
+						SystemModel->GetSystemStackViewModel()->RequestRefreshDeferred();
+					}
 				}
 			});
 
@@ -379,14 +391,8 @@ public:
 
 			const IMainFrameModule& MainFrameModule = FModuleManager::LoadModuleChecked<IMainFrameModule>(TEXT("MainFrame"));
 			const TSharedPtr<SWindow> ParentWindow = MainFrameModule.GetParentWindow();
-			if (ParentWindow.IsValid())
-			{
-				FSlateApplication::Get().AddWindowAsNativeChild(AddModuleWindow, ParentWindow.ToSharedRef());
-			}
-			else
-			{
-				FSlateApplication::Get().AddWindow(AddModuleWindow);
-			}
+			FSlateApplication::Get().AddModalWindow(AddModuleWindow, ParentWindow);
+			return;
 		}
 		else if (ScriptGroupAddAction->IsNewScratchModuleAction())
 		{
