Fixed drag & drop indicators on the overview node not appearing. This was due to the invalidation panel wrapping around the node. During drag & drop we force the invalidation panel to redraw

#tests Dragged & Dropped modules to see if it drag & drop indicators behave correctly

diff --git a/Engine/Plugins/FX/Niagara/Source/NiagaraEditorWidgets/Private/SNiagaraOverviewStackNode.cpp b/Engine/Plugins/FX/Niagara/Source/NiagaraEditorWidgets/Private/SNiagaraOverviewStackNode.cpp
--- a/Engine/Plugins/FX/Niagara/Source/NiagaraEditorWidgets/Private/SNiagaraOverviewStackNode.cpp
+++ b/Engine/Plugins/FX/Niagara/Source/NiagaraEditorWidgets/Private/SNiagaraOverviewStackNode.cpp
@@ -47,6 +47,7 @@
 #include "Widgets/Layout/SScaleBox.h"
 #include "Widgets/Layout/SSpacer.h"
 #include "AssetThumbnail.h"
+#include "Framework/Application/SlateApplication.h"
 
 #define LOCTEXT_NAMESPACE "NiagaraOverviewStackNode"
 
@@ -414,6 +415,12 @@ void SNiagaraOverviewStackNode::Tick(const FGeometry& AllottedGeometry, const do
 {
 	if (OverviewStackNode != nullptr)
 	{
+		// To support drag & drop indicators, we mark the prepass as dirty. This will cause the invalidation widget to render using the slow path, updating the drag & drop indicators properly
+		if(FSlateApplication::Get().IsDragDropping())
+		{
+			ContentAreaWidget->MarkPrepassAsDirty();
+		}
+		
 		if (OverviewStackNode->IsRenamePending() && !SGraphNode::IsRenamePending())
 		{
 			SGraphNode::RequestRename();
