Prevent NDC DI user parameters.

diff --git a/Engine/Plugins/FX/Niagara/Source/Niagara/Private/DataInterface/NiagaraDataInterfaceDataChannelRead.cpp b/Engine/Plugins/FX/Niagara/Source/Niagara/Private/DataInterface/NiagaraDataInterfaceDataChannelRead.cpp
--- a/Engine/Plugins/FX/Niagara/Source/Niagara/Private/DataInterface/NiagaraDataInterfaceDataChannelRead.cpp
+++ b/Engine/Plugins/FX/Niagara/Source/Niagara/Private/DataInterface/NiagaraDataInterfaceDataChannelRead.cpp
@@ -597,7 +597,7 @@ void UNiagaraDataInterfaceDataChannelRead::PostInitProperties()
 
 	if (HasAnyFlags(RF_ClassDefaultObject) && INiagaraModule::DataChannelsEnabled())
 	{
-		ENiagaraTypeRegistryFlags Flags = ENiagaraTypeRegistryFlags::AllowAnyVariable | ENiagaraTypeRegistryFlags::AllowParameter;
+		ENiagaraTypeRegistryFlags Flags = ENiagaraTypeRegistryFlags::AllowNotUserVariable | ENiagaraTypeRegistryFlags::AllowParameter;
 		FNiagaraTypeRegistry::Register(FNiagaraTypeDefinition(GetClass()), Flags);
 	}
 }
diff --git a/Engine/Plugins/FX/Niagara/Source/Niagara/Private/DataInterface/NiagaraDataInterfaceDataChannelWrite.cpp b/Engine/Plugins/FX/Niagara/Source/Niagara/Private/DataInterface/NiagaraDataInterfaceDataChannelWrite.cpp
--- a/Engine/Plugins/FX/Niagara/Source/Niagara/Private/DataInterface/NiagaraDataInterfaceDataChannelWrite.cpp
+++ b/Engine/Plugins/FX/Niagara/Source/Niagara/Private/DataInterface/NiagaraDataInterfaceDataChannelWrite.cpp
@@ -503,7 +503,7 @@ void UNiagaraDataInterfaceDataChannelWrite::PostInitProperties()
 
 	if (HasAnyFlags(RF_ClassDefaultObject) && INiagaraModule::DataChannelsEnabled())
 	{
-		ENiagaraTypeRegistryFlags Flags = ENiagaraTypeRegistryFlags::AllowAnyVariable | ENiagaraTypeRegistryFlags::AllowParameter;
+		ENiagaraTypeRegistryFlags Flags = ENiagaraTypeRegistryFlags::AllowNotUserVariable | ENiagaraTypeRegistryFlags::AllowParameter;
 		FNiagaraTypeRegistry::Register(FNiagaraTypeDefinition(GetClass()), Flags);
 	}
 }
