Clear current frame data when we early out due to no new data. Prevents lingering data in the NDC on GPU sims.

diff --git a/Engine/Plugins/FX/Niagara/Source/Niagara/Private/NiagaraDataChannel.cpp b/Engine/Plugins/FX/Niagara/Source/Niagara/Private/NiagaraDataChannel.cpp
--- a/Engine/Plugins/FX/Niagara/Source/Niagara/Private/NiagaraDataChannel.cpp
+++ b/Engine/Plugins/FX/Niagara/Source/Niagara/Private/NiagaraDataChannel.cpp
@@ -551,6 +551,7 @@ void FNiagaraDataChannelDataProxy::BeginFrame(FNiagaraGpuComputeDispatchInterfac
 
 	if(InstancesToAllocate == 0)
 	{
+		CurrFrameData = nullptr;
 		return;
 	}
 
