NDC Tooltip and error message improvements

diff --git a/Engine/Plugins/FX/Niagara/Source/Niagara/Internal/DataInterface/NiagaraDataInterfaceDataChannelRead.h b/Engine/Plugins/FX/Niagara/Source/Niagara/Internal/DataInterface/NiagaraDataInterfaceDataChannelRead.h
--- a/Engine/Plugins/FX/Niagara/Source/Niagara/Internal/DataInterface/NiagaraDataInterfaceDataChannelRead.h
+++ b/Engine/Plugins/FX/Niagara/Source/Niagara/Internal/DataInterface/NiagaraDataInterfaceDataChannelRead.h
@@ -64,6 +64,10 @@ enum class ENDIDataChannelSpawnScaleMode
 	Max UMETA(Hidden),
 };
 
+/**
+The Data Channel Reader Data Interface allows us to read from a Niagara Data Channel.
+It also allows us to spawn particles into emitters in this system based upon the entries in a Niagara Data Channel.
+*/
 UCLASS(EditInlineNew, Category = "Data Channels", CollapseCategories, meta = (DisplayName = "Data Channel Reader"), MinimalAPI)
 class UNiagaraDataInterfaceDataChannelRead : public UNiagaraDataInterfaceRWBase
 {
@@ -85,18 +89,26 @@ public:
 	UPROPERTY(EditAnywhere, Category="Data Channel")
 	TObjectPtr<UNiagaraDataChannelAsset> Channel;
 	
-	/** True if this reader will read the current frame's data. If false, we read the previous frame.
-	* Reading the current frame introduces a tick order dependency but allows for zero latency reads. Any data channel elements that are generated after this reader is used are missed.
-	* Reading the previous frame's data introduces a frame of latency but ensures we never miss any data as we have access to the whole frame.
+	/** 
+	* If this interface should read the current frame's data from the Data Channel. If false, the interface will read from the previous frame if it's available.
+	* 
+	* Reading the current frame allows us to use the most current data and have the least possible latency from the source.
+	* However, it introduces a tick order dependency between this read and the Blueprints, Game Code or other Niagara Systems writing into this Data Channel.
+	* If this interface reads data before those writing to the Data Channel have executed, then that data will be missed.
+	* 
+	* Reading the previous frame allows us to avoid this tick order dependency.
+	* We can be sure we are reading all data that is written to the Data Channel, regardless of when the writes happened in the frame.
+	* However it does introduce a 1 frame delay in the data being read and so can cause latency/lag.
 	*/
 	UPROPERTY(EditAnywhere, Category = "Data Channel", AdvancedDisplay)
 	bool bReadCurrentFrame = false;
 
 	/**
-	Whether this DI should request updated source data from the Data Channel each tick.
-	Some Data Channels have multiple separate source data elements for things such as spatial subdivision. 
+	The source Data Channel data for this interface will be refreshed every frame.
+	Some Data Channels have multiple separate data elements for things such as spatial subdivision. 
 	Each DI will request the correct one for it's owning system instance from the data channel. 
-	Depending on the data channel this could be an expensive search so we should avoid doing this every tick if possible.
+	Depending on the Data Channel this could be an expensive search so we should avoid doing this every tick if possible.
+	However it may be required. For example if the Niagara System is moving and reading from a spatially sub-divided NDC such as the Islands type.
 	*/
 	UPROPERTY(EditAnywhere, Category = "Data Channel", AdvancedDisplay)
 	bool bUpdateSourceDataEveryTick = true;
@@ -109,7 +121,7 @@ public:
 	It will also mean that Exec Index will be correct on a per NDC Entry level. 
 	Without this settings ExecIndex will be 0...TotalSpawnCount-1. With this it will be 0...SpawnCount for each NDC item individually.
 	Unless absolutely needed this is discouraged as it comes at significant performance cost when spawning and GPU emitters can currently only handle 8 individual spawns per frame.
-	Calling GetNDCSpawnInfo() in the particle spawn script to get the spawning NDC Index is prefered.
+	Calling GetNDCSpawnInfo() in the particle spawn script to get the spawning NDC Index is preferred.
 	*/
 	UPROPERTY(EditAnywhere, Category = "Spawning", AdvancedDisplay)
 	bool bOverrideSpawnGroupToDataChannelIndex = false;
diff --git a/Engine/Plugins/FX/Niagara/Source/Niagara/Internal/DataInterface/NiagaraDataInterfaceDataChannelWrite.h b/Engine/Plugins/FX/Niagara/Source/Niagara/Internal/DataInterface/NiagaraDataInterfaceDataChannelWrite.h
--- a/Engine/Plugins/FX/Niagara/Source/Niagara/Internal/DataInterface/NiagaraDataInterfaceDataChannelWrite.h
+++ b/Engine/Plugins/FX/Niagara/Source/Niagara/Internal/DataInterface/NiagaraDataInterfaceDataChannelWrite.h
@@ -87,7 +87,11 @@ public:
 	FSoftObjectPath DataChannelReference;
 };
 
-
+/**
+The Data Channel Writer Data Interface allows us to write data into a Niagara Data Channel.
+We can specify where the data is made visible with visibility flags for Game, CPU and GPU.
+These can help reduce costs by copying data only where it's needed.
+*/
 UCLASS(EditInlineNew, Category = "Data Channels", CollapseCategories, meta = (DisplayName = "Data Channel Writer"), MinimalAPI)
 class UNiagaraDataInterfaceDataChannelWrite : public UNiagaraDataInterface, public INiagaraSimCacheCustomStorageInterface
 {
@@ -95,36 +99,49 @@ class UNiagaraDataInterfaceDataChannelWrite : public UNiagaraDataInterface, publ
 
 public:
 
-	/** How should we allocate the buffer into which we write data. */
+	/** How should we allocate the buffer into which we write data. Currently only supports static allocation but in future will support more flexible allocation modes. */
 	UPROPERTY(EditAnywhere, Category = "Data Channel")
 	ENiagaraDataChannelAllocationMode AllocationMode = ENiagaraDataChannelAllocationMode::Static;
 
-	/** How many elements to allocate for writing per frame? Usage depends on AllocationMode. TODO: Allow allocation count to be controlled dynamically from script? */
+	/** How many elements to allocate for writing per frame? Usage depends on AllocationMode. */
 	UPROPERTY(EditAnywhere, Category = "Data Channel", meta = (EditCondition="AllocationMode == ENiagaraDataChannelAllocationMode::Static"))
 	uint32 AllocationCount = 0;
 
-	/** Whether the data generated by this DI should be published to the world game data channel. This is require to allow game BP and C++ to read this data. */
+	/** 
+	The data generated by this interface is made visible the Game code and BP. This is required to allow game C++ and BP to read this data. 
+	*/
 	UPROPERTY(EditAnywhere, Category = "Data Channel", meta = (DisplayName = "Visible to Blueprint"))
 	bool bPublishToGame = false;
 
-	/** Whether the data generated by this DI should be published to CPU emitters in other Niagara systems. */
+	/** 
+	Data generated by this interface is made visible to other CPU Niagara emitters. 
+	Allows reading this data in CPU Particle scripts as well as Emitter and System scripts.
+	Also required for this data to be used in Spawning functions on the Data Channel Read Interface such as Spawn Conditional and Spawn Direct.
+	If these spawning functions are used to spawn into a GPU emitter, the relevant data will also be made visible to the GPU. 
+	There is no need to also set "Visible to GPU Systems".
+	*/
 	UPROPERTY(EditAnywhere, Category = "Data Channel", meta = (DisplayName = "Visible to CPU systems"))
 	bool bPublishToCPU = false;
 
-	/** Whether the data generated by this DI should be published to GPU emitters in other Niagara systems.  */
+	/** 
+	Data generated by this interface will be made visible to GPU Niagara emitters.
+	Set this only if you need to directly read this data from Niagara emitters on the GPU.
+	If you are intending to use this data for Spawning particles via Spawn Conditional or Spawn Direct then you should not set this.
+	In that case, set only "Visible To CPU systems" and if required by a GPU emitter, the data will be passed back to the GPU automatically for use in initializing spawned particles.
+	*/
 	UPROPERTY(EditAnywhere, Category = "Data Channel", meta = (DisplayName = "Visible to GPU systems"))
 	bool bPublishToGPU = false;
 	
 	/**
-	Whether this DI should request updated destination data from the Data Channel each tick.
-	Some Data Channels have multiple separate data elements for things such as spacial subdivision. 
+	The target Data Channel data for this interface will be refreshed every frame.
+	Some Data Channels have multiple separate data elements for things such as spatial subdivision. 
 	Each DI will request the correct one for it's owning system instance from the data channel. 
-	Depending on the data channel this could be an expensive search so we should avoid doing this every tick if possible.
+	Depending on the Data Channel this could be an expensive search so we should avoid doing this every tick if possible.
 	*/
 	UPROPERTY(EditAnywhere, Category = "Data Channel", AdvancedDisplay)
 	bool bUpdateDestinationDataEveryTick = true;
 
-	/** The data channel to write to. */
+	/** The Data Channel to write to. */
 	UPROPERTY(EditAnywhere, Category = "Data Channel")
 	TObjectPtr<UNiagaraDataChannelAsset> Channel;
 
diff --git a/Engine/Plugins/FX/Niagara/Source/Niagara/Private/DataInterface/NiagaraDataInterfaceDataChannelRead.cpp b/Engine/Plugins/FX/Niagara/Source/Niagara/Private/DataInterface/NiagaraDataInterfaceDataChannelRead.cpp
--- a/Engine/Plugins/FX/Niagara/Source/Niagara/Private/DataInterface/NiagaraDataInterfaceDataChannelRead.cpp
+++ b/Engine/Plugins/FX/Niagara/Source/Niagara/Private/DataInterface/NiagaraDataInterfaceDataChannelRead.cpp
@@ -99,7 +99,7 @@ namespace NDIDataChannelReadLocal
 		{
 			Sig.Name = TEXT("Num");
 #if WITH_EDITORONLY_DATA
-			Sig.Description = LOCTEXT("NumFunctionDescription", "Returns the current number of DataChannel accessible by this interface.");
+			Sig.Description = LOCTEXT("NumFunctionDescription", "Returns the current number of elements in the Data Channel being read.");
 			NIAGARA_ADD_FUNCTION_SOURCE_INFO(Sig)
 #endif
 			Sig.bMemberFunction = true;
@@ -116,7 +116,7 @@ namespace NDIDataChannelReadLocal
 		{
 			Sig.Name = NDIDataChannelUtilities::GetNDCSpawnDataName;
 #if WITH_EDITORONLY_DATA
-			Sig.Description = LOCTEXT("GetNDCSpawnInfoFunctionDescription", "Returns useful data in relation the the NDC item that spawned this particle.");
+			Sig.Description = LOCTEXT("GetNDCSpawnInfoFunctionDescription", "Returns data in relation the the NDC item that spawned this particle. Only valid for particles spawned from NDC and only on the frame in which they're spawned.");
 			NIAGARA_ADD_FUNCTION_SOURCE_INFO(Sig)
 #endif		
 			Sig.bMemberFunction = true;
@@ -124,8 +124,8 @@ namespace NDIDataChannelReadLocal
 			Sig.AddInput(FNiagaraVariable(FNiagaraTypeDefinition(UNiagaraDataInterfaceDataChannelRead::StaticClass()), TEXT("DataChannel interface")));
 			Sig.AddInputWithoutDefault(FNiagaraVariable(FNiagaraTypeDefinition(FNiagaraEmitterID::StaticStruct()), TEXT("Emitter ID")), LOCTEXT("EmitterIDDesc", "ID of the emitter we'd like to spawn into. This can be obtained from Engine.Emitter.ID."));
 			Sig.AddInputWithDefault(FNiagaraVariable(FNiagaraTypeDefinition::GetIntDef(), TEXT("Spawned Particle Exec Index")), 0, LOCTEXT("GetNDCSpawnData_InExecIndexDesc","The execution index of the spawned particle."));
-			Sig.AddOutput(FNiagaraVariable(FNiagaraTypeDefinition::GetIntDef(), TEXT("NDC Index")), LOCTEXT("GetNDCSpawnData_OutNDCIndexDesc","Index of the NDC item that spawned this particle."));
-			Sig.AddOutput(FNiagaraVariable(FNiagaraTypeDefinition::GetIntDef(), TEXT("NDC Spawn Index")), LOCTEXT("GetNDCSpawnData_OutNDCSpawnIndexDesc","The index of this particle in relation to all the particle spawned by the same NDC item."));
+			Sig.AddOutput(FNiagaraVariable(FNiagaraTypeDefinition::GetIntDef(), TEXT("NDC Index")), LOCTEXT("GetNDCSpawnData_OutNDCIndexDesc","Index of the NDC item that spawned this particle. Can be used to read the NDC data and initialize the spawning particle with data from the NDC."));
+			Sig.AddOutput(FNiagaraVariable(FNiagaraTypeDefinition::GetIntDef(), TEXT("NDC Spawn Index")), LOCTEXT("GetNDCSpawnData_OutNDCSpawnIndexDesc","The index of this particle in relation to all the particle spawned by the same NDC item. Similar to Exec Index but for particles spawned by the same NDC item."));
 			Sig.AddOutput(FNiagaraVariable(FNiagaraTypeDefinition::GetIntDef(), TEXT("NDC Spawn Count")), LOCTEXT("GetNDCSpawnData_OutNDCSpawnCountDesc","The number of particles spawned by the same NDC item."));
 		}
 		return Sig;
@@ -138,7 +138,7 @@ namespace NDIDataChannelReadLocal
 		{
 			Sig.Name = TEXT("Read");
 #if WITH_EDITORONLY_DATA
-			Sig.Description = LOCTEXT("ReadFunctionDescription", "Reads DataChannel data at a specific index. Any values we read that are not in the DataChannel data are set to their default values. Returns success if there was a valid DataChannel to read from.");
+			Sig.Description = LOCTEXT("ReadFunctionDescription", "Reads Data Channel data at a specific index. Any values we read that are not in the Data Channel data are set to their default values. Returns success if there was a valid Data Channel entry to read from at the given index.");
 			NIAGARA_ADD_FUNCTION_SOURCE_INFO(Sig)
 #endif
 			Sig.bMemberFunction = true;
@@ -158,15 +158,15 @@ namespace NDIDataChannelReadLocal
 		{
 			Sig.Name = TEXT("Consume");
 #if WITH_EDITORONLY_DATA
-			Sig.Description = LOCTEXT("ConsumeFunctionDescription", "Consumes an DataChannel from the end of the DataChannel array and reads the specified values. Any values we read that are not in the DataChannel data are set to their default values. Returns success if an DataChannel was available to pop.");
+			Sig.Description = LOCTEXT("ConsumeFunctionDescription", "Consumes an item from the Data Channel and reads the specified values. Any values we read that are not in the Data Channel data are set to their default values. Returns success if an entry was available to be consumed in the Data Channel.");
 			NIAGARA_ADD_FUNCTION_SOURCE_INFO(Sig)
 #endif
 			Sig.bMemberFunction = true;
 			Sig.bReadFunction = true;
 			Sig.AddInput(FNiagaraVariable(FNiagaraTypeDefinition(UNiagaraDataInterfaceDataChannelRead::StaticClass()), TEXT("DataChannel interface")));
-			Sig.AddInputWithDefault(FNiagaraVariable(FNiagaraTypeDefinition::GetBoolDef(), TEXT("Consume")), FNiagaraBool(true), LOCTEXT("ConsumeInputDesc", "True if this instance (particle/emitter etc) should consume data from the data channel in this call."));
+			Sig.AddInputWithDefault(FNiagaraVariable(FNiagaraTypeDefinition::GetBoolDef(), TEXT("Consume")), FNiagaraBool(true), LOCTEXT("ConsumeInputDesc", "True if this instance (particle/emitter etc) should consume data from the Data Channel in this call."));
 			Sig.AddOutput(FNiagaraVariable(FNiagaraTypeDefinition::GetBoolDef(), TEXT("Success")), LOCTEXT("ConsumeSuccessOutputDesc", "True if all reads succeeded."));
-			Sig.AddOutput(FNiagaraVariable(FNiagaraTypeDefinition::GetIntDef(), TEXT("Index")), LOCTEXT("ConsumeIndexOutputDesc", "The index we actually read from. If reading failed this can be -1. This allows subsequent reads of the data channel at this index."));
+			Sig.AddOutput(FNiagaraVariable(FNiagaraTypeDefinition::GetIntDef(), TEXT("Index")), LOCTEXT("ConsumeIndexOutputDesc", "The index we actually read from. If reading failed this can be -1. This allows subsequent reads of the Data Channel at this index."));
 			Sig.RequiredOutputs = IntCastChecked<int16>(Sig.Outputs.Num()); //The user defines what we read in the graph.			
 		}
 		return Sig;
@@ -186,10 +186,10 @@ namespace NDIDataChannelReadLocal
 			Sig.Name = TEXT("SpawnConditional");
 #if WITH_EDITORONLY_DATA
 			NIAGARA_ADD_FUNCTION_SOURCE_INFO(Sig)
-			Sig.Description = LOCTEXT("SpawnCustomFunctionDescription", "Will Spawn particles into the bound Emitter between Min and Max counts for every element in the bound Data Channel.\n\
-		Can take optional additional parameters as conditions on spawning that will be compared against the contents of each data channel element.\n\
-		For example you could spawn only for a particular value of an enum.\n\
-		For compound data types that contain multiple component floats or ints, comparissons are done on a per component basis.\n\
+			Sig.Description = LOCTEXT("SpawnCustomFunctionDescription", "Will Spawn particles into the given Emitter between Min and Max counts for every element in the Data Channel.\n\
+		Can take optional additional parameters as conditions on spawning. The data passed into the function will be compared against the contents of each Data Channel element.\n\
+		For example, you could spawn only for NDC items that match a particular value of an enum.\n\
+		For compound data types that contain multiple component floats or ints, comparisons are done on a per component basis.\n\
 		For example if you add a Vector condition parameter it will be compared against each component of the corresponding Vector in the Data Channel.\n\
 		Result = (Param.X == ChannelValue.X) && (Param.Y == ChannelValue.Y) && (Param.Z == ChannelValue.Z)");
 			Sig.FunctionVersion = static_cast<uint32>(FunctionVersion_SpawnConditional::EmitterIDParameter);
@@ -198,12 +198,12 @@ namespace NDIDataChannelReadLocal
 			Sig.bRequiresExecPin = true;
 			Sig.ModuleUsageBitmask = ENiagaraScriptUsageMask::Emitter | ENiagaraScriptUsageMask::System;
 			Sig.AddInput(FNiagaraVariable(FNiagaraTypeDefinition(UNiagaraDataInterfaceDataChannelRead::StaticClass()), TEXT("DataChannel interface")));
-			Sig.AddInputWithDefault(FNiagaraVariable(FNiagaraTypeDefinition::GetBoolDef(), TEXT("Enable")), FNiagaraBool(true), LOCTEXT("SpawnEnableInputDesc", "Enable or disable this function call. If false, this call with have no effetcs."));
+			Sig.AddInputWithDefault(FNiagaraVariable(FNiagaraTypeDefinition::GetBoolDef(), TEXT("Enable")), FNiagaraBool(true), LOCTEXT("SpawnEnableInputDesc", "Enable or disable this function call. If false, this call with have no effect."));
 			Sig.AddInputWithoutDefault(FNiagaraVariable(FNiagaraTypeDefinition(FNiagaraEmitterID::StaticStruct()), TEXT("Emitter ID")), LOCTEXT("EmitterIDDesc", "ID of the emitter we'd like to spawn into. This can be obtained from Engine.Emitter.ID."));
-			Sig.AddInput(FNiagaraVariable(StaticEnum<ENDIDataChannelSpawnMode>(), TEXT("Mode")), LOCTEXT("SpawnCondModeInputDesc", "A mode switch that controls how this funciton will behave and interact with other calls to spawn functions."));
-			Sig.AddInput(FNiagaraVariable(StaticEnum<ENiagaraConditionalOperator>(), TEXT("Operator")), LOCTEXT("SpawnCondOpInputDesc", "The comparison operator to use when comparing values in the data channel to conditional parameters."));
-			Sig.AddInputWithDefault(FNiagaraVariable(FNiagaraTypeDefinition::GetIntDef(), TEXT("Min Spawn Count")), 1, LOCTEXT("MinSpawnCountInputDesc", "Minimum number of particles to spawn for each element in the data channel."));
-			Sig.AddInputWithDefault(FNiagaraVariable(FNiagaraTypeDefinition::GetIntDef(), TEXT("Max Spawn Count")), 1, LOCTEXT("MaxSpawnCountInputDesc", "Maximum number of particles to spawn for each element in the data channel."));
+			Sig.AddInput(FNiagaraVariable(StaticEnum<ENDIDataChannelSpawnMode>(), TEXT("Mode")), LOCTEXT("SpawnCondModeInputDesc", "Controls how this function will interact with other calls to spawn functions. Spawn counts for each NDC can be accumulated or overwritten."));
+			Sig.AddInput(FNiagaraVariable(StaticEnum<ENiagaraConditionalOperator>(), TEXT("Operator")), LOCTEXT("SpawnCondOpInputDesc", "Compare the input against the Data Channel value:\n\n[Input] [Condition] [Data Channel Value]\n\nFor example:\nSpawn if [100] is [greater than] [Data Channel 'Height' Variable]"));
+			Sig.AddInputWithDefault(FNiagaraVariable(FNiagaraTypeDefinition::GetIntDef(), TEXT("Min Spawn Count")), 1, LOCTEXT("MinSpawnCountInputDesc", "Minimum number of particles to spawn for each element in the Data Channel."));
+			Sig.AddInputWithDefault(FNiagaraVariable(FNiagaraTypeDefinition::GetIntDef(), TEXT("Max Spawn Count")), 1, LOCTEXT("MaxSpawnCountInputDesc", "Maximum number of particles to spawn for each element in the Data Channel."));
 			Sig.RequiredInputs = IntCastChecked<int16>(Sig.Inputs.Num());
 		}
 		return Sig;
@@ -216,7 +216,7 @@ namespace NDIDataChannelReadLocal
 		{
 			Sig.Name = TEXT("SpawnDirect");
 #if WITH_EDITORONLY_DATA
-			Sig.Description = LOCTEXT("SpawnDirectFunctionDescription", "Spawns particles into a given emitter for each entry in the bound NDC. Spawn count is determined directly from a value in the NDC. Additional per NDC random scale and a clamp is available.");
+			Sig.Description = LOCTEXT("SpawnDirectFunctionDescription", "Spawns particles into a given emitter for each entry in the Data Channel. Spawn count is determined directly from a value in the Data Channel. Additional per NDC item random scale and a clamp is available.");
 			NIAGARA_ADD_FUNCTION_SOURCE_INFO(Sig)
 #endif
 			Sig.FunctionSpecifiers.Add(VarNameKey);
@@ -225,13 +225,13 @@ namespace NDIDataChannelReadLocal
 			Sig.bRequiresExecPin = true;
 			Sig.ModuleUsageBitmask = ENiagaraScriptUsageMask::Emitter | ENiagaraScriptUsageMask::System;
 			Sig.AddInput(FNiagaraVariable(FNiagaraTypeDefinition(UNiagaraDataInterfaceDataChannelRead::StaticClass()), TEXT("DataChannel interface")));
-			Sig.AddInputWithDefault(FNiagaraVariable(FNiagaraTypeDefinition::GetBoolDef(), TEXT("Enable")), FNiagaraBool(true), LOCTEXT("SpawnEnableInputDesc", "Enable or disable this function call. If false, this call with have no effetcs."));
+			Sig.AddInputWithDefault(FNiagaraVariable(FNiagaraTypeDefinition::GetBoolDef(), TEXT("Enable")), FNiagaraBool(true), LOCTEXT("SpawnEnableInputDesc", "Enable or disable this function call. If false, this call with have no effect."));
 			Sig.AddInputWithoutDefault(FNiagaraVariable(FNiagaraTypeDefinition(FNiagaraEmitterID::StaticStruct()), TEXT("Emitter ID")), LOCTEXT("EmitterIDDesc", "ID of the emitter we'd like to spawn into. This can be obtained from Engine.Emitter.ID."));
-			Sig.AddInput(FNiagaraVariable(StaticEnum<ENDIDataChannelSpawnMode>(), TEXT("Mode")), LOCTEXT("SpawnCondModeInputDesc", "A mode switch that controls how this funciton will behave and interact with other calls to spawn functions."));			
+			Sig.AddInput(FNiagaraVariable(StaticEnum<ENDIDataChannelSpawnMode>(), TEXT("Mode")), LOCTEXT("SpawnCondModeInputDesc", "Controls how this function will interact with other calls to spawn functions. Spawn counts for each NDC can be accumulated or overwritten."));		
 			Sig.AddInputWithDefault(FNiagaraVariable(FNiagaraTypeDefinition::GetFloatDef(), TEXT("RandomScaleMin")), 1.0f, LOCTEXT("SpawnDirectRandomScaleMinInputDesc", "Minimum value for an additional random scale applied to each NDC spawn count."));
 			Sig.AddInputWithDefault(FNiagaraVariable(FNiagaraTypeDefinition::GetFloatDef(), TEXT("RandomScaleMax")), 1.0f, LOCTEXT("SpawnDirectRandomScaleMaxInputDesc", "Maximum value for an additional random scale applied to each NDC spawn count."));
 			Sig.AddInputWithDefault(FNiagaraVariable(FNiagaraTypeDefinition::GetIntDef(), TEXT("ClampMin")), 0, LOCTEXT("SpawnDirectClampMinInputDesc", "Minimum Spawn Count to use after random scale is applied."));
-			Sig.AddInputWithDefault(FNiagaraVariable(FNiagaraTypeDefinition::GetIntDef(), TEXT("ClampMax")), 1, LOCTEXT("SpawnDirectClampMaxInputDesc", "Minimum Spawn Count to use after random scale is applied. If < 0, No max clamp is applied."));
+			Sig.AddInputWithDefault(FNiagaraVariable(FNiagaraTypeDefinition::GetIntDef(), TEXT("ClampMax")), 1, LOCTEXT("SpawnDirectClampMaxInputDesc", "Maximum Spawn Count to use after random scale is applied. If < 0, No max clamp is applied."));
 		}
 		return Sig;
 	}
@@ -243,7 +243,7 @@ namespace NDIDataChannelReadLocal
 		{
 			Sig.Name = TEXT("ScaleSpawnCount");
 #if WITH_EDITORONLY_DATA
-			Sig.Description = LOCTEXT("ScaleSpawnCountFunctionDescription", "Applies a scaling value for each NDC spawn based on some element in the NDC data. Optional additional random scale and clamp operations to the value read from each NDC entry.");
+			Sig.Description = LOCTEXT("ScaleSpawnCountFunctionDescription", "Applies a scaling value for each NDC Item spawn count based on a variable in the NDC data. Optional additional random scale and clamp operations to the value read from each NDC entry.");
 			NIAGARA_ADD_FUNCTION_SOURCE_INFO(Sig)
 #endif
 			Sig.FunctionSpecifiers.Add(VarNameKey);
@@ -252,7 +252,7 @@ namespace NDIDataChannelReadLocal
 			Sig.bRequiresExecPin = true;
 			Sig.ModuleUsageBitmask = ENiagaraScriptUsageMask::Emitter | ENiagaraScriptUsageMask::System;
 			Sig.AddInput(FNiagaraVariable(FNiagaraTypeDefinition(UNiagaraDataInterfaceDataChannelRead::StaticClass()), TEXT("DataChannel interface")));
-			Sig.AddInputWithDefault(FNiagaraVariable(FNiagaraTypeDefinition::GetBoolDef(), TEXT("Enable")), FNiagaraBool(true), LOCTEXT("SpawnEnableInputDesc", "Enable or disable this function call. If false, this call with have no effetcs."));
+			Sig.AddInputWithDefault(FNiagaraVariable(FNiagaraTypeDefinition::GetBoolDef(), TEXT("Enable")), FNiagaraBool(true), LOCTEXT("SpawnEnableInputDesc", "Enable or disable this function call. If false, this call with have no effects."));
 			Sig.AddInputWithoutDefault(FNiagaraVariable(FNiagaraTypeDefinition(FNiagaraEmitterID::StaticStruct()), TEXT("Emitter ID")), LOCTEXT("EmitterIDDesc", "ID of the emitter we'd like to spawn into. This can be obtained from Engine.Emitter.ID."));
 			Sig.AddInput(FNiagaraVariable(StaticEnum<ENDIDataChannelSpawnScaleMode>(), TEXT("Mode")), LOCTEXT("SpawnScaleModeInputDesc", "Control whether to override or combine this scale with previously set scales when calling this function multiple times."));
 			Sig.AddInputWithDefault(FNiagaraVariable(FNiagaraTypeDefinition::GetFloatDef(), TEXT("RandomScaleMin")), 1.0f, LOCTEXT("ScaleSpawnCountRandomScaleMinInputDesc", "Minimum value for a random additional scale applied to each NDC spawn."));
@@ -554,7 +554,7 @@ bool FNDIDataChannelReadInstanceData::Tick(UNiagaraDataInterfaceDataChannelRead*
 
 	if (!DataChannel.IsValid() /*&& !SourceDI.IsValid()*/)//TODO: Local reads
 	{
-		UE_LOG(LogNiagara, Warning, TEXT("Niagara Data Channel Reader Interface could not find a valid data channel.\nData Channel: %s\nSystem: %s\nComponent:%s\n")
+		UE_LOG(LogNiagara, Warning, TEXT("Niagara Data Channel Reader Interface could not find a valid Data Channel.\nData Channel: %s\nSystem: %s\nComponent:%s\n")
 			, Interface->Channel ? *Interface->Channel->GetName() : TEXT("None")
 			, *Instance->GetSystem()->GetPathName()
 			, *Instance->GetAttachComponent()->GetPathName());
@@ -1019,7 +1019,7 @@ void UNiagaraDataInterfaceDataChannelRead::GetFeedback(UNiagaraSystem* InAsset,
 	{
 		if (const UNiagaraDataChannel* DataChannel = RuntimeReadDI->Channel->Get())
 		{
-			//Ensure the data channel contains all the parameters this function is requesting.
+			//Ensure the Data Channel contains all the parameters this function is requesting.
 			TConstArrayView<FNiagaraDataChannelVariable> ChannelVars = DataChannel->GetVariables();
 			for (const FNDIDataChannelFunctionInfo& FuncInfo : RuntimeReadDI->GetCompiledData().GetFunctionInfo())
 			{
@@ -1098,7 +1098,7 @@ void UNiagaraDataInterfaceDataChannelRead::ValidateFunction(const FNiagaraFuncti
 {
 	Super::ValidateFunction(Function, OutValidationErrors);
 
-	//It would be great to be able to validate the parameters on the function calls here but this is only called on the DI CDO. We don't have the context of which data channel we'll be accessing.
+	//It would be great to be able to validate the parameters on the function calls here but this is only called on the DI CDO. We don't have the context of which Data Channel we'll be accessing.
 	//The translator should have all the required data to use the actual DIs when validating functions. We just need to do some wrangling to pull it from the pre compiled data correctly.
 	//This would probably also allow us to actually call hlsl generation functions on the actual DIs rather than their CDOs. Which would allow for a bunch of better optimized code gen for things like fluids.
 	//TODO!!!
@@ -1756,7 +1756,7 @@ void UNiagaraDataInterfaceDataChannelRead::SpawnConditional(FVectorVMExternalFun
 		int32 SpawnMin = FMath::Max(0, InSpawnMin.GetAndAdvance());
 		int32 SpawnMax = FMath::Max(0, InSpawnMax.GetAndAdvance());
 
-		//Each data channel element has an additional spawn entry which accumulates across all spawning calls and can be nulled independently by a suppression call.
+		//Each Data Channel element has an additional spawn entry which accumulates across all spawning calls and can be nulled independently by a suppression call.
 		FNDIDataChannelRead_EmitterInstanceData& EmitterInstData = InstData->EmitterInstanceData.FindOrAdd(EmitterInst);
 		TArray<FNDIDataChannelRead_EmitterSpawnInfo>& EmitterConditionalSpawns = EmitterInstData.NDCSpawnCounts;
 		EmitterConditionalSpawns.SetNum(NumDataChannelInstances);
@@ -1907,7 +1907,7 @@ void UNiagaraDataInterfaceDataChannelRead::SpawnDirect(FVectorVMExternalFunction
 		int32 ClampMax = InClampMax.GetAndAdvance();
 		ClampMax = ClampMax < 0 ? INT_MAX : ClampMax;
 
-		//Each data channel element has an additional spawn entry which accumulates across all spawning calls and can be nulled independently by a suppression call.
+		//Each Data Channel element has an additional spawn entry which accumulates across all spawning calls and can be nulled independently by a suppression call.
 		FNDIDataChannelRead_EmitterInstanceData& EmitterInstData = InstData->EmitterInstanceData.FindOrAdd(EmitterInst);
 		TArray<FNDIDataChannelRead_EmitterSpawnInfo>& EmitterConditionalSpawns = EmitterInstData.NDCSpawnCounts;
 		EmitterConditionalSpawns.SetNum(NumDataChannelInstances);
@@ -1983,7 +1983,7 @@ void UNiagaraDataInterfaceDataChannelRead::ScaleSpawnCount(FVectorVMExternalFunc
 		float ClampMax = InClampMax.GetAndAdvance();
 		ClampMax = ClampMax < 0.0f ? FLT_MAX : ClampMax;
 
-		//Each data channel element has an additional spawn entry which accumulates across all spawning calls and can be nulled independently by a suppression call.
+		//Each Data Channel element has an additional spawn entry which accumulates across all spawning calls and can be nulled independently by a suppression call.
 		FNDIDataChannelRead_EmitterInstanceData& EmitterInstData = InstData->EmitterInstanceData.FindOrAdd(EmitterInst);
 		TArray<FNDIDataChannelRead_EmitterSpawnInfo>& EmitterConditionalSpawns = EmitterInstData.NDCSpawnCounts;
 		EmitterConditionalSpawns.SetNumZeroed(NumDataChannelInstances);
diff --git a/Engine/Plugins/FX/Niagara/Source/Niagara/Private/DataInterface/NiagaraDataInterfaceDataChannelWrite.cpp b/Engine/Plugins/FX/Niagara/Source/Niagara/Private/DataInterface/NiagaraDataInterfaceDataChannelWrite.cpp
--- a/Engine/Plugins/FX/Niagara/Source/Niagara/Private/DataInterface/NiagaraDataInterfaceDataChannelWrite.cpp
+++ b/Engine/Plugins/FX/Niagara/Source/Niagara/Private/DataInterface/NiagaraDataInterfaceDataChannelWrite.cpp
@@ -74,7 +74,7 @@ namespace NDIDataChannelWriteLocal
 			Sig.Name = TEXT("Num");
 #if WITH_EDITORONLY_DATA
 			NIAGARA_ADD_FUNCTION_SOURCE_INFO(Sig)
-			Sig.Description = LOCTEXT("NumFunctionDescription", "Returns the number of instances allocated for writing in the bound NDC buffer. Writes at an index beyond this will fail.");
+			Sig.Description = LOCTEXT("NumFunctionDescription", "Returns the number of instances allocated for writing into the Data Channel from this interface. Writes at an index beyond this will fail.");
 #endif
 			Sig.bMemberFunction = true;
 			Sig.AddInput(FNiagaraVariable(FNiagaraTypeDefinition(UNiagaraDataInterfaceDataChannelWrite::StaticClass()), TEXT("DataChannel interface")));
@@ -116,14 +116,14 @@ namespace NDIDataChannelWriteLocal
 			Sig.Name = TEXT("Write");
 #if WITH_EDITORONLY_DATA
 			NIAGARA_ADD_FUNCTION_SOURCE_INFO(Sig)
-			Sig.Description = LOCTEXT("WriteFunctionDescription", "Writes DataChannel data at a specific index.  Values in the DataChannel that are not written here are set to their defaults. Returns success if an DataChannel was written to.");
+			Sig.Description = LOCTEXT("WriteFunctionDescription", "Writes data into the Data Channel at a specific index.  Values in the DataChannel that are not written here are set to their defaults. Returns success if the index was valid and data was written into the Data Channel.");
 #endif
 			Sig.bMemberFunction = true;
 			Sig.bRequiresExecPin = true;
 			Sig.bWriteFunction = true;
 			Sig.bSupportsGPU = false;//Cannot use direct index writes on GPU as we write into one shared buffer with all DIs using the same NDC data.
 			Sig.AddInput(FNiagaraVariable(FNiagaraTypeDefinition(UNiagaraDataInterfaceDataChannelWrite::StaticClass()), TEXT("DataChannel interface")));
-			Sig.AddInput(EmitVar, LOCTEXT("ExecuteAppendFlagTooltip", "If true then the append is executed, if false then the append call is skipped"));
+			Sig.AddInput(EmitVar, LOCTEXT("ExecuteWriteFlagTooltip", "If true then the write is executed, if false then this call is ignored and no write occurs."));
 			Sig.AddInput(FNiagaraVariable(FNiagaraTypeDefinition::GetIntDef(), TEXT("Index")));
 			Sig.AddOutput(FNiagaraVariable(FNiagaraTypeDefinition::GetBoolDef(), TEXT("Success")));
 			Sig.RequiredInputs = IntCastChecked<int16>(Sig.Inputs.Num());//The user defines what we write in the graph.
@@ -148,7 +148,7 @@ namespace NDIDataChannelWriteLocal
 			Sig.bRequiresExecPin = true;
 			Sig.bWriteFunction = true;
 			Sig.AddInput(FNiagaraVariable(FNiagaraTypeDefinition(UNiagaraDataInterfaceDataChannelWrite::StaticClass()), TEXT("DataChannel interface")));
-			Sig.AddInput(EmitVar, LOCTEXT("ExecuteAppendFlagTooltip", "If true then the append is executed, if false then the append call is skipped"));
+			Sig.AddInput(EmitVar, LOCTEXT("ExecuteAppendFlagTooltip", "If true then the append is executed, if false then this call is skipped and no append occurs."));
 			Sig.AddOutput(FNiagaraVariable(FNiagaraTypeDefinition::GetBoolDef(), TEXT("Success")));
 			Sig.RequiredInputs = IntCastChecked<int16>(Sig.Inputs.Num());//The user defines what we write in the graph.
 		}
diff --git a/Engine/Plugins/FX/Niagara/Source/Niagara/Private/NiagaraDataChannel.cpp b/Engine/Plugins/FX/Niagara/Source/Niagara/Private/NiagaraDataChannel.cpp
--- a/Engine/Plugins/FX/Niagara/Source/Niagara/Private/NiagaraDataChannel.cpp
+++ b/Engine/Plugins/FX/Niagara/Source/Niagara/Private/NiagaraDataChannel.cpp
@@ -650,7 +650,14 @@ FNiagaraDataBufferRef FNiagaraDataChannelDataProxy::PrepareForWriteAccess(FRDGBu
 		if (CurrBufferAccessCounts > 0)
 		{
 #if !UE_BUILD_SHIPPING
-			UE_LOG(LogNiagara, Warning, TEXT("Attempting to write to a Niagara Data Channel in the same stage in which it's being read. {%s}"), *DebugName);
+			if(!bWarnedAboutSameStageRW)
+			{
+				bWarnedAboutSameStageRW = true;
+				UE_LOG(LogNiagara, Warning, TEXT("Attempting to write to a Niagara Data Channel in the same stage in which it's being read. {%s}\n\
+				This is most often caused by an NDC reading the current frame data and the same NDC being written to in PostRenderOpaque.\n\
+				A possible fix is to read the previous frames data rather than the current frame.\n\
+				Another is to ensure writes to the NDC are done earlier. Things such as Depth/GBuffer reads can make the writing system tick late in the frame."), *DebugName);
+			}
 #endif
 			return nullptr;
 		}
@@ -693,7 +700,14 @@ FNiagaraDataBufferRef FNiagaraDataChannelDataProxy::PrepareForReadAccess(FRDGBui
 			if (CurrBufferAccessCounts < 0)
 			{
 #if !UE_BUILD_SHIPPING
-				UE_LOG(LogNiagara, Warning, TEXT("Attempting to read from a Niagara Data Channel in the same stage in which it's being written. {%s}"), *DebugName);
+				if (!bWarnedAboutSameStageRW)
+				{
+					bWarnedAboutSameStageRW = true;
+					UE_LOG(LogNiagara, Warning, TEXT("Attempting to read from a Niagara Data Channel in the same stage in which it's being written. {%s}\n\
+					This is most often caused by an NDC reading the current frame data and the same NDC being written to in PostRenderOpaque.\n\
+					A possible fix is to read the previous frame data rather than the current frame.\n\
+					Another is to ensure writes to the NDC are done earlier. Things such as Depth/GBuffer reads can make the writing system tick late in the frame."), *DebugName);
+				}
 #endif
 				return nullptr;
 			}
@@ -1174,7 +1188,7 @@ void FNiagaraDataChannelData::CreateRenderThreadProxy(UNiagaraDataChannelHandler
 		RTProxy->Owner = this->AsWeak();
 		RTProxy->LayoutInfo = LayoutInfo;//Keep our own ref to the layout info as the NDCData may die on the GT before the proxy.
 #if !UE_BUILD_SHIPPING
-		RTProxy->DebugName = FString::Printf(TEXT("%s__GPUData"), *DataChannel->GetName());
+		RTProxy->DebugName = FString::Printf(TEXT("%s__GPUData"), *GetNameSafe(DataChannel->GetAsset()));
 		RTProxy->DispatchInterfaceForDebuggingOnly = ComputeDispatchInterface;
 #endif
 		ENQUEUE_RENDER_COMMAND(FNiagaraDataChannelDataProxyInit) (
diff --git a/Engine/Plugins/FX/Niagara/Source/Niagara/Public/NiagaraDataChannel.h b/Engine/Plugins/FX/Niagara/Source/Niagara/Public/NiagaraDataChannel.h
--- a/Engine/Plugins/FX/Niagara/Source/Niagara/Public/NiagaraDataChannel.h
+++ b/Engine/Plugins/FX/Niagara/Source/Niagara/Public/NiagaraDataChannel.h
@@ -91,6 +91,7 @@ struct FNiagaraDataChannelDataProxy : public TSharedFromThis<FNiagaraDataChannel
 	int32 CurrBufferAccessCounts = 0;
 
 	#if !UE_BUILD_SHIPPING
+	bool bWarnedAboutSameStageRW = false;
 	FNiagaraGpuComputeDispatchInterface* DispatchInterfaceForDebuggingOnly = nullptr;
 	
 	FString DebugName;
diff --git a/Engine/Plugins/FX/Niagara/Source/NiagaraEditor/Private/Widgets/DataChannel/NiagaraDataChannelWizard.cpp b/Engine/Plugins/FX/Niagara/Source/NiagaraEditor/Private/Widgets/DataChannel/NiagaraDataChannelWizard.cpp
--- a/Engine/Plugins/FX/Niagara/Source/NiagaraEditor/Private/Widgets/DataChannel/NiagaraDataChannelWizard.cpp
+++ b/Engine/Plugins/FX/Niagara/Source/NiagaraEditor/Private/Widgets/DataChannel/NiagaraDataChannelWizard.cpp
@@ -357,7 +357,7 @@ namespace UE::Niagara::Wizard::DataChannel
 		{
 			if (PreviousPage->Data->SpawnMode == ENiagaraDataChanneSpawnModuleMode::ConditionalSpawn)
 			{
-				return LOCTEXT("SpawnConditionalPageHeader", "Please select which data channel variables should be used as conditions to spawn particles. This is optional!\nFor example, if you select a vector and integer parameter in the data channel they will be compared against the corresponding module inputs.\nResult = (InputA == ChannelValue.A) && (InputB == ChannelValue.B)");
+				return LOCTEXT("SpawnConditionalPageHeader", "OPTIONAL:\nSelect which data channel variables should be used as conditions to spawn particles.\nModule inputs will be created for all selected variables.\nParticles will only be spawned if the data channel variables match the module input values.");
 			}
 			return LOCTEXT("SpawnDirectPageHeader", "Please select which data channel variable should be used as particle spawn count. This needs to be an integer parameter in the data channel.");
 		}
@@ -954,6 +954,7 @@ TSharedRef<FModuleWizardModel> DataChannel::CreateSpawnNDCModuleWizardModel()
 					{
 						UEdGraphPin* EnablePin = Utilities::AddReadParameterPin(FNiagaraTypeDefinition::GetBoolDef(), FName("Spawn Enabled"), MapGetNode);
 						Utilities::SetDefaultValue(Graph, EnablePin->PinName, FNiagaraTypeDefinition::GetBoolDef(), true);
+						Utilities::SetTooltip(Graph, EnablePin->PinName, SpawnFunction->Signature.InputDescriptions[FNiagaraVariable(FNiagaraTypeDefinition::GetBoolDef(), EnableInput->PinName)]);
 						GraphSchema->TryCreateConnection(EnablePin, EnableInput);
 					}
 					if (UEdGraphPin* EmitterIDInput = SpawnFunction->FindPin(FName("Emitter ID"), EGPD_Input))
@@ -961,20 +962,24 @@ TSharedRef<FModuleWizardModel> DataChannel::CreateSpawnNDCModuleWizardModel()
 						UEdGraphPin* EmitterIDPin = Utilities::AddReadParameterPin(FNiagaraTypeDefinition(FNiagaraEmitterID::StaticStruct()), FName("Emitter ID"), MapGetNode);
 						GraphSchema->TryCreateConnection(EmitterIDInput, EmitterIDPin);
 						Utilities::SetDefaultBinding(Graph, EmitterIDPin->PinName, SYS_PARAM_ENGINE_EMITTER_ID.GetName());
+						Utilities::SetTooltip(Graph, EmitterIDPin->PinName, SpawnFunction->Signature.InputDescriptions[FNiagaraVariable(FNiagaraTypeDefinition(FNiagaraEmitterID::StaticStruct()), EmitterIDInput->PinName)]);
 					}
 					if (UEdGraphPin* ModeInput = SpawnFunction->FindPin(FName("Mode"), EGPD_Input))
 					{
 						UEdGraphPin* SpawnModePin = Utilities::AddReadParameterPin(FNiagaraTypeDefinition(StaticEnum<ENDIDataChannelSpawnMode>()), FName("Spawn Mode"), MapGetNode);
 						GraphSchema->TryCreateConnection(ModeInput, SpawnModePin);
+						Utilities::SetTooltip(Graph, SpawnModePin->PinName, SpawnFunction->Signature.InputDescriptions[FNiagaraVariable(FNiagaraTypeDefinition(StaticEnum<ENDIDataChannelSpawnMode>()), ModeInput->PinName)]);
 					}
 					if (UEdGraphPin* OperatorInput = SpawnFunction->FindPin(FName("Operator"), EGPD_Input))
 					{
 						UEdGraphPin* OperatorPin = Utilities::AddReadParameterPin(FNiagaraTypeDefinition(StaticEnum<ENiagaraConditionalOperator>()), FName("Comparison Operator"), MapGetNode);
 						GraphSchema->TryCreateConnection(OperatorInput, OperatorPin);
+						Utilities::SetTooltip(Graph, OperatorPin->PinName, SpawnFunction->Signature.InputDescriptions[FNiagaraVariable(FNiagaraTypeDefinition(StaticEnum<ENiagaraConditionalOperator>()), OperatorInput->PinName)]);
 					}
 					if (UEdGraphPin* MinInput = SpawnFunction->FindPin(FName(SpawnMode == ENiagaraDataChanneSpawnModuleMode::ConditionalSpawn ? "Min Spawn Count" : "ClampMin"), EGPD_Input))
 					{
 						UEdGraphPin* SpawnMinPin = Utilities::AddReadParameterPin(FNiagaraTypeDefinition::GetIntDef(), FName("Min Count"), MapGetNode);
+						Utilities::SetTooltip(Graph, SpawnMinPin->PinName, SpawnFunction->Signature.InputDescriptions[FNiagaraVariable(FNiagaraTypeDefinition::GetIntDef(), MinInput->PinName)]);
 						int32 Default = SpawnMode == ENiagaraDataChanneSpawnModuleMode::ConditionalSpawn ? 1 : -1;//Default to -1 for spawn direct as this is a clamp. aka, no clamp by default.
 						Utilities::SetDefaultValue(Graph, SpawnMinPin->PinName, FNiagaraTypeDefinition::GetIntDef(), Default);
 						GraphSchema->TryCreateConnection(MinInput, SpawnMinPin);
@@ -982,6 +987,7 @@ TSharedRef<FModuleWizardModel> DataChannel::CreateSpawnNDCModuleWizardModel()
 					if (UEdGraphPin* MaxInput = SpawnFunction->FindPin(FName(SpawnMode == ENiagaraDataChanneSpawnModuleMode::ConditionalSpawn ? "Max Spawn Count" : "ClampMax"), EGPD_Input))
 					{
 						UEdGraphPin* SpawnMaxPin = Utilities::AddReadParameterPin(FNiagaraTypeDefinition::GetIntDef(), FName("Max Count"), MapGetNode);
+						Utilities::SetTooltip(Graph, SpawnMaxPin->PinName, SpawnFunction->Signature.InputDescriptions[FNiagaraVariable(FNiagaraTypeDefinition::GetIntDef(), MaxInput->PinName)]);
 						int32 Default = SpawnMode == ENiagaraDataChanneSpawnModuleMode::ConditionalSpawn ? 1 : -1;//Default to -1 for spawn direct as this is a clamp. aka, no clamp by default.
 						Utilities::SetDefaultValue(Graph, SpawnMaxPin->PinName, FNiagaraTypeDefinition::GetIntDef(), Default);
 						GraphSchema->TryCreateConnection(MaxInput, SpawnMaxPin);
@@ -992,12 +998,14 @@ TSharedRef<FModuleWizardModel> DataChannel::CreateSpawnNDCModuleWizardModel()
 						{
 							UEdGraphPin* ScaleMinPin = Utilities::AddReadParameterPin(FNiagaraTypeDefinition::GetFloatDef(), FName("Random Scale Min"), MapGetNode);
 							Utilities::SetDefaultValue(Graph, ScaleMinPin->PinName, FNiagaraTypeDefinition::GetFloatDef(), 1.0f);
+							Utilities::SetTooltip(Graph, ScaleMinPin->PinName, SpawnFunction->Signature.InputDescriptions[FNiagaraVariable(FNiagaraTypeDefinition::GetFloatDef(), ScaleMinInput->PinName)]);
 							GraphSchema->TryCreateConnection(ScaleMinInput, ScaleMinPin);
 						}
 						if (UEdGraphPin* ScaleMaxInput = SpawnFunction->FindPin(FName("RandomScaleMax"), EGPD_Input))
 						{
 							UEdGraphPin* ScaleMaxPin = Utilities::AddReadParameterPin(FNiagaraTypeDefinition::GetFloatDef(), FName("Random Scale Max"), MapGetNode);
 							Utilities::SetDefaultValue(Graph, ScaleMaxPin->PinName, FNiagaraTypeDefinition::GetFloatDef(), 1.0f);
+							Utilities::SetTooltip(Graph, ScaleMaxPin->PinName, SpawnFunction->Signature.InputDescriptions[FNiagaraVariable(FNiagaraTypeDefinition::GetFloatDef(), ScaleMaxInput->PinName)]);
 							GraphSchema->TryCreateConnection(ScaleMaxInput, ScaleMaxPin);
 						}
 
diff --git a/Engine/Plugins/FX/Niagara/Source/NiagaraEditor/Private/Widgets/Wizard/SNiagaraModuleWizard.cpp b/Engine/Plugins/FX/Niagara/Source/NiagaraEditor/Private/Widgets/Wizard/SNiagaraModuleWizard.cpp
--- a/Engine/Plugins/FX/Niagara/Source/NiagaraEditor/Private/Widgets/Wizard/SNiagaraModuleWizard.cpp
+++ b/Engine/Plugins/FX/Niagara/Source/NiagaraEditor/Private/Widgets/Wizard/SNiagaraModuleWizard.cpp
@@ -242,4 +242,16 @@ void Utilities::SetDefaultBinding(UNiagaraGraph* Graph, const FName& VarName, co
 	}
 }
 
+void Utilities::SetTooltip(UNiagaraGraph* Graph, const FName& VarName, const FText& Tooltip)
+{
+	if (Graph)
+	{
+		if (UNiagaraScriptVariable* ScriptVariable = Graph->GetScriptVariable(VarName))
+		{
+			ScriptVariable->Metadata.Description = Tooltip;
+			Graph->ScriptVariableChanged(ScriptVariable->Variable);
+		}
+	}
+}
+
 #undef LOCTEXT_NAMESPACE
diff --git a/Engine/Plugins/FX/Niagara/Source/NiagaraEditor/Public/Widgets/Wizard/SNiagaraModuleWizard.h b/Engine/Plugins/FX/Niagara/Source/NiagaraEditor/Public/Widgets/Wizard/SNiagaraModuleWizard.h
--- a/Engine/Plugins/FX/Niagara/Source/NiagaraEditor/Public/Widgets/Wizard/SNiagaraModuleWizard.h
+++ b/Engine/Plugins/FX/Niagara/Source/NiagaraEditor/Public/Widgets/Wizard/SNiagaraModuleWizard.h
@@ -134,5 +134,6 @@ namespace UE::Niagara::Wizard
 			}
 		}
 		NIAGARAEDITOR_API void SetDefaultBinding(UNiagaraGraph* Graph, const FName& VarName, const FName& DefaultBinding);
+		NIAGARAEDITOR_API void SetTooltip(UNiagaraGraph* Graph, const FName& VarName, const FText& Tooltip);
 	}
 }
