Fix for Niagara categorizations being broken post migration when dealing with versioned data
#rnx

diff --git a/Engine/Plugins/FX/Niagara/Source/Niagara/Private/NiagaraScript.cpp b/Engine/Plugins/FX/Niagara/Source/Niagara/Private/NiagaraScript.cpp
--- a/Engine/Plugins/FX/Niagara/Source/Niagara/Private/NiagaraScript.cpp
+++ b/Engine/Plugins/FX/Niagara/Source/Niagara/Private/NiagaraScript.cpp
@@ -2354,6 +2354,9 @@ void UNiagaraScript::PostLoad()
 			UsageBitmask |= (1 << SimulationStageIndex);
 		}
 	}
+
+	const bool bRequestHierarchyRootMigration = bMigrateParameterDataToHierarchyRoot;
+	bMigrateParameterDataToHierarchyRoot = false;
 	
 	VersionedScriptAdapters.Reserve(VersionData.Num());
 	for (FVersionedNiagaraScriptData& Data : VersionData)
@@ -2363,10 +2366,9 @@ void UNiagaraScript::PostLoad()
 		{
 			Source->ConditionalPostLoad();
 
-			if (bMigrateParameterDataToHierarchyRoot)
+			if (bRequestHierarchyRootMigration)
 			{
 				Source->MigrateParameterDataToHierarchyRoot(Data);
-				bMigrateParameterDataToHierarchyRoot = false;
 			}
 
 			// Synchronize with Definitions after source scripts have been postloaded.
