Niagara - Fix an issue where the preview component can get stuck in a refresh loop when resimming in the system editor.
+ Update the system view model so that it handles the sequencer time in the same way that the FMovieSceneNiagaraSystemTrackTemplate does so that frames are handled consistently when the frame rate is locked.
+ Update the niagara component desired age handling to ignore small changes in time even if it's negative.

diff --git a/Engine/Plugins/FX/Niagara/Source/NiagaraEditor/Private/ViewModels/NiagaraSystemViewModel.cpp b/Engine/Plugins/FX/Niagara/Source/NiagaraEditor/Private/ViewModels/NiagaraSystemViewModel.cpp
--- a/Engine/Plugins/FX/Niagara/Source/NiagaraEditor/Private/ViewModels/NiagaraSystemViewModel.cpp
+++ b/Engine/Plugins/FX/Niagara/Source/NiagaraEditor/Private/ViewModels/NiagaraSystemViewModel.cpp
@@ -1911,6 +1911,7 @@ void FNiagaraSystemViewModel::SetupPreviewComponentAndInstance()
 		PreviewComponent->SetForceSolo(true);
 		PreviewComponent->SetAgeUpdateMode(ENiagaraAgeUpdateMode::DesiredAge);
 		PreviewComponent->SetSeekDelta((float)GetEditorData().GetPlaybackFrameRate().AsInterval());
+		PreviewComponent->SetLockDesiredAgeDeltaTimeToSeekDelta(GetEditorData().GetLockPlaybackFrameRate());
 		PreviewComponent->SetCanRenderWhileSeeking(false);
 		PreviewComponent->Activate(true);
 
@@ -2831,6 +2832,10 @@ void FNiagaraSystemViewModel::SequencerMovieSceneModified(const UMovieScene* Mov
 	if (SystemEditorData.GetLockPlaybackFrameRate() != bMovieSceneIsFrameLocked)
 	{
 		SystemEditorData.SetLockPlaybackFrameRate(bMovieSceneIsFrameLocked);
+		if (PreviewComponent != nullptr)
+		{
+			PreviewComponent->SetLockDesiredAgeDeltaTimeToSeekDelta(bMovieSceneIsFrameLocked);
+		}
 	}
 }
 
@@ -2967,13 +2972,20 @@ void FNiagaraSystemViewModel::SequencerTimeChanged()
 
 			if (bUpdateDesiredAge)
 			{
+				// NOTE: Logic copied from FMovieSceneNiagaraSystemTrackTemplate
+				// Add a quarter of a frame offset here to push the desired age into the middle of the frame since it will be automatically rounded
+				// down to the nearest seek delta.  This prevents a situation where float rounding results in a value which is just slightly less than
+				// the frame boundary, which results in a skipped simulation frame.
+				const float FrameOffset = PreviewComponent->GetLockDesiredAgeDeltaTimeToSeekDelta() ? (PreviewComponent->GetSeekDelta() / 4.0f) : 0.0f;
+				float DesiredAge = CurrentSequencerTime + FrameOffset;
+
 				if (CurrentStatus == EMovieScenePlayerStatus::Playing)
 				{
-					PreviewComponent->SetDesiredAge(FMath::Max(CurrentSequencerTime, 0.0f));
+					PreviewComponent->SetDesiredAge(FMath::Max(DesiredAge, 0.0f));
 				}
 				else
 				{
-					PreviewComponent->SeekToDesiredAge(FMath::Max(CurrentSequencerTime, 0.0f));
+					PreviewComponent->SeekToDesiredAge(FMath::Max(DesiredAge, 0.0f));
 				}
 			}
 
