Fix localization warning.
Removed diverging localization text.

#tests

diff --git a/Engine/Plugins/FX/Niagara/Source/Niagara/Private/DataInterface/NiagaraDataInterfaceDataChannelRead.cpp b/Engine/Plugins/FX/Niagara/Source/Niagara/Private/DataInterface/NiagaraDataInterfaceDataChannelRead.cpp
--- a/Engine/Plugins/FX/Niagara/Source/Niagara/Private/DataInterface/NiagaraDataInterfaceDataChannelRead.cpp
+++ b/Engine/Plugins/FX/Niagara/Source/Niagara/Private/DataInterface/NiagaraDataInterfaceDataChannelRead.cpp
@@ -252,7 +252,7 @@ namespace NDIDataChannelReadLocal
 			Sig.bRequiresExecPin = true;
 			Sig.ModuleUsageBitmask = ENiagaraScriptUsageMask::Emitter | ENiagaraScriptUsageMask::System;
 			Sig.AddInput(FNiagaraVariable(FNiagaraTypeDefinition(UNiagaraDataInterfaceDataChannelRead::StaticClass()), TEXT("DataChannel interface")));
-			Sig.AddInputWithDefault(FNiagaraVariable(FNiagaraTypeDefinition::GetBoolDef(), TEXT("Enable")), FNiagaraBool(true), LOCTEXT("SpawnEnableInputDesc", "Enable or disable this function call. If false, this call with have no effects."));
+			Sig.AddInputWithDefault(FNiagaraVariable(FNiagaraTypeDefinition::GetBoolDef(), TEXT("Enable")), FNiagaraBool(true), LOCTEXT("SpawnEnableInputDesc", "Enable or disable this function call. If false, this call with have no effect."));
 			Sig.AddInputWithoutDefault(FNiagaraVariable(FNiagaraTypeDefinition(FNiagaraEmitterID::StaticStruct()), TEXT("Emitter ID")), LOCTEXT("EmitterIDDesc", "ID of the emitter we'd like to spawn into. This can be obtained from Engine.Emitter.ID."));
 			Sig.AddInput(FNiagaraVariable(StaticEnum<ENDIDataChannelSpawnScaleMode>(), TEXT("Mode")), LOCTEXT("SpawnScaleModeInputDesc", "Control whether to override or combine this scale with previously set scales when calling this function multiple times."));
 			Sig.AddInputWithDefault(FNiagaraVariable(FNiagaraTypeDefinition::GetFloatDef(), TEXT("RandomScaleMin")), 1.0f, LOCTEXT("ScaleSpawnCountRandomScaleMinInputDesc", "Minimum value for a random additional scale applied to each NDC spawn."));
