Fixing issue with stale data in NDC write buffers.
When calling Write to direct indices but only writing at random, unintialized data written in previous frames can cause issues.

diff --git a/Engine/Plugins/FX/Niagara/Source/Niagara/Classes/NiagaraDataSet.h b/Engine/Plugins/FX/Niagara/Source/Niagara/Classes/NiagaraDataSet.h
--- a/Engine/Plugins/FX/Niagara/Source/Niagara/Classes/NiagaraDataSet.h
+++ b/Engine/Plugins/FX/Niagara/Source/Niagara/Classes/NiagaraDataSet.h
@@ -187,6 +187,8 @@ public:
 	FORCEINLINE void ClearGPUInstanceCount() { GPUInstanceCountBufferOffset = INDEX_NONE; }
 
 	NIAGARA_API void BuildRegisterTable();
+	
+	void ZeroCPUBuffers();
 
 #if NIAGARA_MEMORY_TRACKING
 	int32 GetAllocationSizeBytes() const { return AllocationSizeBytes; }
diff --git a/Engine/Plugins/FX/Niagara/Source/Niagara/Internal/DataInterface/NiagaraDataInterfaceDataChannelCommon.h b/Engine/Plugins/FX/Niagara/Source/Niagara/Internal/DataInterface/NiagaraDataInterfaceDataChannelCommon.h
--- a/Engine/Plugins/FX/Niagara/Source/Niagara/Internal/DataInterface/NiagaraDataInterfaceDataChannelCommon.h
+++ b/Engine/Plugins/FX/Niagara/Source/Niagara/Internal/DataInterface/NiagaraDataInterfaceDataChannelCommon.h
@@ -174,6 +174,7 @@ struct FNDIDataChannelCompiledData
 	bool UsedByGPU()const{ return bUsedByGPU; }
 	bool NeedSpawnDataTable()const { return bNeedsSpawnDataTable; }
 	bool SpawnsParticles()const { return bSpawnsParticles; }
+	bool CallsWriteFunction()const { return bCallsWrite; }
 	int32 GetTotalParams()const{ return TotalParams; }
 
 protected:
@@ -208,6 +209,10 @@ protected:
 	UPROPERTY()
 	bool bSpawnsParticles = false;
 
+	//If we call Write() on our CPU buffers we must do some extra buffer book keeping.
+	UPROPERTY()
+	bool bCallsWrite = false;
+
 	/** Iterates over all scripts for the owning system and gathers all functions and parameters accessing this DI. Building the FunctionInfoTable and GPUScriptParameterInfos map.  */
 	void GatherAccessInfo(UNiagaraSystem* System, UNiagaraDataInterface* Owner);
 };
@@ -240,6 +245,7 @@ namespace NDIDataChannelUtilities
 	extern const FName GetNDCSpawnDataName;
 	extern const FName SpawnConditionalName;
 	extern const FName SpawnDirectName;
+	extern const FName WriteName;
 
 	const TGlobalResource<FNDIDummyUAV>& GetDummyUAVFloat();
 	const TGlobalResource<FNDIDummyUAV>& GetDummyUAVInt32();
diff --git a/Engine/Plugins/FX/Niagara/Source/Niagara/Private/DataInterface/NiagaraDataInterfaceDataChannelCommon.cpp b/Engine/Plugins/FX/Niagara/Source/Niagara/Private/DataInterface/NiagaraDataInterfaceDataChannelCommon.cpp
--- a/Engine/Plugins/FX/Niagara/Source/Niagara/Private/DataInterface/NiagaraDataInterfaceDataChannelCommon.cpp
+++ b/Engine/Plugins/FX/Niagara/Source/Niagara/Private/DataInterface/NiagaraDataInterfaceDataChannelCommon.cpp
@@ -216,6 +216,11 @@ void FNDIDataChannelCompiledData::GatherAccessInfo(UNiagaraSystem* System, UNiag
 			bSpawnsParticles = true;			
 		}
 
+		if(BindingInfo.Name == NDIDataChannelUtilities::WriteName)
+		{
+			bCallsWrite = true;
+		}
+
 		return true;
 	};
 	FNiagaraDataInterfaceUtilities::ForEachVMFunction(Owner, System, HandleVMFunc);
@@ -293,6 +298,7 @@ namespace NDIDataChannelUtilities
 	const FName GetNDCSpawnDataName(TEXT("GetNDCSpawnData"));
 	const FName SpawnConditionalName(TEXT("SpawnConditional"));
 	const FName SpawnDirectName(TEXT("SpawnDirect"));
+	const FName WriteName(TEXT("Write"));
 	const TGlobalResource<FNDIDummyUAV> DummyUAVFloat(PF_R32_FLOAT, sizeof(float));
 	const TGlobalResource<FNDIDummyUAV> DummyUAVInt32(PF_R32_SINT, sizeof(int32));
 	const TGlobalResource<FNDIDummyUAV> DummyUAVHalf(PF_R16F, sizeof(FFloat16));
diff --git a/Engine/Plugins/FX/Niagara/Source/Niagara/Private/DataInterface/NiagaraDataInterfaceDataChannelWrite.cpp b/Engine/Plugins/FX/Niagara/Source/Niagara/Private/DataInterface/NiagaraDataInterfaceDataChannelWrite.cpp
--- a/Engine/Plugins/FX/Niagara/Source/Niagara/Private/DataInterface/NiagaraDataInterfaceDataChannelWrite.cpp
+++ b/Engine/Plugins/FX/Niagara/Source/Niagara/Private/DataInterface/NiagaraDataInterfaceDataChannelWrite.cpp
@@ -51,6 +51,14 @@ static FAutoConsoleVariableRef CVarDebugDumpWriterDI(
 	ECVF_Default
 );
 
+int32 GbNDCWriteDIZeroCPUBufferMode = 1;
+static FAutoConsoleVariableRef CVarNDCWriteDIZeroCPUBufferMode(
+	TEXT("fx.Niagara.DataChannels.WriteDIZeroCPUBuffersMode"),
+	GbNDCWriteDIZeroCPUBufferMode,
+	TEXT("Controls how CPU buffers are zeroed for the NDC Write DI\n0 = Do not Zero CPU buffers.\n1 = Zero only when calling \"Write\" function.\n2 = Zero always.\n"),
+	ECVF_Default
+);
+
 namespace NDIDataChannelWriteLocal
 {
 	static const TCHAR* CommonShaderFile = TEXT("/Plugin/FX/Niagara/Private/DataChannel/NiagaraDataInterfaceDataChannelCommon.ush");	
@@ -372,6 +380,15 @@ struct FNDIDataChannelWriteInstanceData
 						if(DestinationData && Interface->AllocationMode == ENiagaraDataChannelAllocationMode::Static)
 						{
 							DestinationData->Allocate(Interface->AllocationCount);
+							
+							//We choose whether to zero the CPU buffers or not.
+							//By default we only do this when we are calling "Write" (fx.Niagara.DataChannels.WriteDIZeroCPUBuffersMode==1) as calling Append should handle uninitialized buffers fine.
+							//However we can use fx.Niagara.DataChannels.WriteDIZeroCPUBuffersMode==2 to do this always or fx.Niagara.DataChannels.WriteDIZeroCPUBuffersMode==0 to never do it.
+							bool bZeroBuffers = GbNDCWriteDIZeroCPUBufferMode == 2 || (GbNDCWriteDIZeroCPUBufferMode == 1 && CompiledData.CallsWriteFunction());							
+							if(bZeroBuffers)
+							{
+								DestinationData->ZeroCPUBuffers();
+							}
 						}
 					}
 				}
diff --git a/Engine/Plugins/FX/Niagara/Source/Niagara/Private/NiagaraDataSet.cpp b/Engine/Plugins/FX/Niagara/Source/Niagara/Private/NiagaraDataSet.cpp
--- a/Engine/Plugins/FX/Niagara/Source/Niagara/Private/NiagaraDataSet.cpp
+++ b/Engine/Plugins/FX/Niagara/Source/Niagara/Private/NiagaraDataSet.cpp
@@ -665,6 +665,24 @@ bool FNiagaraDataBuffer::CheckForNaNs()const
 	return bContainsNaNs;
 }
 
+void FNiagaraDataBuffer::ZeroCPUBuffers()
+{
+	if(FloatData.GetData())
+	{
+		FMemory::Memzero(FloatData.GetData(), FloatData.Num());
+	}
+
+	if (Int32Data.GetData())
+	{
+		FMemory::Memzero(Int32Data.GetData(), Int32Data.Num());
+	}
+
+	if (HalfData.GetData())
+	{
+		FMemory::Memzero(HalfData.GetData(), HalfData.Num());
+	}
+}
+
 void FNiagaraDataBuffer::Allocate(uint32 InNumInstances, bool bMaintainExisting)
 {
 	NumInstances = 0;
