[Niagara]
* Added FNiagaraSceneProxyDesc
* Changed so FNiagaraSceneProxy is created using FNiagaraSceneProxyDesc

diff --git a/Engine/Plugins/FX/Niagara/Source/Niagara/Private/NiagaraComponent.cpp b/Engine/Plugins/FX/Niagara/Source/Niagara/Private/NiagaraComponent.cpp
--- a/Engine/Plugins/FX/Niagara/Source/Niagara/Private/NiagaraComponent.cpp
+++ b/Engine/Plugins/FX/Niagara/Source/Niagara/Private/NiagaraComponent.cpp
@@ -205,9 +205,25 @@ bool FitsIntoFloat(FVector InValue)
 			InValue.Z < UE_OLD_HALF_WORLD_MAX && InValue.Z > -UE_OLD_HALF_WORLD_MAX);
 }
 
-FNiagaraSceneProxy::FNiagaraSceneProxy(UNiagaraComponent* InComponent)
-	: FPrimitiveSceneProxy(InComponent, InComponent->GetAsset() ? InComponent->GetAsset()->GetFName() : FName())
-	, OcclusionQueryMode(InComponent->GetOcclusionQueryMode())
+FNiagaraSceneProxyDesc::FNiagaraSceneProxyDesc()
+{
+	SystemAsset = nullptr;
+	OcclusionQueryMode = ENiagaraOcclusionQueryMode::Default;
+}
+
+FNiagaraSceneProxyDesc::FNiagaraSceneProxyDesc(UNiagaraComponent& InComponent)
+:	FPrimitiveSceneProxyDesc(&InComponent)
+{
+	SystemAsset = InComponent.GetAsset();
+	OcclusionQueryMode = InComponent.GetOcclusionQueryMode();
+	SystemInstanceController = InComponent.GetSystemInstanceController();
+}
+
+FNiagaraSceneProxyDesc::~FNiagaraSceneProxyDesc() = default;
+
+FNiagaraSceneProxy::FNiagaraSceneProxy(const FNiagaraSceneProxyDesc& Desc)
+	: FPrimitiveSceneProxy(Desc, Desc.SystemAsset ? Desc.SystemAsset->GetFName() : NAME_None)
+	, OcclusionQueryMode(Desc.OcclusionQueryMode)
 {
 	SetWireframeColor(FLinearColor(3.0f, 0.0f, 0.0f));
 
@@ -217,8 +233,8 @@ FNiagaraSceneProxy::FNiagaraSceneProxy(UNiagaraComponent* InComponent)
 	// Optionally enable GDME for Niagara, any caveats are covered below
 	bSupportsParallelGDME = GNiagaraParallelGDME;
 
-	FNiagaraSystemInstanceControllerConstPtr SystemInstanceController = InComponent->GetSystemInstanceController();
-	UNiagaraSystem* NiagaraSystem = InComponent->GetAsset();
+	FNiagaraSystemInstanceControllerConstPtr SystemInstanceController = Desc.SystemInstanceController;
+	UNiagaraSystem* NiagaraSystem = Desc.SystemAsset;
 	if (SystemInstanceController && NiagaraSystem)
 	{
 		LLM_SCOPE(ELLMTag::Niagara);
@@ -2371,7 +2387,7 @@ FPrimitiveSceneProxy* UNiagaraComponent::CreateSceneProxy()
 #endif // UE_WITH_PSO_PRECACHING
 
 	// The constructor will set up the System renderers from the component.
-	FNiagaraSceneProxy* Proxy = new FNiagaraSceneProxy(this);
+	FNiagaraSceneProxy* Proxy = new FNiagaraSceneProxy(FNiagaraSceneProxyDesc(*this));
 	return Proxy;
 }
 
diff --git a/Engine/Plugins/FX/Niagara/Source/Niagara/Public/NiagaraSceneProxy.h b/Engine/Plugins/FX/Niagara/Source/Niagara/Public/NiagaraSceneProxy.h
--- a/Engine/Plugins/FX/Niagara/Source/Niagara/Public/NiagaraSceneProxy.h
+++ b/Engine/Plugins/FX/Niagara/Source/Niagara/Public/NiagaraSceneProxy.h
@@ -5,14 +5,29 @@
 #include "Particles/ParticlePerfStats.h"
 #include "PrimitiveSceneProxy.h"
 #include "HeterogeneousVolumeInterface.h"
+#include "PrimitiveSceneProxyDesc.h"
 
 class FNiagaraGpuComputeDispatchInterface;
 class FNiagaraSystemRenderData;
 class UNiagaraComponent;
+class UNiagaraSystem;
+class FNiagaraSystemInstanceController;
 enum class ENiagaraOcclusionQueryMode : uint8;
+using FNiagaraSystemInstanceControllerConstPtr = TSharedPtr<const FNiagaraSystemInstanceController, ESPMode::ThreadSafe>;
 
 #define NIAGARAPROXY_EVENTS_ENABLED	(!STATS && (ENABLE_STATNAMEDEVENTS || CPUPROFILERTRACE_ENABLED))
 
+struct FNiagaraSceneProxyDesc : public FPrimitiveSceneProxyDesc
+{
+	NIAGARA_API FNiagaraSceneProxyDesc();
+	NIAGARA_API FNiagaraSceneProxyDesc(UNiagaraComponent& InComponent);
+	NIAGARA_API ~FNiagaraSceneProxyDesc();
+
+	UNiagaraSystem* SystemAsset;
+	ENiagaraOcclusionQueryMode OcclusionQueryMode;
+	FNiagaraSystemInstanceControllerConstPtr SystemInstanceController;
+};
+
 /**
 * Scene proxy for drawing niagara particle simulations.
 */
@@ -21,7 +36,7 @@ class FNiagaraSceneProxy : public FPrimitiveSceneProxy
 public:
 	NIAGARA_API virtual SIZE_T GetTypeHash() const override;
 
-	NIAGARA_API FNiagaraSceneProxy(UNiagaraComponent* InComponent);
+	NIAGARA_API FNiagaraSceneProxy(const FNiagaraSceneProxyDesc& Desc);
 	NIAGARA_API ~FNiagaraSceneProxy();
 
 	/** Retrieves the render data for a single system */
@@ -63,6 +78,9 @@ public:
 	const FDynamicData& GetProxyDynamicData()const { return DynamicData; }
 	void SetProxyDynamicData(const FDynamicData& NewData) { DynamicData = NewData; }
 
+	UE_DEPRECATED(5.6, "Use FNiagaraSceneProxyDesc ctor instead")
+	inline FNiagaraSceneProxy(UNiagaraComponent* InComponent) : FNiagaraSceneProxy(FNiagaraSceneProxyDesc(*InComponent)) {}
+
 private:
 	NIAGARA_API void ReleaseRenderThreadResources();
 
