Adjust component tag that is used to mark the Niagara component being used with the NiagaraTimedNotify to be unique (using the UniqueID).  Avoids a conflict where notifies from other animations could cancel the incorrect notify.
#tests Used test project to verify that timed notifies continued to work and overlapped animations wouldn't conflict

diff --git a/Engine/Plugins/FX/Niagara/Source/NiagaraAnimNotifies/Private/AnimNotifyState_TimedNiagaraEffect.cpp b/Engine/Plugins/FX/Niagara/Source/NiagaraAnimNotifies/Private/AnimNotifyState_TimedNiagaraEffect.cpp
--- a/Engine/Plugins/FX/Niagara/Source/NiagaraAnimNotifies/Private/AnimNotifyState_TimedNiagaraEffect.cpp
+++ b/Engine/Plugins/FX/Niagara/Source/NiagaraAnimNotifies/Private/AnimNotifyState_TimedNiagaraEffect.cpp
@@ -119,6 +119,15 @@ bool UAnimNotifyState_TimedNiagaraEffect::ValidateParameters(USkeletalMeshCompon
 	return bValid;
 }
 
+FName UAnimNotifyState_TimedNiagaraEffect::GetSpawnedComponentTag() const
+{
+	// we generate a unique tag to associate with our spawned components so that we can clean things up upon completion
+	FName NotifyName = GetFName();
+	NotifyName.SetNumber(GetUniqueID());
+
+	return NotifyName;
+}
+
 FString UAnimNotifyState_TimedNiagaraEffect::GetNotifyName_Implementation() const
 {
 	if (Template)
diff --git a/Engine/Plugins/FX/Niagara/Source/NiagaraAnimNotifies/Public/AnimNotifyState_TimedNiagaraEffect.h b/Engine/Plugins/FX/Niagara/Source/NiagaraAnimNotifies/Public/AnimNotifyState_TimedNiagaraEffect.h
--- a/Engine/Plugins/FX/Niagara/Source/NiagaraAnimNotifies/Public/AnimNotifyState_TimedNiagaraEffect.h
+++ b/Engine/Plugins/FX/Niagara/Source/NiagaraAnimNotifies/Public/AnimNotifyState_TimedNiagaraEffect.h
@@ -68,7 +68,7 @@ protected:
 
 	NIAGARAANIMNOTIFIES_API bool ValidateParameters(USkeletalMeshComponent* MeshComp) const;
 
-	FORCEINLINE FName GetSpawnedComponentTag()const { return GetFName(); }
+	FName GetSpawnedComponentTag() const;
 };
 
 USTRUCT(BlueprintType)
