- Remove HV related code from mesh renderer as the correct path is to use the volume renderer.
- Added a stack warning when using MD_Volume with mesh renderers as it doesn't work correctly, it's a bbox for the component not all the individual meshes that are rendered into it.

#rnx
#tests Create mesh renderer that uses volume domain material, then fix the issue

diff --git a/Engine/Plugins/FX/Niagara/Source/Niagara/Private/NiagaraMeshRendererProperties.cpp b/Engine/Plugins/FX/Niagara/Source/Niagara/Private/NiagaraMeshRendererProperties.cpp
--- a/Engine/Plugins/FX/Niagara/Source/Niagara/Private/NiagaraMeshRendererProperties.cpp
+++ b/Engine/Plugins/FX/Niagara/Source/Niagara/Private/NiagaraMeshRendererProperties.cpp
@@ -323,7 +323,6 @@ bool FNiagaraMeshRendererMeshProperties::HasValidRenderableMesh() const
 
 UNiagaraMeshRendererProperties::UNiagaraMeshRendererProperties()
 	: bOverrideMaterials(false)
-	, bUseHeterogeneousVolumes(false)
 	, bSortOnlyWhenTranslucent(true)
 	, bSubImageBlend(true)
 	, bLockedAxisEnable(false)
@@ -895,20 +894,6 @@ bool UNiagaraMeshRendererProperties::PopulateRequiredBindings(FNiagaraParameterS
 {
 	bool bAnyAdded = Super::PopulateRequiredBindings(InParameterStore);
 
-
-	if (bUseHeterogeneousVolumes)
-	{
-		InParameterStore.AddParameter(
-			FNiagaraVariableBase(FNiagaraTypeDefinition::GetIntDef(), FName("User.ResolutionMaxAxis")),
-			false
-		);
-		InParameterStore.AddParameter(
-			FNiagaraVariableBase(FNiagaraTypeDefinition::GetVec3Def(), FName("User.WorldSpaceSize")),
-			false
-		);
-		bAnyAdded = true;
-	}
-
 	for (const FNiagaraVariableAttributeBinding* Binding : AttributeBindings)
 	{
 		if (Binding && Binding->CanBindToHostParameterMap())
@@ -1126,7 +1111,7 @@ void UNiagaraMeshRendererProperties::GetRendererFeedback(const FVersionedNiagara
 {
 	Super::GetRendererFeedback(InEmitter, OutErrors, OutWarnings, OutInfo);
 
-	GetMaterialUsageFeedback(MATUSAGE_NiagaraMeshParticles, OutWarnings);
+	GetMaterialUsageFeedback(MATUSAGE_NiagaraMeshParticles, { MD_Volume }, OutWarnings);
 
 	if (MaterialParameters.HasAnyBindings())
 	{
@@ -1370,14 +1355,7 @@ void UNiagaraMeshRendererProperties::CheckMaterialUsage()
 				if (MaterialInterface)
 				{
 					const UMaterial* Material = MaterialInterface->GetMaterial();
-					if (UseHeterogeneousVolumes() && Material && Material->MaterialDomain == MD_Volume)
-					{
-						MaterialInterface->CheckMaterialUsage(MATUSAGE_HeterogeneousVolumes);
-					}
-					else
-					{
-						MaterialInterface->CheckMaterialUsage(MATUSAGE_NiagaraMeshParticles);
-					}
+					MaterialInterface->CheckMaterialUsage(MATUSAGE_NiagaraMeshParticles);
 				}
 			}
 		}
diff --git a/Engine/Plugins/FX/Niagara/Source/Niagara/Private/NiagaraRendererMeshes.cpp b/Engine/Plugins/FX/Niagara/Source/Niagara/Private/NiagaraRendererMeshes.cpp
--- a/Engine/Plugins/FX/Niagara/Source/Niagara/Private/NiagaraRendererMeshes.cpp
+++ b/Engine/Plugins/FX/Niagara/Source/Niagara/Private/NiagaraRendererMeshes.cpp
@@ -120,19 +120,8 @@ FNiagaraRendererMeshes::FNiagaraRendererMeshes(ERHIFeatureLevel::Type FeatureLev
 	DistanceCullRangeSquared = FVector2f(0, FLT_MAX);
 	RendererVisibility = Properties->RendererVisibility;
 	bAccurateMotionVectors = Properties->NeedsPreciseMotionVectors();
-	bIsHeterogeneousVolume = Properties->UseHeterogeneousVolumes();
 	bCastShadows = Properties->bCastShadows;
 
-	if (Properties->UseHeterogeneousVolumes())
-	{
-		ResolutionMaxAxisOffset = Emitter->GetRendererBoundVariables().IndexOf(
-			FNiagaraVariableBase(FNiagaraTypeDefinition::GetIntDef(), FName("User.ResolutionMaxAxis"))
-		);
-		WorldSpaceSizeOffset = Emitter->GetRendererBoundVariables().IndexOf(
-			FNiagaraVariableBase(FNiagaraTypeDefinition::GetVec3Def(), FName("User.WorldSpaceSize"))
-		);
-	}
-
 	if (Properties->bEnableCameraDistanceCulling)
 	{
 		DistanceCullRange = FVector2f(Properties->MinCameraDistance, Properties->MaxCameraDistance);
@@ -381,15 +370,6 @@ void FNiagaraRendererMeshes::PrepareParticleMeshRenderData(FRHICommandListBase&
 		return;
 	}
 
-	if (ResolutionMaxAxisOffset != INDEX_NONE)
-	{
-		ParticleMeshRenderData.ResolutionMaxAxis = *reinterpret_cast<const int32*>(ParticleMeshRenderData.DynamicDataMesh->ParameterDataBound.GetData() + ResolutionMaxAxisOffset);
-	}
-	if (WorldSpaceSizeOffset != INDEX_NONE)
-	{
-		ParticleMeshRenderData.WorldSpaceSize = *reinterpret_cast<const FVector3f*>(ParticleMeshRenderData.DynamicDataMesh->ParameterDataBound.GetData() + WorldSpaceSizeOffset);
-	}
-
 	// If the visibility tag comes from a parameter map, so we can evaluate it here and just early out if it doesn't match up
 	if (EmitterRendererVisTagOffset != INDEX_NONE && ParticleMeshRenderData.DynamicDataMesh->ParameterDataBound.IsValidIndex(EmitterRendererVisTagOffset))
 	{
@@ -1592,25 +1572,6 @@ void FNiagaraRendererMeshes::GetDynamicMeshElements(const TArray<const FSceneVie
 							bNeedsPrevTransform
 						);
 
-						if (SceneProxy->IsHeterogeneousVolume())
-						{
-							FHeterogeneousVolumeData* HeterogeneousVolumeData = &Collector.AllocateOneFrameResource<FHeterogeneousVolumeData>(SceneProxy);
-
-							FVector3f WorldSpaceSize = ParticleMeshRenderData.WorldSpaceSize;
-							float WorldSpaceSizeMaxInv = WorldSpaceSize.GetMax() > 0.0 ? 1.0 / WorldSpaceSize.GetMax() : 0.0;
-							FVector3f ResolutionFactor = WorldSpaceSize * WorldSpaceSizeMaxInv;
-
-							FVector3f VolumeResolutionV3f = ResolutionFactor * ParticleMeshRenderData.ResolutionMaxAxis;
-
-							FIntVector VolumeResolution = FIntVector(
-								FMath::CeilToInt(VolumeResolutionV3f.X),
-								FMath::CeilToInt(VolumeResolutionV3f.Y),
-								FMath::CeilToInt(VolumeResolutionV3f.Z));
-
-							HeterogeneousVolumeData->VoxelResolution = VolumeResolution;
-							MeshBatch.Elements[0].UserData = HeterogeneousVolumeData;
-						}
-
 						Collector.AddMesh(ViewIndex, MeshBatch);
 					}
 				}
@@ -1737,10 +1698,6 @@ void FNiagaraRendererMeshes::GetDynamicRayTracingInstances(FRayTracingInstanceCo
 
 			const FMaterialRenderProxy* ConstMaterialProxy = ParticleMeshRenderData.DynamicDataMesh->Materials[RemappedMaterialIndex];
 			const FMaterial& Material = ConstMaterialProxy->GetMaterialWithFallback(FeatureLevel, ConstMaterialProxy);
-			if (Material.IsUsedWithHeterogeneousVolumes())
-			{
-				continue;
-			}
 
 			FMeshBatch MeshBatch;
 			CreateMeshBatchForSection(
@@ -1932,9 +1889,7 @@ FNiagaraDynamicDataBase* FNiagaraRendererMeshes::GenerateDynamicData(const FNiag
 		//In preparation for a material override feature, we pass our material(s) and relevance in via dynamic data.
 		//The renderer ensures we have the correct usage and relevance for materials in BaseMaterials_GT.
 		//Any override feature must also do the same for materials that are set.
-		check( (Properties->UseHeterogeneousVolumes() && Mat->GetMaterial() && (Mat->GetMaterial()->MaterialDomain == MD_Volume)) ?
-			Mat->CheckMaterialUsage_Concurrent(MATUSAGE_HeterogeneousVolumes) :
-			Mat->CheckMaterialUsage_Concurrent(MATUSAGE_NiagaraMeshParticles));
+		check(Mat->CheckMaterialUsage_Concurrent(MATUSAGE_NiagaraMeshParticles));
 		DynamicData->Materials.Add(Mat->GetRenderProxy());
 	}
 
@@ -1965,9 +1920,7 @@ bool FNiagaraRendererMeshes::IsMaterialValid(const UMaterialInterface* Mat)const
 	bool bIsMaterialValid = false;
 	if (Mat)
 	{
-		bIsMaterialValid = (bIsHeterogeneousVolume && Mat->GetMaterial_Concurrent() && (Mat->GetMaterial_Concurrent()->MaterialDomain == MD_Volume)) ?
-			Mat->CheckMaterialUsage_Concurrent(MATUSAGE_HeterogeneousVolumes) :
-			Mat->CheckMaterialUsage_Concurrent(MATUSAGE_NiagaraMeshParticles);
+		bIsMaterialValid = Mat->CheckMaterialUsage_Concurrent(MATUSAGE_NiagaraMeshParticles);
 	}
 
 	return bIsMaterialValid;
diff --git a/Engine/Plugins/FX/Niagara/Source/Niagara/Private/NiagaraRendererProperties.cpp b/Engine/Plugins/FX/Niagara/Source/Niagara/Private/NiagaraRendererProperties.cpp
--- a/Engine/Plugins/FX/Niagara/Source/Niagara/Private/NiagaraRendererProperties.cpp
+++ b/Engine/Plugins/FX/Niagara/Source/Niagara/Private/NiagaraRendererProperties.cpp
@@ -15,6 +15,7 @@
 #include "Stateless/Modules/NiagaraStatelessModule_DynamicMaterialParameters.h"
 
 #include "Interfaces/ITargetPlatform.h"
+#include "Materials/Material.h"
 #include "Materials/MaterialInterface.h"
 #include "Materials/MaterialInstanceConstant.h"
 #include "SceneInterface.h"
@@ -753,31 +754,72 @@ void UNiagaraRendererProperties::GetRendererFeedback(const FVersionedNiagaraEmit
 	}
 }
 
-void UNiagaraRendererProperties::GetMaterialUsageFeedback(EMaterialUsage Usage, TArray<FNiagaraRendererFeedback>& OutFeedback) const
+void UNiagaraRendererProperties::GetMaterialUsageFeedback(EMaterialUsage Usage, TConstArrayView<EMaterialDomain> InvalidMaterialDomains, TArray<FNiagaraRendererFeedback>& OutFeedback) const
 {
 	TArray<UMaterialInterface*> Materials;
 	GetUsedMaterials(nullptr, Materials);
 
-	FString FailedMaterialNames;
+	auto AppendToString =
+		[](FString& InString, const FString& InItem)
+		{
+			if ( InString.Len() > 0 )
+			{
+				InString.Append(TEXT(", "));
+			}
+			InString.Append(InItem);
+		};
+
+	FString FailedUsageNames;
+	FString FailedDomainNames;
 	for (UMaterialInterface* Material : Materials)
 	{
-		if (Material && !Material->CheckMaterialUsage(Usage))
+		if (!Material)
+		{
+			continue;
+		}
+
+		if (!Material->CheckMaterialUsage(Usage))
+		{
+			AppendToString(FailedUsageNames, Material->GetName());
+		}
+
+		if ( UMaterial* BaseMaterial = Material->GetMaterial() )
 		{
-			if (FailedMaterialNames.Len() > 0)
+			if (InvalidMaterialDomains.Contains(BaseMaterial->MaterialDomain))
 			{
-				FailedMaterialNames.Append(TEXT(", "));
+				AppendToString(FailedDomainNames, Material->GetName());
 			}
-			FailedMaterialNames.Append(Material->GetName());
 		}
 	}
 
-	if (FailedMaterialNames.Len() > 0)
+	if (FailedUsageNames.Len() > 0)
 	{
 		OutFeedback.Emplace(
-			FText::Format(LOCTEXT("InvalidMaterialUsage_Desc", "Some materials '{0}' do not have the correct usage flags set, and will use the default material."), FText::FromString(FailedMaterialNames)),
+			FText::Format(LOCTEXT("InvalidMaterialUsage_Desc", "Some materials '{0}' do not have the correct usage flags set, and will use the default material."), FText::FromString(FailedUsageNames)),
 			LOCTEXT("InvalidMaterialUsage_Summary", "Materials might not render correctly.")
 		);
 	}
+
+	if (FailedDomainNames.Len() > 0)
+	{
+		const UEnum* MaterialDomainEnum = StaticEnum<EMaterialDomain>();
+
+		FString InvalidMaterialDomainsString;
+		for (EMaterialDomain MaterialDomain : InvalidMaterialDomains)
+		{
+			AppendToString(InvalidMaterialDomainsString, MaterialDomainEnum->GetDisplayNameTextByIndex(MaterialDomain).ToString());
+		}
+
+		OutFeedback.Emplace(
+			FText::Format(LOCTEXT("InvalidMaterialDomain_Desc", "Some materials '{0}' use material domains '{1}' which are not supported so may not render correctly."), FText::FromString(FailedDomainNames), FText::FromString(InvalidMaterialDomainsString)),
+			LOCTEXT("InvalidMaterialDomain_Summary", "Materials use unsupported material domain.")
+		);
+	}
+}
+
+void UNiagaraRendererProperties::GetMaterialUsageFeedback(EMaterialUsage Usage, TArray<FNiagaraRendererFeedback>& OutFeedback) const
+{
+	GetMaterialUsageFeedback(Usage, MakeArrayView<EMaterialDomain>(nullptr, 0), OutFeedback);
 }
 
 const FSlateBrush* UNiagaraRendererProperties::GetStackIcon() const
diff --git a/Engine/Plugins/FX/Niagara/Source/Niagara/Public/NiagaraMeshRendererProperties.h b/Engine/Plugins/FX/Niagara/Source/Niagara/Public/NiagaraMeshRendererProperties.h
--- a/Engine/Plugins/FX/Niagara/Source/Niagara/Public/NiagaraMeshRendererProperties.h
+++ b/Engine/Plugins/FX/Niagara/Source/Niagara/Public/NiagaraMeshRendererProperties.h
@@ -289,7 +289,6 @@ public:
 	NIAGARA_API virtual TArray<FNiagaraVariable> GetBoundAttributes() const override;
 #endif // WITH_EDITORONLY_DATA
 	NIAGARA_API virtual void CacheFromCompiledData(const FNiagaraDataSetCompiledData* CompiledData) override;
-	virtual bool UseHeterogeneousVolumes() const override { return bUseHeterogeneousVolumes; }
 
 	NIAGARA_API void UpdateMICs();
 
@@ -332,10 +331,6 @@ public:
 	UPROPERTY(EditAnywhere, Category = "Mesh Rendering", DisplayName="Enable Material Overrides")
 	uint32 bOverrideMaterials : 1;
 
-	/** Whether or not to render with heterogeneous volumes.*/
-	UPROPERTY(EditAnywhere, Category = "Rendering")
-	uint32 bUseHeterogeneousVolumes : 1;
-
 	/** If true, the particles are only sorted when using a translucent material. */
 	UPROPERTY(EditAnywhere, Category = "Sorting", meta = (EditCondition = "SortMode != ENiagaraSortMode::None", EditConditionHides))
 	uint32 bSortOnlyWhenTranslucent : 1;
diff --git a/Engine/Plugins/FX/Niagara/Source/Niagara/Public/NiagaraRendererMeshes.h b/Engine/Plugins/FX/Niagara/Source/Niagara/Public/NiagaraRendererMeshes.h
--- a/Engine/Plugins/FX/Niagara/Source/Niagara/Public/NiagaraRendererMeshes.h
+++ b/Engine/Plugins/FX/Niagara/Source/Niagara/Public/NiagaraRendererMeshes.h
@@ -187,7 +187,6 @@ private:
 	uint32 bEnableFrustumCulling : 1;
 	uint32 bEnableLODCulling : 1;
 	uint32 bAccurateMotionVectors : 1;
-	uint32 bIsHeterogeneousVolume : 1;
 	uint32 bCastShadows : 1;
 
 	uint32 bSubImageBlend : 1;
@@ -207,9 +206,6 @@ private:
 	int32 EmitterMeshIndexOffset = INDEX_NONE;
 	uint32 MaterialParamValidMask = 0;
 
-	int32 ResolutionMaxAxisOffset = INDEX_NONE;
-	int32 WorldSpaceSizeOffset = INDEX_NONE;
-
 	int32 VFBoundOffsetsInParamStore[ENiagaraMeshVFLayout::Type::Num_Max];
 	uint32 bSetAnyBoundVars : 1;
 
diff --git a/Engine/Plugins/FX/Niagara/Source/Niagara/Public/NiagaraRendererProperties.h b/Engine/Plugins/FX/Niagara/Source/Niagara/Public/NiagaraRendererProperties.h
--- a/Engine/Plugins/FX/Niagara/Source/Niagara/Public/NiagaraRendererProperties.h
+++ b/Engine/Plugins/FX/Niagara/Source/Niagara/Public/NiagaraRendererProperties.h
@@ -28,6 +28,7 @@ struct FSlateBrush;
 struct FStaticParameterSet;
 struct FStreamingRenderAssetPrimitiveInfo;
 enum EMaterialUsage : int;
+enum EMaterialDomain : int;
 
 #if WITH_EDITOR
 // Helper class for GUI error handling
@@ -372,6 +373,7 @@ public:
 	virtual void GetRendererFeedback(const FVersionedNiagaraEmitter& InEmitter, TArray<FText>& OutErrors, TArray<FText>& OutWarnings, TArray<FText>& OutInfo) const {}
 	NIAGARA_API virtual void GetRendererFeedback(const FVersionedNiagaraEmitter& InEmitter, TArray<FNiagaraRendererFeedback>& OutErrors, TArray<FNiagaraRendererFeedback>& OutWarnings, TArray<FNiagaraRendererFeedback>& OutInfo) const;
 
+	void GetMaterialUsageFeedback(EMaterialUsage Usage, TConstArrayView<EMaterialDomain> InvalidMaterialDomains, TArray<FNiagaraRendererFeedback>& OutFeedback) const;
 	void GetMaterialUsageFeedback(EMaterialUsage Usage, TArray<FNiagaraRendererFeedback>& OutFeedback) const;
 
 	// The icon to display in the niagara stack widget under the renderer section
