[NiagaraAnimNotifies]
* added dllapi to UAnimNotifyState_TimedNiagaraEffect::GetSpawnedComponentTag to fix farm link errors

diff --git a/Engine/Plugins/FX/Niagara/Source/NiagaraAnimNotifies/Public/AnimNotifyState_TimedNiagaraEffect.h b/Engine/Plugins/FX/Niagara/Source/NiagaraAnimNotifies/Public/AnimNotifyState_TimedNiagaraEffect.h
--- a/Engine/Plugins/FX/Niagara/Source/NiagaraAnimNotifies/Public/AnimNotifyState_TimedNiagaraEffect.h
+++ b/Engine/Plugins/FX/Niagara/Source/NiagaraAnimNotifies/Public/AnimNotifyState_TimedNiagaraEffect.h
@@ -68,7 +68,7 @@ protected:
 
 	NIAGARAANIMNOTIFIES_API bool ValidateParameters(USkeletalMeshComponent* MeshComp) const;
 
-	FName GetSpawnedComponentTag() const;
+	NIAGARAANIMNOTIFIES_API FName GetSpawnedComponentTag() const;
 };
 
 USTRUCT(BlueprintType)
