Fix up some shader compiler warnings.  Updated code also seems to generate better code on some platforms (d3d).  Don't claim to understand why the compiler couldn't figure it out before, but this is more consistent pattern throughout the file.
#rnx
#tests Generated shaders for a number of different materials and evaluated that they generated identically behaving code

diff --git a/Engine/Plugins/FX/Niagara/Shaders/Private/NiagaraParticleAccess.ush b/Engine/Plugins/FX/Niagara/Shaders/Private/NiagaraParticleAccess.ush
--- a/Engine/Plugins/FX/Niagara/Shaders/Private/NiagaraParticleAccess.ush
+++ b/Engine/Plugins/FX/Niagara/Shaders/Private/NiagaraParticleAccess.ush
@@ -142,8 +142,13 @@ float NiagaraSafeGetFloat(int RegisterIdx, uint ParticleIdx, float Default)
 	// NOTE: Branch to prevent unnecessary loads
 	BRANCH
 	if (RegisterIdx == -1)
+	{
 		return Default;
-	return NiagaraGetFloat(RegisterIdx, ParticleIdx);
+	}
+	else
+	{
+		return NiagaraGetFloat(RegisterIdx, ParticleIdx);
+	}
 }
 
 float2 NiagaraSafeGetVec2(int RegisterIdx, uint ParticleIdx, float2 Default)
@@ -151,8 +156,13 @@ float2 NiagaraSafeGetVec2(int RegisterIdx, uint ParticleIdx, float2 Default)
 	// NOTE: Branch to prevent unnecessary loads
 	BRANCH
 	if (RegisterIdx == -1)
+	{
 		return Default;
-	return NiagaraGetVec2(RegisterIdx, ParticleIdx);
+	}
+	else
+	{
+		return NiagaraGetVec2(RegisterIdx, ParticleIdx);
+	}
 }
 
 float3 NiagaraSafeGetVec3(int RegisterIdx, uint ParticleIdx, float3 Default)
@@ -160,8 +170,13 @@ float3 NiagaraSafeGetVec3(int RegisterIdx, uint ParticleIdx, float3 Default)
 	// NOTE: Branch to prevent unnecessary loads
 	BRANCH
 	if (RegisterIdx == -1)
+	{
 		return Default;
-	return NiagaraGetVec3(RegisterIdx, ParticleIdx);
+	}
+	else
+	{
+		return NiagaraGetVec3(RegisterIdx, ParticleIdx);
+	}
 }
 
 float4 NiagaraSafeGetVec4(int RegisterIdx, uint ParticleIdx, float4 Default)
@@ -169,8 +184,13 @@ float4 NiagaraSafeGetVec4(int RegisterIdx, uint ParticleIdx, float4 Default)
 	// NOTE: Branch to prevent unnecessary loads
 	BRANCH
 	if (RegisterIdx == -1)
+	{
 		return Default;
-	return NiagaraGetVec4(RegisterIdx, ParticleIdx);
+	}
+	else
+	{
+		return NiagaraGetVec4(RegisterIdx, ParticleIdx);
+	}
 }
 
 #if ENABLE_NIAGARA_INT_DATA_ACCESS
@@ -205,32 +225,52 @@ int NiagaraSafeGetInt(int RegisterIdx, uint ParticleIdx, int Default)
 {
 	BRANCH
 	if (RegisterIdx == -1)
+	{
 		return Default;
-	return NiagaraGetInt(RegisterIdx, ParticleIdx);
+	}
+	else
+	{
+		return NiagaraGetInt(RegisterIdx, ParticleIdx);
+	}
 }
 
 int2 NiagaraSafeGetInt2(int RegisterIdx, uint ParticleIdx, int2 Default)
 {
 	BRANCH
 	if (RegisterIdx == -1)
+	{
 		return Default;
-	return NiagaraGetInt2(RegisterIdx, ParticleIdx);
+	}
+	else
+	{
+		return NiagaraGetInt2(RegisterIdx, ParticleIdx);
+	}
 }
 
 int3 NiagaraSafeGetInt3(int RegisterIdx, uint ParticleIdx, int3 Default)
 {
 	BRANCH
 	if (RegisterIdx == -1)
+	{
 		return Default;
-	return NiagaraGetInt3(RegisterIdx, ParticleIdx);
+	}
+	else
+	{
+		return NiagaraGetInt3(RegisterIdx, ParticleIdx);
+	}
 }
 
 int4 NiagaraSafeGetInt4(int RegisterIdx, uint ParticleIdx, int4 Default)
 {
 	BRANCH
 	if (RegisterIdx == -1)
+	{
 		return Default;
-	return NiagaraGetInt4(RegisterIdx, ParticleIdx);
+	}
+	else
+	{
+		return NiagaraGetInt4(RegisterIdx, ParticleIdx);
+	}
 }
 
 #endif // ENABLE_NIAGARA_INT_DATA_ACCESS
\ No newline at end of file
