- Add owner actor class to the debugger when in Verbose mode

#rnx
#tests Made a BP with a NS inside to confirm it displays the information, also placed a NS in the level

diff --git a/Engine/Plugins/FX/Niagara/Source/Niagara/Private/NiagaraDebugHud.cpp b/Engine/Plugins/FX/Niagara/Source/Niagara/Private/NiagaraDebugHud.cpp
--- a/Engine/Plugins/FX/Niagara/Source/Niagara/Private/NiagaraDebugHud.cpp
+++ b/Engine/Plugins/FX/Niagara/Source/Niagara/Private/NiagaraDebugHud.cpp
@@ -2979,6 +2979,10 @@ void FNiagaraDebugHud::DrawComponents(FNiagaraWorldManager* WorldManager, UCanva
 					StringBuilder.Append(TEXT("/"));
 				}
 				StringBuilder.Append(*GetNameSafe(NiagaraComponent));
+				if (OwnerActor && Settings.SystemDebugVerbosity == ENiagaraDebugHudVerbosity::Verbose)
+				{
+					StringBuilder.Appendf(TEXT(" OwnerClass(%s)"), *OwnerActor->GetClass()->GetName());
+				}
 				StringBuilder.Append(TEXT("\n"));
 
 				if (Settings.SystemDebugVerbosity == ENiagaraDebugHudVerbosity::Verbose)
