Fix NiagaraEmitter asset context menu showing options that do not apply to read-only assets
#rnx

diff --git a/Engine/Plugins/FX/Niagara/Source/NiagaraEditor/Private/AssetDefinitions/AssetDefinition_NiagaraEmitter.cpp b/Engine/Plugins/FX/Niagara/Source/NiagaraEditor/Private/AssetDefinitions/AssetDefinition_NiagaraEmitter.cpp
--- a/Engine/Plugins/FX/Niagara/Source/NiagaraEditor/Private/AssetDefinitions/AssetDefinition_NiagaraEmitter.cpp
+++ b/Engine/Plugins/FX/Niagara/Source/NiagaraEditor/Private/AssetDefinitions/AssetDefinition_NiagaraEmitter.cpp
@@ -144,32 +144,36 @@ namespace MenuExtension_NiagaraEmitter
 			FToolMenuSection& Section = Menu->FindOrAddSection("GetAssetActions");
 			Section.AddDynamicEntry(NAME_None, FNewToolMenuSectionDelegate::CreateLambda([](FToolMenuSection& InSection)
 			{
+				const UContentBrowserAssetContextMenuContext* Context = InSection.FindContext<UContentBrowserAssetContextMenuContext>();
+				if (Context && Context->bCanBeModified)
 				{
-					const TAttribute<FText> Label = LOCTEXT("Emitter_NewNiagaraSystem", "Create Niagara System");
-					const TAttribute<FText> ToolTip = LOCTEXT("Emitter_NewNiagaraSystemTooltip", "Creates a niagara system using this emitter as a base.");
-					const FSlateIcon Icon = FSlateIcon(FAppStyle::GetAppStyleSetName(), "ClassIcon.ParticleSystem");
-
-					FToolUIAction UIAction;
-					UIAction.ExecuteAction = FToolMenuExecuteAction::CreateStatic(&ExecuteNewNiagaraSystem);
-					InSection.AddMenuEntry("Emitter_NewNiagaraSystem", Label, ToolTip, Icon, UIAction);
-				}
-				{
-					const TAttribute<FText> Label = LOCTEXT("Emitter_CreateDuplicateParent", "Create Duplicate Parent");
-					const TAttribute<FText> ToolTip = LOCTEXT("Emitter_CreateDuplicateParentTooltip", "Duplicate this emitter and set this emitter's parent to the new emitter.");
-					const FSlateIcon Icon = FSlateIcon();
-
-					FToolUIAction UIAction;
-					UIAction.ExecuteAction = FToolMenuExecuteAction::CreateStatic(&ExecuteCreateDuplicateParent);
-					InSection.AddMenuEntry("Emitter_CreateDuplicateParent", Label, ToolTip, Icon, UIAction);
-				}
-				{
-					const TAttribute<FText> Label = LOCTEXT("MarkDependentCompilableAssetsDirtyLabel", "Mark Dependent Compilable Assets Dirty");
-					const TAttribute<FText> ToolTip = LOCTEXT("MarkDependentCompilableAssetsDirtyToolTip", "Finds all niagara assets which depend on this asset either directly or indirectly, and marks them dirty so they can be saved with the latest version.");
-					const FSlateIcon Icon = FSlateIcon();
+					{
+						const TAttribute<FText> Label = LOCTEXT("Emitter_NewNiagaraSystem", "Create Niagara System");
+						const TAttribute<FText> ToolTip = LOCTEXT("Emitter_NewNiagaraSystemTooltip", "Creates a niagara system using this emitter as a base.");
+						const FSlateIcon Icon = FSlateIcon(FAppStyle::GetAppStyleSetName(), "ClassIcon.ParticleSystem");
+
+						FToolUIAction UIAction;
+						UIAction.ExecuteAction = FToolMenuExecuteAction::CreateStatic(&ExecuteNewNiagaraSystem);
+						InSection.AddMenuEntry("Emitter_NewNiagaraSystem", Label, ToolTip, Icon, UIAction);
+					}
+					{
+						const TAttribute<FText> Label = LOCTEXT("Emitter_CreateDuplicateParent", "Create Duplicate Parent");
+						const TAttribute<FText> ToolTip = LOCTEXT("Emitter_CreateDuplicateParentTooltip", "Duplicate this emitter and set this emitter's parent to the new emitter.");
+						const FSlateIcon Icon = FSlateIcon();
+
+						FToolUIAction UIAction;
+						UIAction.ExecuteAction = FToolMenuExecuteAction::CreateStatic(&ExecuteCreateDuplicateParent);
+						InSection.AddMenuEntry("Emitter_CreateDuplicateParent", Label, ToolTip, Icon, UIAction);
+					}
+					{
+						const TAttribute<FText> Label = LOCTEXT("MarkDependentCompilableAssetsDirtyLabel", "Mark Dependent Compilable Assets Dirty");
+						const TAttribute<FText> ToolTip = LOCTEXT("MarkDependentCompilableAssetsDirtyToolTip", "Finds all niagara assets which depend on this asset either directly or indirectly, and marks them dirty so they can be saved with the latest version.");
+						const FSlateIcon Icon = FSlateIcon();
 					
-					FToolUIAction UIAction;
-					UIAction.ExecuteAction = FToolMenuExecuteAction::CreateStatic(&ExecuteMarkDependentCompilableAssetsDirty);
-					InSection.AddMenuEntry("MarkDependentCompilableAssetsDirty", Label, ToolTip, Icon, UIAction);
+						FToolUIAction UIAction;
+						UIAction.ExecuteAction = FToolMenuExecuteAction::CreateStatic(&ExecuteMarkDependentCompilableAssetsDirty);
+						InSection.AddMenuEntry("MarkDependentCompilableAssetsDirty", Label, ToolTip, Icon, UIAction);
+					}
 				}
 			}));
 		}));
