- PR #12440: Change from TObjectPtr to TWeakObjectPtr in FNiagaraAction_NewNode

#rnx

diff --git a/Engine/Plugins/FX/Niagara/Source/NiagaraEditor/Private/EdGraphSchema_Niagara.cpp b/Engine/Plugins/FX/Niagara/Source/NiagaraEditor/Private/EdGraphSchema_Niagara.cpp
--- a/Engine/Plugins/FX/Niagara/Source/NiagaraEditor/Private/EdGraphSchema_Niagara.cpp
+++ b/Engine/Plugins/FX/Niagara/Source/NiagaraEditor/Private/EdGraphSchema_Niagara.cpp
@@ -233,7 +233,7 @@ TSharedPtr<FNiagaraAction_NewNode> AddNewNodeMenuAction(
 	if (ensureMsgf(InNodeTemplate == nullptr || NiagaraEditorSettings->IsVisibleClass(InNodeTemplate->GetClass()),
 		TEXT("Can not create a menu action for a node of class %s."), *InNodeTemplate->GetClass()->GetName()))
 	{
-		NewAction->NodeTemplate = InNodeTemplate;
+		NewAction->WeakNodeTemplate = InNodeTemplate;
 	}
 	NewAction->SourceData = SourceData;
 	NewActions.Add(NewAction);
@@ -559,7 +559,7 @@ PRAGMA_ENABLE_DEPRECATION_WARNINGS
 
 			const FText MenuDesc = FText::FromString(FName::NameToDisplayString(Sig.Name.ToString(), false));
 			TSharedPtr<FNiagaraAction_NewNode> FunctionCallAction = AddNewNodeMenuAction(NewActions, FunctionCallNode, MenuDesc, ENiagaraMenuSections::General, TArray<FString>(), MenuDesc);
-			FunctionCallAction->NodeTemplate = FunctionCallNode;
+			FunctionCallAction->WeakNodeTemplate = FunctionCallNode;
 		}
 	}
 
diff --git a/Engine/Plugins/FX/Niagara/Source/NiagaraEditor/Private/NiagaraActions.cpp b/Engine/Plugins/FX/Niagara/Source/NiagaraEditor/Private/NiagaraActions.cpp
--- a/Engine/Plugins/FX/Niagara/Source/NiagaraEditor/Private/NiagaraActions.cpp
+++ b/Engine/Plugins/FX/Niagara/Source/NiagaraEditor/Private/NiagaraActions.cpp
@@ -236,7 +236,7 @@ UEdGraphNode* FNiagaraAction_NewNode::CreateNode(UEdGraph* ParentGraph, UEdGraph
 	UEdGraphNode* ResultNode = nullptr;
 
 	// If there is a template, we actually use it
-	if (NodeTemplate != nullptr)
+	if (UEdGraphNode* NodeTemplate = WeakNodeTemplate.Get())
 	{
 		FString OutErrorMsg;
 		UNiagaraNode* NiagaraNodeTemplate = Cast<UNiagaraNode>(NodeTemplate);
diff --git a/Engine/Plugins/FX/Niagara/Source/NiagaraEditor/Public/NiagaraActions.h b/Engine/Plugins/FX/Niagara/Source/NiagaraEditor/Public/NiagaraActions.h
--- a/Engine/Plugins/FX/Niagara/Source/NiagaraEditor/Public/NiagaraActions.h
+++ b/Engine/Plugins/FX/Niagara/Source/NiagaraEditor/Public/NiagaraActions.h
@@ -215,7 +215,7 @@ struct FNiagaraAction_NewNode : public FNiagaraMenuAction_Generic
 	NIAGARAEDITOR_API class UEdGraphNode* CreateNode(UEdGraph* Graph, TArray<UEdGraphPin*>& FromPins, FVector2D NodePosition, bool bSelectNewNode = true) const;
 
 	UPROPERTY()
-	TObjectPtr<class UEdGraphNode> NodeTemplate = nullptr;
+	TWeakObjectPtr<class UEdGraphNode> WeakNodeTemplate = nullptr;
 };
 
 struct FNiagaraParameterAction : public FEdGraphSchemaAction
