- Fix for missing nullptr check

#rnx

diff --git a/Engine/Plugins/FX/Niagara/Source/NiagaraEditorWidgets/Private/SNiagaraOverviewStackNode.cpp b/Engine/Plugins/FX/Niagara/Source/NiagaraEditorWidgets/Private/SNiagaraOverviewStackNode.cpp
--- a/Engine/Plugins/FX/Niagara/Source/NiagaraEditorWidgets/Private/SNiagaraOverviewStackNode.cpp
+++ b/Engine/Plugins/FX/Niagara/Source/NiagaraEditorWidgets/Private/SNiagaraOverviewStackNode.cpp
@@ -418,7 +418,10 @@ void SNiagaraOverviewStackNode::Tick(const FGeometry& AllottedGeometry, const do
 		// To support drag & drop indicators, we mark the prepass as dirty. This will cause the invalidation widget to render using the slow path, updating the drag & drop indicators properly
 		if(FSlateApplication::Get().IsDragDropping())
 		{
-			ContentAreaWidget->MarkPrepassAsDirty();
+			if (ContentAreaWidget)
+			{
+				ContentAreaWidget->MarkPrepassAsDirty();
+			}
 		}
 		
 		if (OverviewStackNode->IsRenamePending() && !SGraphNode::IsRenamePending())
