[Niagara]
 - Fix race by removing unused static variable

#rnx

diff --git a/Engine/Plugins/FX/Niagara/Source/Niagara/Classes/NiagaraScriptExecutionContext.h b/Engine/Plugins/FX/Niagara/Source/Niagara/Classes/NiagaraScriptExecutionContext.h
--- a/Engine/Plugins/FX/Niagara/Source/Niagara/Classes/NiagaraScriptExecutionContext.h
+++ b/Engine/Plugins/FX/Niagara/Source/Niagara/Classes/NiagaraScriptExecutionContext.h
@@ -145,8 +145,6 @@ struct FNiagaraScriptExecutionContextBase
 
 	TArray<FNiagaraDataSetExecutionInfo, TInlineAllocator<2>> DataSetInfo;
 
-	static uint32 TickCounter;
-
 	/** The script type this context is for. Allows us to access the correct per instance function table on the system instance. */
 	ENiagaraSystemSimulationScript ScriptType;
 
diff --git a/Engine/Plugins/FX/Niagara/Source/Niagara/Private/NiagaraScriptExecutionContext.cpp b/Engine/Plugins/FX/Niagara/Source/Niagara/Private/NiagaraScriptExecutionContext.cpp
--- a/Engine/Plugins/FX/Niagara/Source/Niagara/Private/NiagaraScriptExecutionContext.cpp
+++ b/Engine/Plugins/FX/Niagara/Source/Niagara/Private/NiagaraScriptExecutionContext.cpp
@@ -16,8 +16,6 @@ DECLARE_CYCLE_STAT(TEXT("Rebind DInterface Func Table"), STAT_NiagaraRebindDataI
 	
 	//Internal constants - only needed for non-GPU sim
 
-uint32 FNiagaraScriptExecutionContextBase::TickCounter = 0;
-
 static int32 GbExecVMScripts = 1;
 static FAutoConsoleVariableRef CVarNiagaraExecVMScripts(
 	TEXT("fx.ExecVMScripts"),
@@ -222,8 +220,6 @@ bool FNiagaraScriptExecutionContextBase::Execute(class FNiagaraSystemInstance* I
 		return true;
 	}
 
-	++TickCounter;//Should this be per execution?
-
 	DIStageTickHandler.PreStageTick(Instance, DeltaSeconds);
 
 	if (GbExecVMScripts != 0)
