- Allow optional streaming bounds to be used when we are injected into the streaming system
- In cases where we have dynamic bounds the initial data sent can be small, or the gradual increase of the bounds can result in incorrect mip calculations.
- This allows the user to specify the initial bounds used for streaming to nudge the streaming systems in the right direction from the start

#rnx

diff --git a/Engine/Plugins/FX/Niagara/Source/Niagara/Classes/NiagaraSystem.h b/Engine/Plugins/FX/Niagara/Source/Niagara/Classes/NiagaraSystem.h
--- a/Engine/Plugins/FX/Niagara/Source/Niagara/Classes/NiagaraSystem.h
+++ b/Engine/Plugins/FX/Niagara/Source/Niagara/Classes/NiagaraSystem.h
@@ -208,7 +208,6 @@ struct FNiagaraSystemScalabilityOverrides
 	TArray<FNiagaraSystemScalabilityOverride> Overrides;
 };
 
-
 /** A Niagara System contains multiple Niagara Emitters to create various effects.
  * Niagara Systems can be placed in the world, unlike Emitters, and expose User Parameters to configure an effect at runtime.*/
 UCLASS(BlueprintType, MinimalAPI, meta = (LoadBehavior = "LazyOnDemand"))
@@ -742,6 +741,8 @@ public:
 	NIAGARA_API FBox GetFixedBounds() const;
 	FORCEINLINE void SetFixedBounds(const FBox& Box) { FixedBounds = Box;  }
 
+	FBox GetInitialStreamingBounds() const { return bUseInitialStreamingBounds ? InitialStreamingBounds : FBox(); }
+
 #if WITH_EDITOR
 	NIAGARA_API void SetEffectType(UNiagaraEffectType* EffectType);
 
@@ -942,6 +943,20 @@ protected:
 	UPROPERTY(EditAnywhere, Category = "System", meta = (SkipSystemResetOnChange = "true", EditCondition = "bFixedBounds"))
 	FBox FixedBounds;
 
+	UPROPERTY(EditAnywhere, Category = "System", meta = (InlineEditConditionToggle, SkipSystemResetOnChange = "true"))
+	bool bUseInitialStreamingBounds = false;
+
+	/**
+	Initial streaming bounds are only included when the streaming system asks for the inital bounds,
+	it is the union of the currently calculated generated bounds and these bounds.
+	Any future updates to the bounds in the streamer will use the component bounds only and not these bounds.
+	Setting initial streaming bounds can help in situations where textures do not stream in correctly as the
+	heuristic to updating the bounds inside the streaming system does not update them at all, or frequently enough.
+	This is generally an issue with dynamic bounds, if your system has fixed bounds it is unlikely you will need to use this.
+	*/
+	UPROPERTY(EditAnywhere, Category = "System", meta = (SkipSystemResetOnChange = "true", EditCondition = "bUseInitialStreamingBounds"))
+	FBox InitialStreamingBounds = FBox(FVector(-100), FVector(100));
+
 	UPROPERTY()
 	bool bNeedsGPUContextInitForDataInterfaces;
 
diff --git a/Engine/Plugins/FX/Niagara/Source/Niagara/Private/NiagaraComponent.cpp b/Engine/Plugins/FX/Niagara/Source/Niagara/Private/NiagaraComponent.cpp
--- a/Engine/Plugins/FX/Niagara/Source/Niagara/Private/NiagaraComponent.cpp
+++ b/Engine/Plugins/FX/Niagara/Source/Niagara/Private/NiagaraComponent.cpp
@@ -539,7 +539,7 @@ void FNiagaraSceneProxy::GetDynamicMeshElements(const TArray<const FSceneView*>&
 		}
 	}
 
-	if (ViewFamily.EngineShowFlags.Particles && ViewFamily.EngineShowFlags.Niagara)
+	if (ViewFamily.EngineShowFlags.Bounds && ViewFamily.EngineShowFlags.Particles && ViewFamily.EngineShowFlags.Niagara)
 	{
 		for (int32 ViewIndex = 0; ViewIndex < Views.Num(); ViewIndex++)
 		{
@@ -2419,6 +2419,19 @@ void UNiagaraComponent::GetStreamingRenderAssetInfo(FStreamingTextureLevelContex
 		SystemInstanceController->GetMaterialStreamingInfo(MaterialAndScales);
 
 		LevelContext.BindBuildData(nullptr);
+
+		// Build the streaming bounds
+		FBoxSphereBounds StreamingBounds = Bounds;
+		if (Asset)
+		{
+			const FBox InitialStreamingBounds = Asset->GetInitialStreamingBounds();
+			if (InitialStreamingBounds.IsValid)
+			{
+				const FTransform& LocalToWorld = GetComponentTransform();
+				StreamingBounds = StreamingBounds + InitialStreamingBounds.TransformBy(LocalToWorld);
+			}
+		}
+
 		if (MaterialAndScales.Num() > 0)
 		{
 			static const FMeshUVChannelInfo UVChannelData(1.f);
@@ -2431,11 +2444,11 @@ void UNiagaraComponent::GetStreamingRenderAssetInfo(FStreamingTextureLevelContex
 			{
 				check(MaterialAndScale.Material != nullptr);
 				MaterialData.Material = MaterialAndScale.Material;
-				LevelContext.ProcessMaterial(Bounds, MaterialData, Bounds.SphereRadius * MaterialAndScale.Scale, OutStreamingRenderAssets, bIsValidTextureStreamingBuiltData, this);
+				LevelContext.ProcessMaterial(StreamingBounds, MaterialData, StreamingBounds.SphereRadius * MaterialAndScale.Scale, OutStreamingRenderAssets, bIsValidTextureStreamingBuiltData, this);
 			}
 		}
 
-		SystemInstanceController->GetStreamingMeshInfo(Bounds, LevelContext, OutStreamingRenderAssets);
+		SystemInstanceController->GetStreamingMeshInfo(StreamingBounds, LevelContext, OutStreamingRenderAssets);
 	}
 }
 
diff --git a/Engine/Plugins/FX/Niagara/Source/NiagaraEditor/Private/Widgets/SNiagaraSystemViewport.cpp b/Engine/Plugins/FX/Niagara/Source/NiagaraEditor/Private/Widgets/SNiagaraSystemViewport.cpp
--- a/Engine/Plugins/FX/Niagara/Source/NiagaraEditor/Private/Widgets/SNiagaraSystemViewport.cpp
+++ b/Engine/Plugins/FX/Niagara/Source/NiagaraEditor/Private/Widgets/SNiagaraSystemViewport.cpp
@@ -221,10 +221,23 @@ void FNiagaraSystemViewportClient::Draw(FViewport* InViewport, FCanvas* Canvas)
 		}
 	}
 
-	//-TODO:Stateless: Temporary to draw some debug information
+	const bool bDrawBounds = NiagaraViewport.IsValid() && NiagaraViewport->GetDrawElement(SNiagaraSystemViewport::EDrawElements::Bounds);
+	if (bDrawBounds)
+	{
+		EngineShowFlags.SetBounds(true);
+		EngineShowFlags.Game = 1;
+	}
+	else
+	{
+		EngineShowFlags.SetBounds(false);
+		EngineShowFlags.Game = 0;
+	}
+
 	if (ParticleSystem && Component)
 	{
 		UWorld* World = Component->GetWorld();
+
+		//-TODO:Stateless: Temporary to draw some debug information
 		for (const FNiagaraEmitterHandle& EmitterHandle : ParticleSystem->GetEmitterHandles())
 		{
 			UNiagaraStatelessEmitter* StatelessEmitter = EmitterHandle.GetStatelessEmitter();
@@ -235,17 +248,22 @@ void FNiagaraSystemViewportClient::Draw(FViewport* InViewport, FCanvas* Canvas)
 
 			StatelessEmitter->DrawModuleDebug(World, Component->GetComponentTransform());
 		}
-	}
 
-	if (NiagaraViewport.IsValid() && NiagaraViewport->GetDrawElement(SNiagaraSystemViewport::EDrawElements::Bounds))
-	{
-		EngineShowFlags.SetBounds(true);
-		EngineShowFlags.Game = 1;
-	}
-	else
-	{
-		EngineShowFlags.SetBounds(false);
-		EngineShowFlags.Game = 0;
+		// Draw Streaming Bounds if enabled
+		if (bDrawBounds)
+		{
+			const FBox StreamingBounds = ParticleSystem->GetInitialStreamingBounds();
+			if (StreamingBounds.IsValid)
+			{
+				DrawDebugBox(
+					World,
+					FVector::ZeroVector,//LocalToWorldTransform.TransformPosition(Center),
+					StreamingBounds.GetExtent(),
+					FQuat::Identity,//LocalToWorldTransform.TransformRotation(Rotation),
+					FColor::Green
+				);
+			}
+		}
 	}
 
 	FEditorViewportClient::Draw(InViewport, Canvas);
