- Fix missing center offset for bounds display

#rnx

diff --git a/Engine/Plugins/FX/Niagara/Source/NiagaraEditor/Private/Widgets/SNiagaraSystemViewport.cpp b/Engine/Plugins/FX/Niagara/Source/NiagaraEditor/Private/Widgets/SNiagaraSystemViewport.cpp
--- a/Engine/Plugins/FX/Niagara/Source/NiagaraEditor/Private/Widgets/SNiagaraSystemViewport.cpp
+++ b/Engine/Plugins/FX/Niagara/Source/NiagaraEditor/Private/Widgets/SNiagaraSystemViewport.cpp
@@ -257,9 +257,9 @@ void FNiagaraSystemViewportClient::Draw(FViewport* InViewport, FCanvas* Canvas)
 			{
 				DrawDebugBox(
 					World,
-					FVector::ZeroVector,//LocalToWorldTransform.TransformPosition(Center),
+					StreamingBounds.GetCenter(),
 					StreamingBounds.GetExtent(),
-					FQuat::Identity,//LocalToWorldTransform.TransformRotation(Rotation),
+					FQuat::Identity,
 					FColor::Green
 				);
 			}
