- Add scale to timed niagara anim notify

#rnx
#tests Add timed notify and change scale settings to observer positions scaling

diff --git a/Engine/Plugins/FX/Niagara/Source/NiagaraAnimNotifies/Private/AnimNotifyState_TimedNiagaraEffect.cpp b/Engine/Plugins/FX/Niagara/Source/NiagaraAnimNotifies/Private/AnimNotifyState_TimedNiagaraEffect.cpp
--- a/Engine/Plugins/FX/Niagara/Source/NiagaraAnimNotifies/Private/AnimNotifyState_TimedNiagaraEffect.cpp
+++ b/Engine/Plugins/FX/Niagara/Source/NiagaraAnimNotifies/Private/AnimNotifyState_TimedNiagaraEffect.cpp
@@ -25,7 +25,17 @@ UFXSystemComponent* UAnimNotifyState_TimedNiagaraEffect::SpawnEffect(USkeletalMe
 	// Only spawn if we've got valid params
 	if (ValidateParameters(MeshComp))
 	{
-		if (UNiagaraComponent* NewComponent = UNiagaraFunctionLibrary::SpawnSystemAttached(Template, MeshComp, SocketName, LocationOffset, RotationOffset, EAttachLocation::KeepRelativeOffset, !bDestroyAtEnd))
+		FFXSystemSpawnParameters SpawnParams;
+		SpawnParams.SystemTemplate		= Template;
+		SpawnParams.AttachToComponent	= MeshComp;
+		SpawnParams.AttachPointName		= SocketName;
+		SpawnParams.Location			= LocationOffset;
+		SpawnParams.Rotation			= RotationOffset;
+		SpawnParams.Scale				= Scale;
+		SpawnParams.LocationType		= EAttachLocation::KeepRelativeOffset;
+		SpawnParams.bAutoDestroy		= !bDestroyAtEnd;
+
+		if (UNiagaraComponent* NewComponent = UNiagaraFunctionLibrary::SpawnSystemAttachedWithParams(SpawnParams))
 		{
 			if (bApplyRateScaleAsTimeDilation)
 			{
diff --git a/Engine/Plugins/FX/Niagara/Source/NiagaraAnimNotifies/Public/AnimNotifyState_TimedNiagaraEffect.h b/Engine/Plugins/FX/Niagara/Source/NiagaraAnimNotifies/Public/AnimNotifyState_TimedNiagaraEffect.h
--- a/Engine/Plugins/FX/Niagara/Source/NiagaraAnimNotifies/Public/AnimNotifyState_TimedNiagaraEffect.h
+++ b/Engine/Plugins/FX/Niagara/Source/NiagaraAnimNotifies/Public/AnimNotifyState_TimedNiagaraEffect.h
@@ -37,6 +37,10 @@ class UAnimNotifyState_TimedNiagaraEffect : public UAnimNotifyState
 	UPROPERTY(EditAnywhere, BlueprintReadWrite, Category = NiagaraSystem, meta = (ToolTip = "Rotation offset from the socket or bone for the Niagara system"))
 	FRotator RotationOffset;
 
+	// Scale to spawn the Niagara system at
+	UPROPERTY(EditAnywhere, Category = NiagaraSystem)
+	FVector Scale = FVector(1.0f);
+
 	// Should we set the animation rate scale as time dilation on the spawn effect?
 	UPROPERTY(EditAnywhere, BlueprintReadWrite, Category = NiagaraSystem)
 	bool bApplyRateScaleAsTimeDilation = false;
