- Stateless Parameterizaton
- Shape location - Enable bindings on various properties, fix surface box spawning, add support for edge spawning to plane
- Forces - Allow various forces to use bindings

#tests

#virtualized

diff --git a/Engine/Plugins/FX/Niagara/Shaders/Private/Stateless/Modules/NiagaraStatelessModule_ShapeLocation.ush b/Engine/Plugins/FX/Niagara/Shaders/Private/Stateless/Modules/NiagaraStatelessModule_ShapeLocation.ush
--- a/Engine/Plugins/FX/Niagara/Shaders/Private/Stateless/Modules/NiagaraStatelessModule_ShapeLocation.ush
+++ b/Engine/Plugins/FX/Niagara/Shaders/Private/Stateless/Modules/NiagaraStatelessModule_ShapeLocation.ush
@@ -12,27 +12,31 @@ float3 ShapeLocation_GetLocation(inout FStatelessParticle Particle)
 	BRANCH
 	if ( ShapeLocation_Mode.x == 0 )
 	{
-		const float3 BoxScale = ShapeLocation_Parameters0.xyz;
-		const float3 BoxBias = ShapeLocation_Parameters1.xyz;
+		const float3 BoxSize = ShapeLocation_Parameters0.xyz;
 		const bool bOnSurface = ShapeLocation_Mode.y == 1;
-		const float SurfaceScale = ShapeLocation_Parameters0.w;
-		const float SurfaceBias = ShapeLocation_Parameters1.w;
 
-		const float3 P0 = RandomFloat3(0);
+		const float3 P0 = RandomFloat3(0) - 0.5f;
 		if (bOnSurface)
 		{
-			//-TODO: This isn't quite what we want
-			const float3 L0 = round(P0);
-			const uint S = RandomUInt(1) % 3;
+			const float3 FaceDirA[] = { float3(0.0f, 1.0f, 0.0f), float3( 0.0f, 1.0f, 0.0f), float3(0.0f, 0.0f, 1.0f), float3(0.0f, 0.0f, 1.0f), float3(1.0f, 0.0f, 0.0f), float3(1.0f, 0.0f, 0.0f) };
+			const float3 FaceDirB[] = { float3(0.0f, 0.0f, 1.0f), float3( 0.0f, 0.0f, 1.0f), float3(1.0f, 0.0f, 0.0f), float3(1.0f, 0.0f, 0.0f), float3(0.0f, 1.0f, 0.0f), float3(0.0f, 1.0f, 0.0f) };
+			const float3 FaceDirC[] = { float3(1.0f, 0.0f, 0.0f), float3(-1.0f, 0.0f, 0.0f), float3(0.0f, 1.0f, 0.0f), float3(0.0f,-1.0f, 0.0f), float3(0.0f, 0.0f, 1.0f), float3(0.0f, 0.0f,-1.0f) };
+			const uint NumFaces = ShapeLocation_Mode.z;
+
+			const float3 Thickness = ShapeLocation_Parameters1.xyz;
+			const float3 LocalBoxSize = BoxSize + (Thickness * P0.z);
+				
+			const uint iAxis = RandomUInt(1) % NumFaces;
 			float3 Location;
-			Location.x = S != 0 ? P0.x : L0.x;
-			Location.y = S != 1 ? P0.y : L0.y;
-			Location.z = S != 2 ? P0.z : L0.z;
-			return (Location * BoxScale + BoxBias) + (P0 * SurfaceScale + SurfaceBias);
+			Location  = LocalBoxSize * FaceDirA[iAxis] * P0.x;
+			Location += LocalBoxSize * FaceDirB[iAxis] * P0.y;
+			Location += LocalBoxSize * FaceDirC[iAxis] * 0.5f;
+				
+			return Location;
 		}
 		else
 		{
-			return P0 * BoxScale + BoxBias;
+			return P0 * BoxSize;
 		}
 	}
 
diff --git a/Engine/Plugins/FX/Niagara/Source/Niagara/Internal/Stateless/Modules/NiagaraStatelessModuleCommon.h b/Engine/Plugins/FX/Niagara/Source/Niagara/Internal/Stateless/Modules/NiagaraStatelessModuleCommon.h
--- a/Engine/Plugins/FX/Niagara/Source/Niagara/Internal/Stateless/Modules/NiagaraStatelessModuleCommon.h
+++ b/Engine/Plugins/FX/Niagara/Source/Niagara/Internal/Stateless/Modules/NiagaraStatelessModuleCommon.h
@@ -16,6 +16,7 @@ namespace NiagaraStateless
 
 		ENiagaraCoordinateSpace			VelocityCoordinateSpace = ENiagaraCoordinateSpace::Local;
 		FNiagaraStatelessRangeVector3	VelocityRange = FNiagaraStatelessRangeVector3(FVector3f::ZeroVector);
+		FNiagaraStatelessRangeFloat		LinearVelocityScale = FNiagaraStatelessRangeFloat(1.0f);
 
 		ENiagaraCoordinateSpace			WindCoordinateSpace = ENiagaraCoordinateSpace::Local;
 		FNiagaraStatelessRangeVector3	WindRange = FNiagaraStatelessRangeVector3(FVector3f::ZeroVector);
diff --git a/Engine/Plugins/FX/Niagara/Source/Niagara/Internal/Stateless/Modules/NiagaraStatelessModule_AddVelocity.h b/Engine/Plugins/FX/Niagara/Source/Niagara/Internal/Stateless/Modules/NiagaraStatelessModule_AddVelocity.h
--- a/Engine/Plugins/FX/Niagara/Source/Niagara/Internal/Stateless/Modules/NiagaraStatelessModule_AddVelocity.h
+++ b/Engine/Plugins/FX/Niagara/Source/Niagara/Internal/Stateless/Modules/NiagaraStatelessModule_AddVelocity.h
@@ -27,12 +27,12 @@ public:
 	UPROPERTY(EditAnywhere, Category = "Parameters", meta=(SegmentedDisplay))
 	ENSM_VelocityType VelocityType = ENSM_VelocityType::Linear;
 
-	UPROPERTY(EditAnywhere, Category = "Parameters", meta = (DisplayName = "Velocity", DisableUniformDistribution, DisableBindingDistribution, EditConditionHides, EditCondition = "VelocityType == ENSM_VelocityType::Linear"))
+	UPROPERTY(EditAnywhere, Category = "Parameters", meta = (DisplayName = "Velocity", DisableUniformDistribution, EditConditionHides, EditCondition = "VelocityType == ENSM_VelocityType::Linear"))
 	FNiagaraDistributionRangeVector3 LinearVelocityDistribution = FNiagaraDistributionRangeVector3(FVector3f(0.0f, 0.0f, 100.0f));
-	UPROPERTY(EditAnywhere, Category = "Parameters", meta = (DisplayName = "Velocity Scale", EditConditionHides, EditCondition = "VelocityType == ENSM_VelocityType::Linear"))
-	float LinearVelocityScale = 1.0f;
+	UPROPERTY(EditAnywhere, Category = "Parameters", meta = (DisplayName = "Velocity Scale", DisableRangeDistribution, EditConditionHides, EditCondition = "VelocityType == ENSM_VelocityType::Linear"))
+	FNiagaraDistributionRangeFloat LinearVelocityScale = FNiagaraDistributionRangeFloat(1.0f);
 
-	UPROPERTY(EditAnywhere, Category = "Parameters", meta = (DisplayName = "Cone Velocity", DisableBindingDistribution, EditConditionHides, EditCondition = "VelocityType == ENSM_VelocityType::InCone"))
+	UPROPERTY(EditAnywhere, Category = "Parameters", meta = (DisplayName = "Cone Velocity", EditConditionHides, EditCondition = "VelocityType == ENSM_VelocityType::InCone"))
 	FNiagaraDistributionRangeFloat ConeVelocityDistribution = FNiagaraDistributionRangeFloat(500.0f);
 	UPROPERTY(EditAnywhere, Category = "Parameters", meta = (EditConditionHides, EditCondition = "VelocityType == ENSM_VelocityType::InCone"))
 	FRotator ConeRotation = FRotator::ZeroRotator;
@@ -41,7 +41,7 @@ public:
 	UPROPERTY(EditAnywhere, Category = "Parameters", meta = (EditConditionHides, EditCondition = "VelocityType == ENSM_VelocityType::InCone", ClampMin = "0.0", ClampMax = "360.0", UIMin = "0.0", UIMax = "360.0"))
 	float InnerCone = 0.0f;
 
-	UPROPERTY(EditAnywhere, Category = "Parameters", meta = (DisplayName = "Point Velocity", DisableBindingDistribution, EditConditionHides, EditCondition = "VelocityType == ENSM_VelocityType::FromPoint"))
+	UPROPERTY(EditAnywhere, Category = "Parameters", meta = (DisplayName = "Point Velocity", EditConditionHides, EditCondition = "VelocityType == ENSM_VelocityType::FromPoint"))
 	FNiagaraDistributionRangeFloat PointVelocityDistribution = FNiagaraDistributionRangeFloat(100.0f);
 	UPROPERTY(EditAnywhere, Category = "Parameters", meta = (EditConditionHides, EditCondition = "VelocityType == ENSM_VelocityType::FromPoint"))
 	FVector3f PointOrigin = FVector3f::ZeroVector;
diff --git a/Engine/Plugins/FX/Niagara/Source/Niagara/Internal/Stateless/Modules/NiagaraStatelessModule_ScaleSpriteSize.h b/Engine/Plugins/FX/Niagara/Source/Niagara/Internal/Stateless/Modules/NiagaraStatelessModule_ScaleSpriteSize.h
--- a/Engine/Plugins/FX/Niagara/Source/Niagara/Internal/Stateless/Modules/NiagaraStatelessModule_ScaleSpriteSize.h
+++ b/Engine/Plugins/FX/Niagara/Source/Niagara/Internal/Stateless/Modules/NiagaraStatelessModule_ScaleSpriteSize.h
@@ -83,7 +83,7 @@ public:
 
 		FParameters* Parameters						= SetShaderParameterContext.GetParameterNestedStruct<FParameters>();
 		Parameters->ScaleSpriteSize_Distribution	= ModuleBuiltData->DistributionParameters;
-		SetShaderParameterContext.GetRendererParameterValue(Parameters->ScaleSpriteSize_CurveScale, ModuleBuiltData->CurveScaleOffset, ModuleBuiltData->CurveScale);
+		Parameters->ScaleSpriteSize_CurveScale		= SetShaderParameterContext.GetRendererParameterValue(ModuleBuiltData->CurveScaleOffset, ModuleBuiltData->CurveScale);
 	}
 
 	static void ParticleSimulate(const NiagaraStateless::FParticleSimulationContext& ParticleSimulationContext)
diff --git a/Engine/Plugins/FX/Niagara/Source/Niagara/Internal/Stateless/Modules/NiagaraStatelessModule_ShapeLocation.h b/Engine/Plugins/FX/Niagara/Source/Niagara/Internal/Stateless/Modules/NiagaraStatelessModule_ShapeLocation.h
--- a/Engine/Plugins/FX/Niagara/Source/Niagara/Internal/Stateless/Modules/NiagaraStatelessModule_ShapeLocation.h
+++ b/Engine/Plugins/FX/Niagara/Source/Niagara/Internal/Stateless/Modules/NiagaraStatelessModule_ShapeLocation.h
@@ -9,13 +9,22 @@
 #include "NiagaraStatelessModule_ShapeLocation.generated.h"
 
 UENUM()
-enum class ENSM_ShapePrimitive
+enum class ENSM_ShapePrimitive : uint8
 {
 	Box,
 	Cylinder,
 	Plane,
 	Ring,
 	Sphere,
+	Max UMETA(Hidden)
+};
+
+UENUM()
+enum class ENSM_SurfaceExpansionMode : uint8
+{
+	Inner,
+	Centered,
+	Outside,
 };
 
 UCLASS(MinimalAPI, EditInlineNew, meta = (DisplayName = "Shape Location"))
@@ -29,37 +38,50 @@ public:
 	UPROPERTY(EditAnywhere, Category = "Parameters", meta = (ShowInStackItemHeader, StackItemHeaderAlignment = "Left"))
 	ENSM_ShapePrimitive ShapePrimitive = ENSM_ShapePrimitive::Sphere;
 
-	UPROPERTY(EditAnywhere, Category = "Parameters", meta = (EditConditionHides, EditCondition = "ShapePrimitive == ENSM_ShapePrimitive::Box"))
-	FVector3f BoxSize = FVector3f(100.0f);
+	UPROPERTY(EditAnywhere, Category = "Parameters", meta = (DisableRangeDistribution, EditConditionHides, EditCondition = "ShapePrimitive == ENSM_ShapePrimitive::Box"))
+	FNiagaraDistributionRangeVector3 BoxSize = FNiagaraDistributionRangeVector3(FVector3f(100.0f));
 	UPROPERTY(EditAnywhere, Category = "Parameters", meta = (EditConditionHides, EditCondition = "ShapePrimitive == ENSM_ShapePrimitive::Box"))
 	bool bBoxSurfaceOnly = false;
-	UPROPERTY(EditAnywhere, Category = "Parameters", meta = (EditConditionHides, EditCondition = "ShapePrimitive == ENSM_ShapePrimitive::Box"))
-	float BoxSurfaceThicknessMin = 0.0f;
-	UPROPERTY(EditAnywhere, Category = "Parameters", meta = (EditConditionHides, EditCondition = "ShapePrimitive == ENSM_ShapePrimitive::Box"))
-	float BoxSurfaceThicknessMax = 0.0f;
+	UPROPERTY(EditAnywhere, Category = "Parameters", meta = (EditConditionHides, EditCondition = "ShapePrimitive == ENSM_ShapePrimitive::Box && bBoxSurfaceOnly"))
+	ENSM_SurfaceExpansionMode BoxSurfaceExpansion = ENSM_SurfaceExpansionMode::Centered;
+	UPROPERTY(EditAnywhere, Category = "Parameters", meta = (DisableRangeDistribution, EditConditionHides, EditCondition = "ShapePrimitive == ENSM_ShapePrimitive::Box && bBoxSurfaceOnly"))
+	FNiagaraDistributionRangeFloat BoxSurfaceThickness = FNiagaraDistributionRangeFloat(0.0f);
 
+	UPROPERTY(EditAnywhere, Category = "Parameters", meta = (DisableRangeDistribution, EditConditionHides, EditCondition = "ShapePrimitive == ENSM_ShapePrimitive::Plane"))
+	FNiagaraDistributionRangeVector2 PlaneSize = FNiagaraDistributionRangeVector2(FVector2f(100.0f));
 	UPROPERTY(EditAnywhere, Category = "Parameters", meta = (EditConditionHides, EditCondition = "ShapePrimitive == ENSM_ShapePrimitive::Plane"))
-	FVector2f PlaneSize = FVector2f(100.0f);
-
-	UPROPERTY(EditAnywhere, Category = "Parameters", meta = (EditConditionHides, EditCondition = "ShapePrimitive == ENSM_ShapePrimitive::Cylinder"))
-	float CylinderHeight = 100.0f;
-	UPROPERTY(EditAnywhere, Category = "Parameters", meta = (EditConditionHides, EditCondition = "ShapePrimitive == ENSM_ShapePrimitive::Cylinder"))
-	float CylinderRadius = 100.0f;
-	UPROPERTY(EditAnywhere, Category = "Parameters", meta = (EditConditionHides, EditCondition = "ShapePrimitive == ENSM_ShapePrimitive::Cylinder"))
-	float CylinderHeightMidpoint = 0.5f;
-
-	UPROPERTY(EditAnywhere, Category = "Parameters", meta = (EditConditionHides, EditCondition = "ShapePrimitive == ENSM_ShapePrimitive::Ring"))
-	float RingRadius = 100.0f;
-	UPROPERTY(EditAnywhere, Category = "Parameters", meta = (EditConditionHides, EditCondition = "ShapePrimitive == ENSM_ShapePrimitive::Ring", ClampMin = "0.0", ClampMax = "1.0", UIMin = "0.0", UIMax = "1.0"))
-	float DiscCoverage = 0.0f;
-	UPROPERTY(EditAnywhere, Category = "Parameters", meta = (EditConditionHides, EditCondition = "ShapePrimitive == ENSM_ShapePrimitive::Ring", ClampMin = "0.0", ClampMax = "1.0", UIMin = "0.0", UIMax = "1.0"))
-	float RingUDistribution = 0.0f;
+	bool bPlaneEdgesOnly = false;
+	UPROPERTY(EditAnywhere, Category = "Parameters", meta = (EditConditionHides, EditCondition = "ShapePrimitive == ENSM_ShapePrimitive::Plane && bPlaneEdgesOnly"))
+	ENSM_SurfaceExpansionMode PlaneEdgeExpansion = ENSM_SurfaceExpansionMode::Centered;
+	UPROPERTY(EditAnywhere, Category = "Parameters", meta = (DisableRangeDistribution, EditConditionHides, EditCondition = "ShapePrimitive == ENSM_ShapePrimitive::Plane && bPlaneEdgesOnly"))
+	FNiagaraDistributionRangeFloat PlaneEdgeThickness = FNiagaraDistributionRangeFloat(0.0f);
+
+	UPROPERTY(EditAnywhere, Category = "Parameters", meta = (DisableRangeDistribution, EditConditionHides, EditCondition = "ShapePrimitive == ENSM_ShapePrimitive::Cylinder"))
+	FNiagaraDistributionRangeFloat CylinderHeight = FNiagaraDistributionRangeFloat(100.0f);
+	UPROPERTY(EditAnywhere, Category = "Parameters", meta = (DisableRangeDistribution, EditConditionHides, EditCondition = "ShapePrimitive == ENSM_ShapePrimitive::Cylinder"))
+	FNiagaraDistributionRangeFloat CylinderRadius = FNiagaraDistributionRangeFloat(100.0f);
+	UPROPERTY(EditAnywhere, Category = "Parameters", meta = (DisableRangeDistribution, EditConditionHides, EditCondition = "ShapePrimitive == ENSM_ShapePrimitive::Cylinder"))
+	FNiagaraDistributionRangeFloat CylinderHeightMidpoint = FNiagaraDistributionRangeFloat(0.5f);
+
+	UPROPERTY(EditAnywhere, Category = "Parameters", meta = (DisableRangeDistribution, EditConditionHides, EditCondition = "ShapePrimitive == ENSM_ShapePrimitive::Ring"))
+	FNiagaraDistributionRangeFloat RingRadius = FNiagaraDistributionRangeFloat(100.0f);
+	UPROPERTY(EditAnywhere, Category = "Parameters", meta = (DisableRangeDistribution, EditConditionHides, EditCondition = "ShapePrimitive == ENSM_ShapePrimitive::Ring", ClampMin = "0.0", ClampMax = "1.0", UIMin = "0.0", UIMax = "1.0"))
+	FNiagaraDistributionRangeFloat DiscCoverage = FNiagaraDistributionRangeFloat(0.0f);
+	UPROPERTY(EditAnywhere, Category = "Parameters", meta = (DisableRangeDistribution, EditConditionHides, EditCondition = "ShapePrimitive == ENSM_ShapePrimitive::Ring", ClampMin = "0.0", ClampMax = "1.0", UIMin = "0.0", UIMax = "1.0"))
+	FNiagaraDistributionRangeFloat RingUDistribution = FNiagaraDistributionRangeFloat(0.0f);
 
 	UPROPERTY(EditAnywhere, Category = "Parameters", meta=(EditConditionHides, EditCondition="ShapePrimitive == ENSM_ShapePrimitive::Sphere"))
-	float SphereMin = 0.0f;
-	UPROPERTY(EditAnywhere, Category = "Parameters", meta=(EditConditionHides, EditCondition="ShapePrimitive == ENSM_ShapePrimitive::Sphere"))
-	float SphereMax = 100.0f;
+	FNiagaraDistributionRangeFloat SphereRadius = FNiagaraDistributionRangeFloat(100.0f);
 
+private:
+#if WITH_EDITORONLY_DATA
+	UPROPERTY()
+	float SphereMin_DEPRECATED = 0.0f;
+	UPROPERTY()
+	float SphereMax_DEPRECATED = 100.0f;
+#endif
+
+public:
 	virtual void BuildEmitterData(const FNiagaraStatelessEmitterDataBuildContext& BuildContext) const override;
 	virtual void BuildShaderParameters(FNiagaraStatelessShaderParametersBuilder& ShaderParametersBuilder) const override;
 	virtual void SetShaderParameters(const FNiagaraStatelessSetShaderParameterContext& SetShaderParameterContext) const override;
@@ -77,5 +99,7 @@ public:
 		OutVariables.AddUnique(StatelessGlobals.PositionVariable);
 		OutVariables.AddUnique(StatelessGlobals.PreviousPositionVariable);
 	}
+
+	virtual void PostLoad() override;
 #endif
 };
diff --git a/Engine/Plugins/FX/Niagara/Source/Niagara/Internal/Stateless/Modules/NiagaraStatelessModule_SolveVelocitiesAndForces.h b/Engine/Plugins/FX/Niagara/Source/Niagara/Internal/Stateless/Modules/NiagaraStatelessModule_SolveVelocitiesAndForces.h
--- a/Engine/Plugins/FX/Niagara/Source/Niagara/Internal/Stateless/Modules/NiagaraStatelessModule_SolveVelocitiesAndForces.h
+++ b/Engine/Plugins/FX/Niagara/Source/Niagara/Internal/Stateless/Modules/NiagaraStatelessModule_SolveVelocitiesAndForces.h
@@ -8,10 +8,6 @@
 #include "Stateless/NiagaraStatelessParticleSimContext.h"
 #include "Stateless/Modules/NiagaraStatelessModuleCommon.h"
 
-#include "VectorField/VectorField.h"
-#include "VectorField/VectorFieldStatic.h"
-#include "RHIStaticStates.h"
-
 #include "NiagaraStatelessModule_SolveVelocitiesAndForces.generated.h"
 
 // Integrates all the forces applying them to position
@@ -32,184 +28,11 @@ class UNiagaraStatelessModule_SolveVelocitiesAndForces : public UNiagaraStateles
 public:
 	using FParameters = NiagaraStateless::FSolveVelocitiesAndForcesModule_ShaderParameters;
 
-	virtual void BuildEmitterData(const FNiagaraStatelessEmitterDataBuildContext& BuildContext) const override
-	{
-		NiagaraStateless::FPhysicsBuildData& PhysicsBuildData = BuildContext.GetTransientBuildData<NiagaraStateless::FPhysicsBuildData>();
-
-		FModuleBuiltData* BuiltData				= BuildContext.AllocateBuiltData<FModuleBuiltData>();
-		BuiltData->PhysicsData					= PhysicsBuildData;
-		BuiltData->PhysicsData.DragRange.Min	= FMath::Max(PhysicsBuildData.DragRange.Min, 0.01f);
-		BuiltData->PhysicsData.DragRange.Max	= FMath::Max(PhysicsBuildData.DragRange.Max, 0.01f);
-
-		const FNiagaraStatelessGlobals& StatelessGlobals = FNiagaraStatelessGlobals::Get();
-		BuiltData->PositionVariableOffset				= BuildContext.FindParticleVariableIndex(StatelessGlobals.PositionVariable);
-		BuiltData->VelocityVariableOffset				= BuildContext.FindParticleVariableIndex(StatelessGlobals.VelocityVariable);
-		BuiltData->PreviousPositionVariableOffset		= BuildContext.FindParticleVariableIndex(StatelessGlobals.PreviousPositionVariable);
-		BuiltData->PreviousVelocityVariableOffset		= BuildContext.FindParticleVariableIndex(StatelessGlobals.PreviousVelocityVariable);
-		
-		const bool bAttributesUsed =
-			BuiltData->PositionVariableOffset != INDEX_NONE ||
-			BuiltData->VelocityVariableOffset != INDEX_NONE ||
-			BuiltData->PreviousPositionVariableOffset != INDEX_NONE ||
-			BuiltData->PreviousVelocityVariableOffset != INDEX_NONE;
-
-		if (bAttributesUsed)
-		{
-			BuildContext.AddParticleSimulationExecSimulate(&UNiagaraStatelessModule_SolveVelocitiesAndForces::ParticleSimulate);
-		}
-	}
-
-	virtual void BuildShaderParameters(FNiagaraStatelessShaderParametersBuilder& ShaderParametersBuilder) const override
-	{
-		ShaderParametersBuilder.AddParameterNestedStruct<FParameters>();
-	}
-
-	virtual void SetShaderParameters(const FNiagaraStatelessSetShaderParameterContext& SetShaderParameterContext) const override
-	{
-		const FModuleBuiltData* ModuleBuiltData = SetShaderParameterContext.ReadBuiltData<FModuleBuiltData>();
-		const NiagaraStateless::FPhysicsBuildData& PhysicsData = ModuleBuiltData->PhysicsData;
-
-		const FNiagaraStatelessSpaceTransforms& SpaceTransforms = SetShaderParameterContext.GetSpaceTransforms();
-
-		FParameters* Parameters = SetShaderParameterContext.GetParameterNestedStruct<FParameters>();
-		Parameters->SolveVelocitiesAndForces_MassScale				= PhysicsData.MassRange.GetScale();
-		Parameters->SolveVelocitiesAndForces_MassBias				= PhysicsData.MassRange.Min;
-		Parameters->SolveVelocitiesAndForces_DragScale				= PhysicsData.DragRange.GetScale();
-		Parameters->SolveVelocitiesAndForces_DragBias				= PhysicsData.DragRange.Min;
-		Parameters->SolveVelocitiesAndForces_VelocityScale			= SpaceTransforms.TransformVector(PhysicsData.VelocityCoordinateSpace, PhysicsData.VelocityRange.GetScale());
-		Parameters->SolveVelocitiesAndForces_VelocityBias			= SpaceTransforms.TransformVector(PhysicsData.VelocityCoordinateSpace, PhysicsData.VelocityRange.Min);
-		Parameters->SolveVelocitiesAndForces_WindScale				= SpaceTransforms.TransformVector(PhysicsData.WindCoordinateSpace, PhysicsData.WindRange.GetScale());
-		Parameters->SolveVelocitiesAndForces_WindBias				= SpaceTransforms.TransformVector(PhysicsData.WindCoordinateSpace, PhysicsData.WindRange.Min);
-		Parameters->SolveVelocitiesAndForces_AccelerationScale		= SpaceTransforms.TransformVector(PhysicsData.AccelerationCoordinateSpace, PhysicsData.AccelerationRange.GetScale());
-		Parameters->SolveVelocitiesAndForces_AccelerationBias		= SpaceTransforms.TransformVector(PhysicsData.AccelerationCoordinateSpace, PhysicsData.AccelerationRange.Min);
-		Parameters->SolveVelocitiesAndForces_AccelerationScale		+= SpaceTransforms.TransformVector(ENiagaraCoordinateSpace::World, PhysicsData.GravityRange.GetScale());
-		Parameters->SolveVelocitiesAndForces_AccelerationBias		+= SpaceTransforms.TransformVector(ENiagaraCoordinateSpace::World, PhysicsData.GravityRange.Min);
-
-		Parameters->SolveVelocitiesAndForces_ConeVelocityEnabled	= PhysicsData.bConeVelocity ? 1 : 0;
-		Parameters->SolveVelocitiesAndForces_ConeQuat				= SpaceTransforms.TransformRotation(PhysicsData.ConeCoordinateSpace, PhysicsData.ConeQuat);
-		Parameters->SolveVelocitiesAndForces_ConeVelocityScale		= PhysicsData.ConeVelocityRange.GetScale();
-		Parameters->SolveVelocitiesAndForces_ConeVelocityBias		= PhysicsData.ConeVelocityRange.Min;
-		Parameters->SolveVelocitiesAndForces_ConeAngleScale			= (PhysicsData.ConeOuterAngle - PhysicsData.ConeInnerAngle) * (UE_PI / 360.0f);
-		Parameters->SolveVelocitiesAndForces_ConeAngleBias			= PhysicsData.ConeInnerAngle * (UE_PI / 360.0f);
-		Parameters->SolveVelocitiesAndForces_ConeVelocityFalloff	= PhysicsData.ConeVelocityFalloff;
-
-		Parameters->SolveVelocitiesAndForces_PontVelocityEnabled	= PhysicsData.bPointVelocity ? 1 : 0;
-		Parameters->SolveVelocitiesAndForces_PointVelocityScale		= PhysicsData.PointVelocityRange.GetScale();
-		Parameters->SolveVelocitiesAndForces_PointVelocityBias		= PhysicsData.PointVelocityRange.Min;
-		Parameters->SolveVelocitiesAndForces_PointOrigin			= SpaceTransforms.TransformPosition(PhysicsData.PointCoordinateSpace, PhysicsData.PointOrigin);
-
-		Parameters->SolveVelocitiesAndForces_NoiseEnabled			= PhysicsData.bNoiseEnabled ? 1 : 0;
-		Parameters->SolveVelocitiesAndForces_NoiseAmplitude			= PhysicsData.NoiseAmplitude;
-		Parameters->SolveVelocitiesAndForces_NoiseFrequency			= FVector3f(PhysicsData.NoiseFrequency, PhysicsData.NoiseFrequency, PhysicsData.NoiseFrequency);
-		//SetShaderParameterContext.SetTextureResource(&Parameters->SolveVelocitiesAndForces_NoiseTexture, ModuleBuildData->NoiseTexture);
-		Parameters->SolveVelocitiesAndForces_NoiseMode				= PhysicsData.NoiseMode;
-		Parameters->SolveVelocitiesAndForces_NoiseLUTOffset			= PhysicsData.NoiseLUTOffset;
-		Parameters->SolveVelocitiesAndForces_NoiseLUTNumChannel		= PhysicsData.NoiseLUTNumChannel;
-		Parameters->SolveVelocitiesAndForces_NoiseLUTChannelWidth	= PhysicsData.NoiseLUTChannelWidth;
-
-		FVectorFieldTextureAccessor TextureAccessor(Cast<UVectorField>(PhysicsData.NoiseTexture));
-
-		ENQUEUE_RENDER_COMMAND(FNaughtyTest)(
-			[Parameters, TextureAccessor](FRHICommandListImmediate& RHICmdList)
-			{
-				FRHITexture* NoiseTextureRHI = TextureAccessor.GetTexture();
-				Parameters->SolveVelocitiesAndForces_NoiseTexture = NoiseTextureRHI;
-				Parameters->SolveVelocitiesAndForces_NoiseSampler = TStaticSamplerState<SF_Bilinear, AM_Wrap, AM_Wrap, AM_Wrap>::GetRHI();
-
-				if (Parameters->SolveVelocitiesAndForces_NoiseMode < 2)
-				{
-					const FIntVector TextureSize = NoiseTextureRHI ? NoiseTextureRHI->GetSizeXYZ() : FIntVector(1, 1, 1);
-					Parameters->SolveVelocitiesAndForces_NoiseFrequency.X *= 1.0f / float(TextureSize.X);
-					Parameters->SolveVelocitiesAndForces_NoiseFrequency.Y *= 1.0f / float(TextureSize.Y);
-					Parameters->SolveVelocitiesAndForces_NoiseFrequency.Z *= 1.0f / float(TextureSize.Z);
-				}
-			}
-		);
-	}
-
-	static FVector3f IntegratePosition(float Age, float Mass, float Drag, const FVector3f& Velocity, const FVector3f& Wind, const FVector3f& Acceleration)
-	{
-		const FVector3f IntVelocity = (Velocity - Wind) + (Wind * Age * Age);
-		const float LambdaDragMass = FMath::Max(Drag * (1.0f / Mass), 0.0001f);
-		const float LambdaAge = (1.0f - FMath::Exp(-(LambdaDragMass * Age))) / LambdaDragMass;
-		FVector3f Position = IntVelocity * LambdaAge;
-		Position += (Age - LambdaAge) * (Acceleration / LambdaDragMass);
-		return Position;
-	}
-
-	static void ParticleSimulate(const NiagaraStateless::FParticleSimulationContext& ParticleSimulationContext)
-	{
-		const FModuleBuiltData* ModuleBuiltData = ParticleSimulationContext.ReadBuiltData<FModuleBuiltData>();
-		const FParameters* Parameters = ParticleSimulationContext.ReadParameterNestedStruct<FParameters>();
-
-		const float* AgeData			= ParticleSimulationContext.GetParticleAge();
-		const float* PreviousAgeData	= ParticleSimulationContext.GetParticlePreviousAge();
-		const float InvDeltaTime		= ParticleSimulationContext.GetInvDeltaTime();
-
-		for (uint32 i = 0; i < ParticleSimulationContext.GetNumInstances(); ++i)
-		{
-			const float Mass				= ParticleSimulationContext.RandomScaleBiasFloat(i, 0, Parameters->SolveVelocitiesAndForces_MassScale, Parameters->SolveVelocitiesAndForces_MassBias);
-			const float Drag				= ParticleSimulationContext.RandomScaleBiasFloat(i, 1, Parameters->SolveVelocitiesAndForces_DragScale, Parameters->SolveVelocitiesAndForces_DragBias);
-			FVector3f InitialVelocity		= ParticleSimulationContext.RandomScaleBiasFloat(i, 2, Parameters->SolveVelocitiesAndForces_VelocityScale, Parameters->SolveVelocitiesAndForces_VelocityBias);
-			const FVector3f Wind			= ParticleSimulationContext.RandomScaleBiasFloat(i, 3, Parameters->SolveVelocitiesAndForces_WindScale, Parameters->SolveVelocitiesAndForces_WindBias);
-			const FVector3f Acceleration	= ParticleSimulationContext.RandomScaleBiasFloat(i, 4, Parameters->SolveVelocitiesAndForces_AccelerationScale, Parameters->SolveVelocitiesAndForces_AccelerationBias);
-
-			FVector3f Position			= ParticleSimulationContext.ReadParticleVariable(ModuleBuiltData->PositionVariableOffset, i, FVector3f::ZeroVector);
-			FVector3f PreviousPosition	= ParticleSimulationContext.ReadParticleVariable(ModuleBuiltData->PreviousPositionVariableOffset, i, FVector3f::ZeroVector);
-
-			if (ModuleBuiltData->PhysicsData.bConeVelocity)
-			{
-				const float ConeAngle = ParticleSimulationContext.RandomScaleBiasFloat(i, 5, Parameters->SolveVelocitiesAndForces_ConeAngleScale, Parameters->SolveVelocitiesAndForces_ConeAngleBias);
-				const float ConeRotation = ParticleSimulationContext.RandomFloat(i, 6) * UE_TWO_PI;
-				const FVector2f scAng = FVector2f(FMath::Sin(ConeAngle), FMath::Cos(ConeAngle));
-				const FVector2f scRot = FVector2f(FMath::Sin(ConeRotation), FMath::Cos(ConeRotation));
-				const FVector3f Direction = FVector3f(scRot.X * scAng.X, scRot.Y * scAng.X, scAng.Y);
-
-				float VelocityScale = ParticleSimulationContext.RandomScaleBiasFloat(i, 7, Parameters->SolveVelocitiesAndForces_ConeVelocityScale, Parameters->SolveVelocitiesAndForces_ConeVelocityBias);
-				if (ModuleBuiltData->PhysicsData.bConeVelocity)
-				{
-					const float pf = FMath::Pow(FMath::Clamp(scAng.Y, 0.0f, 1.0f), Parameters->SolveVelocitiesAndForces_ConeVelocityFalloff * 10.0f);
-					VelocityScale *= FMath::Lerp(1.0f, pf, Parameters->SolveVelocitiesAndForces_ConeVelocityFalloff);
-				}
-
-				InitialVelocity += ModuleBuiltData->PhysicsData.ConeQuat.RotateVector(Direction) * VelocityScale;
-			}
-
-			if (ModuleBuiltData->PhysicsData.bPointVelocity)
-			{
-				const FVector3f FallbackDir	= ParticleSimulationContext.RandomUnitFloat3(i, 8);
-				const FVector3f Delta		= Position - Parameters->SolveVelocitiesAndForces_PointOrigin;
-				const FVector3f Direction	= ParticleSimulationContext.SafeNormalize(Delta, FallbackDir);
-				const float		VelocityScale = ParticleSimulationContext.RandomScaleBiasFloat(i, 9, Parameters->SolveVelocitiesAndForces_PointVelocityScale, Parameters->SolveVelocitiesAndForces_PointVelocityBias);
-
-				InitialVelocity += Direction * VelocityScale;
-			}
-
-			if (ModuleBuiltData->PhysicsData.bNoiseEnabled)
-			{
-				//-TODO:
-			}
-
-			Position += IntegratePosition(AgeData[i], Mass, Drag, InitialVelocity, Wind, Acceleration);
-			PreviousPosition += IntegratePosition(PreviousAgeData[i], Mass, Drag, InitialVelocity, Wind, Acceleration);
-
-			ParticleSimulationContext.WriteParticleVariable(ModuleBuiltData->PositionVariableOffset, i, Position);
-			ParticleSimulationContext.WriteParticleVariable(ModuleBuiltData->PreviousPositionVariableOffset, i, PreviousPosition);
-
-			const FVector3f Velocity = (Position - PreviousPosition) * InvDeltaTime;
-			ParticleSimulationContext.WriteParticleVariable(ModuleBuiltData->VelocityVariableOffset, i, Velocity);
-			ParticleSimulationContext.WriteParticleVariable(ModuleBuiltData->PreviousVelocityVariableOffset, i, Velocity);
-		}
-	}
+	virtual void BuildEmitterData(const FNiagaraStatelessEmitterDataBuildContext& BuildContext) const override;
+	virtual void BuildShaderParameters(FNiagaraStatelessShaderParametersBuilder& ShaderParametersBuilder) const override;
+	virtual void SetShaderParameters(const FNiagaraStatelessSetShaderParameterContext& SetShaderParameterContext) const override;
 
 #if WITH_EDITORONLY_DATA
-	virtual void GetOutputVariables(TArray<FNiagaraVariableBase>& OutVariables) const override
-	{
-		const FNiagaraStatelessGlobals& StatelessGlobals = FNiagaraStatelessGlobals::Get();
-		OutVariables.AddUnique(StatelessGlobals.PositionVariable);
-		OutVariables.AddUnique(StatelessGlobals.VelocityVariable);
-		OutVariables.AddUnique(StatelessGlobals.PreviousPositionVariable);
-		OutVariables.AddUnique(StatelessGlobals.PreviousVelocityVariable);
-	}
+	virtual void GetOutputVariables(TArray<FNiagaraVariableBase>& OutVariables) const override;
 #endif
 };
diff --git a/Engine/Plugins/FX/Niagara/Source/Niagara/Internal/Stateless/Modules/NiagaraStatelessModule_SpriteFacingAndAlignment.h b/Engine/Plugins/FX/Niagara/Source/Niagara/Internal/Stateless/Modules/NiagaraStatelessModule_SpriteFacingAndAlignment.h
--- a/Engine/Plugins/FX/Niagara/Source/Niagara/Internal/Stateless/Modules/NiagaraStatelessModule_SpriteFacingAndAlignment.h
+++ b/Engine/Plugins/FX/Niagara/Source/Niagara/Internal/Stateless/Modules/NiagaraStatelessModule_SpriteFacingAndAlignment.h
@@ -17,8 +17,8 @@ class UNiagaraStatelessModule_SpriteFacingAndAlignment : public UNiagaraStateles
 
 	struct FModuleBuiltData
 	{
-		FVector3f	SpriteFacing	= FVector3f::XAxisVector;
-		FVector3f	SpriteAlignment = FVector3f::YAxisVector;
+		FNiagaraStatelessRangeVector3	SpriteFacing	= FNiagaraStatelessRangeVector3(FVector3f::XAxisVector);
+		FNiagaraStatelessRangeVector3	SpriteAlignment = FNiagaraStatelessRangeVector3(FVector3f::YAxisVector);
 
 		int32		SpriteFacingVariableOffset = INDEX_NONE;
 		int32		PreviousSpriteFacingVariableOffset = INDEX_NONE;
@@ -35,11 +35,11 @@ public:
 	UPROPERTY(EditAnywhere, Category = "Parameters", meta = (InlineEditConditionToggle))
 	bool bSpriteAlignmentEnabled = false;
 
-	UPROPERTY(EditAnywhere, Category = "Parameters", meta = (EditCondition = "bSpriteFacingEnabled"))
-	FVector3f	SpriteFacing = FVector3f::XAxisVector;
+	UPROPERTY(EditAnywhere, Category = "Parameters", meta = (DisableRangeDistribution, DisableUniformDistribution, EditCondition = "bSpriteFacingEnabled"))
+	FNiagaraDistributionRangeVector3	SpriteFacing = FNiagaraDistributionRangeVector3(FVector3f::XAxisVector);
 
-	UPROPERTY(EditAnywhere, Category = "Parameters", meta = (EditCondition = "bSpriteAlignmentEnabled"))
-	FVector3f	SpriteAlignment = FVector3f::YAxisVector;
+	UPROPERTY(EditAnywhere, Category = "Parameters", meta = (DisableRangeDistribution, DisableUniformDistribution, EditCondition = "bSpriteAlignmentEnabled"))
+	FNiagaraDistributionRangeVector3	SpriteAlignment = FNiagaraDistributionRangeVector3(FVector3f::YAxisVector);
 
 	virtual void BuildEmitterData(const FNiagaraStatelessEmitterDataBuildContext& BuildContext) const override
 	{
@@ -67,8 +67,14 @@ public:
 			return;
 		}
 
-		BuiltData->SpriteFacing		= SpriteFacing;
-		BuiltData->SpriteAlignment	= SpriteAlignment;
+		if (bSpriteFacingEnabled)
+		{
+			BuiltData->SpriteFacing = BuildContext.ConvertDistributionToRange(SpriteFacing, FVector3f::ZeroVector);
+		}
+		if (bSpriteAlignmentEnabled)
+		{
+			BuiltData->SpriteAlignment	= BuildContext.ConvertDistributionToRange(SpriteAlignment, FVector3f::ZeroVector);
+		}
 
 		BuildContext.AddParticleSimulationExecSimulate(&UNiagaraStatelessModule_SpriteFacingAndAlignment::ParticleSimulate);
 	}
@@ -83,8 +89,8 @@ public:
 		FParameters* Parameters = SetShaderParameterContext.GetParameterNestedStruct<FParameters>();
 		const FModuleBuiltData* ModuleBuiltData = SetShaderParameterContext.ReadBuiltData<FModuleBuiltData>();
 
-		Parameters->SpriteFacingAndAlignment_SpriteFacing		= ModuleBuiltData->SpriteFacing;
-		Parameters->SpriteFacingAndAlignment_SpriteAlignment	= ModuleBuiltData->SpriteAlignment;
+		Parameters->SpriteFacingAndAlignment_SpriteFacing		= SetShaderParameterContext.ConvertRangeToValue(ModuleBuiltData->SpriteFacing);
+		Parameters->SpriteFacingAndAlignment_SpriteAlignment	= SetShaderParameterContext.ConvertRangeToValue(ModuleBuiltData->SpriteAlignment);
 	}
 
 	static void ParticleSimulate(const NiagaraStateless::FParticleSimulationContext& ParticleSimulationContext)
@@ -92,14 +98,15 @@ public:
 		using namespace NiagaraStateless;
 
 		const FModuleBuiltData* ModuleBuiltData = ParticleSimulationContext.ReadBuiltData<FModuleBuiltData>();
+		const FParameters* Parameters = ParticleSimulationContext.ReadParameterNestedStruct<FParameters>();
 
 		for (uint32 i = 0; i < ParticleSimulationContext.GetNumInstances(); ++i)
 		{
-			ParticleSimulationContext.WriteParticleVariable(ModuleBuiltData->SpriteFacingVariableOffset, i, ModuleBuiltData->SpriteFacing);
-			ParticleSimulationContext.WriteParticleVariable(ModuleBuiltData->PreviousSpriteFacingVariableOffset, i, ModuleBuiltData->SpriteFacing);
+			ParticleSimulationContext.WriteParticleVariable(ModuleBuiltData->SpriteFacingVariableOffset, i, Parameters->SpriteFacingAndAlignment_SpriteFacing);
+			ParticleSimulationContext.WriteParticleVariable(ModuleBuiltData->PreviousSpriteFacingVariableOffset, i, Parameters->SpriteFacingAndAlignment_SpriteFacing);
 
-			ParticleSimulationContext.WriteParticleVariable(ModuleBuiltData->SpriteAlignmentVariable, i, ModuleBuiltData->SpriteAlignment);
-			ParticleSimulationContext.WriteParticleVariable(ModuleBuiltData->PreviousSpriteAlignmentVariable, i, ModuleBuiltData->SpriteAlignment);
+			ParticleSimulationContext.WriteParticleVariable(ModuleBuiltData->SpriteAlignmentVariable, i, Parameters->SpriteFacingAndAlignment_SpriteAlignment);
+			ParticleSimulationContext.WriteParticleVariable(ModuleBuiltData->PreviousSpriteAlignmentVariable, i, Parameters->SpriteFacingAndAlignment_SpriteAlignment);
 		}
 	}
 
diff --git a/Engine/Plugins/FX/Niagara/Source/Niagara/Internal/Stateless/NiagaraStatelessCommon.h b/Engine/Plugins/FX/Niagara/Source/Niagara/Internal/Stateless/NiagaraStatelessCommon.h
--- a/Engine/Plugins/FX/Niagara/Source/Niagara/Internal/Stateless/NiagaraStatelessCommon.h
+++ b/Engine/Plugins/FX/Niagara/Source/Niagara/Internal/Stateless/NiagaraStatelessCommon.h
@@ -51,6 +51,14 @@ private:
 	FTransform3f	Transforms[2];			// World to Local & Local to World
 };
 
+template<typename TType> struct FNiagaraStatelessRangeDefaultValue {};
+template<> struct FNiagaraStatelessRangeDefaultValue<int32> { static constexpr int32 Zero() { return 0; } };
+template<> struct FNiagaraStatelessRangeDefaultValue<float> { static constexpr float Zero() { return 0.0f; } };
+template<> struct FNiagaraStatelessRangeDefaultValue<FVector2f> { static const FVector2f Zero() { return FVector2f(0.0f, 0.0f); } };
+template<> struct FNiagaraStatelessRangeDefaultValue<FVector3f> { static const FVector3f Zero() { return FVector3f(0.0f, 0.0f, 0.0f); } };
+template<> struct FNiagaraStatelessRangeDefaultValue<FVector4f> { static const FVector4f Zero() { return FVector4f(0.0f, 0.0f, 0.0f, 0.0f); } };
+template<> struct FNiagaraStatelessRangeDefaultValue<FLinearColor> { static const FLinearColor Zero() { return FLinearColor(0.0f, 0.0f, 0.0f, 0.0f); } };
+
 template<typename TType>
 struct FNiagaraStatelessRange
 {
@@ -64,8 +72,8 @@ struct FNiagaraStatelessRange
 	ValueType GetScale() const { return Max - Min; }
 
 	int32		ParameterOffset = INDEX_NONE;
-	ValueType	Min = {};
-	ValueType	Max = {};
+	ValueType	Min = FNiagaraStatelessRangeDefaultValue<TType>::Zero();
+	ValueType	Max = FNiagaraStatelessRangeDefaultValue<TType>::Zero();
 };
 
 using FNiagaraStatelessRangeInt		= FNiagaraStatelessRange<int32>;
diff --git a/Engine/Plugins/FX/Niagara/Source/Niagara/Internal/Stateless/NiagaraStatelessDistribution.h b/Engine/Plugins/FX/Niagara/Source/Niagara/Internal/Stateless/NiagaraStatelessDistribution.h
--- a/Engine/Plugins/FX/Niagara/Source/Niagara/Internal/Stateless/NiagaraStatelessDistribution.h
+++ b/Engine/Plugins/FX/Niagara/Source/Niagara/Internal/Stateless/NiagaraStatelessDistribution.h
@@ -112,6 +112,7 @@ struct FNiagaraDistributionRangeFloat : public FNiagaraDistributionBase
 
 	FNiagaraDistributionRangeFloat() = default;
 	explicit FNiagaraDistributionRangeFloat(float ConstantValue) { InitConstant(ConstantValue); }
+	explicit FNiagaraDistributionRangeFloat(float MinValue, float MaxValue) { InitRange(MinValue, MaxValue); }
 
 	UPROPERTY(EditAnywhere, Category = "Parameters")
 	float Min = 0.0f;
@@ -120,6 +121,7 @@ struct FNiagaraDistributionRangeFloat : public FNiagaraDistributionBase
 	float Max = 0.0f;
 
 	NIAGARA_API void InitConstant(float Value);
+	NIAGARA_API void InitRange(float MinValue, float MaxValue);
 	NIAGARA_API FNiagaraStatelessRangeFloat CalculateRange(const float Default = 0.0f) const;
 
 #if WITH_EDITORONLY_DATA
@@ -154,6 +156,7 @@ struct FNiagaraDistributionRangeVector2 : public FNiagaraDistributionBase
 	virtual void UpdateValuesFromDistribution() override;
 	virtual FNiagaraTypeDefinition GetBindingTypeDef() const { return FNiagaraTypeDefinition::GetVec2Def(); }
 #endif
+	bool SerializeFromMismatchedTag(const struct FPropertyTag& Tag, FStructuredArchive::FSlot Slot);
 };
 
 USTRUCT()
@@ -345,6 +348,15 @@ struct TStructOpsTypeTraits<FNiagaraDistributionRangeFloat> : public TStructOpsT
 	};
 };
 
+template<>
+struct TStructOpsTypeTraits<FNiagaraDistributionRangeVector2> : public TStructOpsTypeTraitsBase2<FNiagaraDistributionRangeVector2>
+{
+	enum
+	{
+		WithStructuredSerializeFromMismatchedTag = true,
+	};
+};
+
 template<>
 struct TStructOpsTypeTraits<FNiagaraDistributionRangeVector3> : public TStructOpsTypeTraitsBase2<FNiagaraDistributionRangeVector3>
 {
diff --git a/Engine/Plugins/FX/Niagara/Source/Niagara/Internal/Stateless/NiagaraStatelessSetShaderParameterContext.h b/Engine/Plugins/FX/Niagara/Source/Niagara/Internal/Stateless/NiagaraStatelessSetShaderParameterContext.h
--- a/Engine/Plugins/FX/Niagara/Source/Niagara/Internal/Stateless/NiagaraStatelessSetShaderParameterContext.h
+++ b/Engine/Plugins/FX/Niagara/Source/Niagara/Internal/Stateless/NiagaraStatelessSetShaderParameterContext.h
@@ -10,15 +10,7 @@ class FNiagaraStatelessSetShaderParameterContext
 public:
 	UE_NONCOPYABLE(FNiagaraStatelessSetShaderParameterContext);
 
-	explicit FNiagaraStatelessSetShaderParameterContext(const FNiagaraStatelessSpaceTransforms& InSpaceTransforms, TConstArrayView<uint8> InRendererParameterData, TConstArrayView<uint8> InBuiltData, const FShaderParametersMetadata* InShaderParametersMetadata, uint8* InShaderParameters)
-		: SpaceTransforms(InSpaceTransforms)
-		, RendererParameterData(InRendererParameterData)
-		, BuiltData(InBuiltData)
-		, ShaderParametersBase(InShaderParameters)
-		, ParameterOffset(0)
-		, ShaderParametersMetadata(InShaderParametersMetadata)
-	{
-	}
+	explicit FNiagaraStatelessSetShaderParameterContext(const FNiagaraStatelessSpaceTransforms& InSpaceTransforms, TConstArrayView<uint8> InRendererParameterData, TConstArrayView<uint8> InBuiltData, const FShaderParametersMetadata* InShaderParametersMetadata, uint8* InShaderParameters);
 
 	const FNiagaraStatelessSpaceTransforms& GetSpaceTransforms() const { return SpaceTransforms; }
 
@@ -26,9 +18,9 @@ public:
 	T* GetParameterNestedStruct() const
 	{
 		const uint32 StructOffset = Align(ParameterOffset, TShaderParameterStructTypeInfo<T>::Alignment);
-	#if DO_CHECK
+#if DO_CHECK
 		ValidateIncludeStructType(StructOffset, TShaderParameterStructTypeInfo<T>::GetStructMetadata());
-	#endif
+#endif
 		ParameterOffset = StructOffset + TShaderParameterStructTypeInfo<T>::GetStructMetadata()->GetSize();
 		return reinterpret_cast<T*>(ShaderParametersBase + StructOffset);
 	}
@@ -42,30 +34,35 @@ public:
 		return reinterpret_cast<const T*>(BuiltData.GetData() + Offset);
 	}
 
-	template<typename T>
-	void GetRendererParameterValue(T& OutValue, int32 Offset, const T& DefaultValue) const
+	template<typename TValueType>
+	TValueType GetRendererParameterValue(int32 Offset, const TValueType& DefaultValue) const
 	{
 		if (Offset != INDEX_NONE)
 		{
+			TValueType Value;
 			Offset *= sizeof(uint32);
-			check(Offset >= 0 && Offset + sizeof(T) <= RendererParameterData.Num());
-			FMemory::Memcpy(&OutValue, RendererParameterData.GetData() + Offset, sizeof(T));
-		}
-		else
-		{
-			OutValue = DefaultValue;
+			check(Offset >= 0 && Offset + sizeof(TValueType) <= RendererParameterData.Num());
+			FMemory::Memcpy(&Value, RendererParameterData.GetData() + Offset, sizeof(TValueType));
+			return Value;
 		}
+		return DefaultValue;
+	}
+
+	template<typename TValueType>
+	void ConvertRangeToScaleBias(const FNiagaraStatelessRange<TValueType>& Range, TValueType& OutScale, TValueType& OutBias) const
+	{
+		OutBias = GetRendererParameterValue(Range.ParameterOffset, Range.Min);
+		OutScale = Range.ParameterOffset == INDEX_NONE ? Range.GetScale() : FNiagaraStatelessRangeDefaultValue<TValueType>::Zero();
 	}
 
-	void ConvertRangeToScaleBias(const FNiagaraStatelessRangeFloat& Range, float& OutScale, float& OutBias) const { OutScale = Range.GetScale(); GetRendererParameterValue(OutBias, Range.ParameterOffset, Range.Min); }
-	void ConvertRangeToScaleBias(const FNiagaraStatelessRangeVector2& Range, FVector2f& OutScale, FVector2f& OutBias) const { OutScale = Range.GetScale(); GetRendererParameterValue(OutBias, Range.ParameterOffset, Range.Min); }
-	void ConvertRangeToScaleBias(const FNiagaraStatelessRangeVector3& Range, FVector3f& OutScale, FVector3f& OutBias) const { OutScale = Range.GetScale(); GetRendererParameterValue(OutBias, Range.ParameterOffset, Range.Min); }
-	void ConvertRangeToScaleBias(const FNiagaraStatelessRangeColor& Range, FLinearColor& OutScale, FLinearColor& OutBias) const { OutScale = Range.GetScale(); GetRendererParameterValue(OutBias, Range.ParameterOffset, Range.Min); }
+	template<typename TValueType>
+	TValueType ConvertRangeToValue(const FNiagaraStatelessRange<TValueType>& Range) const
+	{
+		return GetRendererParameterValue(Range.ParameterOffset, Range.Min);
+	}
 
-	float ConvertRangeToValue(const FNiagaraStatelessRangeFloat& Range) const { float OutValue; GetRendererParameterValue(OutValue, Range.ParameterOffset, Range.Min); return OutValue; }
-	FVector2f ConvertRangeToValue(const FNiagaraStatelessRangeVector2& Range) const { FVector2f OutValue; GetRendererParameterValue(OutValue, Range.ParameterOffset, Range.Min); return OutValue; }
-	FVector3f ConvertRangeToValue(const FNiagaraStatelessRangeVector3& Range) const { FVector3f OutValue; GetRendererParameterValue(OutValue, Range.ParameterOffset, Range.Min); return OutValue; }
-	FLinearColor ConvertRangeToValue(const FNiagaraStatelessRangeColor& Range) const { FLinearColor OutValue; GetRendererParameterValue(OutValue, Range.ParameterOffset, Range.Min); return OutValue; }
+	void TransformVectorRangeToScaleBias(const FNiagaraStatelessRangeVector3& Range, ENiagaraCoordinateSpace SourceSpace, FVector3f& OutScale, FVector3f& OutBias) const;
+	FVector3f TransformPositionRangeToValue(const FNiagaraStatelessRangeVector3& Range, ENiagaraCoordinateSpace SourceSpace) const;
 
 protected:
 #if DO_CHECK
diff --git a/Engine/Plugins/FX/Niagara/Source/Niagara/Private/Stateless/Modules/NiagaraStatelessModule_AddVelocity.cpp b/Engine/Plugins/FX/Niagara/Source/Niagara/Private/Stateless/Modules/NiagaraStatelessModule_AddVelocity.cpp
--- a/Engine/Plugins/FX/Niagara/Source/Niagara/Private/Stateless/Modules/NiagaraStatelessModule_AddVelocity.cpp
+++ b/Engine/Plugins/FX/Niagara/Source/Niagara/Private/Stateless/Modules/NiagaraStatelessModule_AddVelocity.cpp
@@ -13,10 +13,10 @@ void UNiagaraStatelessModule_AddVelocity::BuildEmitterData(const FNiagaraStatele
 	NiagaraStateless::FPhysicsBuildData& PhysicsBuildData = BuildContext.GetTransientBuildData<NiagaraStateless::FPhysicsBuildData>();
 	if (VelocityType == ENSM_VelocityType::Linear)
 	{
-		const FNiagaraStatelessRangeVector3 VelocityRange = LinearVelocityDistribution.CalculateRange(FVector3f::ZeroVector);
+		const FNiagaraStatelessRangeVector3 VelocityRange = BuildContext.ConvertDistributionToRange(LinearVelocityDistribution, FVector3f::ZeroVector);
 		PhysicsBuildData.VelocityCoordinateSpace = CoordinateSpace;
-		PhysicsBuildData.VelocityRange.Min += VelocityRange.Min * LinearVelocityScale;
-		PhysicsBuildData.VelocityRange.Max += VelocityRange.Max * LinearVelocityScale;
+		PhysicsBuildData.VelocityRange = BuildContext.ConvertDistributionToRange(LinearVelocityDistribution, FVector3f::ZeroVector);
+		PhysicsBuildData.LinearVelocityScale = BuildContext.ConvertDistributionToRange(LinearVelocityScale, 1.0f);
 	}
 	else if (VelocityType == ENSM_VelocityType::FromPoint)
 	{
@@ -24,7 +24,7 @@ void UNiagaraStatelessModule_AddVelocity::BuildEmitterData(const FNiagaraStatele
 
 		PhysicsBuildData.bPointVelocity = true;
 		PhysicsBuildData.PointCoordinateSpace = CoordinateSpace;
-		PhysicsBuildData.PointVelocityRange = PointVelocityDistribution.CalculateRange(0.0f);
+		PhysicsBuildData.PointVelocityRange = BuildContext.ConvertDistributionToRange(PointVelocityDistribution, 0.0f);
 		PhysicsBuildData.PointOrigin = PointOrigin;
 	}
 	else if (VelocityType == ENSM_VelocityType::InCone)
@@ -34,7 +34,7 @@ void UNiagaraStatelessModule_AddVelocity::BuildEmitterData(const FNiagaraStatele
 		PhysicsBuildData.bConeVelocity = true;
 		PhysicsBuildData.ConeCoordinateSpace = CoordinateSpace;
 		PhysicsBuildData.ConeQuat = FQuat4f(ConeRotation.Quaternion());
-		PhysicsBuildData.ConeVelocityRange = ConeVelocityDistribution.CalculateRange(0.0f);
+		PhysicsBuildData.ConeVelocityRange = BuildContext.ConvertDistributionToRange(ConeVelocityDistribution, 0.0f);
 		PhysicsBuildData.ConeOuterAngle = ConeAngle;
 		PhysicsBuildData.ConeInnerAngle = InnerCone;
 		PhysicsBuildData.ConeVelocityFalloff = bSpeedFalloffFromConeAxisEnabled ? FMath::Clamp(SpeedFalloffFromConeAxis, 0.0f, 1.0f) : 0.0f;
@@ -49,8 +49,9 @@ void UNiagaraStatelessModule_AddVelocity::DrawDebug(const FNiagaraStatelessDrawD
 		case ENSM_VelocityType::Linear:
 		{
 			const FNiagaraStatelessRangeVector3 VelocityRange = LinearVelocityDistribution.CalculateRange(FVector3f::ZeroVector);
-			const FVector MinDir = FVector(VelocityRange.Min * LinearVelocityScale);
-			const FVector MaxDir = FVector(VelocityRange.Max * LinearVelocityScale);
+			const FNiagaraStatelessRangeFloat LinearVelocityScaleRange = LinearVelocityScale.CalculateRange();
+			const FVector MinDir = FVector(VelocityRange.Min * LinearVelocityScale.Min);
+			const FVector MaxDir = FVector(VelocityRange.Max * LinearVelocityScale.Max);
 			DrawDebugContext.DrawArrow(FVector::ZeroVector, MinDir);
 
 			if (!FMath::IsNearlyEqual(MinDir.Length(), MaxDir.Length()))
diff --git a/Engine/Plugins/FX/Niagara/Source/Niagara/Private/Stateless/Modules/NiagaraStatelessModule_ShapeLocation.cpp b/Engine/Plugins/FX/Niagara/Source/Niagara/Private/Stateless/Modules/NiagaraStatelessModule_ShapeLocation.cpp
--- a/Engine/Plugins/FX/Niagara/Source/Niagara/Private/Stateless/Modules/NiagaraStatelessModule_ShapeLocation.cpp
+++ b/Engine/Plugins/FX/Niagara/Source/Niagara/Private/Stateless/Modules/NiagaraStatelessModule_ShapeLocation.cpp
@@ -10,50 +10,99 @@ namespace NSMShapeLocationPrivate
 {
 	struct FModuleBuiltData
 	{
-		FUintVector4	Mode		= FUintVector4::ZeroValue;
-		FVector4f		Parameters0	= FVector4f::Zero();
-		FVector4f		Parameters1 = FVector4f::Zero();
+		ENSM_ShapePrimitive	ShapePrimitive = ENSM_ShapePrimitive::Max;
+		int32				PositionVariableOffset = INDEX_NONE;
+		int32				PreviousPositionVariableOffset = INDEX_NONE;
+	};
+
+	struct FModuleBuiltDataBox
+	{
+		bool							bOnSurfaceOnly = false;
+		FNiagaraStatelessRangeVector3	Size;
+		ENSM_SurfaceExpansionMode		SurfaceExpansionMode = ENSM_SurfaceExpansionMode::Centered;
+		FNiagaraStatelessRangeFloat		SurfaceThickness;
+	};
 
-		int32			PositionVariableOffset = INDEX_NONE;
-		int32			PreviousPositionVariableOffset = INDEX_NONE;
+	struct FModuleBuiltDataCylinder
+	{
+		FNiagaraStatelessRangeFloat	Height;
+		FNiagaraStatelessRangeFloat	Radius;
+		FNiagaraStatelessRangeFloat	MidPoint;
+	};
+
+	struct FModuleBuiltDataPlane
+	{
+		bool							bOnEdgesOnly = false;
+		FNiagaraStatelessRangeVector2	Size;
+		ENSM_SurfaceExpansionMode		EdgeExpansionMode = ENSM_SurfaceExpansionMode::Centered;
+		FNiagaraStatelessRangeFloat		EdgeThickness;
+	};
+
+	struct FModuleBuiltDataRing
+	{
+		FNiagaraStatelessRangeFloat	Radius;
+		FNiagaraStatelessRangeFloat	Coverage;
+		FNiagaraStatelessRangeFloat	Distribution;
 	};
 
-	FVector3f ShapeLocation_GetLocation(const NiagaraStateless::FParticleSimulationContext& ParticleSimulationContext, const FModuleBuiltData* ModuleBuiltData, uint32 iInstance)
+	struct FModuleBuiltDataSphere
+	{
+		FNiagaraStatelessRangeFloat	Radius = FNiagaraStatelessRangeFloat(0.0f);
+	};
+
+	TTuple<float, float> CalculateThicknessParameters(ENSM_SurfaceExpansionMode ExpansionMode, float Thickness)
+	{
+		const float HThickness = Thickness * 0.5f;
+		switch (ExpansionMode)
+		{
+			case ENSM_SurfaceExpansionMode::Inner:		return MakeTuple(-HThickness, Thickness);
+			case ENSM_SurfaceExpansionMode::Centered:	return MakeTuple(0.0f, Thickness);
+			case ENSM_SurfaceExpansionMode::Outside:	return MakeTuple(HThickness, Thickness);
+		}
+		checkNoEntry();
+		return MakeTuple(0.0f, 0.0f);
+	}
+
+	FVector3f ShapeLocation_GetLocation(const NiagaraStateless::FParticleSimulationContext& ParticleSimulationContext, const UNiagaraStatelessModule_ShapeLocation::FParameters* ShaderParameters, uint32 iInstance)
 	{
 		// ENSM_ShapePrimitive::Box | ENSM_ShapePrimitive::Plane
-		if ( ModuleBuiltData->Mode.X == 0 )
+		if ( ShaderParameters->ShapeLocation_Mode.X == 0 )
 		{
-			const FVector3f BoxScale = FVector3f(ModuleBuiltData->Parameters0);
-			const FVector3f BoxBias = FVector3f(ModuleBuiltData->Parameters1);
-			const bool bOnSurface = ModuleBuiltData->Mode.Y == 1;
-			const float SurfaceScale = ModuleBuiltData->Parameters0.W;
-			const float SurfaceBias = ModuleBuiltData->Parameters1.W;
-	
-			const FVector3f P0 = ParticleSimulationContext.RandomFloat3(iInstance, 0);
+			const FVector3f BoxSize = FVector3f(ShaderParameters->ShapeLocation_Parameters0);
+			const bool bOnSurface = ShaderParameters->ShapeLocation_Mode.Y == 1;
+
+			const FVector3f P0 = ParticleSimulationContext.RandomFloat3(iInstance, 0) - 0.5f;
 			if (bOnSurface)
 			{
-				//-TODO: This isn't quite what we want
-				const FVector3f L0(FMath::RoundToFloat(P0.X), FMath::RoundToFloat(P0.Y), FMath::RoundToFloat(P0.Z));
-				const uint32 S = ParticleSimulationContext.RandomUInt(iInstance, 1) % 3;
+				static const FVector3f FaceDirA[] = { FVector3f(0.0f, 1.0f, 0.0f), FVector3f(0.0f, 1.0f, 0.0f),  FVector3f(0.0f, 0.0f, 1.0f), FVector3f(0.0f, 0.0f, 1.0f), FVector3f(1.0f, 0.0f, 0.0f), FVector3f(1.0f, 0.0f, 0.0f) };
+				static const FVector3f FaceDirB[] = { FVector3f(0.0f, 0.0f, 1.0f), FVector3f(0.0f, 0.0f, 1.0f),  FVector3f(1.0f, 0.0f, 0.0f), FVector3f(1.0f, 0.0f, 0.0f), FVector3f(0.0f, 1.0f, 0.0f), FVector3f(0.0f, 1.0f, 0.0f) };
+				static const FVector3f FaceDirC[] = { FVector3f(1.0f, 0.0f, 0.0f), FVector3f(-1.0f, 0.0f, 0.0f), FVector3f(0.0f, 1.0f, 0.0f), FVector3f(0.0f,-1.0f, 0.0f), FVector3f(0.0f, 0.0f, 1.0f), FVector3f(0.0f, 0.0f,-1.0f) };
+				const uint32 NumFaces = ShaderParameters->ShapeLocation_Mode.Z;
+
+				const FVector3f Thickness = FVector3f(ShaderParameters->ShapeLocation_Parameters1);
+				const FVector3f LocalBoxSize = BoxSize + (Thickness * P0.Z);
+
+				const uint32 iAxis = ParticleSimulationContext.RandomUInt(iInstance, 1) % NumFaces;
 				FVector3f Location;
-				Location.X = S != 0 ? P0.X : L0.X;
-				Location.Y = S != 1 ? P0.Y : L0.Y;
-				Location.Z = S != 2 ? P0.Z : L0.Z;
-				return (Location * BoxScale + BoxBias) + (P0 * SurfaceScale + SurfaceBias);
+				Location  = LocalBoxSize * FaceDirA[iAxis] * P0.X;
+				Location += LocalBoxSize * FaceDirB[iAxis] * P0.Y;
+				Location += LocalBoxSize * FaceDirC[iAxis] * 0.5f;
+
+				return Location;
 			}
 			else
 			{
-				return P0 * BoxScale + BoxBias;
+				return P0 * BoxSize;
 			}
 		}
 	
 		// ENSM_ShapePrimitive::Cylinder:
-		if ( ModuleBuiltData->Mode.X == 1 )
+		if ( ShaderParameters->ShapeLocation_Mode.X == 1 )
 		{
 			const FVector4f Random = ParticleSimulationContext.RandomFloat4(iInstance, 0);
-			const float HeightScale = ModuleBuiltData->Parameters0.X;
-			const float HeightBias = ModuleBuiltData->Parameters0.Y;
-			const float Radius = ModuleBuiltData->Parameters0.Z;
+			const float HeightScale = ShaderParameters->ShapeLocation_Parameters0.X;
+			const float HeightBias = ShaderParameters->ShapeLocation_Parameters0.Y;
+			const float Radius = ShaderParameters->ShapeLocation_Parameters0.Z;
 	
 			const FVector2f UnitVec = ParticleSimulationContext.SafeNormalize(FVector2f(Random.X - 0.5f, Random.Y - 0.5f));
 			return FVector3f(
@@ -64,12 +113,12 @@ namespace NSMShapeLocationPrivate
 		}
 	
 		// ENSM_ShapePrimitive::Ring:
-		if ( ModuleBuiltData->Mode.X == 2 )
+		if ( ShaderParameters->ShapeLocation_Mode.X == 2 )
 		{
-			const float RadiusScale = ModuleBuiltData->Parameters0.X;
-			const float RadiusBias = ModuleBuiltData->Parameters0.Y;
-			const float UDistributionScale = ModuleBuiltData->Parameters0.Z;
-			const float UDistributionBias = ModuleBuiltData->Parameters0.W;
+			const float RadiusScale = ShaderParameters->ShapeLocation_Parameters0.X;
+			const float RadiusBias = ShaderParameters->ShapeLocation_Parameters0.Y;
+			const float UDistributionScale = ShaderParameters->ShapeLocation_Parameters0.Z;
+			const float UDistributionBias = ShaderParameters->ShapeLocation_Parameters0.W;
 	
 			const float Radius = ParticleSimulationContext.RandomScaleBiasFloat(iInstance, 0, RadiusScale, RadiusBias);
 			const float U = ParticleSimulationContext.RandomScaleBiasFloat(iInstance, 1, UDistributionScale, UDistributionBias);
@@ -79,8 +128,8 @@ namespace NSMShapeLocationPrivate
 	
 		// ENSM_ShapePrimitive::Sphere:
 		{
-			const float SphereScale	= ModuleBuiltData->Parameters0.X;
-			const float SphereBias	= ModuleBuiltData->Parameters0.Y;
+			const float SphereScale	= ShaderParameters->ShapeLocation_Parameters0.X;
+			const float SphereBias	= ShaderParameters->ShapeLocation_Parameters0.Y;
 	
 			const FVector3f Vector = ParticleSimulationContext.RandomUnitFloat3(iInstance, 0);
 			return Vector * ParticleSimulationContext.RandomScaleBiasFloat(iInstance, 1, SphereScale, SphereBias);
@@ -92,10 +141,11 @@ namespace NSMShapeLocationPrivate
 		using namespace NiagaraStateless;
 
 		const FModuleBuiltData* ModuleBuiltData = ParticleSimulationContext.ReadBuiltData<FModuleBuiltData>();
+		const UNiagaraStatelessModule_ShapeLocation::FParameters* ShaderParameters = ParticleSimulationContext.ReadParameterNestedStruct<UNiagaraStatelessModule_ShapeLocation::FParameters>();
 
 		for (uint32 i = 0; i < ParticleSimulationContext.GetNumInstances(); ++i)
 		{
-			const FVector3f ShapeLocation = ShapeLocation_GetLocation(ParticleSimulationContext, ModuleBuiltData, i);
+			const FVector3f ShapeLocation = ShapeLocation_GetLocation(ParticleSimulationContext, ShaderParameters, i);
 			const FVector3f Position = ParticleSimulationContext.ReadParticleVariable(ModuleBuiltData->PositionVariableOffset, i, FVector3f::ZeroVector);
 			const FVector3f PreviousPosition = ParticleSimulationContext.ReadParticleVariable(ModuleBuiltData->PreviousPositionVariableOffset, i, FVector3f::ZeroVector);
 
@@ -123,58 +173,53 @@ void UNiagaraStatelessModule_ShapeLocation::BuildEmitterData(const FNiagaraState
 		return;
 	}
 
-	//-TODO: Build Parameters
+	BuiltData->ShapePrimitive = ShapePrimitive;
+
 	switch (ShapePrimitive)
 	{
 		case ENSM_ShapePrimitive::Box:
 		{
-			BuiltData->Mode.X		= 0;
-			BuiltData->Mode.Y		= bBoxSurfaceOnly ? 1 : 0;
-			BuiltData->Parameters0	= FVector4f(BoxSize, BoxSurfaceThicknessMax - BoxSurfaceThicknessMin);
-			BuiltData->Parameters1	= FVector4f(BoxSize * -0.5f, BoxSurfaceThicknessMin);
+			FModuleBuiltDataBox* BuiltDataBox	= BuildContext.AllocateBuiltData<FModuleBuiltDataBox>();
+			BuiltDataBox->bOnSurfaceOnly		= bBoxSurfaceOnly;
+			BuiltDataBox->Size					= BuildContext.ConvertDistributionToRange(BoxSize, FVector3f::ZeroVector);
+			BuiltDataBox->SurfaceExpansionMode	= BoxSurfaceExpansion;
+			BuiltDataBox->SurfaceThickness		= BuildContext.ConvertDistributionToRange(BoxSurfaceThickness, 0.0f);
 			break;
 		}
-		case ENSM_ShapePrimitive::Plane:
+		case ENSM_ShapePrimitive::Cylinder:
 		{
-			BuiltData->Mode.X		= 0;
-			BuiltData->Parameters0	= FVector4f(PlaneSize.X, PlaneSize.Y, 0.0f, 0.0f);
-			BuiltData->Parameters1	= FVector4f(-PlaneSize.X * 0.5f, -PlaneSize.Y * 0.5f, 0.0f, 0.0f);
+			FModuleBuiltDataCylinder* BuiltDataCylinder = BuildContext.AllocateBuiltData<FModuleBuiltDataCylinder>();
+			BuiltDataCylinder->Height	= BuildContext.ConvertDistributionToRange(CylinderHeight, 0.0f);
+			BuiltDataCylinder->Radius	= BuildContext.ConvertDistributionToRange(CylinderRadius, 0.0f);
+			BuiltDataCylinder->MidPoint	= BuildContext.ConvertDistributionToRange(CylinderHeightMidpoint, 0.0f);
 			break;
 		}
-		case ENSM_ShapePrimitive::Cylinder:
+		case ENSM_ShapePrimitive::Plane:
 		{
-			BuiltData->Mode.X		= 1;
-			BuiltData->Parameters0.X = CylinderHeight;
-			BuiltData->Parameters0.Y = CylinderHeight * -CylinderHeightMidpoint;
-			BuiltData->Parameters0.Z = CylinderRadius;
+			FModuleBuiltDataPlane* BuiltDataPlane = BuildContext.AllocateBuiltData<FModuleBuiltDataPlane>();
+			BuiltDataPlane->bOnEdgesOnly		= bPlaneEdgesOnly;
+			BuiltDataPlane->Size				= BuildContext.ConvertDistributionToRange(PlaneSize, FVector2f::ZeroVector);
+			BuiltDataPlane->EdgeExpansionMode	= PlaneEdgeExpansion;
+			BuiltDataPlane->EdgeThickness		= BuildContext.ConvertDistributionToRange(PlaneEdgeThickness, 0.0f);
 			break;
 		}
 		case ENSM_ShapePrimitive::Ring:
 		{
-			const float DC = FMath::Clamp(1.0f - DiscCoverage, 0.0f, 1.0f);
-			const float SDC = DC > 0.0f ? FMath::Sqrt(DC) : 0.0f;
-
-			BuiltData->Mode.X		= 2;
-			BuiltData->Parameters0.X = RingRadius * (1.0f - SDC);
-			BuiltData->Parameters0.Y = RingRadius * SDC;
-			BuiltData->Parameters0.Z = -UE_TWO_PI * (1.0f - RingUDistribution);
-			BuiltData->Parameters0.W = 0.0f;
+			FModuleBuiltDataRing* BuiltDataRing = BuildContext.AllocateBuiltData<FModuleBuiltDataRing>();
+			BuiltDataRing->Radius		= BuildContext.ConvertDistributionToRange(RingRadius, 0.0f);
+			BuiltDataRing->Coverage		= BuildContext.ConvertDistributionToRange(DiscCoverage, 0.0f);
+			BuiltDataRing->Distribution	= BuildContext.ConvertDistributionToRange(RingUDistribution, 0.0f);
 			break;
 		}
 		case ENSM_ShapePrimitive::Sphere:
 		{
-			BuiltData->Mode.X		= 3;
-			BuiltData->Parameters0.X = SphereMax - SphereMin;
-			BuiltData->Parameters0.Y = SphereMin;
+			FModuleBuiltDataSphere* BuiltDataSphere = BuildContext.AllocateBuiltData<FModuleBuiltDataSphere>();
+			BuiltDataSphere->Radius	= BuildContext.ConvertDistributionToRange(SphereRadius, 0.0f);
 			break;
 		}
 		default:
-		{
-			BuiltData->Mode.X = 0;
-			BuiltData->Parameters0 = FVector4f::Zero();
-			BuiltData->Parameters1 = FVector4f::Zero();
 			checkNoEntry();
-		}
+			return;
 	}
 
 	BuildContext.AddParticleSimulationExecSimulate(&ParticleSimulate);
@@ -192,9 +237,86 @@ void UNiagaraStatelessModule_ShapeLocation::SetShaderParameters(const FNiagaraSt
 	FParameters* Parameters = SetShaderParameterContext.GetParameterNestedStruct<FParameters>();
 	const FModuleBuiltData* ModuleBuiltData = SetShaderParameterContext.ReadBuiltData<FModuleBuiltData>();
 
-	Parameters->ShapeLocation_Mode			= ModuleBuiltData->Mode;
-	Parameters->ShapeLocation_Parameters0	= ModuleBuiltData->Parameters0;
-	Parameters->ShapeLocation_Parameters1	= ModuleBuiltData->Parameters1;
+	Parameters->ShapeLocation_Mode			= FUintVector4::ZeroValue;
+	Parameters->ShapeLocation_Parameters0	= FVector4f::Zero();
+	Parameters->ShapeLocation_Parameters1	= FVector4f::Zero();
+
+	switch (ModuleBuiltData->ShapePrimitive)
+	{
+		case ENSM_ShapePrimitive::Box:
+		{
+			const FModuleBuiltDataBox* BuiltDataBox = SetShaderParameterContext.ReadBuiltData<FModuleBuiltDataBox>();
+			const FVector3f Size	= SetShaderParameterContext.ConvertRangeToValue(BuiltDataBox->Size);
+			const float Thickness	= SetShaderParameterContext.ConvertRangeToValue(BuiltDataBox->SurfaceThickness);
+			const TTuple<float, float> ThicknessParameters = CalculateThicknessParameters(BuiltDataBox->SurfaceExpansionMode, Thickness);
+
+			Parameters->ShapeLocation_Mode.X		= 0;
+			Parameters->ShapeLocation_Mode.Y		= BuiltDataBox->bOnSurfaceOnly ? 1 : 0;
+			Parameters->ShapeLocation_Mode.Z		= 6;
+			Parameters->ShapeLocation_Parameters0	= FVector4f(Size + ThicknessParameters.Key, 0.0f);
+			Parameters->ShapeLocation_Parameters1	= FVector4f(ThicknessParameters.Value, ThicknessParameters.Value, ThicknessParameters.Value, 0.0f);
+
+			break;
+		}
+		case ENSM_ShapePrimitive::Cylinder:
+		{
+			const FModuleBuiltDataCylinder* BuiltDataCylinder = SetShaderParameterContext.ReadBuiltData<FModuleBuiltDataCylinder>();
+			const float Height		= SetShaderParameterContext.ConvertRangeToValue(BuiltDataCylinder->Height);
+			const float Radius		= SetShaderParameterContext.ConvertRangeToValue(BuiltDataCylinder->Radius);
+			const float MidPoint	= SetShaderParameterContext.ConvertRangeToValue(BuiltDataCylinder->MidPoint);
+
+			Parameters->ShapeLocation_Mode.X		= 1;
+			Parameters->ShapeLocation_Parameters0.X	= Height;
+			Parameters->ShapeLocation_Parameters0.Y	= Height * -MidPoint;
+			Parameters->ShapeLocation_Parameters0.Z	= Radius;
+			break;
+		}
+		case ENSM_ShapePrimitive::Plane:
+		{
+			const FModuleBuiltDataPlane* BuiltDataPlane = SetShaderParameterContext.ReadBuiltData<FModuleBuiltDataPlane>();
+			const FVector2f Size	= SetShaderParameterContext.ConvertRangeToValue(BuiltDataPlane->Size);
+			const float Thickness	= SetShaderParameterContext.ConvertRangeToValue(BuiltDataPlane->EdgeThickness);
+			const TTuple<float, float> ThicknessParameters = CalculateThicknessParameters(BuiltDataPlane->EdgeExpansionMode, Thickness);
+
+			Parameters->ShapeLocation_Mode.X		= 0;
+			Parameters->ShapeLocation_Mode.Y		= BuiltDataPlane->bOnEdgesOnly ? 1 : 0;
+			Parameters->ShapeLocation_Mode.Z		= 4;
+			Parameters->ShapeLocation_Parameters0	= FVector4f(Size.X + ThicknessParameters.Key, Size.Y + ThicknessParameters.Key, 0.0f, 0.0f);
+			Parameters->ShapeLocation_Parameters1	= FVector4f(ThicknessParameters.Value, ThicknessParameters.Value, 0.0f, 0.0f);
+			break;
+		}
+		case ENSM_ShapePrimitive::Ring:
+		{
+			const FModuleBuiltDataRing* BuiltDataRing = SetShaderParameterContext.ReadBuiltData<FModuleBuiltDataRing>();
+			const float Radius			= SetShaderParameterContext.ConvertRangeToValue(BuiltDataRing->Radius);
+			const float Coverage		= SetShaderParameterContext.ConvertRangeToValue(BuiltDataRing->Coverage);
+			const float Distribution	= SetShaderParameterContext.ConvertRangeToValue(BuiltDataRing->Distribution);
+
+			const float DC = FMath::Clamp(1.0f - Coverage, 0.0f, 1.0f);
+			const float SDC = DC > 0.0f ? FMath::Sqrt(DC) : 0.0f;
+
+			Parameters->ShapeLocation_Mode.X		= 2;
+			Parameters->ShapeLocation_Parameters0.X = Radius * (1.0f - SDC);
+			Parameters->ShapeLocation_Parameters0.Y = Radius * SDC;
+			Parameters->ShapeLocation_Parameters0.Z = -UE_TWO_PI * (1.0f - Distribution);
+			Parameters->ShapeLocation_Parameters0.W = 0.0f;
+			break;
+		}
+		case ENSM_ShapePrimitive::Sphere:
+		{
+			const FModuleBuiltDataSphere* BuiltDataSphere = SetShaderParameterContext.ReadBuiltData<FModuleBuiltDataSphere>();
+			float SphereScale = 0.0f;
+			float SphereBias = 0.0f;
+			SetShaderParameterContext.ConvertRangeToScaleBias(BuiltDataSphere->Radius, SphereScale, SphereBias);
+
+			Parameters->ShapeLocation_Mode.X		= 3;
+			Parameters->ShapeLocation_Parameters0.X = SphereScale;
+			Parameters->ShapeLocation_Parameters0.Y = SphereBias;
+			break;
+		}
+		default:
+			break;
+	}
 }
 
 #if WITH_EDITOR
@@ -204,33 +326,61 @@ void UNiagaraStatelessModule_ShapeLocation::DrawDebug(const FNiagaraStatelessDra
 	{
 		case ENSM_ShapePrimitive::Box:
 		{
-			DrawDebugContext.DrawBox(FVector::ZeroVector, FVector(BoxSize * 0.5f));
+			if (BoxSize.IsBinding() == false)
+			{
+				DrawDebugContext.DrawBox(FVector::ZeroVector, FVector(BoxSize.Min * 0.5f));
+			}
 			break;
 		}
 		case ENSM_ShapePrimitive::Plane:
 		{
-			DrawDebugContext.DrawBox(FVector::ZeroVector, FVector(PlaneSize.X * 0.5f, PlaneSize.Y * 0.5f, 0.0f));
+			if (PlaneSize.IsBinding() == false)
+			{
+				DrawDebugContext.DrawBox(FVector::ZeroVector, FVector(PlaneSize.Min.X * 0.5f, PlaneSize.Min.Y * 0.5f, 0.0f));
+			}
 			break;
 		}
 		case ENSM_ShapePrimitive::Cylinder:
 		{
-			DrawDebugContext.DrawCylinder(CylinderHeight, CylinderRadius, CylinderHeightMidpoint);
+			if (CylinderHeight.IsBinding() == false && CylinderRadius.IsBinding() == false && CylinderHeightMidpoint.IsBinding() == false)
+			{
+				DrawDebugContext.DrawCylinder(CylinderHeight.Min, CylinderRadius.Min, CylinderHeightMidpoint.Min);
+			}
 			break;
 		}
 		case ENSM_ShapePrimitive::Ring:
 		{
-			const float DC = FMath::Clamp(1.0f - DiscCoverage, 0.0f, 1.0f);
-			const float SDC = DC > 0.0f ? FMath::Sqrt(DC) : 0.0f;
-			DrawDebugContext.DrawCircle(FVector::ZeroVector, RingRadius);
-			DrawDebugContext.DrawCircle(FVector::ZeroVector, RingRadius * SDC);
+			if (RingRadius.IsBinding() == false && DiscCoverage.IsBinding() == false && RingUDistribution.IsBinding() == false)
+			{
+				const float DC = FMath::Clamp(1.0f - DiscCoverage.Min, 0.0f, 1.0f);
+				const float SDC = DC > 0.0f ? FMath::Sqrt(DC) : 0.0f;
+				DrawDebugContext.DrawCircle(FVector::ZeroVector, RingRadius.Min);
+				DrawDebugContext.DrawCircle(FVector::ZeroVector, RingRadius.Min * SDC);
+			}
 			break;
 		}
 		case ENSM_ShapePrimitive::Sphere:
 		{
-			DrawDebugContext.DrawSphere(FVector::ZeroVector, SphereMin);
-			DrawDebugContext.DrawSphere(FVector::ZeroVector, SphereMax);
+			if (SphereRadius.IsBinding() == false)
+			{
+				DrawDebugContext.DrawSphere(FVector::ZeroVector, SphereRadius.Min);
+				if (SphereRadius.IsRange())
+				{
+					DrawDebugContext.DrawSphere(FVector::ZeroVector, SphereRadius.Max);
+				}
+			}
 			break;
 		}
 	}
 }
+
+void UNiagaraStatelessModule_ShapeLocation::PostLoad()
+{
+	Super::PostLoad();
+
+	if (!FMath::IsNearlyEqual(SphereMin_DEPRECATED, 0.0f) || !FMath::IsNearlyEqual(SphereMax_DEPRECATED, 100.0f))
+	{
+		SphereRadius = FNiagaraDistributionRangeFloat(SphereMin_DEPRECATED, SphereMax_DEPRECATED);
+	}
+}
 #endif
diff --git a/Engine/Plugins/FX/Niagara/Source/Niagara/Private/Stateless/Modules/NiagaraStatelessModule_SolveVelocitiesAndForces.cpp b/Engine/Plugins/FX/Niagara/Source/Niagara/Private/Stateless/Modules/NiagaraStatelessModule_SolveVelocitiesAndForces.cpp
new file mode 100644
--- /dev/null
+++ b/Engine/Plugins/FX/Niagara/Source/Niagara/Private/Stateless/Modules/NiagaraStatelessModule_SolveVelocitiesAndForces.cpp
@@ -0,0 +1,216 @@
+// Copyright Epic Games, Inc. All Rights Reserved.
+
+#include "Stateless/Modules/NiagaraStatelessModule_SolveVelocitiesAndForces.h"
+
+#include "VectorField/VectorField.h"
+#include "VectorField/VectorFieldStatic.h"
+#include "RHIStaticStates.h"
+
+namespace NiagaraStatelessModuleSolveVelocitiesAndForcesPrivate
+{
+	struct FModuleBuiltData
+	{
+		NiagaraStateless::FPhysicsBuildData	PhysicsData;
+		int32								PositionVariableOffset = INDEX_NONE;
+		int32								VelocityVariableOffset = INDEX_NONE;
+		int32								PreviousPositionVariableOffset = INDEX_NONE;
+		int32								PreviousVelocityVariableOffset = INDEX_NONE;
+	};
+
+	static FVector3f IntegratePosition(float Age, float Mass, float Drag, const FVector3f& Velocity, const FVector3f& Wind, const FVector3f& Acceleration)
+	{
+		const FVector3f IntVelocity = (Velocity - Wind) + (Wind * Age * Age);
+		const float LambdaDragMass = FMath::Max(Drag * (1.0f / Mass), 0.0001f);
+		const float LambdaAge = (1.0f - FMath::Exp(-(LambdaDragMass * Age))) / LambdaDragMass;
+		FVector3f Position = IntVelocity * LambdaAge;
+		Position += (Age - LambdaAge) * (Acceleration / LambdaDragMass);
+		return Position;
+	}
+
+	static void ParticleSimulate(const NiagaraStateless::FParticleSimulationContext& ParticleSimulationContext)
+	{
+		const FModuleBuiltData* ModuleBuiltData = ParticleSimulationContext.ReadBuiltData<FModuleBuiltData>();
+		const UNiagaraStatelessModule_SolveVelocitiesAndForces::FParameters* Parameters = ParticleSimulationContext.ReadParameterNestedStruct<UNiagaraStatelessModule_SolveVelocitiesAndForces::FParameters>();
+
+		const float* AgeData			= ParticleSimulationContext.GetParticleAge();
+		const float* PreviousAgeData	= ParticleSimulationContext.GetParticlePreviousAge();
+		const float InvDeltaTime		= ParticleSimulationContext.GetInvDeltaTime();
+
+		for (uint32 i = 0; i < ParticleSimulationContext.GetNumInstances(); ++i)
+		{
+			const float Mass				= ParticleSimulationContext.RandomScaleBiasFloat(i, 0, Parameters->SolveVelocitiesAndForces_MassScale, Parameters->SolveVelocitiesAndForces_MassBias);
+			const float Drag				= ParticleSimulationContext.RandomScaleBiasFloat(i, 1, Parameters->SolveVelocitiesAndForces_DragScale, Parameters->SolveVelocitiesAndForces_DragBias);
+			FVector3f InitialVelocity		= ParticleSimulationContext.RandomScaleBiasFloat(i, 2, Parameters->SolveVelocitiesAndForces_VelocityScale, Parameters->SolveVelocitiesAndForces_VelocityBias);
+			const FVector3f Wind			= ParticleSimulationContext.RandomScaleBiasFloat(i, 3, Parameters->SolveVelocitiesAndForces_WindScale, Parameters->SolveVelocitiesAndForces_WindBias);
+			const FVector3f Acceleration	= ParticleSimulationContext.RandomScaleBiasFloat(i, 4, Parameters->SolveVelocitiesAndForces_AccelerationScale, Parameters->SolveVelocitiesAndForces_AccelerationBias);
+
+			FVector3f Position			= ParticleSimulationContext.ReadParticleVariable(ModuleBuiltData->PositionVariableOffset, i, FVector3f::ZeroVector);
+			FVector3f PreviousPosition	= ParticleSimulationContext.ReadParticleVariable(ModuleBuiltData->PreviousPositionVariableOffset, i, FVector3f::ZeroVector);
+
+			if (ModuleBuiltData->PhysicsData.bConeVelocity)
+			{
+				const float ConeAngle = ParticleSimulationContext.RandomScaleBiasFloat(i, 5, Parameters->SolveVelocitiesAndForces_ConeAngleScale, Parameters->SolveVelocitiesAndForces_ConeAngleBias);
+				const float ConeRotation = ParticleSimulationContext.RandomFloat(i, 6) * UE_TWO_PI;
+				const FVector2f scAng = FVector2f(FMath::Sin(ConeAngle), FMath::Cos(ConeAngle));
+				const FVector2f scRot = FVector2f(FMath::Sin(ConeRotation), FMath::Cos(ConeRotation));
+				const FVector3f Direction = FVector3f(scRot.X * scAng.X, scRot.Y * scAng.X, scAng.Y);
+
+				float VelocityScale = ParticleSimulationContext.RandomScaleBiasFloat(i, 7, Parameters->SolveVelocitiesAndForces_ConeVelocityScale, Parameters->SolveVelocitiesAndForces_ConeVelocityBias);
+				if (ModuleBuiltData->PhysicsData.bConeVelocity)
+				{
+					const float pf = FMath::Pow(FMath::Clamp(scAng.Y, 0.0f, 1.0f), Parameters->SolveVelocitiesAndForces_ConeVelocityFalloff * 10.0f);
+					VelocityScale *= FMath::Lerp(1.0f, pf, Parameters->SolveVelocitiesAndForces_ConeVelocityFalloff);
+				}
+
+				InitialVelocity += ModuleBuiltData->PhysicsData.ConeQuat.RotateVector(Direction) * VelocityScale;
+			}
+
+			if (ModuleBuiltData->PhysicsData.bPointVelocity)
+			{
+				const FVector3f FallbackDir	= ParticleSimulationContext.RandomUnitFloat3(i, 8);
+				const FVector3f Delta		= Position - Parameters->SolveVelocitiesAndForces_PointOrigin;
+				const FVector3f Direction	= ParticleSimulationContext.SafeNormalize(Delta, FallbackDir);
+				const float		VelocityScale = ParticleSimulationContext.RandomScaleBiasFloat(i, 9, Parameters->SolveVelocitiesAndForces_PointVelocityScale, Parameters->SolveVelocitiesAndForces_PointVelocityBias);
+
+				InitialVelocity += Direction * VelocityScale;
+			}
+
+			if (ModuleBuiltData->PhysicsData.bNoiseEnabled)
+			{
+				//-TODO:
+			}
+
+			Position += IntegratePosition(AgeData[i], Mass, Drag, InitialVelocity, Wind, Acceleration);
+			PreviousPosition += IntegratePosition(PreviousAgeData[i], Mass, Drag, InitialVelocity, Wind, Acceleration);
+
+			ParticleSimulationContext.WriteParticleVariable(ModuleBuiltData->PositionVariableOffset, i, Position);
+			ParticleSimulationContext.WriteParticleVariable(ModuleBuiltData->PreviousPositionVariableOffset, i, PreviousPosition);
+
+			const FVector3f Velocity = (Position - PreviousPosition) * InvDeltaTime;
+			ParticleSimulationContext.WriteParticleVariable(ModuleBuiltData->VelocityVariableOffset, i, Velocity);
+			ParticleSimulationContext.WriteParticleVariable(ModuleBuiltData->PreviousVelocityVariableOffset, i, Velocity);
+		}
+	}
+}
+
+void UNiagaraStatelessModule_SolveVelocitiesAndForces::BuildEmitterData(const FNiagaraStatelessEmitterDataBuildContext& BuildContext) const
+{
+	using namespace NiagaraStatelessModuleSolveVelocitiesAndForcesPrivate;
+
+	NiagaraStateless::FPhysicsBuildData& PhysicsBuildData = BuildContext.GetTransientBuildData<NiagaraStateless::FPhysicsBuildData>();
+
+	FModuleBuiltData* BuiltData				= BuildContext.AllocateBuiltData<FModuleBuiltData>();
+	BuiltData->PhysicsData					= PhysicsBuildData;
+
+	const FNiagaraStatelessGlobals& StatelessGlobals = FNiagaraStatelessGlobals::Get();
+	BuiltData->PositionVariableOffset				= BuildContext.FindParticleVariableIndex(StatelessGlobals.PositionVariable);
+	BuiltData->VelocityVariableOffset				= BuildContext.FindParticleVariableIndex(StatelessGlobals.VelocityVariable);
+	BuiltData->PreviousPositionVariableOffset		= BuildContext.FindParticleVariableIndex(StatelessGlobals.PreviousPositionVariable);
+	BuiltData->PreviousVelocityVariableOffset		= BuildContext.FindParticleVariableIndex(StatelessGlobals.PreviousVelocityVariable);
+		
+	const bool bAttributesUsed =
+		BuiltData->PositionVariableOffset != INDEX_NONE ||
+		BuiltData->VelocityVariableOffset != INDEX_NONE ||
+		BuiltData->PreviousPositionVariableOffset != INDEX_NONE ||
+		BuiltData->PreviousVelocityVariableOffset != INDEX_NONE;
+
+	if (bAttributesUsed)
+	{
+		BuildContext.AddParticleSimulationExecSimulate(&ParticleSimulate);
+	}
+}
+
+void UNiagaraStatelessModule_SolveVelocitiesAndForces::BuildShaderParameters(FNiagaraStatelessShaderParametersBuilder& ShaderParametersBuilder) const
+{
+	ShaderParametersBuilder.AddParameterNestedStruct<FParameters>();
+}
+
+void UNiagaraStatelessModule_SolveVelocitiesAndForces::SetShaderParameters(const FNiagaraStatelessSetShaderParameterContext& SetShaderParameterContext) const
+{
+	using namespace NiagaraStatelessModuleSolveVelocitiesAndForcesPrivate;
+
+	const FModuleBuiltData* ModuleBuiltData = SetShaderParameterContext.ReadBuiltData<FModuleBuiltData>();
+	const NiagaraStateless::FPhysicsBuildData& PhysicsData = ModuleBuiltData->PhysicsData;
+
+	const FNiagaraStatelessSpaceTransforms& SpaceTransforms = SetShaderParameterContext.GetSpaceTransforms();
+
+	FParameters* Parameters = SetShaderParameterContext.GetParameterNestedStruct<FParameters>();
+	SetShaderParameterContext.ConvertRangeToScaleBias(PhysicsData.MassRange, Parameters->SolveVelocitiesAndForces_MassScale, Parameters->SolveVelocitiesAndForces_MassBias);
+	SetShaderParameterContext.ConvertRangeToScaleBias(PhysicsData.DragRange, Parameters->SolveVelocitiesAndForces_DragScale, Parameters->SolveVelocitiesAndForces_DragBias);
+	Parameters->SolveVelocitiesAndForces_DragScale	= FMath::Max(Parameters->SolveVelocitiesAndForces_DragScale, 0.01f);
+	Parameters->SolveVelocitiesAndForces_DragBias	= FMath::Max(Parameters->SolveVelocitiesAndForces_DragBias, 0.01f);
+
+	SetShaderParameterContext.TransformVectorRangeToScaleBias(PhysicsData.VelocityRange, PhysicsData.VelocityCoordinateSpace, Parameters->SolveVelocitiesAndForces_VelocityScale, Parameters->SolveVelocitiesAndForces_VelocityBias);
+	const float VelocityScale = SetShaderParameterContext.ConvertRangeToValue(PhysicsData.LinearVelocityScale);
+	Parameters->SolveVelocitiesAndForces_VelocityScale *= VelocityScale;
+	Parameters->SolveVelocitiesAndForces_VelocityBias  *= VelocityScale;
+
+	SetShaderParameterContext.TransformVectorRangeToScaleBias(PhysicsData.WindRange, PhysicsData.WindCoordinateSpace, Parameters->SolveVelocitiesAndForces_WindScale, Parameters->SolveVelocitiesAndForces_WindBias);
+	{
+		FVector3f AccelerationScale, AccelerationBias;
+		FVector3f GravityScale, GravityBias;
+		SetShaderParameterContext.ConvertRangeToScaleBias(PhysicsData.AccelerationRange, AccelerationScale, AccelerationBias);
+		SetShaderParameterContext.ConvertRangeToScaleBias(PhysicsData.GravityRange, GravityScale, GravityBias);
+		Parameters->SolveVelocitiesAndForces_AccelerationScale	= SpaceTransforms.TransformVector(PhysicsData.AccelerationCoordinateSpace, AccelerationScale);
+		Parameters->SolveVelocitiesAndForces_AccelerationBias	= SpaceTransforms.TransformVector(PhysicsData.AccelerationCoordinateSpace, AccelerationBias);
+		Parameters->SolveVelocitiesAndForces_AccelerationScale	+= SpaceTransforms.TransformVector(ENiagaraCoordinateSpace::World, GravityScale);
+		Parameters->SolveVelocitiesAndForces_AccelerationBias	+= SpaceTransforms.TransformVector(ENiagaraCoordinateSpace::World, GravityBias);
+	}
+
+	Parameters->SolveVelocitiesAndForces_ConeVelocityEnabled = PhysicsData.bConeVelocity ? 1 : 0;
+	{
+		Parameters->SolveVelocitiesAndForces_ConeQuat				= SpaceTransforms.TransformRotation(PhysicsData.ConeCoordinateSpace, PhysicsData.ConeQuat);
+		SetShaderParameterContext.ConvertRangeToScaleBias(PhysicsData.ConeVelocityRange, Parameters->SolveVelocitiesAndForces_ConeVelocityScale, Parameters->SolveVelocitiesAndForces_ConeVelocityBias);
+		Parameters->SolveVelocitiesAndForces_ConeAngleScale			= (PhysicsData.ConeOuterAngle - PhysicsData.ConeInnerAngle) * (UE_PI / 360.0f);
+		Parameters->SolveVelocitiesAndForces_ConeAngleBias			= PhysicsData.ConeInnerAngle * (UE_PI / 360.0f);
+		Parameters->SolveVelocitiesAndForces_ConeVelocityFalloff	= PhysicsData.ConeVelocityFalloff;
+	}
+
+	Parameters->SolveVelocitiesAndForces_PontVelocityEnabled = PhysicsData.bPointVelocity ? 1 : 0;
+	{
+		SetShaderParameterContext.ConvertRangeToScaleBias(PhysicsData.PointVelocityRange, Parameters->SolveVelocitiesAndForces_PointVelocityScale, Parameters->SolveVelocitiesAndForces_PointVelocityBias);
+		Parameters->SolveVelocitiesAndForces_PointOrigin = SpaceTransforms.TransformPosition(PhysicsData.PointCoordinateSpace, PhysicsData.PointOrigin);
+	}
+
+	Parameters->SolveVelocitiesAndForces_NoiseEnabled = PhysicsData.bNoiseEnabled ? 1 : 0;
+	if (PhysicsData.bNoiseEnabled)
+	{
+		Parameters->SolveVelocitiesAndForces_NoiseAmplitude			= PhysicsData.NoiseAmplitude;
+		Parameters->SolveVelocitiesAndForces_NoiseFrequency			= FVector3f(PhysicsData.NoiseFrequency, PhysicsData.NoiseFrequency, PhysicsData.NoiseFrequency);
+		//SetShaderParameterContext.SetTextureResource(&Parameters->SolveVelocitiesAndForces_NoiseTexture, ModuleBuildData->NoiseTexture);
+		Parameters->SolveVelocitiesAndForces_NoiseMode				= PhysicsData.NoiseMode;
+		Parameters->SolveVelocitiesAndForces_NoiseLUTOffset			= PhysicsData.NoiseLUTOffset;
+		Parameters->SolveVelocitiesAndForces_NoiseLUTNumChannel		= PhysicsData.NoiseLUTNumChannel;
+		Parameters->SolveVelocitiesAndForces_NoiseLUTChannelWidth	= PhysicsData.NoiseLUTChannelWidth;
+	}
+
+	FVectorFieldTextureAccessor TextureAccessor(Cast<UVectorField>(PhysicsData.NoiseTexture));
+
+	ENQUEUE_RENDER_COMMAND(FUpdateNoiseTexture)(
+		[Parameters, TextureAccessor](FRHICommandListImmediate& RHICmdList)
+		{
+			FRHITexture* NoiseTextureRHI = TextureAccessor.GetTexture();
+			Parameters->SolveVelocitiesAndForces_NoiseTexture = NoiseTextureRHI;
+			Parameters->SolveVelocitiesAndForces_NoiseSampler = TStaticSamplerState<SF_Bilinear, AM_Wrap, AM_Wrap, AM_Wrap>::GetRHI();
+
+			if (Parameters->SolveVelocitiesAndForces_NoiseMode < 2)
+			{
+				const FIntVector TextureSize = NoiseTextureRHI ? NoiseTextureRHI->GetSizeXYZ() : FIntVector(1, 1, 1);
+				Parameters->SolveVelocitiesAndForces_NoiseFrequency.X *= 1.0f / float(TextureSize.X);
+				Parameters->SolveVelocitiesAndForces_NoiseFrequency.Y *= 1.0f / float(TextureSize.Y);
+				Parameters->SolveVelocitiesAndForces_NoiseFrequency.Z *= 1.0f / float(TextureSize.Z);
+			}
+		}
+	);
+}
+
+#if WITH_EDITORONLY_DATA
+void UNiagaraStatelessModule_SolveVelocitiesAndForces::GetOutputVariables(TArray<FNiagaraVariableBase>& OutVariables) const
+{
+	const FNiagaraStatelessGlobals& StatelessGlobals = FNiagaraStatelessGlobals::Get();
+	OutVariables.AddUnique(StatelessGlobals.PositionVariable);
+	OutVariables.AddUnique(StatelessGlobals.VelocityVariable);
+	OutVariables.AddUnique(StatelessGlobals.PreviousPositionVariable);
+	OutVariables.AddUnique(StatelessGlobals.PreviousVelocityVariable);
+}
+#endif
diff --git a/Engine/Plugins/FX/Niagara/Source/Niagara/Private/Stateless/NiagaraStatelessDistribution.cpp b/Engine/Plugins/FX/Niagara/Source/Niagara/Private/Stateless/NiagaraStatelessDistribution.cpp
--- a/Engine/Plugins/FX/Niagara/Source/Niagara/Private/Stateless/NiagaraStatelessDistribution.cpp
+++ b/Engine/Plugins/FX/Niagara/Source/Niagara/Private/Stateless/NiagaraStatelessDistribution.cpp
@@ -303,6 +303,16 @@ void FNiagaraDistributionRangeFloat::InitConstant(float Value)
 #endif
 }
 
+void FNiagaraDistributionRangeFloat::InitRange(float MinValue, float MaxValue)
+{
+	Mode = ENiagaraDistributionMode::UniformRange;
+	Min = MinValue;
+	Max = MaxValue;
+#if WITH_EDITORONLY_DATA
+	ChannelConstantsAndRanges = TArray<float>({ MinValue, MaxValue });
+#endif
+}
+
 FNiagaraStatelessRangeFloat FNiagaraDistributionRangeFloat::CalculateRange(const float Default) const
 {
 	return Mode == ENiagaraDistributionMode::UniformConstant ? FNiagaraStatelessRangeFloat(Min, Min) : FNiagaraStatelessRangeFloat(Min, Max);
@@ -696,6 +706,19 @@ bool FNiagaraDistributionRangeFloat::SerializeFromMismatchedTag(const struct FPr
 	return false;
 }
 
+bool FNiagaraDistributionRangeVector2::SerializeFromMismatchedTag(const struct FPropertyTag& Tag, FStructuredArchive::FSlot Slot)
+{
+	if (Tag.GetType().IsStruct(NAME_Vector2f))
+	{
+		FVector2f Value;
+		Slot << Value;
+		*this = FNiagaraDistributionRangeVector2(Value);
+		return true;
+	}
+
+	return false;
+}
+
 bool FNiagaraDistributionRangeVector3::SerializeFromMismatchedTag(const struct FPropertyTag& Tag, FStructuredArchive::FSlot Slot)
 {
 	if (Tag.GetType().IsStruct(NAME_Vector3f))
diff --git a/Engine/Plugins/FX/Niagara/Source/Niagara/Private/Stateless/NiagaraStatelessModule.cpp b/Engine/Plugins/FX/Niagara/Source/Niagara/Private/Stateless/NiagaraStatelessModule.cpp
--- a/Engine/Plugins/FX/Niagara/Source/Niagara/Private/Stateless/NiagaraStatelessModule.cpp
+++ b/Engine/Plugins/FX/Niagara/Source/Niagara/Private/Stateless/NiagaraStatelessModule.cpp
@@ -7,30 +7,6 @@ const FName UNiagaraStatelessModule::PrivateMemberNames::bModuleEnabled = GET_ME
 const FName UNiagaraStatelessModule::PrivateMemberNames::bDebugDrawEnabled = GET_MEMBER_NAME_CHECKED(UNiagaraStatelessModule, bDebugDrawEnabled);
 #endif
 
-#if DO_CHECK
-void FNiagaraStatelessSetShaderParameterContext::ValidateIncludeStructType(uint32 StructOffset, const FShaderParametersMetadata* StructMetaData) const
-{
-	for (const FShaderParametersMetadata::FMember& Member : ShaderParametersMetadata->GetMembers())
-	{
-		if (Member.GetOffset() != StructOffset)
-		{
-			continue;
-		}
-
-		if (Member.GetBaseType() == UBMT_INCLUDED_STRUCT && Member.GetStructMetadata() && Member.GetStructMetadata()->GetLayout() == StructMetaData->GetLayout())
-		{
-			return;
-		}
-
-		const TCHAR* StructType = Member.GetStructMetadata() ? Member.GetStructMetadata()->GetStructTypeName() : TEXT("null");
-		UE_LOG(LogNiagara, Fatal, TEXT("Shader parameter struct member (%s) at offset (%u) is not of type (%s) struct type is (%s)"), Member.GetName(), StructOffset, StructMetaData->GetStructTypeName(), StructType);
-		return;
-	}
-
-	UE_LOG(LogNiagara, Fatal, TEXT("Failed to find shader parameter struct member type (%s) at offset (%u)"), StructMetaData->GetStructTypeName(), StructOffset);
-}
-#endif //DO_CHECK
-
 #if WITH_EDITOR
 bool UNiagaraStatelessModule::CanEditChange(const FProperty* InProperty) const
 {
diff --git a/Engine/Plugins/FX/Niagara/Source/Niagara/Private/Stateless/NiagaraStatelessSetShaderParameterContext.cpp b/Engine/Plugins/FX/Niagara/Source/Niagara/Private/Stateless/NiagaraStatelessSetShaderParameterContext.cpp
new file mode 100644
--- /dev/null
+++ b/Engine/Plugins/FX/Niagara/Source/Niagara/Private/Stateless/NiagaraStatelessSetShaderParameterContext.cpp
@@ -0,0 +1,56 @@
+// Copyright Epic Games, Inc. All Rights Reserved.
+
+#include "Stateless/NiagaraStatelessSetShaderParameterContext.h"
+
+FNiagaraStatelessSetShaderParameterContext::FNiagaraStatelessSetShaderParameterContext(const FNiagaraStatelessSpaceTransforms& InSpaceTransforms, TConstArrayView<uint8> InRendererParameterData, TConstArrayView<uint8> InBuiltData, const FShaderParametersMetadata* InShaderParametersMetadata, uint8* InShaderParameters)
+	: SpaceTransforms(InSpaceTransforms)
+	, RendererParameterData(InRendererParameterData)
+	, BuiltData(InBuiltData)
+	, ShaderParametersBase(InShaderParameters)
+	, ParameterOffset(0)
+	, ShaderParametersMetadata(InShaderParametersMetadata)
+{
+}
+
+void FNiagaraStatelessSetShaderParameterContext::TransformVectorRangeToScaleBias(const FNiagaraStatelessRangeVector3& Range, ENiagaraCoordinateSpace SourceSpace, FVector3f& OutScale, FVector3f& OutBias) const
+{
+	if (Range.ParameterOffset == INDEX_NONE)
+	{
+		OutScale = SpaceTransforms.TransformVector(SourceSpace, Range.GetScale());
+		OutBias  = SpaceTransforms.TransformVector(SourceSpace, Range.Min);
+	}
+	else
+	{
+		OutScale = FNiagaraStatelessRangeDefaultValue<FVector3f>::Zero();
+		OutBias  = SpaceTransforms.TransformVector(SourceSpace, GetRendererParameterValue(Range.ParameterOffset, Range.Min));
+	}
+}
+
+FVector3f FNiagaraStatelessSetShaderParameterContext::TransformPositionRangeToValue(const FNiagaraStatelessRangeVector3& Range, ENiagaraCoordinateSpace SourceSpace) const
+{
+	return SpaceTransforms.TransformPosition(SourceSpace, GetRendererParameterValue(Range.ParameterOffset, Range.Min));
+}
+
+#if DO_CHECK
+void FNiagaraStatelessSetShaderParameterContext::ValidateIncludeStructType(uint32 StructOffset, const FShaderParametersMetadata* StructMetaData) const
+{
+	for (const FShaderParametersMetadata::FMember& Member : ShaderParametersMetadata->GetMembers())
+	{
+		if (Member.GetOffset() != StructOffset)
+		{
+			continue;
+		}
+
+		if (Member.GetBaseType() == UBMT_INCLUDED_STRUCT && Member.GetStructMetadata() && Member.GetStructMetadata()->GetLayout() == StructMetaData->GetLayout())
+		{
+			return;
+		}
+
+		const TCHAR* StructType = Member.GetStructMetadata() ? Member.GetStructMetadata()->GetStructTypeName() : TEXT("null");
+		UE_LOG(LogNiagara, Fatal, TEXT("Shader parameter struct member (%s) at offset (%u) is not of type (%s) struct type is (%s)"), Member.GetName(), StructOffset, StructMetaData->GetStructTypeName(), StructType);
+		return;
+	}
+
+	UE_LOG(LogNiagara, Fatal, TEXT("Failed to find shader parameter struct member type (%s) at offset (%u)"), StructMetaData->GetStructTypeName(), StructOffset);
+}
+#endif //DO_CHECK
