- Fix crash in stack notes when reading editor data during finalize

#rnx

diff --git a/Engine/Plugins/FX/Niagara/Source/NiagaraEditor/Private/ViewModels/Stack/NiagaraStackEntry.cpp b/Engine/Plugins/FX/Niagara/Source/NiagaraEditor/Private/ViewModels/Stack/NiagaraStackEntry.cpp
--- a/Engine/Plugins/FX/Niagara/Source/NiagaraEditor/Private/ViewModels/Stack/NiagaraStackEntry.cpp
+++ b/Engine/Plugins/FX/Niagara/Source/NiagaraEditor/Private/ViewModels/Stack/NiagaraStackEntry.cpp
@@ -250,6 +250,7 @@ FGuid UNiagaraStackEntry::GetSelectionId() const
 
 UNiagaraStackEditorData& UNiagaraStackEntry::GetStackEditorData() const
 {
+	checkf(StackEditorData, TEXT("StackEditorData is nullptr the calling code needs to check for finalized or not"));
 	return *StackEditorData;
 }
 
diff --git a/Engine/Plugins/FX/Niagara/Source/NiagaraEditor/Private/ViewModels/Stack/NiagaraStackNote.cpp b/Engine/Plugins/FX/Niagara/Source/NiagaraEditor/Private/ViewModels/Stack/NiagaraStackNote.cpp
--- a/Engine/Plugins/FX/Niagara/Source/NiagaraEditor/Private/ViewModels/Stack/NiagaraStackNote.cpp
+++ b/Engine/Plugins/FX/Niagara/Source/NiagaraEditor/Private/ViewModels/Stack/NiagaraStackNote.cpp
@@ -21,7 +21,11 @@ FString UNiagaraStackNote::GetTargetStackEntryKey() const
 
 TOptional<FNiagaraStackNoteData> UNiagaraStackNote::GetTargetStackNoteData() const
 {
-	return GetStackEditorData().GetStackNote(GetTargetStackEntryKey());
+	if (!IsFinalized())
+	{
+		return GetStackEditorData().GetStackNote(GetTargetStackEntryKey());
+	}
+	return {};
 }
 
 void UNiagaraStackNote::ToggleInlineDisplay()
