- Make sure we mark the renderer state dirty when we reinit due to dirty data interfaces
- The reinit will destroy the GPU context and leave dynamic data pointing to an invalid context

#rnx
#tests Provided project, that inolves a BP which invalidates the OverrideParameters in PostApplyToComponent and uses a Pause effect type.  You then move the paused component in the editor.

diff --git a/Engine/Plugins/FX/Niagara/Source/Niagara/Private/NiagaraSystemInstance.cpp b/Engine/Plugins/FX/Niagara/Source/Niagara/Private/NiagaraSystemInstance.cpp
--- a/Engine/Plugins/FX/Niagara/Source/Niagara/Private/NiagaraSystemInstance.cpp
+++ b/Engine/Plugins/FX/Niagara/Source/Niagara/Private/NiagaraSystemInstance.cpp
@@ -2484,6 +2484,11 @@ void FNiagaraSystemInstance::Tick_GameThread(float DeltaSeconds)
 	if (OverrideParameters->GetInterfacesDirty())
 	{
 		Reset(EResetMode::ReInit);
+
+		if (USceneComponent* Component = GetAttachComponent())
+		{
+			Component->MarkRenderStateDirty();
+		}
 		return;
 	}
 
