- Fix a crash if the area weighted data has not been built on the SKM

#tests User provided project which uses mutable to combine SKM together.  The SKM has area weighted enabled but the mutable produced content has not built the data.

diff --git a/Engine/Plugins/FX/Niagara/Source/Niagara/Private/NDISkeletalMesh_TriangleSampling.cpp b/Engine/Plugins/FX/Niagara/Source/Niagara/Private/NDISkeletalMesh_TriangleSampling.cpp
--- a/Engine/Plugins/FX/Niagara/Source/Niagara/Private/NDISkeletalMesh_TriangleSampling.cpp
+++ b/Engine/Plugins/FX/Niagara/Source/Niagara/Private/NDISkeletalMesh_TriangleSampling.cpp
@@ -707,16 +707,21 @@ void UNiagaraDataInterfaceSkeletalMesh::RandomTriangle(FVectorVMExternalFunction
 	if (bAreaWeighted)
 	{
 		const FSkeletalMeshSamplingInfo& SamplingInfo = SkelMesh->GetSamplingInfo();
-		const FSkeletalMeshSamplingLODBuiltData& WholeMeshBuiltData = SamplingInfo.GetWholeMeshLODBuiltData(InstData->GetLODIndex());
-		if (WholeMeshBuiltData.AreaWeightedTriangleSampler.GetNumEntries() > 0)
+		const FSkeletalMeshSamplingBuiltData& SamplingInfoBuiltData = SamplingInfo.GetBuiltData();
+		if (SamplingInfoBuiltData.WholeMeshBuiltData.IsValidIndex(LODIndex))
 		{
-			for (int32 i = 0; i < Context.GetNumInstances(); ++i)
+			//-TODO: Consider adding a way to warn the user we failed to get the area weighted data
+			const FSkeletalMeshSamplingLODBuiltData& WholeMeshBuiltData = SamplingInfoBuiltData.WholeMeshBuiltData[LODIndex];
+			if (WholeMeshBuiltData.AreaWeightedTriangleSampler.GetNumEntries() > 0)
 			{
-				RandHelper.GetAndAdvance();
-				OutTri.SetAndAdvance(WholeMeshBuiltData.AreaWeightedTriangleSampler.GetEntryIndex(RandHelper.Rand(i), RandHelper.Rand(i)));
-				OutBary.SetAndAdvance(RandHelper.RandomBarycentricCoord(i));
+				for (int32 i = 0; i < Context.GetNumInstances(); ++i)
+				{
+					RandHelper.GetAndAdvance();
+					OutTri.SetAndAdvance(WholeMeshBuiltData.AreaWeightedTriangleSampler.GetEntryIndex(RandHelper.Rand(i), RandHelper.Rand(i)));
+					OutBary.SetAndAdvance(RandHelper.RandomBarycentricCoord(i));
+				}
+				return;
 			}
-			return;
 		}
 	}
 
