Pass on the FNiagaraTypeRegistry to tackle a few different problems including consistency of threaded access and handling of types needing to be invalidated if they are contained within dynamic plugins.
-adds a spin lock which manages access to the RegisteredTypes array
-moves the RegisteredTypes array to be an array of TUniquePtr instead of a fixed sized static array
-moves access to the internal categories of types to returning copies by value instead of references to ensure locks are in place
-caches registered type indices both on the type definition itself and in a TLS to minimize the number of times we need to access the global map
-add handling of plugin unloading to track if registered types need to be invalidated
-changed the type definition GetTypeHash to be the full path of the underlying type rather than it's address
-moved FNiagaraTypeRegistry stuff to it's own file

diff --git a/Engine/Plugins/FX/Niagara/Source/Niagara/Niagara.Build.cs b/Engine/Plugins/FX/Niagara/Source/Niagara/Niagara.Build.cs
--- a/Engine/Plugins/FX/Niagara/Source/Niagara/Niagara.Build.cs
+++ b/Engine/Plugins/FX/Niagara/Source/Niagara/Niagara.Build.cs
@@ -20,6 +20,7 @@ public class Niagara : ModuleRules
                 "Json",
                 "JsonUtilities",
                 "Landscape",
+                "Projects",
                 "Renderer",
                 "SignalProcessing",
                 "TimeManagement",
diff --git a/Engine/Plugins/FX/Niagara/Source/Niagara/Private/NiagaraModule.cpp b/Engine/Plugins/FX/Niagara/Source/Niagara/Private/NiagaraModule.cpp
--- a/Engine/Plugins/FX/Niagara/Source/Niagara/Private/NiagaraModule.cpp
+++ b/Engine/Plugins/FX/Niagara/Source/Niagara/Private/NiagaraModule.cpp
@@ -2,14 +2,12 @@
 
 #include "NiagaraModule.h"
 #include "GPUSortManager.h"
-#include "Misc/LazySingleton.h"
 #include "Materials/MaterialInterface.h"
 #include "Modules/ModuleManager.h"
 #include "NiagaraComponentSettings.h"
 #include "NiagaraCompileHashVisitor.h"
 #include "NiagaraEmitter.h"
 #include "NiagaraTypes.h"
-#include "NiagaraDebugVis.h" // IWYU pragma: keep
 #include "NiagaraEvents.h"
 #include "NiagaraSettings.h"
 #include "NiagaraDataChannelManager.h"
@@ -34,7 +32,7 @@
 #include "NiagaraCustomVersion.h"
 #include "NiagaraDataChannel.h"
 #include "NiagaraShaderModule.h"
-#include "UObject/CoreRedirects.h"
+#include "NiagaraTypeRegistry.h"
 #include "NiagaraGpuComputeDispatch.h"
 #include "Misc/CoreDelegates.h"
 #include "NiagaraDebuggerClient.h"
@@ -195,17 +193,6 @@ void INiagaraModule::OnDataChannelsEnabledChanged(IConsoleVariable* Variable)
 	Context.AddAll(true);
 }
 
-// these two globals are intended to help ensure that a global variable is accessible to natvis (as defined in Niagara.natvis)
-// while debugging.  GCoreTypeRegistrySingletonPtr is the pointer to the actual data stored within the TLazySingleton<FNiagaraTypeRegistry>
-// while GTypeRegistrySingletonPtr can be declared in each module to ensure that it can be accessed while debugging any Niagara
-// module.  See UE_VISUALIZERS_HELPERS.  Note that currently it seems like we can effectively debug NiagaraEditor/NiagaraShader without
-// further declarations, which is nice, so will be leaving it as it is.
-namespace NiagaraDebugVisHelper
-{
-	const FNiagaraTypeRegistry* GCoreTypeRegistrySingletonPtr = nullptr;
-	UE_SELECT_ANY const FNiagaraTypeRegistry*& GTypeRegistrySingletonPtr = GCoreTypeRegistrySingletonPtr;
-}
-
 FNiagaraVariable INiagaraModule::Engine_WorldDeltaTime;
 FNiagaraVariable INiagaraModule::Engine_DeltaTime;
 FNiagaraVariable INiagaraModule::Engine_InvDeltaTime;
@@ -323,6 +310,8 @@ FNiagaraVariable INiagaraModule::Translator_CallID;
 
 void INiagaraModule::StartupModule()
 {
+	FNiagaraTypeRegistry::Init();
+
 	VectorVM::Init();
 	FNiagaraTypeHelper::InitLWCTypes();
 
@@ -1211,6 +1200,7 @@ FNiagaraTypeDefinition FNiagaraTypeDefinition::ToStaticDef() const
 {
 	FNiagaraTypeDefinition Def(*this);
 	Def.Flags |= TF_Static;
+	Def.InvalidateCachedIndex(Def.Flags != Flags);
 	return Def;
 }
 
@@ -1218,6 +1208,7 @@ FNiagaraTypeDefinition FNiagaraTypeDefinition::RemoveStaticDef() const
 {
 	FNiagaraTypeDefinition Def(*this);
 	Def.Flags &= ~TF_Static;
+	Def.InvalidateCachedIndex(Def.Flags != Flags);
 	return Def;
 }
 
@@ -1537,6 +1528,11 @@ FNiagaraTypeDefinition FNiagaraTypeDefinition::GetNumericOutputType(TConstArrayV
 	}
 }
 
+void FNiagaraTypeDefinition::InvalidateCachedIndex(bool bInvalidate)
+{
+	RegisteredTypeDefIndex = bInvalidate ? INDEX_NONE : RegisteredTypeDefIndex;
+}
+
 bool FNiagaraTypeDefinition::Serialize(FArchive& Ar)
 {
 	Ar.UsingCustomVersion(FNiagaraCustomVersion::GUID);
@@ -1576,6 +1572,8 @@ bool FNiagaraTypeDefinition::Serialize(FArchive& Ar)
 
 	if (Ar.IsLoading())
 	{
+		InvalidateCachedIndex();
+
 #if WITH_EDITORONLY_DATA
 		if (Ar.CustomVer(FNiagaraCustomVersion::GUID) < FNiagaraCustomVersion::MemorySaving)
 		{
@@ -1618,93 +1616,16 @@ bool FNiagaraTypeDefinition::Serialize(FArchive& Ar)
 #endif
 			}
 		}
-	}
-
-	return true;
-}
-
-//////////////////////////////////////////////////////////////////////////
-
-UNiagaraDataInterfaceBase* FNiagaraTypeRegistry::GetDefaultDataInterfaceByName(const FString& DIClassName)
-{
-	UClass* DIClass = nullptr;
-	{
-		FNiagaraTypeRegistry& Registry = Get();
-		FReadScopeLock Lock(Registry.RegisteredTypesLock);
-		for (const FNiagaraTypeDefinition& Def : Registry.RegisteredTypes)
-		{
-			if (Def.IsDataInterface())
-			{
-				UClass* FoundDIClass = Def.GetClass();
-				if (FoundDIClass && (FoundDIClass->GetName() == DIClassName || FoundDIClass->GetFullName() == DIClassName))
-				{
-					DIClass = FoundDIClass;
-					break;
-				}
-			}
-		}
-	}
-
-	// Consider the possibility of a redirector pointing to a new location..
-	if (DIClass == nullptr)
-	{
-		FCoreRedirectObjectName OldObjName;
-		OldObjName.ObjectName = *DIClassName;
-		FCoreRedirectObjectName NewObjName = FCoreRedirects::GetRedirectedName(ECoreRedirectFlags::Type_Class, OldObjName);
-		if (NewObjName.IsValid() && OldObjName != NewObjName)
-		{
-			return GetDefaultDataInterfaceByName(NewObjName.ObjectName.ToString());
-		}
-	}
-
-	if (DIClass)
-	{
-		return CastChecked<UNiagaraDataInterfaceBase>(DIClass->GetDefaultObject(false)); // We wouldn't be registered if the CDO had not already been created...
-	}
-
-	return nullptr;
-}
 
-// this will hold onto references to the classes/enums/structs that are used as a part of type definitions.  In the case that something is forcibly deleted
-// the contents of the RegisteredTypes array will have it's ClassStructorEnum cleared, and it will no longer be treated as a valid type
-// (FNiagaraTypeDefinitionHandle::Resolve will return an invalid Dummy reference).
-// 
-// It might be worth investigating changing RegisteredTypeIndexMap to be hashed based on the name of the class/enum/struct rather than the pointer.  That
-// way if we reload a type old data, with their existing index, will still be treated as valid.  In this current implementation old FNiagaraVariables will
-// have invalid types.  But, for now I think it would be an error to unload a type without having the variables that dependend on it also getting deleted.
-void FNiagaraTypeRegistry::AddReferencedObjects(FReferenceCollector& Collector)
-{
-	FReadScopeLock Lock(RegisteredTypesLock);
-	for (FNiagaraTypeDefinition& RegisteredType : RegisteredTypes)
-	{
-		Collector.AddReferencedObject(RegisteredType.ClassStructOrEnum);
+		GeneratePathNameHash();
 	}
-}
 
-FString FNiagaraTypeRegistry::GetReferencerName() const
-{
-	return TEXT("FNiagaraTypeRegistry");
-}
-
-FNiagaraTypeRegistry::FNiagaraTypeRegistry()
-{
-	NiagaraDebugVisHelper::GCoreTypeRegistrySingletonPtr = this;
-
-}
-
-FNiagaraTypeRegistry::~FNiagaraTypeRegistry()
-{
-	NiagaraDebugVisHelper::GCoreTypeRegistrySingletonPtr = nullptr;
-}
-
-FNiagaraTypeRegistry& FNiagaraTypeRegistry::Get()
-{
-	return TLazySingleton<FNiagaraTypeRegistry>::Get();
+	return true;
 }
 
-void FNiagaraTypeRegistry::TearDown()
+void FNiagaraTypeDefinition::GeneratePathNameHash()
 {
-	TLazySingleton<FNiagaraTypeRegistry>::TearDown();
+	UnderlyingTypePathNameHash = GetTypeHash(GetFullNameSafe(ClassStructOrEnum));
 }
 
 //////////////////////////////////////////////////////////////////////////
@@ -1730,20 +1651,8 @@ void INiagaraModule::ProcessShaderCompilationQueue()
 
 const FNiagaraTypeDefinition& FNiagaraTypeDefinitionHandle::Resolve() const
 {
-	const auto& RegisteredTypes = FNiagaraTypeRegistry::GetRegisteredTypes();
-
-	if (RegisteredTypes.IsValidIndex(RegisteredTypeIndex))
-	{
-		// If the type is invalid, then it's likely that it has been invalidated by GC because the underlying object has
-		// been unloaded.
-		if (RegisteredTypes[RegisteredTypeIndex].IsValid())
-		{
-			return RegisteredTypes[RegisteredTypeIndex];
-		}
-	}
+	return FNiagaraTypeRegistry::Get().GetRegisteredType(RegisteredTypeIndex);
 
-	static FNiagaraTypeDefinition Dummy;
-	return Dummy;
 }
 
 int32 FNiagaraTypeDefinitionHandle::Register(const FNiagaraTypeDefinition& TypeDef) const
@@ -1753,7 +1662,7 @@ int32 FNiagaraTypeDefinitionHandle::Register(const FNiagaraTypeDefinition& TypeD
 		return INDEX_NONE;
 	}
 
-	return FNiagaraTypeRegistry::RegisterIndexed(TypeDef);
+	return FNiagaraTypeRegistry::Get().RegisterTypeInternal(TypeDef);
 }
 
 FArchive& operator<<(FArchive& Ar, FNiagaraTypeDefinitionHandle& Handle)
diff --git a/Engine/Plugins/FX/Niagara/Source/Niagara/Private/NiagaraTypeRegistry.cpp b/Engine/Plugins/FX/Niagara/Source/Niagara/Private/NiagaraTypeRegistry.cpp
new file mode 100644
--- /dev/null
+++ b/Engine/Plugins/FX/Niagara/Source/Niagara/Private/NiagaraTypeRegistry.cpp
@@ -0,0 +1,696 @@
+// Copyright Epic Games, Inc. All Rights Reserved.
+
+#include "NiagaraTypeRegistry.h"
+
+#include "Algo/Transform.h"
+#include "Interfaces/IPluginManager.h"
+#include "Misc/LazySingleton.h"
+#include "NiagaraDataInterfaceBase.h"
+#include "NiagaraDebugVis.h" // IWYU pragma: keep
+#include "UObject/CoreRedirects.h"
+
+// these two globals are intended to help ensure that a global variable is accessible to natvis (as defined in Niagara.natvis)
+// while debugging.  GCoreTypeRegistrySingletonPtr is the pointer to the actual data stored within the TLazySingleton<FNiagaraTypeRegistry>
+// while GTypeRegistrySingletonPtr can be declared in each module to ensure that it can be accessed while debugging any Niagara
+// module.  See UE_VISUALIZERS_HELPERS.  Note that currently it seems like we can effectively debug NiagaraEditor/NiagaraShader without
+// further declarations, which is nice, so will be leaving it as it is.
+namespace NiagaraDebugVisHelper
+{
+	const FNiagaraTypeRegistry* GCoreTypeRegistrySingletonPtr = nullptr;
+	UE_SELECT_ANY const FNiagaraTypeRegistry*& GTypeRegistrySingletonPtr = GCoreTypeRegistrySingletonPtr;
+}
+
+namespace NiagaraTypeRegistryLocal
+{
+	// TLS of the mapping between a TypeDefinition's TypeHash and the index in
+	// the TypeRegistry's RegisteredTypes array.  Allows us to quickly access
+	// the type definition from a handle without having to hit the lock
+	static thread_local TMap<uint32, int32> RegisteredTypeIndexCacheTLS;
+
+	static FDelegateHandle PluginUnmountedHandle;
+
+	void PluginUnmounted(IPlugin& Plugin)
+	{
+		const FString MountedPath = Plugin.GetMountedAssetPath();
+
+		// go through all of the registered types and see if any of them are from the plugin being unmounted.
+		FNiagaraTypeRegistry::Get().InvalidateTypesByPath(MountedPath);
+	}
+}
+
+UNiagaraDataInterfaceBase* FNiagaraTypeRegistry::GetDefaultDataInterfaceByName(const FString& DIClassName)
+{
+	UClass* DIClass = nullptr;
+	{
+		FNiagaraTypeRegistry& Registry = Get();
+		FReadScopeLock Lock(Registry.RegistrationLock);
+		for (const TUniquePtr<FNiagaraTypeDefinition>& TypePtr : Registry.RegisteredTypes)
+		{
+			if (TypePtr.IsValid() && TypePtr->IsDataInterface())
+			{
+				UClass* FoundDIClass = TypePtr->GetClass();
+				if (FoundDIClass && (FoundDIClass->GetName() == DIClassName || FoundDIClass->GetFullName() == DIClassName))
+				{
+					DIClass = FoundDIClass;
+					break;
+				}
+			}
+		}
+	}
+
+	// Consider the possibility of a redirector pointing to a new location..
+	if (DIClass == nullptr)
+	{
+		FCoreRedirectObjectName OldObjName;
+		OldObjName.ObjectName = *DIClassName;
+		FCoreRedirectObjectName NewObjName = FCoreRedirects::GetRedirectedName(ECoreRedirectFlags::Type_Class, OldObjName);
+		if (NewObjName.IsValid() && OldObjName != NewObjName)
+		{
+			return GetDefaultDataInterfaceByName(NewObjName.ObjectName.ToString());
+		}
+	}
+
+	if (DIClass)
+	{
+		return CastChecked<UNiagaraDataInterfaceBase>(DIClass->GetDefaultObject(false)); // We wouldn't be registered if the CDO had not already been created...
+	}
+
+	return nullptr;
+}
+
+// this will hold onto references to the classes/enums/structs that are used as a part of type definitions.  In the case that something is forcibly deleted
+// the contents of the RegisteredTypes array will have it's ClassStructorEnum cleared, and it will no longer be treated as a valid type
+// (FNiagaraTypeDefinitionHandle::Resolve will return an invalid Dummy reference).
+// 
+// It might be worth investigating changing RegisteredTypeIndexMap to be hashed based on the name of the class/enum/struct rather than the pointer.  That
+// way if we reload a type old data, with their existing index, will still be treated as valid.  In this current implementation old FNiagaraVariables will
+// have invalid types.  But, for now I think it would be an error to unload a type without having the variables that dependend on it also getting deleted.
+void FNiagaraTypeRegistry::AddReferencedObjects(FReferenceCollector& Collector)
+{
+	UE::TScopeLock Lock(RegisteredTypesLock);
+	for (const TUniquePtr<FNiagaraTypeDefinition>& RegisteredType : RegisteredTypes)
+	{
+		if (RegisteredType.IsValid())
+		{
+			Collector.AddReferencedObject(RegisteredType->ClassStructOrEnum);
+		}
+	}
+}
+
+FString FNiagaraTypeRegistry::GetReferencerName() const
+{
+	return TEXT("FNiagaraTypeRegistry");
+}
+
+FNiagaraTypeRegistry::FNiagaraTypeRegistry()
+{
+	NiagaraDebugVisHelper::GCoreTypeRegistrySingletonPtr = this;
+
+}
+
+FNiagaraTypeRegistry::~FNiagaraTypeRegistry()
+{
+	NiagaraDebugVisHelper::GCoreTypeRegistrySingletonPtr = nullptr;
+}
+
+FNiagaraTypeRegistry& FNiagaraTypeRegistry::Get()
+{
+	return TLazySingleton<FNiagaraTypeRegistry>::Get();
+}
+
+void FNiagaraTypeRegistry::Init()
+{
+	IPluginManager& PluginManager = IPluginManager::Get();
+	NiagaraTypeRegistryLocal::PluginUnmountedHandle = PluginManager.OnPluginUnmounted().AddStatic(NiagaraTypeRegistryLocal::PluginUnmounted);
+}
+
+void FNiagaraTypeRegistry::TearDown()
+{
+	IPluginManager& PluginManager = IPluginManager::Get();
+	PluginManager.OnPluginUnmounted().Remove(NiagaraTypeRegistryLocal::PluginUnmountedHandle);
+
+	TLazySingleton<FNiagaraTypeRegistry>::TearDown();
+}
+
+void FNiagaraTypeRegistry::InvalidateTypesByPath(const FString& AssetPath)
+{
+	FNiagaraTypeRegistry& Registry = Get();
+
+	if (!Registry.bModuleInitialized)
+	{
+		return;
+	}
+
+	TArray<int32> TypesToInvalidate;
+
+	{
+		FReadScopeLock Lock(Registry.RegistrationLock);
+
+		const int32 TypeCount = Registry.RegisteredSoftObjectPaths.Num();
+		for (int32 TypeIt = 0; TypeIt < TypeCount; ++TypeIt)
+		{
+			if (Registry.RegisteredSoftObjectPaths[TypeIt].GetLongPackageName().StartsWith(AssetPath))
+			{
+				TypesToInvalidate.Add(TypeIt);
+			}
+		}
+	}
+
+	TArray<FNiagaraTypeDefinition> InvalidatedTypes;
+	InvalidatedTypes.Reserve(TypesToInvalidate.Num());
+
+	{
+		UE::TScopeLock Lock(Registry.RegisteredTypesLock);
+
+		for (int32 TypeToInvalidate : TypesToInvalidate)
+		{
+			TUniquePtr<FNiagaraTypeDefinition>& TypePtr = Registry.RegisteredTypes[TypeToInvalidate];
+			InvalidatedTypes.Add(*TypePtr);
+
+			// we leave the type definition around, but invalidate it so that anyone holding onto a reference will still have
+			// valid memory but we aren't going to be holding onto a class that has been invalidated
+			TypePtr->Invalidate();
+		}
+	}
+
+	{
+		FReadScopeLock Lock(Registry.RegistrationLock);
+
+		for (const FNiagaraTypeDefinition& TypeToRemove : InvalidatedTypes)
+		{
+			// we only support unregistering implicitly added types which don't include support for the different
+			// flag based categorizations
+			ensure(!Registry.RegisteredUserVariableTypes.Contains(TypeToRemove));
+			ensure(!Registry.RegisteredUserVariableTypes.Contains(TypeToRemove));
+			ensure(!Registry.RegisteredSystemVariableTypes.Contains(TypeToRemove));
+			ensure(!Registry.RegisteredEmitterVariableTypes.Contains(TypeToRemove));
+			ensure(!Registry.RegisteredParticleVariableTypes.Contains(TypeToRemove));
+			ensure(!Registry.RegisteredParamTypes.Contains(TypeToRemove));
+			ensure(!Registry.RegisteredPayloadTypes.Contains(TypeToRemove));
+			ensure(!Registry.RegisteredUserDefinedTypes.Contains(TypeToRemove));
+			ensure(!Registry.RegisteredNumericTypes.Contains(TypeToRemove));
+			ensure(!Registry.RegisteredIndexTypes.Contains(TypeToRemove));
+
+			// we leave the following properties alone.  These could be considered orphaned, but in the current setup they
+			// would not be repopulated if the type is re-registered (implicitly)
+			//RegisteredStructConversionMap
+			//RegisteredSoftObjectPaths
+		}
+	}
+}
+
+int32 FNiagaraTypeRegistry::RegisterTypeInternal(const FNiagaraTypeDefinition& NewType)
+{
+	uint32 TypeHash = 0;
+
+	if (NewType.RegisteredTypeDefIndex == INDEX_NONE)
+	{
+		TypeHash = GetTypeHash(NewType);
+
+		// first see if we're in the TLS cache
+		if (const int32* ExistingIndex = NiagaraTypeRegistryLocal::RegisteredTypeIndexCacheTLS.Find(TypeHash))
+		{
+			// be sure to update the cached value on the type definition as well
+			NewType.RegisteredTypeDefIndex = *ExistingIndex;
+		}
+	}
+
+	// if we haven't managed to find the index in either of the above caches then we'll have to use a lock to 
+	// check the primary array
+	if (NewType.RegisteredTypeDefIndex == INDEX_NONE)
+	{
+		FReadScopeLock Lock(RegistrationLock);
+		if (const int32* ExistingIndex = RegisteredTypeIndexMap.Find(TypeHash))
+		{
+			NewType.RegisteredTypeDefIndex = *ExistingIndex;
+		}
+	}
+
+	// because types can be invalidated by plugins being unloaded we need to evaluate whether we need
+	// to perform a soft re-registration, i.e. repopulating the type back into the RegisteredTypes array
+	if (NewType.RegisteredTypeDefIndex != INDEX_NONE)
+	{
+		if (NewType.IsValid())
+		{
+			UE::TScopeLock Lock(RegisteredTypesLock);
+			if (ensure(RegisteredTypes.IsValidIndex(NewType.RegisteredTypeDefIndex)))
+			{
+				*RegisteredTypes[NewType.RegisteredTypeDefIndex] = NewType;
+			}
+		}
+	}
+	else
+	{
+		// we need to create a new entry for a registered type
+		{
+			UE::TScopeLock Lock(RegisteredTypesLock);
+			NewType.RegisteredTypeDefIndex = RegisteredTypes.Num();
+			RegisteredTypes.Add(MakeUnique<FNiagaraTypeDefinition>(NewType));
+		}
+
+		{
+			FRWScopeLock Lock(RegistrationLock, SLT_Write);
+			RegisteredSoftObjectPaths.Emplace(NewType.ClassStructOrEnum);
+			RegisteredTypeIndexMap.Add(TypeHash, NewType.RegisteredTypeDefIndex);
+		}
+
+		NiagaraTypeRegistryLocal::RegisteredTypeIndexCacheTLS.Add(TypeHash, NewType.RegisteredTypeDefIndex);
+	}
+
+	return NewType.RegisteredTypeDefIndex;
+}
+
+const FNiagaraTypeDefinition& FNiagaraTypeRegistry::GetRegisteredType(int32 TypeIndex) const
+{
+	if (TypeIndex != INDEX_NONE)
+	{
+		const FNiagaraTypeRegistry& Registry = FNiagaraTypeRegistry::Get();
+
+		UE::TScopeLock Lock(Registry.RegisteredTypesLock);
+
+		if (Registry.RegisteredTypes.IsValidIndex(TypeIndex))
+		{
+			const TUniquePtr<FNiagaraTypeDefinition>& RegisteredType = Registry.RegisteredTypes[TypeIndex];
+			// If the type is invalid, then it's likely that it has been invalidated by GC because the underlying object has
+			// been unloaded.
+			if (RegisteredType.IsValid() && RegisteredType->IsValid())
+			{
+				return *RegisteredType;
+			}
+		}
+	}
+
+	static FNiagaraTypeDefinition Dummy;
+	return Dummy;
+}
+
+TArray<FNiagaraTypeDefinition> FNiagaraTypeRegistry::GetRegisteredTypes()
+{
+	FNiagaraTypeRegistry& Registry = Get();
+
+	TArray<FNiagaraTypeDefinition> Types;
+
+	{
+		UE::TScopeLock Lock(Registry.RegisteredTypesLock);
+		Types.Reserve(Registry.RegisteredTypes.Num());
+		Algo::TransformIf(Registry.RegisteredTypes, Types,
+			[](const TUniquePtr<FNiagaraTypeDefinition>& TypePtr) -> bool
+			{
+				return TypePtr.IsValid() && TypePtr->IsValid();
+			},
+			[](const TUniquePtr<FNiagaraTypeDefinition>& TypePtr) -> FNiagaraTypeDefinition
+			{
+				return *TypePtr;
+			});
+	}
+
+	return Types;
+}
+
+TArray<FNiagaraTypeDefinition> FNiagaraTypeRegistry::GetRegisteredUserVariableTypes()
+{
+	FNiagaraTypeRegistry& Registry = Get();
+
+	TArray<FNiagaraTypeDefinition> Types;
+
+	{
+		FReadScopeLock Lock(Registry.RegistrationLock);
+		Types = Registry.RegisteredUserVariableTypes;
+	}
+
+	return Types;
+}
+
+TArray<FNiagaraTypeDefinition> FNiagaraTypeRegistry::GetRegisteredSystemVariableTypes()
+{
+	FNiagaraTypeRegistry& Registry = Get();
+
+	TArray<FNiagaraTypeDefinition> Types;
+
+	{
+		FReadScopeLock Lock(Registry.RegistrationLock);
+		Types = Registry.RegisteredSystemVariableTypes;
+	}
+
+	return Types;
+}
+
+TArray<FNiagaraTypeDefinition> FNiagaraTypeRegistry::GetRegisteredEmitterVariableTypes()
+{
+	FNiagaraTypeRegistry& Registry = Get();
+
+	TArray<FNiagaraTypeDefinition> Types;
+
+	{
+		FReadScopeLock Lock(Registry.RegistrationLock);
+		Types = Registry.RegisteredEmitterVariableTypes;
+	}
+
+	return Types;
+}
+
+TArray<FNiagaraTypeDefinition> FNiagaraTypeRegistry::GetRegisteredParticleVariableTypes()
+{
+	FNiagaraTypeRegistry& Registry = Get();
+
+	TArray<FNiagaraTypeDefinition> Types;
+
+	{
+		FReadScopeLock Lock(Registry.RegistrationLock);
+		Types = Registry.RegisteredParticleVariableTypes;
+	}
+
+	return Types;
+}
+
+TArray<FNiagaraTypeDefinition> FNiagaraTypeRegistry::GetRegisteredParameterTypes()
+{
+	FNiagaraTypeRegistry& Registry = Get();
+
+	TArray<FNiagaraTypeDefinition> Types;
+
+	{
+		FReadScopeLock Lock(Registry.RegistrationLock);
+		Types = Registry.RegisteredParamTypes;
+	}
+
+	return Types;
+}
+
+TArray<FNiagaraTypeDefinition> FNiagaraTypeRegistry::GetRegisteredPayloadTypes()
+{
+	FNiagaraTypeRegistry& Registry = Get();
+
+	TArray<FNiagaraTypeDefinition> Types;
+
+	{
+		FReadScopeLock Lock(Registry.RegistrationLock);
+		Types = Registry.RegisteredPayloadTypes;
+	}
+
+	return Types;
+}
+
+TArray<FNiagaraTypeDefinition> FNiagaraTypeRegistry::GetUserDefinedTypes()
+{
+	FNiagaraTypeRegistry& Registry = Get();
+
+	TArray<FNiagaraTypeDefinition> Types;
+
+	{
+		FReadScopeLock Lock(Registry.RegistrationLock);
+		Types = Registry.RegisteredUserDefinedTypes;
+	}
+
+	return Types;
+}
+
+TArray<FNiagaraTypeDefinition> FNiagaraTypeRegistry::GetNumericTypes()
+{
+	FNiagaraTypeRegistry& Registry = Get();
+
+	TArray<FNiagaraTypeDefinition> Types;
+
+	{
+		FReadScopeLock Lock(Registry.RegistrationLock);
+		Types = Registry.RegisteredNumericTypes;
+	}
+
+	return Types;
+}
+
+TArray<FNiagaraTypeDefinition> FNiagaraTypeRegistry::GetIndexTypes()
+{
+	FNiagaraTypeRegistry& Registry = Get();
+
+	TArray<FNiagaraTypeDefinition> Types;
+
+	{
+		FReadScopeLock Lock(Registry.RegistrationLock);
+		Types = Registry.RegisteredIndexTypes;
+	}
+
+	return Types;
+}
+
+TOptional<FNiagaraTypeDefinition> FNiagaraTypeRegistry::GetRegisteredTypeByName(FName TypeName)
+{
+	TOptional<FNiagaraTypeDefinition> Results;
+
+	Get().ForAllSelectedRegisteredTypes(
+		[TypeName](const FNiagaraTypeDefinition& TypeDef) -> bool
+		{
+			return (TypeDef.GetFName() == TypeName);
+		},
+		[&Results](const FNiagaraTypeDefinition& TypeDef) -> bool
+		{
+			Results = TypeDef;
+			return false; // continue
+		}
+	);
+
+	return Results;
+}
+
+
+void FNiagaraTypeRegistry::ClearUserDefinedRegistry()
+{
+	FNiagaraTypeRegistry& Registry = Get();
+
+	FRWScopeLock Lock(Registry.RegistrationLock, SLT_Write);
+
+	for (const FNiagaraTypeDefinition& Def : Registry.RegisteredUserDefinedTypes)
+	{
+		Registry.RegisteredPayloadTypes.Remove(Def);
+		Registry.RegisteredParamTypes.Remove(Def);
+		Registry.RegisteredNumericTypes.Remove(Def);
+		Registry.RegisteredIndexTypes.Remove(Def);
+	}
+
+	Registry.RegisteredUserDefinedTypes.Empty();
+
+	// note that we don't worry about cleaning up RegisteredTypes or RegisteredTypeIndexMap because we don't
+	// want to invalidate any indexes that are already stored in FNiagaraTypeDefinitionHandle.  If re-registered
+	// they will be given the same index, and if they are orphaned we don't want to have invalid indices on the handle.
+}
+
+void FNiagaraTypeRegistry::Register(const FNiagaraTypeDefinition& NewType, bool bCanBeParameter, bool bCanBePayload, bool bIsUserDefined)
+{
+	ENiagaraTypeRegistryFlags Flags =
+		ENiagaraTypeRegistryFlags::AllowUserVariable |
+		ENiagaraTypeRegistryFlags::AllowSystemVariable |
+		ENiagaraTypeRegistryFlags::AllowEmitterVariable;
+	if (bCanBeParameter)
+	{
+		Flags |= ENiagaraTypeRegistryFlags::AllowParameter;
+	}
+	if (bCanBePayload)
+	{
+		Flags |= ENiagaraTypeRegistryFlags::AllowPayload;
+	}
+	if (bIsUserDefined)
+	{
+		Flags |= ENiagaraTypeRegistryFlags::IsUserDefined;
+	}
+
+	Register(NewType, Flags);
+}
+
+void FNiagaraTypeRegistry::ProcessRegistryQueue()
+{
+	FNiagaraTypeRegistry& Registry = Get();
+	{
+		FRWScopeLock Lock(Registry.RegistrationLock, SLT_Write);
+		Registry.bModuleInitialized = true;
+	}
+	FQueuedRegistryEntry Entry;
+	while (Registry.RegistryQueue.Dequeue(Entry))
+	{
+		Register(Entry.NewType, Entry.Flags);
+	}
+}
+
+void FNiagaraTypeRegistry::Register(const FNiagaraTypeDefinition& NewType, ENiagaraTypeRegistryFlags Flags)
+{
+	FNiagaraTypeRegistry& Registry = Get();
+	{
+		FReadScopeLock Lock(Registry.RegistrationLock);
+		if (!Registry.bModuleInitialized)
+		{
+			// In a packaged game it can happen that CDOs are created before the Niagara module had a chance to be initialized.
+			// This is problematic, as the swc struct builder tries to access other Niagara types, so we delay the registration until the module is properly initialized.
+			Registry.RegistryQueue.Enqueue({ NewType, Flags });
+			return;
+		}
+	}
+
+	if (FNiagaraTypeHelper::IsLWCType(NewType))
+	{
+		// register the swc type as well if necessary
+		FNiagaraTypeDefinition(FNiagaraTypeHelper::GetSWCStruct(NewType.GetScriptStruct()), FNiagaraTypeDefinition::EAllowUnfriendlyStruct::Deny);
+	}
+
+	Registry.RegisterTypeInternal(NewType);
+
+	{
+		FRWScopeLock Lock(Registry.RegistrationLock, SLT_Write);
+		if (EnumHasAnyFlags(Flags, ENiagaraTypeRegistryFlags::AllowUserVariable))
+		{
+			Registry.RegisteredUserVariableTypes.AddUnique(NewType);
+		}
+
+		if (EnumHasAnyFlags(Flags, ENiagaraTypeRegistryFlags::AllowSystemVariable))
+		{
+			Registry.RegisteredSystemVariableTypes.AddUnique(NewType);
+		}
+
+		if (EnumHasAnyFlags(Flags, ENiagaraTypeRegistryFlags::AllowEmitterVariable))
+		{
+			Registry.RegisteredEmitterVariableTypes.AddUnique(NewType);
+		}
+
+		if (EnumHasAnyFlags(Flags, ENiagaraTypeRegistryFlags::AllowParticleVariable))
+		{
+			Registry.RegisteredParticleVariableTypes.AddUnique(NewType);
+		}
+
+		if (EnumHasAnyFlags(Flags, ENiagaraTypeRegistryFlags::AllowParameter))
+		{
+			Registry.RegisteredParamTypes.AddUnique(NewType);
+		}
+
+		if (EnumHasAnyFlags(Flags, ENiagaraTypeRegistryFlags::AllowPayload))
+		{
+			Registry.RegisteredPayloadTypes.AddUnique(NewType);
+		}
+
+		if (EnumHasAnyFlags(Flags, ENiagaraTypeRegistryFlags::IsUserDefined))
+		{
+			Registry.RegisteredUserDefinedTypes.AddUnique(NewType);
+		}
+
+		if (FNiagaraTypeDefinition::IsValidNumericInput(NewType))
+		{
+			Registry.RegisteredNumericTypes.AddUnique(NewType);
+		}
+
+		if (NewType.IsIndexType())
+		{
+			Registry.RegisteredIndexTypes.AddUnique(NewType);
+		}
+	}
+}
+
+bool FNiagaraTypeRegistry::IsStaticPossible(const FNiagaraTypeDefinition& InSrc)
+{
+	if (InSrc.IsStatic())
+	{
+		return true;
+	}
+
+	FNiagaraTypeRegistry& Registry = Get();
+
+	bool bStaticTypeFound = false;
+	Registry.ForAllSelectedRegisteredTypes(
+		[&InSrc](const FNiagaraTypeDefinition& TypeDef) -> bool
+		{
+			return InSrc.IsSameBaseDefinition(TypeDef) && TypeDef.IsStatic();
+		},
+		[&bStaticTypeFound](const FNiagaraTypeDefinition& TypeDef) -> bool
+		{
+			bStaticTypeFound = true;
+			return false;
+		});
+
+	return bStaticTypeFound;
+}
+
+void FNiagaraTypeRegistry::RegisterStructConverter(const FNiagaraTypeDefinition& SourceType, const FNiagaraLwcStructConverter& StructConverter)
+{
+	FNiagaraTypeRegistry& Registry = Get();
+
+	int32 TypeIndex = Registry.RegisterTypeInternal(SourceType);
+
+	FRWScopeLock Lock(Registry.RegistrationLock, SLT_Write);
+	Registry.RegisteredStructConversionMap.Add(TypeIndex, StructConverter);
+}
+
+FNiagaraLwcStructConverter FNiagaraTypeRegistry::GetStructConverter(const FNiagaraTypeDefinition& SourceType)
+{
+	FNiagaraTypeRegistry& Registry = Get();
+
+	FReadScopeLock Lock(Registry.RegistrationLock);
+	const uint32 TypeHash = GetTypeHash(SourceType);
+	if (const int32* TypeIndex = Registry.RegisteredTypeIndexMap.Find(TypeHash))
+	{
+		if (FNiagaraLwcStructConverter* Converter = Registry.RegisteredStructConversionMap.Find(*TypeIndex))
+		{
+			return *Converter;
+		}
+	}
+
+	return FNiagaraLwcStructConverter();
+}
+
+FNiagaraTypeDefinition FNiagaraTypeRegistry::GetTypeForStruct(UScriptStruct* InStruct)
+{
+	FNiagaraTypeRegistry& Registry = Get();
+
+	FNiagaraTypeDefinition FoundType;
+
+	Registry.ForAllSelectedRegisteredTypes(
+		[InStruct](const FNiagaraTypeDefinition& TypeDef) -> bool
+		{
+			return TypeDef.GetStruct() && InStruct == TypeDef.GetStruct();
+		},
+		[&FoundType](const FNiagaraTypeDefinition& TypeDef) -> bool
+		{
+			FoundType = TypeDef;
+			return false;
+		});
+
+	return FoundType;
+}
+
+#if 0
+template<typename TAction>
+static void ForEachRegisteredType(TAction Func)
+{
+	FNiagaraTypeRegistry& Registry = Get();
+
+	auto TrivialSelect = [](const FNiagaraTypeDefinition& TypeDef) -> bool
+		{
+			return true;
+		};
+
+	Registry.ForAllSelectedRegisteredTypes(TrivialSelect, Func);
+}
+
+	template<typename TSelectAction, typename TAction>
+	void ForAllSelectedRegisteredTypes(TSelectAction Select, TAction Func) const
+	{
+		TArray<FNiagaraTypeDefinition> LocalCopy;
+
+		{
+			UE::TScopeLock Lock(RegisteredTypesLock);
+			LocalCopy.Reserve(RegisteredTypes.Num());
+			Algo::TransformIf(RegisteredTypes, LocalCopy,
+				[Select](const TUniquePtr<FNiagaraTypeDefinition>& DefPtr) -> bool
+				{
+					return DefPtr.IsValid() && DefPtr->IsValid() && Select(*DefPtr);
+				},
+				[](const TUniquePtr<FNiagaraTypeDefinition>& DefPtr) -> FNiagaraTypeDefinition
+				{
+					return *DefPtr;
+				});
+		}
+
+		for (const FNiagaraTypeDefinition& Type : LocalCopy)
+		{
+			if (!Func(Type))
+			{
+				break;
+			}
+		}
+	}
+	#endif
\ No newline at end of file
diff --git a/Engine/Plugins/FX/Niagara/Source/Niagara/Public/NiagaraCommon.h b/Engine/Plugins/FX/Niagara/Source/Niagara/Public/NiagaraCommon.h
--- a/Engine/Plugins/FX/Niagara/Source/Niagara/Public/NiagaraCommon.h
+++ b/Engine/Plugins/FX/Niagara/Source/Niagara/Public/NiagaraCommon.h
@@ -14,6 +14,7 @@
 #include "NiagaraScalabilityState.h"
 #include "NiagaraTickBehaviorEnum.h"
 #include "NiagaraTypes.h"
+#include "NiagaraTypeRegistry.h"
 #include "NiagaraCore.h"
 
 #if UE_ENABLE_INCLUDE_ORDER_DEPRECATED_IN_5_2
diff --git a/Engine/Plugins/FX/Niagara/Source/Niagara/Public/NiagaraTypeRegistry.h b/Engine/Plugins/FX/Niagara/Source/Niagara/Public/NiagaraTypeRegistry.h
new file mode 100644
--- /dev/null
+++ b/Engine/Plugins/FX/Niagara/Source/Niagara/Public/NiagaraTypeRegistry.h
@@ -0,0 +1,182 @@
+// Copyright Epic Games, Inc. All Rights Reserved.
+
+#pragma once
+
+#include "CoreMinimal.h"
+#include "Misc/ScopeLock.h"
+#include "Misc/SecureHash.h"
+#include "Misc/SpinLock.h"
+#include "UObject/GCObject.h"
+#include "UObject/UnrealType.h"
+#include "Containers/Queue.h"
+#include "NiagaraCore.h"
+#include "NiagaraTypes.h"
+
+//////////////////////////////////////////////////////////////////////////
+
+enum class ENiagaraTypeRegistryFlags : uint32
+{
+	None					= 0,
+
+	AllowUserVariable		= (1 << 0),
+	AllowSystemVariable		= (1 << 1),
+	AllowEmitterVariable	= (1 << 2),
+	AllowParticleVariable	= (1 << 3),
+	AllowAnyVariable		= (AllowUserVariable | AllowSystemVariable | AllowEmitterVariable | AllowParticleVariable),
+	AllowNotUserVariable	= (AllowSystemVariable | AllowEmitterVariable | AllowParticleVariable),
+
+	AllowParameter			= (1 << 4),
+	AllowPayload			= (1 << 5),
+
+	IsUserDefined			= (1 << 6),
+
+};
+
+ENUM_CLASS_FLAGS(ENiagaraTypeRegistryFlags)
+
+/* Contains all types currently available for use in Niagara
+* Used by UI to provide selection; new uniforms and variables
+* may be instanced using the types provided here
+*/
+class FNiagaraTypeRegistry : public FGCObject
+{
+public:
+	// In order to simplify the requirements for handling access to the type array from various threads (async loading as an example)
+	// we ensure that the entries in the array will never be invalidated.  While the TUniquePtr may be moved (through reallocation)
+	// the array will not shrink, and the type definition pointer within the TUniquePtr will not be invalidated.
+	using FRegisteredTypesArray = TArray<TUniquePtr<FNiagaraTypeDefinition>>;
+
+	static NIAGARA_API TArray<FNiagaraTypeDefinition> GetRegisteredTypes();
+	static NIAGARA_API TArray<FNiagaraTypeDefinition> GetRegisteredUserVariableTypes();
+	static NIAGARA_API TArray<FNiagaraTypeDefinition> GetRegisteredSystemVariableTypes();
+	static NIAGARA_API TArray<FNiagaraTypeDefinition> GetRegisteredEmitterVariableTypes();
+	static NIAGARA_API TArray<FNiagaraTypeDefinition> GetRegisteredParticleVariableTypes();
+	static NIAGARA_API TArray<FNiagaraTypeDefinition> GetRegisteredParameterTypes();
+	static NIAGARA_API TArray<FNiagaraTypeDefinition> GetRegisteredPayloadTypes();
+	static NIAGARA_API TArray<FNiagaraTypeDefinition> GetUserDefinedTypes();
+	static NIAGARA_API TArray<FNiagaraTypeDefinition> GetNumericTypes();
+	static NIAGARA_API TArray<FNiagaraTypeDefinition> GetIndexTypes();
+
+	static NIAGARA_API TOptional<FNiagaraTypeDefinition> GetRegisteredTypeByName(FName TypeName);
+
+	static NIAGARA_API UNiagaraDataInterfaceBase* GetDefaultDataInterfaceByName(const FString& DIClassName);
+
+	static NIAGARA_API void ClearUserDefinedRegistry();
+
+	UE_DEPRECATED(4.27, "This overload is deprecated, please use the Register function that takes registration flags instead.")
+	static NIAGARA_API void Register(const FNiagaraTypeDefinition& NewType, bool bCanBeParameter, bool bCanBePayload, bool bIsUserDefined);
+
+	static NIAGARA_API void ProcessRegistryQueue();
+
+	static NIAGARA_API void Register(const FNiagaraTypeDefinition &NewType, ENiagaraTypeRegistryFlags Flags);
+
+	static NIAGARA_API bool IsStaticPossible(const FNiagaraTypeDefinition& InSrc);
+
+	static NIAGARA_API void RegisterStructConverter(const FNiagaraTypeDefinition& SourceType, const FNiagaraLwcStructConverter& StructConverter);
+
+	static NIAGARA_API FNiagaraLwcStructConverter GetStructConverter(const FNiagaraTypeDefinition& SourceType);
+
+	static NIAGARA_API FNiagaraTypeDefinition GetTypeForStruct(UScriptStruct* InStruct);
+	static NIAGARA_API void InvalidateTypesByPath(const FString& AssetPath);
+
+	/** LazySingleton interface */
+	static NIAGARA_API FNiagaraTypeRegistry& Get();
+	static void Init();
+	static void TearDown();
+
+	/** FGCObject interface */
+	NIAGARA_API virtual void AddReferencedObjects(FReferenceCollector& Collector) override;
+	NIAGARA_API virtual FString GetReferencerName() const override;
+
+	template<typename TAction>
+	static void ForEachRegisteredType(TAction Func)
+	{
+		FNiagaraTypeRegistry& Registry = Get();
+
+		auto TrivialSelect = [](const FNiagaraTypeDefinition& TypeDef) -> bool
+		{
+			return true;
+		};
+
+		Registry.ForAllSelectedRegisteredTypes(TrivialSelect, Func);
+	}
+
+private:
+	struct FQueuedRegistryEntry
+	{
+		FNiagaraTypeDefinition NewType;
+		ENiagaraTypeRegistryFlags Flags;
+	};
+	friend class FLazySingleton;
+	friend struct FNiagaraTypeDefinitionHandle;
+
+	NIAGARA_API FNiagaraTypeRegistry();
+	NIAGARA_API virtual ~FNiagaraTypeRegistry();
+
+	int32 RegisterTypeInternal(const FNiagaraTypeDefinition& NewType);
+	const FNiagaraTypeDefinition& GetRegisteredType(int32 TypeIndex) const;
+
+	// The spinlock is used for accessing the RegisteredTypes array in performance critical areas (resolve TypeDefinitionHandles being the
+	// most common).
+	mutable UE::FSpinLock RegisteredTypesLock; 
+	FRegisteredTypesArray RegisteredTypes;
+
+	// Covers manipulation of all of the containers populated during registration (except for RegisteredTypes which is handled
+	// by it's own lock so that we can provide a fast path for the common case of reading the array)
+	FRWLock RegistrationLock;
+
+	TArray<FNiagaraTypeDefinition> RegisteredUserVariableTypes;
+	TArray<FNiagaraTypeDefinition> RegisteredSystemVariableTypes;
+	TArray<FNiagaraTypeDefinition> RegisteredEmitterVariableTypes;
+	TArray<FNiagaraTypeDefinition> RegisteredParticleVariableTypes;
+	TArray<FNiagaraTypeDefinition> RegisteredParamTypes;
+	TArray<FNiagaraTypeDefinition> RegisteredPayloadTypes;
+	TArray<FNiagaraTypeDefinition> RegisteredUserDefinedTypes;
+	TArray<FNiagaraTypeDefinition> RegisteredNumericTypes;
+	TArray<FNiagaraTypeDefinition> RegisteredIndexTypes;
+
+	TMap<uint32, int32> RegisteredTypeIndexMap;
+	TMap<uint32, FNiagaraLwcStructConverter> RegisteredStructConversionMap;
+
+	// For each entry in RegisteredTypes we store the SoftObjectPath of the underlying
+	// class/struct so that we can handle dynamically unloading plugins
+	TArray<FSoftObjectPath> RegisteredSoftObjectPaths;
+
+	bool bModuleInitialized = false;
+	TQueue<FQueuedRegistryEntry, EQueueMode::Mpsc> RegistryQueue;
+
+	// helper method to go through the various caching methods for finding the
+	// index into RegisteredTypes for the provided TypeDefinition
+	static NIAGARA_API bool GetCachedRegisteredTypeIndex(const FNiagaraTypeDefinition& Type, uint32& OutTypeHash, int32& OutIndex);
+
+	// populates the cache on the TypeDefinition and TLS
+	static NIAGARA_API void CacheRegisteredTypeIndex(const FNiagaraTypeDefinition& Type, uint32 TypeHash, int32 Index);
+
+	template<typename TSelectAction, typename TAction>
+	void ForAllSelectedRegisteredTypes(TSelectAction Select, TAction Func) const
+	{
+		TArray<FNiagaraTypeDefinition> LocalCopy;
+
+		{
+			UE::TScopeLock Lock(RegisteredTypesLock);
+			LocalCopy.Reserve(RegisteredTypes.Num());
+			Algo::TransformIf(RegisteredTypes, LocalCopy,
+				[Select](const TUniquePtr<FNiagaraTypeDefinition>& DefPtr) -> bool
+				{
+					return DefPtr.IsValid() && DefPtr->IsValid() && Select(*DefPtr);
+				},
+				[](const TUniquePtr<FNiagaraTypeDefinition>& DefPtr) -> FNiagaraTypeDefinition
+				{
+					return *DefPtr;
+				});
+		}
+
+		for (const FNiagaraTypeDefinition& Type : LocalCopy)
+		{
+			if (!Func(Type))
+			{
+				break;
+			}
+		}
+	}
+};
diff --git a/Engine/Plugins/FX/Niagara/Source/Niagara/Public/NiagaraTypes.h b/Engine/Plugins/FX/Niagara/Source/Niagara/Public/NiagaraTypes.h
--- a/Engine/Plugins/FX/Niagara/Source/Niagara/Public/NiagaraTypes.h
+++ b/Engine/Plugins/FX/Niagara/Source/Niagara/Public/NiagaraTypes.h
@@ -671,7 +671,7 @@ public:
 
 	// Construct blank raw type definition 
 	FORCEINLINE FNiagaraTypeDefinition(UObject* Object)
-		: ClassStructOrEnum(Object), Flags(TF_None), Size(INDEX_NONE), Alignment(INDEX_NONE)
+		: ClassStructOrEnum(Object)
 #if WITH_EDITORONLY_DATA
 		, Struct_DEPRECATED(nullptr), Enum_DEPRECATED(nullptr)
 #endif
@@ -692,35 +692,43 @@ public:
 		{
 			check(0);
 		}
+
+		GeneratePathNameHash();
 	}
 
 	FORCEINLINE FNiagaraTypeDefinition(UClass *ClassDef)
-		: ClassStructOrEnum(ClassDef), UnderlyingType(UT_Class), Flags(TF_None), Size(INDEX_NONE), Alignment(INDEX_NONE)
+		: ClassStructOrEnum(ClassDef), UnderlyingType(UT_Class)
 #if WITH_EDITORONLY_DATA
 		, Struct_DEPRECATED(nullptr), Enum_DEPRECATED(nullptr)
 #endif
 	{
 		checkSlow(ClassStructOrEnum != nullptr);
+
+		GeneratePathNameHash();
 	}
 
 	FORCEINLINE FNiagaraTypeDefinition(UEnum *EnumDef)
-		: ClassStructOrEnum(EnumDef), UnderlyingType(UT_Enum), Flags(TF_None), Size(INDEX_NONE), Alignment(INDEX_NONE)
+		: ClassStructOrEnum(EnumDef), UnderlyingType(UT_Enum)
 #if WITH_EDITORONLY_DATA
 		, Struct_DEPRECATED(nullptr), Enum_DEPRECATED(nullptr)
 #endif
 	{
 		checkSlow(ClassStructOrEnum != nullptr);
+
+		GeneratePathNameHash();
 	}
 
 	enum class EAllowUnfriendlyStruct : uint8 { Deny, Allow };
 	FORCEINLINE FNiagaraTypeDefinition(UScriptStruct* StructDef, EAllowUnfriendlyStruct AllowUnfriendlyStruct)
-		: ClassStructOrEnum(StructDef), UnderlyingType(UT_Struct), Flags(TF_None), Size(INDEX_NONE), Alignment(INDEX_NONE)
+		: ClassStructOrEnum(StructDef), UnderlyingType(UT_Struct)
 #if WITH_EDITORONLY_DATA
 		, Struct_DEPRECATED(nullptr), Enum_DEPRECATED(nullptr)
 #endif
 	{
 		checkSlow(ClassStructOrEnum != nullptr);
 		ensureAlwaysMsgf(AllowUnfriendlyStruct == EAllowUnfriendlyStruct::Allow || FNiagaraTypeHelper::IsNiagaraFriendlyTopLevelStruct(StructDef, ENiagaraStructConversion::UserFacing), TEXT("Struct(%s) is not supported."), *StructDef->GetName());
+
+		GeneratePathNameHash();
 	}
 	
 	FORCEINLINE FNiagaraTypeDefinition(UScriptStruct* StructDef) : FNiagaraTypeDefinition(StructDef, EAllowUnfriendlyStruct::Deny) {}
@@ -740,7 +748,7 @@ public:
 
 	// Construct a blank raw type definition
 	FORCEINLINE FNiagaraTypeDefinition()
-		: ClassStructOrEnum(nullptr), UnderlyingType(UT_None), Flags(TF_None), Size(INDEX_NONE), Alignment(INDEX_NONE)
+		: ClassStructOrEnum(nullptr), UnderlyingType(UT_None)
 #if WITH_EDITORONLY_DATA
 		, Struct_DEPRECATED(nullptr), Enum_DEPRECATED(nullptr)
 #endif
@@ -938,6 +946,11 @@ public:
 	{ 
 		return ClassStructOrEnum != nullptr;
 	}
+	void Invalidate()
+	{
+		ClassStructOrEnum = nullptr;
+		UnderlyingTypePathNameHash = 0;
+	}
 
 	NIAGARA_API bool AppendCompileHash(FNiagaraCompileHashVisitor* InVisitor) const;
 
@@ -960,11 +973,23 @@ public:
 	NIAGARA_API bool Serialize(FArchive& Ar);
 
 private:
+	// for allowing access to the cached RegisteredTypeDefIndex
+	friend class FNiagaraTypeRegistry;
+
 	UPROPERTY(EditAnywhere, Category = Type)
-	uint8 Flags;
+	uint8 Flags = TF_None;
+
+	mutable int16 Size = INDEX_NONE;
+	mutable int16 Alignment = INDEX_NONE;
+
+	mutable int32 RegisteredTypeDefIndex = INDEX_NONE;
 
-	mutable int16 Size;
-	mutable int16 Alignment;
+	// the type hash of the full path name for the underlying type.  This is helpful when identifying the type
+	// especially in scenarios where plugins contain the structs/classes and we want to re-use slots in the registry
+	// even if the objects are recreated (with different pointers/object indices)
+	uint32 UnderlyingTypePathNameHash = INDEX_NONE;
+
+	void InvalidateCachedIndex(bool bInvalidate = true);
 
 #if WITH_EDITORONLY_DATA
 	UPROPERTY()
@@ -1145,9 +1170,11 @@ private:
 
 	static NIAGARA_API FNiagaraTypeDefinition CollisionEventDef;
 
+	void NIAGARA_API GeneratePathNameHash();
+
 	friend FORCEINLINE uint32 GetTypeHash(const FNiagaraTypeDefinition& Type)
 	{
-		return HashCombine(HashCombine(GetTypeHash(Type.ClassStructOrEnum), GetTypeHash(Type.UnderlyingType)), GetTypeHash(Type.GetFlags()));
+		return HashCombine(HashCombine(Type.UnderlyingTypePathNameHash, GetTypeHash(Type.UnderlyingType)), GetTypeHash(Type.GetFlags()));
 	}
 };
 
@@ -1182,333 +1209,6 @@ const FNiagaraTypeDefinition& FNiagaraTypeDefinition::Get()
 
 //////////////////////////////////////////////////////////////////////////
 
-enum class ENiagaraTypeRegistryFlags : uint32
-{
-	None					= 0,
-
-	AllowUserVariable		= (1 << 0),
-	AllowSystemVariable		= (1 << 1),
-	AllowEmitterVariable	= (1 << 2),
-	AllowParticleVariable	= (1 << 3),
-	AllowAnyVariable		= (AllowUserVariable | AllowSystemVariable | AllowEmitterVariable | AllowParticleVariable),
-	AllowNotUserVariable	= (AllowSystemVariable | AllowEmitterVariable | AllowParticleVariable),
-
-	AllowParameter			= (1 << 4),
-	AllowPayload			= (1 << 5),
-
-	IsUserDefined			= (1 << 6),
-
-};
-
-ENUM_CLASS_FLAGS(ENiagaraTypeRegistryFlags)
-
-/* Contains all types currently available for use in Niagara
-* Used by UI to provide selection; new uniforms and variables
-* may be instanced using the types provided here
-*/
-class FNiagaraTypeRegistry : public FGCObject
-{
-public:
-	enum
-	{
-		MaxRegisteredTypes = 512,
-	};
-
-	using RegisteredTypesArray = TArray<FNiagaraTypeDefinition, TFixedAllocator<MaxRegisteredTypes>>;
-
-	static const RegisteredTypesArray& GetRegisteredTypes()
-	{
-		return Get().RegisteredTypes;
-	}
-
-	static const TArray<FNiagaraTypeDefinition>& GetRegisteredUserVariableTypes()
-	{
-		return Get().RegisteredUserVariableTypes;
-	}
-
-	static const TArray<FNiagaraTypeDefinition>& GetRegisteredSystemVariableTypes()
-	{
-		return Get().RegisteredSystemVariableTypes;
-	}
-
-	static const TArray<FNiagaraTypeDefinition>& GetRegisteredEmitterVariableTypes()
-	{
-		return Get().RegisteredEmitterVariableTypes;
-	}
-
-	static const TArray<FNiagaraTypeDefinition>& GetRegisteredParticleVariableTypes()
-	{
-		return Get().RegisteredParticleVariableTypes;
-	}
-
-	static const TArray<FNiagaraTypeDefinition>& GetRegisteredParameterTypes()
-	{
-		return Get().RegisteredParamTypes;
-	}
-
-	static const TArray<FNiagaraTypeDefinition>& GetRegisteredPayloadTypes()
-	{
-		return Get().RegisteredPayloadTypes;
-	}
-
-	static const TArray<FNiagaraTypeDefinition>& GetUserDefinedTypes()
-	{
-		return Get().RegisteredUserDefinedTypes;
-	}
-
-	static const TArray<FNiagaraTypeDefinition>& GetNumericTypes()
-	{ 
-		return Get().RegisteredNumericTypes;
-	}
-
-	static const TArray<FNiagaraTypeDefinition>& GetIndexTypes()
-	{
-		return Get().RegisteredIndexTypes;
-	}
-
-	static NIAGARA_API UNiagaraDataInterfaceBase* GetDefaultDataInterfaceByName(const FString& DIClassName);
-
-	static void ClearUserDefinedRegistry()
-	{
-		FNiagaraTypeRegistry& Registry = Get();
-
-		FRWScopeLock Lock(Registry.RegisteredTypesLock, SLT_Write);
-
-		for (const FNiagaraTypeDefinition& Def : Registry.RegisteredUserDefinedTypes)
-		{
-			Registry.RegisteredPayloadTypes.Remove(Def);
-			Registry.RegisteredParamTypes.Remove(Def);
-			Registry.RegisteredNumericTypes.Remove(Def);
-			Registry.RegisteredIndexTypes.Remove(Def);
-		}
-
-		Registry.RegisteredUserDefinedTypes.Empty();
-
-		// note that we don't worry about cleaning up RegisteredTypes or RegisteredTypeIndexMap because we don't
-		// want to invalidate any indexes that are already stored in FNiagaraTypeDefinitionHandle.  If re-registered
-		// they will be given the same index, and if they are orphaned we don't want to have invalid indices on the handle.
-	}
-
-	UE_DEPRECATED(4.27, "This overload is deprecated, please use the Register function that takes registration flags instead.")
-	static void Register(const FNiagaraTypeDefinition& NewType, bool bCanBeParameter, bool bCanBePayload, bool bIsUserDefined)
-	{
-		ENiagaraTypeRegistryFlags Flags =
-			ENiagaraTypeRegistryFlags::AllowUserVariable |
-			ENiagaraTypeRegistryFlags::AllowSystemVariable |
-			ENiagaraTypeRegistryFlags::AllowEmitterVariable;
-		if (bCanBeParameter)
-		{
-			Flags |= ENiagaraTypeRegistryFlags::AllowParameter;
-		}
-		if (bCanBePayload)
-		{
-			Flags |= ENiagaraTypeRegistryFlags::AllowPayload;
-		}
-		if (bIsUserDefined)
-		{
-			Flags |= ENiagaraTypeRegistryFlags::IsUserDefined;
-		}
-										
-		Register(NewType, Flags);
-	}
-
-	static void ProcessRegistryQueue()
-	{
-		FNiagaraTypeRegistry& Registry = Get();
-		{
-			FRWScopeLock Lock(Registry.RegisteredTypesLock, SLT_Write);
-			Registry.bModuleInitialized = true;
-		}
-		FQueuedRegistryEntry Entry;
-		while (Registry.RegistryQueue.Dequeue(Entry))
-		{
-			Register(Entry.NewType, Entry.Flags);
-		}
-	}
-
-	static void Register(const FNiagaraTypeDefinition &NewType, ENiagaraTypeRegistryFlags Flags)
-	{
-		FNiagaraTypeRegistry& Registry = Get();
-		{
-			FReadScopeLock Lock(Registry.RegisteredTypesLock);
-			if (!Registry.bModuleInitialized)
-			{
-				// In a packaged game it can happen that CDOs are created before the Niagara module had a chance to be initialized.
-				// This is problematic, as the swc struct builder tries to access other Niagara types, so we delay the registration until the module is properly initialized.
-				Registry.RegistryQueue.Enqueue({NewType, Flags});
-				return;
-			}
-		}
-
-		if (FNiagaraTypeHelper::IsLWCType(NewType))
-		{
-			// register the swc type as well if necessary
-			FNiagaraTypeDefinition(FNiagaraTypeHelper::GetSWCStruct(NewType.GetScriptStruct()), FNiagaraTypeDefinition::EAllowUnfriendlyStruct::Deny);
-		}
-
-		FRWScopeLock Lock(Registry.RegisteredTypesLock, SLT_Write);
-		//TODO: Make this a map of type to a more verbose set of metadata? Such as the hlsl defs, offset table for conversions etc.
-		Registry.RegisteredTypeIndexMap.Add(GetTypeHash(NewType), Registry.RegisteredTypes.AddUnique(NewType));
-
-		if (EnumHasAnyFlags(Flags, ENiagaraTypeRegistryFlags::AllowUserVariable))
-		{
-			Registry.RegisteredUserVariableTypes.AddUnique(NewType);
-		}
-
-		if (EnumHasAnyFlags(Flags, ENiagaraTypeRegistryFlags::AllowSystemVariable))
-		{
-			Registry.RegisteredSystemVariableTypes.AddUnique(NewType);
-		}
-
-		if (EnumHasAnyFlags(Flags, ENiagaraTypeRegistryFlags::AllowEmitterVariable))
-		{
-			Registry.RegisteredEmitterVariableTypes.AddUnique(NewType);
-		}
-
-		if (EnumHasAnyFlags(Flags, ENiagaraTypeRegistryFlags::AllowParticleVariable))
-		{
-			Registry.RegisteredParticleVariableTypes.AddUnique(NewType);
-		}
-
-		if (EnumHasAnyFlags(Flags, ENiagaraTypeRegistryFlags::AllowParameter))
-		{
-			Registry.RegisteredParamTypes.AddUnique(NewType);
-		}
-
-		if (EnumHasAnyFlags(Flags, ENiagaraTypeRegistryFlags::AllowPayload))
-		{
-			Registry.RegisteredPayloadTypes.AddUnique(NewType);
-		}
-
-		if (EnumHasAnyFlags(Flags, ENiagaraTypeRegistryFlags::IsUserDefined))
-		{
-			Registry.RegisteredUserDefinedTypes.AddUnique(NewType);
-		}
-
-		if (FNiagaraTypeDefinition::IsValidNumericInput(NewType))
-		{
-			Registry.RegisteredNumericTypes.AddUnique(NewType);
-		}
-
-		if(NewType.IsIndexType())
-		{
-			Registry.RegisteredIndexTypes.AddUnique(NewType);
-		}
-	}
-
-	static bool IsStaticPossible(const FNiagaraTypeDefinition& InSrc) 
-	{
-		if (InSrc.IsStatic())
-			return true;
-
-		for (const FNiagaraTypeDefinition& TypeDef : Get().RegisteredTypes)
-		{
-			if (InSrc.IsSameBaseDefinition(TypeDef) && TypeDef.IsStatic())
-				return true;
-		}
-
-		return false;
-	}
-
-	static int32 RegisterIndexed(const FNiagaraTypeDefinition& NewType)
-	{
-		FNiagaraTypeRegistry& Registry = Get();
-
-		const uint32 TypeHash = GetTypeHash(NewType);
-		{
-			FReadScopeLock Lock(Registry.RegisteredTypesLock);
-			if (const int32* ExistingIndex = Registry.RegisteredTypeIndexMap.Find(TypeHash))
-			{
-				return *ExistingIndex;
-			}
-		}
-
-		FRWScopeLock Lock(Registry.RegisteredTypesLock, SLT_Write);
-		const int32 Index = Registry.RegisteredTypes.AddUnique(NewType);
-		Registry.RegisteredTypeIndexMap.Add(TypeHash, Index);
-		return Index;
-	}
-
-	static void RegisterStructConverter(const FNiagaraTypeDefinition& SourceType, const FNiagaraLwcStructConverter& StructConverter)
-	{
-		int32 TypeIndex = RegisterIndexed(SourceType);
-		FNiagaraTypeRegistry& Registry = Get();
-		
-		FRWScopeLock Lock(Registry.RegisteredTypesLock, SLT_Write);
-		Registry.RegisteredStructConversionMap.Add(TypeIndex, StructConverter);
-	}
-
-	static FNiagaraLwcStructConverter GetStructConverter(const FNiagaraTypeDefinition& SourceType)
-	{
-		FNiagaraTypeRegistry& Registry = Get();
-		
-		FReadScopeLock Lock(Registry.RegisteredTypesLock);
-		const uint32 TypeHash = GetTypeHash(SourceType);
-		if (const int32* TypeIndex = Registry.RegisteredTypeIndexMap.Find(TypeHash))
-		{
-			if (FNiagaraLwcStructConverter* Converter = Registry.RegisteredStructConversionMap.Find(*TypeIndex))
-			{
-				return *Converter;
-			}
-		}
-		
-		return FNiagaraLwcStructConverter();
-	}
-
-	static FNiagaraTypeDefinition GetTypeForStruct(UScriptStruct* InStruct)
-	{
-		FNiagaraTypeRegistry& Registry = Get();
-		FReadScopeLock Lock(Registry.RegisteredTypesLock);
-		for (const FNiagaraTypeDefinition& Type : Registry.RegisteredTypes)
-		{
-			if (Type.GetStruct() && InStruct == Type.GetStruct())
-			{
-				return Type;
-			}
-		}
-		return FNiagaraTypeDefinition(InStruct);
-	}
-
-	/** LazySingleton interface */
-	static NIAGARA_API FNiagaraTypeRegistry& Get();
-	static NIAGARA_API void TearDown();
-
-	/** FGCObject interface */
-	NIAGARA_API virtual void AddReferencedObjects(FReferenceCollector& Collector) override;
-	NIAGARA_API virtual FString GetReferencerName() const override;
-
-private:
-	struct FQueuedRegistryEntry
-	{
-		FNiagaraTypeDefinition NewType;
-		ENiagaraTypeRegistryFlags Flags;
-	};
-	friend class FLazySingleton;
-
-	NIAGARA_API FNiagaraTypeRegistry();
-	NIAGARA_API virtual ~FNiagaraTypeRegistry();
-
-	RegisteredTypesArray RegisteredTypes;
-
-	TArray<FNiagaraTypeDefinition> RegisteredUserVariableTypes;
-	TArray<FNiagaraTypeDefinition> RegisteredSystemVariableTypes;
-	TArray<FNiagaraTypeDefinition> RegisteredEmitterVariableTypes;
-	TArray<FNiagaraTypeDefinition> RegisteredParticleVariableTypes;
-	TArray<FNiagaraTypeDefinition> RegisteredParamTypes;
-	TArray<FNiagaraTypeDefinition> RegisteredPayloadTypes;
-	TArray<FNiagaraTypeDefinition> RegisteredUserDefinedTypes;
-	TArray<FNiagaraTypeDefinition> RegisteredNumericTypes;
-	TArray<FNiagaraTypeDefinition> RegisteredIndexTypes;
-
-	TMap<uint32, int32> RegisteredTypeIndexMap;
-	TMap<TWeakObjectPtr<UScriptStruct>, TWeakObjectPtr<UScriptStruct>> LWCRegisteredStructRemapping;
-	TMap<uint32, FNiagaraLwcStructConverter> RegisteredStructConversionMap;
-	FRWLock RegisteredTypesLock;
-
-	bool bModuleInitialized = false;
-	TQueue<FQueuedRegistryEntry, EQueueMode::Mpsc> RegistryQueue;
-};
-
 USTRUCT()
 struct FNiagaraTypeDefinitionHandle
 {
diff --git a/Engine/Plugins/FX/Niagara/Source/NiagaraEditor/Private/NiagaraClipboard.cpp b/Engine/Plugins/FX/Niagara/Source/NiagaraEditor/Private/NiagaraClipboard.cpp
--- a/Engine/Plugins/FX/Niagara/Source/NiagaraEditor/Private/NiagaraClipboard.cpp
+++ b/Engine/Plugins/FX/Niagara/Source/NiagaraEditor/Private/NiagaraClipboard.cpp
@@ -445,14 +445,7 @@ const UNiagaraClipboardFunctionInput* CreateLocalValue(UObject* InOuter, FName I
 
 FNiagaraTypeDefinition UNiagaraClipboardEditorScriptingUtilities::GetRegisteredTypeDefinitionByName(FName InTypeName)
 {
-	for (const FNiagaraTypeDefinition& RegisteredType : FNiagaraTypeRegistry::GetRegisteredTypes())
-	{
-		if (RegisteredType.GetFName() == InTypeName)
-		{
-			return RegisteredType;
-		}
-	}
-	return FNiagaraTypeDefinition();
+	return FNiagaraTypeRegistry::GetRegisteredTypeByName(InTypeName).Get(FNiagaraTypeDefinition());
 }
 
 UNiagaraClipboardFunctionInput* UNiagaraClipboardEditorScriptingUtilities::CreateFloatLocalValueInput(UObject* InOuter, FName InInputName, bool bInHasEditCondition, bool bInEditConditionValue, float InFloatValue)
