Incorporate the cbuffer structs used by the script for generating the CompilationId
#tests Modified cbuffer layout and verified that the compilation id changed

diff --git a/Engine/Plugins/FX/Niagara/Source/Niagara/Private/NiagaraModule.cpp b/Engine/Plugins/FX/Niagara/Source/Niagara/Private/NiagaraModule.cpp
--- a/Engine/Plugins/FX/Niagara/Source/Niagara/Private/NiagaraModule.cpp
+++ b/Engine/Plugins/FX/Niagara/Source/Niagara/Private/NiagaraModule.cpp
@@ -1740,6 +1740,33 @@ void FNiagaraVariable::PostSerialize(const FArchive& Ar)
 #endif
 
 #if WITH_EDITOR
+
+namespace NiagaraCBufferHelperImpl
+{
+	static FSHAHash GenerateHash(TConstArrayView<FNiagaraVariable> Variables)
+	{
+		FSHA1 HashState;
+		FNiagaraCompileHashVisitor Visitor(HashState);
+		for (const FNiagaraVariable& Var : Variables)
+		{
+			Var.AppendCompileHash(&Visitor);
+		}
+		return HashState.Finalize();
+	}
+
+	template<typename TStruct>
+	static FSHAHash GetStructHash()
+	{
+		static TOptional<FSHAHash> StructHash;
+		UE_CALL_ONCE([]()
+		{
+			StructHash.Emplace(NiagaraCBufferHelperImpl::GenerateHash(TStruct::GetVariables()));
+		});
+
+		return StructHash.GetValue();
+	}
+}
+
 const TArray<FNiagaraVariable>& FNiagaraGlobalParameters::GetVariables()
 {
 	static const FName NAME_NiagaraStructPadding0 = "Engine.PaddingInt32_0";
@@ -1760,6 +1787,11 @@ const TArray<FNiagaraVariable>& FNiagaraGlobalParameters::GetVariables()
 	return Variables;
 }
 
+FSHAHash FNiagaraGlobalParameters::GetStructHash()
+{
+	return NiagaraCBufferHelperImpl::GetStructHash<FNiagaraGlobalParameters>();
+}
+
 const TArray<FNiagaraVariable>& FNiagaraSystemParameters::GetVariables()
 {
 	static const FName NAME_NiagaraStructPadding0 = "Engine.System.PaddingInt32_0";
@@ -1789,6 +1821,11 @@ const TArray<FNiagaraVariable>& FNiagaraSystemParameters::GetVariables()
 	return Variables;
 }
 
+FSHAHash FNiagaraSystemParameters::GetStructHash()
+{
+	return NiagaraCBufferHelperImpl::GetStructHash<FNiagaraSystemParameters>();
+}
+
 const TArray<FNiagaraVariable>& FNiagaraOwnerParameters::GetVariables()
 {
 	static const FName NAME_NiagaraStructPadding0 = "Engine.Owner.PaddingInt32_0";
@@ -1826,6 +1863,11 @@ const TArray<FNiagaraVariable>& FNiagaraOwnerParameters::GetVariables()
 	return Variables;
 }
 
+FSHAHash FNiagaraOwnerParameters::GetStructHash()
+{
+	return NiagaraCBufferHelperImpl::GetStructHash<FNiagaraOwnerParameters>();
+}
+
 const TArray<FNiagaraVariable>& FNiagaraEmitterParameters::GetVariables()
 {
 	static const FName NAME_NiagaraStructPadding0 = "Engine.Emitter.PaddingInt32_0";
@@ -1847,6 +1889,11 @@ const TArray<FNiagaraVariable>& FNiagaraEmitterParameters::GetVariables()
 	return Variables;
 }
 
+FSHAHash FNiagaraEmitterParameters::GetStructHash()
+{
+	return NiagaraCBufferHelperImpl::GetStructHash<FNiagaraEmitterParameters>();
+}
+
 #endif
 
 void INiagaraModule::RegisterInternalAssetTagDefinitions()
diff --git a/Engine/Plugins/FX/Niagara/Source/Niagara/Public/NiagaraTypes.h b/Engine/Plugins/FX/Niagara/Source/Niagara/Public/NiagaraTypes.h
--- a/Engine/Plugins/FX/Niagara/Source/Niagara/Public/NiagaraTypes.h
+++ b/Engine/Plugins/FX/Niagara/Source/Niagara/Public/NiagaraTypes.h
@@ -1645,6 +1645,7 @@ struct alignas(16) FNiagaraGlobalParameters
 {
 #if WITH_EDITOR
 	NIAGARA_API static const TArray<FNiagaraVariable>& GetVariables();
+	NIAGARA_API static FSHAHash GetStructHash();
 #endif
 
 	float WorldDeltaTime = 0.0f;
@@ -1665,6 +1666,7 @@ struct alignas(16) FNiagaraSystemParameters
 {
 #if WITH_EDITOR
 	NIAGARA_API static const TArray<FNiagaraVariable>& GetVariables();
+	NIAGARA_API static FSHAHash GetStructHash();
 #endif
 
 	float EngineTimeSinceRendered = 0.0f;
@@ -1695,6 +1697,7 @@ struct alignas(16) FNiagaraOwnerParameters
 {
 #if WITH_EDITOR
 	NIAGARA_API static const TArray<FNiagaraVariable>& GetVariables();
+	NIAGARA_API static FSHAHash GetStructHash();
 #endif
 
 	FMatrix44f	EngineLocalToWorld = FMatrix44f::Identity;
@@ -1726,6 +1729,7 @@ struct alignas(16) FNiagaraEmitterParameters
 {
 #if WITH_EDITOR
 	NIAGARA_API static const TArray<FNiagaraVariable>& GetVariables();
+	NIAGARA_API static FSHAHash GetStructHash();
 #endif
 
 	int32 EmitterNumParticles = 0;
diff --git a/Engine/Plugins/FX/Niagara/Source/NiagaraEditor/Private/NiagaraScriptSource.cpp b/Engine/Plugins/FX/Niagara/Source/NiagaraEditor/Private/NiagaraScriptSource.cpp
--- a/Engine/Plugins/FX/Niagara/Source/NiagaraEditor/Private/NiagaraScriptSource.cpp
+++ b/Engine/Plugins/FX/Niagara/Source/NiagaraEditor/Private/NiagaraScriptSource.cpp
@@ -86,6 +86,20 @@ void UNiagaraScriptSource::ComputeVMCompilationId(FNiagaraVMExecutableDataId& Id
 	HashCollector.AddHash(FNiagaraCompileHash(Hash.Hash, sizeof(Hash.Hash) / sizeof(uint8)), TEXT("/Plugin/FX/Niagara/Private/NiagaraShaderVersion.ush"));
 	Hash = GetShaderFileHash((TEXT("/Engine/Public/ShaderVersion.ush")), EShaderPlatform::SP_PCD3D_SM5);
 	HashCollector.AddHash(FNiagaraCompileHash(Hash.Hash, sizeof(Hash.Hash) / sizeof(uint8)), TEXT("/Engine/Public/ShaderVersion.ush"));
+
+	// Incorporate the cbuffer definitions used by the script
+	FSHAHash StructHash = FNiagaraGlobalParameters::GetStructHash();
+	HashCollector.AddHash(FNiagaraCompileHash(StructHash.Hash, sizeof(StructHash.Hash) / sizeof(uint8)), TEXT("FNiagaraGlobalParameters"));
+
+	if (!UNiagaraScript::IsSystemScript(InUsage))
+	{
+		StructHash = FNiagaraSystemParameters::GetStructHash();
+		HashCollector.AddHash(FNiagaraCompileHash(StructHash.Hash, sizeof(StructHash.Hash) / sizeof(uint8)), TEXT("FNiagaraSystemParameters"));
+		StructHash = FNiagaraOwnerParameters::GetStructHash();
+		HashCollector.AddHash(FNiagaraCompileHash(StructHash.Hash, sizeof(StructHash.Hash) / sizeof(uint8)), TEXT("FNiagaraOwnerParameters"));
+		StructHash = FNiagaraEmitterParameters::GetStructHash();
+		HashCollector.AddHash(FNiagaraCompileHash(StructHash.Hash, sizeof(StructHash.Hash) / sizeof(uint8)), TEXT("FNiagaraEmitterParameters"));
+	}
 }
 
 void UNiagaraScriptSource::RefreshGraphCompileId()
