- Rebuild the widget map when a parameter is deleted

diff --git a/Engine/Plugins/FX/Niagara/Source/NiagaraEditor/Private/Customizations/UserParameters/NiagaraSystemUserParametersBuilder.cpp b/Engine/Plugins/FX/Niagara/Source/NiagaraEditor/Private/Customizations/UserParameters/NiagaraSystemUserParametersBuilder.cpp
--- a/Engine/Plugins/FX/Niagara/Source/NiagaraEditor/Private/Customizations/UserParameters/NiagaraSystemUserParametersBuilder.cpp
+++ b/Engine/Plugins/FX/Niagara/Source/NiagaraEditor/Private/Customizations/UserParameters/NiagaraSystemUserParametersBuilder.cpp
@@ -327,7 +327,7 @@ FReply FNiagaraSystemUserParameterBuilder::SummonHierarchyEditor() const
 	return FReply::Unhandled();
 }
 
-void FNiagaraSystemUserParameterBuilder::DeleteParameter(FNiagaraVariable UserParameter) const
+void FNiagaraSystemUserParameterBuilder::DeleteParameter(FNiagaraVariable UserParameter)
 {
 	FScopedTransaction Transaction(LOCTEXT("DeleteParameterTransaction", "Deleted user parameter"));
 	System->Modify();
@@ -339,6 +339,10 @@ void FNiagaraSystemUserParameterBuilder::DeleteParameter(FNiagaraVariable UserPa
 
 	Cast<UNiagaraSystemEditorData>(System->GetEditorData())->UserParameterHierarchy->Modify();
 	System->HandleVariableRemoved(UserParameter, true);
+
+	UserParamToWidgetMap.Empty();
+	
+	Rebuild();
 }
 
 TSharedRef<SWidget> FNiagaraSystemUserParameterBuilder::CreateUserParameterNameWidget(FNiagaraVariable UserParameter)
diff --git a/Engine/Plugins/FX/Niagara/Source/NiagaraEditor/Private/Customizations/UserParameters/NiagaraSystemUserParametersBuilder.h b/Engine/Plugins/FX/Niagara/Source/NiagaraEditor/Private/Customizations/UserParameters/NiagaraSystemUserParametersBuilder.h
--- a/Engine/Plugins/FX/Niagara/Source/NiagaraEditor/Private/Customizations/UserParameters/NiagaraSystemUserParametersBuilder.h
+++ b/Engine/Plugins/FX/Niagara/Source/NiagaraEditor/Private/Customizations/UserParameters/NiagaraSystemUserParametersBuilder.h
@@ -44,7 +44,7 @@ private:
 	void AddParameter(FNiagaraVariable NewParameter) const;
 	void OnParameterAdded(FNiagaraVariable UserParameter);
 	bool OnAllowMakeType(const FNiagaraTypeDefinition& InType) const;
-	void DeleteParameter(FNiagaraVariable UserParameter) const;
+	void DeleteParameter(FNiagaraVariable UserParameter);
 
 	virtual TSharedRef<SWidget> CreateUserParameterNameWidget(FNiagaraVariable UserParameter) override;
 	
