- Fix for stateless spawn burst going below 1 when spawn count scaling is enabled

#rnx
#tests Compare stateful & stateless spawn burst amounts

diff --git a/Engine/Plugins/FX/Niagara/Source/Niagara/Private/Stateless/NiagaraStatelessEmitterInstance.cpp b/Engine/Plugins/FX/Niagara/Source/Niagara/Private/Stateless/NiagaraStatelessEmitterInstance.cpp
--- a/Engine/Plugins/FX/Niagara/Source/Niagara/Private/Stateless/NiagaraStatelessEmitterInstance.cpp
+++ b/Engine/Plugins/FX/Niagara/Source/Niagara/Private/Stateless/NiagaraStatelessEmitterInstance.cpp
@@ -676,7 +676,7 @@ void FNiagaraStatelessEmitterInstance::InitSpawnInfosForLoop(float Initializatio
 		}
 
 		const int32 UnscaledSpawnAmount = EvaluateDistribution(SpawnInfo.Amount, RandomStream, RendererBindings, DefaultSpawnAmount);
-		const int32 SpawnAmount = FMath::FloorToInt(float(UnscaledSpawnAmount) * EmitterData->SpawnCountScale);
+		const int32 SpawnAmount = UnscaledSpawnAmount > 0 ? FMath::Max(FMath::FloorToInt(float(UnscaledSpawnAmount) * EmitterData->SpawnCountScale), 1) : 0;
 		if (SpawnAmount <= 0)
 		{
 			continue;
