- Add apply & save button to scratchpad editors

#rnx
#tests Modify a scratch, hit apply & save, relaunch to verify it was saved

diff --git a/Engine/Plugins/FX/Niagara/Source/NiagaraEditor/Private/ViewModels/NiagaraScratchPadScriptViewModel.cpp b/Engine/Plugins/FX/Niagara/Source/NiagaraEditor/Private/ViewModels/NiagaraScratchPadScriptViewModel.cpp
--- a/Engine/Plugins/FX/Niagara/Source/NiagaraEditor/Private/ViewModels/NiagaraScratchPadScriptViewModel.cpp
+++ b/Engine/Plugins/FX/Niagara/Source/NiagaraEditor/Private/ViewModels/NiagaraScratchPadScriptViewModel.cpp
@@ -12,6 +12,7 @@
 #include "ViewModels/NiagaraScriptGraphViewModel.h"
 #include "ViewModels/NiagaraSystemViewModel.h"
 
+#include "FileHelpers.h"
 #include "ScopedTransaction.h"
 #include "Framework/Commands/UICommandList.h"
 
@@ -261,6 +262,15 @@ void FNiagaraScratchPadScriptViewModel::ApplyChanges()
 	OnChangesAppliedDelegate.Broadcast(ScratchPadScriptViewModel);
 }
 
+void FNiagaraScratchPadScriptViewModel::ApplyChangesAndSave()
+{
+	ApplyChanges();
+
+	const bool bCheckDirty = false;
+	const bool bPromptToSave = false;
+	FEditorFileUtils::PromptForCheckoutAndSave({ OriginalScript->GetPackage() }, bCheckDirty, bPromptToSave);
+}
+
 void FNiagaraScratchPadScriptViewModel::DiscardChanges()
 {
 	OnRequestDiscardChangesDelegate.ExecuteIfBound();
diff --git a/Engine/Plugins/FX/Niagara/Source/NiagaraEditor/Private/Widgets/NiagaraScratchScriptEditor.cpp b/Engine/Plugins/FX/Niagara/Source/NiagaraEditor/Private/Widgets/NiagaraScratchScriptEditor.cpp
--- a/Engine/Plugins/FX/Niagara/Source/NiagaraEditor/Private/Widgets/NiagaraScratchScriptEditor.cpp
+++ b/Engine/Plugins/FX/Niagara/Source/NiagaraEditor/Private/Widgets/NiagaraScratchScriptEditor.cpp
@@ -55,6 +55,36 @@ void SNiagaraScratchPadScriptEditor::Construct(const FArguments& InArgs, TShared
 					]
 				]
 				+ SHorizontalBox::Slot()
+				.AutoWidth()
+				.Padding(2.0f, 3.0f, 5.0f, 3.0f)
+				[
+					SNew(SButton)
+					.OnClicked(this, &SNiagaraScratchPadScriptEditor::OnApplyAndSaveButtonClicked)
+					.ToolTipText(LOCTEXT("ApplyAndSaveButtonToolTip", "Apply the current changes to this script and save.  This will update the selection stack UI and compile neccessary scripts."))
+					.IsEnabled(this, &SNiagaraScratchPadScriptEditor::GetApplyButtonIsEnabled)
+					.ContentPadding(FMargin(3.0f, 0.0f))
+					.Content()
+					[
+						SNew(SHorizontalBox)
+						+ SHorizontalBox::Slot()
+						.AutoWidth()
+						.Padding(2.0f, 2.0f, 2.0f, 2.0f)
+						[
+							SNew(SImage)
+							.Image(FAppStyle::Get().GetBrush("AssetEditor.Apply"))
+						]
+						+ SHorizontalBox::Slot()
+						.AutoWidth()
+						.VAlign(VAlign_Center)
+						.Padding(2.0f, 2.0f, 2.0f, 3.0f)
+						[
+							SNew(STextBlock)
+							.TextStyle(FNiagaraEditorStyle::Get(), "NiagaraEditor.ScratchPad.EditorHeaderText")
+							.Text(LOCTEXT("ApplyAndSaveButtonLabel", "Apply & Save"))
+						]
+					]
+				]
+				+ SHorizontalBox::Slot()
 				.VAlign(VAlign_Center)
 				.Padding(5, 0, 2, 0)
 				.AutoWidth()
diff --git a/Engine/Plugins/FX/Niagara/Source/NiagaraEditor/Private/Widgets/NiagaraScratchScriptEditor.h b/Engine/Plugins/FX/Niagara/Source/NiagaraEditor/Private/Widgets/NiagaraScratchScriptEditor.h
--- a/Engine/Plugins/FX/Niagara/Source/NiagaraEditor/Private/Widgets/NiagaraScratchScriptEditor.h
+++ b/Engine/Plugins/FX/Niagara/Source/NiagaraEditor/Private/Widgets/NiagaraScratchScriptEditor.h
@@ -71,7 +71,12 @@ private:
 	FReply OnApplyButtonClicked()
 	{
 		ScriptViewModel.Pin()->ApplyChanges();
+		return FReply::Handled();
+	}
 
+	FReply OnApplyAndSaveButtonClicked()
+	{
+		ScriptViewModel.Pin()->ApplyChangesAndSave();
 		return FReply::Handled();
 	}
 
diff --git a/Engine/Plugins/FX/Niagara/Source/NiagaraEditor/Public/ViewModels/NiagaraScratchPadScriptViewModel.h b/Engine/Plugins/FX/Niagara/Source/NiagaraEditor/Public/ViewModels/NiagaraScratchPadScriptViewModel.h
--- a/Engine/Plugins/FX/Niagara/Source/NiagaraEditor/Public/ViewModels/NiagaraScratchPadScriptViewModel.h
+++ b/Engine/Plugins/FX/Niagara/Source/NiagaraEditor/Public/ViewModels/NiagaraScratchPadScriptViewModel.h
@@ -85,6 +85,7 @@ public:
 	NIAGARAEDITOR_API bool HasUnappliedChanges() const;
 
 	NIAGARAEDITOR_API void ApplyChanges();
+	NIAGARAEDITOR_API void ApplyChangesAndSave();
 
 	NIAGARAEDITOR_API void DiscardChanges();
 
