Reduce churn on type registry to only handle invalid types, which should only occur after a plugin is unloaded.  Cleans up TSan report while loading editor.
#rnx
#tests TSan report generated and Niagara's errors went from ~100 to 0

diff --git a/Engine/Plugins/FX/Niagara/Source/Niagara/Private/NiagaraTypeRegistry.cpp b/Engine/Plugins/FX/Niagara/Source/Niagara/Private/NiagaraTypeRegistry.cpp
--- a/Engine/Plugins/FX/Niagara/Source/Niagara/Private/NiagaraTypeRegistry.cpp
+++ b/Engine/Plugins/FX/Niagara/Source/Niagara/Private/NiagaraTypeRegistry.cpp
@@ -235,7 +235,11 @@ int32 FNiagaraTypeRegistry::RegisterTypeInternal(const FNiagaraTypeDefinition& N
 			UE::TScopeLock Lock(RegisteredTypesLock);
 			if (ensure(RegisteredTypes.IsValidIndex(NewType.RegisteredTypeDefIndex)))
 			{
-				*RegisteredTypes[NewType.RegisteredTypeDefIndex] = NewType;
+				TUniquePtr<FNiagaraTypeDefinition>& TypePtr = RegisteredTypes[NewType.RegisteredTypeDefIndex];
+				if (!TypePtr->IsValid())
+				{
+					*TypePtr = NewType;
+				}
 			}
 		}
 	}
