Move the RegisteredTypesLock over to a FTransactionallySafeMutex (from FSpinLock)
#rnx

diff --git a/Engine/Plugins/FX/Niagara/Source/Niagara/Public/NiagaraTypeRegistry.h b/Engine/Plugins/FX/Niagara/Source/Niagara/Public/NiagaraTypeRegistry.h
--- a/Engine/Plugins/FX/Niagara/Source/Niagara/Public/NiagaraTypeRegistry.h
+++ b/Engine/Plugins/FX/Niagara/Source/Niagara/Public/NiagaraTypeRegistry.h
@@ -5,7 +5,7 @@
 #include "CoreMinimal.h"
 #include "Misc/ScopeLock.h"
 #include "Misc/SecureHash.h"
-#include "Misc/SpinLock.h"
+#include "Misc/TransactionallySafeMutex.h"
 #include "UObject/GCObject.h"
 #include "UObject/UnrealType.h"
 #include "Containers/Queue.h"
@@ -118,7 +118,7 @@ private:
 
 	// The spinlock is used for accessing the RegisteredTypes array in performance critical areas (resolve TypeDefinitionHandles being the
 	// most common).
-	mutable UE::FSpinLock RegisteredTypesLock; 
+	mutable FTransactionallySafeMutex RegisteredTypesLock; 
 	FRegisteredTypesArray RegisteredTypes;
 
 	// Covers manipulation of all of the containers populated during registration (except for RegisteredTypes which is handled
