- Handle nullptr for EventStack
- There is a path where the output node (for example) that this can be true

#rnx
#tests Loaded content that would previously crash

diff --git a/Engine/Plugins/FX/Niagara/Source/NiagaraEditor/Private/NiagaraScriptMergeManager.cpp b/Engine/Plugins/FX/Niagara/Source/NiagaraEditor/Private/NiagaraScriptMergeManager.cpp
--- a/Engine/Plugins/FX/Niagara/Source/NiagaraEditor/Private/NiagaraScriptMergeManager.cpp
+++ b/Engine/Plugins/FX/Niagara/Source/NiagaraEditor/Private/NiagaraScriptMergeManager.cpp
@@ -1095,7 +1095,10 @@ void FNiagaraEmitterMergeAdapter::GatherFunctionCallNodes(TArray<UNiagaraNodeFun
 	ParticleUpdateStack->GatherFunctionCallNodes(OutFunctionCallNodes);
 	for (TSharedRef<FNiagaraEventHandlerMergeAdapter> EventHandler : EventHandlers)
 	{
-		EventHandler->GetEventStack()->GatherFunctionCallNodes(OutFunctionCallNodes);
+		if (TSharedPtr<FNiagaraScriptStackMergeAdapter> EventStack = EventHandler->GetEventStack())
+		{
+			EventStack->GatherFunctionCallNodes(OutFunctionCallNodes);
+		}
 	}
 	for (TSharedRef<FNiagaraSimulationStageMergeAdapter> SimulationStage : SimulationStages)
 	{
