add memory tracking
2d dispatch for 2d RT

#virtualized

diff --git a/Engine/Plugins/FX/Niagara/Source/Niagara/Classes/NiagaraDataInterfaceRenderTarget2D.h b/Engine/Plugins/FX/Niagara/Source/Niagara/Classes/NiagaraDataInterfaceRenderTarget2D.h
--- a/Engine/Plugins/FX/Niagara/Source/Niagara/Classes/NiagaraDataInterfaceRenderTarget2D.h
+++ b/Engine/Plugins/FX/Niagara/Source/Niagara/Classes/NiagaraDataInterfaceRenderTarget2D.h
@@ -101,6 +101,11 @@ public:
 	NIAGARA_API virtual void BuildShaderParameters(FNiagaraShaderParametersBuilder& ShaderParametersBuilder) const override;
 	NIAGARA_API virtual void SetShaderParameters(const FNiagaraDataInterfaceSetShaderParametersContext& Context) const override;
 
+#if WITH_EDITOR
+	virtual ENiagaraGpuDispatchType GetGpuDispatchType() const override { return ENiagaraGpuDispatchType::TwoD; }
+	virtual FIntVector GetGpuDispatchNumThreads() const { return FIntVector(8, 8, 1); }
+#endif
+
 	virtual void ProvidePerInstanceDataForRenderThread(void* DataForRenderThread, void* PerInstanceData, const FNiagaraSystemInstanceID& SystemInstance) override {}
 	NIAGARA_API virtual bool InitPerInstanceData(void* PerInstanceData, FNiagaraSystemInstance* SystemInstance) override;
 	NIAGARA_API virtual void DestroyPerInstanceData(void* PerInstanceData, FNiagaraSystemInstance* SystemInstance) override;
diff --git a/Engine/Plugins/FX/Niagara/Source/Niagara/Private/DataInterface/NiagaraDataInterfaceDynamicMesh.cpp b/Engine/Plugins/FX/Niagara/Source/Niagara/Private/DataInterface/NiagaraDataInterfaceDynamicMesh.cpp
--- a/Engine/Plugins/FX/Niagara/Source/Niagara/Private/DataInterface/NiagaraDataInterfaceDynamicMesh.cpp
+++ b/Engine/Plugins/FX/Niagara/Source/Niagara/Private/DataInterface/NiagaraDataInterfaceDynamicMesh.cpp
@@ -7,6 +7,7 @@
 #include "NiagaraSimStageData.h"
 #include "NiagaraShaderParametersBuilder.h"
 #include "NiagaraSystemInstance.h"
+#include "NiagaraStats.h"
 
 #include "Materials/MaterialRenderProxy.h"
 #include "Materials/Material.h"
@@ -527,6 +528,33 @@ namespace NDIDynamicMeshLocal
 			TexCoordOffset = INDEX_NONE;
 			ColorOffset = INDEX_NONE;
 		}
+
+#if STATS
+		void UpdateMemoryStats()
+		{
+			DEC_MEMORY_STAT_BY(STAT_NiagaraGPUDataInterfaceMemory, MemorySize);
+			MemorySize = 0;
+						
+			if (IndexBuffer.GetRHI())
+			{
+				MemorySize += IndexBuffer.GetRHI()->GetSize();
+			}
+			
+			if (VertexBuffer.GetRHI())
+			{
+				MemorySize += VertexBuffer.GetRHI()->GetSize();
+			}
+
+			INC_MEMORY_STAT_BY(STAT_NiagaraGPUDataInterfaceMemory, MemorySize);
+		}
+		
+		void ResetMemoryStats()
+		{
+			DEC_MEMORY_STAT_BY(STAT_NiagaraGPUDataInterfaceMemory, MemorySize);
+			MemorySize = 0;
+		}
+		uint64 MemorySize = 0;
+#endif
 	};
 
 	struct FNDIProxy : public FNiagaraDataInterfaceProxy
@@ -1099,6 +1127,13 @@ void UNiagaraDataInterfaceDynamicMesh::DestroyPerInstanceData(void* PerInstanceD
 	(
 		[RT_Proxy=GetProxyAs<FNDIProxy>(), InstanceID=SystemInstance->GetId()](FRHICommandListImmediate& CmdList)
 		{
+#if STATS
+			if (FNDIInstanceData_RenderThread* InstanceData = RT_Proxy->InstanceData_RT.Find(InstanceID))
+			{
+				InstanceData->ResetMemoryStats();
+			}
+#endif
+			
 			RT_Proxy->InstanceData_RT.Remove(InstanceID);
 		}
 	);
@@ -1151,6 +1186,10 @@ bool UNiagaraDataInterfaceDynamicMesh::PerInstanceTickPostSimulate(void* PerInst
 			{
 				FNDIInstanceData_RenderThread& InstanceData_RT = RT_Proxy->InstanceData_RT.FindOrAdd(RT_SystemInstanceID);
 				InstanceData_RT.UpdateData(RHICmdList, RT_GameToRenderData);
+
+#if STATS
+				InstanceData_RT.UpdateMemoryStats();
+#endif
 			}
 		);
 	}
