Use RegisteredTypesLock in GetDefaultDataInterfaceByName since the function is iterating over Registry.RegisteredTypes

diff --git a/Engine/Plugins/FX/Niagara/Source/Niagara/Private/NiagaraTypeRegistry.cpp b/Engine/Plugins/FX/Niagara/Source/Niagara/Private/NiagaraTypeRegistry.cpp
--- a/Engine/Plugins/FX/Niagara/Source/Niagara/Private/NiagaraTypeRegistry.cpp
+++ b/Engine/Plugins/FX/Niagara/Source/Niagara/Private/NiagaraTypeRegistry.cpp
@@ -43,7 +43,7 @@ UNiagaraDataInterfaceBase* FNiagaraTypeRegistry::GetDefaultDataInterfaceByName(c
 	UClass* DIClass = nullptr;
 	{
 		FNiagaraTypeRegistry& Registry = Get();
-		FReadScopeLock Lock(Registry.RegistrationLock);
+		UE::TScopeLock Lock(Registry.RegisteredTypesLock);
 		for (const TUniquePtr<FNiagaraTypeDefinition>& TypePtr : Registry.RegisteredTypes)
 		{
 			if (TypePtr.IsValid() && TypePtr->IsDataInterface())
