Fix for Spawn module node not having guid, leading to a warning on load
#tests Ran spawn wizard to verify that created nodes don't generate a warning on load
#rnx

diff --git a/Engine/Plugins/FX/Niagara/Source/NiagaraEditor/Private/Widgets/DataChannel/NiagaraDataChannelWizard.cpp b/Engine/Plugins/FX/Niagara/Source/NiagaraEditor/Private/Widgets/DataChannel/NiagaraDataChannelWizard.cpp
--- a/Engine/Plugins/FX/Niagara/Source/NiagaraEditor/Private/Widgets/DataChannel/NiagaraDataChannelWizard.cpp
+++ b/Engine/Plugins/FX/Niagara/Source/NiagaraEditor/Private/Widgets/DataChannel/NiagaraDataChannelWizard.cpp
@@ -1110,6 +1110,8 @@ TSharedRef<FModuleWizardModel> DataChannel::CreateSpawnNDCModuleWizardModel()
 					ExecIndexNode->OpName = FName("Util::ExecIndex");
 					ExecIndexNode->SetFlags(RF_Transactional);
 					Graph->AddNode(ExecIndexNode, false, false);
+					ExecIndexNode->CreateNewGuid();
+					ExecIndexNode->PostPlacedNewNode();
 					ExecIndexNode->AllocateDefaultPins();
 					if (UEdGraphPin* EmitterIDInput = SpawnDataFunction->FindPin(FName("Spawned Particle Exec Index"), EGPD_Input))
 					{
