- Fix for race between VM & external access to the shared array

#rnx

diff --git a/Engine/Plugins/FX/Niagara/Source/Niagara/Classes/NiagaraDataInterfaceArrayImpl.h b/Engine/Plugins/FX/Niagara/Source/Niagara/Classes/NiagaraDataInterfaceArrayImpl.h
--- a/Engine/Plugins/FX/Niagara/Source/Niagara/Classes/NiagaraDataInterfaceArrayImpl.h
+++ b/Engine/Plugins/FX/Niagara/Source/Niagara/Classes/NiagaraDataInterfaceArrayImpl.h
@@ -371,32 +371,19 @@ struct FNDIArrayProxyImpl : public INDIArrayProxyBase
 
 	struct FWriteArrayRef
 	{
-		FWriteArrayRef(TOwnerType* Owner, FNDIArrayInstanceData_GameThread<TArrayType>* InstanceData, bool bFromBP=false)
+		FWriteArrayRef(TOwnerType* Owner, FNDIArrayInstanceData_GameThread<TArrayType>* InstanceData)
 		{
 			if (InstanceData)
 			{
 				LockObject = &InstanceData->ArrayRWGuard;
 				LockObject->WriteLock();
 
-				if (bFromBP)//Writes from BP are for user parameters and we should discard modifications and write to the actual DI.
-				{
-					if (InstanceData->bIsModified)
-					{
-						Owner->GetArrayReference() = InstanceData->ArrayData;
-					}
-					InstanceData->bIsModified = false;
-					InstanceData->ArrayData.Empty();
-					ArrayData = &Owner->GetArrayReference();
- 				}
- 				else
+ 				if (InstanceData->bIsModified == false)
  				{
- 					if (InstanceData->bIsModified == false)
- 					{
- 						InstanceData->bIsModified = true;
- 						InstanceData->ArrayData = Owner->GetArrayReference();
- 					}
-					ArrayData = &InstanceData->ArrayData;
+ 					InstanceData->bIsModified = true;
+ 					InstanceData->ArrayData = Owner->GetArrayReference();
  				}
+				ArrayData = &InstanceData->ArrayData;
 			}
 			else
 			{
@@ -556,6 +543,32 @@ struct FNDIArrayProxyImpl : public INDIArrayProxyBase
 
 	//////////////////////////////////////////////////////////////////////////
 	// BP user parameter accessors, should remove if we every start to share the object between instances
+	void BeginSetArrayFromBP(bool bCopyFromInstanceData)
+	{
+		for (auto It = PerInstanceData_GameThread.CreateConstIterator(); It; ++It)
+		{
+			FNDIArrayInstanceData_GameThread<TArrayType>* InstanceData = It.Value();
+			InstanceData->ArrayRWGuard.WriteLock();
+			if (InstanceData->bIsModified && bCopyFromInstanceData)
+			{
+				Owner->GetArrayReference() = InstanceData->ArrayData;
+			}
+			bCopyFromInstanceData = false;
+			InstanceData->bIsModified = false;
+			InstanceData->bIsRenderDirty |= bShouldSyncToGpu;
+			InstanceData->ArrayData.Empty();
+		}
+	}
+
+	void EndSetArrayFromBP()
+	{
+		for (auto It = PerInstanceData_GameThread.CreateConstIterator(); It; ++It)
+		{
+			FNDIArrayInstanceData_GameThread<TArrayType>* InstanceData = It.Value();
+			InstanceData->ArrayRWGuard.WriteUnlock();
+		}
+	}
+
 	template<typename TFromArrayType>
 	void SetArrayData(TConstArrayView<TFromArrayType> InArrayData)
 	{
@@ -570,17 +583,12 @@ struct FNDIArrayProxyImpl : public INDIArrayProxyBase
 		#if WITH_EDITORONLY_DATA
 			ensure(GIsTransacting || PerInstanceData_GameThread.Num() == 1);
 		#endif
+			BeginSetArrayFromBP(false);
 
-			for (auto It=PerInstanceData_GameThread.CreateConstIterator(); It; ++It)
-			{
-				FNDIArrayInstanceData_GameThread<TArrayType>* InstanceData = It.Value();
-				FWriteScopeLock	ScopeLock(InstanceData->ArrayRWGuard);
-				InstanceData->bIsModified = false;
-				InstanceData->bIsRenderDirty |= bShouldSyncToGpu;
-				InstanceData->ArrayData.Empty();
-			}
 			Owner->GetArrayReference().SetNum(InArrayData.Num());
 			FNDIArrayImplHelper<TArrayType>::CopyCpuToCpuMemory(Owner->GetArrayReference().GetData(), InArrayData.GetData(), InArrayData.Num());
+
+			EndSetArrayFromBP();
 		}
 	}
 
@@ -600,21 +608,21 @@ struct FNDIArrayProxyImpl : public INDIArrayProxyBase
 	void SetArrayValue(int Index, const TFromArrayType& Value, bool bSizeToFit)
 	{
 		ensure(PerInstanceData_GameThread.Num() <= 1);
-		FNDIArrayInstanceData_GameThread<TArrayType>* InstanceData = PerInstanceData_GameThread.IsEmpty() ? nullptr : PerInstanceData_GameThread.CreateConstIterator().Value();
-		FWriteArrayRef ArrayRef(Owner, InstanceData, true);
-		if (!ArrayRef.GetArray().IsValidIndex(Index))
+
+		BeginSetArrayFromBP(true);
+
+		TArray<TArrayType>& ArrayRef = Owner->GetArrayReference();
+		if (!ArrayRef.IsValidIndex(Index))
 		{
 			if (!bSizeToFit)
 			{
 				return;
 			}
-			ArrayRef.GetArray().AddDefaulted(Index + 1 - ArrayRef.GetArray().Num());
+			ArrayRef.AddDefaulted(Index + 1 - ArrayRef.Num());
 		}
-		if (InstanceData)
-		{
-			InstanceData->bIsRenderDirty |= bShouldSyncToGpu;
-		}
-		FNDIArrayImplHelper<TArrayType>::CopyCpuToCpuMemory(ArrayRef.GetArray().GetData() + Index, &Value, 1);
+		FNDIArrayImplHelper<TArrayType>::CopyCpuToCpuMemory(ArrayRef.GetData() + Index, &Value, 1);
+
+		EndSetArrayFromBP();
 	}
 
 	template<typename TToArrayType>
