- Ensure we mark the render state dirty post GC if we have any UObject parameters
- These parameters could be reference as meshes / materials or used in material bindings and can therefore become invalid (i.e. nullptr) potentially leading to dangling data on the render thread

diff --git a/Engine/Plugins/FX/Niagara/Source/Niagara/Private/NiagaraWorldManager.cpp b/Engine/Plugins/FX/Niagara/Source/Niagara/Private/NiagaraWorldManager.cpp
--- a/Engine/Plugins/FX/Niagara/Source/Niagara/Private/NiagaraWorldManager.cpp
+++ b/Engine/Plugins/FX/Niagara/Source/Niagara/Private/NiagaraWorldManager.cpp
@@ -825,9 +825,15 @@ void FNiagaraWorldManager::PreGarbageCollect()
 
 void FNiagaraWorldManager::PostReachabilityAnalysis()
 {
+	// Post garbage collection we must mark parameter stores with UObjects dirty to ensure any objects which are nullptr are updated in the various parameter stores.
+	// We must also mark the component render data as dirty to ensure any material references in renderers are updated as they could also have become nullptr.
 	for (TObjectIterator<UNiagaraComponent> ComponentIt; ComponentIt; ++ComponentIt)
 	{
 		ComponentIt->GetOverrideParameters().MarkUObjectsDirty();
+		if (ComponentIt->GetOverrideParameters().GetUObjects().Num() > 0)
+		{
+			ComponentIt->MarkRenderStateDirty();
+		}
 	}
 }
 
