- Validate CPU execution path when falling back due to count allocation failure

#rnx

diff --git a/Engine/Plugins/FX/Niagara/Source/Niagara/Private/Stateless/NiagaraStatelessComputeManager.cpp b/Engine/Plugins/FX/Niagara/Source/Niagara/Private/Stateless/NiagaraStatelessComputeManager.cpp
--- a/Engine/Plugins/FX/Niagara/Source/Niagara/Private/Stateless/NiagaraStatelessComputeManager.cpp
+++ b/Engine/Plugins/FX/Niagara/Source/Niagara/Private/Stateless/NiagaraStatelessComputeManager.cpp
@@ -334,7 +334,9 @@ FNiagaraDataBuffer* FNiagaraStatelessComputeManager::GetDataBuffer(FRHICommandLi
 			// If it does occur this means some other system has used a count slot but not reserved one
 			else
 			{
-				if (!GenerateCPUDataForGPUSim(RHICmdList, EmitterInstance, CacheData->DataBuffer))
+				const bool bAllowCPUExec = EnumHasAnyFlags(EmitterInstance->EmitterData->FeatureMask, ENiagaraStatelessFeatureMask::ExecuteCPU);
+				const bool bDidGenerate = bAllowCPUExec && GenerateCPUDataForGPUSim(RHICmdList, EmitterInstance, CacheData->DataBuffer);
+				if (!bDidGenerate)
 				{
 					// Note: We can not free the data when in paralell GDME mode as multiple lock / unlock operations could result in the wrong one providing the final data
 					// The plus side to adding back into the cache data is that multiple calls to get the same emitter's data will resolve to no active particles.
