- Add button to Save SVT to separate asset in the Sim Cache viewer

diff --git a/Engine/Plugins/FX/Niagara/Source/NiagaraEditor/Private/Customizations/SimCache/NiagaraRenderTargetVolumeSimCacheVisualizer.cpp b/Engine/Plugins/FX/Niagara/Source/NiagaraEditor/Private/Customizations/SimCache/NiagaraRenderTargetVolumeSimCacheVisualizer.cpp
--- a/Engine/Plugins/FX/Niagara/Source/NiagaraEditor/Private/Customizations/SimCache/NiagaraRenderTargetVolumeSimCacheVisualizer.cpp
+++ b/Engine/Plugins/FX/Niagara/Source/NiagaraEditor/Private/Customizations/SimCache/NiagaraRenderTargetVolumeSimCacheVisualizer.cpp
@@ -22,9 +22,12 @@
 #include "Materials/Material.h"
 #include "SceneView.h"
 
+#include "Widgets/Input/SButton.h"
 #include "Widgets/Input/SNumericEntryBox.h"
 #include "Widgets/Layout/SGridPanel.h"
 #include "Widgets/SBoxPanel.h"
+#include "AssetToolsModule.h"
+#include "IAssetTools.h"
 #include "SEnumCombo.h"
 
 #define LOCTEXT_NAMESPACE "NiagaraVolumeTextureViewport"
@@ -73,85 +76,121 @@ private:
 
 TSharedPtr<SWidget> FNiagaraRenderTargetVolumeSimCacheVisualizer::CreateWidgetFor(const UObject* CachedData, TSharedPtr<FNiagaraSimCacheViewModel> ViewModel)
 {
-	if (const UAnimatedSparseVolumeTexture* VolumeTexture = Cast<const UAnimatedSparseVolumeTexture>(CachedData))
+	const UAnimatedSparseVolumeTexture* SparseVolumeTexture = Cast<const UAnimatedSparseVolumeTexture>(CachedData);
+	if (SparseVolumeTexture == nullptr)
 	{
-		static UEnum* MaskEnum = StaticEnum<ENiagaraRenderTargetVolumeVisualizerMask>();
-		TSharedPtr<SNiagaraVolumeTextureViewport> Viewport = SNew(SNiagaraVolumeTextureViewport, ViewModel, VolumeTexture);
+		return TSharedPtr<SWidget>();
+	}
+
+	static UEnum* MaskEnum = StaticEnum<ENiagaraRenderTargetVolumeVisualizerMask>();
+	TSharedPtr<SNiagaraVolumeTextureViewport> Viewport = SNew(SNiagaraVolumeTextureViewport, ViewModel, SparseVolumeTexture);
 		
-		return SNew(SVerticalBox)
+	return SNew(SVerticalBox)
+	+SVerticalBox::Slot()
+	.Padding(15)
+	.AutoHeight()
+	[
+		SNew(SVerticalBox)
+		+SVerticalBox::Slot()
+		.AutoHeight()
+		[
+			SNew(STextBlock)
+			.Text(LOCTEXT("MaterialAttributes", "Material Attributes"))
+		]
 		+SVerticalBox::Slot()
-		.Padding(15)
+		.Padding(10)
 		.AutoHeight()
 		[
-			SNew(SVerticalBox)
-			+SVerticalBox::Slot()
-			.AutoHeight()
+			SNew(SGridPanel)
+			+SGridPanel::Slot(0, 0)
+			.Padding(5, 0)
+			.VAlign(VAlign_Center)
 			[
 				SNew(STextBlock)
-				.Text(LOCTEXT("MaterialAttributes", "Material Attributes"))
+				.Text(LOCTEXT("TemperatureMask_Label", "Temperature Mask"))
 			]
-			+SVerticalBox::Slot()
-			.Padding(10)
-			.AutoHeight()
+			+SGridPanel::Slot(1, 0)
+			.Padding(5, 0)
 			[
-				SNew(SGridPanel)
-				+SGridPanel::Slot(0, 0)
-				.Padding(5, 0)
-				.VAlign(VAlign_Center)
-				[
-					SNew(STextBlock)
-					.Text(LOCTEXT("TemperatureMask_Label", "Temperature Mask"))
-				]
-				+SGridPanel::Slot(1, 0)
-				.Padding(5, 0)
-				[
-					SNew(SEnumComboBox, MaskEnum)
-					.ContentPadding(FMargin(2, 0))
-					.Font(FAppStyle::Get().GetFontStyle("PropertyWindow.NormalFont"))
-					.CurrentValue_Lambda([Viewport]()
-					{
-						return static_cast<int32>(Viewport->TemparatureMask);
-					})
-					.OnEnumSelectionChanged_Lambda([Viewport](int32 NewValue, ESelectInfo::Type)
-					{
-						ENiagaraRenderTargetVolumeVisualizerMask Value = static_cast<ENiagaraRenderTargetVolumeVisualizerMask>(NewValue);
-						Viewport->SetTemperatureMask(Value);
-					})
-				]
-				+SGridPanel::Slot(0, 1)
-				.Padding(5, 0)
-				.VAlign(VAlign_Center)
-				[
-					SNew(STextBlock)
-					.Text(LOCTEXT("Density_Label", "Density Scale"))
-				]
-				+SGridPanel::Slot(1, 1)
-				.Padding(5, 0)
-				[
-					SNew(SNumericEntryBox<float>)
-					.AllowSpin(true)
-					.MinDesiredValueWidth(80)
-					.MinSliderValue(0.0f)
-					.MaxSliderValue(1.0f)
-					.Delta(0.01f)
-					.Value_Lambda([Viewport]()
-					{
-						return Viewport->DensityScale;
-					})
-					.OnValueChanged_Lambda([Viewport](float Val)
-					{
-						Viewport->SetDensityScale(Val);
-					})
-				]
+				SNew(SEnumComboBox, MaskEnum)
+				.ContentPadding(FMargin(2, 0))
+				.Font(FAppStyle::Get().GetFontStyle("PropertyWindow.NormalFont"))
+				.CurrentValue_Lambda([Viewport]()
+				{
+					return static_cast<int32>(Viewport->TemparatureMask);
+				})
+				.OnEnumSelectionChanged_Lambda([Viewport](int32 NewValue, ESelectInfo::Type)
+				{
+					ENiagaraRenderTargetVolumeVisualizerMask Value = static_cast<ENiagaraRenderTargetVolumeVisualizerMask>(NewValue);
+					Viewport->SetTemperatureMask(Value);
+				})
+			]
+			+SGridPanel::Slot(0, 1)
+			.Padding(5, 0)
+			.VAlign(VAlign_Center)
+			[
+				SNew(STextBlock)
+				.Text(LOCTEXT("Density_Label", "Density Scale"))
+			]
+			+SGridPanel::Slot(1, 1)
+			.Padding(5, 0)
+			[
+				SNew(SNumericEntryBox<float>)
+				.AllowSpin(true)
+				.MinDesiredValueWidth(80)
+				.MinSliderValue(0.0f)
+				.MaxSliderValue(1.0f)
+				.Delta(0.01f)
+				.Value_Lambda([Viewport]()
+				{
+					return Viewport->DensityScale;
+				})
+				.OnValueChanged_Lambda([Viewport](float Val)
+				{
+					Viewport->SetDensityScale(Val);
+				})
 			]
 		]
 		+SVerticalBox::Slot()
 		.AutoHeight()
 		[
-			Viewport.ToSharedRef()
-		];
-	}
-	return TSharedPtr<SWidget>();
+			SNew(SHorizontalBox)
+			+SHorizontalBox::Slot()
+			.AutoWidth()
+			[
+				SNew(SButton)
+				.Text(LOCTEXT("SaveSVT", "Save SVT"))
+				.OnClicked_Lambda(
+					[WeakSparseVolumeTexture =MakeWeakObjectPtr(SparseVolumeTexture)]()
+					{
+						if (const UObject* TextureObject = WeakSparseVolumeTexture.Get())
+						{
+							FString PackagePath = TEXT("/Game/SimCacheSVT");
+							UPackage* Package = TextureObject->GetPackage();
+							if (Package && Package != GetTransientPackage())
+							{
+								PackagePath = Package->GetPathName() / Package->GetName() + TEXT("_SVT");
+							}
+
+							IAssetTools& AssetTools = FModuleManager::Get().LoadModuleChecked<FAssetToolsModule>("AssetTools").Get();
+							AssetTools.DuplicateAssetWithDialogAndTitle(
+								FString(FPathViews::GetCleanFilename(PackagePath)),
+								FString(FPathViews::GetPath(PackagePath)),
+								const_cast<UObject*>(TextureObject),
+								LOCTEXT("SaveSVTAs", "Save SVT as...")
+							);
+						}
+						return FReply::Handled();
+					}
+				)
+			]
+		]
+	]
+	+SVerticalBox::Slot()
+	.AutoHeight()
+	[
+		Viewport.ToSharedRef()
+	];
 }
 
 //////////////////////////////////////////////////////////////////////////
