- Currently selected buffer is now displayed in the header
- Filter state is now saved per cache buffer (i.e. system / emitter)
- Filters are restored and selection is kept when doing consecuitive captures in the attribute spreadsheet
- Components with children (i.e. Color -> R/G/B/A) will show undetermined check state if not all components are the same filter state

diff --git a/Engine/Plugins/FX/Niagara/Source/NiagaraEditor/Private/ViewModels/NiagaraSimCacheViewModel.cpp b/Engine/Plugins/FX/Niagara/Source/NiagaraEditor/Private/ViewModels/NiagaraSimCacheViewModel.cpp
--- a/Engine/Plugins/FX/Niagara/Source/NiagaraEditor/Private/ViewModels/NiagaraSimCacheViewModel.cpp
+++ b/Engine/Plugins/FX/Niagara/Source/NiagaraEditor/Private/ViewModels/NiagaraSimCacheViewModel.cpp
@@ -12,6 +12,21 @@
 
 #define LOCTEXT_NAMESPACE "NiagaraSimCacheViewModel"
 
+namespace NiagaraSimCacheViewModelPrivate
+{
+	TOptional<bool> FindComponentFilteredState(FName ComponentName, TArray<FNiagaraSimCacheViewModel::FComponentInfo> Components)
+	{
+		for (const FNiagaraSimCacheViewModel::FComponentInfo& Component : Components)
+		{
+			if (Component.Name == ComponentName)
+			{
+				return Component.bIsFiltered;
+			}
+		}
+		return {};
+	}
+}
+
 FNiagaraSimCacheViewModel::FNiagaraSimCacheViewModel()
 {
 }
@@ -45,14 +60,6 @@ void FNiagaraSimCacheViewModel::Initialize(TWeakObjectPtr<UNiagaraSimCache> InSi
 	OnViewDataChangedDelegate.Broadcast(true);
 }
 
-void FNiagaraSimCacheViewModel::SetComponentFilters(const TArray<FString>& NewComponentFilterArray)
-{
-	bComponentFilterActive = true;
-	ComponentFilterArray.Empty();
-	ComponentFilterArray.Append(NewComponentFilterArray);
-	OnViewDataChangedDelegate.Broadcast(true);
-}
-
 void FNiagaraSimCacheViewModel::SetupPreviewComponentAndInstance()
 {
 	UNiagaraSystem* System = SimCache ? SimCache->GetSystem(true) : nullptr;
@@ -74,32 +81,25 @@ void FNiagaraSimCacheViewModel::SetupPreviewComponentAndInstance()
 	}
 }
 
-TConstArrayView<FNiagaraSimCacheViewModel::FComponentInfo> FNiagaraSimCacheViewModel::GetSelectedComponentInfos() const
+TArrayView<FNiagaraSimCacheViewModel::FComponentInfo> FNiagaraSimCacheViewModel::GetMutableSelectedComponentInfos()
 {
 	if (SelectionMode == ESelectionMode::SystemInstance)
 	{
-		return GetComponentInfos(INDEX_NONE);
+		return SystemComponentInfos;
 	}
 	else if (SelectionMode == ESelectionMode::Emitter)
 	{
-		return GetComponentInfos(SelectedEmitterIndex);
+		if (TArray<FComponentInfo>* ComponentInfos = EmitterComponentInfos.Find(SelectedEmitterName))
+		{
+			return *ComponentInfos;
+		}
 	}
 	return MakeArrayView<FComponentInfo>(nullptr, 0);
 }
 
-TConstArrayView<FNiagaraSimCacheViewModel::FComponentInfo> FNiagaraSimCacheViewModel::GetComponentInfos(int32 InEmitterIndex) const
+TConstArrayView<FNiagaraSimCacheViewModel::FComponentInfo> FNiagaraSimCacheViewModel::GetSelectedComponentInfos() const
 {
-	if( InEmitterIndex == INDEX_NONE)
-	{
-		return SystemComponentInfos;
-	}
-
-	if(EmitterComponentInfos.IsValidIndex(InEmitterIndex))
-	{
-		return EmitterComponentInfos[InEmitterIndex];
-	}
-	
-	return TConstArrayView<FComponentInfo>();
+	return const_cast<FNiagaraSimCacheViewModel*>(this)->GetMutableSelectedComponentInfos();
 }
 
 FText FNiagaraSimCacheViewModel::GetComponentText(const FName ComponentName, const int32 InstanceIndex) const
@@ -167,6 +167,40 @@ void FNiagaraSimCacheViewModel::SetFrameIndex(const int32 InFrameIndex)
 	OnViewDataChangedDelegate.Broadcast(false);
 }
 
+FText FNiagaraSimCacheViewModel::GetSelectedText() const
+{
+	if (SimCache && SimCache->IsCacheValid())
+	{
+		switch (SelectionMode)
+		{
+			case ESelectionMode::SystemInstance:
+				return LOCTEXT("SystemInstance", "System Instance");
+
+			case ESelectionMode::Emitter:
+				return FText::FromName(SelectedEmitterName);
+
+			case ESelectionMode::DataInterface:
+				if ( SelectedDataInterface.IsValid() )
+				{
+					// Data interfaces can have some long names, i.e. SetVariablesxxx.Emitter.DIName, so let's limit to something sensible
+					constexpr int32 MaxStringLength = 32;
+					FString DataInterfaceName = SelectedDataInterface.GetName().ToString();
+					if (DataInterfaceName.Len() > MaxStringLength)
+					{
+						DataInterfaceName.RightChopInline(DataInterfaceName.Len() - MaxStringLength, EAllowShrinking::No);
+						DataInterfaceName.InsertAt(0, TEXT("..."));
+					}
+					return FText::FromString(DataInterfaceName);
+				}
+				break;
+
+			case ESelectionMode::DebugData:
+				return LOCTEXT("DebugData", "Debug Data");
+		}
+	}
+	return FText();
+}
+
 const UObject* FNiagaraSimCacheViewModel::GetSelectedDataInterfaceStorage() const
 {
 	if (SelectionMode == ESelectionMode::DataInterface && SimCache)
@@ -179,16 +213,16 @@ const UObject* FNiagaraSimCacheViewModel::GetSelectedDataInterfaceStorage() cons
 void FNiagaraSimCacheViewModel::SetSelectedSystemInstance()
 {
 	SelectionMode			= ESelectionMode::SystemInstance;
-	SelectedEmitterIndex	= INDEX_NONE;
+	SelectedEmitterName		= NAME_None;
 	SelectedDataInterface	= FNiagaraVariableBase();
 
 	RefreshFromSelectionChanged();
 }
 
-void FNiagaraSimCacheViewModel::SetSelectedEmitter(int32 EmitterIndex)
+void FNiagaraSimCacheViewModel::SetSelectedEmitter(FName EmitterName)
 {
 	SelectionMode			= ESelectionMode::Emitter;
-	SelectedEmitterIndex	= EmitterIndex;
+	SelectedEmitterName		= EmitterName;
 	SelectedDataInterface	= FNiagaraVariableBase();
 
 	RefreshFromSelectionChanged();
@@ -197,7 +231,7 @@ void FNiagaraSimCacheViewModel::SetSelectedEmitter(int32 EmitterIndex)
 void FNiagaraSimCacheViewModel::SetSelectedDataInterface(FNiagaraVariableBase DIVariable)
 {
 	SelectionMode			= ESelectionMode::DataInterface;
-	SelectedEmitterIndex	= INDEX_NONE;
+	SelectedEmitterName		= NAME_None;
 	SelectedDataInterface	= DIVariable;
 
 	RefreshFromSelectionChanged();
@@ -206,17 +240,111 @@ void FNiagaraSimCacheViewModel::SetSelectedDataInterface(FNiagaraVariableBase DI
 void FNiagaraSimCacheViewModel::SetSelectedDebugData()
 {
 	SelectionMode			= ESelectionMode::DebugData;
-	SelectedEmitterIndex	= INDEX_NONE;
+	SelectedEmitterName		= NAME_None;
 	SelectedDataInterface	= FNiagaraVariableBase();
 
 	RefreshFromSelectionChanged();
 }
 
+bool FNiagaraSimCacheViewModel::IsComponentFilterActive() const
+{
+	return SelectionMode == ESelectionMode::SystemInstance || SelectionMode == ESelectionMode::Emitter;
+}
+
+bool FNiagaraSimCacheViewModel::IsComponentFiltered(FName ComponentName) const
+{
+	for (const FComponentInfo& ComponentInfo : GetSelectedComponentInfos())
+	{
+		if (ComponentInfo.Name == ComponentName)
+		{
+			return ComponentInfo.bIsFiltered;
+		}
+	}
+	return false;
+}
+
+bool FNiagaraSimCacheViewModel::IsComponentFiltered(const FString& ComponentNameString) const
+{
+	return IsComponentFiltered(FName(*ComponentNameString));
+}
+
+void FNiagaraSimCacheViewModel::SetComponentFiltered(const FString& ComponentNameString, bool bFiltered)
+{
+	const FName ComponentName(*ComponentNameString);
+	for (FComponentInfo& ComponentInfo : GetMutableSelectedComponentInfos())
+	{
+		if (ComponentInfo.Name == ComponentName)
+		{
+			ComponentInfo.bIsFiltered = bFiltered;
+			OnViewDataChangedDelegate.Broadcast(true);
+			return;
+		}
+	}
+}
+
+void FNiagaraSimCacheViewModel::ToggleComponentFiltered(const FString& ComponentNameString)
+{
+	const FName ComponentName(*ComponentNameString);
+	for (FComponentInfo& ComponentInfo : GetMutableSelectedComponentInfos())
+	{
+		if (ComponentInfo.Name == ComponentName)
+		{
+			ComponentInfo.bIsFiltered = !ComponentInfo.bIsFiltered;
+			OnViewDataChangedDelegate.Broadcast(true);
+			return;
+		}
+	}
+}
+
+void FNiagaraSimCacheViewModel::SetAllComponentFiltered(bool bFiltered)
+{
+	for (FComponentInfo& ComponentInfo : GetMutableSelectedComponentInfos())
+	{
+		ComponentInfo.bIsFiltered = bFiltered;
+	}
+	OnViewDataChangedDelegate.Broadcast(true);
+}
+
+void FNiagaraSimCacheViewModel::RefreshSelection()
+{
+	bool bSelectionValid = false;
+	if (SimCache && SimCache->IsCacheValid())
+	{
+		switch (SelectionMode)
+		{
+			case ESelectionMode::SystemInstance:
+				bSelectionValid = true;
+				break;
+
+			case ESelectionMode::Emitter:
+				bSelectionValid = SimCache->GetEmitterIndex(SelectedEmitterName) != INDEX_NONE;
+				break;
+
+			case ESelectionMode::DataInterface:
+				bSelectionValid = SimCache->GetDataInterfaceStorageObject(SelectedDataInterface) != nullptr;
+				break;
+
+			case ESelectionMode::DebugData:
+				bSelectionValid = SimCache->GetDebugData() != nullptr;
+				break;
+		}
+	}
+
+	if (!bSelectionValid)
+	{
+		SetSelectedSystemInstance();
+	}
+
+	UpdateCachedFrame();
+	UpdateCurrentEntries();
+	OnBufferChangedDelegate.Broadcast();
+	OnViewDataChangedDelegate.Broadcast(true);
+}
+
 void FNiagaraSimCacheViewModel::RefreshFromSelectionChanged()
 {
 	UpdateCachedFrame();
 	UpdateCurrentEntries();
-	bComponentFilterActive = false;
 	OnBufferChangedDelegate.Broadcast();
 	OnViewDataChangedDelegate.Broadcast(true);
 }
@@ -298,12 +426,15 @@ void FNiagaraSimCacheViewModel::UpdateCachedFrame()
 			break;
 
 		case ESelectionMode::Emitter:
-			if (SelectedEmitterIndex >= 0 && SelectedEmitterIndex < SimCache->GetNumEmitters())
+		{
+			const int32 FoundEmitterIndex = SimCache->GetEmitterIndex(SelectedEmitterName);
+			if (FoundEmitterIndex != INDEX_NONE)
 			{
-				EmitterIndex = SelectedEmitterIndex;
-				NumInstances = SimCache->GetEmitterNumInstances(SelectedEmitterIndex, FrameIndex);
+				EmitterIndex = FoundEmitterIndex;
+				NumInstances = SimCache->GetEmitterNumInstances(FoundEmitterIndex, FrameIndex);
 			}
 			break;
+		}
 
 		case ESelectionMode::DataInterface:
 		case ESelectionMode::DebugData:
@@ -331,12 +462,18 @@ void FNiagaraSimCacheViewModel::UpdateCachedFrame()
 
 void FNiagaraSimCacheViewModel::UpdateComponentInfos()
 {
+	// Save filter state
+	TArray<FComponentInfo>					TempSystemComponentInfos;
+	TMap<FName, TArray<FComponentInfo>>		TempEmitterComponentInfos;
+	Swap(TempSystemComponentInfos, SystemComponentInfos);
+	Swap(TempEmitterComponentInfos, EmitterComponentInfos);
+
+	// Clear data
 	SystemComponentInfos.Empty();
 	EmitterComponentInfos.Empty();
 	FoundFloatComponents = 0;
 	FoundHalfComponents = 0;
 	FoundInt32Components = 0;
-	
 	if (SimCache == nullptr)
 	{
 		return;
@@ -364,16 +501,17 @@ void FNiagaraSimCacheViewModel::UpdateComponentInfos()
 		}
 	);
 
-	for (int32 i = 0; i < SimCache->GetNumEmitters(); ++i)
+	for (int32 iEmitter=0; iEmitter < SimCache->GetNumEmitters(); ++iEmitter)
 	{
-		TArray<FComponentInfo>& CurrentComponentInfos = EmitterComponentInfos.AddDefaulted_GetRef();
+		TArray<FComponentInfo>& CurrentComponentInfos = EmitterComponentInfos.Add(SimCache->GetEmitterName(iEmitter));
 
 		FoundFloatComponents = 0;
 		FoundHalfComponents = 0;
 		FoundInt32Components = 0;
 		
-		SimCache->ForEachEmitterAttribute(i,
-		[&](const FNiagaraSimCacheVariable& Variable)
+		SimCache->ForEachEmitterAttribute(
+			iEmitter,
+			[&](const FNiagaraSimCacheVariable& Variable)
 			{
 				// Build component info
 				const FNiagaraTypeDefinition& TypeDef = Variable.Variable.GetType();
@@ -394,17 +532,60 @@ void FNiagaraSimCacheViewModel::UpdateComponentInfos()
 			}
 		);
 	}
+
+	// Restore filter state
+	for (FComponentInfo& ComponentInfo : SystemComponentInfos)
+	{
+		ComponentInfo.bIsFiltered = NiagaraSimCacheViewModelPrivate::FindComponentFilteredState(ComponentInfo.Name, TempSystemComponentInfos).Get(true);
+	}
+
+	for (auto EmitterIt=EmitterComponentInfos.CreateIterator(); EmitterIt; ++EmitterIt)
+	{
+		TArray<FComponentInfo>* TempEmitterComponents = TempEmitterComponentInfos.Find(EmitterIt->Key);
+		if (TempEmitterComponents == nullptr)
+		{
+			continue;
+		}
+
+		for (FComponentInfo& ComponentInfo : EmitterIt->Value)
+		{
+			ComponentInfo.bIsFiltered = NiagaraSimCacheViewModelPrivate::FindComponentFilteredState(ComponentInfo.Name, *TempEmitterComponents).Get(true);
+		}
+	}
 }
 
 void FNiagaraSimCacheViewModel::BuildTreeItemChildren(TSharedPtr<FNiagaraSimCacheTreeItem> InTreeItem, TWeakPtr<SNiagaraSimCacheTreeView> OwningTreeView)
 {
 	FNiagaraSimCacheTreeItem* TreeItem = InTreeItem.Get();
 
-	if(TreeItem && SimCache)
+	if (!TreeItem || !SimCache)
 	{
-		int32 BufferIndex = TreeItem->GetBufferIndex();
-		
-		SimCache->ForEachEmitterAttribute(BufferIndex,
+		return;
+	}
+
+	TOptional<int32> CacheEmitterIndex;
+	switch (InTreeItem->GetType())
+	{
+		case ENiagaraSimCacheOverviewItemType::System:
+			CacheEmitterIndex = INDEX_NONE;
+			break;
+
+		case ENiagaraSimCacheOverviewItemType::Emitter:
+		{
+			const FName EmitterName = TreeItem->GetEmitterName();
+			const int32 EmitterIndex = SimCache->GetEmitterIndex(EmitterName);
+			if (EmitterIndex != INDEX_NONE)
+			{
+				CacheEmitterIndex = EmitterIndex;
+			}
+			break;
+		}
+	}
+	
+	if (CacheEmitterIndex.IsSet())
+	{
+		SimCache->ForEachEmitterAttribute(
+			CacheEmitterIndex.GetValue(),
 			[&](const FNiagaraSimCacheVariable& Variable)
 			{
 				FNiagaraTypeDefinition TypeDef = Variable.Variable.GetType();
@@ -414,7 +595,7 @@ void FNiagaraSimCacheViewModel::BuildTreeItemChildren(TSharedPtr<FNiagaraSimCach
 				CurrentItem->SetDisplayName(FText::FromName(Variable.Variable.GetName()));
 				CurrentItem->SetFilterName(Variable.Variable.GetName().ToString());
 				CurrentItem->TypeDef = TypeDef;
-				CurrentItem->BufferIndex = TreeItem->GetBufferIndex();
+				CurrentItem->EmitterName = TreeItem->GetEmitterName();
 				//CurrentItem->RootItem = TreeItem;
 
 				TreeItem->AddChild(CurrentItem);
@@ -426,7 +607,6 @@ void FNiagaraSimCacheViewModel::BuildTreeItemChildren(TSharedPtr<FNiagaraSimCach
 				return true;
 			}
 		);
-		
 	}
 }
 
@@ -444,7 +624,7 @@ void FNiagaraSimCacheViewModel::RecursiveBuildTreeItemChildren(FNiagaraSimCacheT
 		
 		CurrentItem->SetDisplayName(FText::FromString(PropertyName));
 		CurrentItem->SetFilterName(Parent->GetFilterName().Append(".").Append(PropertyName));
-		CurrentItem->SetBufferIndex(Root->GetBufferIndex());
+		CurrentItem->SetEmitterName(Root->GetEmitterName());
 		//CurrentItem->RootItem = Root;
 
 		Parent->AddChild(CurrentItem);
@@ -509,12 +689,13 @@ void FNiagaraSimCacheViewModel::BuildEntries(TWeakPtr<SNiagaraSimCacheTreeView>
 	{
 		const TSharedRef<FNiagaraSimCacheEmitterTreeItem> CurrentEmitterItem = MakeShared<FNiagaraSimCacheEmitterTreeItem>(OwningTreeView);
 		const TSharedRef<FNiagaraSimCacheOverviewEmitterItem> CurrentEmitterBufferItem = MakeShared<FNiagaraSimCacheOverviewEmitterItem>();
+		const FName EmitterName = GetEmitterLayoutName(i);
 		
-		CurrentEmitterItem->SetDisplayName(FText::FromName(GetEmitterLayoutName(i)));
-		CurrentEmitterBufferItem->SetDisplayName(FText::FromName(GetEmitterLayoutName(i)));
+		CurrentEmitterItem->SetDisplayName(FText::FromName(EmitterName));
+		CurrentEmitterBufferItem->SetDisplayName(FText::FromName(EmitterName));
 		
-		CurrentEmitterItem->SetBufferIndex(i);
-		CurrentEmitterBufferItem->SetBufferIndex(i);
+		CurrentEmitterItem->SetEmitterName(EmitterName);
+		CurrentEmitterBufferItem->SetEmitterName(EmitterName);
 
 		RootEntries.Add(CurrentEmitterItem);
 		BufferEntries.Add(CurrentEmitterBufferItem);
@@ -533,16 +714,6 @@ void FNiagaraSimCacheViewModel::BuildEntries(TWeakPtr<SNiagaraSimCacheTreeView>
 			CurrentDataInterfaceBufferItem->SetDisplayName(FText::FromName(Var.GetName()));
 			CurrentDataInterfaceBufferItem->DataInterfaceReference = Var;
 
-			for(int32 i = 0; i < GetNumEmitterLayouts(); i++)
-			{
-				if (Var.GetName().ToString().StartsWith(GetEmitterLayoutName(i).ToString() + "."))
-				{
-					CurrentDataInterfaceItem->SetBufferIndex(i);
-					CurrentDataInterfaceBufferItem->SetBufferIndex(i);
-					break;
-				}
-			}
-		
 			RootEntries.Add(CurrentDataInterfaceItem);
 			BufferEntries.Add(CurrentDataInterfaceBufferItem);
 		}
@@ -571,12 +742,23 @@ void FNiagaraSimCacheViewModel::UpdateCurrentEntries()
 		case ESelectionMode::SystemInstance:
 			SelectedRootEntries.Add(RootEntries[0]);
 			return;
+
 		case ESelectionMode::Emitter:
-			if (RootEntries.IsValidIndex(SelectedEmitterIndex + 1))
+		{
+			TSharedRef<FNiagaraSimCacheTreeItem>* EmitterItem = RootEntries.FindByPredicate(
+				[EmitterName=SelectedEmitterName](const TSharedRef<FNiagaraSimCacheTreeItem>& TreeItem)
+				{
+					return (TreeItem->GetType() == ENiagaraSimCacheOverviewItemType::Emitter) && (TreeItem->GetEmitterName() == EmitterName);
+				}
+			);
+
+			if (EmitterItem)
 			{
-				SelectedRootEntries.Add(RootEntries[SelectedEmitterIndex + 1]);
+				SelectedRootEntries.Add(*EmitterItem);
 			}
 			return;
+		}
+
 		default:
 			//-TODO: DO we need to do something here?
 			break;
diff --git a/Engine/Plugins/FX/Niagara/Source/NiagaraEditor/Private/Widgets/SNiagaraDebugCaptureView.cpp b/Engine/Plugins/FX/Niagara/Source/NiagaraEditor/Private/Widgets/SNiagaraDebugCaptureView.cpp
--- a/Engine/Plugins/FX/Niagara/Source/NiagaraEditor/Private/Widgets/SNiagaraDebugCaptureView.cpp
+++ b/Engine/Plugins/FX/Niagara/Source/NiagaraEditor/Private/Widgets/SNiagaraDebugCaptureView.cpp
@@ -228,7 +228,7 @@ void SNiagaraDebugCaptureView::Construct(const FArguments& InArgs, const TShared
 		DebugCaptureToolbarBuilder.AddComboButton(
 			FUIAction(),
 			FOnGetContent::CreateRaw(this, &SNiagaraDebugCaptureView::GenerateFilterMenuContent),
-			FText::GetEmpty(),
+			TAttribute<FText>::Create(TAttribute<FText>::FGetter::CreateSP(SimCacheViewModel.ToSharedRef(), &FNiagaraSimCacheViewModel::GetSelectedText)),
 			LOCTEXT("FilterCombo_Tooltip", "Change filter and selection"),
 			FSlateIcon(FAppStyle::GetAppStyleSetName(),"Icons.Filter")
 		);
@@ -389,27 +389,7 @@ void SNiagaraDebugCaptureView::OnCaptureComplete(UNiagaraSimCache* CapturedSimCa
 	bIsCaptureActive = false;
 	SystemViewModel.Get()->GetSequencer().Get()->Pause();
 				
-	TArray<FGuid> SelectedEmitterHandleIds = SystemViewModel->GetSelectionViewModel()->GetSelectedEmitterHandleIds();
-
-	if(SelectedEmitterHandleIds.Num())
-	{
-		TSharedPtr<FNiagaraEmitterHandleViewModel> SelectedEmitterHandleViewModel = SystemViewModel->GetEmitterHandleViewModelById(SelectedEmitterHandleIds[0]);
-
-		if(SelectedEmitterHandleViewModel.IsValid())
-		{
-			for(int32 i = 0; i < SimCacheViewModel->GetNumEmitterLayouts(); ++i)
-			{
-				if(SimCacheViewModel->GetEmitterLayoutName(i) == SelectedEmitterHandleViewModel->GetName())
-				{
-					if(SimCacheViewModel->GetSelectedEmitter() != i)
-					{
-						SimCacheViewModel->SetSelectedEmitter(i);
-					}
-					break;
-				}
-			}
-		}
-	}
+	SimCacheViewModel->RefreshSelection();
 
 	switch(FrameMode)
 	{
diff --git a/Engine/Plugins/FX/Niagara/Source/NiagaraEditor/Private/Widgets/SNiagaraSimCacheOverview.cpp b/Engine/Plugins/FX/Niagara/Source/NiagaraEditor/Private/Widgets/SNiagaraSimCacheOverview.cpp
--- a/Engine/Plugins/FX/Niagara/Source/NiagaraEditor/Private/Widgets/SNiagaraSimCacheOverview.cpp
+++ b/Engine/Plugins/FX/Niagara/Source/NiagaraEditor/Private/Widgets/SNiagaraSimCacheOverview.cpp
@@ -39,6 +39,12 @@ void SNiagaraSimCacheOverview::OnSimCacheChanged()
 	{
 		BufferListView->RebuildList();
 	}
+	RebuildMainWidget();
+}
+
+void SNiagaraSimCacheOverview::OnViewDataChanged(bool)
+{
+	RebuildMainWidget();
 }
 
 void SNiagaraSimCacheOverview::Construct(const FArguments& InArgs)
@@ -55,14 +61,37 @@ void SNiagaraSimCacheOverview::Construct(const FArguments& InArgs)
 	.SelectionMode(ESelectionMode::Single);
 
 	ViewModel.Get()->OnSimCacheChanged().AddSP(this, &SNiagaraSimCacheOverview::OnSimCacheChanged);
-
-	const float MinSplitterSlotSize = 30.0f; 
+	ViewModel.Get()->OnViewDataChanged().AddSP(this, &SNiagaraSimCacheOverview::OnViewDataChanged);
 
 	ChildSlot
 	[
-		SNew(SSplitter)
+		SAssignNew(MainWidget, SSplitter)
 		.Orientation(Orient_Vertical)
-		+SSplitter::Slot()
+	];
+
+	RebuildMainWidget();
+}
+
+void SNiagaraSimCacheOverview::RebuildMainWidget()
+{
+	// Skip rebuild if we both things active
+	const bool bNeedsComponentFilter = ViewModel->IsComponentFilterActive() && ViewModel->IsCacheValid();
+	const int32 CurrentNumChildren = MainWidget->GetChildren()->Num();
+	const int32 RequiredNumChildren = bNeedsComponentFilter ? 2 : 1;
+	if (CurrentNumChildren == RequiredNumChildren)
+	{
+		return;
+	}
+
+	//-OPT: We can avoid a total rebuild here
+	while (MainWidget->GetChildren()->Num() > 0)
+	{
+		MainWidget->RemoveAt(0);
+	}
+
+	const float MinSplitterSlotSize = 30.0f;
+
+	MainWidget->AddSlot()
 		.Value(0.2f)
 		.MinSize(MinSplitterSlotSize)
 		[
@@ -89,36 +118,39 @@ void SNiagaraSimCacheOverview::Construct(const FArguments& InArgs)
 				// List View
 				BufferListView.ToSharedRef()
 			]
-		]
-		+SSplitter::Slot()
-		.Value(0.8f)
-		.MinSize(MinSplitterSlotSize)
-		[
-			// Component Details
-			SNew(SSplitter)
-			.Orientation(Orient_Vertical)
-			+SSplitter::Slot()
-			.Resizable(false)
-			.SizeRule(SSplitter::SizeToContent)
+		];
+
+	if (bNeedsComponentFilter)
+	{
+		MainWidget->AddSlot()
+			.Value(0.8f)
+			.MinSize(MinSplitterSlotSize)
 			[
-				// Header
-				SNew(SBorder)
-				.BorderImage(FAppStyle::GetNoBrush())
-				.Padding(5.0f)
+				// Component Details
+				SNew(SSplitter)
+				.Orientation(Orient_Vertical)
+				+SSplitter::Slot()
+				.Resizable(false)
+				.SizeRule(SSplitter::SizeToContent)
 				[
-					SNew(STextBlock)
-					.Text(LOCTEXT("ComponentTreeView", "Component Details"))
+					// Header
+					SNew(SBorder)
+					.BorderImage(FAppStyle::GetNoBrush())
+					.Padding(5.0f)
+					[
+						SNew(STextBlock)
+						.Text(LOCTEXT("ComponentTreeView", "Component Details"))
+					]
+				]
+				+SSplitter::Slot()
+				.Resizable(false)
+				[
+					// Tree View
+					TreeViewWidget.ToSharedRef()
 				]
-			]
-			+SSplitter::Slot()
-			.Resizable(false)
-			[
-				// Tree View
-				TreeViewWidget.ToSharedRef()
-			]
 			
-		]
-	];
+			];
+	}
 }
 
 TSharedRef<ITableRow> SNiagaraSimCacheOverview::OnGenerateRowForItem(TSharedRef<FNiagaraSimCacheOverviewItem> Item, const TSharedRef<STableViewBase>& Owner)
@@ -155,7 +187,7 @@ void SNiagaraSimCacheOverview::OnListSelectionChanged(TSharedPtr<FNiagaraSimCach
 				break;
 
 			case ENiagaraSimCacheOverviewItemType::Emitter:
-				ViewModel->SetSelectedEmitter(Item->GetBufferIndex());
+				ViewModel->SetSelectedEmitter(Item->GetEmitterName());
 				break;
 
 			case ENiagaraSimCacheOverviewItemType::DataInterface:
diff --git a/Engine/Plugins/FX/Niagara/Source/NiagaraEditor/Private/Widgets/SNiagaraSimCacheOverview.h b/Engine/Plugins/FX/Niagara/Source/NiagaraEditor/Private/Widgets/SNiagaraSimCacheOverview.h
--- a/Engine/Plugins/FX/Niagara/Source/NiagaraEditor/Private/Widgets/SNiagaraSimCacheOverview.h
+++ b/Engine/Plugins/FX/Niagara/Source/NiagaraEditor/Private/Widgets/SNiagaraSimCacheOverview.h
@@ -12,6 +12,7 @@
 
 class FNiagaraSimCacheViewModel;
 class SNiagaraSimCacheTreeView;
+class SSplitter;
 
 enum class ENiagaraSimCacheOverviewItemType
 {
@@ -31,8 +32,8 @@ struct FNiagaraSimCacheOverviewItem : public TSharedFromThis<FNiagaraSimCacheOve
 	FText GetDisplayNameText() { return DisplayName; }
 	void SetDisplayName(const FText& NewName) { DisplayName = NewName; }
 
-	virtual int32 GetBufferIndex() { return BufferIndex; }
-	void SetBufferIndex(int32 NewIndex) { BufferIndex = NewIndex; }
+	virtual FName GetEmitterName() { return EmitterName; }
+	void SetEmitterName(FName InEmitterName) { EmitterName = InEmitterName; }
 
 	virtual FNiagaraVariableBase GetDataInterface() { return FNiagaraVariableBase(); }
 
@@ -40,7 +41,7 @@ struct FNiagaraSimCacheOverviewItem : public TSharedFromThis<FNiagaraSimCacheOve
 
 	virtual TSharedRef<SWidget> GetRowWidget() { return SNew(STextBlock).Text(GetDisplayNameText()); }
 
-	int32 BufferIndex = INDEX_NONE;
+	FName EmitterName;
 	FText DisplayName;
 };
 
@@ -75,11 +76,15 @@ public:
 	SLATE_END_ARGS()
 
 	void OnSimCacheChanged();
+	void OnViewDataChanged(bool);
 	void Construct(const FArguments& InArgs);
 
+	void RebuildMainWidget();
 	TSharedRef<ITableRow> OnGenerateRowForItem(TSharedRef<FNiagaraSimCacheOverviewItem> Item, const TSharedRef<STableViewBase>& Owner);
 	void OnListSelectionChanged(TSharedPtr<FNiagaraSimCacheOverviewItem> Item, ESelectInfo::Type SelectInfo);
 
+	TSharedPtr<SSplitter> MainWidget;
+
 	TSharedPtr<SListView<TSharedRef<FNiagaraSimCacheOverviewItem>>> BufferListView;
 
 	TSharedPtr<SNiagaraSimCacheTreeView> TreeViewWidget;
diff --git a/Engine/Plugins/FX/Niagara/Source/NiagaraEditor/Private/Widgets/SNiagaraSimCacheTreeView.cpp b/Engine/Plugins/FX/Niagara/Source/NiagaraEditor/Private/Widgets/SNiagaraSimCacheTreeView.cpp
--- a/Engine/Plugins/FX/Niagara/Source/NiagaraEditor/Private/Widgets/SNiagaraSimCacheTreeView.cpp
+++ b/Engine/Plugins/FX/Niagara/Source/NiagaraEditor/Private/Widgets/SNiagaraSimCacheTreeView.cpp
@@ -130,34 +130,16 @@ void SSimCacheTreeViewVisibilityWidget::OnCheckStateChanged(ECheckBoxState NewSt
 }
 
 ECheckBoxState SSimCacheTreeViewVisibilityWidget::GetCheckedState() const
-{
-	return IsInFilter()? ECheckBoxState::Checked : ECheckBoxState::Unchecked;
-}
-
-bool SSimCacheTreeViewVisibilityWidget::IsInFilter() const
 {
 	TSharedPtr<SNiagaraSimCacheTreeView> TreeView = WeakTreeView.Pin();
 	TSharedPtr<FNiagaraSimCacheTreeItem> TreeItem = WeakTreeItem.Pin();
 
-	if(TreeView.IsValid() && TreeItem.IsValid())
+	if (TreeView.IsValid() && TreeItem.IsValid())
 	{
-		return TreeView->IsItemInFilter(TreeItem);
+		return TreeView->GetFilterCheckedState(TreeItem);
 	}
 
-	return false;
-}
-
-bool SSimCacheTreeViewVisibilityWidget::IsItemSelected() const
-{
-	const TSharedPtr<SNiagaraSimCacheTreeView> TreeView = WeakTreeView.Pin();
-	const TSharedPtr<FNiagaraSimCacheTreeItem> TreeItem = WeakTreeItem.Pin();
-
-	if(TreeView.IsValid() && TreeItem.IsValid())
-	{
-		return TreeView->IsItemSelected(TreeItem.ToSharedRef());
-	}
-
-	return false;
+	return ECheckBoxState::Unchecked;
 }
 
 ///// Tree View Widget
@@ -231,9 +213,6 @@ void SNiagaraSimCacheTreeView::OnGetChildren(TSharedRef<FNiagaraSimCacheTreeItem
 void SNiagaraSimCacheTreeView::OnBufferChanged()
 {
 	TreeView->RequestTreeRefresh();
-	SelectionForFilter.Empty();
-	ViewModel->SetComponentFilters(TArray<FString>());
-	SelectAll();
 
 	TArray<TSharedRef<FNiagaraSimCacheTreeItem>>* RootEntries = ViewModel->GetSelectedRootEntries();
 	if (RootEntries->Num() > 0)
@@ -248,57 +227,45 @@ void SNiagaraSimCacheTreeView::OnSimCacheChanged()
 	SetupRootEntries();
 }
 
-void SNiagaraSimCacheTreeView::RecursiveAddToSelectionFilter(TArray<TSharedRef<FNiagaraSimCacheTreeItem>>& ArrayToAdd)
+ECheckBoxState SNiagaraSimCacheTreeView::GetFilterCheckedState(TSharedPtr<FNiagaraSimCacheTreeItem> InItem) const
 {
-	for(const TSharedRef<FNiagaraSimCacheTreeItem>& Item : ArrayToAdd)
+	// Leaf node we get the state directly
+	if (InItem->Children.Num() == 0)
 	{
-		if(Item->Children.Num() > 0)
-		{
-			RecursiveAddToSelectionFilter(Item->Children);
-		}
+		return ViewModel->IsComponentFiltered(InItem->GetFilterName()) ? ECheckBoxState::Checked : ECheckBoxState::Unchecked;
 	}
-	SelectionForFilter.Append(ArrayToAdd);
-}
 
-void SNiagaraSimCacheTreeView::RecursiveRemoveFromSelectionFilter(TArray<TSharedRef<FNiagaraSimCacheTreeItem>>& ArrayToRemove)
-{
-	for(TSharedRef<FNiagaraSimCacheTreeItem> Item : ArrayToRemove)
+	// Recurse children to figure out the state
+	TOptional<ECheckBoxState> IsChecked;
+	for (const TSharedRef<FNiagaraSimCacheTreeItem>& ChildItem : InItem->Children)
 	{
-		if(Item->Children.Num() > 0)
+		const ECheckBoxState ChildCheckState = GetFilterCheckedState(ChildItem);
+		if (!IsChecked.IsSet())
 		{
-			RecursiveRemoveFromSelectionFilter(Item->Children);
+			IsChecked = ChildCheckState;
+		}
+		else if (ChildCheckState != IsChecked.GetValue())
+		{
+			return ECheckBoxState::Undetermined;
 		}
-		SelectionForFilter.Remove(Item);
 	}
+	return IsChecked.Get(ECheckBoxState::Unchecked);
 }
 
-void SNiagaraSimCacheTreeView::UpdateSelectionFilter(TSharedRef<FNiagaraSimCacheTreeItem> ClickedItem)
+void SNiagaraSimCacheTreeView::VisibilityButtonClicked(TSharedRef<FNiagaraSimCacheTreeItem> InItem)
 {
-	TArray<TSharedRef<FNiagaraSimCacheTreeItem>> SelectedItems = TreeView->GetSelectedItems();
-	SelectedItems.AddUnique(ClickedItem);
-
-	SelectedItems.RemoveAll([](const TSharedRef<FNiagaraSimCacheTreeItem> Item){ return Item->GetType() != ENiagaraSimCacheOverviewItemType::Component; });
-
-	if(IsItemInFilter(ClickedItem))
+	// Leaf node, apply the toggle directly
+	if (InItem->Children.Num() == 0)
 	{
-		RecursiveRemoveFromSelectionFilter(SelectedItems);
+		ViewModel->ToggleComponentFiltered(InItem->GetFilterName());
+		return;
 	}
-	else
+
+	// Toggle children
+	for (const TSharedRef<FNiagaraSimCacheTreeItem>& ChildItem : InItem->Children)
 	{
-		RecursiveAddToSelectionFilter(SelectedItems);
+		VisibilityButtonClicked(ChildItem);
 	}
-
-	UpdateStringFilters();
-}
-
-bool SNiagaraSimCacheTreeView::IsItemInFilter(TSharedPtr<FNiagaraSimCacheTreeItem> InItem) const
-{
-	return ViewModel->GetComponentFilters().Contains(InItem->GetFilterName());
-}
-
-void SNiagaraSimCacheTreeView::VisibilityButtonClicked(TSharedRef<FNiagaraSimCacheTreeItem> InItem)
-{
-	UpdateSelectionFilter(InItem);
 }
 
 bool SNiagaraSimCacheTreeView::IsItemSelected(TSharedRef<FNiagaraSimCacheTreeItem> InItem)
@@ -308,22 +275,7 @@ bool SNiagaraSimCacheTreeView::IsItemSelected(TSharedRef<FNiagaraSimCacheTreeIte
 
 void SNiagaraSimCacheTreeView::ClearFilterSelection()
 {
-	SelectionForFilter.Empty();
-	ViewModel->SetComponentFilters(TArray<FString>());
-}
-
-void SNiagaraSimCacheTreeView::UpdateStringFilters()
-{
-	TArray<FString> StringFilters;
-	for(TSharedRef<FNiagaraSimCacheTreeItem> FilterItem : SelectionForFilter)
-	{
-		if(FilterItem->GetType() == ENiagaraSimCacheOverviewItemType::Component)
-		{
-			StringFilters.AddUnique(FilterItem->GetFilterName());
-		}
-	}
-
-	ViewModel->SetComponentFilters(StringFilters);
+	ViewModel->SetAllComponentFiltered(false);
 }
 
 void SNiagaraSimCacheTreeView::SelectAll()
@@ -332,10 +284,8 @@ void SNiagaraSimCacheTreeView::SelectAll()
 	{
 		return;
 	}
-	
-	RecursiveAddToSelectionFilter(*ViewModel->GetSelectedRootEntries());
 
-	UpdateStringFilters();
+	ViewModel->SetAllComponentFiltered(true);
 }
 
 bool SNiagaraSimCacheTreeView::ShouldShowComponentView() const
diff --git a/Engine/Plugins/FX/Niagara/Source/NiagaraEditor/Private/Widgets/SNiagaraSimCacheTreeView.h b/Engine/Plugins/FX/Niagara/Source/NiagaraEditor/Private/Widgets/SNiagaraSimCacheTreeView.h
--- a/Engine/Plugins/FX/Niagara/Source/NiagaraEditor/Private/Widgets/SNiagaraSimCacheTreeView.h
+++ b/Engine/Plugins/FX/Niagara/Source/NiagaraEditor/Private/Widgets/SNiagaraSimCacheTreeView.h
@@ -49,8 +49,6 @@ protected:
 	void OnCheckStateChanged(ECheckBoxState NewState);
 	ECheckBoxState GetCheckedState() const;
 	
-	bool IsInFilter() const;
-	bool IsItemSelected() const;
 private:
 	TWeakPtr<FNiagaraSimCacheTreeItem> WeakTreeItem;
 	TWeakPtr<SNiagaraSimCacheTreeView> WeakTreeView;
@@ -66,12 +64,10 @@ SLATE_END_ARGS()
 	void OnSimCacheChanged();
 	void SetupRootEntries();
 	void Construct(const FArguments& InArgs);
-	bool HasSelectionFilter() const {return !SelectionForFilter.IsEmpty();}
-	bool IsItemInFilter(TSharedPtr<FNiagaraSimCacheTreeItem> InItem) const;
+	ECheckBoxState GetFilterCheckedState(TSharedPtr<FNiagaraSimCacheTreeItem> InItem) const;
 	void VisibilityButtonClicked(TSharedRef<FNiagaraSimCacheTreeItem> InItem);
 	bool IsItemSelected(TSharedRef<FNiagaraSimCacheTreeItem> InItem);
 	void ClearFilterSelection();
-	void UpdateStringFilters();
 	void SelectAll();
 	bool ShouldShowComponentView() const;
 
@@ -82,19 +78,11 @@ private:
 
 	void OnBufferChanged();
 
-	void RecursiveAddToSelectionFilter(TArray<TSharedRef<FNiagaraSimCacheTreeItem>>& ArrayToAdd);
-	void RecursiveRemoveFromSelectionFilter(TArray<TSharedRef<FNiagaraSimCacheTreeItem>>& ArrayToRemove);
-	
-	void UpdateSelectionFilter(TSharedRef<FNiagaraSimCacheTreeItem> ClickedItem);
-	
-	
 	TSharedPtr<STreeView<TSharedRef<FNiagaraSimCacheTreeItem>>> TreeView;
 
 	TArray<TSharedRef<FNiagaraSimCacheTreeItem>> ActiveRootEntries;
 
 	TSharedPtr<FNiagaraSimCacheViewModel> ViewModel;
-
-	TArray<TSharedRef<FNiagaraSimCacheTreeItem>> SelectionForFilter;
 };
 
 struct FNiagaraSimCacheTreeItem : FNiagaraSimCacheOverviewItem
diff --git a/Engine/Plugins/FX/Niagara/Source/NiagaraEditor/Private/Widgets/SNiagaraSimCacheView.cpp b/Engine/Plugins/FX/Niagara/Source/NiagaraEditor/Private/Widgets/SNiagaraSimCacheView.cpp
--- a/Engine/Plugins/FX/Niagara/Source/NiagaraEditor/Private/Widgets/SNiagaraSimCacheView.cpp
+++ b/Engine/Plugins/FX/Niagara/Source/NiagaraEditor/Private/Widgets/SNiagaraSimCacheView.cpp
@@ -369,28 +369,12 @@ void SNiagaraSimCacheView::OnBufferChanged()
 
 bool SNiagaraSimCacheView::GetShouldGenerateWidget(FName Name)
 {
-	TArray<FString> FilterArray;
-	FilterArray = SimCacheViewModel->GetComponentFilters();
-
 	if(!SimCacheViewModel->IsComponentFilterActive())
 	{
 		return true;
 	}
 
-	FString ColumnName = Name.ToString();
-
-	// Always display the instance column
-	if(ColumnName.Equals(NAME_Instance.ToString()))
-	{
-		return true;
-	}
-
-	auto ComponentFilterPred = [&ColumnName](const FString& ComponentFilter)
-	{
-		return ColumnName.Equals(ComponentFilter);
-	};
-	
-	return FilterArray.ContainsByPredicate(ComponentFilterPred);
+	return SimCacheViewModel->IsComponentFiltered(Name);
 }
 
 #undef LOCTEXT_NAMESPACE
diff --git a/Engine/Plugins/FX/Niagara/Source/NiagaraEditor/Public/ViewModels/NiagaraSimCacheViewModel.h b/Engine/Plugins/FX/Niagara/Source/NiagaraEditor/Public/ViewModels/NiagaraSimCacheViewModel.h
--- a/Engine/Plugins/FX/Niagara/Source/NiagaraEditor/Public/ViewModels/NiagaraSimCacheViewModel.h
+++ b/Engine/Plugins/FX/Niagara/Source/NiagaraEditor/Public/ViewModels/NiagaraSimCacheViewModel.h
@@ -27,13 +27,14 @@ public:
 
 	struct FComponentInfo
 	{
-		FName Name = NAME_None;
-		uint32 ComponentOffset = INDEX_NONE;
-		bool bIsFloat = false;
-		bool bIsHalf = false;
-		bool bIsInt32 = false;
-		bool bShowAsBool = false;
-		UEnum* Enum = nullptr;
+		FName	Name = NAME_None;
+		uint32	ComponentOffset = INDEX_NONE;
+		bool	bIsFiltered = true;
+		bool	bIsFloat = false;
+		bool	bIsHalf = false;
+		bool	bIsInt32 = false;
+		bool	bShowAsBool = false;
+		UEnum*	Enum = nullptr;
 	};
 
 	DECLARE_MULTICAST_DELEGATE_OneParam(FOnViewDataChanged, bool)
@@ -42,9 +43,11 @@ public:
 	
 	NIAGARAEDITOR_API FNiagaraSimCacheViewModel();
 	NIAGARAEDITOR_API virtual ~FNiagaraSimCacheViewModel() override;
-	
+
+private:
+	TArrayView<FComponentInfo> GetMutableSelectedComponentInfos();
+public:
 	NIAGARAEDITOR_API TConstArrayView<FComponentInfo> GetSelectedComponentInfos() const;
-	NIAGARAEDITOR_API TConstArrayView<FComponentInfo> GetComponentInfos(int32 InEmitterIndex) const;
 
 	NIAGARAEDITOR_API int32 GetNumInstances() const { return NumInstances; }
 	NIAGARAEDITOR_API int32 GetNumFrames() const;
@@ -57,24 +60,30 @@ public:
 
 	void Initialize(TWeakObjectPtr<UNiagaraSimCache> SimCache);
 
-	TConstArrayView<FString> GetComponentFilters() const { return ComponentFilterArray; }
-	void SetComponentFilters(const TArray<FString>& NewComponentFilterArray);
-	bool IsComponentFilterActive() const { return bComponentFilterActive; }
-
 	ESelectionMode GetSelectionMode() const { return SelectionMode; }
-	int32 GetSelectedEmitter() const { return SelectedEmitterIndex; };
+	FName GetSelectedEmitter() const { return SelectedEmitterName; };
 	FNiagaraVariableBase GetSelectedDataInterface() const { return SelectedDataInterface; };
+	FText GetSelectedText() const;
 
 	NIAGARAEDITOR_API const UObject* GetSelectedDataInterfaceStorage() const;
 
 	void SetSelectedSystemInstance();
-	void SetSelectedEmitter(int32 EmitterIndex);
+	void SetSelectedEmitter(FName EmitterName);
 	void SetSelectedDataInterface(FNiagaraVariableBase DIVariable);
 	void SetSelectedDebugData();
+
+	bool IsComponentFilterActive() const;
+	bool IsComponentFiltered(const FString& ComponentNameString) const;
+	bool IsComponentFiltered(FName ComponentName) const;
+	void SetComponentFiltered(const FString& ComponentNameString, bool bFiltered);
+	void ToggleComponentFiltered(const FString& ComponentNameString);
+	void SetAllComponentFiltered(bool bFiltered);
+
+	void RefreshSelection();
 private:
 	void RefreshFromSelectionChanged();
-public:
 
+public:
 	FText GetComponentText(FName ComponentName, int32 InstanceIndex) const;
 	bool IsCacheValid() const;
 	int32 GetNumEmitterLayouts() const;
@@ -123,15 +132,15 @@ private:
 
 	// Handles what is selected currently
 	ESelectionMode			SelectionMode = ESelectionMode::SystemInstance;
-	int32					SelectedEmitterIndex = INDEX_NONE;
+	FName					SelectedEmitterName;
 	FNiagaraVariableBase	SelectedDataInterface;
 
 	// Number of particles in the given frame
 	int32 NumInstances = 0;
 
 	// Cached Component infos for this system
-	TArray<FComponentInfo>				SystemComponentInfos;
-	TArray<TArray<FComponentInfo>>		EmitterComponentInfos;
+	TArray<FComponentInfo>					SystemComponentInfos;
+	TMap<FName, TArray<FComponentInfo>>		EmitterComponentInfos;
 	// TODO: Component infos for Data Interfaces?
 
 	TArray<TSharedRef<FNiagaraSimCacheTreeItem>> RootEntries;
@@ -139,10 +148,6 @@ private:
 
 	TArray<TSharedRef<FNiagaraSimCacheOverviewItem>> BufferEntries;
 
-	// String filters gathered from the overview tree
-	TArray<FString>						ComponentFilterArray;
-	bool								bComponentFilterActive = false;
-	
 	TArray<float>						FloatComponents;
 	TArray<FFloat16>					HalfComponents;
 	TArray<int32>						Int32Components;
