Add check to try to catch more info from a lingering non-deterministic assertion happening during loading
#rnx

diff --git a/Engine/Plugins/FX/Niagara/Source/NiagaraEditor/Private/NiagaraNodeFunctionCall.cpp b/Engine/Plugins/FX/Niagara/Source/NiagaraEditor/Private/NiagaraNodeFunctionCall.cpp
--- a/Engine/Plugins/FX/Niagara/Source/NiagaraEditor/Private/NiagaraNodeFunctionCall.cpp
+++ b/Engine/Plugins/FX/Niagara/Source/NiagaraEditor/Private/NiagaraNodeFunctionCall.cpp
@@ -305,6 +305,12 @@ UEdGraphPin* UNiagaraNodeFunctionCall::AddStaticSwitchInputPin(FNiagaraVariable
 	TOptional<FGuid> ScriptVariableGuid = Graph->GetScriptVariableGuid(Input);
 	if (SwitchDefaultValue.IsSet() && ScriptVariableGuid.IsSet())
 	{
+		checkf(Input.GetType().GetSize() == sizeof(FNiagaraInt32), TEXT("Incorrectly sized Input variable for UNiagaraNodeFunctionCall::AddStaticSwitchInputPin().  Size: %d (should be %d). Type: %s. Name: %s. Node: %s. ScriptVariableGuid: %s"),
+			Input.GetType().GetSize(), sizeof(FNiagaraInt32),
+			*Input.GetType().GetName(),
+			*Input.GetName().ToString(),
+			*GetFullName(),
+			*ScriptVariableGuid->ToString());
 		Input.AllocateData();
 		Input.SetValue<FNiagaraInt32>( {*SwitchDefaultValue} );
 		PinPersistentGuid = *ScriptVariableGuid;
