Added a getter for the current splitscreen mode to the camera DI

#tests Created local players and different splitscreen configurations to see if cpu and gpu emitters correctly pick it up

diff --git a/Engine/Plugins/FX/Niagara/Shaders/Private/NiagaraDataInterfaceCamera.ush b/Engine/Plugins/FX/Niagara/Shaders/Private/NiagaraDataInterfaceCamera.ush
--- a/Engine/Plugins/FX/Niagara/Shaders/Private/NiagaraDataInterfaceCamera.ush
+++ b/Engine/Plugins/FX/Niagara/Shaders/Private/NiagaraDataInterfaceCamera.ush
@@ -1,12 +1,18 @@
 // Copyright Epic Games, Inc. All Rights Reserved.
 
 float3 {ParameterName}_SystemLWCTile;
+int {ParameterName}_SplitscreenMode;
 
 void GetFieldOfView_{ParameterName}(out float Out_FieldOfViewAngle)
 {
 	Out_FieldOfViewAngle = degrees(View.FieldOfViewWideAngles.x);
 }
 
+void GetSplitscreenInfo_{ParameterName}(out int Out_CurrentSplitscreenMode)
+{
+	Out_CurrentSplitscreenMode = {ParameterName}_SplitscreenMode;
+}
+
 void GetTAAJitter_{ParameterName}(out float2 Out_CurrentJitter, out float2 Out_PreviousJitter)
 {
 	Out_CurrentJitter = View.TemporalAAJitter.xy;
diff --git a/Engine/Plugins/FX/Niagara/Source/Niagara/Classes/NiagaraDataInterfaceCamera.h b/Engine/Plugins/FX/Niagara/Source/Niagara/Classes/NiagaraDataInterfaceCamera.h
--- a/Engine/Plugins/FX/Niagara/Source/Niagara/Classes/NiagaraDataInterfaceCamera.h
+++ b/Engine/Plugins/FX/Niagara/Source/Niagara/Classes/NiagaraDataInterfaceCamera.h
@@ -4,6 +4,7 @@
 #include "NiagaraCommon.h"
 #include "NiagaraShared.h"
 #include "NiagaraDataInterface.h"
+#include "Engine/ViewportSplitScreen.h"
 #include "NiagaraDataInterfaceCamera.generated.h"
 
 struct FDistanceData
@@ -17,6 +18,7 @@ struct FCameraDataInterface_InstanceData
 	FNiagaraPosition CameraLocation = FVector::ZeroVector;
 	FRotator CameraRotation = FRotator::ZeroRotator;
 	float CameraFOV = 0.0f;
+	ESplitScreenType::Type SplitscreenMode = ESplitScreenType::Type::None;
 	FNiagaraLWCConverter LWCConverter;
 
 	TQueue<FDistanceData, EQueueMode::Mpsc> DistanceSortQueue;
@@ -30,6 +32,7 @@ class UNiagaraDataInterfaceCamera : public UNiagaraDataInterface
 
 	BEGIN_SHADER_PARAMETER_STRUCT(FShaderParameters, )
 		SHADER_PARAMETER(FVector3f, SystemLWCTile)
+		SHADER_PARAMETER(int32, SplitscreenMode)
 	END_SHADER_PARAMETER_STRUCT();
 
 public:
@@ -53,6 +56,7 @@ public:
 	NIAGARA_API virtual void DestroyPerInstanceData(void* PerInstanceData, FNiagaraSystemInstance* SystemInstance) override;
 	virtual int32 PerInstanceDataSize() const override { return sizeof(FCameraDataInterface_InstanceData); }
 	NIAGARA_API virtual bool PerInstanceTick(void* PerInstanceData, FNiagaraSystemInstance* SystemInstance, float DeltaSeconds) override;
+	virtual void ProvidePerInstanceDataForRenderThread(void* DataForRenderThread, void* PerInstanceData, const FNiagaraSystemInstanceID& SystemInstance) override;
 #if WITH_EDITORONLY_DATA
 	NIAGARA_API virtual bool AppendCompileHash(FNiagaraCompileHashVisitor* InVisitor) const override;
 	NIAGARA_API virtual bool GetFunctionHLSL(const FNiagaraDataInterfaceGPUParamInfo& ParamInfo, const FNiagaraDataInterfaceGeneratedFunction& FunctionInfo, int FunctionInstanceIndex, FString& OutHLSL) override;
@@ -77,6 +81,7 @@ public:
 	NIAGARA_API void GetClosestParticles(FVectorVMExternalFunctionContext& Context);
 	NIAGARA_API void GetCameraFOV(FVectorVMExternalFunctionContext& Context);
 	NIAGARA_API void GetCameraProperties(FVectorVMExternalFunctionContext& Context);
+	NIAGARA_API void GetGetSplitScreenInfo(FVectorVMExternalFunctionContext& Context);
 
 protected:
 #if WITH_EDITORONLY_DATA
@@ -94,15 +99,7 @@ private:
 	static NIAGARA_API const FName GetCameraPropertiesName;
 	static NIAGARA_API const FName GetFieldOfViewName;
 	static NIAGARA_API const FName GetTAAJitterName;
+	static NIAGARA_API const FName GetSplitScreenName;
 	static NIAGARA_API const FName ApplyPreViewTranslationToPosition;
 	static NIAGARA_API const FName RemovePreViewTranslationFromPosition;
 };
-
-struct FNiagaraDataIntefaceProxyCameraQuery : public FNiagaraDataInterfaceProxy
-{
-	// There's nothing in this proxy.
-	virtual int32 PerInstanceDataPassedToRenderThreadSize() const override
-	{
-		return 0;
-	}
-};
diff --git a/Engine/Plugins/FX/Niagara/Source/Niagara/Private/NiagaraDataInterfaceCamera.cpp b/Engine/Plugins/FX/Niagara/Source/Niagara/Private/NiagaraDataInterfaceCamera.cpp
--- a/Engine/Plugins/FX/Niagara/Source/Niagara/Private/NiagaraDataInterfaceCamera.cpp
+++ b/Engine/Plugins/FX/Niagara/Source/Niagara/Private/NiagaraDataInterfaceCamera.cpp
@@ -3,6 +3,7 @@
 #include "NiagaraDataInterfaceCamera.h"
 
 #include "Engine/LocalPlayer.h"
+#include "Engine/GameViewportClient.h"
 #include "NiagaraCompileHashVisitor.h"
 #include "NiagaraTypes.h"
 #include "NiagaraWorldManager.h"
@@ -43,13 +44,44 @@ const FName UNiagaraDataInterfaceCamera::GetFieldOfViewName(TEXT("GetFieldOfView
 const FName UNiagaraDataInterfaceCamera::CalculateDistancesName(TEXT("CalculateParticleDistancesCPU"));
 const FName UNiagaraDataInterfaceCamera::QueryClosestName(TEXT("QueryClosestParticlesCPU"));
 const FName UNiagaraDataInterfaceCamera::GetTAAJitterName(TEXT("GetTAAJitter"));
+const FName UNiagaraDataInterfaceCamera::GetSplitScreenName(TEXT("GetSplitscreenInfo"));
 const FName UNiagaraDataInterfaceCamera::ApplyPreViewTranslationToPosition(TEXT("ApplyPreViewTranslationToPositionGPU"));
 const FName UNiagaraDataInterfaceCamera::RemovePreViewTranslationFromPosition(TEXT("RemovePreViewTranslationFromPositionGPU"));
 
+struct FCameraDataInterface_InstanceDataRT
+{
+	ESplitScreenType::Type SplitscreenMode = ESplitScreenType::Type::None;
+};
+
+struct FNiagaraDataInterfaceProxyCameraQuery : FNiagaraDataInterfaceProxy
+{
+	virtual int32 PerInstanceDataPassedToRenderThreadSize() const override { return sizeof(FCameraDataInterface_InstanceDataRT); }
+
+	static void ProvidePerInstanceDataForRenderThread(void* InDataForRenderThread, void* InDataFromGameThread, const FNiagaraSystemInstanceID& SystemInstance)
+	{
+		FCameraDataInterface_InstanceDataRT* DataForRenderThread = new (InDataForRenderThread) FCameraDataInterface_InstanceDataRT();
+		const FCameraDataInterface_InstanceData* DataFromGameThread = static_cast<FCameraDataInterface_InstanceData*>(InDataFromGameThread);
+		
+		DataForRenderThread->SplitscreenMode = DataFromGameThread->SplitscreenMode;
+	}
+
+	virtual void ConsumePerInstanceDataFromGameThread(void* PerInstanceData, const FNiagaraSystemInstanceID& InstanceID) override
+	{
+		FCameraDataInterface_InstanceDataRT* InstanceDataFromGT = static_cast<FCameraDataInterface_InstanceDataRT*>(PerInstanceData);
+		FCameraDataInterface_InstanceDataRT& InstanceData = SystemInstancesToInstanceData_RT.FindOrAdd(InstanceID);
+		InstanceData = *InstanceDataFromGT;
+
+		// clean up the GT data
+		InstanceDataFromGT->~FCameraDataInterface_InstanceDataRT();
+	}
+
+	TMap<FNiagaraSystemInstanceID, FCameraDataInterface_InstanceDataRT> SystemInstancesToInstanceData_RT;
+};
+
 UNiagaraDataInterfaceCamera::UNiagaraDataInterfaceCamera(FObjectInitializer const& ObjectInitializer)
 	: Super(ObjectInitializer)
 {
-	Proxy.Reset(new FNiagaraDataIntefaceProxyCameraQuery());
+	Proxy.Reset(new FNiagaraDataInterfaceProxyCameraQuery());
 }
 
 void UNiagaraDataInterfaceCamera::PostInitProperties()
@@ -73,6 +105,14 @@ void UNiagaraDataInterfaceCamera::DestroyPerInstanceData(void* PerInstanceData,
 {
 	FCameraDataInterface_InstanceData* PIData = static_cast<FCameraDataInterface_InstanceData*>(PerInstanceData);
 	PIData->~FCameraDataInterface_InstanceData();
+
+	ENQUEUE_RENDER_COMMAND(RemoveProxy)
+	(
+		[RT_Proxy=GetProxyAs<FNiagaraDataInterfaceProxyCameraQuery>(), InstanceID=SystemInstance->GetId()](FRHICommandListImmediate& CmdList)
+		{
+			RT_Proxy->SystemInstancesToInstanceData_RT.Remove(InstanceID);
+		}
+	);
 }
 
 bool UNiagaraDataInterfaceCamera::PerInstanceTick(void* PerInstanceData, FNiagaraSystemInstance* SystemInstance, float DeltaSeconds)
@@ -115,6 +155,10 @@ bool UNiagaraDataInterfaceCamera::PerInstanceTick(void* PerInstanceData, FNiagar
 						PIData->CameraLocation = LWCConverter.ConvertWorldToSimulationPosition(PlayerController->PlayerCameraManager->GetCameraLocation());
 						PIData->CameraRotation = PlayerController->PlayerCameraManager->GetCameraRotation();
 						PIData->CameraFOV = PlayerController->PlayerCameraManager->GetFOVAngle();
+						if (ULocalPlayer* LocalPlayer = PlayerController->GetLocalPlayer())
+						{
+							PIData->SplitscreenMode = LocalPlayer->ViewportClient ? LocalPlayer->ViewportClient->GetCurrentSplitscreenConfiguration() : ESplitScreenType::Type::None;
+						}
 						return false;
 					}
 					else
@@ -147,6 +191,11 @@ bool UNiagaraDataInterfaceCamera::PerInstanceTick(void* PerInstanceData, FNiagar
 	return false;
 }
 
+void UNiagaraDataInterfaceCamera::ProvidePerInstanceDataForRenderThread(void* DataForRenderThread, void* PerInstanceData, const FNiagaraSystemInstanceID& SystemInstance)
+{
+	FNiagaraDataInterfaceProxyCameraQuery::ProvidePerInstanceDataForRenderThread(DataForRenderThread, PerInstanceData, SystemInstance);
+}
+
 #if WITH_EDITORONLY_DATA
 void UNiagaraDataInterfaceCamera::GetFunctionsInternal(TArray<FNiagaraFunctionSignature>& OutFunctions) const
 {
@@ -177,7 +226,6 @@ void UNiagaraDataInterfaceCamera::GetFunctionsInternal(TArray<FNiagaraFunctionSi
 	Sig.Description = LOCTEXT("GetClipSpaceTransformsDescription", "This function returns the clip transforms for the current view. Only valid for gpu particles.");
 	Sig.FunctionVersion = FNiagaraCameraDIFunctionVersion::LatestVersion;
 	Sig.bMemberFunction = true;
-	Sig.bRequiresContext = false;
 	Sig.bSupportsCPU = false;
 	Sig.AddInput(FNiagaraVariable(FNiagaraTypeDefinition(GetClass()), TEXT("Camera interface")));
 	Sig.AddOutput(FNiagaraVariable(FNiagaraTypeDefinition::GetMatrix4Def(), TEXT("World To Clip Transform")), LOCTEXT("WorldToClipTransformDescription", "Transforms a world space position to clip space"));
@@ -196,7 +244,6 @@ void UNiagaraDataInterfaceCamera::GetFunctionsInternal(TArray<FNiagaraFunctionSi
 	Sig.Description = LOCTEXT("GetViewSpaceTransformsDescription", "This function returns the relevant transforms for the current view. Only valid for gpu particles.");
 	Sig.FunctionVersion = FNiagaraCameraDIFunctionVersion::LatestVersion;
 	Sig.bMemberFunction = true;
-	Sig.bRequiresContext = false;
 	Sig.bSupportsCPU = false;
 	Sig.AddInput(FNiagaraVariable(FNiagaraTypeDefinition(GetClass()), TEXT("Camera interface")));
 	Sig.AddOutput(FNiagaraVariable(FNiagaraTypeDefinition::GetMatrix4Def(), TEXT("Translated World To View Transform")), LOCTEXT("TranslatedWorldToViewTransformDescription", "Transforms a position from world to view, but the view position is assumed to be at the origin.\nThis allows for more precision of the transform when dealing with big coordinates.\nTo use this transform, you first need to add PreViewTranslation to your transform target."));
@@ -213,7 +260,6 @@ void UNiagaraDataInterfaceCamera::GetFunctionsInternal(TArray<FNiagaraFunctionSi
 	Sig.Description = LOCTEXT("ApplyPreViewTranslationToPositionDescription", "This function applies the pre view translation to a lwc positions such as Particles.Position. The resulting vector can be used with any of the 'Translated X to Y' transform matrices.");
 	Sig.FunctionVersion = FNiagaraCameraDIFunctionVersion::LatestVersion;
 	Sig.bMemberFunction = true;
-	Sig.bRequiresContext = false;
 	Sig.bSupportsCPU = false;
 	Sig.AddInput(FNiagaraVariable(FNiagaraTypeDefinition(GetClass()), TEXT("Camera interface")));
 	Sig.AddInput(FNiagaraVariable(FNiagaraTypeDefinition::GetPositionDef(), TEXT("World Position")), LOCTEXT("ApplyPreViewTranslationToPositionInputDescription", "Position input to be transformed."));
@@ -225,7 +271,6 @@ void UNiagaraDataInterfaceCamera::GetFunctionsInternal(TArray<FNiagaraFunctionSi
 	Sig.Description = LOCTEXT("RemovePreViewTranslationFromPositionDescription", "This function removes the pre view translation from a vector to return a lwc position. Use the result of a 'X to Translated Y' transform matrix multiplication as input vector.");
 	Sig.FunctionVersion = FNiagaraCameraDIFunctionVersion::LatestVersion;
 	Sig.bMemberFunction = true;
-	Sig.bRequiresContext = false;
 	Sig.bSupportsCPU = false;
 	Sig.AddInput(FNiagaraVariable(FNiagaraTypeDefinition(GetClass()), TEXT("Camera interface")));
 	Sig.AddInput(FNiagaraVariable(FNiagaraTypeDefinition::GetVec4Def(), TEXT("Translated Position")), LOCTEXT("RemovePreViewTranslationFromPositionInputDescription", "Translated vector that is a result of a 'X to Translated Y' transforms matrix."));
@@ -238,7 +283,6 @@ void UNiagaraDataInterfaceCamera::GetFunctionsInternal(TArray<FNiagaraFunctionSi
 	Sig.Description = LOCTEXT("GetNiagaraFOVDescription", "This function returns the field of view angle (in degrees) for the active camera. For gpu particles this returns the x axis fov.");
 	Sig.FunctionVersion = FNiagaraCameraDIFunctionVersion::LatestVersion;
 	Sig.bMemberFunction = true;
-	Sig.bRequiresContext = false;
 	Sig.AddInput(FNiagaraVariable(FNiagaraTypeDefinition(GetClass()), TEXT("Camera interface")));
 	Sig.AddOutput(FNiagaraVariable(FNiagaraTypeDefinition::GetFloatDef(), TEXT("Field Of View Angle")), LOCTEXT("FieldOfViewAngleDescription", "Return the camera's field of view in degrees."));
 	OutFunctions.Add(Sig);
@@ -249,7 +293,6 @@ void UNiagaraDataInterfaceCamera::GetFunctionsInternal(TArray<FNiagaraFunctionSi
 	Sig.Description = LOCTEXT("GetCameraPositionDescription", "This function returns the position of the currently active camera.");
 	Sig.FunctionVersion = FNiagaraCameraDIFunctionVersion::LatestVersion;
 	Sig.bMemberFunction = true;
-	Sig.bRequiresContext = false;
 	Sig.AddInput(FNiagaraVariable(FNiagaraTypeDefinition(GetClass()), TEXT("Camera interface")));
 	Sig.AddOutput(FNiagaraVariable(FNiagaraTypeDefinition::GetPositionDef(), TEXT("Camera Position World")), LOCTEXT("CameraPositionWorldDescription", "The camera position in world space."));
 	Sig.AddOutput(FNiagaraVariable(FNiagaraTypeDefinition::GetVec3Def(), TEXT("Forward Vector World")), LOCTEXT("ForwardVectorWorldDescription", "The world space direction that the camera is pointing."));
@@ -262,7 +305,6 @@ void UNiagaraDataInterfaceCamera::GetFunctionsInternal(TArray<FNiagaraFunctionSi
 	Sig.Description = LOCTEXT("QueryClosestDescription", "This function checks the previously calculated distance of each particle and then returns true for the closest particles and false for the other ones.\nThis function needs to be paired with CalculateParticleDistancesCPU to work correctly.");
 	Sig.FunctionVersion = FNiagaraCameraDIFunctionVersion::LatestVersion;
 	Sig.bMemberFunction = true;
-	Sig.bRequiresContext = false;
 	Sig.bSupportsGPU = false;
 	Sig.AddInput(FNiagaraVariable(FNiagaraTypeDefinition(GetClass()), TEXT("Camera interface")));
 	Sig.AddInput(FNiagaraVariable(FNiagaraTypeDefinition::GetIDDef(), TEXT("Particle ID")));
@@ -275,7 +317,6 @@ void UNiagaraDataInterfaceCamera::GetFunctionsInternal(TArray<FNiagaraFunctionSi
 	Sig.Description = LOCTEXT("CalculateDistancesDescription", "This function compares the particle position against the camera position and stores the result to be queried in the next frame.\nThe results can then be queried with QueryClosestParticlesCPU.");
 	Sig.FunctionVersion = FNiagaraCameraDIFunctionVersion::LatestVersion;
 	Sig.bMemberFunction = true;
-	Sig.bRequiresContext = false;
 	Sig.bSupportsGPU = false;
 	Sig.bRequiresExecPin = true;
 	Sig.AddInput(FNiagaraVariable(FNiagaraTypeDefinition(GetClass()), TEXT("Camera interface")));
@@ -288,11 +329,20 @@ void UNiagaraDataInterfaceCamera::GetFunctionsInternal(TArray<FNiagaraFunctionSi
 	Sig.Description = NSLOCTEXT("Niagara", "GetTAAJitterDescription", "This function returns the TAA jitter values of the currently active camera.");
 	Sig.FunctionVersion = FNiagaraCameraDIFunctionVersion::LatestVersion;
 	Sig.bMemberFunction = true;
-	Sig.bRequiresContext = false;
+	Sig.bSupportsCPU = false;
 	Sig.Inputs.Add(FNiagaraVariable(FNiagaraTypeDefinition(GetClass()), TEXT("Camera interface")));
 	Sig.Outputs.Add(FNiagaraVariable(FNiagaraTypeDefinition::GetVec2Def(), TEXT("Current TAA Jitter (clip)")));
 	Sig.Outputs.Add(FNiagaraVariable(FNiagaraTypeDefinition::GetVec2Def(), TEXT("Previous TAA Jitter (clip)")));
 	OutFunctions.Add(Sig);
+
+	Sig = FNiagaraFunctionSignature();
+	Sig.Name = GetSplitScreenName;
+	Sig.Description = NSLOCTEXT("Niagara", "GetSplitScreenNameDescription", "This function returns the currently active splitscreen mode as ESplitScreenType value.");
+	Sig.FunctionVersion = FNiagaraCameraDIFunctionVersion::LatestVersion;
+	Sig.bMemberFunction = true;
+	Sig.Inputs.Add(FNiagaraVariable(FNiagaraTypeDefinition(GetClass()), TEXT("Camera interface")));
+	Sig.AddOutput(FNiagaraVariable(FNiagaraTypeDefinition(StaticEnum<ESplitScreenType::Type>()), TEXT("Current Splitscreen Mode")), LOCTEXT("SplitscreenModeOutput", "The splitscreen mode as ESplitScreenType value. A value of 0 means there is no splitscreen active."));
+	OutFunctions.Add(Sig);
 }
 
 bool UNiagaraDataInterfaceCamera::GetFunctionHLSL(const FNiagaraDataInterfaceGPUParamInfo& ParamInfo, const FNiagaraDataInterfaceGeneratedFunction& FunctionInfo, int FunctionInstanceIndex, FString& OutHLSL)
@@ -304,6 +354,7 @@ bool UNiagaraDataInterfaceCamera::GetFunctionHLSL(const FNiagaraDataInterfaceGPU
 		(FunctionInfo.DefinitionName == GetCameraPropertiesName) ||
 		(FunctionInfo.DefinitionName == ApplyPreViewTranslationToPosition) ||
 		(FunctionInfo.DefinitionName == RemovePreViewTranslationFromPosition) ||
+		(FunctionInfo.DefinitionName == GetSplitScreenName) ||
 		(FunctionInfo.DefinitionName == GetTAAJitterName))
 	{
 		return true;
@@ -387,14 +438,19 @@ void UNiagaraDataInterfaceCamera::BuildShaderParameters(FNiagaraShaderParameters
 
 void UNiagaraDataInterfaceCamera::SetShaderParameters(const FNiagaraDataInterfaceSetShaderParametersContext& Context) const
 {
+	FNiagaraDataInterfaceProxyCameraQuery& DataInterfaceProxy = Context.GetProxy<FNiagaraDataInterfaceProxyCameraQuery>();
+	FCameraDataInterface_InstanceDataRT& InstanceData = DataInterfaceProxy.SystemInstancesToInstanceData_RT.FindChecked(Context.GetSystemInstanceID());
+	
 	FShaderParameters* ShaderParameters = Context.GetParameterNestedStruct<FShaderParameters>();
 	ShaderParameters->SystemLWCTile = Context.GetSystemLWCTile();
+	ShaderParameters->SplitscreenMode = static_cast<int32>(InstanceData.SplitscreenMode);
 }
 
 DEFINE_NDI_DIRECT_FUNC_BINDER(UNiagaraDataInterfaceCamera, GetClosestParticles);
 DEFINE_NDI_DIRECT_FUNC_BINDER(UNiagaraDataInterfaceCamera, CalculateParticleDistances);
 DEFINE_NDI_DIRECT_FUNC_BINDER(UNiagaraDataInterfaceCamera, GetCameraFOV);
 DEFINE_NDI_DIRECT_FUNC_BINDER(UNiagaraDataInterfaceCamera, GetCameraProperties);
+DEFINE_NDI_DIRECT_FUNC_BINDER(UNiagaraDataInterfaceCamera, GetGetSplitScreenInfo);
 void UNiagaraDataInterfaceCamera::GetVMExternalFunction(const FVMExternalFunctionBindingInfo& BindingInfo, void* InstanceData, FVMExternalFunction &OutFunc)
 {
 	if (BindingInfo.Name == GetFieldOfViewName)
@@ -413,6 +469,10 @@ void UNiagaraDataInterfaceCamera::GetVMExternalFunction(const FVMExternalFunctio
 	{
 		NDI_FUNC_BINDER(UNiagaraDataInterfaceCamera, GetCameraProperties)::Bind(this, OutFunc);
 	}
+	else if (BindingInfo.Name == GetSplitScreenName)
+	{
+		NDI_FUNC_BINDER(UNiagaraDataInterfaceCamera, GetGetSplitScreenInfo)::Bind(this, OutFunc);
+	}
 	else
 	{
 		UE_LOG(LogNiagara, Display, TEXT("Could not find data interface external function in %s. Received Name: %s"), *GetPathNameSafe(this), *BindingInfo.Name.ToString());
@@ -482,6 +542,19 @@ void UNiagaraDataInterfaceCamera::GetCameraProperties(FVectorVMExternalFunctionC
 	}
 }
 
+void UNiagaraDataInterfaceCamera::GetGetSplitScreenInfo(FVectorVMExternalFunctionContext& Context)
+{
+	VectorVM::FUserPtrHandler<FCameraDataInterface_InstanceData> InstData(Context);
+	VectorVM::FExternalFuncRegisterHandler<ESplitScreenType::Type> OutMode(Context);
+
+	ESplitScreenType::Type SplitscreenMode = InstData.Get()->SplitscreenMode;
+
+	for (int32 i = 0; i < Context.GetNumInstances(); ++i)
+	{
+		*OutMode.GetDestAndAdvance() = SplitscreenMode;
+	}
+}
+
 void UNiagaraDataInterfaceCamera::GetClosestParticles(FVectorVMExternalFunctionContext& Context)
 {
 	VectorVM::FUserPtrHandler<FCameraDataInterface_InstanceData> InstData(Context);
