- Add support for sending doubles to float array
- They will be truncated, but it avoids the user needing to do a conversion to a float array

#rnx

diff --git a/Engine/Plugins/FX/Niagara/Source/Niagara/Classes/NiagaraDataInterfaceArrayFunctionLibrary.h b/Engine/Plugins/FX/Niagara/Source/Niagara/Classes/NiagaraDataInterfaceArrayFunctionLibrary.h
--- a/Engine/Plugins/FX/Niagara/Source/Niagara/Classes/NiagaraDataInterfaceArrayFunctionLibrary.h
+++ b/Engine/Plugins/FX/Niagara/Source/Niagara/Classes/NiagaraDataInterfaceArrayFunctionLibrary.h
@@ -165,6 +165,7 @@ public:
 
 	//////////////////////////////////////////////////////////////////////////////////////////////////////////////////
 	// None BP compatable set functions
+	static NIAGARA_API void SetNiagaraArrayFloat(UNiagaraComponent* NiagaraSystem, FName OverrideName, TConstArrayView<double> ArrayData);
 	static NIAGARA_API void SetNiagaraArrayVector2D(UNiagaraComponent* NiagaraSystem, FName OverrideName, TConstArrayView<FVector2f> ArrayData);
 	static NIAGARA_API void SetNiagaraArrayVector(UNiagaraComponent* NiagaraSystem, FName OverrideName, TConstArrayView<FVector3f> ArrayData);
 	static NIAGARA_API void SetNiagaraArrayVector4(UNiagaraComponent* NiagaraSystem, FName OverrideName, TConstArrayView<FVector4f> ArrayData);
diff --git a/Engine/Plugins/FX/Niagara/Source/Niagara/Private/NiagaraDataInterfaceArrayFunctionLibrary.cpp b/Engine/Plugins/FX/Niagara/Source/Niagara/Private/NiagaraDataInterfaceArrayFunctionLibrary.cpp
--- a/Engine/Plugins/FX/Niagara/Source/Niagara/Private/NiagaraDataInterfaceArrayFunctionLibrary.cpp
+++ b/Engine/Plugins/FX/Niagara/Source/Niagara/Private/NiagaraDataInterfaceArrayFunctionLibrary.cpp
@@ -418,6 +418,11 @@ bool UNiagaraDataInterfaceArrayFunctionLibrary::GetNiagaraArrayBoolValue(UNiagar
 
 //////////////////////////////////////////////////////////////////////////////////////////////////////////////////
 // None BP compatable set functions
+void UNiagaraDataInterfaceArrayFunctionLibrary::SetNiagaraArrayFloat(UNiagaraComponent* NiagaraSystem, FName OverrideName, TConstArrayView<double> ArrayData)
+{
+	SetNiagaraArray<double, UNiagaraDataInterfaceArrayFloat>(NiagaraSystem, OverrideName, ArrayData);
+}
+
 void UNiagaraDataInterfaceArrayFunctionLibrary::SetNiagaraArrayVector2D(UNiagaraComponent* NiagaraSystem, FName OverrideName, TConstArrayView<FVector2f> ArrayData)
 {
 	SetNiagaraArray<FVector2f, UNiagaraDataInterfaceArrayFloat2>(NiagaraSystem, OverrideName, ArrayData);
diff --git a/Engine/Plugins/FX/Niagara/Source/Niagara/Private/NiagaraDataInterfaceArrayImplHelpers.h b/Engine/Plugins/FX/Niagara/Source/Niagara/Private/NiagaraDataInterfaceArrayImplHelpers.h
--- a/Engine/Plugins/FX/Niagara/Source/Niagara/Private/NiagaraDataInterfaceArrayImplHelpers.h
+++ b/Engine/Plugins/FX/Niagara/Source/Niagara/Private/NiagaraDataInterfaceArrayImplHelpers.h
@@ -21,6 +21,19 @@ struct FNDIArrayImplHelper<float> : public FNDIArrayImplHelperBase<float>
 	static const FNiagaraTypeDefinition& GetTypeDefinition() { return FNiagaraTypeDefinition::GetFloatDef(); }
 	static const float GetDefaultValue() { return 0.0f; }
 
+	static void CopyCpuToCpuMemory(float* Dest, const float* Src, int32 NumElements)
+	{
+		FMemory::Memcpy(Dest, Src, NumElements * sizeof(float));
+	}
+
+	static void CopyCpuToCpuMemory(float* Dest, const double* Src, int32 NumElements)
+	{
+		for (int32 i = 0; i < NumElements; ++i)
+		{
+			Dest[i] = float(Src[i]);
+		}
+	}
+
 	static void AppendValueToString(const float Value, FString& OutString)
 	{
 		OutString.Appendf(TEXT("%f"), Value);
