- Add support for converting an int64 array to int32
- The data will be truncated, but avoids making an intermediate copy of the array

#rnx

diff --git a/Engine/Plugins/FX/Niagara/Source/Niagara/Classes/NiagaraDataInterfaceArrayFunctionLibrary.h b/Engine/Plugins/FX/Niagara/Source/Niagara/Classes/NiagaraDataInterfaceArrayFunctionLibrary.h
--- a/Engine/Plugins/FX/Niagara/Source/Niagara/Classes/NiagaraDataInterfaceArrayFunctionLibrary.h
+++ b/Engine/Plugins/FX/Niagara/Source/Niagara/Classes/NiagaraDataInterfaceArrayFunctionLibrary.h
@@ -165,6 +165,7 @@ public:
 
 	//////////////////////////////////////////////////////////////////////////////////////////////////////////////////
 	// None BP compatable set functions
+	static NIAGARA_API void SetNiagaraArrayInt32(UNiagaraComponent* NiagaraSystem, FName OverrideName, TConstArrayView<int64> ArrayData);
 	static NIAGARA_API void SetNiagaraArrayFloat(UNiagaraComponent* NiagaraSystem, FName OverrideName, TConstArrayView<double> ArrayData);
 	static NIAGARA_API void SetNiagaraArrayVector2D(UNiagaraComponent* NiagaraSystem, FName OverrideName, TConstArrayView<FVector2f> ArrayData);
 	static NIAGARA_API void SetNiagaraArrayVector(UNiagaraComponent* NiagaraSystem, FName OverrideName, TConstArrayView<FVector3f> ArrayData);
diff --git a/Engine/Plugins/FX/Niagara/Source/Niagara/Private/NiagaraDataInterfaceArrayFunctionLibrary.cpp b/Engine/Plugins/FX/Niagara/Source/Niagara/Private/NiagaraDataInterfaceArrayFunctionLibrary.cpp
--- a/Engine/Plugins/FX/Niagara/Source/Niagara/Private/NiagaraDataInterfaceArrayFunctionLibrary.cpp
+++ b/Engine/Plugins/FX/Niagara/Source/Niagara/Private/NiagaraDataInterfaceArrayFunctionLibrary.cpp
@@ -418,6 +418,11 @@ bool UNiagaraDataInterfaceArrayFunctionLibrary::GetNiagaraArrayBoolValue(UNiagar
 
 //////////////////////////////////////////////////////////////////////////////////////////////////////////////////
 // None BP compatable set functions
+void UNiagaraDataInterfaceArrayFunctionLibrary::SetNiagaraArrayInt32(UNiagaraComponent* NiagaraSystem, FName OverrideName, TConstArrayView<int64> ArrayData)
+{
+	SetNiagaraArray<int64, UNiagaraDataInterfaceArrayInt32>(NiagaraSystem, OverrideName, ArrayData);
+}
+
 void UNiagaraDataInterfaceArrayFunctionLibrary::SetNiagaraArrayFloat(UNiagaraComponent* NiagaraSystem, FName OverrideName, TConstArrayView<double> ArrayData)
 {
 	SetNiagaraArray<double, UNiagaraDataInterfaceArrayFloat>(NiagaraSystem, OverrideName, ArrayData);
diff --git a/Engine/Plugins/FX/Niagara/Source/Niagara/Private/NiagaraDataInterfaceArrayImplHelpers.h b/Engine/Plugins/FX/Niagara/Source/Niagara/Private/NiagaraDataInterfaceArrayImplHelpers.h
--- a/Engine/Plugins/FX/Niagara/Source/Niagara/Private/NiagaraDataInterfaceArrayImplHelpers.h
+++ b/Engine/Plugins/FX/Niagara/Source/Niagara/Private/NiagaraDataInterfaceArrayImplHelpers.h
@@ -412,6 +412,19 @@ struct FNDIArrayImplHelper<int32> : public FNDIArrayImplHelperBase<int32>
 	static const FNiagaraTypeDefinition& GetTypeDefinition() { return FNiagaraTypeDefinition::GetIntDef(); }
 	static const int32 GetDefaultValue() { return 0; }
 
+	static void CopyCpuToCpuMemory(int32* Dest, const int32* Src, int32 NumElements)
+	{
+		FMemory::Memcpy(Dest, Src, NumElements * sizeof(int32));
+	}
+
+	static void CopyCpuToCpuMemory(int32* Dest, const int64* Src, int32 NumElements)
+	{
+		for (int32 i = 0; i < NumElements; ++i)
+		{
+			Dest[i] = int32(Src[i]);
+		}
+	}
+
 	static int32 AtomicAdd(int32* Dest, int32 Value)
 	{
 		return FPlatformAtomics::InterlockedAdd(Dest, Value);
