- Fix memory leak from multiview code by moving all relevant code into WITH_MGPU

#rnx

diff --git a/Engine/Plugins/FX/Niagara/Source/Niagara/Classes/NiagaraComputeExecutionContext.h b/Engine/Plugins/FX/Niagara/Source/Niagara/Classes/NiagaraComputeExecutionContext.h
--- a/Engine/Plugins/FX/Niagara/Source/Niagara/Classes/NiagaraComputeExecutionContext.h
+++ b/Engine/Plugins/FX/Niagara/Source/Niagara/Classes/NiagaraComputeExecutionContext.h
@@ -79,8 +79,12 @@ struct FNiagaraComputeExecutionContext : public INiagaraComputeDataBufferInterfa
 		}
 		else
 		{
+		#if WITH_MGPU
 			// Non-translucent rendering uses previous frame (older) data, which is MultiViewPreviousDataToRender if present
 			return MultiViewPreviousDataToRender ? MultiViewPreviousDataToRender : DataToRender;
+		#else //WITH_MGPU
+			return DataToRender;
+		#endif //WITH_MGPU
 		}
 	}
 	// End: INiagaraComputeDataBufferInterface
@@ -97,7 +101,9 @@ struct FNiagaraComputeExecutionContext : public INiagaraComputeDataBufferInterfa
 
 	void SetDataToRender(FNiagaraDataBuffer* InDataToRender);
 	void SetTranslucentDataToRender(FNiagaraDataBuffer* InTranslucentDataToRender);
+#if WITH_MGPU
 	void SetMultiViewPreviousDataToRender(FNiagaraDataBuffer* InMultiViewPreviousDataToRender);
+#endif //WITH_MGPU
 
 	int32 GetConstantBufferSize() const;
 	uint8* WriteConstantBufferInstanceData(uint8* InTargetBuffer, FNiagaraComputeInstanceData& InstanceData) const;
@@ -167,8 +173,10 @@ public:
 	// Optional buffer which can be used to render translucent data with no latency (i.e. this frames data)
 	FNiagaraDataBufferRef TranslucentDataToRender = nullptr;
 
+#if WITH_MGPU
 	// For MultiView, we need to track the previous frame's data for rendering systems on other view families
 	FNiagaraDataBufferRef MultiViewPreviousDataToRender = nullptr;
+#endif //WITH_MGPU
 
 	// Game thread spawn info will be sent to the render thread inside FNiagaraComputeInstanceData
 	FNiagaraGpuSpawnInfo GpuSpawnInfo_GT;
diff --git a/Engine/Plugins/FX/Niagara/Source/Niagara/Private/NiagaraComputeExecutionContext.cpp b/Engine/Plugins/FX/Niagara/Source/Niagara/Private/NiagaraComputeExecutionContext.cpp
--- a/Engine/Plugins/FX/Niagara/Source/Niagara/Private/NiagaraComputeExecutionContext.cpp
+++ b/Engine/Plugins/FX/Niagara/Source/Niagara/Private/NiagaraComputeExecutionContext.cpp
@@ -242,10 +242,12 @@ void FNiagaraComputeExecutionContext::SetTranslucentDataToRender(FNiagaraDataBuf
 	TranslucentDataToRender = InTranslucentDataToRender;
 }
 
+#if WITH_MGPU
 void FNiagaraComputeExecutionContext::SetMultiViewPreviousDataToRender(FNiagaraDataBuffer* InMultiViewPreviousDataToRender)
 {
 	MultiViewPreviousDataToRender = InMultiViewPreviousDataToRender;
 }
+#endif //WITH_MGPU
 
 int32 FNiagaraComputeExecutionContext::GetConstantBufferSize() const
 {
diff --git a/Engine/Plugins/FX/Niagara/Source/Niagara/Private/NiagaraGpuComputeDispatch.cpp b/Engine/Plugins/FX/Niagara/Source/Niagara/Private/NiagaraGpuComputeDispatch.cpp
--- a/Engine/Plugins/FX/Niagara/Source/Niagara/Private/NiagaraGpuComputeDispatch.cpp
+++ b/Engine/Plugins/FX/Niagara/Source/Niagara/Private/NiagaraGpuComputeDispatch.cpp
@@ -1320,14 +1320,14 @@ void FNiagaraGpuComputeDispatch::ExecuteTicks(FRDGBuilder& GraphBuilder, TConstS
 				ComputeContext->SetTranslucentDataToRender(nullptr);
 				ComputeContext->SetDataToRender(CurrentData);
 
-#if WITH_MGPU
+			#if WITH_MGPU
 				if (bCrossGPUTransferEnabled)
 				{
 					AddCrossGPUTransfer(GraphBuilder.RHICmdList, CurrentData->GetGPUBufferFloat().Buffer);
 					AddCrossGPUTransfer(GraphBuilder.RHICmdList, CurrentData->GetGPUBufferHalf().Buffer);
 					AddCrossGPUTransfer(GraphBuilder.RHICmdList, CurrentData->GetGPUBufferInt().Buffer);
 				}
-#endif // WITH_MGPU
+			#endif // WITH_MGPU
 
 				if (bIsLastViewFamily)
 				{
@@ -1336,10 +1336,12 @@ void FNiagaraGpuComputeDispatch::ExecuteTicks(FRDGBuilder& GraphBuilder, TConstS
 				}
 				else
 				{
+				#if WITH_MGPU
 					// Track the previous data to render for multi-view, so additional view families can render using consistent data with
 					// the first view family, and add it to array marking that we need to clean it up on the last view family.
 					ComputeContext->SetMultiViewPreviousDataToRender(FinalSimStageDataBuffer);
 					NeedsMultiViewPreviousDataClear.Emplace(ComputeContext);
+				#endif //WITH_MGPU
 
 					// Mark CurrentData as not to be used until PostOpaqueRender.  Render data generated before PostOpaqueRender should use
 					// previous frame data (MultiViewPreviousDataToRender, set above), instead of CurrentData.  Setting the ready stage back to
