- Fix crash when reading parameter store data for aligned types (i.e. quat / FVector4f)

diff --git a/Engine/Plugins/FX/Niagara/Source/Niagara/Public/NiagaraParameterStore.h b/Engine/Plugins/FX/Niagara/Source/Niagara/Public/NiagaraParameterStore.h
--- a/Engine/Plugins/FX/Niagara/Source/Niagara/Public/NiagaraParameterStore.h
+++ b/Engine/Plugins/FX/Niagara/Source/Niagara/Public/NiagaraParameterStore.h
@@ -377,6 +377,16 @@ public:
 	
 	NIAGARA_API virtual const FNiagaraVariableWithOffset* FindParameterVariable(const FNiagaraVariable& Parameter, bool IgnoreType = false) const;
 
+private:
+	template<typename T>
+	T GetParameterValueFromOffset(int32 Offset) const
+	{
+		T OutValue;
+		FMemory::Memcpy(&OutValue, GetParameterData(Offset), sizeof(T));
+		return OutValue;
+	}
+
+public:
 	/** Gets the typed parameter data. */
 	template<typename T>
 	FORCEINLINE_DEBUGGABLE void GetParameterValue(T& OutValue, const FNiagaraVariableBase& Parameter)const
@@ -385,7 +395,7 @@ public:
 		int32 Offset = IndexOf(Parameter);
 		if (Offset != INDEX_NONE)
 		{
-			OutValue = *(const T*)(GetParameterData(Offset));
+			OutValue = GetParameterValueFromOffset<T>(Offset);
 		}
 	}
 
@@ -396,7 +406,7 @@ public:
 		int32 Offset = IndexOf(Parameter);
 		if (Offset != INDEX_NONE)
 		{
-			return *(const T*)(GetParameterData(Offset));
+			return GetParameterValueFromOffset<T>(Offset);
 		}
 		return T();
 	}
@@ -409,7 +419,7 @@ public:
 		const int32 Offset = IndexOf(Parameter);
 		if (Offset != INDEX_NONE)
 		{
-			OptionalValue.Emplace(*(const T*)(GetParameterData(Offset)));
+			OptionalValue.Emplace(GetParameterValueFromOffset<T>(Offset));
 		}
 		return OptionalValue;
 	}
@@ -419,7 +429,7 @@ public:
 	{
 		check(Parameter.GetSizeInBytes() == sizeof(T));
 		int32 Offset = IndexOf(Parameter);
-		return Offset != INDEX_NONE ? *(const T*)(GetParameterData(Offset)) : DefaultValue;
+		return Offset != INDEX_NONE ? GetParameterValueFromOffset<T>(Offset) : DefaultValue;
 	}
 
 	FORCEINLINE const uint8* GetParameterData(int32 Offset)const
