- GBuffer reads for partial depth

diff --git a/Engine/Plugins/FX/Niagara/Shaders/Private/NiagaraDataInterfaceGBufferTemplate.ush b/Engine/Plugins/FX/Niagara/Shaders/Private/NiagaraDataInterfaceGBufferTemplate.ush
--- a/Engine/Plugins/FX/Niagara/Shaders/Private/NiagaraDataInterfaceGBufferTemplate.ush
+++ b/Engine/Plugins/FX/Niagara/Shaders/Private/NiagaraDataInterfaceGBufferTemplate.ush
@@ -245,6 +245,22 @@ void DecodeStencil_{ParameterName}(in float2 ScreenUV, in bool bApplyViewportOff
 }
 #endif
 
+void DecodePartialDepth_{ParameterName}(in float2 ScreenUV, in bool bApplyViewportOffset, out bool IsValid, out float OutValue)
+{
+	ScreenUV = ApplyViewportOffset_{ParameterName}(ScreenUV, bApplyViewportOffset);
+#if NDIGBUFFER_HAS_GBUFFER
+	IsValid = true;
+	OutValue = ConvertFromDeviceZ(Texture2DSampleLevel(SceneTexturesStruct.ScenePartialDepthTexture, SceneTexturesStruct_ScenePartialDepthTextureSampler, ScreenUV, 0).r);
+#elif NDIGBUFFER_IS_MOBILE
+	IsValid = true;
+	OutValue = ConvertFromDeviceZ(Texture2DSampleLevel(MobileSceneTextures.ScenePartialDepthTexture, MobileSceneTextures.ScenePartialDepthTextureSampler, ScreenUV, 0).r);
+#else
+	IsValid = false;
+	OutValue = 0;
+#endif
+}
+
+
 void DecodeCustomDepth_{ParameterName}(in float2 ScreenUV, in bool bApplyViewportOffset, out bool IsValid, out float OutValue)
 {
 #if NDIGBUFFER_HAS_GBUFFER
diff --git a/Engine/Plugins/FX/Niagara/Source/Niagara/Private/DataInterface/NiagaraDataInterfaceGBuffer.cpp b/Engine/Plugins/FX/Niagara/Source/Niagara/Private/DataInterface/NiagaraDataInterfaceGBuffer.cpp
--- a/Engine/Plugins/FX/Niagara/Source/Niagara/Private/DataInterface/NiagaraDataInterfaceGBuffer.cpp
+++ b/Engine/Plugins/FX/Niagara/Source/Niagara/Private/DataInterface/NiagaraDataInterfaceGBuffer.cpp
@@ -40,11 +40,12 @@ namespace NiagaraDataInterfaceGBufferLocal
 
 	struct FGBufferAttribute
 	{
-		FGBufferAttribute(const TCHAR* InAttributeName, const TCHAR* InAttributeType, FNiagaraTypeDefinition InTypeDef, FText InDescription)
+		FGBufferAttribute(const TCHAR* InAttributeName, const TCHAR* InAttributeType, FNiagaraTypeDefinition InTypeDef, FText InDescription, uint16 InMiscUsageBitMask)
 			: AttributeName(InAttributeName)
 			, AttributeType(InAttributeType)
 			, TypeDef(InTypeDef)
 			, Description(InDescription)
+			, MiscUsageBitMask(InMiscUsageBitMask)
 		{
 			FString TempName;
 			TempName = TEXT("Decode");
@@ -57,6 +58,7 @@ namespace NiagaraDataInterfaceGBufferLocal
 		FName					ScreenUVFunctionName;
 		FNiagaraTypeDefinition	TypeDef;
 		FText					Description;
+		uint16					MiscUsageBitMask;
 	};
 
 	static FText GetDescription_ScreenVelocity()
@@ -90,23 +92,25 @@ namespace NiagaraDataInterfaceGBufferLocal
 	{
 		static const TArray<FGBufferAttribute> GBufferAttributes =
 		{
-			FGBufferAttribute(TEXT("DiffuseColor"),		TEXT("float3"),	FNiagaraTypeDefinition::GetVec3Def(), FText::GetEmpty()),
-			FGBufferAttribute(TEXT("WorldNormal"),		TEXT("float3"),	FNiagaraTypeDefinition::GetVec3Def(), FText::GetEmpty()),
-			FGBufferAttribute(TEXT("ScreenVelocity"),	TEXT("float3"),	FNiagaraTypeDefinition::GetVec3Def(), GetDescription_ScreenVelocity()),
-			FGBufferAttribute(TEXT("WorldVelocity"),	TEXT("float3"),	FNiagaraTypeDefinition::GetVec3Def(), GetDescription_WorldVelocity()),
-			FGBufferAttribute(TEXT("BaseColor"),		TEXT("float3"),	FNiagaraTypeDefinition::GetVec3Def(), FText::GetEmpty()),
-			//FGBufferAttribute(TEXT("SpecularColor"),	TEXT("float3"),	FNiagaraTypeDefinition::GetVec3Def(), FText::GetEmpty()),
-			FGBufferAttribute(TEXT("Metallic"),			TEXT("float"),	FNiagaraTypeDefinition::GetFloatDef(), FText::GetEmpty()),
-			FGBufferAttribute(TEXT("Specular"),			TEXT("float"),	FNiagaraTypeDefinition::GetFloatDef(), FText::GetEmpty()),
-			FGBufferAttribute(TEXT("Roughness"),		TEXT("float"),	FNiagaraTypeDefinition::GetFloatDef(), FText::GetEmpty()),
-			FGBufferAttribute(TEXT("Depth"),			TEXT("float"),	FNiagaraTypeDefinition::GetFloatDef(), FText::GetEmpty()),
-			//FGBufferAttribute(TEXT("Stencil"),			TEXT("int"),	FNiagaraTypeDefinition::GetIntDef(), FText::GetEmpty()),
-
-			FGBufferAttribute(TEXT("CustomDepth"),		TEXT("float"),	FNiagaraTypeDefinition::GetFloatDef(), FText::GetEmpty()),
-			FGBufferAttribute(TEXT("CustomStencil"),	TEXT("int"),	FNiagaraTypeDefinition::GetIntDef(), FText::GetEmpty()),
-
-			FGBufferAttribute(TEXT("SceneColor"),		TEXT("float4"),	FNiagaraTypeDefinition::GetVec4Def(), GetDescription_SceneColor()),
-			FGBufferAttribute(TEXT("ShadingModelID"),	TEXT("int"),	FNiagaraTypeDefinition::GetIntDef(), FText::GetEmpty()),
+			FGBufferAttribute(TEXT("DiffuseColor"),		TEXT("float3"),	FNiagaraTypeDefinition::GetVec3Def(), FText::GetEmpty(), 0),
+			FGBufferAttribute(TEXT("WorldNormal"),		TEXT("float3"),	FNiagaraTypeDefinition::GetVec3Def(), FText::GetEmpty(), 0),
+			FGBufferAttribute(TEXT("ScreenVelocity"),	TEXT("float3"),	FNiagaraTypeDefinition::GetVec3Def(), GetDescription_ScreenVelocity(), 0),
+			FGBufferAttribute(TEXT("WorldVelocity"),	TEXT("float3"),	FNiagaraTypeDefinition::GetVec3Def(), GetDescription_WorldVelocity(), 0),
+			FGBufferAttribute(TEXT("BaseColor"),		TEXT("float3"),	FNiagaraTypeDefinition::GetVec3Def(), FText::GetEmpty(), 0),
+			//FGBufferAttribute(TEXT("SpecularColor"),	TEXT("float3"),	FNiagaraTypeDefinition::GetVec3Def(), FText::GetEmpty(), 0),
+			FGBufferAttribute(TEXT("Metallic"),			TEXT("float"),	FNiagaraTypeDefinition::GetFloatDef(), FText::GetEmpty(), 0),
+			FGBufferAttribute(TEXT("Specular"),			TEXT("float"),	FNiagaraTypeDefinition::GetFloatDef(), FText::GetEmpty(), 0),
+			FGBufferAttribute(TEXT("Roughness"),		TEXT("float"),	FNiagaraTypeDefinition::GetFloatDef(), FText::GetEmpty(), 0),
+			FGBufferAttribute(TEXT("Depth"),			TEXT("float"),	FNiagaraTypeDefinition::GetFloatDef(), FText::GetEmpty(), 0),
+			//FGBufferAttribute(TEXT("Stencil"),			TEXT("int"),	FNiagaraTypeDefinition::GetIntDef(), FText::GetEmpty(), 0),
+
+			FGBufferAttribute(TEXT("PartialDepth"),		TEXT("float"),	FNiagaraTypeDefinition::GetFloatDef(), FText::GetEmpty(), uint16(ENiagaraScriptMiscUsageMask::UsesPartialDepthCollisionQuery)),
+
+			FGBufferAttribute(TEXT("CustomDepth"),		TEXT("float"),	FNiagaraTypeDefinition::GetFloatDef(), FText::GetEmpty(), 0),
+			FGBufferAttribute(TEXT("CustomStencil"),	TEXT("int"),	FNiagaraTypeDefinition::GetIntDef(), FText::GetEmpty(), 0),
+
+			FGBufferAttribute(TEXT("SceneColor"),		TEXT("float4"),	FNiagaraTypeDefinition::GetVec4Def(), GetDescription_SceneColor(), 0),
+			FGBufferAttribute(TEXT("ShadingModelID"),	TEXT("int"),	FNiagaraTypeDefinition::GetIntDef(), FText::GetEmpty(), 0),
 		};
 
 		return MakeArrayView(GBufferAttributes);
@@ -157,6 +161,7 @@ void UNiagaraDataInterfaceGBuffer::GetFunctionsInternal(TArray<FNiagaraFunctionS
 		Signature.bRequiresContext = false;
 		Signature.bSupportsCPU = false;
 		Signature.bExperimental = true;
+		Signature.MiscUsageBitMask = Attribute.MiscUsageBitMask;
 		Signature.Inputs.Add(FNiagaraVariable(FNiagaraTypeDefinition(GetClass()), TEXT("GBufferInterface")));
 		Signature.Inputs.Add(FNiagaraVariable(FNiagaraTypeDefinition::GetVec2Def(), TEXT("ScreenUV")));
 		Signature.Inputs.Add_GetRef(FNiagaraVariable(FNiagaraTypeDefinition::GetBoolDef(), TEXT("ApplyViewportOffset"))).SetValue(true);
