- Validation rule for using events

diff --git a/Engine/Plugins/FX/Niagara/Source/NiagaraEditor/Private/NiagaraValidationRules.cpp b/Engine/Plugins/FX/Niagara/Source/NiagaraEditor/Private/NiagaraValidationRules.cpp
--- a/Engine/Plugins/FX/Niagara/Source/NiagaraEditor/Private/NiagaraValidationRules.cpp
+++ b/Engine/Plugins/FX/Niagara/Source/NiagaraEditor/Private/NiagaraValidationRules.cpp
@@ -370,6 +370,39 @@ void UNiagaraValidationRule_NoWarmupTime::CheckValidity(const FNiagaraValidation
 	}
 }
 
+void UNiagaraValidationRule_NoEvents::CheckValidity(const FNiagaraValidationContext& Context, TArray<FNiagaraValidationResult>& OutResults)  const
+{
+	TArray<TSharedRef<FNiagaraEmitterHandleViewModel>> EmitterHandleViewModels = Context.ViewModel->GetEmitterHandleViewModels();
+	for (TSharedRef<FNiagaraEmitterHandleViewModel> EmitterHandleModel : EmitterHandleViewModels)
+	{
+		FNiagaraEmitterHandle* EmitterHandle = EmitterHandleModel.Get().GetEmitterHandle();
+		if (!EmitterHandle->GetIsEnabled())
+		{
+			continue;
+		}
+
+		FVersionedNiagaraEmitterData* EmitterData = EmitterHandleModel.Get().GetEmitterHandle()->GetEmitterData();
+		if (EmitterData->GetEventHandlers().Num() == 0)
+		{
+			continue;
+		}
+
+		TArray<FNiagaraPlatformSetConflictInfo> Conflicts = NiagaraValidation::GatherPlatformSetConflicts(&Platforms, &EmitterData->Platforms);
+		if (Conflicts.Num() == 0)
+		{
+			continue;
+		}
+		
+		const FString PlatformConflicts = NiagaraValidation::GetPlatformConflictsString(Conflicts);
+
+		FNiagaraValidationResult& Result = OutResults.AddDefaulted_GetRef();
+		Result.Severity = Severity;
+		Result.SummaryText = LOCTEXT("NoEventsSummary", "Events are not allowed.");
+		Result.Description = FText::Format(LOCTEXT("NoEventsDesc", "Events are not allowed on '{0}'."), FText::FromString(PlatformConflicts));
+		Result.SourceObject = NiagaraValidation::GetStackEntry<UNiagaraStackEmitterPropertiesItem>(EmitterHandleModel.Get().GetEmitterStackViewModel());
+	}
+}
+
 void UNiagaraValidationRule_FixedGPUBoundsSet::CheckValidity(const FNiagaraValidationContext& Context, TArray<FNiagaraValidationResult>& Results)  const
 {
 	// if the system has fixed bounds set then it overrides the emitter settings
diff --git a/Engine/Plugins/FX/Niagara/Source/NiagaraEditor/Public/NiagaraValidationRules.h b/Engine/Plugins/FX/Niagara/Source/NiagaraEditor/Public/NiagaraValidationRules.h
--- a/Engine/Plugins/FX/Niagara/Source/NiagaraEditor/Public/NiagaraValidationRules.h
+++ b/Engine/Plugins/FX/Niagara/Source/NiagaraEditor/Public/NiagaraValidationRules.h
@@ -98,6 +98,22 @@ public:
 	virtual void CheckValidity(const FNiagaraValidationContext& Context, TArray<FNiagaraValidationResult>& OutResults) const override;
 };
 
+/** This validation rule ensures that emitters do not use events. */
+UCLASS(Category = "Validation", DisplayName = "No Events")
+class UNiagaraValidationRule_NoEvents : public UNiagaraValidationRule
+{
+	GENERATED_BODY()
+
+public:
+	UPROPERTY(EditAnywhere, Category = Validation)
+	ENiagaraValidationSeverity Severity = ENiagaraValidationSeverity::Warning;
+
+	UPROPERTY(EditAnywhere, Category = Validation)
+	FNiagaraPlatformSet Platforms;
+
+	virtual void CheckValidity(const FNiagaraValidationContext& Context, TArray<FNiagaraValidationResult>& OutResults) const override;
+};
+
 /** This validation rule ensures that GPU emitters have fixed bounds set. */
 UCLASS(Category = "Validation", DisplayName = "Fixed GPU Bounds Set")
 class UNiagaraValidationRule_FixedGPUBoundsSet : public UNiagaraValidationRule
