Fix for default type being returned when a new type is encountered via the SWC/LWC pipeline

diff --git a/Engine/Plugins/FX/Niagara/Source/Niagara/Private/NiagaraTypeRegistry.cpp b/Engine/Plugins/FX/Niagara/Source/Niagara/Private/NiagaraTypeRegistry.cpp
--- a/Engine/Plugins/FX/Niagara/Source/Niagara/Private/NiagaraTypeRegistry.cpp
+++ b/Engine/Plugins/FX/Niagara/Source/Niagara/Private/NiagaraTypeRegistry.cpp
@@ -653,48 +653,5 @@ FNiagaraTypeDefinition FNiagaraTypeRegistry::GetTypeForStruct(UScriptStruct* InS
 			return false;
 		});
 
-	return FoundType;
+	return FoundType.IsValid() ? FoundType : FNiagaraTypeDefinition(InStruct);
 }
-
-#if 0
-template<typename TAction>
-static void ForEachRegisteredType(TAction Func)
-{
-	FNiagaraTypeRegistry& Registry = Get();
-
-	auto TrivialSelect = [](const FNiagaraTypeDefinition& TypeDef) -> bool
-		{
-			return true;
-		};
-
-	Registry.ForAllSelectedRegisteredTypes(TrivialSelect, Func);
-}
-
-	template<typename TSelectAction, typename TAction>
-	void ForAllSelectedRegisteredTypes(TSelectAction Select, TAction Func) const
-	{
-		TArray<FNiagaraTypeDefinition> LocalCopy;
-
-		{
-			UE::TScopeLock Lock(RegisteredTypesLock);
-			LocalCopy.Reserve(RegisteredTypes.Num());
-			Algo::TransformIf(RegisteredTypes, LocalCopy,
-				[Select](const TUniquePtr<FNiagaraTypeDefinition>& DefPtr) -> bool
-				{
-					return DefPtr.IsValid() && DefPtr->IsValid() && Select(*DefPtr);
-				},
-				[](const TUniquePtr<FNiagaraTypeDefinition>& DefPtr) -> FNiagaraTypeDefinition
-				{
-					return *DefPtr;
-				});
-		}
-
-		for (const FNiagaraTypeDefinition& Type : LocalCopy)
-		{
-			if (!Func(Type))
-			{
-				break;
-			}
-		}
-	}
-	#endif
\ No newline at end of file
