- Fix crash if you delete the desired age track in the sequencer timeline

#12557

diff --git a/Engine/Plugins/FX/Niagara/Source/NiagaraEditor/Private/Sequencer/LevelSequence/NiagaraSystemTrackEditor.cpp b/Engine/Plugins/FX/Niagara/Source/NiagaraEditor/Private/Sequencer/LevelSequence/NiagaraSystemTrackEditor.cpp
--- a/Engine/Plugins/FX/Niagara/Source/NiagaraEditor/Private/Sequencer/LevelSequence/NiagaraSystemTrackEditor.cpp
+++ b/Engine/Plugins/FX/Niagara/Source/NiagaraEditor/Private/Sequencer/LevelSequence/NiagaraSystemTrackEditor.cpp
@@ -154,8 +154,9 @@ TSharedPtr<SWidget> FNiagaraSystemTrackEditor::BuildOutlinerColumnWidget(const F
 				.Font(FAppStyle::GetFontStyle(TEXT("PropertyWindow.NormalFont")))
 				.CurrentValue_Lambda([NiagaraSystemTrack]() -> int32
 				{
-					UMovieSceneNiagaraSystemSpawnSection* SpawnSection = CastChecked<UMovieSceneNiagaraSystemSpawnSection>(NiagaraSystemTrack->GetAllSections()[0]);
-					return static_cast<int32>(SpawnSection->GetAgeUpdateMode());
+					const TArray<UMovieSceneSection*>& Sections = NiagaraSystemTrack->GetAllSections();
+					UMovieSceneNiagaraSystemSpawnSection* Section = Sections.Num() == 1 ? Cast<UMovieSceneNiagaraSystemSpawnSection>(Sections[0]) : nullptr;
+					return static_cast<int32>(Section ? Section->GetAgeUpdateMode() : ENiagaraAgeUpdateMode::TickDeltaTime);
 				})
 				.OnEnumSelectionChanged_Lambda([NiagaraSystemTrack](int32 NewSelection, ESelectInfo::Type)
 				{
@@ -164,10 +165,14 @@ TSharedPtr<SWidget> FNiagaraSystemTrackEditor::BuildOutlinerColumnWidget(const F
 						return;
 					}
 
-					FScopedTransaction Transaction(LOCTEXT("ChangeValue", "Change Value"));
-					UMovieSceneNiagaraSystemSpawnSection* SpawnSection = CastChecked<UMovieSceneNiagaraSystemSpawnSection>(NiagaraSystemTrack->GetAllSections()[0]);
-					SpawnSection->Modify();
-					SpawnSection->SetAgeUpdateMode( static_cast<ENiagaraAgeUpdateMode>(NewSelection));
+					TArray<UMovieSceneSection*> Sections = NiagaraSystemTrack->GetAllSections();
+					UMovieSceneNiagaraSystemSpawnSection* Section = Sections.Num() == 1 ? Cast<UMovieSceneNiagaraSystemSpawnSection>(Sections[0]) : nullptr;
+					if (Section)
+					{
+						FScopedTransaction Transaction(LOCTEXT("ChangeValue", "Change Value"));
+						Section->Modify();
+						Section->SetAgeUpdateMode(static_cast<ENiagaraAgeUpdateMode>(NewSelection));
+					}
 				})
 				.IsEnabled_Lambda([NiagaraSystemTrack]() -> bool
 				{
