- Fix a crash in the editor during EOF updates if a mesh renderer references a user parameter and the mesh is still compiling

diff --git a/Engine/Plugins/FX/Niagara/Source/Niagara/Private/NiagaraMeshRendererProperties.cpp b/Engine/Plugins/FX/Niagara/Source/Niagara/Private/NiagaraMeshRendererProperties.cpp
--- a/Engine/Plugins/FX/Niagara/Source/Niagara/Private/NiagaraMeshRendererProperties.cpp
+++ b/Engine/Plugins/FX/Niagara/Source/Niagara/Private/NiagaraMeshRendererProperties.cpp
@@ -198,6 +198,28 @@ public:
 
 namespace NiagaraMeshRendererPropertiesInternal
 {
+	bool IsStaticMeshSafeToUse(const FNiagaraEmitterInstance* EmitterInstance, UStaticMesh* StaticMesh)
+	{
+		if (StaticMesh)
+		{
+			bool bInvalidDueToCompiling = false;
+		#if WITH_EDITOR
+			// During EOF updates we can not test GetRenderData as we will cause a wait / reregister to occur which is invalid while in EOF updates
+			// The assumption is that when we do hit this situation we are waiting on a static mesh build from reimport, etc, so we can just skip until
+			// we get the post build callback, ideally we would have a better way to handle this by not rendering while the build is in progress but
+			// that requires quite a large rework of how we handle reading from the static mesh data, especially around updating MICs in PostLoad.
+			if (StaticMesh->IsCompiling())
+			{
+				FNiagaraSystemInstance* SystemInstance = EmitterInstance ? EmitterInstance->GetParentSystemInstance() : nullptr;
+				UWorld* World = SystemInstance ? SystemInstance->GetWorld() : nullptr;
+				bInvalidDueToCompiling = World && World->bPostTickComponentUpdate;
+			}
+		#endif
+			return !bInvalidDueToCompiling && StaticMesh->GetRenderData();
+		}
+		return false;
+	}
+
 	void ResolveRenderableMeshInternal(const FNiagaraMeshRendererMeshProperties& MeshProperties, const FNiagaraEmitterInstance* EmitterInstance, INiagaraRenderableMeshInterface*& OutInterface, UStaticMesh*& OutStaticMesh)
 	{
 		OutInterface = nullptr;
@@ -217,32 +239,16 @@ namespace NiagaraMeshRendererPropertiesInternal
 				}
 
 				UStaticMesh* BoundMesh = Cast<UStaticMesh>(EmitterInstance->GetRendererBoundVariables().GetUObject(MeshParameter));
-				if (BoundMesh && BoundMesh->GetRenderData())
+				if (IsStaticMeshSafeToUse(EmitterInstance, BoundMesh))
 				{
 					OutStaticMesh = BoundMesh;
 					return;
 				}
 			}
 		}
-		if (MeshProperties.Mesh)
+		if (IsStaticMeshSafeToUse(EmitterInstance, MeshProperties.Mesh))
 		{
-			bool bInvalidForAsynCompiling = false;
-			#if WITH_EDITOR
-				// During EOF updates we can not test GetRenderData as we will cause a wait / reregister to occur which is invalid while in EOF updates
-				// The assumption is that when we do hit this situation we are waiting on a static mesh build from reimport, etc, so we can just skip until
-				// we get the post build callback, ideally we would have a better way to handle this by not rendering while the build is in progress but
-				// that requires quite a large rework of how we handle reading from the static mesh data, especially around updating MICs in PostLoad.
-				if (MeshProperties.Mesh->IsCompiling())
-				{
-					FNiagaraSystemInstance* SystemInstance = EmitterInstance ? EmitterInstance->GetParentSystemInstance() : nullptr;
-					UWorld* World = SystemInstance ? SystemInstance->GetWorld() : nullptr;
-					bInvalidForAsynCompiling = World && World->bPostTickComponentUpdate;
-				}
-			#endif
-			if (!bInvalidForAsynCompiling && MeshProperties.Mesh->GetRenderData())
-			{
-				OutStaticMesh = MeshProperties.Mesh;
-			}
+			OutStaticMesh = MeshProperties.Mesh;
 		}
 	}
 }
