PR #12558: Bug: UMovieSceneNiagaraTrack::IsEmpty() returns true when non-empty

diff --git a/Engine/Plugins/FX/Niagara/Source/Niagara/Private/MovieScene/MovieSceneNiagaraTrack.cpp b/Engine/Plugins/FX/Niagara/Source/Niagara/Private/MovieScene/MovieSceneNiagaraTrack.cpp
--- a/Engine/Plugins/FX/Niagara/Source/Niagara/Private/MovieScene/MovieSceneNiagaraTrack.cpp
+++ b/Engine/Plugins/FX/Niagara/Source/Niagara/Private/MovieScene/MovieSceneNiagaraTrack.cpp
@@ -21,7 +21,7 @@ bool UMovieSceneNiagaraTrack::HasSection(const UMovieSceneSection& Section) cons
 
 bool UMovieSceneNiagaraTrack::IsEmpty() const
 {
-	return Sections.Num() != 0;
+	return Sections.Num() == 0;
 }
 
 void UMovieSceneNiagaraTrack::RemoveAllAnimationData()
