- Add ability to opt out of hit proxy to sprite / mesh / ribbon renderer

diff --git a/Engine/Plugins/FX/Niagara/Source/Niagara/Private/NiagaraRendererMeshes.cpp b/Engine/Plugins/FX/Niagara/Source/Niagara/Private/NiagaraRendererMeshes.cpp
--- a/Engine/Plugins/FX/Niagara/Source/Niagara/Private/NiagaraRendererMeshes.cpp
+++ b/Engine/Plugins/FX/Niagara/Source/Niagara/Private/NiagaraRendererMeshes.cpp
@@ -121,6 +121,9 @@ FNiagaraRendererMeshes::FNiagaraRendererMeshes(ERHIFeatureLevel::Type FeatureLev
 	RendererVisibility = Properties->RendererVisibility;
 	bAccurateMotionVectors = Properties->NeedsPreciseMotionVectors();
 	bCastShadows = Properties->bCastShadows;
+#if WITH_EDITORONLY_DATA
+	bIncludeInHitProxy = Properties->bIncludeInHitProxy;
+#endif
 
 	if (Properties->bEnableCameraDistanceCulling)
 	{
@@ -1271,6 +1274,12 @@ void FNiagaraRendererMeshes::CreateMeshBatchForSection(
 	MeshBatch.CastRayTracedShadow = SceneProxy.CastsDynamicShadow() && bCastShadows && Section.bCastShadow;
 #endif
 	MeshBatch.DepthPriorityGroup = (ESceneDepthPriorityGroup)SceneProxy.GetDepthPriorityGroup(&View);
+#if WITH_EDITORONLY_DATA
+	if (bIncludeInHitProxy == false)
+	{
+		MeshBatch.BatchHitProxyId = FHitProxyId::InvisibleHitProxyId;
+	}
+#endif
 
 	const FBox LocalBounds = ScaleFBox(MeshData.RenderableMesh->GetLocalBounds(), FVector3f::OneVector, MeshBoundsScale);
 
diff --git a/Engine/Plugins/FX/Niagara/Source/Niagara/Private/NiagaraRendererRibbons.cpp b/Engine/Plugins/FX/Niagara/Source/Niagara/Private/NiagaraRendererRibbons.cpp
--- a/Engine/Plugins/FX/Niagara/Source/Niagara/Private/NiagaraRendererRibbons.cpp
+++ b/Engine/Plugins/FX/Niagara/Source/Niagara/Private/NiagaraRendererRibbons.cpp
@@ -741,6 +741,9 @@ FNiagaraRendererRibbons::FNiagaraRendererRibbons(ERHIFeatureLevel::Type FeatureL
 	DrawDirection = Properties->DrawDirection;
 	RendererLayout = &Properties->RendererLayout;
 	bCastShadows = Properties->bCastShadows;
+#if WITH_EDITORONLY_DATA
+	bIncludeInHitProxy = Properties->bIncludeInHitProxy;
+#endif
 	bUseGeometryNormals = Properties->bUseGeometryNormals;
 	bGpuRibbonLinkIsFloat = Properties->bGpuRibbonLinkIsFloat;
 	GpuRibbonLinkOrderOffset = Properties->GpuRibbonLinkOrderOffset;
@@ -913,7 +916,13 @@ void FNiagaraRendererRibbons::GetDynamicMeshElements(const TArray<const FSceneVi
 			}
 
 			FMeshBatch& MeshBatch = Collector.AllocateMesh();
-	
+		#if WITH_EDITORONLY_DATA
+			if (bIncludeInHitProxy == false)
+			{
+				MeshBatch.BatchHitProxyId = FHitProxyId::InvisibleHitProxyId;
+			}
+		#endif
+
 			const FVector ViewOriginForDistanceCulling = View->ViewMatrices.GetViewOrigin();
 
 			FNiagaraRibbonRenderingFrameViewResources* RenderingViewResources = RenderingResources.RibbonResources->ViewResources.Add_GetRef(MakeShared<FNiagaraRibbonRenderingFrameViewResources>()).Get();
diff --git a/Engine/Plugins/FX/Niagara/Source/Niagara/Private/NiagaraRendererSprites.cpp b/Engine/Plugins/FX/Niagara/Source/Niagara/Private/NiagaraRendererSprites.cpp
--- a/Engine/Plugins/FX/Niagara/Source/Niagara/Private/NiagaraRendererSprites.cpp
+++ b/Engine/Plugins/FX/Niagara/Source/Niagara/Private/NiagaraRendererSprites.cpp
@@ -103,6 +103,9 @@ FNiagaraRendererSprites::FNiagaraRendererSprites(ERHIFeatureLevel::Type FeatureL
 	RendererVisibility = Properties->RendererVisibility;
 	bAccurateMotionVectors = Properties->NeedsPreciseMotionVectors();
 	bCastShadows = Properties->bCastShadows;
+#if WITH_EDITORONLY_DATA
+	bIncludeInHitProxy = Properties->bIncludeInHitProxy;
+#endif
 
 	PixelCoverageMode = Properties->PixelCoverageMode;
 	if (PixelCoverageMode == ENiagaraRendererPixelCoverageMode::Automatic)
@@ -262,7 +265,6 @@ void FNiagaraRendererSprites::PrepareParticleSpriteRenderData(FRHICommandListBas
 		bLowLatencyTranslucencyEnabled &= !Material.IsTranslucencyWritingCustomDepth();
 	}
 
-
 	ParticleSpriteRenderData.SourceParticleData = ParticleSpriteRenderData.DynamicDataSprites->GetParticleDataToRender(RHICmdList, bLowLatencyTranslucencyEnabled);
 	if ( !ParticleSpriteRenderData.SourceParticleData || (SourceMode == ENiagaraRendererSourceDataMode::Particles && ParticleSpriteRenderData.SourceParticleData->GetNumInstances() == 0) )
 	{
@@ -915,6 +917,13 @@ void FNiagaraRendererSprites::CreateMeshBatchForView(
 	MeshBatch.bUseWireframeSelectionColoring = SceneProxy.IsSelected();
 	MeshBatch.SegmentIndex = 0;
 
+#if WITH_EDITORONLY_DATA
+	if (bIncludeInHitProxy == false)
+	{
+		MeshBatch.BatchHitProxyId = FHitProxyId::InvisibleHitProxyId;
+	}
+#endif
+
 	const bool bIsWireframe = View.Family->EngineShowFlags.Wireframe;
 	if (bIsWireframe)
 	{
diff --git a/Engine/Plugins/FX/Niagara/Source/Niagara/Public/NiagaraMeshRendererProperties.h b/Engine/Plugins/FX/Niagara/Source/Niagara/Public/NiagaraMeshRendererProperties.h
--- a/Engine/Plugins/FX/Niagara/Source/Niagara/Public/NiagaraMeshRendererProperties.h
+++ b/Engine/Plugins/FX/Niagara/Source/Niagara/Public/NiagaraMeshRendererProperties.h
@@ -362,6 +362,14 @@ public:
 	UPROPERTY(EditAnywhere, Category = "Rendering")
 	uint8 bCastShadows : 1 = 1; //-V570
 
+#if WITH_EDITORONLY_DATA
+	/**
+	When disabled the mesh batches do not contribute to the hit proxy, i.e. they are not part of editor selection.
+	*/
+	UPROPERTY(EditAnywhere, Category = "Rendering")
+	uint8 bIncludeInHitProxy : 1 = 1;
+#endif
+
 	/** The materials to be used instead of the StaticMesh's materials. Note that each material must have the Niagara Mesh Particles flag checked. If the ParticleMesh
 	requires more materials than exist in this array or any entry in this array is set to None, we will use the ParticleMesh's existing Material instead.*/
 	UPROPERTY(EditAnywhere, Category = "Mesh Rendering", meta = (EditCondition = "bOverrideMaterials", EditConditionHides))
diff --git a/Engine/Plugins/FX/Niagara/Source/Niagara/Public/NiagaraRendererMeshes.h b/Engine/Plugins/FX/Niagara/Source/Niagara/Public/NiagaraRendererMeshes.h
--- a/Engine/Plugins/FX/Niagara/Source/Niagara/Public/NiagaraRendererMeshes.h
+++ b/Engine/Plugins/FX/Niagara/Source/Niagara/Public/NiagaraRendererMeshes.h
@@ -188,6 +188,9 @@ private:
 	uint32 bEnableLODCulling : 1;
 	uint32 bAccurateMotionVectors : 1;
 	uint32 bCastShadows : 1;
+#if WITH_EDITORONLY_DATA
+	uint32 bIncludeInHitProxy : 1 = true;
+#endif
 
 	uint32 bSubImageBlend : 1;
 	FVector2f SubImageSize = FVector2f::ZeroVector;
diff --git a/Engine/Plugins/FX/Niagara/Source/Niagara/Public/NiagaraRendererRibbons.h b/Engine/Plugins/FX/Niagara/Source/Niagara/Public/NiagaraRendererRibbons.h
--- a/Engine/Plugins/FX/Niagara/Source/Niagara/Public/NiagaraRendererRibbons.h
+++ b/Engine/Plugins/FX/Niagara/Source/Niagara/Public/NiagaraRendererRibbons.h
@@ -390,9 +390,12 @@ protected:
 	ENiagaraRibbonDrawDirection DrawDirection;
 	
 	const FNiagaraRendererLayout* RendererLayout;
-	bool bCastShadows = true;
-	bool bUseGeometryNormals = false;
-	bool bGpuRibbonLinkIsFloat = false;
+	uint32 bCastShadows : 1 = true;
+#if WITH_EDITORONLY_DATA
+	uint32 bIncludeInHitProxy : 1 = true;
+#endif
+	uint32 bUseGeometryNormals : 1 = false;
+	uint32 bGpuRibbonLinkIsFloat : 1 = false;
 	uint32 GpuRibbonLinkOrderOffset = INDEX_NONE;
 
 	mutable UE::FMutex VertexBuffersGuard;
diff --git a/Engine/Plugins/FX/Niagara/Source/Niagara/Public/NiagaraRendererSprites.h b/Engine/Plugins/FX/Niagara/Source/Niagara/Public/NiagaraRendererSprites.h
--- a/Engine/Plugins/FX/Niagara/Source/Niagara/Public/NiagaraRendererSprites.h
+++ b/Engine/Plugins/FX/Niagara/Source/Niagara/Public/NiagaraRendererSprites.h
@@ -110,6 +110,9 @@ private:
 	uint32 bEnableDistanceCulling : 1;
 	uint32 bAccurateMotionVectors : 1;
 	uint32 bCastShadows : 1;
+#if WITH_EDITORONLY_DATA
+	uint32 bIncludeInHitProxy : 1 = true;
+#endif
 	uint32 bSetAnyBoundVars : 1;
 	uint32 bVisTagInParamStore : 1;
 
diff --git a/Engine/Plugins/FX/Niagara/Source/Niagara/Public/NiagaraRibbonRendererProperties.h b/Engine/Plugins/FX/Niagara/Source/Niagara/Public/NiagaraRibbonRendererProperties.h
--- a/Engine/Plugins/FX/Niagara/Source/Niagara/Public/NiagaraRibbonRendererProperties.h
+++ b/Engine/Plugins/FX/Niagara/Source/Niagara/Public/NiagaraRibbonRendererProperties.h
@@ -338,6 +338,14 @@ public:
 	UPROPERTY(EditAnywhere, Category = "Rendering")
 	uint8 bCastShadows : 1 = 1; //-V570
 
+#if WITH_EDITORONLY_DATA
+	/**
+	When disabled the mesh batches do not contribute to the hit proxy, i.e. they are not part of editor selection.
+	*/
+	UPROPERTY(EditAnywhere, Category = "Rendering")
+	uint8 bIncludeInHitProxy : 1 = 1;
+#endif
+
 	/** Tessellation factor to apply to the width of the ribbon.
 	* Ranges from 1 to 16. Greater values increase amount of tessellation.
 	*/
diff --git a/Engine/Plugins/FX/Niagara/Source/Niagara/Public/NiagaraSpriteRendererProperties.h b/Engine/Plugins/FX/Niagara/Source/Niagara/Public/NiagaraSpriteRendererProperties.h
--- a/Engine/Plugins/FX/Niagara/Source/Niagara/Public/NiagaraSpriteRendererProperties.h
+++ b/Engine/Plugins/FX/Niagara/Source/Niagara/Public/NiagaraSpriteRendererProperties.h
@@ -221,6 +221,14 @@ public:
 	UPROPERTY(EditAnywhere, Category = "Rendering")
 	uint8 bCastShadows : 1 = 1; //-V570
 
+#if WITH_EDITORONLY_DATA
+	/**
+	When disabled the mesh batches do not contribute to the hit proxy, i.e. they are not part of editor selection.
+	*/
+	UPROPERTY(EditAnywhere, Category = "Rendering")
+	uint8 bIncludeInHitProxy : 1 = 1;
+#endif
+
 	/** Sort precision to use when sorting is active. */
 	UPROPERTY(EditAnywhere, Category = "Sorting", meta = (EditCondition = "SortMode != ENiagaraSortMode::None", EditConditionHides))
 	ENiagaraRendererSortPrecision SortPrecision = ENiagaraRendererSortPrecision::Default;
