[Viewport Toolbar] Niagara Editor - Update preview scene settings menu

diff --git a/Engine/Plugins/FX/Niagara/Source/NiagaraEditor/Private/Widgets/NiagaraSystemEditorViewportToolbarSections.cpp b/Engine/Plugins/FX/Niagara/Source/NiagaraEditor/Private/Widgets/NiagaraSystemEditorViewportToolbarSections.cpp
--- a/Engine/Plugins/FX/Niagara/Source/NiagaraEditor/Private/Widgets/NiagaraSystemEditorViewportToolbarSections.cpp
+++ b/Engine/Plugins/FX/Niagara/Source/NiagaraEditor/Private/Widgets/NiagaraSystemEditorViewportToolbarSections.cpp
@@ -68,8 +68,6 @@ void UE::NiagaraSystemEditor::FillShowSubmenu(UToolMenu* InMenu, bool bInShowVie
 	FToolMenuSection& CommonShowFlagsSection =
 		InMenu->FindOrAddSection("CommonShowFlags", LOCTEXT("CommonShowFlagsLabel", "Common Show Flags"));
 
-	CommonShowFlagsSection.AddMenuEntry(Commands.TogglePreviewGrid);
-	CommonShowFlagsSection.AddSeparator(NAME_None);
 	CommonShowFlagsSection.AddMenuEntry(Commands.ToggleEmitterExecutionOrder);
 	CommonShowFlagsSection.AddMenuEntry(Commands.ToggleGpuTickInformation);
 	CommonShowFlagsSection.AddMenuEntry(Commands.ToggleInstructionCounts);
@@ -78,27 +76,6 @@ void UE::NiagaraSystemEditor::FillShowSubmenu(UToolMenu* InMenu, bool bInShowVie
 	CommonShowFlagsSection.AddMenuEntry(Commands.ToggleStatelessInfo);
 }
 
-FToolMenuEntry UE::NiagaraSystemEditor::CreateSettingsSubmenu()
-{
-	FToolMenuEntry Entry = FToolMenuEntry::InitSubMenu(
-		"Settings",
-		LOCTEXT("SettingsSubmenuLabel", "Settings"),
-		LOCTEXT("SettingsSubmenuTooltip", "Settings options"),
-		FNewToolMenuDelegate::CreateLambda(
-			[](UToolMenu* Submenu) -> void
-			{
-				UE::NiagaraSystemEditor::FillSettingsSubmenu(Submenu);
-			}
-		)
-	);
-	Entry.Icon = FSlateIcon(FAppStyle::GetAppStyleSetName(), "LevelEditor.GameSettings");
-	Entry.InsertPosition.Position = EToolMenuInsertType::Last;
-	Entry.ToolBarData.LabelOverride = FText();
-	Entry.ToolBarData.ResizeParams.ClippingPriority = 1000;
-
-	return Entry;
-}
-
 TSharedRef<SWidget> UE::NiagaraSystemEditor::CreateMotionMenuWidget(const TSharedRef<SNiagaraSystemViewport>& InNiagaraSystemEditorViewport
 )
 {
@@ -203,4 +180,20 @@ void UE::NiagaraSystemEditor::AddMotionSettingsToSection(FToolMenuSection& InSec
 	);
 }
 
+void UE::NiagaraSystemEditor::ExtendPreviewSceneSettingsSubmenu(FName InSubmenuName)
+{
+	UToolMenu* const Submenu = UToolMenus::Get()->ExtendMenu(InSubmenuName);
+	if (!Submenu)
+	{
+		return;
+	}
+
+	FToolMenuInsert AudioInsertPosition("AssetViewerProfileSelectionSection", EToolMenuInsertType::Before);
+	FToolMenuSection& PreviewControlsSection = Submenu->FindOrAddSection(
+		"AssetViewerPreviewControlsSection", LOCTEXT("AssetViewerPreviewControlsSectionLabel", "Preview Controls"), AudioInsertPosition
+	);
+
+	UE::NiagaraSystemEditor::AddMotionSettingsToSection(PreviewControlsSection);
+}
+
 #undef LOCTEXT_NAMESPACE
diff --git a/Engine/Plugins/FX/Niagara/Source/NiagaraEditor/Private/Widgets/NiagaraSystemEditorViewportToolbarSections.h b/Engine/Plugins/FX/Niagara/Source/NiagaraEditor/Private/Widgets/NiagaraSystemEditorViewportToolbarSections.h
--- a/Engine/Plugins/FX/Niagara/Source/NiagaraEditor/Private/Widgets/NiagaraSystemEditorViewportToolbarSections.h
+++ b/Engine/Plugins/FX/Niagara/Source/NiagaraEditor/Private/Widgets/NiagaraSystemEditorViewportToolbarSections.h
@@ -13,7 +13,7 @@ FToolMenuEntry CreateShowSubmenu();
 void FillShowSubmenu(UToolMenu* InMenu, bool bInShowViewportStatsToggle = true);
 
 TSharedRef<SWidget> CreateMotionMenuWidget(const TSharedRef<SNiagaraSystemViewport>& InNiagaraSystemEditorViewport);
-FToolMenuEntry CreateSettingsSubmenu();
 void FillSettingsSubmenu(UToolMenu* InMenu);
 void AddMotionSettingsToSection(FToolMenuSection& InSection);
+void ExtendPreviewSceneSettingsSubmenu(FName InSubmenuName);
 }
diff --git a/Engine/Plugins/FX/Niagara/Source/NiagaraEditor/Private/Widgets/SNiagaraSystemViewport.cpp b/Engine/Plugins/FX/Niagara/Source/NiagaraEditor/Private/Widgets/SNiagaraSystemViewport.cpp
--- a/Engine/Plugins/FX/Niagara/Source/NiagaraEditor/Private/Widgets/SNiagaraSystemViewport.cpp
+++ b/Engine/Plugins/FX/Niagara/Source/NiagaraEditor/Private/Widgets/SNiagaraSystemViewport.cpp
@@ -1276,8 +1276,15 @@ TSharedPtr<SWidget> SNiagaraSystemViewport::BuildViewportToolbar()
 
 			RightSection.AddEntry(UE::NiagaraSystemEditor::CreateShowSubmenu());
 			RightSection.AddEntry(UE::UnrealEd::CreatePerformanceAndScalabilitySubmenu());
-			RightSection.AddEntry(UE::NiagaraSystemEditor::CreateSettingsSubmenu());
-			RightSection.AddEntry(UE::UnrealEd::CreateAssetViewerProfileSubmenu());
+
+			// Add Preview Scene Submenu
+			{
+				UE::NiagaraSystemEditor::ExtendPreviewSceneSettingsSubmenu(
+					"NiagaraSystemEditor.ViewportToolbar.AssetViewerProfile"
+				);
+
+				RightSection.AddEntry(UE::UnrealEd::CreateAssetViewerProfileSubmenu());
+			}
 		}
 	}
 
@@ -1310,6 +1317,7 @@ TSharedPtr<SWidget> SNiagaraSystemViewport::BuildViewportToolbar()
 				}
 			);
 
+			ContextObject->bShowGridToggleInPreviewSceneSettings = true;
 			ViewportToolbarContext.AddObject(ContextObject);
 		}
 	}
