Niagara - Update the niagara clipboard so that non-niagara copies are handled as generic portable values.  This allows pasting from external details panels into module inputs or objects displayed in the stack.

diff --git a/Engine/Plugins/FX/Niagara/Source/NiagaraEditor/Private/NiagaraClipboard.cpp b/Engine/Plugins/FX/Niagara/Source/NiagaraEditor/Private/NiagaraClipboard.cpp
--- a/Engine/Plugins/FX/Niagara/Source/NiagaraEditor/Private/NiagaraClipboard.cpp
+++ b/Engine/Plugins/FX/Niagara/Source/NiagaraEditor/Private/NiagaraClipboard.cpp
@@ -366,6 +366,16 @@ const UNiagaraClipboardContent* FNiagaraClipboard::GetClipboardContent() const
 		return ClipboardContentFactory.ClipboardContent;
 	}
 
+	if (ClipboardText.IsEmpty() == false)
+	{
+		// If the clipboard text wasn't a niagara clipboard object, it's likely been copied from elsewhere in the editor and may be 
+		// a valid portable value so we consturct an empty clipboard object and set the portable data.
+		UNiagaraClipboardContent* ClipboardContent = UNiagaraClipboardContent::Create();
+		FNiagaraClipboardPortableValue& PortableValue = ClipboardContent->PortableValues.AddDefaulted_GetRef();
+		PortableValue.ValueString = ClipboardText;
+		return ClipboardContent;
+	}
+
 	return nullptr;
 }
 
