Implemented Sim Cached Writes for NDC GPU Writes.
Limited scope and risk changes that improve NDC debugability.

diff --git a/Engine/Plugins/FX/Niagara/Source/Niagara/Internal/DataInterface/NiagaraDataInterfaceDataChannelWrite.h b/Engine/Plugins/FX/Niagara/Source/Niagara/Internal/DataInterface/NiagaraDataInterfaceDataChannelWrite.h
--- a/Engine/Plugins/FX/Niagara/Source/Niagara/Internal/DataInterface/NiagaraDataInterfaceDataChannelWrite.h
+++ b/Engine/Plugins/FX/Niagara/Source/Niagara/Internal/DataInterface/NiagaraDataInterfaceDataChannelWrite.h
@@ -22,6 +22,7 @@ This avoids cumbersome work in the graph to access data but requires special han
 #include "NiagaraDataInterface.h"
 #include "NiagaraSimCacheCustomStorageInterface.h"
 #include "NiagaraDataInterfaceRW.h"
+#include "NiagaraDataChannel.h"
 #include "NiagaraDataInterfaceDataChannelWrite.generated.h"
 
 /** Additional compile time information used by the Write DI. */
@@ -85,6 +86,11 @@ public:
 
 	UPROPERTY()
 	FSoftObjectPath DataChannelReference;
+
+	UPROPERTY(Transient)
+	TObjectPtr<const UNiagaraDataInterfaceDataChannelWrite> DataInterface = nullptr;
+
+	FNiagaraSystemInstanceID InstanceID; 
 };
 
 /**
@@ -203,6 +209,7 @@ public:
 	//sim cache functions
 	virtual UObject* SimCacheBeginWrite(UObject* SimCache, FNiagaraSystemInstance* NiagaraSystemInstance, const void* OptionalPerInstanceData, FNiagaraSimCacheFeedbackContext& FeedbackContext) const override;
 	virtual bool SimCacheWriteFrame(UObject* StorageObject, int FrameIndex, FNiagaraSystemInstance* SystemInstance, const void* OptionalPerInstanceData, FNiagaraSimCacheFeedbackContext& FeedbackContext) const override;
+	virtual bool SimCacheEndWrite(UObject* StorageObject) const override;
 	virtual bool SimCacheReadFrame(UObject* StorageObject, int FrameA, int FrameB, float Interp, FNiagaraSystemInstance* SystemInstance, void* OptionalPerInstanceData) override;
 	virtual void SimCachePostReadFrame(void* OptionalPerInstanceData, FNiagaraSystemInstance* SystemInstance) override;
 	virtual bool SimCacheCompareFrame(const UObject* LhsStorageObject, const UObject* RhsStorageObject, int FrameIndex, TOptional<float> Tolerance, FString& OutErrors) const override;
@@ -260,9 +267,11 @@ struct FNiagaraDataInterfaceProxy_DataChannelWrite : public FNiagaraDataInterfac
 		bool bPublishToCPU = false;
 		bool bPublishToGPU = false;
 
-
 		int32 AllocationCount = 0;
 		FVector3f LwcTile;
+
+		mutable bool bCapturingSimCache = false;
+		mutable TArray<FNiagaraDataBufferRef> PendingSimCacheReadbacks;
 	};
 
 	TMap<FNiagaraSystemInstanceID, FInstanceData> SystemInstancesToProxyData_RT;
diff --git a/Engine/Plugins/FX/Niagara/Source/Niagara/Private/DataInterface/NiagaraDataInterfaceDataChannelWrite.cpp b/Engine/Plugins/FX/Niagara/Source/Niagara/Private/DataInterface/NiagaraDataInterfaceDataChannelWrite.cpp
--- a/Engine/Plugins/FX/Niagara/Source/Niagara/Private/DataInterface/NiagaraDataInterfaceDataChannelWrite.cpp
+++ b/Engine/Plugins/FX/Niagara/Source/Niagara/Private/DataInterface/NiagaraDataInterfaceDataChannelWrite.cpp
@@ -29,6 +29,9 @@
 
 #include "NiagaraDataInterfaceUtilities.h"
 
+#include "NiagaraDataSetReadback.h"
+#include "NiagaraGpuReadbackManager.h"
+
 #if WITH_EDITOR
 #include "INiagaraEditorOnlyDataUtlities.h"
 #include "Modules/ModuleManager.h"
@@ -765,7 +768,56 @@ bool UNiagaraDataInterfaceDataChannelWrite::Equals(const UNiagaraDataInterface*
 
 UObject* UNiagaraDataInterfaceDataChannelWrite::SimCacheBeginWrite(UObject* SimCache, FNiagaraSystemInstance* NiagaraSystemInstance, const void* OptionalPerInstanceData, FNiagaraSimCacheFeedbackContext& FeedbackContext) const
 {
-	return NewObject<UNDIDataChannelWriteSimCacheData>(SimCache);
+	UNDIDataChannelWriteSimCacheData* NewSimCacheStorage = NewObject<UNDIDataChannelWriteSimCacheData>(SimCache);
+	NewSimCacheStorage->DataInterface = this;
+	NewSimCacheStorage->InstanceID = NiagaraSystemInstance->GetId();
+
+	const FNiagaraDataInterfaceProxy_DataChannelWrite* RT_Proxy = GetProxyAs<FNiagaraDataInterfaceProxy_DataChannelWrite>();
+	if(RT_Proxy)
+	{
+		ENQUEUE_RENDER_COMMAND(NDISimCacheGPUBeginWrite)
+			(
+				[RT_Proxy, RT_InstanceID = NiagaraSystemInstance->GetId()](FRHICommandListImmediate& RHICmdList)
+				{
+					const FNiagaraDataInterfaceProxy_DataChannelWrite::FInstanceData* InstanceData = RT_Proxy->SystemInstancesToProxyData_RT.Find(RT_InstanceID);
+					if (InstanceData)
+					{
+						InstanceData->bCapturingSimCache = true;
+					}
+				}
+			);
+	}
+	return NewSimCacheStorage;
+}
+
+bool UNiagaraDataInterfaceDataChannelWrite::SimCacheEndWrite(UObject* StorageObject) const
+{
+	if (UNDIDataChannelWriteSimCacheData* Storage = Cast<UNDIDataChannelWriteSimCacheData>(StorageObject))
+	{
+		if(Storage->DataInterface)
+		{
+			const FNiagaraDataInterfaceProxy_DataChannelWrite* RT_Proxy = Storage->DataInterface->GetProxyAs<FNiagaraDataInterfaceProxy_DataChannelWrite>();
+			if(RT_Proxy)
+			{
+				ENQUEUE_RENDER_COMMAND(NDISimCacheGPUBeginWrite)
+				(
+					[RT_Proxy, RT_InstanceID = Storage->InstanceID](FRHICommandListImmediate& RHICmdList)
+					{
+						const FNiagaraDataInterfaceProxy_DataChannelWrite::FInstanceData* InstanceData = RT_Proxy->SystemInstancesToProxyData_RT.Find(RT_InstanceID);
+						if (InstanceData)
+						{
+							InstanceData->bCapturingSimCache = false;
+						}
+					}
+				);
+			}
+		}
+
+		Storage->DataInterface = nullptr;
+		Storage->InstanceID = 0;
+	}
+
+	return true;
 }
 
 bool UNiagaraDataInterfaceDataChannelWrite::SimCacheWriteFrame(UObject* StorageObject, int FrameIndex, FNiagaraSystemInstance* SystemInstance, const void* OptionalPerInstanceData, FNiagaraSimCacheFeedbackContext& FeedbackContext) const
@@ -780,20 +832,74 @@ bool UNiagaraDataInterfaceDataChannelWrite::SimCacheWriteFrame(UObject* StorageO
 		FeedbackContext.Errors.Add(TEXT("Missing per instance data for data channel writer DI"));
 		return false;
 	}
-	// put data from instance data into the storage object
+
+	UWorld* World = SystemInstance->GetWorld();
+	if(World == nullptr)
+	{
+		FeedbackContext.Errors.Add(TEXT("Missing world for data channel writer DI's System Instace."));
+		return false;
+	}
+
+
 	const FNDIDataChannelWriteInstanceData* InstanceData = static_cast<const FNDIDataChannelWriteInstanceData*>(OptionalPerInstanceData);
 	if (UNDIDataChannelWriteSimCacheData* Storage = Cast<UNDIDataChannelWriteSimCacheData>(StorageObject))
 	{
 		ensure(Storage->FrameData.Num() == FrameIndex);
 		Storage->DataChannelReference = Channel.Get();
 		FNDIDataChannelWriteSimCacheFrame& FrameData = Storage->FrameData.AddDefaulted_GetRef();
-		
-		if (InstanceData->DataChannelData && ShouldPublish() && InstanceData->DestinationData && InstanceData->DestinationData->GetNumInstances() > 0)
+
+		if (InstanceData->DataChannelData && ShouldPublish())
 		{
 			FNiagaraDataChannelGameData GameData;
 			GameData.Init(Channel->Get());
-			GameData.AppendFromDataSet(InstanceData->DestinationData, SystemInstance->GetLWCTile());
-			
+
+			if(InstanceData->DestinationData && InstanceData->DestinationData->GetNumInstances() > 0)
+			{
+				GameData.AppendFromDataSet(InstanceData->DestinationData, SystemInstance->GetLWCTile());
+			}
+
+			FNiagaraGpuComputeDispatchInterface*  DispatchInterface = FNiagaraGpuComputeDispatchInterface::Get(World);
+			if (CompiledData.UsedByGPU() && DispatchInterface)
+			{
+				const FNiagaraDataInterfaceProxy_DataChannelWrite* RT_Proxy = GetProxyAs<FNiagaraDataInterfaceProxy_DataChannelWrite>();
+
+				TArray<TSharedPtr<FNiagaraDataBufferReadback, ESPMode::ThreadSafe>> FrameReadbacks;
+
+				ENQUEUE_RENDER_COMMAND(NDISimCacheGPUWriteFrame)
+				(
+					[RT_Proxy, RT_InstanceID = SystemInstance->GetId(), NDCData_RT = InstanceData->DataChannelData, &FrameReadbacks, DispatchInterface](FRHICommandListImmediate& RHICmdList)
+					{
+						const FNiagaraDataInterfaceProxy_DataChannelWrite::FInstanceData* InstanceData_RT = RT_Proxy->SystemInstancesToProxyData_RT.Find(RT_InstanceID);
+						if (InstanceData_RT)
+						{
+							//Instance data already has a buffer we can use to read back this frame's data.
+							for (const FNiagaraDataBufferRef& ReadbackBuffer : InstanceData_RT->PendingSimCacheReadbacks)
+							{
+								//We enqueue a readback for the data if we're wanting to store in the sim cache
+								TSharedPtr<FNiagaraDataBufferReadback, ESPMode::ThreadSafe>& NewReadback = FrameReadbacks.Add_GetRef(MakeShared<FNiagaraDataBufferReadback, ESPMode::ThreadSafe>());
+								//TODO: Rework to avoid additional readback. With some refactoring we can just re-use the existing readback if available.
+								NewReadback->EnqueueReadback(RHICmdList, ReadbackBuffer, DispatchInterface->GetGpuReadbackManager(), DispatchInterface->GetGPUInstanceCounterManager());
+							}
+							InstanceData_RT->PendingSimCacheReadbacks.Reset();
+						}
+
+						DispatchInterface->GetGpuReadbackManager()->WaitCompletion(RHICmdList);
+					}
+				);
+
+				//Is it enough that we flush the RT here?
+				//Can we be sure all work has been submitted and the dispatcher has finished processing the frame etc?
+				FlushRenderingCommands();
+
+				for(TSharedPtr<FNiagaraDataBufferReadback, ESPMode::ThreadSafe>& Readback : FrameReadbacks)
+				{
+					//TODO: Direct copy path, for now copy over to a data set then to the frame data.
+					FNiagaraDataBuffer* ReadbackBuffer = InstanceData->DataChannelData->GetBufferForCPUWrite();
+					Readback->ReadResultsToDataBuffer(ReadbackBuffer);
+					GameData.AppendFromDataSet(ReadbackBuffer->UnlockForRead(), SystemInstance->GetLWCTile());
+				}
+			}
+
 			FrameData.NumElements = GameData.Num();
 			for (const FNiagaraDataChannelVariableBuffer& VarBuffer : GameData.GetVariableBuffers())
 			{
@@ -806,11 +912,12 @@ bool UNiagaraDataInterfaceDataChannelWrite::SimCacheWriteFrame(UObject* StorageO
 			{
 				FrameData.VariableData[VarPair.Value].SourceVar = VarPair.Key;
 			}
-			
+
 			FrameData.bVisibleToGame = bPublishToGame;
 			FrameData.bVisibleToCPUSims = bPublishToCPU;
 			FrameData.bVisibleToGPUSims = bPublishToGPU;
 		}
+
 		return true;
 	} 
 	return false;
@@ -1537,7 +1644,7 @@ void FNiagaraDataInterfaceProxy_DataChannelWrite::PreStage(const FNDIGpuComputeP
 
 	if (InstanceData && InstanceData->ChannelDataRTProxy)
 	{
-		if( InstanceData->BufferForCPU == nullptr && (InstanceData->bPublishToGame || InstanceData->bPublishToCPU) )
+		if( InstanceData->BufferForCPU == nullptr && (InstanceData->bPublishToGame || InstanceData->bPublishToCPU || InstanceData->bCapturingSimCache) )
 		{
 			//Allocate a separate buffer that we will write into and ship back to the CPU.			 
 			InstanceData->BufferForCPU = InstanceData->ChannelDataRTProxy->AllocateBufferForCPU(Context.GetGraphBuilder(), Context.GetComputeDispatchInterface().GetFeatureLevel(), InstanceData->AllocationCount, InstanceData->bPublishToGame, InstanceData->bPublishToCPU, InstanceData->LwcTile);
@@ -1546,6 +1653,11 @@ void FNiagaraDataInterfaceProxy_DataChannelWrite::PreStage(const FNDIGpuComputeP
 			uint32 Offset = InstanceData->BufferForCPU->GetGPUInstanceCountBufferOffset();
 			Context.GetInstanceCountManager().FreeEntry(Offset);
 			InstanceData->BufferForCPU->SetGPUInstanceCountBufferOffset(Context.GetInstanceCountManager().AcquireEntry());
+
+			if(InstanceData->bCapturingSimCache)
+			{
+				InstanceData->PendingSimCacheReadbacks.Emplace(InstanceData->BufferForCPU);
+			}
 		}
 
 		if(InstanceData->bPublishToGPU && InstanceData->GPUBuffer == nullptr)
diff --git a/Engine/Plugins/FX/Niagara/Source/Niagara/Private/NiagaraDataChannel.cpp b/Engine/Plugins/FX/Niagara/Source/Niagara/Private/NiagaraDataChannel.cpp
--- a/Engine/Plugins/FX/Niagara/Source/Niagara/Private/NiagaraDataChannel.cpp
+++ b/Engine/Plugins/FX/Niagara/Source/Niagara/Private/NiagaraDataChannel.cpp
@@ -518,6 +518,7 @@ void FNiagaraDataChannelDataProxy::OnRemovedFromDispatcher(FNiagaraGpuComputeDis
 		ComputeDispatchInterface->GetGPUInstanceCounterManager().FreeEntry(Offset);
 		ReadbackInfo.Buffer->SetGPUInstanceCountBufferOffset(INDEX_NONE);
 	}
+	PendingGPUReadbackBuffers.Reset();
 }
 
 void FNiagaraDataChannelDataProxy::Reset()
@@ -534,6 +535,15 @@ void FNiagaraDataChannelDataProxy::BeginFrame(FNiagaraGpuComputeDispatchInterfac
 #endif
 
 	FNiagaraGPUInstanceCountManager& InstCountManager = DispatchInterface->GetGPUInstanceCounterManager();	
+
+	for (FNDCGpuReadbackInfo& ReadbackInfo : PendingGPUReadbackBuffers)
+	{
+		uint32 Offset = ReadbackInfo.Buffer->GetGPUInstanceCountBufferOffset();
+		DispatchInterface->GetGPUInstanceCounterManager().FreeEntry(Offset);
+		ReadbackInfo.Buffer->SetGPUInstanceCountBufferOffset(INDEX_NONE);
+	}
+	PendingGPUReadbackBuffers.Reset();
+
 	if(this->bNeedsPrevFrameData)
 	{
 		Swap(PrevFrameData, CurrFrameData);
@@ -634,9 +644,10 @@ void FNiagaraDataChannelDataProxy::EndFrame(FNiagaraGpuComputeDispatchInterface*
 			NewReadback->GetOnReadbackComplete().BindLambda(PublishOnCPU);
 			NewReadback->EnqueueReadback(RHICmdList, ReadbackInfo.Buffer, DispatchInterface->GetGpuReadbackManager(), DispatchInterface->GetGPUInstanceCounterManager());
 
-			uint32 Offset = ReadbackInfo.Buffer->GetGPUInstanceCountBufferOffset();
-			DispatchInterface->GetGPUInstanceCounterManager().FreeEntry(Offset);
-			ReadbackInfo.Buffer->SetGPUInstanceCountBufferOffset(INDEX_NONE);
+			//Defer the release of this until next frame to allow other users of this instance count
+// 			uint32 Offset = ReadbackInfo.Buffer->GetGPUInstanceCountBufferOffset();
+// 			DispatchInterface->GetGPUInstanceCounterManager().FreeEntry(Offset);
+// 			ReadbackInfo.Buffer->SetGPUInstanceCountBufferOffset(INDEX_NONE);
 		}
 		PendingGPUReadbackBuffers.Reset();
 	}
@@ -776,15 +787,20 @@ FNiagaraDataBufferRef FNiagaraDataChannelDataProxy::AllocateBufferForCPU(FRDGBui
 	FNiagaraDataBuffer* Buffer = &GPUDataSet->AllocateBuffer();
 	Buffer->AllocateGPU(GraphBuilder.RHICmdList, AllocationSize, FeatureLevel, TEXT("NDC GPU Buffers for CPU"));
 
-	FNDCGpuReadbackInfo& NewReadback = PendingGPUReadbackBuffers.AddDefaulted_GetRef();
-	NewReadback.Buffer = Buffer->UnlockForRead();
-	NewReadback.bPublishToCPU = bPublishToCPU;
-	NewReadback.bPublishToGame = bPublishToGame;
-	NewReadback.LWCTile = LWCTile;
+	FNiagaraDataBufferRef Ret = Buffer->UnlockForRead();
+
+	if(bPublishToCPU || bPublishToGame)
+	{
+		FNDCGpuReadbackInfo& NewReadback = PendingGPUReadbackBuffers.AddDefaulted_GetRef();
+		NewReadback.Buffer = Ret;
+		NewReadback.bPublishToCPU = bPublishToCPU;
+		NewReadback.bPublishToGame = bPublishToGame;
+		NewReadback.LWCTile = LWCTile;
+	}
 
 	AddTransition(GraphBuilder, ERHIAccess::SRVMask, ERHIAccess::UAVCompute, Buffer);
 
-	return NewReadback.Buffer;
+	return Ret;
 }
 
 void FNiagaraDataChannelDataProxy::AddBuffersFromCPU(const TArray<FNiagaraDataBufferRef>& BuffersFromCPU)
