- Fix a crash recording over an existing Sim Cache if the data interface stored object was no longer valid

#test Load up a cache with DI storage where the type no longer exists
#rnx

diff --git a/Engine/Plugins/FX/Niagara/Source/Niagara/Private/NiagaraSimCache.cpp b/Engine/Plugins/FX/Niagara/Source/Niagara/Private/NiagaraSimCache.cpp
--- a/Engine/Plugins/FX/Niagara/Source/Niagara/Private/NiagaraSimCache.cpp
+++ b/Engine/Plugins/FX/Niagara/Source/Niagara/Private/NiagaraSimCache.cpp
@@ -346,7 +346,10 @@ bool UNiagaraSimCache::BeginWrite(FNiagaraSimCacheCreateParameters InCreateParam
 
 	for ( auto it=DataInterfaceStorage.CreateIterator(); it; ++it )
 	{
-		it->Value->MarkAsGarbage();
+		if (it->Value)
+		{
+			it->Value->MarkAsGarbage();
+		}
 	}
 	DataInterfaceStorage.Empty();
 
