Fixed array sim cache comparison when no cpu or gpu data is available

#rnx

diff --git a/Engine/Plugins/FX/Niagara/Source/Niagara/Private/NiagaraDataInterfaceArray.cpp b/Engine/Plugins/FX/Niagara/Source/Niagara/Private/NiagaraDataInterfaceArray.cpp
--- a/Engine/Plugins/FX/Niagara/Source/Niagara/Private/NiagaraDataInterfaceArray.cpp
+++ b/Engine/Plugins/FX/Niagara/Source/Niagara/Private/NiagaraDataInterfaceArray.cpp
@@ -121,9 +121,12 @@ bool UNiagaraDataInterfaceArray::SimCacheCompareFrame(const UObject* LhsStorageO
 {
 	const UNDIArraySimCacheData* LhsCacheData = CastChecked<const UNDIArraySimCacheData>(LhsStorageObject);
 	const UNDIArraySimCacheData* RhsCacheData = CastChecked<const UNDIArraySimCacheData>(RhsStorageObject);
+	const bool bHasCpuData = LhsCacheData->CpuFrameData.IsValidIndex(FrameIndex) || RhsCacheData->CpuFrameData.IsValidIndex(FrameIndex);
+	const bool bHasGpuData = LhsCacheData->GpuFrameData.IsValidIndex(FrameIndex) || RhsCacheData->GpuFrameData.IsValidIndex(FrameIndex);
+	const bool bIsCpuFrameIndexValid = LhsCacheData->CpuFrameData.IsValidIndex(FrameIndex) && RhsCacheData->CpuFrameData.IsValidIndex(FrameIndex);
+	const bool bIsGpuFrameIndexValid = !bHasGpuData || (LhsCacheData->GpuFrameData.IsValidIndex(FrameIndex) && RhsCacheData->GpuFrameData.IsValidIndex(FrameIndex));
 
-	if (!LhsCacheData->CpuFrameData.IsValidIndex(FrameIndex) || !RhsCacheData->CpuFrameData.IsValidIndex(FrameIndex) ||
-		!LhsCacheData->GpuFrameData.IsValidIndex(FrameIndex) || !RhsCacheData->GpuFrameData.IsValidIndex(FrameIndex) )
+	if (!bIsCpuFrameIndexValid || !bIsGpuFrameIndexValid )
 	{
 		OutErrors = TEXT("FrameIndex was not valid");
 		return false;
@@ -137,7 +140,7 @@ bool UNiagaraDataInterfaceArray::SimCacheCompareFrame(const UObject* LhsStorageO
 		{
 			if (LhsFrame.NumElements != RhsFrame.NumElements)
 			{
-				OutErrors = FString::Printf(TEXT("Element Count Mismatch (%d -> %d) for %s data"), LhsCacheData->CpuFrameData[FrameIndex].NumElements, RhsCacheData->CpuFrameData[FrameIndex].NumElements, SimType);
+				OutErrors = FString::Printf(TEXT("Element Count Mismatch (%d -> %d) for %s data"), LhsFrame.NumElements, RhsFrame.NumElements, SimType);
 				return false;
 			}
 
@@ -157,8 +160,8 @@ bool UNiagaraDataInterfaceArray::SimCacheCompareFrame(const UObject* LhsStorageO
 		};
 
 	return
-		CompareFrames(LhsCacheData->CpuFrameData[FrameIndex], RhsCacheData->CpuFrameData[FrameIndex], TEXT("CPU")) &&
-		CompareFrames(LhsCacheData->GpuFrameData[FrameIndex], RhsCacheData->GpuFrameData[FrameIndex], TEXT("GPU"));
+		(!bHasCpuData || CompareFrames(LhsCacheData->CpuFrameData[FrameIndex], RhsCacheData->CpuFrameData[FrameIndex], TEXT("CPU"))) &&
+		(!bHasGpuData || CompareFrames(LhsCacheData->GpuFrameData[FrameIndex], RhsCacheData->GpuFrameData[FrameIndex], TEXT("GPU")));
 }
 
 FString UNiagaraDataInterfaceArray::SimCacheVisualizerRead(const UNDIArraySimCacheData* CacheData, const FNDIArraySimCacheDataFrame& FrameData, int Element) const
