- Render Target Sim Cache Storage changes
- Reworked so that it can be shared between the different data interfaces
- Added sim cache comparison for RT2D
- Added visualizer for RT2D data

#tests Made a cache test using the new functionality

diff --git a/Engine/Plugins/FX/Niagara/Config/BaseNiagara.ini b/Engine/Plugins/FX/Niagara/Config/BaseNiagara.ini
--- a/Engine/Plugins/FX/Niagara/Config/BaseNiagara.ini
+++ b/Engine/Plugins/FX/Niagara/Config/BaseNiagara.ini
@@ -42,6 +42,9 @@
 +EnumRedirects=(OldName="ENiagaraMessageSeverity",ValueChanges=(("ENiagaraMessageSeverity::CustomNote", "ENiagaraMessageSeverity::Info")) )
 +EnumRedirects=(OldName="EStackIssueSeverity",ValueChanges=(("EStackIssueSeverity::CustomNote", "EStackIssueSeverity::Info")) )
 
++ClassRedirects=(OldName="NDIRenderTargetVolumeSimCacheData", NewName="/Script/Niagara.NDIRenderTargetSimCacheData")
++StructRedirects=(OldName="NDIRenderTargetVolumeSimCacheFrame", NewName="/Script/Niagara.NDIRenderTargetSimCacheFrame")
+
 ; Flipbook to Baker
 +ClassRedirects=(OldName="NiagaraFlipbookSettings", NewName="/Script/Niagara.NiagaraBakerSettings")
 +PropertyRedirects=(OldName="NiagaraSystem.FlipbookSettings", NewName="NiagaraSystem.BakerSettings")
diff --git a/Engine/Plugins/FX/Niagara/Source/Niagara/Classes/NiagaraDataInterfaceRenderTarget2D.h b/Engine/Plugins/FX/Niagara/Source/Niagara/Classes/NiagaraDataInterfaceRenderTarget2D.h
--- a/Engine/Plugins/FX/Niagara/Source/Niagara/Classes/NiagaraDataInterfaceRenderTarget2D.h
+++ b/Engine/Plugins/FX/Niagara/Source/Niagara/Classes/NiagaraDataInterfaceRenderTarget2D.h
@@ -129,6 +129,7 @@ public:
 	NIAGARA_API virtual bool SimCacheWriteFrame(UObject* StorageObject, int FrameIndex, FNiagaraSystemInstance* SystemInstance, const void* OptionalPerInstanceData, FNiagaraSimCacheFeedbackContext& FeedbackContext) const override;
 	NIAGARA_API virtual bool SimCacheEndWrite(UObject* StorageObject) const override;
 	NIAGARA_API virtual bool SimCacheReadFrame(UObject* StorageObject, int FrameA, int FrameB, float Interp, FNiagaraSystemInstance* SystemInstance, void* OptionalPerInstanceData) override;
+	NIAGARA_API virtual bool SimCacheCompareFrame(const UObject* LhsStorageObject, const UObject* RhsStorageObject, int FrameIndex, TOptional<float> Tolerance, FString& OutErrors) const override;
 	NIAGARA_API virtual TSharedPtr<FJsonObject> SimCacheToJson(const UObject* StorageObject, int FrameIndex, TOptional<FString> TargetFolder, TOptional<FString> FilenamePrefix) const override;
 	//~ UNiagaraDataInterface interface END
 
diff --git a/Engine/Plugins/FX/Niagara/Source/Niagara/Internal/DataInterface/NDIRenderTargetSimCacheData.h b/Engine/Plugins/FX/Niagara/Source/Niagara/Internal/DataInterface/NDIRenderTargetSimCacheData.h
new file mode 100644
--- /dev/null
+++ b/Engine/Plugins/FX/Niagara/Source/Niagara/Internal/DataInterface/NDIRenderTargetSimCacheData.h
@@ -0,0 +1,77 @@
+// Copyright Epic Games, Inc. All Rights Reserved.
+#pragma once
+
+#include "RHIDefinitions.h"
+#include "Serialization/BulkData.h"
+#include "PixelFormat.h"
+
+#include "NDIRenderTargetSimCacheData.generated.h"
+
+class FJsonObject;
+
+USTRUCT()
+struct FNDIRenderTargetSimCacheFrame
+{
+	GENERATED_BODY();
+
+	UPROPERTY()
+	FIntVector Size = FIntVector(EForceInit::ForceInitToZero);
+
+	UPROPERTY()
+	TEnumAsByte<EPixelFormat> Format = EPixelFormat::PF_A16B16G16R16;
+
+	UPROPERTY()
+	int32 UncompressedSize = 0;
+
+	UPROPERTY()
+	int32 CompressedSize = 0;
+
+	mutable FByteBulkData BulkData;
+};
+
+UCLASS(MinimalAPI)
+class UNDIRenderTargetSimCacheData : public UObject
+{
+	GENERATED_BODY()
+
+public:
+	// Begin UObject Interface
+	NIAGARA_API virtual void BeginDestroy() override;
+	NIAGARA_API virtual void Serialize(FArchive& Ar) override;
+
+	bool IsValidFrame(int32 FrameIndex) const { return Frames.IsValidIndex(FrameIndex); }
+	bool HasPixelData(int32 FrameIndex) const { return Frames.IsValidIndex(FrameIndex) && Frames[FrameIndex].CompressedSize > 0; }
+	FIntVector GetTextureSize(int32 FrameIndex) const { check(Frames.IsValidIndex(FrameIndex)); return Frames[FrameIndex].Size; }
+	EPixelFormat GetTextureFormat(int32 FrameIndex) const { check(Frames.IsValidIndex(FrameIndex)); return Frames[FrameIndex].Format; }
+	int32 GetCompressedSize(int32 FrameIndex) const { check(Frames.IsValidIndex(FrameIndex)); return Frames[FrameIndex].CompressedSize; }
+	int32 GetUncompressedSize(int32 FrameIndex) const { check(Frames.IsValidIndex(FrameIndex)); return Frames[FrameIndex].UncompressedSize; }
+
+	bool CompareFrame(const UNDIRenderTargetSimCacheData* Other, int32 FrameIndex, TOptional<float> Tolerance, FString& OutErrors) const;
+
+	TSharedPtr<FJsonObject> FrameToJson(int FrameIndex, TOptional<FString> TargetFolder, TOptional<FString> FilenamePrefix) const;
+
+	// Get a copy of the uncompressed pixel data
+	NIAGARA_API TArray<uint8> GetPixelData(int32 FrameIndex) const;
+
+	// Copy the pixel data into the provided buffer
+	// Note:  you must have allocated the appropriate space
+	NIAGARA_API void GetPixelData(int32 FrameIndex, uint8* DestPixelData, int32 DestRowPitch, int32 DestSlicePitch) const;
+
+	// Release the bulk data copy we are holding onto
+	NIAGARA_API void ReleasePixelData(int32 FrameIndex);
+
+	// Release all of the bulk data copies
+	NIAGARA_API void ReleaseAllPixelData();
+
+	// Set pixel data for the given frame
+	NIAGARA_API void SetPixelData(int32 FrameIndex, FIntVector Size, EPixelFormat Format, TArrayView<uint8> PixelData);
+
+	UPROPERTY()
+	FName CompressionType;
+
+protected:
+	UPROPERTY()
+	TArray<FNDIRenderTargetSimCacheFrame> Frames;
+
+	mutable TArray<uint8*> PixelDataFrames;
+};
diff --git a/Engine/Plugins/FX/Niagara/Source/Niagara/Private/DataInterface/NDIRenderTargetSimCacheData.cpp b/Engine/Plugins/FX/Niagara/Source/Niagara/Private/DataInterface/NDIRenderTargetSimCacheData.cpp
new file mode 100644
--- /dev/null
+++ b/Engine/Plugins/FX/Niagara/Source/Niagara/Private/DataInterface/NDIRenderTargetSimCacheData.cpp
@@ -0,0 +1,344 @@
+// Copyright Epic Games, Inc. All Rights Reserved.
+
+#include "DataInterface/NDIRenderTargetSimCacheData.h"
+#include "NiagaraCommon.h"
+
+#include "Dom/JsonObject.h"
+#include "Misc/Compression.h"
+#include "Misc/Paths.h"
+#include "ImageCoreUtils.h"
+#include "ImageUtils.h"
+
+namespace NDIRenderTargetSimCacheDataPrivate
+{
+	bool CanConvertPixelFormat(EPixelFormat Format)
+	{
+		switch (Format)
+		{
+			case PF_FloatRGBA:		return true;
+			case PF_G16R16F:
+			case PF_G16R16F_FILTER:	return true;
+			default:				return false;
+		}
+	}
+
+	FLinearColor ConvertPixel(const uint8* SourceData, EPixelFormat Format)
+	{
+		FLinearColor Color;
+		switch (Format)
+		{
+			case PF_FloatRGBA:
+				Color.R = *reinterpret_cast<const FFloat16*>(SourceData);
+				Color.G = *reinterpret_cast<const FFloat16*>(SourceData + 2);
+				Color.B = *reinterpret_cast<const FFloat16*>(SourceData + 4);
+				Color.A = *reinterpret_cast<const FFloat16*>(SourceData + 6);
+				break;
+
+			case PF_G16R16F:
+			case PF_G16R16F_FILTER:
+				Color.R = *reinterpret_cast<const FFloat16*>(SourceData);
+				Color.G = *reinterpret_cast<const FFloat16*>(SourceData + 2);
+				Color.B = 0.0f;
+				Color.A = 1.0f;
+				break;
+
+			default: // do nothing;
+				Color.R = 0.0f;
+				Color.G = 0.0f;
+				Color.B = 0.0f;
+				Color.A = 1.0f;
+				break;
+		}
+		return Color;
+	}
+}
+
+
+void UNDIRenderTargetSimCacheData::BeginDestroy()
+{
+	Super::BeginDestroy();
+
+	ReleaseAllPixelData();
+	Frames.Empty();
+}
+
+void UNDIRenderTargetSimCacheData::Serialize(FArchive& Ar)
+{
+	Super::Serialize(Ar);
+
+	for (FNDIRenderTargetSimCacheFrame& Frame : Frames)
+	{
+		if (Frame.UncompressedSize == 0)
+		{
+			continue;
+		}
+
+		Frame.BulkData.Serialize(Ar, this);
+	}
+}
+
+bool UNDIRenderTargetSimCacheData::CompareFrame(const UNDIRenderTargetSimCacheData* Other, int32 FrameIndex, TOptional<float> InTolerance, FString& OutErrors) const
+{
+	if (HasPixelData(FrameIndex) != Other->HasPixelData(FrameIndex))
+	{
+		OutErrors = TEXT("HasPixelData mismatched between caches");
+		return false;
+	}
+
+	const FIntVector LhsTextureSize = GetTextureSize(FrameIndex);
+	const FIntVector RhsTextureSize = Other->GetTextureSize(FrameIndex);
+	if (LhsTextureSize != RhsTextureSize)
+	{
+		OutErrors = FString::Printf(
+			TEXT("TextureSizess (%dx%dx%d & %dx%dx%d) mismatched between caches"),
+			LhsTextureSize.X, LhsTextureSize.Y, LhsTextureSize.Z,
+			RhsTextureSize.X, RhsTextureSize.Y, RhsTextureSize.Z
+		);
+		return false;
+	}
+
+	const EPixelFormat LhsPixelFormat = GetTextureFormat(FrameIndex);
+	const EPixelFormat RhsPixelFormat = GetTextureFormat(FrameIndex);
+	if (LhsPixelFormat != RhsPixelFormat)
+	{
+		OutErrors = FString::Printf(
+			TEXT("PixelFormats (%s & %s) mismatched between caches"),
+			GPixelFormats[LhsPixelFormat].Name,
+			GPixelFormats[RhsPixelFormat].Name
+		);
+		return false;
+	}
+
+	if (!NDIRenderTargetSimCacheDataPrivate::CanConvertPixelFormat(LhsPixelFormat))
+	{
+		OutErrors = FString::Printf(TEXT("PixelFormat (%s) is not supported for comparison"), GPixelFormats[LhsPixelFormat].Name);
+		return false;
+	}
+
+	const TArray<uint8> LhsPixelData = GetPixelData(FrameIndex);
+	const TArray<uint8> RhsPixelData = Other->GetPixelData(FrameIndex);
+	if (LhsPixelData.Num() != RhsPixelData.Num())
+	{
+		OutErrors = TEXT("Arraysize mismatch between caches, but texture size matches?");
+		return false;
+	}
+
+	const float Tolerance = InTolerance.Get(1.e-3f);
+
+	const uint32 BlockBytes = GPixelFormats[LhsPixelFormat].BlockBytes;
+	const uint32 PixelCount = LhsPixelData.Num() / BlockBytes;
+
+	for ( uint32 i=0; i < PixelCount; ++i)
+	{
+		const uint32 Offset = i * BlockBytes;
+		const FLinearColor LhsPixel = NDIRenderTargetSimCacheDataPrivate::ConvertPixel(LhsPixelData.GetData() + Offset, LhsPixelFormat);
+		const FLinearColor RhsPixel = NDIRenderTargetSimCacheDataPrivate::ConvertPixel(RhsPixelData.GetData() + Offset, RhsPixelFormat);
+		if (!LhsPixel.Equals(RhsPixel, Tolerance))
+		{
+			const int32 PixelX = i % LhsTextureSize.X;
+			const int32 PixelY = (i / LhsTextureSize.X) % LhsTextureSize.Y;
+			const int32 PixelZ = i / (LhsTextureSize.X * LhsTextureSize.Y);
+			
+			OutErrors = FString::Printf(
+				TEXT("Pixel (%dx%dx%d) is different (R:%g G:%g B:%g A:%g) vs (R:%g G:%g B:%g A:%g)"),
+				PixelX, PixelY, PixelZ,
+				LhsPixel.R, LhsPixel.G, LhsPixel.B, LhsPixel.A,
+				RhsPixel.R, RhsPixel.G, RhsPixel.B, RhsPixel.A
+			);
+			return false;
+		}
+	}
+
+	return true;
+}
+
+TSharedPtr<FJsonObject> UNDIRenderTargetSimCacheData::FrameToJson(int FrameIndex, TOptional<FString> TargetFolder, TOptional<FString> FilenamePrefix) const
+{
+	if (!HasPixelData(FrameIndex))
+	{
+		return TSharedPtr<FJsonObject>();
+	}
+
+	const FNDIRenderTargetSimCacheFrame& CacheFrame = Frames[FrameIndex];
+
+	TSharedPtr<FJsonObject> JsonCacheObject = MakeShared<FJsonObject>();
+	JsonCacheObject->SetStringField(TEXT("CompressionType"), CompressionType.ToString());
+	JsonCacheObject->SetStringField(TEXT("PixelFormat"), GPixelFormats[CacheFrame.Format].Name);
+	JsonCacheObject->SetNumberField(TEXT("CompressedSize"), CacheFrame.CompressedSize);
+	JsonCacheObject->SetNumberField(TEXT("UncompressedSize"), CacheFrame.UncompressedSize);
+	JsonCacheObject->SetNumberField(TEXT("SizeX"), CacheFrame.Size.X);
+	JsonCacheObject->SetNumberField(TEXT("SizeY"), CacheFrame.Size.Y);
+	JsonCacheObject->SetNumberField(TEXT("SizeZ"), CacheFrame.Size.Z);
+
+	TArray<uint8> PixelData = GetPixelData(FrameIndex);
+	if (PixelData.Num() == 0 || !TargetFolder.IsSet() || !FilenamePrefix.IsSet())
+	{
+		return JsonCacheObject;
+	}
+
+	const ERawImageFormat::Type ImageFormat = FImageCoreUtils::GetRawImageFormatForPixelFormat(CacheFrame.Format);
+	if (ImageFormat != ERawImageFormat::RGBA16F)
+	{
+		UE_LOG(LogNiagara, Error, TEXT("Unable to save render target to file with current pixel format"));
+		return JsonCacheObject;
+	}
+
+	const uint32 PixelCount = CacheFrame.Size.X * CacheFrame.Size.Y * CacheFrame.Size.Z;
+	const uint32 BlockBytes = GPixelFormats[CacheFrame.Format].BlockBytes;
+
+	// we need to convert the raw pixel data into a format that the image utils understand 
+	TArray<FFloat16Color> ImagePixelData;
+	ImagePixelData.AddZeroed(PixelCount);
+
+	for (uint32 Index = 0; Index < PixelCount; Index++)
+	{
+		const uint32 Offset = Index * BlockBytes;
+		ImagePixelData[Index] = NDIRenderTargetSimCacheDataPrivate::ConvertPixel(PixelData.GetData() + Offset, CacheFrame.Format);
+	}
+		
+	FImageView ImageView(ImagePixelData.GetData(), CacheFrame.Size.X, CacheFrame.Size.Y, CacheFrame.Size.Z, ImageFormat, EGammaSpace::Linear);
+	if (FImageUtils::SaveImageByExtension(*FPaths::Combine(TargetFolder.GetValue(), FilenamePrefix.GetValue() + ".exr"), ImageView))
+	{
+		JsonCacheObject->SetStringField(TEXT("TextureData"), FilenamePrefix.GetValue() + ".exr");
+	}		
+	return JsonCacheObject;
+}
+
+TArray<uint8> UNDIRenderTargetSimCacheData::GetPixelData(int32 FrameIndex) const
+{
+	if (!Frames.IsValidIndex(FrameIndex))
+	{
+		return TArray<uint8>();
+	}
+
+	const FNDIRenderTargetSimCacheFrame& FrameData = Frames[FrameIndex];
+
+	PixelDataFrames.SetNum(Frames.Num());
+	if (PixelDataFrames[FrameIndex] == nullptr)
+	{
+		FrameData.BulkData.GetCopy(reinterpret_cast<void**>(&PixelDataFrames[FrameIndex]), true);
+	}
+
+	TArray<uint8> OutPixelData;
+	OutPixelData.AddUninitialized(FrameData.UncompressedSize);
+	if (FrameData.CompressedSize > 0)
+	{
+		FCompression::UncompressMemory(CompressionType, OutPixelData.GetData(), FrameData.UncompressedSize, PixelDataFrames[FrameIndex], FrameData.CompressedSize);
+	}
+	else
+	{
+		FMemory::Memcpy(OutPixelData.GetData(), PixelDataFrames[FrameIndex], FrameData.UncompressedSize);
+	}
+
+	return OutPixelData;
+}
+
+void UNDIRenderTargetSimCacheData::GetPixelData(int32 FrameIndex, uint8* DestPixelData, int32 DestRowPitch, int32 DestSlicePitch) const
+{
+	if (!Frames.IsValidIndex(FrameIndex))
+	{
+		return;
+	}
+
+	const FNDIRenderTargetSimCacheFrame& FrameData = Frames[FrameIndex];
+
+	PixelDataFrames.SetNum(Frames.Num());
+	if (PixelDataFrames[FrameIndex] == nullptr)
+	{
+		FrameData.BulkData.GetCopy(reinterpret_cast<void**>(&PixelDataFrames[FrameIndex]), true);
+	}
+
+	const uint32 BlockBytes = GPixelFormats[FrameData.Format].BlockBytes;
+	const int32 SrcRowPitch = FrameData.Size.X * BlockBytes;
+	const int32 SrcSlicePitch = FrameData.Size.Y * SrcRowPitch;
+
+	// Fast path (direct copy / decompress)
+	if (DestRowPitch == SrcRowPitch && (FrameData.Size.Z == 1 || DestSlicePitch == SrcSlicePitch))
+	{
+		if (FrameData.CompressedSize > 0)
+		{
+			FCompression::UncompressMemory(CompressionType, DestPixelData, FrameData.UncompressedSize, PixelDataFrames[FrameIndex], FrameData.CompressedSize);
+		}
+		else
+		{
+			FMemory::Memcpy(DestPixelData, PixelDataFrames[FrameIndex], FrameData.UncompressedSize);
+		}
+	}
+	// Slow path as we have to adjust for row / slice size differences
+	else
+	{
+		TArray<uint8> Decompressed;
+		if (FrameData.CompressedSize > 0)
+		{
+			Decompressed.AddUninitialized(BlockBytes * FrameData.Size.X * FrameData.Size.Y * FrameData.Size.Z);
+			FCompression::UncompressMemory(CompressionType, Decompressed.GetData(), Decompressed.Num(), PixelDataFrames[FrameIndex], FrameData.CompressedSize);
+		}
+
+		const uint8* SrcData = Decompressed.Num() > 0 ? Decompressed.GetData() : PixelDataFrames[FrameIndex];
+		for (int32 Slice=0; Slice < FrameData.Size.Z; ++Slice)
+		{
+			uint8* DstData = DestPixelData + (Slice * DestSlicePitch);
+			for (int32 Row=0; Row < FrameData.Size.Y; ++Row)
+			{
+				FMemory::Memcpy(DstData, SrcData, SrcRowPitch);
+				SrcData += SrcRowPitch;
+				DstData += DestRowPitch;
+			}
+		}
+	}
+}
+
+void UNDIRenderTargetSimCacheData::ReleasePixelData(int32 FrameIndex)
+{
+	if (PixelDataFrames.IsValidIndex(FrameIndex))
+	{
+		FMemory::Free(PixelDataFrames[FrameIndex]);
+		PixelDataFrames[FrameIndex] = nullptr;
+	}
+}
+
+void UNDIRenderTargetSimCacheData::ReleaseAllPixelData()
+{
+	for (uint8* PixelData : PixelDataFrames)
+	{
+		FMemory::Free(PixelData);
+	}
+	PixelDataFrames.Empty();
+}
+
+void UNDIRenderTargetSimCacheData::SetPixelData(int32 FrameIndex, FIntVector Size, EPixelFormat Format, TArrayView<uint8> PixelData)
+{
+	ReleaseAllPixelData();
+
+	Frames.SetNum(FMath::Max(Frames.Num(), FrameIndex) + 1);
+
+	FNDIRenderTargetSimCacheFrame& FrameData = Frames[FrameIndex];
+	FrameData.Size				= Size;
+	FrameData.Format			= Format;
+	FrameData.UncompressedSize	= PixelData.Num();
+	FrameData.CompressedSize	= 0;
+
+	// Compressed
+	if (!CompressionType.IsNone())
+	{
+		const ECompressionFlags CompressionFlags = COMPRESS_BiasMemory;
+		int CompressedSize = FCompression::CompressMemoryBound(CompressionType, PixelData.Num(), CompressionFlags);
+
+		TArray<uint8> CompressedPixelData;
+		CompressedPixelData.Reserve(CompressedSize);
+		if (FCompression::CompressMemory(CompressionType, CompressedPixelData.GetData(), CompressedSize, PixelData.GetData(), PixelData.Num(), CompressionFlags))
+		{
+			FrameData.BulkData.Lock(LOCK_READ_WRITE);
+			FMemory::Memcpy(FrameData.BulkData.Realloc(CompressedSize), CompressedPixelData.GetData(), CompressedSize);
+			FrameData.BulkData.Unlock();
+
+			FrameData.CompressedSize = CompressedSize;
+			return;
+		}
+	}
+
+	// Uncompressed
+	FrameData.BulkData.Lock(LOCK_READ_WRITE);
+	FMemory::Memcpy(FrameData.BulkData.Realloc(PixelData.Num()), PixelData.GetData(), PixelData.Num());
+	FrameData.BulkData.Unlock();
+}
diff --git a/Engine/Plugins/FX/Niagara/Source/Niagara/Private/NDIRenderTargetVolumeSimCacheData.h b/Engine/Plugins/FX/Niagara/Source/Niagara/Private/NDIRenderTargetVolumeSimCacheData.h
deleted file mode 100644
diff --git a/Engine/Plugins/FX/Niagara/Source/Niagara/Private/NiagaraDataInterfaceRenderTarget2D.cpp b/Engine/Plugins/FX/Niagara/Source/Niagara/Private/NiagaraDataInterfaceRenderTarget2D.cpp
--- a/Engine/Plugins/FX/Niagara/Source/Niagara/Private/NiagaraDataInterfaceRenderTarget2D.cpp
+++ b/Engine/Plugins/FX/Niagara/Source/Niagara/Private/NiagaraDataInterfaceRenderTarget2D.cpp
@@ -4,11 +4,9 @@
 #include "NiagaraCompileHashVisitor.h"
 #include "TextureResource.h"
 #include "CanvasTypes.h"
-#include "ImageCoreUtils.h"
-#include "ImageUtils.h"
 #include "Engine/TextureRenderTarget2D.h"
 
-#include "NDIRenderTargetVolumeSimCacheData.h"
+#include "DataInterface/NDIRenderTargetSimCacheData.h"
 #include "NiagaraSimCache.h"
 #include "NiagaraDataInterfaceRenderTargetCommon.h"
 #include "NiagaraSystemInstance.h"
@@ -583,8 +581,8 @@ NIAGARA_API UObject* UNiagaraDataInterfaceRenderTarget2D::SimCacheBeginWrite(UOb
 		return nullptr;
 	}
 
-	UNDIRenderTargetVolumeSimCacheData* SimCacheData = nullptr;
-	SimCacheData = NewObject<UNDIRenderTargetVolumeSimCacheData>(InSimCache);
+	UNDIRenderTargetSimCacheData* SimCacheData = nullptr;
+	SimCacheData = NewObject<UNDIRenderTargetSimCacheData>(InSimCache);
 	SimCacheData->CompressionType = NDIRenderTarget2DLocal::GetSimCacheCompressionType();
 	return SimCacheData;	
 }
@@ -643,50 +641,13 @@ NIAGARA_API bool UNiagaraDataInterfaceRenderTarget2D::SimCacheWriteFrame(UObject
 
 		if (TextureData.Num() > 0)
 		{
-			UNDIRenderTargetVolumeSimCacheData* SimCacheData = CastChecked<UNDIRenderTargetVolumeSimCacheData>(StorageObject);
-			SimCacheData->Frames.SetNum(FMath::Max(SimCacheData->Frames.Num(), FrameIndex + 1));
-
-			FNDIRenderTargetVolumeSimCacheFrame* CacheFrame = &SimCacheData->Frames[FrameIndex];
-
-			CacheFrame->Size = FIntVector(InstanceData_GT->Size.X, InstanceData_GT->Size.Y, 1);
-			CacheFrame->Format = GetPixelFormatFromRenderTargetFormat(InstanceData_GT->Format);
-
-			const FName CompressionType = SimCacheData->CompressionType;
-			const bool bUseCompression = CompressionType.IsNone() == false;
-			const int TextureSizeBytes = TextureData.Num() * TextureData.GetTypeSize();
-			int CompressedSize = bUseCompression ? FCompression::CompressMemoryBound(CompressionType, TextureSizeBytes) : TextureSizeBytes;
-
-			CacheFrame->UncompressedSize = TextureSizeBytes;
-			CacheFrame->CompressedSize = 0;
-			uint8* FinalPixelData = reinterpret_cast<uint8*>(FMemory::Malloc(CompressedSize));
-			if (bUseCompression)
-			{
-				if (FCompression::CompressMemory(CompressionType, FinalPixelData, CompressedSize, TextureData.GetData(), TextureSizeBytes, COMPRESS_BiasMemory))
-				{
-					CacheFrame->CompressedSize = CompressedSize;
-				}
-				else
-				{
-					UE_LOG(LogNiagara, Error, TEXT("Error compressing render target data"));
-					return false;
-				}
-			}
-
-			if (CacheFrame->CompressedSize == 0)
-			{
-				FMemory::Memcpy(FinalPixelData, TextureData.GetData(), TextureSizeBytes);
-			}
-
-			// Update bulk data
-			CacheFrame->BulkData.Lock(LOCK_READ_WRITE);
-			{
-				const int32 FinalByteSize = CacheFrame->CompressedSize > 0 ? CacheFrame->CompressedSize : CacheFrame->UncompressedSize;
-				uint8* BulkDataPtr = reinterpret_cast<uint8*>(CacheFrame->BulkData.Realloc(FinalByteSize));
-				FMemory::Memcpy(BulkDataPtr, FinalPixelData, FinalByteSize);
-			}
-			CacheFrame->BulkData.Unlock();
-
-			FMemory::Free(FinalPixelData);
+			UNDIRenderTargetSimCacheData* SimCacheData = CastChecked<UNDIRenderTargetSimCacheData>(StorageObject);
+			SimCacheData->SetPixelData(
+				FrameIndex,
+				FIntVector(InstanceData_GT->Size.X, InstanceData_GT->Size.Y, 1),
+				GetPixelFormatFromRenderTargetFormat(InstanceData_GT->Format),
+				TextureData
+			);
 		}
 	}
 
@@ -700,65 +661,41 @@ NIAGARA_API bool UNiagaraDataInterfaceRenderTarget2D::SimCacheEndWrite(UObject*
 
 NIAGARA_API bool UNiagaraDataInterfaceRenderTarget2D::SimCacheReadFrame(UObject* StorageObject, int FrameA, int FrameB, float Interp, FNiagaraSystemInstance* SystemInstance, void* OptionalPerInstanceData)
 {
-	UNDIRenderTargetVolumeSimCacheData* SimCacheData = CastChecked<UNDIRenderTargetVolumeSimCacheData>(StorageObject);
+	UNDIRenderTargetSimCacheData* SimCacheData = CastChecked<UNDIRenderTargetSimCacheData>(StorageObject);
 
 	FRenderTarget2DRWInstanceData_GameThread* InstanceData_GT = reinterpret_cast<FRenderTarget2DRWInstanceData_GameThread*>(OptionalPerInstanceData);
 
-
 	const int FrameIndex = Interp >= 0.5f ? FrameB : FrameA;
-	if (!SimCacheData->Frames.IsValidIndex(FrameIndex))
+	if (!SimCacheData->IsValidFrame(FrameIndex))
 	{
 		return false;
 	}
 
-	FNDIRenderTargetVolumeSimCacheFrame* CacheFrame = &SimCacheData->Frames[FrameIndex];
+	{
+		const FIntVector TextureSize = SimCacheData->GetTextureSize(FrameIndex);
+		const EPixelFormat TextureFormat = SimCacheData->GetTextureFormat(FrameIndex);
 
-	InstanceData_GT->Size = FIntPoint(CacheFrame->Size.X, CacheFrame->Size.Y);
-	InstanceData_GT->Format = NiagaraDataInterfaceRenderTargetCommon::GetRenderTargetFormatFromPixelFormat(CacheFrame->Format);
-	InstanceData_GT->Filter = TextureFilter::TF_Default;
+		InstanceData_GT->Size = FIntPoint(TextureSize.X, TextureSize.Y);
+		InstanceData_GT->Format = NiagaraDataInterfaceRenderTargetCommon::GetRenderTargetFormatFromPixelFormat(TextureFormat);
+		InstanceData_GT->Filter = TextureFilter::TF_Default;
+	}
 
 	PerInstanceTick(InstanceData_GT, SystemInstance, 0.0f);
 	PerInstanceTickPostSimulate(InstanceData_GT, SystemInstance, 0.0f);
 
-	if (InstanceData_GT->TargetTexture && CacheFrame->GetPixelData() != nullptr)
+	if (InstanceData_GT->TargetTexture && SimCacheData->HasPixelData(FrameIndex))
 	{
 		FNiagaraDataInterfaceProxyRenderTarget2DProxy* RT_Proxy = GetProxyAs<FNiagaraDataInterfaceProxyRenderTarget2DProxy>();
 		FTextureRenderTargetResource* RT_TargetTexture = InstanceData_GT->TargetTexture->GameThread_GetRenderTargetResource();
-		ENQUEUE_RENDER_COMMAND(NDIRenderTargetVolumeUpdate)
+		ENQUEUE_RENDER_COMMAND(NDIRenderTargetUpdate)
 			(
-				[RT_PixelFormat = CacheFrame->Format, RT_Proxy, RT_InstanceID = SystemInstance->GetId(), RT_TargetTexture, RT_CacheFrame = CacheFrame, RT_CompressionType = SimCacheData->CompressionType, RT_Format = InstanceData_GT->Format](FRHICommandListImmediate& RHICmdList)
+				[RT_Proxy, RT_InstanceID=SystemInstance->GetId(), RT_TargetTexture, RT_SimCacheData=SimCacheData, RT_FrameIndex=FrameIndex](FRHICommandListImmediate& RHICmdList)
 				{
 					if (FRenderTarget2DRWInstanceData_RenderThread* InstanceData_RT = RT_Proxy->SystemInstancesToProxyData_RT.Find(RT_InstanceID))
 					{
-						uint32 BlockBytes = GPixelFormats[RT_PixelFormat].BlockBytes;
-						FUpdateTextureRegion2D UpdateRegion(0, 0, 0, 0, InstanceData_RT->Size.X, InstanceData_RT->Size.Y);
-						const int32 DestPitch = InstanceData_RT->Size.X * BlockBytes;
-									
-						TArray<uint8> Decompressed;
-					
-						if (RT_CacheFrame->CompressedSize > 0)
-						{
-							Decompressed.AddUninitialized(BlockBytes * InstanceData_RT->Size.X * InstanceData_RT->Size.Y);
-							bool Success = FCompression::UncompressMemory(RT_CompressionType, Decompressed.GetData(), Decompressed.Num(), RT_CacheFrame->GetPixelData(), RT_CacheFrame->CompressedSize);
-							check(Success);
-						}
-
 						uint32 Stride;
 						uint8* OutBuffer = static_cast<uint8*>(RHICmdList.LockTexture2D(InstanceData_RT->TextureRHI, 0, RLM_WriteOnly, Stride, true));
-
-						const int32 SourcePitch = InstanceData_RT->Size.X * BlockBytes;
-
-						const uint8* SrcData = Decompressed.Num() > 0 ? Decompressed.GetData() : RT_CacheFrame->GetPixelData();
-							
-						uint8* DstData = OutBuffer;
-							
-						for (int32 y = 0; y < InstanceData_RT->Size.Y; ++y)
-						{
-							FMemory::Memcpy(DstData, SrcData, SourcePitch);
-							SrcData += SourcePitch;
-							DstData += Stride;
-						}
-
+						RT_SimCacheData->GetPixelData(RT_FrameIndex, OutBuffer, Stride, 0);
 						RHICmdList.UnlockTexture2D(InstanceData_RT->TextureRHI, 0, true);						
 						//RHICmdList.UpdateTexture2D(InstanceData_RT->TextureRHI, 0, UpdateRegion, DestPitch, RT_CacheFrame->PixelData);													
 					}
@@ -769,88 +706,17 @@ NIAGARA_API bool UNiagaraDataInterfaceRenderTarget2D::SimCacheReadFrame(UObject*
 	return true;
 }
 
-namespace NiagaraRenderTarget2DImageHelper
+bool UNiagaraDataInterfaceRenderTarget2D::SimCacheCompareFrame(const UObject* LhsStorageObject, const UObject* RhsStorageObject, int FrameIndex, TOptional<float> Tolerance, FString& OutErrors) const
 {
-	void ConvertPixel(const uint8* SourceData, EPixelFormat Format, FFloat16Color& OutColor)
-	{
-		switch (Format)
-		{
-		case PF_FloatRGB:
-			OutColor.R = *reinterpret_cast<const float*>(SourceData);
-			OutColor.G = *reinterpret_cast<const float*>(SourceData + 4);
-			OutColor.B = *reinterpret_cast<const float*>(SourceData + 8);
-		case PF_FloatRGBA:
-			OutColor.R = *reinterpret_cast<const float*>(SourceData);
-			OutColor.G = *reinterpret_cast<const float*>(SourceData + 4);
-			OutColor.B = *reinterpret_cast<const float*>(SourceData + 8);
-			OutColor.A = *reinterpret_cast<const float*>(SourceData + 12);
-		case PF_G16R16F:
-		case PF_G16R16F_FILTER:
-			FMemory::Memcpy(&OutColor.R, SourceData, 4); //-V512
-		default: // do nothing;
-			break;
-		}
-	}
+	const UNDIRenderTargetSimCacheData* LhsSimCacheData = CastChecked<UNDIRenderTargetSimCacheData>(LhsStorageObject);
+	const UNDIRenderTargetSimCacheData* RhsSimCacheData = CastChecked<UNDIRenderTargetSimCacheData>(RhsStorageObject);
+	return LhsSimCacheData->CompareFrame(RhsSimCacheData, FrameIndex, Tolerance, OutErrors);
 }
 
 TSharedPtr<FJsonObject> UNiagaraDataInterfaceRenderTarget2D::SimCacheToJson(const UObject* StorageObject, int FrameIndex, TOptional<FString> TargetFolder, TOptional<FString> FilenamePrefix) const
 {
-	const UNDIRenderTargetVolumeSimCacheData* SimCacheData = CastChecked<UNDIRenderTargetVolumeSimCacheData>(StorageObject);
-	if (SimCacheData->Frames.IsValidIndex(FrameIndex))
-	{
-		FNDIRenderTargetVolumeSimCacheFrame CacheFrame = SimCacheData->Frames[FrameIndex];
-		TSharedPtr<FJsonObject> JsonCacheObject = MakeShared<FJsonObject>();
-		JsonCacheObject->SetStringField(TEXT("CompressionType"), SimCacheData->CompressionType.ToString());
-		JsonCacheObject->SetNumberField(TEXT("PixelFormat"), static_cast<int32>(CacheFrame.Format));
-		JsonCacheObject->SetNumberField(TEXT("CompressedSize"), CacheFrame.CompressedSize);
-		JsonCacheObject->SetNumberField(TEXT("UncompressedSize"), CacheFrame.UncompressedSize);
-		JsonCacheObject->SetNumberField(TEXT("SizeX"), CacheFrame.Size.X);
-		JsonCacheObject->SetNumberField(TEXT("SizeY"), CacheFrame.Size.Y);
-		
-		const uint8* PixelData = CacheFrame.GetPixelData();
-		if (!PixelData || !TargetFolder.IsSet() || !FilenamePrefix.IsSet())
-		{
-			return JsonCacheObject;
-		}
-
-		ERawImageFormat::Type ImageFormat = FImageCoreUtils::GetRawImageFormatForPixelFormat(CacheFrame.Format);
-		if (ImageFormat != ERawImageFormat::RGBA16F)
-		{
-			UE_LOG(LogNiagara, Error, TEXT("Unable to save render target to file with current pixel format"));
-			return JsonCacheObject;
-		}
-
-		uint32 PixelCount = CacheFrame.Size.X * CacheFrame.Size.Y;
-		uint32 BlockBytes = GPixelFormats[CacheFrame.Format].BlockBytes;
-		TArray<uint8> Decompressed;
-		if (CacheFrame.CompressedSize > 0)
-		{
-			Decompressed.AddUninitialized(BlockBytes * PixelCount);
-			if (!FCompression::UncompressMemory(SimCacheData->CompressionType, Decompressed.GetData(), Decompressed.Num(), PixelData, CacheFrame.CompressedSize))
-			{
-				UE_LOG(LogNiagara, Error, TEXT("Error decompressing render target data"));
-				return JsonCacheObject;
-			}
-		}
-		const uint8* SrcData = Decompressed.Num() > 0 ? Decompressed.GetData() : PixelData;
-
-		// we need to convert the raw pixel data into a format that the image utils understand 
-		TArray<FFloat16Color> ImagePixelData;
-		ImagePixelData.AddZeroed(PixelCount);
-		for (uint32 Index = 0; Index < PixelCount; Index++)
-		{
-			uint32 Offset = Index * BlockBytes;
-			NiagaraRenderTarget2DImageHelper::ConvertPixel(SrcData + Offset, CacheFrame.Format, ImagePixelData[Index]);
-		}
-		
-		FImageView ImageView(ImagePixelData.GetData(), CacheFrame.Size.X, CacheFrame.Size.Y, ImageFormat);
-		if (FImageUtils::SaveImageByExtension(*FPaths::Combine(TargetFolder.GetValue(), FilenamePrefix.GetValue() + ".exr"), ImageView))
-		{
-			JsonCacheObject->SetStringField(TEXT("TextureData"), FilenamePrefix.GetValue() + ".exr");
-			return JsonCacheObject;
-		}		
-	}
-	return TSharedPtr<FJsonObject>();
+	const UNDIRenderTargetSimCacheData* SimCacheData = CastChecked<UNDIRenderTargetSimCacheData>(StorageObject);
+	return SimCacheData->FrameToJson(FrameIndex, TargetFolder, FilenamePrefix);
 }
 
 void UNiagaraDataInterfaceRenderTarget2D::VMSetSize(FVectorVMExternalFunctionContext& Context)
diff --git a/Engine/Plugins/FX/Niagara/Source/Niagara/Private/NiagaraDataInterfaceRenderTargetVolume.cpp b/Engine/Plugins/FX/Niagara/Source/Niagara/Private/NiagaraDataInterfaceRenderTargetVolume.cpp
--- a/Engine/Plugins/FX/Niagara/Source/Niagara/Private/NiagaraDataInterfaceRenderTargetVolume.cpp
+++ b/Engine/Plugins/FX/Niagara/Source/Niagara/Private/NiagaraDataInterfaceRenderTargetVolume.cpp
@@ -8,7 +8,7 @@
 #include "Engine/TextureRenderTargetVolume.h"
 #include "HAL/PlatformFileManager.h"
 
-#include "NDIRenderTargetVolumeSimCacheData.h"
+#include "DataInterface/NDIRenderTargetSimCacheData.h"
 #include "NiagaraDataInterfaceRenderTargetCommon.h"
 #include "NiagaraGpuComputeDispatchInterface.h"
 #include "NiagaraSettings.h"
@@ -600,17 +600,15 @@ UObject* UNiagaraDataInterfaceRenderTargetVolume::SimCacheBeginWrite(UObject* In
 
 		return CurrCache;
 	}	
-	else
+	else if (NDIRenderTargetVolumeLocal::GSimCacheEnabled)
 	{
-		UNDIRenderTargetVolumeSimCacheData* SimCacheData = nullptr;
-		if (NDIRenderTargetVolumeLocal::GSimCacheEnabled)
-		{
-			SimCacheData = NewObject<UNDIRenderTargetVolumeSimCacheData>(InSimCache);
-			SimCacheData->CompressionType = NDIRenderTargetVolumeLocal::GetSimCacheCompressionType();
-		}
+		UNDIRenderTargetSimCacheData* SimCacheData = NewObject<UNDIRenderTargetSimCacheData>(InSimCache);
+		SimCacheData->CompressionType = NDIRenderTargetVolumeLocal::GetSimCacheCompressionType();
 
 		return SimCacheData;
 	}
+
+	return nullptr;
 }
 
 bool UNiagaraDataInterfaceRenderTargetVolume::SimCacheWriteFrame(UObject* StorageObject, int FrameIndex, FNiagaraSystemInstance* SystemInstance, const void* OptionalPerInstanceData, FNiagaraSimCacheFeedbackContext& FeedbackContext) const
@@ -702,45 +700,13 @@ bool UNiagaraDataInterfaceRenderTargetVolume::SimCacheWriteFrame(UObject* Storag
 			}
 			else
 			{
-				UNDIRenderTargetVolumeSimCacheData* SimCacheData = CastChecked<UNDIRenderTargetVolumeSimCacheData>(StorageObject);
-				SimCacheData->Frames.SetNum(FMath::Max(SimCacheData->Frames.Num(), FrameIndex + 1));
-
-				FNDIRenderTargetVolumeSimCacheFrame* CacheFrame = &SimCacheData->Frames[FrameIndex];
-				
-				CacheFrame->Size = InstanceData_GT->Size;
-				CacheFrame->Format = InstanceData_GT->Format;
-
-				const FName CompressionType = SimCacheData->CompressionType;
-				const bool bUseCompression = CompressionType.IsNone() == false;
-				const int TextureSizeBytes = TextureData.Num() * TextureData.GetTypeSize();
-				int CompressedSize = bUseCompression ? FCompression::CompressMemoryBound(CompressionType, TextureSizeBytes) : TextureSizeBytes;
-
-				CacheFrame->UncompressedSize = TextureSizeBytes;
-				CacheFrame->CompressedSize = 0;
-				uint8* FinalPixelData = reinterpret_cast<uint8*>(FMemory::Malloc(CompressedSize));
-				if (bUseCompression)
-				{
-					if (FCompression::CompressMemory(CompressionType, FinalPixelData, CompressedSize, TextureData.GetData(), TextureSizeBytes, COMPRESS_BiasMemory))
-					{
-						CacheFrame->CompressedSize = CompressedSize;
-					}
-				}
-
-				if (CacheFrame->CompressedSize == 0)
-				{
-					FMemory::Memcpy(FinalPixelData, TextureData.GetData(), TextureSizeBytes);
-				}
-
-				// Update bulk data
-				CacheFrame->BulkData.Lock(LOCK_READ_WRITE);
-				{
-					const int32 FinalByteSize = CacheFrame->CompressedSize > 0 ? CacheFrame->CompressedSize : CacheFrame->UncompressedSize;
-					uint8* BulkDataPtr = reinterpret_cast<uint8*>(CacheFrame->BulkData.Realloc(FinalByteSize));
-					FMemory::Memcpy(BulkDataPtr, FinalPixelData, FinalByteSize);
-				}
-				CacheFrame->BulkData.Unlock();
-
-				FMemory::Free(FinalPixelData);
+				UNDIRenderTargetSimCacheData* SimCacheData = CastChecked<UNDIRenderTargetSimCacheData>(StorageObject);
+				SimCacheData->SetPixelData(
+					FrameIndex,
+					InstanceData_GT->Size,
+					InstanceData_GT->Format,
+					TextureData
+				);
 			}
 		}
 	}
@@ -888,73 +854,34 @@ bool UNiagaraDataInterfaceRenderTargetVolume::SimCacheReadFrame(UObject* Storage
 	}
 	else
 	{
-		UNDIRenderTargetVolumeSimCacheData* SimCacheData = CastChecked<UNDIRenderTargetVolumeSimCacheData>(StorageObject);
+		UNDIRenderTargetSimCacheData* SimCacheData = CastChecked<UNDIRenderTargetSimCacheData>(StorageObject);
 
 		const int FrameIndex = Interp >= 0.5f ? FrameB : FrameA;
-		if (!SimCacheData->Frames.IsValidIndex(FrameIndex))
+		if (!SimCacheData->IsValidFrame(FrameIndex))
 		{
 			return false;
 		}
 
-		FNDIRenderTargetVolumeSimCacheFrame* CacheFrame = &SimCacheData->Frames[FrameIndex];
-
-		InstanceData_GT->Size = CacheFrame->Size;
-		InstanceData_GT->Format = CacheFrame->Format;
+		InstanceData_GT->Size   = SimCacheData->GetTextureSize(FrameIndex);
+		InstanceData_GT->Format = SimCacheData->GetTextureFormat(FrameIndex);
 		InstanceData_GT->Filter = TextureFilter::TF_Default;
 
 		PerInstanceTick(InstanceData_GT, SystemInstance, 0.0f);
 		PerInstanceTickPostSimulate(InstanceData_GT, SystemInstance, 0.0f);
 
-		if (InstanceData_GT->TargetTexture && CacheFrame->GetPixelData() != nullptr)
+		if (InstanceData_GT->TargetTexture && SimCacheData->HasPixelData(FrameIndex))
 		{
 			FNiagaraDataInterfaceProxyRenderTargetVolumeProxy* RT_Proxy = GetProxyAs<FNiagaraDataInterfaceProxyRenderTargetVolumeProxy>();
 			FTextureRenderTargetResource* RT_TargetTexture = InstanceData_GT->TargetTexture->GameThread_GetRenderTargetResource();
 			ENQUEUE_RENDER_COMMAND(NDIRenderTargetVolumeUpdate)
 			(
-				[RT_Proxy, RT_InstanceID=SystemInstance->GetId(), RT_TargetTexture, RT_CacheFrame=CacheFrame, RT_CompressionType=SimCacheData->CompressionType, RT_Format = InstanceData_GT->Format](FRHICommandListImmediate& RHICmdList)
+				[RT_Proxy, RT_InstanceID=SystemInstance->GetId(), RT_TargetTexture, RT_SimCacheData=SimCacheData, RT_FrameIndex=FrameIndex](FRHICommandListImmediate& RHICmdList)
 				{
 					if (FRenderTargetVolumeRWInstanceData_RenderThread* InstanceData_RT = RT_Proxy->SystemInstancesToProxyData_RT.Find(RT_InstanceID))
 					{
 						FUpdateTextureRegion3D UpdateRegion(FIntVector::ZeroValue, FIntVector::ZeroValue, InstanceData_RT->Size);
-
 						FUpdateTexture3DData UpdateTexture = RHICmdList.BeginUpdateTexture3D(InstanceData_RT->RenderTarget->GetRHI(), 0, UpdateRegion);
-
-						uint32 BlockBytes = GPixelFormats[RT_Format].BlockBytes;
-
-						const int32 SrcRowPitch = InstanceData_RT->Size.X * BlockBytes;
-						const int32 SrcDepthPitch = InstanceData_RT->Size.Y * SrcRowPitch;
-						if (UpdateTexture.RowPitch == SrcRowPitch && UpdateTexture.DepthPitch == SrcDepthPitch)
-						{
-							if (RT_CacheFrame->CompressedSize > 0)
-							{
-								FCompression::UncompressMemory(RT_CompressionType, UpdateTexture.Data, UpdateTexture.DataSizeBytes, RT_CacheFrame->GetPixelData(), RT_CacheFrame->CompressedSize);
-							}
-							else
-							{
-								FMemory::Memcpy(UpdateTexture.Data, RT_CacheFrame->GetPixelData(), InstanceData_RT->Size.X * InstanceData_RT->Size.Y * sizeof(FFloat16Color));
-							}
-						}
-						else
-						{
-							TArray<uint8> Decompressed;
-							if (RT_CacheFrame->CompressedSize > 0)
-							{
-								Decompressed.AddUninitialized(BlockBytes * InstanceData_RT->Size.X * InstanceData_RT->Size.Y * InstanceData_RT->Size.Z);
-								FCompression::UncompressMemory(RT_CompressionType, Decompressed.GetData(), Decompressed.Num(), RT_CacheFrame->GetPixelData(), RT_CacheFrame->CompressedSize);
-							}
-
-							const uint8* SrcData = Decompressed.Num() > 0 ? Decompressed.GetData() : RT_CacheFrame->GetPixelData();
-							for (int32 z = 0; z < InstanceData_RT->Size.Z; ++z)
-							{
-								uint8* DstData = UpdateTexture.Data + (z * UpdateTexture.DepthPitch);
-								for (int32 y = 0; y < InstanceData_RT->Size.Y; ++y)
-								{
-									FMemory::Memcpy(DstData, SrcData, SrcRowPitch);
-									SrcData += SrcRowPitch;
-									DstData += UpdateTexture.RowPitch;
-								}
-							}
-						}
+						RT_SimCacheData->GetPixelData(RT_FrameIndex, UpdateTexture.Data, UpdateTexture.RowPitch, UpdateTexture.DepthPitch);
 						RHICmdList.EndUpdateTexture3D(UpdateTexture);
 					}
 				}
@@ -1143,6 +1070,11 @@ bool UNiagaraDataInterfaceRenderTargetVolume::PerInstanceTickPostSimulate(void*
 			//(InstanceData->TargetTexture->bAutoGenerateMips != bAutoGenerateMips) ||
 			!InstanceData->TargetTexture->GetResource())
 		{
+			if (InstanceData->Format == PF_Unknown)
+			{
+				return false;
+			}
+
 			// resize RT to match what we need for the output
 			InstanceData->TargetTexture->bCanCreateUAV = true;
 			//InstanceData->TargetTexture->bAutoGenerateMips = bAutoGenerateMips;
diff --git a/Engine/Plugins/FX/Niagara/Source/NiagaraEditor/Private/Customizations/SimCache/NiagaraRenderTargetSimCacheVisualizer.cpp b/Engine/Plugins/FX/Niagara/Source/NiagaraEditor/Private/Customizations/SimCache/NiagaraRenderTargetSimCacheVisualizer.cpp
new file mode 100644
--- /dev/null
+++ b/Engine/Plugins/FX/Niagara/Source/NiagaraEditor/Private/Customizations/SimCache/NiagaraRenderTargetSimCacheVisualizer.cpp
@@ -0,0 +1,319 @@
+// Copyright Epic Games, Inc. All Rights Reserved.
+
+#include "NiagaraRenderTargetSimCacheVisualizer.h"
+
+#include "DataInterface/NDIRenderTargetSimCacheData.h"
+
+#include "Engine/Canvas.h"
+#include "Engine/Texture2D.h"
+#include "UObject/StrongObjectPtr.h"
+#include "Slate/SceneViewport.h"
+#include "ViewModels/NiagaraSimCacheViewModel.h"
+#include "Widgets/Input/SNumericEntryBox.h"
+#include "Widgets/Input/SSpinBox.h"
+#include "Widgets/Layout/SGridPanel.h"
+#include "Widgets/Text/STextBlock.h"
+#include "CanvasItem.h"
+#include "DetailLayoutBuilder.h"
+#include "EditorViewportClient.h"
+#include "SCommonEditorViewportToolbarBase.h"
+#include "SEnumCombo.h"
+#include "SEditorViewport.h"
+#include "TextureResource.h"
+
+#define LOCTEXT_NAMESPACE "NiagaraRenderTargetSimCacheVisualizer"
+
+////////////////////////////////////////////////////////////////////////////////////////////////////////////////////
+
+namespace NDIRenderTargetSimCacheVisualizer
+{
+class FTextureViewportClient final : public FEditorViewportClient
+{
+public:
+	FTextureViewportClient(const TSharedRef<SEditorViewport>& InOwnerViewport, TSharedPtr<FNiagaraSimCacheViewModel> InViewModel, const UNDIRenderTargetSimCacheData* InCacheData)
+		: FEditorViewportClient(nullptr, nullptr, StaticCastSharedRef<SEditorViewport>(InOwnerViewport))
+		, ViewModel(InViewModel)
+		, CacheData(InCacheData)
+	{
+		ViewModel->OnViewDataChanged().AddRaw(this, &FTextureViewportClient::OnViewDataChanged);
+
+		UpdateTexture();
+	}
+
+	virtual ~FTextureViewportClient()
+	{
+		ViewModel->OnViewDataChanged().RemoveAll(this);
+	}
+
+	virtual void Draw(FViewport* InViewport, FCanvas* Canvas) override
+	{
+		Canvas->Clear(FLinearColor::Transparent);
+
+		const float DPIScaleFactor = ShouldDPIScaleSceneCanvas() ? GetDPIScale() : 1.0f;
+		const FIntRect ViewRect(
+			FMath::Floor(Canvas->GetViewRect().Min.X / DPIScaleFactor) + 2,
+			FMath::Floor(Canvas->GetViewRect().Min.Y / DPIScaleFactor) + 2,
+			FMath::Floor(Canvas->GetViewRect().Max.X / DPIScaleFactor) - 2,
+			FMath::Floor(Canvas->GetViewRect().Max.Y / DPIScaleFactor) - 2
+		);
+		if (ViewRect.Width() <= 0 || ViewRect.Height() <= 0)
+		{
+			return;
+		}
+
+		const FVector2D HalfPixel(0.5f / float(Texture->GetSizeX()), 0.5f / float(Texture->GetSizeY()));
+		FCanvasTileItem TileItem(
+			FVector2D(0.0f, 0.0f),
+			Texture->GetResource(),
+			FVector2D(ViewRect.Width(), ViewRect.Height()),
+			FVector2D(HalfPixel.X, HalfPixel.Y),
+			FVector2D(1.0f - HalfPixel.X, 1.0f - HalfPixel.Y),
+			FLinearColor::White
+		);
+		Canvas->DrawItem(TileItem);
+	}
+
+	void UpdateTexture()
+	{
+		const int32 FrameIndex = ViewModel->GetFrameIndex();
+		const bool bIsValid = CacheData->HasPixelData(FrameIndex);
+		const FIntVector TextureSize = bIsValid ? CacheData->GetTextureSize(FrameIndex) : FIntVector(1, 1, 1);
+		const EPixelFormat TextureFormat = bIsValid ? CacheData->GetTextureFormat(FrameIndex) : EPixelFormat::PF_B8G8R8A8;
+
+		if (Texture == nullptr || Texture->GetSizeX() != TextureSize.X || Texture->GetSizeY() != TextureSize.Y || Texture->GetPixelFormat() != TextureFormat )
+		{
+			Texture.Reset(UTexture2D::CreateTransient(TextureSize.X, TextureSize.Y, TextureFormat));
+			Texture->MipGenSettings = TMGS_NoMipmaps;
+			Texture->UpdateResource();
+		}
+
+		void* UntypedMipData = Texture->GetPlatformData()->Mips[0].BulkData.Lock(LOCK_READ_WRITE);
+
+		if (bIsValid)
+		{
+			//-TODO: This assumes single slice (i.e. 2D) when we use this viualizer for other types it will need updating.
+			checkf(TextureSize.Z == 1, TEXT("Only one texture slice is supported right now"));
+			TArray<uint8> Data = CacheData->GetPixelData(FrameIndex);
+			FMemory::Memcpy(UntypedMipData, Data.GetData(), Data.Num());
+		}
+		else
+		{
+			FColor* TypedMipData = static_cast<FColor*>(UntypedMipData);
+			TypedMipData[0] = FColor::Black;
+		}
+
+		Texture->GetPlatformData()->Mips[0].BulkData.Unlock();
+		Texture->UpdateResource();
+
+		bNeedsRedraw = true;
+	}
+
+	void OnViewDataChanged(bool)
+	{
+		UpdateTexture();
+	}
+
+	FVector2D GetTextureSize() const
+	{
+		return FVector2D(Texture->GetSizeX(), Texture->GetSizeY());
+	}
+
+private:
+	TSharedPtr<FNiagaraSimCacheViewModel>					ViewModel;
+	TStrongObjectPtr<const UNDIRenderTargetSimCacheData>	CacheData;
+	TStrongObjectPtr<UTexture2D>							Texture;
+};
+
+class STextureViewport : public SEditorViewport
+{
+public:
+	SLATE_BEGIN_ARGS(STextureViewport) {}
+		SLATE_ARGUMENT(TSharedPtr<FNiagaraSimCacheViewModel>, ViewModel)
+		SLATE_ARGUMENT(const UNDIRenderTargetSimCacheData*, CacheData)
+	SLATE_END_ARGS()
+
+	void Construct(const FArguments& InArgs)
+	{
+		ViewModel = InArgs._ViewModel;
+		CacheData.Reset(InArgs._CacheData);
+
+		SEditorViewport::Construct(
+			SEditorViewport::FArguments()
+			.ViewportSize(this, &STextureViewport::GetViewportSize)
+		);
+	}
+
+	virtual ~STextureViewport() override
+	{
+	}
+
+	TSharedRef<FEditorViewportClient> MakeEditorViewportClient() override
+	{
+		ViewportClient = MakeShareable(new FTextureViewportClient(SharedThis(this), ViewModel, CacheData.Get()));
+		return ViewportClient.ToSharedRef();
+	}
+
+	FVector2D GetViewportSize() const
+	{
+		return ViewportClient ? ViewportClient->GetTextureSize() * ZoomLevel : FVector2D(1.0f);
+	}
+
+	TOptional<float> GetZoomLevel() const
+	{
+		return TOptional<float>(ZoomLevel);
+	}
+
+	void SetZoomLevel(float InZoomLevel)
+	{
+		if (FMath::IsNearlyEqual(ZoomLevel, InZoomLevel))
+		{
+			return;
+		}
+		ViewportClient->RedrawRequested(nullptr);
+		ZoomLevel = InZoomLevel;
+	}
+
+private:
+	TSharedPtr<FNiagaraSimCacheViewModel>					ViewModel;
+	TStrongObjectPtr<const UNDIRenderTargetSimCacheData>	CacheData;
+	TSharedPtr<FTextureViewportClient>						ViewportClient;
+
+	float	ZoomLevel = 1.0f;
+};
+
+class SSimCacheView : public SCompoundWidget
+{
+public:
+	SLATE_BEGIN_ARGS(SSimCacheView)
+	{}
+	SLATE_END_ARGS()
+
+	virtual ~SSimCacheView() override
+	{
+		ViewModel->OnViewDataChanged().RemoveAll(this);
+	}
+
+	void Construct(const FArguments& InArgs, TSharedPtr<FNiagaraSimCacheViewModel> InViewModel, const UNDIRenderTargetSimCacheData* InCacheData)
+	{
+		ViewModel = InViewModel;
+		CacheData.Reset(InCacheData);
+
+		TSharedRef<STextureViewport> Viewport =
+			SNew(STextureViewport)
+			.ViewModel(ViewModel)
+			.CacheData(CacheData.Get());
+
+		ChildSlot
+		[
+			SNew(SVerticalBox)
+			+ SVerticalBox::Slot()
+			.Padding(10.0f)
+			.AutoHeight()
+			[
+				SNew(SGridPanel)
+				+SGridPanel::Slot(0, 0)
+				.HAlign(HAlign_Right)
+				.Padding(5, 0)
+				[
+					SNew(STextBlock)
+					.Text(LOCTEXT("TextureDetails", "Texture Details:"))
+				]
+				+SGridPanel::Slot(1, 0)
+				.Padding(5, 0)
+				[
+					SNew(STextBlock)
+					.Text(this, &SSimCacheView::GetTextureDetails)
+				]
+				+SGridPanel::Slot(0, 1)
+				.HAlign(HAlign_Right)
+				.Padding(5, 0)
+				[
+					SNew(STextBlock)
+					.Text(LOCTEXT("TextureMemory", "Texture Memory:"))
+				]
+				+SGridPanel::Slot(1, 1)
+				.Padding(5, 0)
+				[
+					SNew(STextBlock)
+					.Text(this, &SSimCacheView::GetTextureMemoryDetails)
+				]
+				+SGridPanel::Slot(0, 2)
+				.HAlign(HAlign_Right)
+				.Padding(5, 0)
+				[
+					SNew(STextBlock)
+					.Text(LOCTEXT("Zoom", "Zoom:"))
+				]
+				+SGridPanel::Slot(1, 2)
+				.Padding(5, 0)
+				[
+					SNew(SNumericEntryBox<float>)
+					.AllowSpin(true)
+					.MinDesiredValueWidth(80)
+					.MinSliderValue(0.25f)
+					.MaxSliderValue(16.0f)
+					.Delta(0.25f)
+					.Value(Viewport, &STextureViewport::GetZoomLevel)
+					.OnValueChanged(Viewport, &STextureViewport::SetZoomLevel)
+				]
+			]
+			+ SVerticalBox::Slot()
+			.AutoHeight()
+			[
+				SNew(SHorizontalBox)
+				+ SHorizontalBox::Slot()
+				.AutoWidth()
+				[
+					Viewport
+				]
+			]
+		];
+	}
+
+	FText GetTextureDetails() const
+	{
+		const int32 FrameIndex = ViewModel->GetFrameIndex();
+		const bool bIsValid = CacheData->IsValidFrame(FrameIndex);
+		const FIntVector TextureSize = bIsValid ? CacheData->GetTextureSize(FrameIndex) : FIntVector(0,0,0);
+		const EPixelFormat TextureFormat = bIsValid ? CacheData->GetTextureFormat(FrameIndex) : EPixelFormat::PF_Unknown;
+		return FText::Format(
+			LOCTEXT("TextureDetailsFormat", "{0}x{1}x{2} {3}"),
+			TextureSize.X, TextureSize.Y, TextureSize.Z,
+			FText::FromString(GPixelFormats[TextureFormat].Name)
+		);
+	}
+
+	FText GetTextureMemoryDetails() const
+	{
+		const int32 FrameIndex = ViewModel->GetFrameIndex();
+		const bool bIsValid = CacheData->IsValidFrame(FrameIndex);
+		const int32 CompressedSize = CacheData->GetCompressedSize(FrameIndex);
+		const int32 UncompressedSize = CacheData->GetUncompressedSize(FrameIndex);
+		return FText::Format(
+			LOCTEXT("TextureMemoryDetailsFormat", "{0}mb (Compressed) {1}mb (Decompressed)"),
+			FText::AsNumber(float(CompressedSize) / 1024.0f / 1024.0f),
+			FText::AsNumber(float(UncompressedSize) / 1024.0f / 1024.0f)
+		);
+	}
+
+private:
+	TSharedPtr<FNiagaraSimCacheViewModel>					ViewModel;
+	TStrongObjectPtr<const UNDIRenderTargetSimCacheData>	CacheData;
+};
+
+} // NDIRenderTargetSimCacheVisualizer
+
+////////////////////////////////////////////////////////////////////////////////////////////////////////////////////
+
+TSharedPtr<SWidget> FNiagaraRenderTargetSimCacheVisualizer::CreateWidgetFor(const UObject* InCachedData, TSharedPtr<FNiagaraSimCacheViewModel> ViewModel)
+{
+	using namespace NDIRenderTargetSimCacheVisualizer;
+
+	if (const UNDIRenderTargetSimCacheData* CachedData = Cast<const UNDIRenderTargetSimCacheData>(InCachedData))
+	{
+		return SNew(SSimCacheView, ViewModel, CachedData);
+	}
+	return TSharedPtr<SWidget>();
+}
+
+#undef LOCTEXT_NAMESPACE
diff --git a/Engine/Plugins/FX/Niagara/Source/NiagaraEditor/Private/Customizations/SimCache/NiagaraRenderTargetSimCacheVisualizer.h b/Engine/Plugins/FX/Niagara/Source/NiagaraEditor/Private/Customizations/SimCache/NiagaraRenderTargetSimCacheVisualizer.h
new file mode 100644
--- /dev/null
+++ b/Engine/Plugins/FX/Niagara/Source/NiagaraEditor/Private/Customizations/SimCache/NiagaraRenderTargetSimCacheVisualizer.h
@@ -0,0 +1,15 @@
+// Copyright Epic Games, Inc. All Rights Reserved.
+
+#pragma once
+
+#include "NiagaraCommon.h"
+#include "Customizations/NiagaraDataInterfaceSimCacheVisualizer.h"
+
+class FNiagaraRenderTargetSimCacheVisualizer : public INiagaraDataInterfaceSimCacheVisualizer
+{
+public:
+	FNiagaraRenderTargetSimCacheVisualizer() {}
+	virtual ~FNiagaraRenderTargetSimCacheVisualizer() override = default;
+	
+	virtual TSharedPtr<SWidget> CreateWidgetFor(const UObject* CachedData, TSharedPtr<FNiagaraSimCacheViewModel> ViewModel) override;
+};
diff --git a/Engine/Plugins/FX/Niagara/Source/NiagaraEditor/Private/NiagaraEditorModule.cpp b/Engine/Plugins/FX/Niagara/Source/NiagaraEditor/Private/NiagaraEditorModule.cpp
--- a/Engine/Plugins/FX/Niagara/Source/NiagaraEditor/Private/NiagaraEditorModule.cpp
+++ b/Engine/Plugins/FX/Niagara/Source/NiagaraEditor/Private/NiagaraEditorModule.cpp
@@ -59,6 +59,7 @@
 #include "DataInterface/NiagaraDataInterfaceDataChannelWrite.h"
 #include "DataInterface/NiagaraDataInterfaceMemoryBuffer.h"
 #include "DataInterface/NiagaraDataInterfaceSimpleCounter.h"
+#include "NiagaraDataInterfaceRenderTarget2D.h"
 #include "NiagaraDataInterfaceRenderTargetVolume.h"
 
 #include "ViewModels/NiagaraScriptViewModel.h"
@@ -122,6 +123,7 @@
 #include "Customizations/SimCache/NiagaraArraySimCacheVisualizer.h"
 #include "Customizations/SimCache/FNiagaraDataChannelSimCacheVisualizer.h"
 #include "Customizations/SimCache/NiagaraMemoryBufferSimCacheVisualizer.h"
+#include "Customizations/SimCache/NiagaraRenderTargetSimCacheVisualizer.h"
 #include "Customizations/SimCache/NiagaraRenderTargetVolumeSimCacheVisualizer.h"
 #include "Customizations/SimCache/NiagaraSimpleCounterSimCacheVisualizer.h"
 
@@ -1561,6 +1563,7 @@ void FNiagaraEditorModule::StartupModule()
 	// Register sim cache visualizers
 	RegisterDataInterfaceCacheVisualizer(UNiagaraDataInterfaceDataChannelWrite::StaticClass(), MakeShared<FNiagaraDataChannelSimCacheVisualizer>());
 	RegisterDataInterfaceCacheVisualizer(UNiagaraDataInterfaceMemoryBuffer::StaticClass(), MakeShared<FNiagaraMemoryBufferSimCacheVisualizer>());
+	RegisterDataInterfaceCacheVisualizer(UNiagaraDataInterfaceRenderTarget2D::StaticClass(), MakeShared<FNiagaraRenderTargetSimCacheVisualizer>());
 	RegisterDataInterfaceCacheVisualizer(UNiagaraDataInterfaceRenderTargetVolume::StaticClass(), MakeShared<FNiagaraRenderTargetVolumeSimCacheVisualizer>());
 	RegisterDataInterfaceCacheVisualizer(UNiagaraDataInterfaceSimpleCounter::StaticClass(), MakeShared<FNiagaraSimpleCounterSimCacheVisualizer>());
 	for (TObjectIterator<UClass> It; It; ++It)
