- Add Export To Disk button to sim cache viewer
- This allows you to export transient / none asset caches for further insepection, i.e. those captured in the attribute spreadsheet / from running clients

#tests Exported a cache from the cache tests where the cache is a child object of an asset

diff --git a/Engine/Plugins/FX/Niagara/Source/NiagaraEditor/Private/AssetDefinitions/AssetDefinition_NiagaraSimCache.cpp b/Engine/Plugins/FX/Niagara/Source/NiagaraEditor/Private/AssetDefinitions/AssetDefinition_NiagaraSimCache.cpp
--- a/Engine/Plugins/FX/Niagara/Source/NiagaraEditor/Private/AssetDefinitions/AssetDefinition_NiagaraSimCache.cpp
+++ b/Engine/Plugins/FX/Niagara/Source/NiagaraEditor/Private/AssetDefinitions/AssetDefinition_NiagaraSimCache.cpp
@@ -1,6 +1,7 @@
 // Copyright Epic Games, Inc. All Rights Reserved.
 
 #include "AssetDefinition_NiagaraSimCache.h"
+#include "NiagaraEditorSimCacheUtils.h"
 
 #include "ContentBrowserMenuContexts.h"
 #include "DesktopPlatformModule.h"
@@ -15,8 +16,6 @@
 
 #define LOCTEXT_NAMESPACE "UAssetDefinition_NiagaraSimCache"
 
-DEFINE_LOG_CATEGORY_STATIC(LogNiagaraSimCache, Log, All);
-
 FLinearColor UAssetDefinition_NiagaraSimCache::GetAssetColor() const
 {
 	return FNiagaraEditorStyle::Get().GetColor("NiagaraEditor.AssetColors.SimCache").ToFColor(true);
@@ -38,67 +37,13 @@ EAssetCommandResult UAssetDefinition_NiagaraSimCache::OpenAssets(const FAssetOpe
 
 namespace MenuExtension_NiagaraSimCache
 {	
-	void ExportToDisk(const FToolMenuContext& InContext)
+	static void ExportToDisk(const FToolMenuContext& InContext)
 	{
-		FString ExportFolder;
-		
-		if (IDesktopPlatform* DesktopPlatform = FDesktopPlatformModule::Get())
-		{
-			const bool bFolderPicked = DesktopPlatform->OpenDirectoryDialog(
-				FSlateApplication::Get().FindBestParentWindowHandleForDialogs(nullptr),
-				LOCTEXT("ExportSimCache", "Pick SimCache Export Folder" ).ToString(),
-				*(FEditorDirectories::Get().GetLastDirectory(ELastDirectory::GENERIC_EXPORT)),
-				ExportFolder);
-
-			if (!bFolderPicked)
-			{
-				return;
-			}
-			FEditorDirectories::Get().SetLastDirectory(ELastDirectory::GENERIC_EXPORT, ExportFolder); 
-		}
-		
-		bool bWarning = false;
 		const UContentBrowserAssetContextMenuContext* CBContext = UContentBrowserAssetContextMenuContext::FindContextWithAssets(InContext);
-		IFileManager& FileManager = IFileManager::Get();
-		for (UNiagaraSimCache* Cache : CBContext->LoadSelectedObjects<UNiagaraSimCache>())
-		{
-			if (Cache == nullptr)
-			{
-				continue;
-			}
-
-			FString CacheRootFolder = FPaths::Combine(ExportFolder, FPaths::MakeValidFileName(Cache->GetName(), '_'));
-			if (FileManager.DirectoryExists(*CacheRootFolder))
-			{
-				if (!FileManager.DeleteDirectory(*CacheRootFolder, true, true))
-				{
-					UE_LOG(LogNiagaraSimCache, Warning, TEXT("Unable to delete existing folder %s"), *CacheRootFolder);
-					bWarning = true;
-					continue;
-				}
-			}
-			if (!FileManager.MakeDirectory(*CacheRootFolder, true))
-			{
-				UE_LOG(LogNiagaraSimCache, Warning, TEXT("Unable to create folder %s"), *CacheRootFolder);
-				bWarning = true;
-				continue;
-			}
-			if (!FNiagaraSimCacheJson::DumpToFile(*Cache, CacheRootFolder, FNiagaraSimCacheJson::EExportType::SeparateEachFrame))
-			{
-				bWarning = true;
-			}
-		}
-		
-		FNotificationInfo Info(LOCTEXT("ExportToDisk_DoneInfo", "Export completed."));
-		Info.ExpireDuration = 4.0f;
-		if (bWarning)
-		{
-			Info.Text = LOCTEXT("ExportData_DoneWarn", "Export completed with warnings.\nPlease check the log.");
-			Info.Image = FCoreStyle::Get().GetBrush(TEXT("MessageLog.Warning"));
-		}
-		FSlateNotificationManager::Get().AddNotification(Info);
+		TArray<UNiagaraSimCache*> Assets = CBContext->LoadSelectedObjects<UNiagaraSimCache>();
+		FNiagaraEditorSimCacheUtils::ExportToDisk(Assets);
 	}
-	
+
 	static FDelayedAutoRegisterHelper DelayedAutoRegister(EDelayedRegisterRunPhase::EndOfEngineInit, []{ 
 		UToolMenus::RegisterStartupCallback(FSimpleMulticastDelegate::FDelegate::CreateLambda([]()
 		{
diff --git a/Engine/Plugins/FX/Niagara/Source/NiagaraEditor/Private/NiagaraEditorSimCacheUtils.cpp b/Engine/Plugins/FX/Niagara/Source/NiagaraEditor/Private/NiagaraEditorSimCacheUtils.cpp
new file mode 100644
--- /dev/null
+++ b/Engine/Plugins/FX/Niagara/Source/NiagaraEditor/Private/NiagaraEditorSimCacheUtils.cpp
@@ -0,0 +1,98 @@
+// Copyright Epic Games, Inc. All Rights Reserved.
+
+#include "NiagaraEditorSimCacheUtils.h"
+#include "NiagaraSimCache.h"
+
+#include "ContentBrowserMenuContexts.h"
+#include "DesktopPlatformModule.h"
+#include "EditorDirectories.h"
+#include "NiagaraEditorStyle.h"
+#include "NiagaraSimCacheJson.h"
+#include "Framework/Application/SlateApplication.h"
+#include "Framework/Notifications/NotificationManager.h"
+#include "Toolkits/NiagaraSimCacheToolkit.h"
+#include "ToolMenus.h"
+#include "Widgets/Notifications/SNotificationList.h"
+
+DEFINE_LOG_CATEGORY_STATIC(LogNiagaraSimCache, Log, All);
+
+#define LOCTEXT_NAMESPACE "NiagaraEditorSimCacheUtils"
+
+namespace FNiagaraEditorSimCacheUtils
+{
+
+template<typename TArrayType>
+void ExportToDiskInternal(TArrayType CachesToExport)
+{
+	FString ExportFolder;
+		
+	if (IDesktopPlatform* DesktopPlatform = FDesktopPlatformModule::Get())
+	{
+		const bool bFolderPicked = DesktopPlatform->OpenDirectoryDialog(
+			FSlateApplication::Get().FindBestParentWindowHandleForDialogs(nullptr),
+			LOCTEXT("ExportSimCache", "Pick SimCache Export Folder" ).ToString(),
+			*(FEditorDirectories::Get().GetLastDirectory(ELastDirectory::GENERIC_EXPORT)),
+			ExportFolder);
+
+		if (!bFolderPicked)
+		{
+			return;
+		}
+		FEditorDirectories::Get().SetLastDirectory(ELastDirectory::GENERIC_EXPORT, ExportFolder); 
+	}
+		
+	bool bWarning = false;
+	IFileManager& FileManager = IFileManager::Get();
+	for (const UNiagaraSimCache* Cache : CachesToExport)
+	{
+		if (Cache == nullptr)
+		{
+			continue;
+		}
+
+		FString CacheRootFolder = FPaths::Combine(ExportFolder, FPaths::MakeValidFileName(Cache->GetName(), '_'));
+		if (FileManager.DirectoryExists(*CacheRootFolder))
+		{
+			if (!FileManager.DeleteDirectory(*CacheRootFolder, true, true))
+			{
+				UE_LOG(LogNiagaraSimCache, Warning, TEXT("Unable to delete existing folder %s"), *CacheRootFolder);
+				bWarning = true;
+				continue;
+			}
+		}
+		if (!FileManager.MakeDirectory(*CacheRootFolder, true))
+		{
+			UE_LOG(LogNiagaraSimCache, Warning, TEXT("Unable to create folder %s"), *CacheRootFolder);
+			bWarning = true;
+			continue;
+		}
+		if (!FNiagaraSimCacheJson::DumpToFile(*Cache, CacheRootFolder, FNiagaraSimCacheJson::EExportType::SeparateEachFrame))
+		{
+			bWarning = true;
+		}
+	}
+		
+	FNotificationInfo Info(LOCTEXT("ExportToDisk_DoneInfo", "Export completed."));
+	Info.ExpireDuration = 4.0f;
+	if (bWarning)
+	{
+		Info.Text = LOCTEXT("ExportData_DoneWarn", "Export completed with warnings.\nPlease check the log.");
+		Info.Image = FCoreStyle::Get().GetBrush(TEXT("MessageLog.Warning"));
+	}
+	FSlateNotificationManager::Get().AddNotification(Info);
+}
+
+} //namespace FNiagaraEditorSimCacheUtils
+
+
+void FNiagaraEditorSimCacheUtils::ExportToDisk(TArrayView<UNiagaraSimCache*> CachesToExport)
+{
+	ExportToDiskInternal(CachesToExport);
+}
+
+void FNiagaraEditorSimCacheUtils::ExportToDisk(const UNiagaraSimCache* CacheToExport)
+{
+	ExportToDiskInternal(MakeArrayView(&CacheToExport, 1));
+}
+
+#undef LOCTEXT_NAMESPACE
diff --git a/Engine/Plugins/FX/Niagara/Source/NiagaraEditor/Private/NiagaraEditorSimCacheUtils.h b/Engine/Plugins/FX/Niagara/Source/NiagaraEditor/Private/NiagaraEditorSimCacheUtils.h
new file mode 100644
--- /dev/null
+++ b/Engine/Plugins/FX/Niagara/Source/NiagaraEditor/Private/NiagaraEditorSimCacheUtils.h
@@ -0,0 +1,13 @@
+// Copyright Epic Games, Inc. All Rights Reserved.
+
+#pragma once
+
+#include "CoreMinimal.h"
+
+class UNiagaraSimCache;
+
+namespace FNiagaraEditorSimCacheUtils
+{
+	extern void ExportToDisk(TArrayView<UNiagaraSimCache*> CachesToExport);
+	extern void ExportToDisk(const UNiagaraSimCache* CacheToExport);
+}
diff --git a/Engine/Plugins/FX/Niagara/Source/NiagaraEditor/Private/Toolkits/NiagaraSimCacheToolkit.cpp b/Engine/Plugins/FX/Niagara/Source/NiagaraEditor/Private/Toolkits/NiagaraSimCacheToolkit.cpp
--- a/Engine/Plugins/FX/Niagara/Source/NiagaraEditor/Private/Toolkits/NiagaraSimCacheToolkit.cpp
+++ b/Engine/Plugins/FX/Niagara/Source/NiagaraEditor/Private/Toolkits/NiagaraSimCacheToolkit.cpp
@@ -277,6 +277,17 @@ void FNiagaraSimCacheToolkit::ExtendToolbar()
 					LOCTEXT("CopyToCSVTooltip", "Copies the currently selected view to the clipboard in CSV format."),
 					FSlateIcon(FAppStyle::GetAppStyleSetName(), "GenericCommands.Copy")
 				);
+
+				ToolbarBuilder.AddToolBarButton(
+					FUIAction(
+						FExecuteAction::CreateSP(SimCacheViewModelRef, &FNiagaraSimCacheViewModel::ExportToDisk),
+						FCanExecuteAction::CreateSP(SimCacheViewModelRef, &FNiagaraSimCacheViewModel::CanExportToDisk)
+					),
+					NAME_None,
+					LOCTEXT("ExportToJson", "Export To Disk"),
+					LOCTEXT("ExportToJsonTooltip", "Exports the raw data for each frame to disk. Note that data from data interfaces is only exported if they implement support for it."),
+					FSlateIcon()//-TODO:FAppStyle::GetAppStyleSetName(), "GenericCommands.Copy")
+				);
 			}
 			ToolbarBuilder.EndSection();
 		}
diff --git a/Engine/Plugins/FX/Niagara/Source/NiagaraEditor/Private/ViewModels/NiagaraSimCacheViewModel.cpp b/Engine/Plugins/FX/Niagara/Source/NiagaraEditor/Private/ViewModels/NiagaraSimCacheViewModel.cpp
--- a/Engine/Plugins/FX/Niagara/Source/NiagaraEditor/Private/ViewModels/NiagaraSimCacheViewModel.cpp
+++ b/Engine/Plugins/FX/Niagara/Source/NiagaraEditor/Private/ViewModels/NiagaraSimCacheViewModel.cpp
@@ -4,6 +4,7 @@
 #include "NiagaraComponent.h"
 #include "NiagaraSimCache.h"
 #include "NiagaraEditorCommon.h"
+#include "NiagaraEditorSimCacheUtils.h"
 #include "NiagaraDebuggerCommon.h"
 #include "HAL/PlatformApplicationMisc.h"
 #include "Serialization/MemoryReader.h"
@@ -811,6 +812,20 @@ void FNiagaraSimCacheViewModel::CopyActiveToClipboard() const
 	FPlatformApplicationMisc::ClipboardCopy(*ClipboardString);
 }
 
+bool FNiagaraSimCacheViewModel::CanExportToDisk() const
+{
+	return IsCacheValid();
+}
+
+void FNiagaraSimCacheViewModel::ExportToDisk() const
+{
+	if (!IsCacheValid())
+	{
+		return;
+	}
+	FNiagaraEditorSimCacheUtils::ExportToDisk(SimCache.Get());
+}
+
 void FNiagaraSimCacheViewModel::AddReferencedObjects(FReferenceCollector& Collector)
 {
 	if(SimCache != nullptr)
diff --git a/Engine/Plugins/FX/Niagara/Source/NiagaraEditor/Public/ViewModels/NiagaraSimCacheViewModel.h b/Engine/Plugins/FX/Niagara/Source/NiagaraEditor/Public/ViewModels/NiagaraSimCacheViewModel.h
--- a/Engine/Plugins/FX/Niagara/Source/NiagaraEditor/Public/ViewModels/NiagaraSimCacheViewModel.h
+++ b/Engine/Plugins/FX/Niagara/Source/NiagaraEditor/Public/ViewModels/NiagaraSimCacheViewModel.h
@@ -99,6 +99,9 @@ public:
 	bool CanCopyActiveToClipboard() const;
 	void CopyActiveToClipboard() const;
 
+	bool CanExportToDisk() const;
+	void ExportToDisk() const;
+
 	UNiagaraComponent* GetPreviewComponent() { return PreviewComponent; }
 
 	//~ FGCObject interface
