- Add menu options to enable/disable parts of Niagara
- Main component activation control + renderer controls
- Does not push to the client, only works in the editor

diff --git a/Engine/Plugins/FX/Niagara/Source/Niagara/Private/NiagaraComponent.cpp b/Engine/Plugins/FX/Niagara/Source/Niagara/Private/NiagaraComponent.cpp
--- a/Engine/Plugins/FX/Niagara/Source/Niagara/Private/NiagaraComponent.cpp
+++ b/Engine/Plugins/FX/Niagara/Source/Niagara/Private/NiagaraComponent.cpp
@@ -71,11 +71,11 @@ static FAutoConsoleVariableRef CVarNiagaraSoloAllowAsyncWorkToEndOfFrame(
 	ECVF_Default
 );
 
-static int32 GbSuppressNiagaraSystems = 0;
-static FAutoConsoleVariableRef CVarSuppressNiagaraSystems(
-	TEXT("fx.SuppressNiagaraSystems"),
-	GbSuppressNiagaraSystems,
-	TEXT("If > 0 Niagara particle systems will not be activated. \n"),
+static bool GbNiagaraComponentsEnabled = true;
+static FAutoConsoleVariableRef CVarNiagaraComponentsEnabled(
+	TEXT("fx.NiagaraComponentsEnabled"),
+	GbNiagaraComponentsEnabled,
+	TEXT("If false Niagara particle components will not be activated."),
 	ECVF_Default
 );
 
@@ -1242,7 +1242,7 @@ void UNiagaraComponent::ActivateInternal(bool bReset /* = false */, bool bIsScal
 		CurrLocalBounds.Init();
 	}
 
-	if (GbSuppressNiagaraSystems != 0)
+	if (!GbNiagaraComponentsEnabled)
 	{
 		UnregisterWithScalabilityManager();
 		OnSystemComplete(true);
diff --git a/Engine/Plugins/FX/Niagara/Source/Niagara/Private/NiagaraRendererLights.cpp b/Engine/Plugins/FX/Niagara/Source/Niagara/Private/NiagaraRendererLights.cpp
--- a/Engine/Plugins/FX/Niagara/Source/Niagara/Private/NiagaraRendererLights.cpp
+++ b/Engine/Plugins/FX/Niagara/Source/Niagara/Private/NiagaraRendererLights.cpp
@@ -20,11 +20,11 @@
 
 DECLARE_CYCLE_STAT(TEXT("Generate Particle Lights"), STAT_NiagaraGenLights, STATGROUP_Niagara);
 
-static int32 GbEnableNiagaraLightRendering = 1;
+static bool GbEnableNiagaraLightRendering = true;
 static FAutoConsoleVariableRef CVarEnableNiagaraLightRendering(
 	TEXT("fx.EnableNiagaraLightRendering"),
 	GbEnableNiagaraLightRendering,
-	TEXT("If == 0, Niagara Light Renderers are disabled. \n"),
+	TEXT("If false Niagara Light Renderers are disabled."),
 	ECVF_Default
 );
 
@@ -203,7 +203,7 @@ FNiagaraDynamicDataBase* FNiagaraRendererLights::GenerateDynamicData(const FNiag
 
 void FNiagaraRendererLights::GatherSimpleLights(FSimpleLightArray& OutParticleLights)const
 {
-	if (GbEnableNiagaraLightRendering == 0)
+	if (GbEnableNiagaraLightRendering == false)
 	{
 		return;
 	}
diff --git a/Engine/Plugins/FX/Niagara/Source/Niagara/Private/NiagaraRendererMeshes.cpp b/Engine/Plugins/FX/Niagara/Source/Niagara/Private/NiagaraRendererMeshes.cpp
--- a/Engine/Plugins/FX/Niagara/Source/Niagara/Private/NiagaraRendererMeshes.cpp
+++ b/Engine/Plugins/FX/Niagara/Source/Niagara/Private/NiagaraRendererMeshes.cpp
@@ -35,11 +35,11 @@ DECLARE_CYCLE_STAT(TEXT("Generate Mesh Vertex Data [GT]"), STAT_NiagaraGenMeshVe
 DECLARE_DWORD_COUNTER_STAT(TEXT("NumMeshesRenderer"), STAT_NiagaraNumMeshes, STATGROUP_Niagara);
 DECLARE_DWORD_COUNTER_STAT(TEXT("NumMeshVerts"), STAT_NiagaraNumMeshVerts, STATGROUP_Niagara);
 
-static int32 GbEnableNiagaraMeshRendering = 1;
+static bool GbEnableNiagaraMeshRendering = true;
 static FAutoConsoleVariableRef CVarEnableNiagaraMeshRendering(
 	TEXT("fx.EnableNiagaraMeshRendering"),
 	GbEnableNiagaraMeshRendering,
-	TEXT("If == 0, Niagara Mesh Renderers are disabled. \n"),
+	TEXT("If false Niagara Mesh Renderers are disabled."),
 	ECVF_Default
 );
 
@@ -1382,7 +1382,7 @@ void FNiagaraRendererMeshes::GetDynamicMeshElements(const TArray<const FSceneVie
 	FParticleMeshRenderData ParticleMeshRenderData;
 	PrepareParticleMeshRenderData(Collector.GetRHICommandList(), ParticleMeshRenderData, ViewFamily, Collector, DynamicDataRender, SceneProxy, false, ENiagaraGpuComputeTickStage::Last);
 
-	if (ParticleMeshRenderData.SourceParticleData == nullptr || GbEnableNiagaraMeshRendering == 0)
+	if (ParticleMeshRenderData.SourceParticleData == nullptr || GbEnableNiagaraMeshRendering == false)
 	{
 		return;
 	}
diff --git a/Engine/Plugins/FX/Niagara/Source/Niagara/Private/NiagaraRendererRibbons.cpp b/Engine/Plugins/FX/Niagara/Source/Niagara/Private/NiagaraRendererRibbons.cpp
--- a/Engine/Plugins/FX/Niagara/Source/Niagara/Private/NiagaraRendererRibbons.cpp
+++ b/Engine/Plugins/FX/Niagara/Source/Niagara/Private/NiagaraRendererRibbons.cpp
@@ -104,11 +104,11 @@ static FAutoConsoleVariableRef CVarNiagaraRibbonMinSegmentLength(
 	ECVF_Scalability
 );
 
-static int32 GbEnableNiagaraRibbonRendering = 1;
+static bool GbEnableNiagaraRibbonRendering = true;
 static FAutoConsoleVariableRef CVarEnableNiagaraRibbonRendering(
 	TEXT("fx.EnableNiagaraRibbonRendering"),
 	GbEnableNiagaraRibbonRendering,
-	TEXT("If == 0, Niagara Ribbon Renderers are disabled. \n"),
+	TEXT("If false Niagara Ribbon Renderers are disabled."),
 	ECVF_Default
 );
 
@@ -827,7 +827,7 @@ void FNiagaraRendererRibbons::GetDynamicMeshElements(const TArray<const FSceneVi
 
 	FNiagaraDataBuffer* SourceParticleData = DynamicData->GetParticleDataToRender(Collector.GetRHICommandList());
 
-	if (GbEnableNiagaraRibbonRendering == 0 || SourceParticleData == nullptr)
+	if (GbEnableNiagaraRibbonRendering == false || SourceParticleData == nullptr)
 	{
 		return;
 	}
diff --git a/Engine/Plugins/FX/Niagara/Source/Niagara/Private/NiagaraRendererSprites.cpp b/Engine/Plugins/FX/Niagara/Source/Niagara/Private/NiagaraRendererSprites.cpp
--- a/Engine/Plugins/FX/Niagara/Source/Niagara/Private/NiagaraRendererSprites.cpp
+++ b/Engine/Plugins/FX/Niagara/Source/Niagara/Private/NiagaraRendererSprites.cpp
@@ -22,11 +22,11 @@
 
 DECLARE_DWORD_COUNTER_STAT(TEXT("NumSprites"), STAT_NiagaraNumSprites, STATGROUP_Niagara);
 
-static int32 GbEnableNiagaraSpriteRendering = 1;
+static bool GbEnableNiagaraSpriteRendering = true;
 static FAutoConsoleVariableRef CVarEnableNiagaraSpriteRendering(
 	TEXT("fx.EnableNiagaraSpriteRendering"),
 	GbEnableNiagaraSpriteRendering,
-	TEXT("If == 0, Niagara Sprite Renderers are disabled. \n"),
+	TEXT("If false Niagara Sprite Renderers are disabled."),
 	ECVF_Default
 );
 
@@ -238,7 +238,7 @@ void FNiagaraRendererSprites::PrepareParticleSpriteRenderData(FRHICommandListBas
 
 	// Early out if we have no data or instances, this must be done before we read the material
 	FNiagaraDataBuffer* CurrentParticleData = ParticleSpriteRenderData.DynamicDataSprites->GetParticleDataToRender(RHICmdList, bGpuLowLatencyTranslucency);
-	if (!CurrentParticleData || (SourceMode == ENiagaraRendererSourceDataMode::Particles && CurrentParticleData->GetNumInstances() == 0) || (GbEnableNiagaraSpriteRendering == 0))
+	if (!CurrentParticleData || (SourceMode == ENiagaraRendererSourceDataMode::Particles && CurrentParticleData->GetNumInstances() == 0) || (GbEnableNiagaraSpriteRendering == false))
 	{
 		return;
 	}
diff --git a/Engine/Plugins/FX/Niagara/Source/NiagaraEditor/Private/Widgets/SNiagaraDebugger.cpp b/Engine/Plugins/FX/Niagara/Source/NiagaraEditor/Private/Widgets/SNiagaraDebugger.cpp
--- a/Engine/Plugins/FX/Niagara/Source/NiagaraEditor/Private/Widgets/SNiagaraDebugger.cpp
+++ b/Engine/Plugins/FX/Niagara/Source/NiagaraEditor/Private/Widgets/SNiagaraDebugger.cpp
@@ -55,6 +55,28 @@ namespace NiagaraDebuggerLocal
 	{
 		return TAttribute<T>::Create(TAttribute<T>::FGetter::CreateLambda(InFunction));
 	}
+
+	static void CreateConsoleVariableMenuEntry(FMenuBuilder& MenuBuilder, const FText& Label, const TCHAR* ConsoleVariableName)
+	{
+		IConsoleVariable* ConsoleVariable = IConsoleManager::Get().FindConsoleVariable(ConsoleVariableName);
+		if ( !ensure(ConsoleVariable) )
+		{
+			return;
+		}
+
+		MenuBuilder.AddMenuEntry(
+			Label,
+			FText::Format(LOCTEXT("ConsoleVariableFormat", "{0}\nSee console variable '{1}'"), ConsoleVariable->GetDetailedHelp(), FText::FromString(ConsoleVariableName)),
+			FSlateIcon(),
+			FUIAction(
+				FExecuteAction::CreateLambda([ConsoleVariable]() { ConsoleVariable->Set(!ConsoleVariable->GetBool()); }),
+				FCanExecuteAction(),
+				FIsActionChecked::CreateLambda([ConsoleVariable]() { return ConsoleVariable->GetBool(); })
+			),
+			NAME_None,
+			EUserInterfaceActionType::ToggleButton
+		);
+	}
 }
 
 namespace NiagaraDebugHudTab
@@ -1470,6 +1492,25 @@ TSharedRef<SWidget> SNiagaraDebugger::MakeModeOptionsMenu()
 			),
 			NAME_None, EUserInterfaceActionType::ToggleButton);
 	}
+	{
+		MenuBuilder.AddMenuSeparator();
+
+		NiagaraDebuggerLocal::CreateConsoleVariableMenuEntry(MenuBuilder, LOCTEXT("NiagaraComponentsEnabled", "Niagara Components Enabled"), TEXT("fx.NiagaraComponentsEnabled"));
+
+		MenuBuilder.AddSubMenu(
+			LOCTEXT("RendererEnabled", "Renderer Enabled"),
+			FText(),
+			FNewMenuDelegate::CreateLambda(
+				[](FMenuBuilder& SubMenuBuilder)
+				{
+					NiagaraDebuggerLocal::CreateConsoleVariableMenuEntry(SubMenuBuilder, LOCTEXT("LightRenderingEnabled", "Light Renderer Enabled"), TEXT("fx.EnableNiagaraLightRendering"));
+					NiagaraDebuggerLocal::CreateConsoleVariableMenuEntry(SubMenuBuilder, LOCTEXT("MeshRenderingEnabled", "Mesh Renderer Enabled"), TEXT("fx.EnableNiagaraMeshRendering"));
+					NiagaraDebuggerLocal::CreateConsoleVariableMenuEntry(SubMenuBuilder, LOCTEXT("RibbonRenderingEnabled", "Ribbon Renderer Enabled"), TEXT("fx.EnableNiagaraRibbonRendering"));
+					NiagaraDebuggerLocal::CreateConsoleVariableMenuEntry(SubMenuBuilder, LOCTEXT("SpriteRendererEnabled", "Sprite Renderer Enabled"), TEXT("fx.EnableNiagaraSpriteRendering"));
+				}
+			)
+		);
+	}
 
 	return MenuBuilder.MakeWidget();
 }
