Remove excessive checks with new compilation mode and object parameters
#tests made a set node with an object parameter and leave it as none, have the system compiled with the new compilation mode
#rnx

diff --git a/Engine/Plugins/FX/Niagara/Source/NiagaraEditor/Private/NiagaraGraphDigest.cpp b/Engine/Plugins/FX/Niagara/Source/NiagaraEditor/Private/NiagaraGraphDigest.cpp
--- a/Engine/Plugins/FX/Niagara/Source/NiagaraEditor/Private/NiagaraGraphDigest.cpp
+++ b/Engine/Plugins/FX/Niagara/Source/NiagaraEditor/Private/NiagaraGraphDigest.cpp
@@ -3224,9 +3224,7 @@ FNiagaraCompilationNodeInput::FNiagaraCompilationNodeInput(const UNiagaraNodeInp
 	else if (InputVariable.IsUObject())
 	{
 		UObject* SourceObjectAsset = InNode->GetObjectAsset();
-
 		ObjectAssetName = InputVariable.GetName();
-		check(SourceObjectAsset);
 		Context.ParentGraph.RegisterObjectAsset(ObjectAssetName, SourceObjectAsset);
 		ObjectAssetPath = FSoftObjectPath(SourceObjectAsset);
 	}
@@ -3329,7 +3327,6 @@ void FNiagaraCompilationNodeInput::Compile(FTranslator* Translator, TArray<int32
 					}
 					else if (!ObjectAssetName.IsNone())
 					{
-						check(ObjectAssetPath.IsValid());
 						check(InputVariable.IsUObject());
 						Outputs.Add(Translator->RegisterUObjectPath(InputVariable, ObjectAssetPath, false));
 						return;
@@ -3355,7 +3352,6 @@ void FNiagaraCompilationNodeInput::Compile(FTranslator* Translator, TArray<int32
 		}
 		else if (!ObjectAssetName.IsNone())
 		{
-			check(ObjectAssetPath.IsValid());
 			check(InputVariable.IsUObject());
 			Outputs.Add(Translator->RegisterUObjectPath(InputVariable, ObjectAssetPath, false));
 			break;
