- Use GetFormat() rather than peeking OverrideFormat as it can incorrectly be PF_Unknown if the override format is not used

#rnx

diff --git a/Engine/Plugins/FX/Niagara/Source/Niagara/Private/NiagaraDataInterfaceRenderTarget2DArray.cpp b/Engine/Plugins/FX/Niagara/Source/Niagara/Private/NiagaraDataInterfaceRenderTarget2DArray.cpp
--- a/Engine/Plugins/FX/Niagara/Source/Niagara/Private/NiagaraDataInterfaceRenderTarget2DArray.cpp
+++ b/Engine/Plugins/FX/Niagara/Source/Niagara/Private/NiagaraDataInterfaceRenderTarget2DArray.cpp
@@ -639,7 +639,7 @@ bool UNiagaraDataInterfaceRenderTarget2DArray::PerInstanceTick(void* PerInstance
 			//{
 			//	InstanceData->MipMapGeneration = ENiagaraMipMapGeneration::Disabled;
 			//}
-			InstanceData->Format = InstanceData->TargetTexture->OverrideFormat;
+			InstanceData->Format = InstanceData->TargetTexture->GetFormat();
 			InstanceData->Filter = InstanceData->TargetTexture->Filter;
 		}
 		else
@@ -689,7 +689,7 @@ bool UNiagaraDataInterfaceRenderTarget2DArray::PerInstanceTickPostSimulate(void*
 	{
 		//const bool bAutoGenerateMips = InstanceData->MipMapGeneration != ENiagaraMipMapGeneration::Disabled;
 		if ((InstanceData->TargetTexture->SizeX != InstanceData->Size.X) || (InstanceData->TargetTexture->SizeY != InstanceData->Size.Y) || (InstanceData->TargetTexture->Slices != InstanceData->Size.Z) ||
-			(InstanceData->TargetTexture->OverrideFormat != InstanceData->Format) ||
+			(InstanceData->TargetTexture->GetFormat() != InstanceData->Format) ||
 			(InstanceData->TargetTexture->Filter != InstanceData->Filter) ||
 			!InstanceData->TargetTexture->bCanCreateUAV ||
 			//(InstanceData->TargetTexture->bAutoGenerateMips != bAutoGenerateMips) ||
diff --git a/Engine/Plugins/FX/Niagara/Source/Niagara/Private/NiagaraDataInterfaceRenderTargetCube.cpp b/Engine/Plugins/FX/Niagara/Source/Niagara/Private/NiagaraDataInterfaceRenderTargetCube.cpp
--- a/Engine/Plugins/FX/Niagara/Source/Niagara/Private/NiagaraDataInterfaceRenderTargetCube.cpp
+++ b/Engine/Plugins/FX/Niagara/Source/Niagara/Private/NiagaraDataInterfaceRenderTargetCube.cpp
@@ -590,7 +590,7 @@ bool UNiagaraDataInterfaceRenderTargetCube::PerInstanceTick(void* PerInstanceDat
 			//{
 			//	InstanceData->MipMapGeneration = ENiagaraMipMapGeneration::Disabled;
 			//}
-			InstanceData->Format = InstanceData->TargetTexture->OverrideFormat;
+			InstanceData->Format = InstanceData->TargetTexture->GetFormat();
 			InstanceData->Filter = InstanceData->TargetTexture->Filter;
 		}
 		else
@@ -641,7 +641,7 @@ bool UNiagaraDataInterfaceRenderTargetCube::PerInstanceTickPostSimulate(void* Pe
 	{
 		//const bool bAutoGenerateMips = InstanceData->MipMapGeneration != ENiagaraMipMapGeneration::Disabled;
 		if ((InstanceData->TargetTexture->SizeX != InstanceData->Size) || 
-			(InstanceData->TargetTexture->OverrideFormat != InstanceData->Format) ||
+			(InstanceData->TargetTexture->GetFormat() != InstanceData->Format) ||
 			(InstanceData->TargetTexture->Filter != InstanceData->Filter) ||
 			!InstanceData->TargetTexture->bCanCreateUAV ||
 			//(InstanceData->TargetTexture->bAutoGenerateMips != bAutoGenerateMips) ||
diff --git a/Engine/Plugins/FX/Niagara/Source/Niagara/Private/NiagaraDataInterfaceRenderTargetVolume.cpp b/Engine/Plugins/FX/Niagara/Source/Niagara/Private/NiagaraDataInterfaceRenderTargetVolume.cpp
--- a/Engine/Plugins/FX/Niagara/Source/Niagara/Private/NiagaraDataInterfaceRenderTargetVolume.cpp
+++ b/Engine/Plugins/FX/Niagara/Source/Niagara/Private/NiagaraDataInterfaceRenderTargetVolume.cpp
@@ -1013,7 +1013,7 @@ bool UNiagaraDataInterfaceRenderTargetVolume::PerInstanceTick(void* PerInstanceD
 			//{
 			//	InstanceData->MipMapGeneration = ENiagaraMipMapGeneration::Disabled;
 			//}
-			InstanceData->Format = InstanceData->TargetTexture->OverrideFormat;
+			InstanceData->Format = InstanceData->TargetTexture->GetFormat();
 			InstanceData->Filter = InstanceData->TargetTexture->Filter;
 		}
 		else
@@ -1064,17 +1064,12 @@ bool UNiagaraDataInterfaceRenderTargetVolume::PerInstanceTickPostSimulate(void*
 	{
 		//const bool bAutoGenerateMips = InstanceData->MipMapGeneration != ENiagaraMipMapGeneration::Disabled;
 		if ((InstanceData->TargetTexture->SizeX != InstanceData->Size.X) || (InstanceData->TargetTexture->SizeY != InstanceData->Size.Y) || (InstanceData->TargetTexture->SizeZ != InstanceData->Size.Z) ||
-			(InstanceData->TargetTexture->OverrideFormat != InstanceData->Format) ||
+			(InstanceData->TargetTexture->GetFormat() != InstanceData->Format) ||
 			(InstanceData->TargetTexture->Filter != InstanceData->Filter) ||
 			!InstanceData->TargetTexture->bCanCreateUAV ||
 			//(InstanceData->TargetTexture->bAutoGenerateMips != bAutoGenerateMips) ||
 			!InstanceData->TargetTexture->GetResource())
 		{
-			if (InstanceData->Format == PF_Unknown)
-			{
-				return false;
-			}
-
 			// resize RT to match what we need for the output
 			InstanceData->TargetTexture->bCanCreateUAV = true;
 			//InstanceData->TargetTexture->bAutoGenerateMips = bAutoGenerateMips;
