- Don't allow bStartWithTickEnabled to be editable in the BP editor
- We already will override this field as Niagara controls this value but it can be confusing in the BP editor as it will show as adjustable

#rnx

diff --git a/Engine/Plugins/FX/Niagara/Source/Niagara/Private/NiagaraComponent.cpp b/Engine/Plugins/FX/Niagara/Source/Niagara/Private/NiagaraComponent.cpp
--- a/Engine/Plugins/FX/Niagara/Source/Niagara/Private/NiagaraComponent.cpp
+++ b/Engine/Plugins/FX/Niagara/Source/Niagara/Private/NiagaraComponent.cpp
@@ -3176,6 +3176,19 @@ static FNiagaraVariant GetParameterValueFromStore(const FNiagaraVariableBase& Va
 }
 
 #if WITH_EDITOR
+bool UNiagaraComponent::CanEditChange(const FProperty* InProperty) const
+{
+	if (InProperty)
+	{
+		// We don't support changing this property as it is fully controlled by Niagara in both batached and solo mode
+		if (InProperty->GetFName() == GET_MEMBER_NAME_CHECKED(FTickFunction, bStartWithTickEnabled))
+		{
+			return false;
+		}
+	}
+
+	return Super::CanEditChange(InProperty);
+}
 
 void UNiagaraComponent::PreEditChange(FProperty* PropertyAboutToChange)
 {
diff --git a/Engine/Plugins/FX/Niagara/Source/Niagara/Public/NiagaraComponent.h b/Engine/Plugins/FX/Niagara/Source/Niagara/Public/NiagaraComponent.h
--- a/Engine/Plugins/FX/Niagara/Source/Niagara/Public/NiagaraComponent.h
+++ b/Engine/Plugins/FX/Niagara/Source/Niagara/Public/NiagaraComponent.h
@@ -607,6 +607,8 @@ public:
 #endif
 
 #if WITH_EDITOR
+	NIAGARA_API virtual bool CanEditChange(const FProperty* InProperty) const override;
+
 	NIAGARA_API virtual void PreEditChange(FProperty* PropertyAboutToChange) override;
 	NIAGARA_API virtual void PostEditChangeProperty(FPropertyChangedEvent& PropertyChangedEvent) override;
 
