Allow Niagara Systems placed in a level to correctly interact with NDC Islands.

- Read accesses can now create new islands. Allows readers to be initialized before writers but function correclty going forwards rather than fail to init and be permeantly broken.
- Islands do not clean themselves up until they are unused and empty.
- Also fixes minor perf issue in NDC Islands search.

diff --git a/Engine/Plugins/FX/Niagara/Source/Niagara/Internal/NiagaraDataChannelCommon.h b/Engine/Plugins/FX/Niagara/Source/Niagara/Internal/NiagaraDataChannelCommon.h
--- a/Engine/Plugins/FX/Niagara/Source/Niagara/Internal/NiagaraDataChannelCommon.h
+++ b/Engine/Plugins/FX/Niagara/Source/Niagara/Internal/NiagaraDataChannelCommon.h
@@ -105,6 +105,10 @@ struct FNiagaraDataChannelData final : public TSharedFromThis<FNiagaraDataChanne
 
 	//Returns if this data is still valid. This can return false in cases where the owning data channel has been modified for example.
 	bool IsLayoutValid(UNiagaraDataChannelHandler* Owner)const;
+
+	//Return true if data has been written to this NDC data for the current frame.
+	bool HasData()const;
+
 private:
 
 	void CreateRenderThreadProxy(UNiagaraDataChannelHandler* Owner);
diff --git a/Engine/Plugins/FX/Niagara/Source/Niagara/Private/NiagaraDataChannel.cpp b/Engine/Plugins/FX/Niagara/Source/Niagara/Private/NiagaraDataChannel.cpp
--- a/Engine/Plugins/FX/Niagara/Source/Niagara/Private/NiagaraDataChannel.cpp
+++ b/Engine/Plugins/FX/Niagara/Source/Niagara/Private/NiagaraDataChannel.cpp
@@ -1183,6 +1183,11 @@ bool FNiagaraDataChannelData::IsLayoutValid(UNiagaraDataChannelHandler* Owner)co
 	return false;
 }
 
+bool FNiagaraDataChannelData::HasData()const
+{
+	return GameData->Num() > 0 || (CPUSimData && CPUSimData->GetCurrentData()->GetNumInstances() > 0) || PublishRequestsForGPU.Num() > 0;
+}
+
 void FNiagaraDataChannelData::DestroyRenderThreadProxy(FNiagaraGpuComputeDispatchInterface* ComputeDispatchInterface)
 {
 	//We can leave the proxy to be destroyed by the dispatcher ( or in some edge cases perhaps a DI proxy ).
diff --git a/Engine/Plugins/FX/Niagara/Source/Niagara/Private/NiagaraDataChannel_Islands.cpp b/Engine/Plugins/FX/Niagara/Source/Niagara/Private/NiagaraDataChannel_Islands.cpp
--- a/Engine/Plugins/FX/Niagara/Source/Niagara/Private/NiagaraDataChannel_Islands.cpp
+++ b/Engine/Plugins/FX/Niagara/Source/Niagara/Private/NiagaraDataChannel_Islands.cpp
@@ -152,10 +152,24 @@ void UNiagaraDataChannelHandler_Islands::BeginFrame(float DeltaTime, FNiagaraWor
 void UNiagaraDataChannelHandler_Islands::EndFrame(float DeltaTime, FNiagaraWorldManager* OwningWorld)
 {
 	Super::EndFrame(DeltaTime, OwningWorld);
+	const UNiagaraDataChannel_Islands* IslandChannel = GetChannelTyped<const UNiagaraDataChannel_Islands>();
 	for (auto It = ActiveIslands.CreateIterator(); It; ++It)
 	{
 		FNDCIsland& Island = IslandPool[*It];
 		Island.EndFrame();
+
+		if (IslandChannel->GetDebugDrawSettings().ShowBounds())
+		{
+			Island.DebugDrawBounds();
+		}
+
+		//Free up islands that are not being used.
+		if (Island.IsBeingUsed() == false && Island.GetData()->HasData() == false)
+		{
+			FreeIslands.Emplace(*It);
+			Island.OnReleased();
+			It.RemoveCurrentSwap();
+		}
 	}
 }
 
@@ -170,20 +184,6 @@ void UNiagaraDataChannelHandler_Islands::Tick(float DeltaTime, ETickingGroup Tic
 	{
 		FNDCIsland& Island = IslandPool[*It];
 		Island.Tick(TickGroup);
-
-		if(IslandChannel->GetDebugDrawSettings().ShowBounds())
-		{
-			Island.DebugDrawBounds();
-		}
-
-		//Free up islands that are not being used.
-		//We do this immediately as we leave it up to the handler systems to delay their cleanup as long as they want to.
-		if (Island.IsBeingUsed() == false)
-		{
-			FreeIslands.Emplace(*It);
-			Island.OnReleased();
-			It.RemoveCurrentSwap();
-		}
 	}
 }
 
@@ -203,9 +203,6 @@ FNDCIsland* UNiagaraDataChannelHandler_Islands::FindOrCreateIsland(const FNiagar
 	FVector Location = SearchParams.GetLocation();
 	if (AccessType == ENiagaraResourceAccess::ReadOnly)
 	{
-		//If we're reading only then we just try to find the right island to read from.
-		//When writing we'll grow/spawn islands to accommodate the data.
-
 		//First we see if this is an islands handler system.
 		if(SearchParams.GetOwner())
 		{
@@ -217,71 +214,59 @@ FNDCIsland* UNiagaraDataChannelHandler_Islands::FindOrCreateIsland(const FNiagar
 					return &Island;
 				}
 			}
-			//Failing that, we'll just return the first island that intersects the owners bounds.
-			for (int32 i : ActiveIslands)
-			{
-				FNDCIsland& Island = IslandPool[i];
-				if (Island.Intersects(SearchParams.GetOwner()->Bounds))
-				{
-					return &Island;
-				}
-			}
 		}
-		else
+	}
+
+	UWorld* World = GetWorld();
+	check(World);
+	double WorldTime = World->GetTimeSeconds();
+
+	//Find the first island that could contain this point.
+	//For now we do a linear search of active islands.
+	//Assuming that overall active island count will be low.
+	//If this is not the case then we'll want to add an acceleration structure to speed up this search.
+	FVector MaxExtents = IslandChannel->GetMaxExtents();
+	FVector PerElementExtents = IslandChannel->GetPerElementExtents();
+
+	if(IslandChannel->GetMode() == ENiagraDataChannel_IslandMode::AlignedStatic)
+	{
+		//for static islands we can just find the first one that contains our point.
+		for (int32 i : ActiveIslands)
 		{
-			//We don't have an owning component so just return the first island to contain the location.
-			for (int32 i : ActiveIslands)
+			FNDCIsland& Island = IslandPool[i];
+			if(Island.Contains(Location))
 			{
-				FNDCIsland& Island = IslandPool[i];
-				if (Island.Contains(Location))
-				{
-					return &Island;
-				}
+				return &Island;
 			}
 		}
-		return nullptr;
 	}
-	else
+	else if(IslandChannel->GetMode() == ENiagraDataChannel_IslandMode::Dynamic)
 	{
-		UWorld* World = GetWorld();
-		check(World);
-		double WorldTime = World->GetTimeSeconds();
-
-		//Find the first island that could contain this point.
-		//For now we do a linear search of active islands.
-		//Assuming that overall active island count will be low.
-		//If this is not the case then we'll want to add an acceleration structure to speed up this search.
-		FNDCIsland* IslandToUse = nullptr;
-		FVector MaxExtents = IslandChannel->GetMaxExtents();
-		FVector PerElementExtents = IslandChannel->GetPerElementExtents();
+		//for dynamic islands we find the first island that can grow to accommodate this point.
 		for (int32 i : ActiveIslands)
 		{
 			FNDCIsland& Island = IslandPool[i];
-			if (Island.TryGrow(Location, PerElementExtents, MaxExtents))
+			if(Island.TryGrow(Location, PerElementExtents, MaxExtents))
 			{
-				IslandToUse = &Island;
+				return &Island;
 			}
 		}
+	}
 
-		if (IslandToUse)
-		{
-			return IslandToUse;
-		}
-
-		//Failing that, get a init a new island from the pool.
-		int32 NewIslandIndex = ActivateNewIsland(Location);
-		if (ensure(NewIslandIndex != INDEX_NONE))
-		{
-			IslandToUse = &IslandPool[NewIslandIndex];
-		}
-
-		if (IslandToUse)
-		{
-			IslandToUse->OnAcquired(Location);
-		}
+	//Failing that, get a init a new island from the pool.
+	FNDCIsland* IslandToUse = nullptr;
+	int32 NewIslandIndex = ActivateNewIsland(Location);
+	if (ensure(NewIslandIndex != INDEX_NONE))
+	{
+		IslandToUse = &IslandPool[NewIslandIndex];
+	}
 
-		return IslandToUse;
+	if (IslandToUse)
+	{
+		IslandToUse->OnAcquired(Location);
 	}
+
+	return IslandToUse;
 }
 
 int32 UNiagaraDataChannelHandler_Islands::ActivateNewIsland(FVector Location)
