Couple of tweaks for the new compilation mode for Niagara
-moves thes sytem compilation task to a TSharedFromThis to ensure that the task won't get removed out from under active worker tasks (via an abort)
-moves DDC access to Records of individual values rather than just simple values
-adds optional (based on a value in NiagaraEditorSettings) metadata to be attached to DDC puts to help track down future issues
#rnx
#tests Ran through some different compilation scenarios

diff --git a/Engine/Plugins/FX/Niagara/Source/NiagaraEditor/Private/NiagaraCompilationTasks.cpp b/Engine/Plugins/FX/Niagara/Source/NiagaraEditor/Private/NiagaraCompilationTasks.cpp
--- a/Engine/Plugins/FX/Niagara/Source/NiagaraEditor/Private/NiagaraCompilationTasks.cpp
+++ b/Engine/Plugins/FX/Niagara/Source/NiagaraEditor/Private/NiagaraCompilationTasks.cpp
@@ -9,6 +9,7 @@
 #include "Misc/PathViews.h"
 #include "NiagaraCompilationTypes.h"
 #include "NiagaraDigestDatabase.h"
+#include "NiagaraEditorSettings.h"
 #include "NiagaraEmitter.h"
 #include "NiagaraGraphDigest.h"
 #include "NiagaraCompilationBridge.h"
@@ -19,6 +20,8 @@
 #include "NiagaraSystem.h"
 #include "NiagaraSystemCompilingManager.h"
 #include "ProfilingDebugging/CookStats.h"
+#include "Serialization/CompactBinary.h"
+#include "Serialization/CompactBinaryWriter.h"
 #include "Serialization/ObjectAndNameAsStringProxyArchive.h"
 #include "ShaderParameterMetadataBuilder.h"
 
@@ -55,9 +58,18 @@ namespace NiagaraSystemCookStats
 
 namespace NiagaraCompilationTasksImpl
 {
-	static const FGuid UE_NIAGARA_ASYNC_TASK_COMPILER_VER = FGuid(0x0E657251, 0xB787406B, 0x8DE09B90, 0x817ABFB9);
+	static const FGuid UE_NIAGARA_ASYNC_TASK_COMPILER_VER = FGuid(0x73F57E8B, 0xD0AC44E8, 0xBEE372D4, 0x1CD85A69);
 	
-	static UE::DerivedData::FCacheBucket NiagaraDDCBucket("NiagaraScript");
+	static const UE::DerivedData::FCacheBucket NiagaraDDCBucket("NiagaraScript");
+	static const UE::DerivedData::FValueId ScriptDataValue = UE::DerivedData::FValueId::FromName(TEXT("ScriptData"));
+
+	static bool GDebugParseCompileMetaDataOnGet = false;
+	static FAutoConsoleVariableRef CVarGDebugParseCompileMetaDataOnGet(
+		TEXT("fx.Niagara.ParseCompileMetaDataOnGet"),
+		GDebugParseCompileMetaDataOnGet,
+		TEXT("When true the compilation metadata from DDC will be parsed on get."),
+		ECVF_Default
+	);
 
 	void GetUsagesToDuplicate(ENiagaraScriptUsage TargetUsage, TArray<ENiagaraScriptUsage>& DuplicateUsages)
 	{
@@ -345,6 +357,100 @@ namespace NiagaraCompilationTasksImpl
 		return InputVariables;
 	}
 
+	struct FScriptCompileMetaData
+	{
+		FScriptCompileMetaData(const FCbObject& MetaData)
+		{
+			Version = static_cast<EMetaDataVersion>(MetaData.Find(ANSITEXTVIEW("Version")).AsInt32());
+			if (Version == EMetaDataVersion::Latest)
+			{
+				BuildString = FString(MetaData.Find(ANSITEXTVIEW("Build")).AsString());
+				DateTime = MetaData.Find(ANSITEXTVIEW("Time")).AsDateTime();
+				CompileIdString = FString(MetaData.Find(ANSITEXTVIEW("CompileId")).AsString());
+
+				const FCbArray RIParamArray = MetaData.Find(ANSITEXTVIEW("RIParam")).AsArray();
+				for (FCbFieldViewIterator It = RIParamArray.CreateViewIterator(); It; ++It)
+				{
+					TTuple<FName, FNiagaraVariant>& NamedVariant = BakedRapidIterationParameters.AddDefaulted_GetRef();
+
+					FCbObjectView RIParamView = It->AsObjectView();
+
+					NamedVariant.Key = *FString(RIParamView.FindView(ANSITEXTVIEW("VarName")).AsString());
+
+					FMemoryView VariantDataView = RIParamView.FindView(ANSITEXTVIEW("VarData")).AsBinaryView();
+
+					NamedVariant.Value = FNiagaraVariant(VariantDataView.GetData(), VariantDataView.GetSize());
+				}
+			}
+		}
+
+		FScriptCompileMetaData(const FNiagaraSystemCompilationTask::FCompileTaskInfo& CompileTask)
+		{
+			Version = EMetaDataVersion::Latest;
+			BuildString = FApp::GetBuildVersion();
+			DateTime = FDateTime::UtcNow();
+			CompileTask.CompileId.AppendKeyString(CompileIdString);
+
+			BakedRapidIterationParameters.Reserve(CompileTask.BakedRapidIterationParameters.Num());
+			Algo::Transform(CompileTask.BakedRapidIterationParameters, BakedRapidIterationParameters, [](const FNiagaraVariable& InVar) -> TTuple<FName, FNiagaraVariant>
+			{
+				return MakeTuple(InVar.GetName(), FNiagaraVariant(InVar.GetData(), InVar.GetAllocatedSizeInBytes()));
+			});
+		}
+
+		FCbObject ToCompactBinary() const
+		{
+			FCbWriter Writer;
+			Writer.BeginObject();
+			{
+				Writer.AddInteger(ANSITEXTVIEW("Version"), static_cast<int32>(Version));
+				Writer.AddString(ANSITEXTVIEW("Build"), BuildString);
+				Writer.AddDateTime(ANSITEXTVIEW("Time"), DateTime);
+				Writer.AddString(ANSITEXTVIEW("CompileId"), CompileIdString);
+
+				Writer.BeginArray(ANSITEXTVIEW("RIParam"));
+				{
+					for (const TTuple<FName, FNiagaraVariant>& Var : BakedRapidIterationParameters)
+					{
+						Writer.BeginObject();
+						{
+							Writer.AddString(ANSITEXTVIEW("VarName"), *Var.Key.ToString());
+							Writer.AddBinary(ANSITEXTVIEW("VarData"), Var.Value.GetBytes(), Var.Value.GetNumBytes());
+						}
+
+						Writer.EndObject();
+					}
+				}
+				Writer.EndArray();
+			}
+			Writer.EndObject();
+
+			return Writer.Save().AsObject();
+		}
+
+		enum class EMetaDataVersion : int32
+		{
+			Invalid = 0,
+			Latest = 1
+		};
+
+		EMetaDataVersion Version = EMetaDataVersion::Invalid;
+		FString BuildString;
+		FDateTime DateTime;
+		FString CompileIdString;
+		TArray<TTuple<FName, FNiagaraVariant>> BakedRapidIterationParameters;
+	};
+
+	static FCbObject GenerateScriptCompileMetaData(const FNiagaraSystemCompilationTask::FCompileTaskInfo& CompileTask)
+	{
+		return FScriptCompileMetaData(CompileTask).ToCompactBinary();
+	}
+
+	static void ParseScriptCompileMetaData(const FCbObject& MetaData)
+	{
+		FScriptCompileMetaData Data(MetaData);
+	}
+
 	void AddRapidIterationParameter(
 		const FNiagaraParameterHandle& ParameterHandle,
 		const FNiagaraVariable& InputVariable,
@@ -500,6 +606,7 @@ FNiagaraSystemCompilationTask::FNiagaraSystemCompilationTask(FNiagaraCompilation
 	, DDCRequestOwner(UE::DerivedData::EPriority::Normal)
 	, System_GT(InSystem)
 {
+	bGenerateCompileMetaData = GetDefault<UNiagaraEditorSettings>()->GetGenerateMetaDataOnCompile();
 }
 
 void FNiagaraSystemCompilationTask::Abort()
@@ -1158,7 +1265,7 @@ void FNiagaraSystemCompilationTask::FDispatchAndProcessDataCacheGetRequests::Lau
 
 	using namespace UE::DerivedData;
 
-	TArray<FCacheGetValueRequest> GetRequests;
+	TArray<FCacheGetRequest> GetRequests;
 
 	const int32 CompileTaskCount = SystemCompileTask->CompileTasks.Num();
 	const int32 ShaderCompileTaskCount = SystemCompileTask->CompileComputeShaderTasks.Num();
@@ -1178,7 +1285,7 @@ void FNiagaraSystemCompilationTask::FDispatchAndProcessDataCacheGetRequests::Lau
 				Index.CompileTaskIndex = CompileTaskIt;
 				Index.bShaderCompileTask = false;
 
-				FCacheGetValueRequest& GetValueRequest = GetRequests.AddDefaulted_GetRef();
+				FCacheGetRequest& GetValueRequest = GetRequests.AddDefaulted_GetRef();
 				GetValueRequest.Name = CompileTask.AssetPath;
 				GetValueRequest.Key = CacheKey;
 				GetValueRequest.UserData = Index.UserData;
@@ -1202,7 +1309,7 @@ void FNiagaraSystemCompilationTask::FDispatchAndProcessDataCacheGetRequests::Lau
 				Index.CompileTaskIndex = ShaderCompileTaskIt;
 				Index.bShaderCompileTask = true;
 
-				FCacheGetValueRequest& ShaderGetValueRequest = GetRequests.AddDefaulted_GetRef();
+				FCacheGetRequest& ShaderGetValueRequest = GetRequests.AddDefaulted_GetRef();
 				ShaderGetValueRequest.Name = CompileTask.AssetPath;
 				ShaderGetValueRequest.Key = CacheKey;
 				ShaderGetValueRequest.UserData = Index.UserData;
@@ -1217,7 +1324,7 @@ void FNiagaraSystemCompilationTask::FDispatchAndProcessDataCacheGetRequests::Lau
 	if (PendingGetRequestCount > 0)
 	{
 		FRequestBarrier AsyncBarrier(SystemCompileTask->DDCRequestOwner);
-		GetCache().GetValue(GetRequests, SystemCompileTask->DDCRequestOwner, [this, SystemCompileTask](FCacheGetValueResponse&& Response)
+		GetCache().Get(GetRequests, SystemCompileTask->DDCRequestOwner, [this, SystemCompileTask = SystemCompileTask->AsShared()](FCacheGetResponse&& Response)
 		{
 			FDDCUserData Index;
 			Index.UserData = Response.UserData;
@@ -1240,7 +1347,16 @@ void FNiagaraSystemCompilationTask::FDispatchAndProcessDataCacheGetRequests::Lau
 
 			if (DDCTaskInfo && Response.Status == EStatus::Ok)
 			{
-				DDCTaskInfo->PendingDDCData = Response.Value.GetData().Decompress().MoveToUnique();
+				const FValueWithId& ValueWithId = Response.Record.GetValue(NiagaraCompilationTasksImpl::ScriptDataValue);
+				DDCTaskInfo->PendingDDCData = ValueWithId.GetData().Decompress().MoveToUnique();
+
+				if (NiagaraCompilationTasksImpl::GDebugParseCompileMetaDataOnGet)
+				{
+					if (!Index.bShaderCompileTask)
+					{
+						NiagaraCompilationTasksImpl::ParseScriptCompileMetaData(Response.Record.GetMeta());
+					}
+				}
 			}
 
 			if (PendingGetRequestCount.fetch_sub(1, std::memory_order_relaxed) == 1)
@@ -1305,7 +1421,9 @@ void FNiagaraSystemCompilationTask::FDispatchDataCachePutRequests::Launch(FNiaga
 {
 	using namespace UE::DerivedData;
 
-	TArray<FCachePutValueRequest> PutRequests;
+	const bool bIncludeMetaData = SystemCompileTask->GenerateCompileMetaData();
+
+	TArray<FCachePutRequest> PutRequests;
 
 	const int32 CompileTaskCount = SystemCompileTask->CompileTasks.Num();
 	const int32 ShaderCompileTaskCount = SystemCompileTask->CompileComputeShaderTasks.Num();
@@ -1321,10 +1439,17 @@ void FNiagaraSystemCompilationTask::FDispatchDataCachePutRequests::Launch(FNiaga
 				FValue PutValue = FValue::Compress(SharedBuffer);
 				for (const FCacheKey& CachePutKey : CompileTask.DDCTaskInfo.DataCachePutKeys)
 				{
-					FCachePutValueRequest& PutRequest = PutRequests.AddDefaulted_GetRef();
+					FCacheRecordBuilder RecordBuilder(CachePutKey);
+					RecordBuilder.AddValue(NiagaraCompilationTasksImpl::ScriptDataValue, PutValue);
+					if (bIncludeMetaData)
+					{
+						RecordBuilder.SetMeta(NiagaraCompilationTasksImpl::GenerateScriptCompileMetaData(CompileTask));
+					}
+
+					FCachePutRequest& PutRequest = PutRequests.AddZeroed_GetRef();
 					PutRequest.Name = CompileTask.AssetPath;
-					PutRequest.Key = CachePutKey;
-					PutRequest.Value = PutValue;
+					PutRequest.Record = RecordBuilder.Build();
+					PutRequest.Policy = ECachePolicy::Default;
 				}
 			}
 		}
@@ -1341,10 +1466,13 @@ void FNiagaraSystemCompilationTask::FDispatchDataCachePutRequests::Launch(FNiaga
 				FValue PutValue = FValue::Compress(SharedBuffer);
 				for (const FCacheKey& CachePutKey : CompileTask.DDCTaskInfo.DataCachePutKeys)
 				{
-					FCachePutValueRequest& PutRequest = PutRequests.AddDefaulted_GetRef();
+					FCacheRecordBuilder RecordBuilder(CachePutKey);
+					RecordBuilder.AddValue(NiagaraCompilationTasksImpl::ScriptDataValue, PutValue);
+
+					FCachePutRequest& PutRequest = PutRequests.AddZeroed_GetRef();
 					PutRequest.Name = ParentCompileTask.AssetPath;
-					PutRequest.Key = CachePutKey;
-					PutRequest.Value = PutValue;
+					PutRequest.Record = RecordBuilder.Build();
+					PutRequest.Policy = ECachePolicy::Default;
 				}
 			}
 		}
@@ -1355,7 +1483,7 @@ void FNiagaraSystemCompilationTask::FDispatchDataCachePutRequests::Launch(FNiaga
 	if (!PutRequests.IsEmpty())
 	{
 		FRequestBarrier AsyncBarrier(SystemCompileTask->DDCRequestOwner);
-		GetCache().PutValue(PutRequests, SystemCompileTask->DDCRequestOwner, [this](FCachePutValueResponse&& Response)
+		GetCache().Put(PutRequests, SystemCompileTask->DDCRequestOwner, [this](FCachePutResponse&& Response)
 		{
 			if (PendingPutRequestCount.fetch_sub(1, std::memory_order_relaxed) == 1)
 			{
@@ -1910,6 +2038,11 @@ void FNiagaraSystemCompilationTask::GetAvailableCollections(TArray<FNiagaraCompi
 	DigestedParameterCollections.GenerateValueArray(OutCollections);
 }
 
+bool FNiagaraSystemCompilationTask::GenerateCompileMetaData() const
+{
+	return bGenerateCompileMetaData;
+}
+
 bool FNiagaraSystemCompilationTask::GetStageName(int32 EmitterIndex, const FNiagaraCompilationNodeOutput* OutputNode, FName& OutStageName) const
 {
 	OutStageName = NAME_None;
@@ -2056,18 +2189,17 @@ struct FNiagaraSystemCompilationTask::FCollectStaticVariablesTaskBuilder
 
 	using FStaticVariableBuilderTaskHandle = TSharedPtr<FStaticVariableBuilderTask, ESPMode::ThreadSafe>;
 
-	UE::Tasks::FTask LaunchCollectedTasks()
+	static UE::Tasks::FTask LaunchCollectedTasks(FCollectStaticVariablesTaskBuilder&& Context)
 	{
 		using namespace UE::Tasks;
-
-		if (!StaticVariableTasksToProcess.IsEmpty())
+		if (!Context.StaticVariableTasksToProcess.IsEmpty())
 		{
 			// fire off the tasks related to collecting static variables
-			const int32 TaskCount = StaticVariableTasksToProcess.Num();
+			const int32 TaskCount = Context.StaticVariableTasksToProcess.Num();
 			TArray<FTask> PendingTasks;
 			PendingTasks.Reserve(TaskCount);
 
-			for (FStaticVariableBuilderTaskHandle& BuilderTask : StaticVariableTasksToProcess)
+			for (FStaticVariableBuilderTaskHandle& BuilderTask : Context.StaticVariableTasksToProcess)
 			{
 				PendingTasks.Add(Launch(UE_SOURCE_LOCATION, [BuilderTask]
 				{
@@ -2100,8 +2232,8 @@ struct FNiagaraSystemCompilationTask::FCollectStaticVariablesTaskBuilder
 				}));
 			}
 
-			TArray<FStaticVariableBuilderTaskHandle> MergeTaskInput = StaticVariableTasksToProcess;
-			TArray<FNiagaraVariable>& MergeTaskOutput = FoundStaticVariables;
+			TArray<FStaticVariableBuilderTaskHandle> MergeTaskInput = Context.StaticVariableTasksToProcess;
+			TArray<FNiagaraVariable>& MergeTaskOutput = Context.FoundStaticVariables;
 			return Launch(UE_SOURCE_LOCATION, [MergeTaskInput, &MergeTaskOutput]
 			{
 				if (MergeTaskInput.Num() == 1)
@@ -2350,8 +2482,7 @@ UE::Tasks::FTask FNiagaraSystemCompilationTask::BuildRapidIterationParametersAsy
 
 	TMap<int32 /*EmitterIndex*/, FTask> CollectStaticVariableTasks;
 	
-	FCollectStaticVariablesTaskBuilder SystemTaskBuilder(*this);
-	FTask SystemCollectStaticVariableTask = SystemTaskBuilder.LaunchCollectedTasks();
+	FTask SystemCollectStaticVariableTask = FCollectStaticVariablesTaskBuilder::LaunchCollectedTasks({*this});
 	CollectStaticVariableTasks.Add(INDEX_NONE, SystemCollectStaticVariableTask);
 
 	if (SystemCollectStaticVariableTask.IsValid())
@@ -2362,10 +2493,9 @@ UE::Tasks::FTask FNiagaraSystemCompilationTask::BuildRapidIterationParametersAsy
 			{
 				continue;
 			}
-			FTask EmitterCollectStaticVariableTask = Launch(UE_SOURCE_LOCATION, [this, &EmitterInfo]
+			FTask EmitterCollectStaticVariableTask = Launch(UE_SOURCE_LOCATION, [Self = AsShared(), &EmitterInfo]
 			{
-				FCollectStaticVariablesTaskBuilder EmitterTaskBuilder(*this, EmitterInfo);
-				AddNested(EmitterTaskBuilder.LaunchCollectedTasks());
+				AddNested(FCollectStaticVariablesTaskBuilder::LaunchCollectedTasks({*Self, EmitterInfo}));
 			}, SystemCollectStaticVariableTask);
 
 			CollectStaticVariableTasks.Add(EmitterInfo.DigestedEmitterIndex, EmitterCollectStaticVariableTask);
@@ -2378,9 +2508,9 @@ UE::Tasks::FTask FNiagaraSystemCompilationTask::BuildRapidIterationParametersAsy
 
 			const FEmitterInfo* EmitterInfo = SystemInfo.EmitterInfoBySourceEmitter(ScriptInfo.SourceEmitterIndex);
 
-			PendingTasks.Add(Launch(UE_SOURCE_LOCATION, [this, &ScriptInfo]
+			PendingTasks.Add(Launch(UE_SOURCE_LOCATION, [Self = AsShared(), &ScriptInfo]
 			{
-				FBuildRapidIterationTaskBuilder TaskBuilder(*this, ScriptInfo);
+				FBuildRapidIterationTaskBuilder TaskBuilder(*Self, ScriptInfo);
 				AddNested(TaskBuilder.LaunchCollectedTasks());
 			}, CollectStaticVariableTasks.FindRef(EmitterInfo ? EmitterInfo->DigestedEmitterIndex : INDEX_NONE)));
 		}
@@ -2388,7 +2518,7 @@ UE::Tasks::FTask FNiagaraSystemCompilationTask::BuildRapidIterationParametersAsy
 		// We need a single task to copy over parameters between dependent scripts and make sure that the static variable
 		// list is up to date
 		TMap<TObjectKey<UNiagaraScript>, FScriptInfo>* ScriptInfoMapPtr = &DigestedScriptInfo;
-		return Launch(UE_SOURCE_LOCATION, [ScriptInfoMapPtr, this]
+		return Launch(UE_SOURCE_LOCATION, [ScriptInfoMapPtr, Self = AsShared()]
 		{
 			for (TMap<TObjectKey<UNiagaraScript>, FScriptInfo>::ElementType& CurrentIt : *ScriptInfoMapPtr)
 			{
@@ -2409,9 +2539,9 @@ UE::Tasks::FTask FNiagaraSystemCompilationTask::BuildRapidIterationParametersAsy
 
 						if (SrcScriptInfo.SourceEmitterIndex == INDEX_NONE)
 						{
-							SystemInfo.StaticVariableResults.AddUnique(StaticVariable);
+							Self->SystemInfo.StaticVariableResults.AddUnique(StaticVariable);
 						}
-						else if (FEmitterInfo* EmitterInfo = SystemInfo.EmitterInfoBySourceEmitter(SrcScriptInfo.SourceEmitterIndex))
+						else if (FEmitterInfo* EmitterInfo = Self->SystemInfo.EmitterInfoBySourceEmitter(SrcScriptInfo.SourceEmitterIndex))
 						{
 							EmitterInfo->StaticVariableResults.AddUnique(StaticVariable);
 						}
@@ -2462,9 +2592,9 @@ void FNiagaraSystemCompilationTask::IssuePostResultsProcessedTasks()
 	PutRequestHelper = MakeUnique<FDispatchDataCachePutRequests>();
 	PutRequestHelper->Launch(this);
 
-	Launch(UE_SOURCE_LOCATION, [this]
+	Launch(UE_SOURCE_LOCATION, [Self = AsShared()]
 	{
-		CompilationState = EState::Completed;
+		Self->CompilationState = EState::Completed;
 	}, PutRequestHelper->CompletionEvent);
 }
 
@@ -2477,11 +2607,11 @@ void FNiagaraSystemCompilationTask::IssueCompilationTasks()
 
 	if (HasOutstandingCompileTasks())
 	{
-		FTask PrecompileTask = Launch(UE_SOURCE_LOCATION, [this]
+		FTask PrecompileTask = Launch(UE_SOURCE_LOCATION, [Self = AsShared()]
 		{
 			COOK_STAT(auto Timer = NiagaraSystemCookStats::UsageStats.TimeSyncWork(); Timer.TrackCyclesOnly(););
 
-			Precompile();
+			Self->Precompile();
 		});
 
 		// async compilation copy tasks
@@ -2489,11 +2619,11 @@ void FNiagaraSystemCompilationTask::IssueCompilationTasks()
 		{
 			if (GroupInfo.HasOutstandingCompileTasks(*this))
 			{
-				FTask InstantiateGraphTask = Launch(UE_SOURCE_LOCATION, [this, &GroupInfo]
+				FTask InstantiateGraphTask = Launch(UE_SOURCE_LOCATION, [Self = AsShared(), &GroupInfo]
 				{
 					COOK_STAT(auto Timer = NiagaraSystemCookStats::UsageStats.TimeSyncWork(); Timer.TrackCyclesOnly(););
 
-					GroupInfo.InstantiateCompileGraph(*this);
+					GroupInfo.InstantiateCompileGraph(*Self);
 				}, PrecompileTask);
 
 				for (int32 CompileTaskIndex : GroupInfo.CompileTaskIndices)
@@ -2504,45 +2634,47 @@ void FNiagaraSystemCompilationTask::IssueCompilationTasks()
 						CompileTask.CompileResultsReadyEvent = MakeUnique<FTaskEvent>(UE_SOURCE_LOCATION);
 						CompileTask.CompileResultsProcessedEvent = MakeUnique<FTaskEvent>(UE_SOURCE_LOCATION);
 
-						IssuedCompilationTasks.Add(Launch(UE_SOURCE_LOCATION, [this, &GroupInfo, &CompileTask]
+						IssuedCompilationTasks.Add(Launch(UE_SOURCE_LOCATION, [Self = AsShared(), &GroupInfo, &CompileTask]
 						{
 							COOK_STAT(auto Timer = NiagaraSystemCookStats::UsageStats.TimeSyncWork(); Timer.TrackCyclesOnly(););
 
-							CompileTask.CollectNamedDataInterfaces(this, GroupInfo);
+							FNiagaraSystemCompilationTask* TaskPtr = &Self.Get();
+
+							CompileTask.CollectNamedDataInterfaces(TaskPtr, GroupInfo);
 							switch (CompileTask.ScriptCompileType)
 							{
 								case EScriptCompileType::CompileForVm:
-									CompileTask.Translate(this, GroupInfo);
-									CompileTask.IssueCompileVm(this, GroupInfo);
+									CompileTask.Translate(TaskPtr, GroupInfo);
+									CompileTask.IssueCompileVm(TaskPtr, GroupInfo);
 								break;
 
 								case EScriptCompileType::CompileForGpu:
-									CompileTask.Translate(this, GroupInfo);
-									CompileTask.IssueCompileGpu(this, GroupInfo);
+									CompileTask.Translate(TaskPtr, GroupInfo);
+									CompileTask.IssueCompileGpu(TaskPtr, GroupInfo);
 								break;
 
 								// path for handling the translation of GPU scripts but we're not handling the compilation of
 								// the compute shaders ourselves, rather it's using the old path of the FNiagaraShaderScript getting
 								// cached by the UNiagaraScript post CPU translation/compilation
 								case EScriptCompileType::TranslateForGpu:
-									CompileTask.Translate(this, GroupInfo);
-									CompileTask.IssueTranslateGpu(this, GroupInfo);
+									CompileTask.Translate(TaskPtr, GroupInfo);
+									CompileTask.IssueTranslateGpu(TaskPtr, GroupInfo);
 								break;
 
 								// path for dealing with the CPU scripts (Particle spawn/update) when the emitter is GPU
 								// These scripts shouldn't need to be processed, but currently they are required for supplying
 								// RI parameters to the GPU execution context
 								case EScriptCompileType::DummyCompileForCpuStubs:
-									CompileTask.Translate(this, GroupInfo);
+									CompileTask.Translate(TaskPtr, GroupInfo);
 								break;
 							}
 						}, InstantiateGraphTask));
 
-						Launch(UE_SOURCE_LOCATION, [this, &GroupInfo, &CompileTask]
+						Launch(UE_SOURCE_LOCATION, [Self = AsShared(), &GroupInfo, &CompileTask]
 						{
 							COOK_STAT(auto Timer = NiagaraSystemCookStats::UsageStats.TimeSyncWork(); Timer.TrackCyclesOnly(););
 
-							CompileTask.ProcessCompilation(this, GroupInfo);
+							CompileTask.ProcessCompilation(&Self.Get(), GroupInfo);
 						}, *CompileTask.CompileResultsReadyEvent);
 
 						CompilationTasks.Add(*CompileTask.CompileResultsProcessedEvent);
@@ -2552,14 +2684,14 @@ void FNiagaraSystemCompilationTask::IssueCompilationTasks()
 		}
 	}
 
-	Launch(UE_SOURCE_LOCATION, [this]
+	Launch(UE_SOURCE_LOCATION, [Self = AsShared()]
 	{
-		CompilationState = EState::WaitingForProcessing;
+		Self->CompilationState = EState::WaitingForProcessing;
 	}, IssuedCompilationTasks);
 
-	Launch(UE_SOURCE_LOCATION, [this]
+	Launch(UE_SOURCE_LOCATION, [Self = AsShared()]
 	{
-		IssuePostResultsProcessedTasks();
+		Self->IssuePostResultsProcessedTasks();
 	}, CompilationTasks);
 }
 
@@ -2611,35 +2743,35 @@ UE::Tasks::FTask FNiagaraSystemCompilationTask::BeginTasks()
 		}
 	}
 
-	FTask SystemTask = Launch(UE_SOURCE_LOCATION, [this]
+	FTask SystemTask = Launch(UE_SOURCE_LOCATION, [Self = AsShared()]
 	{
 		COOK_STAT(auto Timer = NiagaraSystemCookStats::UsageStats.TimeSyncWork(); Timer.TrackCyclesOnly(););
 
-		if (HasOutstandingCompileTasks())
+		if (Self->HasOutstandingCompileTasks())
 		{
 			TArray<FTask> CompilationTaskPrerequisites;
 
-			if (!SystemInfo.bUseRapidIterationParams)
+			if (!Self->SystemInfo.bUseRapidIterationParams)
 			{
-				FTask BuildRIParamTask = BuildRapidIterationParametersAsync();
-				PostRIParameterGetRequestHelper = MakeUnique<FDispatchAndProcessDataCacheGetRequests>();
+				FTask BuildRIParamTask = Self->BuildRapidIterationParametersAsync();
+				Self->PostRIParameterGetRequestHelper = MakeUnique<FDispatchAndProcessDataCacheGetRequests>();
 
-				Launch(UE_SOURCE_LOCATION, [this]
+				Launch(UE_SOURCE_LOCATION, [Self]
 				{
-					PostRIParameterGetRequestHelper->Launch(this);
+					Self->PostRIParameterGetRequestHelper->Launch(&Self.Get());
 				}, BuildRIParamTask);
 
-				CompilationTaskPrerequisites.Add(PostRIParameterGetRequestHelper->CompletionEvent);
+				CompilationTaskPrerequisites.Add(Self->PostRIParameterGetRequestHelper->CompletionEvent);
 			}
 
-			FTask InnerTask = Launch(UE_SOURCE_LOCATION, [this]
+			FTask InnerTask = Launch(UE_SOURCE_LOCATION, [Self]
 			{
-				IssueCompilationTasks();
+				Self->IssueCompilationTasks();
 			}, CompilationTaskPrerequisites);
 		}
 		else
 		{
-			IssuePostResultsProcessedTasks();
+			Self->IssuePostResultsProcessedTasks();
 		}
 	}, SystemTaskPrerequisites);
 
diff --git a/Engine/Plugins/FX/Niagara/Source/NiagaraEditor/Private/NiagaraCompilationTasks.h b/Engine/Plugins/FX/Niagara/Source/NiagaraEditor/Private/NiagaraCompilationTasks.h
--- a/Engine/Plugins/FX/Niagara/Source/NiagaraEditor/Private/NiagaraCompilationTasks.h
+++ b/Engine/Plugins/FX/Niagara/Source/NiagaraEditor/Private/NiagaraCompilationTasks.h
@@ -7,8 +7,9 @@
 #include "DerivedDataRequestOwner.h"
 #include "NiagaraCompiler.h"
 #include "NiagaraCompilationPrivate.h"
+#include "Templates/SharedPointer.h"
 
-struct FNiagaraSystemCompilationTask
+struct FNiagaraSystemCompilationTask : public TSharedFromThis<FNiagaraSystemCompilationTask, ESPMode::ThreadSafe>
 {
 	// helper for script DDC get requests
 	struct FDispatchAndProcessDataCacheGetRequests
@@ -19,7 +20,7 @@ struct FNiagaraSystemCompilationTask
 		std::atomic<int32> PendingGetRequestCount = 0;
 	};
 
-	// helepr for script DDC put requests
+	// helper for script DDC put requests
 	struct FDispatchDataCachePutRequests
 	{
 		void Launch(FNiagaraSystemCompilationTask* SystemCompileTask);
@@ -252,6 +253,8 @@ struct FNiagaraSystemCompilationTask
 
 	void GetAvailableCollections(TArray<FNiagaraCompilationNPCHandle>& OutCollections) const;
 
+	bool GenerateCompileMetaData() const;
+
 	UE::DerivedData::FRequestOwner DDCRequestOwner;
 
 private:
@@ -274,4 +277,6 @@ private:
 	UE::Tasks::FTaskEvent CompileCompletionEvent = UE::Tasks::FTaskEvent(UE_SOURCE_LOCATION);
 
 	bool bAborting = false;
+
+	bool bGenerateCompileMetaData = false;
 };
diff --git a/Engine/Plugins/FX/Niagara/Source/NiagaraEditor/Private/NiagaraEditorSettings.cpp b/Engine/Plugins/FX/Niagara/Source/NiagaraEditor/Private/NiagaraEditorSettings.cpp
--- a/Engine/Plugins/FX/Niagara/Source/NiagaraEditor/Private/NiagaraEditorSettings.cpp
+++ b/Engine/Plugins/FX/Niagara/Source/NiagaraEditor/Private/NiagaraEditorSettings.cpp
@@ -909,5 +909,10 @@ bool UNiagaraEditorSettings::GetUpdateStackValuesOnCommitOnly() const
 	return bUpdateStackValuesOnCommitOnly;
 }
 
+bool UNiagaraEditorSettings::GetGenerateMetaDataOnCompile() const
+{
+	return bGenerateMetaDataOnCompile;
+}
+
 #undef LOCTEXT_NAMESPACE
 
diff --git a/Engine/Plugins/FX/Niagara/Source/NiagaraEditor/Public/NiagaraEditorSettings.h b/Engine/Plugins/FX/Niagara/Source/NiagaraEditor/Public/NiagaraEditorSettings.h
--- a/Engine/Plugins/FX/Niagara/Source/NiagaraEditor/Public/NiagaraEditorSettings.h
+++ b/Engine/Plugins/FX/Niagara/Source/NiagaraEditor/Public/NiagaraEditorSettings.h
@@ -566,6 +566,10 @@ private:
 	UPROPERTY(config, EditAnywhere, Category = Niagara)
 	bool bUpdateStackValuesOnCommitOnly = false;
 
+	/** */
+	UPROPERTY(config, EditAnywhere, Category = Niagara)
+	bool bGenerateMetaDataOnCompile = false;
+
 	/** Speeds used for slowing down and speeding up the playback speeds */
 	UPROPERTY(config, EditAnywhere, Category = Niagara)
 	TArray<float> PlaybackSpeeds;
@@ -622,4 +626,6 @@ public:
 	NIAGARAEDITOR_API void SetShowMemoryInfo(bool bInShowInfo);
 	NIAGARAEDITOR_API bool IsShowStatelessInfo() const;
 	NIAGARAEDITOR_API void SetShowStatelessInfo(bool bInShowInfo);
+	NIAGARAEDITOR_API bool GetGenerateMetaDataOnCompile() const;
+
 };
