Fix bug where Niagara's old viewport toolbar didn't show if CVAR ToolMenusViewportToolbars was 2 when Niagara viewport was opened

diff --git a/Engine/Plugins/FX/Niagara/Source/NiagaraEditor/Private/Widgets/SNiagaraSystemViewport.cpp b/Engine/Plugins/FX/Niagara/Source/NiagaraEditor/Private/Widgets/SNiagaraSystemViewport.cpp
--- a/Engine/Plugins/FX/Niagara/Source/NiagaraEditor/Private/Widgets/SNiagaraSystemViewport.cpp
+++ b/Engine/Plugins/FX/Niagara/Source/NiagaraEditor/Private/Widgets/SNiagaraSystemViewport.cpp
@@ -1366,22 +1366,19 @@ FText SNiagaraSystemViewport::GetViewportCompileStatusText() const
 
 void SNiagaraSystemViewport::PopulateViewportOverlays(TSharedRef<class SOverlay> Overlay)
 {
-	if (UE::UnrealEd::ShowOldViewportToolbars())
-	{
-		// clang-format off
-		Overlay->AddSlot()
-		.VAlign(VAlign_Top)
-		[
-			SNew(SNiagaraSystemViewportToolBar, SharedThis(this))
-			.Sequencer(Sequencer)
-			.Visibility_Lambda([]()
-				{
-					return UE::UnrealEd::ShowOldViewportToolbars() ? EVisibility::Visible : EVisibility::Collapsed;
-				}
-			)
-		];
-		// clang-format on
-	}
+	// clang-format off
+	Overlay->AddSlot()
+	.VAlign(VAlign_Top)
+	[
+		SNew(SNiagaraSystemViewportToolBar, SharedThis(this))
+		.Sequencer(Sequencer)
+		.Visibility_Lambda([]()
+			{
+				return UE::UnrealEd::ShowOldViewportToolbars() ? EVisibility::Visible : EVisibility::Collapsed;
+			}
+		)
+	];
+	// clang-format on
 
 	Overlay->AddSlot()
 	.VAlign(VAlign_Center)
