- Add feedback if no parameters are found that we can bind to

diff --git a/Engine/Plugins/FX/Niagara/Source/NiagaraEditor/Private/Customizations/NiagaraTypeCustomizations.cpp b/Engine/Plugins/FX/Niagara/Source/NiagaraEditor/Private/Customizations/NiagaraTypeCustomizations.cpp
--- a/Engine/Plugins/FX/Niagara/Source/NiagaraEditor/Private/Customizations/NiagaraTypeCustomizations.cpp
+++ b/Engine/Plugins/FX/Niagara/Source/NiagaraEditor/Private/Customizations/NiagaraTypeCustomizations.cpp
@@ -2351,6 +2351,19 @@ TSharedRef<SWidget> FNiagaraRendererMaterialParameterCustomization::OnGetMateria
 				);
 			}
 		}
+
+		if (AddedParameterNames.IsEmpty())
+		{
+			MenuBuilder.AddMenuEntry(
+				LOCTEXT("NoParametersFound", "No Parameters Found"),
+				FText(),
+				FSlateIcon(),
+				FUIAction(
+					FExecuteAction(),
+					FCanExecuteAction::CreateLambda([]() { return false;} )
+				)
+			);
+		}
 	}
 
 	return MenuBuilder.MakeWidget();
