- Add Gpu Cone / Torus support for debug draw

diff --git a/Engine/Plugins/FX/Niagara/Shaders/Private/NiagaraDataInterfaceDebugDraw.ush b/Engine/Plugins/FX/Niagara/Shaders/Private/NiagaraDataInterfaceDebugDraw.ush
--- a/Engine/Plugins/FX/Niagara/Shaders/Private/NiagaraDataInterfaceDebugDraw.ush
+++ b/Engine/Plugins/FX/Niagara/Shaders/Private/NiagaraDataInterfaceDebugDraw.ush
@@ -3,6 +3,7 @@
 #pragma once
 
 #include "NiagaraCommon.ush"
+#include "NiagaraQuaternionUtils.ush"
 #include "NiagaraDebugDraw.ush"
 
 RWBuffer<uint>		RWNDIDebugDrawArgs;
@@ -379,80 +380,124 @@ void NDIDebugDraw_Cylinder(bool bDraw, float3 Location, float3 Axis, float Heigh
 
 void NDIDebugDraw_Cone(bool bDraw, float3 Location, float3 Axis, float Height, float RadiusTop, float RadiusBottom, int NumHeightSegments, int NumRadiusSegments, float4 Color)
 {
-	if ( bDraw && (NDIDebugDrawLineMaxInstances > 0) )
+    if (!bDraw || (NDIDebugDrawLineMaxInstances <= 0))
+    {
+        return;
+    }
+
+    NumHeightSegments = clamp(NumHeightSegments, 2, 16);
+    NumRadiusSegments = clamp(NumRadiusSegments, 3, 16);
+    const int NumLines = NumRadiusSegments + (NumRadiusSegments * NumHeightSegments);
+ 
+	uint InstanceIndex = 0;
+	InterlockedAdd(RWNDIDebugDrawArgs[1], NumLines, InstanceIndex);
+	if (InstanceIndex + NumLines > NDIDebugDrawLineMaxInstances)
 	{
-		NumHeightSegments = clamp(NumHeightSegments, 1, 16);
-		NumRadiusSegments = clamp(NumRadiusSegments, 3, 16);
+		InterlockedAdd(RWNDIDebugDrawArgs[1], -NumLines);
+		return;
+	}
+    
+    const float3 AxisNorm = normalize(Axis);
+	const float4 Rotation = FindQuatBetweenNormals(float3(0, 0, 1), AxisNorm);
+    
+	const float HeightOffset = Height * 0.5f;
+	const float HeightStep = Height / float(NumHeightSegments - 1);
+
+	const float HeightSegmentStepAlpha = 1.0f / float(NumHeightSegments - 1);
+
+	float CurrentRotation = 0.0f;
+	float3 LastPointNormal = float3(sin(CurrentRotation), cos(CurrentRotation), 0.0f);
+	
+	for (int Idx = 1; Idx <= NumRadiusSegments; Idx++)
+	{
+		CurrentRotation = 2.0f * PI / float(NumRadiusSegments) * Idx;
 
-		int NumLines = NumRadiusSegments + (NumRadiusSegments * NumHeightSegments + 1);
+		const float3 CurrentPointNormal = float3(sin(CurrentRotation), cos(CurrentRotation), 0.0f);
+		const float3 CurrentBottom = RotateVectorByQuat(CurrentPointNormal * RadiusBottom + float3(0.0f, 0.0f, -HeightOffset), Rotation);
+		const float3 CurrentTop = RotateVectorByQuat(CurrentPointNormal * RadiusTop + float3(0.0f, 0.0f, Height - HeightOffset), Rotation);
 
-		uint InstanceIndex = 0;
-		InterlockedAdd(RWNDIDebugDrawArgs[1], NumLines, InstanceIndex);
-		if ( InstanceIndex + NumLines <= NDIDebugDrawLineMaxInstances )
+        PackLine(RWNDIDebugDrawLineVertex, InstanceIndex++, Location + CurrentBottom, Location + CurrentTop, Color);
+        
+		for (int HeightIdx=0; HeightIdx < NumHeightSegments; ++HeightIdx)
 		{
-			float3 AxisNorm = normalize(Axis);
+			const float RingRadius = lerp(RadiusBottom, RadiusTop, HeightSegmentStepAlpha * float(HeightIdx));
+			const float3 RingOffset = float3(0.0f, 0.0f, (HeightStep * float(HeightIdx)) - HeightOffset);
 
-			const float3x3 AxisRotation = NiagaraFindRotationBetweenNormals(float3(0,0,1), AxisNorm);
-
-			float3 TangentXScaled = mul(AxisRotation, float3(1, 0, 0));
-			float3 TangentYScaled = mul(AxisRotation, float3(0, 1, 0));
+			const float3 LastRingPosition = RotateVectorByQuat(LastPointNormal * RingRadius + RingOffset, Rotation);
+			const float3 CurrentRingPosition = RotateVectorByQuat(CurrentPointNormal * RingRadius + RingOffset, Rotation);
 
-			float3 HalfHeightOffset = AxisNorm * Height * 0.5f;
-			float3 AxisStep = AxisNorm * Height / NumHeightSegments;
-
-			float CurrentRotation = 0.0f;
-			float3 LastPointNormal = Location + (TangentXScaled * cos(CurrentRotation)) + (TangentYScaled * sin(CurrentRotation));
+			PackLine(RWNDIDebugDrawLineVertex, InstanceIndex++, Location + LastRingPosition, Location + CurrentRingPosition, Color);
+		}
 
-			float HeightSegmentStepAlpha = 1.0 / NumHeightSegments;
+		LastPointNormal = CurrentPointNormal;
+	}
+}
 
-			for (int Idx = 1; Idx <= NumRadiusSegments; Idx++)
-			{
-				CurrentRotation = 2.0f * PI / NumRadiusSegments * Idx;
-				float3 CurrentPointNormal = (TangentXScaled * cos(CurrentRotation)) + (TangentYScaled * sin(CurrentRotation));
-				float3 CurrentPoint = Location + CurrentPointNormal;
-				
-				float3 LocationBottom = Location - HalfHeightOffset;
-				float3 LocationTop = Location + HalfHeightOffset;
+void NDIDebugDraw_Torus(bool bDraw, float3 Location, float3 Axis, float MajorRadius, float MinorRadius, float MajorRadiusSegments, float MinorRadiusSegments, float4 Color)
+{
+	if (!bDraw || (NDIDebugDrawLineMaxInstances <= 0))
+	{
+		return;
+	}
 
-				float3 LastBottom = LocationBottom + (LastPointNormal * RadiusBottom);
-				float3 LastTop = LocationTop + (LastPointNormal * RadiusTop);
+	MinorRadiusSegments = clamp(MinorRadiusSegments, 3, 32);
+	MajorRadiusSegments = clamp(MajorRadiusSegments, 3, 32);
 
-				float3 CurrentBottom = (LocationBottom + (CurrentPointNormal * RadiusBottom));
-				float3 CurrentTop = (LocationTop + (CurrentPointNormal * RadiusTop));
-				
-				// Add Height Line
-				PackLine(RWNDIDebugDrawLineVertex, InstanceIndex++, CurrentBottom, CurrentTop, Color);
+	const int NumLines = MajorRadiusSegments * MinorRadiusSegments * 2;
+ 
+	uint InstanceIndex = 0;
+	InterlockedAdd(RWNDIDebugDrawArgs[1], NumLines, InstanceIndex);
+	if (InstanceIndex + NumLines > NDIDebugDrawLineMaxInstances)
+	{
+		InterlockedAdd(RWNDIDebugDrawArgs[1], -NumLines);
+		return;
+	}
+    
+	const float3 AxisNorm = normalize(Axis);
+	const float4 Rotation = FindQuatBetweenNormals(float3(0, 0, 1), AxisNorm);
 
-				// Add Bottom
-				PackLine(RWNDIDebugDrawLineVertex, InstanceIndex++, LastBottom, CurrentBottom, Color);
-				
-				// Add all Rings + top
-				for (int HeightIdx = 1; HeightIdx <= NumHeightSegments; HeightIdx++)
-				{
-					float RingRadius = lerp(RadiusBottom, RadiusTop, HeightSegmentStepAlpha * HeightIdx);
+	const float3 TangentX = RotateVectorByQuat(float3(1, 0, 0), Rotation);
+	const float3 TangentY = RotateVectorByQuat(float3(0, 1, 0), Rotation);
 
-					float3 RingOffset = AxisStep * HeightIdx;
-					float3 LocationRing = LocationBottom + RingOffset;
+	const float MajorRotationDelta = 2.0f * PI / float(MajorRadiusSegments);
+	const float MinorRotationDelta = 2.0f * PI / float(MinorRadiusSegments);
 
-					float3 LastRingPosition = LocationRing + (LastPointNormal * RingRadius);
-					float3 CurrRingPosition = LocationRing + (CurrentPointNormal * RingRadius);
+	float CurrentMajorRotation = 0.0f;
+	float3 PreviousMajorNormal = (TangentX * cos(CurrentMajorRotation)) + (TangentY * sin(CurrentMajorRotation));
+	
+	for (int MajorIdx = 0; MajorIdx < MajorRadiusSegments; ++MajorIdx)
+	{
+		CurrentMajorRotation += MajorRotationDelta;
+		const float3 CurrentMajorNormal = (TangentX * cos(CurrentMajorRotation)) + (TangentY * sin(CurrentMajorRotation));
 
-					PackLine(RWNDIDebugDrawLineVertex, InstanceIndex++, LastRingPosition, CurrRingPosition, Color);
-				}
+		float CurrentMinorRotation = 0.0f;
+		float3 PreviousMinorNormal = (CurrentMajorNormal * cos(CurrentMinorRotation)) + (AxisNorm * sin(CurrentMinorRotation));
 
-				LastPointNormal = CurrentPointNormal;
-			}
-		}
-		else
+		for (int MinorIdx = 0; MinorIdx < MinorRadiusSegments; ++MinorIdx)
 		{
-			// Not enough space, reduce count back to original
-			InterlockedAdd(RWNDIDebugDrawArgs[1], -NumLines);
+			CurrentMinorRotation += MinorRotationDelta;
+			const float3 CurrentMinorNormal = (CurrentMajorNormal * cos(CurrentMinorRotation)) + (AxisNorm * sin(CurrentMinorRotation));
+
+			PackLine(
+				RWNDIDebugDrawLineVertex, InstanceIndex++,
+				Location + CurrentMajorNormal * MajorRadius + PreviousMinorNormal * MinorRadius,
+				Location + CurrentMajorNormal * MajorRadius + CurrentMinorNormal * MinorRadius,
+				Color
+			);
+
+			const float3 PreviousMajorMinorNormal = (PreviousMajorNormal * cos(CurrentMinorRotation)) + (AxisNorm * sin(CurrentMinorRotation));
+			PackLine(
+				RWNDIDebugDrawLineVertex, InstanceIndex++,
+				Location + PreviousMajorNormal * MajorRadius + PreviousMajorMinorNormal * MinorRadius,
+				Location + CurrentMajorNormal * MajorRadius + CurrentMinorNormal * MinorRadius,
+				Color
+			);
+
+			PreviousMinorNormal = CurrentMinorNormal;
 		}
-	}
-}
 
-void NDIDebugDraw_Torus(bool bDraw, float3 Location, float3 Axis, float MajorRadius, float MinorRadius, float MajorRadiusSegments, float MinorRadiusSegments, float4 Color)
-{
+		PreviousMajorNormal = CurrentMajorNormal;
+	}
 }
 
 
diff --git a/Engine/Plugins/FX/Niagara/Shaders/Private/NiagaraDataInterfaceDebugDrawTemplate.ush b/Engine/Plugins/FX/Niagara/Shaders/Private/NiagaraDataInterfaceDebugDrawTemplate.ush
--- a/Engine/Plugins/FX/Niagara/Shaders/Private/NiagaraDataInterfaceDebugDrawTemplate.ush
+++ b/Engine/Plugins/FX/Niagara/Shaders/Private/NiagaraDataInterfaceDebugDrawTemplate.ush
@@ -33,10 +33,12 @@ void DrawCylinder_{ParameterName}_UEImpureCall(bool bExecute, float3 Location, f
 
 void DrawCone_{ParameterName}_UEImpureCall(bool bExecute, float3 Location, float3 Axis, float Height, float RadiusTop, float RadiusBottom, int NumHeightSegments, int NumRadiusSegments, float4 Color)
 {
+	NDIDebugDraw_Cone(bExecute, Location, Axis, Height, RadiusTop, RadiusBottom, NumHeightSegments, NumRadiusSegments, Color);
 }
 
 void DrawTorus_{ParameterName}_UEImpureCall(bool bExecute, float3 Location, float3 Axis, float MajorRadius, float MinorRadius, int MajorRadiusSegments, int MinorRadiusSegments, float4 Color)
 {
+	NDIDebugDraw_Torus(bExecute, Location, Axis, MajorRadius, MinorRadius, MajorRadiusSegments, MinorRadiusSegments, Color);
 }
 
 void DrawCoordinateSystem_{ParameterName}_UEImpureCall(bool bExecute, float3 Location, float4 Rotation, float Scale)
diff --git a/Engine/Plugins/FX/Niagara/Source/Niagara/Private/DataInterface/NiagaraDataInterfaceDebugDraw.cpp b/Engine/Plugins/FX/Niagara/Source/Niagara/Private/DataInterface/NiagaraDataInterfaceDebugDraw.cpp
--- a/Engine/Plugins/FX/Niagara/Source/Niagara/Private/DataInterface/NiagaraDataInterfaceDebugDraw.cpp
+++ b/Engine/Plugins/FX/Niagara/Source/Niagara/Private/DataInterface/NiagaraDataInterfaceDebugDraw.cpp
@@ -264,18 +264,21 @@ struct FNDIDebugDrawInstanceData_GameThread
 
 	void AddCone(const FVector3f& Location, const FVector3f& Axis, float Height, float RadiusTop, float RadiusBottom, int32 NumHeightSegments, int32 NumRadiusSegments, const FLinearColor& Color)
 	{
-		FQuat4f Rotation = FQuat4f::FindBetweenNormals(FVector3f::UpVector, Axis);
+		FQuat4f Rotation = FQuat4f::FindBetweenVectors(FVector3f::UpVector, Axis);
 
 		AddCone(Location, Rotation, Height, RadiusTop, RadiusBottom, NumHeightSegments, NumRadiusSegments, Color, FVector3f::OneVector, true);
 	}
 
 	void AddCone(const FVector3f& Location, const FQuat4f& Rotation, float Height, float RadiusTop, float RadiusBottom, int32 NumHeightSegments, int32 NumRadiusSegments, const FLinearColor& Color, const FVector3f& NonUniformScale, bool bCenterVertically)
 	{
+		NumHeightSegments = FMath::Clamp(NumHeightSegments, 2, 32);
+		NumRadiusSegments = FMath::Clamp(NumRadiusSegments, 3, 32);
+
 		const FVector3f HeightVector = FVector3f::ZAxisVector * Height;
 		const FVector3f HalfHeightOffset = bCenterVertically ? HeightVector * 0.5f : FVector3f::ZeroVector;
-		const FVector3f AxisStep = FVector3f::ZAxisVector * Height / float(NumHeightSegments);
+		const FVector3f AxisStep = FVector3f::ZAxisVector * Height / float(NumHeightSegments - 1);
 
-		const float HeightSegmentStepAlpha = 1.0f / float(NumHeightSegments);
+		const float HeightSegmentStepAlpha = 1.0f / float(NumHeightSegments - 1);
 
 		float CurrentRotation = 0.0f;
 		FVector3f LastPointNormal = (FVector3f::YAxisVector * FMath::Cos(CurrentRotation)) + (FVector3f::XAxisVector * FMath::Sin(CurrentRotation));
@@ -286,20 +289,14 @@ struct FNDIDebugDrawInstanceData_GameThread
 
 			FVector3f CurrentPointNormal = (FVector3f::YAxisVector * FMath::Cos(CurrentRotation)) + (FVector3f::XAxisVector * FMath::Sin(CurrentRotation));
 
-			const FVector3f LastBottom = LastPointNormal * RadiusBottom - HalfHeightOffset;
-			const FVector3f LastTop = LastPointNormal * RadiusTop + HeightVector - HalfHeightOffset;
-
 			const FVector3f CurrentBottom = CurrentPointNormal * RadiusBottom - HalfHeightOffset;
 			const FVector3f CurrentTop = CurrentPointNormal * RadiusTop + HeightVector - HalfHeightOffset;
 
 			// Add Height line
 			AddLine(Location + Rotation.RotateVector(CurrentBottom * NonUniformScale), Location + Rotation.RotateVector(CurrentTop * NonUniformScale), Color);
 
-			// Add Bottom
-			AddLine(Location + Rotation.RotateVector(LastBottom * NonUniformScale), Location + Rotation.RotateVector(CurrentBottom * NonUniformScale), Color);
-
-			// Add all Rings + top
-			for (int32 HeightIdx = 1; HeightIdx <= NumHeightSegments; HeightIdx++)
+			// Add Rings
+			for (int32 HeightIdx = 0; HeightIdx < NumHeightSegments; HeightIdx++)
 			{
 				float RingRadius = FMath::Lerp(RadiusBottom, RadiusTop, HeightSegmentStepAlpha * HeightIdx);
 
@@ -318,6 +315,9 @@ struct FNDIDebugDrawInstanceData_GameThread
 
 	void AddTorus(const FVector3f& Location, const FVector3f& Axis, float MajorRadius, float MinorRadius, int32 MajorRadiusSegments, int32 MinorRadiusSegments, const FLinearColor& Color)
 	{
+		MinorRadiusSegments = FMath::Clamp(MinorRadiusSegments, 3, 32);
+		MajorRadiusSegments = FMath::Clamp(MajorRadiusSegments, 3, 32);
+
 		const FVector3f AxisNorm = Axis.GetSafeNormal();
 
 		const FQuat4f AxisRotation = FQuat4f::FindBetweenNormals(FVector3f::UpVector, AxisNorm);
@@ -331,7 +331,7 @@ struct FNDIDebugDrawInstanceData_GameThread
 		float CurrentMajorRotation = 0.0f;
 		FVector3f PreviousMajorNormal = (TangentX * FMath::Cos(CurrentMajorRotation)) + (TangentY * FMath::Sin(CurrentMajorRotation));
 
-		for (int32 MajorIdx = 0; MajorIdx < MajorRadiusSegments; MajorIdx++)
+		for (int32 MajorIdx = 0; MajorIdx < MajorRadiusSegments; ++MajorIdx)
 		{
 			CurrentMajorRotation += MajorRotationDelta;
 			const FVector3f CurrentMajorNormal = (TangentX * FMath::Cos(CurrentMajorRotation)) + (TangentY * FMath::Sin(CurrentMajorRotation));
@@ -339,7 +339,7 @@ struct FNDIDebugDrawInstanceData_GameThread
 			float CurrentMinorRotation = 0.0f;
 			FVector3f PreviousMinorNormal = (CurrentMajorNormal * FMath::Cos(CurrentMinorRotation)) + (AxisNorm * FMath::Sin(CurrentMinorRotation));
 
-			for (int32 MinorIdx = 0; MinorIdx < MinorRadiusSegments; MinorIdx++)
+			for (int32 MinorIdx = 0; MinorIdx < MinorRadiusSegments; ++MinorIdx)
 			{
 				CurrentMinorRotation += MinorRotationDelta;
 				FVector3f CurrentMinorNormal = (CurrentMajorNormal * FMath::Cos(CurrentMinorRotation)) + (AxisNorm * FMath::Sin(CurrentMinorRotation));
@@ -356,7 +356,6 @@ struct FNDIDebugDrawInstanceData_GameThread
 
 			PreviousMajorNormal = CurrentMajorNormal;
 		}
-
 	}
 
 	void AddCoordinateSystem(const FVector3f& Location, const FQuat4f& Rotation, float Scale)
