- Fix race when generating GPU ribbons with parallel GDME enabled

diff --git a/Engine/Plugins/FX/Niagara/Source/Niagara/Private/NiagaraRendererRibbons.cpp b/Engine/Plugins/FX/Niagara/Source/Niagara/Private/NiagaraRendererRibbons.cpp
--- a/Engine/Plugins/FX/Niagara/Source/Niagara/Private/NiagaraRendererRibbons.cpp
+++ b/Engine/Plugins/FX/Niagara/Source/Niagara/Private/NiagaraRendererRibbons.cpp
@@ -594,7 +594,7 @@ public:
 	}
 
 	//-OPT: These caches should be more central and are as a simple solution to reduce memory thrashing / poor performance for ribbons
-	const TSharedPtr<FRWBuffer>& GetOrAllocateIndirectDrawBuffer(FRHICommandListBase& RHICmdList)
+	TSharedPtr<FRWBuffer> GetOrAllocateIndirectDrawBuffer(FRHICommandListBase& RHICmdList)
 	{
 		UE::TScopeLock ScopeLock(AllocateGuard);
 
@@ -614,7 +614,7 @@ public:
 		return BufferEntry->Buffer;
 	}
 
-	const TSharedPtr<FNiagaraRibbonIndexBuffer>& GetOrAllocateIndexBuffer(FRHICommandListBase& RHICmdList, int32 NumIndices, int32 MaxIndicesEstimate)
+	TSharedPtr<FNiagaraRibbonIndexBuffer> GetOrAllocateIndexBuffer(FRHICommandListBase& RHICmdList, int32 NumIndices, int32 MaxIndicesEstimate)
 	{
 		UE::TScopeLock ScopeLock(AllocateGuard);
 
