- Only consider validation rules active if they are not nullptr

diff --git a/Engine/Plugins/FX/Niagara/Source/Niagara/Classes/NiagaraValidationRuleSet.h b/Engine/Plugins/FX/Niagara/Source/Niagara/Classes/NiagaraValidationRuleSet.h
--- a/Engine/Plugins/FX/Niagara/Source/Niagara/Classes/NiagaraValidationRuleSet.h
+++ b/Engine/Plugins/FX/Niagara/Source/Niagara/Classes/NiagaraValidationRuleSet.h
@@ -13,6 +13,10 @@ class UNiagaraValidationRuleSet : public UObject
 {
 	GENERATED_BODY()
 public:
+	bool HasAnyRules() const
+	{
+		return ValidationRules.ContainsByPredicate([](UNiagaraValidationRule* Rule) { return Rule != nullptr; });
+	}
 
 	UPROPERTY(EditAnywhere, Category = "Validation", Instanced)
 	TArray<TObjectPtr<UNiagaraValidationRule>> ValidationRules;
diff --git a/Engine/Plugins/FX/Niagara/Source/NiagaraEditor/Private/NiagaraValidationRules.cpp b/Engine/Plugins/FX/Niagara/Source/NiagaraEditor/Private/NiagaraValidationRules.cpp
--- a/Engine/Plugins/FX/Niagara/Source/NiagaraEditor/Private/NiagaraValidationRules.cpp
+++ b/Engine/Plugins/FX/Niagara/Source/NiagaraEditor/Private/NiagaraValidationRules.cpp
@@ -49,7 +49,7 @@ bool NiagaraValidation::HasValidationRules(UNiagaraSystem* NiagaraSystem)
 			for (const TSoftObjectPtr<UNiagaraValidationRuleSet>& ValidationRuleSetPtr : EditorSettings->DefaultValidationRuleSets)
 			{
 				const UNiagaraValidationRuleSet* ValidationRuleSet = ValidationRuleSetPtr.LoadSynchronous();
-				if (ValidationRuleSet != nullptr && ValidationRuleSet->ValidationRules.Num() > 0)
+				if (ValidationRuleSet != nullptr && ValidationRuleSet->HasAnyRules())
 				{
 					return true;
 				}
@@ -58,14 +58,17 @@ bool NiagaraValidation::HasValidationRules(UNiagaraSystem* NiagaraSystem)
 
 		if (UNiagaraEffectType* EffectType = NiagaraSystem->GetEffectType())
 		{
-			if (EffectType->ValidationRules.Num() > 0)
+			for (UNiagaraValidationRule* Rule : EffectType->ValidationRules)
 			{
-				return true;
+				if (Rule != nullptr)
+				{
+					return true;
+				}
 			}
 
 			for (UNiagaraValidationRuleSet* ValidationRuleSet : EffectType->ValidationRuleSets)
 			{
-				if (ValidationRuleSet != nullptr && ValidationRuleSet->ValidationRules.Num() > 0)
+				if (ValidationRuleSet != nullptr && ValidationRuleSet->HasAnyRules())
 				{
 					return true;
 				}
