Remove redundant niagara stat scopes and mark noisy ones as verbose.

#rnx

diff --git a/Engine/Plugins/FX/Niagara/Source/Niagara/Private/DataInterface/NiagaraDataInterfaceDataChannelRead.cpp b/Engine/Plugins/FX/Niagara/Source/Niagara/Private/DataInterface/NiagaraDataInterfaceDataChannelRead.cpp
--- a/Engine/Plugins/FX/Niagara/Source/Niagara/Private/DataInterface/NiagaraDataInterfaceDataChannelRead.cpp
+++ b/Engine/Plugins/FX/Niagara/Source/Niagara/Private/DataInterface/NiagaraDataInterfaceDataChannelRead.cpp
@@ -37,11 +37,11 @@
 
 #define LOCTEXT_NAMESPACE "NiagaraDataInterfaceDataChannelRead"
 
-DECLARE_CYCLE_STAT(TEXT("NDIDataChannelRead Read"), STAT_NDIDataChannelRead_Read, STATGROUP_NiagaraDataChannels);
-DECLARE_CYCLE_STAT(TEXT("NDIDataChannelRead Consume"), STAT_NDIDataChannelRead_Consume, STATGROUP_NiagaraDataChannels);
-DECLARE_CYCLE_STAT(TEXT("NDIDataChannelRead Spawn"), STAT_NDIDataChannelRead_Spawn, STATGROUP_NiagaraDataChannels);
-DECLARE_CYCLE_STAT(TEXT("NDIDataChannelRead Tick"), STAT_NDIDataChannelRead_Tick, STATGROUP_NiagaraDataChannels);
-DECLARE_CYCLE_STAT(TEXT("NDIDataChannelRead PostTick"), STAT_NDIDataChannelRead_PostTick, STATGROUP_NiagaraDataChannels);
+DECLARE_CYCLE_STAT_WITH_FLAGS(TEXT("NDIDataChannelRead Read"), STAT_NDIDataChannelRead_Read, STATGROUP_NiagaraDataChannels, EStatFlags::Verbose);
+DECLARE_CYCLE_STAT_WITH_FLAGS(TEXT("NDIDataChannelRead Consume"), STAT_NDIDataChannelRead_Consume, STATGROUP_NiagaraDataChannels, EStatFlags::Verbose);
+DECLARE_CYCLE_STAT_WITH_FLAGS(TEXT("NDIDataChannelRead Spawn"), STAT_NDIDataChannelRead_Spawn, STATGROUP_NiagaraDataChannels, EStatFlags::Verbose);
+DECLARE_CYCLE_STAT_WITH_FLAGS(TEXT("NDIDataChannelRead Tick"), STAT_NDIDataChannelRead_Tick, STATGROUP_NiagaraDataChannels, EStatFlags::Verbose);
+DECLARE_CYCLE_STAT_WITH_FLAGS(TEXT("NDIDataChannelRead PostTick"), STAT_NDIDataChannelRead_PostTick, STATGROUP_NiagaraDataChannels, EStatFlags::Verbose);
 
 static int GNDCReadForceTG = -1;
 static FAutoConsoleVariableRef CVarNDCReadForceTG(
diff --git a/Engine/Plugins/FX/Niagara/Source/Niagara/Private/NiagaraComponent.cpp b/Engine/Plugins/FX/Niagara/Source/Niagara/Private/NiagaraComponent.cpp
--- a/Engine/Plugins/FX/Niagara/Source/Niagara/Private/NiagaraComponent.cpp
+++ b/Engine/Plugins/FX/Niagara/Source/Niagara/Private/NiagaraComponent.cpp
@@ -40,7 +40,7 @@ DECLARE_CYCLE_STAT(TEXT("Component Tick (GT)"), STAT_NiagaraComponentTick, STATG
 DECLARE_CYCLE_STAT(TEXT("Activate (GT)"), STAT_NiagaraComponentActivate, STATGROUP_Niagara);
 DECLARE_CYCLE_STAT(TEXT("Deactivate (GT)"), STAT_NiagaraComponentDeactivate, STATGROUP_Niagara);
 DECLARE_CYCLE_STAT(TEXT("Send Render Data (GT)"), STAT_NiagaraComponentSendRenderData, STATGROUP_Niagara);
-DECLARE_CYCLE_STAT(TEXT("Set Dynamic Data (RT)"), STAT_NiagaraSetDynamicData, STATGROUP_Niagara);
+DECLARE_CYCLE_STAT_WITH_FLAGS(TEXT("Set Dynamic Data (RT)"), STAT_NiagaraSetDynamicData, STATGROUP_Niagara, EStatFlags::Verbose);
 DECLARE_CYCLE_STAT(TEXT("Get Dynamic Mesh Elements (RT)"), STAT_NiagaraComponentGetDynamicMeshElements, STATGROUP_Niagara);
 
 DEFINE_LOG_CATEGORY(LogNiagara);
diff --git a/Engine/Plugins/FX/Niagara/Source/Niagara/Private/NiagaraEmitterInstanceImpl.cpp b/Engine/Plugins/FX/Niagara/Source/Niagara/Private/NiagaraEmitterInstanceImpl.cpp
--- a/Engine/Plugins/FX/Niagara/Source/Niagara/Private/NiagaraEmitterInstanceImpl.cpp
+++ b/Engine/Plugins/FX/Niagara/Source/Niagara/Private/NiagaraEmitterInstanceImpl.cpp
@@ -21,7 +21,7 @@ DECLARE_DWORD_COUNTER_STAT(TEXT("Num Custom Events"), STAT_NiagaraNumCustomEvent
 //DECLARE_CYCLE_STAT(TEXT("Tick"), STAT_NiagaraTick, STATGROUP_Niagara);
 DECLARE_CYCLE_STAT(TEXT("Emitter Simulate [CNC]"), STAT_NiagaraSimulate, STATGROUP_Niagara);
 DECLARE_CYCLE_STAT(TEXT("Emitter Spawn [CNC]"), STAT_NiagaraSpawn, STATGROUP_Niagara);
-DECLARE_CYCLE_STAT(TEXT("Emitter Post Tick [CNC]"), STAT_NiagaraEmitterPostTick, STATGROUP_Niagara);
+DECLARE_CYCLE_STAT_WITH_FLAGS(TEXT("Emitter Post Tick [CNC]"), STAT_NiagaraEmitterPostTick, STATGROUP_Niagara, EStatFlags::Verbose);
 DECLARE_CYCLE_STAT(TEXT("Emitter Event Handling [CNC]"), STAT_NiagaraEventHandle, STATGROUP_Niagara);
 DECLARE_CYCLE_STAT(TEXT("Emitter Event CopyBuffer [CNC]"), STAT_NiagaraEvent_CopyBuffer, STATGROUP_Niagara);
 DECLARE_CYCLE_STAT(TEXT("Emitter Error Check [CNC]"), STAT_NiagaraEmitterErrorCheck, STATGROUP_Niagara);
diff --git a/Engine/Plugins/FX/Niagara/Source/Niagara/Private/NiagaraRendererMeshes.cpp b/Engine/Plugins/FX/Niagara/Source/Niagara/Private/NiagaraRendererMeshes.cpp
--- a/Engine/Plugins/FX/Niagara/Source/Niagara/Private/NiagaraRendererMeshes.cpp
+++ b/Engine/Plugins/FX/Niagara/Source/Niagara/Private/NiagaraRendererMeshes.cpp
@@ -30,7 +30,7 @@
 #include "SceneRendererInterface.h"
 #include "StaticMeshResources.h"
 
-DECLARE_CYCLE_STAT(TEXT("Generate Mesh Vertex Data [GT]"), STAT_NiagaraGenMeshVertexData, STATGROUP_Niagara);
+DECLARE_CYCLE_STAT_WITH_FLAGS(TEXT("Generate Mesh Vertex Data [GT]"), STAT_NiagaraGenMeshVertexData, STATGROUP_Niagara, EStatFlags::Verbose);
 
 DECLARE_DWORD_COUNTER_STAT(TEXT("NumMeshesRenderer"), STAT_NiagaraNumMeshes, STATGROUP_Niagara);
 DECLARE_DWORD_COUNTER_STAT(TEXT("NumMeshVerts"), STAT_NiagaraNumMeshVerts, STATGROUP_Niagara);
diff --git a/Engine/Plugins/FX/Niagara/Source/Niagara/Private/NiagaraStats.h b/Engine/Plugins/FX/Niagara/Source/Niagara/Private/NiagaraStats.h
--- a/Engine/Plugins/FX/Niagara/Source/Niagara/Private/NiagaraStats.h
+++ b/Engine/Plugins/FX/Niagara/Source/Niagara/Private/NiagaraStats.h
@@ -7,7 +7,7 @@
 
 DECLARE_STATS_GROUP(TEXT("Niagara"), STATGROUP_Niagara, STATCAT_Niagara);
 DECLARE_DWORD_COUNTER_STAT(TEXT("# CPU Particles"), STAT_NiagaraNumParticles, STATGROUP_Niagara);
-DECLARE_CYCLE_STAT(TEXT("Constant Setup [CNC]"), STAT_NiagaraConstants, STATGROUP_Niagara);
+DECLARE_CYCLE_STAT_WITH_FLAGS(TEXT("Constant Setup [CNC]"), STAT_NiagaraConstants, STATGROUP_Niagara, EStatFlags::Verbose);
 DECLARE_CYCLE_STAT(TEXT("Emitter Tick [CNC]"), STAT_NiagaraTick, STATGROUP_Niagara);
 DECLARE_DWORD_COUNTER_STAT(TEXT("NumSystems"), STAT_NiagaraNumSystems, STATGROUP_Niagara);
 DECLARE_MEMORY_STAT(TEXT("Niagara particle data memory"), STAT_NiagaraParticleMemory, STATGROUP_Niagara);
diff --git a/Engine/Plugins/FX/Niagara/Source/Niagara/Private/NiagaraSystemInstance.cpp b/Engine/Plugins/FX/Niagara/Source/Niagara/Private/NiagaraSystemInstance.cpp
--- a/Engine/Plugins/FX/Niagara/Source/Niagara/Private/NiagaraSystemInstance.cpp
+++ b/Engine/Plugins/FX/Niagara/Source/Niagara/Private/NiagaraSystemInstance.cpp
@@ -2259,10 +2259,8 @@ void FNiagaraSystemInstance::ManualTick(float DeltaSeconds, const FGraphEventRef
 
 void FNiagaraSystemInstance::SimCacheTick_GameThread(UNiagaraSimCache* SimCache, float DesiredAge, float DeltaSeconds, const FGraphEventRef& MyCompletionGraphEvent)
 {
-	SCOPE_CYCLE_COUNTER(STAT_NiagaraOverview_GT);
 	SCOPE_CYCLE_COUNTER(STAT_NiagaraSystemInst_ComponentTickGT);
 	SCOPE_CYCLE_COUNTER(STAT_NiagaraSystemInst_TickGT);
-	CSV_SCOPED_TIMING_STAT_EXCLUSIVE(Effects);
 	LLM_SCOPE(ELLMTag::Niagara);
 
 	check(IsInGameThread());
@@ -2457,8 +2455,6 @@ bool FNiagaraSystemInstance::HandleCompletion()
 void FNiagaraSystemInstance::Tick_GameThread(float DeltaSeconds)
 {
 	SCOPE_CYCLE_COUNTER(STAT_NiagaraSystemInst_TickGT);
-	SCOPE_CYCLE_COUNTER(STAT_NiagaraOverview_GT);
-	CSV_SCOPED_TIMING_STAT_EXCLUSIVE(Effects);
 	LLM_SCOPE(ELLMTag::Niagara);
 
 	FNiagaraCrashReporterScope CRScope(this);
@@ -2523,8 +2519,6 @@ void FNiagaraSystemInstance::Tick_GameThread(float DeltaSeconds)
 void FNiagaraSystemInstance::Tick_Concurrent(bool bEnqueueGPUTickIfNeeded)
 {
 	SCOPE_CYCLE_COUNTER(STAT_NiagaraSystemInst_TickCNC);
-	SCOPE_CYCLE_COUNTER(STAT_NiagaraOverview_GT_CNC);
-	CSV_SCOPED_TIMING_STAT_EXCLUSIVE(Effects);
 	LLM_SCOPE(ELLMTag::Niagara);
 	FScopeCycleCounterUObject AdditionalScope(GetSystem(), GET_STATID(STAT_NiagaraOverview_GT_CNC));
 
@@ -2707,7 +2701,6 @@ void FNiagaraSystemInstance::FinalizeTick_GameThread(bool bEnqueueGPUTickIfNeede
 
 	SCOPE_CYCLE_COUNTER(STAT_NiagaraOverview_GT);
 	SCOPE_CYCLE_COUNTER(STAT_NiagaraSystemInst_FinalizeGT);
-	CSV_SCOPED_TIMING_STAT_EXCLUSIVE(Effects);
 	LLM_SCOPE(ELLMTag::Niagara);
 
 	// Ensure concurrent work is complete and clear the finalize ref
diff --git a/Engine/Plugins/FX/Niagara/Source/Niagara/Private/NiagaraSystemSimulation.cpp b/Engine/Plugins/FX/Niagara/Source/Niagara/Private/NiagaraSystemSimulation.cpp
--- a/Engine/Plugins/FX/Niagara/Source/Niagara/Private/NiagaraSystemSimulation.cpp
+++ b/Engine/Plugins/FX/Niagara/Source/Niagara/Private/NiagaraSystemSimulation.cpp
@@ -307,6 +307,7 @@ struct FNiagaraSystemSimulationTickConcurrentTask
 #if NIAGARA_SYSTEMSIMULATION_DEBUGGING
 		NiagaraSystemSimulationLocal::DebugDelayConcurrentTask();
 #endif
+		CSV_SCOPED_TIMING_STAT_EXCLUSIVE(Effects);
 
 		Context.Owner->Tick_Concurrent(Context);
 		CompletionTask->Unlock();
@@ -339,6 +340,7 @@ struct FNiagaraSystemSimulationSpawnConcurrentTask
 #if NIAGARA_SYSTEMSIMULATION_DEBUGGING
 		NiagaraSystemSimulationLocal::DebugDelayConcurrentTask();
 #endif
+		CSV_SCOPED_TIMING_STAT_EXCLUSIVE(Effects);
 
 		Context.BeforeInstancesTickGraphEvents.Add(MyCompletionGraphEvent);
 		Context.Owner->Spawn_Concurrent(Context);
@@ -1123,9 +1125,7 @@ void FNiagaraSystemSimulation::Tick_GameThread_Internal(float DeltaSeconds, cons
 	// Work may not be complete if we back to back tick the GameThread without sending EOF updates
 	WaitForInstancesTickComplete();
 
-	SCOPE_CYCLE_COUNTER(STAT_NiagaraOverview_GT);
 	SCOPE_CYCLE_COUNTER(STAT_NiagaraSystemSim_TickGT);
-	CSV_SCOPED_TIMING_STAT_EXCLUSIVE(Effects);
 	LLM_SCOPE(ELLMTag::Niagara);
 	FScopeCycleCounterUObject AdditionalScope(GetSystem(), GET_STATID(STAT_NiagaraOverview_GT_CNC));
 
@@ -1674,8 +1674,6 @@ void FNiagaraSystemSimulation::WaitForInstancesTickComplete(bool bEnsureComplete
 void FNiagaraSystemSimulation::Tick_Concurrent(FNiagaraSystemSimulationTickContext& Context)
 {
 	SCOPE_CYCLE_COUNTER(STAT_NiagaraSystemSim_TickCNC);
-	SCOPE_CYCLE_COUNTER(STAT_NiagaraOverview_GT_CNC);
-	CSV_SCOPED_TIMING_STAT_EXCLUSIVE(Effects);
 	LLM_SCOPE(ELLMTag::Niagara);
 
 	FScopeCycleCounterUObject AdditionalScope(Context.System, GET_STATID(STAT_NiagaraOverview_GT_CNC));
