PR #11600: added trace complex to niagara

#rnx

diff --git a/Engine/Plugins/FX/Niagara/Source/Niagara/Classes/NiagaraCollision.h b/Engine/Plugins/FX/Niagara/Source/Niagara/Classes/NiagaraCollision.h
--- a/Engine/Plugins/FX/Niagara/Source/Niagara/Classes/NiagaraCollision.h
+++ b/Engine/Plugins/FX/Niagara/Source/Niagara/Classes/NiagaraCollision.h
@@ -90,8 +90,8 @@ public:
 		CurrBuffer = 0;
 	}
 
-	int32 SubmitQuery(FVector StartPos, FVector EndPos, ECollisionChannel TraceChannel);
-	bool PerformQuery(FVector StartPos, FVector EndPos, FNiagaraDICollsionQueryResult &Result, ECollisionChannel TraceChannel);
+	int32 SubmitQuery(FVector StartPos, FVector EndPos, ECollisionChannel TraceChannel, bool bTraceComplex = false);
+	bool PerformQuery(FVector StartPos, FVector EndPos, FNiagaraDICollsionQueryResult &Result, ECollisionChannel TraceChannel, bool bTraceComplex = false);
 	bool GetQueryResult(uint32 TraceID, FNiagaraDICollsionQueryResult &Result);
 
 private:
diff --git a/Engine/Plugins/FX/Niagara/Source/Niagara/Private/NiagaraCollision.cpp b/Engine/Plugins/FX/Niagara/Source/Niagara/Private/NiagaraCollision.cpp
--- a/Engine/Plugins/FX/Niagara/Source/Niagara/Private/NiagaraCollision.cpp
+++ b/Engine/Plugins/FX/Niagara/Source/Niagara/Private/NiagaraCollision.cpp
@@ -107,7 +107,7 @@ void FNiagaraDICollisionQueryBatch::CollectResults()
 	}
 }
 
-int32 FNiagaraDICollisionQueryBatch::SubmitQuery(FVector StartPos, FVector EndPos, ECollisionChannel TraceChannel)
+int32 FNiagaraDICollisionQueryBatch::SubmitQuery(FVector StartPos, FVector EndPos, ECollisionChannel TraceChannel, bool bTraceComplex /*= false*/)
 {
 	if (!GNiagaraCollisionCPUEnabled)
 	{
@@ -120,6 +120,7 @@ int32 FNiagaraDICollisionQueryBatch::SubmitQuery(FVector StartPos, FVector EndPo
 	CollisionQueryParams.bReturnFaceIndex = false;
 	CollisionQueryParams.bReturnPhysicalMaterial = true;
 	CollisionQueryParams.bIgnoreTouches = true;
+	CollisionQueryParams.bTraceComplex = bTraceComplex;
 
 	int32 TraceIdx = INDEX_NONE;
 
@@ -136,7 +137,7 @@ int32 FNiagaraDICollisionQueryBatch::SubmitQuery(FVector StartPos, FVector EndPo
 }
 
 // Work has been done on the collision world side of things to support synchronous queries from multiple threads
-bool FNiagaraDICollisionQueryBatch::PerformQuery(FVector StartPos, FVector EndPos, FNiagaraDICollsionQueryResult &Result, ECollisionChannel TraceChannel)
+bool FNiagaraDICollisionQueryBatch::PerformQuery(FVector StartPos, FVector EndPos, FNiagaraDICollsionQueryResult &Result, ECollisionChannel TraceChannel, bool bTraceComplex /*= false*/)
 {
 	if (!CollisionWorld || !GNiagaraCollisionCPUEnabled)
 	{
@@ -150,6 +151,7 @@ bool FNiagaraDICollisionQueryBatch::PerformQuery(FVector StartPos, FVector EndPo
 	QueryParams.bReturnFaceIndex = false;
 	QueryParams.bReturnPhysicalMaterial = true;
 	QueryParams.bIgnoreTouches = true;
+	QueryParams.bTraceComplex = bTraceComplex;
 	FHitResult TraceResult;
 	bool ValidHit = CollisionWorld->LineTraceSingleByChannel(TraceResult, StartPos, EndPos, TraceChannel, QueryParams);
 	if (ValidHit)
diff --git a/Engine/Plugins/FX/Niagara/Source/Niagara/Private/NiagaraDataInterfaceCollisionQuery.cpp b/Engine/Plugins/FX/Niagara/Source/Niagara/Private/NiagaraDataInterfaceCollisionQuery.cpp
--- a/Engine/Plugins/FX/Niagara/Source/Niagara/Private/NiagaraDataInterfaceCollisionQuery.cpp
+++ b/Engine/Plugins/FX/Niagara/Source/Niagara/Private/NiagaraDataInterfaceCollisionQuery.cpp
@@ -55,6 +55,7 @@ struct FNiagaraCollisionDIFunctionVersion
 		ReturnCollisionMaterialIdx = 3,
 		LargeWorldCoordinates = 4,
 		DeprecatedAsyncGpuTrace = 5,
+		AddedTraceComplex = 6,
 
 		VersionPlusOne,
 		LatestVersion = VersionPlusOne - 1
@@ -360,6 +361,7 @@ void UNiagaraDataInterfaceCollisionQuery::GetFunctionsInternal(TArray<FNiagaraFu
 	{
 		const FText TraceChannelDescription = LOCTEXT("TraceChannelDescription", "The trace channel to collide against. Trace channels can be configured in the project settings.");
 		const FText SkipTraceDescription = LOCTEXT("SkipTraceDescription", "If true then the trace will be skipped completely.\nThis can be used as a performance optimization, as branch nodes in the graph still execute every path.");
+		const FText TraceComplexDescription = LOCTEXT("TraceComplexDescription", "If true then the trace will hit complex collision.");
 		const FText CollisionValidDescription = LOCTEXT("CollisionValidDescription", "Returns true if the trace was not skipped and the trace was blocked by some world geometry.");
 		const FText IsTraceInsideMeshDescription = LOCTEXT("IsTraceInsideMeshDescription", "If true then the trace started in penetration, i.e. with an initial blocking overlap.");
 		const FText CollisionPosWorldDescription = LOCTEXT("CollisionPosWorldDescription", "If the collision is valid, this returns the location of the blocking hit.");
@@ -379,6 +381,7 @@ void UNiagaraDataInterfaceCollisionQuery::GetFunctionsInternal(TArray<FNiagaraFu
 		SigCpuSync.AddInput(FNiagaraVariable(FNiagaraTypeDefinition::GetPositionDef(), TEXT("TraceEndWorld")), TraceEndWorldDescription);
 		SigCpuSync.AddInput(FNiagaraVariable(FNiagaraTypeDefinition(TraceChannelEnum), TEXT("TraceChannel")), TraceChannelDescription);
 		SigCpuSync.AddInput(FNiagaraVariable(FNiagaraTypeDefinition::GetBoolDef(), TEXT("SkipTrace")), SkipTraceDescription);
+		SigCpuSync.AddInput(FNiagaraVariable(FNiagaraTypeDefinition::GetBoolDef(), TEXT("TraceComplex")), TraceComplexDescription);
 		SigCpuSync.AddOutput(FNiagaraVariable(FNiagaraTypeDefinition::GetBoolDef(), TEXT("CollisionValid")), CollisionValidDescription);
 		SigCpuSync.AddOutput(FNiagaraVariable(FNiagaraTypeDefinition::GetBoolDef(), TEXT("IsTraceInsideMesh")), IsTraceInsideMeshDescription);
 		SigCpuSync.AddOutput(FNiagaraVariable(FNiagaraTypeDefinition::GetPositionDef(), TEXT("CollisionPosWorld")), CollisionPosWorldDescription);
@@ -400,6 +403,7 @@ void UNiagaraDataInterfaceCollisionQuery::GetFunctionsInternal(TArray<FNiagaraFu
 		SigCpuAsync.AddInput(FNiagaraVariable(FNiagaraTypeDefinition::GetPositionDef(), TEXT("TraceEndWorld")), TraceEndWorldDescription);
 		SigCpuAsync.AddInput(FNiagaraVariable(FNiagaraTypeDefinition(TraceChannelEnum), TEXT("TraceChannel")), TraceChannelDescription);
 		SigCpuAsync.AddInput(FNiagaraVariable(FNiagaraTypeDefinition::GetBoolDef(), TEXT("SkipTrace")), SkipTraceDescription);
+		SigCpuAsync.AddInput(FNiagaraVariable(FNiagaraTypeDefinition::GetBoolDef(), TEXT("TraceComplex")), TraceComplexDescription);
 		SigCpuAsync.AddOutput(FNiagaraVariable(FNiagaraTypeDefinition::GetIntDef(), TEXT("NextFrameQueryID")), LOCTEXT("TraceAsync_NextFrameQueryIDDescription", "The query ID to save and use as input to this function in the next frame."));
 		SigCpuAsync.AddOutput(FNiagaraVariable(FNiagaraTypeDefinition::GetBoolDef(), TEXT("CollisionValid")), CollisionValidDescription);
 		SigCpuAsync.AddOutput(FNiagaraVariable(FNiagaraTypeDefinition::GetBoolDef(), TEXT("IsTraceInsideMesh")), IsTraceInsideMeshDescription);
@@ -555,6 +559,7 @@ void UNiagaraDataInterfaceCollisionQuery::PerformQuerySyncCPU(FVectorVMExternalF
 	FNDIInputParam<FNiagaraPosition> EndPosParam(Context);
 	FNDIInputParam<ECollisionChannel> TraceChannelParam(Context);
 	FNDIInputParam<FNiagaraBool> IsSkipTrace(Context);
+	FNDIInputParam<FNiagaraBool> IsTraceComplex(Context);
 	
 	FNDIOutputParam<FNiagaraBool> OutQueryValid(Context);
 	FNDIOutputParam<FNiagaraBool> OutInsideMesh(Context);
@@ -573,11 +578,12 @@ void UNiagaraDataInterfaceCollisionQuery::PerformQuerySyncCPU(FVectorVMExternalF
 		FNiagaraPosition TraceEnd = EndPosParam.GetAndAdvance();
 		ECollisionChannel TraceChannel = TraceChannelParam.GetAndAdvance();
 		bool Skip = IsSkipTrace.GetAndAdvance();
+		bool TraceComplex = IsTraceComplex.GetAndAdvance();
 		ensure(!TraceStart.ContainsNaN());
 		ensure(!TraceEnd.ContainsNaN());
 		FNiagaraDICollsionQueryResult Res;
 
-		if (!Skip && InstanceData->CollisionBatch.PerformQuery(LWCConverter.ConvertSimulationPositionToWorld(TraceStart), LWCConverter.ConvertSimulationPositionToWorld(TraceEnd), Res, TraceChannel))
+		if (!Skip && InstanceData->CollisionBatch.PerformQuery(LWCConverter.ConvertSimulationPositionToWorld(TraceStart), LWCConverter.ConvertSimulationPositionToWorld(TraceEnd), Res, TraceChannel, TraceComplex))
 		{
 			OutQueryValid.SetAndAdvance(true);
 			OutInsideMesh.SetAndAdvance(Res.IsInsideMesh);
@@ -609,6 +615,7 @@ void UNiagaraDataInterfaceCollisionQuery::PerformQueryAsyncCPU(FVectorVMExternal
 	FNDIInputParam<FNiagaraPosition> EndPosParam(Context);
 	FNDIInputParam<ECollisionChannel> TraceChannelParam(Context);
 	FNDIInputParam<FNiagaraBool> IsSkipTrace(Context);
+	FNDIInputParam<FNiagaraBool> IsTraceComplex(Context);
 	
 	FNDIOutputParam<int32> OutQueryID(Context);
 	FNDIOutputParam<FNiagaraBool> OutQueryValid(Context);
@@ -628,10 +635,11 @@ void UNiagaraDataInterfaceCollisionQuery::PerformQueryAsyncCPU(FVectorVMExternal
 		FNiagaraPosition TraceEnd = EndPosParam.GetAndAdvance();
 		ECollisionChannel TraceChannel = TraceChannelParam.GetAndAdvance();
 		bool Skip = IsSkipTrace.GetAndAdvance();
+		bool TraceComplex = IsTraceComplex.GetAndAdvance();
 		ensure(!TraceStart.ContainsNaN());
 		ensure(!TraceEnd.ContainsNaN());
 
-		int QueryID = Skip ? INDEX_NONE : InstanceData->CollisionBatch.SubmitQuery(LWCConverter.ConvertSimulationPositionToWorld(TraceStart), LWCConverter.ConvertSimulationPositionToWorld(TraceEnd), TraceChannel);
+		int QueryID = Skip ? INDEX_NONE : InstanceData->CollisionBatch.SubmitQuery(LWCConverter.ConvertSimulationPositionToWorld(TraceStart), LWCConverter.ConvertSimulationPositionToWorld(TraceEnd), TraceChannel, TraceComplex);
 		OutQueryID.SetAndAdvance(QueryID + 1);
 
 		// try to retrieve a query with the supplied query ID
