- Code to track down modifications of scalability managers when a NS is activating
- Theory is that a component renderer is activating a NS but need to confirm

#rnx

diff --git a/Engine/Plugins/FX/Niagara/Source/Niagara/Private/NiagaraWorldManager.cpp b/Engine/Plugins/FX/Niagara/Source/Niagara/Private/NiagaraWorldManager.cpp
--- a/Engine/Plugins/FX/Niagara/Source/Niagara/Private/NiagaraWorldManager.cpp
+++ b/Engine/Plugins/FX/Niagara/Source/Niagara/Private/NiagaraWorldManager.cpp
@@ -1643,6 +1643,10 @@ void FNiagaraWorldManager::UpdateScalabilityManagers(float DeltaSeconds, bool bN
 	SCOPE_CYCLE_COUNTER(STAT_UpdateScalabilityManagers);
 	CSV_SCOPED_TIMING_STAT(Particles, CoreSystems_NiagaraUpdateScalabilityManagers);
 
+#if DO_CHECK
+	TGuardValue<bool> GuardScalabilityManagersLocked(bScalabilityManagersLocked, true);
+#endif
+
 	for (auto& Pair : ScalabilityManagers)
 	{
 		FNiagaraScalabilityManager& ScalabilityMan = Pair.Value;
@@ -1674,6 +1678,9 @@ void FNiagaraWorldManager::RegisterWithScalabilityManager(UNiagaraComponent* Com
 
 		if (!ScalabilityManager)
 		{
+		#if DO_CHECK
+			ensureMsgf(bScalabilityManagersLocked == false, TEXT("ScalabilityManagers modified while iterating, by Component (%s) System (%s)"), *GetNameSafe(Component), *GetNameSafe(Component ? Component->GetAsset() : nullptr));
+		#endif
 			ScalabilityManager = &ScalabilityManagers.Add(EffectType);
 			ScalabilityManager->EffectType = EffectType;
 		}
diff --git a/Engine/Plugins/FX/Niagara/Source/Niagara/Public/NiagaraWorldManager.h b/Engine/Plugins/FX/Niagara/Source/Niagara/Public/NiagaraWorldManager.h
--- a/Engine/Plugins/FX/Niagara/Source/Niagara/Public/NiagaraWorldManager.h
+++ b/Engine/Plugins/FX/Niagara/Source/Niagara/Public/NiagaraWorldManager.h
@@ -397,6 +397,10 @@ private:
 	/** Instances that have been queued for deletion this frame, serviced in PostActorTick */
 	TArray<FNiagaraSystemInstancePtr> DeferredDeletionQueue;
 
+#if DO_CHECK
+	bool bScalabilityManagersLocked = false;
+#endif
+
 	UPROPERTY(transient)
 	TMap<TObjectPtr<UNiagaraEffectType>, FNiagaraScalabilityManager> ScalabilityManagers;
 
