Changed Niagara Parameter Reference counts to count all references of all emitters & system itself when nothing is selected.
Previously, if nothing was selected we'd show the system node's references instead

diff --git a/Engine/Plugins/FX/Niagara/Source/NiagaraEditor/Private/ViewModels/NiagaraParameterPanelViewModel.cpp b/Engine/Plugins/FX/Niagara/Source/NiagaraEditor/Private/ViewModels/NiagaraParameterPanelViewModel.cpp
--- a/Engine/Plugins/FX/Niagara/Source/NiagaraEditor/Private/ViewModels/NiagaraParameterPanelViewModel.cpp
+++ b/Engine/Plugins/FX/Niagara/Source/NiagaraEditor/Private/ViewModels/NiagaraParameterPanelViewModel.cpp
@@ -46,6 +46,7 @@
 #include "Framework/MultiBox/MultiBoxBuilder.h"
 #include "HAL/PlatformApplicationMisc.h"
 #include "ScopedTransaction.h"
+#include "ViewModels/NiagaraOverviewGraphViewModel.h"
 #include "Widgets/SNullWidget.h"
 #define LOCTEXT_NAMESPACE "NiagaraParameterPanelViewModel"
 
@@ -846,9 +847,17 @@ void FNiagaraSystemToolkitParameterPanelViewModel::Cleanup()
 		}
 	}
 
-	if (SystemViewModel->GetEditMode() == ENiagaraSystemViewModelEditMode::SystemAsset && SystemGraphSelectionViewModelWeak.IsValid())
+	if (SystemViewModel->GetEditMode() == ENiagaraSystemViewModelEditMode::SystemAsset)
 	{
-		SystemGraphSelectionViewModelWeak.Pin()->GetOnSelectedEmitterScriptGraphsRefreshedDelegate().RemoveAll(this);
+		if(SystemGraphSelectionViewModelWeak.IsValid())
+		{
+			SystemGraphSelectionViewModelWeak.Pin()->GetOnSelectedEmitterScriptGraphsRefreshedDelegate().RemoveAll(this);
+		}
+
+		if(SystemViewModel->GetSelectionViewModel())
+		{
+			SystemViewModel->GetSelectionViewModel()->OnSystemIsSelectedChanged().RemoveAll(this);
+		}
 	}
 
 	if (SystemViewModel->GetEditMode() == ENiagaraSystemViewModelEditMode::SystemAsset)
@@ -944,9 +953,14 @@ void FNiagaraSystemToolkitParameterPanelViewModel::Init(const FSystemToolkitUICo
 	}
 
 	// Bind OnChanged() callback for system edit mode active emitter handle selection changing.
-	if (SystemViewModel->GetEditMode() == ENiagaraSystemViewModelEditMode::SystemAsset && ensureMsgf(SystemGraphSelectionViewModelWeak.IsValid(), TEXT("SystemGraphSelectionViewModel was null for System edit mode!")))
+	if (SystemViewModel->GetEditMode() == ENiagaraSystemViewModelEditMode::SystemAsset)
 	{
-		SystemGraphSelectionViewModelWeak.Pin()->GetOnSelectedEmitterScriptGraphsRefreshedDelegate().AddSP(this, &FNiagaraSystemToolkitParameterPanelViewModel::Refresh);
+		if(ensureMsgf(SystemGraphSelectionViewModelWeak.IsValid(), TEXT("SystemGraphSelectionViewModel was null for System edit mode!")))
+		{
+			SystemGraphSelectionViewModelWeak.Pin()->GetOnSelectedEmitterScriptGraphsRefreshedDelegate().AddSP(this, &FNiagaraSystemToolkitParameterPanelViewModel::RefreshNextTick);
+		}
+		
+		SystemViewModel->GetSelectionViewModel()->OnSystemIsSelectedChanged().AddSP(this, &FNiagaraSystemToolkitParameterPanelViewModel::Refresh);
 	}
 
 	// Bind OnChanged() for System Exposed Parameters (User Parameters).
@@ -1881,11 +1895,21 @@ TArray<FNiagaraParameterPanelItem> FNiagaraSystemToolkitParameterPanelViewModel:
 	TMap<FNiagaraVariable, TObjectPtr<UNiagaraScriptVariable>> ParameterToScriptVariableMap;
 	TArray<UNiagaraGraph*> EditableGraphs = GetEditableGraphsConst();
 	TArray<UNiagaraGraph*> AllGraphs = GetAllGraphsConst();
-
+	
+	TArray<UNiagaraGraph*> GraphsToCheck;
+	if(SystemViewModel->GetSelectionViewModel()->GetSystemIsSelected() == false && SystemViewModel->GetSelectionViewModel()->GetSelectedEmitterHandleIds().Num() == 0)
+	{	
+		GraphsToCheck = AllGraphs;
+	}
+	else
+	{
+		GraphsToCheck = EditableGraphs;
+	}
+	
 	// Collect all metadata to be packaged with the FNiagaraParameterPanelItems.
-	for (const UNiagaraGraph* EditableGraph : EditableGraphs)
+	for (const UNiagaraGraph* GraphToCheck : GraphsToCheck)
 	{
-		ParameterToScriptVariableMap.Append(EditableGraph->GetAllMetaData());
+		ParameterToScriptVariableMap.Append(GraphToCheck->GetAllMetaData());
 	}
 	ParameterToScriptVariableMap.Append(TransientParameterToScriptVarMap);
 
@@ -1938,11 +1962,11 @@ TArray<FNiagaraParameterPanelItem> FNiagaraSystemToolkitParameterPanelViewModel:
 		// Collect parameters for all emitters.
 		TArray<FNiagaraVariable> VisitedInvalidParameters;
 
-		// -For each selected graph perform a parameter map history traversal and record each unique visited parameter.
-		for (const UNiagaraGraph* EditableGraph : EditableGraphs)
+		// -For each applicable graph perform a parameter map history traversal and record each unique visited parameter.
+		for (const UNiagaraGraph* GraphToCheck : GraphsToCheck)
 		{
 			TArray<UNiagaraNodeOutput*> OutputNodes;
-			EditableGraph->GetNodesOfClass<UNiagaraNodeOutput>(OutputNodes);
+			GraphToCheck->GetNodesOfClass<UNiagaraNodeOutput>(OutputNodes);
 			for (UNiagaraNodeOutput* OutputNode : OutputNodes)
 			{
 				UNiagaraNode* NodeToTraverse = OutputNode;
@@ -1965,7 +1989,7 @@ TArray<FNiagaraParameterPanelItem> FNiagaraSystemToolkitParameterPanelViewModel:
 
 				bool bIgnoreDisabled = true;
 				FNiagaraParameterMapHistoryBuilder Builder;
-				FVersionedNiagaraEmitter GraphOwningEmitter = EditableGraph->GetOwningEmitter();
+				FVersionedNiagaraEmitter GraphOwningEmitter = GraphToCheck->GetOwningEmitter();
 				FCompileConstantResolver ConstantResolver = GraphOwningEmitter.Emitter != nullptr
 					? FCompileConstantResolver(GraphOwningEmitter, OutputNode->GetUsage())
 					: FCompileConstantResolver();
@@ -2015,10 +2039,10 @@ TArray<FNiagaraParameterPanelItem> FNiagaraSystemToolkitParameterPanelViewModel:
 					ParameterToScriptVariableMap.Append(VisitedExternalGraph->GetAllMetaData());
 				}
 
-				auto AddReadReference = [EditableGraph](const FNiagaraParameterMapHistory::FReadHistory& History, FNiagaraParameterPanelItem& Item)
+				auto AddReadReference = [GraphToCheck](const FNiagaraParameterMapHistory::FReadHistory& History, FNiagaraParameterPanelItem& Item)
 				{
 					FNiagaraParameterReferencePath Source;
-					Source.SourceGraph = EditableGraph;
+					Source.SourceGraph = GraphToCheck;
 					Source.ModuleName = History.ReadPin.ModuleName;
 					if (Source.ModuleName == NAME_None && History.ReadPin.Pin->LinkedTo.Num() > 0)
 					{
@@ -2060,7 +2084,7 @@ TArray<FNiagaraParameterPanelItem> FNiagaraSystemToolkitParameterPanelViewModel:
 					}
 					Item.AddToReadCount(Source);
 				};
-				auto AddWriteReference = [EditableGraph](const FNiagaraParameterMapHistory::FModuleScopedPin& History, FNiagaraParameterPanelItem& Item)
+				auto AddWriteReference = [GraphToCheck](const FNiagaraParameterMapHistory::FModuleScopedPin& History, FNiagaraParameterPanelItem& Item)
 				{
 					if (Cast<UNiagaraNodeParameterMapGet>(History.Pin->GetOwningNode()))
 					{
@@ -2068,7 +2092,7 @@ TArray<FNiagaraParameterPanelItem> FNiagaraSystemToolkitParameterPanelViewModel:
 						return;
 					}
 					FNiagaraParameterReferencePath Source;
-					Source.SourceGraph = EditableGraph;
+					Source.SourceGraph = GraphToCheck;
 					Source.ModuleName = History.ModuleName;
 					Item.AddToWriteCount(Source);
 				};
diff --git a/Engine/Plugins/FX/Niagara/Source/NiagaraEditor/Private/ViewModels/NiagaraSystemSelectionViewModel.cpp b/Engine/Plugins/FX/Niagara/Source/NiagaraEditor/Private/ViewModels/NiagaraSystemSelectionViewModel.cpp
--- a/Engine/Plugins/FX/Niagara/Source/NiagaraEditor/Private/ViewModels/NiagaraSystemSelectionViewModel.cpp
+++ b/Engine/Plugins/FX/Niagara/Source/NiagaraEditor/Private/ViewModels/NiagaraSystemSelectionViewModel.cpp
@@ -457,6 +457,7 @@ void UNiagaraSystemSelectionViewModel::RemoveEntryFromSelectionInternal(UNiagara
 
 void UNiagaraSystemSelectionViewModel::UpdateExternalSelectionState()
 {
+	bool bWasSystemSelected = bSystemIsSelected;
 	bSystemIsSelected = false;
 	TArray<FGuid> OldSelectedEmitterHandleIds = SelectedEmitterHandleIds;
 	SelectedEmitterHandleIds.Empty();
@@ -488,6 +489,7 @@ void UNiagaraSystemSelectionViewModel::UpdateExternalSelectionState()
 
 	// If there were ids left over, or we added more than we removed than the id selection has changed.
 	bool bEmitterHandleIdSelectionChanged = OldSelectedEmitterHandleIds.Num() != 0 || SelectedEmitterHandleIds.Num() != NumRemoved;
+	bool bSystemSelectionChanged = bWasSystemSelected != bSystemIsSelected;
 
 	// Check the new selection to make sure that there are no entries which are owned by other entries and if there are
 	// remove them.
@@ -521,4 +523,8 @@ void UNiagaraSystemSelectionViewModel::UpdateExternalSelectionState()
 	{
 		OnEmitterHandleIdSelectionChangedDelegate.Broadcast();
 	}
+	if(bSystemSelectionChanged)
+	{
+		OnSystemIsSelectedChangedDelegate.Broadcast();
+	}
 }
diff --git a/Engine/Plugins/FX/Niagara/Source/NiagaraEditor/Private/ViewModels/Stack/NiagaraStackFunctionInput.cpp b/Engine/Plugins/FX/Niagara/Source/NiagaraEditor/Private/ViewModels/Stack/NiagaraStackFunctionInput.cpp
--- a/Engine/Plugins/FX/Niagara/Source/NiagaraEditor/Private/ViewModels/Stack/NiagaraStackFunctionInput.cpp
+++ b/Engine/Plugins/FX/Niagara/Source/NiagaraEditor/Private/ViewModels/Stack/NiagaraStackFunctionInput.cpp
@@ -3291,6 +3291,21 @@ const UNiagaraStackFunctionInput::FCollectedUsageData& UNiagaraStackFunctionInpu
 			CachedCollectedUsageData = FCollectedUsageData();
 			FNiagaraVariableBase Var(InputType, AliasedInputParameterHandle.GetParameterHandleString());
 			CachedCollectedUsageData.GetValue().bHasReferencedParameterRead = ParamVM->IsVariableSelected(Var) || bFoundOverride;
+
+			TArray<UNiagaraStackFunctionInput*> ChildrenInputs;
+			GetUnfilteredChildrenOfType<UNiagaraStackFunctionInput>(ChildrenInputs);
+			for(UNiagaraStackFunctionInput* ChildInput : ChildrenInputs)
+			{
+				if(ChildInput->GetCollectedUsageData().bHasReferencedParameterRead)
+				{
+					CachedCollectedUsageData.GetValue().bHasReferencedParameterRead = true;
+				}
+
+				if(ChildInput->GetCollectedUsageData().bHasReferencedParameterWrite)
+				{
+					CachedCollectedUsageData.GetValue().bHasReferencedParameterWrite = true;
+				}
+			}
 		}
 		else
 		{
diff --git a/Engine/Plugins/FX/Niagara/Source/NiagaraEditor/Private/ViewModels/Stack/NiagaraStackModuleItem.cpp b/Engine/Plugins/FX/Niagara/Source/NiagaraEditor/Private/ViewModels/Stack/NiagaraStackModuleItem.cpp
--- a/Engine/Plugins/FX/Niagara/Source/NiagaraEditor/Private/ViewModels/Stack/NiagaraStackModuleItem.cpp
+++ b/Engine/Plugins/FX/Niagara/Source/NiagaraEditor/Private/ViewModels/Stack/NiagaraStackModuleItem.cpp
@@ -254,7 +254,17 @@ const UNiagaraStackModuleItem::FCollectedUsageData& UNiagaraStackModuleItem::Get
 				CachedCollectedUsageData.GetValue().bHasReferencedParameterWrite = true;
 		}
 
-
+		if(InputRoot)
+		{
+			if(InputRoot->GetCollectedUsageData().bHasReferencedParameterRead)
+			{
+				CachedCollectedUsageData.GetValue().bHasReferencedParameterRead = true;
+			}
+			if (InputRoot->GetCollectedUsageData().bHasReferencedParameterWrite)
+			{
+				CachedCollectedUsageData.GetValue().bHasReferencedParameterWrite = true;
+			}
+		}
 
 		if (OutputCollection)
 		{
