- On pool reuse, update solo mode & clear the component tick

diff --git a/Engine/Plugins/FX/Niagara/Source/Niagara/Private/NiagaraComponent.cpp b/Engine/Plugins/FX/Niagara/Source/Niagara/Private/NiagaraComponent.cpp
--- a/Engine/Plugins/FX/Niagara/Source/Niagara/Private/NiagaraComponent.cpp
+++ b/Engine/Plugins/FX/Niagara/Source/Niagara/Private/NiagaraComponent.cpp
@@ -1901,8 +1901,11 @@ void UNiagaraComponent::OnPooledReuse(UWorld* NewWorld)
 	SystemFixedBounds.Init();
 	EmitterOverrideInfos.Empty();
 
+	PrimaryComponentTick.SetTickFunctionEnable(false);
+
 	if (SystemInstanceController != nullptr)
 	{
+		UpdateInstanceSoloMode();
 		SystemInstanceController->OnPooledReuse(*NewWorld);
 		SystemInstanceController->SetLastRenderTime(GetLastRenderTime());
 	}
