- Add missing transition when copying a texture 2D to the cache

diff --git a/Engine/Plugins/FX/Niagara/Source/Niagara/Private/NiagaraDataInterfaceRenderTarget2D.cpp b/Engine/Plugins/FX/Niagara/Source/Niagara/Private/NiagaraDataInterfaceRenderTarget2D.cpp
--- a/Engine/Plugins/FX/Niagara/Source/Niagara/Private/NiagaraDataInterfaceRenderTarget2D.cpp
+++ b/Engine/Plugins/FX/Niagara/Source/Niagara/Private/NiagaraDataInterfaceRenderTarget2D.cpp
@@ -604,6 +604,8 @@ NIAGARA_API bool UNiagaraDataInterfaceRenderTarget2D::SimCacheWriteFrame(UObject
 			{
 				if (const FRenderTarget2DRWInstanceData_RenderThread* InstanceData_RT = RT_Proxy->SystemInstancesToProxyData_RT.Find(RT_InstanceID))
 				{
+					RHICmdList.Transition(FRHITransitionInfo(InstanceData_RT->TextureRHI, ERHIAccess::SRVMask, ERHIAccess::CopySrc));
+
 					FRHIGPUTextureReadback RenderTargetReadback("ReadRTTexture");
 
 					RenderTargetReadback.EnqueueCopy(RHICmdList, InstanceData_RT->TextureRHI, FIntVector(0, 0, 0), 0, FIntVector(RT_VolumeResolution.X, RT_VolumeResolution.Y, 1));
@@ -634,6 +636,8 @@ NIAGARA_API bool UNiagaraDataInterfaceRenderTarget2D::SimCacheWriteFrame(UObject
 										
 					//Unlock the staging texture
 					RenderTargetReadback.Unlock();
+
+					RHICmdList.Transition(FRHITransitionInfo(InstanceData_RT->TextureRHI, ERHIAccess::CopySrc, ERHIAccess::SRVMask));
 				}
 			}
 		);
