- Add constant value option for static bool parameter bindings
- Change action menu to use SNiagaraParameter widget

diff --git a/Engine/Plugins/FX/Niagara/Source/Niagara/Private/NiagaraRendererProperties.cpp b/Engine/Plugins/FX/Niagara/Source/Niagara/Private/NiagaraRendererProperties.cpp
--- a/Engine/Plugins/FX/Niagara/Source/Niagara/Private/NiagaraRendererProperties.cpp
+++ b/Engine/Plugins/FX/Niagara/Source/Niagara/Private/NiagaraRendererProperties.cpp
@@ -444,68 +444,81 @@ bool UNiagaraRendererProperties::BuildMaterialStaticParameterSet(const FNiagaraR
 	}
 
 	bool bModified = false;
-	for (const FNiagaraRendererMaterialStaticBoolParameter& ParameterBinding : MaterialParameters.StaticBoolParameters)
-	{
-		NiagaraSystem->ForEachScript(
-			[&](UNiagaraScript* NiagaraScript)
+
+	auto SetMaterialStaticParameter =
+		[&](FName ParameterName, bool bParameterValue)
+		{
+			for (const FMaterialParameterInfo& ParameterInfo : AllStaticSwitchParameterInfos)
 			{
-				for (const FNiagaraVariable& StaticVariable : NiagaraScript->GetVMExecutableData().StaticVariablesWritten)
+				if (ParameterInfo.Name != ParameterName)
 				{
-					if (StaticVariable.GetType() != FNiagaraTypeDefinition::GetBoolDef().ToStaticDef())
+					continue;
+				}
+
+				FStaticSwitchParameter* StaticParameter = StaticParameterSet.StaticSwitchParameters.FindByPredicate(
+					[ParameterInfo](const FStaticSwitchParameter& StaticParameter)
 					{
-						continue;
+						return StaticParameter.ParameterInfo == ParameterInfo;
 					}
+				);
 
-					FNiagaraVariableBase ResolvedStaticVariable = StaticVariable;
-					if (NiagaraEmitter)
+				if (StaticParameter == nullptr)
+				{
+					FGuid ParameterGuid;
+					bool bDefaultValue = false;
+					if (Material->GetStaticSwitchParameterDefaultValue(ParameterInfo, bDefaultValue, ParameterGuid))
 					{
-						ResolvedStaticVariable.ReplaceRootNamespace(NiagaraEmitter->GetUniqueEmitterName(), FNiagaraConstants::EmitterNamespaceString);
+						if (bDefaultValue != bParameterValue)
+						{
+							StaticParameterSet.SetParameterValue(ParameterInfo, FMaterialParameterValue(bParameterValue), EMaterialSetParameterValueFlags::None);
+							bModified = true;
+						}
 					}
-					if (ResolvedStaticVariable.GetName() != ParameterBinding.StaticVariableName)
+				}
+				else
+				{
+					if (StaticParameter && StaticParameter->Value != bParameterValue)
 					{
-						continue;
+						StaticParameter->Value = bParameterValue;
+						bModified = true;
 					}
+				}
+			}
+		};
 
-					for (const FMaterialParameterInfo& ParameterInfo : AllStaticSwitchParameterInfos)
+	for (const FNiagaraRendererMaterialStaticBoolParameter& ParameterBinding : MaterialParameters.StaticBoolParameters)
+	{
+		if (ParameterBinding.StaticValue.IsSet())
+		{
+			SetMaterialStaticParameter(ParameterBinding.MaterialParameterName, ParameterBinding.StaticValue.GetValue());
+		}
+		else
+		{
+			NiagaraSystem->ForEachScript(
+				[&](UNiagaraScript* NiagaraScript)
+				{
+					for (const FNiagaraVariable& StaticVariable : NiagaraScript->GetVMExecutableData().StaticVariablesWritten)
 					{
-						if (ParameterInfo.Name != ParameterBinding.MaterialParameterName)
+						if (StaticVariable.GetType() != FNiagaraTypeDefinition::GetBoolDef().ToStaticDef())
 						{
 							continue;
 						}
 
-						FStaticSwitchParameter* StaticParameter = StaticParameterSet.StaticSwitchParameters.FindByPredicate(
-							[ParameterInfo](const FStaticSwitchParameter& StaticParameter)
-							{
-								return StaticParameter.ParameterInfo == ParameterInfo;
-							}
-						);
-
-						const bool bNewValue = StaticVariable.GetValue<bool>();
-						if (StaticParameter == nullptr)
+						FNiagaraVariableBase ResolvedStaticVariable = StaticVariable;
+						if (NiagaraEmitter)
 						{
-							FGuid ParameterGuid;
-							bool bDefaultValue = false;
-							if (Material->GetStaticSwitchParameterDefaultValue(ParameterInfo, bDefaultValue, ParameterGuid))
-							{
-								if (bDefaultValue != bNewValue)
-								{
-									StaticParameterSet.SetParameterValue(ParameterInfo, FMaterialParameterValue(bNewValue), EMaterialSetParameterValueFlags::None);
-									bModified = true;
-								}
-							}
+							ResolvedStaticVariable.ReplaceRootNamespace(NiagaraEmitter->GetUniqueEmitterName(), FNiagaraConstants::EmitterNamespaceString);
 						}
-						else
+						if (ResolvedStaticVariable.GetName() != ParameterBinding.StaticVariableName)
 						{
-							if (StaticParameter && StaticParameter->Value != bNewValue)
-							{
-								StaticParameter->Value = bNewValue;
-								bModified = true;
-							}
+							continue;
 						}
+
+						SetMaterialStaticParameter(ParameterBinding.MaterialParameterName, StaticVariable.GetValue<bool>());
 					}
 				}
-			}
-		);
+			);
+		}
 	}
 
 	return bModified;
diff --git a/Engine/Plugins/FX/Niagara/Source/Niagara/Public/NiagaraRendererProperties.h b/Engine/Plugins/FX/Niagara/Source/Niagara/Public/NiagaraRendererProperties.h
--- a/Engine/Plugins/FX/Niagara/Source/Niagara/Public/NiagaraRendererProperties.h
+++ b/Engine/Plugins/FX/Niagara/Source/Niagara/Public/NiagaraRendererProperties.h
@@ -234,6 +234,9 @@ struct FNiagaraRendererMaterialStaticBoolParameter
 
 	UPROPERTY(EditAnywhere, Category = "Material")
 	FName StaticVariableName;
+
+	UPROPERTY()
+	TOptional<bool> StaticValue;
 };
 
 
diff --git a/Engine/Plugins/FX/Niagara/Source/NiagaraEditor/Private/Customizations/NiagaraTypeCustomizations.cpp b/Engine/Plugins/FX/Niagara/Source/NiagaraEditor/Private/Customizations/NiagaraTypeCustomizations.cpp
--- a/Engine/Plugins/FX/Niagara/Source/NiagaraEditor/Private/Customizations/NiagaraTypeCustomizations.cpp
+++ b/Engine/Plugins/FX/Niagara/Source/NiagaraEditor/Private/Customizations/NiagaraTypeCustomizations.cpp
@@ -24,6 +24,7 @@
 #include "Widgets/Input/SCheckBox.h"
 #include "Widgets/Input/SComboButton.h"
 #include "Widgets/Input/SEditableTextBox.h"
+#include "Widgets/Layout/SWidgetSwitcher.h"
 #include "Widgets/Text/STextBlock.h"
 
 #include "NiagaraConstants.h"
@@ -2389,31 +2390,51 @@ void FNiagaraRendererMaterialTextureParameterCustomization::GetMaterialParameter
 
 bool FNiagaraRendererMaterialStaticBoolParameterCustomization::CustomizeChildProperty(class IDetailChildrenBuilder& ChildBuilder, TSharedPtr<IPropertyHandle> PropertyHandle)
 {
-	if (PropertyHandle->GetProperty()->GetFName() == GET_MEMBER_NAME_CHECKED(FNiagaraRendererMaterialStaticBoolParameter, StaticVariableName))
+	if (PropertyHandle->GetProperty()->GetFName() != GET_MEMBER_NAME_CHECKED(FNiagaraRendererMaterialStaticBoolParameter, StaticVariableName))
 	{
-		ChildBuilder.AddCustomRow(PropertyHandle->GetPropertyDisplayName())
-		.NameWidget
-		[
-			PropertyHandle->CreatePropertyNameWidget()
-		]
-		.ValueWidget
+		return false;
+	}
+
+	StructPropertyHandle = PropertyHandle ? PropertyHandle->GetParentHandle() : nullptr;
+	FStructProperty* StructProperty = StructPropertyHandle ? CastField<FStructProperty>(StructPropertyHandle->GetProperty()) : nullptr;
+	if (StructProperty == nullptr || StructProperty->Struct != FNiagaraRendererMaterialStaticBoolParameter::StaticStruct())
+	{
+		StructPropertyHandle = nullptr;
+		return false;
+	}
+
+	ChildBuilder.AddCustomRow(PropertyHandle->GetPropertyDisplayName())
+	.NameWidget
+	[
+		PropertyHandle->CreatePropertyNameWidget()
+	]
+	.ValueWidget
+	[
+		SNew(SComboButton)
+		.OnGetMenuContent(this, &FNiagaraRendererMaterialStaticBoolParameterCustomization::OnGetStaticVariablelBindingNameMenuContent)
+		.ContentPadding(1)
+		.ButtonStyle(FAppStyle::Get(), "PropertyEditor.AssetComboStyle")
+		.ForegroundColor(FAppStyle::GetColor("PropertyEditor.AssetName.ColorAndOpacity"))
+		.ButtonContent()
 		[
-			SNew(SComboButton)
-			.OnGetMenuContent(this, &FNiagaraRendererMaterialStaticBoolParameterCustomization::OnGetStaticVariablelBindingNameMenuContent, PropertyHandle)
-			.ContentPadding(1)
-			.ButtonStyle(FAppStyle::Get(), "PropertyEditor.AssetComboStyle")
-			.ForegroundColor(FAppStyle::GetColor("PropertyEditor.AssetName.ColorAndOpacity"))
-			.ButtonContent()
+			SNew(SWidgetSwitcher)
+			.WidgetIndex(this, &FNiagaraRendererMaterialStaticBoolParameterCustomization::GetWidgetIndex)
+			+SWidgetSwitcher::Slot()
 			[
 				SNew(SNiagaraParameterName)
 				.ParameterName(this, &FNiagaraRendererMaterialParameterCustomization::GetBindingName, PropertyHandle)
 				.IsReadOnly(true)
 			]
-		];
+			+SWidgetSwitcher::Slot()
+			[
+				SNew(STextBlock)
+				.Font(IDetailLayoutBuilder::GetDetailFont())
+				.Text(this, &FNiagaraRendererMaterialStaticBoolParameterCustomization::GetCurrentConstantValue)
+			]
+		]
+	];
 
-		return true;
-	}
-	return false;
+	return true;
 }
 
 void FNiagaraRendererMaterialStaticBoolParameterCustomization::GetMaterialParameterInfos(UMaterialInterface* Material, TArray<FMaterialParameterInfo>& OutMaterialParameterInfos) const
@@ -2426,13 +2447,28 @@ void FNiagaraRendererMaterialStaticBoolParameterCustomization::GetMaterialParame
 	}
 }
 
-TSharedRef<SWidget> FNiagaraRendererMaterialStaticBoolParameterCustomization::OnGetStaticVariablelBindingNameMenuContent(TSharedPtr<IPropertyHandle> PropertyHandle) const
+TSharedRef<SWidget> FNiagaraRendererMaterialStaticBoolParameterCustomization::OnGetStaticVariablelBindingNameMenuContent() const
 {
 	FMenuBuilder MenuBuilder(true, NULL);
 
 	UNiagaraRendererProperties* RenderProperties = WeakRenderProperties.Get();
-	if (RenderProperties != nullptr && PropertyHandle.IsValid())
+
+	if (RenderProperties != nullptr && StructPropertyHandle.IsValid())
 	{
+		MenuBuilder.AddMenuEntry(
+			GetConstantText(true),
+			TAttribute<FText>(),
+			FSlateIcon(),
+			FUIAction(FExecuteAction::CreateSP(this, &FNiagaraRendererMaterialStaticBoolParameterCustomization::SetValue, true))
+		);
+		MenuBuilder.AddMenuEntry(
+			GetConstantText(false),
+			TAttribute<FText>(),
+			FSlateIcon(),
+			FUIAction(FExecuteAction::CreateSP(this, &FNiagaraRendererMaterialStaticBoolParameterCustomization::SetValue, false))
+		);
+		MenuBuilder.AddSeparator();
+
 		const FVersionedNiagaraEmitter NiagaraEmitter = RenderProperties->GetOuterEmitter();
 		const FNiagaraTypeDefinition StaticBoolDef = FNiagaraTypeDefinition::GetBoolDef().ToStaticDef();
 		FNiagaraAliasContext AliasContext;
@@ -2459,18 +2495,12 @@ TSharedRef<SWidget> FNiagaraRendererMaterialStaticBoolParameterCustomization::On
 						continue;
 					}
 
+
 					MenuBuilder.AddMenuEntry(
-						FText::FromName(ResolvedVariable.GetName()),
-						TAttribute<FText>(),
-						FSlateIcon(),
-						FUIAction(
-							FExecuteAction::CreateLambda(
-								[ValidBinding= ResolvedVariable.GetName(), PropertyHandle]()
-								{
-									PropertyHandle->SetValue(ValidBinding);
-								}
-							)
-						)
+						FUIAction(FExecuteAction::CreateSP(this, &FNiagaraRendererMaterialStaticBoolParameterCustomization::SetValue, ResolvedVariable.GetName())),
+						SNew(SNiagaraParameterName)
+						.ParameterName(ResolvedVariable.GetName())
+						.IsReadOnly(true)
 					);
 				}
 			};
@@ -2492,6 +2522,75 @@ TSharedRef<SWidget> FNiagaraRendererMaterialStaticBoolParameterCustomization::On
 	return MenuBuilder.MakeWidget();
 }
 
+FNiagaraRendererMaterialStaticBoolParameter* FNiagaraRendererMaterialStaticBoolParameterCustomization::ResolveParameterStruct() const
+{
+	if (UNiagaraRendererProperties* RenderProperties = WeakRenderProperties.Get())
+	{
+		return StructPropertyHandle.IsValid() ? reinterpret_cast<FNiagaraRendererMaterialStaticBoolParameter*>(StructPropertyHandle->GetValueBaseAddress(reinterpret_cast<uint8*>(RenderProperties))) : nullptr;
+	}
+	return nullptr;
+}
+
+int32 FNiagaraRendererMaterialStaticBoolParameterCustomization::GetWidgetIndex() const
+{
+	if (FNiagaraRendererMaterialStaticBoolParameter* Parameter = ResolveParameterStruct())
+	{
+		return Parameter->StaticValue.IsSet() ? 1 : 0;
+	}
+	return 0;
+}
+
+void FNiagaraRendererMaterialStaticBoolParameterCustomization::SetValue(FName ParameterName) const
+{
+	if (UNiagaraRendererProperties* RenderProperties = WeakRenderProperties.Get())
+	{
+		FScopedTransaction Transaction(LOCTEXT("EditStaticBoolBinding", "Edit Static Bool Binding"));
+		RenderProperties->Modify();
+
+		StructPropertyHandle->NotifyPreChange();
+
+		FNiagaraRendererMaterialStaticBoolParameter* Parameter = ResolveParameterStruct();
+		Parameter->StaticVariableName = ParameterName;
+		Parameter->StaticValue.Reset();
+
+		StructPropertyHandle->NotifyPostChange(EPropertyChangeType::ValueSet);
+		StructPropertyHandle->NotifyFinishedChangingProperties();
+	}
+}
+
+void FNiagaraRendererMaterialStaticBoolParameterCustomization::SetValue(bool Value) const
+{
+	if (UNiagaraRendererProperties* RenderProperties = WeakRenderProperties.Get())
+	{
+		FScopedTransaction Transaction(LOCTEXT("EditStaticBoolBinding", "Edit Static Bool Binding"));
+		RenderProperties->Modify();
+
+		StructPropertyHandle->NotifyPreChange();
+
+		FNiagaraRendererMaterialStaticBoolParameter* Parameter = ResolveParameterStruct();
+		Parameter->StaticVariableName = NAME_None;
+		Parameter->StaticValue = Value;
+
+		StructPropertyHandle->NotifyPostChange(EPropertyChangeType::ValueSet);
+		StructPropertyHandle->NotifyFinishedChangingProperties();
+	}
+}
+
+
+FText FNiagaraRendererMaterialStaticBoolParameterCustomization::GetCurrentConstantValue() const
+{
+	bool bValue = false;
+	if (FNiagaraRendererMaterialStaticBoolParameter* Parameter = ResolveParameterStruct())
+	{
+		bValue = Parameter->StaticValue.Get(false);
+	}
+	return GetConstantText(bValue);
+}
+
+FText FNiagaraRendererMaterialStaticBoolParameterCustomization::GetConstantText(bool bValue) const
+{
+	return bValue ? LOCTEXT("ConstantTrue", "Constant: True") : LOCTEXT("ConstantFalse", "Constant: False");
+}
 
 //////////////////////////////////////////////////////////////////////////
 
diff --git a/Engine/Plugins/FX/Niagara/Source/NiagaraEditor/Private/Customizations/NiagaraTypeCustomizations.h b/Engine/Plugins/FX/Niagara/Source/NiagaraEditor/Private/Customizations/NiagaraTypeCustomizations.h
--- a/Engine/Plugins/FX/Niagara/Source/NiagaraEditor/Private/Customizations/NiagaraTypeCustomizations.h
+++ b/Engine/Plugins/FX/Niagara/Source/NiagaraEditor/Private/Customizations/NiagaraTypeCustomizations.h
@@ -443,7 +443,20 @@ public:
 
 	virtual bool CustomizeChildProperty(class IDetailChildrenBuilder& ChildBuilder, TSharedPtr<IPropertyHandle> PropertyHandle) override;
 	virtual void GetMaterialParameterInfos(UMaterialInterface* Material, TArray<FMaterialParameterInfo>& OutMaterialParameterInfos) const override;
-	TSharedRef<SWidget> OnGetStaticVariablelBindingNameMenuContent(TSharedPtr<IPropertyHandle> PropertyHandle) const;
+	TSharedRef<SWidget> OnGetStaticVariablelBindingNameMenuContent() const;
+
+private:
+	struct FNiagaraRendererMaterialStaticBoolParameter* ResolveParameterStruct() const;
+
+	int32 GetWidgetIndex() const;
+
+	void SetValue(FName ParameterName) const;
+	void SetValue(bool Value) const;
+
+	FText GetCurrentConstantValue() const;
+	FText GetConstantText(bool bValue) const;
+
+	TSharedPtr<IPropertyHandle> StructPropertyHandle;
 };
 
 //** Properties customization for FNiagaraVariables. */
