- Fixes for ray traced ribbons on CPU & GPU
- Fixed various transition issues also

diff --git a/Engine/Plugins/FX/Niagara/Shaders/Private/NiagaraRibbonVertexFactory.ush b/Engine/Plugins/FX/Niagara/Shaders/Private/NiagaraRibbonVertexFactory.ush
--- a/Engine/Plugins/FX/Niagara/Shaders/Private/NiagaraRibbonVertexFactory.ush
+++ b/Engine/Plugins/FX/Niagara/Shaders/Private/NiagaraRibbonVertexFactory.ush
@@ -24,6 +24,8 @@
 #define SLICEVERTEXDATA_TEXTUREV_OFFSET 4
 #define SLICEVERTEXDATA_STRIDE 5
 
+#define INDEX_GEN_DRAW_INDIRECT_OFFSET				3		// Duplicated in NiagaraRibbonCommon.ush
+
 #define INDEX_GEN_NUM_SUB_SEGMENTS_OFFSET			15		// Duplicated in NiagaraRibbonCommon.ush
 #define INDEX_GEN_ONE_OVER_NUM_SUB_SEGMENTS_OFFSET	16		// Duplicated in NiagaraRibbonCommon.ush
 
@@ -96,15 +98,6 @@ float3 CubicInterp(float3 P0, float3 T0, float3 P1, float3 T1, float A, float TS
  *							Particle specific
  *********************************************************************************/
  
-Buffer<uint> SortedIndices;
-
-Buffer<float> TangentsAndDistances;
-Buffer<uint> MultiRibbonIndices;
-Buffer<uint> PackedPerRibbonDataByIndex;
-
-uint FacingMode;
-
-
 // Naming convention :
 // RawParticleId : Used to access per particle data stored in the VF. Actual particle data needs to be indirected through SortedIndices.
 // ParticleId : Used to access particle data in NiagaraVFParticleAccess.
@@ -1006,15 +999,20 @@ FVertexFactoryInput LoadVertexFactoryInputForDynamicUpdate(uint TriangleIndex, i
 	FVertexFactoryInput Input;
 
 	Input.InterpVtxId = TriangleIndex * 3 + VertexIndex;
+	if ( NiagaraRibbonVFLooseParameters.UseIndexBufferForRayTracing != 0 )
+	{
+		Input.InterpVtxId = NiagaraRibbonVFLooseParameters.IndexBuffer[Input.InterpVtxId];
+	}
 
 #if FEATURE_LEVEL >= FEATURE_LEVEL_SM4 && ONEPASS_POINTLIGHT_SHADOW && USING_VERTEX_SHADER_LAYER
 	Input.InstanceId = 0;
 #endif
 	return Input;
 }
+
 uint GetNumRayTracingDynamicMeshVerticesIndirect()
 {
-	return 0;
+	return NiagaraRibbonVFLooseParameters.IndirectDrawOutput[NiagaraRibbonVFLooseParameters.IndirectDrawOutputOffset + INDEX_GEN_DRAW_INDIRECT_OFFSET + 0];
 }
 #endif
 
@@ -1031,9 +1029,10 @@ FVertexFactoryInput LoadVertexFactoryInputForHGS(uint TriangleIndex, int VertexI
 #endif
 	return Input;
 }
+
 uint GetNumRayTracingDynamicMeshVerticesIndirect()
 {
-	return 0;
+	return NiagaraRibbonVFLooseParameters.IndirectDrawOutput[NiagaraRibbonVFLooseParameters.IndirectDrawOutputOffset + INDEX_GEN_DRAW_INDIRECT_OFFSET + 0];
 }
 #endif
 
diff --git a/Engine/Plugins/FX/Niagara/Shaders/Private/Ribbons/NiagaraRibbonCommon.ush b/Engine/Plugins/FX/Niagara/Shaders/Private/Ribbons/NiagaraRibbonCommon.ush
--- a/Engine/Plugins/FX/Niagara/Shaders/Private/Ribbons/NiagaraRibbonCommon.ush
+++ b/Engine/Plugins/FX/Niagara/Shaders/Private/Ribbons/NiagaraRibbonCommon.ush
@@ -37,7 +37,7 @@
 
 // This layout must match FNiagaraRibbonIndirectDrawBufferLayout
 #define INDEX_GEN_EXECUTE_INDIRECT_OFFSET			0
-#define INDEX_GEN_DRAW_INDIRECT_OFFSET				3
+#define INDEX_GEN_DRAW_INDIRECT_OFFSET				3		// Duplicated in NiagaraRibbonVertexFactory.ush
 #define INDEX_GEN_STEREO_DRAW_INDIRECT_OFFSET		8
 
 #define INDEX_GEN_TESSELLATION_FACTOR_OFFSET		13
diff --git a/Engine/Plugins/FX/Niagara/Source/Niagara/Private/NiagaraRendererRibbons.cpp b/Engine/Plugins/FX/Niagara/Source/Niagara/Private/NiagaraRendererRibbons.cpp
--- a/Engine/Plugins/FX/Niagara/Source/Niagara/Private/NiagaraRendererRibbons.cpp
+++ b/Engine/Plugins/FX/Niagara/Source/Niagara/Private/NiagaraRendererRibbons.cpp
@@ -168,9 +168,14 @@ static FAutoConsoleVariableRef CVarNiagaraRibbonShareGeneratedData(
 
 static TAutoConsoleVariable<int32> CVarRayTracingNiagaraRibbons(
 	TEXT("r.RayTracing.Geometry.NiagaraRibbons"),
-	0,
+	1,
 	TEXT("Include Niagara ribbons in ray tracing effects (default = 1 (Niagara ribbons enabled in ray tracing))"));
 
+static TAutoConsoleVariable<int32> CVarRayTracingNiagaraRibbonsGPU(
+	TEXT("r.RayTracing.Geometry.NiagaraRibbons.GPU"),
+	1,
+	TEXT("If we allow GPU ribbon raytracing"));
+
 // max absolute error 9.0x10^-3
 // Eberly's polynomial degree 1 - respect bounds
 // input [-1, 1] and output [0, PI]
@@ -271,6 +276,13 @@ struct FNiagaraRibbonIndexBuffer final : FIndexBuffer
 
 		IndexBufferRHI = IndexAllocation.IndexBuffer->IndexBufferRHI;
 		FirstIndex = IndexAllocation.FirstIndex;
+	#if RHI_RAYTRACING
+		if (IsRayTracingAllowed())
+		{
+			const uint32 IndexStride = IndexBufferRHI->GetDesc().Stride;
+			SRV = RHICmdList.CreateShaderResourceView(IndexBufferRHI, IndexStride, IndexStride == 2 ? PF_R16_UINT : PF_R32_UINT);
+		}
+	#endif
 	}
 
 	// GPU allocation path assumes 32 bit indicies
@@ -280,16 +292,25 @@ struct FNiagaraRibbonIndexBuffer final : FIndexBuffer
 
 		FRHIResourceCreateInfo CreateInfo(TEXT("NiagaraRibbonIndexBuffer"));
 		IndexBufferRHI = RHICmdList.CreateIndexBuffer(sizeof(uint32), sizeof(uint32) * NumElements, BUF_Static | BUF_UnorderedAccess, ERHIAccess::VertexOrIndexBuffer, CreateInfo);
+	#if RHI_RAYTRACING
+		SRV = RHICmdList.CreateShaderResourceView(IndexBufferRHI, sizeof(uint32), PF_R32_UINT);
+	#endif
 		UAV = RHICmdList.CreateUnorderedAccessView(IndexBufferRHI, PF_R32_UINT);
 	}
 
 	virtual void ReleaseRHI() override
 	{
 		UAV.SafeRelease();
+	#if RHI_RAYTRACING
+		SRV.SafeRelease();
+	#endif
 		FIndexBuffer::ReleaseRHI();
 	}
 
 	uint32						FirstIndex = 0;
+#if RHI_RAYTRACING
+	FShaderResourceViewRHIRef	SRV;
+#endif
 	FUnorderedAccessViewRHIRef	UAV;
 };
 
@@ -456,7 +477,7 @@ FNiagaraRibbonVertexBuffers::FNiagaraRibbonVertexBuffers()
 
 void FNiagaraRibbonVertexBuffers::InitializeOrUpdateBuffers(FRHICommandListBase& RHICmdList, const FNiagaraRibbonGenerationConfig& GenerationConfig, const TSharedPtr<FNiagaraRibbonCPUGeneratedVertexData>& GeneratedGeometryData, const FNiagaraDataBuffer* SourceParticleData, int32 MaxAllocationCount, bool bIsUsingGPUInit)
 {	
-	constexpr ERHIAccess InitialBufferAccessFlags = ERHIAccess::SRVMask | ERHIAccess::VertexOrIndexBuffer;
+	constexpr ERHIAccess InitialBufferAccessFlags = ERHIAccess::SRVMask;
 
 	if (bIsUsingGPUInit)
 	{
@@ -467,7 +488,7 @@ void FNiagaraRibbonVertexBuffers::InitializeOrUpdateBuffers(FRHICommandListBase&
 		//       enough space to assume each particle will be the start of a unique ribbon to avoid running OOB on the buffers.
 		const int32 TotalRibbons = FMath::Clamp<int32>(TotalParticles, 1, MaxAllocatedRibbons);
 
-		SortedIndicesBuffer.Allocate(RHICmdList, TotalParticles, MaxAllocationCount, InitialBufferAccessFlags, false);
+		SortedIndicesBuffer.Allocate(RHICmdList, TotalParticles, MaxAllocationCount, InitialBufferAccessFlags | ERHIAccess::VertexOrIndexBuffer, false, EBufferUsageFlags::None);
 		TangentsAndDistancesBuffer.Allocate(RHICmdList, TotalParticles * 4, MaxAllocationCount * 4, InitialBufferAccessFlags, false);
 		MultiRibbonIndicesBuffer.Allocate(RHICmdList, GenerationConfig.HasRibbonIDs() ? TotalParticles : 0, MaxAllocationCount, InitialBufferAccessFlags, false);
 		RibbonLookupTableBuffer.Allocate(RHICmdList, TotalRibbons * FRibbonMultiRibbonInfoBufferEntry::NumElements, MaxAllocatedRibbons * FRibbonMultiRibbonInfoBufferEntry::NumElements, InitialBufferAccessFlags, false);
@@ -478,7 +499,7 @@ void FNiagaraRibbonVertexBuffers::InitializeOrUpdateBuffers(FRHICommandListBase&
 	{		
 		check(GeneratedGeometryData.IsValid());
 
-		SortedIndicesBuffer.Allocate(RHICmdList, GeneratedGeometryData->SortedIndices.Num(), MaxAllocationCount, InitialBufferAccessFlags, true);
+		SortedIndicesBuffer.Allocate(RHICmdList, GeneratedGeometryData->SortedIndices.Num(), MaxAllocationCount, InitialBufferAccessFlags | ERHIAccess::VertexOrIndexBuffer, true, EBufferUsageFlags::None);
 		TangentsAndDistancesBuffer.Allocate(RHICmdList, GeneratedGeometryData->TangentAndDistances.Num() * 4, MaxAllocationCount * 4, InitialBufferAccessFlags, true);
 		MultiRibbonIndicesBuffer.Allocate(RHICmdList, GenerationConfig.HasRibbonIDs() ? GeneratedGeometryData->MultiRibbonIndices.Num() : 0, MaxAllocationCount, InitialBufferAccessFlags, true);
 		RibbonLookupTableBuffer.Allocate(RHICmdList, GeneratedGeometryData->RibbonInfoLookup.Num() * FRibbonMultiRibbonInfoBufferEntry::NumElements, MaxAllocationCount * FRibbonMultiRibbonInfoBufferEntry::NumElements, InitialBufferAccessFlags, true);
@@ -510,7 +531,7 @@ struct FNiagaraRibbonGPUInitComputeBuffers
 	void InitOrUpdateBuffers(FRHICommandListBase& RHICmdList, int32 NeededSize, bool bWantsMultiRibbon, bool bWantsTessellation, bool bWantsTessellationTwist)
 	{
 		// TODO: Downsize these when we haven't needed the size for a bit
-		constexpr ERHIAccess InitialAccess = ERHIAccess::SRVMask | ERHIAccess::VertexOrIndexBuffer;
+		constexpr ERHIAccess InitialAccess = ERHIAccess::SRVMask;
 		
 		if (SortBuffer.NumBytes < NeededSize * sizeof(int32))
 		{
@@ -1093,40 +1114,21 @@ void FNiagaraRendererRibbons::GetDynamicRayTracingInstances(FRayTracingInstanceC
 		return;
 	}
 
-	if (DynamicDataRibbon->bIsGPUSystem)
+	if (SourceParticleData->GetNumInstances() < 2)
 	{
-		//-TODO: RayTracing does not support GPU Ribbons currently
-		return;
-
 		// Bail if we don't have enough particle data to have a valid ribbon
-		// or if somehow the sim targets don't match
-		//if (SimTarget != ENiagaraSimTarget::GPUComputeSim || SourceParticleData->GetNumInstances() < 2)
-		//{
-		//	return;
-		//}
+		return;
 	}
-	else
-	{
-		check(SimTarget == ENiagaraSimTarget::CPUSim);
-
-		if (SourceParticleData->GetNumInstances() < 2)
-		{
-			// Bail if we don't have enough particle data to have a valid ribbon
-			return;
-		}
 
-		//-TODO: RayTracing does not support GPU Init for Ribbons
-		if (DynamicDataRibbon->bUseGPUInit)
-		{
-			return;
-		}
+	if (!DynamicDataRibbon->bUseGPUInit && (!DynamicDataRibbon->GenerationOutput.IsValid() ||DynamicDataRibbon->GenerationOutput->SegmentData.Num() <= 0))
+	{
+		return;
+	}
 
-		if (!DynamicDataRibbon->bUseGPUInit && (
-			!DynamicDataRibbon->GenerationOutput.IsValid() ||
-			DynamicDataRibbon->GenerationOutput->SegmentData.Num() <= 0))
-		{
-			return;
-		}
+	const bool bDataGeneratedOnGpu = DynamicDataRibbon->bUseGPUInit || DynamicDataRibbon->bIsGPUSystem;
+	if ( !CVarRayTracingNiagaraRibbonsGPU.GetValueOnRenderThread() && bDataGeneratedOnGpu )
+	{
+		return;
 	}
 	
 	auto View = Collector.GetReferenceView();
@@ -1156,35 +1158,45 @@ void FNiagaraRendererRibbons::GetDynamicRayTracingInstances(FRayTracingInstanceC
 	{
 		return;
 	}
-	
-	FRayTracingInstance RayTracingInstance;
-	RayTracingInstance.Geometry = &RayTracingGeometry;
-	RayTracingInstance.InstanceTransforms.Add(FMatrix::Identity);
-	
-	RayTracingGeometry.Initializer.IndexBuffer = RenderingViewResources->IndexBuffer->IndexBufferRHI;// PerViewGeneratedData.IndexAllocation.IndexBuffer->IndexBufferRHI;
-	RayTracingGeometry.Initializer.IndexBufferOffset = RenderingViewResources->IndexBuffer->FirstIndex * RenderingViewResources->IndexBuffer->IndexBufferRHI->GetStride();//PerViewGeneratedData.IndexAllocation.FirstIndex * PerViewGeneratedData.IndexAllocation.IndexStride;
-	
+
 	FMeshBatch MeshBatch;
-	
 	SetupMeshBatchAndCollectorResourceForView(RHICmdList, RenderingViewResources->IndexGenerationSettings, DynamicDataRibbon, SourceParticleData, View, *ViewFamily, SceneProxy, RenderingResources.RibbonResources, RenderingViewResources, MeshBatch, DynamicDataRibbon->bUseGPUInit);
 
+	// Get the Vertex / Triangle counts, this is known for CPU but unknown for GPU
+	const uint32 VertexCount = bDataGeneratedOnGpu ? RenderingViewResources->IndexGenerationSettings.TotalNumIndices : MeshBatch.Elements[0].NumPrimitives * 3;
+	const uint32 TriangleCount = bDataGeneratedOnGpu ? RenderingViewResources->IndexGenerationSettings.TotalNumIndices / 3 : MeshBatch.Elements[0].NumPrimitives;
+	if (TriangleCount == 0)
+	{
+		return;
+	}
+
+	FRayTracingInstance RayTracingInstance;
+	RayTracingInstance.Geometry = &RayTracingGeometry;
+	RayTracingInstance.InstanceTransforms.Add(FMatrix::Identity);
+	if (bDataGeneratedOnGpu == false)
+	{
+		RayTracingGeometry.Initializer.IndexBuffer = RenderingViewResources->IndexBuffer->IndexBufferRHI;
+		RayTracingGeometry.Initializer.IndexBufferOffset = RenderingViewResources->IndexBuffer->FirstIndex * RenderingViewResources->IndexBuffer->IndexBufferRHI->GetStride();
+	}
 	RayTracingInstance.Materials.Add(MeshBatch);
 	
 	// Use the internal vertex buffer only when initialized otherwise used the shared vertex buffer - needs to be updated every frame
 	FRWBuffer* VertexBuffer = RayTracingDynamicVertexBuffer.NumBytes > 0 ? &RayTracingDynamicVertexBuffer : nullptr;
-	
-	const uint32 VertexCount = DynamicDataRibbon->bUseGPUInit ? SourceParticleData->GetNumInstances() : DynamicDataRibbon->GenerationOutput->SortedIndices.Num();
-	
-	const int32 MaxTriangleCount = RenderingViewResources->IndexGenerationSettings.MaxSegmentCount * RenderingViewResources->IndexGenerationSettings.SubSegmentCount * ShapeState.TrianglesPerSegment;
-	
+
+	if (bDataGeneratedOnGpu)
+	{
+		//-TODO: We can optimize this to potentially share the data
+		GpuRibbonDataManager.RegisterRenderer(this, SourceParticleData, RenderingResources.RibbonResources);
+	}
+
 	Collector.AddRayTracingGeometryUpdate(
 		FRayTracingDynamicGeometryUpdateParams
 		{
 			RayTracingInstance.Materials,
-			DynamicDataRibbon->bUseGPUInit,
+			bDataGeneratedOnGpu,
 			VertexCount,
 			VertexCount * static_cast<uint32>(sizeof(FVector3f)),
-			DynamicDataRibbon->bUseGPUInit? MaxTriangleCount : MeshBatch.Elements[0].NumPrimitives,
+			TriangleCount,
 			&RayTracingGeometry,
 			VertexBuffer,
 			true
@@ -2154,6 +2166,13 @@ inline void FNiagaraRendererRibbons::SetupMeshBatchAndCollectorResourceForView(F
 	
 	// Set common data on vertex factory
 	FNiagaraRibbonVFLooseParameters VFLooseParams;
+#if RHI_RAYTRACING
+	VFLooseParams.IndexBuffer = GetSrvOrDefaultUInt(RenderingViewResources->IndexBuffer->SRV);
+	VFLooseParams.UseIndexBufferForRayTracing = bShouldUseGPUInitIndices;
+#else
+	VFLooseParams.IndexBuffer = GetDummyUIntBuffer();
+	VFLooseParams.UseIndexBufferForRayTracing = false;
+#endif
 	VFLooseParams.SortedIndices = VertexBuffers.SortedIndicesBuffer.SRV;
 	VFLooseParams.TangentsAndDistances = VertexBuffers.TangentsAndDistancesBuffer.SRV;
 	VFLooseParams.MultiRibbonIndices = GetSrvOrDefaultUInt(VertexBuffers.MultiRibbonIndicesBuffer.SRV);
@@ -2206,10 +2225,6 @@ inline void FNiagaraRendererRibbons::SetupMeshBatchAndCollectorResourceForView(F
 	else
 	{
 		MeshElement.NumPrimitives = GeneratedData.CPUTriangleCount;
-
-		const uint32 MaxVertexCount = bShouldUseGPUInitIndices ? SourceParticleData->GetNumInstances() : DynamicDataRibbon->GenerationOutput->SortedIndices.Num();
-		MeshElement.MaxVertexIndex = MaxVertexCount > 0 ? MaxVertexCount - 1 : 0;
-
 		check(MeshElement.NumPrimitives > 0);
 	}	
 	
@@ -2494,9 +2509,9 @@ void FNiagaraRendererRibbons::InitializeVertexBuffersGPU(FRHICommandListImmediat
 	// Clear the command buffer if we just initialized it, or if the sim doesn't have enough data to run
 	if ((!bCanRun || VertexBuffers.bJustCreatedCommandBuffer) && VertexBuffers.GPUComputeCommandBuffer.NumBytes > 0)
 	{
-		RHICmdList.Transition(FRHITransitionInfo(VertexBuffers.GPUComputeCommandBuffer.Buffer, ERHIAccess::SRVMask | ERHIAccess::VertexOrIndexBuffer | ERHIAccess::IndirectArgs, ERHIAccess::UAVCompute));
+		RHICmdList.Transition(FRHITransitionInfo(VertexBuffers.GPUComputeCommandBuffer.Buffer, ERHIAccess::SRVMask | ERHIAccess::IndirectArgs, ERHIAccess::UAVCompute));
 		RHICmdList.ClearUAVUint(VertexBuffers.GPUComputeCommandBuffer.UAV, FUintVector4(0));
-		RHICmdList.Transition(FRHITransitionInfo(VertexBuffers.GPUComputeCommandBuffer.Buffer, ERHIAccess::UAVCompute, ERHIAccess::SRVMask | ERHIAccess::VertexOrIndexBuffer | ERHIAccess::IndirectArgs));
+		RHICmdList.Transition(FRHITransitionInfo(VertexBuffers.GPUComputeCommandBuffer.Buffer, ERHIAccess::UAVCompute, ERHIAccess::SRVMask | ERHIAccess::IndirectArgs));
 		VertexBuffers.bJustCreatedCommandBuffer = false;
 	}
 	
@@ -2744,15 +2759,15 @@ void FNiagaraRendererRibbons::InitializeVertexBuffersGPU(FRHICommandListImmediat
 				FMath::DivideAndRoundUp<uint32>(NumExecutableInstances, FNiagaraRibbonComputeCommon::VertexGenReductionFinalizationThreadSize) :
 				1;
 			RHICmdList.Transition({
-				FRHITransitionInfo(VertexBuffers.RibbonLookupTableBuffer.Buffer, ERHIAccess::SRVMask | ERHIAccess::VertexOrIndexBuffer, ERHIAccess::UAVCompute),
-				FRHITransitionInfo(VertexBuffers.GPUComputeCommandBuffer.Buffer, ERHIAccess::SRVMask | ERHIAccess::VertexOrIndexBuffer | ERHIAccess::IndirectArgs, ERHIAccess::UAVCompute)});
+				FRHITransitionInfo(VertexBuffers.RibbonLookupTableBuffer.Buffer, ERHIAccess::SRVMask, ERHIAccess::UAVCompute),
+				FRHITransitionInfo(VertexBuffers.GPUComputeCommandBuffer.Buffer, ERHIAccess::SRVMask | ERHIAccess::IndirectArgs, ERHIAccess::UAVCompute)});
 
 			FComputeShaderUtils::Dispatch(RHICmdList, ComputeShader, FinalizationParams, FIntVector(NumThreadGroups, 1, 1));
 			
 			// We don't need to transition RibbonLookupTableBuffer as it's still needed for the next shader
 			RHICmdList.Transition({
 				FRHITransitionInfo(VertexBuffers.RibbonLookupTableBuffer.Buffer, ERHIAccess::UAVCompute, ERHIAccess::UAVCompute),
-				FRHITransitionInfo(VertexBuffers.GPUComputeCommandBuffer.Buffer, ERHIAccess::UAVCompute, ERHIAccess::SRVMask | ERHIAccess::VertexOrIndexBuffer | ERHIAccess::IndirectArgs)});		
+				FRHITransitionInfo(VertexBuffers.GPUComputeCommandBuffer.Buffer, ERHIAccess::UAVCompute, ERHIAccess::SRVMask | ERHIAccess::IndirectArgs)});		
 		}
 		
 		{
@@ -2797,13 +2812,13 @@ void FNiagaraRendererRibbons::InitializeVertexBuffersGPU(FRHICommandListImmediat
 				// We don't need to transition RibbonLookupTableBuffer as it's still setup for UAV from the last shader
 				RHICmdList.Transition({
 					FRHITransitionInfo(VertexBuffers.RibbonLookupTableBuffer.Buffer, ERHIAccess::UAVCompute, ERHIAccess::UAVCompute),
-					FRHITransitionInfo(VertexBuffers.TangentsAndDistancesBuffer.Buffer, ERHIAccess::SRVMask | ERHIAccess::VertexOrIndexBuffer, ERHIAccess::UAVCompute)});
+					FRHITransitionInfo(VertexBuffers.TangentsAndDistancesBuffer.Buffer, ERHIAccess::SRVMask, ERHIAccess::UAVCompute)});
 
 				FComputeShaderUtils::DispatchIndirect(RHICmdList, ComputeShader, FinalizeParams, VertexBuffers.GPUComputeCommandBuffer.Buffer, CommandBufferOffset * FNiagaraRibbonCommandBufferLayout::NumElements);
 
 				RHICmdList.Transition({
-					FRHITransitionInfo(VertexBuffers.TangentsAndDistancesBuffer.Buffer, ERHIAccess::UAVCompute, ERHIAccess::SRVMask | ERHIAccess::VertexOrIndexBuffer),
-					FRHITransitionInfo(VertexBuffers.RibbonLookupTableBuffer.Buffer, ERHIAccess::UAVCompute, ERHIAccess::SRVMask | ERHIAccess::VertexOrIndexBuffer)});
+					FRHITransitionInfo(VertexBuffers.TangentsAndDistancesBuffer.Buffer, ERHIAccess::UAVCompute, ERHIAccess::SRVMask),
+					FRHITransitionInfo(VertexBuffers.RibbonLookupTableBuffer.Buffer, ERHIAccess::UAVCompute, ERHIAccess::SRVMask)});
 			}		
 		}
 	}
diff --git a/Engine/Plugins/FX/Niagara/Source/NiagaraVertexFactories/Public/NiagaraRibbonVertexFactory.h b/Engine/Plugins/FX/Niagara/Source/NiagaraVertexFactories/Public/NiagaraRibbonVertexFactory.h
--- a/Engine/Plugins/FX/Niagara/Source/NiagaraVertexFactories/Public/NiagaraRibbonVertexFactory.h
+++ b/Engine/Plugins/FX/Niagara/Source/NiagaraVertexFactories/Public/NiagaraRibbonVertexFactory.h
@@ -75,6 +75,8 @@ BEGIN_GLOBAL_SHADER_PARAMETER_STRUCT(FNiagaraRibbonVFLooseParameters, NIAGARAVER
 	SHADER_PARAMETER(uint32, Shape)
 	SHADER_PARAMETER(uint32, NeedsPreciseMotionVectors)
 	SHADER_PARAMETER(uint32, UseGeometryNormals)
+	SHADER_PARAMETER(uint32, UseIndexBufferForRayTracing)
+	SHADER_PARAMETER_SRV(Buffer<uint>, IndexBuffer)
 	SHADER_PARAMETER_SRV(Buffer<uint>, SortedIndices)
 	SHADER_PARAMETER_SRV(Buffer<float>, TangentsAndDistances)
 	SHADER_PARAMETER_SRV(Buffer<uint>, MultiRibbonIndices)
