PR #12610: Improve SNiagaraBaselineViewport performance

diff --git a/Engine/Plugins/FX/Niagara/Source/NiagaraEditor/Private/Widgets/SNiagaraSystemViewport.cpp b/Engine/Plugins/FX/Niagara/Source/NiagaraEditor/Private/Widgets/SNiagaraSystemViewport.cpp
--- a/Engine/Plugins/FX/Niagara/Source/NiagaraEditor/Private/Widgets/SNiagaraSystemViewport.cpp
+++ b/Engine/Plugins/FX/Niagara/Source/NiagaraEditor/Private/Widgets/SNiagaraSystemViewport.cpp
@@ -1477,9 +1477,9 @@ void FNiagaraBaselineViewportClient::Tick(float DeltaSeconds)
 	{
 		if (!World->GetBegunPlay())
 		{
-			for (FActorIterator It(World); It; ++It)
+			for (TActorIterator<ANiagaraPerfBaselineActor> It(World); It; ++It)
 			{
-				if (ANiagaraPerfBaselineActor* BaselineActor = Cast<ANiagaraPerfBaselineActor>(*It))
+				if (ANiagaraPerfBaselineActor* BaselineActor = *It)
 				{
 					It->DispatchBeginPlay();
 				}
