- Speculative change for a low repro crash

diff --git a/Engine/Plugins/FX/Niagara/Source/Niagara/Private/NiagaraDataInterfaceCamera.cpp b/Engine/Plugins/FX/Niagara/Source/Niagara/Private/NiagaraDataInterfaceCamera.cpp
--- a/Engine/Plugins/FX/Niagara/Source/Niagara/Private/NiagaraDataInterfaceCamera.cpp
+++ b/Engine/Plugins/FX/Niagara/Source/Niagara/Private/NiagaraDataInterfaceCamera.cpp
@@ -142,11 +142,7 @@ bool UNiagaraDataInterfaceCamera::PerInstanceTick(void* PerInstanceData, FNiagar
 		for (FConstPlayerControllerIterator Iterator = World->GetPlayerControllerIterator(); Iterator; ++Iterator)
 		{
 			APlayerController* PlayerController = Iterator->Get();
-			const bool IsLocalPlayer = (PlayerController && PlayerController->Player)
-				? PlayerController->Player->IsA<ULocalPlayer>()
-				: false;
-
-			if (IsLocalPlayer)
+			if (PlayerController != nullptr && PlayerController->GetLocalPlayer() != nullptr)
 			{
 				if (LocalPlayerIndex == PlayerControllerIndex)
 				{
