- Add missing transactions to node

diff --git a/Engine/Plugins/FX/Niagara/Source/NiagaraEditor/Private/Customizations/NiagaraStaticSwitchNodeDetails.cpp b/Engine/Plugins/FX/Niagara/Source/NiagaraEditor/Private/Customizations/NiagaraStaticSwitchNodeDetails.cpp
--- a/Engine/Plugins/FX/Niagara/Source/NiagaraEditor/Private/Customizations/NiagaraStaticSwitchNodeDetails.cpp
+++ b/Engine/Plugins/FX/Niagara/Source/NiagaraEditor/Private/Customizations/NiagaraStaticSwitchNodeDetails.cpp
@@ -281,6 +281,9 @@ void FNiagaraStaticSwitchNodeDetails::DefaultIntValueCommitted(int32 Value, ETex
 	UNiagaraScriptVariable* ScriptVar = GetSwitchParameterScriptVar();
 	if (ScriptVar)
 	{
+		const FScopedTransaction Transaction(LOCTEXT("SetStaticSwitchIntValue ", "Set Static Switch Int Value"));
+		ScriptVar->Modify();
+
 		ScriptVar->SetStaticSwitchDefaultValue(Value);
 		SetSwitchParameterMetadata(ScriptVar->Metadata);
 	}
@@ -291,6 +294,9 @@ void FNiagaraStaticSwitchNodeDetails::DefaultBoolValueCommitted(ECheckBoxState N
 	UNiagaraScriptVariable* ScriptVar = GetSwitchParameterScriptVar();
 	if (ScriptVar)
 	{
+		const FScopedTransaction Transaction(LOCTEXT("SetStaticSwitchBoolValue ", "Set Static Switch Bool Value"));
+		ScriptVar->Modify();
+
 		ScriptVar->SetStaticSwitchDefaultValue((NewState == ECheckBoxState::Checked) ? 1 : 0);
 		SetSwitchParameterMetadata(ScriptVar->Metadata);
 	}
@@ -301,6 +307,7 @@ void FNiagaraStaticSwitchNodeDetails::ExposeAsPinCommitted(ECheckBoxState NewSta
 {
 	if (Node.IsValid())
 	{
+		const FScopedTransaction Transaction(LOCTEXT("SetStaticSwitchExposed", "Set Static Switch Exposed"));
 		Node->Modify();
 		Node->SwitchTypeData.bExposeAsPin = (NewState == ECheckBoxState::Checked);
 
@@ -362,6 +369,9 @@ void FNiagaraStaticSwitchNodeDetails::OnSelectionChanged(TSharedPtr<DefaultEnumO
 		return;
 	}
 
+	const FScopedTransaction Transaction(LOCTEXT("SetStaticSwitchDefaultValue", "Set Static Switch Default Value"));
+	ScriptVar->Modify();
+
 	ScriptVar->SetStaticSwitchDefaultValue(SelectedDefaultValue->EnumIndex);
 	SetSwitchParameterMetadata(ScriptVar->Metadata);
 	if (SwitchTypeDataProperty.IsValid())
@@ -377,7 +387,10 @@ void FNiagaraStaticSwitchNodeDetails::OnSelectionChanged(TSharedPtr<SwitchDropdo
 	{
 		return;
 	}
-	
+
+	const FScopedTransaction Transaction(LOCTEXT("SetStaticSwitchType", "Set Static Switch Type"));
+	Node->Modify();
+
 	FNiagaraTypeDefinition OldType = Node->GetInputType();
 	Node->SwitchTypeData.Enum = nullptr;
 	Node->SwitchTypeData.SwitchConstant = NAME_None;
@@ -412,6 +425,10 @@ void FNiagaraStaticSwitchNodeDetails::OnSelectionChanged(TSharedPtr<ConstantDrop
 	{
 		return;
 	}
+
+	const FScopedTransaction Transaction(LOCTEXT("SetStaticSwitchCompilerConstant", "Set Static Switch Compiler Constant"));
+	Node->Modify();
+
 	Node->SwitchTypeData.SwitchConstant = NewValue->Constant.GetName();
 	
 	// in case of an enum constant we also need the enum type
@@ -605,8 +622,14 @@ void FNiagaraStaticSwitchNodeDetails::OnParameterNameCommited(const FText& InTex
 {
 	if (Node.IsValid())
 	{
-		Node->ChangeSwitchParameterName(FName(*InText.ToString()));
-		Node->GetNiagaraGraph()->NotifyGraphNeedsRecompile();
+		FName NewName(*InText.ToString());
+		if ( Node->InputParameterName != NewName )
+		{
+			const FScopedTransaction Transaction(LOCTEXT("SetStaticSwitchName", "Set Static Switch Name"));
+			Node->Modify();
+			Node->ChangeSwitchParameterName(NewName);
+			Node->GetNiagaraGraph()->NotifyGraphNeedsRecompile();
+		}
 	}
 }
 
