- Add SRV access to ribbon index buffer allocation

#rnx
#tests Ran test level with GPU ribbons in it

diff --git a/Engine/Plugins/FX/Niagara/Source/Niagara/Private/NiagaraRendererRibbons.cpp b/Engine/Plugins/FX/Niagara/Source/Niagara/Private/NiagaraRendererRibbons.cpp
--- a/Engine/Plugins/FX/Niagara/Source/Niagara/Private/NiagaraRendererRibbons.cpp
+++ b/Engine/Plugins/FX/Niagara/Source/Niagara/Private/NiagaraRendererRibbons.cpp
@@ -291,7 +291,7 @@ struct FNiagaraRibbonIndexBuffer final : FIndexBuffer
 		InitResource(RHICmdList);
 
 		FRHIResourceCreateInfo CreateInfo(TEXT("NiagaraRibbonIndexBuffer"));
-		IndexBufferRHI = RHICmdList.CreateIndexBuffer(sizeof(uint32), sizeof(uint32) * NumElements, BUF_Static | BUF_UnorderedAccess, ERHIAccess::VertexOrIndexBuffer, CreateInfo);
+		IndexBufferRHI = RHICmdList.CreateIndexBuffer(sizeof(uint32), sizeof(uint32) * NumElements, BUF_Static | BUF_UnorderedAccess | BUF_ShaderResource, ERHIAccess::VertexOrIndexBuffer, CreateInfo);
 	#if RHI_RAYTRACING
 		SRV = RHICmdList.CreateShaderResourceView(IndexBufferRHI, sizeof(uint32), PF_R32_UINT);
 	#endif
