integrate looping into svt baker output

diff --git a/Engine/Plugins/FX/Niagara/Shaders/Private/NiagaraBlendSVTsToDenseBuffer.usf b/Engine/Plugins/FX/Niagara/Shaders/Private/NiagaraBlendSVTsToDenseBuffer.usf
new file mode 100644
--- /dev/null
+++ b/Engine/Plugins/FX/Niagara/Shaders/Private/NiagaraBlendSVTsToDenseBuffer.usf
@@ -0,0 +1,40 @@
+// Copyright Epic Games, Inc. All Rights Reserved.
+
+#include "/Engine/Private/Common.ush"
+#include "/Engine/Private/SparseVolumeTexture/SparseVolumeTextureCommon.ush"
+RWTexture3D<float4>	DestinationBuffer;
+
+SamplerState		TileDataTextureSampler_A;
+Texture3D<uint>		SparseVolumeTexturePageTable_A;
+Texture3D			SparseVolumeTextureA_A;
+uint4				PackedSVTUniforms0_A;
+uint4				PackedSVTUniforms1_A;
+int3				TextureSize_A;
+int					MipLevels_A;
+
+SamplerState		TileDataTextureSampler_B;
+Texture3D<uint>		SparseVolumeTexturePageTable_B;
+Texture3D			SparseVolumeTextureA_B;
+uint4				PackedSVTUniforms0_B;
+uint4				PackedSVTUniforms1_B;
+int3				TextureSize_B;
+int					MipLevels_B;
+
+float				LerpAmount;
+
+[numthreads(4, 4, 4)]
+void PerformBlendCS(uint3 Index : SV_DispatchThreadID)
+{
+	const int MipLevel = 0; // We want to sample level 0 but the page table can actually point to lower resolution mip level voxels	
+	const int3 VirtualCoord = Index >> MipLevel;	
+
+	const FSparseVolumeTextureUniforms Uniforms_A = SparseVolumeTextureUnpackUniforms(PackedSVTUniforms0_A, PackedSVTUniforms1_A);		
+	const int3 VoxelCoordA = SparseVolumeTextureLoadPageTable(SparseVolumeTexturePageTable_A, Uniforms_A, VirtualCoord, MipLevel);
+	float4 ValueA = SparseVolumeTextureA_A.Load(int4(VoxelCoordA, 0));
+
+	const FSparseVolumeTextureUniforms Uniforms_B = SparseVolumeTextureUnpackUniforms(PackedSVTUniforms0_B, PackedSVTUniforms1_B);		
+	const int3 VoxelCoordB = SparseVolumeTextureLoadPageTable(SparseVolumeTexturePageTable_B, Uniforms_B, VirtualCoord, MipLevel);
+	float4 ValueB = SparseVolumeTextureA_B.Load(int4(VoxelCoordB, 0));
+	
+	DestinationBuffer[Index] = lerp(ValueA, ValueB, LerpAmount);
+}
diff --git a/Engine/Plugins/FX/Niagara/Source/Niagara/Classes/NiagaraBakerOutputSparseVolumeTexture.h b/Engine/Plugins/FX/Niagara/Source/Niagara/Classes/NiagaraBakerOutputSparseVolumeTexture.h
--- a/Engine/Plugins/FX/Niagara/Source/Niagara/Classes/NiagaraBakerOutputSparseVolumeTexture.h
+++ b/Engine/Plugins/FX/Niagara/Source/Niagara/Classes/NiagaraBakerOutputSparseVolumeTexture.h
@@ -32,6 +32,30 @@ public:
 	UPROPERTY(EditAnywhere, Category = "Settings")
 	FString SparseVolumeTextureAssetPathFormat = TEXT("{AssetFolder}/{AssetName}_BakedSVT_{OutputName}");
 
+	/**
+	Enable outputting a seamlessly blended looped SVT sequence
+	*/
+	UPROPERTY(EditAnywhere, Category = "Looping")
+	bool bEnableLoopedOutput = false;
+	
+	/**
+	Path for the looped SVT
+	*/
+	UPROPERTY(EditAnywhere, Category = "Looping", meta = (EditCondition = "bEnableLoopedOutput", EditConditionHides))
+	FString LoopedSparseVolumeTextureAssetPathFormat = TEXT("{AssetFolder}/{AssetName}_Looped_BakedSVT_{OutputName}");
+
+	/**
+	Time in seconds to start the looped output from
+	*/
+	UPROPERTY(EditAnywhere, Category = "Looping", meta = (EditCondition = "bEnableLoopedOutput", EditConditionHides))
+	float StartTime = 4.0f;
+
+	/**
+	Number of seconds to blend the output for
+	*/
+	UPROPERTY(EditAnywhere, Category = "Looping", meta = (EditCondition = "bEnableLoopedOutput", EditConditionHides))
+	float BlendDuration = 2.0f;
+
 	NIAGARA_API virtual bool Equals(const UNiagaraBakerOutput& Other) const override;
 
 #if WITH_EDITOR
diff --git a/Engine/Plugins/FX/Niagara/Source/Niagara/Classes/NiagaraBakerSettings.h b/Engine/Plugins/FX/Niagara/Source/Niagara/Classes/NiagaraBakerSettings.h
--- a/Engine/Plugins/FX/Niagara/Source/Niagara/Classes/NiagaraBakerSettings.h
+++ b/Engine/Plugins/FX/Niagara/Source/Niagara/Classes/NiagaraBakerSettings.h
@@ -196,6 +196,10 @@ public:
 	UPROPERTY(EditAnywhere, Category = "Settings")
 	uint8 bRenderComponentOnly : 1;
 
+	/** Should we preview the baked looped simulation if it has been generated, or the full baked sim */
+	UPROPERTY(EditAnywhere, Category = "Settings")
+	uint8 bPreviewLoopedOutput : 1;
+	
 	NIAGARA_API bool Equals(const UNiagaraBakerSettings& Other) const;
 
 	float GetSeekDelta() const { return 1.0f / float(FramesPerSecond); }
diff --git a/Engine/Plugins/FX/Niagara/Source/Niagara/Private/NiagaraBakerOutputSparseVolumeTexture.cpp b/Engine/Plugins/FX/Niagara/Source/Niagara/Private/NiagaraBakerOutputSparseVolumeTexture.cpp
--- a/Engine/Plugins/FX/Niagara/Source/Niagara/Private/NiagaraBakerOutputSparseVolumeTexture.cpp
+++ b/Engine/Plugins/FX/Niagara/Source/Niagara/Private/NiagaraBakerOutputSparseVolumeTexture.cpp
@@ -9,7 +9,13 @@ bool UNiagaraBakerOutputSparseVolumeTexture::Equals(const UNiagaraBakerOutput& O
 	const UNiagaraBakerOutputSparseVolumeTexture& Other = *CastChecked<UNiagaraBakerOutputSparseVolumeTexture>(&OtherBase);
 	return
 		Super::Equals(Other) &&
-		SourceBinding.SourceName == Other.SourceBinding.SourceName;
+		SourceBinding.SourceName == Other.SourceBinding.SourceName && 
+		VolumeWorldSpaceSizeBinding == Other.VolumeWorldSpaceSizeBinding &&
+		SparseVolumeTextureAssetPathFormat == Other.SparseVolumeTextureAssetPathFormat &&
+		bEnableLoopedOutput == Other.bEnableLoopedOutput &&
+		LoopedSparseVolumeTextureAssetPathFormat == Other.LoopedSparseVolumeTextureAssetPathFormat &&
+		StartTime == Other.StartTime &&
+		BlendDuration == Other.BlendDuration;
 }
 
 #if WITH_EDITOR
diff --git a/Engine/Plugins/FX/Niagara/Source/Niagara/Private/NiagaraBakerSettings.cpp b/Engine/Plugins/FX/Niagara/Source/Niagara/Private/NiagaraBakerSettings.cpp
--- a/Engine/Plugins/FX/Niagara/Source/Niagara/Private/NiagaraBakerSettings.cpp
+++ b/Engine/Plugins/FX/Niagara/Source/Niagara/Private/NiagaraBakerSettings.cpp
@@ -19,6 +19,7 @@ UNiagaraBakerSettings::UNiagaraBakerSettings(const FObjectInitializer& Init)
 	bLockToSimulationFrameRate = false;
 	bPreviewLooping = true;
 	bRenderComponentOnly = true;
+	bPreviewLoopedOutput = true;
 	Outputs.Add(Init.CreateDefaultSubobject<UNiagaraBakerOutputTexture2D>(this, "DefaultOutput"));
 
 	// Fill in default camera settings
@@ -69,7 +70,8 @@ bool UNiagaraBakerSettings::Equals(const UNiagaraBakerSettings& Other) const
 		FramesPerSecond == Other.FramesPerSecond &&
 		FramesPerDimension == Other.FramesPerDimension &&
 		bPreviewLooping == Other.bPreviewLooping &&
-		bRenderComponentOnly == Other.bRenderComponentOnly;
+		bRenderComponentOnly == Other.bRenderComponentOnly && 
+		bPreviewLoopedOutput == Other.bPreviewLoopedOutput;
 }
 
 FVector UNiagaraBakerSettings::GetCameraLocation() const
diff --git a/Engine/Plugins/FX/Niagara/Source/NiagaraEditor/Private/NiagaraBakerRenderer.cpp b/Engine/Plugins/FX/Niagara/Source/NiagaraEditor/Private/NiagaraBakerRenderer.cpp
--- a/Engine/Plugins/FX/Niagara/Source/NiagaraEditor/Private/NiagaraBakerRenderer.cpp
+++ b/Engine/Plugins/FX/Niagara/Source/NiagaraEditor/Private/NiagaraBakerRenderer.cpp
@@ -690,7 +690,10 @@ void FNiagaraBakerRenderer::RenderSparseVolumeTexture(UTextureRenderTarget2D* Re
 	{
 		SVTPreviewComponent->OverrideMaterials[0]->GetMaterial()->SetSparseVolumeTextureParameterValueEditorOnly("SparseVolumeTexture", SVT);
 	}
-	SVTPreviewComponent->SetFrame(Indices.FrameIndexA);
+	
+	// do mod if we are previewing the looped result and the timeline doesn't line up with the baked number of frames for the looped
+	// sim.  This will at least loop it, but there could be a frame pop when the timeline resets
+	SVTPreviewComponent->SetFrame(Indices.FrameIndexA % SVT->GetNumFrames());
 	
 	// #todo(dmp): apply world scale and pivot to HV actor
 
diff --git a/Engine/Plugins/FX/Niagara/Source/NiagaraEditor/Private/NiagaraBakerRendererOutputSparseVolumeTexture.cpp b/Engine/Plugins/FX/Niagara/Source/NiagaraEditor/Private/NiagaraBakerRendererOutputSparseVolumeTexture.cpp
--- a/Engine/Plugins/FX/Niagara/Source/NiagaraEditor/Private/NiagaraBakerRendererOutputSparseVolumeTexture.cpp
+++ b/Engine/Plugins/FX/Niagara/Source/NiagaraEditor/Private/NiagaraBakerRendererOutputSparseVolumeTexture.cpp
@@ -11,6 +11,8 @@
 #include "NiagaraDataInterfaceGrid3DCollection.h"
 #include "NiagaraDataInterfaceRenderTargetVolume.h"
 #include "NiagaraBatchedElements.h"
+#include "NiagaraSVTShaders.h"
+#include "NiagaraShader.h"
 
 #include "Engine/Canvas.h"
 #include "Engine/TextureRenderTarget2D.h"
@@ -21,6 +23,7 @@
 #include "Factories/VolumeTextureFactory.h"
 #include "TextureResource.h"
 #include "UObject/UObjectGlobals.h"
+#include "RenderGraphUtils.h"
 
 #include "TextureRenderTargetVolumeResource.h"
 #include "Engine/TextureRenderTargetVolume.h"
@@ -96,15 +99,31 @@ FIntPoint FNiagaraBakerRendererOutputSparseVolumeTexture::GetGeneratedSize(UNiag
 void FNiagaraBakerRendererOutputSparseVolumeTexture::RenderGenerated(UNiagaraBakerOutput* InBakerOutput, const FNiagaraBakerRenderer& BakerRenderer, UTextureRenderTarget2D* RenderTarget, TOptional<FString>& OutErrorString) const
 {
 	static FString SVTNotFoundError(TEXT("Sparse Volume Texture asset not found.\nPlease bake to see the result."));
+	static FString LoopedSVTNotFoundError(TEXT("Looped Sparse Volume Texture asset not found.\reverting to full frame range baked result."));
 
 	UNiagaraBakerOutputSparseVolumeTexture* BakerOutput = CastChecked<UNiagaraBakerOutputSparseVolumeTexture>(InBakerOutput);
-	UNiagaraBakerSettings* BakerSettings = BakerOutput->GetTypedOuter<UNiagaraBakerSettings>();
+	UNiagaraBakerSettings* BakerSettings = BakerRenderer.GetBakerSettings();
+
+	UAnimatedSparseVolumeTexture* SVT = nullptr;
+	
+	bool bPreviewLoopedOutput = true;
+	if (BakerOutput->bEnableLoopedOutput && BakerSettings->bPreviewLoopedOutput)
+	{
+		SVT = BakerOutput->GetAsset<UAnimatedSparseVolumeTexture>(BakerOutput->LoopedSparseVolumeTextureAssetPathFormat, 0);
+		if (SVT == nullptr)
+		{
+			OutErrorString = SVTNotFoundError;			
+		}
+	}
 
-	UAnimatedSparseVolumeTexture* SVT = BakerOutput->GetAsset<UAnimatedSparseVolumeTexture>(BakerOutput->SparseVolumeTextureAssetPathFormat, 0);
 	if (SVT == nullptr)
 	{
-		OutErrorString = SVTNotFoundError;
-		return;
+		SVT = BakerOutput->GetAsset<UAnimatedSparseVolumeTexture>(BakerOutput->SparseVolumeTextureAssetPathFormat, 0);
+		if (SVT == nullptr)
+		{
+			OutErrorString = SVTNotFoundError;
+			return;
+		}
 	}
 
 	const float WorldTime = BakerRenderer.GetWorldTime();
@@ -335,6 +354,8 @@ void FNiagaraBakerRendererOutputSparseVolumeTexture::BakeFrame(FNiagaraBakerFeed
 
 void FNiagaraBakerRendererOutputSparseVolumeTexture::EndBake(FNiagaraBakerFeedbackContext& FeedbackContext, UNiagaraBakerOutput* InBakerOutput)
 {
+	UNiagaraBakerSettings* BakerSettings = InBakerOutput->GetTypedOuter<UNiagaraBakerSettings>();
+	
 	UNiagaraBakerOutputSparseVolumeTexture* OutputVolumeTexture = CastChecked<UNiagaraBakerOutputSparseVolumeTexture>(InBakerOutput);
 		
 	if (!SVTAsset->EndInitialize())
@@ -343,5 +364,322 @@ void FNiagaraBakerRendererOutputSparseVolumeTexture::EndBake(FNiagaraBakerFeedba
 	}
 
 	SVTAsset->PostLoad();
+
+	// make a second pass over the data and loop it
+	if (OutputVolumeTexture->bEnableLoopedOutput)
+	{		
+		const FString LoopedAssetFullName = OutputVolumeTexture->GetAssetPath(OutputVolumeTexture->LoopedSparseVolumeTextureAssetPathFormat, 0);
+
+		IAssetRegistry* AssetRegistry = IAssetRegistry::Get();
+		TArray<FAssetData> FoundAssets;
+		bool FoundAsset = false;
+		if (AssetRegistry->GetAssetsByPackageName(FName(LoopedAssetFullName), FoundAssets))
+		{
+			if (FoundAssets.Num() > 0)
+			{
+				if (UObject* ExistingOject = StaticLoadObject(UAnimatedSparseVolumeTexture::StaticClass(), nullptr, *LoopedAssetFullName))
+				{
+					FoundAsset = true;
+				}
+			}
+		}
+
+		UNiagaraBakerOutputSparseVolumeTexture* BakerOutput = CastChecked<UNiagaraBakerOutputSparseVolumeTexture>(InBakerOutput);
+		UAnimatedSparseVolumeTexture* LoopedSVT = BakerOutput->GetAsset<UAnimatedSparseVolumeTexture>(BakerOutput->LoopedSparseVolumeTextureAssetPathFormat, 0);
+
+		if (LoopedSVT == nullptr)
+		{
+			LoopedSVTAsset = UNiagaraBakerOutput::GetOrCreateAsset<UAnimatedSparseVolumeTexture, USparseVolumeTextureFactory>(LoopedAssetFullName);
+		}
+		else
+		{
+			LoopedSVTAsset = NewObject<UAnimatedSparseVolumeTexture>(LoopedSVT->GetOuter(), UAnimatedSparseVolumeTexture::StaticClass(), *LoopedSVT->GetName(), RF_Public | RF_Standalone);
+			LoopedSVTAsset->PostEditChange();
+		}
+
+		if (!LoopedSVTAsset->BeginInitialize(1))
+		{
+			UE_LOG(LogNiagaraBaker, Error, TEXT("Cannot initialize looped SVT for baking"));
+			return;
+		}
+
+		// create dense output buffer
+		const float FrameRate = BakerSettings->FramesPerSecond;
+		const int32 TotalNumFrames = SVTAsset->GetNumFrames();		
+		const int32 StartFrame = OutputVolumeTexture->StartTime * FrameRate;
+		const int32 BlendFrames = OutputVolumeTexture->BlendDuration * FrameRate;
+		const int32 LoopedFrames = TotalNumFrames - StartFrame - BlendFrames;
+		const int32 BlendStartFrame = TotalNumFrames - BlendFrames;
+
+		if (TotalNumFrames <= 0)
+		{
+			UE_LOG(LogNiagaraBaker, Error, TEXT("SVT sequence to loop must have > 0 frames"));
+			return;
+		}
+
+		if (StartFrame < 0)
+		{
+			UE_LOG(LogNiagaraBaker, Error, TEXT("Start frame must be greater than 0"));
+			return;
+		}
+
+		if (BlendFrames < 0)
+		{
+			UE_LOG(LogNiagaraBaker, Error, TEXT("Blend amount must be greater than 0"));
+			return;
+		}
+
+		if (LoopedFrames < 0)
+		{
+			UE_LOG(LogNiagaraBaker, Error, TEXT("Cannot loop SVT due to insufficient frames.  Please reduce blend amount and start time or bake more frames."));
+			return;
+		}
+
+		for (int i = 0; i < LoopedFrames; ++i)
+		{
+			int32 OutFrameA = StartFrame + BlendFrames + i;
+			float LerpAmount = FMath::Clamp((1.0 * OutFrameA - BlendStartFrame) / BlendFrames, 0.0f, 1.0f);
+
+			int32 OutFrameB = StartFrame + 1.0 * LerpAmount * BlendFrames;
+
+			//
+			// nonblocking read SVT frames
+			//			
+			const int32 MipLevel = 0;
+			const float SVTFrameRate = 0.0f;
+			const bool bBlocking = true;
+			const bool bHasFrameRate = false;
+
+			USparseVolumeTextureFrame* SparseVolumeTextureFrameA = 
+				USparseVolumeTextureFrame::GetFrameAndIssueStreamingRequest(SVTAsset, GetTypeHash(InBakerOutput), SVTFrameRate, OutFrameA, MipLevel, bBlocking, bHasFrameRate);
+
+			USparseVolumeTextureFrame* SparseVolumeTextureFrameB =
+				USparseVolumeTextureFrame::GetFrameAndIssueStreamingRequest(SVTAsset, GetTypeHash(InBakerOutput+12345), SVTFrameRate, OutFrameB, MipLevel, bBlocking, bHasFrameRate);
+
+			// The streaming manager normally ticks in FDeferredShadingSceneRenderer::Render(), but the SVT->DenseTexture conversion compute shader happens in a render command before that.
+			// At execution time of that command, the streamer hasn't had the chance to do any streaming yet, so we force another tick here.
+			// Assuming blocking requests are used, this guarantees that the requested frame is fully streamed in (if there is memory available).
+			UE::SVT::GetStreamingManager().Update_GameThread();
+
+			if (SparseVolumeTextureFrameA == nullptr)
+			{
+				UE_LOG(LogNiagaraBaker, Error, TEXT("Invalid frame from baked SVT for looping"));
+				return;
+			}
+
+			if (SparseVolumeTextureFrameB == nullptr)
+			{
+				UE_LOG(LogNiagaraBaker, Error, TEXT("Invalid frame from baked SVT for looping"));
+				return;
+			}
+
+			FIntVector VolumeResolution = SparseVolumeTextureFrameA->GetVolumeResolution();
+			EPixelFormat VolumeFormat = SparseVolumeTextureFrameA->GetFormat(0);
+			
+			//
+			// Perform blend and output results
+			//
+
+			TArray<uint8> TextureData;		
+			ENQUEUE_RENDER_COMMAND(NDIRenderTargetVolumeUpdate)
+				(
+					[RT_TextureData = &TextureData, RT_VolumeResolution = VolumeResolution, RT_VolumeFormat = VolumeFormat, RT_LerpAmount = LerpAmount,
+					RT_SVTRenderResources_A = SparseVolumeTextureFrameA->GetTextureRenderResources(),
+					RT_SVTRenderResources_B = SparseVolumeTextureFrameB->GetTextureRenderResources()](FRHICommandListImmediate& RHICmdList)
+					{
+						if (RT_SVTRenderResources_A == nullptr)
+						{
+							UE_LOG(LogNiagaraBaker, Error, TEXT("Null svt resource"));
+							return;
+						}
+
+						if (RT_SVTRenderResources_B == nullptr)
+						{
+							UE_LOG(LogNiagaraBaker, Error, TEXT("Null svt resource"));
+							return;
+						}
+
+						//
+						// execute compute shader to output blended result
+						//
+						FRDGBuilder GraphBuilder(RHICmdList);
+						
+						TShaderMapRef<FNiagaraBlendSVTsToDenseBufferCS> ComputeShader(GetGlobalShaderMap(GMaxRHIFeatureLevel));
+						FNiagaraBlendSVTsToDenseBufferCS::FParameters* PassParameters = GraphBuilder.AllocParameters<FNiagaraBlendSVTsToDenseBufferCS::FParameters>();
+
+						FRDGTextureDesc TempTextureDesc = FRDGTextureDesc::Create3D(
+							RT_VolumeResolution,
+							RT_VolumeFormat,
+							FClearValueBinding::Black,
+							ETextureCreateFlags::ShaderResource | ETextureCreateFlags::RenderTargetable | ETextureCreateFlags::UAV);
+
+						FRDGTextureRef TempTexture = GraphBuilder.CreateTexture(TempTextureDesc, TEXT("TempOutput"));
+						PassParameters->DestinationBuffer = GraphBuilder.CreateUAV(TempTexture);
+					
+						FRHITexture* PageTableTexture_A = RT_SVTRenderResources_A->GetPageTableTexture();
+						FRHITexture* TextureA_A = RT_SVTRenderResources_A->GetPhysicalTileDataATexture();
+						if (PageTableTexture_A == nullptr || TextureA_A == nullptr)
+						{
+							UE_LOG(LogNiagaraBaker, Error, TEXT("Null svt texture"));
+							return;
+						}
+
+						// build pass parameters for the "A" SVT frame
+						FUintVector4 CurrentPackedUniforms0_A = FUintVector4();
+						FUintVector4 CurrentPackedUniforms1_A = FUintVector4();
+						RT_SVTRenderResources_A->GetPackedUniforms(CurrentPackedUniforms0_A, CurrentPackedUniforms1_A);
+
+						PassParameters->TileDataTextureSampler_A = TStaticSamplerState<SF_Point, AM_Clamp, AM_Clamp, AM_Clamp>::GetRHI();
+						PassParameters->SparseVolumeTexturePageTable_A = PageTableTexture_A;
+						PassParameters->SparseVolumeTextureA_A = TextureA_A;
+						PassParameters->PackedSVTUniforms0_A = CurrentPackedUniforms0_A;
+						PassParameters->PackedSVTUniforms1_A = CurrentPackedUniforms1_A;
+						PassParameters->TextureSize_A = RT_VolumeResolution;
+						PassParameters->MipLevels_A = 0;
+
+						// build pass parameters for the "B" SVT frame
+						FRHITexture* PageTableTexture_B = RT_SVTRenderResources_B->GetPageTableTexture();
+						FRHITexture* TextureA_B = RT_SVTRenderResources_B->GetPhysicalTileDataATexture();
+
+						FUintVector4 CurrentPackedUniforms0_B = FUintVector4();
+						FUintVector4 CurrentPackedUniforms1_B = FUintVector4();
+						RT_SVTRenderResources_B->GetPackedUniforms(CurrentPackedUniforms0_B, CurrentPackedUniforms1_B);
+
+						PassParameters->TileDataTextureSampler_B = TStaticSamplerState<SF_Point, AM_Clamp, AM_Clamp, AM_Clamp>::GetRHI();
+						PassParameters->SparseVolumeTexturePageTable_B = PageTableTexture_B;
+						PassParameters->SparseVolumeTextureA_B = TextureA_B;
+						PassParameters->PackedSVTUniforms0_B = CurrentPackedUniforms0_B;
+						PassParameters->PackedSVTUniforms1_B = CurrentPackedUniforms1_B;
+						PassParameters->TextureSize_B = RT_VolumeResolution;
+						PassParameters->MipLevels_B = 0;
+
+						PassParameters->LerpAmount = RT_LerpAmount;
+
+						FIntVector ThreadGroupSize = FNiagaraShader::GetDefaultThreadGroupSize(ENiagaraGpuDispatchType::ThreeD);
+						const FIntVector NumThreadGroups(
+							FMath::DivideAndRoundUp(RT_VolumeResolution.X, ThreadGroupSize.X),
+							FMath::DivideAndRoundUp(RT_VolumeResolution.Y, ThreadGroupSize.Y),
+							FMath::DivideAndRoundUp(RT_VolumeResolution.Z, ThreadGroupSize.Z)
+						);
+
+						GraphBuilder.AddPass(
+							// Friendly name of the pass for profilers using printf semantics.
+							RDG_EVENT_NAME("Blend SVTs"),
+							// Parameters provided to RDG.
+							PassParameters,
+							// Issues compute commands.
+							ERDGPassFlags::Compute,
+							// This is deferred until Execute. May execute in parallel with other passes.
+							[PassParameters, ComputeShader, NumThreadGroups](FRDGAsyncTask, FRHIComputeCommandList& RHICmdList)
+							{
+								FComputeShaderUtils::Dispatch(RHICmdList, ComputeShader, *PassParameters, NumThreadGroups);
+							});
+			
+						//
+						// Readback dense texture		
+						//
+						FRHIGPUTextureReadback RenderTargetReadback("ReadVolumeTexture");						
+						AddEnqueueCopyPass(GraphBuilder, &RenderTargetReadback, TempTexture);
+						
+						// Execute the graph.
+						GraphBuilder.Execute();
+
+						RHICmdList.BlockUntilGPUIdle();
+
+						check(RenderTargetReadback.IsReady());
+
+						//Lock the readback staging texture
+						int32 RowPitchInPixels;
+						int32 BufferHeight;
+						const uint8* LockedData = (const uint8*)RenderTargetReadback.Lock(RowPitchInPixels, &BufferHeight);
+
+						if (LockedData == nullptr)
+						{
+							UE_LOG(LogNiagaraBaker, Error, TEXT("Readback failed and returned null locked data"));
+							return;
+						}
+
+						if (RowPitchInPixels == 0 || BufferHeight == 0)
+						{
+							UE_LOG(LogNiagaraBaker, Error, TEXT("Readback failed and returned data with zero pitch/buffer height"));
+							return;
+						}
+						
+						uint32 BlockBytes = GPixelFormats[RT_VolumeFormat].BlockBytes;
+						int32 Count = RT_VolumeResolution.X * RT_VolumeResolution.Y * RT_VolumeResolution.Z * BlockBytes;
+						RT_TextureData->AddUninitialized(Count);
+
+						if (RT_TextureData->Num() == 0)
+						{
+							UE_LOG(LogNiagaraBaker, Error, TEXT("Output looped texture data has no elements"));
+							return;
+						}
+						
+						const uint8* SliceStart = LockedData;
+						for (int32 Z = 0; Z < RT_VolumeResolution.Z; ++Z)
+						{
+							const uint8* RowStart = SliceStart;
+							for (int32 Y = 0; Y < RT_VolumeResolution.Y; ++Y)
+							{
+								int32 Offset = 0 + Y * RT_VolumeResolution.X + Z * RT_VolumeResolution.X * RT_VolumeResolution.Y;
+								FMemory::Memcpy(RT_TextureData->GetData() + Offset * BlockBytes, RowStart, BlockBytes * RT_VolumeResolution.X);
+
+								RowStart += RowPitchInPixels * BlockBytes;
+							}
+
+							SliceStart += BufferHeight * RowPitchInPixels * BlockBytes;
+						}						
+
+						//Unlock the staging texture
+						RenderTargetReadback.Unlock();					
+					}					
+				);
+
+			FlushRenderingCommands();
+
+			if (TextureData.Num() == 0)
+			{
+				UE_LOG(LogNiagaraBaker, Error, TEXT("Readback failed when trying to add looped render target to SVT"));
+				return;
+			}
+
+			//
+			// add dense texture to LoopedSVTAsset
+			//
+#if WITH_EDITOR			
+			UE::SVT::FTextureDataCreateInfo SVTCreateInfo;
+			SVTCreateInfo.VirtualVolumeAABBMin = FIntVector3::ZeroValue;
+			SVTCreateInfo.VirtualVolumeAABBMax = VolumeResolution;
+			SVTCreateInfo.FallbackValues[0] = FVector4f(0, 0, 0, 0);
+			SVTCreateInfo.FallbackValues[1] = FVector4f(0, 0, 0, 0);
+			SVTCreateInfo.AttributesFormats[0] = VolumeFormat;
+			SVTCreateInfo.AttributesFormats[1] = PF_Unknown;
+
+			UE::SVT::FTextureData SparseTextureData{};
+			bool Success = SparseTextureData.CreateFromDense(SVTCreateInfo, TArrayView<uint8, int64>((uint8*)TextureData.GetData(), (int64)TextureData.Num() * sizeof(TextureData[0])), TArrayView<uint8>());
+
+			if (!Success)
+			{
+				UE_LOG(LogNiagaraBaker, Error, TEXT("Cannot create looped SVT output"));
+				return;
+			}
+
+			FTransform TransformToUse = SparseVolumeTextureFrameA->GetFrameTransform();
+
+			if (!LoopedSVTAsset->AppendFrame(SparseTextureData, TransformToUse))
+			{
+				UE_LOG(LogNiagaraBaker, Error, TEXT("Cannot append frame to looped SVT"));
+				return;
+			}
+#endif
+		}		
+
+		if (!LoopedSVTAsset->EndInitialize())
+		{
+			UE_LOG(LogNiagaraBaker, Error, TEXT("Cannot EndInitialize on creating looped SVT"));
+		}
+
+		LoopedSVTAsset->PostLoad();
+	}
 }
 
diff --git a/Engine/Plugins/FX/Niagara/Source/NiagaraEditor/Private/NiagaraBakerRendererOutputSparseVolumeTexture.h b/Engine/Plugins/FX/Niagara/Source/NiagaraEditor/Private/NiagaraBakerRendererOutputSparseVolumeTexture.h
--- a/Engine/Plugins/FX/Niagara/Source/NiagaraEditor/Private/NiagaraBakerRendererOutputSparseVolumeTexture.h
+++ b/Engine/Plugins/FX/Niagara/Source/NiagaraEditor/Private/NiagaraBakerRendererOutputSparseVolumeTexture.h
@@ -25,4 +25,5 @@ public:
 protected:
 
 	TObjectPtr<UAnimatedSparseVolumeTexture> SVTAsset = nullptr;	
+	TObjectPtr<UAnimatedSparseVolumeTexture> LoopedSVTAsset = nullptr;
 };
diff --git a/Engine/Plugins/FX/Niagara/Source/NiagaraEditor/Private/ViewModels/NiagaraBakerViewModel.cpp b/Engine/Plugins/FX/Niagara/Source/NiagaraEditor/Private/ViewModels/NiagaraBakerViewModel.cpp
--- a/Engine/Plugins/FX/Niagara/Source/NiagaraEditor/Private/ViewModels/NiagaraBakerViewModel.cpp
+++ b/Engine/Plugins/FX/Niagara/Source/NiagaraEditor/Private/ViewModels/NiagaraBakerViewModel.cpp
@@ -99,6 +99,22 @@ void FNiagaraBakerViewModel::ToggleRenderComponentOnly()
 	}
 }
 
+bool FNiagaraBakerViewModel::PreviewLoopedOutput() const
+{
+	UNiagaraBakerSettings* BakerSettings = GetBakerSettings();
+	return BakerSettings ? BakerSettings->bPreviewLoopedOutput : false;
+}
+
+void FNiagaraBakerViewModel::TogglePreviewLoopedOutput()
+{
+	if (UNiagaraBakerSettings* BakerSettings = GetBakerSettings())
+	{
+		const FScopedTransaction Transaction(LOCTEXT("SetRenderComponentOnly", "Set Render Component Only"));
+		BakerSettings->Modify();
+		BakerSettings->bPreviewLoopedOutput = !BakerSettings->bPreviewLoopedOutput;
+	}
+}
+
 void FNiagaraBakerViewModel::SetCameraSettingsIndex(int CamerSettingsIndex)
 {
 	if (UNiagaraBakerSettings* BakerSettings = GetBakerSettings())
diff --git a/Engine/Plugins/FX/Niagara/Source/NiagaraEditor/Private/ViewModels/NiagaraBakerViewModel.h b/Engine/Plugins/FX/Niagara/Source/NiagaraEditor/Private/ViewModels/NiagaraBakerViewModel.h
--- a/Engine/Plugins/FX/Niagara/Source/NiagaraEditor/Private/ViewModels/NiagaraBakerViewModel.h
+++ b/Engine/Plugins/FX/Niagara/Source/NiagaraEditor/Private/ViewModels/NiagaraBakerViewModel.h
@@ -64,6 +64,9 @@ public:
 	bool ShowRenderComponentOnly() const;
 	void ToggleRenderComponentOnly();
 
+	bool PreviewLoopedOutput() const;
+	void TogglePreviewLoopedOutput();
+
 	void SetCameraSettingsIndex(int CamerSettingsIndex);
 	bool IsCameraSettingIndex(int CamerSettingsIndex) const;
 
diff --git a/Engine/Plugins/FX/Niagara/Source/NiagaraEditor/Private/Widgets/SNiagaraBakerWidget.cpp b/Engine/Plugins/FX/Niagara/Source/NiagaraEditor/Private/Widgets/SNiagaraBakerWidget.cpp
--- a/Engine/Plugins/FX/Niagara/Source/NiagaraEditor/Private/Widgets/SNiagaraBakerWidget.cpp
+++ b/Engine/Plugins/FX/Niagara/Source/NiagaraEditor/Private/Widgets/SNiagaraBakerWidget.cpp
@@ -667,6 +667,19 @@ TSharedRef<SWidget> SNiagaraBakerWidget::MakeSettingsMenu()
 		EUserInterfaceActionType::ToggleButton
 	);
 
+	MenuBuilder.AddMenuEntry(
+		LOCTEXT("ShowNiagaraOnly", "Preview Looped Output"),
+		LOCTEXT("ShowNiagaraOnlyTooltip", "Renders only the Niagara System when enabled."),
+		FSlateIcon(),
+		FUIAction(
+			FExecuteAction::CreateSP(ViewModel, &FNiagaraBakerViewModel::TogglePreviewLoopedOutput),
+			FCanExecuteAction(),
+			FIsActionChecked::CreateSP(ViewModel, &FNiagaraBakerViewModel::PreviewLoopedOutput)
+		),
+		NAME_None,
+		EUserInterfaceActionType::ToggleButton
+	);
+
 	MenuBuilder.AddSubMenu(
 		LOCTEXT("BakeQualityLevel", "Bake Quality Level"),
 		LOCTEXT("BakeQualityLevelToolTip", "Allows you to optionally override the quality level when baking, i.e. use cinematic quality while baking"),
diff --git a/Engine/Plugins/FX/Niagara/Source/NiagaraShader/Private/NiagaraSVTShaders.cpp b/Engine/Plugins/FX/Niagara/Source/NiagaraShader/Private/NiagaraSVTShaders.cpp
--- a/Engine/Plugins/FX/Niagara/Source/NiagaraShader/Private/NiagaraSVTShaders.cpp
+++ b/Engine/Plugins/FX/Niagara/Source/NiagaraShader/Private/NiagaraSVTShaders.cpp
@@ -8,3 +8,19 @@ void FNiagaraCopySVTToDenseBufferCS::ModifyCompilationEnvironment(const FGlobalS
 }
 
 IMPLEMENT_GLOBAL_SHADER(FNiagaraCopySVTToDenseBufferCS, "/Plugin/FX/Niagara/Private/NiagaraSVTToDenseBuffer.usf", "PerformCopyCS", SF_Compute);
+
+void FNiagaraBlendSVTsToDenseBufferCS::ModifyCompilationEnvironment(const FGlobalShaderPermutationParameters& Parameters, FShaderCompilerEnvironment& OutEnvironment)
+{
+	FGlobalShader::ModifyCompilationEnvironment(Parameters, OutEnvironment);
+}
+
+bool FNiagaraBlendSVTsToDenseBufferCS::ShouldCompilePermutation(const FGlobalShaderPermutationParameters& Parameters)
+{
+#if WITH_EDITOR
+	return true;
+#else
+	return false;
+#endif
+}
+
+IMPLEMENT_GLOBAL_SHADER(FNiagaraBlendSVTsToDenseBufferCS, "/Plugin/FX/Niagara/Private/NiagaraBlendSVTsToDenseBuffer.usf", "PerformBlendCS", SF_Compute);
\ No newline at end of file
diff --git a/Engine/Plugins/FX/Niagara/Source/NiagaraShader/Public/NiagaraSVTShaders.h b/Engine/Plugins/FX/Niagara/Source/NiagaraShader/Public/NiagaraSVTShaders.h
--- a/Engine/Plugins/FX/Niagara/Source/NiagaraShader/Public/NiagaraSVTShaders.h
+++ b/Engine/Plugins/FX/Niagara/Source/NiagaraShader/Public/NiagaraSVTShaders.h
@@ -23,3 +23,35 @@ class FNiagaraCopySVTToDenseBufferCS : public FGlobalShader
 		SHADER_PARAMETER(int32, MipLevel)
 	END_SHADER_PARAMETER_STRUCT()
 };
+
+class FNiagaraBlendSVTsToDenseBufferCS : public FGlobalShader
+{
+	DECLARE_EXPORTED_GLOBAL_SHADER(FNiagaraBlendSVTsToDenseBufferCS, NIAGARASHADER_API);
+	SHADER_USE_PARAMETER_STRUCT(FNiagaraBlendSVTsToDenseBufferCS, FGlobalShader);
+
+	static void ModifyCompilationEnvironment(const FGlobalShaderPermutationParameters& Parameters, FShaderCompilerEnvironment& OutEnvironment);
+	static bool ShouldCompilePermutation(const FGlobalShaderPermutationParameters& Parameters);
+
+	BEGIN_SHADER_PARAMETER_STRUCT(FParameters, NIAGARASHADER_API)
+		SHADER_PARAMETER_RDG_TEXTURE_UAV(RWTexture3D<float4>, DestinationBuffer)
+
+		SHADER_PARAMETER_SAMPLER(SamplerState, TileDataTextureSampler_A)
+		SHADER_PARAMETER_TEXTURE(Texture3D<uint>, SparseVolumeTexturePageTable_A)
+		SHADER_PARAMETER_TEXTURE(Texture3D, SparseVolumeTextureA_A)
+		SHADER_PARAMETER(FUintVector4, PackedSVTUniforms0_A)
+		SHADER_PARAMETER(FUintVector4, PackedSVTUniforms1_A)
+		SHADER_PARAMETER(FIntVector, TextureSize_A)
+		SHADER_PARAMETER(int32, MipLevels_A)
+
+		SHADER_PARAMETER_SAMPLER(SamplerState, TileDataTextureSampler_B)
+		SHADER_PARAMETER_TEXTURE(Texture3D<uint>, SparseVolumeTexturePageTable_B)
+		SHADER_PARAMETER_TEXTURE(Texture3D, SparseVolumeTextureA_B)
+		SHADER_PARAMETER(FUintVector4, PackedSVTUniforms0_B)
+		SHADER_PARAMETER(FUintVector4, PackedSVTUniforms1_B)
+		SHADER_PARAMETER(FIntVector, TextureSize_B)
+		SHADER_PARAMETER(int32, MipLevels_B)
+
+		SHADER_PARAMETER(float, LerpAmount)
+						
+	END_SHADER_PARAMETER_STRUCT()
+};
