Fix ensure caused by activating a new NS during activation flow of another.
An edge case flow inside Scalability culling and NDC initialization could cause this.

diff --git a/Engine/Plugins/FX/Niagara/Source/Niagara/Private/DataInterface/NiagaraDataInterfaceDataChannelRead.cpp b/Engine/Plugins/FX/Niagara/Source/Niagara/Private/DataInterface/NiagaraDataInterfaceDataChannelRead.cpp
--- a/Engine/Plugins/FX/Niagara/Source/Niagara/Private/DataInterface/NiagaraDataInterfaceDataChannelRead.cpp
+++ b/Engine/Plugins/FX/Niagara/Source/Niagara/Private/DataInterface/NiagaraDataInterfaceDataChannelRead.cpp
@@ -482,7 +482,7 @@ bool FNDIDataChannelReadInstanceData::Tick(UNiagaraDataInterfaceDataChannelRead*
 		if (DataChannelPtr)
 		{
 			bool bNDCDataIsValid = DataChannelData && DataChannelData->IsLayoutValid(DataChannelPtr);
-			if(bNDCDataIsValid == false || Interface->bUpdateSourceDataEveryTick)
+			if(!bIsInit && (bNDCDataIsValid == false || Interface->bUpdateSourceDataEveryTick))
 			{
 				//TODO: Automatically modify tick group if we have DIs that require current frame info?
 				FNiagaraDataChannelSearchParameters SearchParams(Instance->GetAttachComponent());
