fix cis loctext

diff --git a/Engine/Plugins/FX/Niagara/Source/NiagaraEditor/Private/Widgets/SNiagaraBakerWidget.cpp b/Engine/Plugins/FX/Niagara/Source/NiagaraEditor/Private/Widgets/SNiagaraBakerWidget.cpp
--- a/Engine/Plugins/FX/Niagara/Source/NiagaraEditor/Private/Widgets/SNiagaraBakerWidget.cpp
+++ b/Engine/Plugins/FX/Niagara/Source/NiagaraEditor/Private/Widgets/SNiagaraBakerWidget.cpp
@@ -668,8 +668,8 @@ TSharedRef<SWidget> SNiagaraBakerWidget::MakeSettingsMenu()
 	);
 
 	MenuBuilder.AddMenuEntry(
-		LOCTEXT("ShowNiagaraOnly", "Preview Looped Output"),
-		LOCTEXT("ShowNiagaraOnlyTooltip", "Renders only the Niagara System when enabled."),
+		LOCTEXT("ShowLoopedOutput", "Preview Looped Output"),
+		LOCTEXT("ShowLoopedOutputTooltip", "Renders only the Niagara System when enabled."),
 		FSlateIcon(),
 		FUIAction(
 			FExecuteAction::CreateSP(ViewModel, &FNiagaraBakerViewModel::TogglePreviewLoopedOutput),
