- Defer the initial emitter state setup for stateless until first tick to ensure any bindings from the system spawn / update script have run

#test Add a user / system parameter to delay the loop start, set to none zero and see it was ignored

#virtualized

diff --git a/Engine/Plugins/FX/Niagara/Source/Niagara/Internal/Stateless/NiagaraStatelessEmitterInstance.h b/Engine/Plugins/FX/Niagara/Source/Niagara/Internal/Stateless/NiagaraStatelessEmitterInstance.h
--- a/Engine/Plugins/FX/Niagara/Source/Niagara/Internal/Stateless/NiagaraStatelessEmitterInstance.h
+++ b/Engine/Plugins/FX/Niagara/Source/Niagara/Internal/Stateless/NiagaraStatelessEmitterInstance.h
@@ -112,6 +112,7 @@ private:
 	uint32										bEmitterEnabled_GT : 1 = true;
 	uint32										bEmitterEnabled_CNC : 1 = true;
 	uint32										bSpawnInfosDirty : 1 = false;
+	uint32										bNeedsEmitterStateInit : 1 = false;
 
 	int32										RandomSeed = 0;
 	FRandomStream								RandomStream;
diff --git a/Engine/Plugins/FX/Niagara/Source/Niagara/Private/Stateless/NiagaraStatelessEmitterInstance.cpp b/Engine/Plugins/FX/Niagara/Source/Niagara/Private/Stateless/NiagaraStatelessEmitterInstance.cpp
--- a/Engine/Plugins/FX/Niagara/Source/Niagara/Private/Stateless/NiagaraStatelessEmitterInstance.cpp
+++ b/Engine/Plugins/FX/Niagara/Source/Niagara/Private/Stateless/NiagaraStatelessEmitterInstance.cpp
@@ -128,8 +128,7 @@ void FNiagaraStatelessEmitterInstance::Init(int32 InEmitterIndex)
 
 	GPUDataBufferInterfaces = RenderThreadDataPtr.Get();
 
-	InitEmitterState();
-	InitSpawnInfos(0.0f);
+	bNeedsEmitterStateInit = true;
 }
 
 void FNiagaraStatelessEmitterInstance::ResetSimulation(bool bKillExisting)
@@ -165,8 +164,7 @@ void FNiagaraStatelessEmitterInstance::ResetSimulation(bool bKillExisting)
 	Age = 0.0f;
 	bEmitterEnabled_CNC = bEmitterEnabled_GT;
 
-	InitEmitterState();
-	InitSpawnInfos(0.0f);
+	bNeedsEmitterStateInit = true;
 
 	InternalExecutionState = ENiagaraExecutionState::Active;
 	ExecutionState = InternalExecutionState;
@@ -272,6 +270,13 @@ bool FNiagaraStatelessEmitterInstance::ShouldTick() const
 
 void FNiagaraStatelessEmitterInstance::Tick(float DeltaSeconds)
 {
+	if ( bNeedsEmitterStateInit )
+	{
+		bNeedsEmitterStateInit = false;
+		InitEmitterState();
+		InitSpawnInfos(0.0f);
+	}
+
 	Age += DeltaSeconds;
 
 	TickSpawnInfos();
@@ -284,7 +289,6 @@ void FNiagaraStatelessEmitterInstance::InitEmitterState()
 {
 	using namespace NiagaraStatelessEmitterInstancePrivate;
 
-
 	const FNiagaraEmitterStateData& EmitterState = EmitterData->EmitterState;
 	LoopCount			= 0;
 
