- Optimize socket reading when in filtered / unfiltered mode

diff --git a/Engine/Plugins/FX/Niagara/Source/Niagara/Private/DataInterface/NiagaraDataInterfaceSocketReader.cpp b/Engine/Plugins/FX/Niagara/Source/Niagara/Private/DataInterface/NiagaraDataInterfaceSocketReader.cpp
--- a/Engine/Plugins/FX/Niagara/Source/Niagara/Private/DataInterface/NiagaraDataInterfaceSocketReader.cpp
+++ b/Engine/Plugins/FX/Niagara/Source/Niagara/Private/DataInterface/NiagaraDataInterfaceSocketReader.cpp
@@ -5,6 +5,7 @@
 #include "GameFramework/Pawn.h"
 #include "NiagaraCompileHashVisitor.h"
 #include "NiagaraComponent.h"
+#include "NiagaraDataInterfaceUtilities.h"
 #include "NiagaraTypes.h"
 #include "NiagaraShaderParametersBuilder.h"
 #include "NiagaraSystemInstance.h"
@@ -28,6 +29,14 @@
 
 namespace NDISocketReaderLocal
 {
+	static bool GDisableTransformReadType = 0;
+	static FAutoConsoleVariableRef CVarDisableTransformReadType(
+		TEXT("fx.Niagara.SocketReaderDI.DisableTransformReadType"),
+		GDisableTransformReadType,
+		TEXT("Disables optimization around removing reading transforms that we don't consume."),
+		ECVF_Default
+	);
+
 	static const TCHAR* TemplateShaderFile = TEXT("/Plugin/FX/Niagara/Private/NiagaraDataInterfaceSocketReaderTemplate.ush");
 
 	static const FName	IsValidName("IsValid");
@@ -49,9 +58,11 @@ namespace NDISocketReaderLocal
 	struct FInstanceData_GameThread
 	{
 		FNiagaraParameterDirectBinding<UObject*>	UserParamBinding;
-		bool										bIsDataValid = false;
-		bool										bNeedsSocketRecache = true;
-		bool										bNeedsRenderUpdate = true;
+		uint32										bIsDataValid : 1 = false;
+		uint32										bNeedsSocketRecache : 1 = true;
+		uint32										bNeedsRenderUpdate : 1 = true;
+		uint32										bReadsFilteredTransforms : 1 = false;
+		uint32										bReadsUnfilteredTransforms : 1 = false;
 
 		float										DeltaSeconds = 0.0f;
 		float										InvDeltaSeconds = 0.0f;
@@ -416,7 +427,8 @@ namespace NDISocketReaderLocal
 		}
 
 		// Update the socket data
-		if (bNeedsSocketUpdate && InstanceData->NumSockets > 0)
+		const bool bReadsAnyTransformData = InstanceData->bReadsFilteredTransforms || InstanceData->bReadsUnfilteredTransforms;
+		if (bReadsAnyTransformData && bNeedsSocketUpdate && InstanceData->NumSockets > 0)
 		{
 			// If we are not resetting copy the current transform to previous
 			if (!bResetPreviousTransform)
@@ -427,10 +439,25 @@ namespace NDISocketReaderLocal
 			if (USceneComponent* SceneComponent = Cast<USceneComponent>(InstanceData->ResolvedObject))
 			{
 				InstanceData->bNeedsRenderUpdate = true;
-				for (int32 i = 0; i < InstanceData->NumSockets; ++i)
+
+				if (InstanceData->bReadsFilteredTransforms && InstanceData->bReadsUnfilteredTransforms)
+				{
+					for (int32 SocketIndex=0; SocketIndex < InstanceData->NumSockets; ++SocketIndex)
+					{
+						const FTransform SocketTransform = SceneComponent->GetSocketTransform(InstanceData->SocketNames[SocketIndex], ERelativeTransformSpace::RTS_Component);
+						InstanceData->SocketTransforms[SocketIndex] = FTransform3f(SocketTransform);
+					}
+				}
+				else
 				{
-					const FTransform SocketTransform			= SceneComponent->GetSocketTransform(InstanceData->SocketNames[i], ERelativeTransformSpace::RTS_Component);
-					InstanceData->SocketTransforms[i]			= FTransform3f(SocketTransform);
+					const int32 SocketTableStart = InstanceData->bReadsFilteredTransforms ? 0 : InstanceData->NumFilteredSockets;
+					const int32 SocketTableEnd   = SocketTableStart + (InstanceData->bReadsFilteredTransforms ? InstanceData->NumFilteredSockets : InstanceData->NumUnfilteredSockets);
+					for (int32 i=SocketTableStart; i < SocketTableEnd; ++i)
+					{
+						const int32 SocketIndex = InstanceData->SocketFilterUnfilteredIndex[i + 1];
+						const FTransform SocketTransform = SceneComponent->GetSocketTransform(InstanceData->SocketNames[SocketIndex], ERelativeTransformSpace::RTS_Component);
+						InstanceData->SocketTransforms[SocketIndex] = FTransform3f(SocketTransform);
+					}
 				}
 			}
 			else if (UStaticMesh* StaticMesh = Cast<UStaticMesh>(InstanceData->ResolvedObject))
@@ -520,6 +547,7 @@ namespace NDISocketReaderLocal
 		Default,
 		Filtered,
 		Unfiltered,
+		None,
 	};
 
 	void VMIsValid(FVectorVMExternalFunctionContext& Context)
@@ -635,19 +663,25 @@ namespace NDISocketReaderLocal
 		}
 	}
 
-	static const TMap<FName, FVMExternalFunction> VMFunctionBindings =
+	struct FVMFunctionInfo
+	{
+		FVMExternalFunction	FunctionBinding;
+		ESocketReadType		TransformReadType;
+	};
+
+	static const TMap<FName, FVMFunctionInfo> VMFunctionInfos =
 	{
-		{IsValidName,									FVMExternalFunction::CreateStatic(VMIsValid)},
-		{GetComponentToWorldName,						FVMExternalFunction::CreateStatic(VMGetComponentToWorld)},
-		{GetSocketCountName,							FVMExternalFunction::CreateStatic(VMGetSocketCount<ESocketReadType::Default>)},
-		{GetFilteredSocketCountName,					FVMExternalFunction::CreateStatic(VMGetSocketCount<ESocketReadType::Filtered>)},
-		{GetUnfilteredSocketCountName,					FVMExternalFunction::CreateStatic(VMGetSocketCount<ESocketReadType::Unfiltered>)},
-		{GetSocketTransformName,						FVMExternalFunction::CreateStatic(VMGetSocketTransform<ESocketReadType::Default, false>)},
-		{GetFilteredSocketTransformName,				FVMExternalFunction::CreateStatic(VMGetSocketTransform<ESocketReadType::Filtered, false>)},
-		{GetUnfilteredSocketTransformName,				FVMExternalFunction::CreateStatic(VMGetSocketTransform<ESocketReadType::Unfiltered, false>)},
-		{GetSocketTransformInterpolatedName,			FVMExternalFunction::CreateStatic(VMGetSocketTransform<ESocketReadType::Default, true>)},
-		{GetFilteredSocketTransformInterpolatedName,	FVMExternalFunction::CreateStatic(VMGetSocketTransform<ESocketReadType::Filtered, true>)},
-		{GetUnfilteredSocketTransformInterpolatedName,	FVMExternalFunction::CreateStatic(VMGetSocketTransform<ESocketReadType::Unfiltered, true>)},
+		{IsValidName,									FVMFunctionInfo{FVMExternalFunction::CreateStatic(VMIsValid), ESocketReadType::None}},
+		{GetComponentToWorldName,						FVMFunctionInfo{FVMExternalFunction::CreateStatic(VMGetComponentToWorld), ESocketReadType::None}},
+		{GetSocketCountName,							FVMFunctionInfo{FVMExternalFunction::CreateStatic(VMGetSocketCount<ESocketReadType::Default>), ESocketReadType::None}},
+		{GetFilteredSocketCountName,					FVMFunctionInfo{FVMExternalFunction::CreateStatic(VMGetSocketCount<ESocketReadType::Filtered>), ESocketReadType::None}},
+		{GetUnfilteredSocketCountName,					FVMFunctionInfo{FVMExternalFunction::CreateStatic(VMGetSocketCount<ESocketReadType::Unfiltered>), ESocketReadType::None}},
+		{GetSocketTransformName,						FVMFunctionInfo{FVMExternalFunction::CreateStatic(VMGetSocketTransform<ESocketReadType::Default, false>), ESocketReadType::Default}},
+		{GetFilteredSocketTransformName,				FVMFunctionInfo{FVMExternalFunction::CreateStatic(VMGetSocketTransform<ESocketReadType::Filtered, false>), ESocketReadType::Filtered}},
+		{GetUnfilteredSocketTransformName,				FVMFunctionInfo{FVMExternalFunction::CreateStatic(VMGetSocketTransform<ESocketReadType::Unfiltered, false>), ESocketReadType::Unfiltered}},
+		{GetSocketTransformInterpolatedName,			FVMFunctionInfo{FVMExternalFunction::CreateStatic(VMGetSocketTransform<ESocketReadType::Default, true>), ESocketReadType::Default}},
+		{GetFilteredSocketTransformInterpolatedName,	FVMFunctionInfo{FVMExternalFunction::CreateStatic(VMGetSocketTransform<ESocketReadType::Filtered, true>), ESocketReadType::Filtered}},
+		{GetUnfilteredSocketTransformInterpolatedName,	FVMFunctionInfo{FVMExternalFunction::CreateStatic(VMGetSocketTransform<ESocketReadType::Unfiltered, true>), ESocketReadType::Unfiltered}},
 	};
 }
 
@@ -722,10 +756,17 @@ void UNiagaraDataInterfaceSocketReader::GetFunctionsInternal(TArray<FNiagaraFunc
 }
 #endif
 
-void UNiagaraDataInterfaceSocketReader::GetVMExternalFunction(const FVMExternalFunctionBindingInfo& BindingInfo, void* InstanceData, FVMExternalFunction& OutFunc)
+void UNiagaraDataInterfaceSocketReader::GetVMExternalFunction(const FVMExternalFunctionBindingInfo& BindingInfo, void* PerInstanceData, FVMExternalFunction& OutFunc)
 {
 	using namespace NDISocketReaderLocal;
-	OutFunc = VMFunctionBindings.FindRef(BindingInfo.Name);
+	if (const FVMFunctionInfo* FunctionInfo = VMFunctionInfos.Find(BindingInfo.Name))
+	{
+		OutFunc = FunctionInfo->FunctionBinding;
+
+		FInstanceData_GameThread* InstanceData = static_cast<FInstanceData_GameThread*>(PerInstanceData);
+		InstanceData->bReadsFilteredTransforms		|= FunctionInfo->TransformReadType == ESocketReadType::Default || FunctionInfo->TransformReadType == ESocketReadType::Filtered;
+		InstanceData->bReadsUnfilteredTransforms	|= FunctionInfo->TransformReadType == ESocketReadType::Default || FunctionInfo->TransformReadType == ESocketReadType::Unfiltered;
+	}
 }
 
 #if WITH_EDITORONLY_DATA
@@ -795,6 +836,26 @@ bool UNiagaraDataInterfaceSocketReader::InitPerInstanceData(void* PerInstanceDat
 	InstanceData->UserParamBinding.Init(SystemInstance->GetInstanceParameters(), ObjectParameterBinding.Parameter);
 	ResolveObject(this, InstanceData, SystemInstance);
 
+	InstanceData->bReadsFilteredTransforms |= GDisableTransformReadType;
+	InstanceData->bReadsUnfilteredTransforms |= GDisableTransformReadType;
+
+	if (IsUsedWithGPUScript())
+	{
+		FNiagaraDataInterfaceUtilities::ForEachGpuFunction(
+			this, SystemInstance,
+			[&](const UNiagaraScript* Script, const FNiagaraDataInterfaceGeneratedFunction& Function)
+			{
+				if (const FVMFunctionInfo* FunctionInfo = VMFunctionInfos.Find(Function.DefinitionName))
+				{
+					InstanceData->bReadsFilteredTransforms |= FunctionInfo->TransformReadType == ESocketReadType::Default || FunctionInfo->TransformReadType == ESocketReadType::Filtered;
+					InstanceData->bReadsUnfilteredTransforms |= FunctionInfo->TransformReadType == ESocketReadType::Default || FunctionInfo->TransformReadType == ESocketReadType::Unfiltered;
+				}
+				// Keeping going unless there is no more information we can gather
+				return (InstanceData->bReadsFilteredTransforms && InstanceData->bReadsUnfilteredTransforms) != true;
+			}
+		);
+	}
+
 	return true;
 }
 
