Added color option for notes & revamped note UI.
Also made context actions more consistent with iconography

diff --git a/Engine/Plugins/FX/Niagara/Source/NiagaraEditor/Private/NiagaraEditorSettings.cpp b/Engine/Plugins/FX/Niagara/Source/NiagaraEditor/Private/NiagaraEditorSettings.cpp
--- a/Engine/Plugins/FX/Niagara/Source/NiagaraEditor/Private/NiagaraEditorSettings.cpp
+++ b/Engine/Plugins/FX/Niagara/Source/NiagaraEditor/Private/NiagaraEditorSettings.cpp
@@ -93,6 +93,11 @@ FLinearColor UNiagaraEditorSettings::GetSourceColor(EScriptSource  Source) const
 	return FLinearColor(1.f,1.f,1.f,0.3);
 }
 
+FLinearColor UNiagaraEditorSettings::GetDefaultNoteColor() const
+{
+	return DefaultNoteColor;
+}
+
 void UNiagaraEditorSettings::SetupNamespaceMetadata()
 {
 	DefaultNamespaceMetadata = FNiagaraNamespaceMetadata({ NAME_None })
diff --git a/Engine/Plugins/FX/Niagara/Source/NiagaraEditor/Private/NiagaraEditorStyle.cpp b/Engine/Plugins/FX/Niagara/Source/NiagaraEditor/Private/NiagaraEditorStyle.cpp
--- a/Engine/Plugins/FX/Niagara/Source/NiagaraEditor/Private/NiagaraEditorStyle.cpp
+++ b/Engine/Plugins/FX/Niagara/Source/NiagaraEditor/Private/NiagaraEditorStyle.cpp
@@ -417,7 +417,6 @@ void FNiagaraEditorStyle::InitIcons()
 	Set("NiagaraEditor.Module.Pin.TypeSelector", new IMAGE_BRUSH("Icons/Scratch", Icon16x16, FLinearColor::Gray));
 	Set("NiagaraEditor.Module.AddPin", new IMAGE_BRUSH("Icons/PlusSymbol_12x", Icon12x12, FLinearColor::Gray));
 	Set("NiagaraEditor.Module.RemovePin", new IMAGE_BRUSH("Icons/MinusSymbol_12x", Icon12x12, FLinearColor::Gray));
-	Set("NiagaraEditor.Message.CustomNote", new IMAGE_BRUSH("Icons/icon_custom_note_16x", Icon16x16));
 	Set("NiagaraEditor.Module.DynamicInput", new IMAGE_BRUSH_SVG("Icons/DynamicInput", Icon16x16));
 	Set("NiagaraEditor.Module.TypeIconPill", new IMAGE_BRUSH_SVG("Icons/TypeIconPill", FVector2D(5.0f, 16.0f)));
 
@@ -561,9 +560,7 @@ void FNiagaraEditorStyle::InitStackWidgets()
 	
 	FTextBlockStyle NoteHeaderStyle;
 	
-	FFontOutlineSettings NoteHeaderOutlineSettings;
-	NoteHeaderOutlineSettings.OutlineSize = 1;
-	NoteHeaderStyle.SetFont(DEFAULT_FONT("Regular", 12.f, NoteHeaderOutlineSettings));
+	NoteHeaderStyle.SetFont(DEFAULT_FONT("Regular", 11.f));
 	NoteHeaderStyle.SetColorAndOpacity(FSlateColor(FStyleColors::White));
 	FInlineEditableTextBlockStyle NoteHeaderInlineEditableStyle;
 	NoteHeaderInlineEditableStyle.SetTextStyle(NoteHeaderStyle);
diff --git a/Engine/Plugins/FX/Niagara/Source/NiagaraEditor/Private/NiagaraEditorUtilities.cpp b/Engine/Plugins/FX/Niagara/Source/NiagaraEditor/Private/NiagaraEditorUtilities.cpp
--- a/Engine/Plugins/FX/Niagara/Source/NiagaraEditor/Private/NiagaraEditorUtilities.cpp
+++ b/Engine/Plugins/FX/Niagara/Source/NiagaraEditor/Private/NiagaraEditorUtilities.cpp
@@ -4747,7 +4747,8 @@ TSharedRef<SToolTip> FNiagaraEditorUtilities::Tooltips::CreateStackNoteTooltip(U
 			.VAlign(VAlign_Center)
 			[
 				SNew(SImage)
-				.Image(FNiagaraEditorStyle::Get().GetBrush("NiagaraEditor.Message.CustomNote"))
+				.Image(FAppStyle::GetBrush("Icons.Comment"))
+				.ColorAndOpacity(StackNote.GetColor())
 			]
 		]
 		+ SHorizontalBox::Slot()
diff --git a/Engine/Plugins/FX/Niagara/Source/NiagaraEditor/Private/ViewModels/Stack/NiagaraStackEntry.cpp b/Engine/Plugins/FX/Niagara/Source/NiagaraEditor/Private/ViewModels/Stack/NiagaraStackEntry.cpp
--- a/Engine/Plugins/FX/Niagara/Source/NiagaraEditor/Private/ViewModels/Stack/NiagaraStackEntry.cpp
+++ b/Engine/Plugins/FX/Niagara/Source/NiagaraEditor/Private/ViewModels/Stack/NiagaraStackEntry.cpp
@@ -14,6 +14,7 @@
 #include "ScopedTransaction.h"
 #include "ViewModels/HierarchyEditor/NiagaraHierarchyViewModelBase.h"
 #include "ViewModels/HierarchyEditor/NiagaraSummaryViewViewModel.h"
+#include "ViewModels/Stack/NiagaraStackClipboardUtilities.h"
 
 #include UE_INLINE_GENERATED_CPP_BY_NAME(NiagaraStackEntry)
 
@@ -816,6 +817,7 @@ void UNiagaraStackEntry::RefreshChildren()
 		{
 			StackNote = NewObject<UNiagaraStackNote>(this);
 			StackNote->Initialize(CreateDefaultChildRequiredData(), GetStackEditorDataKey());
+			StackNote->IndentLevel = GetIndentLevel();
 			StackNote->OnNoteChanged().BindUObject(this, &UNiagaraStackEntry::SetStackNoteData);
 		}
 		
@@ -1192,6 +1194,67 @@ void UNiagaraStackEntry::DeleteStackNoteData()
 	}
 }
 
+void UNiagaraStackEntry::AddStackNote()
+{
+	if(SupportsStackNotes())
+	{
+		if(HasStackNoteData() == false)
+		{
+			FScopedTransaction Transaction(LOCTEXT("NewStackNoteAdded", "Added Stack Note"));
+			GetStackEditorData().Modify();
+		
+			FNiagaraStackNoteData NewStackNote;
+			NewStackNote.MessageHeader = LOCTEXT("DefaultStackNoteHeader", "Note");
+			NewStackNote.Message = LOCTEXT("DefaultStackNoteMessage", "Text");
+			NewStackNote.Color = GetDefault<UNiagaraEditorSettings>()->GetDefaultNoteColor();
+
+			GetStackEditorData().AddOrReplaceStackNote(GetStackEditorDataKey(), NewStackNote);
+			RefreshChildren();
+		
+			if(UNiagaraStackNote* Note = GetStackNote())
+			{
+				Note->SetIsRenamePending(true);
+			}
+		}
+	}
+}
+
+void UNiagaraStackEntry::ToggleStackNoteInlineDisplay()
+{
+	if(GetStackNote() != nullptr)
+	{
+		GetStackNote()->ToggleInlineDisplay();
+	}
+}
+
+void UNiagaraStackEntry::CopyStackNote()
+{
+	FNiagaraStackClipboardUtilities::CopyNote(this);
+}
+
+void UNiagaraStackEntry::PasteStackNote()
+{
+	if(SupportsStackNotes())
+	{
+		if(const UNiagaraClipboardContent* ClipboardContent = FNiagaraEditorModule::Get().GetClipboard().GetClipboardContent())
+		{
+			if(ClipboardContent->StackNote.IsValid())
+			{
+				FScopedTransaction Transaction(LOCTEXT("PasteNoteTransaction", "Note pasted"));
+				GetStackEditorData().Modify();
+				
+				SetStackNoteData(ClipboardContent->StackNote);
+
+				// if the note we pasted is not inlined, we expand the stack entry in order to show the note
+				if(GetStackNoteData().bInlineNote == false)
+				{
+					SetIsExpanded(true);
+				}
+			}
+		}
+	}
+}
+
 FNiagaraHierarchyIdentity UNiagaraStackEntry::DetermineSummaryIdentity() const
 {
 	return FNiagaraHierarchyIdentity();
diff --git a/Engine/Plugins/FX/Niagara/Source/NiagaraEditor/Private/ViewModels/Stack/NiagaraStackNote.cpp b/Engine/Plugins/FX/Niagara/Source/NiagaraEditor/Private/ViewModels/Stack/NiagaraStackNote.cpp
--- a/Engine/Plugins/FX/Niagara/Source/NiagaraEditor/Private/ViewModels/Stack/NiagaraStackNote.cpp
+++ b/Engine/Plugins/FX/Niagara/Source/NiagaraEditor/Private/ViewModels/Stack/NiagaraStackNote.cpp
@@ -1,6 +1,8 @@
 ﻿// Copyright Epic Games, Inc. All Rights Reserved.
 
 #include "ViewModels/Stack/NiagaraStackNote.h"
+
+#include "NiagaraEditorSettings.h"
 #include "NiagaraMessages.h"
 #include "ScopedTransaction.h"
 
@@ -33,10 +35,10 @@ void UNiagaraStackNote::ToggleInlineDisplay()
 	FScopedTransaction Transaction(LOCTEXT("ToggleNoteDisplayTransaction", "Toggled Note Display"));
 	GetStackEditorData().Modify();
 
-	FNiagaraStackNoteData UpdatedMessage = GetTargetStackNoteData().GetValue();
-	UpdatedMessage.bInlineNote = !UpdatedMessage.bInlineNote;
+	FNiagaraStackNoteData UpdatedStackNoteData = GetTargetStackNoteData().GetValue();
+	UpdatedStackNoteData.bInlineNote = !UpdatedStackNoteData.bInlineNote;
 
-	OnNoteChangedDelegate.ExecuteIfBound(UpdatedMessage);
+	OnNoteChangedDelegate.ExecuteIfBound(UpdatedStackNoteData);
 }
 
 void UNiagaraStackNote::DeleteTargetStackNote()
@@ -49,6 +51,25 @@ void UNiagaraStackNote::DeleteTargetStackNote()
 	OnNoteChangedDelegate.ExecuteIfBound(FNiagaraStackNoteData());
 }
 
+void UNiagaraStackNote::UpdateNoteColor(FLinearColor LinearColor)
+{
+	FScopedTransaction Transaction(LOCTEXT("ChangeNoteColorTransaction", "Changed Note Color"));
+	GetStackEditorData().Modify();
+
+	FNiagaraStackNoteData UpdatedStackNoteData = GetTargetStackNoteData().GetValue();
+
+	if(UpdatedStackNoteData.Color != LinearColor)
+	{
+		UpdatedStackNoteData.Color = LinearColor;
+		
+		OnNoteChangedDelegate.ExecuteIfBound(UpdatedStackNoteData);
+	}
+	else
+	{
+		Transaction.Cancel();
+	}
+}
+
 bool UNiagaraStackNote::GetShouldShowInStack() const
 {
 	if (GetTargetStackNoteData().IsSet())
@@ -105,4 +126,14 @@ bool UNiagaraStackNote::TestCanPasteWithMessage(const UNiagaraClipboardContent*
 	return false;
 }
 
+FLinearColor UNiagaraStackNote::GetColor() const
+{
+	if(GetTargetStackNoteData().IsSet())
+	{
+		return GetTargetStackNoteData()->Color.IsAlmostBlack() == false ? GetTargetStackNoteData()->Color : GetDefault<UNiagaraEditorSettings>()->GetDefaultNoteColor();
+	}
+
+	return GetDefault<UNiagaraEditorSettings>()->GetDefaultNoteColor(); 
+}
+
 #undef LOCTEXT_NAMESPACE
diff --git a/Engine/Plugins/FX/Niagara/Source/NiagaraEditor/Public/NiagaraEditorSettings.h b/Engine/Plugins/FX/Niagara/Source/NiagaraEditor/Public/NiagaraEditorSettings.h
--- a/Engine/Plugins/FX/Niagara/Source/NiagaraEditor/Public/NiagaraEditorSettings.h
+++ b/Engine/Plugins/FX/Niagara/Source/NiagaraEditor/Public/NiagaraEditorSettings.h
@@ -34,24 +34,6 @@ namespace FNiagaraEditorGuids
 	const extern FGuid StaticSwitchNamespaceMetaDataGuid;
 }
 
-USTRUCT()
-struct FNiagaraNewAssetDialogConfig
-{
-	GENERATED_BODY()
-
-	UPROPERTY()
-	int32 SelectedOptionIndex;
-
-	UPROPERTY()
-	FVector2D WindowSize;
-
-	FNiagaraNewAssetDialogConfig()
-	{
-		SelectedOptionIndex = 0;
-		WindowSize = FVector2D(450, 600);
-	}
-};
-
 UENUM()
 enum class ENiagaraCategoryExpandState : uint8
 {
@@ -475,6 +457,8 @@ public:
 
 	NIAGARAEDITOR_API FLinearColor GetSourceColor(EScriptSource Source) const;
 
+	NIAGARAEDITOR_API FLinearColor GetDefaultNoteColor() const;
+
 	NIAGARAEDITOR_API FNiagaraParameterPanelSectionStorage& FindOrAddParameterPanelSectionStorage(FGuid PanelSectionId, bool& bOutAdded);
 
 	DECLARE_DELEGATE_RetVal_OneParam(bool, FOnIsClassAllowed, const UClass* /*InClass*/);
@@ -577,6 +561,9 @@ private:
 	UPROPERTY(config, EditAnywhere, Category = "Niagara Colors")
 	FNiagaraActionColors ActionColors;
 
+	UPROPERTY(Config, EditAnywhere, Category ="Niagara Colors")
+	FLinearColor DefaultNoteColor = FLinearColor(FColor(56, 111, 75));
+
 	/** This is built using PlaybackSpeeds, populated whenever it is accessed using GetPlaybackSpeeds() */
 	mutable TOptional<TArray<float>> CachedPlaybackSpeeds;
 
diff --git a/Engine/Plugins/FX/Niagara/Source/NiagaraEditor/Public/NiagaraMessages.h b/Engine/Plugins/FX/Niagara/Source/NiagaraEditor/Public/NiagaraMessages.h
--- a/Engine/Plugins/FX/Niagara/Source/NiagaraEditor/Public/NiagaraMessages.h
+++ b/Engine/Plugins/FX/Niagara/Source/NiagaraEditor/Public/NiagaraMessages.h
@@ -268,6 +268,9 @@ struct FNiagaraStackNoteData
 
 	UPROPERTY()
 	FText MessageHeader;
+
+	UPROPERTY()
+	FLinearColor Color = FLinearColor::Black;
 	
 	UPROPERTY()
 	FText Message;
diff --git a/Engine/Plugins/FX/Niagara/Source/NiagaraEditor/Public/ViewModels/Stack/NiagaraStackEntry.h b/Engine/Plugins/FX/Niagara/Source/NiagaraEditor/Public/ViewModels/Stack/NiagaraStackEntry.h
--- a/Engine/Plugins/FX/Niagara/Source/NiagaraEditor/Public/ViewModels/Stack/NiagaraStackEntry.h
+++ b/Engine/Plugins/FX/Niagara/Source/NiagaraEditor/Public/ViewModels/Stack/NiagaraStackEntry.h
@@ -565,11 +565,16 @@ public:
 	NIAGARAEDITOR_API virtual void OnRenamed(FText NewName);
 
 	/** Generally, all stack items with a valid stack key would support the note feature. We only allow them on specific types of of entries. */
-	NIAGARAEDITOR_API virtual bool SupportsStackNotes() { return false; }
+	NIAGARAEDITOR_API virtual bool SupportsStackNotes() const { return false; }
 	NIAGARAEDITOR_API bool HasStackNoteData() const;
 	NIAGARAEDITOR_API FNiagaraStackNoteData GetStackNoteData() const;
 	NIAGARAEDITOR_API void SetStackNoteData(FNiagaraStackNoteData InStackNoteData);
 	NIAGARAEDITOR_API void DeleteStackNoteData();
+
+	NIAGARAEDITOR_API void AddStackNote();
+	NIAGARAEDITOR_API void ToggleStackNoteInlineDisplay();
+	NIAGARAEDITOR_API void CopyStackNote();
+	NIAGARAEDITOR_API void PasteStackNote();
 	
 	virtual bool SupportsSummaryView() const { return false; }
 	NIAGARAEDITOR_API virtual struct FNiagaraHierarchyIdentity DetermineSummaryIdentity() const;
@@ -612,6 +617,8 @@ protected:
 
 	NIAGARAEDITOR_API virtual int32 GetChildIndentLevel() const;
 
+	NIAGARAEDITOR_API void SetIndentLevelOverride(int32 IndentLevelOverride);
+
 	NIAGARAEDITOR_API virtual TOptional<FDropRequestResponse> CanDropInternal(const FDropRequest& DropRequest);
 
 	NIAGARAEDITOR_API virtual TOptional<FDropRequestResponse> DropInternal(const FDropRequest& DropRequest);
diff --git a/Engine/Plugins/FX/Niagara/Source/NiagaraEditor/Public/ViewModels/Stack/NiagaraStackFunctionInput.h b/Engine/Plugins/FX/Niagara/Source/NiagaraEditor/Public/ViewModels/Stack/NiagaraStackFunctionInput.h
--- a/Engine/Plugins/FX/Niagara/Source/NiagaraEditor/Public/ViewModels/Stack/NiagaraStackFunctionInput.h
+++ b/Engine/Plugins/FX/Niagara/Source/NiagaraEditor/Public/ViewModels/Stack/NiagaraStackFunctionInput.h
@@ -108,7 +108,7 @@ public:
 	NIAGARAEDITOR_API virtual FText GetTooltipText() const override;
 	NIAGARAEDITOR_API virtual bool GetIsEnabled() const override;
 	NIAGARAEDITOR_API virtual UObject* GetExternalAsset() const override;
-	virtual bool SupportsStackNotes() override { return true; }
+	virtual bool SupportsStackNotes() const override { return true; }
 	virtual bool SupportsCut() const override { return true; }
 	NIAGARAEDITOR_API virtual bool TestCanCutWithMessage(FText& OutMessage) const override;
 	NIAGARAEDITOR_API virtual FText GetCutTransactionText() const override;
diff --git a/Engine/Plugins/FX/Niagara/Source/NiagaraEditor/Public/ViewModels/Stack/NiagaraStackItem.h b/Engine/Plugins/FX/Niagara/Source/NiagaraEditor/Public/ViewModels/Stack/NiagaraStackItem.h
--- a/Engine/Plugins/FX/Niagara/Source/NiagaraEditor/Public/ViewModels/Stack/NiagaraStackItem.h
+++ b/Engine/Plugins/FX/Niagara/Source/NiagaraEditor/Public/ViewModels/Stack/NiagaraStackItem.h
@@ -63,7 +63,7 @@ public:
 	virtual bool TestCanResetToBaseWithMessage(FText& OutCanResetToBaseMessage) const { return false; }
 	virtual void ResetToBase() { }
 
-	virtual bool SupportsStackNotes() override { return true; }
+	virtual bool SupportsStackNotes() const override { return true; }
 	
 	virtual bool SupportsEditMode() const { return false; }
 	virtual void OnEditButtonClicked() { }
diff --git a/Engine/Plugins/FX/Niagara/Source/NiagaraEditor/Public/ViewModels/Stack/NiagaraStackNote.h b/Engine/Plugins/FX/Niagara/Source/NiagaraEditor/Public/ViewModels/Stack/NiagaraStackNote.h
--- a/Engine/Plugins/FX/Niagara/Source/NiagaraEditor/Public/ViewModels/Stack/NiagaraStackNote.h
+++ b/Engine/Plugins/FX/Niagara/Source/NiagaraEditor/Public/ViewModels/Stack/NiagaraStackNote.h
@@ -7,6 +7,7 @@
 #include "NiagaraMessages.h"
 #include "NiagaraStackEditorData.h"
 #include "NiagaraStackEntry.h"
+#include "Styling/SlateTypes.h"
 #include "NiagaraStackNote.generated.h"
 
 UCLASS(MinimalAPI)
@@ -24,16 +25,21 @@ public:
 	void RequestEditHeader() const { OnRequestEditHeaderDelegate.ExecuteIfBound(); }
 	NIAGARAEDITOR_API void ToggleInlineDisplay();
 	NIAGARAEDITOR_API void DeleteTargetStackNote();
-	
+	NIAGARAEDITOR_API void UpdateNoteColor(FLinearColor LinearColor);
+
 	FSimpleDelegate& OnRequestEditHeader() { return OnRequestEditHeaderDelegate; }
 
 	/** The delegate for the target stack entry to bind to. Forwards note changes which then need to be set via UNiagaraStackEntry::SetStackNoteData. */
 	FOnNoteChanged& OnNoteChanged() { return OnNoteChangedDelegate; }
 
 	virtual bool GetShouldShowInStack() const override;
+
+	NIAGARAEDITOR_API FLinearColor GetColor() const;
 protected:
 	virtual bool GetShouldShowInOverview() const override { return false; }
 	virtual EStackRowStyle GetStackRowStyle() const override;
+	virtual bool GetCanExpand() const override { return false; }
+	virtual bool KeepExpanderIndentation() const override { return false; }
 	
 	virtual bool SupportsRename() const override { return true; }
 	virtual bool SupportsCopy() const override { return true; }
diff --git a/Engine/Plugins/FX/Niagara/Source/NiagaraEditor/Public/ViewModels/Stack/NiagaraStackPropertyRow.h b/Engine/Plugins/FX/Niagara/Source/NiagaraEditor/Public/ViewModels/Stack/NiagaraStackPropertyRow.h
--- a/Engine/Plugins/FX/Niagara/Source/NiagaraEditor/Public/ViewModels/Stack/NiagaraStackPropertyRow.h
+++ b/Engine/Plugins/FX/Niagara/Source/NiagaraEditor/Public/ViewModels/Stack/NiagaraStackPropertyRow.h
@@ -38,7 +38,7 @@ public:
 
 	NIAGARAEDITOR_API virtual bool IsExpandedByDefault() const override;
 
-	virtual bool SupportsStackNotes() override { return true; }
+	virtual bool SupportsStackNotes() const override { return true; }
 	
 	NIAGARAEDITOR_API virtual bool CanDrag() const override;
 
diff --git a/Engine/Plugins/FX/Niagara/Source/NiagaraEditor/Public/ViewModels/Stack/NiagaraStackRendererItem.h b/Engine/Plugins/FX/Niagara/Source/NiagaraEditor/Public/ViewModels/Stack/NiagaraStackRendererItem.h
--- a/Engine/Plugins/FX/Niagara/Source/NiagaraEditor/Public/ViewModels/Stack/NiagaraStackRendererItem.h
+++ b/Engine/Plugins/FX/Niagara/Source/NiagaraEditor/Public/ViewModels/Stack/NiagaraStackRendererItem.h
@@ -57,7 +57,7 @@ public:
 	
 	virtual bool SupportsRename() const override { return true; }
 
-	virtual bool SupportsStackNotes() override { return true; }
+	virtual bool SupportsStackNotes() const override { return true; }
 
 	NIAGARAEDITOR_API bool HasBaseRenderer() const;
 
diff --git a/Engine/Plugins/FX/Niagara/Source/NiagaraEditor/Public/ViewModels/Stack/NiagaraStackSimulationStageGroup.h b/Engine/Plugins/FX/Niagara/Source/NiagaraEditor/Public/ViewModels/Stack/NiagaraStackSimulationStageGroup.h
--- a/Engine/Plugins/FX/Niagara/Source/NiagaraEditor/Public/ViewModels/Stack/NiagaraStackSimulationStageGroup.h
+++ b/Engine/Plugins/FX/Niagara/Source/NiagaraEditor/Public/ViewModels/Stack/NiagaraStackSimulationStageGroup.h
@@ -80,7 +80,7 @@ public:
 	NIAGARAEDITOR_API virtual bool GetIsInherited() const override;
 	NIAGARAEDITOR_API virtual FText GetInheritanceMessage() const override;
 
-	virtual bool SupportsStackNotes() override { return true; }
+	virtual bool SupportsStackNotes() const override { return true; }
 	
 	virtual bool CanDrag() const override { return true; }
 
diff --git a/Engine/Plugins/FX/Niagara/Source/NiagaraEditorWidgets/Private/NiagaraEditorWidgetsStyle.cpp b/Engine/Plugins/FX/Niagara/Source/NiagaraEditorWidgets/Private/NiagaraEditorWidgetsStyle.cpp
--- a/Engine/Plugins/FX/Niagara/Source/NiagaraEditorWidgets/Private/NiagaraEditorWidgetsStyle.cpp
+++ b/Engine/Plugins/FX/Niagara/Source/NiagaraEditorWidgets/Private/NiagaraEditorWidgetsStyle.cpp
@@ -66,8 +66,7 @@ FNiagaraEditorWidgetsStyle::FNiagaraEditorWidgetsStyle() : FSlateStyleSet("Niaga
 {
 	SetContentRoot(FPaths::EnginePluginsDir() / TEXT("FX/Niagara/Content/Slate"));
 	SetCoreContentRoot(FPaths::EngineContentDir() / TEXT("Editor/Slate"));
-
-	PreInit();
+	
 	InitColors();
 	InitTextStyles();
 	InitRowStyles();
@@ -81,14 +80,8 @@ FNiagaraEditorWidgetsStyle::FNiagaraEditorWidgetsStyle() : FSlateStyleSet("Niaga
 	InitMisc();
 }
 
-void FNiagaraEditorWidgetsStyle::PreInit()
-{
-	CustomNoteBackgroundColor = FLinearColor(FColor(56, 111, 75));
-}
-
 void FNiagaraEditorWidgetsStyle::InitColors()
 {
-	Set("NiagaraEditor.Stack.Item.CustomNoteBackgroundColor", CustomNoteBackgroundColor);
 	Set("NiagaraEditor.Stack.UnknownColor", FLinearColor(1, 0, 1));
 
 	Set("NiagaraEditor.Stack.ForegroundColor", FLinearColor(FColor(220, 220, 220)));
@@ -256,7 +249,6 @@ void FNiagaraEditorWidgetsStyle::InitRowStyles()
 	FLinearColor PanelHoverColor = FStyleColors::Panel.GetSpecifiedColor() + HoverAdd;
 	FLinearColor RecessedHoverColor = FStyleColors::Recessed.GetSpecifiedColor() + HoverAdd;
 	FLinearColor HeaderHoverColor = FStyleColors::Header.GetSpecifiedColor() + HoverAdd;
-	FLinearColor NoteHoverColor = CustomNoteBackgroundColor + HoverAdd;
 	
 	FSlateBoxBrush ScratchPadCategoryBrush = BOX_BRUSH("Icons/CategoryRow", FMargin(2.0f / 8.0f), FLinearColor(FColor(48, 48, 48)));
 	FSlateBoxBrush ScratchPadHoveredCategoryBrush = BOX_BRUSH("Icons/CategoryRow", FMargin(2.0f / 8.0f), FLinearColor(FColor(38, 38, 38)));
@@ -286,18 +278,14 @@ void FNiagaraEditorWidgetsStyle::InitRowStyles()
 		.SetEvenRowBackgroundBrush(FSlateColorBrush(FStyleColors::Recessed))
 		.SetOddRowBackgroundBrush(FSlateColorBrush(FStyleColors::Recessed)));
 
-	Set("NiagaraEditor.Stack.TableViewRow.ItemContentNote", FTableRowStyle(BaseStackTableRowStyle)
-		.SetEvenRowBackgroundHoveredBrush(FSlateColorBrush(NoteHoverColor))
-		.SetOddRowBackgroundHoveredBrush(FSlateColorBrush(NoteHoverColor))
-		.SetEvenRowBackgroundBrush(FSlateColorBrush(CustomNoteBackgroundColor))
-		.SetOddRowBackgroundBrush(FSlateColorBrush(CustomNoteBackgroundColor)));
-
 	Set("NiagaraEditor.Stack.TableViewRow.ItemHeader", FTableRowStyle(NormalTableRowStyle)
 		.SetEvenRowBackgroundHoveredBrush(FSlateColorBrush(HeaderHoverColor))
 		.SetOddRowBackgroundHoveredBrush(FSlateColorBrush(HeaderHoverColor))
 		.SetEvenRowBackgroundBrush(FSlateColorBrush(FStyleColors::Header))
 		.SetOddRowBackgroundBrush(FSlateColorBrush(FStyleColors::Header)));
 
+	Set("NiagaraEditor.Stack.TableViewRow.ItemContentNote", FTableRowStyle(BaseStackTableRowStyle));
+	
 	Set("NiagaraEditor.Stack.TableViewRow.Spacer", FTableRowStyle(NormalTableRowStyle)
 		.SetSelectorFocusedBrush(FSlateNoResource())
 		.SetActiveBrush(FSlateNoResource())
diff --git a/Engine/Plugins/FX/Niagara/Source/NiagaraEditorWidgets/Private/NiagaraEditorWidgetsStyle.h b/Engine/Plugins/FX/Niagara/Source/NiagaraEditorWidgets/Private/NiagaraEditorWidgetsStyle.h
--- a/Engine/Plugins/FX/Niagara/Source/NiagaraEditorWidgets/Private/NiagaraEditorWidgetsStyle.h
+++ b/Engine/Plugins/FX/Niagara/Source/NiagaraEditorWidgets/Private/NiagaraEditorWidgetsStyle.h
@@ -22,7 +22,6 @@ public:
 	static void ReinitializeStyle();
 private:
 	FNiagaraEditorWidgetsStyle();
-	void PreInit();
 	void InitColors();
 	void InitTextStyles();
 	void InitRowStyles();
@@ -40,6 +39,4 @@ private:
 	void InitStackTextStyles();
 
 	static TSharedPtr<FNiagaraEditorWidgetsStyle> NiagaraEditorWidgetsStyle;
-
-	FLinearColor CustomNoteBackgroundColor;
 };
diff --git a/Engine/Plugins/FX/Niagara/Source/NiagaraEditorWidgets/Private/NiagaraEditorWidgetsUtilities.cpp b/Engine/Plugins/FX/Niagara/Source/NiagaraEditorWidgets/Private/NiagaraEditorWidgetsUtilities.cpp
--- a/Engine/Plugins/FX/Niagara/Source/NiagaraEditorWidgets/Private/NiagaraEditorWidgetsUtilities.cpp
+++ b/Engine/Plugins/FX/Niagara/Source/NiagaraEditorWidgets/Private/NiagaraEditorWidgetsUtilities.cpp
@@ -505,6 +505,68 @@ void FNiagaraStackEditorWidgetsUtilities::AddStackSimulationStageGroupContextMen
 	MenuBuilder.EndSection();
 }
 
+bool FNiagaraStackEditorWidgetsUtilities::AddStackNoteContextMenuActions(FMenuBuilder& MenuBuilder, UNiagaraStackEntry& StackEntry)
+{
+	if(StackEntry.SupportsStackNotes())
+	{
+		MenuBuilder.BeginSection("NoteActions", LOCTEXT("NoteActions", "Note Actions"));
+		
+		if(StackEntry.HasStackNoteData() == false)
+		{
+			MenuBuilder.AddMenuEntry(
+			LOCTEXT("AddNote", "Add Note"),
+			LOCTEXT("AddNoteTooltip", "Add a note to this row"),
+			FSlateIcon(FAppStyle::GetAppStyleSetName(), "Icons.Comment"),
+			FUIAction(FExecuteAction::CreateUObject(&StackEntry, &UNiagaraStackEntry::AddStackNote))
+			);
+		}
+		else
+		{				
+			MenuBuilder.AddMenuEntry(
+			LOCTEXT("ToggleInlineNote", "Toggle Inline Note Display"),
+			LOCTEXT("ToggleInlineNoteTooltip", "Toggle the inline display for this row's note.\nAn inlined note will show up in the row itself, saving on space"),
+			FSlateIcon(FAppStyle::GetAppStyleSetName(), "Icons.Comment"),
+			FUIAction(FExecuteAction::CreateUObject(&StackEntry, &UNiagaraStackEntry::ToggleStackNoteInlineDisplay))
+			);
+		
+			MenuBuilder.AddMenuEntry(
+			LOCTEXT("CopyNote", "Copy Note"),
+			LOCTEXT("CopyNoteTooltip", "Copy the note of this row"),
+			FSlateIcon(FAppStyle::GetAppStyleSetName(), "GenericCommands.Copy"),
+			FUIAction(FExecuteAction::CreateUObject(&StackEntry, &UNiagaraStackEntry::CopyStackNote))
+			);
+		}
+			
+		if(const UNiagaraClipboardContent* ClipboardContent = FNiagaraEditorModule::Get().GetClipboard().GetClipboardContent())
+		{
+			if(ClipboardContent->StackNote.IsValid())
+			{
+				MenuBuilder.AddMenuEntry(
+				LOCTEXT("PasteNote", "Paste Note"),
+				LOCTEXT("TPasteNoteTooltip", "Paste the note copied in the clipboard"),
+				FSlateIcon(FAppStyle::GetAppStyleSetName(), "GenericCommands.Paste"),
+				FUIAction(FExecuteAction::CreateUObject(&StackEntry, &UNiagaraStackEntry::PasteStackNote))
+				);
+			}
+		}
+
+		if(StackEntry.HasStackNoteData())
+		{
+			MenuBuilder.AddMenuEntry(
+			LOCTEXT("DeleteNote", "Delete Note"),
+			LOCTEXT("DeleteNoteTooltip", "Delete this row's note"),
+			FSlateIcon(FAppStyle::GetAppStyleSetName(), "GenericCommands.Delete"),
+			FUIAction(FExecuteAction::CreateUObject(&StackEntry, &UNiagaraStackEntry::DeleteStackNoteData))
+			);
+		}
+
+		MenuBuilder.EndSection();
+		return true;
+	}
+
+	return false;
+}
+
 TSharedRef<FDragDropOperation> FNiagaraStackEditorWidgetsUtilities::ConstructDragDropOperationForStackEntries(const TArray<UNiagaraStackEntry*>& DraggedEntries)
 {
 	TSharedRef<FNiagaraStackEntryDragDropOp> DragDropOp = MakeShared<FNiagaraStackEntryDragDropOp>(DraggedEntries);
diff --git a/Engine/Plugins/FX/Niagara/Source/NiagaraEditorWidgets/Private/NiagaraEditorWidgetsUtilities.h b/Engine/Plugins/FX/Niagara/Source/NiagaraEditorWidgets/Private/NiagaraEditorWidgetsUtilities.h
--- a/Engine/Plugins/FX/Niagara/Source/NiagaraEditorWidgets/Private/NiagaraEditorWidgetsUtilities.h
+++ b/Engine/Plugins/FX/Niagara/Source/NiagaraEditorWidgets/Private/NiagaraEditorWidgetsUtilities.h
@@ -39,6 +39,8 @@ namespace FNiagaraStackEditorWidgetsUtilities
 
 	void AddStackSimulationStageGroupContextMenuActions(FMenuBuilder& MenuBuilder, UNiagaraStackSimulationStageGroup* SimulationStageGroup, TSharedRef<SWidget> TagetWidget);
 
+	bool AddStackNoteContextMenuActions(FMenuBuilder& MenuBuilder, UNiagaraStackEntry& StackEntry);
+	
 	TSharedRef<FDragDropOperation> ConstructDragDropOperationForStackEntries(const TArray<UNiagaraStackEntry*>& DraggedEntries);
 
 	void HandleDragLeave(const FDragDropEvent& InDragDropEvent);
diff --git a/Engine/Plugins/FX/Niagara/Source/NiagaraEditorWidgets/Private/SNiagaraOverviewStack.cpp b/Engine/Plugins/FX/Niagara/Source/NiagaraEditorWidgets/Private/SNiagaraOverviewStack.cpp
--- a/Engine/Plugins/FX/Niagara/Source/NiagaraEditorWidgets/Private/SNiagaraOverviewStack.cpp
+++ b/Engine/Plugins/FX/Niagara/Source/NiagaraEditorWidgets/Private/SNiagaraOverviewStack.cpp
@@ -283,6 +283,7 @@ class SNiagaraSystemOverviewEntryListRow : public STableRow<UNiagaraStackEntry*>
 
 			bMenuItemsAdded |= FNiagaraStackEditorWidgetsUtilities::AddStackEntryAssetContextMenuActions(MenuBuilder, *StackEntry);
 			bMenuItemsAdded |= StackCommandContext->AddEditMenuItems(MenuBuilder);
+			bMenuItemsAdded |= FNiagaraStackEditorWidgetsUtilities::AddStackNoteContextMenuActions(MenuBuilder, *StackEntry);
 		
 			if (bMenuItemsAdded)
 			{
diff --git a/Engine/Plugins/FX/Niagara/Source/NiagaraEditorWidgets/Private/SNiagaraStackNote.cpp b/Engine/Plugins/FX/Niagara/Source/NiagaraEditorWidgets/Private/SNiagaraStackNote.cpp
--- a/Engine/Plugins/FX/Niagara/Source/NiagaraEditorWidgets/Private/SNiagaraStackNote.cpp
+++ b/Engine/Plugins/FX/Niagara/Source/NiagaraEditorWidgets/Private/SNiagaraStackNote.cpp
@@ -11,6 +11,8 @@
 #include "ScopedTransaction.h"
 #include "ViewModels/Stack/NiagaraStackEntry.h"
 #include "Widgets/SToolTip.h"
+#include "Widgets/Colors/SColorBlock.h"
+#include "Widgets/Colors/SColorPicker.h"
 
 #define LOCTEXT_NAMESPACE "SNiagaraStackNote"
 
@@ -22,7 +24,8 @@ void SNiagaraStackInlineNote::Construct(const FArguments& InArgs, UNiagaraStackE
 	bInteractable = InArgs._bInteractable;
 
 	TSharedRef<SWidget> ChildWidget = SNew(SImage)
-		.Image(FNiagaraEditorStyle::Get().GetBrush("NiagaraEditor.Message.CustomNote"));
+		.Image(FAppStyle::GetBrush("Icons.Comment"))
+		.ColorAndOpacity(this, &SNiagaraStackInlineNote::GetStackNoteColor);
 
 	if(bInteractable)
 	{
@@ -56,6 +59,16 @@ FReply SNiagaraStackInlineNote::OnClicked() const
 	return FReply::Handled();
 }
 
+FSlateColor SNiagaraStackInlineNote::GetStackNoteColor() const
+{
+	if(StackEntry.IsValid() && StackEntry->IsFinalized() == false && StackEntry->GetStackNote())
+	{
+		return StackEntry->GetStackNote()->GetColor();
+	}
+
+	return GetDefault<UNiagaraEditorSettings>()->GetDefaultNoteColor();
+}
+
 void SNiagaraStackInlineNote::OnMouseEnter(const FGeometry& MyGeometry, const FPointerEvent& MouseEvent)
 {
 	// this is a workaround. We force a regeneration of the tooltip whenever we enter it.
@@ -96,46 +109,94 @@ void SNiagaraStackNote::Rebuild()
 	}
 	
 	FNiagaraStackNoteData DisplayedStackNote = MatchingStackNote.GetValue();
-	
+
 	ChildSlot
 	[
-		SAssignNew(ExpandableArea, SExpandableArea)
-		.BorderImage(FAppStyle::GetBrush("NoBorder"))
-		.HeaderContent()
+		SNew(SHorizontalBox)
+		+ SHorizontalBox::Slot()
+		.AutoWidth()
 		[
-			SNew(SHorizontalBox)
-			+ SHorizontalBox::Slot()
-			.AutoWidth()
-			.Padding(2.f)
+			SNew(SBox)
+			.WidthOverride(16.f)
 			[
-				SNew(SButton)
-				.ButtonStyle(&FAppStyle::GetWidgetStyle<FButtonStyle>("HoverHintOnly"))
-				.OnClicked(this, &SNiagaraStackNote::OnToggleInlineDisplayClicked)
-				.ToolTipText(LOCTEXT("InlineNoteButtonTooltip", "Converts this note to an inline note."))
-				.Content()
+				SNew(SColorBlock)
+				.Color(this, &SNiagaraStackNote::GetStackNoteColor)
+				.OnMouseButtonDown(this, &SNiagaraStackNote::SummonColorPicker)
+			]
+		]
+		+ SHorizontalBox::Slot()
+		[
+			SAssignNew(ExpandableArea, SExpandableArea)
+			.BorderImage(FAppStyle::GetBrush("NoBorder"))
+			.Padding(FMargin(20.f, 2.f, 5.f, 2.f))
+			.HeaderContent()
+			[
+				SNew(SHorizontalBox)
+				+ SHorizontalBox::Slot()
+				.AutoWidth()
+				.Padding(2.f)
 				[
-					SNew(SBox)
-					.HeightOverride(16.f)
-					.WidthOverride(16.f)
+					SNew(SButton)
+					.ButtonStyle(&FAppStyle::GetWidgetStyle<FButtonStyle>("HoverHintOnly"))
+					.OnClicked(this, &SNiagaraStackNote::OnToggleInlineDisplayClicked)
+					.ToolTipText(LOCTEXT("InlineNoteButtonTooltip", "Converts this note to an inline note."))
+					.Content()
+					[
+						SNew(SBox)
+						.HeightOverride(16.f)
+						.WidthOverride(16.f)
+						.HAlign(HAlign_Center)
+						.VAlign(VAlign_Center)
+						[
+							SNew(SImage)
+							.Image(FAppStyle::GetBrush("Icons.Comment"))
+							.ColorAndOpacity(this, &SNiagaraStackNote::GetSlateStackNoteColor)
+						]
+					]
+				]
+				+ SHorizontalBox::Slot()
+				[
+					SNew(SHorizontalBox)
+					+ SHorizontalBox::Slot()
+					[
+						SAssignNew(HeaderText, SInlineEditableTextBlock)
+						.Text(this, &SNiagaraStackNote::GetStackNoteHeader)
+						.OnTextCommitted(this, &SNiagaraStackNote::CommitStackNoteHeaderUpdate)
+						.Style(&FNiagaraEditorStyle::Get().GetWidgetStyle<FInlineEditableTextBlockStyle>("NiagaraEditor.Stack.Note.HeaderEditableText"))
+						.IsSelected(FIsSelected::CreateLambda([]() { return false; }))
+						.AutoWrapNonEditText(true)
+					]
+					+ SHorizontalBox::Slot()
+					.AutoWidth()
 					.HAlign(HAlign_Center)
-					.VAlign(VAlign_Center)
+					.VAlign(VAlign_Bottom)
+					.Padding(10.f, 2.f)
 					[
-						SNew(SImage)
-						.Image(FNiagaraEditorStyle::Get().GetBrush("NiagaraEditor.Message.CustomNote"))
+						SNew(SButton)
+						.OnClicked(this, &SNiagaraStackNote::OnEditHeaderButtonClicked)
+						.ButtonStyle(&FAppStyle::GetWidgetStyle<FButtonStyle>("HoverHintOnly"))
+						.Visibility(this, &SNiagaraStackNote::GetEditNoteHeaderButtonVisibility)
+						.ToolTipText(LOCTEXT("EditHeaderButtonTooltip", "Edit the header of this note."))
+						[
+							SNew(SImage)
+							.Image(FAppStyle::GetBrush("Icons.Edit"))
+						]
 					]
 				]
 			]
-			+ SHorizontalBox::Slot()
+			.BodyContent()
 			[
 				SNew(SHorizontalBox)
 				+ SHorizontalBox::Slot()
 				[
-					SAssignNew(HeaderText, SInlineEditableTextBlock)
-					.Text(this, &SNiagaraStackNote::GetStackNoteHeader)
-					.OnTextCommitted(this, &SNiagaraStackNote::CommitStackNoteHeaderUpdate)
-					.Style(&FNiagaraEditorStyle::Get().GetWidgetStyle<FInlineEditableTextBlockStyle>("NiagaraEditor.Stack.Note.HeaderEditableText"))
-					.IsSelected(FIsSelected::CreateLambda([]() { return false; }))
+					SAssignNew(BodyText, SInlineEditableTextBlock)
+					.Text(this, &SNiagaraStackNote::GetStackNoteBody)
+					.OnTextCommitted(this, &SNiagaraStackNote::CommitStackNoteBodyUpdate)
 					.AutoWrapNonEditText(true)
+					.AutoWrapMultilineEditText(true)
+					.MultiLine(true)
+					.ModiferKeyForNewLine(EModifierKey::Shift)
+					.IsSelected(FIsSelected::CreateLambda([]() { return false; }))
 				]
 				+ SHorizontalBox::Slot()
 				.AutoWidth()
@@ -144,10 +205,10 @@ void SNiagaraStackNote::Rebuild()
 				.Padding(10.f, 2.f)
 				[
 					SNew(SButton)
-					.OnClicked(this, &SNiagaraStackNote::OnEditHeaderButtonClicked)
+					.OnClicked(this, &SNiagaraStackNote::OnEditBodyButtonClicked)
 					.ButtonStyle(&FAppStyle::GetWidgetStyle<FButtonStyle>("HoverHintOnly"))
-					.Visibility(this, &SNiagaraStackNote::GetEditNoteHeaderButtonVisibility)
-					.ToolTipText(LOCTEXT("EditHeaderButtonTooltip", "Edit the header of this note."))
+					.Visibility(this, &SNiagaraStackNote::GetEditNoteBodyButtonVisibility)
+					.ToolTipText(LOCTEXT("EditMessageButtonTooltip", "Edit the message of this note."))
 					[
 						SNew(SImage)
 						.Image(FAppStyle::GetBrush("Icons.Edit"))
@@ -155,38 +216,6 @@ void SNiagaraStackNote::Rebuild()
 				]
 			]
 		]
-		.BodyContent()
-		[
-			SNew(SHorizontalBox)
-			+ SHorizontalBox::Slot()
-			[
-				SAssignNew(BodyText, SInlineEditableTextBlock)
-				.Text(this, &SNiagaraStackNote::GetStackNoteBody)
-				.OnTextCommitted(this, &SNiagaraStackNote::CommitStackNoteBodyUpdate)
-				.AutoWrapNonEditText(true)
-				.AutoWrapMultilineEditText(true)
-				.MultiLine(true)
-				.ModiferKeyForNewLine(EModifierKey::Shift)
-				.IsSelected(FIsSelected::CreateLambda([]() { return false; }))
-			]
-			+ SHorizontalBox::Slot()
-			.AutoWidth()
-			.HAlign(HAlign_Center)
-			.VAlign(VAlign_Bottom)
-			.Padding(10.f, 2.f)
-			[
-				SNew(SButton)
-				.OnClicked(this, &SNiagaraStackNote::OnEditBodyButtonClicked)
-				.ButtonStyle(&FAppStyle::GetWidgetStyle<FButtonStyle>("HoverHintOnly"))
-				.Visibility(this, &SNiagaraStackNote::GetEditNoteBodyButtonVisibility)
-				.ToolTipText(LOCTEXT("EditMessageButtonTooltip", "Edit the message of this note."))
-				[
-					SNew(SImage)
-					.Image(FAppStyle::GetBrush("Icons.Edit"))
-				]
-			]
-		]
-		.Padding(FMargin(20.f, 2.f, 5.f, 2.f))
 	];
 }
 
@@ -197,25 +226,25 @@ void SNiagaraStackNote::FillRowContextMenu(FMenuBuilder& MenuBuilder)
 		MenuBuilder.AddMenuEntry(
 			LOCTEXT("EditTitle", "Edit Title"),
 			LOCTEXT("EditTitleTooltip", "Edit the title of this note."),
-			FSlateIcon(),
+			FSlateIcon(FAppStyle::GetAppStyleSetName(), "Icons.Edit"),
 			FUIAction(FExecuteAction::CreateSP(this, &SNiagaraStackNote::EditHeaderText)));
 
 		MenuBuilder.AddMenuEntry(
 			LOCTEXT("EditBody", "Edit Message"),
 			LOCTEXT("EditMessageTooltip", "Edit the message of this note."),
-			FSlateIcon(),
+			FSlateIcon(FAppStyle::GetAppStyleSetName(), "Icons.Edit"),
 			FUIAction(FExecuteAction::CreateSP(this, &SNiagaraStackNote::EditBodyText)));
 
 		MenuBuilder.AddMenuEntry(
 			LOCTEXT("ToggleInline", "Toggle Inline Display"),
 			LOCTEXT("ToggleInlineTooltip", "Toggle the Inline Display for this note.\nAn inlined note will show up in the row itself, saving on space."),
-			FSlateIcon(),
+			FSlateIcon(FAppStyle::GetAppStyleSetName(), "Icons.Comment"),
 			FUIAction(FExecuteAction::CreateSP(this, &SNiagaraStackNote::ToggleInlineDisplay)));
 
 		MenuBuilder.AddMenuEntry(
 			LOCTEXT("DeleteNote", "Delete Note"),
 			LOCTEXT("DeleteNoteTooltip", "Delete this note."),
-			FSlateIcon(),
+			FSlateIcon(FAppStyle::GetAppStyleSetName(), "GenericCommands.Delete"),
 			FUIAction(FExecuteAction::CreateSP(this, &SNiagaraStackNote::DeleteStackNote)));
 	}
 	MenuBuilder.EndSection();
@@ -227,7 +256,7 @@ void SNiagaraStackNote::Tick(const FGeometry& AllottedGeometry, const double InC
 	{
 		if(HeaderText.IsValid())
 		{
-			EditHeaderText();
+			EditBodyText();
 			StackNote->SetIsRenamePending(false);
 		}
 	}
@@ -293,6 +322,16 @@ FText SNiagaraStackNote::GetStackNoteBody() const
 	return GetStackNoteData().GetValue().Message;
 }
 
+FLinearColor SNiagaraStackNote::GetStackNoteColor() const
+{
+	return StackNote->GetColor();
+}
+
+FSlateColor SNiagaraStackNote::GetSlateStackNoteColor() const
+{
+	return GetStackNoteColor();
+}
+
 FReply SNiagaraStackNote::OnToggleInlineDisplayClicked() const
 {
 	ToggleInlineDisplay();
@@ -319,6 +358,39 @@ FReply SNiagaraStackNote::OnEditBodyButtonClicked()
 	return FReply::Handled();
 }
 
+FReply SNiagaraStackNote::SummonColorPicker(const FGeometry& Geometry, const FPointerEvent& PointerEvent) const
+{	
+	FColorPickerArgs PickerArgs;
+	{
+		PickerArgs.bOnlyRefreshOnOk = true;
+		PickerArgs.DisplayGamma = TAttribute<float>::Create(TAttribute<float>::FGetter::CreateUObject(GEngine, &UEngine::GetDisplayGamma));
+		PickerArgs.OnColorCommitted = FOnLinearColorValueChanged::CreateSP(this, &SNiagaraStackNote::OnColorPickedCommitted);
+		PickerArgs.OnColorPickerCancelled = FOnColorPickerCancelled::CreateSP(this, &SNiagaraStackNote::OnColorPickerCancelled);
+		PickerArgs.OnColorPickerWindowClosed = FOnWindowClosed::CreateSP(this, &SNiagaraStackNote::OnColorPickerClosed);
+		PickerArgs.bUseAlpha = false;
+		PickerArgs.InitialColor = GetStackNoteColor();
+		PickerArgs.bOnlyRefreshOnMouseUp = true;
+		PickerArgs.bOnlyRefreshOnOk = false;
+	}
+	
+	OpenColorPicker(PickerArgs);
+	return FReply::Handled();
+}
+
+void SNiagaraStackNote::OnColorPickedCommitted(FLinearColor LinearColor) const
+{
+	StackNote->UpdateNoteColor(LinearColor);
+}
+
+void SNiagaraStackNote::OnColorPickerCancelled(FLinearColor LinearColor) const
+{
+	StackNote->UpdateNoteColor(LinearColor);
+}
+
+void SNiagaraStackNote::OnColorPickerClosed(const TSharedRef<SWindow>& Window) const
+{
+}
+
 EVisibility SNiagaraStackNote::GetEditNoteHeaderButtonVisibility() const
 {
 	return ShowEditTextButtons.Get() ? EVisibility::Visible : EVisibility::Hidden;
@@ -331,4 +403,4 @@ EVisibility SNiagaraStackNote::GetEditNoteBodyButtonVisibility() const
 
 END_SLATE_FUNCTION_BUILD_OPTIMIZATION
 
-#undef LOCTEXT_NAMESPACE
\ No newline at end of file
+#undef LOCTEXT_NAMESPACE
diff --git a/Engine/Plugins/FX/Niagara/Source/NiagaraEditorWidgets/Private/SNiagaraStackNote.h b/Engine/Plugins/FX/Niagara/Source/NiagaraEditorWidgets/Private/SNiagaraStackNote.h
--- a/Engine/Plugins/FX/Niagara/Source/NiagaraEditorWidgets/Private/SNiagaraStackNote.h
+++ b/Engine/Plugins/FX/Niagara/Source/NiagaraEditorWidgets/Private/SNiagaraStackNote.h
@@ -26,6 +26,8 @@ public:
 private:
 	FReply OnClicked() const;
 
+	FSlateColor GetStackNoteColor() const;
+
 	virtual void OnMouseEnter(const FGeometry& MyGeometry, const FPointerEvent& MouseEvent) override;
 private:
 	TWeakObjectPtr<UNiagaraStackEntry> StackEntry;
@@ -44,7 +46,7 @@ public:
 	
 	void Construct(const FArguments& InArgs, UNiagaraStackNote& StackNote);
 	virtual ~SNiagaraStackNote() override;
-	
+
 	void Rebuild();
 
 	void FillRowContextMenu(FMenuBuilder& MenuBuilder);
@@ -64,6 +66,8 @@ private:
 	
 	FText GetStackNoteHeader() const;
 	FText GetStackNoteBody() const;
+	FLinearColor GetStackNoteColor() const;
+	FSlateColor GetSlateStackNoteColor() const;
 
 private:
 	FReply OnToggleInlineDisplayClicked() const;
@@ -71,6 +75,11 @@ private:
 	FReply OnEditHeaderButtonClicked();
 	FReply OnEditBodyButtonClicked();
 	
+	FReply SummonColorPicker(const FGeometry& Geometry, const FPointerEvent& PointerEvent) const;
+	void OnColorPickedCommitted(FLinearColor LinearColor) const;
+	void OnColorPickerCancelled(FLinearColor LinearColor) const;
+	void OnColorPickerClosed(const TSharedRef<SWindow>& Window) const;
+
 	EVisibility GetEditNoteHeaderButtonVisibility() const;
 	EVisibility GetEditNoteBodyButtonVisibility() const;
 
diff --git a/Engine/Plugins/FX/Niagara/Source/NiagaraEditorWidgets/Private/Stack/SNiagaraStackTableRow.cpp b/Engine/Plugins/FX/Niagara/Source/NiagaraEditorWidgets/Private/Stack/SNiagaraStackTableRow.cpp
--- a/Engine/Plugins/FX/Niagara/Source/NiagaraEditorWidgets/Private/Stack/SNiagaraStackTableRow.cpp
+++ b/Engine/Plugins/FX/Niagara/Source/NiagaraEditorWidgets/Private/Stack/SNiagaraStackTableRow.cpp
@@ -3,6 +3,7 @@
 #include "Stack/SNiagaraStackTableRow.h"
 
 #include "NiagaraEditorModule.h"
+#include "NiagaraEditorSettings.h"
 #include "NiagaraEditorStyle.h"
 #include "NiagaraEditorWidgetsStyle.h"
 #include "Styling/AppStyle.h"
@@ -436,59 +437,10 @@ FReply SNiagaraStackTableRow::OnMouseButtonUp(const FGeometry& MyGeometry, const
 			FSlateIcon(),
 			FUIAction(FExecuteAction::CreateSP(this, &SNiagaraStackTableRow::CollapseChildren)));
 
-		if(StackEntry->SupportsStackNotes())
-		{
-			if(StackEntry->HasStackNoteData() == false)
-			{
-				MenuBuilder.AddMenuEntry(
-				LOCTEXT("AddNote", "Add Note"),
-				LOCTEXT("AddNoteTooltip", "Add a note to this row"),
-				FSlateIcon(FNiagaraEditorStyle::Get().GetStyleSetName(), "NiagaraEditor.Message.CustomNote"),
-				FUIAction(FExecuteAction::CreateSP(this, &SNiagaraStackTableRow::AddStackNote))
-				);
-			}
-			else
-			{				
-				MenuBuilder.AddMenuEntry(
-				LOCTEXT("ToggleInlineNote", "Toggle Inline Note Display"),
-				LOCTEXT("ToggleInlineNoteTooltip", "Toggle the inline display for this row's note.\nAn inlined note will show up in the row itself, saving on space"),
-				FSlateIcon(FNiagaraEditorStyle::Get().GetStyleSetName(), "NiagaraEditor.Message.CustomNote"),
-				FUIAction(FExecuteAction::CreateSP(this, &SNiagaraStackTableRow::ToggleStackNoteInlineDisplay))
-				);
-
-				MenuBuilder.AddMenuEntry(
-				LOCTEXT("CopyNote", "Copy Note"),
-			LOCTEXT("CopyNoteTooltip", "Copy the note of this row"),
-				FSlateIcon(FAppStyle::GetAppStyleSetName(), "GenericCommands.Copy"),
-				FUIAction(FExecuteAction::CreateSP(this, &SNiagaraStackTableRow::CopyStackNote))
-				);
-			}
-			
-			if(const UNiagaraClipboardContent* ClipboardContent = FNiagaraEditorModule::Get().GetClipboard().GetClipboardContent())
-			{
-				if(ClipboardContent->StackNote.IsValid())
-				{
-					MenuBuilder.AddMenuEntry(
-				LOCTEXT("PasteNote", "Paste Note"),
-				LOCTEXT("TPasteNoteTooltip", "Paste the note copied in the clipboard"),
-					FSlateIcon(FAppStyle::GetAppStyleSetName(), "GenericCommands.Paste"),
-					FUIAction(FExecuteAction::CreateSP(this, &SNiagaraStackTableRow::PasteStackNote))
-					);
-				}
-			}
-
-			if(StackEntry->HasStackNoteData())
-			{
-				MenuBuilder.AddMenuEntry(
-				LOCTEXT("DeleteNote", "Delete Note"),
-				LOCTEXT("DeleteNoteTooltip", "Delete this row's note"),
-				FSlateIcon(FAppStyle::GetAppStyleSetName(), "GenericCommands.Delete"),
-				FUIAction(FExecuteAction::CreateSP(this, &SNiagaraStackTableRow::DeleteStackNote))
-				);
-			}
-		}
 		MenuBuilder.EndSection();
-
+		
+		FNiagaraStackEditorWidgetsUtilities::AddStackNoteContextMenuActions(MenuBuilder, *StackEntry);
+		
 		if(StackEntry->GetEmitterViewModel().IsValid() && StackEntry->GetEmitterViewModel())
 		{			
 			if(StackEntry->SupportsSummaryView())
@@ -603,68 +555,6 @@ void SNiagaraStackTableRow::ExpandChildren()
 	SetExpansionStateRecursive(StackEntry, bIsExpanded);
 }
 
-void SNiagaraStackTableRow::AddStackNote() const
-{
-	if(StackEntry->HasStackNoteData() == false)
-	{
-		FScopedTransaction Transaction(LOCTEXT("NewStackNoteAdded", "Added Stack Note"));
-		StackEntry->GetStackEditorData().Modify();
-		
-		FNiagaraStackNoteData NewStackNote;
-		NewStackNote.MessageHeader = LOCTEXT("DefaultStackNoteHeader", "Title");
-		NewStackNote.Message = LOCTEXT("DefaultStackNoteMessage", "Text");
-
-		StackEntry->GetStackEditorData().AddOrReplaceStackNote(StackEntry->GetStackEditorDataKey(), NewStackNote);
-		StackEntry->RefreshChildren();
-		
-		if(UNiagaraStackNote* Note = StackEntry->GetStackNote())
-		{
-			Note->SetIsRenamePending(true);
-		}
-	}
-}
-
-void SNiagaraStackTableRow::DeleteStackNote() const
-{
-	if(StackEntry->HasStackNoteData())
-	{
-		StackEntry->DeleteStackNoteData();
-	}
-}
-
-void SNiagaraStackTableRow::ToggleStackNoteInlineDisplay() const
-{
-	if(StackEntry->GetStackNote() != nullptr)
-	{
-		StackEntry->GetStackNote()->ToggleInlineDisplay();
-	}
-}
-
-void SNiagaraStackTableRow::CopyStackNote() const
-{
-	FNiagaraStackClipboardUtilities::CopyNote(StackEntry);
-}
-
-void SNiagaraStackTableRow::PasteStackNote() const
-{
-	if(const UNiagaraClipboardContent* ClipboardContent = FNiagaraEditorModule::Get().GetClipboard().GetClipboardContent())
-	{
-		if(ClipboardContent->StackNote.IsValid())
-		{
-			FScopedTransaction Transaction(LOCTEXT("PasteNoteTransaction", "Note pasted"));
-			StackEntry->GetStackEditorData().Modify();
-			
-			StackEntry->SetStackNoteData(ClipboardContent->StackNote);
-
-			// if the note we pasted is not inlined, we expand the stack entry in order to show the note
-			if(StackEntry->GetStackNoteData().bInlineNote == false)
-			{
-				StackEntry->SetIsExpanded(true);
-			}
-		}
-	}
-}
-
 EVisibility SNiagaraStackTableRow::GetRowVisibility() const
 {
 	return StackEntry->GetShouldShowInStack()
diff --git a/Engine/Plugins/FX/Niagara/Source/NiagaraEditorWidgets/Private/Stack/SNiagaraStackTableRow.h b/Engine/Plugins/FX/Niagara/Source/NiagaraEditorWidgets/Private/Stack/SNiagaraStackTableRow.h
--- a/Engine/Plugins/FX/Niagara/Source/NiagaraEditorWidgets/Private/Stack/SNiagaraStackTableRow.h
+++ b/Engine/Plugins/FX/Niagara/Source/NiagaraEditorWidgets/Private/Stack/SNiagaraStackTableRow.h
@@ -76,13 +76,7 @@ private:
 
 	void CollapseChildren();
 	void ExpandChildren();
-
-	void AddStackNote() const;
-	void DeleteStackNote() const;
-	void ToggleStackNoteInlineDisplay() const;
-	void CopyStackNote() const;
-	void PasteStackNote() const;
-
+	
 	EVisibility GetRowVisibility() const;
 
 	EVisibility GetExecutionCategoryIconVisibility() const;
