- Handle invalid socket when reading from filter / unfiltered bones
- Move VM usage into Init as the instance data during the binding might not be the correct one

#rnx

diff --git a/Engine/Plugins/FX/Niagara/Source/Niagara/Private/DataInterface/NiagaraDataInterfaceSocketReader.cpp b/Engine/Plugins/FX/Niagara/Source/Niagara/Private/DataInterface/NiagaraDataInterfaceSocketReader.cpp
--- a/Engine/Plugins/FX/Niagara/Source/Niagara/Private/DataInterface/NiagaraDataInterfaceSocketReader.cpp
+++ b/Engine/Plugins/FX/Niagara/Source/Niagara/Private/DataInterface/NiagaraDataInterfaceSocketReader.cpp
@@ -455,8 +455,11 @@ namespace NDISocketReaderLocal
 					for (int32 i=SocketTableStart; i < SocketTableEnd; ++i)
 					{
 						const int32 SocketIndex = InstanceData->SocketFilterUnfilteredIndex[i + 1];
-						const FTransform SocketTransform = SceneComponent->GetSocketTransform(InstanceData->SocketNames[SocketIndex], ERelativeTransformSpace::RTS_Component);
-						InstanceData->SocketTransforms[SocketIndex] = FTransform3f(SocketTransform);
+						if (SocketIndex >= 0 && SocketIndex < InstanceData->NumSockets)
+						{
+							const FTransform SocketTransform = SceneComponent->GetSocketTransform(InstanceData->SocketNames[SocketIndex], ERelativeTransformSpace::RTS_Component);
+							InstanceData->SocketTransforms[SocketIndex] = FTransform3f(SocketTransform);
+						}
 					}
 				}
 			}
@@ -544,10 +547,10 @@ namespace NDISocketReaderLocal
 
 	enum class ESocketReadType
 	{
-		Default,
+		None,
+		Any,
 		Filtered,
 		Unfiltered,
-		None,
 	};
 
 	void VMIsValid(FVectorVMExternalFunctionContext& Context)
@@ -673,13 +676,13 @@ namespace NDISocketReaderLocal
 	{
 		{IsValidName,									FVMFunctionInfo{FVMExternalFunction::CreateStatic(VMIsValid), ESocketReadType::None}},
 		{GetComponentToWorldName,						FVMFunctionInfo{FVMExternalFunction::CreateStatic(VMGetComponentToWorld), ESocketReadType::None}},
-		{GetSocketCountName,							FVMFunctionInfo{FVMExternalFunction::CreateStatic(VMGetSocketCount<ESocketReadType::Default>), ESocketReadType::None}},
+		{GetSocketCountName,							FVMFunctionInfo{FVMExternalFunction::CreateStatic(VMGetSocketCount<ESocketReadType::Any>), ESocketReadType::None}},
 		{GetFilteredSocketCountName,					FVMFunctionInfo{FVMExternalFunction::CreateStatic(VMGetSocketCount<ESocketReadType::Filtered>), ESocketReadType::None}},
 		{GetUnfilteredSocketCountName,					FVMFunctionInfo{FVMExternalFunction::CreateStatic(VMGetSocketCount<ESocketReadType::Unfiltered>), ESocketReadType::None}},
-		{GetSocketTransformName,						FVMFunctionInfo{FVMExternalFunction::CreateStatic(VMGetSocketTransform<ESocketReadType::Default, false>), ESocketReadType::Default}},
+		{GetSocketTransformName,						FVMFunctionInfo{FVMExternalFunction::CreateStatic(VMGetSocketTransform<ESocketReadType::Any, false>), ESocketReadType::Any}},
 		{GetFilteredSocketTransformName,				FVMFunctionInfo{FVMExternalFunction::CreateStatic(VMGetSocketTransform<ESocketReadType::Filtered, false>), ESocketReadType::Filtered}},
 		{GetUnfilteredSocketTransformName,				FVMFunctionInfo{FVMExternalFunction::CreateStatic(VMGetSocketTransform<ESocketReadType::Unfiltered, false>), ESocketReadType::Unfiltered}},
-		{GetSocketTransformInterpolatedName,			FVMFunctionInfo{FVMExternalFunction::CreateStatic(VMGetSocketTransform<ESocketReadType::Default, true>), ESocketReadType::Default}},
+		{GetSocketTransformInterpolatedName,			FVMFunctionInfo{FVMExternalFunction::CreateStatic(VMGetSocketTransform<ESocketReadType::Any, true>), ESocketReadType::Any}},
 		{GetFilteredSocketTransformInterpolatedName,	FVMFunctionInfo{FVMExternalFunction::CreateStatic(VMGetSocketTransform<ESocketReadType::Filtered, true>), ESocketReadType::Filtered}},
 		{GetUnfilteredSocketTransformInterpolatedName,	FVMFunctionInfo{FVMExternalFunction::CreateStatic(VMGetSocketTransform<ESocketReadType::Unfiltered, true>), ESocketReadType::Unfiltered}},
 	};
@@ -762,10 +765,6 @@ void UNiagaraDataInterfaceSocketReader::GetVMExternalFunction(const FVMExternalF
 	if (const FVMFunctionInfo* FunctionInfo = VMFunctionInfos.Find(BindingInfo.Name))
 	{
 		OutFunc = FunctionInfo->FunctionBinding;
-
-		FInstanceData_GameThread* InstanceData = static_cast<FInstanceData_GameThread*>(PerInstanceData);
-		InstanceData->bReadsFilteredTransforms		|= FunctionInfo->TransformReadType == ESocketReadType::Default || FunctionInfo->TransformReadType == ESocketReadType::Filtered;
-		InstanceData->bReadsUnfilteredTransforms	|= FunctionInfo->TransformReadType == ESocketReadType::Default || FunctionInfo->TransformReadType == ESocketReadType::Unfiltered;
 	}
 }
 
@@ -839,7 +838,23 @@ bool UNiagaraDataInterfaceSocketReader::InitPerInstanceData(void* PerInstanceDat
 	InstanceData->bReadsFilteredTransforms |= GDisableTransformReadType;
 	InstanceData->bReadsUnfilteredTransforms |= GDisableTransformReadType;
 
-	if (IsUsedWithGPUScript())
+	if (IsUsedWithCPUScript() && (!InstanceData->bReadsFilteredTransforms || !InstanceData->bReadsUnfilteredTransforms))
+	{
+		FNiagaraDataInterfaceUtilities::ForEachVMFunction(
+			this, SystemInstance,
+			[&](const UNiagaraScript* Script, const FVMExternalFunctionBindingInfo& Function)
+			{
+				if (const FVMFunctionInfo* FunctionInfo = VMFunctionInfos.Find(Function.Name))
+				{
+					InstanceData->bReadsFilteredTransforms |= FunctionInfo->TransformReadType == ESocketReadType::Any || FunctionInfo->TransformReadType == ESocketReadType::Filtered;
+					InstanceData->bReadsUnfilteredTransforms |= FunctionInfo->TransformReadType == ESocketReadType::Any || FunctionInfo->TransformReadType == ESocketReadType::Unfiltered;
+				}
+				return !InstanceData->bReadsFilteredTransforms || !InstanceData->bReadsUnfilteredTransforms;
+			}
+		);
+	}
+
+	if (IsUsedWithGPUScript() && (!InstanceData->bReadsFilteredTransforms || !InstanceData->bReadsUnfilteredTransforms))
 	{
 		FNiagaraDataInterfaceUtilities::ForEachGpuFunction(
 			this, SystemInstance,
@@ -847,11 +862,10 @@ bool UNiagaraDataInterfaceSocketReader::InitPerInstanceData(void* PerInstanceDat
 			{
 				if (const FVMFunctionInfo* FunctionInfo = VMFunctionInfos.Find(Function.DefinitionName))
 				{
-					InstanceData->bReadsFilteredTransforms |= FunctionInfo->TransformReadType == ESocketReadType::Default || FunctionInfo->TransformReadType == ESocketReadType::Filtered;
-					InstanceData->bReadsUnfilteredTransforms |= FunctionInfo->TransformReadType == ESocketReadType::Default || FunctionInfo->TransformReadType == ESocketReadType::Unfiltered;
+					InstanceData->bReadsFilteredTransforms |= FunctionInfo->TransformReadType == ESocketReadType::Any || FunctionInfo->TransformReadType == ESocketReadType::Filtered;
+					InstanceData->bReadsUnfilteredTransforms |= FunctionInfo->TransformReadType == ESocketReadType::Any || FunctionInfo->TransformReadType == ESocketReadType::Unfiltered;
 				}
-				// Keeping going unless there is no more information we can gather
-				return (InstanceData->bReadsFilteredTransforms && InstanceData->bReadsUnfilteredTransforms) != true;
+				return !InstanceData->bReadsFilteredTransforms || !InstanceData->bReadsUnfilteredTransforms;
 			}
 		);
 	}
