- Remove experimental from various data interface functions

#tests Add nodes which were experimental before / run ET

diff --git a/Engine/Plugins/FX/Niagara/Source/Niagara/Classes/NiagaraDataInterfaceIntRenderTarget2D.h b/Engine/Plugins/FX/Niagara/Source/Niagara/Classes/NiagaraDataInterfaceIntRenderTarget2D.h
--- a/Engine/Plugins/FX/Niagara/Source/Niagara/Classes/NiagaraDataInterfaceIntRenderTarget2D.h
+++ b/Engine/Plugins/FX/Niagara/Source/Niagara/Classes/NiagaraDataInterfaceIntRenderTarget2D.h
@@ -9,7 +9,7 @@
 
 class UTextureRenderTarget2D;
 
-UCLASS(EditInlineNew, Category = "Render Target", CollapseCategories, meta = (DisplayName = "Integer Render Target 2D", Experimental), Blueprintable, BlueprintType, MinimalAPI)
+UCLASS(EditInlineNew, Category = "Render Target", CollapseCategories, meta = (DisplayName = "Integer Render Target 2D"), Blueprintable, BlueprintType, MinimalAPI)
 class UNiagaraDataInterfaceIntRenderTarget2D : public UNiagaraDataInterfaceRWBase
 {
 	GENERATED_UCLASS_BODY()
diff --git a/Engine/Plugins/FX/Niagara/Source/Niagara/Private/DataInterface/NiagaraDataInterfaceGBuffer.cpp b/Engine/Plugins/FX/Niagara/Source/Niagara/Private/DataInterface/NiagaraDataInterfaceGBuffer.cpp
--- a/Engine/Plugins/FX/Niagara/Source/Niagara/Private/DataInterface/NiagaraDataInterfaceGBuffer.cpp
+++ b/Engine/Plugins/FX/Niagara/Source/Niagara/Private/DataInterface/NiagaraDataInterfaceGBuffer.cpp
@@ -160,7 +160,6 @@ void UNiagaraDataInterfaceGBuffer::GetFunctionsInternal(TArray<FNiagaraFunctionS
 		Signature.bMemberFunction = true;
 		Signature.bRequiresContext = false;
 		Signature.bSupportsCPU = false;
-		Signature.bExperimental = true;
 		Signature.MiscUsageBitMask = Attribute.MiscUsageBitMask;
 		Signature.Inputs.Add(FNiagaraVariable(FNiagaraTypeDefinition(GetClass()), TEXT("GBufferInterface")));
 		Signature.Inputs.Add(FNiagaraVariable(FNiagaraTypeDefinition::GetVec2Def(), TEXT("ScreenUV")));
diff --git a/Engine/Plugins/FX/Niagara/Source/Niagara/Private/DataInterface/NiagaraDataInterfaceVirtualTexture.cpp b/Engine/Plugins/FX/Niagara/Source/Niagara/Private/DataInterface/NiagaraDataInterfaceVirtualTexture.cpp
--- a/Engine/Plugins/FX/Niagara/Source/Niagara/Private/DataInterface/NiagaraDataInterfaceVirtualTexture.cpp
+++ b/Engine/Plugins/FX/Niagara/Source/Niagara/Private/DataInterface/NiagaraDataInterfaceVirtualTexture.cpp
@@ -187,7 +187,6 @@ void UNiagaraDataInterfaceVirtualTexture::GetFunctionsInternal(TArray<FNiagaraFu
 		Sig.bRequiresContext = false;
 		Sig.bSupportsCPU = false;
 		Sig.bSupportsGPU = true;
-		Sig.bExperimental = true;
 		Sig.SetFunctionVersion(FNDIFunctionVersion::LatestVersion);
 		Sig.Inputs.Emplace(FNiagaraTypeDefinition(GetClass()), TEXT("Texture"));
 		Sig.Outputs.Emplace(FNiagaraTypeDefinition::GetBoolDef(), TEXT("BaseColorValid"));
@@ -205,7 +204,6 @@ void UNiagaraDataInterfaceVirtualTexture::GetFunctionsInternal(TArray<FNiagaraFu
 		Sig.bRequiresContext = false;		
 		Sig.bSupportsCPU = false;
 		Sig.bSupportsGPU = true;
-		Sig.bExperimental = true;
 		Sig.SetFunctionVersion(FNDIFunctionVersion::LatestVersion);
 		Sig.Inputs.Emplace(FNiagaraTypeDefinition(GetClass()), TEXT("Texture"));
 		Sig.Inputs.Emplace(FNiagaraTypeDefinition::GetPositionDef(), TEXT("WorldPosition"));
diff --git a/Engine/Plugins/FX/Niagara/Source/Niagara/Private/DataInterface/NiagaraDataInterfaceVirtualTexture.h b/Engine/Plugins/FX/Niagara/Source/Niagara/Private/DataInterface/NiagaraDataInterfaceVirtualTexture.h
--- a/Engine/Plugins/FX/Niagara/Source/Niagara/Private/DataInterface/NiagaraDataInterfaceVirtualTexture.h
+++ b/Engine/Plugins/FX/Niagara/Source/Niagara/Private/DataInterface/NiagaraDataInterfaceVirtualTexture.h
@@ -12,7 +12,7 @@ struct FNiagaraDataInterfaceGeneratedFunction;
 class URuntimeVirtualTexture;
 
 /** Data Interface allowing sampling of a texture */
-UCLASS(Experimental, EditInlineNew, Category = "Texture", CollapseCategories, meta = (DisplayName = "Runtime Virtual Texture Sample"), MinimalAPI)
+UCLASS(EditInlineNew, Category = "Texture", CollapseCategories, meta = (DisplayName = "Runtime Virtual Texture Sample"), MinimalAPI)
 class UNiagaraDataInterfaceVirtualTexture : public UNiagaraDataInterface
 {
 	GENERATED_UCLASS_BODY()
diff --git a/Engine/Plugins/FX/Niagara/Source/Niagara/Private/NDISkeletalMesh_TriangleSampling.cpp b/Engine/Plugins/FX/Niagara/Source/Niagara/Private/NDISkeletalMesh_TriangleSampling.cpp
--- a/Engine/Plugins/FX/Niagara/Source/Niagara/Private/NDISkeletalMesh_TriangleSampling.cpp
+++ b/Engine/Plugins/FX/Niagara/Source/Niagara/Private/NDISkeletalMesh_TriangleSampling.cpp
@@ -298,7 +298,6 @@ void UNiagaraDataInterfaceSkeletalMesh::GetTriangleSamplingFunctions(TArray<FNia
 		Sig.Outputs.Add(FNiagaraVariable(FNiagaraTypeDefinition::GetBoolDef(), TEXT("IsValid")));
 		Sig.bMemberFunction = true;
 		Sig.bRequiresContext = false;
-		Sig.bExperimental = true;
 	}
 
 	{
@@ -316,7 +315,6 @@ void UNiagaraDataInterfaceSkeletalMesh::GetTriangleSamplingFunctions(TArray<FNia
 		Sig.Outputs.Add(FNiagaraVariable(FNiagaraTypeDefinition::GetBoolDef(), TEXT("IsValid")));
 		Sig.bMemberFunction = true;
 		Sig.bRequiresContext = false;
-		Sig.bExperimental = true;
 	}
 
 	{
@@ -329,7 +327,6 @@ void UNiagaraDataInterfaceSkeletalMesh::GetTriangleSamplingFunctions(TArray<FNia
 		Sig.Outputs.Add(FNiagaraVariable(FNiagaraTypeDefinition::GetBoolDef(), TEXT("IsValid")));
 		Sig.bMemberFunction = true;
 		Sig.bRequiresContext = false;
-		Sig.bExperimental = true;
 		Sig.bSupportsCPU = false;
 	}
 
@@ -344,7 +341,6 @@ void UNiagaraDataInterfaceSkeletalMesh::GetTriangleSamplingFunctions(TArray<FNia
 		Sig.Outputs.Add(FNiagaraVariable(FNiagaraTypeDefinition::GetBoolDef(), TEXT("IsValid")));
 		Sig.bMemberFunction = true;
 		Sig.bRequiresContext = false;
-		Sig.bExperimental = true;
 		Sig.bSupportsCPU = false;
 	}
 }
diff --git a/Engine/Plugins/FX/Niagara/Source/Niagara/Private/NiagaraDataInterfaceArrayImpl.cpp b/Engine/Plugins/FX/Niagara/Source/Niagara/Private/NiagaraDataInterfaceArrayImpl.cpp
--- a/Engine/Plugins/FX/Niagara/Source/Niagara/Private/NiagaraDataInterfaceArrayImpl.cpp
+++ b/Engine/Plugins/FX/Niagara/Source/Niagara/Private/NiagaraDataInterfaceArrayImpl.cpp
@@ -51,8 +51,6 @@ bool FNiagaraDataInterfaceArrayImplHelper::UpgradeFunctionCall(FNiagaraFunctionS
 			}
 		}
 
-		FunctionSignature.bExperimental = false;
-
 		if (FunctionSignature.Name == FNiagaraDataInterfaceArrayImplHelper::Function_SetArrayElemName)
 		{
 			FunctionSignature.Inputs.EmplaceAt(1, FNiagaraTypeDefinition::GetBoolDef(), TEXT("SkipSet"));
diff --git a/Engine/Plugins/FX/Niagara/Source/Niagara/Private/NiagaraDataInterfaceGrid2DCollection.cpp b/Engine/Plugins/FX/Niagara/Source/Niagara/Private/NiagaraDataInterfaceGrid2DCollection.cpp
--- a/Engine/Plugins/FX/Niagara/Source/Niagara/Private/NiagaraDataInterfaceGrid2DCollection.cpp
+++ b/Engine/Plugins/FX/Niagara/Source/Niagara/Private/NiagaraDataInterfaceGrid2DCollection.cpp
@@ -421,7 +421,6 @@ void UNiagaraDataInterfaceGrid2DCollection::GetFunctionsInternal(TArray<FNiagara
 		Sig.Outputs.Add(FNiagaraVariable(FNiagaraTypeDefinition::GetBoolDef(), TEXT("Success")));
 
 		Sig.ModuleUsageBitmask = ENiagaraScriptUsageMask::Emitter | ENiagaraScriptUsageMask::System;
-		Sig.bExperimental = true;
 		Sig.bMemberFunction = true;
 		Sig.bRequiresExecPin = true;
 		Sig.bRequiresContext = false;
@@ -440,7 +439,6 @@ void UNiagaraDataInterfaceGrid2DCollection::GetFunctionsInternal(TArray<FNiagara
 		Sig.Inputs.Add(FNiagaraVariable(FNiagaraTypeDefinition::GetIntDef(), TEXT("AttributeIndex")));
 		Sig.Outputs.Add(FNiagaraVariable(FNiagaraTypeDefinition::GetFloatDef(), TEXT("Value")));
 
-		Sig.bExperimental = true;
 		Sig.bMemberFunction = true;
 		Sig.bRequiresContext = false;
 		Sig.ModuleUsageBitmask = ENiagaraScriptUsageMask::Particle;
@@ -462,7 +460,6 @@ void UNiagaraDataInterfaceGrid2DCollection::GetFunctionsInternal(TArray<FNiagara
 		Sig.Inputs.Add(FNiagaraVariable(FNiagaraTypeDefinition::GetIntDef(), TEXT("AttributeIndex")));
 		Sig.Outputs.Add(FNiagaraVariable(FNiagaraTypeDefinition::GetFloatDef(), TEXT("Value")));
 
-		Sig.bExperimental = true;
 		Sig.bMemberFunction = true;
 		Sig.bRequiresContext = false;
 		Sig.ModuleUsageBitmask = ENiagaraScriptUsageMask::Particle;
@@ -485,7 +482,6 @@ void UNiagaraDataInterfaceGrid2DCollection::GetFunctionsInternal(TArray<FNiagara
 		Sig.Inputs.Add(FNiagaraVariable(FNiagaraTypeDefinition::GetFloatDef(), TEXT("Value")));
 		Sig.Outputs.Add(FNiagaraVariable(FNiagaraTypeDefinition::GetIntDef(), TEXT("IGNORE")));
 
-		Sig.bExperimental = true;
 		Sig.bMemberFunction = true;
 		Sig.bRequiresContext = false;
 		Sig.bWriteFunction = true;
@@ -506,7 +502,6 @@ void UNiagaraDataInterfaceGrid2DCollection::GetFunctionsInternal(TArray<FNiagara
 		Sig.Inputs.Add(FNiagaraVariable(FNiagaraTypeDefinition::GetIntDef(), TEXT("AttributeIndex")));
 		Sig.Inputs.Add(FNiagaraVariable(FNiagaraTypeDefinition::GetFloatDef(), TEXT("Value")));
 
-		Sig.bExperimental = true;
 		Sig.bRequiresExecPin = true;
 		Sig.bMemberFunction = true;
 		Sig.bRequiresContext = false;
@@ -526,7 +521,6 @@ void UNiagaraDataInterfaceGrid2DCollection::GetFunctionsInternal(TArray<FNiagara
 		Sig.Inputs.Add(FNiagaraVariable(FNiagaraTypeDefinition::GetIntDef(), TEXT("IndexX")));
 		Sig.Inputs.Add(FNiagaraVariable(FNiagaraTypeDefinition::GetIntDef(), TEXT("IndexY")));
 
-		Sig.bExperimental = true;
 		Sig.bMemberFunction = true;
 		Sig.bRequiresContext = false;
 		Sig.bWriteFunction = true;
@@ -545,7 +539,6 @@ void UNiagaraDataInterfaceGrid2DCollection::GetFunctionsInternal(TArray<FNiagara
 		Sig.Inputs.Add(FNiagaraVariable(FNiagaraTypeDefinition::GetIntDef(), TEXT("IndexX")));
 		Sig.Inputs.Add(FNiagaraVariable(FNiagaraTypeDefinition::GetIntDef(), TEXT("IndexY")));
 
-		Sig.bExperimental = true;
 		Sig.bMemberFunction = true;
 		Sig.bRequiresContext = false;
 		Sig.bWriteFunction = true;
@@ -568,7 +561,6 @@ void UNiagaraDataInterfaceGrid2DCollection::GetFunctionsInternal(TArray<FNiagara
 		Sig.FunctionSpecifiers.Add(FName("Attribute"));
 		Sig.bMemberFunction = true;
 		Sig.bRequiresContext = false;
-		Sig.bExperimental = true;
 		Sig.bRequiresExecPin = true;
 		Sig.bWriteFunction = true;
 		Sig.ModuleUsageBitmask = ENiagaraScriptUsageMask::Particle;
@@ -588,7 +580,6 @@ void UNiagaraDataInterfaceGrid2DCollection::GetFunctionsInternal(TArray<FNiagara
 		Sig.FunctionSpecifiers.Add(FName("Attribute"));
 		Sig.bMemberFunction = true;
 		Sig.bRequiresContext = false;
-		Sig.bExperimental = true;
 		Sig.ModuleUsageBitmask = ENiagaraScriptUsageMask::Particle;
 		Sig.bSupportsCPU = false;
 		Sig.bSupportsGPU = true;
@@ -611,7 +602,6 @@ void UNiagaraDataInterfaceGrid2DCollection::GetFunctionsInternal(TArray<FNiagara
 		Sig.FunctionSpecifiers.Add(FName("Attribute"));
 		Sig.bMemberFunction = true;
 		Sig.bRequiresContext = false;
-		Sig.bExperimental = true;
 		Sig.ModuleUsageBitmask = ENiagaraScriptUsageMask::Particle;
 		Sig.bSupportsCPU = false;
 		Sig.bSupportsGPU = true;
@@ -634,7 +624,6 @@ void UNiagaraDataInterfaceGrid2DCollection::GetFunctionsInternal(TArray<FNiagara
 		Sig.FunctionSpecifiers.Add(FName("Attribute"));
 		Sig.bMemberFunction = true;
 		Sig.bRequiresContext = false;
-		Sig.bExperimental = true;
 		Sig.ModuleUsageBitmask = ENiagaraScriptUsageMask::Particle;
 		Sig.bSupportsCPU = false;
 		Sig.bSupportsGPU = true;
@@ -658,7 +647,6 @@ void UNiagaraDataInterfaceGrid2DCollection::GetFunctionsInternal(TArray<FNiagara
 		Sig.FunctionSpecifiers.Add(FName("Attribute"));
 		Sig.bMemberFunction = true;
 		Sig.bRequiresContext = false;
-		Sig.bExperimental = true;
 		Sig.bRequiresExecPin = true;
 		Sig.bWriteFunction = true;
 		Sig.ModuleUsageBitmask = ENiagaraScriptUsageMask::Particle;
@@ -683,7 +671,6 @@ void UNiagaraDataInterfaceGrid2DCollection::GetFunctionsInternal(TArray<FNiagara
 		Sig.FunctionSpecifiers.Add(FName("Attribute"));
 		Sig.bMemberFunction = true;
 		Sig.bRequiresContext = false;
-		Sig.bExperimental = true;
 		Sig.ModuleUsageBitmask = ENiagaraScriptUsageMask::Particle;
 		Sig.bSupportsCPU = false;
 		Sig.bSupportsGPU = true;
@@ -706,7 +693,6 @@ void UNiagaraDataInterfaceGrid2DCollection::GetFunctionsInternal(TArray<FNiagara
 		Sig.FunctionSpecifiers.Add(FName("Attribute"));
 		Sig.bMemberFunction = true;
 		Sig.bRequiresContext = false;
-		Sig.bExperimental = true;
 		Sig.ModuleUsageBitmask = ENiagaraScriptUsageMask::Particle;
 		Sig.bSupportsCPU = false;
 		Sig.bSupportsGPU = true;
@@ -729,7 +715,6 @@ void UNiagaraDataInterfaceGrid2DCollection::GetFunctionsInternal(TArray<FNiagara
 		Sig.FunctionSpecifiers.Add(FName("Attribute"));
 		Sig.bMemberFunction = true;
 		Sig.bRequiresContext = false;
-		Sig.bExperimental = true;
 		Sig.ModuleUsageBitmask = ENiagaraScriptUsageMask::Particle;
 		Sig.bSupportsCPU = false;
 		Sig.bSupportsGPU = true;
@@ -753,7 +738,6 @@ void UNiagaraDataInterfaceGrid2DCollection::GetFunctionsInternal(TArray<FNiagara
 		Sig.FunctionSpecifiers.Add(FName("Attribute"));
 		Sig.bMemberFunction = true;
 		Sig.bRequiresContext = false;
-		Sig.bExperimental = true;
 		Sig.bRequiresExecPin = true;
 		Sig.bWriteFunction = true;
 		Sig.ModuleUsageBitmask = ENiagaraScriptUsageMask::Particle;
@@ -778,7 +762,6 @@ void UNiagaraDataInterfaceGrid2DCollection::GetFunctionsInternal(TArray<FNiagara
 		Sig.FunctionSpecifiers.Add(FName("Attribute"));
 		Sig.bMemberFunction = true;
 		Sig.bRequiresContext = false;
-		Sig.bExperimental = true;
 		Sig.ModuleUsageBitmask = ENiagaraScriptUsageMask::Particle;
 		Sig.bSupportsCPU = false;
 		Sig.bSupportsGPU = true;
@@ -801,7 +784,6 @@ void UNiagaraDataInterfaceGrid2DCollection::GetFunctionsInternal(TArray<FNiagara
 		Sig.FunctionSpecifiers.Add(FName("Attribute"));
 		Sig.bMemberFunction = true;
 		Sig.bRequiresContext = false;
-		Sig.bExperimental = true;
 		Sig.ModuleUsageBitmask = ENiagaraScriptUsageMask::Particle;
 		Sig.bSupportsCPU = false;
 		Sig.bSupportsGPU = true;
@@ -824,7 +806,6 @@ void UNiagaraDataInterfaceGrid2DCollection::GetFunctionsInternal(TArray<FNiagara
 		Sig.FunctionSpecifiers.Add(FName("Attribute"));
 		Sig.bMemberFunction = true;
 		Sig.bRequiresContext = false;
-		Sig.bExperimental = true;
 		Sig.ModuleUsageBitmask = ENiagaraScriptUsageMask::Particle;
 		Sig.bSupportsCPU = false;
 		Sig.bSupportsGPU = true;
@@ -848,7 +829,6 @@ void UNiagaraDataInterfaceGrid2DCollection::GetFunctionsInternal(TArray<FNiagara
 		Sig.FunctionSpecifiers.Add(FName("Attribute"));
 		Sig.bMemberFunction = true;
 		Sig.bRequiresContext = false;
-		Sig.bExperimental = true;
 		Sig.bRequiresExecPin = true;
 		Sig.bWriteFunction = true;
 		Sig.ModuleUsageBitmask = ENiagaraScriptUsageMask::Particle;
@@ -868,7 +848,6 @@ void UNiagaraDataInterfaceGrid2DCollection::GetFunctionsInternal(TArray<FNiagara
 		Sig.FunctionSpecifiers.Add(FName("Attribute"));
 		Sig.bMemberFunction = true;
 		Sig.bRequiresContext = false;
-		Sig.bExperimental = true;
 		Sig.ModuleUsageBitmask = ENiagaraScriptUsageMask::Particle;
 		Sig.bSupportsCPU = false;
 		Sig.bSupportsGPU = true;
@@ -891,7 +870,6 @@ void UNiagaraDataInterfaceGrid2DCollection::GetFunctionsInternal(TArray<FNiagara
 		Sig.FunctionSpecifiers.Add(FName("Attribute"));
 		Sig.bMemberFunction = true;
 		Sig.bRequiresContext = false;
-		Sig.bExperimental = true;
 		Sig.ModuleUsageBitmask = ENiagaraScriptUsageMask::Particle;
 		Sig.bSupportsCPU = false;
 		Sig.bSupportsGPU = true;
@@ -914,7 +892,6 @@ void UNiagaraDataInterfaceGrid2DCollection::GetFunctionsInternal(TArray<FNiagara
 		Sig.FunctionSpecifiers.Add(FName("Attribute"));
 		Sig.bMemberFunction = true;
 		Sig.bRequiresContext = false;
-		Sig.bExperimental = true;
 		Sig.ModuleUsageBitmask = ENiagaraScriptUsageMask::Particle;
 		Sig.bSupportsCPU = false;
 		Sig.bSupportsGPU = true;
@@ -937,7 +914,6 @@ void UNiagaraDataInterfaceGrid2DCollection::GetFunctionsInternal(TArray<FNiagara
 		Sig.Inputs.Add(FNiagaraVariable(FNiagaraTypeDefinition::GetIntDef(), TEXT("AttributeIndex")));
 		Sig.Outputs.Add(FNiagaraVariable(FNiagaraTypeDefinition::GetFloatDef(), TEXT("Value")));
 
-		Sig.bExperimental = true;
 		Sig.bMemberFunction = true;
 		Sig.bRequiresContext = false;
 		Sig.ModuleUsageBitmask = ENiagaraScriptUsageMask::Particle;
@@ -957,7 +933,6 @@ void UNiagaraDataInterfaceGrid2DCollection::GetFunctionsInternal(TArray<FNiagara
 		Sig.Inputs.Add(FNiagaraVariable(FNiagaraTypeDefinition::GetIntDef(), TEXT("AttributeIndex")));
 		Sig.Outputs.Add(FNiagaraVariable(FNiagaraTypeDefinition::GetFloatDef(), TEXT("Value")));
 
-		Sig.bExperimental = true;
 		Sig.bMemberFunction = true;
 		Sig.bRequiresContext = false;
 		Sig.ModuleUsageBitmask = ENiagaraScriptUsageMask::Particle;
@@ -976,7 +951,6 @@ void UNiagaraDataInterfaceGrid2DCollection::GetFunctionsInternal(TArray<FNiagara
 		Sig.Inputs.Add(FNiagaraVariable(FNiagaraTypeDefinition::GetIntDef(), TEXT("AttributeIndex")));
 		Sig.Outputs.Add(FNiagaraVariable(FNiagaraTypeDefinition::GetFloatDef(), TEXT("Value")));
 
-		Sig.bExperimental = true;
 		Sig.bMemberFunction = true;
 		Sig.bRequiresContext = false;
 		Sig.ModuleUsageBitmask = ENiagaraScriptUsageMask::Particle;
@@ -994,7 +968,6 @@ void UNiagaraDataInterfaceGrid2DCollection::GetFunctionsInternal(TArray<FNiagara
 		Sig.FunctionSpecifiers.Add(FName("Attribute"));
 		Sig.bMemberFunction = true;
 		Sig.bRequiresContext = false;
-		Sig.bExperimental = true;
 		Sig.Description = NSLOCTEXT("Niagara", "NiagaraDataInterfaceGridColl2D_GetVector4AttributeIndex", "Gets a attribute starting index value for Vector4 on the Grid by Attribute name. Returns -1 if not found.");
 		OutFunctions.Add(Sig);
 	}
@@ -1007,7 +980,6 @@ void UNiagaraDataInterfaceGrid2DCollection::GetFunctionsInternal(TArray<FNiagara
 		Sig.FunctionSpecifiers.Add(FName("Attribute"));
 		Sig.bMemberFunction = true;
 		Sig.bRequiresContext = false;
-		Sig.bExperimental = true;
 		Sig.Description = NSLOCTEXT("Niagara", "NiagaraDataInterfaceGridColl2D_GetVector3AttributeIndex", "Gets a attribute starting index value for Vector3 on the Grid by Attribute name. Returns -1 if not found.");
 		OutFunctions.Add(Sig);
 	}
@@ -1020,7 +992,6 @@ void UNiagaraDataInterfaceGrid2DCollection::GetFunctionsInternal(TArray<FNiagara
 		Sig.FunctionSpecifiers.Add(FName("Attribute"));
 		Sig.bMemberFunction = true;
 		Sig.bRequiresContext = false;
-		Sig.bExperimental = true;
 		Sig.Description = NSLOCTEXT("Niagara", "NiagaraDataInterfaceGridColl2D_GetVector2AttributeIndex", "Gets a attribute starting index value for Vector2 on the Grid by Attribute name. Returns -1 if not found.");
 		OutFunctions.Add(Sig);
 	}
@@ -1033,7 +1004,6 @@ void UNiagaraDataInterfaceGrid2DCollection::GetFunctionsInternal(TArray<FNiagara
 		Sig.FunctionSpecifiers.Add(FName("Attribute"));
 		Sig.bMemberFunction = true;
 		Sig.bRequiresContext = false;
-		Sig.bExperimental = true;
 		Sig.Description = NSLOCTEXT("Niagara", "NiagaraDataInterfaceGridColl2D_GetFloatAttributeIndex", "Gets a attribute starting index value for float on the Grid by Attribute name. Returns -1 if not found.");
 		OutFunctions.Add(Sig);
 	}
diff --git a/Engine/Plugins/FX/Niagara/Source/Niagara/Private/NiagaraDataInterfaceGrid3DCollection.cpp b/Engine/Plugins/FX/Niagara/Source/Niagara/Private/NiagaraDataInterfaceGrid3DCollection.cpp
--- a/Engine/Plugins/FX/Niagara/Source/Niagara/Private/NiagaraDataInterfaceGrid3DCollection.cpp
+++ b/Engine/Plugins/FX/Niagara/Source/Niagara/Private/NiagaraDataInterfaceGrid3DCollection.cpp
@@ -432,7 +432,6 @@ void UNiagaraDataInterfaceGrid3DCollection::GetFunctionsInternal(TArray<FNiagara
 		Sig.Outputs.Add(FNiagaraVariable(FNiagaraTypeDefinition::GetBoolDef(), TEXT("Success")));
 
 		Sig.ModuleUsageBitmask = ENiagaraScriptUsageMask::Emitter | ENiagaraScriptUsageMask::System;
-		Sig.bExperimental = true;
 		Sig.bMemberFunction = true;
 		Sig.bRequiresExecPin = true;
 		Sig.bRequiresContext = false;
@@ -452,7 +451,6 @@ void UNiagaraDataInterfaceGrid3DCollection::GetFunctionsInternal(TArray<FNiagara
 		Sig.Inputs.Add(FNiagaraVariable(FNiagaraTypeDefinition::GetIntDef(), TEXT("AttributeIndex")));
 		Sig.Outputs.Add(FNiagaraVariable(FNiagaraTypeDefinition::GetFloatDef(), TEXT("Value")));
 
-		Sig.bExperimental = true;
 		Sig.bMemberFunction = true;
 		Sig.bRequiresContext = false;
 		Sig.ModuleUsageBitmask = ENiagaraScriptUsageMask::Particle;
@@ -475,7 +473,6 @@ void UNiagaraDataInterfaceGrid3DCollection::GetFunctionsInternal(TArray<FNiagara
 		Sig.Inputs.Add(FNiagaraVariable(FNiagaraTypeDefinition::GetIntDef(), TEXT("AttributeIndex")));
 		Sig.Outputs.Add(FNiagaraVariable(FNiagaraTypeDefinition::GetFloatDef(), TEXT("Value")));
 
-		Sig.bExperimental = true;
 		Sig.bMemberFunction = true;
 		Sig.bRequiresContext = false;
 		Sig.ModuleUsageBitmask = ENiagaraScriptUsageMask::Particle;
@@ -496,7 +493,6 @@ void UNiagaraDataInterfaceGrid3DCollection::GetFunctionsInternal(TArray<FNiagara
 		Sig.Inputs.Add(FNiagaraVariable(FNiagaraTypeDefinition::GetIntDef(), TEXT("IndexZ")));
 		Sig.Outputs.Add(FNiagaraVariable(FNiagaraTypeDefinition::GetFloatDef(), TEXT("Value")));
 
-		Sig.bExperimental = true;
 		Sig.bMemberFunction = true;
 		Sig.bRequiresContext = false;
 		Sig.ModuleUsageBitmask = ENiagaraScriptUsageMask::Particle;
@@ -521,7 +517,6 @@ void UNiagaraDataInterfaceGrid3DCollection::GetFunctionsInternal(TArray<FNiagara
 		Sig.Inputs.Add(FNiagaraVariable(FNiagaraTypeDefinition::GetFloatDef(), TEXT("Value")));
 		Sig.Outputs.Add(FNiagaraVariable(FNiagaraTypeDefinition::GetIntDef(), TEXT("IGNORE")));
 
-		Sig.bExperimental = true;
 		Sig.bMemberFunction = true;
 		Sig.bRequiresContext = false;
 		Sig.bWriteFunction = true;
@@ -543,7 +538,6 @@ void UNiagaraDataInterfaceGrid3DCollection::GetFunctionsInternal(TArray<FNiagara
 		Sig.Inputs.Add(FNiagaraVariable(FNiagaraTypeDefinition::GetFloatDef(), TEXT("Value")));
 		Sig.Outputs.Add(FNiagaraVariable(FNiagaraTypeDefinition::GetIntDef(), TEXT("IGNORE")));
 
-		Sig.bExperimental = true;
 		Sig.bMemberFunction = true;
 		Sig.bRequiresContext = false;
 		Sig.bWriteFunction = true;
@@ -564,7 +558,6 @@ void UNiagaraDataInterfaceGrid3DCollection::GetFunctionsInternal(TArray<FNiagara
 		Sig.Inputs.Add(FNiagaraVariable(FNiagaraTypeDefinition::GetIntDef(), TEXT("IndexY")));
 		Sig.Inputs.Add(FNiagaraVariable(FNiagaraTypeDefinition::GetIntDef(), TEXT("IndexZ")));
 
-		Sig.bExperimental = true;
 		Sig.bMemberFunction = true;
 		Sig.bRequiresContext = false;
 		Sig.bWriteFunction = true;
@@ -584,7 +577,6 @@ void UNiagaraDataInterfaceGrid3DCollection::GetFunctionsInternal(TArray<FNiagara
 		Sig.Inputs.Add(FNiagaraVariable(FNiagaraTypeDefinition::GetIntDef(), TEXT("IndexY")));
 		Sig.Inputs.Add(FNiagaraVariable(FNiagaraTypeDefinition::GetIntDef(), TEXT("IndexZ")));
 
-		Sig.bExperimental = true;
 		Sig.bMemberFunction = true;
 		Sig.bRequiresContext = false;
 		Sig.bWriteFunction = true;
@@ -607,7 +599,6 @@ void UNiagaraDataInterfaceGrid3DCollection::GetFunctionsInternal(TArray<FNiagara
 		Sig.Inputs.Add(FNiagaraVariable(FNiagaraTypeDefinition::GetIntDef(), TEXT("NumAttributesSet")));
 		Sig.Inputs.Add(FNiagaraVariable(FNiagaraTypeDefinition::GetIntDef(), TEXT("AttributeMask")));
 
-		Sig.bExperimental = true;
 		Sig.bMemberFunction = true;
 		Sig.bRequiresContext = false;
 		Sig.bWriteFunction = true;
@@ -631,7 +622,6 @@ void UNiagaraDataInterfaceGrid3DCollection::GetFunctionsInternal(TArray<FNiagara
 		Sig.FunctionSpecifiers.Add(FName(UNiagaraDataInterfaceRWBase::NAME_Attribute));
 		Sig.bMemberFunction = true;
 		Sig.bRequiresContext = false;
-		Sig.bExperimental = true;
 		Sig.bRequiresExecPin = true;
 		Sig.bWriteFunction = true;
 		Sig.ModuleUsageBitmask = ENiagaraScriptUsageMask::Particle;
@@ -652,7 +642,6 @@ void UNiagaraDataInterfaceGrid3DCollection::GetFunctionsInternal(TArray<FNiagara
 		Sig.FunctionSpecifiers.Add(FName(UNiagaraDataInterfaceRWBase::NAME_Attribute));
 		Sig.bMemberFunction = true;
 		Sig.bRequiresContext = false;
-		Sig.bExperimental = true;
 		Sig.ModuleUsageBitmask = ENiagaraScriptUsageMask::Particle;
 		Sig.bSupportsCPU = false;
 		Sig.bSupportsGPU = true;
@@ -671,7 +660,6 @@ void UNiagaraDataInterfaceGrid3DCollection::GetFunctionsInternal(TArray<FNiagara
 		Sig.FunctionSpecifiers.Add(FName(UNiagaraDataInterfaceRWBase::NAME_Attribute));
 		Sig.bMemberFunction = true;
 		Sig.bRequiresContext = false;
-		Sig.bExperimental = true;
 		Sig.ModuleUsageBitmask = ENiagaraScriptUsageMask::Particle;
 		Sig.bSupportsCPU = false;
 		Sig.bSupportsGPU = true;
@@ -690,7 +678,6 @@ void UNiagaraDataInterfaceGrid3DCollection::GetFunctionsInternal(TArray<FNiagara
 		Sig.FunctionSpecifiers.Add(FName(UNiagaraDataInterfaceRWBase::NAME_Attribute));
 		Sig.bMemberFunction = true;
 		Sig.bRequiresContext = false;
-		Sig.bExperimental = true;
 		Sig.ModuleUsageBitmask = ENiagaraScriptUsageMask::Particle;
 		Sig.bSupportsCPU = false;
 		Sig.bSupportsGPU = true;
@@ -711,7 +698,6 @@ void UNiagaraDataInterfaceGrid3DCollection::GetFunctionsInternal(TArray<FNiagara
 		Sig.FunctionSpecifiers.Add(FName(UNiagaraDataInterfaceRWBase::NAME_Attribute));
 		Sig.bMemberFunction = true;
 		Sig.bRequiresContext = false;
-		Sig.bExperimental = true;
 		Sig.bRequiresExecPin = true;
 		Sig.bWriteFunction = true;
 		Sig.ModuleUsageBitmask = ENiagaraScriptUsageMask::Particle;
@@ -733,7 +719,6 @@ void UNiagaraDataInterfaceGrid3DCollection::GetFunctionsInternal(TArray<FNiagara
 		Sig.FunctionSpecifiers.Add(FName(UNiagaraDataInterfaceRWBase::NAME_Attribute));
 		Sig.bMemberFunction = true;
 		Sig.bRequiresContext = false;
-		Sig.bExperimental = true;
 		Sig.ModuleUsageBitmask = ENiagaraScriptUsageMask::Particle;
 		Sig.bSupportsCPU = false;
 		Sig.bSupportsGPU = true;
@@ -752,7 +737,6 @@ void UNiagaraDataInterfaceGrid3DCollection::GetFunctionsInternal(TArray<FNiagara
 		Sig.FunctionSpecifiers.Add(FName(UNiagaraDataInterfaceRWBase::NAME_Attribute));
 		Sig.bMemberFunction = true;
 		Sig.bRequiresContext = false;
-		Sig.bExperimental = true;
 		Sig.ModuleUsageBitmask = ENiagaraScriptUsageMask::Particle;
 		Sig.bSupportsCPU = false;
 		Sig.bSupportsGPU = true;
@@ -771,7 +755,6 @@ void UNiagaraDataInterfaceGrid3DCollection::GetFunctionsInternal(TArray<FNiagara
 		Sig.FunctionSpecifiers.Add(FName(UNiagaraDataInterfaceRWBase::NAME_Attribute));
 		Sig.bMemberFunction = true;
 		Sig.bRequiresContext = false;
-		Sig.bExperimental = true;
 		Sig.ModuleUsageBitmask = ENiagaraScriptUsageMask::Particle;
 		Sig.bSupportsCPU = false;
 		Sig.bSupportsGPU = true;
@@ -792,7 +775,6 @@ void UNiagaraDataInterfaceGrid3DCollection::GetFunctionsInternal(TArray<FNiagara
 		Sig.FunctionSpecifiers.Add(FName(UNiagaraDataInterfaceRWBase::NAME_Attribute));
 		Sig.bMemberFunction = true;
 		Sig.bRequiresContext = false;
-		Sig.bExperimental = true;
 		Sig.bRequiresExecPin = true;
 		Sig.bWriteFunction = true;
 		Sig.ModuleUsageBitmask = ENiagaraScriptUsageMask::Particle;
@@ -814,7 +796,6 @@ void UNiagaraDataInterfaceGrid3DCollection::GetFunctionsInternal(TArray<FNiagara
 		Sig.FunctionSpecifiers.Add(FName(UNiagaraDataInterfaceRWBase::NAME_Attribute));
 		Sig.bMemberFunction = true;
 		Sig.bRequiresContext = false;
-		Sig.bExperimental = true;
 		Sig.ModuleUsageBitmask = ENiagaraScriptUsageMask::Particle;
 		Sig.bSupportsCPU = false;
 		Sig.bSupportsGPU = true;
@@ -833,7 +814,6 @@ void UNiagaraDataInterfaceGrid3DCollection::GetFunctionsInternal(TArray<FNiagara
 		Sig.FunctionSpecifiers.Add(FName(UNiagaraDataInterfaceRWBase::NAME_Attribute));
 		Sig.bMemberFunction = true;
 		Sig.bRequiresContext = false;
-		Sig.bExperimental = true;
 		Sig.ModuleUsageBitmask = ENiagaraScriptUsageMask::Particle;
 		Sig.bSupportsCPU = false;
 		Sig.bSupportsGPU = true;
@@ -852,7 +832,6 @@ void UNiagaraDataInterfaceGrid3DCollection::GetFunctionsInternal(TArray<FNiagara
 		Sig.FunctionSpecifiers.Add(FName(UNiagaraDataInterfaceRWBase::NAME_Attribute));
 		Sig.bMemberFunction = true;
 		Sig.bRequiresContext = false;
-		Sig.bExperimental = true;
 		Sig.ModuleUsageBitmask = ENiagaraScriptUsageMask::Particle;
 		Sig.bSupportsCPU = false;
 		Sig.bSupportsGPU = true;
@@ -873,7 +852,6 @@ void UNiagaraDataInterfaceGrid3DCollection::GetFunctionsInternal(TArray<FNiagara
 		Sig.FunctionSpecifiers.Add(FName(UNiagaraDataInterfaceRWBase::NAME_Attribute));
 		Sig.bMemberFunction = true;
 		Sig.bRequiresContext = false;
-		Sig.bExperimental = true;
 		Sig.bRequiresExecPin = true;
 		Sig.bWriteFunction = true;
 		Sig.ModuleUsageBitmask = ENiagaraScriptUsageMask::Particle;
@@ -894,7 +872,6 @@ void UNiagaraDataInterfaceGrid3DCollection::GetFunctionsInternal(TArray<FNiagara
 		Sig.FunctionSpecifiers.Add(FName(UNiagaraDataInterfaceRWBase::NAME_Attribute));
 		Sig.bMemberFunction = true;
 		Sig.bRequiresContext = false;
-		Sig.bExperimental = true;
 		Sig.ModuleUsageBitmask = ENiagaraScriptUsageMask::Particle;
 		Sig.bSupportsCPU = false;
 		Sig.bSupportsGPU = true;
@@ -913,7 +890,6 @@ void UNiagaraDataInterfaceGrid3DCollection::GetFunctionsInternal(TArray<FNiagara
 		Sig.FunctionSpecifiers.Add(FName(UNiagaraDataInterfaceRWBase::NAME_Attribute));
 		Sig.bMemberFunction = true;
 		Sig.bRequiresContext = false;
-		Sig.bExperimental = true;
 		Sig.ModuleUsageBitmask = ENiagaraScriptUsageMask::Particle;
 		Sig.bSupportsCPU = false;
 		Sig.bSupportsGPU = true;
@@ -932,7 +908,6 @@ void UNiagaraDataInterfaceGrid3DCollection::GetFunctionsInternal(TArray<FNiagara
 		Sig.FunctionSpecifiers.Add(FName(UNiagaraDataInterfaceRWBase::NAME_Attribute));
 		Sig.bMemberFunction = true;
 		Sig.bRequiresContext = false;
-		Sig.bExperimental = true;
 		Sig.ModuleUsageBitmask = ENiagaraScriptUsageMask::Particle;
 		Sig.bSupportsCPU = false;
 		Sig.bSupportsGPU = true;
@@ -952,7 +927,6 @@ void UNiagaraDataInterfaceGrid3DCollection::GetFunctionsInternal(TArray<FNiagara
 		Sig.Inputs.Add(FNiagaraVariable(FNiagaraTypeDefinition::GetIntDef(), TEXT("AttributeIndex")));
 		Sig.Outputs.Add(FNiagaraVariable(FNiagaraTypeDefinition::GetFloatDef(), TEXT("Value")));
 
-		Sig.bExperimental = true;
 		Sig.bMemberFunction = true;
 		Sig.bRequiresContext = false;
 		Sig.ModuleUsageBitmask = ENiagaraScriptUsageMask::Particle;
@@ -973,7 +947,6 @@ void UNiagaraDataInterfaceGrid3DCollection::GetFunctionsInternal(TArray<FNiagara
 		Sig.Inputs.Add(FNiagaraVariable(FNiagaraTypeDefinition::GetIntDef(), TEXT("AttributeIndex")));
 		Sig.Outputs.Add(FNiagaraVariable(FNiagaraTypeDefinition::GetFloatDef(), TEXT("Value")));
 
-		Sig.bExperimental = true;
 		Sig.bMemberFunction = true;
 		Sig.bRequiresContext = false;
 		Sig.ModuleUsageBitmask = ENiagaraScriptUsageMask::Particle;
@@ -994,7 +967,6 @@ void UNiagaraDataInterfaceGrid3DCollection::GetFunctionsInternal(TArray<FNiagara
 		Sig.Inputs.Add(FNiagaraVariable(FNiagaraTypeDefinition::GetIntDef(), TEXT("AttributeIndex")));
 		Sig.Outputs.Add(FNiagaraVariable(FNiagaraTypeDefinition::GetFloatDef(), TEXT("Value")));
 
-		Sig.bExperimental = true;
 		Sig.bMemberFunction = true;
 		Sig.bRequiresContext = false;
 		Sig.ModuleUsageBitmask = ENiagaraScriptUsageMask::Particle;
@@ -1014,7 +986,6 @@ void UNiagaraDataInterfaceGrid3DCollection::GetFunctionsInternal(TArray<FNiagara
 		Sig.Inputs.Add(FNiagaraVariable(FNiagaraTypeDefinition::GetIntDef(), TEXT("AttributeIndex")));
 		Sig.Outputs.Add(FNiagaraVariable(FNiagaraTypeDefinition::GetFloatDef(), TEXT("Value")));
 
-		Sig.bExperimental = true;
 		Sig.bMemberFunction = true;
 		Sig.bRequiresContext = false;
 		Sig.ModuleUsageBitmask = ENiagaraScriptUsageMask::Particle;
@@ -1031,7 +1002,6 @@ void UNiagaraDataInterfaceGrid3DCollection::GetFunctionsInternal(TArray<FNiagara
 		Sig.Inputs.Add(FNiagaraVariable(FNiagaraTypeDefinition::GetVec3Def(), TEXT("Unit")));
 		Sig.Outputs.Add(FNiagaraVariable(FNiagaraTypeDefinition::GetFloatDef(), TEXT("Value")));
 
-		Sig.bExperimental = true;
 		Sig.bMemberFunction = true;
 		Sig.bRequiresContext = false;
 		Sig.ModuleUsageBitmask = ENiagaraScriptUsageMask::Particle;
@@ -1050,7 +1020,6 @@ void UNiagaraDataInterfaceGrid3DCollection::GetFunctionsInternal(TArray<FNiagara
 		Sig.Inputs.Add(FNiagaraVariable(FNiagaraTypeDefinition::GetVec3Def(), TEXT("Unit")));
 		Sig.Outputs.Add(FNiagaraVariable(FNiagaraTypeDefinition::GetFloatDef(), TEXT("Value")));
 
-		Sig.bExperimental = true;
 		Sig.bMemberFunction = true;
 		Sig.bRequiresContext = false;
 		Sig.ModuleUsageBitmask = ENiagaraScriptUsageMask::Particle;
@@ -1070,7 +1039,6 @@ void UNiagaraDataInterfaceGrid3DCollection::GetFunctionsInternal(TArray<FNiagara
 		Sig.FunctionSpecifiers.Add(FName(UNiagaraDataInterfaceRWBase::NAME_Attribute));
 		Sig.bMemberFunction = true;
 		Sig.bRequiresContext = false;
-		Sig.bExperimental = true;
 		Sig.Description = NSLOCTEXT("Niagara", "NiagaraDataInterfaceGridColl2D_GetVector4AttributeIndex", "Gets a attribute starting index value for Vector4 on the Grid by Attribute name. Returns -1 if not found.");
 		OutFunctions.Add(Sig);
 	}
@@ -1083,7 +1051,6 @@ void UNiagaraDataInterfaceGrid3DCollection::GetFunctionsInternal(TArray<FNiagara
 		Sig.FunctionSpecifiers.Add(FName(UNiagaraDataInterfaceRWBase::NAME_Attribute));
 		Sig.bMemberFunction = true;
 		Sig.bRequiresContext = false;
-		Sig.bExperimental = true;
 		Sig.Description = NSLOCTEXT("Niagara", "NiagaraDataInterfaceGridColl2D_GetVector3AttributeIndex", "Gets a attribute starting index value for Vector3 on the Grid by Attribute name. Returns -1 if not found.");
 		OutFunctions.Add(Sig);
 	}
@@ -1096,7 +1063,6 @@ void UNiagaraDataInterfaceGrid3DCollection::GetFunctionsInternal(TArray<FNiagara
 		Sig.FunctionSpecifiers.Add(FName(UNiagaraDataInterfaceRWBase::NAME_Attribute));
 		Sig.bMemberFunction = true;
 		Sig.bRequiresContext = false;
-		Sig.bExperimental = true;
 		Sig.Description = NSLOCTEXT("Niagara", "NiagaraDataInterfaceGridColl2D_GetVector2AttributeIndex", "Gets a attribute starting index value for Vector2 on the Grid by Attribute name. Returns -1 if not found.");
 		OutFunctions.Add(Sig);
 	}
@@ -1109,7 +1075,6 @@ void UNiagaraDataInterfaceGrid3DCollection::GetFunctionsInternal(TArray<FNiagara
 		Sig.FunctionSpecifiers.Add(FName(UNiagaraDataInterfaceRWBase::NAME_Attribute));
 		Sig.bMemberFunction = true;
 		Sig.bRequiresContext = false;
-		Sig.bExperimental = true;
 		Sig.Description = NSLOCTEXT("Niagara", "NiagaraDataInterfaceGridColl2D_GetFloatAttributeIndex", "Gets a attribute starting index value for float on the Grid by Attribute name. Returns -1 if not found.");
 		OutFunctions.Add(Sig);
 	}
@@ -1124,7 +1089,6 @@ void UNiagaraDataInterfaceGrid3DCollection::GetFunctionsInternal(TArray<FNiagara
 		Sig.Inputs.Add(FNiagaraVariable(FNiagaraTypeDefinition::GetIntDef(), TEXT("AttributeIndex")));
 		Sig.Outputs.Add(FNiagaraVariable(FNiagaraTypeDefinition::GetFloatDef(), TEXT("Value")));
 
-		Sig.bExperimental = true;
 		Sig.bMemberFunction = true;
 		Sig.bRequiresContext = false;
 		Sig.ModuleUsageBitmask = ENiagaraScriptUsageMask::Particle;
@@ -1146,7 +1110,6 @@ void UNiagaraDataInterfaceGrid3DCollection::GetFunctionsInternal(TArray<FNiagara
 		Sig.Inputs.Add(FNiagaraVariable(FNiagaraTypeDefinition::GetIntDef(), TEXT("AttributeIndex")));
 		Sig.Outputs.Add(FNiagaraVariable(FNiagaraTypeDefinition::GetVec2Def(), TEXT("Value")));
 
-		Sig.bExperimental = true;
 		Sig.bMemberFunction = true;
 		Sig.bRequiresContext = false;
 		Sig.ModuleUsageBitmask = ENiagaraScriptUsageMask::Particle;
@@ -1168,7 +1131,6 @@ void UNiagaraDataInterfaceGrid3DCollection::GetFunctionsInternal(TArray<FNiagara
 		Sig.Inputs.Add(FNiagaraVariable(FNiagaraTypeDefinition::GetIntDef(), TEXT("AttributeIndex")));
 		Sig.Outputs.Add(FNiagaraVariable(FNiagaraTypeDefinition::GetVec3Def(), TEXT("Value")));
 
-		Sig.bExperimental = true;
 		Sig.bMemberFunction = true;
 		Sig.bRequiresContext = false;
 		Sig.ModuleUsageBitmask = ENiagaraScriptUsageMask::Particle;
@@ -1190,7 +1152,6 @@ void UNiagaraDataInterfaceGrid3DCollection::GetFunctionsInternal(TArray<FNiagara
 		Sig.Inputs.Add(FNiagaraVariable(FNiagaraTypeDefinition::GetIntDef(), TEXT("AttributeIndex")));
 		Sig.Outputs.Add(FNiagaraVariable(FNiagaraTypeDefinition::GetVec4Def(), TEXT("Value")));
 
-		Sig.bExperimental = true;
 		Sig.bMemberFunction = true;
 		Sig.bRequiresContext = false;
 		Sig.ModuleUsageBitmask = ENiagaraScriptUsageMask::Particle;
@@ -1210,7 +1171,6 @@ void UNiagaraDataInterfaceGrid3DCollection::GetFunctionsInternal(TArray<FNiagara
 		Sig.Inputs.Add(FNiagaraVariable(FNiagaraTypeDefinition::GetIntDef(), TEXT("AttributeIndex")));
 		Sig.Outputs.Add(FNiagaraVariable(FNiagaraTypeDefinition::GetFloatDef(), TEXT("Value")));
 
-		Sig.bExperimental = true;
 		Sig.bMemberFunction = true;
 		Sig.bRequiresContext = false;
 		Sig.ModuleUsageBitmask = ENiagaraScriptUsageMask::Particle;
@@ -1228,7 +1188,6 @@ void UNiagaraDataInterfaceGrid3DCollection::GetFunctionsInternal(TArray<FNiagara
 		Sig.Inputs.Add(FNiagaraVariable(FNiagaraTypeDefinition::GetIntDef(), TEXT("AttributeIndex")));
 		Sig.Outputs.Add(FNiagaraVariable(FNiagaraTypeDefinition::GetVec2Def(), TEXT("Value")));
 
-		Sig.bExperimental = true;
 		Sig.bMemberFunction = true;
 		Sig.bRequiresContext = false;
 		Sig.ModuleUsageBitmask = ENiagaraScriptUsageMask::Particle;
@@ -1246,7 +1205,6 @@ void UNiagaraDataInterfaceGrid3DCollection::GetFunctionsInternal(TArray<FNiagara
 		Sig.Inputs.Add(FNiagaraVariable(FNiagaraTypeDefinition::GetIntDef(), TEXT("AttributeIndex")));
 		Sig.Outputs.Add(FNiagaraVariable(FNiagaraTypeDefinition::GetVec3Def(), TEXT("Value")));
 
-		Sig.bExperimental = true;
 		Sig.bMemberFunction = true;
 		Sig.bRequiresContext = false;
 		Sig.ModuleUsageBitmask = ENiagaraScriptUsageMask::Particle;
@@ -1264,7 +1222,6 @@ void UNiagaraDataInterfaceGrid3DCollection::GetFunctionsInternal(TArray<FNiagara
 		Sig.Inputs.Add(FNiagaraVariable(FNiagaraTypeDefinition::GetIntDef(), TEXT("AttributeIndex")));
 		Sig.Outputs.Add(FNiagaraVariable(FNiagaraTypeDefinition::GetVec4Def(), TEXT("Value")));
 
-		Sig.bExperimental = true;
 		Sig.bMemberFunction = true;
 		Sig.bRequiresContext = false;
 		Sig.ModuleUsageBitmask = ENiagaraScriptUsageMask::Particle;
@@ -1282,7 +1239,6 @@ void UNiagaraDataInterfaceGrid3DCollection::GetFunctionsInternal(TArray<FNiagara
 		Sig.Inputs.Add(FNiagaraVariable(FNiagaraTypeDefinition::GetIntDef(), TEXT("AttributeIndex")));
 		Sig.Outputs.Add(FNiagaraVariable(FNiagaraTypeDefinition::GetFloatDef(), TEXT("Value")));
 
-		Sig.bExperimental = true;
 		Sig.bMemberFunction = true;
 		Sig.bRequiresContext = false;
 		Sig.ModuleUsageBitmask = ENiagaraScriptUsageMask::Particle;
@@ -1300,7 +1256,6 @@ void UNiagaraDataInterfaceGrid3DCollection::GetFunctionsInternal(TArray<FNiagara
 		Sig.Inputs.Add(FNiagaraVariable(FNiagaraTypeDefinition::GetIntDef(), TEXT("AttributeIndex")));
 		Sig.Outputs.Add(FNiagaraVariable(FNiagaraTypeDefinition::GetVec2Def(), TEXT("Value")));
 
-		Sig.bExperimental = true;
 		Sig.bMemberFunction = true;
 		Sig.bRequiresContext = false;
 		Sig.ModuleUsageBitmask = ENiagaraScriptUsageMask::Particle;
@@ -1318,7 +1273,6 @@ void UNiagaraDataInterfaceGrid3DCollection::GetFunctionsInternal(TArray<FNiagara
 		Sig.Inputs.Add(FNiagaraVariable(FNiagaraTypeDefinition::GetIntDef(), TEXT("AttributeIndex")));
 		Sig.Outputs.Add(FNiagaraVariable(FNiagaraTypeDefinition::GetVec3Def(), TEXT("Value")));
 
-		Sig.bExperimental = true;
 		Sig.bMemberFunction = true;
 		Sig.bRequiresContext = false;
 		Sig.ModuleUsageBitmask = ENiagaraScriptUsageMask::Particle;
@@ -1336,7 +1290,6 @@ void UNiagaraDataInterfaceGrid3DCollection::GetFunctionsInternal(TArray<FNiagara
 		Sig.Inputs.Add(FNiagaraVariable(FNiagaraTypeDefinition::GetIntDef(), TEXT("AttributeIndex")));
 		Sig.Outputs.Add(FNiagaraVariable(FNiagaraTypeDefinition::GetVec4Def(), TEXT("Value")));
 
-		Sig.bExperimental = true;
 		Sig.bMemberFunction = true;
 		Sig.bRequiresContext = false;
 		Sig.ModuleUsageBitmask = ENiagaraScriptUsageMask::Particle;
@@ -1357,7 +1310,6 @@ void UNiagaraDataInterfaceGrid3DCollection::GetFunctionsInternal(TArray<FNiagara
 		Sig.Inputs.Add(FNiagaraVariable(FNiagaraTypeDefinition::GetFloatDef(), TEXT("Value")));
 		Sig.bMemberFunction = true;
 		Sig.bRequiresContext = false;
-		Sig.bExperimental = true;
 		Sig.bRequiresExecPin = true;
 		Sig.bWriteFunction = true;
 		Sig.ModuleUsageBitmask = ENiagaraScriptUsageMask::Particle;
@@ -1378,7 +1330,6 @@ void UNiagaraDataInterfaceGrid3DCollection::GetFunctionsInternal(TArray<FNiagara
 		Sig.Inputs.Add(FNiagaraVariable(FNiagaraTypeDefinition::GetVec2Def(), TEXT("Value")));
 		Sig.bMemberFunction = true;
 		Sig.bRequiresContext = false;
-		Sig.bExperimental = true;
 		Sig.bRequiresExecPin = true;
 		Sig.bWriteFunction = true;
 		Sig.ModuleUsageBitmask = ENiagaraScriptUsageMask::Particle;
@@ -1399,7 +1350,6 @@ void UNiagaraDataInterfaceGrid3DCollection::GetFunctionsInternal(TArray<FNiagara
 		Sig.Inputs.Add(FNiagaraVariable(FNiagaraTypeDefinition::GetVec3Def(), TEXT("Value")));
 		Sig.bMemberFunction = true;
 		Sig.bRequiresContext = false;
-		Sig.bExperimental = true;
 		Sig.bRequiresExecPin = true;
 		Sig.bWriteFunction = true;
 		Sig.ModuleUsageBitmask = ENiagaraScriptUsageMask::Particle;
@@ -1420,7 +1370,6 @@ void UNiagaraDataInterfaceGrid3DCollection::GetFunctionsInternal(TArray<FNiagara
 		Sig.Inputs.Add(FNiagaraVariable(FNiagaraTypeDefinition::GetVec4Def(), TEXT("Value")));
 		Sig.bMemberFunction = true;
 		Sig.bRequiresContext = false;
-		Sig.bExperimental = true;
 		Sig.bRequiresExecPin = true;
 		Sig.bWriteFunction = true;
 		Sig.ModuleUsageBitmask = ENiagaraScriptUsageMask::Particle;
diff --git a/Engine/Plugins/FX/Niagara/Source/Niagara/Private/NiagaraDataInterfaceIntRenderTarget2D.cpp b/Engine/Plugins/FX/Niagara/Source/Niagara/Private/NiagaraDataInterfaceIntRenderTarget2D.cpp
--- a/Engine/Plugins/FX/Niagara/Source/Niagara/Private/NiagaraDataInterfaceIntRenderTarget2D.cpp
+++ b/Engine/Plugins/FX/Niagara/Source/Niagara/Private/NiagaraDataInterfaceIntRenderTarget2D.cpp
@@ -202,7 +202,6 @@ void UNiagaraDataInterfaceIntRenderTarget2D::GetFunctionsInternal(TArray<FNiagar
 		Sig.Inputs.Add(FNiagaraVariable(FNiagaraTypeDefinition::GetIntDef(), TEXT("PixelX")));
 		Sig.Inputs.Add(FNiagaraVariable(FNiagaraTypeDefinition::GetIntDef(), TEXT("PixelY")));
 		Sig.Outputs.Add(FNiagaraVariable(FNiagaraTypeDefinition::GetIntDef(), TEXT("Value")));
-		Sig.bExperimental = true;
 		Sig.bMemberFunction = true;
 		Sig.bRequiresExecPin = false;
 		Sig.bRequiresContext = false;
@@ -220,7 +219,6 @@ void UNiagaraDataInterfaceIntRenderTarget2D::GetFunctionsInternal(TArray<FNiagar
 		Sig.Inputs.Add(FNiagaraVariable(FNiagaraTypeDefinition::GetIntDef(), TEXT("PixelX")));
 		Sig.Inputs.Add(FNiagaraVariable(FNiagaraTypeDefinition::GetIntDef(), TEXT("PixelY")));
 		Sig.Inputs.Add(FNiagaraVariable(FNiagaraTypeDefinition::GetIntDef(), TEXT("Value")));
-		Sig.bExperimental = true;
 		Sig.bMemberFunction = true;
 		Sig.bRequiresExecPin = true;
 		Sig.bRequiresContext = false;
@@ -240,7 +238,6 @@ void UNiagaraDataInterfaceIntRenderTarget2D::GetFunctionsInternal(TArray<FNiagar
 		Sig.Inputs.Add(FNiagaraVariable(FNiagaraTypeDefinition::GetIntDef(), TEXT("Amount")));
 		Sig.Outputs.Add(FNiagaraVariable(FNiagaraTypeDefinition::GetIntDef(), TEXT("CurrentValue")));
 		Sig.Outputs.Add(FNiagaraVariable(FNiagaraTypeDefinition::GetIntDef(), TEXT("PreviousValue")));
-		Sig.bExperimental = true;
 		Sig.bMemberFunction = true;
 		Sig.bRequiresExecPin = true;
 		Sig.bRequiresContext = false;
@@ -262,7 +259,6 @@ void UNiagaraDataInterfaceIntRenderTarget2D::GetFunctionsInternal(TArray<FNiagar
 		Sig.Inputs.Add(FNiagaraVariable(FNiagaraTypeDefinition::GetIntDef(), TEXT("Value")));
 		Sig.Outputs.Add(FNiagaraVariable(FNiagaraTypeDefinition::GetIntDef(), TEXT("Original Value")));
 		Sig.bHidden = true;
-		Sig.bExperimental = true;
 		Sig.bMemberFunction = true;
 		Sig.bRequiresExecPin = true;
 		Sig.bRequiresContext = false;
@@ -284,7 +280,6 @@ void UNiagaraDataInterfaceIntRenderTarget2D::GetFunctionsInternal(TArray<FNiagar
 		Sig.Inputs.Add(FNiagaraVariable(FNiagaraTypeDefinition::GetIntDef(), TEXT("Value")));
 		Sig.Outputs.Add(FNiagaraVariable(FNiagaraTypeDefinition::GetIntDef(), TEXT("CurrentValue")));
 		Sig.Outputs.Add(FNiagaraVariable(FNiagaraTypeDefinition::GetIntDef(), TEXT("PreviousValue")));
-		Sig.bExperimental = true;
 		Sig.bMemberFunction = true;
 		Sig.bRequiresExecPin = true;
 		Sig.bRequiresContext = false;
@@ -304,7 +299,6 @@ void UNiagaraDataInterfaceIntRenderTarget2D::GetFunctionsInternal(TArray<FNiagar
 		Sig.Inputs.Add(FNiagaraVariable(FNiagaraTypeDefinition::GetIntDef(), TEXT("Value")));
 		Sig.Outputs.Add(FNiagaraVariable(FNiagaraTypeDefinition::GetIntDef(), TEXT("CurrentValue")));
 		Sig.Outputs.Add(FNiagaraVariable(FNiagaraTypeDefinition::GetIntDef(), TEXT("PreviousValue")));
-		Sig.bExperimental = true;
 		Sig.bMemberFunction = true;
 		Sig.bRequiresExecPin = true;
 		Sig.bRequiresContext = false;
@@ -323,7 +317,6 @@ void UNiagaraDataInterfaceIntRenderTarget2D::GetFunctionsInternal(TArray<FNiagar
 		Sig.Inputs.Add(FNiagaraVariable(FNiagaraTypeDefinition(GetClass()), TEXT("RenderTarget")));
 		Sig.Outputs.Add(FNiagaraVariable(FNiagaraTypeDefinition::GetIntDef(), TEXT("Width")));
 		Sig.Outputs.Add(FNiagaraVariable(FNiagaraTypeDefinition::GetIntDef(), TEXT("Height")));
-		Sig.bExperimental = true;
 		Sig.bMemberFunction = true;
 		Sig.bRequiresContext = false;
 		Sig.bSupportsCPU = true;
@@ -339,7 +332,6 @@ void UNiagaraDataInterfaceIntRenderTarget2D::GetFunctionsInternal(TArray<FNiagar
 		Sig.Inputs.Add(FNiagaraVariable(FNiagaraTypeDefinition::GetIntDef(), TEXT("Height")));
 		Sig.Outputs.Add(FNiagaraVariable(FNiagaraTypeDefinition::GetBoolDef(), TEXT("Success")));
 		Sig.ModuleUsageBitmask = EmitterSystemOnlyBitmask;
-		Sig.bExperimental = true;
 		Sig.bMemberFunction = true;
 		Sig.bRequiresContext = false;
 		Sig.bRequiresExecPin = true;
@@ -356,7 +348,6 @@ void UNiagaraDataInterfaceIntRenderTarget2D::GetFunctionsInternal(TArray<FNiagar
 		Sig.Inputs.Add(FNiagaraVariable(FNiagaraTypeDefinition::GetIntDef(), TEXT("Linear")));
 		Sig.Outputs.Add(FNiagaraVariable(FNiagaraTypeDefinition::GetIntDef(), TEXT("PixelX")));
 		Sig.Outputs.Add(FNiagaraVariable(FNiagaraTypeDefinition::GetIntDef(), TEXT("PixelY")));
-		Sig.bExperimental = true;
 		Sig.bMemberFunction = true;
 		Sig.bRequiresContext = false;
 		Sig.bSupportsCPU = false;
@@ -370,7 +361,6 @@ void UNiagaraDataInterfaceIntRenderTarget2D::GetFunctionsInternal(TArray<FNiagar
 		Sig.Inputs.Add(FNiagaraVariable(FNiagaraTypeDefinition(GetClass()), TEXT("RenderTarget")));
 		Sig.Inputs.Add(FNiagaraVariable(FNiagaraTypeDefinition::GetIntDef(), TEXT("Linear")));
 		Sig.Outputs.Add(FNiagaraVariable(FNiagaraTypeDefinition::GetVec2Def(), TEXT("UV")));
-		Sig.bExperimental = true;
 		Sig.bMemberFunction = true;
 		Sig.bRequiresContext = false;
 		Sig.bSupportsCPU = false;
@@ -385,7 +375,6 @@ void UNiagaraDataInterfaceIntRenderTarget2D::GetFunctionsInternal(TArray<FNiagar
 		Sig.Inputs.Add(FNiagaraVariable(FNiagaraTypeDefinition(GetClass()), TEXT("RenderTarget")));
 		Sig.Outputs.Add(FNiagaraVariable(FNiagaraTypeDefinition::GetIntDef(), TEXT("PixelX")));
 		Sig.Outputs.Add(FNiagaraVariable(FNiagaraTypeDefinition::GetIntDef(), TEXT("PixelY")));
-		Sig.bExperimental = true;
 		Sig.bMemberFunction = true;
 		Sig.bRequiresContext = false;
 		Sig.bSupportsCPU = false;
@@ -398,7 +387,6 @@ void UNiagaraDataInterfaceIntRenderTarget2D::GetFunctionsInternal(TArray<FNiagar
 		Sig.Name = NDIIntRenderTarget2DLocal::ExecToUVFunctionName;
 		Sig.Inputs.Add(FNiagaraVariable(FNiagaraTypeDefinition(GetClass()), TEXT("RenderTarget")));
 		Sig.Outputs.Add(FNiagaraVariable(FNiagaraTypeDefinition::GetVec2Def(), TEXT("UV")));
-		Sig.bExperimental = true;
 		Sig.bMemberFunction = true;
 		Sig.bRequiresContext = false;
 		Sig.bSupportsCPU = false;
diff --git a/Engine/Plugins/FX/Niagara/Source/Niagara/Private/NiagaraDataInterfaceLandscape.cpp b/Engine/Plugins/FX/Niagara/Source/Niagara/Private/NiagaraDataInterfaceLandscape.cpp
--- a/Engine/Plugins/FX/Niagara/Source/Niagara/Private/NiagaraDataInterfaceLandscape.cpp
+++ b/Engine/Plugins/FX/Niagara/Source/Niagara/Private/NiagaraDataInterfaceLandscape.cpp
@@ -1491,7 +1491,6 @@ void UNiagaraDataInterfaceLandscape::GetFunctionsInternal(TArray<FNiagaraFunctio
 	{
 		FNiagaraFunctionSignature Sig;
 		Sig.Name = GetBaseColorName;
-		Sig.bExperimental = true;
 		Sig.bMemberFunction = true;
 		Sig.bRequiresContext = false;
 		Sig.bSupportsCPU = false;
@@ -1506,7 +1505,6 @@ void UNiagaraDataInterfaceLandscape::GetFunctionsInternal(TArray<FNiagaraFunctio
 	{
 		FNiagaraFunctionSignature Sig;
 		Sig.Name = GetHeightName;
-		Sig.bExperimental = true;
 		Sig.bMemberFunction = true;
 		Sig.bRequiresContext = false;
 		Sig.bSupportsCPU = false;
@@ -1521,7 +1519,6 @@ void UNiagaraDataInterfaceLandscape::GetFunctionsInternal(TArray<FNiagaraFunctio
 	{
 		FNiagaraFunctionSignature Sig;
 		Sig.Name = GetWorldNormalName;
-		Sig.bExperimental = true;
 		Sig.bMemberFunction = true;
 		Sig.bRequiresContext = false;
 		Sig.bSupportsCPU = false;
@@ -1536,7 +1533,6 @@ void UNiagaraDataInterfaceLandscape::GetFunctionsInternal(TArray<FNiagaraFunctio
 	{
 		FNiagaraFunctionSignature Sig;
 		Sig.Name = GetPhysicalMaterialIndexName;
-		Sig.bExperimental = true;
 		Sig.bMemberFunction = true;
 		Sig.bRequiresContext = false;
 		Sig.bSupportsCPU = false;
diff --git a/Engine/Plugins/FX/Niagara/Source/Niagara/Private/NiagaraDataInterfaceNeighborGrid3D.cpp b/Engine/Plugins/FX/Niagara/Source/Niagara/Private/NiagaraDataInterfaceNeighborGrid3D.cpp
--- a/Engine/Plugins/FX/Niagara/Source/Niagara/Private/NiagaraDataInterfaceNeighborGrid3D.cpp
+++ b/Engine/Plugins/FX/Niagara/Source/Niagara/Private/NiagaraDataInterfaceNeighborGrid3D.cpp
@@ -84,7 +84,6 @@ void UNiagaraDataInterfaceNeighborGrid3D::GetFunctionsInternal(TArray<FNiagaraFu
 		Sig.Outputs.Add(FNiagaraVariable(FNiagaraTypeDefinition::GetBoolDef(), TEXT("Success")));
 
 		Sig.ModuleUsageBitmask = ENiagaraScriptUsageMask::Emitter | ENiagaraScriptUsageMask::System;
-		Sig.bExperimental = true;
 		Sig.bMemberFunction = true;
 		Sig.bRequiresExecPin = true;
 		Sig.bRequiresContext = false;
@@ -99,7 +98,6 @@ void UNiagaraDataInterfaceNeighborGrid3D::GetFunctionsInternal(TArray<FNiagaraFu
 		Sig.Inputs.Add(FNiagaraVariable(FNiagaraTypeDefinition(GetClass()), TEXT("Grid")));
 		Sig.Outputs.Add(FNiagaraVariable(FNiagaraTypeDefinition::GetIntDef(), TEXT("MaxNeighborsPerCell")));
 
-		Sig.bExperimental = true;
 		Sig.bMemberFunction = true;
 		Sig.bRequiresContext = false;
 		OutFunctions.Add(Sig);
@@ -115,7 +113,6 @@ void UNiagaraDataInterfaceNeighborGrid3D::GetFunctionsInternal(TArray<FNiagaraFu
 		Sig.Inputs.Add(FNiagaraVariable(FNiagaraTypeDefinition::GetIntDef(), TEXT("Neighbor")));
 		Sig.Outputs.Add(FNiagaraVariable(FNiagaraTypeDefinition::GetIntDef(), TEXT("Linear Index")));
 
-		Sig.bExperimental = true;
 		Sig.bMemberFunction = true;
 		Sig.bRequiresContext = false;
 		OutFunctions.Add(Sig);
@@ -128,7 +125,6 @@ void UNiagaraDataInterfaceNeighborGrid3D::GetFunctionsInternal(TArray<FNiagaraFu
 		Sig.Inputs.Add(FNiagaraVariable(FNiagaraTypeDefinition::GetIntDef(), TEXT("Linear")));
 		Sig.Outputs.Add(FNiagaraVariable(FNiagaraTypeDefinition::GetIntDef(), TEXT("NeighborIndex")));
 
-		Sig.bExperimental = true;
 		Sig.bMemberFunction = true;
 		Sig.bRequiresContext = false;
 		OutFunctions.Add(Sig);
@@ -142,7 +138,6 @@ void UNiagaraDataInterfaceNeighborGrid3D::GetFunctionsInternal(TArray<FNiagaraFu
 		Sig.Inputs.Add(FNiagaraVariable(FNiagaraTypeDefinition::GetIntDef(), TEXT("NeighborIndex")));
 		Sig.Outputs.Add(FNiagaraVariable(FNiagaraTypeDefinition::GetIntDef(), TEXT("IGNORE")));
 
-		Sig.bExperimental = true;
 		Sig.bWriteFunction = true;
 		Sig.bMemberFunction = true;
 		Sig.bRequiresContext = false;
@@ -156,7 +151,6 @@ void UNiagaraDataInterfaceNeighborGrid3D::GetFunctionsInternal(TArray<FNiagaraFu
 		Sig.Inputs.Add(FNiagaraVariable(FNiagaraTypeDefinition::GetIntDef(), TEXT("Linear")));
 		Sig.Outputs.Add(FNiagaraVariable(FNiagaraTypeDefinition::GetIntDef(), TEXT("NeighborCount")));
 
-		Sig.bExperimental = true;
 		Sig.bMemberFunction = true;
 		Sig.bRequiresContext = false;
 		OutFunctions.Add(Sig);
@@ -170,7 +164,6 @@ void UNiagaraDataInterfaceNeighborGrid3D::GetFunctionsInternal(TArray<FNiagaraFu
 		Sig.Inputs.Add(FNiagaraVariable(FNiagaraTypeDefinition::GetIntDef(), TEXT("Increment")));
 		Sig.Outputs.Add(FNiagaraVariable(FNiagaraTypeDefinition::GetIntDef(), TEXT("PrevNeighborCount")));
 
-		Sig.bExperimental = true;
 		Sig.bWriteFunction = true;
 		Sig.bMemberFunction = true;
 		Sig.bRequiresContext = false;
@@ -187,7 +180,6 @@ void UNiagaraDataInterfaceNeighborGrid3D::GetFunctionsInternal(TArray<FNiagaraFu
 		Sig.Inputs.Add(FNiagaraVariable(FNiagaraTypeDefinition::GetIntDef(), TEXT("ParticleIndex")));
 		Sig.Outputs.Add(FNiagaraVariable(FNiagaraTypeDefinition::GetBoolDef(), TEXT("Success")));
 
-		Sig.bExperimental = true;
 		Sig.bWriteFunction = true;
 		Sig.bMemberFunction = true;
 		Sig.bRequiresContext = false;
diff --git a/Engine/Plugins/FX/Niagara/Source/Niagara/Private/NiagaraDataInterfaceParticleRead.cpp b/Engine/Plugins/FX/Niagara/Source/Niagara/Private/NiagaraDataInterfaceParticleRead.cpp
--- a/Engine/Plugins/FX/Niagara/Source/Niagara/Private/NiagaraDataInterfaceParticleRead.cpp
+++ b/Engine/Plugins/FX/Niagara/Source/Niagara/Private/NiagaraDataInterfaceParticleRead.cpp
@@ -507,7 +507,6 @@ void UNiagaraDataInterfaceParticleRead::GetFunctionsInternal(TArray<FNiagaraFunc
 		Sig.Inputs.Emplace(FNiagaraTypeDefinition(GetClass()), TEXT("Particle Reader"));
 		Sig.bMemberFunction = true;
 		Sig.bRequiresContext = false;
-		Sig.bExperimental = true;
 		Sig.Outputs.Emplace(FNiagaraTypeDefinition::GetBoolDef(), TEXT("IsLocalSpace"));
 		Sig.SetDescription(LOCTEXT("GetEmitterLocalSpace", "Returns if the emitter is using local space or world space."));
 	}
@@ -522,7 +521,6 @@ void UNiagaraDataInterfaceParticleRead::GetFunctionsInternal(TArray<FNiagaraFunc
 		Sig.Outputs.Add(FNiagaraVariable(FNiagaraTypeDefinition::GetIntDef(), TEXT("Num Spawned")));
 		Sig.bMemberFunction = true;
 		Sig.bRequiresContext = false;
-		Sig.bExperimental = true;
 		Sig.Description = NSLOCTEXT("Niagara", "NiagaraDataInterfaceParticleRead_GetNumSpawnedDesc", "Returns the current number of spawned particles this frame.  Note: This is the request spawn amount, when killing particles during spawn it will not be accurate.");
 		OutFunctions.Add(Sig);
 	}
@@ -536,7 +534,6 @@ void UNiagaraDataInterfaceParticleRead::GetFunctionsInternal(TArray<FNiagaraFunc
 		Sig.Outputs.Add(FNiagaraVariable(FNiagaraTypeDefinition::GetIDDef(), TEXT("ID")));
 		Sig.bMemberFunction = true;
 		Sig.bRequiresContext = false;
-		Sig.bExperimental = true;
 		Sig.Description = NSLOCTEXT("Niagara", "NiagaraDataInterfaceParticleRead_GetSpawnedIDAtIndexDesc", "Returns the Niagara ID for the particle spawned at Index.");
 		OutFunctions.Add(Sig);
 	}
@@ -571,7 +568,6 @@ void UNiagaraDataInterfaceParticleRead::GetPersistentIDFunctions(TArray<FNiagara
 		Sig.Outputs.Add(FNiagaraVariable(FNiagaraTypeDefinition::GetIntDef(), TEXT("Index")));
 		Sig.bMemberFunction = true;
 		Sig.bRequiresContext = false;
-		Sig.bExperimental = true;
 #if WITH_EDITORONLY_DATA
 		Sig.Description = NSLOCTEXT("Niagara", "NiagaraDataInterfaceParticleRead_GetParticleIndex", "Returns the particle index from the NiagaraID or -1 if not valid.");
 #endif
@@ -587,7 +583,6 @@ void UNiagaraDataInterfaceParticleRead::GetPersistentIDFunctions(TArray<FNiagara
 		Sig.Outputs.Add(FNiagaraVariable(FNiagaraTypeDefinition::GetIntDef(), TEXT("Particle Index")));
 		Sig.bMemberFunction = true;
 		Sig.bRequiresContext = false;
-		Sig.bExperimental = true;
 #if WITH_EDITORONLY_DATA
 		Sig.Description = NSLOCTEXT("Niagara", "NiagaraDataInterfaceParticleRead_GetParticleIndexFromIDTable", "Returns the particle index from the ID table, or -1 if that index or particle is not valid.\nThe ID to Index table is guaranteed to give a stable particle index over a range should the source emitter only ever spawn and not kill particles.");
 #endif
@@ -609,7 +604,6 @@ void UNiagaraDataInterfaceParticleRead::GetPersistentIDFunctions(TArray<FNiagara
 
 		Sig.bMemberFunction = true;
 		Sig.bRequiresContext = false;
-		Sig.bExperimental = true;
 #if WITH_EDITORONLY_DATA
 		Sig.Description = NSLOCTEXT("Niagara", "NiagaraDataInterfaceParticleRead_GetIntByIDDesc", "Returns an Int value from a particle by NiagaraID.  Note: When reading from self this will be the previous frames data.");
 #endif
@@ -626,7 +620,6 @@ void UNiagaraDataInterfaceParticleRead::GetPersistentIDFunctions(TArray<FNiagara
 		Sig.FunctionSpecifiers.Add(FName("Attribute"));
 		Sig.bMemberFunction = true;
 		Sig.bRequiresContext = false;
-		Sig.bExperimental = true;
 #if WITH_EDITORONLY_DATA
 		Sig.Description = NSLOCTEXT("Niagara", "NiagaraDataInterfaceParticleRead_GetBoolByIDDesc", "Returns a Bool value from a particle by NiagaraID.  Note: When reading from self this will be the previous frames data.");
 #endif
@@ -643,7 +636,6 @@ void UNiagaraDataInterfaceParticleRead::GetPersistentIDFunctions(TArray<FNiagara
 		Sig.FunctionSpecifiers.Add(FName("Attribute"));
 		Sig.bMemberFunction = true;
 		Sig.bRequiresContext = false;
-		Sig.bExperimental = true;
 #if WITH_EDITORONLY_DATA
 		Sig.Description = NSLOCTEXT("Niagara", "NiagaraDataInterfaceParticleRead_GetFloatByIDDesc", "Returns a Float value from a particle by NiagaraID.  Note: When reading from self this will be the previous frames data.");
 #endif
@@ -660,7 +652,6 @@ void UNiagaraDataInterfaceParticleRead::GetPersistentIDFunctions(TArray<FNiagara
 		Sig.FunctionSpecifiers.Add(FName("Attribute"));
 		Sig.bMemberFunction = true;
 		Sig.bRequiresContext = false;
-		Sig.bExperimental = true;
 #if WITH_EDITORONLY_DATA
 		Sig.Description = NSLOCTEXT("Niagara", "NiagaraDataInterfaceParticleRead_GetVector2ByIDDesc", "Returns a Vector2 value from a particle by NiagaraID.  Note: When reading from self this will be the previous frames data.");
 #endif
@@ -677,7 +668,6 @@ void UNiagaraDataInterfaceParticleRead::GetPersistentIDFunctions(TArray<FNiagara
 		Sig.FunctionSpecifiers.Add(FName("Attribute"));
 		Sig.bMemberFunction = true;
 		Sig.bRequiresContext = false;
-		Sig.bExperimental = true;
 #if WITH_EDITORONLY_DATA
 		Sig.Description = NSLOCTEXT("Niagara", "NiagaraDataInterfaceParticleRead_GetVector3ByIDDesc", "Returns a Vector3 value from a particle by NiagaraID.  Note: When reading from self this will be the previous frames data.");
 #endif
@@ -694,7 +684,6 @@ void UNiagaraDataInterfaceParticleRead::GetPersistentIDFunctions(TArray<FNiagara
 		Sig.FunctionSpecifiers.Add(FName("Attribute"));
 		Sig.bMemberFunction = true;
 		Sig.bRequiresContext = false;
-		Sig.bExperimental = true;
 #if WITH_EDITORONLY_DATA
 		Sig.Description = NSLOCTEXT("Niagara", "NiagaraDataInterfaceParticleRead_GetPositionByIDDesc", "Returns a Position value from a particle by NiagaraID.  Note: When reading from self this will be the previous frames data.");
 #endif
@@ -711,7 +700,6 @@ void UNiagaraDataInterfaceParticleRead::GetPersistentIDFunctions(TArray<FNiagara
 		Sig.FunctionSpecifiers.Add(FName("Attribute"));
 		Sig.bMemberFunction = true;
 		Sig.bRequiresContext = false;
-		Sig.bExperimental = true;
 #if WITH_EDITORONLY_DATA
 		Sig.Description = NSLOCTEXT("Niagara", "NiagaraDataInterfaceParticleRead_GetVector4ByIDDesc", "Returns a Vector4 value from a particle by NiagaraID.  Note: When reading from self this will be the previous frames data.");
 #endif
@@ -728,7 +716,6 @@ void UNiagaraDataInterfaceParticleRead::GetPersistentIDFunctions(TArray<FNiagara
 		Sig.FunctionSpecifiers.Add(FName("Attribute"));
 		Sig.bMemberFunction = true;
 		Sig.bRequiresContext = false;
-		Sig.bExperimental = true;
 #if WITH_EDITORONLY_DATA
 		Sig.Description = NSLOCTEXT("Niagara", "NiagaraDataInterfaceParticleRead_GetColorByIDDesc", "Returns a Color value from a particle by NiagaraID.  Note: When reading from self this will be the previous frames data.");
 #endif
@@ -745,7 +732,6 @@ void UNiagaraDataInterfaceParticleRead::GetPersistentIDFunctions(TArray<FNiagara
 		Sig.FunctionSpecifiers.Add(FName("Attribute"));
 		Sig.bMemberFunction = true;
 		Sig.bRequiresContext = false;
-		Sig.bExperimental = true;
 #if WITH_EDITORONLY_DATA
 		Sig.Description = NSLOCTEXT("Niagara", "NiagaraDataInterfaceParticleRead_GetQuatByIDDesc", "Returns a Quat value from a particle by NiagaraID.  Note: When reading from self this will be the previous frames data.");
 #endif
@@ -762,7 +748,6 @@ void UNiagaraDataInterfaceParticleRead::GetPersistentIDFunctions(TArray<FNiagara
 		Sig.FunctionSpecifiers.Add(FName("Attribute"));
 		Sig.bMemberFunction = true;
 		Sig.bRequiresContext = false;
-		Sig.bExperimental = true;
 #if WITH_EDITORONLY_DATA
 		Sig.Description = NSLOCTEXT("Niagara", "NiagaraDataInterfaceParticleRead_GetIDByIDDesc", "Returns a NiagaraID value from a particle by NiagaraID.  Note: When reading from self this will be the previous frames data.");
 #endif
@@ -788,7 +773,6 @@ void UNiagaraDataInterfaceParticleRead::GetIndexFunctions(TArray<FNiagaraFunctio
 		Sig.FunctionSpecifiers.Add(FName("Attribute"));
 		Sig.bMemberFunction = true;
 		Sig.bRequiresContext = false;
-		Sig.bExperimental = true;
 #if WITH_EDITORONLY_DATA
 		Sig.Description = NSLOCTEXT("Niagara", "NiagaraDataInterfaceParticleRead_GetIntByIndexDesc", "Returns an Int value from a particle by index.  Note: When reading from self this will be the previous frames data.");
 #endif
@@ -805,7 +789,6 @@ void UNiagaraDataInterfaceParticleRead::GetIndexFunctions(TArray<FNiagaraFunctio
 		Sig.FunctionSpecifiers.Add(FName("Attribute"));
 		Sig.bMemberFunction = true;
 		Sig.bRequiresContext = false;
-		Sig.bExperimental = true;
 #if WITH_EDITORONLY_DATA
 		Sig.Description = NSLOCTEXT("Niagara", "NiagaraDataInterfaceParticleRead_GetBoolByIndexDesc", "Returns a Bool value from a particle by index.  Note: When reading from self this will be the previous frames data.");
 #endif
@@ -822,7 +805,6 @@ void UNiagaraDataInterfaceParticleRead::GetIndexFunctions(TArray<FNiagaraFunctio
 		Sig.FunctionSpecifiers.Add(FName("Attribute"));
 		Sig.bMemberFunction = true;
 		Sig.bRequiresContext = false;
-		Sig.bExperimental = true;
 #if WITH_EDITORONLY_DATA
 		Sig.Description = NSLOCTEXT("Niagara", "NiagaraDataInterfaceParticleRead_GetFloatByIndexDesc", "Returns a Float value from a particle by index.  Note: When reading from self this will be the previous frames data.");
 #endif
@@ -839,7 +821,6 @@ void UNiagaraDataInterfaceParticleRead::GetIndexFunctions(TArray<FNiagaraFunctio
 		Sig.FunctionSpecifiers.Add(FName("Attribute"));
 		Sig.bMemberFunction = true;
 		Sig.bRequiresContext = false;
-		Sig.bExperimental = true;
 #if WITH_EDITORONLY_DATA
 		Sig.Description = NSLOCTEXT("Niagara", "NiagaraDataInterfaceParticleRead_GetVector2ByIndexDesc", "Returns a Vector2 value from a particle by index.  Note: When reading from self this will be the previous frames data.");
 #endif
@@ -856,7 +837,6 @@ void UNiagaraDataInterfaceParticleRead::GetIndexFunctions(TArray<FNiagaraFunctio
 		Sig.FunctionSpecifiers.Add(FName("Attribute"));
 		Sig.bMemberFunction = true;
 		Sig.bRequiresContext = false;
-		Sig.bExperimental = true;
 #if WITH_EDITORONLY_DATA
 		Sig.Description = NSLOCTEXT("Niagara", "NiagaraDataInterfaceParticleRead_GetVector3ByIndexDesc", "Returns a Vector3 value from a particle by index.  Note: When reading from self this will be the previous frames data.");
 #endif
@@ -873,7 +853,6 @@ void UNiagaraDataInterfaceParticleRead::GetIndexFunctions(TArray<FNiagaraFunctio
 		Sig.FunctionSpecifiers.Add(FName("Attribute"));
 		Sig.bMemberFunction = true;
 		Sig.bRequiresContext = false;
-		Sig.bExperimental = true;
 #if WITH_EDITORONLY_DATA
 		Sig.Description = NSLOCTEXT("Niagara", "NiagaraDataInterfaceParticleRead_GetVector4ByIndexDesc", "Returns a Vector4 value from a particle by index.  Note: When reading from self this will be the previous frames data.");
 #endif
@@ -890,7 +869,6 @@ void UNiagaraDataInterfaceParticleRead::GetIndexFunctions(TArray<FNiagaraFunctio
 		Sig.FunctionSpecifiers.Add(FName("Attribute"));
 		Sig.bMemberFunction = true;
 		Sig.bRequiresContext = false;
-		Sig.bExperimental = true;
 #if WITH_EDITORONLY_DATA
 		Sig.Description = NSLOCTEXT("Niagara", "NiagaraDataInterfaceParticleRead_GetPositionByIndexDesc", "Returns a position value from a particle by index.  Note: When reading from self this will be the previous frames data.");
 #endif
@@ -907,7 +885,6 @@ void UNiagaraDataInterfaceParticleRead::GetIndexFunctions(TArray<FNiagaraFunctio
 		Sig.FunctionSpecifiers.Add(FName("Attribute"));
 		Sig.bMemberFunction = true;
 		Sig.bRequiresContext = false;
-		Sig.bExperimental = true;
 #if WITH_EDITORONLY_DATA
 		Sig.Description = NSLOCTEXT("Niagara", "NiagaraDataInterfaceParticleRead_GetColorByIndexDesc", "Returns a Color value from a particle by index.  Note: When reading from self this will be the previous frames data.");
 #endif
@@ -924,7 +901,6 @@ void UNiagaraDataInterfaceParticleRead::GetIndexFunctions(TArray<FNiagaraFunctio
 		Sig.FunctionSpecifiers.Add(FName("Attribute"));
 		Sig.bMemberFunction = true;
 		Sig.bRequiresContext = false;
-		Sig.bExperimental = true;
 #if WITH_EDITORONLY_DATA
 		Sig.Description = NSLOCTEXT("Niagara", "NiagaraDataInterfaceParticleRead_GetQuatByIndexDesc", "Returns a Quat value from a particle by index.  Note: When reading from self this will be the previous frames data.");
 #endif
@@ -941,7 +917,6 @@ void UNiagaraDataInterfaceParticleRead::GetIndexFunctions(TArray<FNiagaraFunctio
 		Sig.FunctionSpecifiers.Add(FName("Attribute"));
 		Sig.bMemberFunction = true;
 		Sig.bRequiresContext = false;
-		Sig.bExperimental = true;
 #if WITH_EDITORONLY_DATA
 		Sig.Description = NSLOCTEXT("Niagara", "NiagaraDataInterfaceParticleRead_GetIDByIndexDesc", "Returns a NiagaraID value from a particle by index.  Note: When reading from self this will be the previous frames data.");
 #endif
diff --git a/Engine/Plugins/FX/Niagara/Source/Niagara/Private/NiagaraDataInterfaceRasterizationGrid3D.cpp b/Engine/Plugins/FX/Niagara/Source/Niagara/Private/NiagaraDataInterfaceRasterizationGrid3D.cpp
--- a/Engine/Plugins/FX/Niagara/Source/Niagara/Private/NiagaraDataInterfaceRasterizationGrid3D.cpp
+++ b/Engine/Plugins/FX/Niagara/Source/Niagara/Private/NiagaraDataInterfaceRasterizationGrid3D.cpp
@@ -90,7 +90,6 @@ void UNiagaraDataInterfaceRasterizationGrid3D::GetFunctionsInternal(TArray<FNiag
 		Sig.Outputs.Add(FNiagaraVariable(FNiagaraTypeDefinition::GetBoolDef(), TEXT("Success")));
 
 		Sig.ModuleUsageBitmask = ENiagaraScriptUsageMask::Emitter | ENiagaraScriptUsageMask::System;
-		Sig.bExperimental = true;
 		Sig.bMemberFunction = true;
 		Sig.bRequiresExecPin = true;
 		Sig.bRequiresContext = false;
@@ -107,7 +106,6 @@ void UNiagaraDataInterfaceRasterizationGrid3D::GetFunctionsInternal(TArray<FNiag
 		Sig.Outputs.Add(FNiagaraVariable(FNiagaraTypeDefinition::GetBoolDef(), TEXT("Success")));
 
 		Sig.ModuleUsageBitmask = ENiagaraScriptUsageMask::Emitter | ENiagaraScriptUsageMask::System;
-		Sig.bExperimental = true;
 		Sig.bMemberFunction = true;
 		Sig.bRequiresExecPin = true;
 		Sig.bRequiresContext = false;
@@ -124,7 +122,6 @@ void UNiagaraDataInterfaceRasterizationGrid3D::GetFunctionsInternal(TArray<FNiag
 		Sig.Outputs.Add(FNiagaraVariable(FNiagaraTypeDefinition::GetBoolDef(), TEXT("Success")));
 
 		Sig.ModuleUsageBitmask = ENiagaraScriptUsageMask::Emitter | ENiagaraScriptUsageMask::System;
-		Sig.bExperimental = true;
 		Sig.bMemberFunction = true;
 		Sig.bRequiresExecPin = true;
 		Sig.bRequiresContext = false;
@@ -144,7 +141,6 @@ void UNiagaraDataInterfaceRasterizationGrid3D::GetFunctionsInternal(TArray<FNiag
 		Sig.Inputs.Add(FNiagaraVariable(FNiagaraTypeDefinition::GetFloatDef(), TEXT("Value")));
 		Sig.Outputs.Add(FNiagaraVariable(FNiagaraTypeDefinition::GetIntDef(), TEXT("IGNORE")));
 
-		Sig.bExperimental = true;
 		Sig.bMemberFunction = true;
 		Sig.bRequiresContext = false;
 		Sig.bWriteFunction = true;
@@ -167,7 +163,6 @@ void UNiagaraDataInterfaceRasterizationGrid3D::GetFunctionsInternal(TArray<FNiag
 		Sig.Inputs.Add(FNiagaraVariable(FNiagaraTypeDefinition::GetIntDef(), TEXT("Value")));		
 		Sig.Outputs.Add(FNiagaraVariable(FNiagaraTypeDefinition::GetIntDef(), TEXT("IGNORE")));
 
-		Sig.bExperimental = true;
 		Sig.bMemberFunction = true;
 		Sig.bRequiresContext = false;
 		Sig.bWriteFunction = true;
@@ -190,7 +185,6 @@ void UNiagaraDataInterfaceRasterizationGrid3D::GetFunctionsInternal(TArray<FNiag
 		Sig.Inputs.Add(FNiagaraVariable(FNiagaraTypeDefinition::GetFloatDef(), TEXT("Value")));
 		Sig.Outputs.Add(FNiagaraVariable(FNiagaraTypeDefinition::GetIntDef(), TEXT("IGNORE")));
 
-		Sig.bExperimental = true;
 		Sig.bMemberFunction = true;
 		Sig.bRequiresContext = false;
 		Sig.bWriteFunction = true;
@@ -216,7 +210,6 @@ void UNiagaraDataInterfaceRasterizationGrid3D::GetFunctionsInternal(TArray<FNiag
 		Sig.Inputs.Add(FNiagaraVariable(FNiagaraTypeDefinition::GetFloatDef(), TEXT("Value")));
 		Sig.Outputs.Add(FNiagaraVariable(FNiagaraTypeDefinition::GetFloatDef(), TEXT("OriginalValue")));
 
-		Sig.bExperimental = true;
 		Sig.bMemberFunction = true;
 		Sig.bRequiresContext = false;
 		Sig.bWriteFunction = true;
@@ -239,7 +232,6 @@ void UNiagaraDataInterfaceRasterizationGrid3D::GetFunctionsInternal(TArray<FNiag
 		Sig.Inputs.Add(FNiagaraVariable(FNiagaraTypeDefinition::GetIntDef(), TEXT("Value")));
 		Sig.Outputs.Add(FNiagaraVariable(FNiagaraTypeDefinition::GetIntDef(), TEXT("IGNORE")));
 
-		Sig.bExperimental = true;
 		Sig.bMemberFunction = true;
 		Sig.bRequiresContext = false;
 		Sig.bWriteFunction = true;
@@ -262,7 +254,6 @@ void UNiagaraDataInterfaceRasterizationGrid3D::GetFunctionsInternal(TArray<FNiag
 		Sig.Inputs.Add(FNiagaraVariable(FNiagaraTypeDefinition::GetFloatDef(), TEXT("Value")));
 		Sig.Outputs.Add(FNiagaraVariable(FNiagaraTypeDefinition::GetIntDef(), TEXT("IsSafe")));
 
-		Sig.bExperimental = true;
 		Sig.bMemberFunction = true;
 		Sig.bRequiresContext = false;
 		Sig.bWriteFunction = true;
@@ -285,7 +276,6 @@ void UNiagaraDataInterfaceRasterizationGrid3D::GetFunctionsInternal(TArray<FNiag
 		Sig.Inputs.Add(FNiagaraVariable(FNiagaraTypeDefinition::GetFloatDef(), TEXT("Value")));
 		Sig.Outputs.Add(FNiagaraVariable(FNiagaraTypeDefinition::GetIntDef(), TEXT("IGNORE")));
 
-		Sig.bExperimental = true;
 		Sig.bMemberFunction = true;
 		Sig.bRequiresContext = false;
 		Sig.bWriteFunction = true;
@@ -311,7 +301,6 @@ void UNiagaraDataInterfaceRasterizationGrid3D::GetFunctionsInternal(TArray<FNiag
 		Sig.Inputs.Add(FNiagaraVariable(FNiagaraTypeDefinition::GetFloatDef(), TEXT("Value")));
 		Sig.Outputs.Add(FNiagaraVariable(FNiagaraTypeDefinition::GetFloatDef(), TEXT("OriginalValue")));
 
-		Sig.bExperimental = true;
 		Sig.bMemberFunction = true;
 		Sig.bRequiresContext = false;
 		Sig.bWriteFunction = true;
@@ -331,7 +320,6 @@ void UNiagaraDataInterfaceRasterizationGrid3D::GetFunctionsInternal(TArray<FNiag
 		Sig.Inputs.Add(FNiagaraVariable(FNiagaraTypeDefinition::GetFloatDef(), TEXT("Value")));
 		Sig.Outputs.Add(FNiagaraVariable(FNiagaraTypeDefinition::GetIntDef(), TEXT("IGNORE")));
 
-		Sig.bExperimental = true;
 		Sig.bMemberFunction = true;
 		Sig.bRequiresContext = false;
 		Sig.bWriteFunction = true;
@@ -356,7 +344,6 @@ void UNiagaraDataInterfaceRasterizationGrid3D::GetFunctionsInternal(TArray<FNiag
 		Sig.Inputs.Add(FNiagaraVariable(FNiagaraTypeDefinition::GetFloatDef(), TEXT("Value")));
 		Sig.Outputs.Add(FNiagaraVariable(FNiagaraTypeDefinition::GetFloatDef(), TEXT("OriginalValue")));
 
-		Sig.bExperimental = true;
 		Sig.bMemberFunction = true;
 		Sig.bRequiresContext = false;
 		Sig.bWriteFunction = true;
@@ -379,7 +366,6 @@ void UNiagaraDataInterfaceRasterizationGrid3D::GetFunctionsInternal(TArray<FNiag
 		Sig.Inputs.Add(FNiagaraVariable(FNiagaraTypeDefinition::GetIntDef(), TEXT("AttributeIndex")));
 		Sig.Outputs.Add(FNiagaraVariable(FNiagaraTypeDefinition::GetFloatDef(), TEXT("Value")));
 
-		Sig.bExperimental = true;
 		Sig.bMemberFunction = true;
 		Sig.bRequiresContext = false;
 		Sig.ModuleUsageBitmask = ENiagaraScriptUsageMask::Particle;
@@ -402,7 +388,6 @@ void UNiagaraDataInterfaceRasterizationGrid3D::GetFunctionsInternal(TArray<FNiag
 		Sig.Inputs.Add(FNiagaraVariable(FNiagaraTypeDefinition::GetIntDef(), TEXT("AttributeIndex")));
 		Sig.Outputs.Add(FNiagaraVariable(FNiagaraTypeDefinition::GetIntDef(), TEXT("Value")));
 
-		Sig.bExperimental = true;
 		Sig.bMemberFunction = true;
 		Sig.bRequiresContext = false;
 		Sig.ModuleUsageBitmask = ENiagaraScriptUsageMask::Particle;
