- Move SetEmitterEnabled fail into verbose logging to avoid log spam in cooked builds
- Remove verbose logging from shipping builds also

#rnx

diff --git a/Engine/Plugins/FX/Niagara/Source/Niagara/Private/NiagaraCommon.cpp b/Engine/Plugins/FX/Niagara/Source/Niagara/Private/NiagaraCommon.cpp
--- a/Engine/Plugins/FX/Niagara/Source/Niagara/Private/NiagaraCommon.cpp
+++ b/Engine/Plugins/FX/Niagara/Source/Niagara/Private/NiagaraCommon.cpp
@@ -937,7 +937,7 @@ void FNiagaraMaterialAttributeBinding::CacheValues(const UNiagaraEmitter* InEmit
 
 //////////////////////////////////////////////////////////////////////////
 
-#if !NO_LOGGING
+#if NIAGARA_VERBOSE_LOGGING_ENABLED
 bool FNiagaraUtilities::LogVerboseWarnings()
 {
 	return GNiagaraLogVerboseWarnings != 0;
diff --git a/Engine/Plugins/FX/Niagara/Source/Niagara/Private/NiagaraSystemInstance.cpp b/Engine/Plugins/FX/Niagara/Source/Niagara/Private/NiagaraSystemInstance.cpp
--- a/Engine/Plugins/FX/Niagara/Source/Niagara/Private/NiagaraSystemInstance.cpp
+++ b/Engine/Plugins/FX/Niagara/Source/Niagara/Private/NiagaraSystemInstance.cpp
@@ -176,8 +176,10 @@ void FNiagaraSystemInstance::SetEmitterEnable(FName EmitterName, bool bNewEnable
 		}
 	}
 
-	// Failed to find emitter
-	UE_LOG(LogNiagara, Warning, TEXT("SetEmitterEnable: Failed to find Emitter(%s) System(%s) Component(%s)"), *EmitterName.ToString(), *GetNameSafe(System), *GetFullNameSafe(AttachComponent.Get()));
+	if (FNiagaraUtilities::LogVerboseWarnings())
+	{
+		UE_LOG(LogNiagara, Warning, TEXT("SetEmitterEnable: Failed to find Emitter(%s) System(%s) Component(%s)"), *EmitterName.ToString(), *GetNameSafe(System), *GetFullNameSafe(AttachComponent.Get()));
+	}
 }
 
 void FNiagaraSystemInstance::Init(bool bInForceSolo)
diff --git a/Engine/Plugins/FX/Niagara/Source/Niagara/Public/NiagaraCommon.h b/Engine/Plugins/FX/Niagara/Source/Niagara/Public/NiagaraCommon.h
--- a/Engine/Plugins/FX/Niagara/Source/Niagara/Public/NiagaraCommon.h
+++ b/Engine/Plugins/FX/Niagara/Source/Niagara/Public/NiagaraCommon.h
@@ -48,6 +48,9 @@ using FNiagaraShaderMapRef = TRefCountPtr<FNiagaraShaderMap>;
 
 #define NIAGARA_SYSTEM_CAPTURE (!UE_BUILD_SHIPPING || WITH_EDITORONLY_DATA)
 
+#define NIAGARA_VERBOSE_LOGGING_ENABLED (!NO_LOGGING && !UE_BUILD_SHIPPING)
+
+
 #ifndef NIAGARA_COMPUTEDEBUG_ENABLED
 	#define NIAGARA_COMPUTEDEBUG_ENABLED WITH_EDITOR
 #endif
@@ -1524,10 +1527,10 @@ namespace FNiagaraUtilities
 	}
 
 	// When enabled log more information for the end user
-#if NO_LOGGING
-	inline bool LogVerboseWarnings() { return false; }
-#else
+#if NIAGARA_VERBOSE_LOGGING_ENABLED
 	bool LogVerboseWarnings();
+#else
+	constexpr bool LogVerboseWarnings() { return false; }
 #endif
 
 	// Whether GPU particles are currently allowed. Could change depending on config and runtime switches.
