- Add scene component interface to abstract actor vs entity hierarchy
- Use the new interface in the static mesh data interface
- Allow the digest and runtime to use entities, currently disabled in the digest as converting all object types it's likely not the right solution

diff --git a/Engine/Plugins/FX/Niagara/Source/Niagara/Internal/DataInterface/NiagaraDataInterfaceStaticMesh.h b/Engine/Plugins/FX/Niagara/Source/Niagara/Internal/DataInterface/NiagaraDataInterfaceStaticMesh.h
--- a/Engine/Plugins/FX/Niagara/Source/Niagara/Internal/DataInterface/NiagaraDataInterfaceStaticMesh.h
+++ b/Engine/Plugins/FX/Niagara/Source/Niagara/Internal/DataInterface/NiagaraDataInterfaceStaticMesh.h
@@ -189,7 +189,7 @@ public:
 #endif
 	//~ UNiagaraDataInterface interface
 
-	NIAGARA_API UStaticMesh* GetStaticMesh(USceneComponent*& OutComponent, class FNiagaraSystemInstance* SystemInstance = nullptr, UObject* ParameterBindingValue = nullptr);
+	NIAGARA_API UStaticMesh* ResolveStaticMesh(UObject*& OutComponent, class FNiagaraSystemInstance* SystemInstance = nullptr, UObject* ParameterBindingValue = nullptr) const;
 	NIAGARA_API void SetSourceComponentFromBlueprints(UStaticMeshComponent* ComponentToUse);
 	NIAGARA_API void SetDefaultMeshFromBlueprints(UStaticMesh* MeshToUse);
 
diff --git a/Engine/Plugins/FX/Niagara/Source/Niagara/Private/DataInterface/NiagaraDataInterfaceStaticMesh.cpp b/Engine/Plugins/FX/Niagara/Source/Niagara/Private/DataInterface/NiagaraDataInterfaceStaticMesh.cpp
--- a/Engine/Plugins/FX/Niagara/Source/Niagara/Private/DataInterface/NiagaraDataInterfaceStaticMesh.cpp
+++ b/Engine/Plugins/FX/Niagara/Source/Niagara/Private/DataInterface/NiagaraDataInterfaceStaticMesh.cpp
@@ -9,6 +9,7 @@
 #include "Engine/StaticMesh.h"
 #include "NiagaraRenderer.h"
 #include "Misc/LargeWorldRenderPosition.h"
+#include "NiagaraSceneComponentUtils.h"
 #include "NiagaraSettings.h"
 #include "NiagaraCompileHashVisitor.h"
 #include "NiagaraFunctionLibrary.h"
@@ -573,8 +574,10 @@ namespace NDIStaticMeshLocal
 
 	struct FInstanceData_GameThread
 	{
+		const INiagaraSceneComponentUtils* SceneComponentUtils = nullptr;
+
 		/** Cached ptr to StaticMeshComponent we sample from, when found. Otherwise, the SceneComponent to use to transform the Default or Preview mesh. */
-		TWeakObjectPtr<USceneComponent> SceneComponentWeakPtr;
+		TWeakObjectPtr<UObject> ComponentWeakPtr;
 
 		/** Cached ptr to the mesh so that we can make sure that we haven't been deleted. */
 		TWeakObjectPtr<UStaticMesh> StaticMeshWeakPtr;
@@ -744,34 +747,18 @@ namespace NDIStaticMeshLocal
 		}
 #endif
 
-		void UpdateTransforms(USceneComponent* SceneComponent, FNiagaraSystemInstance* SystemInstance)
+		void UpdateTransforms(UObject* Component, FNiagaraSystemInstance* SystemInstance)
 		{
 			SystemInstanceWorldTransform = SystemInstance->GetWorldTransform();
 			LWCTileOffset = FVector(SystemInstance->GetLWCTile()) * -FLargeWorldRenderScalar::GetTileSize();
 
 			FTransform ComponentTransform;
-			if (SceneComponent)
+			if (SceneComponentUtils->GetStaticMeshTransforms(Component, ComponentTransform, ISMTransforms))
 			{
-				ComponentTransform = SceneComponent->GetComponentToWorld();
-				if (UInstancedStaticMeshComponent* ISMComponent = Cast<UInstancedStaticMeshComponent>(SceneComponent))
-				{
-					const int32 NumInstances = ISMComponent->PerInstanceSMData.Num();
-					ISMTransforms.SetNum(NumInstances);
-					for (int32 i = 0; i < NumInstances; ++i)
-					{
-						ISMComponent->GetInstanceTransform(i, ISMTransforms[i], true);
-					}
-
-					if (ISMTransforms.IsValidIndex(InstanceIndex))
-					{
-						ComponentTransform = ISMTransforms[InstanceIndex];
-					}
-				}
-				else
+				if (ISMTransforms.IsValidIndex(InstanceIndex))
 				{
-					ISMTransforms.Empty();
+					ComponentTransform = ISMTransforms[InstanceIndex];
 				}
-
 				OwnerToMeshVector = FVector3f(ComponentTransform.GetLocation() - SystemInstanceWorldTransform.GetLocation());
 				ComponentTransform.AddToTranslation(LWCTileOffset);
 			}
@@ -787,20 +774,20 @@ namespace NDIStaticMeshLocal
 			Rotation = FQuat4f(ComponentTransform.GetRotation());
 		}
 
-		void UpdateTransformsResetPrevious(USceneComponent* SceneComponent, FNiagaraSystemInstance* SystemInstance)
+		void UpdateTransformsResetPrevious(UObject* Component, FNiagaraSystemInstance* SystemInstance)
 		{
-			UpdateTransforms(SceneComponent, SystemInstance);
+			UpdateTransforms(Component, SystemInstance);
 			PrevTransform = Transform;
 			PrevTransformInverseTransposed = TransformInverseTransposed;
 			PrevRotation = Rotation;
 		}
 
-		void UpdateTransformsCopyPrevious(USceneComponent* SceneComponent, FNiagaraSystemInstance* SystemInstance)
+		void UpdateTransformsCopyPrevious(UObject* Component, FNiagaraSystemInstance* SystemInstance)
 		{
 			PrevTransform = Transform;
 			PrevTransformInverseTransposed = TransformInverseTransposed;
 			PrevRotation = Rotation;
-			UpdateTransforms(SceneComponent, SystemInstance);
+			UpdateTransforms(Component, SystemInstance);
 		}
 
 		void UpdateTransformsInstanceIndexChanged()
@@ -824,7 +811,10 @@ namespace NDIStaticMeshLocal
 			check(SystemInstance);
 
 			// Initialize members
-			SceneComponentWeakPtr = nullptr;
+			PRAGMA_DISABLE_DEPRECATION_WARNINGS
+			SceneComponentUtils = &FNiagaraSystemInstance::GetSceneComponentUtils(SystemInstance);
+			PRAGMA_ENABLE_DEPRECATION_WARNINGS
+			ComponentWeakPtr = nullptr;
 			StaticMeshWeakPtr = nullptr;
 			Transform = FMatrix::Identity;
 			TransformInverseTransposed = FMatrix::Identity;
@@ -860,20 +850,17 @@ namespace NDIStaticMeshLocal
 			MeshParameterBinding.Init(SystemInstance->GetInstanceParameters(), Interface->MeshParameterBinding.Parameter);
 
 			// Get component / mesh we are using
-			USceneComponent* SceneComponent = nullptr;
-			UStaticMesh* StaticMesh = Interface->GetStaticMesh(SceneComponent, SystemInstance, MeshParameterBinding.GetValue());
-			SceneComponentWeakPtr = SceneComponent;
+			UObject* Component = nullptr;
+			UStaticMesh* StaticMesh = Interface->ResolveStaticMesh(Component, SystemInstance, MeshParameterBinding.GetValue());
+			ComponentWeakPtr = Component;
 
 			// Gather attached information
 			InstanceIndex = Interface->InstanceIndex;
-			UpdateTransformsResetPrevious(SceneComponent, SystemInstance);
+			UpdateTransformsResetPrevious(Component, SystemInstance);
 
 			if (bUsePhysicsVelocity)
 			{
-				if (UStaticMeshComponent* MeshComponent = Cast<UStaticMeshComponent>(SceneComponent))
-				{
-					PhysicsVelocity = MeshComponent->GetPhysicsLinearVelocity();
-				}
+				PhysicsVelocity = SceneComponentUtils->GetPhysicsLinearVelocity(Component);
 			}
 
 			if (GNiagaraFailStaticMeshDataInterface != 0)
@@ -1005,7 +992,7 @@ namespace NDIStaticMeshLocal
 				{
 					if (NumMeshSockets > TNumericLimits<uint16>::Max())
 					{
-						UE_LOG(LogNiagara, Warning, TEXT("Static Mesh Data Interface '%s' requires more sockets '%d' than we currently support '%d' Mesh '%s' Component '%s'"), *GetFullNameSafe(Interface), NumMeshSockets, TNumericLimits<uint16>::Max(), *GetFullNameSafe(StaticMesh), *GetFullNameSafe(SceneComponent));
+						UE_LOG(LogNiagara, Warning, TEXT("Static Mesh Data Interface '%s' requires more sockets '%d' than we currently support '%d' Mesh '%s' Component '%s'"), *GetFullNameSafe(Interface), NumMeshSockets, TNumericLimits<uint16>::Max(), *GetFullNameSafe(StaticMesh), *GetFullNameSafe(Component));
 						return false;
 					}
 
@@ -1027,7 +1014,7 @@ namespace NDIStaticMeshLocal
 							const int32 SocketIndex = StaticMesh->Sockets.IndexOfByPredicate([&](const UStaticMeshSocket* Socket) { return Socket->SocketName == FilteredSocketName; });
 							if (SocketIndex == INDEX_NONE)
 							{
-								UE_LOG(LogNiagara, Warning, TEXT("Static Mesh Data Interface '%s' could not find socket '%s' Mesh '%s' Component '%s'"), *GetFullNameSafe(Interface), *FilteredSocketName.ToString(), *GetFullNameSafe(StaticMesh), *GetFullNameSafe(SceneComponent));
+								UE_LOG(LogNiagara, Warning, TEXT("Static Mesh Data Interface '%s' could not find socket '%s' Mesh '%s' Component '%s'"), *GetFullNameSafe(Interface), *FilteredSocketName.ToString(), *GetFullNameSafe(StaticMesh), *GetFullNameSafe(Component));
 								continue;
 							}
 							++NumFilteredSockets;
@@ -1100,22 +1087,15 @@ namespace NDIStaticMeshLocal
 
 			DeltaSeconds = InDeltaSeconds;
 
-			USceneComponent* SceneComponent = SceneComponentWeakPtr.Get();
+			UObject* Component = ComponentWeakPtr.Get();
 			if (bCaptureTransformsPerFrame)
 			{
-				UpdateTransformsCopyPrevious(SceneComponent, SystemInstance);
+				UpdateTransformsCopyPrevious(Component, SystemInstance);
 			}
 
 			if (bUsePhysicsVelocity)
 			{
-				if (UStaticMeshComponent* MeshComponent = Cast<UStaticMeshComponent>(SceneComponent))
-				{
-					PhysicsVelocity = MeshComponent->GetPhysicsLinearVelocity();
-				}
-				else
-				{
-					PhysicsVelocity = FVector::ZeroVector;
-				}
+				PhysicsVelocity = SceneComponentUtils->GetPhysicsLinearVelocity(Component);
 			}
 
 			if (UvMappingIndexSet != INDEX_NONE && !UvMapping)
@@ -1139,7 +1119,7 @@ namespace NDIStaticMeshLocal
 
 		bool ResetRequired(UNiagaraDataInterfaceStaticMesh* Interface, FNiagaraSystemInstance* Instance)
 		{
-			USceneComponent* Component = SceneComponentWeakPtr.Get();
+			UObject* Component = ComponentWeakPtr.Get();
 			if (bComponentValid && !Component)
 			{
 				// The component we were bound to is no longer valid so we have to trigger a reset.
@@ -1160,8 +1140,8 @@ namespace NDIStaticMeshLocal
 			}
 
 			// Check to see if we would resolve to a different mesh / component
-			USceneComponent* ResolvedComponent = nullptr;
-			UStaticMesh* ResolvedMesh = Interface->GetStaticMesh(ResolvedComponent, Instance, MeshParameterBinding.GetValue());
+			UObject* ResolvedComponent = nullptr;
+			UStaticMesh* ResolvedMesh = Interface->ResolveStaticMesh(ResolvedComponent, Instance, MeshParameterBinding.GetValue());
 			if (Component != ResolvedComponent || Mesh != ResolvedMesh)
 			{
 				return true;
@@ -1216,13 +1196,7 @@ namespace NDIStaticMeshLocal
 
 				if (CachedLODIdx >= CurrentFirstLOD && StaticMesh->GetRenderData()->LODResources.IsValidIndex(CachedLODIdx))
 				{
-					if(UStaticMeshComponent* SMComp = Cast<UStaticMeshComponent>(SceneComponentWeakPtr.Get()))
-					{
-						if(SMComp->LODData.IsValidIndex(CachedLODIdx))
-						{
-							OutVertexColorOverrides = SMComp->LODData[CachedLODIdx].OverrideVertexColors;
-						}
-					}
+					OutVertexColorOverrides = SceneComponentUtils->GetStaticMeshOverrideColors(ComponentWeakPtr.Get(), CachedLODIdx);
 				
 					if (const FStaticMeshLODResources* LODResource = &StaticMesh->GetRenderData()->LODResources[CachedLODIdx])
 					{
@@ -1936,10 +1910,7 @@ void UNiagaraDataInterfaceStaticMesh::ProvidePerInstanceDataForRenderThread(void
 	InstanceData->GetCurrentLODWithVertexColorOverrides(DataFromGT->OverrideVertexColors);
 #endif
 
-	if ( UPrimitiveComponent* PrimitiveComponent = Cast<UPrimitiveComponent>(InstanceData->SceneComponentWeakPtr) )
-	{
-		DataFromGT->DistanceFieldPrimitiveId = PrimitiveComponent->GetPrimitiveSceneId();
-	}
+	DataFromGT->DistanceFieldPrimitiveId = InstanceData->SceneComponentUtils->GetPrimitiveSceneId(InstanceData->ComponentWeakPtr.Get());
 }
 
 #if WITH_EDITORONLY_DATA
@@ -2363,7 +2334,6 @@ void UNiagaraDataInterfaceStaticMesh::GetUVMappingFunctions(TArray<FNiagaraFunct
 		Sig.Outputs.Add(FNiagaraVariable(FNiagaraTypeDefinition::GetIntDef(), TEXT("Triangle")));
 		Sig.Outputs.Add(FNiagaraVariable(FNiagaraTypeDefinition::GetVec3Def(), TEXT("BaryCoord")));
 		Sig.Outputs.Add(FNiagaraVariable(FNiagaraTypeDefinition::GetBoolDef(), TEXT("IsValid")));
-		Sig.bExperimental = true;
 	}
 
 	{
@@ -2377,7 +2347,6 @@ void UNiagaraDataInterfaceStaticMesh::GetUVMappingFunctions(TArray<FNiagaraFunct
 		Sig.Outputs.Add(FNiagaraVariable(FNiagaraTypeDefinition::GetIntDef(), TEXT("Triangle")));
 		Sig.Outputs.Add(FNiagaraVariable(FNiagaraTypeDefinition::GetVec3Def(), TEXT("BaryCoord")));
 		Sig.Outputs.Add(FNiagaraVariable(FNiagaraTypeDefinition::GetBoolDef(), TEXT("IsValid")));
-		Sig.bExperimental = true;
 	}
 
 	{
@@ -2388,8 +2357,6 @@ void UNiagaraDataInterfaceStaticMesh::GetUVMappingFunctions(TArray<FNiagaraFunct
 		Sig.bRequiresExecPin = true;
 
 		Sig.Inputs.Emplace(FNiagaraVariable(FNiagaraTypeDefinition::GetIntDef(), TEXT("UvSet")));
-
-		Sig.bExperimental = true;
 	}
 }
 
@@ -2403,7 +2370,6 @@ void UNiagaraDataInterfaceStaticMesh::GetDistanceFieldFunctions(TArray<FNiagaraF
 		FNiagaraFunctionSignature& Sig = OutFunctions.Add_GetRef(BaseSignature);
 		Sig.Name = QueryDistanceFieldName;
 		Sig.bSupportsCPU = false;
-		Sig.bExperimental = true;
 		Sig.Inputs.Emplace_GetRef(FNiagaraTypeDefinition::GetBoolDef(), TEXT("Execute")).SetValue(true);
 		Sig.Inputs.Emplace(FNiagaraTypeDefinition::GetPositionDef(), TEXT("World Position"));
 		Sig.Inputs.Emplace(UseMaxDistanceVar);
@@ -3345,9 +3311,9 @@ void UNiagaraDataInterfaceStaticMesh::DrawDebugHud(FNDIDrawDebugHudContext& Debu
 		return;
 	}
 
-	USceneComponent* SceneComponent = InstanceData_GT->SceneComponentWeakPtr.Get();
+	UObject* Component = InstanceData_GT->ComponentWeakPtr.Get();
 	UStaticMesh* SkeletalMesh = InstanceData_GT->StaticMeshWeakPtr.Get();
-	DebugHudContext.GetOutputString().Appendf(TEXT("StaticMesh(%s) StaticMeshComp(%s)"), *GetNameSafe(SkeletalMesh), *GetNameSafe(SceneComponent));
+	DebugHudContext.GetOutputString().Appendf(TEXT("StaticMesh(%s) StaticMeshComp(%s)"), *GetNameSafe(SkeletalMesh), *GetNameSafe(Component));
 }
 #endif
 
@@ -3458,132 +3424,46 @@ void UNiagaraDataInterfaceStaticMesh::GetFeedback(UNiagaraSystem* Asset, UNiagar
 }
 #endif //WITH_EDITOR
 
-UStaticMesh* UNiagaraDataInterfaceStaticMesh::GetStaticMesh(USceneComponent*& OutComponent, class FNiagaraSystemInstance* SystemInstance, UObject* ParameterBindingValue)
+UStaticMesh* UNiagaraDataInterfaceStaticMesh::ResolveStaticMesh(UObject*& OutComponent, FNiagaraSystemInstance* SystemInstance, UObject* ParameterBindingValue) const
 {
-	// Helper to scour an actor (or its parents) for a valid Static mesh component
-	auto FindActorMeshComponent = [](AActor* Actor, bool bRecurseParents = false) -> UStaticMeshComponent*
-	{
-		if (AStaticMeshActor* MeshActor = Cast<AStaticMeshActor>(Actor))
-		{
-			UStaticMeshComponent* Comp = MeshActor->GetStaticMeshComponent();
-			if (::IsValid(Comp))
-			{
-				return Comp;
-			}
-		}
-
-		// Fall back on any valid component on the actor
-		while (Actor)
-		{
-			for (UActorComponent* ActorComp : Actor->GetComponents())
-			{
-				UStaticMeshComponent* Comp = Cast<UStaticMeshComponent>(ActorComp);
-				if (::IsValid(Comp) && Comp->GetStaticMesh() != nullptr)
-				{
-					return Comp;
-				}
-			}
-
-			if (bRecurseParents)
-			{
-				Actor = Actor->GetParentActor();
-			}
-			else
-			{
-				break;
-			}
-		}
-
-		return nullptr;
-	};
-
 	const bool bTrySource = SourceMode == ENDIStaticMesh_SourceMode::Default || SourceMode == ENDIStaticMesh_SourceMode::Source;
 	const bool bTryMeshParameter = SourceMode == ENDIStaticMesh_SourceMode::Default || SourceMode == ENDIStaticMesh_SourceMode::MeshParameterBinding;
 	const bool bTryAttachParent = SourceMode == ENDIStaticMesh_SourceMode::Default || SourceMode == ENDIStaticMesh_SourceMode::AttachParent;
 	const bool bTryDefaultMesh = SourceMode == ENDIStaticMesh_SourceMode::Default || SourceMode == ENDIStaticMesh_SourceMode::DefaultMeshOnly;
 
 	UStaticMesh* OutMesh = nullptr;
-	UStaticMeshComponent* OutStaticMeshComponent = nullptr;
 	if (bTrySource && ::IsValid(SourceComponent))
 	{
-		OutStaticMeshComponent = SourceComponent;
-		OutMesh = OutStaticMeshComponent->GetStaticMesh();
+		OutComponent = SourceComponent;
+		OutMesh = SourceComponent->GetStaticMesh();
 	}
 	else if (bTrySource && SoftSourceActor.Get())
 	{
-		OutStaticMeshComponent = FindActorMeshComponent(SoftSourceActor.Get());
-		if (OutStaticMeshComponent)
-		{
-			OutMesh = OutStaticMeshComponent->GetStaticMesh();
-		}
+		PRAGMA_DISABLE_DEPRECATION_WARNINGS
+		const INiagaraSceneComponentUtils& SceneComponentUtils = FNiagaraSystemInstance::GetSceneComponentUtils(SystemInstance);
+		PRAGMA_ENABLE_DEPRECATION_WARNINGS
+
+		SceneComponentUtils.ResolveStaticMesh(SoftSourceActor.Get(), false, OutComponent, OutMesh);
 	}
 	else
 	{
+		PRAGMA_DISABLE_DEPRECATION_WARNINGS
+		const INiagaraSceneComponentUtils& SceneComponentUtils = FNiagaraSystemInstance::GetSceneComponentUtils(SystemInstance);
+		PRAGMA_ENABLE_DEPRECATION_WARNINGS
+
 		// Try parameter binding, this could either be a component or a static mesh
 		if (bTryMeshParameter && ParameterBindingValue)
 		{
-			if (AActor* Actor = Cast<AActor>(ParameterBindingValue))
-			{
-				OutStaticMeshComponent = FindActorMeshComponent(Actor, true);
-				if (OutStaticMeshComponent)
-				{
-					OutMesh = OutStaticMeshComponent->GetStaticMesh();
-				}
-			}
-			else
-			{
-				OutStaticMeshComponent = Cast<UStaticMeshComponent>(ParameterBindingValue);
-				if (OutStaticMeshComponent)
-				{
-					OutMesh = OutStaticMeshComponent->GetStaticMesh();
-				}
-				else
-				{
-					OutMesh = Cast<UStaticMesh>(ParameterBindingValue);
-				}
-			}
+			SceneComponentUtils.ResolveStaticMesh(ParameterBindingValue, true, OutComponent, OutMesh);
 		}
 
 		// Try attached parent, only if we did not already find a valid component / mesh
-		if (bTryAttachParent && SystemInstance && !OutStaticMeshComponent && !OutMesh)
+		if (bTryAttachParent && SystemInstance && !OutComponent && !OutMesh)
 		{
-			if (USceneComponent* AttachComponent = SystemInstance->GetAttachComponent())
-			{
-				// First, try to find the mesh component up the attachment hierarchy
-				for (USceneComponent* Curr = AttachComponent; Curr; Curr = Curr->GetAttachParent())
-				{
-					UStaticMeshComponent* ParentComp = Cast<UStaticMeshComponent>(Curr);
-					if (::IsValid(ParentComp))
-					{
-						OutStaticMeshComponent = ParentComp;
-						break;
-					}
-				}
-
-				if (!OutStaticMeshComponent)
-				{
-					// Next, try to find one in our outer chain
-					UStaticMeshComponent* OuterComp = AttachComponent->GetTypedOuter<UStaticMeshComponent>();
-					if (::IsValid(OuterComp))
-					{
-						OutStaticMeshComponent = OuterComp;
-					}
-					else if (AActor* Actor = AttachComponent->GetAttachmentRootActor())
-					{
-						// Final fall-back, look for any mesh component on our root actor or any of its parents
-						OutStaticMeshComponent = FindActorMeshComponent(Actor, true);
-					}
-				}
-
-				if (OutStaticMeshComponent)
-				{
-					OutMesh = OutStaticMeshComponent->GetStaticMesh();
-				}
-			}
+			SceneComponentUtils.ResolveStaticMesh(true, OutComponent, OutMesh);
 		}
 	}
 
-	OutComponent = OutStaticMeshComponent;
 	if (bTryDefaultMesh && !OutComponent && !OutMesh)
 	{
 		OutMesh = DefaultMesh;
diff --git a/Engine/Plugins/FX/Niagara/Source/Niagara/Private/NiagaraComponent.cpp b/Engine/Plugins/FX/Niagara/Source/Niagara/Private/NiagaraComponent.cpp
--- a/Engine/Plugins/FX/Niagara/Source/Niagara/Private/NiagaraComponent.cpp
+++ b/Engine/Plugins/FX/Niagara/Source/Niagara/Private/NiagaraComponent.cpp
@@ -13,6 +13,7 @@
 #include "NiagaraFunctionLibrary.h"
 #include "NiagaraGpuComputeDispatchInterface.h"
 #include "NiagaraRenderer.h"
+#include "NiagaraSceneComponentUtils.h"
 #include "NiagaraSceneProxy.h"
 #include "NiagaraSimCache.h"
 #include "NiagaraStats.h"
@@ -4087,6 +4088,15 @@ UNiagaraSimCache* UNiagaraComponent::GetSimCache() const
 	return SimCache;
 }
 
+INiagaraSceneComponentUtils* UNiagaraComponent::CreateSceneComponentUtils()
+{
+	if (CreateSceneComponentUtilsFunction != nullptr)
+	{
+		return CreateSceneComponentUtilsFunction();
+	}
+	return new FNiagaraActorSceneComponentUtils(this);
+}
+
 void UNiagaraComponent::SetAllowScalability(bool bAllow)
 {
 	bAllowScalability = bAllow; 
diff --git a/Engine/Plugins/FX/Niagara/Source/Niagara/Private/NiagaraSceneComponentUtils.cpp b/Engine/Plugins/FX/Niagara/Source/Niagara/Private/NiagaraSceneComponentUtils.cpp
new file mode 100644
--- /dev/null
+++ b/Engine/Plugins/FX/Niagara/Source/Niagara/Private/NiagaraSceneComponentUtils.cpp
@@ -0,0 +1,166 @@
+// Copyright Epic Games, Inc. All Rights Reserved.
+
+#include "NiagaraSceneComponentUtils.h"
+#include "NiagaraComponent.h"
+#include "NiagaraSystemInstance.h"
+
+#include "Components/StaticMeshComponent.h"
+#include "Components/InstancedStaticMeshComponent.h"
+#include "Engine/StaticMeshActor.h"
+#include "Engine/StaticMesh.h"
+#include "StaticMeshComponentLODInfo.h"
+
+namespace NiagaraSceneComponentUtilsPrivate
+{
+	UStaticMeshComponent* FindStaticMeshComponent(AActor* Actor, bool bRecurseParents)
+	{
+		if (AStaticMeshActor* MeshActor = Cast<AStaticMeshActor>(Actor))
+		{
+			UStaticMeshComponent* Comp = MeshActor->GetStaticMeshComponent();
+			if (::IsValid(Comp))
+			{
+				return Comp;
+			}
+		}
+
+		// Fall back on any valid component on the actor
+		while (Actor)
+		{
+			for (UActorComponent* ActorComp : Actor->GetComponents())
+			{
+				UStaticMeshComponent* Comp = Cast<UStaticMeshComponent>(ActorComp);
+				if (::IsValid(Comp) && Comp->GetStaticMesh() != nullptr)
+				{
+					return Comp;
+				}
+			}
+
+			if (bRecurseParents)
+			{
+				Actor = Actor->GetParentActor();
+			}
+			else
+			{
+				break;
+			}
+		}
+
+		return nullptr;
+	};
+}
+
+FNiagaraActorSceneComponentUtils::FNiagaraActorSceneComponentUtils(UNiagaraComponent* OwnerComponent)
+	: WeakOwnerComponent(OwnerComponent)
+{
+}
+
+void FNiagaraActorSceneComponentUtils::ResolveStaticMesh(bool bRecurseParents, UObject*& OutComponent, UStaticMesh*& OutStaticMesh) const
+{
+	USceneComponent* OwnerComponent = WeakOwnerComponent.Get();
+	if (OwnerComponent == nullptr)
+	{
+		return;
+	}
+
+	for (USceneComponent* Curr = OwnerComponent; Curr; Curr = Curr->GetAttachParent())
+	{
+		UStaticMeshComponent* ParentComp = Cast<UStaticMeshComponent>(Curr);
+		if (::IsValid(ParentComp))
+		{
+			OutComponent = ParentComp;
+			OutStaticMesh = ParentComp->GetStaticMesh();
+			return;
+		}
+	}
+
+	// Next, try to find one in our outer chain
+	UStaticMeshComponent* OuterComp = OwnerComponent->GetTypedOuter<UStaticMeshComponent>();
+	if (::IsValid(OuterComp))
+	{
+		OutComponent = OuterComp;
+		OutStaticMesh = OuterComp->GetStaticMesh();
+		return;
+	}
+
+	if (AActor* Actor = OwnerComponent->GetAttachmentRootActor())
+	{
+		if (UStaticMeshComponent* StaticMeshComponent = NiagaraSceneComponentUtilsPrivate::FindStaticMeshComponent(Actor, bRecurseParents))
+		{
+			OutComponent = StaticMeshComponent;
+			OutStaticMesh = StaticMeshComponent->GetStaticMesh();
+		}
+	}
+}
+
+void FNiagaraActorSceneComponentUtils::ResolveStaticMesh(UObject* ObjectFrom, bool bRecurseParents, UObject*& OutComponent, UStaticMesh*& OutStaticMesh) const
+{
+	if (AActor* Actor = Cast<AActor>(ObjectFrom))
+	{
+		if (UStaticMeshComponent* StaticMeshComponent = NiagaraSceneComponentUtilsPrivate::FindStaticMeshComponent(Actor, bRecurseParents))
+		{
+			OutComponent = StaticMeshComponent;
+			OutStaticMesh = StaticMeshComponent->GetStaticMesh();
+		}
+	}
+	else if (UStaticMeshComponent* StaticMeshComponent = Cast<UStaticMeshComponent>(ObjectFrom))
+	{
+		OutComponent = StaticMeshComponent;
+		OutStaticMesh = StaticMeshComponent->GetStaticMesh();
+	}
+	else
+	{
+		OutComponent = nullptr;
+		OutStaticMesh = Cast<UStaticMesh>(ObjectFrom);
+	}
+}
+
+bool FNiagaraActorSceneComponentUtils::GetStaticMeshTransforms(UObject* Component, FTransform& OutComponentTransform, TArray<FTransform>& OutInstanceTransforms) const
+{
+	USceneComponent* SceneComponent = Cast<USceneComponent>(Component);
+	if (!SceneComponent)
+	{
+		return false;
+	}
+
+	OutComponentTransform = SceneComponent->GetComponentToWorld();
+	if (UInstancedStaticMeshComponent* ISMComponent = Cast<UInstancedStaticMeshComponent>(SceneComponent))
+	{
+		const int32 NumInstances = ISMComponent->PerInstanceSMData.Num();
+		OutInstanceTransforms.SetNum(NumInstances);
+		for (int32 i = 0; i < NumInstances; ++i)
+		{
+			ISMComponent->GetInstanceTransform(i, OutInstanceTransforms[i], true);
+		}
+	}
+	else
+	{
+		OutInstanceTransforms.Empty();
+	}
+
+	return true;
+}
+
+FColorVertexBuffer* FNiagaraActorSceneComponentUtils::GetStaticMeshOverrideColors(UObject* Component, int32 LODIndex) const
+{
+	if (UStaticMeshComponent* StaticMeshComponent = Cast<UStaticMeshComponent>(Component))
+	{
+		if (StaticMeshComponent->LODData.IsValidIndex(LODIndex))
+		{
+			return StaticMeshComponent->LODData[LODIndex].OverrideVertexColors;
+		}
+	}
+
+	return nullptr;
+}
+
+FPrimitiveComponentId FNiagaraActorSceneComponentUtils::GetPrimitiveSceneId(UObject* Component) const
+{
+	UPrimitiveComponent* PrimitiveComponent = Cast<UPrimitiveComponent>(Component);
+	return PrimitiveComponent ? PrimitiveComponent->GetPrimitiveSceneId() : FPrimitiveComponentId();
+}
+
+FVector FNiagaraActorSceneComponentUtils::GetPhysicsLinearVelocity(UObject* Component) const
+{
+	UPrimitiveComponent* PrimitiveComponent = Cast<UPrimitiveComponent>(Component);
+	return PrimitiveComponent ? PrimitiveComponent->GetPhysicsLinearVelocity() : FVector::ZeroVector;
+}
diff --git a/Engine/Plugins/FX/Niagara/Source/Niagara/Private/NiagaraSystemInstance.cpp b/Engine/Plugins/FX/Niagara/Source/Niagara/Private/NiagaraSystemInstance.cpp
--- a/Engine/Plugins/FX/Niagara/Source/Niagara/Private/NiagaraSystemInstance.cpp
+++ b/Engine/Plugins/FX/Niagara/Source/Niagara/Private/NiagaraSystemInstance.cpp
@@ -10,6 +10,7 @@
 #include "NiagaraEmitterInstance.h"
 #include "NiagaraEmitterInstanceImpl.h"
 #include "NiagaraGPUSystemTick.h"
+#include "NiagaraSceneComponentUtils.h"
 #include "NiagaraStats.h"
 #include "NiagaraParameterCollection.h"
 #include "NiagaraSystemSimulation.h"
@@ -162,8 +163,14 @@ FNiagaraSystemInstance::FNiagaraSystemInstance(UWorld& InWorld, UNiagaraSystem&
 	{
 		TickBehavior = ENiagaraTickBehavior::ForceTickFirst;
 	}
-}
 
+	if (UNiagaraComponent* NiagaraComponent = Cast<UNiagaraComponent>(InAttachComponent))
+	{
+		PRAGMA_DISABLE_DEPRECATION_WARNINGS
+		SceneComponentUtils.Reset(NiagaraComponent->CreateSceneComponentUtils());
+		PRAGMA_ENABLE_DEPRECATION_WARNINGS
+	}
+}
 
 void FNiagaraSystemInstance::SetEmitterEnable(FName EmitterName, bool bNewEnableState)
 {
@@ -2850,6 +2857,12 @@ FNiagaraEmitterInstancePtr FNiagaraSystemInstance::GetSimulationForHandle(const
 	return nullptr;
 }
 
+const INiagaraSceneComponentUtils& FNiagaraSystemInstance::GetSceneComponentUtils(FNiagaraSystemInstance* Instance)
+{
+	static FNiagaraActorSceneComponentUtils GNullUtils(nullptr);
+	return Instance && Instance->SceneComponentUtils.IsValid() ? *Instance->SceneComponentUtils.Get() : GNullUtils;
+}
+
 TConstArrayView<FNiagaraEmitterExecutionIndex> FNiagaraSystemInstance::GetEmitterExecutionOrder() const
 {
 	if (SystemSimulation != nullptr)
diff --git a/Engine/Plugins/FX/Niagara/Source/Niagara/Public/NiagaraComponent.h b/Engine/Plugins/FX/Niagara/Source/Niagara/Public/NiagaraComponent.h
--- a/Engine/Plugins/FX/Niagara/Source/Niagara/Public/NiagaraComponent.h
+++ b/Engine/Plugins/FX/Niagara/Source/Niagara/Public/NiagaraComponent.h
@@ -24,6 +24,7 @@
 class FMeshElementCollector;
 class FNiagaraRenderer;
 class UNiagaraEffectType;
+class INiagaraSceneComponentUtils;
 class UNiagaraSystem;
 class UNiagaraParameterCollection;
 class UNiagaraParameterCollectionInstance;
@@ -785,6 +786,14 @@ public:
 	UFUNCTION(BlueprintCallable, Category = SimCache)
 	void ClearSimCache(bool bResetSystem = false) { SetSimCache(nullptr, bResetSystem); }
 
+	// DO NOT USE THIS IN EXTERNAL CODE as it is subject to change
+	// Methods to handle interop between Actor/Entity/Desc
+	UE_DEPRECATED(5.5, "This function is experimental and subject to change.")
+	INiagaraSceneComponentUtils* CreateSceneComponentUtils();
+	UE_DEPRECATED(5.5, "This function is experimental and subject to change.")
+	void SetCreateSceneComponentUtilsFunction(TFunction<INiagaraSceneComponentUtils*()> Function) { CreateSceneComponentUtilsFunction = Function; }
+	// DO NOT USE THIS IN EXTERNAL CODE as it is subject to change
+
 	/** Set whether this component is allowed to perform scalability checks and potentially be culled etc. Occasionally it is useful to disable this for specific components. E.g. Effects on the local player. */
 	UFUNCTION(BlueprintSetter, Category = Scalability, meta = (Keywords = "LOD scalability"))
 	NIAGARA_API void SetAllowScalability(bool bAllow);
@@ -882,6 +891,8 @@ private:
 
 	float CustomTimeDilation = 1.0f;
 
+	TFunction<INiagaraSceneComponentUtils*()> CreateSceneComponentUtilsFunction;
+
 	UPROPERTY(transient)
 	TObjectPtr<class UNiagaraSimCache> SimCache;
 
diff --git a/Engine/Plugins/FX/Niagara/Source/Niagara/Public/NiagaraSceneComponentUtils.h b/Engine/Plugins/FX/Niagara/Source/Niagara/Public/NiagaraSceneComponentUtils.h
new file mode 100644
--- /dev/null
+++ b/Engine/Plugins/FX/Niagara/Source/Niagara/Public/NiagaraSceneComponentUtils.h
@@ -0,0 +1,59 @@
+// Copyright Epic Games, Inc. All Rights Reserved.
+
+#pragma once
+
+#include "Math/MathFwd.h"
+#include "Components/SceneComponent.h"
+#include "UObject/Object.h"
+#include "PrimitiveComponentId.h"
+
+class FColorVertexBuffer;
+class UNiagaraComponent;
+class UStaticMesh;
+
+// Helper class to abstract how we need to search the scene for components in various data interfaces
+// This is temporary until scene graph is folded into the core engine, or we have more official APIs to abstract Actors/Entity/Desc
+// DO NOT USE THIS IN EXTERNAL CODE as it is subject to change
+class INiagaraSceneComponentUtils
+{
+public:
+	virtual ~INiagaraSceneComponentUtils() = default;
+
+	// Resolve the static mesh from what the interface's owner object
+	virtual void ResolveStaticMesh(bool bRecurseParents, UObject*& OutComponent, UStaticMesh*& OutStaticMesh) const = 0;
+	// Resolve the static mesh from the provided object which could be a static mesh / component / actor / entity / etc
+	virtual void ResolveStaticMesh(UObject* ObjectFrom, bool bRecurseParents, UObject*& OutComponent, UStaticMesh*& OutStaticMesh) const = 0;
+	// Get the component transform and any ISM instance transforms
+	virtual bool GetStaticMeshTransforms(UObject* Component, FTransform& OutComponentTransform, TArray<FTransform>& OutInstanceTransforms) const = 0;
+	// Get the statich mesh override vertex colors, if any
+	virtual FColorVertexBuffer* GetStaticMeshOverrideColors(UObject* Component, int32 LODIndex) const = 0;
+
+	// Get the FPrimitiveComponentId for the provided component
+	virtual FPrimitiveComponentId GetPrimitiveSceneId(UObject* Component) const = 0;
+	// Get the physics lineary velocit for the provided component
+	virtual FVector GetPhysicsLinearVelocity(UObject* Component) const = 0;
+};
+
+// Implementation for AActor & UActorComponents
+class FNiagaraActorSceneComponentUtils final : public INiagaraSceneComponentUtils
+{
+public:
+	FNiagaraActorSceneComponentUtils(UNiagaraComponent* OwnerComponent);
+
+	// Resolve the static mesh from what the interface's owner object
+	virtual void ResolveStaticMesh(bool bRecurseParents, UObject*& OutComponent, UStaticMesh*& OutStaticMesh) const override;
+	// Resolve the static mesh from the provided object which could be a static mesh / component / actor / entity / etc
+	virtual void ResolveStaticMesh(UObject* ObjectFrom, bool bRecurseParents, UObject*& OutComponent, UStaticMesh*& OutStaticMesh) const override;
+	// Get the component transform and any ISM instance transforms
+	virtual bool GetStaticMeshTransforms(UObject* Component, FTransform& OutComponentTransform, TArray<FTransform>& OutInstanceTransforms) const override;
+	// Get the statich mesh override vertex colors, if any
+	virtual FColorVertexBuffer* GetStaticMeshOverrideColors(UObject* Component, int32 LODIndex) const override;
+
+	// Get the FPrimitiveComponentId for the provided component
+	virtual FPrimitiveComponentId GetPrimitiveSceneId(UObject* Component) const override;
+	// Get the physics lineary velocit for the provided component
+	virtual FVector GetPhysicsLinearVelocity(UObject* Component) const override;
+
+private:
+	TWeakObjectPtr<USceneComponent> WeakOwnerComponent;
+};
diff --git a/Engine/Plugins/FX/Niagara/Source/Niagara/Public/NiagaraSystemInstance.h b/Engine/Plugins/FX/Niagara/Source/Niagara/Public/NiagaraSystemInstance.h
--- a/Engine/Plugins/FX/Niagara/Source/Niagara/Public/NiagaraSystemInstance.h
+++ b/Engine/Plugins/FX/Niagara/Source/Niagara/Public/NiagaraSystemInstance.h
@@ -13,6 +13,7 @@
 
 class FNiagaraWorldManager;
 class FNiagaraEmitterInstanceImpl;
+class INiagaraSceneComponentUtils;
 class FNiagaraSystemInstance;
 class FNiagaraSystemSimulation;
 class FNiagaraGpuComputeDispatchInterface;
@@ -263,6 +264,9 @@ public:
 	[[nodiscard]] NIAGARA_API TArrayView<FNiagaraEmitterInstanceRef> GetEmitters() { return Emitters; }
 	[[nodiscard]] NIAGARA_API TConstArrayView<FNiagaraEmitterInstanceRef> GetEmitters() const { return Emitters; }
 
+	UE_DEPRECATED(5.5, "This function is experimental and subject to change.")
+	static const INiagaraSceneComponentUtils& GetSceneComponentUtils(FNiagaraSystemInstance* Instance);
+
 	FORCEINLINE const FBox& GetLocalBounds() const { return LocalBounds;  }
 	FORCEINLINE const FVector3f& GetLWCTile() const { return LWCTile;  }
 	NIAGARA_API TConstArrayView<FNiagaraEmitterExecutionIndex> GetEmitterExecutionOrder() const;
@@ -500,6 +504,7 @@ private:
 	UNiagaraSystem* System = nullptr;
 	FNiagaraUserRedirectionParameterStore* OverrideParameters;
 	TWeakObjectPtr<USceneComponent> AttachComponent;
+	TUniquePtr<INiagaraSceneComponentUtils> SceneComponentUtils;
 
 	FTransform WorldTransform;
 	TOptional<FVector> PreviousLocation;
diff --git a/Engine/Plugins/FX/Niagara/Source/NiagaraEditorWidgets/Private/DetailCustomizations/NiagaraDataInterfaceStaticMeshDetails.cpp b/Engine/Plugins/FX/Niagara/Source/NiagaraEditorWidgets/Private/DetailCustomizations/NiagaraDataInterfaceStaticMeshDetails.cpp
--- a/Engine/Plugins/FX/Niagara/Source/NiagaraEditorWidgets/Private/DetailCustomizations/NiagaraDataInterfaceStaticMeshDetails.cpp
+++ b/Engine/Plugins/FX/Niagara/Source/NiagaraEditorWidgets/Private/DetailCustomizations/NiagaraDataInterfaceStaticMeshDetails.cpp
@@ -41,10 +41,10 @@ void FNiagaraDataInterfaceStaticMeshDetails::CustomizeDetails(IDetailLayoutBuild
 	 MeshInterface->OnChanged().RemoveAll(this);
 	 MeshInterface->OnChanged().AddSP(this, &FNiagaraDataInterfaceStaticMeshDetails::OnInterfaceChanged);
 
-	 USceneComponent* SceneComponent = nullptr;
+	 UObject* Component = nullptr;
 	 if (UNiagaraDataInterfaceStaticMesh* Interface = MeshInterface.Get())
 	 {
-		 MeshObject = Interface->GetStaticMesh(SceneComponent);
+		 MeshObject = Interface->ResolveStaticMesh(Component);
 	 }
 	 if (MeshObject.IsValid())
 	 {
@@ -91,10 +91,10 @@ void FNiagaraDataInterfaceStaticMeshDetails::CustomizeDetails(IDetailLayoutBuild
 		 MeshObject = nullptr;
 	 }
 
-	 USceneComponent* SceneComponent = nullptr;
+	 UObject* Component = nullptr;
 	 if (UNiagaraDataInterfaceStaticMesh* Interface = MeshInterface.Get())
 	 {
-		 MeshObject = Interface->GetStaticMesh(SceneComponent);
+		 MeshObject = Interface->ResolveStaticMesh(Component);
 	 }
 
 	 if (MeshObject.IsValid())
