- Fix static analysis issue

#rnx

diff --git a/Engine/Plugins/FX/Niagara/Source/Niagara/Private/DataInterface/NiagaraDataInterfaceStaticMesh.cpp b/Engine/Plugins/FX/Niagara/Source/Niagara/Private/DataInterface/NiagaraDataInterfaceStaticMesh.cpp
--- a/Engine/Plugins/FX/Niagara/Source/Niagara/Private/DataInterface/NiagaraDataInterfaceStaticMesh.cpp
+++ b/Engine/Plugins/FX/Niagara/Source/Niagara/Private/DataInterface/NiagaraDataInterfaceStaticMesh.cpp
@@ -1141,7 +1141,12 @@ namespace NDIStaticMeshLocal
 
 			// Check to see if we would resolve to a different mesh / component
 			UObject* ResolvedComponent = nullptr;
-			UStaticMesh* ResolvedMesh = Interface->ResolveStaticMesh(ResolvedComponent, Instance, MeshParameterBinding.GetValue());
+			UStaticMesh* ResolvedMesh = nullptr;
+			if (Interface != nullptr)
+			{
+				ResolvedMesh = Interface->ResolveStaticMesh(ResolvedComponent, Instance, MeshParameterBinding.GetValue());
+			}
+
 			if (Component != ResolvedComponent || Mesh != ResolvedMesh)
 			{
 				return true;
