Fixed 'Open Asset' on stack entries that point to dynamic inputs opening up the asset the input belongs to instead of the dynamic input itself

diff --git a/Engine/Plugins/FX/Niagara/Source/NiagaraEditor/Private/ViewModels/Stack/NiagaraStackFunctionInput.cpp b/Engine/Plugins/FX/Niagara/Source/NiagaraEditor/Private/ViewModels/Stack/NiagaraStackFunctionInput.cpp
--- a/Engine/Plugins/FX/Niagara/Source/NiagaraEditor/Private/ViewModels/Stack/NiagaraStackFunctionInput.cpp
+++ b/Engine/Plugins/FX/Niagara/Source/NiagaraEditor/Private/ViewModels/Stack/NiagaraStackFunctionInput.cpp
@@ -267,7 +267,14 @@ bool UNiagaraStackFunctionInput::GetIsEnabled() const
 
 UObject* UNiagaraStackFunctionInput::GetExternalAsset() const
 {
-	if (OwningFunctionCallNode.IsValid() && OwningFunctionCallNode->FunctionScript != nullptr && OwningFunctionCallNode->FunctionScript->IsAsset())
+	if(GetValueMode() == EValueMode::Dynamic)
+	{
+		if(InputValues.DynamicNode.IsValid() && InputValues.DynamicNode->FunctionScript != nullptr && InputValues.DynamicNode->FunctionScript->IsAsset())
+		{
+			return InputValues.DynamicNode->FunctionScript;
+		}
+	}
+	else if (OwningFunctionCallNode.IsValid() && OwningFunctionCallNode->FunctionScript != nullptr && OwningFunctionCallNode->FunctionScript->IsAsset())
 	{
 		return OwningFunctionCallNode->FunctionScript;
 	}
