- Move SystemStateData to an instanced struct
- Fixes issues with internal includes, and opens up more possibilities in the future

#rnx
#tests Test state is still correctly created, Ran ET

diff --git a/Engine/Plugins/FX/Niagara/Source/Niagara/Classes/NiagaraSystem.h b/Engine/Plugins/FX/Niagara/Source/Niagara/Classes/NiagaraSystem.h
--- a/Engine/Plugins/FX/Niagara/Source/Niagara/Classes/NiagaraSystem.h
+++ b/Engine/Plugins/FX/Niagara/Source/Niagara/Classes/NiagaraSystem.h
@@ -12,9 +12,9 @@
 #include "NiagaraMessageStore.h"
 #include "NiagaraParameterCollection.h"
 #include "NiagaraParameterDefinitionsSubscriber.h"
-#include "NiagaraSystemEmitterState.h"
 #include "NiagaraUserRedirectionParameterStore.h"
 #include "Particles/ParticleSystem.h"
+#include "StructUtils/InstancedStruct.h"
 #include "UObject/Object.h"
 #include "UObject/ObjectMacros.h"
 #include "Particles/FXBudget.h"
@@ -24,6 +24,7 @@ class FNiagaraActiveCompilation;
 class UNiagaraEffectType;
 enum class ENiagaraCullProxyMode : uint32;
 enum class ERendererStencilMask : uint8;
+struct FNiagaraSystemStateData;
 struct FNiagaraSystemScalabilityOverride;
 struct FNiagaraSystemScalabilitySettings;
 struct FNiagaraSystemStaticBuffers;
@@ -803,7 +804,7 @@ public:
 #endif
 
 	/** Access the code system state data. */
-	const FNiagaraSystemStateData& GetSystemStateData() const { return SystemStateData; }
+	const FNiagaraSystemStateData& GetSystemStateData() const;
 
 	/** Used for debug HUD / viewport to convey what mode we are running in */
 	NIAGARA_API const TCHAR* GetSystemStateModeString() const;
@@ -1082,7 +1083,7 @@ protected:
 	TUniquePtr<FNiagaraSystemStaticBuffers, FStaticBuffersDeletor> StaticBuffers;
 
 	UPROPERTY()
-	FNiagaraSystemStateData SystemStateData;
+	FInstancedStruct SystemStateDataStruct;
 
 	FRenderCommandFence WaitRenderCommandsFence;
 };
diff --git a/Engine/Plugins/FX/Niagara/Source/Niagara/Private/NiagaraSystem.cpp b/Engine/Plugins/FX/Niagara/Source/Niagara/Private/NiagaraSystem.cpp
--- a/Engine/Plugins/FX/Niagara/Source/Niagara/Private/NiagaraSystem.cpp
+++ b/Engine/Plugins/FX/Niagara/Source/Niagara/Private/NiagaraSystem.cpp
@@ -2857,8 +2857,16 @@ const UNiagaraScript* UNiagaraSystem::GetSystemUpdateScript() const
 	return SystemUpdateScript;
 }
 
+const FNiagaraSystemStateData& UNiagaraSystem::GetSystemStateData() const
+{
+	static const FNiagaraSystemStateData DefaultSystemStateData;
+	const FNiagaraSystemStateData* SystemStateData = SystemStateDataStruct.GetPtr<const FNiagaraSystemStateData>();
+	return SystemStateData ? *SystemStateData : DefaultSystemStateData;
+}
+
 const TCHAR* UNiagaraSystem::GetSystemStateModeString() const
 {
+	const FNiagaraSystemStateData& SystemStateData = GetSystemStateData();
 	if ( SystemStateData.bRunUpdateScript == false )
 	{
 		return SystemStateData.bRunSpawnScript ? TEXT("[Fast-U]") : TEXT("[Fast-SU]");
@@ -3985,15 +3993,16 @@ void UNiagaraSystem::GatherStaticVariables(TArray<FNiagaraVariable>& OutVars, TA
 void UNiagaraSystem::ResolveRequiresScripts()
 {
 #if WITH_EDITORONLY_DATA
+	SystemStateDataStruct.Reset();
 	if (bAllowSystemStateFastPath)
 	{
 		INiagaraModule& NiagaraModule = FModuleManager::GetModuleChecked<INiagaraModule>("Niagara");
 		const INiagaraEditorOnlyDataUtilities& EditorOnlyDataUtilities = NiagaraModule.GetEditorOnlyDataUtilities();
-		SystemStateData = EditorOnlyDataUtilities.GetSystemStateData(*this);
-	}
-	else
-	{
-		SystemStateData = FNiagaraSystemStateData();
+		TOptional<FNiagaraSystemStateData> SystemStateData = EditorOnlyDataUtilities.GetSystemStateData(*this);
+		if (SystemStateData.IsSet())
+		{
+			SystemStateDataStruct = FInstancedStruct::Make(SystemStateData.GetValue());
+		}
 	}
 #endif //WITH_EDITORONLY_DATA
 }
diff --git a/Engine/Plugins/FX/Niagara/Source/Niagara/Public/INiagaraEditorOnlyDataUtlities.h b/Engine/Plugins/FX/Niagara/Source/Niagara/Public/INiagaraEditorOnlyDataUtlities.h
--- a/Engine/Plugins/FX/Niagara/Source/Niagara/Public/INiagaraEditorOnlyDataUtlities.h
+++ b/Engine/Plugins/FX/Niagara/Source/Niagara/Public/INiagaraEditorOnlyDataUtlities.h
@@ -32,5 +32,5 @@ public:
 
 	virtual UNiagaraDataInterface* GetResolvedRuntimeInstanceForEditorDataInterfaceInstance(const UNiagaraSystem& OwningSystem, UNiagaraDataInterface& EditorDataInterfaceInstance) const = 0;
 
-	virtual FNiagaraSystemStateData GetSystemStateData(const UNiagaraSystem& System) const = 0;
+	virtual TOptional<FNiagaraSystemStateData> GetSystemStateData(const UNiagaraSystem& System) const = 0;
 };
diff --git a/Engine/Plugins/FX/Niagara/Source/NiagaraEditor/Private/NiagaraEditorModule.cpp b/Engine/Plugins/FX/Niagara/Source/NiagaraEditor/Private/NiagaraEditorModule.cpp
--- a/Engine/Plugins/FX/Niagara/Source/NiagaraEditor/Private/NiagaraEditorModule.cpp
+++ b/Engine/Plugins/FX/Niagara/Source/NiagaraEditor/Private/NiagaraEditorModule.cpp
@@ -48,6 +48,7 @@
 #include "PropertyEditorModule.h"
 #include "NiagaraSettings.h"
 #include "NiagaraShaderModule.h"
+#include "NiagaraSystemEmitterState.h"
 
 #include "NiagaraDataInterfaceArray.h"
 #include "NiagaraDataInterfaceCurve.h"
@@ -313,7 +314,7 @@ class FNiagaraEditorOnlyDataUtilities : public INiagaraEditorOnlyDataUtilities
 		return FNiagaraEditorUtilities::GetResolvedRuntimeInstanceForEditorDataInterfaceInstance(OwningSystem, EditorDataInterfaceInstance);
 	}
 
-	virtual FNiagaraSystemStateData GetSystemStateData(const UNiagaraSystem& System) const override
+	virtual TOptional<FNiagaraSystemStateData> GetSystemStateData(const UNiagaraSystem& System) const override
 	{
 		// All emitters must be stateless currently
 		// We can perhaps look at this again, but we always write Emitter.RandomSeed currently even with an empty script
@@ -321,7 +322,7 @@ class FNiagaraEditorOnlyDataUtilities : public INiagaraEditorOnlyDataUtilities
 		{
 			if ( EmitterHandle.GetIsEnabled() && EmitterHandle.GetEmitterMode() != ENiagaraEmitterMode::Stateless )
 			{
-				return FNiagaraSystemStateData();
+				return TOptional<FNiagaraSystemStateData>();
 			}
 		}
 
@@ -330,7 +331,7 @@ class FNiagaraEditorOnlyDataUtilities : public INiagaraEditorOnlyDataUtilities
 		const UNiagaraScriptSource* ScriptSource = Script ? Cast<UNiagaraScriptSource>(Script->GetLatestSource()) : nullptr;
 		if (!ScriptSource || !ScriptSource->NodeGraph)
 		{
-			return FNiagaraSystemStateData();
+			return TOptional<FNiagaraSystemStateData>();
 		}
 
 		// Look for update script nodes to see if it's possible to avoid running the update script
