- Fix no unity IWYU

diff --git a/Engine/Plugins/FX/Niagara/Source/Niagara/Private/NiagaraSystemSimulation.cpp b/Engine/Plugins/FX/Niagara/Source/Niagara/Private/NiagaraSystemSimulation.cpp
--- a/Engine/Plugins/FX/Niagara/Source/Niagara/Private/NiagaraSystemSimulation.cpp
+++ b/Engine/Plugins/FX/Niagara/Source/Niagara/Private/NiagaraSystemSimulation.cpp
@@ -7,6 +7,7 @@
 #include "NiagaraTypes.h"
 #include "NiagaraParameterCollection.h"
 #include "NiagaraSimulationTaskPriority.h"
+#include "NiagaraSystemEmitterState.h"
 #include "NiagaraSystemInstance.h"
 #include "NiagaraSystemGpuComputeProxy.h"
 #include "NiagaraConstants.h"
