- Move NiagaraDataChannelCommon from internal to public

#rnx

diff --git a/Engine/Plugins/FX/Niagara/Source/Niagara/Internal/NiagaraDataChannelCommon.h b/Engine/Plugins/FX/Niagara/Source/Niagara/Internal/NiagaraDataChannelCommon.h
deleted file mode 100644
diff --git a/Engine/Plugins/FX/Niagara/Source/Niagara/Public/NiagaraDataChannelCommon.h b/Engine/Plugins/FX/Niagara/Source/Niagara/Public/NiagaraDataChannelCommon.h
new file mode 100644
--- /dev/null
+++ b/Engine/Plugins/FX/Niagara/Source/Niagara/Public/NiagaraDataChannelCommon.h
@@ -0,0 +1,158 @@
+// Copyright Epic Games, Inc. All Rights Reserved.
+
+#pragma once
+
+#include "CoreMinimal.h"
+#include "UObject/ObjectMacros.h"
+#include "NiagaraDataSet.h"
+#include "NiagaraDataChannelPublic.h"
+
+class UWorld;
+class UActorComponentc;
+class FNiagaraWorldManager;
+class FNiagaraDataBuffer;
+class FNiagaraSystemInstance;
+class UNiagaraDataChannel;
+class UNiagaraDataChannelHandler;
+class UNiagaraDataInterfaceDataChannelWrite;
+class UNiagaraDataInterfaceDataChannelRead;
+
+using FNiagaraDataChannelDataProxyPtr = TSharedPtr<struct FNiagaraDataChannelDataProxy>;
+
+/** A request to publish data into a Niagara Data Channel.  */
+struct FNiagaraDataChannelPublishRequest
+{
+	/** The buffer containing the data to be published. This can come from a data channel DI or can be the direct contents on a Niagara simulation. */
+	FNiagaraDataBufferRef Data;
+
+	/** Game level data if this request comes from the game code. */
+	TSharedPtr<FNiagaraDataChannelGameData> GameData = nullptr;
+
+	/** If true, data in this request will be made visible to BP and C++ game code.*/
+	bool bVisibleToGame = false;
+
+	/** If true, data in this request will be made visible to Niagara CPU simulations. */
+	bool bVisibleToCPUSims = false;
+
+	/** If true, data in this request will be made visible to Niagara GPU simulations. */
+	bool bVisibleToGPUSims = false;
+
+	/** 
+	LWC Tile for the originator system of this request.
+	Allows us to convert from the Niagara Simulation space into LWC coordinates.
+	*/
+	FVector3f LwcTile = FVector3f::ZeroVector;
+
+#if WITH_NIAGARA_DEBUGGER
+	/** Instigator of this write, used for debug tracking */
+	FString DebugSource;
+#endif
+
+	FNiagaraDataChannelPublishRequest() = default;
+	explicit FNiagaraDataChannelPublishRequest(FNiagaraDataBufferRef InData)
+		: Data(InData)
+	{
+	}
+
+	explicit FNiagaraDataChannelPublishRequest(FNiagaraDataBufferRef InData, bool bInVisibleToGame, bool bInVisibleToCPUSims, bool bInVisibleToGPUSims, FVector3f InLwcTile)
+	: Data(InData), bVisibleToGame(bInVisibleToGame), bVisibleToCPUSims(bInVisibleToCPUSims), bVisibleToGPUSims(bInVisibleToGPUSims), LwcTile(InLwcTile)
+	{
+	}
+};
+
+
+/**
+Underlying storage class for data channel data.
+Some data channels will have many of these and can distribute them as needed to different accessing systems.
+For example, some data channel handlers may subdivide the scene such that distant systems are not interacting.
+In this case, each subdivision would have it's own FNiagaraDataChannelData and distribute these to the relevant NiagaraSystems.
+*/
+struct FNiagaraDataChannelData final : public TSharedFromThis<FNiagaraDataChannelData, ESPMode::ThreadSafe>
+{
+	UE_NONCOPYABLE(FNiagaraDataChannelData)
+	NIAGARA_API explicit FNiagaraDataChannelData();
+	NIAGARA_API ~FNiagaraDataChannelData();
+
+	NIAGARA_API void Init(UNiagaraDataChannelHandler* Owner);
+	NIAGARA_API void Reset();
+
+	NIAGARA_API void BeginFrame(UNiagaraDataChannelHandler* Owner);
+	NIAGARA_API void EndFrame(UNiagaraDataChannelHandler* Owner);
+	NIAGARA_API int32 ConsumePublishRequests(UNiagaraDataChannelHandler* Owner, const ETickingGroup& TickGroup);
+
+	NIAGARA_API FNiagaraDataChannelGameData* GetGameData();
+	NIAGARA_API FNiagaraDataBufferRef GetCPUData(bool bPreviousFrame);
+	FNiagaraDataChannelDataProxyPtr GetRTProxy(){ return RTProxy; }
+	
+	/** Adds a request to publish some data into the channel on the next tick. */
+	NIAGARA_API void Publish(const FNiagaraDataChannelPublishRequest& Request);
+
+	NIAGARA_API void PublishFromGPU(const FNiagaraDataChannelPublishRequest& Request);
+
+	NIAGARA_API const FNiagaraDataSetCompiledData& GetCompiledData(ENiagaraSimTarget SimTarget);
+
+	void SetLwcTile(FVector3f InLwcTile){ LwcTile = InLwcTile; }
+	FVector3f GetLwcTile()const { return LwcTile; }
+
+	//This will get a buffer from the CPU dataset intended to be written to on the CPU.
+	FNiagaraDataBuffer* GetBufferForCPUWrite();
+
+	void DestroyRenderThreadProxy(FNiagaraGpuComputeDispatchInterface* ComputeDispatchInterface);
+
+	void RegisterGPUSpawningReader() { ++NumGPUSpawningReaders; }
+	void UnregisterGPUSpawningReader() { --NumGPUSpawningReaders; }
+	int32 NumRegisteredGPUSpawningReaders()const{ return NumGPUSpawningReaders; }
+
+	//Returns if this data is still valid. This can return false in cases where the owning data channel has been modified for example.
+	bool IsLayoutValid(UNiagaraDataChannelHandler* Owner)const;
+
+	//Return true if data has been written to this NDC data for the current frame.
+	bool HasData()const;
+
+private:
+
+	void CreateRenderThreadProxy(UNiagaraDataChannelHandler* Owner);
+
+	/** DataChannel data accessible from Game/BP. AoS Layout. LWC types. */
+	FNiagaraDataChannelGameDataPtr GameData;
+
+	/** DataChannel data accessible to Niagara CPU sims. SoA layout. Non LWC types. */
+	FNiagaraDataSet* CPUSimData = nullptr;
+
+	// Cached off buffer with the previous frame's CPU Sim accessible data.
+	// Some systems can choose to read this to avoid any current frame tick ordering issues.
+	FNiagaraDataBufferRef PrevCPUSimData = nullptr;
+
+	/** Dataset we use for staging game data for the consumption by RT/GPU sims. */
+	FNiagaraDataSet* GameDataStaging = nullptr;
+
+	/** Data buffers we'll be passing to the RT proxy for uploading to the GPU */
+	TArray<FNiagaraDataChannelPublishRequest> PublishRequestsForGPU;
+
+	/** Render thread proxy for this data. Owns all RT side data meant for GPU simulations. */
+	FNiagaraDataChannelDataProxyPtr RTProxy;
+
+	/** Pending requests to publish data into this data channel. These requests are consumed at tick tick group. */
+	TArray<FNiagaraDataChannelPublishRequest> PublishRequests;
+
+	/** Pending requests to publish data into this data channel from the GPU. To alleviate data race behavior with data coming back from the GPU, we always consume GPU requests at the start of the frame only. */
+	TArray<FNiagaraDataChannelPublishRequest> PublishRequestsFromGPU;
+
+	/** The world we were initialized with, used to get the compute interface. */
+	TWeakObjectPtr<UWorld> WeakOwnerWorld;
+
+	FVector3f LwcTile = FVector3f::ZeroVector;
+
+	/** Critical section protecting shared state for multiple writers publishing from different threads. */
+	FCriticalSection PublishCritSec;
+
+	//Keep reference to the layout this data was built with.
+	FNiagaraDataChannelLayoutInfoPtr LayoutInfo;
+
+	/** 
+	Track number of explicitly registered readers that spawn GPU particles from this data.
+	If we're spawning GPU particles using the CPU data (Spawn Conditional etc) then we have to send all CPU data to the GPU every frame.
+	Can possibly extend this to be a more automatic, registration based approach to shipping NDC data around rather than exmplicit flags on write.
+	*/
+	std::atomic<int32> NumGPUSpawningReaders;
+};
