Fix for compiling nodes without any valid output pins (likely only orphaned/add pins are left).  Avoids an array out of bounds exception

diff --git a/Engine/Plugins/FX/Niagara/Source/NiagaraEditor/Private/NiagaraHlslTranslator.cpp b/Engine/Plugins/FX/Niagara/Source/NiagaraEditor/Private/NiagaraHlslTranslator.cpp
--- a/Engine/Plugins/FX/Niagara/Source/NiagaraEditor/Private/NiagaraHlslTranslator.cpp
+++ b/Engine/Plugins/FX/Niagara/Source/NiagaraEditor/Private/NiagaraHlslTranslator.cpp
@@ -10126,7 +10126,8 @@ int32 TNiagaraHlslTranslator<GraphBridge>::CompileOutputPin(const FPin* InPin)
 	// if the node is disabled (which commonly happens when a module is disabled in the stack), we skip it and follow the execution pin to the next node
 	if (!GraphBridge::NodeIsEnabled(Node))
 	{
-		const FPin* OutputPin = GraphBridge::GetOutputPins(Node)[0];
+		TArray<const FOutputPin*> Outputs = GraphBridge::GetOutputPins(Node);
+		const FPin* OutputPin = Outputs.IsEmpty() ? nullptr : Outputs[0];
 		if ((OutputPin == Pin) && GraphBridge::GetPinType(Pin, ENiagaraStructConversion::UserFacing) == FNiagaraTypeDefinition::GetParameterMapDef())
 		{
 			TArray<const FInputPin*> InputPins = GraphBridge::GetInputPins(Node);
