Fix for compile failures on Select nodes if wildcard output pins are left (even disconnected)
#rnx
#tests Make a system using a scratch module with a Select node with adding a pin (rather than replacing the auto generated wildcard pin)

diff --git a/Engine/Plugins/FX/Niagara/Source/NiagaraEditor/Private/NiagaraGraphDigest.cpp b/Engine/Plugins/FX/Niagara/Source/NiagaraEditor/Private/NiagaraGraphDigest.cpp
--- a/Engine/Plugins/FX/Niagara/Source/NiagaraEditor/Private/NiagaraGraphDigest.cpp
+++ b/Engine/Plugins/FX/Niagara/Source/NiagaraEditor/Private/NiagaraGraphDigest.cpp
@@ -1561,6 +1561,11 @@ UEdGraphPin* FNiagaraCompilationPin::GetSourcePin() const
 	return OwningNode->SourceNode->Pins[SourcePinIndex];
 }
 
+bool FNiagaraCompilationPin::IsWildcard() const
+{
+	return PinType.PinCategory == UEdGraphSchema_Niagara::PinCategoryType && PinType.PinSubCategoryObject == FNiagaraTypeDefinition::GetWildcardStruct();
+}
+
 FNiagaraCompilationInputPin::FNiagaraCompilationInputPin(const UEdGraphPin* InPin)
 	: FNiagaraCompilationPin(InPin)
 {
@@ -4638,25 +4643,47 @@ void FNiagaraCompilationNodeSelect::Compile(FTranslator* Translator, TArray<int3
 	const FNiagaraCompilationInputPin& SelectorPin = InputPins[SelectorPinIndex];
 	int32 Selection = Translator->CompileInputPin(&SelectorPin);
 
+	const int32 OutputVarCount = OutputVars.Num();
+	TArray<int32> ConnectedOutputVarIndices;
+	TArray<FNiagaraVariable> ConnectedOutputVars;
+	ConnectedOutputVarIndices.Reserve(OutputVarCount);
+	ConnectedOutputVars.Reserve(OutputVarCount);
+	for (int32 OutputVarIndex = 0; OutputVarIndex < OutputVarCount; ++OutputVarIndex)
+	{
+		const FNiagaraVariable& OutputVar = OutputVars[OutputVarIndex];
+		const FNiagaraCompilationOutputPin* OutputPin = FindOutputPin(OutputVar);
+		if (OutputPin && !OutputPin->LinkedTo.IsEmpty())
+		{
+			ConnectedOutputVarIndices.Add(OutputVarIndex);
+			ConnectedOutputVars.Add(OutputVar);
+		}
+	}
+
 	// a map from selector value to compiled option pins (i.e.: for selector value 0 all pins that should be case "if 0" get their compiled index added under key 0)
 	TMap<int32, TArray<int32>> OptionValues;
 	OptionValues.Reserve(NumOptionsPerVariable);
 
-	const int32 OutputCount = OutputVars.Num();
-	for (int32 OutputIndex = 0; OutputIndex < OutputCount; ++OutputIndex)
+	for (const FNiagaraVariable& OutputVar : ConnectedOutputVars)
 	{
-		const FNiagaraCompilationOutputPin* PinByVariable = FindOutputPin(OutputVars[OutputIndex]);
+		const FNiagaraCompilationOutputPin* PinByVariable = FindOutputPin(OutputVar);
 		if (ensure(PinByVariable))
 		{
-			if (OutputVars[OutputIndex].GetType() != PinByVariable->Variable.GetType())
+			if (OutputVar.GetType() != PinByVariable->Variable.GetType())
 			{
 				Translator->Error(FText::Format(LOCTEXT("PinTypeOutputVarTypeMismatch", "Internal output variable type {0} does not match pin type {1}. Please refresh node."),
-					FText::FromString(OutputVars[OutputIndex].GetType().GetName()), FText::FromString(PinByVariable->Variable.GetType().GetName())),
+					FText::FromString(OutputVar.GetType().GetName()), FText::FromString(PinByVariable->Variable.GetType().GetName())),
+					this, PinByVariable);
+			}
+			else if (PinByVariable->IsWildcard())
+			{
+				Translator->Error(FText::Format(LOCTEXT("PinTypeOutputWildcard", "Wildcard types are intended to be placeholders and replaced by real types.  Please resolve type for {0}."),
+					FText::FromString(OutputVar.GetName().ToString())),
 					this, PinByVariable);
 			}
 		}
 	}
 
+	int32 HandledInputCount = 0;
 	const int32 InputPinCount = InputPins.Num();
 	for (int32 InputPinIt = 0; InputPinIt < InputPinCount; ++InputPinIt)
 	{
@@ -4668,12 +4695,31 @@ void FNiagaraCompilationNodeSelect::Compile(FTranslator* Translator, TArray<int3
 			continue;
 		}
 
+		const int32 OutputVarIndex = HandledInputCount % OutputVarCount;
+		const int32 SelectorIndex = HandledInputCount / OutputVarCount;
+
+		++HandledInputCount;
+		
+		// also skip if the output isn't connected
+		if (!ConnectedOutputVarIndices.Contains(OutputVarIndex))
+		{
+			continue;
+		}
+
 		int32 CompiledInput = Translator->CompileInputPin(&InputPin);
-		const int32 SelectorValue = SelectorValues[InputPinIt / OutputCount];
+		const int32 SelectorValue = SelectorValues[SelectorIndex];
 		OptionValues.FindOrAdd(SelectorValue).Add(CompiledInput);
 	}
 
-	Translator->Select(this, Selection, OutputVars, OptionValues, Outputs);
+	TArray<int32> SelectOutputs;
+	Translator->Select(this, Selection, ConnectedOutputVars, OptionValues, SelectOutputs);
+
+	// remap the compiled outputs based on list of connected outputs
+	Outputs.Init(INDEX_NONE, OutputVarCount);
+	for (int32 ConnectedVarIndex = 0; ConnectedVarIndex < ConnectedOutputVars.Num(); ++ConnectedVarIndex)
+	{
+		Outputs[ConnectedOutputVarIndices[ConnectedVarIndex]] = SelectOutputs[ConnectedVarIndex];
+	}
 }
 
 void FNiagaraCompilationNodeSelect::BuildParameterMapHistory(FParameterMapHistoryBuilder& Builder, bool bRecursive, bool bFilterForCompilation) const
diff --git a/Engine/Plugins/FX/Niagara/Source/NiagaraEditor/Private/NiagaraGraphDigest.h b/Engine/Plugins/FX/Niagara/Source/NiagaraEditor/Private/NiagaraGraphDigest.h
--- a/Engine/Plugins/FX/Niagara/Source/NiagaraEditor/Private/NiagaraGraphDigest.h
+++ b/Engine/Plugins/FX/Niagara/Source/NiagaraEditor/Private/NiagaraGraphDigest.h
@@ -66,6 +66,8 @@ public:
 
 	UEdGraphPin* GetSourcePin() const;
 
+	bool IsWildcard() const;
+
 	FNiagaraVariable Variable;
 	int32 SourcePinIndex = INDEX_NONE;
 	FName PinName;
diff --git a/Engine/Plugins/FX/Niagara/Source/NiagaraEditor/Private/NiagaraNodeSelect.cpp b/Engine/Plugins/FX/Niagara/Source/NiagaraEditor/Private/NiagaraNodeSelect.cpp
--- a/Engine/Plugins/FX/Niagara/Source/NiagaraEditor/Private/NiagaraNodeSelect.cpp
+++ b/Engine/Plugins/FX/Niagara/Source/NiagaraEditor/Private/NiagaraNodeSelect.cpp
@@ -311,8 +311,24 @@ void UNiagaraNodeSelect::Compile(FTranslator* Translator, TArray<int32>& Outputs
 	{
 		Translator->Warning(LOCTEXT("SelectNodePinSelectorTypeInvalid", "Select node selector pin should have a valid type."), this, nullptr);
 	}
-	
-	for(const FNiagaraVariable& Variable : OutputVars)
+
+	const int32 OutputVarCount = OutputVars.Num();
+	TArray<int32> ConnectedOutputVarIndices;
+	TArray<FNiagaraVariable> ConnectedOutputVars;
+	ConnectedOutputVarIndices.Reserve(OutputVarCount);
+	ConnectedOutputVars.Reserve(OutputVarCount);
+	for (int32 OutputVarIndex = 0; OutputVarIndex < OutputVarCount; ++OutputVarIndex)
+	{
+		const FNiagaraVariable& OutputVar = OutputVars[OutputVarIndex];
+		const UEdGraphPin* OutputPin = GetOutputPin(OutputVar);
+		if (OutputPin && !OutputPin->LinkedTo.IsEmpty())
+		{
+			ConnectedOutputVarIndices.Add(OutputVarIndex);
+			ConnectedOutputVars.Add(OutputVar);
+		}
+	}
+
+	for (const FNiagaraVariable& Variable : ConnectedOutputVars)
 	{
 		if(!Variable.GetType().IsValid())
 		{
@@ -341,27 +357,44 @@ void UNiagaraNodeSelect::Compile(FTranslator* Translator, TArray<int32>& Outputs
 		OptionValues.Add(SelectorValues[Index]);
 	}
 
-	for (int32 OutputIndex = 0; OutputIndex < OutputVars.Num(); OutputIndex++)
+	for (const FNiagaraVariable& OutputVar : ConnectedOutputVars)
 	{
-		if(OutputVars[OutputIndex].GetType() != UEdGraphSchema_Niagara::PinToTypeDefinition(GetOutputPin(OutputVars[OutputIndex])))
+		const UEdGraphPin* OutputPin = GetOutputPin(OutputVar);
+		const FNiagaraTypeDefinition OutputPinType = UEdGraphSchema_Niagara::PinToTypeDefinition(OutputPin);
+
+		if(OutputVar.GetType() != OutputPinType)
 		{
 			Translator->Error(FText::Format(LOCTEXT("PinTypeOutputVarTypeMismatch", "Internal output variable type {0} does not match pin type {1}. Please refresh node."),
-				FText::FromString(OutputVars[OutputIndex].GetType().GetName()), FText::FromString(UEdGraphSchema_Niagara::PinToTypeDefinition(GetOutputPin(OutputVars[OutputIndex])).GetName())),
-				this, GetOutputPin(OutputVars[OutputIndex]));
+				FText::FromString(OutputVar.GetType().GetName()), FText::FromString(OutputPinType.GetName())),
+				this, OutputPin);
+		}
+		else if (UEdGraphSchema_Niagara::IsPinWildcard(OutputPin))
+		{
+			Translator->Error(FText::Format(LOCTEXT("PinTypeOutputWildcard", "Wildcard types are intended to be placeholders and replaced by real types.  Please resolve type for {0}."),
+				FText::FromString(OutputVar.GetName().ToString())),
+				this, OutputPin);
 		}
 	}
 	
 	for (int32 SelectorValueIndex = 0; SelectorValueIndex < NumOptionsPerVariable; SelectorValueIndex++)
 	{
-		for (int32 OutputIndex = 0; OutputIndex < OutputVars.Num(); OutputIndex++)
+		for (int32 OutputVarIndex : ConnectedOutputVarIndices)
 		{
-			TArray<UEdGraphPin*> OptionPins = GetOptionPins(OutputIndex);
+			TArray<UEdGraphPin*> OptionPins = GetOptionPins(OutputVarIndex);
 			int32 CodeChunkIndex = Translator->CompileInputPin(OptionPins[SelectorValueIndex]);
 			OptionValues[SelectorValues[SelectorValueIndex]].Add(CodeChunkIndex);
 		}
 	}	
-	
-	Translator->Select(this, Selection, OutputVars, OptionValues, Outputs);
+
+	TArray<int32> SelectOutputs;
+	Translator->Select(this, Selection, ConnectedOutputVars, OptionValues, SelectOutputs);
+
+	// remap the compiled outputs based on list of connected outputs
+	Outputs.Init(INDEX_NONE, OutputVarCount);
+	for (int32 ConnectedVarIndex = 0; ConnectedVarIndex < ConnectedOutputVars.Num(); ++ConnectedVarIndex)
+	{
+		Outputs[ConnectedOutputVarIndices[ConnectedVarIndex]] = SelectOutputs[ConnectedVarIndex];
+	}
 }
 
 bool UNiagaraNodeSelect::AllowExternalPinTypeChanges(const UEdGraphPin* InGraphPin) const
@@ -695,6 +728,7 @@ TArray<int32> UNiagaraNodeSelect::GetOptionValues() const
 	else if(FNiagaraTypeDefinition::GetIntDef().IsSameBaseDefinition(SelectorPinType))
 	{
 		int32 NewOptionsCount = FMath::Max(2, NumOptionsPerVariable);
+		SelectorValues.Reserve(NewOptionsCount);
 		for(int32 Index = 0; Index < NewOptionsCount; Index++)
 		{
 			SelectorValues.Add(Index);
@@ -703,11 +737,12 @@ TArray<int32> UNiagaraNodeSelect::GetOptionValues() const
 	else if(SelectorPinType.IsEnum() && SelectorPinType.IsValid() && SelectorPinType.GetEnum())
 	{
 		UEnum* Enum = SelectorPinType.GetEnum();
-		const int32 EnumEntryCount = Enum->NumEnums();
-		if (EnumEntryCount > 0)
+		const int32 TotalEnumEntryCount = Enum->NumEnums() - 1;
+		if (TotalEnumEntryCount > 0)
 		{
+			SelectorValues.Reserve(TotalEnumEntryCount);
 			int32 ValidEnumEntryCount = 0;
-			for (int32 EnumIndex = 0; EnumIndex < Enum->NumEnums()-1; EnumIndex++)
+			for (int32 EnumIndex = 0; EnumIndex < TotalEnumEntryCount; EnumIndex++)
 			{
 				if(FNiagaraEditorUtilities::IsEnumIndexVisible(Enum, EnumIndex))
 				{
