Added Niagara Script Graph zoom limit break option

diff --git a/Engine/Plugins/FX/Niagara/Source/NiagaraEditor/Private/NiagaraEditorSettings.cpp b/Engine/Plugins/FX/Niagara/Source/NiagaraEditor/Private/NiagaraEditorSettings.cpp
--- a/Engine/Plugins/FX/Niagara/Source/NiagaraEditor/Private/NiagaraEditorSettings.cpp
+++ b/Engine/Plugins/FX/Niagara/Source/NiagaraEditor/Private/NiagaraEditorSettings.cpp
@@ -789,6 +789,11 @@ bool UNiagaraEditorSettings::IsAllowedAssetObjectByClassUsageInternal(const UObj
 	return true;
 }
 
+EGraphZoomLimitHandling UNiagaraEditorSettings::GetScriptEditorZoomLimitHandling() const
+{
+	return ScriptGraphZoomLimitHandling;
+}
+
 bool UNiagaraEditorSettings::IsAllowedObjectByClassUsageInternal(const UObject& InObject, TSet<const UObject*>& CheckedObjects) const
 {
 	CheckedObjects.Add(&InObject);
diff --git a/Engine/Plugins/FX/Niagara/Source/NiagaraEditor/Private/Widgets/SNiagaraScriptGraph.cpp b/Engine/Plugins/FX/Niagara/Source/NiagaraEditor/Private/Widgets/SNiagaraScriptGraph.cpp
--- a/Engine/Plugins/FX/Niagara/Source/NiagaraEditor/Private/Widgets/SNiagaraScriptGraph.cpp
+++ b/Engine/Plugins/FX/Niagara/Source/NiagaraEditor/Private/Widgets/SNiagaraScriptGraph.cpp
@@ -30,11 +30,104 @@
 #include "GraphEditorActions.h"
 #include "EditorFontGlyphs.h"
 #include "NiagaraEditorSettings.h"
+#include "SGraphPanel.h"
 #include "AssetRegistry/AssetRegistryModule.h"
 #include "Widgets/SNiagaraGraphActionMenu.h"
 
 #define LOCTEXT_NAMESPACE "NiagaraScriptGraph"
 
+struct FNiagaraScriptZoomLevelsContainer : public FZoomLevelsContainer
+{
+	struct FZoomLevelEntry
+	{
+	public:
+		FZoomLevelEntry(float InZoomAmount, const FText& InDisplayText, EGraphRenderingLOD::Type InLOD)
+			: DisplayText(FText::Format(LOCTEXT("Zoom", "Zoom {0}"), InDisplayText))
+			, ZoomAmount(InZoomAmount)
+			, LOD(InLOD)
+		{
+		}
+
+	public:
+		FText DisplayText;
+		float ZoomAmount;
+		EGraphRenderingLOD::Type LOD;
+	};
+	
+	FNiagaraScriptZoomLevelsContainer()
+	{
+		ZoomLevels.Reserve(20);
+		ZoomLevels.Add(FZoomLevelEntry(0.100f, FText::FromString(TEXT("-12")), EGraphRenderingLOD::LowestDetail));
+		ZoomLevels.Add(FZoomLevelEntry(0.125f, FText::FromString(TEXT("-11")), EGraphRenderingLOD::LowestDetail));
+		ZoomLevels.Add(FZoomLevelEntry(0.150f, FText::FromString(TEXT("-10")), EGraphRenderingLOD::LowestDetail));
+		ZoomLevels.Add(FZoomLevelEntry(0.175f, FText::FromString(TEXT("-9")), EGraphRenderingLOD::LowestDetail));
+		ZoomLevels.Add(FZoomLevelEntry(0.200f, FText::FromString(TEXT("-8")), EGraphRenderingLOD::LowestDetail));
+		ZoomLevels.Add(FZoomLevelEntry(0.225f, FText::FromString(TEXT("-7")), EGraphRenderingLOD::LowDetail));
+		ZoomLevels.Add(FZoomLevelEntry(0.250f, FText::FromString(TEXT("-6")), EGraphRenderingLOD::LowDetail));
+		ZoomLevels.Add(FZoomLevelEntry(0.375f, FText::FromString(TEXT("-5")), EGraphRenderingLOD::MediumDetail));
+		ZoomLevels.Add(FZoomLevelEntry(0.500f, FText::FromString(TEXT("-4")), EGraphRenderingLOD::MediumDetail));
+		ZoomLevels.Add(FZoomLevelEntry(0.675f, FText::FromString(TEXT("-3")), EGraphRenderingLOD::MediumDetail));
+		ZoomLevels.Add(FZoomLevelEntry(0.750f, FText::FromString(TEXT("-2")), EGraphRenderingLOD::DefaultDetail));
+		ZoomLevels.Add(FZoomLevelEntry(0.875f, FText::FromString(TEXT("-1")), EGraphRenderingLOD::DefaultDetail));
+		ZoomLevels.Add(FZoomLevelEntry(1.000f, FText::FromString(TEXT("1:1")), EGraphRenderingLOD::DefaultDetail));
+		ZoomLevels.Add(FZoomLevelEntry(1.250f, FText::FromString(TEXT("+1")), EGraphRenderingLOD::DefaultDetail));
+		ZoomLevels.Add(FZoomLevelEntry(1.375f, FText::FromString(TEXT("+2")), EGraphRenderingLOD::DefaultDetail));
+		ZoomLevels.Add(FZoomLevelEntry(1.500f, FText::FromString(TEXT("+3")), EGraphRenderingLOD::FullyZoomedIn));
+		ZoomLevels.Add(FZoomLevelEntry(1.675f, FText::FromString(TEXT("+4")), EGraphRenderingLOD::FullyZoomedIn));
+		ZoomLevels.Add(FZoomLevelEntry(1.750f, FText::FromString(TEXT("+5")), EGraphRenderingLOD::FullyZoomedIn));
+		ZoomLevels.Add(FZoomLevelEntry(1.875f, FText::FromString(TEXT("+6")), EGraphRenderingLOD::FullyZoomedIn));
+		ZoomLevels.Add(FZoomLevelEntry(2.000f, FText::FromString(TEXT("+7")), EGraphRenderingLOD::FullyZoomedIn));
+	}
+
+	float GetZoomAmount(int32 InZoomLevel) const override
+	{
+		checkSlow(ZoomLevels.IsValidIndex(InZoomLevel));
+		return ZoomLevels[InZoomLevel].ZoomAmount;
+	}
+
+	int32 GetNearestZoomLevel(float InZoomAmount) const override
+	{
+		for (int32 ZoomLevelIndex=0; ZoomLevelIndex < GetNumZoomLevels(); ++ZoomLevelIndex)
+		{
+			if (InZoomAmount <= GetZoomAmount(ZoomLevelIndex))
+			{
+				return ZoomLevelIndex;
+			}
+		}
+
+		return GetDefaultZoomLevel();
+	}
+	
+	FText GetZoomText(int32 InZoomLevel) const override
+	{
+		checkSlow(ZoomLevels.IsValidIndex(InZoomLevel));
+		return ZoomLevels[InZoomLevel].DisplayText;
+	}
+	
+	int32 GetNumZoomLevels() const override
+	{
+		return ZoomLevels.Num();
+	}
+	
+	int32 GetDefaultZoomLevel() const override
+	{
+		return 12;
+	}
+
+	EGraphRenderingLOD::Type GetLOD(int32 InZoomLevel) const override
+	{
+		checkSlow(ZoomLevels.IsValidIndex(InZoomLevel));
+		return ZoomLevels[InZoomLevel].LOD;
+	}
+
+	virtual EGraphZoomLimitHandling GetZoomLimitHandling() const override
+	{
+		return GetDefault<UNiagaraEditorSettings>()->GetScriptEditorZoomLimitHandling();
+	}
+	
+	TArray<FZoomLevelEntry> ZoomLevels;
+};
+
 void SNiagaraScriptGraph::Construct(const FArguments& InArgs, TSharedRef<FNiagaraScriptGraphViewModel> InViewModel, const FAssetData& InEditedAsset)
 {
 	EditedAsset = InEditedAsset;
@@ -286,12 +379,12 @@ TSharedRef<SGraphEditor> SNiagaraScriptGraph::ConstructGraphEditor()
 
 	// Set a niagara node factory.
 	CreatedGraphEditor->SetNodeFactory(MakeShareable(new FNiagaraNodeFactory()));
+	CreatedGraphEditor->GetGraphPanel()->SetZoomLevelsContainer<FNiagaraScriptZoomLevelsContainer>();
 	CreatedGraphEditor->ZoomToFit(false);
 
 	return CreatedGraphEditor;
 }
 
-
 void SNiagaraScriptGraph::ViewModelSelectedNodesChanged()
 {
 	if (FNiagaraEditorUtilities::SetsMatch(GraphEditor->GetSelectedNodes(), ViewModel->GetNodeSelection()->GetSelectedObjectsResolved()) == false)
diff --git a/Engine/Plugins/FX/Niagara/Source/NiagaraEditor/Public/NiagaraEditorSettings.h b/Engine/Plugins/FX/Niagara/Source/NiagaraEditor/Public/NiagaraEditorSettings.h
--- a/Engine/Plugins/FX/Niagara/Source/NiagaraEditor/Public/NiagaraEditorSettings.h
+++ b/Engine/Plugins/FX/Niagara/Source/NiagaraEditor/Public/NiagaraEditorSettings.h
@@ -9,6 +9,7 @@
 #include "NiagaraValidationRuleSet.h"
 #include "Engine/DeveloperSettings.h"
 #include "NiagaraSpawnShortcut.h"
+#include "SNodePanel.h"
 #include "Editor/UnrealEdTypes.h"
 #include "Viewports.h"
 #include "NiagaraEditorSettings.generated.h"
@@ -378,6 +379,9 @@ public:
 	UPROPERTY(config, EditAnywhere, Category = Niagara)
 	ENiagaraAddDefaultsTrackMode DefaultsSequencerSubtracks = ENiagaraAddDefaultsTrackMode::NoSubtracks;
 
+	UPROPERTY(config, EditAnywhere, Category = "Niagara Graph")
+	EGraphZoomLimitHandling ScriptGraphZoomLimitHandling = EGraphZoomLimitHandling::Default;
+
 	NIAGARAEDITOR_API TArray<float> GetPlaybackSpeeds() const;
 
 	/** Gets whether or not auto-compile is enabled in the editors. */
@@ -503,7 +507,8 @@ public:
 	NIAGARAEDITOR_API void SetAutoExposureValue(bool bInUseAutoExposure);
 
 	NIAGARAEDITOR_API bool GetForceSilentLoadingOfCachedAssets() const { return bForceSilentLoadingOfCachedAssets; }
-	
+
+	NIAGARAEDITOR_API EGraphZoomLimitHandling GetScriptEditorZoomLimitHandling() const;
 private:
 	NIAGARAEDITOR_API bool IsAllowedObjectByClassUsageInternal(const UObject& InObject, TSet<const UObject*>& CheckedObjects) const;
 	NIAGARAEDITOR_API bool IsAllowedAssetObjectByClassUsageInternal(const UObject& InAssetObject, TSet<const UObject*>& CheckedAssetObjects) const;
