- Remove constructing the signature to get the Name

#rnx
#tests Ran ET for data channels

diff --git a/Engine/Plugins/FX/Niagara/Source/Niagara/Private/DataInterface/NiagaraDataInterfaceDataChannelRead.cpp b/Engine/Plugins/FX/Niagara/Source/Niagara/Private/DataInterface/NiagaraDataInterfaceDataChannelRead.cpp
--- a/Engine/Plugins/FX/Niagara/Source/Niagara/Private/DataInterface/NiagaraDataInterfaceDataChannelRead.cpp
+++ b/Engine/Plugins/FX/Niagara/Source/Niagara/Private/DataInterface/NiagaraDataInterfaceDataChannelRead.cpp
@@ -90,6 +90,13 @@ namespace NDIDataChannelReadLocal
 	static const TCHAR* TemplateShaderFile_Read = TEXT("/Plugin/FX/Niagara/Private/DataChannel/NiagaraDataInterfaceDataChannelTemplate_Read.ush");
 	static const TCHAR* TemplateShaderFile_Consume = TEXT("/Plugin/FX/Niagara/Private/DataChannel/NiagaraDataInterfaceDataChannelTemplate_Consume.ush");
 
+	static const FName NAME_SpawnDirect("SpawnDirect");
+	static const FName NAME_SpawnConditional("SpawnConditional");
+	static const FName NAME_Consume("Consume");
+	static const FName NAME_Read("Read");
+	static const FName NAME_Num("Num");
+	static const FName NAME_ScaleSpawnCount("ScaleSpawnCount");
+
 	//////////////////////////////////////////////////////////////////////////
 	//Function definitions
 	const FNiagaraFunctionSignature& GetFunctionSig_Num()
@@ -97,7 +104,7 @@ namespace NDIDataChannelReadLocal
 		static FNiagaraFunctionSignature Sig;
 		if(!Sig.IsValid())
 		{
-			Sig.Name = TEXT("Num");
+			Sig.Name = NAME_Num;
 #if WITH_EDITORONLY_DATA
 			Sig.Description = LOCTEXT("NumFunctionDescription", "Returns the current number of elements in the Data Channel being read.");
 			NIAGARA_ADD_FUNCTION_SOURCE_INFO(Sig)
@@ -136,7 +143,7 @@ namespace NDIDataChannelReadLocal
 		static FNiagaraFunctionSignature Sig;
 		if(!Sig.IsValid())
 		{
-			Sig.Name = TEXT("Read");
+			Sig.Name = NAME_Read;
 #if WITH_EDITORONLY_DATA
 			Sig.Description = LOCTEXT("ReadFunctionDescription", "Reads Data Channel data at a specific index. Any values we read that are not in the Data Channel data are set to their default values. Returns success if there was a valid Data Channel entry to read from at the given index.");
 			NIAGARA_ADD_FUNCTION_SOURCE_INFO(Sig)
@@ -156,7 +163,7 @@ namespace NDIDataChannelReadLocal
 		static FNiagaraFunctionSignature Sig;
 		if(!Sig.IsValid())
 		{
-			Sig.Name = TEXT("Consume");
+			Sig.Name = NAME_Consume;
 #if WITH_EDITORONLY_DATA
 			Sig.Description = LOCTEXT("ConsumeFunctionDescription", "Consumes an item from the Data Channel and reads the specified values. Any values we read that are not in the Data Channel data are set to their default values. Returns success if an entry was available to be consumed in the Data Channel.");
 			NIAGARA_ADD_FUNCTION_SOURCE_INFO(Sig)
@@ -183,7 +190,7 @@ namespace NDIDataChannelReadLocal
 		static FNiagaraFunctionSignature Sig;
 		if(!Sig.IsValid())
 		{
-			Sig.Name = TEXT("SpawnConditional");
+			Sig.Name = NAME_SpawnConditional;
 #if WITH_EDITORONLY_DATA
 			NIAGARA_ADD_FUNCTION_SOURCE_INFO(Sig)
 			Sig.Description = LOCTEXT("SpawnCustomFunctionDescription", "Will Spawn particles into the given Emitter between Min and Max counts for every element in the Data Channel.\n\
@@ -214,7 +221,7 @@ namespace NDIDataChannelReadLocal
 		static FNiagaraFunctionSignature Sig;
 		if (!Sig.IsValid())
 		{
-			Sig.Name = TEXT("SpawnDirect");
+			Sig.Name = NAME_SpawnDirect;
 #if WITH_EDITORONLY_DATA
 			Sig.Description = LOCTEXT("SpawnDirectFunctionDescription", "Spawns particles into a given emitter for each entry in the Data Channel. Spawn count is determined directly from a value in the Data Channel. Additional per NDC item random scale and a clamp is available.");
 			NIAGARA_ADD_FUNCTION_SOURCE_INFO(Sig)
@@ -241,7 +248,7 @@ namespace NDIDataChannelReadLocal
 		static FNiagaraFunctionSignature Sig;
 		if (!Sig.IsValid())
 		{
-			Sig.Name = TEXT("ScaleSpawnCount");
+			Sig.Name = NAME_ScaleSpawnCount;
 #if WITH_EDITORONLY_DATA
 			Sig.Description = LOCTEXT("ScaleSpawnCountFunctionDescription", "Applies a scaling value for each NDC Item spawn count based on a variable in the NDC data. Optional additional random scale and clamp operations to the value read from each NDC entry.");
 			NIAGARA_ADD_FUNCTION_SOURCE_INFO(Sig)
@@ -270,8 +277,8 @@ namespace NDIDataChannelReadLocal
 		LoadShaderSourceFile(TemplateShaderFile_Common, EShaderPlatform::SP_PCD3D_SM5, &OutCommonTemplateShaders.AddDefaulted_GetRef(), nullptr);
 		LoadShaderSourceFile(TemplateShaderFile_ReadCommon, EShaderPlatform::SP_PCD3D_SM5, &OutCommonTemplateShaders.AddDefaulted_GetRef(), nullptr);
 
-		LoadShaderSourceFile(TemplateShaderFile_Read, EShaderPlatform::SP_PCD3D_SM5, &OutMap.Add(GetFunctionSig_Read().Name), nullptr);
-		LoadShaderSourceFile(TemplateShaderFile_Consume, EShaderPlatform::SP_PCD3D_SM5, &OutMap.Add(GetFunctionSig_Consume().Name), nullptr);
+		LoadShaderSourceFile(TemplateShaderFile_Read, EShaderPlatform::SP_PCD3D_SM5, &OutMap.Add(NAME_Read), nullptr);
+		LoadShaderSourceFile(TemplateShaderFile_Consume, EShaderPlatform::SP_PCD3D_SM5, &OutMap.Add(NAME_Consume), nullptr);
 	}
 
 	BEGIN_SHADER_PARAMETER_STRUCT(FShaderParameters, )
@@ -1200,15 +1207,15 @@ void UNiagaraDataInterfaceDataChannelRead::GetFunctionsInternal(TArray<FNiagaraF
 
 void UNiagaraDataInterfaceDataChannelRead::GetVMExternalFunction(const FVMExternalFunctionBindingInfo& BindingInfo, void* InstanceData, FVMExternalFunction& OutFunc)
 {
-	if (BindingInfo.Name == NDIDataChannelReadLocal::GetFunctionSig_Num().Name)
+	if (BindingInfo.Name == NDIDataChannelReadLocal::NAME_Num)
 	{
 		OutFunc = FVMExternalFunction::CreateLambda([this](FVectorVMExternalFunctionContext& Context) { this->Num(Context); });
 	}
-	else if (BindingInfo.Name == NDIDataChannelReadLocal::GetFunctionSig_GetNDCSpawnData().Name)
+	else if (BindingInfo.Name == NDIDataChannelUtilities::GetNDCSpawnDataName)
 	{
 		OutFunc = FVMExternalFunction::CreateLambda([this](FVectorVMExternalFunctionContext& Context) { this->GetNDCSpawnData(Context); });
 	}
-	else if (BindingInfo.Name == NDIDataChannelReadLocal::GetFunctionSig_SpawnDirect().Name)
+	else if (BindingInfo.Name == NDIDataChannelReadLocal::NAME_SpawnDirect)
 	{
 		//Find the appropriate variable from the function binding and determine it's type to bind to the correct impl for ScaleSpawnCount and provide the variable index.
 		FNiagaraVariableBase FuncSpecVariable = DecodeVariableFromSpecifiers(BindingInfo);
@@ -1236,7 +1243,7 @@ void UNiagaraDataInterfaceDataChannelRead::GetVMExternalFunction(const FVMExtern
 			UE_LOG(LogTemp, Display, TEXT("Failed to bind Data Interface function due to missing NDC variable for %s.\nReceived Name: %s\nNDC: %s"), *GetPathNameSafe(this), *BindingInfo.Name.ToString(), *GetPathNameSafe(Channel));
 		}
 	}
-	else if (BindingInfo.Name == NDIDataChannelReadLocal::GetFunctionSig_ScaleSpawnCount().Name)
+	else if (BindingInfo.Name == NDIDataChannelReadLocal::NAME_ScaleSpawnCount)
 	{
 		//Find the appropriate variable from the function binding and determine it's type to bind to the correct impl for ScaleSpawnCount and provide the variable index.
 			//Can we do this decode on the PostCompile? Also would be good to validate/warn on missing or invalid types.
@@ -1273,15 +1280,15 @@ void UNiagaraDataInterfaceDataChannelRead::GetVMExternalFunction(const FVMExtern
 	else
 	{
 		int32 FuncIndex = CompiledData.FindFunctionInfoIndex(BindingInfo.Name, BindingInfo.VariadicInputs, BindingInfo.VariadicOutputs);
-		if (BindingInfo.Name == NDIDataChannelReadLocal::GetFunctionSig_Read().Name)
+		if (BindingInfo.Name == NDIDataChannelReadLocal::NAME_Read)
 		{
 			OutFunc = FVMExternalFunction::CreateLambda([this, FuncIndex](FVectorVMExternalFunctionContext& Context) { this->Read(Context, FuncIndex); });
 		}
-		else if (BindingInfo.Name == NDIDataChannelReadLocal::GetFunctionSig_Consume().Name)
+		else if (BindingInfo.Name == NDIDataChannelReadLocal::NAME_Consume)
 		{
 			OutFunc = FVMExternalFunction::CreateLambda([this, FuncIndex](FVectorVMExternalFunctionContext& Context) { this->Consume(Context, FuncIndex); });
 		}
-		else if (BindingInfo.Name == NDIDataChannelReadLocal::GetFunctionSig_SpawnConditional().Name)
+		else if (BindingInfo.Name == NDIDataChannelReadLocal::NAME_SpawnConditional)
 		{
 			OutFunc = FVMExternalFunction::CreateLambda([this, FuncIndex](FVectorVMExternalFunctionContext& Context) { this->SpawnConditional(Context, FuncIndex); });
 		}
