- Move some of the stateless headers into public so they can be shared by ET

#rnx

diff --git a/Engine/Plugins/FX/Niagara/Source/Niagara/Internal/Stateless/NiagaraStatelessCommon.h b/Engine/Plugins/FX/Niagara/Source/Niagara/Internal/Stateless/NiagaraStatelessCommon.h
--- a/Engine/Plugins/FX/Niagara/Source/Niagara/Internal/Stateless/NiagaraStatelessCommon.h
+++ b/Engine/Plugins/FX/Niagara/Source/Niagara/Internal/Stateless/NiagaraStatelessCommon.h
@@ -51,37 +51,6 @@ private:
 	FTransform3f	Transforms[2];			// World to Local & Local to World
 };
 
-template<typename TType> struct FNiagaraStatelessRangeDefaultValue {};
-template<> struct FNiagaraStatelessRangeDefaultValue<int32> { static constexpr int32 Zero() { return 0; } };
-template<> struct FNiagaraStatelessRangeDefaultValue<float> { static constexpr float Zero() { return 0.0f; } };
-template<> struct FNiagaraStatelessRangeDefaultValue<FVector2f> { static const FVector2f Zero() { return FVector2f(0.0f, 0.0f); } };
-template<> struct FNiagaraStatelessRangeDefaultValue<FVector3f> { static const FVector3f Zero() { return FVector3f(0.0f, 0.0f, 0.0f); } };
-template<> struct FNiagaraStatelessRangeDefaultValue<FVector4f> { static const FVector4f Zero() { return FVector4f(0.0f, 0.0f, 0.0f, 0.0f); } };
-template<> struct FNiagaraStatelessRangeDefaultValue<FLinearColor> { static const FLinearColor Zero() { return FLinearColor(0.0f, 0.0f, 0.0f, 0.0f); } };
-
-template<typename TType>
-struct FNiagaraStatelessRange
-{
-	using ValueType = TType;
-
-	FNiagaraStatelessRange() = default;
-	FNiagaraStatelessRange(const FNiagaraStatelessRange& Other) = default;
-	explicit FNiagaraStatelessRange(const ValueType& InMinMax) : Min(InMinMax), Max(InMinMax) {}
-	explicit FNiagaraStatelessRange(const ValueType& InMin, const ValueType& InMax) : Min(InMin), Max(InMax) {}
-
-	ValueType GetScale() const { return Max - Min; }
-
-	int32		ParameterOffset = INDEX_NONE;
-	ValueType	Min = FNiagaraStatelessRangeDefaultValue<TType>::Zero();
-	ValueType	Max = FNiagaraStatelessRangeDefaultValue<TType>::Zero();
-};
-
-using FNiagaraStatelessRangeInt		= FNiagaraStatelessRange<int32>;
-using FNiagaraStatelessRangeFloat	= FNiagaraStatelessRange<float>;
-using FNiagaraStatelessRangeVector2	= FNiagaraStatelessRange<FVector2f>;
-using FNiagaraStatelessRangeVector3	= FNiagaraStatelessRange<FVector3f>;
-using FNiagaraStatelessRangeVector4	= FNiagaraStatelessRange<FVector4f>;
-using FNiagaraStatelessRangeColor	= FNiagaraStatelessRange<FLinearColor>;
 
 struct FNiagaraStatelessGlobals
 {
diff --git a/Engine/Plugins/FX/Niagara/Source/Niagara/Internal/Stateless/NiagaraStatelessDistribution.h b/Engine/Plugins/FX/Niagara/Source/Niagara/Internal/Stateless/NiagaraStatelessDistribution.h
deleted file mode 100644
diff --git a/Engine/Plugins/FX/Niagara/Source/Niagara/Internal/Stateless/NiagaraStatelessEmitterDataBuildContext.h b/Engine/Plugins/FX/Niagara/Source/Niagara/Internal/Stateless/NiagaraStatelessEmitterDataBuildContext.h
--- a/Engine/Plugins/FX/Niagara/Source/Niagara/Internal/Stateless/NiagaraStatelessEmitterDataBuildContext.h
+++ b/Engine/Plugins/FX/Niagara/Source/Niagara/Internal/Stateless/NiagaraStatelessEmitterDataBuildContext.h
@@ -2,8 +2,8 @@
 
 #pragma once
 
-#include "NiagaraStatelessCommon.h"
-#include "NiagaraStatelessDistribution.h"
+#include "Stateless/NiagaraStatelessCommon.h"
+#include "Stateless/NiagaraStatelessDistribution.h"
 #include "NiagaraStatelessBuiltDistribution.h"
 
 struct FNiagaraDataSetCompiledData;
diff --git a/Engine/Plugins/FX/Niagara/Source/Niagara/Internal/Stateless/NiagaraStatelessModule.h b/Engine/Plugins/FX/Niagara/Source/Niagara/Internal/Stateless/NiagaraStatelessModule.h
--- a/Engine/Plugins/FX/Niagara/Source/Niagara/Internal/Stateless/NiagaraStatelessModule.h
+++ b/Engine/Plugins/FX/Niagara/Source/Niagara/Internal/Stateless/NiagaraStatelessModule.h
@@ -3,10 +3,10 @@
 #pragma once
 
 #include "NiagaraMergeable.h"
-#include "NiagaraStatelessCommon.h"
-#include "NiagaraStatelessDistribution.h"
-#include "NiagaraStatelessSetShaderParameterContext.h"
-#include "NiagaraStatelessShaderParametersBuilder.h"
+#include "Stateless/NiagaraStatelessCommon.h"
+#include "Stateless/NiagaraStatelessDistribution.h"
+#include "Stateless/NiagaraStatelessSetShaderParameterContext.h"
+#include "Stateless/NiagaraStatelessShaderParametersBuilder.h"
 
 #include "NiagaraStatelessModule.generated.h"
 
diff --git a/Engine/Plugins/FX/Niagara/Source/Niagara/Internal/Stateless/NiagaraStatelessSetShaderParameterContext.h b/Engine/Plugins/FX/Niagara/Source/Niagara/Internal/Stateless/NiagaraStatelessSetShaderParameterContext.h
--- a/Engine/Plugins/FX/Niagara/Source/Niagara/Internal/Stateless/NiagaraStatelessSetShaderParameterContext.h
+++ b/Engine/Plugins/FX/Niagara/Source/Niagara/Internal/Stateless/NiagaraStatelessSetShaderParameterContext.h
@@ -2,7 +2,8 @@
 
 #pragma once
 
-#include "NiagaraStatelessCommon.h"
+#include "Stateless/NiagaraStatelessCommon.h"
+#include "Stateless/NiagaraStatelessRange.h"
 #include "ShaderParameterStruct.h"
 
 class FNiagaraStatelessSetShaderParameterContext
diff --git a/Engine/Plugins/FX/Niagara/Source/Niagara/Internal/Stateless/NiagaraStatelessSpawnInfo.h b/Engine/Plugins/FX/Niagara/Source/Niagara/Internal/Stateless/NiagaraStatelessSpawnInfo.h
--- a/Engine/Plugins/FX/Niagara/Source/Niagara/Internal/Stateless/NiagaraStatelessSpawnInfo.h
+++ b/Engine/Plugins/FX/Niagara/Source/Niagara/Internal/Stateless/NiagaraStatelessSpawnInfo.h
@@ -2,8 +2,8 @@
 
 #pragma once
 
-#include "NiagaraStatelessCommon.h"
-#include "NiagaraStatelessDistribution.h"
+#include "Stateless/NiagaraStatelessCommon.h"
+#include "Stateless/NiagaraStatelessDistribution.h"
 
 #include "NiagaraStatelessSpawnInfo.generated.h"
 
diff --git a/Engine/Plugins/FX/Niagara/Source/Niagara/Public/Stateless/NiagaraStatelessDistribution.h b/Engine/Plugins/FX/Niagara/Source/Niagara/Public/Stateless/NiagaraStatelessDistribution.h
new file mode 100644
--- /dev/null
+++ b/Engine/Plugins/FX/Niagara/Source/Niagara/Public/Stateless/NiagaraStatelessDistribution.h
@@ -0,0 +1,368 @@
+// Copyright Epic Games, Inc. All Rights Reserved.
+
+#pragma once
+
+#include "Curves/RichCurve.h"
+#include "NiagaraStatelessRange.h"
+#include "NiagaraParameterBinding.h"
+
+#include "NiagaraStatelessDistribution.generated.h"
+
+UENUM()
+enum class ENiagaraDistributionMode
+{
+	Binding,
+	UniformConstant,
+	NonUniformConstant,
+	UniformRange,
+	NonUniformRange,
+	UniformCurve,
+	NonUniformCurve
+};
+
+USTRUCT()
+struct FNiagaraDistributionBase
+{
+	GENERATED_BODY()
+
+	virtual ~FNiagaraDistributionBase() = default;
+
+	UPROPERTY(EditAnywhere, Category = "Parameters")
+	ENiagaraDistributionMode Mode = ENiagaraDistributionMode::UniformConstant;
+
+	UPROPERTY(EditAnywhere, Category = "Parameters")
+	FNiagaraVariableBase ParameterBinding;
+
+	bool IsBinding() const { return Mode == ENiagaraDistributionMode::Binding; }
+	bool IsConstant() const { return Mode == ENiagaraDistributionMode::UniformConstant || Mode == ENiagaraDistributionMode::NonUniformConstant; }
+	bool IsUniform() const { return Mode == ENiagaraDistributionMode::UniformConstant || Mode == ENiagaraDistributionMode::UniformRange; }
+	bool IsCurve() const { return Mode == ENiagaraDistributionMode::UniformCurve || Mode == ENiagaraDistributionMode::NonUniformCurve; }
+	bool IsRange() const { return Mode == ENiagaraDistributionMode::UniformRange || Mode == ENiagaraDistributionMode::NonUniformRange; }
+
+#if WITH_EDITORONLY_DATA
+	UPROPERTY(EditAnywhere, Category = "Parameters")
+	TArray<float> ChannelConstantsAndRanges;
+
+	UPROPERTY(EditAnywhere, Category = "Parameters")
+	TArray<FRichCurve> ChannelCurves;
+
+	UPROPERTY(EditAnywhere, Category = "Parameters")
+	int32 MaxLutSampleCount = 128;
+
+	bool operator==(const FNiagaraDistributionBase& Other) const
+	{
+		return
+			Mode == Other.Mode &&
+			ParameterBinding == Other.ParameterBinding &&
+			ChannelConstantsAndRanges == Other.ChannelConstantsAndRanges &&
+			ChannelCurves == Other.ChannelCurves &&
+			MaxLutSampleCount == Other.MaxLutSampleCount;
+	}
+
+	virtual bool AllowBinding() const { return true; }
+	virtual bool AllowCurves() const { return true; }
+	virtual bool DisplayAsColor() const { return false; }
+	virtual int32 GetBaseNumberOfChannels() const { return 0; }
+	virtual void UpdateValuesFromDistribution() { }
+
+	virtual FNiagaraTypeDefinition GetBindingTypeDef() const { return FNiagaraTypeDefinition(); }
+
+	static void PostEditChangeProperty(UObject* OwnerObject, FPropertyChangedEvent& PropertyChangedEvent);
+#endif
+};
+
+USTRUCT()
+struct FNiagaraDistributionRangeInt
+{
+	GENERATED_BODY()
+
+	FNiagaraDistributionRangeInt() = default;
+	explicit FNiagaraDistributionRangeInt(int32 ConstantValue) { InitConstant(ConstantValue); }
+
+	UPROPERTY(EditAnywhere, Category = "Parameters")
+	ENiagaraDistributionMode Mode = ENiagaraDistributionMode::UniformConstant;
+
+	UPROPERTY(EditAnywhere, Category = "Parameters")
+	FNiagaraVariableBase ParameterBinding;
+
+	UPROPERTY(EditAnywhere, Category = "Parameters")
+	int32 Min = 0;
+
+	UPROPERTY(EditAnywhere, Category = "Parameters")
+	int32 Max = 0;
+
+	NIAGARA_API void InitConstant(int32 Value);
+	NIAGARA_API FNiagaraStatelessRangeInt CalculateRange(const int32 Default = 0) const;
+
+	bool IsBinding() const { return Mode == ENiagaraDistributionMode::Binding; }
+	bool IsConstant() const { return Mode == ENiagaraDistributionMode::UniformConstant || Mode == ENiagaraDistributionMode::NonUniformConstant; }
+	bool IsUniform() const { return Mode == ENiagaraDistributionMode::UniformConstant || Mode == ENiagaraDistributionMode::UniformRange; }
+	bool IsCurve() const { return Mode == ENiagaraDistributionMode::UniformCurve || Mode == ENiagaraDistributionMode::NonUniformCurve; }
+	bool IsRange() const { return Mode == ENiagaraDistributionMode::UniformRange || Mode == ENiagaraDistributionMode::NonUniformRange; }
+
+#if WITH_EDITORONLY_DATA
+	FNiagaraTypeDefinition GetBindingTypeDef() const { return FNiagaraTypeDefinition::GetIntDef(); }
+#endif
+};
+
+USTRUCT()
+struct FNiagaraDistributionRangeFloat : public FNiagaraDistributionBase
+{
+	GENERATED_BODY()
+
+	FNiagaraDistributionRangeFloat() = default;
+	explicit FNiagaraDistributionRangeFloat(float ConstantValue) { InitConstant(ConstantValue); }
+	explicit FNiagaraDistributionRangeFloat(float MinValue, float MaxValue) { InitRange(MinValue, MaxValue); }
+
+	UPROPERTY(EditAnywhere, Category = "Parameters")
+	float Min = 0.0f;
+
+	UPROPERTY(EditAnywhere, Category = "Parameters")
+	float Max = 0.0f;
+
+	NIAGARA_API void InitConstant(float Value);
+	NIAGARA_API void InitRange(float MinValue, float MaxValue);
+	NIAGARA_API FNiagaraStatelessRangeFloat CalculateRange(const float Default = 0.0f) const;
+
+#if WITH_EDITORONLY_DATA
+	virtual bool AllowCurves() const override { return false; }
+	virtual int32 GetBaseNumberOfChannels() const override { return 1; }
+	NIAGARA_API virtual void UpdateValuesFromDistribution() override;
+	virtual FNiagaraTypeDefinition GetBindingTypeDef() const { return FNiagaraTypeDefinition::GetFloatDef(); }
+#endif
+	bool SerializeFromMismatchedTag(const struct FPropertyTag& Tag, FStructuredArchive::FSlot Slot);
+};
+
+USTRUCT()
+struct FNiagaraDistributionRangeVector2 : public FNiagaraDistributionBase
+{
+	GENERATED_BODY()
+
+	FNiagaraDistributionRangeVector2() = default;
+	explicit FNiagaraDistributionRangeVector2(const FVector2f& ConstantValue) { InitConstant(ConstantValue); }
+
+	UPROPERTY(EditAnywhere, Category = "Parameters")
+	FVector2f Min = FVector2f::ZeroVector;
+
+	UPROPERTY(EditAnywhere, Category = "Parameters")
+	FVector2f Max = FVector2f::ZeroVector;
+
+	NIAGARA_API void InitConstant(const FVector2f& Value);
+	NIAGARA_API FNiagaraStatelessRangeVector2 CalculateRange(const FVector2f& Default = FVector2f::ZeroVector) const;
+
+#if WITH_EDITORONLY_DATA
+	virtual bool AllowCurves() const override { return false; }
+	virtual int32 GetBaseNumberOfChannels() const override { return 2; }
+	virtual void UpdateValuesFromDistribution() override;
+	virtual FNiagaraTypeDefinition GetBindingTypeDef() const { return FNiagaraTypeDefinition::GetVec2Def(); }
+#endif
+	bool SerializeFromMismatchedTag(const struct FPropertyTag& Tag, FStructuredArchive::FSlot Slot);
+};
+
+USTRUCT()
+struct FNiagaraDistributionRangeVector3 : public FNiagaraDistributionBase
+{
+	GENERATED_BODY()
+
+	FNiagaraDistributionRangeVector3() = default;
+	explicit FNiagaraDistributionRangeVector3(const FVector3f& ConstantValue) { InitConstant(ConstantValue); }
+
+	UPROPERTY(EditAnywhere, Category = "Parameters")
+	FVector3f Min = FVector3f::ZeroVector;
+
+	UPROPERTY(EditAnywhere, Category = "Parameters")
+	FVector3f Max = FVector3f::ZeroVector;
+
+	NIAGARA_API void InitConstant(const FVector3f& Value);
+	NIAGARA_API FNiagaraStatelessRangeVector3 CalculateRange(const FVector3f& Default = FVector3f::ZeroVector) const;
+
+#if WITH_EDITORONLY_DATA
+	virtual bool AllowCurves() const override { return false; }
+	virtual int32 GetBaseNumberOfChannels() const override { return 3; }
+	virtual void UpdateValuesFromDistribution() override;
+	virtual FNiagaraTypeDefinition GetBindingTypeDef() const { return FNiagaraTypeDefinition::GetVec3Def(); }
+#endif
+	bool SerializeFromMismatchedTag(const struct FPropertyTag& Tag, FStructuredArchive::FSlot Slot);
+};
+
+USTRUCT()
+struct FNiagaraDistributionRangeColor : public FNiagaraDistributionBase
+{
+	GENERATED_BODY()
+
+	FNiagaraDistributionRangeColor() = default;
+	explicit FNiagaraDistributionRangeColor(const FLinearColor& ConstantValue) { InitConstant(ConstantValue); }
+
+	UPROPERTY(EditAnywhere, Category = "Parameters")
+	FLinearColor Min = FLinearColor::White;
+
+	UPROPERTY(EditAnywhere, Category = "Parameters")
+	FLinearColor Max = FLinearColor::White;
+
+	NIAGARA_API void InitConstant(const FLinearColor& Value);
+	NIAGARA_API FNiagaraStatelessRangeColor CalculateRange(const FLinearColor& Default = FLinearColor::White) const;
+
+#if WITH_EDITORONLY_DATA
+	virtual bool AllowCurves() const override { return false; }
+	virtual bool DisplayAsColor() const { return true; }
+	virtual int32 GetBaseNumberOfChannels() const override { return 4; }
+	virtual void UpdateValuesFromDistribution() override;
+	virtual FNiagaraTypeDefinition GetBindingTypeDef() const { return FNiagaraTypeDefinition::GetColorDef(); }
+#endif
+};
+
+USTRUCT()
+struct FNiagaraDistributionFloat : public FNiagaraDistributionBase
+{
+	GENERATED_BODY()
+
+	FNiagaraDistributionFloat() = default;
+	explicit FNiagaraDistributionFloat(float ConstantValue) { InitConstant(ConstantValue); }
+	explicit FNiagaraDistributionFloat(std::initializer_list<float> CurvePoints) { InitCurve(CurvePoints); }
+
+	UPROPERTY(EditAnywhere, Category = "Parameters")
+	TArray<float> Values;
+
+	UPROPERTY(EditAnywhere, Category = "Parameters")
+	FVector2f ValuesTimeRange = FVector2f(0.0f, 1.0f);
+
+	NIAGARA_API void InitConstant(float Value);
+	NIAGARA_API void InitCurve(std::initializer_list<float> CurvePoints);
+#if WITH_EDITORONLY_DATA
+	NIAGARA_API void InitCurve(const TArray<FRichCurveKey>& CurveKeys);
+#endif
+	NIAGARA_API FNiagaraStatelessRangeFloat CalculateRange(const float Default = 0.0f) const;
+
+#if WITH_EDITORONLY_DATA
+	bool operator==(const FNiagaraDistributionFloat& Other) const
+	{
+		return (FNiagaraDistributionBase)*this == (FNiagaraDistributionBase)Other && Values == Other.Values;
+	}
+	virtual int32 GetBaseNumberOfChannels() const override { return 1; }
+	NIAGARA_API virtual void UpdateValuesFromDistribution() override;
+	virtual FNiagaraTypeDefinition GetBindingTypeDef() const { return FNiagaraTypeDefinition::GetFloatDef(); }
+#endif
+};
+
+USTRUCT()
+struct FNiagaraDistributionVector2 : public FNiagaraDistributionBase
+{
+	GENERATED_BODY()
+
+	FNiagaraDistributionVector2() = default;
+	explicit FNiagaraDistributionVector2(const float ConstantValue) { InitConstant(ConstantValue); }
+	explicit FNiagaraDistributionVector2(const FVector2f& ConstantValue) { InitConstant(ConstantValue); }
+
+	UPROPERTY(EditAnywhere, Category = "Parameters")
+	TArray<FVector2f> Values;
+
+	UPROPERTY(EditAnywhere, Category = "Parameters")
+	FVector2f ValuesTimeRange = FVector2f(0.0f, 1.0f);
+
+	NIAGARA_API void InitConstant(const float Value);
+	NIAGARA_API void InitConstant(const FVector2f& Value);
+	NIAGARA_API FNiagaraStatelessRangeVector2 CalculateRange(const FVector2f& Default = FVector2f::ZeroVector) const;
+
+#if WITH_EDITORONLY_DATA
+	virtual int32 GetBaseNumberOfChannels() const override { return 2; }
+	NIAGARA_API virtual void UpdateValuesFromDistribution() override;
+	virtual FNiagaraTypeDefinition GetBindingTypeDef() const { return FNiagaraTypeDefinition::GetVec2Def(); }
+#endif
+};
+
+USTRUCT()
+struct FNiagaraDistributionVector3 : public FNiagaraDistributionBase
+{
+	GENERATED_BODY()
+
+	FNiagaraDistributionVector3() = default;
+	explicit FNiagaraDistributionVector3(const float ConstantValue) { InitConstant(ConstantValue); }
+	explicit FNiagaraDistributionVector3(const FVector3f& ConstantValue) { InitConstant(ConstantValue); }
+	explicit FNiagaraDistributionVector3(std::initializer_list<float> CurvePoints) { InitCurve(CurvePoints); }
+	explicit FNiagaraDistributionVector3(std::initializer_list<FVector3f> CurvePoints) { InitCurve(CurvePoints); }
+
+	UPROPERTY(EditAnywhere, Category = "Parameters")
+	TArray<FVector3f> Values;
+
+	UPROPERTY(EditAnywhere, Category = "Parameters")
+	FVector2f ValuesTimeRange = FVector2f(0.0f, 1.0f);
+
+	NIAGARA_API void InitConstant(const float Value);
+	NIAGARA_API void InitConstant(const FVector3f& Value);
+	NIAGARA_API void InitCurve(std::initializer_list<float> CurvePoints);
+	NIAGARA_API void InitCurve(std::initializer_list<FVector3f> CurvePoints);
+	NIAGARA_API FNiagaraStatelessRangeVector3 CalculateRange(const FVector3f& Default = FVector3f::ZeroVector) const;
+
+#if WITH_EDITORONLY_DATA
+	virtual int32 GetBaseNumberOfChannels() const override { return 3; }
+	NIAGARA_API virtual void UpdateValuesFromDistribution() override;
+	virtual FNiagaraTypeDefinition GetBindingTypeDef() const { return FNiagaraTypeDefinition::GetVec3Def(); }
+#endif
+};
+
+USTRUCT()
+struct FNiagaraDistributionPosition : public FNiagaraDistributionVector3
+{
+	GENERATED_BODY()
+
+	FNiagaraDistributionPosition() = default;
+	explicit FNiagaraDistributionPosition(const float ConstantValue) { InitConstant(ConstantValue); }
+	explicit FNiagaraDistributionPosition(const FVector3f& ConstantValue) { InitConstant(ConstantValue); }
+
+#if WITH_EDITORONLY_DATA
+	virtual FNiagaraTypeDefinition GetBindingTypeDef() const { return FNiagaraTypeDefinition::GetPositionDef(); }
+#endif
+};
+
+USTRUCT()
+struct FNiagaraDistributionColor : public FNiagaraDistributionBase
+{
+	GENERATED_BODY()
+
+	FNiagaraDistributionColor() = default;
+	explicit FNiagaraDistributionColor(const FLinearColor& ConstantValue) { InitConstant(ConstantValue); }
+
+	UPROPERTY(EditAnywhere, Category = "Parameters")
+	TArray<FLinearColor> Values;
+
+	UPROPERTY(EditAnywhere, Category = "Parameters")
+	FVector2f ValuesTimeRange = FVector2f(0.0f, 1.0f);
+
+	NIAGARA_API void InitConstant(const FLinearColor& Value);
+	NIAGARA_API FNiagaraStatelessRangeColor CalculateRange(const FLinearColor& Default = FLinearColor::White) const;
+
+#if WITH_EDITORONLY_DATA
+	virtual bool DisplayAsColor() const override { return true; }
+	virtual int32 GetBaseNumberOfChannels() const override { return 4; }
+	NIAGARA_API virtual void UpdateValuesFromDistribution() override;
+	virtual FNiagaraTypeDefinition GetBindingTypeDef() const { return FNiagaraTypeDefinition::GetColorDef(); }
+#endif
+};
+
+template<>
+struct TStructOpsTypeTraits<FNiagaraDistributionRangeFloat> : public TStructOpsTypeTraitsBase2<FNiagaraDistributionRangeFloat>
+{
+	enum
+	{
+		WithStructuredSerializeFromMismatchedTag = true,
+	};
+};
+
+template<>
+struct TStructOpsTypeTraits<FNiagaraDistributionRangeVector2> : public TStructOpsTypeTraitsBase2<FNiagaraDistributionRangeVector2>
+{
+	enum
+	{
+		WithStructuredSerializeFromMismatchedTag = true,
+	};
+};
+
+template<>
+struct TStructOpsTypeTraits<FNiagaraDistributionRangeVector3> : public TStructOpsTypeTraitsBase2<FNiagaraDistributionRangeVector3>
+{
+	enum
+	{
+		WithStructuredSerializeFromMismatchedTag = true,
+	};
+};
+
diff --git a/Engine/Plugins/FX/Niagara/Source/Niagara/Public/Stateless/NiagaraStatelessRange.h b/Engine/Plugins/FX/Niagara/Source/Niagara/Public/Stateless/NiagaraStatelessRange.h
new file mode 100644
--- /dev/null
+++ b/Engine/Plugins/FX/Niagara/Source/Niagara/Public/Stateless/NiagaraStatelessRange.h
@@ -0,0 +1,37 @@
+// Copyright Epic Games, Inc. All Rights Reserved.
+
+#pragma once
+
+#include "CoreMinimal.h"
+
+template<typename TType> struct FNiagaraStatelessRangeDefaultValue {};
+template<> struct FNiagaraStatelessRangeDefaultValue<int32> { static constexpr int32 Zero() { return 0; } };
+template<> struct FNiagaraStatelessRangeDefaultValue<float> { static constexpr float Zero() { return 0.0f; } };
+template<> struct FNiagaraStatelessRangeDefaultValue<FVector2f> { static const FVector2f Zero() { return FVector2f(0.0f, 0.0f); } };
+template<> struct FNiagaraStatelessRangeDefaultValue<FVector3f> { static const FVector3f Zero() { return FVector3f(0.0f, 0.0f, 0.0f); } };
+template<> struct FNiagaraStatelessRangeDefaultValue<FVector4f> { static const FVector4f Zero() { return FVector4f(0.0f, 0.0f, 0.0f, 0.0f); } };
+template<> struct FNiagaraStatelessRangeDefaultValue<FLinearColor> { static const FLinearColor Zero() { return FLinearColor(0.0f, 0.0f, 0.0f, 0.0f); } };
+
+template<typename TType>
+struct FNiagaraStatelessRange
+{
+	using ValueType = TType;
+
+	FNiagaraStatelessRange() = default;
+	FNiagaraStatelessRange(const FNiagaraStatelessRange& Other) = default;
+	explicit FNiagaraStatelessRange(const ValueType& InMinMax) : Min(InMinMax), Max(InMinMax) {}
+	explicit FNiagaraStatelessRange(const ValueType& InMin, const ValueType& InMax) : Min(InMin), Max(InMax) {}
+
+	ValueType GetScale() const { return Max - Min; }
+
+	int32		ParameterOffset = INDEX_NONE;
+	ValueType	Min = FNiagaraStatelessRangeDefaultValue<TType>::Zero();
+	ValueType	Max = FNiagaraStatelessRangeDefaultValue<TType>::Zero();
+};
+
+using FNiagaraStatelessRangeInt = FNiagaraStatelessRange<int32>;
+using FNiagaraStatelessRangeFloat = FNiagaraStatelessRange<float>;
+using FNiagaraStatelessRangeVector2 = FNiagaraStatelessRange<FVector2f>;
+using FNiagaraStatelessRangeVector3 = FNiagaraStatelessRange<FVector3f>;
+using FNiagaraStatelessRangeVector4 = FNiagaraStatelessRange<FVector4f>;
+using FNiagaraStatelessRangeColor = FNiagaraStatelessRange<FLinearColor>;
