- IWYU fix

#rnx

diff --git a/Engine/Plugins/FX/Niagara/Source/Niagara/Internal/Stateless/NiagaraStatelessParticleSimContext.h b/Engine/Plugins/FX/Niagara/Source/Niagara/Internal/Stateless/NiagaraStatelessParticleSimContext.h
--- a/Engine/Plugins/FX/Niagara/Source/Niagara/Internal/Stateless/NiagaraStatelessParticleSimContext.h
+++ b/Engine/Plugins/FX/Niagara/Source/Niagara/Internal/Stateless/NiagaraStatelessParticleSimContext.h
@@ -2,7 +2,8 @@
 
 #pragma once
 
-#include "NiagaraStatelessCommon.h"
+#include "Stateless/NiagaraStatelessCommon.h"
+#include "Stateless/NiagaraStatelessRange.h"
 #include "NiagaraStatelessBuiltDistribution.h"
 
 #include "ShaderParameterMacros.h"
