error out if we have a baked svt with no frames

diff --git a/Engine/Plugins/FX/Niagara/Source/NiagaraEditor/Private/NiagaraBakerRenderer.cpp b/Engine/Plugins/FX/Niagara/Source/NiagaraEditor/Private/NiagaraBakerRenderer.cpp
--- a/Engine/Plugins/FX/Niagara/Source/NiagaraEditor/Private/NiagaraBakerRenderer.cpp
+++ b/Engine/Plugins/FX/Niagara/Source/NiagaraEditor/Private/NiagaraBakerRenderer.cpp
@@ -691,6 +691,13 @@ void FNiagaraBakerRenderer::RenderSparseVolumeTexture(UTextureRenderTarget2D* Re
 		SVTPreviewComponent->OverrideMaterials[0]->GetMaterial()->SetSparseVolumeTextureParameterValueEditorOnly("SparseVolumeTexture", SVT);
 	}
 	
+	int32 NumFrames = SVT->GetNumFrames();
+	if (NumFrames == 0)
+	{
+		UE_LOG(LogNiagaraEditor, Warning, TEXT("No frames in output SVT"));
+		return;
+	}
+
 	// do mod if we are previewing the looped result and the timeline doesn't line up with the baked number of frames for the looped
 	// sim.  This will at least loop it, but there could be a frame pop when the timeline resets
 	SVTPreviewComponent->SetFrame(Indices.FrameIndexA % SVT->GetNumFrames());
