- Downgrade ensure to warning in the log

#rnx

diff --git a/Engine/Plugins/FX/Niagara/Source/NiagaraEditor/Private/NiagaraNodeFunctionCall.cpp b/Engine/Plugins/FX/Niagara/Source/NiagaraEditor/Private/NiagaraNodeFunctionCall.cpp
--- a/Engine/Plugins/FX/Niagara/Source/NiagaraEditor/Private/NiagaraNodeFunctionCall.cpp
+++ b/Engine/Plugins/FX/Niagara/Source/NiagaraEditor/Private/NiagaraNodeFunctionCall.cpp
@@ -434,10 +434,14 @@ void UNiagaraNodeFunctionCall::AllocateDefaultPins()
 		TArray<FNiagaraVariable> SwitchNodeInputs = Graph->FindStaticSwitchInputs();
 		for (FNiagaraVariable& Input : SwitchNodeInputs)
 		{
-			if (ensureMsgf(Input.IsValid(), TEXT("Static switch input %s on function call node %s is invalid"), *Input.GetName().ToString(), *GetPathName()))
+			if (Input.IsValid())
 			{
 				AddStaticSwitchInputPin(Input);
 			}
+			else
+			{
+				UE_LOG(LogNiagaraEditor, Warning, TEXT("Static switch input %s on function call node %s is invalid"), *Input.GetName().ToString(), *GetPathName());
+			}
 		}
 
 		for (FNiagaraVariable& Output : Outputs)
