Providing exmaple of how a custom plugin for a Niagara test suite could look like.

#rnx

diff --git a/Engine/Plugins/FX/Niagara/Source/NiagaraEditor/Internal/NiagaraStackEntryEnumerable.h b/Engine/Plugins/FX/Niagara/Source/NiagaraEditor/Internal/NiagaraStackEntryEnumerable.h
new file mode 100644
--- /dev/null
+++ b/Engine/Plugins/FX/Niagara/Source/NiagaraEditor/Internal/NiagaraStackEntryEnumerable.h
@@ -0,0 +1,293 @@
+// Copyright Epic Games, Inc. All Rights Reserved.
+
+#pragma once
+
+#include "NiagaraDataInterface.h"
+#include "Containers/Array.h"
+#include "Delegates/Delegate.h"
+#include "ViewModels/Stack/NiagaraStackEntry.h"
+
+template<typename TNiagaraStackEntry>
+class INiagaraStackEntryIterator
+{
+public:
+	virtual ~INiagaraStackEntryIterator()
+	{
+	}
+
+	virtual bool IsValid() const = 0;
+	virtual void MoveNext() = 0;
+	virtual TNiagaraStackEntry* GetCurrent() const = 0;
+};
+
+template<typename TNiagaraStackEntry>
+class TNiagaraStackEntryNullIterator : public INiagaraStackEntryIterator<TNiagaraStackEntry>
+{
+public:
+	TNiagaraStackEntryNullIterator()
+	{
+	}
+
+	virtual bool IsValid() const override { return false; }
+	virtual void MoveNext() override { }
+	virtual TNiagaraStackEntry* GetCurrent() const override { return nullptr; }
+};
+
+template<typename TNiagaraStackEntry>
+class TNiagaraStackEntryArrayIterator : public INiagaraStackEntryIterator<TNiagaraStackEntry>
+{
+public:
+	TNiagaraStackEntryArrayIterator(const TArray<TNiagaraStackEntry*>& InArrayEntries)
+		: ArrayEntries(InArrayEntries)
+		, ArrayIndex(0)
+	{
+	}
+
+	virtual bool IsValid() const override
+	{
+		return ArrayIndex >= 0 && ArrayIndex < ArrayEntries.Num();
+	}
+
+	virtual void MoveNext() override
+	{
+		ArrayIndex++;
+	}
+
+	virtual TNiagaraStackEntry* GetCurrent() const override
+	{
+		return ArrayEntries[ArrayIndex];
+	}
+
+private:
+	TArray<TNiagaraStackEntry*> ArrayEntries;
+	int32 ArrayIndex = INDEX_NONE;
+};
+
+template<typename TNiagaraStackEntry>
+class TNiagaraStackEntryPredicateIterator : public INiagaraStackEntryIterator<TNiagaraStackEntry>
+{
+public:
+	DECLARE_DELEGATE_RetVal_OneParam(bool, FEntryPredicate, TNiagaraStackEntry* /* Entry */);
+
+public:
+	TNiagaraStackEntryPredicateIterator(TSharedRef<INiagaraStackEntryIterator<TNiagaraStackEntry>> InIterator, FEntryPredicate InPredicate)
+		: Iterator(InIterator)
+		, Predicate(InPredicate)
+	{
+		if (Iterator->IsValid())
+		{
+			if (Predicate.Execute(Iterator->GetCurrent()) == false)
+			{
+				MoveToNextValidEntry();
+			}
+		}
+	}
+
+	virtual bool IsValid() const override
+	{
+		return Iterator->IsValid();
+	}
+
+	virtual void MoveNext() override
+	{
+		MoveToNextValidEntry();
+	}
+
+	virtual TNiagaraStackEntry* GetCurrent() const override
+	{
+		return Iterator->GetCurrent();
+	}
+
+private:
+	void MoveToNextValidEntry()
+	{
+		Iterator->MoveNext();
+		while (Iterator->IsValid() && Predicate.Execute(Iterator->GetCurrent()) == false)
+		{
+			Iterator->MoveNext();
+		}
+	}
+
+private:
+	TSharedRef<INiagaraStackEntryIterator<TNiagaraStackEntry>> Iterator;
+	FEntryPredicate Predicate;
+};
+
+template<typename TNiagaraStackEntrySource, typename TNiagaraStackEntryTarget>
+class TNiagaraStackEntryOfTypeIterator : public INiagaraStackEntryIterator<TNiagaraStackEntryTarget>
+{
+public:
+	TNiagaraStackEntryOfTypeIterator(TSharedRef<INiagaraStackEntryIterator<TNiagaraStackEntrySource>> InSourceEntryIterator)
+		: PredicateIterator(MakeShared<TNiagaraStackEntryPredicateIterator<TNiagaraStackEntrySource>>(InSourceEntryIterator,
+			TNiagaraStackEntryPredicateIterator<TNiagaraStackEntrySource>::FEntryPredicate::CreateStatic(&TNiagaraStackEntryOfTypeIterator<TNiagaraStackEntrySource, TNiagaraStackEntryTarget>::IsOfType)))
+	{
+	}
+
+	virtual bool IsValid() const override
+	{
+		return PredicateIterator->IsValid();
+	}
+
+	virtual void MoveNext() override
+	{
+		PredicateIterator->MoveNext();
+	}
+
+	virtual TNiagaraStackEntryTarget* GetCurrent() const override
+	{
+		return CastChecked<TNiagaraStackEntryTarget>(PredicateIterator->GetCurrent());
+	}
+
+private:
+	static bool IsOfType(TNiagaraStackEntrySource* Entry)
+	{
+		return Entry->template IsA<TNiagaraStackEntryTarget>();
+	}
+
+private:
+	TSharedRef<INiagaraStackEntryIterator<TNiagaraStackEntrySource>> PredicateIterator;
+};
+
+template<typename TNiagaraStackEntry>
+class TNiagaraStackEntryChildrenIterator : public INiagaraStackEntryIterator<UNiagaraStackEntry>
+{
+public:
+	TNiagaraStackEntryChildrenIterator(TSharedRef<INiagaraStackEntryIterator<TNiagaraStackEntry>> InIterator)
+		: Iterator(InIterator)
+	{
+		if (Iterator->IsValid())
+		{
+			MoveNext();
+		}
+	}
+
+	virtual bool IsValid() const override
+	{
+		return Iterator->IsValid() && CurrentChildEntryIndex >= 0 && CurrentChildEntryIndex < CurrentChildEntries.Num();
+	}
+
+	virtual void MoveNext() override
+	{
+		if (CurrentChildEntryIndex == INDEX_NONE)
+		{
+			Iterator->GetCurrent()->GetFilteredChildren(CurrentChildEntries);
+			CurrentChildEntryIndex = 0;
+		}
+		else
+		{
+			CurrentChildEntryIndex++;
+		}
+
+		while (Iterator->IsValid() && CurrentChildEntryIndex >= CurrentChildEntries.Num())
+		{
+			Iterator->MoveNext();
+			if (Iterator->IsValid())
+			{
+				Iterator->GetCurrent()->GetFilteredChildren(CurrentChildEntries);
+				CurrentChildEntryIndex = 0;
+			}
+		}
+	}
+
+	virtual UNiagaraStackEntry* GetCurrent() const override
+	{
+		return CurrentChildEntries[CurrentChildEntryIndex];
+	}
+
+private:
+	TSharedRef<INiagaraStackEntryIterator<TNiagaraStackEntry>> Iterator;
+	TArray<UNiagaraStackEntry*> CurrentChildEntries;
+	int32 CurrentChildEntryIndex = INDEX_NONE;
+};
+
+template<typename TNiagaraStackEntry>
+class TNiagaraStackEntryEnumerable
+{
+public:
+	class FIterator
+	{
+	public:
+		FIterator(const TNiagaraStackEntryEnumerable<TNiagaraStackEntry>& InOwner, const TSharedRef<INiagaraStackEntryIterator<TNiagaraStackEntry>>& InIterator)
+			: Owner(&InOwner)
+			, Iterator(InIterator)
+		{
+		}
+
+		bool operator!=(const FIterator& Other)
+		{
+			return !(Owner == Other.Owner && Iterator->IsValid() == false && Other.Iterator->IsValid() == false);
+		}
+
+		FIterator& operator++()
+		{
+			Iterator->MoveNext();
+			return *this;
+		}
+
+		TNiagaraStackEntry* operator*() const
+		{
+			return Iterator->GetCurrent();
+		}
+
+	private:
+		const TNiagaraStackEntryEnumerable<TNiagaraStackEntry>* Owner;
+		TSharedRef<INiagaraStackEntryIterator<TNiagaraStackEntry>> Iterator;
+	};
+
+	TNiagaraStackEntryEnumerable(TNiagaraStackEntry& StackEntry)
+		: StackEntryIterator(MakeShared<TNiagaraStackEntryArrayIterator<TNiagaraStackEntry>>(TArray<TNiagaraStackEntry*>({ &StackEntry })))
+	{
+	}
+
+	TNiagaraStackEntryEnumerable(TSharedRef<INiagaraStackEntryIterator<TNiagaraStackEntry>> InStackEntryIterator)
+		: StackEntryIterator(InStackEntryIterator)
+	{
+	}
+
+	FIterator begin() const
+	{
+		return FIterator(*this, StackEntryIterator);
+	}
+
+	FIterator end() const
+	{
+		TSharedRef<TNiagaraStackEntryNullIterator<TNiagaraStackEntry>> NullIterator = MakeShared<TNiagaraStackEntryNullIterator<TNiagaraStackEntry>>();
+		return FIterator(*this, NullIterator);
+	}
+
+	TNiagaraStackEntryEnumerable<UNiagaraStackEntry> Children() const
+	{
+		return TNiagaraStackEntryEnumerable<UNiagaraStackEntry>(MakeShared<TNiagaraStackEntryChildrenIterator<TNiagaraStackEntry>>(StackEntryIterator));
+	}
+
+	template<typename TNiagaraStackEntryTarget>
+	TNiagaraStackEntryEnumerable<TNiagaraStackEntryTarget> OfType() const
+	{
+		return TNiagaraStackEntryEnumerable<TNiagaraStackEntryTarget>(MakeShared<TNiagaraStackEntryOfTypeIterator<TNiagaraStackEntry, TNiagaraStackEntryTarget>>(StackEntryIterator));
+	}
+
+	template<typename TPredicate>
+	TNiagaraStackEntryEnumerable<TNiagaraStackEntry> Where(TPredicate Predicate) const
+	{
+		return TNiagaraStackEntryEnumerable<TNiagaraStackEntry>(MakeShared<TNiagaraStackEntryPredicateIterator<TNiagaraStackEntry>>(StackEntryIterator,
+			TNiagaraStackEntryPredicateIterator<TNiagaraStackEntry>::FEntryPredicate::CreateLambda([Predicate](TNiagaraStackEntry* Entry) { return Predicate(Entry); })));
+	}
+
+	TNiagaraStackEntry* First() const
+	{
+		for (TNiagaraStackEntry* StackEntry : *this)
+		{
+			return StackEntry;
+		}
+		return nullptr;
+	}
+
+	TArray<TNiagaraStackEntry*> ToArray() const
+	{
+		TArray<TNiagaraStackEntry*> Array;
+		Array.Append(*this);
+	}
+
+private:
+	TSharedRef<INiagaraStackEntryIterator<TNiagaraStackEntry>> StackEntryIterator;
+};
\ No newline at end of file
diff --git a/Engine/Plugins/FX/Niagara/Source/NiagaraEditor/Internal/NiagaraStackQuery.cpp b/Engine/Plugins/FX/Niagara/Source/NiagaraEditor/Internal/NiagaraStackQuery.cpp
new file mode 100644
--- /dev/null
+++ b/Engine/Plugins/FX/Niagara/Source/NiagaraEditor/Internal/NiagaraStackQuery.cpp
@@ -0,0 +1,185 @@
+// Copyright Epic Games, Inc. All Rights Reserved.
+
+#include "NiagaraStackQuery.h"
+
+#include "IDetailTreeNode.h"
+#include "NiagaraStackEntryEnumerable.h"
+#include "PropertyHandle.h"
+#include "ViewModels/NiagaraEmitterHandleViewModel.h"
+#include "ViewModels/NiagaraEmitterViewModel.h"
+#include "ViewModels/NiagaraSystemViewModel.h"
+#include "ViewModels/Stack/NiagaraStackEntry.h"
+#include "ViewModels/Stack/NiagaraStackInputCategory.h"
+#include "ViewModels/Stack/NiagaraStackFunctionInput.h"
+#include "ViewModels/Stack/NiagaraStackValueCollection.h"
+#include "ViewModels/Stack/NiagaraStackModuleItem.h"
+#include "ViewModels/Stack/NiagaraStackObject.h"
+#include "ViewModels/Stack/NiagaraStackPropertyRow.h"
+#include "ViewModels/Stack/NiagaraStackRoot.h"
+#include "ViewModels/Stack/NiagaraStackScriptItemGroup.h"
+#include "ViewModels/Stack/NiagaraStackViewModel.h"
+
+#define LOCTEXT_NAMESPACE "NiagaraStackQuery"
+
+FNiagaraStackPropertyRowQuery FNiagaraStackObjectQuery::FindPropertyRow(FName PropertyName) const
+{
+	if (GetEntry() == nullptr)
+	{
+		return FNiagaraStackPropertyRowQuery(GetErrorMessage());
+	}
+
+	UNiagaraStackPropertyRow* FoundPropertyRow = TNiagaraStackEntryEnumerable<UNiagaraStackObject>(*GetEntry())
+		.Children().OfType<UNiagaraStackPropertyRow>()
+		.Where([](UNiagaraStackPropertyRow* PropertyRow) { return PropertyRow->GetDetailTreeNode()->GetNodeType() == EDetailNodeType::Category; })
+		.Children().OfType<UNiagaraStackPropertyRow>()
+		.Where([PropertyName](UNiagaraStackPropertyRow* PropertyRow)
+		{
+			TSharedPtr<IPropertyHandle> PropertyHandle = PropertyRow->GetDetailTreeNode()->CreatePropertyHandle();
+			return PropertyHandle.IsValid() && PropertyHandle->GetProperty() != nullptr && PropertyHandle->GetProperty()->GetFName() == PropertyName;
+		})
+		.First();
+
+	return FoundPropertyRow != nullptr
+		? FNiagaraStackPropertyRowQuery(*FoundPropertyRow)
+		: FNiagaraStackPropertyRowQuery(FText::Format(LOCTEXT("PropertyRowFailFormat", "Failed to find property named {0}."), FText::FromName(PropertyName)));
+}
+
+FNiagaraStackObjectQuery FNiagaraStackFunctionInputQuery::FindObjectValue() const
+{
+	if (GetEntry() == nullptr)
+	{
+		return FNiagaraStackObjectQuery(GetErrorMessage());
+	}
+
+	UNiagaraStackObject* FoundObjectValue = TNiagaraStackEntryEnumerable<UNiagaraStackFunctionInput>(*GetEntry())
+		.Children().OfType<UNiagaraStackObject>()
+		.First();
+
+	return FoundObjectValue != nullptr
+		? FNiagaraStackObjectQuery(*FoundObjectValue)
+		: FNiagaraStackObjectQuery(LOCTEXT("ObjectValueFailFormat", "Failed to input object value."));
+}
+
+FNiagaraStackFunctionInputQuery FNiagaraStackModuleItemQuery::FindFunctionInput(FName InputName) const
+{
+	if (GetEntry() == nullptr)
+	{
+		return FNiagaraStackFunctionInputQuery(GetErrorMessage());
+	}
+
+	// Check each category for its inputs. This does not handle nested categories
+	UNiagaraStackFunctionInput* FoundInput = TNiagaraStackEntryEnumerable<UNiagaraStackModuleItem>(*GetEntry())
+		.Children().OfType<UNiagaraStackScriptHierarchyRoot>()
+		.Children().OfType<UNiagaraStackScriptHierarchyCategory>()
+		.Children().OfType<UNiagaraStackFunctionInput>()
+		.Where([InputName](UNiagaraStackFunctionInput* StackFunctionInput) { return StackFunctionInput->GetInputParameterHandle().GetName() == InputName;})
+		.First();
+
+	// Check children right below the root that don't belong to a category, is the input wasn't found
+	if(FoundInput == nullptr)
+	{
+		FoundInput = TNiagaraStackEntryEnumerable<UNiagaraStackModuleItem>(*GetEntry())
+		.Children().OfType<UNiagaraStackScriptHierarchyRoot>()
+		.Children().OfType<UNiagaraStackFunctionInput>()
+		.Where([InputName](UNiagaraStackFunctionInput* StackFunctionInput) { return StackFunctionInput->GetInputParameterHandle().GetName() == InputName;})
+		.First();
+	}
+	
+	return FoundInput != nullptr
+		? FNiagaraStackFunctionInputQuery(*FoundInput)
+		: FNiagaraStackFunctionInputQuery(FText::Format(LOCTEXT("FunctionInputFailFormat", "Failed to find input named {0}."), FText::FromName(InputName)));
+}
+
+FNiagaraStackModuleItemQuery FNiagaraStackScriptItemGroupQuery::FindSetParametersItem(FName ParameterName) const
+{
+	if (GetEntry() == nullptr)
+	{
+		return FNiagaraStackModuleItemQuery(GetErrorMessage());
+	}
+
+	UNiagaraStackModuleItem* FoundModuleItem = TNiagaraStackEntryEnumerable<UNiagaraStackScriptItemGroup>(*GetEntry())
+		.Children().OfType<UNiagaraStackModuleItem>()
+		.Where([ParameterName](UNiagaraStackModuleItem* ModuleItem)
+		{
+			UNiagaraNodeAssignment* AssignmentModuleNode = Cast<UNiagaraNodeAssignment>(&ModuleItem->GetModuleNode());
+			return AssignmentModuleNode != nullptr &&
+				AssignmentModuleNode->GetAssignmentTargets().ContainsByPredicate(
+					[ParameterName](const FNiagaraVariable& AssignmentTarget) { return AssignmentTarget.GetName() == ParameterName; });
+		})
+		.First();
+
+	return FoundModuleItem != nullptr
+		? FNiagaraStackModuleItemQuery(*FoundModuleItem)
+		: FNiagaraStackModuleItemQuery(FText::Format(LOCTEXT("SetParametersItemFailFormat", "Failed to find a set parameters module item with parameter {0}"), FText::FromName(ParameterName)));
+}
+
+FNiagaraStackModuleItemQuery FNiagaraStackScriptItemGroupQuery::FindModuleItem(const FString& ModuleName) const
+{
+	if (GetEntry() == nullptr)
+	{
+		return FNiagaraStackModuleItemQuery(GetErrorMessage());
+	}
+	UNiagaraStackModuleItem* FoundModuleItem = TNiagaraStackEntryEnumerable<UNiagaraStackScriptItemGroup>(*GetEntry())
+		.Children().OfType<UNiagaraStackModuleItem>()
+		.Where([ModuleName](UNiagaraStackModuleItem* ModuleItem)
+			{
+				return ModuleItem->GetModuleNode().GetFunctionName() == ModuleName;
+			})
+		.First();
+	return FoundModuleItem != nullptr
+		? FNiagaraStackModuleItemQuery(*FoundModuleItem)
+		: FNiagaraStackModuleItemQuery(FText::Format(LOCTEXT("ModuleItemFailFormat", "Failed to find a module named {0}"), FText::FromString(ModuleName)));
+}
+
+FNiagaraStackRootQuery FNiagaraStackRootQuery::SystemStackRootEntry(TSharedRef<FNiagaraSystemViewModel> SystemViewModel)
+{
+	UNiagaraStackViewModel* StackViewModel = SystemViewModel->GetSystemStackViewModel();
+	return 
+		StackViewModel != nullptr &&
+		StackViewModel->GetRootEntry() != nullptr &&
+		StackViewModel->GetRootEntry()->IsA<UNiagaraStackRoot>()
+			? FNiagaraStackRootQuery(*CastChecked<UNiagaraStackRoot>(StackViewModel->GetRootEntry()))
+			: FNiagaraStackRootQuery(LOCTEXT("SystemStackViewModelFail", "Failed to find system stack root entry."));
+}
+
+FNiagaraStackRootQuery FNiagaraStackRootQuery::EmitterStackRootEntry(TSharedRef<FNiagaraSystemViewModel> SystemViewModel, FName EmitterName)
+{
+	const TSharedRef<FNiagaraEmitterHandleViewModel>* EmitterHandleViewModelPtr = SystemViewModel->GetEmitterHandleViewModels().FindByPredicate(
+		[EmitterName](const TSharedRef<FNiagaraEmitterHandleViewModel>& EmitterHandleViewModel) { return EmitterHandleViewModel->GetName() == EmitterName; });
+	
+	UNiagaraStackViewModel* StackViewModel = EmitterHandleViewModelPtr != nullptr ? (*EmitterHandleViewModelPtr)->GetEmitterStackViewModel() : nullptr;
+	return 
+		StackViewModel != nullptr &&
+		StackViewModel->GetRootEntry() != nullptr &&
+		StackViewModel->GetRootEntry()->IsA<UNiagaraStackRoot>()
+			? FNiagaraStackRootQuery(*CastChecked<UNiagaraStackRoot>(StackViewModel->GetRootEntry()))
+			: FNiagaraStackRootQuery(FText::Format(LOCTEXT("EmitterStackViewModelFailFormat", "Failed to find emitter stack view model for emitter {0}."),
+				FText::FromName(EmitterName)));
+}
+
+FNiagaraStackScriptItemGroupQuery FNiagaraStackRootQuery::FindScriptGroup(ENiagaraScriptUsage ScriptUsage, FGuid ScriptUsageId) const
+{
+	if (GetEntry() == nullptr)
+	{
+		return FNiagaraStackScriptItemGroupQuery(GetErrorMessage());
+	}
+
+	UNiagaraStackScriptItemGroup* FoundGroup = TNiagaraStackEntryEnumerable<UNiagaraStackRoot>(*GetEntry())
+		.Children().OfType<UNiagaraStackScriptItemGroup>()
+		.Where([ScriptUsage, ScriptUsageId](UNiagaraStackScriptItemGroup* ScriptItemGroup)
+		{ 
+			return ScriptItemGroup->GetScriptUsage() == ScriptUsage && ScriptItemGroup->GetScriptUsageId() == ScriptUsageId; 
+		})
+		.First();
+
+	if (FoundGroup != nullptr)
+	{
+		return FNiagaraStackScriptItemGroupQuery(*FoundGroup);
+	}
+
+	static const UEnum* UsageEnum = StaticEnum<ENiagaraScriptUsage>();
+	return FNiagaraStackScriptItemGroupQuery(FText::Format(LOCTEXT("ScriptGroupEntryFailFormat", "Failed to find a script group with usage: {0} and id: {1}"),
+		UsageEnum->GetDisplayNameTextByValue((int64)ScriptUsage), FText::FromString(ScriptUsageId.ToString(EGuidFormats::DigitsWithHyphens))));
+}
+
+#undef LOCTEXT_NAMESPACE
diff --git a/Engine/Plugins/FX/Niagara/Source/NiagaraEditor/Internal/NiagaraStackQuery.h b/Engine/Plugins/FX/Niagara/Source/NiagaraEditor/Internal/NiagaraStackQuery.h
new file mode 100644
--- /dev/null
+++ b/Engine/Plugins/FX/Niagara/Source/NiagaraEditor/Internal/NiagaraStackQuery.h
@@ -0,0 +1,123 @@
+// Copyright Epic Games, Inc. All Rights Reserved.
+
+#pragma once
+
+#include "Internationalization/Text.h"
+#include "NiagaraCommon.h"
+#include "Templates/SharedPointer.h"
+#include "ViewModels/Stack/NiagaraStackEntry.h"
+
+class FNiagaraSystemViewModel;
+class UNiagaraStackEntry;
+class UNiagaraStackFunctionInput;
+class UNiagaraStackRoot;
+class UNiagaraStackScriptItemGroup;
+class UNiagaraStackModuleItem;
+class UNiagaraStackObject;
+class UNiagaraStackPropertyRow;
+
+template<typename TNiagaraStackEntry>
+struct TNiagaraStackQueryResult
+{
+	TNiagaraStackQueryResult(TNiagaraStackEntry* InStackEntry, FText InErrorMessage)
+		: StackEntry(InStackEntry)
+		, ErrorMessage(InErrorMessage)
+	{
+	}
+
+	bool IsValid() const 
+	{
+		return StackEntry != nullptr;
+	}
+
+	TNiagaraStackEntry* StackEntry;
+	FText ErrorMessage;
+};
+
+template<typename TNiagaraStackEntry>
+class TNiagaraStackQueryBase
+{
+public:
+	TNiagaraStackQueryBase(TNiagaraStackEntry& InStackEntry)
+		: StackEntry(&InStackEntry)
+	{
+	}
+
+	TNiagaraStackQueryBase(FText InErrorMessage)
+		: StackEntry(nullptr)
+		, ErrorMessage(InErrorMessage)
+	{
+		ensureMsgf(ErrorMessage.IsEmpty() == false, TEXT("Error message must not be empty."));
+	}
+
+	TNiagaraStackEntry* GetEntry() const { return StackEntry; }
+
+	FText GetErrorMessage() const { return ErrorMessage; }
+
+	TNiagaraStackQueryResult<TNiagaraStackEntry> GetResult() const
+	{
+		return TNiagaraStackQueryResult<TNiagaraStackEntry>(StackEntry, ErrorMessage);
+	}
+
+private:
+	TNiagaraStackEntry* StackEntry;
+	FText ErrorMessage;
+};
+
+class FNiagaraStackPropertyRowQuery : public TNiagaraStackQueryBase<UNiagaraStackPropertyRow>
+{
+public:
+	FNiagaraStackPropertyRowQuery(UNiagaraStackPropertyRow& InStackEntry) : TNiagaraStackQueryBase<UNiagaraStackPropertyRow>(InStackEntry) { }
+	FNiagaraStackPropertyRowQuery(FText InErrorMessage) : TNiagaraStackQueryBase<UNiagaraStackPropertyRow>(InErrorMessage) { }
+};
+
+class FNiagaraStackObjectQuery : public TNiagaraStackQueryBase<UNiagaraStackObject>
+{
+public:
+	FNiagaraStackObjectQuery(UNiagaraStackObject& InStackEntry) : TNiagaraStackQueryBase<UNiagaraStackObject>(InStackEntry) { }
+	FNiagaraStackObjectQuery(FText InErrorMessage) : TNiagaraStackQueryBase<UNiagaraStackObject>(InErrorMessage) { }
+
+	FNiagaraStackPropertyRowQuery NIAGARAEDITOR_API FindPropertyRow(FName PropertyName) const;
+};
+
+class FNiagaraStackFunctionInputQuery : public TNiagaraStackQueryBase<UNiagaraStackFunctionInput>
+{
+public:
+	FNiagaraStackFunctionInputQuery(UNiagaraStackFunctionInput& InStackEntry)	: TNiagaraStackQueryBase<UNiagaraStackFunctionInput>(InStackEntry) { }
+	FNiagaraStackFunctionInputQuery(FText InErrorMessage) : TNiagaraStackQueryBase<UNiagaraStackFunctionInput>(InErrorMessage) { }
+
+	FNiagaraStackObjectQuery NIAGARAEDITOR_API FindObjectValue() const;
+};
+
+class FNiagaraStackModuleItemQuery : public TNiagaraStackQueryBase<UNiagaraStackModuleItem>
+{
+public:
+	FNiagaraStackModuleItemQuery(UNiagaraStackModuleItem& InStackEntry) : TNiagaraStackQueryBase<UNiagaraStackModuleItem>(InStackEntry) { }
+	FNiagaraStackModuleItemQuery(FText InErrorMessage) : TNiagaraStackQueryBase<UNiagaraStackModuleItem>(InErrorMessage) { }
+
+	FNiagaraStackFunctionInputQuery NIAGARAEDITOR_API FindFunctionInput(FName InputName) const;
+};
+
+class FNiagaraStackScriptItemGroupQuery : public TNiagaraStackQueryBase<UNiagaraStackScriptItemGroup>
+{
+public:
+	FNiagaraStackScriptItemGroupQuery(UNiagaraStackScriptItemGroup& InStackEntry) : TNiagaraStackQueryBase<UNiagaraStackScriptItemGroup>(InStackEntry) { }
+	FNiagaraStackScriptItemGroupQuery(FText InErrorMessage) : TNiagaraStackQueryBase<UNiagaraStackScriptItemGroup>(InErrorMessage) { }
+
+	FNiagaraStackModuleItemQuery NIAGARAEDITOR_API FindSetParametersItem(FName ParameterName) const;
+	FNiagaraStackModuleItemQuery NIAGARAEDITOR_API FindModuleItem(const FString& ModuleName) const;
+};
+
+class FNiagaraStackRootQuery : public TNiagaraStackQueryBase<UNiagaraStackRoot>
+{
+public:
+	static FNiagaraStackRootQuery SystemStackRootEntry(TSharedRef<FNiagaraSystemViewModel> InSystemViewModel);
+
+	static FNiagaraStackRootQuery NIAGARAEDITOR_API EmitterStackRootEntry(TSharedRef<FNiagaraSystemViewModel> InSystemViewModel, FName EmitterName);
+
+	FNiagaraStackScriptItemGroupQuery NIAGARAEDITOR_API FindScriptGroup(ENiagaraScriptUsage ScriptUsage, FGuid ScriptUsageId) const;
+
+private:
+	FNiagaraStackRootQuery(UNiagaraStackRoot& InStackEntry) : TNiagaraStackQueryBase<UNiagaraStackRoot>(InStackEntry) { }
+	FNiagaraStackRootQuery(FText InErrorMessage) : TNiagaraStackQueryBase<UNiagaraStackRoot>(InErrorMessage) { }
+};
\ No newline at end of file
diff --git a/Engine/Plugins/FX/Niagara/Source/NiagaraEditor/NiagaraEditor.Build.cs b/Engine/Plugins/FX/Niagara/Source/NiagaraEditor/NiagaraEditor.Build.cs
--- a/Engine/Plugins/FX/Niagara/Source/NiagaraEditor/NiagaraEditor.Build.cs
+++ b/Engine/Plugins/FX/Niagara/Source/NiagaraEditor/NiagaraEditor.Build.cs
@@ -86,6 +86,7 @@ public class NiagaraEditor : ModuleRules
                 "UnrealEd",
                 "AppFramework",
 				"ToolWidgets",
+				"Sequencer",
             }
         );
 
