- Fix for color curve resetting when changing interpolation type

#rnx

diff --git a/Engine/Plugins/FX/Niagara/Source/NiagaraEditor/Private/Stateless/NiagaraDistributionPropertyCustomization.cpp b/Engine/Plugins/FX/Niagara/Source/NiagaraEditor/Private/Stateless/NiagaraDistributionPropertyCustomization.cpp
--- a/Engine/Plugins/FX/Niagara/Source/NiagaraEditor/Private/Stateless/NiagaraDistributionPropertyCustomization.cpp
+++ b/Engine/Plugins/FX/Niagara/Source/NiagaraEditor/Private/Stateless/NiagaraDistributionPropertyCustomization.cpp
@@ -194,7 +194,7 @@ public:
 		}
 	}
 
-	static ENiagaraDistributionEditorMode GetDistributionEditorModeFromSourceMode(int32 InSourceNumChannels, bool bInDisplayAsColor, ENiagaraDistributionMode InSourceMode)
+	static ENiagaraDistributionEditorMode GetDistributionEditorModeFromSourceMode(int32 InSourceNumChannels, bool bInDisplayAsColor, ENiagaraDistributionMode InSourceMode, TOptional<ENiagaraDistributionEditorMode> CurrentMode = {})
 	{
 		if (bInDisplayAsColor)
 		{
@@ -207,7 +207,7 @@ public:
 			case ENiagaraDistributionMode::NonUniformRange:
 				return ENiagaraDistributionEditorMode::ColorRange;
 			case ENiagaraDistributionMode::NonUniformCurve:
-				return ENiagaraDistributionEditorMode::ColorGradient;
+				return CurrentMode.IsSet() && CurrentMode.GetValue() == ENiagaraDistributionEditorMode::NonUniformCurve ? ENiagaraDistributionEditorMode::NonUniformCurve : ENiagaraDistributionEditorMode::ColorGradient;
 			default:
 				return ENiagaraDistributionEditorMode::ColorConstant;
 			}
@@ -618,7 +618,7 @@ private:
 	void DistributionPropertyChanged()
 	{
 		// Watch for editor mode changes from external sources like copy/paste.
-		ENiagaraDistributionEditorMode NewEditorMode = GetDistributionEditorModeFromSourceMode(SourceNumChannels, bDisplayAsColor, SourceDistribution->Mode);
+		const ENiagaraDistributionEditorMode NewEditorMode = GetDistributionEditorModeFromSourceMode(SourceNumChannels, bDisplayAsColor, SourceDistribution->Mode, EditorMode);
 		if (NewEditorMode != EditorMode && bUpdatingDistributionModeInternally == false)
 		{
 			EditorMode = NewEditorMode;
