- Fix for not being able to change isolation when the system is paused

#tests Create a system with multiple emitters, pause in the editor, toggle isolation states on / off

diff --git a/Engine/Plugins/FX/Niagara/Source/NiagaraEditor/Private/ViewModels/NiagaraSystemViewModel.cpp b/Engine/Plugins/FX/Niagara/Source/NiagaraEditor/Private/ViewModels/NiagaraSystemViewModel.cpp
--- a/Engine/Plugins/FX/Niagara/Source/NiagaraEditor/Private/ViewModels/NiagaraSystemViewModel.cpp
+++ b/Engine/Plugins/FX/Niagara/Source/NiagaraEditor/Private/ViewModels/NiagaraSystemViewModel.cpp
@@ -2095,7 +2095,7 @@ void FNiagaraSystemViewModel::IsolateEmitters(TArray<FGuid> EmitterHandlesIdsToI
 		}
 	}
 
-	GetSystem().SetIsolateEnabled(bAnyEmitterIsolated);
+	SetSystenIsolationEnabled(bAnyEmitterIsolated);
 }
 
 void FNiagaraSystemViewModel::CacheIsolatedEmitterState()
@@ -2116,6 +2116,21 @@ void FNiagaraSystemViewModel::RestoreIsolatedEmitterState()
 	CachedEmitterIsolationState.Empty();
 }
 
+void FNiagaraSystemViewModel::SetSystenIsolationEnabled(bool bIsolationEnabled)
+{
+	UNiagaraSystem* NiagaraSystem = &GetSystem();
+	NiagaraSystem->SetIsolateEnabled(bIsolationEnabled);
+
+	for (TObjectIterator<UNiagaraComponent> It; It; ++It)
+	{
+		UNiagaraComponent* Comp = *It;
+		if (::IsValid(Comp) && Comp->IsRegistered() && Comp->GetAsset() == NiagaraSystem)
+		{
+			Comp->MarkRenderStateDirty();
+		}
+	}
+}
+
 void FNiagaraSystemViewModel::DisableEmitters(TArray<FGuid> EmitterHandlesIdsToDisable)
 {
 	FScopedTransaction Transaction(LOCTEXT("DisableSelectedEmitters", "Disabled selected emitters"));
@@ -2156,7 +2171,7 @@ void FNiagaraSystemViewModel::ToggleEmitterIsolation(TSharedRef<FNiagaraEmitterH
 		}
 	}
 
-	GetSystem().SetIsolateEnabled(bAnyEmitterIsolated);
+	SetSystenIsolationEnabled(bAnyEmitterIsolated);
 }
 
 void FNiagaraSystemViewModel::ResetEmitterHandleViewModelsAndTracks()
@@ -2241,7 +2256,8 @@ void FNiagaraSystemViewModel::RefreshEmitterHandleViewModels()
 			bAnyEmitterIsolated = true;
 		}
 	}
-	GetSystem().SetIsolateEnabled(bAnyEmitterIsolated);
+	SetSystenIsolationEnabled(bAnyEmitterIsolated);
+
 	OnEmitterHandleViewModelsChangedDelegate.Broadcast();
 }
 
diff --git a/Engine/Plugins/FX/Niagara/Source/NiagaraEditor/Public/ViewModels/NiagaraSystemViewModel.h b/Engine/Plugins/FX/Niagara/Source/NiagaraEditor/Public/ViewModels/NiagaraSystemViewModel.h
--- a/Engine/Plugins/FX/Niagara/Source/NiagaraEditor/Public/ViewModels/NiagaraSystemViewModel.h
+++ b/Engine/Plugins/FX/Niagara/Source/NiagaraEditor/Public/ViewModels/NiagaraSystemViewModel.h
@@ -356,7 +356,9 @@ public:
 
 	void CacheIsolatedEmitterState();
 	void RestoreIsolatedEmitterState();
-	
+
+	void SetSystenIsolationEnabled(bool bIsolationEnabled);
+
 	/** Disable the supplied emitters. */
 	NIAGARAEDITOR_API void DisableEmitters(TArray<FGuid> EmitterHandlesIdsToDisable);
 
