suppress warning that shows each frame for degenerate SVT output

diff --git a/Engine/Plugins/FX/Niagara/Source/NiagaraEditor/Private/NiagaraBakerRenderer.cpp b/Engine/Plugins/FX/Niagara/Source/NiagaraEditor/Private/NiagaraBakerRenderer.cpp
--- a/Engine/Plugins/FX/Niagara/Source/NiagaraEditor/Private/NiagaraBakerRenderer.cpp
+++ b/Engine/Plugins/FX/Niagara/Source/NiagaraEditor/Private/NiagaraBakerRenderer.cpp
@@ -693,8 +693,7 @@ void FNiagaraBakerRenderer::RenderSparseVolumeTexture(UTextureRenderTarget2D* Re
 	
 	int32 NumFrames = SVT->GetNumFrames();
 	if (NumFrames == 0)
-	{
-		UE_LOG(LogNiagaraEditor, Warning, TEXT("No frames in output SVT"));
+	{		
 		return;
 	}
 
