- For velocity aligned sprites make sure we never have Age as 0.0 or the first frame will be incorrect

#rnx

diff --git a/Engine/Plugins/FX/Niagara/Shaders/Private/Stateless/Modules/NiagaraStatelessModule_SolveVelocitiesAndForces.ush b/Engine/Plugins/FX/Niagara/Shaders/Private/Stateless/Modules/NiagaraStatelessModule_SolveVelocitiesAndForces.ush
--- a/Engine/Plugins/FX/Niagara/Shaders/Private/Stateless/Modules/NiagaraStatelessModule_SolveVelocitiesAndForces.ush
+++ b/Engine/Plugins/FX/Niagara/Shaders/Private/Stateless/Modules/NiagaraStatelessModule_SolveVelocitiesAndForces.ush
@@ -159,7 +159,8 @@ void SolveVelocitiesAndForces_Simulate(inout FStatelessParticle Particle)
 	}
 	
 	// Simulate
-	SolveVelocitiesAndForces_IntegratePosition(ModuleData, Particle.Age, Particle.Position);
+	// Note we need to ensure velocity aligned sprites work we need to make sure Age & PreviousAge != 0.0
+	SolveVelocitiesAndForces_IntegratePosition(ModuleData, max(Particle.Age, UE_KINDA_SMALL_NUMBER), Particle.Position);
 	SolveVelocitiesAndForces_IntegratePosition(ModuleData, Particle.PreviousAge, Particle.PreviousPosition);
 		
 	Particle.Velocity = (Particle.Position - Particle.PreviousPosition) * Particle.InvDeltaTime;
diff --git a/Engine/Plugins/FX/Niagara/Source/Niagara/Private/Stateless/Modules/NiagaraStatelessModule_SolveVelocitiesAndForces.cpp b/Engine/Plugins/FX/Niagara/Source/Niagara/Private/Stateless/Modules/NiagaraStatelessModule_SolveVelocitiesAndForces.cpp
--- a/Engine/Plugins/FX/Niagara/Source/Niagara/Private/Stateless/Modules/NiagaraStatelessModule_SolveVelocitiesAndForces.cpp
+++ b/Engine/Plugins/FX/Niagara/Source/Niagara/Private/Stateless/Modules/NiagaraStatelessModule_SolveVelocitiesAndForces.cpp
@@ -80,7 +80,7 @@ namespace NiagaraStatelessModuleSolveVelocitiesAndForcesPrivate
 				//-TODO:
 			}
 
-			Position += IntegratePosition(AgeData[i], Mass, Drag, InitialVelocity, Wind, Acceleration);
+			Position += IntegratePosition(FMath::Max(AgeData[i], UE_KINDA_SMALL_NUMBER), Mass, Drag, InitialVelocity, Wind, Acceleration);
 			PreviousPosition += IntegratePosition(PreviousAgeData[i], Mass, Drag, InitialVelocity, Wind, Acceleration);
 
 			ParticleSimulationContext.WriteParticleVariable(ModuleBuiltData->PositionVariableOffset, i, Position);
