- Remove passing in LWC tile to various data interface and use Engine.Owner.LWCTile instead

#rnx
#tests Ran ET / tested gbuffer sampling in LWC space

diff --git a/Engine/Plugins/Experimental/ChaosNiagara/Shaders/NiagaraDataInterfacePhysicsField.ush b/Engine/Plugins/Experimental/ChaosNiagara/Shaders/NiagaraDataInterfacePhysicsField.ush
--- a/Engine/Plugins/Experimental/ChaosNiagara/Shaders/NiagaraDataInterfacePhysicsField.ush
+++ b/Engine/Plugins/Experimental/ChaosNiagara/Shaders/NiagaraDataInterfacePhysicsField.ush
@@ -24,12 +24,11 @@ int				{ParameterName}_ClipmapExponent;
 int				{ParameterName}_ClipmapCount;
 int				{ParameterName}_TargetCount;
 int4			{ParameterName}_FieldTargets[MAX_PHYSICS_FIELD_TARGETS];
-float3			{ParameterName}_SystemLWCTile;
 
 float3 PhysicsVectorField_{ParameterName}(in float3 WorldPosition, in int TargetIndex, in bool bUseClipmap)
 {
 	// LWC_TODO: fix this when the physics field functions are lwc ready
-	FLWCVector3 SamplePos = MakeLWCVector3({ParameterName}_SystemLWCTile, WorldPosition);
+	FLWCVector3 SamplePos = MakeLWCVector3(GetEngineOwnerLWCTile(), WorldPosition);
 
 	if ( {ParameterName}_bClipmapAvailable && bUseClipmap )
 	{
@@ -48,7 +47,7 @@ float3 PhysicsVectorField_{ParameterName}(in float3 WorldPosition, in int Target
 float PhysicsScalarField_{ParameterName}(in float3 WorldPosition, in int TargetIndex, in bool bUseClipmap)
 {
 	// LWC_TODO: fix this when the physics field functions are lwc ready
-	FLWCVector3 SamplePos = MakeLWCVector3({ParameterName}_SystemLWCTile, WorldPosition);
+	FLWCVector3 SamplePos = MakeLWCVector3(GetEngineOwnerLWCTile(), WorldPosition);
 
 	if ( {ParameterName}_bClipmapAvailable && bUseClipmap )
 	{
@@ -67,7 +66,7 @@ float PhysicsScalarField_{ParameterName}(in float3 WorldPosition, in int TargetI
 int PhysicsIntegerField_{ParameterName}(in float3 WorldPosition, in int TargetIndex, in bool bUseClipmap)
 {
 	// LWC_TODO: fix this when the physics field functions are lwc ready
-	FLWCVector3 SamplePos = MakeLWCVector3({ParameterName}_SystemLWCTile, WorldPosition);
+	FLWCVector3 SamplePos = MakeLWCVector3(GetEngineOwnerLWCTile(), WorldPosition);
 
 	if ( {ParameterName}_bClipmapAvailable && bUseClipmap )
 	{
diff --git a/Engine/Plugins/Experimental/ChaosNiagara/Source/ChaosNiagara/Private/NiagaraDataInterfacePhysicsField.cpp b/Engine/Plugins/Experimental/ChaosNiagara/Source/ChaosNiagara/Private/NiagaraDataInterfacePhysicsField.cpp
--- a/Engine/Plugins/Experimental/ChaosNiagara/Source/ChaosNiagara/Private/NiagaraDataInterfacePhysicsField.cpp
+++ b/Engine/Plugins/Experimental/ChaosNiagara/Source/ChaosNiagara/Private/NiagaraDataInterfacePhysicsField.cpp
@@ -49,7 +49,6 @@ namespace NDIPhysicsFieldLocal
 		SHADER_PARAMETER(int, 				ClipmapCount)
 		SHADER_PARAMETER(int,				TargetCount)
 		SHADER_PARAMETER_ARRAY(FIntVector4, FieldTargets, [MAX_PHYSICS_FIELD_TARGETS])
-		SHADER_PARAMETER(FVector3f,			SystemLWCTile)
 	END_SHADER_PARAMETER_STRUCT()
 
 	static const TCHAR* TemplateShaderFilePath = TEXT("/Plugin/Experimental/ChaosNiagara/NiagaraDataInterfacePhysicsField.ush");
@@ -740,7 +739,6 @@ void UNiagaraDataInterfacePhysicsField::SetShaderParameters(const FNiagaraDataIn
 			ShaderInstanceParameters->FieldTargets[i] = FIntVector4(0, 0, 0, 0);
 		}
 	}
-	ShaderInstanceParameters->SystemLWCTile = Context.GetSystemLWCTile();
 }
 
 void UNiagaraDataInterfacePhysicsField::ProvidePerInstanceDataForRenderThread(void* DataForRenderThread, void* PerInstanceData, const FNiagaraSystemInstanceID& SystemInstance)
diff --git a/Engine/Plugins/FX/Niagara/Shaders/Private/NiagaraDataInterfaceAsyncGpuTraceTemplate.ush b/Engine/Plugins/FX/Niagara/Shaders/Private/NiagaraDataInterfaceAsyncGpuTraceTemplate.ush
--- a/Engine/Plugins/FX/Niagara/Shaders/Private/NiagaraDataInterfaceAsyncGpuTraceTemplate.ush
+++ b/Engine/Plugins/FX/Niagara/Shaders/Private/NiagaraDataInterfaceAsyncGpuTraceTemplate.ush
@@ -1,6 +1,5 @@
 // Copyright Epic Games, Inc. All Rights Reserved.
 
-float3											{ParameterName}_SystemLWCTile;
 /** Niagara ray tracing uses shared buffers between dispatches so each buffer has a buffer + and offset into that buffer for the current dispatch. */
 uint											{ParameterName}_MaxRayTraceCount;
 RWStructuredBuffer<FNiagaraAsyncGpuTrace>		{ParameterName}_RWRayRequests;
@@ -15,7 +14,7 @@ void IssueAsyncRayTraceGpu_{ParameterName}_UEImpureCall(bool bExecute, int In_Qu
 	Out_IsQueryValid = false;
 	if ( bExecute )
 	{
-		NDIAsyncGpuTrace_IssueAsyncRayTrace({ParameterName}_MaxRayTraceCount, {ParameterName}_RWRayRequests, {ParameterName}_RayRequestsOffset, {ParameterName}_RWRayTraceCounts, {ParameterName}_RayTraceCountsOffset, {ParameterName}_SystemLWCTile, In_QueryID, In_TraceStart, In_TraceEnd, In_CollisionGroup, Out_IsQueryValid);
+		NDIAsyncGpuTrace_IssueAsyncRayTrace({ParameterName}_MaxRayTraceCount, {ParameterName}_RWRayRequests, {ParameterName}_RayRequestsOffset, {ParameterName}_RWRayTraceCounts, {ParameterName}_RayTraceCountsOffset, GetEngineOwnerLWCTile(), In_QueryID, In_TraceStart, In_TraceEnd, In_CollisionGroup, Out_IsQueryValid);
 	}
 }
 
@@ -35,7 +34,7 @@ void CreateAsyncRayTraceGpu_{ParameterName}_UEImpureCall(bool bExecute, float3 I
 	Out_IsQueryValid	= false;
 	if (bExecute)
 	{
-		NDIAsyncGpuTrace_CreateAsyncRayTrace({ParameterName}_MaxRayTraceCount, {ParameterName}_RWRayRequests, {ParameterName}_RayRequestsOffset, {ParameterName}_RWRayTraceCounts, {ParameterName}_RayTraceCountsOffset, {ParameterName}_SystemLWCTile, In_TraceStart, In_TraceEnd, In_CollisionGroup, Out_QueryID, Out_IsQueryValid);
+		NDIAsyncGpuTrace_CreateAsyncRayTrace({ParameterName}_MaxRayTraceCount, {ParameterName}_RWRayRequests, {ParameterName}_RayRequestsOffset, {ParameterName}_RWRayTraceCounts, {ParameterName}_RayTraceCountsOffset, GetEngineOwnerLWCTile(), In_TraceStart, In_TraceEnd, In_CollisionGroup, Out_QueryID, Out_IsQueryValid);
 	}
 }
 
@@ -47,6 +46,6 @@ void ReadAsyncRayTraceGpu_{ParameterName}(bool bExecute, int In_PreviousFrameQue
 	Out_CollisionNormal		= float3(0.0f, 0.0f, 0.0f);
 	if (bExecute)
 	{
-		NDIAsyncGpuTrace_ReadAsyncRayTrace({ParameterName}_MaxRayTraceCount, {ParameterName}_IntersectionResults, {ParameterName}_IntersectionResultsOffset, {ParameterName}_SystemLWCTile, In_PreviousFrameQueryID, Out_CollisionValid, Out_CollisionDistance, Out_CollisionPosWorld, Out_CollisionNormal);
+		NDIAsyncGpuTrace_ReadAsyncRayTrace({ParameterName}_MaxRayTraceCount, {ParameterName}_IntersectionResults, {ParameterName}_IntersectionResultsOffset, GetEngineOwnerLWCTile(), In_PreviousFrameQueryID, Out_CollisionValid, Out_CollisionDistance, Out_CollisionPosWorld, Out_CollisionNormal);
 	}
 }
diff --git a/Engine/Plugins/FX/Niagara/Shaders/Private/NiagaraDataInterfaceCamera.ush b/Engine/Plugins/FX/Niagara/Shaders/Private/NiagaraDataInterfaceCamera.ush
--- a/Engine/Plugins/FX/Niagara/Shaders/Private/NiagaraDataInterfaceCamera.ush
+++ b/Engine/Plugins/FX/Niagara/Shaders/Private/NiagaraDataInterfaceCamera.ush
@@ -1,6 +1,5 @@
 // Copyright Epic Games, Inc. All Rights Reserved.
 
-float3 {ParameterName}_SystemLWCTile;
 int {ParameterName}_SplitscreenMode;
 
 void GetFieldOfView_{ParameterName}(out float Out_FieldOfViewAngle)
@@ -22,7 +21,7 @@ void GetTAAJitter_{ParameterName}(out float2 Out_CurrentJitter, out float2 Out_P
 void GetCameraPropertiesCPUASC47GPU_{ParameterName}(out float3 Out_CameraPositionWorld, out float3 Out_ViewForwardVector, out float3 Out_ViewUpVector, out float3 Out_ViewRightVector)
 {
 	FLWCVector3 CameraPos = PrimaryView.TileOffset.WorldCameraOrigin;
-	LWCSetTile(CameraPos, LWCGetTile(CameraPos) - {ParameterName}_SystemLWCTile); // convert to simulation space
+	LWCSetTile(CameraPos, LWCGetTile(CameraPos) - GetEngineOwnerLWCTile()); // convert to simulation space
 	Out_CameraPositionWorld = LWCToFloat(CameraPos);
 	Out_ViewForwardVector = View.ViewForward;
 	Out_ViewUpVector = View.ViewUp;
@@ -33,7 +32,7 @@ void GetViewPropertiesGPU_{ParameterName}(out float3 Out_ViewPositionWorld, out
 	out float4 Out_ScreenToViewSpace, out float2 Out_Current_TAAJitter, out float2 Out_Previous_TAAJitter, out float3 Out_PreViewTranslation, out float4 Out_BufferSizeAndInverseSize, out float2 Out_ViewportOffset)
 {
 	FLWCVector3 ViewPos = PrimaryView.TileOffset.WorldViewOrigin;
-	LWCSetTile(ViewPos, LWCGetTile(ViewPos) - {ParameterName}_SystemLWCTile); // convert to simulation space
+	LWCSetTile(ViewPos, LWCGetTile(ViewPos) - GetEngineOwnerLWCTile()); // convert to simulation space
 	Out_ViewPositionWorld.xyz = LWCToFloat(ViewPos);
 	Out_ViewForwardVector = View.ViewForward;
 	Out_ViewUpVector = View.ViewUp;
@@ -73,7 +72,7 @@ void GetViewSpaceTransformsGPU_{ParameterName}(out float4x4 Out_TranslatedWorldT
 
 void ApplyPreViewTranslationToPositionGPU_{ParameterName}(in float3 In_WorldPosition, out float4 Out_TranslatedPosition)
 {
-	FLWCVector3 LwcPos = MakeLWCVector3({ParameterName}_SystemLWCTile, In_WorldPosition);
+	FLWCVector3 LwcPos = MakeLWCVector3(GetEngineOwnerLWCTile(), In_WorldPosition);
 	float3 ResolvedLwcPos = LWCToFloat(LWCAdd(LwcPos, PrimaryView.TileOffset.PreViewTranslation));
 	Out_TranslatedPosition = float4(ResolvedLwcPos, 1);
 }
@@ -82,6 +81,6 @@ void RemovePreViewTranslationFromPositionGPU_{ParameterName}(in float4 In_Transl
 {
 	FLWCVector3 TranslatedPos = LWCPromote(In_TranslatedPosition.xyz / In_TranslatedPosition.w);
 	FLWCVector3 WorldPos = LWCSubtract(TranslatedPos, PrimaryView.TileOffset.PreViewTranslation);
-	WorldPos.Tile -= {ParameterName}_SystemLWCTile;
+	WorldPos.Tile -= GetEngineOwnerLWCTile();
 	Out_WorldPosition = LWCToFloat(WorldPos);
 }
diff --git a/Engine/Plugins/FX/Niagara/Shaders/Private/NiagaraDataInterfaceCollisionQueryTemplate.ush b/Engine/Plugins/FX/Niagara/Shaders/Private/NiagaraDataInterfaceCollisionQueryTemplate.ush
--- a/Engine/Plugins/FX/Niagara/Shaders/Private/NiagaraDataInterfaceCollisionQueryTemplate.ush
+++ b/Engine/Plugins/FX/Niagara/Shaders/Private/NiagaraDataInterfaceCollisionQueryTemplate.ush
@@ -1,25 +1,23 @@
 // Copyright Epic Games, Inc. All Rights Reserved.
 
-float3 {ParameterName}_SystemLWCTile;
-
 void QuerySceneDepthGPU_{ParameterName}(in float3 In_SamplePos, out float Out_SceneDepth, out float3 Out_CameraPosWorld, out bool Out_IsInsideView, out float3 Out_WorldPos, out float3 Out_WorldNormal)
 {
-	NDICollisionQuery_QuerySceneDepthGPU(In_SamplePos, {ParameterName}_SystemLWCTile, Out_SceneDepth, Out_CameraPosWorld, Out_IsInsideView, Out_WorldPos, Out_WorldNormal);
+	NDICollisionQuery_QuerySceneDepthGPU(In_SamplePos, GetEngineOwnerLWCTile(), Out_SceneDepth, Out_CameraPosWorld, Out_IsInsideView, Out_WorldPos, Out_WorldNormal);
 }
 
 void QueryScenePartialDepthGPU_{ParameterName}(in float3 In_SamplePos, out float Out_SceneDepth, out float3 Out_CameraPosWorld, out bool Out_IsInsideView, out float3 Out_WorldPos, out float3 Out_WorldNormal)
 {
-	NDICollisionQuery_QueryScenePartialDepthGPU(In_SamplePos, {ParameterName}_SystemLWCTile, Out_SceneDepth, Out_CameraPosWorld, Out_IsInsideView, Out_WorldPos, Out_WorldNormal);
+	NDICollisionQuery_QueryScenePartialDepthGPU(In_SamplePos, GetEngineOwnerLWCTile(), Out_SceneDepth, Out_CameraPosWorld, Out_IsInsideView, Out_WorldPos, Out_WorldNormal);
 }
 
 void QueryCustomDepthGPU_{ParameterName}(in float3 In_SamplePos, out float Out_SceneDepth, out float3 Out_CameraPosWorld, out bool Out_IsInsideView, out float3 Out_WorldPos, out float3 Out_WorldNormal)
 {
-	NDICollisionQuery_QueryCustomDepthGPU(In_SamplePos, {ParameterName}_SystemLWCTile, Out_SceneDepth, Out_CameraPosWorld, Out_IsInsideView, Out_WorldPos, Out_WorldNormal);
+	NDICollisionQuery_QueryCustomDepthGPU(In_SamplePos, GetEngineOwnerLWCTile(), Out_SceneDepth, Out_CameraPosWorld, Out_IsInsideView, Out_WorldPos, Out_WorldNormal);
 }
 
 void QueryMeshDistanceFieldGPU_{ParameterName}(in float3 In_SamplePos, out float Out_DistanceToNearestSurface, out float3 Out_FieldGradient, out bool Out_IsDistanceFieldValid)
 {
-	NDICollisionQuery_QueryMeshDistanceFieldGPU(In_SamplePos, {ParameterName}_SystemLWCTile, Out_DistanceToNearestSurface, Out_FieldGradient, Out_IsDistanceFieldValid);
+	NDICollisionQuery_QueryMeshDistanceFieldGPU(In_SamplePos, GetEngineOwnerLWCTile(), Out_DistanceToNearestSurface, Out_FieldGradient, Out_IsDistanceFieldValid);
 }
 
 // DEPRECATE_BEGIN
diff --git a/Engine/Plugins/FX/Niagara/Shaders/Private/NiagaraDataInterfaceLandscape.ush b/Engine/Plugins/FX/Niagara/Shaders/Private/NiagaraDataInterfaceLandscape.ush
--- a/Engine/Plugins/FX/Niagara/Shaders/Private/NiagaraDataInterfaceLandscape.ush
+++ b/Engine/Plugins/FX/Niagara/Shaders/Private/NiagaraDataInterfaceLandscape.ush
@@ -56,11 +56,10 @@ int2				{ParameterName}_CachedHeightTextureDimension;
 SamplerState		{ParameterName}_PointClampedSampler;
 Texture2D<uint>		{ParameterName}_CachedPhysMatTexture;
 int2				{ParameterName}_CachedPhysMatTextureDimension;
-float3				{ParameterName}_SystemLWCTile;
 
 float2 GetCachedHeightTextureUv_{ParameterName}(float3 InWorldPos)
 {
-	FLWCVector4 LwcWorldPos = MakeLWCVector4({ParameterName}_SystemLWCTile, float4(InWorldPos, 1.0f));
+	FLWCVector4 LwcWorldPos = MakeLWCVector4(GetEngineOwnerLWCTile(), float4(InWorldPos, 1.0f));
 	FLWCInverseMatrix LwcWorldToUv = MakeLWCInverseMatrix({ParameterName}_CachedHeightTextureLWCTile, {ParameterName}_CachedHeightTextureWorldToUvTransform);
 
 	float2 Uv = LWCMultiply(LwcWorldPos, LwcWorldToUv).xy;
@@ -77,7 +76,7 @@ void GetBaseColor_{ParameterName}(float3 InWorldPos, out float3 OutBaseColor, ou
 		float3 UnpackBaseColorU = {ParameterName}_BaseColorVirtualTextureWorldToUvTransform[1].xyz;
 		float3 UnpackBaseColorV = {ParameterName}_BaseColorVirtualTextureWorldToUvTransform[2].xyz;
 
-		FDFVector3 LwcWorldPos = DFFromTileOffset_Hack(MakeLWCVector3({ParameterName}_SystemLWCTile, InWorldPos));
+		FDFVector3 LwcWorldPos = DFFromTileOffset_Hack(MakeLWCVector3(GetEngineOwnerLWCTile(), InWorldPos));
 		FDFVector3 LwcBaseColorOrigin = DFFromTileOffset_Hack(MakeLWCVector3({ParameterName}_BaseColorVirtualTextureLWCTile, UnpackBaseColorOrigin));
 		float2 SampleUv = VirtualTextureWorldToUV(LwcWorldPos, LwcBaseColorOrigin, UnpackBaseColorU, UnpackBaseColorV);
 
@@ -115,7 +114,7 @@ void GetHeight_{ParameterName}(float3 InWorldPos, out float OutHeight, out bool
 		float3 UnpackHeightV = {ParameterName}_HeightVirtualTextureWorldToUvTransform[2].xyz;
 		float2 UnpackHeightScaleBias = {ParameterName}_HeightVirtualTextureWorldToUvTransform[3].xy;
 
-		FDFVector3 LwcWorldPos = DFFromTileOffset_Hack(MakeLWCVector3({ParameterName}_SystemLWCTile, InWorldPos));
+		FDFVector3 LwcWorldPos = DFFromTileOffset_Hack(MakeLWCVector3(GetEngineOwnerLWCTile(), InWorldPos));
 		FDFVector3 LwcHeightOrigin = DFFromTileOffset_Hack(MakeLWCVector3({ParameterName}_HeightVirtualTextureLWCTile, UnpackHeightOrigin));
 		float2 SampleUv = VirtualTextureWorldToUV(LwcWorldPos, LwcHeightOrigin, UnpackHeightU, UnpackHeightV);
 
@@ -159,7 +158,7 @@ void GetWorldNormal_{ParameterName}(float3 InWorldPos, out float3 OutNormal, out
 		float3 UnpackNormalU = {ParameterName}_NormalVirtualTextureWorldToUvTransform[1].xyz;
 		float3 UnpackNormalV = {ParameterName}_NormalVirtualTextureWorldToUvTransform[2].xyz;
 
-		FDFVector3 LwcWorldPos = DFFromTileOffset_Hack(MakeLWCVector3({ParameterName}_SystemLWCTile, InWorldPos));
+		FDFVector3 LwcWorldPos = DFFromTileOffset_Hack(MakeLWCVector3(GetEngineOwnerLWCTile(), InWorldPos));
 		FDFVector3 LwcNormalOrigin = DFFromTileOffset_Hack(MakeLWCVector3({ParameterName}_NormalVirtualTextureLWCTile, UnpackNormalOrigin));
 		float2 SampleUv = VirtualTextureWorldToUV(LwcWorldPos, LwcNormalOrigin, UnpackNormalU, UnpackNormalV);
 
diff --git a/Engine/Plugins/FX/Niagara/Shaders/Private/NiagaraDataInterfaceOcclusion.ush b/Engine/Plugins/FX/Niagara/Shaders/Private/NiagaraDataInterfaceOcclusion.ush
--- a/Engine/Plugins/FX/Niagara/Shaders/Private/NiagaraDataInterfaceOcclusion.ush
+++ b/Engine/Plugins/FX/Niagara/Shaders/Private/NiagaraDataInterfaceOcclusion.ush
@@ -3,7 +3,6 @@
 #define AP_SUPPORT_SAMPLE_ATMOSHPERE 0
 #include "/Engine/Private/SkyAtmosphereCommon.ush"
 
-float3				{ParameterName}_SystemLWCTile;
 float2				{ParameterName}_CloudVolumetricTextureUVScale;
 float2				{ParameterName}_CloudVolumetricTextureUVMax;
 Texture2D<float4>	{ParameterName}_CloudVolumetricTexture;
@@ -11,7 +10,7 @@ SamplerState		{ParameterName}_CloudVolumetricTextureSampler;
 
 void QueryOcclusionFactorWithRectangle_{ParameterName}(in float3 In_SampleCenterWorldPos, in float In_SampleWindowWidthWorld, in float In_SampleWindowHeightWorld, in float In_SampleSteps, out float Out_VisibilityFraction, out float Out_SampleFraction)
 {
-	FLWCVector3 LwcSamplePos = MakeLWCVector3({ParameterName}_SystemLWCTile, In_SampleCenterWorldPos);
+	FLWCVector3 LwcSamplePos = MakeLWCVector3(GetEngineOwnerLWCTile(), In_SampleCenterWorldPos);
 	float3 ResolvedLwcSamplePos = DFFastToTranslatedWorld(DFFromTileOffset_Hack(LwcSamplePos), PrimaryView.PreViewTranslation);
 	const float3 TranslatedWorldViewOrigin = DFFastToTranslatedWorld(PrimaryView.WorldViewOrigin, PrimaryView.PreViewTranslation);
 	float CameraDistance = abs(dot(ResolvedLwcSamplePos - TranslatedWorldViewOrigin, View.ViewForward.xyz));
@@ -72,7 +71,7 @@ void QueryOcclusionFactorWithCircleGPU_{ParameterName}(in float3 In_SampleCenter
 {
 	const float PI = 3.14159265;
 	const float SPIRAL_TURN = 2 * PI * 0.61803399; // use golden ratio to rotate sample pattern each ring so we get a spiral
-	FLWCVector3 LwcSamplePos = MakeLWCVector3({ParameterName}_SystemLWCTile, In_SampleCenterWorldPos);
+	FLWCVector3 LwcSamplePos = MakeLWCVector3(GetEngineOwnerLWCTile(), In_SampleCenterWorldPos);
 	float3 ResolvedLwcSamplePos = DFFastToTranslatedWorld(DFFromTileOffset_Hack(LwcSamplePos), PrimaryView.PreViewTranslation);
 	const float3 TranslatedWorldViewOrigin = DFFastToTranslatedWorld(PrimaryView.WorldViewOrigin, PrimaryView.PreViewTranslation);
 	float CameraDistance = abs(dot(ResolvedLwcSamplePos - TranslatedWorldViewOrigin, View.ViewForward.xyz));
@@ -141,7 +140,7 @@ void QueryOcclusionFactorWithCircle_{ParameterName}(float3 WorldPosition, float
 	const float PI = 3.14159265;
 	const float SPIRAL_TURN = 2.0f * PI * 0.61803399; // use golden ratio to rotate sample pattern each ring so we get a spiral
 
-	const FLWCVector3 LWCWorldPosition = MakeLWCVector3({ParameterName}_SystemLWCTile, WorldPosition);
+	const FLWCVector3 LWCWorldPosition = MakeLWCVector3(GetEngineOwnerLWCTile(), WorldPosition);
 	const float3 PreViewPosition = DFFastToTranslatedWorld(DFFromTileOffset_Hack(LWCWorldPosition), PrimaryView.PreViewTranslation);
 	const float3 TranslatedWorldViewOrigin = DFFastToTranslatedWorld(PrimaryView.WorldViewOrigin, PrimaryView.PreViewTranslation);
 	const float SampleDepth = abs(dot(PreViewPosition - TranslatedWorldViewOrigin, View.ViewForward.xyz));
@@ -201,7 +200,7 @@ void QueryCloudOcclusionWithCircle_{ParameterName}(float3 WorldPosition, float W
 	const float PI = 3.14159265;
 	const float SPIRAL_TURN = 2.0f * PI * 0.61803399; // use golden ratio to rotate sample pattern each ring so we get a spiral
 
-	const FLWCVector3 LWCWorldPosition = MakeLWCVector3({ParameterName}_SystemLWCTile, WorldPosition);
+	const FLWCVector3 LWCWorldPosition = MakeLWCVector3(GetEngineOwnerLWCTile(), WorldPosition);
 	const float3 PreViewPosition = DFFastToTranslatedWorld(DFFromTileOffset_Hack(LWCWorldPosition), PrimaryView.PreViewTranslation);
 	const float4 ClipPosition = mul(float4(PreViewPosition, 1), PrimaryView.TranslatedWorldToClip);
 	const float2 ScreenPosition = ClipPosition.xy / ClipPosition.w;
diff --git a/Engine/Plugins/FX/Niagara/Shaders/Private/NiagaraDataInterfaceRigidMeshCollisionQuery.ush b/Engine/Plugins/FX/Niagara/Shaders/Private/NiagaraDataInterfaceRigidMeshCollisionQuery.ush
--- a/Engine/Plugins/FX/Niagara/Shaders/Private/NiagaraDataInterfaceRigidMeshCollisionQuery.ush
+++ b/Engine/Plugins/FX/Niagara/Shaders/Private/NiagaraDataInterfaceRigidMeshCollisionQuery.ush
@@ -26,7 +26,6 @@ Buffer<float4>	{ParameterName}_MeshScaleBuffer;
 Buffer<float4>	{ParameterName}_ElementExtentBuffer;
 Buffer<uint>	{ParameterName}_PhysicsTypeBuffer;
 Buffer<uint>	{ParameterName}_DFIndexBuffer;
-float3			{ParameterName}_SystemLWCTile;
 float3			{ParameterName}_CombinedBBoxWorldMin;
 float3			{ParameterName}_CombinedBBoxWorldMax;
 
@@ -91,7 +90,7 @@ float3 GetMeshScale_{ParameterName}(in int ElementIndex)
 
 void IsWorldPositionInsideCombinedBounds_{ParameterName}(in float3 WorldPosition, out bool IsInside)
 {
-	const FDFVector3 LWCWorldPosition = DFFromTileOffset_Hack(MakeLWCVector3({ParameterName}_SystemLWCTile, WorldPosition));
+	const FDFVector3 LWCWorldPosition = DFFromTileOffset_Hack(MakeLWCVector3(GetEngineOwnerLWCTile(), WorldPosition));
 
 	IsInside =  
 		WorldPosition.x >= {ParameterName}_CombinedBBoxWorldMin.x && 
@@ -379,7 +378,7 @@ void GetElementPointMeshDistanceFieldNoNormal_{ParameterName}(in float3 WorldPos
 		if (DFIndex >= 0 && DFIndex < NumSceneObjects)
 		{
 			MaxDistance = max(.001, MaxDistance);
-			const FDFVector3 LWCWorldPosition = DFFromTileOffset_Hack(MakeLWCVector3({ParameterName}_SystemLWCTile, WorldPosition));
+			const FDFVector3 LWCWorldPosition = DFFromTileOffset_Hack(MakeLWCVector3(GetEngineOwnerLWCTile(), WorldPosition));
 			OutClosestDistance = DistanceToNearestSurfaceForObject(DFIndex, LWCWorldPosition, MaxDistance);			
 		}
 
@@ -514,7 +513,7 @@ void EvaluateWorldPositionOnElement_{ParameterName}(int ElementIndex, float3 Wor
 // Given a world space position (WorldPosition) compute the static mesh closest point (position,normal,velocity)
 void GetClosestPointMeshDistanceField_{ParameterName}(in float3 WorldPosition, in float DeltaTime, in float TimeFraction, in float MaxDistance, out float OutClosestDistance, out float3 OutClosestPosition, out float3 OutClosestNormal, out float3 OutClosestVelocity, out bool NormalIsValid, out float MaxEncodedDistance)
 {
-	const FDFVector3 LWCWorldPosition = DFFromTileOffset_Hack(MakeLWCVector3({ParameterName}_SystemLWCTile, WorldPosition));
+	const FDFVector3 LWCWorldPosition = DFFromTileOffset_Hack(MakeLWCVector3(GetEngineOwnerLWCTile(), WorldPosition));
 
 	float3 CollisionPosition = float3(0,0,0);
 	float3 CollisionNormal = float3(0,0,0);
@@ -556,7 +555,7 @@ void GetClosestPointMeshDistanceField_{ParameterName}(in float3 WorldPosition, i
 // Given a world space position (WorldPosition) compute the static mesh closest point (position,normal,velocity) by querying all overlapping mesh distance fields
 void GetClosestPointMeshDistanceFieldAccurate_{ParameterName}(in float3 WorldPosition, in float DeltaTime, in float TimeFraction, in float MaxDistance, out float OutClosestDistance, out float3 OutClosestPosition, out float3 OutClosestNormal, out float3 OutClosestVelocity, out bool NormalIsValid, out float MaxEncodedDistance)
 {
-	const FDFVector3 LWCWorldPosition = DFFromTileOffset_Hack(MakeLWCVector3({ParameterName}_SystemLWCTile, WorldPosition));
+	const FDFVector3 LWCWorldPosition = DFFromTileOffset_Hack(MakeLWCVector3(GetEngineOwnerLWCTile(), WorldPosition));
 
 	float3 CollisionPosition = float3(0,0,0);
 	float3 CollisionNormal = float3(0,0,0);
@@ -598,7 +597,7 @@ void GetClosestPointMeshDistanceFieldAccurate_{ParameterName}(in float3 WorldPos
 // Given a world space position (WorldPosition) compute the static mesh closest point (position,normal,velocity)
 void GetClosestPointMeshDistanceFieldNoNormal_{ParameterName}(in float3 WorldPosition, in float DeltaTime, in float TimeFraction, in float MaxDistance, out float OutClosestDistance, out float3 OutClosestPosition, out float3 OutClosestVelocity)
 {	
-	const FDFVector3 LWCWorldPosition = DFFromTileOffset_Hack(MakeLWCVector3({ParameterName}_SystemLWCTile, WorldPosition));
+	const FDFVector3 LWCWorldPosition = DFFromTileOffset_Hack(MakeLWCVector3(GetEngineOwnerLWCTile(), WorldPosition));
 
 	float3 CollisionPosition = float3(0,0,0);
 	float3 CollisionNormal = float3(0,0,0);
diff --git a/Engine/Plugins/FX/Niagara/Shaders/Private/NiagaraDataInterfaceStaticMeshTemplate.ush b/Engine/Plugins/FX/Niagara/Shaders/Private/NiagaraDataInterfaceStaticMeshTemplate.ush
--- a/Engine/Plugins/FX/Niagara/Shaders/Private/NiagaraDataInterfaceStaticMeshTemplate.ush
+++ b/Engine/Plugins/FX/Niagara/Shaders/Private/NiagaraDataInterfaceStaticMeshTemplate.ush
@@ -37,7 +37,6 @@ float3			{ParameterName}_PreSkinnedLocalBoundsExtents;
 float3			{ParameterName}_MeshBoundsWSCenter;
 float3			{ParameterName}_MeshBoundsWSExtents;
 float3			{ParameterName}_OwnerToMeshVector;
-float3			{ParameterName}_SystemLWCTile;
 
 #if DISTATICMESH_ALLOWDISTANCEFIELD
 int				{ParameterName}_InstanceDistanceFieldIndex;
@@ -780,7 +779,7 @@ void QueryDistanceField_{ParameterName}(bool bExecute, float3 WorldPosition, boo
 	{
 		// Needs to be a smidge over 0.0 for DistanceToNearestSurfaceForObject to query the SDF
 		MaxDistance = max(0.001f, MaxDistance);
-		const FDFVector3 LWCWorldPosition = DFFromTileOffset_Hack(MakeLWCVector3({ParameterName}_SystemLWCTile, WorldPosition));
+		const FDFVector3 LWCWorldPosition = DFFromTileOffset_Hack(MakeLWCVector3(GetEngineOwnerLWCTile(), WorldPosition));
 		OutDistance = DistanceToNearestSurfaceForObject(DFIndex, LWCWorldPosition, MaxDistance);
 	}
 #else
diff --git a/Engine/Plugins/FX/Niagara/Shaders/Private/NiagaraDataInterfaceVirtualTextureTemplate.ush b/Engine/Plugins/FX/Niagara/Shaders/Private/NiagaraDataInterfaceVirtualTextureTemplate.ush
--- a/Engine/Plugins/FX/Niagara/Shaders/Private/NiagaraDataInterfaceVirtualTextureTemplate.ush
+++ b/Engine/Plugins/FX/Niagara/Shaders/Private/NiagaraDataInterfaceVirtualTextureTemplate.ush
@@ -17,7 +17,6 @@ int					{ParameterName}_MaterialType;
 float4				{ParameterName}_UVUniforms[3];
 uint4				{ParameterName}_PageTableUniforms[2];
 float2				{ParameterName}_WorldHeightUnpack;
-float3				{ParameterName}_SystemLWCTile;
 SamplerState		{ParameterName}_SharedSampler;
 
 void GetAttributesValid_{ParameterName}(out bool BaseColorValid, out bool SpecularValid, out bool RoughnessValid, out bool NormalValid, out bool WorldHeightValid, out bool MaskValid)
@@ -87,8 +86,8 @@ void SampleRVT_{ParameterName}(in float3 WorldPosition, out bool bInsideVolume,
 	Mask = 1.0f;
 
 	// Get Sample Location
-	FDFVector3 LWCWorldPosition = DFFromTileOffset_Hack(MakeLWCVector3({ParameterName}_SystemLWCTile, WorldPosition));
-	FDFVector3 LWCUVOrigin = DFFromTileOffset_Hack(MakeLWCVector3({ParameterName}_SystemLWCTile, {ParameterName}_UVUniforms[0].xyz));
+	FDFVector3 LWCWorldPosition = DFFromTileOffset_Hack(MakeLWCVector3(GetEngineOwnerLWCTile(), WorldPosition));
+	FDFVector3 LWCUVOrigin = DFFromTileOffset_Hack(MakeLWCVector3(GetEngineOwnerLWCTile(), {ParameterName}_UVUniforms[0].xyz));
 
 	float2 SampleUV = VirtualTextureWorldToUV(LWCWorldPosition, LWCUVOrigin, {ParameterName}_UVUniforms[1].xyz, {ParameterName}_UVUniforms[2].xyz);
 
diff --git a/Engine/Plugins/FX/Niagara/Shaders/Private/NiagaraEmitterInstanceShader.usf b/Engine/Plugins/FX/Niagara/Shaders/Private/NiagaraEmitterInstanceShader.usf
--- a/Engine/Plugins/FX/Niagara/Shaders/Private/NiagaraEmitterInstanceShader.usf
+++ b/Engine/Plugins/FX/Niagara/Shaders/Private/NiagaraEmitterInstanceShader.usf
@@ -65,6 +65,8 @@ bool NiagaraAny(bool3 b) { return b.x || b.y || b.z; }
 bool NiagaraAny(bool4 b) { return b.x || b.y || b.z || b.w; }
 #endif
 
+#define GetEngineOwnerLWCTile() Engine_Owner_LWCTile.xyz
+
 /* -----------------------------------------------------------------
  * GPU simulation utility functions
  * -----------------------------------------------------------------
diff --git a/Engine/Plugins/FX/Niagara/Source/Niagara/Classes/NiagaraDataInterfaceCamera.h b/Engine/Plugins/FX/Niagara/Source/Niagara/Classes/NiagaraDataInterfaceCamera.h
--- a/Engine/Plugins/FX/Niagara/Source/Niagara/Classes/NiagaraDataInterfaceCamera.h
+++ b/Engine/Plugins/FX/Niagara/Source/Niagara/Classes/NiagaraDataInterfaceCamera.h
@@ -31,7 +31,6 @@ class UNiagaraDataInterfaceCamera : public UNiagaraDataInterface
 	GENERATED_UCLASS_BODY()
 
 	BEGIN_SHADER_PARAMETER_STRUCT(FShaderParameters, )
-		SHADER_PARAMETER(FVector3f, SystemLWCTile)
 		SHADER_PARAMETER(int32, SplitscreenMode)
 	END_SHADER_PARAMETER_STRUCT();
 
diff --git a/Engine/Plugins/FX/Niagara/Source/Niagara/Classes/NiagaraDataInterfaceCollisionQuery.h b/Engine/Plugins/FX/Niagara/Source/Niagara/Classes/NiagaraDataInterfaceCollisionQuery.h
--- a/Engine/Plugins/FX/Niagara/Source/Niagara/Classes/NiagaraDataInterfaceCollisionQuery.h
+++ b/Engine/Plugins/FX/Niagara/Source/Niagara/Classes/NiagaraDataInterfaceCollisionQuery.h
@@ -27,10 +27,6 @@ class UNiagaraDataInterfaceCollisionQuery : public UNiagaraDataInterface
 {
 	GENERATED_UCLASS_BODY()
 
-	BEGIN_SHADER_PARAMETER_STRUCT(FShaderParameters, )
-		SHADER_PARAMETER(FVector3f,	SystemLWCTile)
-	END_SHADER_PARAMETER_STRUCT()
-
 public:
 	//UObject Interface
 	NIAGARA_API virtual void PostInitProperties() override;
diff --git a/Engine/Plugins/FX/Niagara/Source/Niagara/Classes/NiagaraDataInterfaceOcclusion.h b/Engine/Plugins/FX/Niagara/Source/Niagara/Classes/NiagaraDataInterfaceOcclusion.h
--- a/Engine/Plugins/FX/Niagara/Source/Niagara/Classes/NiagaraDataInterfaceOcclusion.h
+++ b/Engine/Plugins/FX/Niagara/Source/Niagara/Classes/NiagaraDataInterfaceOcclusion.h
@@ -12,7 +12,6 @@ class UNiagaraDataInterfaceOcclusion : public UNiagaraDataInterface
 	GENERATED_UCLASS_BODY()
 
 	BEGIN_SHADER_PARAMETER_STRUCT(FShaderParameters, )
-		SHADER_PARAMETER(FVector3f,				SystemLWCTile)
 		SHADER_PARAMETER(FVector2f,				CloudVolumetricTextureUVScale)
 		SHADER_PARAMETER(FVector2f,				CloudVolumetricTextureUVMax)
 		SHADER_PARAMETER_RDG_TEXTURE(Texture2D, CloudVolumetricTexture)
diff --git a/Engine/Plugins/FX/Niagara/Source/Niagara/Private/DataInterface/NiagaraDataInterfaceAsyncGpuTrace.cpp b/Engine/Plugins/FX/Niagara/Source/Niagara/Private/DataInterface/NiagaraDataInterfaceAsyncGpuTrace.cpp
--- a/Engine/Plugins/FX/Niagara/Source/Niagara/Private/DataInterface/NiagaraDataInterfaceAsyncGpuTrace.cpp
+++ b/Engine/Plugins/FX/Niagara/Source/Niagara/Private/DataInterface/NiagaraDataInterfaceAsyncGpuTrace.cpp
@@ -24,7 +24,6 @@ namespace NDIAsyncGpuTraceLocal
 	static const TCHAR* TemplateShaderFile = TEXT("/Plugin/FX/Niagara/Private/NiagaraDataInterfaceAsyncGpuTraceTemplate.ush");
 
 	BEGIN_SHADER_PARAMETER_STRUCT(FShaderParameters,)
-		SHADER_PARAMETER(FVector3f,											SystemLWCTile)
 		SHADER_PARAMETER(uint32,											MaxRayTraceCount)
 		SHADER_PARAMETER_UAV(RWStructuredBuffer<FNiagaraAsyncGpuTrace>,		RWRayRequests)
 		SHADER_PARAMETER(uint32,											RayRequestsOffset)
@@ -411,7 +410,6 @@ void UNiagaraDataInterfaceAsyncGpuTrace::BuildShaderParameters(FNiagaraShaderPar
 void UNiagaraDataInterfaceAsyncGpuTrace::SetShaderParameters(const FNiagaraDataInterfaceSetShaderParametersContext& Context) const
 {
 	NDIAsyncGpuTraceLocal::FShaderParameters* ShaderParameters = Context.GetParameterNestedStruct<NDIAsyncGpuTraceLocal::FShaderParameters>();
-	ShaderParameters->SystemLWCTile = Context.GetSystemLWCTile();
 
 	const bool bRWRayRequestsBound = Context.IsResourceBound(&ShaderParameters->RWRayRequests);
 	const bool bIntersectionResultsBound = Context.IsResourceBound(&ShaderParameters->IntersectionResults);
diff --git a/Engine/Plugins/FX/Niagara/Source/Niagara/Private/DataInterface/NiagaraDataInterfaceStaticMesh.cpp b/Engine/Plugins/FX/Niagara/Source/Niagara/Private/DataInterface/NiagaraDataInterfaceStaticMesh.cpp
--- a/Engine/Plugins/FX/Niagara/Source/Niagara/Private/DataInterface/NiagaraDataInterfaceStaticMesh.cpp
+++ b/Engine/Plugins/FX/Niagara/Source/Niagara/Private/DataInterface/NiagaraDataInterfaceStaticMesh.cpp
@@ -93,7 +93,6 @@ namespace NDIStaticMeshLocal
 		SHADER_PARAMETER(FVector3f,				MeshBoundsWSCenter)
 		SHADER_PARAMETER(FVector3f,				MeshBoundsWSExtents)
 		SHADER_PARAMETER(FVector3f,				OwnerToMeshVector)
-		SHADER_PARAMETER(FVector3f,				SystemLWCTile)
 
 		SHADER_PARAMETER_SRV(Buffer<int>,		UvMappingBuffer)
 		SHADER_PARAMETER(uint32,				UvMappingBufferLength)
@@ -3236,8 +3235,6 @@ void UNiagaraDataInterfaceStaticMesh::SetShaderParameters(const FNiagaraDataInte
 	ShaderParameters->MeshBoundsWSExtents = WorldBounds.GetExtent();
 	ShaderParameters->OwnerToMeshVector = InstanceData.OwnerToMeshVector;
 
-	ShaderParameters->SystemLWCTile = Context.GetSystemLWCTile();
-
 	if (InstanceData.UvMappingBuffer)
 	{
 		ShaderParameters->UvMappingBuffer = InstanceData.UvMappingBuffer->GetSrv();
diff --git a/Engine/Plugins/FX/Niagara/Source/Niagara/Private/DataInterface/NiagaraDataInterfaceVirtualTexture.cpp b/Engine/Plugins/FX/Niagara/Source/Niagara/Private/DataInterface/NiagaraDataInterfaceVirtualTexture.cpp
--- a/Engine/Plugins/FX/Niagara/Source/Niagara/Private/DataInterface/NiagaraDataInterfaceVirtualTexture.cpp
+++ b/Engine/Plugins/FX/Niagara/Source/Niagara/Private/DataInterface/NiagaraDataInterfaceVirtualTexture.cpp
@@ -36,7 +36,6 @@ namespace NDIVirtualTextureLocal
 		SHADER_PARAMETER_ARRAY(FUintVector4,		PageTableUniforms, [2])
 		SHADER_PARAMETER_ARRAY(FVector4f,			UVUniforms, [3])
 		SHADER_PARAMETER(FVector2f,					WorldHeightUnpack)
-		SHADER_PARAMETER(FVector3f,					SystemLWCTile)
 		SHADER_PARAMETER_SAMPLER(SamplerState,		SharedSampler)
 	END_SHADER_PARAMETER_STRUCT()
 
@@ -440,7 +439,6 @@ void UNiagaraDataInterfaceVirtualTexture::SetShaderParameters(const FNiagaraData
 		ShaderParameters->WorldHeightUnpack		= FVector2f::ZeroVector;
 		ShaderParameters->SharedSampler			= TStaticSamplerState<SF_Bilinear, AM_Clamp, AM_Clamp, AM_Clamp>::GetRHI();
 	}
-	ShaderParameters->SystemLWCTile	= Context.GetSystemLWCTile();
 	
 	if ( bLayerValid[0] == false )
 	{
diff --git a/Engine/Plugins/FX/Niagara/Source/Niagara/Private/NiagaraDataInterfaceCamera.cpp b/Engine/Plugins/FX/Niagara/Source/Niagara/Private/NiagaraDataInterfaceCamera.cpp
--- a/Engine/Plugins/FX/Niagara/Source/Niagara/Private/NiagaraDataInterfaceCamera.cpp
+++ b/Engine/Plugins/FX/Niagara/Source/Niagara/Private/NiagaraDataInterfaceCamera.cpp
@@ -438,7 +438,6 @@ void UNiagaraDataInterfaceCamera::SetShaderParameters(const FNiagaraDataInterfac
 	FCameraDataInterface_InstanceDataRT& InstanceData = DataInterfaceProxy.SystemInstancesToInstanceData_RT.FindChecked(Context.GetSystemInstanceID());
 	
 	FShaderParameters* ShaderParameters = Context.GetParameterNestedStruct<FShaderParameters>();
-	ShaderParameters->SystemLWCTile = Context.GetSystemLWCTile();
 	ShaderParameters->SplitscreenMode = static_cast<int32>(InstanceData.SplitscreenMode);
 }
 
diff --git a/Engine/Plugins/FX/Niagara/Source/Niagara/Private/NiagaraDataInterfaceCollisionQuery.cpp b/Engine/Plugins/FX/Niagara/Source/Niagara/Private/NiagaraDataInterfaceCollisionQuery.cpp
--- a/Engine/Plugins/FX/Niagara/Source/Niagara/Private/NiagaraDataInterfaceCollisionQuery.cpp
+++ b/Engine/Plugins/FX/Niagara/Source/Niagara/Private/NiagaraDataInterfaceCollisionQuery.cpp
@@ -525,7 +525,6 @@ bool UNiagaraDataInterfaceCollisionQuery::AppendCompileHash(FNiagaraCompileHashV
 	InVisitor->UpdatePOD(TEXT("NiagaraCollisionDI_DistanceField"), IsDistanceFieldEnabled());
 	InVisitor->UpdateShaderFile(NDICollisionQueryLocal::CommonShaderFile);
 	InVisitor->UpdateShaderFile(NDICollisionQueryLocal::TemplateShaderFile);
-	InVisitor->UpdateShaderParameters<FShaderParameters>();
 	InVisitor->UpdateShaderParameters<FGlobalDistanceFieldParameters2>();
 
 	return true;
@@ -535,15 +534,11 @@ bool UNiagaraDataInterfaceCollisionQuery::AppendCompileHash(FNiagaraCompileHashV
 
 void UNiagaraDataInterfaceCollisionQuery::BuildShaderParameters(FNiagaraShaderParametersBuilder& ShaderParametersBuilder) const
 {
-	ShaderParametersBuilder.AddNestedStruct<FShaderParameters>();
 	ShaderParametersBuilder.AddIncludedStruct<FGlobalDistanceFieldParameters2>();
 }
 
 void UNiagaraDataInterfaceCollisionQuery::SetShaderParameters(const FNiagaraDataInterfaceSetShaderParametersContext& Context) const
 {
-	FShaderParameters* ShaderParameters = Context.GetParameterNestedStruct<FShaderParameters>();
-	ShaderParameters->SystemLWCTile = Context.GetSystemLWCTile();	//-OPT: Seems like a lot of DIs might want tile
-
 	FGlobalDistanceFieldParameters2* ShaderGDFParameters = Context.GetParameterIncludedStruct<FGlobalDistanceFieldParameters2>();
 	if (Context.IsStructBound(ShaderGDFParameters))
 	{
diff --git a/Engine/Plugins/FX/Niagara/Source/Niagara/Private/NiagaraDataInterfaceLandscape.cpp b/Engine/Plugins/FX/Niagara/Source/Niagara/Private/NiagaraDataInterfaceLandscape.cpp
--- a/Engine/Plugins/FX/Niagara/Source/Niagara/Private/NiagaraDataInterfaceLandscape.cpp
+++ b/Engine/Plugins/FX/Niagara/Source/Niagara/Private/NiagaraDataInterfaceLandscape.cpp
@@ -114,7 +114,6 @@ namespace NiagaraDataInterfaceLandscape
 		SHADER_PARAMETER_SAMPLER(SamplerState,		PointClampedSampler)
 		SHADER_PARAMETER_TEXTURE(Texture2D<uint>,	CachedPhysMatTexture)
 		SHADER_PARAMETER(FIntPoint,					CachedPhysMatTextureDimension)
-		SHADER_PARAMETER(FVector3f,					SystemLWCTile)
 	END_SHADER_PARAMETER_STRUCT()
 }
 
@@ -1688,9 +1687,6 @@ void UNiagaraDataInterfaceLandscape::SetShaderParameters(const FNiagaraDataInter
 	{
 		FNDILandscapeData_RenderThread::SetCachedHeightTextureParameters_Defaults(ShaderParameters);
 	}
-
-	// System tile for LWC
-	ShaderParameters->SystemLWCTile = Context.GetSystemLWCTile();
 }
 
 bool UNiagaraDataInterfaceLandscape::InitPerInstanceData(void* PerInstanceData, FNiagaraSystemInstance* SystemInstance)
diff --git a/Engine/Plugins/FX/Niagara/Source/Niagara/Private/NiagaraDataInterfaceOcclusion.cpp b/Engine/Plugins/FX/Niagara/Source/Niagara/Private/NiagaraDataInterfaceOcclusion.cpp
--- a/Engine/Plugins/FX/Niagara/Source/Niagara/Private/NiagaraDataInterfaceOcclusion.cpp
+++ b/Engine/Plugins/FX/Niagara/Source/Niagara/Private/NiagaraDataInterfaceOcclusion.cpp
@@ -227,7 +227,6 @@ void UNiagaraDataInterfaceOcclusion::BuildShaderParameters(FNiagaraShaderParamet
 void UNiagaraDataInterfaceOcclusion::SetShaderParameters(const FNiagaraDataInterfaceSetShaderParametersContext& Context) const
 {
 	FShaderParameters* ShaderParameters = Context.GetParameterNestedStruct<FShaderParameters>();
-	ShaderParameters->SystemLWCTile = Context.GetSystemLWCTile();
 
 	if (Context.IsResourceBound(&ShaderParameters->CloudVolumetricTexture))
 	{
diff --git a/Engine/Plugins/FX/Niagara/Source/Niagara/Private/NiagaraDataInterfaceRigidMeshCollisionQuery.cpp b/Engine/Plugins/FX/Niagara/Source/Niagara/Private/NiagaraDataInterfaceRigidMeshCollisionQuery.cpp
--- a/Engine/Plugins/FX/Niagara/Source/Niagara/Private/NiagaraDataInterfaceRigidMeshCollisionQuery.cpp
+++ b/Engine/Plugins/FX/Niagara/Source/Niagara/Private/NiagaraDataInterfaceRigidMeshCollisionQuery.cpp
@@ -76,7 +76,6 @@ BEGIN_SHADER_PARAMETER_STRUCT(FShaderParameters,)
 	SHADER_PARAMETER_SRV(Buffer<float4>,	MeshScaleBuffer)
 	SHADER_PARAMETER_SRV(Buffer<uint32>,	PhysicsTypeBuffer)
 	SHADER_PARAMETER_SRV(Buffer<uint32>,	DFIndexBuffer)
-	SHADER_PARAMETER(FVector3f,				SystemLWCTile)
 	SHADER_PARAMETER(FVector3f,				CombinedBBoxWorldMin)
 	SHADER_PARAMETER(FVector3f,				CombinedBBoxWorldMax)
 END_SHADER_PARAMETER_STRUCT()
@@ -1998,8 +1997,6 @@ void UNiagaraDataInterfaceRigidMeshCollisionQuery::SetShaderParameters(const FNi
 		UE::FXRenderingUtils::DistanceFields::SetupObjectBufferParameters(Context.GetGraphBuilder(), ShaderDistanceFieldObjectParameters, PrimaryView);
 		UE::FXRenderingUtils::DistanceFields::SetupAtlasParameters(Context.GetGraphBuilder(), ShaderDistanceFieldAtlasParameters, PrimaryView);
 	}
-
-	ShaderParameters->SystemLWCTile = Context.GetSystemLWCTile();
 }
 
 #if WITH_EDITOR
