Deprecated old parameterized CreateShaderResourceView/CreateUnorderedAccessView helpers

diff --git a/Engine/Plugins/2D/Paper2D/Source/Paper2D/Private/PaperSpriteVertexBuffer.cpp b/Engine/Plugins/2D/Paper2D/Source/Paper2D/Private/PaperSpriteVertexBuffer.cpp
--- a/Engine/Plugins/2D/Paper2D/Source/Paper2D/Private/PaperSpriteVertexBuffer.cpp
+++ b/Engine/Plugins/2D/Paper2D/Source/Paper2D/Private/PaperSpriteVertexBuffer.cpp
@@ -31,7 +31,11 @@ void FPaperSpriteVertexBuffer::CreateBuffers(FRHICommandListBase& RHICmdList, in
 		PositionBuffer.VertexBufferRHI = RHICmdList.CreateVertexBuffer(PositionSize, Usage, CreateInfo);
 		if (RHISupportsManualVertexFetch(GMaxRHIShaderPlatform))
 		{
-			PositionBufferSRV = RHICmdList.CreateShaderResourceView(PositionBuffer.VertexBufferRHI, sizeof(float), PF_R32_FLOAT);
+			PositionBufferSRV = RHICmdList.CreateShaderResourceView(
+				PositionBuffer.VertexBufferRHI, 
+				FRHIViewDesc::CreateBufferSRV()
+					.SetType(FRHIViewDesc::EBufferType::Typed)
+					.SetFormat(PF_R32_FLOAT));
 		}
 
 	}
@@ -43,7 +47,11 @@ void FPaperSpriteVertexBuffer::CreateBuffers(FRHICommandListBase& RHICmdList, in
 		TangentBuffer.VertexBufferRHI = RHICmdList.CreateVertexBuffer(TangentSize, Usage, CreateInfo);
 		if (RHISupportsManualVertexFetch(GMaxRHIShaderPlatform))
 		{
-			TangentBufferSRV = RHICmdList.CreateShaderResourceView(TangentBuffer.VertexBufferRHI, sizeof(FPackedNormal), PF_R8G8B8A8_SNORM);
+			TangentBufferSRV = RHICmdList.CreateShaderResourceView(
+				TangentBuffer.VertexBufferRHI, 
+				FRHIViewDesc::CreateBufferSRV()
+					.SetType(FRHIViewDesc::EBufferType::Typed)
+					.SetFormat(PF_R8G8B8A8_SNORM));
 		}
 	}
 
@@ -54,7 +62,11 @@ void FPaperSpriteVertexBuffer::CreateBuffers(FRHICommandListBase& RHICmdList, in
 		TexCoordBuffer.VertexBufferRHI = RHICmdList.CreateVertexBuffer(TexCoordSize, Usage, CreateInfo);
 		if (RHISupportsManualVertexFetch(GMaxRHIShaderPlatform))
 		{
-			TexCoordBufferSRV = RHICmdList.CreateShaderResourceView(TexCoordBuffer.VertexBufferRHI, sizeof(FVector2f), PF_G32R32F);
+			TexCoordBufferSRV = RHICmdList.CreateShaderResourceView(
+				TexCoordBuffer.VertexBufferRHI, 
+				FRHIViewDesc::CreateBufferSRV()
+					.SetType(FRHIViewDesc::EBufferType::Typed)
+					.SetFormat(PF_G32R32F));
 		}
 	}
 
@@ -65,7 +77,11 @@ void FPaperSpriteVertexBuffer::CreateBuffers(FRHICommandListBase& RHICmdList, in
 		ColorBuffer.VertexBufferRHI = RHICmdList.CreateVertexBuffer(ColorSize, Usage, CreateInfo);
 		if (RHISupportsManualVertexFetch(GMaxRHIShaderPlatform))
 		{
-			ColorBufferSRV = RHICmdList.CreateShaderResourceView(ColorBuffer.VertexBufferRHI, sizeof(FColor), PF_R8G8B8A8);
+			ColorBufferSRV = RHICmdList.CreateShaderResourceView(
+				ColorBuffer.VertexBufferRHI,
+				FRHIViewDesc::CreateBufferSRV()
+					.SetType(FRHIViewDesc::EBufferType::Typed)
+					.SetFormat(PF_R8G8B8A8));
 		}
 	}
 
diff --git a/Engine/Plugins/Animation/MLDeformer/MLDeformerFramework/Source/MLDeformerFramework/Private/MLDeformerModel.cpp b/Engine/Plugins/Animation/MLDeformer/MLDeformerFramework/Source/MLDeformerFramework/Private/MLDeformerModel.cpp
--- a/Engine/Plugins/Animation/MLDeformer/MLDeformerFramework/Source/MLDeformerFramework/Private/MLDeformerModel.cpp
+++ b/Engine/Plugins/Animation/MLDeformer/MLDeformerFramework/Source/MLDeformerFramework/Private/MLDeformerModel.cpp
@@ -37,7 +37,11 @@ namespace UE::MLDeformer
 			RHICmdList.UnlockBuffer(VertexBufferRHI);
 			VertexMap.Empty();
 
-			ShaderResourceViewRHI = RHICmdList.CreateShaderResourceView(VertexBufferRHI, 4, PF_R32_UINT);
+			ShaderResourceViewRHI = RHICmdList.CreateShaderResourceView(
+				VertexBufferRHI, 
+				FRHIViewDesc::CreateBufferSRV()
+					.SetType(FRHIViewDesc::EBufferType::Typed)
+					.SetFormat(PF_R32_UINT));
 		}
 		else
 		{
diff --git a/Engine/Plugins/Enterprise/LidarPointCloud/Source/LidarPointCloudRuntime/Private/Rendering/LidarPointCloudRenderBuffers.cpp b/Engine/Plugins/Enterprise/LidarPointCloud/Source/LidarPointCloudRuntime/Private/Rendering/LidarPointCloudRenderBuffers.cpp
--- a/Engine/Plugins/Enterprise/LidarPointCloud/Source/LidarPointCloudRuntime/Private/Rendering/LidarPointCloudRenderBuffers.cpp
+++ b/Engine/Plugins/Enterprise/LidarPointCloud/Source/LidarPointCloudRuntime/Private/Rendering/LidarPointCloudRenderBuffers.cpp
@@ -92,7 +92,11 @@ void FLidarPointCloudRenderBuffer::InitRHI(FRHICommandListBase& RHICmdList)
 {
 	FRHIResourceCreateInfo CreateInfo(TEXT("FLidarPointCloudRenderBuffer"));
 	Buffer = RHICmdList.CreateVertexBuffer(sizeof(uint32) * Capacity, BUF_ShaderResource | BUF_Dynamic, CreateInfo);
-	SRV = RHICmdList.CreateShaderResourceView(Buffer, sizeof(uint32), PF_R32_FLOAT);
+	SRV = RHICmdList.CreateShaderResourceView(
+		Buffer,
+		FRHIViewDesc::CreateBufferSRV()
+			.SetType(FRHIViewDesc::EBufferType::Typed)
+			.SetFormat(PF_R32_FLOAT));
 
 	FLidarPointCloudVertexFactoryUniformShaderParameters UniformParameters;
 	UniformParameters.VertexFetch_Buffer = SRV;
@@ -332,7 +336,11 @@ void FLidarPointCloudVertexFactory::FPointCloudVertexBuffer::InitRHI(FRHICommand
 	const uint32 BufferSize = NumPoints * 4 * sizeof(FLidarPointCloudPoint);
 	FRHIResourceCreateInfo CreateInfo(TEXT("FPointCloudVertexBuffer"));
 	VertexBufferRHI = RHICmdList.CreateBuffer(BufferSize, BUF_ShaderResource | BUF_Static | BUF_VertexBuffer, sizeof(FLidarPointCloudPoint), ERHIAccess::VertexOrIndexBuffer, CreateInfo);
-	VertexBufferSRV = RHICmdList.CreateShaderResourceView(VertexBufferRHI, sizeof(float), PF_R32_FLOAT);
+	VertexBufferSRV = RHICmdList.CreateShaderResourceView(
+		VertexBufferRHI, 
+		FRHIViewDesc::CreateBufferSRV()
+			.SetType(FRHIViewDesc::EBufferType::Typed)
+			.SetFormat(PF_R32_FLOAT));
 	void* Buffer = RHICmdList.LockBuffer(VertexBufferRHI, 0, BufferSize, RLM_WriteOnly);
 
 	FLidarPointCloudPoint* Dest = (FLidarPointCloudPoint*)Buffer;
diff --git a/Engine/Plugins/Enterprise/LidarPointCloud/Source/LidarPointCloudRuntime/Private/Rendering/LidarPointCloudRendering.cpp b/Engine/Plugins/Enterprise/LidarPointCloud/Source/LidarPointCloudRuntime/Private/Rendering/LidarPointCloudRendering.cpp
--- a/Engine/Plugins/Enterprise/LidarPointCloud/Source/LidarPointCloudRuntime/Private/Rendering/LidarPointCloudRendering.cpp
+++ b/Engine/Plugins/Enterprise/LidarPointCloud/Source/LidarPointCloudRuntime/Private/Rendering/LidarPointCloudRendering.cpp
@@ -108,7 +108,11 @@ private:
 
 			if (RHISupportsManualVertexFetch(GMaxRHIShaderPlatform))
 			{
-				NewData.PositionComponentSRV = RHICmdList.CreateShaderResourceView(InVertexBuffer->VertexBufferRHI, 4, PF_R32_FLOAT);
+				NewData.PositionComponentSRV = RHICmdList.CreateShaderResourceView(
+					InVertexBuffer->VertexBufferRHI, 
+					FRHIViewDesc::CreateBufferSRV()
+						.SetType(FRHIViewDesc::EBufferType::Typed)
+						.SetFormat(PF_R32_FLOAT));
 				NewData.ColorComponentsSRV = GNullColorVertexBuffer.VertexBufferSRV;
 				NewData.TangentsSRV = GNullColorVertexBuffer.VertexBufferSRV;
 				NewData.TextureCoordinatesSRV = GNullColorVertexBuffer.VertexBufferSRV;
diff --git a/Engine/Plugins/Experimental/ChaosFlesh/Source/ChaosFleshEngine/Private/ChaosFlesh/ChaosDeformableGPUBuffers.cpp b/Engine/Plugins/Experimental/ChaosFlesh/Source/ChaosFleshEngine/Private/ChaosFlesh/ChaosDeformableGPUBuffers.cpp
--- a/Engine/Plugins/Experimental/ChaosFlesh/Source/ChaosFleshEngine/Private/ChaosFlesh/ChaosDeformableGPUBuffers.cpp
+++ b/Engine/Plugins/Experimental/ChaosFlesh/Source/ChaosFleshEngine/Private/ChaosFlesh/ChaosDeformableGPUBuffers.cpp
@@ -59,7 +59,11 @@ void FFloatArrayBufferWithSRV::InitRHI(FRHICommandListBase& RHICmdList)
 			Array.Empty();
 		}
 
-        ShaderResourceViewRHI = RHICmdList.CreateShaderResourceView(VertexBufferRHI, 4, PF_R32_FLOAT);
+        ShaderResourceViewRHI = RHICmdList.CreateShaderResourceView(
+			VertexBufferRHI, 
+			FRHIViewDesc::CreateBufferSRV()
+				.SetType(FRHIViewDesc::EBufferType::Typed)
+				.SetFormat(PF_R32_FLOAT));
 
 		UE_LOG(LogFleshDeformerBuffers, Log,
 			TEXT("FFloatArrayBufferWithSRV - '%s' - Buffer size: %d, Array num: %d"),
@@ -129,7 +133,11 @@ void FHalfArrayBufferWithSRV::InitRHI(FRHICommandListBase& RHICmdList)
 		RHICmdList.UnlockBuffer(VertexBufferRHI);
         Array.Empty();
 
-        ShaderResourceViewRHI = RHICmdList.CreateShaderResourceView(VertexBufferRHI, sizeof(FFloat16), PF_R16F);
+        ShaderResourceViewRHI = RHICmdList.CreateShaderResourceView(
+			VertexBufferRHI, 
+			FRHIViewDesc::CreateBufferSRV()
+				.SetType(FRHIViewDesc::EBufferType::Typed)
+				.SetFormat(PF_R16F));
 
 		UE_LOG(LogFleshDeformerBuffers, Log,
 			TEXT("FHalfArrayBufferWithSRV - '%s' - Buffer size: %d, Array num: %d"),
@@ -330,7 +338,11 @@ void FIndexArrayBufferWithSRV::InitRHI(FRHICommandListBase& RHICmdList)
         Array.Empty();
 
 		EPixelFormat Format = PF_R32_UINT;
-        ShaderResourceViewRHI = RHICmdList.CreateShaderResourceView(VertexBufferRHI, sizeof(uint32), Format);
+        ShaderResourceViewRHI = RHICmdList.CreateShaderResourceView(
+			VertexBufferRHI,
+			FRHIViewDesc::CreateBufferSRV()
+				.SetType(FRHIViewDesc::EBufferType::Typed)
+				.SetFormat(Format));
 
 		UE_LOG(LogFleshDeformerBuffers, Log,
 			TEXT("FIndexArrayBufferWithSRV - '%s' - Data stride: %d, Buffer size: %d, Input array size: %d"),
diff --git a/Engine/Plugins/Experimental/GPULightmass/Source/GPULightmass/Private/IrradianceCaching.cpp b/Engine/Plugins/Experimental/GPULightmass/Source/GPULightmass/Private/IrradianceCaching.cpp
--- a/Engine/Plugins/Experimental/GPULightmass/Source/GPULightmass/Private/IrradianceCaching.cpp
+++ b/Engine/Plugins/Experimental/GPULightmass/Source/GPULightmass/Private/IrradianceCaching.cpp
@@ -20,9 +20,9 @@ FIrradianceCache::FIrradianceCache(int32 Quality, float Spacing, float CornerRej
 	{
 		FRHIResourceCreateInfo CreateInfo(TEXT("FIrradianceCache"));
 		IrradianceCacheRecords = RHICmdList.CreateStructuredBuffer(sizeof(FVector4f), sizeof(FVector4f) * IrradianceCacheMaxSize, BUF_UnorderedAccess | BUF_ShaderResource, CreateInfo);
-		IrradianceCacheRecordsUAV = RHICmdList.CreateUnorderedAccessView(IrradianceCacheRecords, false, false);
+		IrradianceCacheRecordsUAV = RHICmdList.CreateUnorderedAccessView(IrradianceCacheRecords, FRHIViewDesc::CreateBufferUAV().SetTypeFromBuffer(IrradianceCacheRecords));
 		IrradianceCacheRecordBackfaceHits = RHICmdList.CreateStructuredBuffer(sizeof(FVector4f), sizeof(FVector4f) * IrradianceCacheMaxSize, BUF_UnorderedAccess | BUF_ShaderResource, CreateInfo);
-		IrradianceCacheRecordBackfaceHitsUAV = RHICmdList.CreateUnorderedAccessView(IrradianceCacheRecordBackfaceHits, false, false);
+		IrradianceCacheRecordBackfaceHitsUAV = RHICmdList.CreateUnorderedAccessView(IrradianceCacheRecordBackfaceHits, FRHIViewDesc::CreateBufferUAV().SetTypeFromBuffer(IrradianceCacheRecordBackfaceHits));
 		
 		IrradianceCacheTotalBytes += sizeof(FVector4f) * IrradianceCacheMaxSize;
 		IrradianceCacheTotalBytes += sizeof(FVector4f) * IrradianceCacheMaxSize;
diff --git a/Engine/Plugins/Experimental/GPULightmass/Source/GPULightmass/Private/LightmapRenderer.cpp b/Engine/Plugins/Experimental/GPULightmass/Source/GPULightmass/Private/LightmapRenderer.cpp
--- a/Engine/Plugins/Experimental/GPULightmass/Source/GPULightmass/Private/LightmapRenderer.cpp
+++ b/Engine/Plugins/Experimental/GPULightmass/Source/GPULightmass/Private/LightmapRenderer.cpp
@@ -213,7 +213,7 @@ struct FGPUBatchedTileRequests
 			CreateInfo.ResourceArray = &BatchedTilesDesc;
 
 			BatchedTilesBuffer = RHICmdList.CreateStructuredBuffer(sizeof(FGPUTileDescription), BatchedTilesDesc.GetResourceDataSize(), BUF_Static | BUF_ShaderResource, CreateInfo);
-			BatchedTilesSRV = RHICmdList.CreateShaderResourceView(BatchedTilesBuffer);
+			BatchedTilesSRV = RHICmdList.CreateShaderResourceView(BatchedTilesBuffer, FRHIViewDesc::CreateBufferSRV().SetTypeFromBuffer(BatchedTilesBuffer));
 		}
 	}
 };
@@ -1196,7 +1196,7 @@ bool FSceneRenderState::SetupRayTracingScene(FRDGBuilder& GraphBuilder, FSceneUn
 			{
 				FRHIResourceCreateInfo CreateInfo(TEXT("LightmassRayTracingInstanceUploadBuffer"));
 				InstanceUploadBuffer = RHICmdList.CreateStructuredBuffer(sizeof(FRayTracingInstanceDescriptorInput), InstanceUploadBufferSize, BUF_ShaderResource | BUF_Volatile, CreateInfo);
-				InstanceUploadSRV = RHICmdList.CreateShaderResourceView(InstanceUploadBuffer);
+				InstanceUploadSRV = RHICmdList.CreateShaderResourceView(InstanceUploadBuffer, FRHIViewDesc::CreateBufferSRV().SetTypeFromBuffer(InstanceUploadBuffer));
 			}
 
 			const uint32 TransformUploadBufferSize = RayTracingSceneInitializationData.NumNativeCPUInstances * 3 * sizeof(FVector4f);
@@ -1205,7 +1205,7 @@ bool FSceneRenderState::SetupRayTracingScene(FRDGBuilder& GraphBuilder, FSceneUn
 			{
 				FRHIResourceCreateInfo CreateInfo(TEXT("LightmassRayTracingTransformUploadBuffer"));
 				TransformUploadBuffer = RHICmdList.CreateStructuredBuffer(sizeof(FVector4f), TransformUploadBufferSize, BUF_ShaderResource | BUF_Volatile, CreateInfo);
-				TransformUploadSRV = RHICmdList.CreateShaderResourceView(TransformUploadBuffer);
+				TransformUploadSRV = RHICmdList.CreateShaderResourceView(TransformUploadBuffer, FRHIViewDesc::CreateBufferSRV().SetTypeFromBuffer(TransformUploadBuffer));
 			}
 
 			{
@@ -1857,14 +1857,14 @@ void FLightmapRenderer::Finalize(FRDGBuilder& GraphBuilder)
 						FRHIResourceCreateInfo CreateInfo(TEXT("SrcTilePositionsBuffer"), &SrcTilePositions);
 						CreateInfo.GPUMask = FRHIGPUMask::GPU0();
 						SrcTilePositionsBuffer = RHICmdList.CreateStructuredBuffer(sizeof(FIntPoint), SrcTilePositions.GetResourceDataSize(), BUF_Dynamic | BUF_ShaderResource, CreateInfo);
-						SrcTilePositionsSRV = RHICmdList.CreateShaderResourceView(SrcTilePositionsBuffer);
+						SrcTilePositionsSRV = RHICmdList.CreateShaderResourceView(SrcTilePositionsBuffer, FRHIViewDesc::CreateBufferSRV().SetTypeFromBuffer(SrcTilePositionsBuffer));
 					}
 
 					{
 						FRHIResourceCreateInfo CreateInfo(TEXT("DstTilePositionsBuffer"), &DstTilePositions);
 						CreateInfo.GPUMask = FRHIGPUMask::GPU0();
 						DstTilePositionsBuffer = RHICmdList.CreateStructuredBuffer(sizeof(FIntPoint), DstTilePositions.GetResourceDataSize(), BUF_Dynamic | BUF_ShaderResource, CreateInfo);
-						DstTilePositionsSRV = RHICmdList.CreateShaderResourceView(DstTilePositionsBuffer);
+						DstTilePositionsSRV = RHICmdList.CreateShaderResourceView(DstTilePositionsBuffer, FRHIViewDesc::CreateBufferSRV().SetTypeFromBuffer(DstTilePositionsBuffer));
 					}
 
 					{
@@ -2743,28 +2743,40 @@ void FLightmapRenderer::Finalize(FRDGBuilder& GraphBuilder)
 						FRHIResourceCreateInfo CreateInfo(TEXT("LightTypeBuffer"), &LightTypeArray);
 						CreateInfo.GPUMask = FRHIGPUMask::FromIndex(GPUIndex);
 						LightTypeBuffer = RHICmdList.CreateVertexBuffer(LightTypeArray.GetResourceDataSize(), BUF_Dynamic | BUF_ShaderResource, CreateInfo);
-						LightTypeSRV = RHICmdList.CreateShaderResourceView(LightTypeBuffer, sizeof(int32), PF_R32_SINT);
+						LightTypeSRV = RHICmdList.CreateShaderResourceView(
+							LightTypeBuffer, 
+							FRHIViewDesc::CreateBufferSRV()
+								.SetType(FRHIViewDesc::EBufferType::Typed)
+								.SetFormat(PF_R32_SINT));
 					}
 
 					{
 						FRHIResourceCreateInfo CreateInfo(TEXT("ChannelIndexBuffer"), &ChannelIndexArray);
 						CreateInfo.GPUMask = FRHIGPUMask::FromIndex(GPUIndex);
 						ChannelIndexBuffer = RHICmdList.CreateVertexBuffer(ChannelIndexArray.GetResourceDataSize(), BUF_Dynamic | BUF_ShaderResource, CreateInfo);
-						ChannelIndexSRV = RHICmdList.CreateShaderResourceView(ChannelIndexBuffer, sizeof(int32), PF_R32_SINT);
+						ChannelIndexSRV = RHICmdList.CreateShaderResourceView(
+							ChannelIndexBuffer, 
+							FRHIViewDesc::CreateBufferSRV()
+								.SetType(FRHIViewDesc::EBufferType::Typed)
+								.SetFormat(PF_R32_SINT));
 					}
 
 					{
 						FRHIResourceCreateInfo CreateInfo(TEXT("LightSampleIndexSRV"), &LightSampleIndexArray);
 						CreateInfo.GPUMask = FRHIGPUMask::FromIndex(GPUIndex);
 						LightSampleIndexBuffer = RHICmdList.CreateVertexBuffer(LightSampleIndexArray.GetResourceDataSize(), BUF_Dynamic | BUF_ShaderResource, CreateInfo);
-						LightSampleIndexSRV = RHICmdList.CreateShaderResourceView(LightSampleIndexBuffer, sizeof(int32), PF_R32_SINT);
+						LightSampleIndexSRV = RHICmdList.CreateShaderResourceView(
+							LightSampleIndexBuffer, 
+							FRHIViewDesc::CreateBufferSRV()
+								.SetType(FRHIViewDesc::EBufferType::Typed)
+								.SetFormat(PF_R32_SINT));
 					}
 
 					{
 						FRHIResourceCreateInfo CreateInfo(TEXT("LightShaderParameterBuffer"), &LightShaderParameterArray);
 						CreateInfo.GPUMask = FRHIGPUMask::FromIndex(GPUIndex);
 						LightShaderParameterBuffer = RHICmdList.CreateStructuredBuffer(sizeof(FLightShaderConstants), LightShaderParameterArray.GetResourceDataSize(), BUF_Dynamic | BUF_ShaderResource, CreateInfo);
-						LightShaderParameterSRV = RHICmdList.CreateShaderResourceView(LightShaderParameterBuffer);
+						LightShaderParameterSRV = RHICmdList.CreateShaderResourceView(LightShaderParameterBuffer, FRHIViewDesc::CreateBufferSRV().SetTypeFromBuffer(LightShaderParameterBuffer));
 					}
 
 					// Render GBuffer
diff --git a/Engine/Plugins/Experimental/VirtualHeightfieldMesh/Source/VirtualHeightfieldMesh/Private/VirtualHeightfieldMeshSceneProxy.cpp b/Engine/Plugins/Experimental/VirtualHeightfieldMesh/Source/VirtualHeightfieldMesh/Private/VirtualHeightfieldMeshSceneProxy.cpp
--- a/Engine/Plugins/Experimental/VirtualHeightfieldMesh/Source/VirtualHeightfieldMesh/Private/VirtualHeightfieldMeshSceneProxy.cpp
+++ b/Engine/Plugins/Experimental/VirtualHeightfieldMesh/Source/VirtualHeightfieldMesh/Private/VirtualHeightfieldMeshSceneProxy.cpp
@@ -997,13 +997,20 @@ namespace VirtualHeightfieldMesh
 			const int32 InstanceSize = sizeof(VirtualHeightfieldMesh::QuadRenderInstance);
 			const int32 InstanceBufferSize = CVarVHMMaxRenderItems.GetValueOnRenderThread() * InstanceSize;
 			InBuffers.InstanceBuffer = RHICmdList.CreateStructuredBuffer(InstanceSize, InstanceBufferSize, BUF_UnorderedAccess|BUF_ShaderResource, ERHIAccess::SRVMask, CreateInfo);
-			InBuffers.InstanceBufferUAV = RHICmdList.CreateUnorderedAccessView(InBuffers.InstanceBuffer, false, false);
-			InBuffers.InstanceBufferSRV = RHICmdList.CreateShaderResourceView(InBuffers.InstanceBuffer);
+			InBuffers.InstanceBufferUAV = RHICmdList.CreateUnorderedAccessView(
+				InBuffers.InstanceBuffer, 
+				FRHIViewDesc::CreateBufferUAV()
+					.SetTypeFromBuffer(InBuffers.InstanceBuffer));
+			InBuffers.InstanceBufferSRV = RHICmdList.CreateShaderResourceView(InBuffers.InstanceBuffer, FRHIViewDesc::CreateBufferSRV().SetTypeFromBuffer(InBuffers.InstanceBuffer));
 		}
 		{
 			FRHIResourceCreateInfo CreateInfo(TEXT("VirtualHeightfieldMesh.InstanceIndirectArgsBuffer"));
 			InBuffers.IndirectArgsBuffer = RHICmdList.CreateVertexBuffer(5 * sizeof(uint32), BUF_UnorderedAccess|BUF_DrawIndirect, ERHIAccess::IndirectArgs, CreateInfo);
-			InBuffers.IndirectArgsBufferUAV = RHICmdList.CreateUnorderedAccessView(InBuffers.IndirectArgsBuffer, PF_R32_UINT);
+			InBuffers.IndirectArgsBufferUAV = RHICmdList.CreateUnorderedAccessView(
+				InBuffers.IndirectArgsBuffer, 
+				FRHIViewDesc::CreateBufferUAV()
+					.SetType(FRHIViewDesc::EBufferType::Typed)
+					.SetFormat(PF_R32_UINT));
 		}
 	}
 
diff --git a/Engine/Plugins/Experimental/Water/Source/Runtime/Private/WaterViewExtension.cpp b/Engine/Plugins/Experimental/Water/Source/Runtime/Private/WaterViewExtension.cpp
--- a/Engine/Plugins/Experimental/Water/Source/Runtime/Private/WaterViewExtension.cpp
+++ b/Engine/Plugins/Experimental/Water/Source/Runtime/Private/WaterViewExtension.cpp
@@ -345,7 +345,11 @@ void FWaterViewExtension::UpdateGPUBuffers()
 					ERHIAccess::SRVMask,
 					MakeConstArrayView(WaterAuxDataBuffer)
 				);
-				WaterGPUData->AuxDataSRV = RHICmdList.CreateShaderResourceView(WaterGPUData->AuxDataBuffer, sizeof(FVector4f), PF_A32B32G32R32F);
+				WaterGPUData->AuxDataSRV = RHICmdList.CreateShaderResourceView(
+					WaterGPUData->AuxDataBuffer, 
+					FRHIViewDesc::CreateBufferSRV()
+						.SetType(FRHIViewDesc::EBufferType::Typed)
+						.SetFormat(PF_A32B32G32R32F));
 				
 				WaterGPUData->WaterBodyDataBuffer = UE::RHIResourceUtils::CreateBufferFromArray(
 					RHICmdList,
@@ -354,7 +358,11 @@ void FWaterViewExtension::UpdateGPUBuffers()
 					ERHIAccess::SRVMask,
 					MakeConstArrayView(WaterBodyDataBuffer)
 				);
-				WaterGPUData->WaterBodyDataSRV = RHICmdList.CreateShaderResourceView(WaterGPUData->WaterBodyDataBuffer, sizeof(FVector4f), PF_A32B32G32R32F);
+				WaterGPUData->WaterBodyDataSRV = RHICmdList.CreateShaderResourceView(
+					WaterGPUData->WaterBodyDataBuffer, 
+					FRHIViewDesc::CreateBufferSRV()
+						.SetType(FRHIViewDesc::EBufferType::Typed)
+						.SetFormat(PF_A32B32G32R32F));
 			}
 		);
 
diff --git a/Engine/Plugins/Experimental/Water/Source/Runtime/Public/WaterVertexFactory.h b/Engine/Plugins/Experimental/Water/Source/Runtime/Public/WaterVertexFactory.h
--- a/Engine/Plugins/Experimental/Water/Source/Runtime/Public/WaterVertexFactory.h
+++ b/Engine/Plugins/Experimental/Water/Source/Runtime/Public/WaterVertexFactory.h
@@ -137,7 +137,11 @@ public:
 
 		RHICmdList.UnlockBuffer(VertexBufferRHI);
 
-		SRV = RHICmdList.CreateShaderResourceView(VertexBufferRHI, sizeof(float), PF_R32_FLOAT);
+		SRV = RHICmdList.CreateShaderResourceView(
+			VertexBufferRHI, 
+			FRHIViewDesc::CreateBufferSRV()
+				.SetType(FRHIViewDesc::EBufferType::Typed)
+				.SetFormat(PF_R32_FLOAT));
 	}
 
 	virtual void ReleaseRHI() override
diff --git a/Engine/Plugins/FX/Niagara/Source/Niagara/Classes/NiagaraDataInterfaceArrayImpl.h b/Engine/Plugins/FX/Niagara/Source/Niagara/Classes/NiagaraDataInterfaceArrayImpl.h
--- a/Engine/Plugins/FX/Niagara/Source/Niagara/Classes/NiagaraDataInterfaceArrayImpl.h
+++ b/Engine/Plugins/FX/Niagara/Source/Niagara/Classes/NiagaraDataInterfaceArrayImpl.h
@@ -267,10 +267,18 @@ struct FNDIArrayInstanceData_RenderThread
 			const ERHIAccess DefaultAccess = IsReadOnly() ? ERHIAccess::SRVCompute : ERHIAccess::UAVCompute;
 			ArrayBuffer = RHICmdList.CreateBuffer(ArrayNumBytes, BufferUsage, TypeStride, DefaultAccess, CreateInfo);
 
-			ArraySRV = RHICmdList.CreateShaderResourceView(ArrayBuffer, TypeStride, PixelFormat);
+			ArraySRV = RHICmdList.CreateShaderResourceView(
+				ArrayBuffer, 
+				FRHIViewDesc::CreateBufferSRV()
+					.SetType(FRHIViewDesc::EBufferType::Typed)
+					.SetFormat(PixelFormat));
 			if ( !IsReadOnly() )
 			{
-				ArrayUAV = RHICmdList.CreateUnorderedAccessView(ArrayBuffer, PixelFormat);
+				ArrayUAV = RHICmdList.CreateUnorderedAccessView(
+					ArrayBuffer, 
+					FRHIViewDesc::CreateBufferUAV()
+						.SetType(FRHIViewDesc::EBufferType::Typed)
+						.SetFormat(PixelFormat));
 			}
 		}
 
diff --git a/Engine/Plugins/FX/Niagara/Source/Niagara/Private/DataInterface/NiagaraDataInterfaceDynamicMesh.cpp b/Engine/Plugins/FX/Niagara/Source/Niagara/Private/DataInterface/NiagaraDataInterfaceDynamicMesh.cpp
--- a/Engine/Plugins/FX/Niagara/Source/Niagara/Private/DataInterface/NiagaraDataInterfaceDynamicMesh.cpp
+++ b/Engine/Plugins/FX/Niagara/Source/Niagara/Private/DataInterface/NiagaraDataInterfaceDynamicMesh.cpp
@@ -485,7 +485,11 @@ namespace NDIDynamicMeshLocal
 
 				IndexBuffer.InitResource(RHICmdList);
 				IndexBuffer.IndexBufferRHI = IndexPooledBuffer->GetRHI();
-				IndexBufferSRV = RHICmdList.CreateShaderResourceView(IndexBuffer.IndexBufferRHI, sizeof(uint32), EPixelFormat::PF_R32_UINT);
+				IndexBufferSRV = RHICmdList.CreateShaderResourceView(
+					IndexBuffer.IndexBufferRHI, 
+					FRHIViewDesc::CreateBufferSRV()
+						.SetType(FRHIViewDesc::EBufferType::Typed)
+						.SetFormat(EPixelFormat::PF_R32_UINT));
 			}
 
 			if (NumVertices > 0)
diff --git a/Engine/Plugins/FX/Niagara/Source/Niagara/Private/DataInterface/NiagaraDataInterfaceStaticMesh.cpp b/Engine/Plugins/FX/Niagara/Source/Niagara/Private/DataInterface/NiagaraDataInterfaceStaticMesh.cpp
--- a/Engine/Plugins/FX/Niagara/Source/Niagara/Private/DataInterface/NiagaraDataInterfaceStaticMesh.cpp
+++ b/Engine/Plugins/FX/Niagara/Source/Niagara/Private/DataInterface/NiagaraDataInterfaceStaticMesh.cpp
@@ -447,7 +447,11 @@ namespace NDIStaticMeshLocal
 				if (bCanCreateIndexSRV)
 				{
 					bool b32Bit = GpuInitializeData.LODResource->IndexBuffer.Is32Bit();
-					MeshIndexBufferSRV = RHICmdList.CreateShaderResourceView(IndexBufferRHIRef, b32Bit ? 4 : 2, b32Bit ? PF_R32_UINT : PF_R16_UINT);
+					MeshIndexBufferSRV = RHICmdList.CreateShaderResourceView(
+						IndexBufferRHIRef, 
+						FRHIViewDesc::CreateBufferSRV()
+							.SetType(FRHIViewDesc::EBufferType::Typed)
+							.SetFormat(b32Bit ? PF_R32_UINT : PF_R16_UINT));
 				}
 				else
 				{
diff --git a/Engine/Plugins/FX/Niagara/Source/Niagara/Private/Experimental/NiagaraMeshUvMapping.cpp b/Engine/Plugins/FX/Niagara/Source/Niagara/Private/Experimental/NiagaraMeshUvMapping.cpp
--- a/Engine/Plugins/FX/Niagara/Source/Niagara/Private/Experimental/NiagaraMeshUvMapping.cpp
+++ b/Engine/Plugins/FX/Niagara/Source/Niagara/Private/Experimental/NiagaraMeshUvMapping.cpp
@@ -124,7 +124,11 @@ FMeshUvMappingBufferProxy::InitRHI(FRHICommandListBase& RHICmdList)
 	const int32 BufferSize = FrozenQuadTree.Num();
 
 	UvMappingBuffer = RHICmdList.CreateVertexBuffer(BufferSize, BUF_ShaderResource | BUF_Static, CreateInfo);
-	UvMappingSrv = RHICmdList.CreateShaderResourceView(UvMappingBuffer, sizeof(int32), PF_R32_SINT);
+	UvMappingSrv = RHICmdList.CreateShaderResourceView(
+		UvMappingBuffer, 
+		FRHIViewDesc::CreateBufferSRV()
+			.SetType(FRHIViewDesc::EBufferType::Typed)
+			.SetFormat(PF_R32_SINT));
 
 #if STATS
 	check(GpuMemoryUsage == 0);
diff --git a/Engine/Plugins/FX/Niagara/Source/Niagara/Private/NiagaraDataInterfaceSkeletalMesh.cpp b/Engine/Plugins/FX/Niagara/Source/Niagara/Private/NiagaraDataInterfaceSkeletalMesh.cpp
--- a/Engine/Plugins/FX/Niagara/Source/Niagara/Private/NiagaraDataInterfaceSkeletalMesh.cpp
+++ b/Engine/Plugins/FX/Niagara/Source/Niagara/Private/NiagaraDataInterfaceSkeletalMesh.cpp
@@ -970,7 +970,11 @@ void FSkeletalMeshGpuSpawnStaticBuffers::InitRHI(FRHICommandListBase& RHICmdList
 		uint32* PackedData = (uint32*)RHICmdList.LockBuffer(BufferTriangleUniformSamplerProbAliasRHI, 0, SizeByte, RLM_WriteOnly);
 		NDISkelMeshLocal::PackProbAlias(PackedData, triangleSampler);
 		RHICmdList.UnlockBuffer(BufferTriangleUniformSamplerProbAliasRHI);
-		BufferTriangleUniformSamplerProbAliasSRV = RHICmdList.CreateShaderResourceView(BufferTriangleUniformSamplerProbAliasRHI, sizeof(uint32), PF_R32_UINT);
+		BufferTriangleUniformSamplerProbAliasSRV = RHICmdList.CreateShaderResourceView(
+			BufferTriangleUniformSamplerProbAliasRHI, 
+			FRHIViewDesc::CreateBufferSRV()
+				.SetType(FRHIViewDesc::EBufferType::Typed)
+				.SetFormat(PF_R32_UINT));
 #if STATS
 		GPUMemoryUsage += SizeByte;
 #endif
@@ -991,7 +995,11 @@ void FSkeletalMeshGpuSpawnStaticBuffers::InitRHI(FRHICommandListBase& RHICmdList
 		{
 			CreateInfo.ResourceArray = &SampleRegionsProbAlias;
 			SampleRegionsProbAliasBuffer = RHICmdList.CreateVertexBuffer(SampleRegionsProbAlias.Num() * SampleRegionsProbAlias.GetTypeSize(), BUF_Static | BUF_ShaderResource, CreateInfo);
-			SampleRegionsProbAliasSRV = RHICmdList.CreateShaderResourceView(SampleRegionsProbAliasBuffer, sizeof(uint32), PF_R32_UINT);
+			SampleRegionsProbAliasSRV = RHICmdList.CreateShaderResourceView(
+				SampleRegionsProbAliasBuffer, 
+				FRHIViewDesc::CreateBufferSRV()
+					.SetType(FRHIViewDesc::EBufferType::Typed)
+					.SetFormat(PF_R32_UINT));
 #if STATS
 			GPUMemoryUsage += SampleRegionsProbAlias.Num() * SampleRegionsProbAlias.GetTypeSize();
 #endif
@@ -1004,7 +1012,11 @@ void FSkeletalMeshGpuSpawnStaticBuffers::InitRHI(FRHICommandListBase& RHICmdList
 		CreateInfo.DebugName = (TEXT("SampleRegionsTriangleIndicesBuffer"));
 		CreateInfo.ResourceArray = &SampleRegionsTriangleIndicies;
 		SampleRegionsTriangleIndicesBuffer = RHICmdList.CreateVertexBuffer(SampleRegionsTriangleIndicies.Num() * SampleRegionsTriangleIndicies.GetTypeSize(), BUF_Static | BUF_ShaderResource, CreateInfo);
-		SampleRegionsTriangleIndicesSRV = RHICmdList.CreateShaderResourceView(SampleRegionsTriangleIndicesBuffer, sizeof(int32), PF_R32_UINT);
+		SampleRegionsTriangleIndicesSRV = RHICmdList.CreateShaderResourceView(
+			SampleRegionsTriangleIndicesBuffer, 
+			FRHIViewDesc::CreateBufferSRV()
+				.SetType(FRHIViewDesc::EBufferType::Typed)
+				.SetFormat(PF_R32_UINT));
 #if STATS
 		GPUMemoryUsage += SampleRegionsTriangleIndicies.Num() * SampleRegionsTriangleIndicies.GetTypeSize();
 #endif
@@ -1012,7 +1024,11 @@ void FSkeletalMeshGpuSpawnStaticBuffers::InitRHI(FRHICommandListBase& RHICmdList
 		CreateInfo.DebugName = (TEXT("SampleRegionsVerticesBuffer"));
 		CreateInfo.ResourceArray = &SampleRegionsVertices;
 		SampleRegionsVerticesBuffer = RHICmdList.CreateVertexBuffer(SampleRegionsVertices.Num() * SampleRegionsVertices.GetTypeSize(), BUF_Static | BUF_ShaderResource, CreateInfo);
-		SampleRegionsVerticesSRV = RHICmdList.CreateShaderResourceView(SampleRegionsVerticesBuffer, sizeof(int32), PF_R32_UINT);
+		SampleRegionsVerticesSRV = RHICmdList.CreateShaderResourceView(
+			SampleRegionsVerticesBuffer, 
+			FRHIViewDesc::CreateBufferSRV()
+				.SetType(FRHIViewDesc::EBufferType::Typed)
+				.SetFormat(PF_R32_UINT));
 #if STATS
 		GPUMemoryUsage += SampleRegionsVertices.Num() * SampleRegionsVertices.GetTypeSize();
 #endif
@@ -1045,7 +1061,11 @@ void FSkeletalMeshGpuSpawnStaticBuffers::InitRHI(FRHICommandListBase& RHICmdList
 			AccumulatedMatrixOffset += Section.BoneMap.Num();
 		}
 		RHICmdList.UnlockBuffer(BufferTriangleMatricesOffsetRHI);
-		BufferTriangleMatricesOffsetSRV = RHICmdList.CreateShaderResourceView(BufferTriangleMatricesOffsetRHI, sizeof(uint32), PF_R32_UINT);
+		BufferTriangleMatricesOffsetSRV = RHICmdList.CreateShaderResourceView(
+			BufferTriangleMatricesOffsetRHI, 
+			FRHIViewDesc::CreateBufferSRV()
+				.SetType(FRHIViewDesc::EBufferType::Typed)
+				.SetFormat(PF_R32_UINT));
 #if STATS
 		GPUMemoryUsage += VertexCount * sizeof(uint32);
 #endif
@@ -1062,7 +1082,11 @@ void FSkeletalMeshGpuSpawnStaticBuffers::InitRHI(FRHICommandListBase& RHICmdList
 		CreateInfo.ResourceArray = &FilteredAndUnfilteredBonesArray;
 
 		FilteredAndUnfilteredBonesBuffer = RHICmdList.CreateVertexBuffer(FilteredAndUnfilteredBonesArray.Num() * FilteredAndUnfilteredBonesArray.GetTypeSize(), BUF_Static | BUF_ShaderResource, CreateInfo);
-		FilteredAndUnfilteredBonesSRV = RHICmdList.CreateShaderResourceView(FilteredAndUnfilteredBonesBuffer, sizeof(uint16), PF_R16_UINT);
+		FilteredAndUnfilteredBonesSRV = RHICmdList.CreateShaderResourceView(
+			FilteredAndUnfilteredBonesBuffer,
+			FRHIViewDesc::CreateBufferSRV()
+				.SetType(FRHIViewDesc::EBufferType::Typed)
+				.SetFormat(PF_R16_UINT));
 	}
 	else
 	{
@@ -1141,10 +1165,18 @@ void FSkeletalMeshGpuDynamicBufferProxy::InitRHI(FRHICommandListBase& RHICmdList
 	{
 		FRHIResourceCreateInfo CreateInfo(TEXT("SkeletalMeshGpuDynamicBuffer"));
 		Buffer.SectionBuffer = RHICmdList.CreateVertexBuffer(sizeof(FVector4f) * 3 * SectionBoneCount, BUF_ShaderResource | BUF_Dynamic, CreateInfo);
-		Buffer.SectionSRV = RHICmdList.CreateShaderResourceView(Buffer.SectionBuffer, sizeof(FVector4f), PF_A32B32G32R32F);
+		Buffer.SectionSRV = RHICmdList.CreateShaderResourceView(
+			Buffer.SectionBuffer, 
+			FRHIViewDesc::CreateBufferSRV()
+				.SetType(FRHIViewDesc::EBufferType::Typed)
+				.SetFormat(PF_A32B32G32R32F));
 
 		Buffer.SamplingBuffer = RHICmdList.CreateVertexBuffer(sizeof(FVector4f) * 3 * (SamplingBoneCount + SamplingSocketCount), BUF_ShaderResource | BUF_Dynamic, CreateInfo);
-		Buffer.SamplingSRV = RHICmdList.CreateShaderResourceView(Buffer.SamplingBuffer, sizeof(FVector4f), PF_A32B32G32R32F);
+		Buffer.SamplingSRV = RHICmdList.CreateShaderResourceView(
+			Buffer.SamplingBuffer, 
+			FRHIViewDesc::CreateBufferSRV()
+				.SetType(FRHIViewDesc::EBufferType::Typed)
+				.SetFormat(PF_A32B32G32R32F));
 
 #if STATS
 		GPUMemoryUsage += sizeof(FVector4f) * 3 * SectionBoneCount;
diff --git a/Engine/Plugins/FX/Niagara/Source/Niagara/Private/NiagaraDataInterfaceSkeletalMeshConnectivity.cpp b/Engine/Plugins/FX/Niagara/Source/Niagara/Private/NiagaraDataInterfaceSkeletalMeshConnectivity.cpp
--- a/Engine/Plugins/FX/Niagara/Source/Niagara/Private/NiagaraDataInterfaceSkeletalMeshConnectivity.cpp
+++ b/Engine/Plugins/FX/Niagara/Source/Niagara/Private/NiagaraDataInterfaceSkeletalMeshConnectivity.cpp
@@ -406,7 +406,11 @@ FSkeletalMeshConnectivityProxy::InitRHI(FRHICommandListBase& RHICmdList)
 	const int32 BufferSize = AdjacencyResource.Num();
 
 	AdjacencyBuffer = RHICmdList.CreateVertexBuffer(BufferSize, BUF_ShaderResource | BUF_Static, CreateInfo);
-	AdjacencySrv = RHICmdList.CreateShaderResourceView(AdjacencyBuffer, sizeof(uint32), PF_R32_UINT);
+	AdjacencySrv = RHICmdList.CreateShaderResourceView(
+		AdjacencyBuffer, 
+		FRHIViewDesc::CreateBufferSRV()
+			.SetType(FRHIViewDesc::EBufferType::Typed)
+			.SetFormat(PF_R32_UINT));
 
 #if STATS
 	check(GpuMemoryUsage == 0);
diff --git a/Engine/Plugins/FX/Niagara/Source/Niagara/Private/NiagaraEmptyUAVPool.cpp b/Engine/Plugins/FX/Niagara/Source/Niagara/Private/NiagaraEmptyUAVPool.cpp
--- a/Engine/Plugins/FX/Niagara/Source/Niagara/Private/NiagaraEmptyUAVPool.cpp
+++ b/Engine/Plugins/FX/Niagara/Source/Niagara/Private/NiagaraEmptyUAVPool.cpp
@@ -72,7 +72,11 @@ FRHIUnorderedAccessView* FNiagaraEmptyUAVPool::GetEmptyUAVFromPool(FRHICommandLi
 			{
 				const uint32 BytesPerElement = GPixelFormats[Format].BlockBytes;
 				NewUAV.Buffer = RHICmdList.CreateVertexBuffer(BytesPerElement, BUF_UnorderedAccess | BUF_ShaderResource, CreateInfo);
-				NewUAV.UAV = RHICmdList.CreateUnorderedAccessView(NewUAV.Buffer, Format);
+				NewUAV.UAV = RHICmdList.CreateUnorderedAccessView(
+					NewUAV.Buffer, 
+					FRHIViewDesc::CreateBufferUAV()
+						.SetType(FRHIViewDesc::EBufferType::Typed)
+						.SetFormat(Format));
 				break;
 			}
 			
@@ -83,7 +87,7 @@ FRHIUnorderedAccessView* FNiagaraEmptyUAVPool::GetEmptyUAVFromPool(FRHICommandLi
 					.SetFlags(ETextureCreateFlags::ShaderResource | ETextureCreateFlags::UAV);
 
 				NewUAV.Texture = RHICreateTexture(Desc);
-				NewUAV.UAV = RHICmdList.CreateUnorderedAccessView(NewUAV.Texture, 0);
+				NewUAV.UAV = RHICmdList.CreateUnorderedAccessView(NewUAV.Texture, FRHIViewDesc::CreateTextureUAV().SetDimensionFromTexture(NewUAV.Texture));
 				break;
 			}
 			
@@ -94,7 +98,7 @@ FRHIUnorderedAccessView* FNiagaraEmptyUAVPool::GetEmptyUAVFromPool(FRHICommandLi
 					.SetFlags(ETextureCreateFlags::ShaderResource | ETextureCreateFlags::UAV);
 
 				NewUAV.Texture = RHICreateTexture(Desc);
-				NewUAV.UAV = RHICmdList.CreateUnorderedAccessView(NewUAV.Texture, 0);
+				NewUAV.UAV = RHICmdList.CreateUnorderedAccessView(NewUAV.Texture, FRHIViewDesc::CreateTextureUAV().SetDimensionFromTexture(NewUAV.Texture));
 				break;
 			}
 			
@@ -105,7 +109,7 @@ FRHIUnorderedAccessView* FNiagaraEmptyUAVPool::GetEmptyUAVFromPool(FRHICommandLi
 					.SetFlags(ETextureCreateFlags::ShaderResource | ETextureCreateFlags::UAV);
 
 				NewUAV.Texture = RHICreateTexture(Desc);
-				NewUAV.UAV = RHICmdList.CreateUnorderedAccessView(NewUAV.Texture, 0);
+				NewUAV.UAV = RHICmdList.CreateUnorderedAccessView(NewUAV.Texture, FRHIViewDesc::CreateTextureUAV().SetDimensionFromTexture(NewUAV.Texture));
 				break;
 			}
 			
@@ -116,7 +120,7 @@ FRHIUnorderedAccessView* FNiagaraEmptyUAVPool::GetEmptyUAVFromPool(FRHICommandLi
 					.SetFlags(ETextureCreateFlags::ShaderResource | ETextureCreateFlags::UAV);
 
 				NewUAV.Texture = RHICreateTexture(Desc);
-				NewUAV.UAV = RHICmdList.CreateUnorderedAccessView(NewUAV.Texture, 0);
+				NewUAV.UAV = RHICmdList.CreateUnorderedAccessView(NewUAV.Texture, FRHIViewDesc::CreateTextureUAV().SetDimensionFromTexture(NewUAV.Texture));
 				break;
 			}
 
@@ -128,7 +132,7 @@ FRHIUnorderedAccessView* FNiagaraEmptyUAVPool::GetEmptyUAVFromPool(FRHICommandLi
 					.SetFlags(ETextureCreateFlags::ShaderResource | ETextureCreateFlags::UAV);
 
 				NewUAV.Texture = RHICreateTexture(Desc);
-				NewUAV.UAV = RHICmdList.CreateUnorderedAccessView(NewUAV.Texture, 0);
+				NewUAV.UAV = RHICmdList.CreateUnorderedAccessView(NewUAV.Texture, FRHIViewDesc::CreateTextureUAV().SetDimensionFromTexture(NewUAV.Texture));
 				break;
 			}
 
diff --git a/Engine/Plugins/FX/Niagara/Source/Niagara/Private/NiagaraRenderer.cpp b/Engine/Plugins/FX/Niagara/Source/Niagara/Private/NiagaraRenderer.cpp
--- a/Engine/Plugins/FX/Niagara/Source/Niagara/Private/NiagaraRenderer.cpp
+++ b/Engine/Plugins/FX/Niagara/Source/Niagara/Private/NiagaraRenderer.cpp
@@ -69,7 +69,11 @@ public:
 
 		RHICmdList.UnlockBuffer(Buffer);
 
-		SRV = RHICmdList.CreateShaderResourceView(Buffer, NumBytes, PixelFormat);
+		SRV = RHICmdList.CreateShaderResourceView(
+			Buffer, 
+			FRHIViewDesc::CreateBufferSRV()
+				.SetType(FRHIViewDesc::EBufferType::Typed)
+				.SetFormat(PixelFormat));
 	}
 
 	virtual void ReleaseRHI() override
@@ -144,7 +148,7 @@ public:
 				return;
 		}
 
-		SRV = RHICmdList.CreateShaderResourceView(Texture, 0);
+		SRV = RHICmdList.CreateShaderResourceView(Texture, FRHIViewDesc::CreateTextureSRV().SetDimensionFromTexture(Texture));
 	}
 
 	virtual void ReleaseRHI() override
diff --git a/Engine/Plugins/FX/Niagara/Source/Niagara/Private/NiagaraRendererRibbons.cpp b/Engine/Plugins/FX/Niagara/Source/Niagara/Private/NiagaraRendererRibbons.cpp
--- a/Engine/Plugins/FX/Niagara/Source/Niagara/Private/NiagaraRendererRibbons.cpp
+++ b/Engine/Plugins/FX/Niagara/Source/Niagara/Private/NiagaraRendererRibbons.cpp
@@ -280,7 +280,11 @@ struct FNiagaraRibbonIndexBuffer final : FIndexBuffer
 		if (IsRayTracingAllowed())
 		{
 			const uint32 IndexStride = IndexBufferRHI->GetDesc().Stride;
-			SRV = RHICmdList.CreateShaderResourceView(IndexBufferRHI, IndexStride, IndexStride == 2 ? PF_R16_UINT : PF_R32_UINT);
+			SRV = RHICmdList.CreateShaderResourceView(
+				IndexBufferRHI, 
+				FRHIViewDesc::CreateBufferSRV()
+					.SetType(FRHIViewDesc::EBufferType::Typed)
+					.SetFormat(IndexStride == 2 ? PF_R16_UINT : PF_R32_UINT));
 		}
 	#endif
 	}
@@ -293,9 +297,17 @@ struct FNiagaraRibbonIndexBuffer final : FIndexBuffer
 		FRHIResourceCreateInfo CreateInfo(TEXT("NiagaraRibbonIndexBuffer"));
 		IndexBufferRHI = RHICmdList.CreateIndexBuffer(sizeof(uint32), sizeof(uint32) * NumElements, BUF_Static | BUF_UnorderedAccess | BUF_ShaderResource, ERHIAccess::VertexOrIndexBuffer, CreateInfo);
 	#if RHI_RAYTRACING
-		SRV = RHICmdList.CreateShaderResourceView(IndexBufferRHI, sizeof(uint32), PF_R32_UINT);
+		SRV = RHICmdList.CreateShaderResourceView(
+			IndexBufferRHI, 
+			FRHIViewDesc::CreateBufferSRV()
+				.SetType(FRHIViewDesc::EBufferType::Typed)
+				.SetFormat(PF_R32_UINT));
 	#endif
-		UAV = RHICmdList.CreateUnorderedAccessView(IndexBufferRHI, PF_R32_UINT);
+		UAV = RHICmdList.CreateUnorderedAccessView(
+			IndexBufferRHI, 
+			FRHIViewDesc::CreateBufferUAV()
+				.SetType(FRHIViewDesc::EBufferType::Typed)
+				.SetFormat(PF_R32_UINT));
 	}
 
 	virtual void ReleaseRHI() override
@@ -449,8 +461,16 @@ bool FNiagaraRibbonGpuBuffer::Allocate(FRHICommandListBase& RHICmdList, uint32 N
 
 			NumBytes = ElementBytes * NumElements;
 			Buffer = RHICmdList.CreateVertexBuffer(NumBytes, UsageFlags, InResourceState, CreateInfo);
-			SRV = RHICmdList.CreateShaderResourceView(Buffer, ElementBytes, UE_PIXELFORMAT_TO_UINT8(PixelFormat));
-			UAV = bGpuReadOnly ? nullptr : RHICmdList.CreateUnorderedAccessView(Buffer, UE_PIXELFORMAT_TO_UINT8(PixelFormat));
+			SRV = RHICmdList.CreateShaderResourceView(
+				Buffer, 
+				FRHIViewDesc::CreateBufferSRV()
+					.SetType(FRHIViewDesc::EBufferType::Typed)
+					.SetFormat(PixelFormat));
+			UAV = bGpuReadOnly ? nullptr : RHICmdList.CreateUnorderedAccessView(
+				Buffer, 
+				FRHIViewDesc::CreateBufferUAV()
+					.SetType(FRHIViewDesc::EBufferType::Typed)
+					.SetFormat(PixelFormat));
 			return true;
 		}
 	}
diff --git a/Engine/Plugins/FX/Niagara/Source/NiagaraVertexFactories/Private/NiagaraCutoutVertexBuffer.cpp b/Engine/Plugins/FX/Niagara/Source/NiagaraVertexFactories/Private/NiagaraCutoutVertexBuffer.cpp
--- a/Engine/Plugins/FX/Niagara/Source/NiagaraVertexFactories/Private/NiagaraCutoutVertexBuffer.cpp
+++ b/Engine/Plugins/FX/Niagara/Source/NiagaraVertexFactories/Private/NiagaraCutoutVertexBuffer.cpp
@@ -28,7 +28,11 @@ void FNiagaraCutoutVertexBuffer::InitRHI(FRHICommandListBase& RHICmdList)
 		void* BufferData = RHICmdList.LockBuffer(VertexBufferRHI, 0, DataSize, RLM_WriteOnly);
 		FMemory::Memcpy(BufferData, Data.GetData(), DataSize);
 		RHICmdList.UnlockBuffer(VertexBufferRHI);
-		VertexBufferSRV = RHICmdList.CreateShaderResourceView(VertexBufferRHI, sizeof(FVector2f), PF_G32R32F);
+		VertexBufferSRV = RHICmdList.CreateShaderResourceView(
+			VertexBufferRHI, 
+			FRHIViewDesc::CreateBufferSRV()
+				.SetType(FRHIViewDesc::EBufferType::Typed)
+				.SetFormat(PF_G32R32F));
 
 		Data.Empty();
 	}
diff --git a/Engine/Plugins/FX/Niagara/Source/NiagaraVertexFactories/Public/NiagaraVertexFactory.h b/Engine/Plugins/FX/Niagara/Source/NiagaraVertexFactories/Public/NiagaraVertexFactory.h
--- a/Engine/Plugins/FX/Niagara/Source/NiagaraVertexFactories/Public/NiagaraVertexFactory.h
+++ b/Engine/Plugins/FX/Niagara/Source/NiagaraVertexFactories/Public/NiagaraVertexFactory.h
@@ -30,7 +30,11 @@ public:
 		FMemory::Memzero(BufferData, sizeof(uint32));
 		RHICmdList.UnlockBuffer(VertexBufferRHI);
 
-		VertexBufferSRV = RHICmdList.CreateShaderResourceView(VertexBufferRHI, sizeof(uint32), PF_R32_UINT);
+		VertexBufferSRV = RHICmdList.CreateShaderResourceView(
+			VertexBufferRHI, 
+			FRHIViewDesc::CreateBufferSRV()
+				.SetType(FRHIViewDesc::EBufferType::Typed)
+				.SetFormat(PF_R32_UINT));
 	}
 
 	virtual void ReleaseRHI() override
diff --git a/Engine/Plugins/Media/ElectraUtil/Source/ElectraSamples/Private/ElectraTextureSampleAVF.cpp b/Engine/Plugins/Media/ElectraUtil/Source/ElectraSamples/Private/ElectraTextureSampleAVF.cpp
--- a/Engine/Plugins/Media/ElectraUtil/Source/ElectraSamples/Private/ElectraTextureSampleAVF.cpp
+++ b/Engine/Plugins/Media/ElectraUtil/Source/ElectraSamples/Private/ElectraTextureSampleAVF.cpp
@@ -255,8 +255,18 @@ void FElectraMediaTexConvApple::ConvertTexture(FTextureRHIRef & InDstTexture, CV
 						GraphicsPSOInit.BoundShaderState.PixelShaderRHI = PixelShader.GetPixelShader();
 						SetGraphicsPipelineState(RHICmdList, GraphicsPSOInit, 0);
 
-						FShaderResourceViewRHIRef Y_SRV = RHICmdList.CreateShaderResourceView(YTex, 0, 1, PF_G8);
-						FShaderResourceViewRHIRef UV_SRV = RHICmdList.CreateShaderResourceView(UVTex, 0, 1, PF_R8G8);
+						FShaderResourceViewRHIRef Y_SRV = RHICmdList.CreateShaderResourceView(
+							YTex, 
+							FRHIViewDesc::CreateTextureSRV()
+								.SetDimensionFromTexture(GetTextureInternal(0))
+								.SetMipRange(0, 1)
+								.SetFormat(PF_G8));
+						FShaderResourceViewRHIRef UV_SRV = RHICmdList.CreateShaderResourceView(
+							UVTex, 
+							FRHIViewDesc::CreateTextureSRV()
+								.SetDimensionFromTexture(GetTextureInternal(0))
+								.SetMipRange(0, 1)
+								.SetFormat(PF_R8G8));
 
 						SetShaderParametersLegacyPS(RHICmdList, PixelShader, FIntPoint(YWidth, YHeight), Y_SRV, UV_SRV, FIntPoint(YWidth, YHeight), YUVMtx, EncodingType, ColorSpaceMtx, false);
 					}
@@ -272,8 +282,18 @@ void FElectraMediaTexConvApple::ConvertTexture(FTextureRHIRef & InDstTexture, CV
 						GraphicsPSOInit.BoundShaderState.PixelShaderRHI = PixelShader.GetPixelShader();
 						SetGraphicsPipelineState(RHICmdList, GraphicsPSOInit, 0);
 
-						FShaderResourceViewRHIRef Y_SRV = RHICmdList.CreateShaderResourceView(YTex, 0, 1, PF_G16);
-						FShaderResourceViewRHIRef UV_SRV = RHICmdList.CreateShaderResourceView(UVTex, 0, 1, PF_G16R16);
+						FShaderResourceViewRHIRef Y_SRV = RHICmdList.CreateShaderResourceView(
+							YTex, 
+							FRHIViewDesc::CreateTextureSRV()
+								.SetDimensionFromTexture(GetTextureInternal(0))
+								.SetMipRange(0, 1)
+								.SetFormat(PF_G16));
+						FShaderResourceViewRHIRef UV_SRV = RHICmdList.CreateShaderResourceView(
+							UVTex, 
+							FRHIViewDesc::CreateTextureSRV()
+								.SetDimensionFromTexture(GetTextureInternal(0))
+								.SetMipRange(0, 1)
+								.SetFormat(PF_G16R16));
 						
 						// Update shader uniform parameters.
 						SetShaderParametersLegacyPS(RHICmdList, PixelShader, FIntPoint(YWidth, YHeight), Y_SRV, UV_SRV, FIntPoint(YWidth, YHeight), YUVMtx, ColorSpaceMtx, EncodingType);
diff --git a/Engine/Plugins/Media/ImgMedia/Source/ImgMedia/Private/Readers/ExrImgMediaReaderGpu.cpp b/Engine/Plugins/Media/ImgMedia/Source/ImgMedia/Private/Readers/ExrImgMediaReaderGpu.cpp
--- a/Engine/Plugins/Media/ImgMedia/Source/ImgMedia/Private/Readers/ExrImgMediaReaderGpu.cpp
+++ b/Engine/Plugins/Media/ImgMedia/Source/ImgMedia/Private/Readers/ExrImgMediaReaderGpu.cpp
@@ -515,7 +515,7 @@ void FExrImgMediaReaderGpu::CreateSampleConverterCallback(TSharedPtr<FExrMediaTe
 				void* MappedBuffer = RHICmdList.LockBuffer(BufferRef, 0, NumElements, RLM_WriteOnly);
 				FMemory::Memcpy(MappedBuffer, &ConverterParams.TileInfoPerMipLevel[SampleMipLevel][0], sizeof(FExrReader::FTileDesc) * ConverterParams.TileInfoPerMipLevel[SampleMipLevel].Num());
 				RHICmdList.UnlockBuffer(BufferRef);
-				Parameters.TileDescBuffer = RHICmdList.CreateShaderResourceView(BufferRef);
+				Parameters.TileDescBuffer = RHICmdList.CreateShaderResourceView(BufferRef, FRHIViewDesc::CreateBufferSRV().SetTypeFromBuffer(BufferRef));
 				PermutationVector.Set<FExrSwizzlePS::FPartialTiles>(true);
 			}
 
@@ -661,11 +661,11 @@ FStructuredBufferPoolItemSharedPtr FExrImgMediaReaderGpu::AllocateGpuBufferFromP
 				{
 					CreateInfo.DebugName = TEXT("ExrReaderGpu.DestBuffer");
 					AllocatedBuffer->ShaderAccessBufferRef = RHICmdList.CreateStructuredBuffer(sizeof(uint16) * 2., AllocSize, BUF_ShaderResource | BUF_FastVRAM, CreateInfo);
-					AllocatedBuffer->ShaderResourceView = RHICmdList.CreateShaderResourceView(AllocatedBuffer->ShaderAccessBufferRef);
+					AllocatedBuffer->ShaderResourceView = RHICmdList.CreateShaderResourceView(AllocatedBuffer->ShaderAccessBufferRef, FRHIViewDesc::CreateBufferSRV().SetTypeFromBuffer(AllocatedBuffer->ShaderAccessBufferRef));
 				}
 				else
 				{
-					AllocatedBuffer->ShaderResourceView = RHICmdList.CreateShaderResourceView(AllocatedBuffer->UploadBufferRef);
+					AllocatedBuffer->ShaderResourceView = RHICmdList.CreateShaderResourceView(AllocatedBuffer->UploadBufferRef, FRHIViewDesc::CreateBufferSRV().SetTypeFromBuffer(AllocatedBuffer->UploadBufferRef));
 				}
 
 				AllocatedBuffer->AllocationReadyEvent->Trigger();
diff --git a/Engine/Plugins/Media/PixelCapture/Source/PixelCapture/Private/PixelCaptureCapturerRHIToI420Compute.cpp b/Engine/Plugins/Media/PixelCapture/Source/PixelCapture/Private/PixelCaptureCapturerRHIToI420Compute.cpp
--- a/Engine/Plugins/Media/PixelCapture/Source/PixelCapture/Private/PixelCaptureCapturerRHIToI420Compute.cpp
+++ b/Engine/Plugins/Media/PixelCapture/Source/PixelCapture/Private/PixelCaptureCapturerRHIToI420Compute.cpp
@@ -82,9 +82,9 @@ void FPixelCaptureCapturerRHIToI420Compute::Initialize(int32 InputWidth, int32 I
 	StagingTextureU = RHICreateTexture(StagingDescUV);
 	StagingTextureV = RHICreateTexture(StagingDescUV);
 
-	TextureYUAV = RHICmdList.CreateUnorderedAccessView(TextureY);
-	TextureUUAV = RHICmdList.CreateUnorderedAccessView(TextureU);
-	TextureVUAV = RHICmdList.CreateUnorderedAccessView(TextureV);
+	TextureYUAV = RHICmdList.CreateUnorderedAccessView(TextureY, FRHIViewDesc::CreateTextureUAV().SetDimensionFromTexture(TextureY));
+	TextureUUAV = RHICmdList.CreateUnorderedAccessView(TextureU, FRHIViewDesc::CreateTextureUAV().SetDimensionFromTexture(TextureU));
+	TextureVUAV = RHICmdList.CreateUnorderedAccessView(TextureV, FRHIViewDesc::CreateTextureUAV().SetDimensionFromTexture(TextureV));
 
 	int32 OutWidth, OutHeight;
 	RHICmdList.MapStagingSurface(StagingTextureY, nullptr, MappedY, OutWidth, OutHeight);
diff --git a/Engine/Plugins/Media/WmfMedia/Source/WmfMedia/Private/Wmf/WmfMediaHardwareVideoDecodingRendering.cpp b/Engine/Plugins/Media/WmfMedia/Source/WmfMedia/Private/Wmf/WmfMediaHardwareVideoDecodingRendering.cpp
--- a/Engine/Plugins/Media/WmfMedia/Source/WmfMedia/Private/Wmf/WmfMediaHardwareVideoDecodingRendering.cpp
+++ b/Engine/Plugins/Media/WmfMedia/Source/WmfMedia/Private/Wmf/WmfMediaHardwareVideoDecodingRendering.cpp
@@ -173,8 +173,18 @@ bool FWmfMediaHardwareVideoDecodingParameters::ConvertTextureFormat_RenderThread
 
 			SetGraphicsPipelineState(RHICmdList, GraphicsPSOInit, 0);
 
-			FShaderResourceViewRHIRef Y_SRV = RHICmdList.CreateShaderResourceView(SampleDestinationTexture, 0, 1, PF_G8);
-			FShaderResourceViewRHIRef UV_SRV = RHICmdList.CreateShaderResourceView(SampleDestinationTexture, 0, 1, PF_R8G8);
+			FShaderResourceViewRHIRef Y_SRV = RHICmdList.CreateShaderResourceView(
+				SampleDestinationTexture, 
+				FRHIViewDesc::CreateTextureSRV()
+					.SetDimensionFromTexture(SampleDestinationTexture)
+					.SetMipRange(0, 1)
+					.SetFormat(PF_G8));
+			FShaderResourceViewRHIRef UV_SRV = RHICmdList.CreateShaderResourceView(
+				SampleDestinationTexture, 
+				FRHIViewDesc::CreateTextureSRV()
+					.SetDimensionFromTexture(SampleDestinationTexture)
+					.SetMipRange(0, 1)
+					.SetFormat(PF_R8G8));
 
 			const FHardwareVideoDecodingShaderParams Parameters = FWmfMediaHardwareVideoDecodingShader::GetParameters(Y_SRV, UV_SRV, InSample->IsOutputSrgb());
 
@@ -201,7 +211,11 @@ bool FWmfMediaHardwareVideoDecodingParameters::ConvertTextureFormat_RenderThread
 			{
 				PixelFormat = PF_DXT5;
 			}
-			FShaderResourceViewRHIRef SRV = RHICmdList.CreateShaderResourceView(SampleDestinationTexture, 0, 1, PixelFormat);
+			FShaderResourceViewRHIRef SRV = RHICmdList.CreateShaderResourceView(
+				SampleDestinationTexture, 
+				FRHIViewDesc::CreateTextureSRV()
+					.SetDimensionFromTexture(SampleDestinationTexture)
+					.SetFormat(PixelFormat));
 
 			const FHardwareVideoDecodingShaderParams Parameters = FWmfMediaHardwareVideoDecodingShader::GetParameters(SRV, InSample->IsOutputSrgb());
 
@@ -217,7 +231,11 @@ bool FWmfMediaHardwareVideoDecodingParameters::ConvertTextureFormat_RenderThread
 
 			SetGraphicsPipelineState(RHICmdList, GraphicsPSOInit, 0);
 
-			FShaderResourceViewRHIRef SRV = RHICmdList.CreateShaderResourceView(SampleDestinationTexture, 0, 1, PF_A16B16G16R16);
+			FShaderResourceViewRHIRef SRV = RHICmdList.CreateShaderResourceView(
+				SampleDestinationTexture, 
+				FRHIViewDesc::CreateTextureSRV()
+					.SetDimensionFromTexture(SampleDestinationTexture)
+					.SetFormat(PF_A16B16G16R16));
 
 			const FHardwareVideoDecodingShaderParams Parameters = FWmfMediaHardwareVideoDecodingShader::GetParameters(SRV, InSample->IsOutputSrgb());
 
@@ -233,7 +251,11 @@ bool FWmfMediaHardwareVideoDecodingParameters::ConvertTextureFormat_RenderThread
 
 			SetGraphicsPipelineState(RHICmdList, GraphicsPSOInit, 0);
 
-			FShaderResourceViewRHIRef SRV = RHICmdList.CreateShaderResourceView(SampleDestinationTexture, 0, 1, PF_DXT5);
+			FShaderResourceViewRHIRef SRV = RHICmdList.CreateShaderResourceView(
+				SampleDestinationTexture, 
+				FRHIViewDesc::CreateTextureSRV()
+					.SetDimensionFromTexture(SampleDestinationTexture)
+					.SetFormat(PF_DXT5));
 
 			const FHardwareVideoDecodingShaderParams Parameters = FWmfMediaHardwareVideoDecodingShader::GetParameters(SRV, InSample->IsOutputSrgb());
 
@@ -260,8 +282,16 @@ bool FWmfMediaHardwareVideoDecodingParameters::ConvertTextureFormat_RenderThread
 
 			SetGraphicsPipelineState(RHICmdList, GraphicsPSOInit, 0);
 
-			FShaderResourceViewRHIRef SRV = RHICmdList.CreateShaderResourceView(SampleDestinationTexture, 0, 1, PF_DXT5);
-			FShaderResourceViewRHIRef SRVAlpha = RHICmdList.CreateShaderResourceView(SampleDestinationAlphaTexture, 0, 1, PF_BC4);
+			FShaderResourceViewRHIRef SRV = RHICmdList.CreateShaderResourceView(
+				SampleDestinationTexture, 
+				FRHIViewDesc::CreateTextureSRV()
+					.SetDimensionFromTexture(SampleDestinationTexture)
+					.SetFormat(PF_DXT5));
+			FShaderResourceViewRHIRef SRVAlpha = RHICmdList.CreateShaderResourceView(
+				SampleDestinationAlphaTexture, 
+				FRHIViewDesc::CreateTextureSRV()
+					.SetDimensionFromTexture(SampleDestinationAlphaTexture)
+					.SetFormat(PF_BC4));
 
 			const FHardwareVideoDecodingShaderParams VsParameters = FWmfMediaHardwareVideoDecodingShader::GetParameters(SRV, InSample->IsOutputSrgb());
 			const FHardwareVideoDecodingShaderParams PsParameters = FWmfMediaHardwareVideoDecodingShader::GetParameters(SRV, SRVAlpha, InSample->IsOutputSrgb());
diff --git a/Engine/Plugins/Runtime/AR/AppleAR/AppleARKit/Source/AppleARKit/Private/AppleARKitTextures.cpp b/Engine/Plugins/Runtime/AR/AppleAR/AppleARKit/Source/AppleARKit/Private/AppleARKitTextures.cpp
--- a/Engine/Plugins/Runtime/AR/AppleAR/AppleARKit/Source/AppleARKit/Private/AppleARKitTextures.cpp
+++ b/Engine/Plugins/Runtime/AR/AppleAR/AppleARKit/Source/AppleARKit/Private/AppleARKitTextures.cpp
@@ -1013,7 +1013,7 @@ public:
 		
 		if (!DecodedTextureUAV)
 		{
-			DecodedTextureUAV = RHICmdList.CreateUnorderedAccessView(DecodedTextureRef);
+			DecodedTextureUAV = RHICmdList.CreateUnorderedAccessView(DecodedTextureRef, FRHIViewDesc::CreateTextureUAV().SetDimensionFromTexture(DecodedTextureRef));
 		}
 		
 		{
diff --git a/Engine/Plugins/Runtime/GeometryCache/Source/GeometryCache/Private/GeometryCacheSceneProxy.cpp b/Engine/Plugins/Runtime/GeometryCache/Source/GeometryCache/Private/GeometryCacheSceneProxy.cpp
--- a/Engine/Plugins/Runtime/GeometryCache/Source/GeometryCache/Private/GeometryCacheSceneProxy.cpp
+++ b/Engine/Plugins/Runtime/GeometryCache/Source/GeometryCache/Private/GeometryCacheSceneProxy.cpp
@@ -1500,7 +1500,11 @@ void FGeomCacheIndexBuffer::InitRHI(FRHICommandListBase& RHICmdList)
 
 	if (IndexBufferRHI && NumAllocatedIndices)
 	{
-		BufferSRV = RHICmdList.CreateShaderResourceView(IndexBufferRHI, sizeof(uint32), PF_R32_UINT);
+		BufferSRV = RHICmdList.CreateShaderResourceView(
+			IndexBufferRHI, 
+			FRHIViewDesc::CreateBufferSRV()
+				.SetType(FRHIViewDesc::EBufferType::Typed)
+				.SetFormat(PF_R32_UINT));
 	}
 }
 
@@ -1554,7 +1558,11 @@ void FGeomCacheIndexBuffer::Update(FRHICommandListBase& RHICmdList, const TArray
 
 	if (bReallocate && IndexBufferRHI && NumAllocatedIndices)
 	{
-		BufferSRV = RHICmdList.CreateShaderResourceView(IndexBufferRHI, sizeof(uint32), PF_R32_UINT);
+		BufferSRV = RHICmdList.CreateShaderResourceView(
+			IndexBufferRHI, 
+			FRHIViewDesc::CreateBufferSRV()
+				.SetType(FRHIViewDesc::EBufferType::Typed)
+				.SetFormat(PF_R32_UINT));
 	}
 }
 
@@ -1573,7 +1581,11 @@ void FGeomCacheIndexBuffer::UpdateSizeOnly(FRHICommandListBase& RHICmdList, int3
 
 	if (bReallocate && IndexBufferRHI && NumAllocatedIndices)
 	{
-		BufferSRV = RHICmdList.CreateShaderResourceView(IndexBufferRHI, sizeof(uint32), PF_R32_UINT);
+		BufferSRV = RHICmdList.CreateShaderResourceView(
+			IndexBufferRHI, 
+			FRHIViewDesc::CreateBufferSRV()
+				.SetType(FRHIViewDesc::EBufferType::Typed)
+				.SetFormat(PF_R32_UINT));
 	}
 }
 
@@ -1584,7 +1596,11 @@ void FGeomCacheVertexBuffer::InitRHI(FRHICommandListBase& RHICmdList)
 
 	if (VertexBufferRHI && RHISupportsManualVertexFetch(GMaxRHIShaderPlatform))
 	{
-		BufferSRV = RHICmdList.CreateShaderResourceView(VertexBufferRHI, sizeof(float), PF_R32_FLOAT);
+		BufferSRV = RHICmdList.CreateShaderResourceView(
+			VertexBufferRHI, 
+			FRHIViewDesc::CreateBufferSRV()
+				.SetType(FRHIViewDesc::EBufferType::Typed)
+				.SetFormat(PF_R32_FLOAT));
 	}
 }
 
@@ -1601,7 +1617,11 @@ void FGeomCacheTangentBuffer::InitRHI(FRHICommandListBase& RHICmdList)
 
 	if (VertexBufferRHI && RHISupportsManualVertexFetch(GMaxRHIShaderPlatform))
 	{
-		BufferSRV = RHICmdList.CreateShaderResourceView(VertexBufferRHI, sizeof(FPackedNormal), PF_R8G8B8A8_SNORM);
+		BufferSRV = RHICmdList.CreateShaderResourceView(
+			VertexBufferRHI, 
+			FRHIViewDesc::CreateBufferSRV()
+				.SetType(FRHIViewDesc::EBufferType::Typed)
+				.SetFormat(PF_R8G8B8A8_SNORM));
 	}
 }
 
@@ -1612,7 +1632,11 @@ void FGeomCacheColorBuffer::InitRHI(FRHICommandListBase& RHICmdList)
 
 	if (VertexBufferRHI && RHISupportsManualVertexFetch(GMaxRHIShaderPlatform))
 	{
-		BufferSRV = RHICmdList.CreateShaderResourceView(VertexBufferRHI, sizeof(FColor), PF_B8G8R8A8);
+		BufferSRV = RHICmdList.CreateShaderResourceView(
+			VertexBufferRHI, 
+			FRHIViewDesc::CreateBufferSRV()
+				.SetType(FRHIViewDesc::EBufferType::Typed)
+				.SetFormat(PF_B8G8R8A8));
 	}
 }
 
@@ -1653,7 +1677,11 @@ void FGeomCacheVertexBuffer::UpdateRaw(FRHICommandListBase& RHICmdList, const vo
 
 	if (bReallocate && VertexBufferRHI && RHISupportsManualVertexFetch(GMaxRHIShaderPlatform))
 	{
-		BufferSRV = RHICmdList.CreateShaderResourceView(VertexBufferRHI, sizeof(float), PF_R32_FLOAT);
+		BufferSRV = RHICmdList.CreateShaderResourceView(
+			VertexBufferRHI, 
+			FRHIViewDesc::CreateBufferSRV()
+				.SetType(FRHIViewDesc::EBufferType::Typed)
+				.SetFormat(PF_R32_FLOAT));
 	}
 }
 
@@ -1670,6 +1698,10 @@ void FGeomCacheVertexBuffer::UpdateSize(FRHICommandListBase& RHICmdList, int32 N
 
 	if (bReallocate && VertexBufferRHI && RHISupportsManualVertexFetch(GMaxRHIShaderPlatform))
 	{
-		BufferSRV = RHICmdList.CreateShaderResourceView(VertexBufferRHI, sizeof(float), PF_R32_FLOAT);
+		BufferSRV = RHICmdList.CreateShaderResourceView(
+			VertexBufferRHI, 
+			FRHIViewDesc::CreateBufferSRV()
+				.SetType(FRHIViewDesc::EBufferType::Typed)
+				.SetFormat(PF_R32_FLOAT));
 	}
 }
diff --git a/Engine/Plugins/Runtime/HairStrands/Source/HairStrandsCore/Private/GroomResources.cpp b/Engine/Plugins/Runtime/HairStrands/Source/HairStrandsCore/Private/GroomResources.cpp
--- a/Engine/Plugins/Runtime/HairStrands/Source/HairStrandsCore/Private/GroomResources.cpp
+++ b/Engine/Plugins/Runtime/HairStrands/Source/HairStrandsCore/Private/GroomResources.cpp
@@ -179,7 +179,11 @@ void CreateBuffer(FRHICommandListBase& RHICmdList, const TArray<typename FormatT
 	void* BufferData = RHICmdList.LockBuffer(OutBuffer.VertexBufferRHI, 0, DataSizeInBytes, RLM_WriteOnly);
 	FMemory::Memcpy(BufferData, InData.GetData(), DataSizeInBytes);
 	RHICmdList.UnlockBuffer(OutBuffer.VertexBufferRHI);
-	OutBuffer.ShaderResourceViewRHI = RHICmdList.CreateShaderResourceView(OutBuffer.VertexBufferRHI, FormatType::SizeInByte, FormatType::Format);
+	OutBuffer.ShaderResourceViewRHI = RHICmdList.CreateShaderResourceView(
+		OutBuffer.VertexBufferRHI, 
+		FRHIViewDesc::CreateBufferSRV()
+			.SetType(FRHIViewDesc::EBufferType::Typed)
+			.SetFormat(FormatType::Format));
 }
 
 /////////////////////////////////////////////////////////////////////////////////////////
diff --git a/Engine/Plugins/Runtime/HairStrands/Source/HairStrandsCore/Private/HairStrandsVertexFactory.cpp b/Engine/Plugins/Runtime/HairStrands/Source/HairStrandsCore/Private/HairStrandsVertexFactory.cpp
--- a/Engine/Plugins/Runtime/HairStrands/Source/HairStrandsCore/Private/HairStrandsVertexFactory.cpp
+++ b/Engine/Plugins/Runtime/HairStrands/Source/HairStrandsCore/Private/HairStrandsVertexFactory.cpp
@@ -64,11 +64,28 @@ public:
 			RHICmdList.UnlockBuffer(ByteAddressBufferRHI);
 		}
 
-		SRVUint = RHICmdList.CreateShaderResourceView(VertexBufferRHI, sizeof(uint32), PF_R32_UINT);
-		SRVFloat = RHICmdList.CreateShaderResourceView(VertexBufferRHI, sizeof(uint32), PF_R32_FLOAT);
-		SRVRGBA = RHICmdList.CreateShaderResourceView(VertexBufferRHI, sizeof(uint32), PF_R8G8B8A8);
-		SRVRGBA_Uint = RHICmdList.CreateShaderResourceView(VertexBufferRHI, sizeof(uint32), PF_R8G8B8A8_UINT);
-		SRVByteAddress = RHICmdList.CreateShaderResourceView(ByteAddressBufferRHI, sizeof(uint32), PF_R32_UINT);
+		SRVUint = RHICmdList.CreateShaderResourceView(
+			VertexBufferRHI, 
+			FRHIViewDesc::CreateBufferSRV()
+				.SetType(FRHIViewDesc::EBufferType::Typed)
+				.SetFormat(PF_R32_UINT));
+		SRVFloat = RHICmdList.CreateShaderResourceView(
+			VertexBufferRHI, 
+			FRHIViewDesc::CreateBufferSRV()
+				.SetType(FRHIViewDesc::EBufferType::Typed)
+				.SetFormat(PF_R32_FLOAT));
+		SRVRGBA = RHICmdList.CreateShaderResourceView(
+			VertexBufferRHI, 
+			FRHIViewDesc::CreateBufferSRV()
+				.SetType(FRHIViewDesc::EBufferType::Typed)
+				.SetFormat(PF_R8G8B8A8));
+		SRVRGBA_Uint = RHICmdList.CreateShaderResourceView(
+			VertexBufferRHI, 
+			FRHIViewDesc::CreateBufferSRV()
+				.SetType(FRHIViewDesc::EBufferType::Typed)
+				.SetFormat(PF_R8G8B8A8_UINT));
+
+		SRVByteAddress = RHICmdList.CreateShaderResourceView(ByteAddressBufferRHI, FRHIViewDesc::CreateBufferSRV().SetType(FRHIViewDesc::EBufferType::Raw));
 	}
 
 	virtual void ReleaseRHI() override
diff --git a/Engine/Plugins/Tests/RHITests/Source/RHITests/Private/RHIClearTextureTests.cpp b/Engine/Plugins/Tests/RHITests/Source/RHITests/Private/RHIClearTextureTests.cpp
--- a/Engine/Plugins/Tests/RHITests/Source/RHITests/Private/RHIClearTextureTests.cpp
+++ b/Engine/Plugins/Tests/RHITests/Source/RHITests/Private/RHIClearTextureTests.cpp
@@ -156,7 +156,7 @@ FRHITextureCreateDesc CreateTexture2DAtlasDesc(EPixelFormat PixelFormat, ERHIAcc
 FTextureRHIRef CreateTexture2DAtlasRT(FRHICommandListImmediate& RHICmdList, EPixelFormat PixelFormat, FUnorderedAccessViewRHIRef* OutUAV)
 {
 	FTextureRHIRef Texture2DAtlasRT =  RHICreateTexture(CreateTexture2DAtlasDesc(PixelFormat, ERHIAccess::UAVCompute, TexCreate_UAV | TexCreate_ShaderResource, TEXT("Texture2DAtlasRT")));
-	*OutUAV = RHICmdList.CreateUnorderedAccessView(Texture2DAtlasRT);
+	*OutUAV = RHICmdList.CreateUnorderedAccessView(Texture2DAtlasRT, FRHIViewDesc::CreateTextureUAV().SetDimensionFromTexture(Texture2DAtlasRT));
 	(*OutUAV)->SetOwnerName(TEXT("Texture2DAtlasRT"));
 
 	return Texture2DAtlasRT;
diff --git a/Engine/Plugins/Tests/RHITests/Source/RHITests/Public/RHITextureTests.h b/Engine/Plugins/Tests/RHITests/Source/RHITests/Public/RHITextureTests.h
--- a/Engine/Plugins/Tests/RHITests/Source/RHITests/Public/RHITextureTests.h
+++ b/Engine/Plugins/Tests/RHITests/Source/RHITests/Public/RHITextureTests.h
@@ -25,7 +25,11 @@ public:
 			// Test clear whole resource to zero
 			for (uint32 Mip = 0; Mip < TextureRHI->GetNumMips(); ++Mip)
 			{
-				FUnorderedAccessViewRHIRef MipUAV = RHICmdList.CreateUnorderedAccessView(TextureRHI, Mip);
+				FUnorderedAccessViewRHIRef MipUAV = RHICmdList.CreateUnorderedAccessView(
+					TextureRHI, 
+					FRHIViewDesc::CreateTextureUAV()
+						.SetDimensionFromTexture(TextureRHI)
+						.SetMipLevel(Mip));
 
 				ValueType ZerosValue;
 				FMemory::Memset(&ZerosValue, 0, sizeof(ZerosValue));
@@ -77,7 +81,11 @@ public:
 			// Clear the selected mip index to the provided value
 			
 			RHICmdList.Transition(FRHITransitionInfo(TextureRHI, ERHIAccess::CopySrc, ERHIAccess::UAVCompute));
-			FUnorderedAccessViewRHIRef SpecificMipUAV = RHICmdList.CreateUnorderedAccessView(TextureRHI, MipIndex);
+			FUnorderedAccessViewRHIRef SpecificMipUAV = RHICmdList.CreateUnorderedAccessView(
+				TextureRHI, 
+				FRHIViewDesc::CreateTextureUAV()
+					.SetDimensionFromTexture(TextureRHI)
+					.SetMipLevel(MipIndex));
 			(RHICmdList.*ClearPtr)(SpecificMipUAV, ClearValue);
 
 			RHICmdList.Transition(FRHITransitionInfo(TextureRHI, ERHIAccess::UAVCompute, ERHIAccess::CopySrc));
@@ -261,8 +269,12 @@ public:
 
 			if (Texture && SRVFormat != PF_Unknown)
 			{
-				FRHITextureSRVCreateInfo ViewInfo(0, 1, SRVFormat);
-				FShaderResourceViewRHIRef SRV = RHICmdList.CreateShaderResourceView(Texture, ViewInfo);
+				FShaderResourceViewRHIRef SRV = RHICmdList.CreateShaderResourceView(
+					Texture, 
+					FRHIViewDesc::CreateTextureSRV()
+						.SetDimensionFromTexture(Texture)
+						.SetFormat(SRVFormat)
+						.SetMipRange(0, 1));
 				bResult = (SRV != nullptr);
 			}
 
diff --git a/Engine/Plugins/TextureGraph/Source/TextureGraphEngine/FxMat/FxMaterial.h b/Engine/Plugins/TextureGraph/Source/TextureGraphEngine/FxMat/FxMaterial.h
--- a/Engine/Plugins/TextureGraph/Source/TextureGraphEngine/FxMat/FxMaterial.h
+++ b/Engine/Plugins/TextureGraph/Source/TextureGraphEngine/FxMat/FxMaterial.h
@@ -342,7 +342,7 @@ public:
 
 			if (!RenderTargetUAV)
 			{
-				RenderTargetUAV = RHI.CreateUnorderedAccessView(Target);
+				RenderTargetUAV = RHI.CreateUnorderedAccessView(Target, FRHIViewDesc::CreateTextureUAV().SetDimensionFromTexture(Target));
 
 				/// These access flags are taken from UE4 where and match what ERHIAccess::ERWNoBarrier was 
 				/// defined as. The definition has been discarded in UE5
diff --git a/Engine/Plugins/TextureGraph/Source/TextureGraphEngine/Transform/Utility/T_TextureHistogram.cpp b/Engine/Plugins/TextureGraph/Source/TextureGraphEngine/Transform/Utility/T_TextureHistogram.cpp
--- a/Engine/Plugins/TextureGraph/Source/TextureGraphEngine/Transform/Utility/T_TextureHistogram.cpp
+++ b/Engine/Plugins/TextureGraph/Source/TextureGraphEngine/Transform/Utility/T_TextureHistogram.cpp
@@ -136,7 +136,7 @@ public:
 
 			if (!RenderTargetUAV)
 			{
-				RenderTargetUAV = RHI.CreateUnorderedAccessView(Target);
+				RenderTargetUAV = RHI.CreateUnorderedAccessView(Target, FRHIViewDesc::CreateTextureUAV().SetDimensionFromTexture(Target));
 				RHI.Transition(FRHITransitionInfo(Target, ERHIAccess::Unknown, ERHIAccess::UAVMask));
 			}
 
diff --git a/Engine/Plugins/VirtualProduction/CameraCalibrationCore/Source/CameraCalibrationCore/Private/DistortionRenderingUtils.cpp b/Engine/Plugins/VirtualProduction/CameraCalibrationCore/Source/CameraCalibrationCore/Private/DistortionRenderingUtils.cpp
--- a/Engine/Plugins/VirtualProduction/CameraCalibrationCore/Source/CameraCalibrationCore/Private/DistortionRenderingUtils.cpp
+++ b/Engine/Plugins/VirtualProduction/CameraCalibrationCore/Source/CameraCalibrationCore/Private/DistortionRenderingUtils.cpp
@@ -93,7 +93,7 @@ namespace DistortionRenderingUtils
 				// Create an SRV for the input buffer of image points
 				FRHIResourceCreateInfo CreateInfo(TEXT("ImagePointsInitialData"), &InputPointsResourceArray);
 				FBufferRHIRef InputPointsBuffer = RHICmdList.CreateStructuredBuffer(sizeof(FVector2f), BufferSize, BUF_Static | BUF_ShaderResource, CreateInfo);
-				FShaderResourceViewRHIRef InputPointsSRV = RHICmdList.CreateShaderResourceView(InputPointsBuffer);
+				FShaderResourceViewRHIRef InputPointsSRV = RHICmdList.CreateShaderResourceView(InputPointsBuffer, FRHIViewDesc::CreateBufferSRV().SetTypeFromBuffer(InputPointsBuffer));
 				Parameters.InputPoints = InputPointsSRV;
 
 				// Create a RWBuffer to use as a UAV for the output buffer of undistorted points
diff --git a/Engine/Source/Runtime/Engine/Private/DynamicMeshBuilder.cpp b/Engine/Source/Runtime/Engine/Private/DynamicMeshBuilder.cpp
--- a/Engine/Source/Runtime/Engine/Private/DynamicMeshBuilder.cpp
+++ b/Engine/Source/Runtime/Engine/Private/DynamicMeshBuilder.cpp
@@ -457,11 +457,27 @@ public:
 
 		if (RHISupportsManualVertexFetch(GMaxRHIShaderPlatform))
 		{
-			TangentBufferSRV = RHICmdList.CreateShaderResourceView(TangentBuffer.VertexBufferRHI, 4, PF_R8G8B8A8_SNORM);
-			TexCoordBufferSRV = RHICmdList.CreateShaderResourceView(TexCoordBuffer.VertexBufferRHI, TextureStride, TextureFormat);
-			ColorBufferSRV = RHICmdList.CreateShaderResourceView(ColorBuffer.VertexBufferRHI, 4, PF_R8G8B8A8);
-
-			PositionBufferSRV = RHICmdList.CreateShaderResourceView(PositionBuffer.VertexBufferRHI, sizeof(float), PF_R32_FLOAT);
+			TangentBufferSRV = RHICmdList.CreateShaderResourceView(
+				TangentBuffer.VertexBufferRHI, 
+				FRHIViewDesc::CreateBufferSRV()
+					.SetType(FRHIViewDesc::EBufferType::Typed)
+					.SetFormat(PF_R8G8B8A8_SNORM));
+			TexCoordBufferSRV = RHICmdList.CreateShaderResourceView(
+				TexCoordBuffer.VertexBufferRHI, 
+				FRHIViewDesc::CreateBufferSRV()
+					.SetType(FRHIViewDesc::EBufferType::Typed)
+					.SetFormat(TextureFormat));
+			ColorBufferSRV = RHICmdList.CreateShaderResourceView(
+				ColorBuffer.VertexBufferRHI, 
+				FRHIViewDesc::CreateBufferSRV()
+					.SetType(FRHIViewDesc::EBufferType::Typed)
+					.SetFormat(PF_R8G8B8A8));
+
+			PositionBufferSRV = RHICmdList.CreateShaderResourceView(
+				PositionBuffer.VertexBufferRHI, 
+				FRHIViewDesc::CreateBufferSRV()
+					.SetType(FRHIViewDesc::EBufferType::Typed)
+					.SetFormat(PF_R32_FLOAT));
 		}
 	}
 
diff --git a/Engine/Source/Runtime/Engine/Private/GPUSkinCache.cpp b/Engine/Source/Runtime/Engine/Private/GPUSkinCache.cpp
--- a/Engine/Source/Runtime/Engine/Private/GPUSkinCache.cpp
+++ b/Engine/Source/Runtime/Engine/Private/GPUSkinCache.cpp
@@ -1827,7 +1827,11 @@ bool FGPUSkinCache::ProcessEntry(
 
 	        FRHIResourceCreateInfo CreateInfo(TEXT("ClothPositionAndNormalsBuffer"), ResourceArray);
 	        ClothPositionAndNormalsBuffer.VertexBufferRHI = RHICmdList.CreateVertexBuffer( ResourceArray->GetResourceDataSize(), BUF_Static | BUF_ShaderResource, CreateInfo);
-	        ClothPositionAndNormalsBuffer.VertexBufferSRV = RHICmdList.CreateShaderResourceView(ClothPositionAndNormalsBuffer.VertexBufferRHI, sizeof(FVector2f), PF_G32R32F);
+	        ClothPositionAndNormalsBuffer.VertexBufferSRV = RHICmdList.CreateShaderResourceView(
+				ClothPositionAndNormalsBuffer.VertexBufferRHI, 
+				FRHIViewDesc::CreateBufferSRV()
+					.SetType(FRHIViewDesc::EBufferType::Typed)
+					.SetFormat(PF_G32R32F));
 			SectionDispatchData.ClothPositionsAndNormalsBuffer = ClothPositionAndNormalsBuffer.VertexBufferSRV;
 		}
 		else
diff --git a/Engine/Source/Runtime/Engine/Private/GPUSkinVertexFactory.cpp b/Engine/Source/Runtime/Engine/Private/GPUSkinVertexFactory.cpp
--- a/Engine/Source/Runtime/Engine/Private/GPUSkinVertexFactory.cpp
+++ b/Engine/Source/Runtime/Engine/Private/GPUSkinVertexFactory.cpp
@@ -194,7 +194,11 @@ FVertexBufferAndSRV FBoneBufferPoolPolicy::CreateResource(FRHICommandListBase& R
 	FVertexBufferAndSRV Buffer;
 	FRHIResourceCreateInfo CreateInfo(TEXT("FBoneBufferPoolPolicy"));
 	Buffer.VertexBufferRHI = RHICmdList.CreateVertexBuffer( BufferSize, (BUF_Dynamic | BUF_ShaderResource), CreateInfo );
-	Buffer.VertexBufferSRV = RHICmdList.CreateShaderResourceView( Buffer.VertexBufferRHI, sizeof(FVector4f), PF_A32B32G32R32F );
+	Buffer.VertexBufferSRV = RHICmdList.CreateShaderResourceView(
+		Buffer.VertexBufferRHI, 
+		FRHIViewDesc::CreateBufferSRV()
+			.SetType(FRHIViewDesc::EBufferType::Typed)
+			.SetFormat(PF_A32B32G32R32F));
 	return Buffer;
 }
 
@@ -215,7 +219,11 @@ FVertexBufferAndSRV FClothBufferPoolPolicy::CreateResource(FRHICommandListBase&
 	FVertexBufferAndSRV Buffer;
 	FRHIResourceCreateInfo CreateInfo(TEXT("FClothBufferPoolPolicy"));
 	Buffer.VertexBufferRHI = RHICmdList.CreateVertexBuffer( BufferSize, (BUF_Dynamic | BUF_ShaderResource), CreateInfo );
-	Buffer.VertexBufferSRV = RHICmdList.CreateShaderResourceView( Buffer.VertexBufferRHI, sizeof(FVector2f), PF_G32R32F );
+	Buffer.VertexBufferSRV = RHICmdList.CreateShaderResourceView(
+		Buffer.VertexBufferRHI, 
+		FRHIViewDesc::CreateBufferSRV()
+			.SetType(FRHIViewDesc::EBufferType::Typed)
+			.SetFormat(PF_G32R32F));
 	return Buffer;
 }
 
diff --git a/Engine/Source/Runtime/Engine/Private/GPUSort.cpp b/Engine/Source/Runtime/Engine/Private/GPUSort.cpp
--- a/Engine/Source/Runtime/Engine/Private/GPUSort.cpp
+++ b/Engine/Source/Runtime/Engine/Private/GPUSort.cpp
@@ -98,11 +98,14 @@ public:
 				CreateInfo);
 			BufferSRVs[BufferIndex] = RHICmdList.CreateShaderResourceView(
 				Buffers[BufferIndex],
-				/*Stride=*/ sizeof(uint32),
-				/*Format=*/ PF_R32_UINT );
+				FRHIViewDesc::CreateBufferSRV()
+					.SetType(FRHIViewDesc::EBufferType::Typed)
+					.SetFormat(PF_R32_UINT));
 			BufferUAVs[BufferIndex] = RHICmdList.CreateUnorderedAccessView(
 				Buffers[BufferIndex],
-				/*Format=*/ PF_R32_UINT );
+				FRHIViewDesc::CreateBufferUAV()
+					.SetType(FRHIViewDesc::EBufferType::Typed)
+					.SetFormat(PF_R32_UINT));
 		}
 	}
 
@@ -190,8 +193,10 @@ public:
 				/*Usage=*/ BUF_Volatile | BUF_ShaderResource,
 				CreateInfo);
 			SortParametersBufferSRV = RHICmdList.CreateShaderResourceView(
-				SortParametersBufferRHI, /*Stride=*/ sizeof(uint32), PF_R32_UINT 
-				);
+				SortParametersBufferRHI,
+				FRHIViewDesc::CreateBufferSRV()
+					.SetType(FRHIViewDesc::EBufferType::Typed)
+					.SetFormat(PF_R32_UINT));
 	}
 
 	/**
@@ -736,12 +741,28 @@ static bool RunGPUSortTest(FRHICommandListImmediate& RHICmdList, int32 TestSize,
 	{
 		FRHIResourceCreateInfo CreateInfo(TEXT("KeysBuffer"));
 		KeysBufferRHI[BufferIndex] = RHICmdList.CreateVertexBuffer(BufferSize, BUF_Static | BUF_ShaderResource | BUF_UnorderedAccess, CreateInfo);
-		KeysBufferSRV[BufferIndex] = RHICmdList.CreateShaderResourceView(KeysBufferRHI[BufferIndex], /*Stride=*/ sizeof(uint32), PF_R32_UINT);
-		KeysBufferUAV[BufferIndex] = RHICmdList.CreateUnorderedAccessView(KeysBufferRHI[BufferIndex], PF_R32_UINT);
+		KeysBufferSRV[BufferIndex] = RHICmdList.CreateShaderResourceView(
+			KeysBufferRHI[BufferIndex],
+			FRHIViewDesc::CreateBufferSRV()
+				.SetType(FRHIViewDesc::EBufferType::Typed)
+				.SetFormat(PF_R32_UINT));
+		KeysBufferUAV[BufferIndex] = RHICmdList.CreateUnorderedAccessView(
+			KeysBufferRHI[BufferIndex], 
+			FRHIViewDesc::CreateBufferUAV()
+				.SetType(FRHIViewDesc::EBufferType::Typed)
+				.SetFormat(PF_R32_UINT));
 		CreateInfo.DebugName = TEXT("ValuesBuffer");
 		ValuesBufferRHI[BufferIndex] = RHICmdList.CreateVertexBuffer(BufferSize, BUF_Static | BUF_ShaderResource | BUF_UnorderedAccess, CreateInfo);
-		ValuesBufferSRV[BufferIndex] = RHICmdList.CreateShaderResourceView(ValuesBufferRHI[BufferIndex], /*Stride=*/ sizeof(uint32), PF_R32_UINT);
-		ValuesBufferUAV[BufferIndex] = RHICmdList.CreateUnorderedAccessView(ValuesBufferRHI[BufferIndex], PF_R32_UINT);
+		ValuesBufferSRV[BufferIndex] = RHICmdList.CreateShaderResourceView(
+			ValuesBufferRHI[BufferIndex],
+			FRHIViewDesc::CreateBufferSRV()
+				.SetType(FRHIViewDesc::EBufferType::Typed)
+				.SetFormat(PF_R32_UINT));
+		ValuesBufferUAV[BufferIndex] = RHICmdList.CreateUnorderedAccessView(
+			ValuesBufferRHI[BufferIndex], 
+			FRHIViewDesc::CreateBufferUAV()
+				.SetType(FRHIViewDesc::EBufferType::Typed)
+				.SetFormat(PF_R32_UINT));
 	}
 
 	// Upload initial keys and values to the GPU.
diff --git a/Engine/Source/Runtime/Engine/Private/GPUSortManager.cpp b/Engine/Source/Runtime/Engine/Private/GPUSortManager.cpp
--- a/Engine/Source/Runtime/Engine/Private/GPUSortManager.cpp
+++ b/Engine/Source/Runtime/Engine/Private/GPUSortManager.cpp
@@ -271,18 +271,45 @@ FGPUSortManager::FValueBuffer::FValueBuffer(FRHICommandListBase& RHICmdList, int
 	FRHIResourceCreateInfo CreateInfo(TEXT("ValueBuffer"));
 	VertexBufferRHI = RHICmdList.CreateVertexBuffer((uint32)InAllocatedCount * sizeof(uint32), BUF_Static | BUF_ShaderResource | BUF_UnorderedAccess, ERHIAccess::SRVGraphics, CreateInfo);
 
-	UInt32SRV = RHICmdList.CreateShaderResourceView(VertexBufferRHI, sizeof(uint32), PF_R32_UINT);
-	UInt32UAV = RHICmdList.CreateUnorderedAccessView(VertexBufferRHI, PF_R32_UINT);
+	UInt32SRV = RHICmdList.CreateShaderResourceView(
+		VertexBufferRHI, 
+		FRHIViewDesc::CreateBufferSRV()
+			.SetType(FRHIViewDesc::EBufferType::Typed)
+			.SetFormat(PF_R32_UINT));
+
+	UInt32UAV = RHICmdList.CreateUnorderedAccessView(
+		VertexBufferRHI, 
+		FRHIViewDesc::CreateBufferUAV()
+			.SetType(FRHIViewDesc::EBufferType::Typed)
+			.SetFormat(PF_R32_UINT));
 
 	if (EnumHasAnyFlags(InSettings.AllowedFlags, EGPUSortFlags::ValuesAsInt32))
 	{
-		Int32SRV = RHICmdList.CreateShaderResourceView(VertexBufferRHI, sizeof(int32), PF_R32_SINT);
-		Int32UAV = RHICmdList.CreateUnorderedAccessView(VertexBufferRHI, PF_R32_SINT);
+		Int32SRV = RHICmdList.CreateShaderResourceView(
+			VertexBufferRHI, 
+			FRHIViewDesc::CreateBufferSRV()
+				.SetType(FRHIViewDesc::EBufferType::Typed)
+				.SetFormat(PF_R32_SINT));
+
+		Int32UAV = RHICmdList.CreateUnorderedAccessView(
+			VertexBufferRHI, 
+			FRHIViewDesc::CreateBufferUAV()
+				.SetType(FRHIViewDesc::EBufferType::Typed)
+				.SetFormat(PF_R32_SINT));
 	}
 	if (EnumHasAnyFlags(InSettings.AllowedFlags, EGPUSortFlags::ValuesAsG16R16F))
 	{
-		G16R16SRV = RHICmdList.CreateShaderResourceView(VertexBufferRHI, sizeof(FFloat16) * 2, PF_G16R16F);
-		G16R16UAV = RHICmdList.CreateUnorderedAccessView(VertexBufferRHI, PF_G16R16F);
+		G16R16SRV = RHICmdList.CreateShaderResourceView(
+			VertexBufferRHI, 
+			FRHIViewDesc::CreateBufferSRV()
+				.SetType(FRHIViewDesc::EBufferType::Typed)
+				.SetFormat(PF_G16R16F));
+
+		G16R16UAV = RHICmdList.CreateUnorderedAccessView(
+			VertexBufferRHI, 
+			FRHIViewDesc::CreateBufferUAV()
+				.SetType(FRHIViewDesc::EBufferType::Typed)
+				.SetFormat(PF_G16R16F));
 	}
 }
 
diff --git a/Engine/Source/Runtime/Engine/Private/GeometryCacheVertexFactory.cpp b/Engine/Source/Runtime/Engine/Private/GeometryCacheVertexFactory.cpp
--- a/Engine/Source/Runtime/Engine/Private/GeometryCacheVertexFactory.cpp
+++ b/Engine/Source/Runtime/Engine/Private/GeometryCacheVertexFactory.cpp
@@ -127,7 +127,11 @@ public:
 		DummyContents[1] = FVector4f(1.0f, 1.0f, 1.0f, 1.0f);
 		RHICmdList.UnlockBuffer(VertexBufferRHI);
 
-		SRV = RHICmdList.CreateShaderResourceView(VertexBufferRHI, sizeof(float), PF_R32_FLOAT);
+		SRV = RHICmdList.CreateShaderResourceView(
+			VertexBufferRHI, 
+			FRHIViewDesc::CreateBufferSRV()
+				.SetType(FRHIViewDesc::EBufferType::Typed)
+				.SetFormat(PF_R32_FLOAT));
 	}
 
 	virtual void ReleaseRHI() override
@@ -152,7 +156,11 @@ public:
 		DummyContents[1] = FVector4f(1.0f, 1.0f, 1.0f, 1.0f);
 		RHICmdList.UnlockBuffer(VertexBufferRHI);
 
-		SRV = RHICmdList.CreateShaderResourceView(VertexBufferRHI, sizeof(FPackedNormal), PF_R8G8B8A8_SNORM);
+		SRV = RHICmdList.CreateShaderResourceView(
+			VertexBufferRHI, 
+			FRHIViewDesc::CreateBufferSRV()
+				.SetType(FRHIViewDesc::EBufferType::Typed)
+				.SetFormat(PF_R8G8B8A8_SNORM));
 	}
 
 	virtual void ReleaseRHI() override
@@ -280,7 +288,11 @@ void FGeometryCacheVertexVertexFactory::CreateManualVertexFetchUniformBuffer(
 
 	if (PoistionBuffer != NULL)
 	{
-		OutUserData.PositionSRV = RHICmdList.CreateShaderResourceView(PoistionBuffer->VertexBufferRHI, sizeof(float), PF_R32_FLOAT);
+		OutUserData.PositionSRV = RHICmdList.CreateShaderResourceView(
+			PoistionBuffer->VertexBufferRHI, 
+			FRHIViewDesc::CreateBufferSRV()
+				.SetType(FRHIViewDesc::EBufferType::Typed)
+				.SetFormat(PF_R32_FLOAT));
 		// Position will need per-component fetch since we don't have R32G32B32 pixel format
 		ManualVertexFetchParameters.Position = OutUserData.PositionSRV;
 	}
@@ -291,7 +303,11 @@ void FGeometryCacheVertexVertexFactory::CreateManualVertexFetchUniformBuffer(
 
 	if (Data.TangentBasisComponents[0].VertexBuffer != NULL)
 	{
-		OutUserData.TangentXSRV = RHICmdList.CreateShaderResourceView(Data.TangentBasisComponents[0].VertexBuffer->VertexBufferRHI, sizeof(FPackedNormal), PF_R8G8B8A8_SNORM);
+		OutUserData.TangentXSRV = RHICmdList.CreateShaderResourceView(
+			Data.TangentBasisComponents[0].VertexBuffer->VertexBufferRHI, 
+			FRHIViewDesc::CreateBufferSRV()
+				.SetType(FRHIViewDesc::EBufferType::Typed)
+				.SetFormat(PF_R8G8B8A8_SNORM));
 		ManualVertexFetchParameters.TangentX = OutUserData.TangentXSRV;
 	}
 	else
@@ -301,7 +317,11 @@ void FGeometryCacheVertexVertexFactory::CreateManualVertexFetchUniformBuffer(
 
 	if (Data.TangentBasisComponents[1].VertexBuffer != NULL)
 	{
-		OutUserData.TangentZSRV = RHICmdList.CreateShaderResourceView(Data.TangentBasisComponents[1].VertexBuffer->VertexBufferRHI, sizeof(FPackedNormal), PF_R8G8B8A8_SNORM);
+		OutUserData.TangentZSRV = RHICmdList.CreateShaderResourceView(
+			Data.TangentBasisComponents[1].VertexBuffer->VertexBufferRHI, 
+			FRHIViewDesc::CreateBufferSRV()
+				.SetType(FRHIViewDesc::EBufferType::Typed)
+				.SetFormat(PF_R8G8B8A8_SNORM));
 		ManualVertexFetchParameters.TangentZ = OutUserData.TangentZSRV;
 	}
 	else
@@ -311,7 +331,11 @@ void FGeometryCacheVertexVertexFactory::CreateManualVertexFetchUniformBuffer(
 
 	if (Data.ColorComponent.VertexBuffer)
 	{
-		OutUserData.ColorSRV = RHICmdList.CreateShaderResourceView(Data.ColorComponent.VertexBuffer->VertexBufferRHI, sizeof(FColor), PF_B8G8R8A8);
+		OutUserData.ColorSRV = RHICmdList.CreateShaderResourceView(
+			Data.ColorComponent.VertexBuffer->VertexBufferRHI, 
+			FRHIViewDesc::CreateBufferSRV()
+				.SetType(FRHIViewDesc::EBufferType::Typed)
+				.SetFormat(PF_B8G8R8A8));
 		ManualVertexFetchParameters.Color = OutUserData.ColorSRV;
 	}
 	else
@@ -322,7 +346,11 @@ void FGeometryCacheVertexVertexFactory::CreateManualVertexFetchUniformBuffer(
 
 	if (MotionBlurBuffer)
 	{
-		OutUserData.MotionBlurDataSRV = RHICmdList.CreateShaderResourceView(MotionBlurBuffer->VertexBufferRHI, sizeof(float), PF_R32_FLOAT);
+		OutUserData.MotionBlurDataSRV = RHICmdList.CreateShaderResourceView(
+			MotionBlurBuffer->VertexBufferRHI, 
+			FRHIViewDesc::CreateBufferSRV()
+				.SetType(FRHIViewDesc::EBufferType::Typed)
+				.SetFormat(PF_R32_FLOAT));
 		ManualVertexFetchParameters.MotionBlurData = OutUserData.MotionBlurDataSRV;
 	}
 	else if (PoistionBuffer != NULL)
@@ -337,7 +365,11 @@ void FGeometryCacheVertexVertexFactory::CreateManualVertexFetchUniformBuffer(
 	if (Data.TextureCoordinates.Num())
 	{
 		checkf(Data.TextureCoordinates.Num() <= 1, TEXT("We're assuming FGeometryCacheSceneProxy uses only one TextureCoordinates vertex buffer"));
-		OutUserData.TexCoordsSRV = RHICmdList.CreateShaderResourceView(Data.TextureCoordinates[0].VertexBuffer->VertexBufferRHI, sizeof(float), PF_R32_FLOAT);
+		OutUserData.TexCoordsSRV = RHICmdList.CreateShaderResourceView(
+			Data.TextureCoordinates[0].VertexBuffer->VertexBufferRHI, 
+			FRHIViewDesc::CreateBufferSRV()
+				.SetType(FRHIViewDesc::EBufferType::Typed)
+				.SetFormat(PF_R32_FLOAT));
 		// TexCoords will need per-component fetch since we don't have R32G32 pixel format
 		ManualVertexFetchParameters.TexCoords = OutUserData.TexCoordsSRV;
 	}
diff --git a/Engine/Source/Runtime/Engine/Private/InstancedStaticMesh.cpp b/Engine/Source/Runtime/Engine/Private/InstancedStaticMesh.cpp
--- a/Engine/Source/Runtime/Engine/Private/InstancedStaticMesh.cpp
+++ b/Engine/Source/Runtime/Engine/Private/InstancedStaticMesh.cpp
@@ -313,7 +313,11 @@ public:
 		RHICmdList.UnlockBuffer(VertexBufferRHI);
 
 		// Create a view of the buffer
-		ShaderResourceViewRHI = RHICmdList.CreateShaderResourceView(VertexBufferRHI, sizeof(float), PF_R32_FLOAT);
+		ShaderResourceViewRHI = RHICmdList.CreateShaderResourceView(
+			VertexBufferRHI, 
+			FRHIViewDesc::CreateBufferSRV()
+				.SetType(FRHIViewDesc::EBufferType::Typed)
+				.SetFormat(PF_R32_FLOAT));
 	}
 };
 
@@ -423,7 +427,11 @@ void FStaticMeshInstanceBuffer::InitRHI(FRHICommandListBase& RHICmdList)
 			// Make sure we still create custom data SRV on platforms that do not support/use MVF 
 			if (InstanceCustomDataSRV == nullptr)
 			{
-				InstanceCustomDataSRV = RHICmdList.CreateShaderResourceView(InstanceCustomDataBuffer.VertexBufferRHI, 4, PF_R32_FLOAT);
+				InstanceCustomDataSRV = RHICmdList.CreateShaderResourceView(
+					InstanceCustomDataBuffer.VertexBufferRHI, 
+					FRHIViewDesc::CreateBufferSRV()
+						.SetType(FRHIViewDesc::EBufferType::Typed)
+						.SetFormat(PF_R32_FLOAT));
 			}
 		}
 		else
@@ -484,7 +492,11 @@ void FStaticMeshInstanceBuffer::CreateVertexBuffer(FRHICommandListBase& RHICmdLi
 	
 	if (RHISupportsManualVertexFetch(GMaxRHIShaderPlatform))
 	{
-		OutInstanceSRV = RHICmdList.CreateShaderResourceView(OutVertexBufferRHI, InStride, InFormat);
+		OutInstanceSRV = RHICmdList.CreateShaderResourceView(
+			OutVertexBufferRHI, 
+			FRHIViewDesc::CreateBufferSRV()
+				.SetType(FRHIViewDesc::EBufferType::Typed)
+				.SetFormat(static_cast<EPixelFormat>(InFormat)));
 	}
 }
 
diff --git a/Engine/Source/Runtime/Engine/Private/MeshParticleVertexFactory.cpp b/Engine/Source/Runtime/Engine/Private/MeshParticleVertexFactory.cpp
--- a/Engine/Source/Runtime/Engine/Private/MeshParticleVertexFactory.cpp
+++ b/Engine/Source/Runtime/Engine/Private/MeshParticleVertexFactory.cpp
@@ -58,7 +58,11 @@ public:
 	{
 		FRHIResourceCreateInfo CreateInfo(TEXT("FDummyPrevTransformBuffer"));
 		VB = RHICmdList.CreateVertexBuffer(sizeof(FVector4f) * 3, BUF_Static | BUF_ShaderResource, CreateInfo);
-		SRV = RHICmdList.CreateShaderResourceView(VB, sizeof(FVector4f), PF_A32B32G32R32F);
+		SRV = RHICmdList.CreateShaderResourceView(
+			VB,
+			FRHIViewDesc::CreateBufferSRV()
+				.SetType(FRHIViewDesc::EBufferType::Typed)
+				.SetFormat(PF_A32B32G32R32F));
 	}
 
 	virtual void ReleaseRHI()
diff --git a/Engine/Source/Runtime/Engine/Private/Particles/ParticleGpuSimulation.cpp b/Engine/Source/Runtime/Engine/Private/Particles/ParticleGpuSimulation.cpp
--- a/Engine/Source/Runtime/Engine/Private/Particles/ParticleGpuSimulation.cpp
+++ b/Engine/Source/Runtime/Engine/Private/Particles/ParticleGpuSimulation.cpp
@@ -2297,7 +2297,9 @@ static FBox ComputeParticleBounds(
 			CreateInfo);
 		FUnorderedAccessViewRHIRef BoundsVertexBufferUAV = RHICmdList.CreateUnorderedAccessView(
 			BoundsVertexBufferRHI,
-			PF_A32B32G32R32F );
+			FRHIViewDesc::CreateBufferUAV()
+				.SetType(FRHIViewDesc::EBufferType::Typed)
+				.SetFormat(PF_A32B32G32R32F));
 
 		// Grab the shader.
 		TShaderMapRef<FParticleBoundsCS> ParticleBoundsCS(GetGlobalShaderMap(FeatureLevel));
@@ -2452,7 +2454,11 @@ public:
 			check(TileBufferSize > 0);
 			FRHIResourceCreateInfo CreateInfo(TEXT("FParticleTileVertexBuffer"));
 			VertexBufferRHI = RHICmdList.CreateVertexBuffer( TileBufferSize, BUF_Static | BUF_KeepCPUAccessible | BUF_ShaderResource, CreateInfo );
-			VertexBufferSRV = RHICmdList.CreateShaderResourceView( VertexBufferRHI, Stride, PF_A32B32G32R32F);
+			VertexBufferSRV = RHICmdList.CreateShaderResourceView(
+				VertexBufferRHI, 
+				FRHIViewDesc::CreateBufferSRV()
+					.SetType(FRHIViewDesc::EBufferType::Typed)
+					.SetFormat(PF_A32B32G32R32F));
 		}
 	}
 
@@ -2512,7 +2518,11 @@ public:
 			const EBufferUsageFlags Flags = BUF_Static | /*BUF_KeepCPUAccessible | */BUF_ShaderResource;
 			FRHIResourceCreateInfo CreateInfo(TEXT("FGPUParticleVertexBuffer"));
 			VertexBufferRHI = RHICmdList.CreateVertexBuffer(BufferSize, Flags, CreateInfo);
-			VertexBufferSRV = RHICmdList.CreateShaderResourceView(VertexBufferRHI, BufferStride, PF_FloatRGBA);
+			VertexBufferSRV = RHICmdList.CreateShaderResourceView(
+				VertexBufferRHI, 
+				FRHIViewDesc::CreateBufferSRV()
+					.SetType(FRHIViewDesc::EBufferType::Typed)
+					.SetFormat(PF_FloatRGBA));
 		}
 	}
 
diff --git a/Engine/Source/Runtime/Engine/Private/Particles/ParticleResources.cpp b/Engine/Source/Runtime/Engine/Private/Particles/ParticleResources.cpp
--- a/Engine/Source/Runtime/Engine/Private/Particles/ParticleResources.cpp
+++ b/Engine/Source/Runtime/Engine/Private/Particles/ParticleResources.cpp
@@ -144,7 +144,11 @@ void FParticleScratchVertexBuffer::InitRHI(FRHICommandListBase& RHICmdList)
 
 	FRHIResourceCreateInfo CreateInfo(TEXT("FParticleScratchVertexBuffer"));
 	VertexBufferRHI = RHICmdList.CreateVertexBuffer(GParticleScratchVertexBufferSize, Flags, CreateInfo);
-	VertexBufferSRV_A32B32G32R32F = RHICmdList.CreateShaderResourceView(VertexBufferRHI, Stride, PF_A32B32G32R32F);
+	VertexBufferSRV_A32B32G32R32F = RHICmdList.CreateShaderResourceView(
+		VertexBufferRHI, 
+		FRHIViewDesc::CreateBufferSRV()
+			.SetType(FRHIViewDesc::EBufferType::Typed)
+			.SetFormat(PF_A32B32G32R32F));
 }
 
 FParticleShaderParamRef FParticleScratchVertexBuffer::GetShaderParam()
diff --git a/Engine/Source/Runtime/Engine/Private/Particles/ParticleSortingGPU.cpp b/Engine/Source/Runtime/Engine/Private/Particles/ParticleSortingGPU.cpp
--- a/Engine/Source/Runtime/Engine/Private/Particles/ParticleSortingGPU.cpp
+++ b/Engine/Source/Runtime/Engine/Private/Particles/ParticleSortingGPU.cpp
@@ -191,14 +191,30 @@ void FParticleSortBuffers::InitRHI(FRHICommandListBase& RHICmdList)
 		FRHIResourceCreateInfo CreateInfo(TEXT("PartialSortKeyBuffer"));
 
 		KeyBuffers[BufferIndex] = RHICmdList.CreateVertexBuffer( BufferSize * sizeof(uint32), BUF_Static | BUF_ShaderResource | BUF_UnorderedAccess, CreateInfo);
-		KeyBufferSRVs[BufferIndex] = RHICmdList.CreateShaderResourceView( KeyBuffers[BufferIndex], /*Stride=*/ sizeof(uint32), PF_R32_UINT );
-		KeyBufferUAVs[BufferIndex] = RHICmdList.CreateUnorderedAccessView( KeyBuffers[BufferIndex], PF_R32_UINT );
+		KeyBufferSRVs[BufferIndex] = RHICmdList.CreateShaderResourceView(
+			KeyBuffers[BufferIndex],
+			FRHIViewDesc::CreateBufferSRV()
+				.SetType(FRHIViewDesc::EBufferType::Typed)
+				.SetFormat(PF_R32_UINT));
+		KeyBufferUAVs[BufferIndex] = RHICmdList.CreateUnorderedAccessView(
+			KeyBuffers[BufferIndex],
+			FRHIViewDesc::CreateBufferUAV()
+				.SetType(FRHIViewDesc::EBufferType::Typed)
+				.SetFormat(PF_R32_UINT));
 
 		CreateInfo.DebugName = TEXT("PartialSortVertexBuffer");
 		VertexBuffers[BufferIndex] = RHICmdList.CreateVertexBuffer( BufferSize * sizeof(uint32), BUF_Static | BUF_ShaderResource | BUF_UnorderedAccess, CreateInfo);
 
-		VertexBufferSortSRVs[BufferIndex] = RHICmdList.CreateShaderResourceView(VertexBuffers[BufferIndex], /*Stride=*/ sizeof(uint32), PF_R32_UINT);
-		VertexBufferSortUAVs[BufferIndex] = RHICmdList.CreateUnorderedAccessView(VertexBuffers[BufferIndex], PF_R32_UINT);
+		VertexBufferSortSRVs[BufferIndex] = RHICmdList.CreateShaderResourceView(
+			VertexBuffers[BufferIndex],
+			FRHIViewDesc::CreateBufferSRV()
+				.SetType(FRHIViewDesc::EBufferType::Typed)
+				.SetFormat(PF_R32_UINT));
+		VertexBufferSortUAVs[BufferIndex] = RHICmdList.CreateUnorderedAccessView(
+			VertexBuffers[BufferIndex],
+			FRHIViewDesc::CreateBufferUAV()
+				.SetType(FRHIViewDesc::EBufferType::Typed)
+				.SetFormat(PF_R32_UINT));
 	}
 }
 
diff --git a/Engine/Source/Runtime/Engine/Private/Particles/ParticleVertexFactory.cpp b/Engine/Source/Runtime/Engine/Private/Particles/ParticleVertexFactory.cpp
--- a/Engine/Source/Runtime/Engine/Private/Particles/ParticleVertexFactory.cpp
+++ b/Engine/Source/Runtime/Engine/Private/Particles/ParticleVertexFactory.cpp
@@ -33,7 +33,11 @@ public:
 		FMemory::Memzero(BufferData, Size);
 		RHICmdList.UnlockBuffer(VertexBufferRHI);
 
-		VertexBufferSRV = RHICmdList.CreateShaderResourceView(VertexBufferRHI, sizeof(FVector2f), PF_G32R32F);
+		VertexBufferSRV = RHICmdList.CreateShaderResourceView(
+			VertexBufferRHI,
+			FRHIViewDesc::CreateBufferSRV()
+				.SetType(FRHIViewDesc::EBufferType::Typed)
+				.SetFormat(PF_G32R32F));
 	}
 	
 	virtual void ReleaseRHI() override
diff --git a/Engine/Source/Runtime/Engine/Private/Particles/SubUVAnimation.cpp b/Engine/Source/Runtime/Engine/Private/Particles/SubUVAnimation.cpp
--- a/Engine/Source/Runtime/Engine/Private/Particles/SubUVAnimation.cpp
+++ b/Engine/Source/Runtime/Engine/Private/Particles/SubUVAnimation.cpp
@@ -63,7 +63,11 @@ void FSubUVBoundingGeometryBuffer::InitRHI(FRHICommandListBase& RHICmdList)
 	if (Vertices->Num() > 0)
 	{
 		VertexBufferRHI = UE::RHIResourceUtils::CreateVertexBufferFromArray(RHICmdList, TEXT("FSubUVBoundingGeometryBuffer"), EBufferUsageFlags::ShaderResource | EBufferUsageFlags::Static, MakeConstArrayView(*Vertices));
-		ShaderResourceView = RHICmdList.CreateShaderResourceView(VertexBufferRHI, sizeof(FVector2f), PF_G32R32F);
+		ShaderResourceView = RHICmdList.CreateShaderResourceView(
+			VertexBufferRHI,
+			FRHIViewDesc::CreateBufferSRV()
+				.SetType(FRHIViewDesc::EBufferType::Typed)
+				.SetFormat(PF_G32R32F));
 	}
 }
 
diff --git a/Engine/Source/Runtime/Engine/Private/PositionVertexBuffer.cpp b/Engine/Source/Runtime/Engine/Private/PositionVertexBuffer.cpp
--- a/Engine/Source/Runtime/Engine/Private/PositionVertexBuffer.cpp
+++ b/Engine/Source/Runtime/Engine/Private/PositionVertexBuffer.cpp
@@ -240,7 +240,11 @@ void FPositionVertexBuffer::InitRHI(FRHICommandListBase& RHICmdList)
 	// Always create SRV if VertexBufferRHI is valid, as several systems rely on having position SRV. Memory impact has been measured to be minimal.
 	if (VertexBufferRHI)
 	{
-		PositionComponentSRV = RHICmdList.CreateShaderResourceView(VertexBufferRHI, 4, PF_R32_FLOAT);
+		PositionComponentSRV = RHICmdList.CreateShaderResourceView(
+			VertexBufferRHI, 
+			FRHIViewDesc::CreateBufferSRV()
+				.SetType(FRHIViewDesc::EBufferType::Typed)
+				.SetFormat(PF_R32_FLOAT));
 	}
 }
 
diff --git a/Engine/Source/Runtime/Engine/Private/PrecomputedVolumetricLightmap.cpp b/Engine/Source/Runtime/Engine/Private/PrecomputedVolumetricLightmap.cpp
--- a/Engine/Source/Runtime/Engine/Private/PrecomputedVolumetricLightmap.cpp
+++ b/Engine/Source/Runtime/Engine/Private/PrecomputedVolumetricLightmap.cpp
@@ -52,7 +52,7 @@ void FVolumetricLightmapDataLayer::CreateUAV(FRHICommandListBase& RHICmdList)
 {
 	check(Texture);
 
-	UAV = RHICmdList.CreateUnorderedAccessView(Texture);
+	UAV = RHICmdList.CreateUnorderedAccessView(Texture, FRHIViewDesc::CreateTextureUAV().SetDimensionFromTexture(Texture));
 }
 
 TGlobalResource<FVolumetricLightmapBrickAtlas> GVolumetricLightmapBrickAtlas;
@@ -272,13 +272,21 @@ ENGINE_API void FPrecomputedVolumetricLightmapData::InitRHIForSubLevelResources(
 		{
 			FRHIResourceCreateInfo CreateInfo(TEXT("SubLevelBrickPositionsBuffer"), &SubLevelBrickPositions);
 			SubLevelBrickPositionsBuffer = RHICmdList.CreateVertexBuffer(SubLevelBrickPositions.Num() * SubLevelBrickPositions.GetTypeSize(), BUF_Static | BUF_ShaderResource, CreateInfo);
-			SubLevelBrickPositionsSRV = RHICmdList.CreateShaderResourceView(SubLevelBrickPositionsBuffer, sizeof(uint32), PF_R32_UINT);
+			SubLevelBrickPositionsSRV = RHICmdList.CreateShaderResourceView(
+				SubLevelBrickPositionsBuffer, 
+				FRHIViewDesc::CreateBufferSRV()
+					.SetType(FRHIViewDesc::EBufferType::Typed)
+					.SetFormat(PF_R32_UINT));
 		}
 
 		{
 			FRHIResourceCreateInfo CreateInfo(TEXT("IndirectionTextureOriginalValuesBuffer"), &IndirectionTextureOriginalValues);
 			IndirectionTextureOriginalValuesBuffer = RHICmdList.CreateVertexBuffer(IndirectionTextureOriginalValues.Num() * IndirectionTextureOriginalValues.GetTypeSize(), BUF_Static | BUF_ShaderResource, CreateInfo);
-			IndirectionTextureOriginalValuesSRV = RHICmdList.CreateShaderResourceView(IndirectionTextureOriginalValuesBuffer, sizeof(FColor), PF_R8G8B8A8_UINT);
+			IndirectionTextureOriginalValuesSRV = RHICmdList.CreateShaderResourceView(
+				IndirectionTextureOriginalValuesBuffer, 
+				FRHIViewDesc::CreateBufferSRV()
+					.SetType(FRHIViewDesc::EBufferType::Typed)
+					.SetFormat(PF_R8G8B8A8_UINT));
 		}
 	}
 }
diff --git a/Engine/Source/Runtime/Engine/Private/PrimitiveUniformShaderParameters.cpp b/Engine/Source/Runtime/Engine/Private/PrimitiveUniformShaderParameters.cpp
--- a/Engine/Source/Runtime/Engine/Private/PrimitiveUniformShaderParameters.cpp
+++ b/Engine/Source/Runtime/Engine/Private/PrimitiveUniformShaderParameters.cpp
@@ -48,12 +48,12 @@ void FSinglePrimitiveStructured::InitRHI(FRHICommandListBase& RHICmdList)
 			.SetClassName(ClassName);
 
 		PrimitiveSceneDataTextureRHI = RHICreateTexture(Desc);
-		PrimitiveSceneDataTextureSRV = RHICmdList.CreateShaderResourceView(PrimitiveSceneDataTextureRHI, 0);
+		PrimitiveSceneDataTextureSRV = RHICmdList.CreateShaderResourceView(PrimitiveSceneDataTextureRHI, FRHIViewDesc::CreateTextureSRV().SetDimensionFromTexture(PrimitiveSceneDataTextureRHI));
 	}
 
 	FRHIResourceCreateInfo CreateInfo(TEXT("SkyIrradianceEnvironmentMap"));
 	SkyIrradianceEnvironmentMapRHI = RHICmdList.CreateStructuredBuffer(sizeof(FVector4f), sizeof(FVector4f) * 8, BUF_Static | BUF_ShaderResource, CreateInfo);
-	SkyIrradianceEnvironmentMapSRV = RHICmdList.CreateShaderResourceView(SkyIrradianceEnvironmentMapRHI);
+	SkyIrradianceEnvironmentMapSRV = RHICmdList.CreateShaderResourceView(SkyIrradianceEnvironmentMapRHI, FRHIViewDesc::CreateBufferSRV().SetTypeFromBuffer(SkyIrradianceEnvironmentMapRHI));
 
 	UploadToGPU(RHICmdList);
 }
diff --git a/Engine/Source/Runtime/Engine/Private/Rendering/ColorVertexBuffer.cpp b/Engine/Source/Runtime/Engine/Private/Rendering/ColorVertexBuffer.cpp
--- a/Engine/Source/Runtime/Engine/Private/Rendering/ColorVertexBuffer.cpp
+++ b/Engine/Source/Runtime/Engine/Private/Rendering/ColorVertexBuffer.cpp
@@ -426,7 +426,11 @@ void FColorVertexBuffer::InitRHI(FRHICommandListBase& RHICmdList)
 	VertexBufferRHI = CreateRHIBuffer(RHICmdList);
 	if (VertexBufferRHI && RHISupportsManualVertexFetch(GMaxRHIShaderPlatform))
 	{
-		ColorComponentsSRV = RHICmdList.CreateShaderResourceView(VertexBufferRHI, 4, PF_R8G8B8A8);
+		ColorComponentsSRV = RHICmdList.CreateShaderResourceView(
+			VertexBufferRHI,
+			FRHIViewDesc::CreateBufferSRV()
+				.SetType(FRHIViewDesc::EBufferType::Typed)
+				.SetFormat(PF_R8G8B8A8));
 	}
 }
 
diff --git a/Engine/Source/Runtime/Engine/Private/Rendering/MorphTargetVertexInfoBuffers.cpp b/Engine/Source/Runtime/Engine/Private/Rendering/MorphTargetVertexInfoBuffers.cpp
--- a/Engine/Source/Runtime/Engine/Private/Rendering/MorphTargetVertexInfoBuffers.cpp
+++ b/Engine/Source/Runtime/Engine/Private/Rendering/MorphTargetVertexInfoBuffers.cpp
@@ -30,7 +30,7 @@ void FMorphTargetVertexInfoBuffers::InitRHI(FRHICommandListBase& RHICmdList)
 	void* BufferPtr = RHICmdList.LockBuffer(MorphDataBuffer, 0, BufferSize, RLM_WriteOnly);
 	FMemory::ParallelMemcpy(BufferPtr, MorphData.GetData(), BufferSize, EMemcpyCachePolicy::StoreUncached);
 	RHICmdList.UnlockBuffer(MorphDataBuffer);
-	MorphDataSRV = RHICmdList.CreateShaderResourceView(MorphDataBuffer);
+	MorphDataSRV = RHICmdList.CreateShaderResourceView(MorphDataBuffer, FRHIViewDesc::CreateBufferSRV().SetTypeFromBuffer(MorphDataBuffer));
 
 	if (bEmptyMorphCPUDataOnInitRHI)
 	{
diff --git a/Engine/Source/Runtime/Engine/Private/Rendering/SkeletalMeshAttributeVertexBuffer.cpp b/Engine/Source/Runtime/Engine/Private/Rendering/SkeletalMeshAttributeVertexBuffer.cpp
--- a/Engine/Source/Runtime/Engine/Private/Rendering/SkeletalMeshAttributeVertexBuffer.cpp
+++ b/Engine/Source/Runtime/Engine/Private/Rendering/SkeletalMeshAttributeVertexBuffer.cpp
@@ -176,7 +176,11 @@ void FSkeletalMeshAttributeVertexBuffer::InitRHI(FRHICommandListBase& RHICmdList
 	FResourceArrayInterface* ResourceArray = ValueData->GetResourceArray();
 	FRHIResourceCreateInfo CreateInfo(TEXT("SkeletalMeshAttributeVertexBuffer"), ResourceArray);
 	ValueBuffer.VertexBufferRHI = RHICmdList.CreateVertexBuffer(ResourceArray->GetResourceDataSize(), BUF_Static | BUF_ShaderResource, CreateInfo);
-	ValueBuffer.VertexBufferSRV = RHICmdList.CreateShaderResourceView(ValueBuffer.VertexBufferRHI, ComponentStride, PixelFormat);
+	ValueBuffer.VertexBufferSRV = RHICmdList.CreateShaderResourceView(
+		ValueBuffer.VertexBufferRHI,
+		FRHIViewDesc::CreateBufferSRV()
+			.SetType(FRHIViewDesc::EBufferType::Typed)
+			.SetFormat(static_cast<EPixelFormat>(PixelFormat)));
 }
 
 
diff --git a/Engine/Source/Runtime/Engine/Private/Rendering/SkeletalMeshHalfEdgeBuffer.cpp b/Engine/Source/Runtime/Engine/Private/Rendering/SkeletalMeshHalfEdgeBuffer.cpp
--- a/Engine/Source/Runtime/Engine/Private/Rendering/SkeletalMeshHalfEdgeBuffer.cpp
+++ b/Engine/Source/Runtime/Engine/Private/Rendering/SkeletalMeshHalfEdgeBuffer.cpp
@@ -298,11 +298,11 @@ void FSkeletalMeshHalfEdgeBuffer::InitRHI(FRHICommandListBase& RHICmdList)
 	//    the RHI will takeover ownership of the streamed buffer and becomes usable
 	if (ensure(VertexToEdgeBufferRHI))
 	{
-		VertexToEdgeBufferSRV = RHICmdList.CreateShaderResourceView(VertexToEdgeBufferRHI);
+		VertexToEdgeBufferSRV = RHICmdList.CreateShaderResourceView(VertexToEdgeBufferRHI, FRHIViewDesc::CreateBufferSRV().SetTypeFromBuffer(VertexToEdgeBufferRHI));
 	}
 	if (ensure(EdgeToTwinEdgeBufferRHI))
 	{
-		EdgeToTwinEdgeBufferSRV = RHICmdList.CreateShaderResourceView(EdgeToTwinEdgeBufferRHI);
+		EdgeToTwinEdgeBufferSRV = RHICmdList.CreateShaderResourceView(EdgeToTwinEdgeBufferRHI, FRHIViewDesc::CreateBufferSRV().SetTypeFromBuffer(EdgeToTwinEdgeBufferRHI));
 	}
 }
 
diff --git a/Engine/Source/Runtime/Engine/Private/Rendering/SkeletalMeshVertexClothBuffer.cpp b/Engine/Source/Runtime/Engine/Private/Rendering/SkeletalMeshVertexClothBuffer.cpp
--- a/Engine/Source/Runtime/Engine/Private/Rendering/SkeletalMeshVertexClothBuffer.cpp
+++ b/Engine/Source/Runtime/Engine/Private/Rendering/SkeletalMeshVertexClothBuffer.cpp
@@ -106,7 +106,11 @@ void FSkeletalMeshVertexClothBuffer::InitRHI(FRHICommandListBase& RHICmdList)
 	VertexBufferRHI = CreateRHIBuffer(RHICmdList);
 	if (VertexBufferRHI)
 	{
-		VertexBufferSRV = RHICmdList.CreateShaderResourceView(VertexBufferRHI, 16, PF_A32B32G32R32F);
+		VertexBufferSRV = RHICmdList.CreateShaderResourceView(
+			VertexBufferRHI,
+			FRHIViewDesc::CreateBufferSRV()
+				.SetType(FRHIViewDesc::EBufferType::Typed)
+				.SetFormat(PF_A32B32G32R32F));
 	}
 }
 
diff --git a/Engine/Source/Runtime/Engine/Private/Rendering/SkinWeightVertexBuffer.cpp b/Engine/Source/Runtime/Engine/Private/Rendering/SkinWeightVertexBuffer.cpp
--- a/Engine/Source/Runtime/Engine/Private/Rendering/SkinWeightVertexBuffer.cpp
+++ b/Engine/Source/Runtime/Engine/Private/Rendering/SkinWeightVertexBuffer.cpp
@@ -145,7 +145,11 @@ void FSkinWeightLookupVertexBuffer::InitRHI(FRHICommandListBase& RHICmdList)
 
 		if (bSRV)
 		{
-			SRVValue = RHICmdList.CreateShaderResourceView(VertexBufferRHI, PixelFormatStride, PixelFormat);
+			SRVValue = RHICmdList.CreateShaderResourceView(
+				VertexBufferRHI,
+				FRHIViewDesc::CreateBufferSRV()
+					.SetType(FRHIViewDesc::EBufferType::Typed)
+					.SetFormat(PixelFormat));
 		}
 	}
 }
@@ -440,7 +444,11 @@ void FSkinWeightDataVertexBuffer::InitRHI(FRHICommandListBase& RHICmdList)
 
 	if (bSRV)
 	{
-		SRVValue = RHICmdList.CreateShaderResourceView(VertexBufferRHI, GetPixelFormatStride(), GetPixelFormat());
+		SRVValue = RHICmdList.CreateShaderResourceView(
+			VertexBufferRHI,
+			FRHIViewDesc::CreateBufferSRV()
+				.SetType(FRHIViewDesc::EBufferType::Typed)
+				.SetFormat(GetPixelFormat()));
 	}
 }
 
diff --git a/Engine/Source/Runtime/Engine/Private/Rendering/StaticMeshVertexBuffer.cpp b/Engine/Source/Runtime/Engine/Private/Rendering/StaticMeshVertexBuffer.cpp
--- a/Engine/Source/Runtime/Engine/Private/Rendering/StaticMeshVertexBuffer.cpp
+++ b/Engine/Source/Runtime/Engine/Private/Rendering/StaticMeshVertexBuffer.cpp
@@ -326,7 +326,11 @@ void FStaticMeshVertexBuffer::InitRHI(FRHICommandListBase& RHICmdList)
 	{
 		uint32       Stride = GetUseHighPrecisionTangentBasis() ? 8 : 4;
 		EPixelFormat Format = GetUseHighPrecisionTangentBasis() ? PF_R16G16B16A16_SNORM : PF_R8G8B8A8_SNORM;
-		TangentsSRV = RHICmdList.CreateShaderResourceView(TangentsVertexBuffer.VertexBufferRHI, Stride, Format);
+		TangentsSRV = RHICmdList.CreateShaderResourceView(
+			TangentsVertexBuffer.VertexBufferRHI,
+			FRHIViewDesc::CreateBufferSRV()
+				.SetType(FRHIViewDesc::EBufferType::Typed)
+				.SetFormat(Format));
 	}
 
 	const bool bHadTexCoordData = TexcoordData != nullptr;
@@ -336,7 +340,11 @@ void FStaticMeshVertexBuffer::InitRHI(FRHICommandListBase& RHICmdList)
 	{
 		uint32       Stride = GetUseFullPrecisionUVs() ? 8 : 4;
 		EPixelFormat Format = GetUseFullPrecisionUVs() ? PF_G32R32F : PF_G16R16F;
-		TextureCoordinatesSRV = RHICmdList.CreateShaderResourceView(TexCoordVertexBuffer.VertexBufferRHI, Stride, Format);
+		TextureCoordinatesSRV = RHICmdList.CreateShaderResourceView(
+			TexCoordVertexBuffer.VertexBufferRHI,
+			FRHIViewDesc::CreateBufferSRV()
+				.SetType(FRHIViewDesc::EBufferType::Typed)
+				.SetFormat(Format));
 	}
 }
 
diff --git a/Engine/Source/Runtime/Engine/Private/SkeletalRenderGPUSkin.cpp b/Engine/Source/Runtime/Engine/Private/SkeletalRenderGPUSkin.cpp
--- a/Engine/Source/Runtime/Engine/Private/SkeletalRenderGPUSkin.cpp
+++ b/Engine/Source/Runtime/Engine/Private/SkeletalRenderGPUSkin.cpp
@@ -162,7 +162,11 @@ void FMorphVertexBuffer::InitRHI(FRHICommandListBase& RHICmdList)
 
 	VertexBufferRHI = RHICmdList.CreateVertexBuffer(Size, Flags, CreateInfo);
 	VertexBufferRHI->SetOwnerName(GetOwnerName());
-	SRVValue = RHICmdList.CreateShaderResourceView(VertexBufferRHI, 4, PF_R32_FLOAT);
+	SRVValue = RHICmdList.CreateShaderResourceView(
+		VertexBufferRHI, 
+		FRHIViewDesc::CreateBufferSRV()
+			.SetType(FRHIViewDesc::EBufferType::Typed)
+			.SetFormat(PF_R32_FLOAT));
 
 	if (!bUseGPUMorphTargets)
 	{
@@ -176,7 +180,11 @@ void FMorphVertexBuffer::InitRHI(FRHICommandListBase& RHICmdList)
 	}
 	else
 	{
-		UAVValue = RHICmdList.CreateUnorderedAccessView(VertexBufferRHI, PF_R32_UINT);
+		UAVValue = RHICmdList.CreateUnorderedAccessView(
+			VertexBufferRHI, 
+			FRHIViewDesc::CreateBufferUAV()
+				.SetType(FRHIViewDesc::EBufferType::Typed)
+				.SetFormat(PF_R32_UINT));
 		bNeedsInitialClear = true;
 	}
 
diff --git a/Engine/Source/Runtime/Engine/Private/StaticMesh.cpp b/Engine/Source/Runtime/Engine/Private/StaticMesh.cpp
--- a/Engine/Source/Runtime/Engine/Private/StaticMesh.cpp
+++ b/Engine/Source/Runtime/Engine/Private/StaticMesh.cpp
@@ -291,7 +291,11 @@ void FStaticMeshSectionAreaWeightedTriangleSamplerBuffer::InitRHI(FRHICommandLis
 		}
 		RHICmdList.UnlockBuffer(BufferSectionTriangleRHI);
 
-		BufferSectionTriangleSRV = RHICmdList.CreateShaderResourceView(BufferSectionTriangleRHI, sizeof(SectionTriangleInfo), PF_R32G32_UINT);
+		BufferSectionTriangleSRV = RHICmdList.CreateShaderResourceView(
+			BufferSectionTriangleRHI, 
+			FRHIViewDesc::CreateBufferSRV()
+				.SetType(FRHIViewDesc::EBufferType::Typed)
+				.SetFormat(PF_R32G32_UINT));
 	}
 }
 
diff --git a/Engine/Source/Runtime/Engine/Private/TextureRenderTarget2D.cpp b/Engine/Source/Runtime/Engine/Private/TextureRenderTarget2D.cpp
--- a/Engine/Source/Runtime/Engine/Private/TextureRenderTarget2D.cpp
+++ b/Engine/Source/Runtime/Engine/Private/TextureRenderTarget2D.cpp
@@ -567,7 +567,7 @@ void FTextureRenderTarget2DResource::InitRHI(FRHICommandListBase& RHICmdList)
 
 		if (EnumHasAnyFlags(TexCreateFlags, ETextureCreateFlags::UAV))
 		{
-			UnorderedAccessViewRHI = RHICmdList.CreateUnorderedAccessView(RenderTargetTextureRHI);
+			UnorderedAccessViewRHI = RHICmdList.CreateUnorderedAccessView(RenderTargetTextureRHI, FRHIViewDesc::CreateTextureUAV().SetDimensionFromTexture(RenderTargetTextureRHI));
 		}
 
 		SetGPUMask(FRHIGPUMask::All());
diff --git a/Engine/Source/Runtime/Engine/Private/TextureRenderTarget2DArray.cpp b/Engine/Source/Runtime/Engine/Private/TextureRenderTarget2DArray.cpp
--- a/Engine/Source/Runtime/Engine/Private/TextureRenderTarget2DArray.cpp
+++ b/Engine/Source/Runtime/Engine/Private/TextureRenderTarget2DArray.cpp
@@ -260,7 +260,7 @@ void FTextureRenderTarget2DArrayResource::InitRHI(FRHICommandListBase& RHICmdLis
 
 		if (EnumHasAnyFlags(TexCreateFlags, ETextureCreateFlags::UAV))
 		{
-			UnorderedAccessViewRHI = RHICmdList.CreateUnorderedAccessView(TextureRHI);
+			UnorderedAccessViewRHI = RHICmdList.CreateUnorderedAccessView(TextureRHI, FRHIViewDesc::CreateTextureUAV().SetDimensionFromTexture(TextureRHI));
 		}
 
 		RHIUpdateTextureReference(Owner->TextureReference.TextureReferenceRHI, TextureRHI);
diff --git a/Engine/Source/Runtime/Engine/Private/TextureRenderTargetCube.cpp b/Engine/Source/Runtime/Engine/Private/TextureRenderTargetCube.cpp
--- a/Engine/Source/Runtime/Engine/Private/TextureRenderTargetCube.cpp
+++ b/Engine/Source/Runtime/Engine/Private/TextureRenderTargetCube.cpp
@@ -316,7 +316,7 @@ void FTextureRenderTargetCubeResource::InitRHI(FRHICommandListBase& RHICmdList)
 
 		if (EnumHasAnyFlags(TexCreateFlags, ETextureCreateFlags::UAV))
 		{
-			UnorderedAccessViewRHI = RHICmdList.CreateUnorderedAccessView(TextureRHI);
+			UnorderedAccessViewRHI = RHICmdList.CreateUnorderedAccessView(TextureRHI, FRHIViewDesc::CreateTextureUAV().SetDimensionFromTexture(TextureRHI));
 		}
 
 		RHIUpdateTextureReference(Owner->TextureReference.TextureReferenceRHI, TextureRHI);
diff --git a/Engine/Source/Runtime/Engine/Private/TextureRenderTargetVolume.cpp b/Engine/Source/Runtime/Engine/Private/TextureRenderTargetVolume.cpp
--- a/Engine/Source/Runtime/Engine/Private/TextureRenderTargetVolume.cpp
+++ b/Engine/Source/Runtime/Engine/Private/TextureRenderTargetVolume.cpp
@@ -264,7 +264,7 @@ void FTextureRenderTargetVolumeResource::InitRHI(FRHICommandListBase& RHICmdList
 
 		if (EnumHasAnyFlags(TexCreateFlags, TexCreate_UAV))
 		{
-			UnorderedAccessViewRHI = RHICmdList.CreateUnorderedAccessView(TextureRHI);
+			UnorderedAccessViewRHI = RHICmdList.CreateUnorderedAccessView(TextureRHI, FRHIViewDesc::CreateTextureUAV().SetDimensionFromTexture(TextureRHI));
 		}
 
 		RHIUpdateTextureReference(Owner->TextureReference.TextureReferenceRHI,TextureRHI);
diff --git a/Engine/Source/Runtime/Engine/Private/VT/VirtualTextureAdapter.cpp b/Engine/Source/Runtime/Engine/Private/VT/VirtualTextureAdapter.cpp
--- a/Engine/Source/Runtime/Engine/Private/VT/VirtualTextureAdapter.cpp
+++ b/Engine/Source/Runtime/Engine/Private/VT/VirtualTextureAdapter.cpp
@@ -328,7 +328,11 @@ public:
 		SourceSRVs.Reserve(SourceTexture->GetNumMips());
 		for (uint32 MipIndex = 0; MipIndex < SourceTexture->GetNumMips(); ++MipIndex)
 		{
-			SourceSRVs.Add(RHICmdList.CreateShaderResourceView(SourceTexture, MipIndex));
+			SourceSRVs.Add(RHICmdList.CreateShaderResourceView(
+				SourceTexture,
+				FRHIViewDesc::CreateTextureSRV()
+					.SetDimensionFromTexture(SourceTexture)
+					.SetMipRange(MipIndex, 1)));
 		}
 	}
 
diff --git a/Engine/Source/Runtime/Engine/Private/VectorField.cpp b/Engine/Source/Runtime/Engine/Private/VectorField.cpp
--- a/Engine/Source/Runtime/Engine/Private/VectorField.cpp
+++ b/Engine/Source/Runtime/Engine/Private/VectorField.cpp
@@ -1091,7 +1091,7 @@ public:
 
 			if (GetFeatureLevel() >= ERHIFeatureLevel::SM5)
 			{
-				VolumeTextureUAV = RHICmdList.CreateUnorderedAccessView(VolumeTextureRHI);
+				VolumeTextureUAV = RHICmdList.CreateUnorderedAccessView(VolumeTextureRHI, FRHIViewDesc::CreateTextureUAV().SetDimensionFromTexture(VolumeTextureRHI));
 			}
 		}
 	}
diff --git a/Engine/Source/Runtime/Engine/Public/RawIndexBuffer.h b/Engine/Source/Runtime/Engine/Public/RawIndexBuffer.h
--- a/Engine/Source/Runtime/Engine/Public/RawIndexBuffer.h
+++ b/Engine/Source/Runtime/Engine/Public/RawIndexBuffer.h
@@ -400,7 +400,11 @@ public:
 		if (IndexBufferRHI && IsSRVNeeded(Indices.GetAllowCPUAccess()) && bHadIndexData)
 		{
 			// If the index buffer is a placeholder we still need to create a FRHIShaderResourceView.
-			SRVValue = RHICmdList.CreateShaderResourceView(IndexBufferRHI, sizeof(INDEX_TYPE), sizeof(INDEX_TYPE) == 2 ? PF_R16_UINT : PF_R32_UINT);
+			SRVValue = RHICmdList.CreateShaderResourceView(
+				IndexBufferRHI,
+				FRHIViewDesc::CreateBufferSRV()
+					.SetType(FRHIViewDesc::EBufferType::Typed)
+					.SetFormat(sizeof(INDEX_TYPE) == 2 ? PF_R16_UINT : PF_R32_UINT));
 		}
 	}
 	
diff --git a/Engine/Source/Runtime/Engine/Public/Rendering/SkeletalMeshDuplicatedVerticesBuffer.h b/Engine/Source/Runtime/Engine/Public/Rendering/SkeletalMeshDuplicatedVerticesBuffer.h
--- a/Engine/Source/Runtime/Engine/Public/Rendering/SkeletalMeshDuplicatedVerticesBuffer.h
+++ b/Engine/Source/Runtime/Engine/Public/Rendering/SkeletalMeshDuplicatedVerticesBuffer.h
@@ -117,7 +117,11 @@ public:
 			CreateInfo.OwnerName = GetOwnerName();
             DuplicatedVerticesIndexBuffer.VertexBufferRHI = RHICmdList.CreateVertexBuffer(ResourceArray->GetResourceDataSize(), BUF_Static | BUF_ShaderResource, CreateInfo);
 			DuplicatedVerticesIndexBuffer.VertexBufferRHI->SetOwnerName(GetOwnerName());
-            DuplicatedVerticesIndexBuffer.VertexBufferSRV = RHICmdList.CreateShaderResourceView(DuplicatedVerticesIndexBuffer.VertexBufferRHI, sizeof(uint32), PF_R32_UINT);
+            DuplicatedVerticesIndexBuffer.VertexBufferSRV = RHICmdList.CreateShaderResourceView(
+				DuplicatedVerticesIndexBuffer.VertexBufferRHI,
+				FRHIViewDesc::CreateBufferSRV()
+					.SetType(FRHIViewDesc::EBufferType::Typed)
+					.SetFormat(PF_R32_UINT));
         }
 
         {
@@ -129,7 +133,11 @@ public:
 			CreateInfo.OwnerName = GetOwnerName();
             LengthAndIndexDuplicatedVerticesIndexBuffer.VertexBufferRHI = RHICmdList.CreateVertexBuffer(ResourceArray->GetResourceDataSize(), BUF_Static | BUF_ShaderResource, CreateInfo);
 			LengthAndIndexDuplicatedVerticesIndexBuffer.VertexBufferRHI->SetOwnerName(GetOwnerName());
-            LengthAndIndexDuplicatedVerticesIndexBuffer.VertexBufferSRV = RHICmdList.CreateShaderResourceView(LengthAndIndexDuplicatedVerticesIndexBuffer.VertexBufferRHI, sizeof(uint32), PF_R32_UINT);
+            LengthAndIndexDuplicatedVerticesIndexBuffer.VertexBufferSRV = RHICmdList.CreateShaderResourceView(
+				LengthAndIndexDuplicatedVerticesIndexBuffer.VertexBufferRHI,
+				FRHIViewDesc::CreateBufferSRV()
+					.SetType(FRHIViewDesc::EBufferType::Typed)
+					.SetFormat(PF_R32_UINT));
         }
     }
 
diff --git a/Engine/Source/Runtime/Experimental/GeometryCollectionEngine/Private/GeometryCollection/GeometryCollectionSceneProxy.h b/Engine/Source/Runtime/Experimental/GeometryCollectionEngine/Private/GeometryCollection/GeometryCollectionSceneProxy.h
--- a/Engine/Source/Runtime/Experimental/GeometryCollectionEngine/Private/GeometryCollection/GeometryCollectionSceneProxy.h
+++ b/Engine/Source/Runtime/Experimental/GeometryCollectionEngine/Private/GeometryCollection/GeometryCollectionSceneProxy.h
@@ -39,7 +39,11 @@ public:
 		// #note: This differs from instanced static mesh in that we are storing the entire transform in the buffer rather than
 		// splitting out the translation.  This is to simplify transferring data at runtime as a memcopy
 		VertexBufferRHI = RHICmdList.CreateVertexBuffer(NumTransforms * sizeof(FVector4f) * 4, BUF_Dynamic | BUF_ShaderResource, CreateInfo);		
-		VertexBufferSRV = RHICmdList.CreateShaderResourceView(VertexBufferRHI, 16, PF_A32B32G32R32F);
+		VertexBufferSRV = RHICmdList.CreateShaderResourceView(
+			VertexBufferRHI,
+			FRHIViewDesc::CreateBufferSRV()
+				.SetType(FRHIViewDesc::EBufferType::Typed)
+				.SetFormat(PF_A32B32G32R32F));
 	}
 
 	void UpdateDynamicData(FRHICommandListBase& RHICmdList, const TArray<FMatrix44f>& Transforms, EResourceLockMode LockMode);
diff --git a/Engine/Source/Runtime/Landscape/Private/LandscapeEditLayers.cpp b/Engine/Source/Runtime/Landscape/Private/LandscapeEditLayers.cpp
--- a/Engine/Source/Runtime/Landscape/Private/LandscapeEditLayers.cpp
+++ b/Engine/Source/Runtime/Landscape/Private/LandscapeEditLayers.cpp
@@ -1017,7 +1017,7 @@ public:
 	{
 		FRHIResourceCreateInfo CreateInfo(TEXT("FLandscapeLayerWeightmapExtractMaterialLayersComputeShaderResource"));
 		ComponentsData = RHICmdList.CreateStructuredBuffer(sizeof(FLandscapeLayerWeightmapExtractMaterialLayersComponentData), OriginalComponentsData.Num() * sizeof(FLandscapeLayerWeightmapExtractMaterialLayersComponentData), BUF_ShaderResource | BUF_Volatile, CreateInfo);
-		ComponentsDataSRV = RHICmdList.CreateShaderResourceView(ComponentsData);
+		ComponentsDataSRV = RHICmdList.CreateShaderResourceView(ComponentsData, FRHIViewDesc::CreateBufferSRV().SetTypeFromBuffer(ComponentsData));
 
 		uint8* Buffer = (uint8*)RHICmdList.LockBuffer(ComponentsData, 0, OriginalComponentsData.Num() * sizeof(FLandscapeLayerWeightmapExtractMaterialLayersComponentData), RLM_WriteOnly);
 		FMemory::Memcpy(Buffer, OriginalComponentsData.GetData(), OriginalComponentsData.Num() * sizeof(FLandscapeLayerWeightmapExtractMaterialLayersComponentData));
@@ -1177,7 +1177,7 @@ public:
 		FRHIResourceCreateInfo CreateInfo(TEXT("ComponentsData"));
 		uint32 ComponentsDataMemSize = OriginalComponentsData.Num() * sizeof(FLandscapeLayerWeightmapPackMaterialLayersComponentData);
 		ComponentsData = RHICmdList.CreateStructuredBuffer(sizeof(FLandscapeLayerWeightmapPackMaterialLayersComponentData), ComponentsDataMemSize, BUF_ShaderResource | BUF_Volatile, CreateInfo);
-		ComponentsDataSRV = RHICmdList.CreateShaderResourceView(ComponentsData);
+		ComponentsDataSRV = RHICmdList.CreateShaderResourceView(ComponentsData, FRHIViewDesc::CreateBufferSRV().SetTypeFromBuffer(ComponentsData));
 
 		uint8* Buffer = (uint8*)RHICmdList.LockBuffer(ComponentsData, 0, ComponentsDataMemSize, RLM_WriteOnly);
 		FMemory::Memcpy(Buffer, OriginalComponentsData.GetData(), ComponentsDataMemSize);
@@ -1186,7 +1186,11 @@ public:
 		FRHIResourceCreateInfo WeightBlendCreateInfo(TEXT("WeightmapWeightBlendMode"));
 		uint32 WeightBlendMemSize = OriginalWeightmapWeightBlendModeData.Num() * sizeof(float);
 		WeightmapWeightBlendMode = RHICmdList.CreateVertexBuffer(WeightBlendMemSize, BUF_ShaderResource | BUF_Volatile, WeightBlendCreateInfo);
-		WeightmapWeightBlendModeSRV = RHICmdList.CreateShaderResourceView(WeightmapWeightBlendMode, sizeof(float), PF_R32_FLOAT);
+		WeightmapWeightBlendModeSRV = RHICmdList.CreateShaderResourceView(
+			WeightmapWeightBlendMode, 
+			FRHIViewDesc::CreateBufferSRV()
+				.SetType(FRHIViewDesc::EBufferType::Typed)
+				.SetFormat(PF_R32_FLOAT));
 
 		void* WeightmapWeightBlendModePtr = RHICmdList.LockBuffer(WeightmapWeightBlendMode, 0, WeightBlendMemSize, RLM_WriteOnly);
 		FMemory::Memcpy(WeightmapWeightBlendModePtr, OriginalWeightmapWeightBlendModeData.GetData(), WeightBlendMemSize);
@@ -1195,7 +1199,11 @@ public:
 		FRHIResourceCreateInfo TextureOutputOffsetCreateInfo(TEXT("WeightmapTextureOutputOffset"));
 		uint32 TextureOutputOffsetMemSize = OriginalTextureOutputOffset.Num() * sizeof(FVector2f);
 		WeightmapTextureOutputOffset = RHICmdList.CreateVertexBuffer(TextureOutputOffsetMemSize, BUF_ShaderResource | BUF_Volatile, TextureOutputOffsetCreateInfo);
-		WeightmapTextureOutputOffsetSRV = RHICmdList.CreateShaderResourceView(WeightmapTextureOutputOffset, sizeof(FVector2f), PF_G32R32F);
+		WeightmapTextureOutputOffsetSRV = RHICmdList.CreateShaderResourceView(
+			WeightmapTextureOutputOffset,
+			FRHIViewDesc::CreateBufferSRV()
+				.SetType(FRHIViewDesc::EBufferType::Typed)
+				.SetFormat(PF_G32R32F));
 
 		void* TextureOutputOffsetPtr = RHICmdList.LockBuffer(WeightmapTextureOutputOffset, 0, TextureOutputOffsetMemSize, RLM_WriteOnly);
 		FMemory::Memcpy(TextureOutputOffsetPtr, OriginalTextureOutputOffset.GetData(), TextureOutputOffsetMemSize);
diff --git a/Engine/Source/Runtime/Landscape/Private/LandscapeEditResources.cpp b/Engine/Source/Runtime/Landscape/Private/LandscapeEditResources.cpp
--- a/Engine/Source/Runtime/Landscape/Private/LandscapeEditResources.cpp
+++ b/Engine/Source/Runtime/Landscape/Private/LandscapeEditResources.cpp
@@ -56,7 +56,12 @@ void FLandscapeTexture2DResource::InitRHI(FRHICommandListBase& RHICmdList)
 
 	if (bCreateSRV)
 	{
-		TextureSRV = RHICmdList.CreateShaderResourceView(TextureRHI, /*MipLevel = */0, static_cast<uint8>(NumMips), Format);
+		TextureSRV = RHICmdList.CreateShaderResourceView(
+			TextureRHI,
+			FRHIViewDesc::CreateTextureSRV()
+				.SetDimensionFromTexture(TextureRHI)
+				.SetMipRange(0, static_cast<uint8>(NumMips))
+				.SetFormat(Format));
 	}
 }
 
@@ -137,7 +142,12 @@ void FLandscapeTexture2DArrayResource::InitRHI(FRHICommandListBase& RHICmdList)
 
 	if (bCreateSRV)
 	{
-		TextureSRV = RHICmdList.CreateShaderResourceView(TextureRHI, /*MipLevel = */0, static_cast<uint8>(NumMips), Format);
+		TextureSRV = RHICmdList.CreateShaderResourceView(
+			TextureRHI,
+			FRHIViewDesc::CreateTextureSRV()
+				.SetDimensionFromTexture(TextureRHI)
+				.SetMipRange(0, static_cast<uint8>(NumMips))
+				.SetFormat(Format));
 	}
 }
 
diff --git a/Engine/Source/Runtime/Landscape/Private/LandscapeRender.cpp b/Engine/Source/Runtime/Landscape/Private/LandscapeRender.cpp
--- a/Engine/Source/Runtime/Landscape/Private/LandscapeRender.cpp
+++ b/Engine/Source/Runtime/Landscape/Private/LandscapeRender.cpp
@@ -1064,7 +1064,11 @@ void FLandscapeRenderSystem::UpdateBuffers(FRHICommandListBase& RHICmdList)
 			const static FLazyName ClassName(TEXT("FLandscapeRenderSystem"));
 			CreateInfo.ClassName = ClassName;
 			SectionLODBiasBuffer = RHICmdList.CreateVertexBuffer(SectionLODBiases.GetResourceDataSize(), BUF_ShaderResource | BUF_Dynamic, CreateInfo);
-			SectionLODBiasSRV = RHICmdList.CreateShaderResourceView(SectionLODBiasBuffer, sizeof(float), PF_R32_FLOAT);
+			SectionLODBiasSRV = RHICmdList.CreateShaderResourceView(
+				SectionLODBiasBuffer,
+				FRHIViewDesc::CreateBufferSRV()
+					.SetType(FRHIViewDesc::EBufferType::Typed)
+					.SetFormat(PF_R32_FLOAT));
 			bUpdateUB = true;
 		}
 		else
@@ -1187,11 +1191,19 @@ void FLandscapeSceneViewExtension::PreRenderView_RenderThread(FRDGBuilder& Graph
 			{
 				FRHIResourceCreateInfo CreateInfoLODBuffer(TEXT("LandscapeLODDataBuffer"), &LandscapeView.LandscapeLODData);
 				FBufferRHIRef LandscapeLODDataBuffer = RHICmdList.CreateVertexBuffer(LandscapeView.LandscapeLODData.GetResourceDataSize(), BUF_ShaderResource | BUF_Volatile, CreateInfoLODBuffer);
-				LandscapeView.View->LandscapePerComponentDataBuffer = RHICmdList.CreateShaderResourceView(LandscapeLODDataBuffer, sizeof(float), PF_R32_FLOAT);
+				LandscapeView.View->LandscapePerComponentDataBuffer = RHICmdList.CreateShaderResourceView(
+					LandscapeLODDataBuffer,
+					FRHIViewDesc::CreateBufferSRV()
+						.SetType(FRHIViewDesc::EBufferType::Typed)
+						.SetFormat(PF_R32_FLOAT));
 
 				FRHIResourceCreateInfo CreateInfoIndirection(TEXT("LandscapeIndirectionBuffer"), &LandscapeView.LandscapeIndirection);
 				FBufferRHIRef LandscapeIndirectionBuffer = RHICmdList.CreateVertexBuffer(LandscapeView.LandscapeIndirection.GetResourceDataSize(), BUF_ShaderResource | BUF_Volatile, CreateInfoIndirection);
-				LandscapeView.View->LandscapeIndirectionBuffer = RHICmdList.CreateShaderResourceView(LandscapeIndirectionBuffer, sizeof(uint32), PF_R32_UINT);
+				LandscapeView.View->LandscapeIndirectionBuffer = RHICmdList.CreateShaderResourceView(
+					LandscapeIndirectionBuffer,
+					FRHIViewDesc::CreateBufferSRV()
+						.SetType(FRHIViewDesc::EBufferType::Typed)
+						.SetFormat(PF_R32_UINT));
 			}
 
 		}, GIsThreadedRendering && GLandscapeUseAsyncTasksForLODComputation);
diff --git a/Engine/Source/Runtime/MRMesh/Private/MRMeshComponent.cpp b/Engine/Source/Runtime/MRMesh/Private/MRMeshComponent.cpp
--- a/Engine/Source/Runtime/MRMesh/Private/MRMeshComponent.cpp
+++ b/Engine/Source/Runtime/MRMesh/Private/MRMeshComponent.cpp
@@ -218,7 +218,11 @@ public:
 		{
 			NewSection->PositionBuffer.InitResource(RHICmdList);
 			NewSection->PositionBuffer.InitRHIWith(RHICmdList, Args.PositionData);
-			NewSection->PositionBufferSRV = RHICmdList.CreateShaderResourceView(NewSection->PositionBuffer.VertexBufferRHI, sizeof(float), PF_R32_FLOAT);
+			NewSection->PositionBufferSRV = RHICmdList.CreateShaderResourceView(
+				NewSection->PositionBuffer.VertexBufferRHI, 
+				FRHIViewDesc::CreateBufferSRV()
+					.SetType(FRHIViewDesc::EBufferType::Typed)
+					.SetFormat(PF_R32_FLOAT));
 		}
 
 		// TEXTURE COORDS BUFFER
@@ -227,7 +231,11 @@ public:
 			if (Args.UVData.Num())
 			{
 				NewSection->UVBuffer.InitRHIWith(RHICmdList, Args.UVData);
-				NewSection->UVBufferSRV = RHICmdList.CreateShaderResourceView(NewSection->UVBuffer.VertexBufferRHI, 8, PF_G32R32F);
+				NewSection->UVBufferSRV = RHICmdList.CreateShaderResourceView(
+					NewSection->UVBuffer.VertexBufferRHI,
+					FRHIViewDesc::CreateBufferSRV()
+						.SetType(FRHIViewDesc::EBufferType::Typed)
+						.SetFormat(PF_G32R32F));
 			}
 		}
 
@@ -241,7 +249,11 @@ public:
 
 			if (RHISupportsManualVertexFetch(GMaxRHIShaderPlatform))
 			{
-				NewSection->TangentXZBufferSRV = RHICmdList.CreateShaderResourceView(NewSection->TangentXZBuffer.VertexBufferRHI, 4, PF_R8G8B8A8_SNORM);
+				NewSection->TangentXZBufferSRV = RHICmdList.CreateShaderResourceView(
+					NewSection->TangentXZBuffer.VertexBufferRHI,
+					FRHIViewDesc::CreateBufferSRV()
+						.SetType(FRHIViewDesc::EBufferType::Typed)
+						.SetFormat(PF_R8G8B8A8_SNORM));
 			}
 		}
 
@@ -251,7 +263,11 @@ public:
 			if (Args.ColorData.Num())
 			{
 				NewSection->ColorBuffer.InitRHIWith(RHICmdList, Args.ColorData);
-				NewSection->ColorBufferSRV = RHICmdList.CreateShaderResourceView(NewSection->ColorBuffer.VertexBufferRHI, 4, PF_R8G8B8A8);
+				NewSection->ColorBufferSRV = RHICmdList.CreateShaderResourceView(
+					NewSection->ColorBuffer.VertexBufferRHI,
+					FRHIViewDesc::CreateBufferSRV()
+						.SetType(FRHIViewDesc::EBufferType::Typed)
+						.SetFormat(PF_R8G8B8A8));
 			}
 		}
 
diff --git a/Engine/Source/Runtime/MediaAssets/Private/Misc/MediaTextureResource.cpp b/Engine/Source/Runtime/MediaAssets/Private/Misc/MediaTextureResource.cpp
--- a/Engine/Source/Runtime/MediaAssets/Private/Misc/MediaTextureResource.cpp
+++ b/Engine/Source/Runtime/MediaAssets/Private/Misc/MediaTextureResource.cpp
@@ -1107,8 +1107,16 @@ void FMediaTextureResource::GetColorSpaceConversionMatrixForSample(const TShared
 						GraphicsPSOInit.BoundShaderState.PixelShaderRHI = ConvertShader.GetPixelShader();
 						SetGraphicsPipelineState(RHICmdList, GraphicsPSOInit, 0);
 						FIntPoint TexDim = InputTexture->GetSizeXY();
-						TempSRV0 = RHICmdList.CreateShaderResourceView(InputTexture, 0, 1, PF_G8);								// note: RHI does provide "magic" to select Y vs. UV planes based on the pixel format (D3D/DXGI)
-						TempSRV1 = RHICmdList.CreateShaderResourceView(InputTexture, 0, 1, PF_R8G8);
+						TempSRV0 = RHICmdList.CreateShaderResourceView(
+							InputTexture,
+							FRHIViewDesc::CreateTextureSRV()
+								.SetDimensionFromTexture(InputTexture)
+								.SetFormat(PF_G8)); // note: RHI does provide "magic" to select Y vs. UV planes based on the pixel format (D3D/DXGI)
+						TempSRV1 = RHICmdList.CreateShaderResourceView(
+							InputTexture,
+							FRHIViewDesc::CreateTextureSRV()
+								.SetDimensionFromTexture(InputTexture)
+								.SetFormat(PF_R8G8));
 						SetShaderParametersLegacyPS(RHICmdList, ConvertShader, TexDim, TempSRV0, TempSRV1, OutputDim, YUVMtx, Sample->GetEncodingType(), ColorSpaceMtx, SampleFormat == EMediaTextureSampleFormat::CharNV21, Sample->GetToneMapMethod());
 					}
 					else
@@ -1135,8 +1143,16 @@ void FMediaTextureResource::GetColorSpaceConversionMatrixForSample(const TShared
 						GraphicsPSOInit.BoundShaderState.PixelShaderRHI = ConvertShader.GetPixelShader();
 						SetGraphicsPipelineState(RHICmdList, GraphicsPSOInit, 0);
 
-						FShaderResourceViewRHIRef Y_SRV = RHICmdList.CreateShaderResourceView(InputTexture, 0, 1, PF_G16);		// note: RHI does provide "magic" to select Y vs. UV planes based on the pixel format (D3D/DXGI)
-						FShaderResourceViewRHIRef UV_SRV = RHICmdList.CreateShaderResourceView(InputTexture, 0, 1, PF_G16R16);
+						FShaderResourceViewRHIRef Y_SRV = RHICmdList.CreateShaderResourceView(
+							InputTexture,
+							FRHIViewDesc::CreateTextureSRV()
+								.SetDimensionFromTexture(InputTexture)
+								.SetFormat(PF_G16)); // note: RHI does provide "magic" to select Y vs. UV planes based on the pixel format (D3D/DXGI)
+						FShaderResourceViewRHIRef UV_SRV = RHICmdList.CreateShaderResourceView(
+							InputTexture,
+							FRHIViewDesc::CreateTextureSRV()
+								.SetDimensionFromTexture(InputTexture)
+								.SetFormat(PF_G16R16));
 						SetShaderParametersLegacyPS(RHICmdList, ConvertShader, TexDim, Y_SRV, UV_SRV, OutputDim, YUVMtx, ColorSpaceMtx, Sample->GetEncodingType(), Sample->GetToneMapMethod());
 					}
 					else
@@ -1200,7 +1216,11 @@ void FMediaTextureResource::GetColorSpaceConversionMatrixForSample(const TShared
 					TShaderMapRef<FYUVY416ConvertPS> ConvertShader(ShaderMap);
 					GraphicsPSOInit.BoundShaderState.PixelShaderRHI = ConvertShader.GetPixelShader();
 					SetGraphicsPipelineState(RHICmdList, GraphicsPSOInit, 0);
-					FShaderResourceViewRHIRef SRV = RHICmdList.CreateShaderResourceView(InputTexture, 0, 1, (Sample->GetFormat() == EMediaTextureSampleFormat::Y416) ? PF_A16B16G16R16 : PF_A32B32G32R32F);
+					FShaderResourceViewRHIRef SRV = RHICmdList.CreateShaderResourceView(
+						InputTexture,
+						FRHIViewDesc::CreateTextureSRV()
+							.SetDimensionFromTexture(InputTexture)
+							.SetFormat((Sample->GetFormat() == EMediaTextureSampleFormat::Y416) ? PF_A16B16G16R16 : PF_A32B32G32R32F));
 
 					SetShaderParametersLegacyPS(RHICmdList, ConvertShader, SRV, YUVMtx, Sample->GetEncodingType(), ColorSpaceMtx,
 												InputTexture->GetFormat() == PF_A8R8G8B8, Sample->GetToneMapMethod());
@@ -1248,7 +1268,11 @@ void FMediaTextureResource::GetColorSpaceConversionMatrixForSample(const TShared
 					TShaderMapRef<FARGB16BigConvertPS> ConvertShader(ShaderMap);
 					GraphicsPSOInit.BoundShaderState.PixelShaderRHI = ConvertShader.GetPixelShader();
 					SetGraphicsPipelineState(RHICmdList, GraphicsPSOInit, 0);
-					FShaderResourceViewRHIRef SRV = RHICmdList.CreateShaderResourceView(InputTexture, 0, 1, PF_R16G16B16A16_UINT);
+					FShaderResourceViewRHIRef SRV = RHICmdList.CreateShaderResourceView(
+						InputTexture,
+						FRHIViewDesc::CreateTextureSRV()
+							.SetDimensionFromTexture(InputTexture)
+							.SetFormat(PF_R16G16B16A16_UINT));
 
 					SetShaderParametersLegacyPS(RHICmdList, ConvertShader, SRV, OutputDim, Sample->GetEncodingType(), ColorSpaceMtx, Sample->GetToneMapMethod());
 				}
diff --git a/Engine/Source/Runtime/RHI/Public/RHICommandList.h b/Engine/Source/Runtime/RHI/Public/RHICommandList.h
--- a/Engine/Source/Runtime/RHI/Public/RHICommandList.h
+++ b/Engine/Source/Runtime/RHI/Public/RHICommandList.h
@@ -946,7 +946,7 @@ public:
 		return CreateShaderResourceView(Initializer.Buffer, Initializer);
 	}
 
-	//UE_DEPRECATED(5.3, "Use the CreateUnorderedAccessView function that takes an FRHIViewDesc.")
+	UE_DEPRECATED(5.6, "Use the CreateUnorderedAccessView function that takes an FRHIViewDesc.")
 	FORCEINLINE FUnorderedAccessViewRHIRef CreateUnorderedAccessView(FRHIBuffer* Buffer, bool bUseUAVCounter, bool bAppendBuffer)
 	{
 		return CreateUnorderedAccessView(Buffer, FRHIViewDesc::CreateBufferUAV()
@@ -956,7 +956,7 @@ public:
 		);
 	}
 
-	//UE_DEPRECATED(5.3, "Use the CreateUnorderedAccessView function that takes an FRHIViewDesc.")
+	UE_DEPRECATED(5.6, "Use the CreateUnorderedAccessView function that takes an FRHIViewDesc.")
 	FORCEINLINE FUnorderedAccessViewRHIRef CreateUnorderedAccessView(FRHIBuffer* Buffer, uint8 Format)
 	{
 		// For back-compat reasons, SRVs of byte-address buffers created via this function ignore the Format, and instead create raw views.
@@ -975,7 +975,7 @@ public:
 		}
 	}
 
-	//UE_DEPRECATED(5.3, "Use the CreateUnorderedAccessView function that takes an FRHIViewDesc.")
+	UE_DEPRECATED(5.6, "Use the CreateUnorderedAccessView function that takes an FRHIViewDesc.")
 	FORCEINLINE FUnorderedAccessViewRHIRef CreateUnorderedAccessView(FRHITexture* Texture, uint32 MipLevel = 0, uint16 FirstArraySlice = 0, uint16 NumArraySlices = 0)
 	{
 		check(MipLevel < 256);
@@ -987,7 +987,7 @@ public:
 		);
 	}
 
-	//UE_DEPRECATED(5.3, "Use the CreateUnorderedAccessView function that takes an FRHIViewDesc.")
+	UE_DEPRECATED(5.6, "Use the CreateUnorderedAccessView function that takes an FRHIViewDesc.")
 	FORCEINLINE FUnorderedAccessViewRHIRef CreateUnorderedAccessView(FRHITexture* Texture, uint32 MipLevel, uint8 Format, uint16 FirstArraySlice = 0, uint16 NumArraySlices = 0)
 	{
 		check(MipLevel < 256);
@@ -1000,7 +1000,7 @@ public:
 		);
 	}
 
-	//UE_DEPRECATED(5.3, "Use the CreateShaderResourceView function that takes an FRHIViewDesc.")
+	UE_DEPRECATED(5.6, "Use the CreateShaderResourceView function that takes an FRHIViewDesc.")
 	FORCEINLINE FShaderResourceViewRHIRef CreateShaderResourceView(FRHIBuffer* Buffer)
 	{
 		FShaderResourceViewRHIRef SRVRef = CreateShaderResourceView(Buffer, FRHIViewDesc::CreateBufferSRV()
@@ -1010,7 +1010,7 @@ public:
 		return SRVRef;
 	}
 
-	//UE_DEPRECATED(5.3, "Use the CreateShaderResourceView function that takes an FRHIViewDesc.")
+	UE_DEPRECATED(5.6, "Use the CreateShaderResourceView function that takes an FRHIViewDesc.")
 	FORCEINLINE FShaderResourceViewRHIRef CreateShaderResourceView(FRHIBuffer* Buffer, uint32 Stride, uint8 Format)
 	{
 		check(Format != PF_Unknown);
@@ -1032,7 +1032,7 @@ public:
 		}
 	}
 
-	//UE_DEPRECATED(5.3, "Use the CreateShaderResourceView function that takes an FRHIViewDesc.")
+	UE_DEPRECATED(5.6, "Use the CreateShaderResourceView function that takes an FRHIViewDesc.")
 	FORCEINLINE FShaderResourceViewRHIRef CreateShaderResourceView(FRHITexture* Texture, const FRHITextureSRVCreateInfo& CreateInfo)
 	{
 		return CreateShaderResourceView(Texture, FRHIViewDesc::CreateTextureSRV()
@@ -1045,7 +1045,7 @@ public:
 		);
 	}
 
-	//UE_DEPRECATED(5.3, "Use the CreateShaderResourceView function that takes an FRHIViewDesc.")
+	UE_DEPRECATED(5.6, "Use the CreateShaderResourceView function that takes an FRHIViewDesc.")
 	FORCEINLINE FShaderResourceViewRHIRef CreateShaderResourceView(FRHITexture* Texture, uint8 MipLevel)
 	{
 		return CreateShaderResourceView(Texture, FRHIViewDesc::CreateTextureSRV()
@@ -1054,7 +1054,7 @@ public:
 		);
 	}
 
-	//UE_DEPRECATED(5.3, "Use the CreateShaderResourceView function that takes an FRHIViewDesc.")
+	UE_DEPRECATED(5.6, "Use the CreateShaderResourceView function that takes an FRHIViewDesc.")
 	FORCEINLINE FShaderResourceViewRHIRef CreateShaderResourceView(FRHITexture* Texture, uint8 MipLevel, uint8 NumMipLevels, EPixelFormat Format)
 	{
 		return CreateShaderResourceView(Texture, FRHIViewDesc::CreateTextureSRV()
@@ -1064,7 +1064,7 @@ public:
 		);
 	}
 
-	//UE_DEPRECATED(5.3, "Use the CreateShaderResourceView function that takes an FRHIViewDesc.")
+	UE_DEPRECATED(5.6, "Use the CreateShaderResourceView function that takes an FRHIViewDesc.")
 	FORCEINLINE FShaderResourceViewRHIRef CreateShaderResourceViewWriteMask(FRHITexture* Texture2DRHI)
 	{
 		return CreateShaderResourceView(Texture2DRHI, FRHIViewDesc::CreateTextureSRV()
@@ -1073,7 +1073,7 @@ public:
 		);
 	}
 
-	//UE_DEPRECATED(5.3, "Use the CreateShaderResourceView function that takes an FRHIViewDesc.")
+	UE_DEPRECATED(5.6, "Use the CreateShaderResourceView function that takes an FRHIViewDesc.")
 	FORCEINLINE FShaderResourceViewRHIRef CreateShaderResourceViewFMask(FRHITexture* Texture2DRHI)
 	{
 		return CreateShaderResourceView(Texture2DRHI, FRHIViewDesc::CreateTextureSRV()
diff --git a/Engine/Source/Runtime/RHI/Public/RHIUtilities.h b/Engine/Source/Runtime/RHI/Public/RHIUtilities.h
--- a/Engine/Source/Runtime/RHI/Public/RHIUtilities.h
+++ b/Engine/Source/Runtime/RHI/Public/RHIUtilities.h
@@ -137,8 +137,8 @@ struct FTextureRWBuffer
 		FRHICommandListBase& RHICmdList = FRHICommandListImmediate::Get();
 
 		Buffer = RHICreateTexture(Desc);
-		UAV = RHICmdList.CreateUnorderedAccessView(Buffer, 0);
-		SRV = RHICmdList.CreateShaderResourceView(Buffer, 0);
+		UAV = RHICmdList.CreateUnorderedAccessView(Buffer, FRHIViewDesc::CreateTextureUAV().SetDimensionFromTexture(Buffer));
+		SRV = RHICmdList.CreateShaderResourceView(Buffer, FRHIViewDesc::CreateTextureSRV().SetDimensionFromTexture(Buffer));
 	}
 
 	void Initialize3D(const TCHAR* InDebugName, uint32 BytesPerElement, uint32 SizeX, uint32 SizeY, uint32 SizeZ, EPixelFormat Format, ETextureCreateFlags Flags = DefaultTextureInitFlag)
@@ -152,8 +152,8 @@ struct FTextureRWBuffer
 		FRHICommandListBase& RHICmdList = FRHICommandListImmediate::Get();
 
 		Buffer = RHICreateTexture(Desc);
-		UAV = RHICmdList.CreateUnorderedAccessView(Buffer, 0);
-		SRV = RHICmdList.CreateShaderResourceView(Buffer, 0);
+		UAV = RHICmdList.CreateUnorderedAccessView(Buffer, FRHIViewDesc::CreateTextureUAV().SetDimensionFromTexture(Buffer));
+		SRV = RHICmdList.CreateShaderResourceView(Buffer, FRHIViewDesc::CreateTextureSRV().SetDimensionFromTexture(Buffer));
 	}
 
 	void Release()
@@ -233,8 +233,18 @@ struct FRWBuffer
 		CreateInfo.OwnerName = OwnerName;
 		CreateInfo.ResourceArray = InResourceArray;
 		Buffer = RHICmdList.CreateBuffer(NumBytes, BUF_VertexBuffer | BUF_UnorderedAccess | BUF_ShaderResource | AdditionalUsage, 0, InResourceState, CreateInfo);
-		UAV = RHICmdList.CreateUnorderedAccessView(Buffer, UE_PIXELFORMAT_TO_UINT8(Format));
-		SRV = RHICmdList.CreateShaderResourceView(Buffer, BytesPerElement, UE_PIXELFORMAT_TO_UINT8(Format));
+		
+		UAV = RHICmdList.CreateUnorderedAccessView(
+			Buffer,
+			FRHIViewDesc::CreateBufferUAV()
+				.SetType(FRHIViewDesc::EBufferType::Typed)
+				.SetFormat(Format));
+
+		SRV = RHICmdList.CreateShaderResourceView(
+			Buffer,
+			FRHIViewDesc::CreateBufferSRV()
+				.SetType(FRHIViewDesc::EBufferType::Typed)
+				.SetFormat(Format));
 	}
 
 	void Initialize(FRHICommandListBase& RHICmdList, const TCHAR* InDebugName, uint32 BytesPerElement, uint32 NumElements, EPixelFormat Format, EBufferUsageFlags AdditionalUsage = BUF_None, FResourceArrayUploadInterface* InResourceArray = nullptr)
@@ -280,8 +290,7 @@ struct FTextureReadBuffer2D
 			.SetBulkData(InBulkData);
 
 		Buffer = RHICreateTexture(Desc);
-		
-		SRV = RHICmdList.CreateShaderResourceView(Buffer, 0);
+		SRV = RHICmdList.CreateShaderResourceView(Buffer, FRHIViewDesc::CreateTextureSRV().SetDimensionFromTexture(Buffer));
 	}
 
 	void Release()
@@ -307,7 +316,12 @@ struct FReadBuffer
 		FRHIResourceCreateInfo CreateInfo(InDebugName);
 		CreateInfo.ResourceArray = InResourceArray;
 		Buffer = RHICmdList.CreateVertexBuffer(NumBytes, BUF_ShaderResource | AdditionalUsage, ERHIAccess::SRVMask, CreateInfo);
-		SRV = RHICmdList.CreateShaderResourceView(Buffer, BytesPerElement, UE_PIXELFORMAT_TO_UINT8(Format));
+
+		SRV = RHICmdList.CreateShaderResourceView(
+			Buffer, 
+			FRHIViewDesc::CreateBufferSRV()
+				.SetType(FRHIViewDesc::EBufferType::Typed)
+				.SetFormat(Format));
 	}
 
 	void Release()
@@ -341,8 +355,15 @@ struct FRWBufferStructured
 		NumBytes = BytesPerElement * NumElements;
 		FRHIResourceCreateInfo CreateInfo(InDebugName);
 		Buffer = RHICmdList.CreateStructuredBuffer(BytesPerElement, NumBytes, BUF_UnorderedAccess | BUF_ShaderResource | AdditionalUsage, InitialState, CreateInfo);
-		UAV = RHICmdList.CreateUnorderedAccessView(Buffer, bUseUavCounter, bAppendBuffer);
-		SRV = RHICmdList.CreateShaderResourceView(Buffer);
+		
+		UAV = RHICmdList.CreateUnorderedAccessView(
+			Buffer,
+			FRHIViewDesc::CreateBufferUAV()
+				.SetTypeFromBuffer(Buffer)
+				.SetAtomicCounter(bUseUavCounter)
+				.SetAppendBuffer(bAppendBuffer));
+		
+		SRV = RHICmdList.CreateShaderResourceView(Buffer, FRHIViewDesc::CreateBufferSRV().SetTypeFromBuffer(Buffer));
 	}
 
 	void Release()
@@ -368,7 +389,7 @@ struct FByteAddressBuffer
 		check( NumBytes % 4 == 0 );
 		FRHIResourceCreateInfo CreateInfo(InDebugName);
 		Buffer = RHICmdList.CreateStructuredBuffer(4, NumBytes, BUF_ShaderResource | BUF_ByteAddressBuffer | AdditionalUsage, ERHIAccess::SRVMask, CreateInfo);
-		SRV = RHICmdList.CreateShaderResourceView(Buffer);
+		SRV = RHICmdList.CreateShaderResourceView(Buffer, FRHIViewDesc::CreateBufferSRV().SetType(FRHIViewDesc::EBufferType::Raw));
 	}
 
 	void Release()
@@ -387,7 +408,7 @@ struct FRWByteAddressBuffer : public FByteAddressBuffer
 	void Initialize(FRHICommandListBase& RHICmdList, const TCHAR* DebugName, uint32 InNumBytes, EBufferUsageFlags AdditionalUsage = BUF_None)
 	{
 		FByteAddressBuffer::Initialize(RHICmdList, DebugName, InNumBytes, BUF_UnorderedAccess | AdditionalUsage);
-		UAV = RHICmdList.CreateUnorderedAccessView(Buffer, false, false);
+		UAV = RHICmdList.CreateUnorderedAccessView(Buffer, FRHIViewDesc::CreateBufferUAV().SetType(FRHIViewDesc::EBufferType::Raw));
 	}
 
 	void Release()
diff --git a/Engine/Source/Runtime/RenderCore/Private/DumpGPU.cpp b/Engine/Source/Runtime/RenderCore/Private/DumpGPU.cpp
--- a/Engine/Source/Runtime/RenderCore/Private/DumpGPU.cpp
+++ b/Engine/Source/Runtime/RenderCore/Private/DumpGPU.cpp
@@ -1024,16 +1024,20 @@ public:
 		(*Texture)->DisableLifetimeExtension();
 	}
 
-	FShaderResourceViewRHIRef CreateSRVNoLifetimeExtension(FRHICommandList& RHICmdList, FRHITexture* Texture, const FRHITextureSRVCreateInfo& Desc)
+	FShaderResourceViewRHIRef CreateSRVNoLifetimeExtension(FRHICommandList& RHICmdList, FRHITexture* Texture, FRHIViewDesc::FTextureSRV::FInitializer const& CreateInfo)
 	{
-		FShaderResourceViewRHIRef SRV = RHICmdList.CreateShaderResourceView(Texture, Desc);
+		FShaderResourceViewRHIRef SRV = RHICmdList.CreateShaderResourceView(Texture, CreateInfo);
 		SRV->DisableLifetimeExtension();
 		return SRV;
 	}
 
 	FUnorderedAccessViewRHIRef CreateUAVNoLifetimeExtension(FRHICommandList& RHICmdList, FRHITexture* Texture, uint32 MipLevel)
 	{
-		FUnorderedAccessViewRHIRef UAV = RHICmdList.CreateUnorderedAccessView(Texture, MipLevel);
+		FUnorderedAccessViewRHIRef UAV = RHICmdList.CreateUnorderedAccessView(
+			Texture,
+			FRHIViewDesc::CreateTextureUAV()
+				.SetDimensionFromTexture(Texture)
+				.SetMipLevel(uint8(MipLevel)));
 		UAV->DisableLifetimeExtension();
 		return UAV;
 	}
diff --git a/Engine/Source/Runtime/RenderCore/Private/GlobalRenderResources.cpp b/Engine/Source/Runtime/RenderCore/Private/GlobalRenderResources.cpp
--- a/Engine/Source/Runtime/RenderCore/Private/GlobalRenderResources.cpp
+++ b/Engine/Source/Runtime/RenderCore/Private/GlobalRenderResources.cpp
@@ -66,7 +66,7 @@ public:
 		SamplerStateRHI = GetOrCreateSamplerState(SamplerStateInitializer);
 
 		// Create a view of the texture
-		ShaderResourceViewRHI = RHICmdList.CreateShaderResourceView(TextureRHI, 0u);
+		ShaderResourceViewRHI = RHICmdList.CreateShaderResourceView(TextureRHI, FRHIViewDesc::CreateTextureSRV().SetDimensionFromTexture(TextureRHI));
 	}
 
 	virtual uint32 GetSizeX() const override { return 1; }
@@ -84,8 +84,16 @@ public:
 		VertexBufferRHI = RHICmdList.CreateVertexBuffer(16u, BUF_Static | BUF_ShaderResource | BUF_UnorderedAccess, CreateInfo);
 
 		// Create a view of the buffer
-		ShaderResourceViewRHI = RHICmdList.CreateShaderResourceView(VertexBufferRHI, 4u, PF_R32_UINT);
-		UnorderedAccessViewRHI = RHICmdList.CreateUnorderedAccessView(VertexBufferRHI, PF_R32_UINT);
+		ShaderResourceViewRHI = RHICmdList.CreateShaderResourceView(
+			VertexBufferRHI,
+			FRHIViewDesc::CreateBufferSRV()
+				.SetType(FRHIViewDesc::EBufferType::Typed)
+				.SetFormat(PF_R32_UINT));
+		UnorderedAccessViewRHI = RHICmdList.CreateUnorderedAccessView(
+			VertexBufferRHI,
+			FRHIViewDesc::CreateBufferUAV()
+				.SetType(FRHIViewDesc::EBufferType::Typed)
+				.SetFormat(PF_R32_UINT));
 	}
 };
 
@@ -101,8 +109,8 @@ public:
 		VertexBufferRHI = RHICmdList.CreateStructuredBuffer(sizeof(uint32), BufferSize, BUF_Static | BUF_ShaderResource | BUF_UnorderedAccess, CreateInfo);
 
 		// Create a view of the buffer
-		ShaderResourceViewRHI = RHICmdList.CreateShaderResourceView(VertexBufferRHI);
-		UnorderedAccessViewRHI = RHICmdList.CreateUnorderedAccessView(VertexBufferRHI, false, false);
+		ShaderResourceViewRHI = RHICmdList.CreateShaderResourceView(VertexBufferRHI, FRHIViewDesc::CreateBufferSRV().SetTypeFromBuffer(VertexBufferRHI));
+		UnorderedAccessViewRHI = RHICmdList.CreateUnorderedAccessView(VertexBufferRHI, FRHIViewDesc::CreateBufferUAV().SetTypeFromBuffer(VertexBufferRHI));
 	}
 };
 
@@ -122,8 +130,8 @@ public:
 		RHICmdList.UnlockBuffer(VertexBufferRHI);
 
 		// Create a view of the buffer
-		ShaderResourceViewRHI = RHICmdList.CreateShaderResourceView(VertexBufferRHI);
-		UnorderedAccessViewRHI = RHICmdList.CreateUnorderedAccessView(VertexBufferRHI, false, false);
+		ShaderResourceViewRHI = RHICmdList.CreateShaderResourceView(VertexBufferRHI, FRHIViewDesc::CreateBufferSRV().SetTypeFromBuffer(VertexBufferRHI));
+		UnorderedAccessViewRHI = RHICmdList.CreateUnorderedAccessView(VertexBufferRHI, FRHIViewDesc::CreateBufferUAV().SetTypeFromBuffer(VertexBufferRHI));
 	}
 };
 
@@ -142,8 +150,15 @@ public:
 		RHICmdList.UnlockBuffer(VertexBufferRHI);
 
 		// Create a view of the buffer
-		ShaderResourceViewRHI = RHICmdList.CreateShaderResourceView(VertexBufferRHI, sizeof(FVector4f), PF_A32B32G32R32F);
-		UnorderedAccessViewRHI = RHICmdList.CreateUnorderedAccessView(VertexBufferRHI, PF_A32B32G32R32F);
+		ShaderResourceViewRHI = RHICmdList.CreateShaderResourceView(
+			VertexBufferRHI,
+			FRHIViewDesc::CreateBufferSRV()
+				.SetType(FRHIViewDesc::EBufferType::Typed)
+				.SetFormat(PF_A32B32G32R32F));
+		UnorderedAccessViewRHI = RHICmdList.CreateUnorderedAccessView(
+			VertexBufferRHI, FRHIViewDesc::CreateBufferUAV()
+				.SetType(FRHIViewDesc::EBufferType::Typed)
+				.SetFormat(PF_A32B32G32R32F));
 	}
 };
 
@@ -190,7 +205,11 @@ public:
 		RHICmdList.UnlockBuffer(VertexBufferRHI);
 
 		// Create a view of the buffer
-		ShaderResourceViewRHI = RHICmdList.CreateShaderResourceView(VertexBufferRHI, sizeof(FVector4f), PF_A32B32G32R32F);
+		ShaderResourceViewRHI = RHICmdList.CreateShaderResourceView(
+			VertexBufferRHI,
+			FRHIViewDesc::CreateBufferSRV()
+				.SetType(FRHIViewDesc::EBufferType::Typed)
+				.SetFormat(PF_A32B32G32R32F));
 	}
 };
 
@@ -211,7 +230,11 @@ public:
 		RHICmdList.UnlockBuffer(VertexBufferRHI);
 
 		// Create a view of the buffer
-		ShaderResourceViewRHI = RHICmdList.CreateShaderResourceView(VertexBufferRHI, sizeof(FVector4f), PF_A32B32G32R32F);
+		ShaderResourceViewRHI = RHICmdList.CreateShaderResourceView(
+			VertexBufferRHI,
+			FRHIViewDesc::CreateBufferSRV()
+				.SetType(FRHIViewDesc::EBufferType::Typed)
+				.SetFormat(PF_A32B32G32R32F));
 	}
 };
 
@@ -537,7 +560,7 @@ public:
 		SamplerStateRHI = GetOrCreateSamplerState(SamplerStateInitializer);
 
 		// Create a view of the texture
-		ShaderResourceViewRHI = RHICmdList.CreateShaderResourceView(TextureRHI, 0u);
+		ShaderResourceViewRHI = RHICmdList.CreateShaderResourceView(TextureRHI, FRHIViewDesc::CreateTextureSRV().SetDimensionFromTexture(TextureRHI));
 	}
 
 	virtual uint32 GetSizeX() const override { return 1; }
@@ -604,7 +627,11 @@ void FNullColorVertexBuffer::InitRHI(FRHICommandListBase& RHICmdList)
 	Vertices[2] = FColor(255, 255, 255, 255).DWColor();
 	Vertices[3] = FColor(255, 255, 255, 255).DWColor();
 	RHICmdList.UnlockBuffer(VertexBufferRHI);
-	VertexBufferSRV = RHICmdList.CreateShaderResourceView(VertexBufferRHI, sizeof(FColor), PF_R8G8B8A8);
+	VertexBufferSRV = RHICmdList.CreateShaderResourceView(
+		VertexBufferRHI,
+		FRHIViewDesc::CreateBufferSRV()
+			.SetType(FRHIViewDesc::EBufferType::Typed)
+			.SetFormat(PF_R8G8B8A8));
 }
 
 void FNullColorVertexBuffer::ReleaseRHI()
@@ -631,7 +658,11 @@ void FNullVertexBuffer::InitRHI(FRHICommandListBase& RHICmdList)
 	*LockedData = FVector3f(0.0f);
 	RHICmdList.UnlockBuffer(VertexBufferRHI);
 
-	VertexBufferSRV = RHICmdList.CreateShaderResourceView(VertexBufferRHI, sizeof(FColor), PF_R8G8B8A8);
+	VertexBufferSRV = RHICmdList.CreateShaderResourceView(
+		VertexBufferRHI,
+		FRHIViewDesc::CreateBufferSRV()
+			.SetType(FRHIViewDesc::EBufferType::Typed)
+			.SetFormat(PF_R8G8B8A8));
 }
 
 void FNullVertexBuffer::ReleaseRHI()
diff --git a/Engine/Source/Runtime/RenderCore/Private/RayTracingValidationShaders.cpp b/Engine/Source/Runtime/RenderCore/Private/RayTracingValidationShaders.cpp
--- a/Engine/Source/Runtime/RenderCore/Private/RayTracingValidationShaders.cpp
+++ b/Engine/Source/Runtime/RenderCore/Private/RayTracingValidationShaders.cpp
@@ -33,7 +33,7 @@ void FRayTracingValidateGeometryBuildParamsCS::Dispatch(FRHICommandList& RHICmdL
 	const uint32 IndexStride = Initializer.IndexBuffer->GetStride();
 
 	const FRawBufferShaderResourceViewInitializer IBViewInitializer(Initializer.IndexBuffer);
-	FShaderResourceViewRHIRef IndexBufferSRV = RHICmdList.CreateShaderResourceView(IBViewInitializer);
+	FShaderResourceViewRHIRef IndexBufferSRV = RHICmdList.CreateShaderResourceView(IBViewInitializer.Buffer, FRHIViewDesc::CreateBufferSRV().SetTypeFromBuffer(IBViewInitializer.Buffer));
 
 	for (const FRayTracingGeometrySegment& Segment : Initializer.Segments)
 	{
@@ -55,7 +55,7 @@ void FRayTracingValidateGeometryBuildParamsCS::Dispatch(FRHICommandList& RHICmdL
 		SetShaderValue(BatchedParameters, ComputeShader->MaxVerticesParam, Segment.MaxVertices);
 
 		const FRawBufferShaderResourceViewInitializer VBViewInitializer(Segment.VertexBuffer);
-		FShaderResourceViewRHIRef VertexBufferSRV = RHICmdList.CreateShaderResourceView(VBViewInitializer);
+		FShaderResourceViewRHIRef VertexBufferSRV = RHICmdList.CreateShaderResourceView(VBViewInitializer.Buffer, FRHIViewDesc::CreateBufferSRV().SetTypeFromBuffer(VBViewInitializer.Buffer));
 
 		SetSRVParameter(BatchedParameters, ComputeShader->VertexBufferParam, VertexBufferSRV);
 		SetSRVParameter(BatchedParameters, ComputeShader->IndexBufferParam, IndexBufferSRV);
@@ -102,7 +102,7 @@ void FRayTracingValidateSceneBuildParamsCS::Dispatch(FRHICommandList& RHICmdList
 		FRHIBatchedShaderParameters& BatchedParameters = RHICmdList.GetScratchShaderParameters();
 
 		const FRawBufferShaderResourceViewInitializer InstanceBufferViewInitializer(InstanceBuffer);
-		FShaderResourceViewRHIRef InstanceBufferSRV = RHICmdList.CreateShaderResourceView(InstanceBufferViewInitializer);
+		FShaderResourceViewRHIRef InstanceBufferSRV = RHICmdList.CreateShaderResourceView(InstanceBufferViewInitializer.Buffer, FRHIViewDesc::CreateBufferSRV().SetTypeFromBuffer(InstanceBufferViewInitializer.Buffer));
 
 		SetShaderValue(BatchedParameters, ComputeShader->NumInstancesParam, NumInstances);
 		SetShaderValue(BatchedParameters, ComputeShader->NumHitGroupsParam, NumHitGroups);
diff --git a/Engine/Source/Runtime/RenderCore/Private/RenderTargetPool.cpp b/Engine/Source/Runtime/RenderCore/Private/RenderTargetPool.cpp
--- a/Engine/Source/Runtime/RenderCore/Private/RenderTargetPool.cpp
+++ b/Engine/Source/Runtime/RenderCore/Private/RenderTargetPool.cpp
@@ -96,12 +96,16 @@ FPooledRenderTarget* FRenderTargetPool::CreateRenderTarget(FRHICommandListBase&
 				? Desc.UAVFormat
 				: Desc.Format;
 
-			Result->RenderTargetItem.UAV = RHICmdList.CreateUnorderedAccessView(Result->GetRHI(), 0, (uint8)AliasFormat, 0, 0);
+			Result->RenderTargetItem.UAV = RHICmdList.CreateUnorderedAccessView(
+				Result->GetRHI(),
+				FRHIViewDesc::CreateTextureUAV()
+					.SetDimensionFromTexture(Result->GetRHI())
+					.SetFormat(AliasFormat));
 		}
 		else
 		{
 			checkf(Desc.UAVFormat == PF_Unknown || Desc.UAVFormat == Desc.Format, TEXT("UAV aliasing is not supported by the current RHI."));
-			Result->RenderTargetItem.UAV = RHICmdList.CreateUnorderedAccessView(Result->GetRHI(), 0);
+			Result->RenderTargetItem.UAV = RHICmdList.CreateUnorderedAccessView(Result->GetRHI(), FRHIViewDesc::CreateTextureUAV().SetDimensionFromTexture(Result->GetRHI()));
 		}
 	}
 
diff --git a/Engine/Source/Runtime/RenderCore/Private/UnifiedBuffer.cpp b/Engine/Source/Runtime/RenderCore/Private/UnifiedBuffer.cpp
--- a/Engine/Source/Runtime/RenderCore/Private/UnifiedBuffer.cpp
+++ b/Engine/Source/Runtime/RenderCore/Private/UnifiedBuffer.cpp
@@ -1488,7 +1488,7 @@ void FScatterUploadBuffer::Init( uint32 NumElements, uint32 InNumBytesPerElement
 
 			FRHIResourceCreateInfo CreateInfo(DebugName);
 			ScatterBuffer.Buffer = RHICmdList.CreateStructuredBuffer(sizeof(uint32), ScatterBuffer.NumBytes, BUF_ShaderResource | BUF_Volatile | Usage, CreateInfo);
-			ScatterBuffer.SRV = RHICmdList.CreateShaderResourceView(ScatterBuffer.Buffer);
+			ScatterBuffer.SRV = RHICmdList.CreateShaderResourceView(ScatterBuffer.Buffer, FRHIViewDesc::CreateBufferSRV().SetTypeFromBuffer(ScatterBuffer.Buffer));
 		}
 
 		if (UploadBytes > UploadBuffer.NumBytes || UploadBufferSize < UploadBuffer.NumBytes / 2)
@@ -1499,7 +1499,7 @@ void FScatterUploadBuffer::Init( uint32 NumElements, uint32 InNumBytesPerElement
 
 			FRHIResourceCreateInfo CreateInfo(DebugName);
 			UploadBuffer.Buffer = RHICmdList.CreateStructuredBuffer(TypeSize, UploadBuffer.NumBytes, BUF_ShaderResource | BUF_Volatile | Usage, CreateInfo);
-			UploadBuffer.SRV = RHICmdList.CreateShaderResourceView(UploadBuffer.Buffer);
+			UploadBuffer.SRV = RHICmdList.CreateShaderResourceView(UploadBuffer.Buffer, FRHIViewDesc::CreateBufferSRV().SetTypeFromBuffer(UploadBuffer.Buffer));
 		}
 
 		ScatterData = (uint32*)RHICmdList.LockBuffer(ScatterBuffer.Buffer, 0, ScatterBytes, RLM_WriteOnly);
@@ -1538,11 +1538,11 @@ void FScatterUploadBuffer::ResourceUploadTo(FRHICommandList& RHICmdList, const R
 
 		{
 			ScatterBuffer.Buffer = UE::RHIResourceUtils::CreateBufferFromArray(RHICmdList, TEXT("ScatterResourceArray"), Usage, sizeof(uint32), ScatterData, ScatterDataSize);
-			ScatterBuffer.SRV = RHICmdList.CreateShaderResourceView(ScatterBuffer.Buffer);
+			ScatterBuffer.SRV = RHICmdList.CreateShaderResourceView(ScatterBuffer.Buffer, FRHIViewDesc::CreateBufferSRV().SetTypeFromBuffer(ScatterBuffer.Buffer));
 		}
 		{
 			UploadBuffer.Buffer = UE::RHIResourceUtils::CreateBufferFromArray(RHICmdList, TEXT("ScatterUploadBuffer"), Usage, TypeSize, UploadData, UploadDataSize);
-			UploadBuffer.SRV = RHICmdList.CreateShaderResourceView(UploadBuffer.Buffer);
+			UploadBuffer.SRV = RHICmdList.CreateShaderResourceView(UploadBuffer.Buffer, FRHIViewDesc::CreateBufferSRV().SetTypeFromBuffer(UploadBuffer.Buffer));
 		}
 	}
 	else
diff --git a/Engine/Source/Runtime/RenderCore/Private/VertexFactory.cpp b/Engine/Source/Runtime/RenderCore/Private/VertexFactory.cpp
--- a/Engine/Source/Runtime/RenderCore/Private/VertexFactory.cpp
+++ b/Engine/Source/Runtime/RenderCore/Private/VertexFactory.cpp
@@ -409,7 +409,11 @@ void FPrimitiveIdDummyBuffer::InitRHI(FRHICommandListBase& RHICmdList)
 	uint32* Vertices = (uint32*)RHICmdList.LockBuffer(VertexBufferRHI, 0, sizeof(uint32), RLM_WriteOnly);
 	Vertices[0] = 0;
 	RHICmdList.UnlockBuffer(VertexBufferRHI);
-	VertexBufferSRV = RHICmdList.CreateShaderResourceView(VertexBufferRHI, sizeof(uint32), PF_R32_UINT);
+	VertexBufferSRV = RHICmdList.CreateShaderResourceView(
+		VertexBufferRHI,
+		FRHIViewDesc::CreateBufferSRV()
+			.SetType(FRHIViewDesc::EBufferType::Typed)
+			.SetFormat(PF_R32_UINT));
 }
 
 TGlobalResource<FPrimitiveIdDummyBuffer> GPrimitiveIdDummy;
diff --git a/Engine/Source/Runtime/Renderer/Private/MeshPassProcessor.cpp b/Engine/Source/Runtime/Renderer/Private/MeshPassProcessor.cpp
--- a/Engine/Source/Runtime/Renderer/Private/MeshPassProcessor.cpp
+++ b/Engine/Source/Runtime/Renderer/Private/MeshPassProcessor.cpp
@@ -1857,7 +1857,7 @@ void DrawDynamicMeshPassPrivate(
 
 		if (IsUniformBufferStaticSlotValid(SceneArgs.BatchedPrimitiveSlot))
 		{
-			FShaderResourceViewRHIRef BatchedPrimitiveSRV = RHICmdList.CreateShaderResourceView(SceneArgs.PrimitiveIdsBuffer);
+			FShaderResourceViewRHIRef BatchedPrimitiveSRV = RHICmdList.CreateShaderResourceView(SceneArgs.PrimitiveIdsBuffer, FRHIViewDesc::CreateBufferSRV().SetTypeFromBuffer(SceneArgs.PrimitiveIdsBuffer));
 			SceneArgs.PrimitiveIdsBuffer = nullptr;
 			FRHIShaderResourceView* SRV = BatchedPrimitiveSRV.GetReference();
 			FUniformBufferRHIRef UBRef = RHICreateUniformBuffer(&SRV, GDynamicBatchedPrimitiveLayout.LayoutRHI, UniformBuffer_SingleFrame, EUniformBufferValidation::None);
diff --git a/Engine/Source/Runtime/Renderer/Private/OIT/OIT.cpp b/Engine/Source/Runtime/Renderer/Private/OIT/OIT.cpp
--- a/Engine/Source/Runtime/Renderer/Private/OIT/OIT.cpp
+++ b/Engine/Source/Runtime/Renderer/Private/OIT/OIT.cpp
@@ -289,8 +289,16 @@ public:
 
 		FRHIResourceCreateInfo CreateInfo(DebugName);
 		IndexBufferRHI = RHICmdList.CreateBuffer(NumIndices * BytesPerElement, BUF_UnorderedAccess | BUF_ShaderResource | BUF_IndexBuffer, BytesPerElement /*Stride*/, ERHIAccess::VertexOrIndexBuffer, CreateInfo);
-		SortedIndexUAV = RHICmdList.CreateUnorderedAccessView(IndexBufferRHI, Format);
-		SourceIndexSRV = RHICmdList.CreateShaderResourceView(SourceIndexBuffer->IndexBufferRHI, BytesPerElement, Format);
+		SortedIndexUAV = RHICmdList.CreateUnorderedAccessView(
+			IndexBufferRHI,
+			FRHIViewDesc::CreateBufferUAV()
+				.SetType(FRHIViewDesc::EBufferType::Typed)
+				.SetFormat(Format));
+		SourceIndexSRV = RHICmdList.CreateShaderResourceView(
+			SourceIndexBuffer->IndexBufferRHI, 
+			FRHIViewDesc::CreateBufferSRV()
+				.SetType(FRHIViewDesc::EBufferType::Typed)
+				.SetFormat(Format));
 	}
 
 	virtual void ReleaseRHI() override
diff --git a/Engine/Source/Runtime/Renderer/Private/RayTracing/RayTracingScene.cpp b/Engine/Source/Runtime/Renderer/Private/RayTracing/RayTracingScene.cpp
--- a/Engine/Source/Runtime/Renderer/Private/RayTracing/RayTracingScene.cpp
+++ b/Engine/Source/Runtime/Renderer/Private/RayTracing/RayTracingScene.cpp
@@ -167,7 +167,7 @@ void FRayTracingScene::Create(FRDGBuilder& GraphBuilder, const FViewInfo& View,
 			{
 				FRHIResourceCreateInfo CreateInfo(TEXT("FRayTracingScene::InstanceUploadBuffer"));
 				Layer.InstanceUploadBuffer = RHICmdList.CreateStructuredBuffer(sizeof(FRayTracingInstanceDescriptorInput), UploadBufferSize, BUF_ShaderResource | BUF_Volatile, CreateInfo);
-				Layer.InstanceUploadSRV = RHICmdList.CreateShaderResourceView(Layer.InstanceUploadBuffer);
+				Layer.InstanceUploadSRV = RHICmdList.CreateShaderResourceView(Layer.InstanceUploadBuffer, FRHIViewDesc::CreateBufferSRV().SetTypeFromBuffer(Layer.InstanceUploadBuffer));
 			}
 		}
 
@@ -181,7 +181,7 @@ void FRayTracingScene::Create(FRDGBuilder& GraphBuilder, const FViewInfo& View,
 			{
 				FRHIResourceCreateInfo CreateInfo(TEXT("FRayTracingScene::TransformUploadBuffer"));
 				Layer.TransformUploadBuffer = RHICmdList.CreateStructuredBuffer(sizeof(FVector4f), UploadBufferSize, BUF_ShaderResource | BUF_Volatile, CreateInfo);
-				Layer.TransformUploadSRV = RHICmdList.CreateShaderResourceView(Layer.TransformUploadBuffer);
+				Layer.TransformUploadSRV = RHICmdList.CreateShaderResourceView(Layer.TransformUploadBuffer, FRHIViewDesc::CreateBufferSRV().SetTypeFromBuffer(Layer.TransformUploadBuffer));
 			}
 		}
 
diff --git a/Engine/Source/Runtime/Renderer/Private/SystemTextures.cpp b/Engine/Source/Runtime/Renderer/Private/SystemTextures.cpp
--- a/Engine/Source/Runtime/Renderer/Private/SystemTextures.cpp
+++ b/Engine/Source/Runtime/Renderer/Private/SystemTextures.cpp
@@ -249,7 +249,7 @@ void FSystemTextures::InitializeCommonTextures(FRHICommandListImmediate& RHICmdL
 	// Create a WhiteDummy texture
 	{
 		WhiteDummy = CreateRenderTarget(GWhiteTexture->TextureRHI, TEXT("WhiteDummy"));
-		WhiteDummySRV = RHICmdList.CreateShaderResourceView((FRHITexture*)WhiteDummy->GetRHI(), 0);
+		WhiteDummySRV = RHICmdList.CreateShaderResourceView((FRHITexture*)WhiteDummy->GetRHI(),FRHIViewDesc::CreateTextureSRV().SetDimensionFromTexture(WhiteDummy->GetRHI()));
 	}
 
 	// Create a BlackDummy texture
@@ -404,7 +404,7 @@ void FSystemTextures::InitializeCommonTextures(FRHICommandListImmediate& RHICmdL
 		FTextureRHIRef Texture = RHICreateTexture(Desc);
 		SetDummyTextureData<FColor>(Texture, FColor::White);
 		StencilDummy = CreateRenderTarget(Texture, Desc.DebugName);
-		StencilDummySRV = RHICmdList.CreateShaderResourceView(StencilDummy->GetRHI(), 0);
+		StencilDummySRV = RHICmdList.CreateShaderResourceView(StencilDummy->GetRHI(), FRHIViewDesc::CreateTextureSRV().SetDimensionFromTexture(Texture));
 	}
 
 	if (GPixelFormats[PF_FloatRGBA].Supported)
diff --git a/Engine/Source/Runtime/Renderer/Private/Tests/RayTracingTestbed.cpp b/Engine/Source/Runtime/Renderer/Private/Tests/RayTracingTestbed.cpp
--- a/Engine/Source/Runtime/Renderer/Private/Tests/RayTracingTestbed.cpp
+++ b/Engine/Source/Runtime/Renderer/Private/Tests/RayTracingTestbed.cpp
@@ -218,7 +218,7 @@ bool RunRayTracingTestbed_RenderThread(const FString& Parameters)
 	{
 		FRHIResourceCreateInfo CreateInfo(TEXT("RayTracingTestBedInstanceUploadBuffer"));
 		InstanceUploadBuffer = RHICmdList.CreateStructuredBuffer(sizeof(FRayTracingInstanceDescriptorInput), InstanceUploadBufferSize, BUF_ShaderResource | BUF_Volatile, CreateInfo);
-		InstanceUploadSRV = RHICmdList.CreateShaderResourceView(InstanceUploadBuffer);
+		InstanceUploadSRV = RHICmdList.CreateShaderResourceView(InstanceUploadBuffer, FRHIViewDesc::CreateBufferSRV().SetTypeFromBuffer(InstanceUploadBuffer));
 	}
 
 	const uint32 TransformUploadBufferSize = RayTracingSceneInitializationData.NumNativeCPUInstances * 3 * sizeof(FVector4f);
@@ -227,7 +227,7 @@ bool RunRayTracingTestbed_RenderThread(const FString& Parameters)
 	{
 		FRHIResourceCreateInfo CreateInfo(TEXT("RayTracingTestBedInstanceUploadBuffer"));
 		TransformUploadBuffer = RHICmdList.CreateStructuredBuffer(sizeof(FVector4f), TransformUploadBufferSize, BUF_ShaderResource | BUF_Volatile, CreateInfo);
-		TransformUploadSRV = RHICmdList.CreateShaderResourceView(TransformUploadBuffer);
+		TransformUploadSRV = RHICmdList.CreateShaderResourceView(TransformUploadBuffer, FRHIViewDesc::CreateBufferSRV().SetTypeFromBuffer(TransformUploadBuffer));
 	}
 
 	{
diff --git a/Engine/Source/Runtime/Renderer/Private/VT/VirtualTexturePhysicalSpace.cpp b/Engine/Source/Runtime/Renderer/Private/VT/VirtualTexturePhysicalSpace.cpp
--- a/Engine/Source/Runtime/Renderer/Private/VT/VirtualTexturePhysicalSpace.cpp
+++ b/Engine/Source/Runtime/Renderer/Private/VT/VirtualTexturePhysicalSpace.cpp
@@ -178,10 +178,20 @@ void FVirtualTexturePhysicalSpace::InitRHI(FRHICommandListBase& RHICmdList)
 		FRHITextureSRVCreateInfo SRVCreateInfo;
 		SRVCreateInfo.Format = FormatSRV;
 		SRVCreateInfo.SRGBOverride = SRGBO_ForceDisable;
-		TextureSRV[Layer] = RHICmdList.CreateShaderResourceView(TextureRHI, SRVCreateInfo);
+		TextureSRV[Layer] = RHICmdList.CreateShaderResourceView(
+			TextureRHI,
+			FRHIViewDesc::CreateTextureSRV()
+				.SetDimensionFromTexture(TextureRHI)
+				.SetFormat(SRVCreateInfo.Format)
+				.SetDisableSRGB(SRVCreateInfo.SRGBOverride == SRGBO_ForceDisable));
 
 		SRVCreateInfo.SRGBOverride = SRGBO_Default;
-		TextureSRV_SRGB[Layer] = RHICmdList.CreateShaderResourceView(TextureRHI, SRVCreateInfo);
+		TextureSRV_SRGB[Layer] = RHICmdList.CreateShaderResourceView(
+			TextureRHI,
+			FRHIViewDesc::CreateTextureSRV()
+				.SetDimensionFromTexture(TextureRHI)
+				.SetFormat(SRVCreateInfo.Format)
+				.SetDisableSRGB(SRVCreateInfo.SRGBOverride == SRGBO_ForceDisable));
 	}
 
 	INC_MEMORY_STAT_BY(STAT_TotalPhysicalMemory, GetSizeInBytes());
diff --git a/Engine/Source/Runtime/Renderer/Private/VT/VirtualTextureSpace.cpp b/Engine/Source/Runtime/Renderer/Private/VT/VirtualTextureSpace.cpp
--- a/Engine/Source/Runtime/Renderer/Private/VT/VirtualTextureSpace.cpp
+++ b/Engine/Source/Runtime/Renderer/Private/VT/VirtualTextureSpace.cpp
@@ -446,7 +446,11 @@ void FVirtualTextureSpace::ApplyUpdates(FVirtualTextureSystem* System, FRDGBuild
 
 		FRHIResourceCreateInfo CreateInfo(TEXT("VirtualTexture_PageTableUpdateBuffer"));
 		UpdateBuffer = RHICmdList.CreateVertexBuffer(NewBufferSize, BUF_ShaderResource | BUF_Volatile, CreateInfo);
-		UpdateBufferSRV = RHICmdList.CreateShaderResourceView(UpdateBuffer, sizeof(FPageTableUpdate), PF_R16G16B16A16_UINT);
+		UpdateBufferSRV = RHICmdList.CreateShaderResourceView(
+			UpdateBuffer, 
+			FRHIViewDesc::CreateBufferSRV()
+				.SetType(FRHIViewDesc::EBufferType::Typed)
+				.SetFormat(PF_R16G16B16A16_UINT));
 	}
 
 	// This flushes the RHI thread!
diff --git a/Engine/Source/Runtime/VulkanRHI/Private/VulkanTexture.cpp b/Engine/Source/Runtime/VulkanRHI/Private/VulkanTexture.cpp
--- a/Engine/Source/Runtime/VulkanRHI/Private/VulkanTexture.cpp
+++ b/Engine/Source/Runtime/VulkanRHI/Private/VulkanTexture.cpp
@@ -919,7 +919,11 @@ FTextureReferenceRHIRef FVulkanDynamicRHI::RHICreateTextureReference(FRHICommand
 	// If the referenced texture is configured for bindless, make sure we also create an SRV to use for bindless.
 	if (ReferencedTexture && ReferencedTexture->GetDefaultBindlessHandle().IsValid())
 	{
-		FShaderResourceViewRHIRef BindlessView = RHICmdList.CreateShaderResourceView(ReferencedTexture, 0u);
+		FShaderResourceViewRHIRef BindlessView = RHICmdList.CreateShaderResourceView(
+			ReferencedTexture,
+			FRHIViewDesc::CreateTextureSRV()
+				.SetDimensionFromTexture(ReferencedTexture)
+				.SetMipRange(0, 1));
 		return new FVulkanTextureReference(ReferencedTexture, ResourceCast(BindlessView.GetReference()));
 	}
 #endif
