- Make FNiagaraEmitterID constructor explicit
- Add comparison operators

#rnx

diff --git a/Engine/Plugins/FX/Niagara/Source/Niagara/Public/NiagaraCompilationTypes.h b/Engine/Plugins/FX/Niagara/Source/Niagara/Public/NiagaraCompilationTypes.h
--- a/Engine/Plugins/FX/Niagara/Source/Niagara/Public/NiagaraCompilationTypes.h
+++ b/Engine/Plugins/FX/Niagara/Source/Niagara/Public/NiagaraCompilationTypes.h
@@ -51,7 +51,7 @@ struct FNiagaraScriptAsyncCompileData
 	TSharedPtr<struct FNiagaraVMExecutableData> ExeData;
 	TArray<FNiagaraCompiledShaderInfo> CompiledShaders;
 	FString UniqueEmitterName;
-	FNiagaraEmitterID EmitterID = INDEX_NONE;
+	FNiagaraEmitterID EmitterID;
 	FNiagaraScriptCompileMetrics CompileMetrics;
 
 	bool bFromDerivedDataCache = false;
diff --git a/Engine/Plugins/FX/Niagara/Source/Niagara/Public/NiagaraTypes.h b/Engine/Plugins/FX/Niagara/Source/Niagara/Public/NiagaraTypes.h
--- a/Engine/Plugins/FX/Niagara/Source/Niagara/Public/NiagaraTypes.h
+++ b/Engine/Plugins/FX/Niagara/Source/Niagara/Public/NiagaraTypes.h
@@ -192,10 +192,13 @@ struct FNiagaraEmitterID
 	GENERATED_BODY()
 	
 	FNiagaraEmitterID() = default;
-	FNiagaraEmitterID(int32 InID):ID(InID) {}
+	explicit FNiagaraEmitterID(int32 InID) : ID(InID) {}
 
 	UPROPERTY(EditAnywhere, Category = Emitter)
 	int32 ID = INDEX_NONE;
+
+	bool operator==(const FNiagaraEmitterID OtherID) const { return ID == OtherID.ID; }
+	bool operator!=(const FNiagaraEmitterID OtherID) const { return ID != OtherID.ID; }
 };
 
 USTRUCT()
diff --git a/Engine/Plugins/FX/Niagara/Source/NiagaraEditor/Private/NiagaraNodeEmitter.h b/Engine/Plugins/FX/Niagara/Source/NiagaraEditor/Private/NiagaraNodeEmitter.h
--- a/Engine/Plugins/FX/Niagara/Source/NiagaraEditor/Private/NiagaraNodeEmitter.h
+++ b/Engine/Plugins/FX/Niagara/Source/NiagaraEditor/Private/NiagaraNodeEmitter.h
@@ -93,7 +93,7 @@ private:
 	ENiagaraScriptUsage ScriptType;
 
 	FName CachedUniqueName;
-	FNiagaraEmitterID CachedEmitterID = INDEX_NONE;
+	FNiagaraEmitterID CachedEmitterID;
 	TWeakObjectPtr<UNiagaraGraph> CachedGraphWeakPtr;
 	TWeakObjectPtr<UNiagaraScriptSourceBase> CachedScriptSourceWeakPtr;
 };
\ No newline at end of file
