- Capture owning instance with array types & fix some compilation issues when GPU is not enabled

#rnx

diff --git a/Engine/Plugins/FX/Niagara/Source/Niagara/Classes/NiagaraDataInterfaceArrayImpl.h b/Engine/Plugins/FX/Niagara/Source/Niagara/Classes/NiagaraDataInterfaceArrayImpl.h
--- a/Engine/Plugins/FX/Niagara/Source/Niagara/Classes/NiagaraDataInterfaceArrayImpl.h
+++ b/Engine/Plugins/FX/Niagara/Source/Niagara/Classes/NiagaraDataInterfaceArrayImpl.h
@@ -204,10 +204,11 @@ struct FNiagaraDataInterfaceArrayImplHelper
 template<typename TArrayType>
 struct FNDIArrayInstanceData_GameThread
 {
-	bool				bIsModified = false;		// True if the array has ever been modified and we are reading instance data
-	bool				bIsRenderDirty = true;		// True if we have made modifications that could be pushed to the render thread
-	FRWLock				ArrayRWGuard;
-	TArray<TArrayType>	ArrayData;					// Modified array data
+	FNiagaraSystemInstance*	OwnerInstance = nullptr;
+	bool					bIsModified = false;		// True if the array has ever been modified and we are reading instance data
+	bool					bIsRenderDirty = true;		// True if we have made modifications that could be pushed to the render thread
+	FRWLock					ArrayRWGuard;
+	TArray<TArrayType>		ArrayData;					// Modified array data
 };
 
 template<typename TArrayType>
@@ -311,7 +312,7 @@ struct FNDIArrayInstanceData_RenderThread
 	}
 
 	template<typename T = FNDIArrayImplHelper<TArrayType>>
-	typename TEnableIf<!T::bSupportsGPU>::Type UpdateDataImpl(FRHICommandList& RHICmdList, TArray<TArrayType>& InArrayData)
+	typename TEnableIf<!T::bSupportsGPU>::Type UpdateDataImpl(FRHICommandList& RHICmdList, TConstArrayView<TArrayType> InArrayData)
 	{
 	}
 
@@ -1139,9 +1140,9 @@ struct FNDIArrayProxyImpl : public INDIArrayProxyBase
 		}
 
 		// Read GPU Data
-		if (FNDIArrayImplHelper<TArrayType>::bSupportsGPU && Owner->IsUsedWithGPUScript())
+		if constexpr (FNDIArrayImplHelper<TArrayType>::bSupportsGPU)
 		{
-			if (CacheData->GpuFrameData.IsValidIndex(FrameIndex))
+			if (Owner->IsUsedWithGPUScript() && CacheData->GpuFrameData.IsValidIndex(FrameIndex))
 			{
 				FNiagaraGpuComputeDispatchInterface* ComputeInterface = FNiagaraGpuComputeDispatchInterface::Get(SystemInstance->GetWorld());
 				ENQUEUE_RENDER_COMMAND(NDIArray_SimCacheWrite)(
@@ -1203,6 +1204,7 @@ struct FNDIArrayProxyImpl : public INDIArrayProxyBase
 		CachePropertiesFromOwner();
 
 		auto* InstanceData_GT = new(InPerInstanceData) FNDIArrayInstanceData_GameThread<TArrayType>();
+		InstanceData_GT->OwnerInstance = SystemInstance;
 		InstanceData_GT->bIsRenderDirty = true;
 
 		PerInstanceData_GameThread.Emplace(SystemInstance->GetId(), InstanceData_GT);
