- Move stateless ensure into verbose warning

#rnx

diff --git a/Engine/Plugins/FX/Niagara/Source/Niagara/Private/Stateless/NiagaraStatelessComputeManager.cpp b/Engine/Plugins/FX/Niagara/Source/Niagara/Private/Stateless/NiagaraStatelessComputeManager.cpp
--- a/Engine/Plugins/FX/Niagara/Source/Niagara/Private/Stateless/NiagaraStatelessComputeManager.cpp
+++ b/Engine/Plugins/FX/Niagara/Source/Niagara/Private/Stateless/NiagaraStatelessComputeManager.cpp
@@ -322,7 +322,7 @@ FNiagaraDataBuffer* FNiagaraStatelessComputeManager::GetDataBuffer(FRHICommandLi
 		{
 			FNiagaraGPUInstanceCountManager& CountManager = ComputeInterface->GetGPUInstanceCounterManager();
 			const uint32 CountOffset = CountManager.AcquireEntry();
-			if (ensure(CountOffset != INDEX_NONE))
+			if (CountOffset != INDEX_NONE)
 			{
 				CacheData->DataBuffer->SetNumInstances(ActiveParticles);
 				CacheData->DataBuffer->SetGPUInstanceCountBufferOffset(CountOffset);
@@ -334,6 +334,11 @@ FNiagaraDataBuffer* FNiagaraStatelessComputeManager::GetDataBuffer(FRHICommandLi
 			// If it does occur this means some other system has used a count slot but not reserved one
 			else
 			{
+				if (FNiagaraUtilities::LogVerboseWarnings())
+				{
+					ensureMsgf(false, TEXT("Count reserved for stateless was not available, attemping to generate on the CPU."));
+				}
+
 				const bool bAllowCPUExec = EnumHasAnyFlags(EmitterInstance->EmitterData->FeatureMask, ENiagaraStatelessFeatureMask::ExecuteCPU);
 				const bool bDidGenerate = bAllowCPUExec && GenerateCPUDataForGPUSim(RHICmdList, EmitterInstance, CacheData->DataBuffer);
 				if (!bDidGenerate)
