- Handle missing types for data interface user parameters

#rnx
#test Add a DI type (user) remove the DI type and observe the crash

diff --git a/Engine/Plugins/FX/Niagara/Source/NiagaraEditor/Private/ViewModels/NiagaraScriptInputCollectionViewModel.cpp b/Engine/Plugins/FX/Niagara/Source/NiagaraEditor/Private/ViewModels/NiagaraScriptInputCollectionViewModel.cpp
--- a/Engine/Plugins/FX/Niagara/Source/NiagaraEditor/Private/ViewModels/NiagaraScriptInputCollectionViewModel.cpp
+++ b/Engine/Plugins/FX/Niagara/Source/NiagaraEditor/Private/ViewModels/NiagaraScriptInputCollectionViewModel.cpp
@@ -304,11 +304,14 @@ void FNiagaraScriptInputCollectionViewModel::RefreshParameterViewModels()
 				ParameterViewModel = MakeShareable(new FNiagaraScriptParameterViewModel(GraphVariable, *InputNode, ObjectAsset, ParameterEditMode));
 			}
 
-			ParameterViewModel->OnNameChanged().AddSP(this, &FNiagaraScriptInputCollectionViewModel::OnParameterNameChanged, TWeakObjectPtr<UNiagaraNodeInput>(InputNode));
-			ParameterViewModel->OnTypeChanged().AddSP(this, &FNiagaraScriptInputCollectionViewModel::OnParameterTypeChanged, &GraphVariable);
-			ParameterViewModel->OnDefaultValueChanged().AddSP(this, &FNiagaraScriptInputCollectionViewModel::OnParameterValueChangedInternal, ParameterViewModel.ToSharedRef());
-			ParameterViewModels.Add(ParameterViewModel.ToSharedRef());
-			AddedInputNames.Add(InputNode->Input.GetName());
+			if (ParameterViewModel)
+			{
+				ParameterViewModel->OnNameChanged().AddSP(this, &FNiagaraScriptInputCollectionViewModel::OnParameterNameChanged, TWeakObjectPtr<UNiagaraNodeInput>(InputNode));
+				ParameterViewModel->OnTypeChanged().AddSP(this, &FNiagaraScriptInputCollectionViewModel::OnParameterTypeChanged, &GraphVariable);
+				ParameterViewModel->OnDefaultValueChanged().AddSP(this, &FNiagaraScriptInputCollectionViewModel::OnParameterValueChangedInternal, ParameterViewModel.ToSharedRef());
+				ParameterViewModels.Add(ParameterViewModel.ToSharedRef());
+				AddedInputNames.Add(InputNode->Input.GetName());
+			}
 		}
 	}
 
