- Mesh Array Data Interface
- Can be bound to a mesh renderer so that the array of meshes can be modified at runtime
- Can not be manipulated in the VVM currently as Object handling through the graph is not functional

diff --git a/Engine/Plugins/FX/Niagara/Source/Niagara/Classes/NiagaraDataInterfaceArray.h b/Engine/Plugins/FX/Niagara/Source/Niagara/Classes/NiagaraDataInterfaceArray.h
--- a/Engine/Plugins/FX/Niagara/Source/Niagara/Classes/NiagaraDataInterfaceArray.h
+++ b/Engine/Plugins/FX/Niagara/Source/Niagara/Classes/NiagaraDataInterfaceArray.h
@@ -24,6 +24,7 @@ struct FNDIArraySimCacheDataFrame;
 	void SetVariantArrayData(TConstArrayView<TFromArrayType> InArrayData); \
 	template<typename TFromArrayType> \
 	void SetVariantArrayValue(int Index, const TFromArrayType& Value, bool bSizeToFit); \
+	const TArray<TYPENAME>& GetArrayReference() const { return MEMBERNAME; } \
 	TArray<TYPENAME>& GetArrayReference() { return MEMBERNAME; }
 
 #if WITH_EDITORONLY_DATA
@@ -38,6 +39,7 @@ struct FNDIArraySimCacheDataFrame;
 		void SetVariantArrayData(TConstArrayView<TFromArrayType> InArrayData); \
 		template<typename TFromArrayType> \
 		void SetVariantArrayValue(int Index, const TFromArrayType& Value, bool bSizeToFit); \
+		const TArray<TYPENAME>& GetArrayReference() const { return Internal##MEMBERNAME; } \
 		TArray<TYPENAME>& GetArrayReference() { return Internal##MEMBERNAME; }
 #else
 	#define NDIARRAY_GENERATE_BODY_LWC(CLASSNAME, TYPENAME, MEMBERNAME) NDIARRAY_GENERATE_BODY(CLASSNAME, TYPENAME, Internal##MEMBERNAME)
diff --git a/Engine/Plugins/FX/Niagara/Source/Niagara/Classes/NiagaraDataInterfaceArrayImpl.h b/Engine/Plugins/FX/Niagara/Source/Niagara/Classes/NiagaraDataInterfaceArrayImpl.h
--- a/Engine/Plugins/FX/Niagara/Source/Niagara/Classes/NiagaraDataInterfaceArrayImpl.h
+++ b/Engine/Plugins/FX/Niagara/Source/Niagara/Classes/NiagaraDataInterfaceArrayImpl.h
@@ -207,7 +207,7 @@ struct FNDIArrayInstanceData_GameThread
 	FNiagaraSystemInstance*	OwnerInstance = nullptr;
 	bool					bIsModified = false;		// True if the array has ever been modified and we are reading instance data
 	bool					bIsRenderDirty = true;		// True if we have made modifications that could be pushed to the render thread
-	FRWLock					ArrayRWGuard;
+	mutable FRWLock			ArrayRWGuard;
 	TArray<TArrayType>		ArrayData;					// Modified array data
 };
 
@@ -349,7 +349,7 @@ struct FNDIArrayProxyImpl : public INDIArrayProxyBase
 
 	struct FReadArrayRef
 	{
-		FReadArrayRef(TOwnerType* Owner, FNDIArrayInstanceData_GameThread<TArrayType>* InstanceData)
+		FReadArrayRef(const TOwnerType* Owner, const FNDIArrayInstanceData_GameThread<TArrayType>* InstanceData)
 		{
 			if ( InstanceData )
 			{
@@ -362,6 +362,7 @@ struct FNDIArrayProxyImpl : public INDIArrayProxyBase
 				ArrayData = &Owner->GetArrayReference();
 			}
 		}
+
 		~FReadArrayRef()
 		{
 			if ( LockObject )
@@ -601,6 +602,23 @@ struct FNDIArrayProxyImpl : public INDIArrayProxyBase
 		}
 	}
 
+	template<typename TFromArrayType>
+	void SetArrayDataAndRecreateRenderState(TConstArrayView<TFromArrayType> InArrayData)
+	{
+		SetArrayData(InArrayData);
+
+		for (auto It=PerInstanceData_GameThread.CreateConstIterator(); It; ++It)
+		{
+			FNDIArrayInstanceData_GameThread<TArrayType>* InstanceData = It.Value();
+			FNiagaraSystemInstance* SystemInstance = InstanceData->OwnerInstance;
+			USceneComponent* SceneComponent = SystemInstance ? SystemInstance->GetAttachComponent() : nullptr;
+			if (SceneComponent)
+			{
+				SceneComponent->RecreateRenderState_Concurrent();
+			}
+		}
+	}
+
 	template<typename TToArrayType>
 	TArray<TToArrayType> GetArrayDataCopy()
 	{
@@ -681,6 +699,7 @@ struct FNDIArrayProxyImpl : public INDIArrayProxyBase
 		{
 			FNiagaraFunctionSignature& Sig = OutFunctions.Add_GetRef(DefaultImmutableSig);
 			Sig.Name = FNiagaraDataInterfaceArrayImplHelper::Function_LengthName;
+			Sig.bSupportsCPU = true;
 			Sig.Outputs.Emplace(FNiagaraTypeDefinition::GetIntDef(), TEXT("Num"));
 			Sig.Description = NSLOCTEXT("Niagara", "Array_LengthDesc", "Gets the number of elements in the array.");
 		}
@@ -688,6 +707,7 @@ struct FNDIArrayProxyImpl : public INDIArrayProxyBase
 		{
 			FNiagaraFunctionSignature& Sig = OutFunctions.Add_GetRef(DefaultImmutableSig);
 			Sig.Name = FNiagaraDataInterfaceArrayImplHelper::Function_IsValidIndexName;
+			Sig.bSupportsCPU = true;
 			Sig.Inputs.Emplace(FNiagaraTypeDefinition::GetIntDef(), TEXT("Index"));
 			Sig.Outputs.Emplace(FNiagaraTypeDefinition::GetBoolDef(), TEXT("Valid"));
 			Sig.Description = NSLOCTEXT("Niagara", "Array_IsValidIndexDesc", "Tests to see if the index is valid and exists in the array.");
@@ -696,6 +716,7 @@ struct FNDIArrayProxyImpl : public INDIArrayProxyBase
 		{
 			FNiagaraFunctionSignature& Sig = OutFunctions.Add_GetRef(DefaultImmutableSig);
 			Sig.Name = FNiagaraDataInterfaceArrayImplHelper::Function_LastIndexName;
+			Sig.bSupportsCPU = true;
 			Sig.Outputs.Emplace(FNiagaraTypeDefinition::GetIntDef(), TEXT("Index"));
 			Sig.Description = NSLOCTEXT("Niagara", "Array_LastIndexDesc", "Returns the last valid index in the array, will be -1 if no elements.");
 		}
@@ -802,10 +823,8 @@ struct FNDIArrayProxyImpl : public INDIArrayProxyBase
 	}
 #endif
 
-	template<typename T = FNDIArrayImplHelper<TArrayType>>
-	typename TEnableIf<T::bSupportsCPU>::Type GetVMExternalFunction_Internal(const FVMExternalFunctionBindingInfo& BindingInfo, void* InstanceData, FVMExternalFunction& OutFunc)
+	void GetVMExternalFunction_Internal(const FVMExternalFunctionBindingInfo& BindingInfo, void* InstanceData, FVMExternalFunction& OutFunc)
 	{
-		// Immutable functions
 		if (BindingInfo.Name == FNiagaraDataInterfaceArrayImplHelper::Function_LengthName)
 		{
 			check(BindingInfo.GetNumInputs() == 1 && BindingInfo.GetNumOutputs() == 1);
@@ -821,7 +840,13 @@ struct FNDIArrayProxyImpl : public INDIArrayProxyBase
 			check(BindingInfo.GetNumInputs() == 1 && BindingInfo.GetNumOutputs() == 1);
 			OutFunc = FVMExternalFunction::CreateLambda([this](FVectorVMExternalFunctionContext& Context) { this->VMGetLastIndex(Context); });
 		}
-		else if (BindingInfo.Name == FNiagaraDataInterfaceArrayImplHelper::Function_GetName)
+	}
+
+	template<typename T = FNDIArrayImplHelper<TArrayType>>
+	typename TEnableIf<T::bSupportsCPU>::Type GetVMExternalFunction_CPUAccessInternal(const FVMExternalFunctionBindingInfo& BindingInfo, void* InstanceData, FVMExternalFunction& OutFunc)
+	{
+		// Immutable functions
+		if (BindingInfo.Name == FNiagaraDataInterfaceArrayImplHelper::Function_GetName)
 		{
 			// Note: Outputs is variable based upon type
 			//check(BindingInfo.GetNumInputs() == 2 && BindingInfo.GetNumOutputs() == 1);
@@ -855,7 +880,7 @@ struct FNDIArrayProxyImpl : public INDIArrayProxyBase
 	}
 
 	template<typename T = FNDIArrayImplHelper<TArrayType>>
-	typename TEnableIf<!T::bSupportsCPU>::Type GetVMExternalFunction_Internal(const FVMExternalFunctionBindingInfo& BindingInfo, void* InstanceData, FVMExternalFunction& OutFunc)
+	typename TEnableIf<!T::bSupportsCPU>::Type GetVMExternalFunction_CPUAccessInternal(const FVMExternalFunctionBindingInfo& BindingInfo, void* InstanceData, FVMExternalFunction& OutFunc)
 	{
 	}
 
@@ -885,6 +910,10 @@ struct FNDIArrayProxyImpl : public INDIArrayProxyBase
 	{
 		GetVMExternalFunction_Internal(BindingInfo, InstanceData, OutFunc);
 		if (OutFunc.IsBound() == false)
+		{
+			GetVMExternalFunction_CPUAccessInternal(BindingInfo, InstanceData, OutFunc);
+		}
+		if (OutFunc.IsBound() == false)
 		{
 			GetVMExternalFunction_AtomicInternal(BindingInfo, InstanceData, OutFunc);
 		}
@@ -1556,6 +1585,11 @@ struct FNDIArrayProxyImpl : public INDIArrayProxyBase
 		InstanceData->bIsRenderDirty |= bShouldSyncToGpu;
 	}
 
+	const FNDIArrayInstanceData_GameThread<TArrayType>* GetPerInstanceData_GameThread(FNiagaraSystemInstanceID SystemInstanceID) const
+	{
+		return PerInstanceData_GameThread.FindRef(SystemInstanceID);
+	}
+
 private:
 	TOwnerType*	Owner = nullptr;
 	bool bShouldSyncToGpu = false;
diff --git a/Engine/Plugins/FX/Niagara/Source/Niagara/Internal/DataInterface/NiagaraDataInterfaceArrayMesh.h b/Engine/Plugins/FX/Niagara/Source/Niagara/Internal/DataInterface/NiagaraDataInterfaceArrayMesh.h
new file mode 100644
--- /dev/null
+++ b/Engine/Plugins/FX/Niagara/Source/Niagara/Internal/DataInterface/NiagaraDataInterfaceArrayMesh.h
@@ -0,0 +1,32 @@
+// Copyright Epic Games, Inc. All Rights Reserved.
+
+#pragma once
+
+#include "NiagaraDataInterfaceArray.h"
+#include "NiagaraDataInterfaceArrayImpl.h"
+#include "NiagaraRenderableMeshArrayInterface.h"
+#include "NiagaraDataInterfaceArrayMesh.generated.h"
+
+class FNiagaraSystemInstance;
+
+UCLASS(EditInlineNew, Category = "Array", meta = (DisplayName = "Mesh Array", Experimental), Blueprintable, BlueprintType, MinimalAPI)
+class UNiagaraDataInterfaceArrayMesh : public UNiagaraDataInterfaceArray, public INiagaraRenderableMeshArrayInterface
+{
+public:
+	GENERATED_BODY()
+
+	UPROPERTY(EditAnywhere, Category = "Array")
+	TArray<FNiagaraMeshRendererMeshPropertiesBase> MeshData;
+
+	NDIARRAY_GENERATE_BODY(UNiagaraDataInterfaceArrayMesh, FNiagaraMeshRendererMeshPropertiesBase, MeshData)
+
+	//INiagaraRenderableMeshArrayInterface Interface Begin
+	virtual void ForEachMesh(FNiagaraSystemInstance* SystemInstance, TFunction<void(int32)> NumMeshesDelegate, TFunction<void(const FNiagaraMeshRendererMeshProperties&)> IterateDelegate) const override;
+	//INiagaraRenderableMeshArrayInterface Interface End
+
+	UFUNCTION(BlueprintCallable, Category = Niagara, meta = (DisplayName = "Niagara Set Mesh Array", UnsafeDuringActorConstruction = "true"))
+	static NIAGARA_API void SetNiagaraArrayMesh(UNiagaraComponent* NiagaraComponent, FName OverrideName, const TArray<FNiagaraMeshRendererMeshPropertiesBase>& ArrayData);
+
+	UFUNCTION(BlueprintCallable, Category = Niagara, meta = (DisplayName = "Niagara Set Mesh Array (Static Mesh)", UnsafeDuringActorConstruction = "true"))
+	static NIAGARA_API void SetNiagaraArrayMeshSM(UNiagaraComponent* NiagaraComponent, FName OverrideName, const TArray<UStaticMesh*>& ArrayData);
+};
diff --git a/Engine/Plugins/FX/Niagara/Source/Niagara/Private/DataInterface/NiagaraDataInterfaceArrayMesh.cpp b/Engine/Plugins/FX/Niagara/Source/Niagara/Private/DataInterface/NiagaraDataInterfaceArrayMesh.cpp
new file mode 100644
--- /dev/null
+++ b/Engine/Plugins/FX/Niagara/Source/Niagara/Private/DataInterface/NiagaraDataInterfaceArrayMesh.cpp
@@ -0,0 +1,91 @@
+// Copyright Epic Games, Inc. All Rights Reserved.
+
+#include "DataInterface/NiagaraDataInterfaceArrayMesh.h"
+#include "NiagaraEmitterInstanceImpl.h"
+#include "NiagaraEmitterInstance.h"
+#include "NiagaraFunctionLibrary.h"
+#include "NiagaraSystemInstance.h"
+
+#include "Engine/StaticMesh.h"
+
+#include UE_INLINE_GENERATED_CPP_BY_NAME(NiagaraDataInterfaceArrayMesh)
+
+template<>
+struct FNDIArrayImplHelper<FNiagaraMeshRendererMeshPropertiesBase> : public FNDIArrayImplHelperBase<FNiagaraMeshRendererMeshPropertiesBase>
+{
+	typedef FNiagaraMeshRendererMeshPropertiesBase TVMArrayType;
+
+	static constexpr bool bSupportsCPU = false;
+	static constexpr bool bSupportsGPU = false;
+
+	static const FNiagaraTypeDefinition GetTypeDefinition() { return FNiagaraTypeDefinition(FNiagaraMeshRendererMeshPropertiesBase::StaticStruct()); }
+	static const FNiagaraMeshRendererMeshPropertiesBase GetDefaultValue() { return FNiagaraMeshRendererMeshPropertiesBase(); }
+
+	static void CopyCpuToCpuMemory(FNiagaraMeshRendererMeshPropertiesBase* Dest, const FNiagaraMeshRendererMeshPropertiesBase* Src, int32 NumElements)
+	{
+		FMemory::Memcpy(Dest, Src, NumElements * sizeof(FNiagaraMeshRendererMeshPropertiesBase));
+	}
+
+	static void CopyCpuToCpuMemory(FNiagaraMeshRendererMeshPropertiesBase* Dest, UStaticMesh*const* Src, int32 NumElements)
+	{
+		for (int32 i = 0; i < NumElements; ++i)
+		{
+			Dest[i] = FNiagaraMeshRendererMeshPropertiesBase();
+			Dest[i].Mesh = Src[i];
+		}
+	}
+
+	static void AppendValueToString(const FNiagaraMeshRendererMeshPropertiesBase& Value, FString& OutString)
+	{
+		GetFNameSafe(Value.Mesh).AppendString(OutString);
+	}
+
+	static bool IsNearlyEqual(const FNiagaraMeshRendererMeshPropertiesBase& Lhs, const FNiagaraMeshRendererMeshPropertiesBase& Rhs, float Tolerance)
+	{
+		return Lhs.IsNearlyEqual(Rhs, Tolerance);
+	}
+};
+
+NDIARRAY_GENERATE_IMPL(UNiagaraDataInterfaceArrayMesh, FNiagaraMeshRendererMeshPropertiesBase, MeshData)
+
+void UNiagaraDataInterfaceArrayMesh::ForEachMesh(FNiagaraSystemInstance* SystemInstance, TFunction<void(int32)> NumMeshesDelegate, TFunction<void(const FNiagaraMeshRendererMeshProperties&)> IterateDelegate) const
+{
+	const FNDIArrayInstanceData_GameThread<FNiagaraMeshRendererMeshPropertiesBase>* PerInstanceData = nullptr;
+	if (SystemInstance)
+	{
+		PerInstanceData = GetProxyAs<FProxyType>()->GetPerInstanceData_GameThread(SystemInstance->GetId());
+	}
+
+	FProxyType::FReadArrayRef ArrayRef(this, PerInstanceData);
+
+	NumMeshesDelegate(ArrayRef.GetArray().Num());
+
+	FNiagaraMeshRendererMeshProperties MeshProperties;
+	for (const FNiagaraMeshRendererMeshPropertiesBase& MeshEntry : ArrayRef.GetArray())
+	{
+		MeshProperties.Mesh				= MeshEntry.Mesh;
+		MeshProperties.Scale			= MeshEntry.Scale;
+		MeshProperties.Rotation			= MeshEntry.Rotation;
+		MeshProperties.PivotOffset		= MeshEntry.PivotOffset;
+		MeshProperties.PivotOffsetSpace	= MeshEntry.PivotOffsetSpace;
+		IterateDelegate(MeshProperties);
+	}
+}
+
+void UNiagaraDataInterfaceArrayMesh::SetNiagaraArrayMesh(UNiagaraComponent* NiagaraComponent, FName OverrideName, const TArray<FNiagaraMeshRendererMeshPropertiesBase>& ArrayData)
+{
+	if (UNiagaraDataInterfaceArrayMesh* ArrayDI = UNiagaraFunctionLibrary::GetDataInterface<UNiagaraDataInterfaceArrayMesh>(NiagaraComponent, OverrideName))
+	{
+		auto* ArrayProxy = static_cast<typename UNiagaraDataInterfaceArrayMesh::FProxyType*>(ArrayDI->GetProxy());
+		ArrayProxy->SetArrayDataAndRecreateRenderState(MakeArrayView(ArrayData));
+	}
+}
+
+void UNiagaraDataInterfaceArrayMesh::SetNiagaraArrayMeshSM(UNiagaraComponent* NiagaraComponent, FName OverrideName, const TArray<UStaticMesh*>& ArrayData)
+{
+	if (UNiagaraDataInterfaceArrayMesh* ArrayDI = UNiagaraFunctionLibrary::GetDataInterface<UNiagaraDataInterfaceArrayMesh>(NiagaraComponent, OverrideName))
+	{
+		auto* ArrayProxy = static_cast<typename UNiagaraDataInterfaceArrayMesh::FProxyType*>(ArrayDI->GetProxy());
+		ArrayProxy->SetArrayDataAndRecreateRenderState(MakeArrayView(ArrayData));
+	}
+}
diff --git a/Engine/Plugins/FX/Niagara/Source/Niagara/Private/NiagaraMeshRendererProperties.cpp b/Engine/Plugins/FX/Niagara/Source/Niagara/Private/NiagaraMeshRendererProperties.cpp
--- a/Engine/Plugins/FX/Niagara/Source/Niagara/Private/NiagaraMeshRendererProperties.cpp
+++ b/Engine/Plugins/FX/Niagara/Source/Niagara/Private/NiagaraMeshRendererProperties.cpp
@@ -3,6 +3,7 @@
 #include "NiagaraMeshRendererProperties.h"
 #include "Materials/Material.h"
 #include "NiagaraRendererMeshes.h"
+#include "NiagaraRenderableMeshArrayInterface.h"
 #include "NiagaraConstants.h"
 #include "NiagaraBoundsCalculatorHelper.h"
 #include "NiagaraCustomVersion.h"
@@ -276,17 +277,11 @@ bool FNiagaraMeshMaterialOverride::SerializeFromMismatchedTag(const struct FProp
 ////////////////////////////////////////////////////////////////////////////////////////////////////////////////////////////////
 
 FNiagaraMeshRendererMeshProperties::FNiagaraMeshRendererMeshProperties()
-	: Mesh(nullptr)
-#if WITH_EDITORONLY_DATA
-	, UserParamBinding_DEPRECATED(FNiagaraTypeDefinition(UStaticMesh::StaticClass()))
-#endif
-	, LODRange(0, MAX_STATIC_MESH_LODS)
-	, Scale(1.0f, 1.0f, 1.0f)
-	, Rotation(FRotator::ZeroRotator)
-	, PivotOffset(ForceInitToZero)
-	, PivotOffsetSpace(ENiagaraMeshPivotOffsetSpace::Mesh)
 {
+	LODRange = FIntVector2(0, MAX_STATIC_MESH_LODS);
 #if WITH_EDITORONLY_DATA
+	UserParamBinding_DEPRECATED = FNiagaraUserParameterBinding(FNiagaraTypeDefinition(UStaticMesh::StaticClass()));
+
 	MeshParameterBinding.SetUsage(ENiagaraParameterBindingUsage::NotParticle);
 	MeshParameterBinding.SetAllowedInterfaces({UNiagaraRenderableMeshInterface::StaticClass()});
 	MeshParameterBinding.SetAllowedObjects({UStaticMesh::StaticClass()});
@@ -301,6 +296,11 @@ FNiagaraMeshRendererMeshProperties::FNiagaraMeshRendererMeshProperties()
 #endif
 }
 
+bool FNiagaraMeshRendererMeshProperties::HasValidRenderableMesh() const
+{
+	return Mesh || MeshParameterBinding.ResolvedParameter.IsValid();
+}
+
 FNiagaraRenderableMeshPtr FNiagaraMeshRendererMeshProperties::ResolveRenderableMesh(const FNiagaraEmitterInstance* EmitterInstance) const
 {
 	INiagaraRenderableMeshInterface* RenderableMeshInterface = nullptr;
@@ -322,10 +322,7 @@ FNiagaraRenderableMeshPtr FNiagaraMeshRendererMeshProperties::ResolveRenderableM
 	return nullptr;
 }
 
-bool FNiagaraMeshRendererMeshProperties::HasValidRenderableMesh() const
-{
-	return Mesh || MeshParameterBinding.ResolvedParameter.IsValid();
-}
+////////////////////////////////////////////////////////////////////////////////////////////////////////////////////////////////
 
 UNiagaraMeshRendererProperties::UNiagaraMeshRendererProperties()
 	: bOverrideMaterials(false)
@@ -337,6 +334,9 @@ UNiagaraMeshRendererProperties::UNiagaraMeshRendererProperties()
 	FlipbookSuffixFormat = TEXT("_{frame_number}");
 	FlipbookSuffixNumDigits = 1;
 	NumFlipbookFrames = 1;
+
+	MeshesBinding.SetUsage(ENiagaraParameterBindingUsage::NotParticle);
+	MeshesBinding.SetAllowedInterfaces({ UNiagaraRenderableMeshArrayInterface::StaticClass() });
 #endif
 
 	AttributeBindings =
@@ -411,6 +411,8 @@ FNiagaraBoundsCalculator* UNiagaraMeshRendererProperties::CreateBoundsCalculator
 		bLocalSpace = EmitterData->bLocalSpace;
 	}
 
+	//-TODO: Bounds are incorrect when binding to a mesh inside the mesh properties or a mesh array
+
 	for (const auto& MeshProperties : Meshes)
 	{
 		if (MeshProperties.Mesh)
@@ -568,6 +570,10 @@ void UNiagaraMeshRendererProperties::UpdateSourceModeDerivates(ENiagaraRendererS
 	FVersionedNiagaraEmitter SrcEmitter = GetOuterEmitter();
 	if (SrcEmitter.Emitter)
 	{
+	#if WITH_EDITORONLY_DATA
+		MeshesBinding.OnRenameEmitter(SrcEmitter.Emitter->GetUniqueEmitterName());
+	#endif
+
 		for (FNiagaraMaterialAttributeBinding& MaterialParamBinding : MaterialParameters.AttributeBindings)
 		{
 			MaterialParamBinding.CacheValues(SrcEmitter.Emitter);
@@ -770,6 +776,25 @@ void UNiagaraMeshRendererProperties::ApplyMaterialOverrides(const FNiagaraEmitte
 	}
 }
 
+void UNiagaraMeshRendererProperties::ForEachMesh(FNiagaraSystemInstance* SystemInstance, TFunction<void(int32)> NumMeshesDelegate, TFunction<void(const FNiagaraMeshRendererMeshProperties&)> IterateDelegate) const
+{
+	NumMeshesDelegate(Meshes.Num());
+	for (const FNiagaraMeshRendererMeshProperties& MeshProperties : Meshes)
+	{
+		IterateDelegate(MeshProperties);
+	}
+}
+
+const INiagaraRenderableMeshArrayInterface* UNiagaraMeshRendererProperties::GetMeshArrayInterface(const FNiagaraEmitterInstance* EmitterInstance) const
+{
+	INiagaraRenderableMeshArrayInterface* ResolvedInterface = nullptr;
+	if (EmitterInstance)
+	{
+		ResolvedInterface = Cast<INiagaraRenderableMeshArrayInterface>(EmitterInstance->GetRendererBoundVariables().GetDataInterface(MeshesBinding.ResolvedParameter));
+	}
+	return ResolvedInterface ? ResolvedInterface : this;
+}
+
 const FVertexFactoryType* UNiagaraMeshRendererProperties::GetVertexFactoryType() const
 {
 	return &FNiagaraMeshVertexFactory::StaticType;
@@ -780,84 +805,93 @@ void UNiagaraMeshRendererProperties::GetUsedMaterials(const FNiagaraEmitterInsta
 	FNiagaraSystemInstance* SystemInstance = EmitterInstance ? EmitterInstance->GetParentSystemInstance() : nullptr;
 
 	TArray<UMaterialInterface*> OrderedMeshMaterials;
-	for (int32 MeshIndex=0; MeshIndex < Meshes.Num(); ++MeshIndex)
-	{
-		OrderedMeshMaterials.Reset(0);
 
-		INiagaraRenderableMeshInterface* RenderableMeshInterface = nullptr;
-		UStaticMesh* StaticMesh = nullptr;
-		NiagaraMeshRendererPropertiesInternal::ResolveRenderableMeshInternal(Meshes[MeshIndex], EmitterInstance, RenderableMeshInterface, StaticMesh);
-
-		if (RenderableMeshInterface && SystemInstance)
+	const INiagaraRenderableMeshArrayInterface* MeshArrayInterface = GetMeshArrayInterface(EmitterInstance);
+	MeshArrayInterface->ForEachMesh(
+		SystemInstance,
+		[&](const FNiagaraMeshRendererMeshProperties& MeshProperties)
 		{
-			RenderableMeshInterface->GetUsedMaterials(SystemInstance->GetId(), OrderedMeshMaterials);
-		}
-		else if (StaticMesh)
-		{
-			FNiagaraRenderableStaticMesh(StaticMesh).GetUsedMaterials(OrderedMeshMaterials);
-		}
+			OrderedMeshMaterials.Reset(0);
 
-		if (OrderedMeshMaterials.Num() > 0)
-		{
-			ApplyMaterialOverrides(EmitterInstance, OrderedMeshMaterials);
+			INiagaraRenderableMeshInterface* RenderableMeshInterface = nullptr;
+			UStaticMesh* StaticMesh = nullptr;
+			NiagaraMeshRendererPropertiesInternal::ResolveRenderableMeshInternal(MeshProperties, EmitterInstance, RenderableMeshInterface, StaticMesh);
 
-			OutMaterials.Reserve(OutMaterials.Num() + OrderedMeshMaterials.Num());
-			for (TObjectPtr<UMaterialInterface> MaterialInterface : OrderedMeshMaterials)
+			if (RenderableMeshInterface && SystemInstance)
 			{
-				if (MaterialInterface)
+				RenderableMeshInterface->GetUsedMaterials(SystemInstance->GetId(), OrderedMeshMaterials);
+			}
+			else if (StaticMesh)
+			{
+				FNiagaraRenderableStaticMesh(StaticMesh).GetUsedMaterials(OrderedMeshMaterials);
+			}
+
+			if (OrderedMeshMaterials.Num() > 0)
+			{
+				ApplyMaterialOverrides(EmitterInstance, OrderedMeshMaterials);
+
+				OutMaterials.Reserve(OutMaterials.Num() + OrderedMeshMaterials.Num());
+				for (TObjectPtr<UMaterialInterface> MaterialInterface : OrderedMeshMaterials)
 				{
-					OutMaterials.AddUnique(MaterialInterface);
+					if (MaterialInterface)
+					{
+						OutMaterials.AddUnique(MaterialInterface);
+					}
 				}
 			}
 		}
-	}
+	);
 }
 
 void UNiagaraMeshRendererProperties::CollectPSOPrecacheData(const FNiagaraEmitterInstance* InEmitter, FPSOPrecacheParamsList& OutParams) const
 {
 	const FVertexFactoryType* VFType = GetVertexFactoryType();
 	bool bSupportsManualVertexFetch = VFType->SupportsManualVertexFetch(GMaxRHIFeatureLevel);
-		
-	for (int32 MeshIndex = 0; MeshIndex < Meshes.Num(); ++MeshIndex)
-	{
-		INiagaraRenderableMeshInterface* RenderableMeshInterface = nullptr;
-		UStaticMesh* StaticMesh = nullptr;
-		NiagaraMeshRendererPropertiesInternal::ResolveRenderableMeshInternal(Meshes[MeshIndex], InEmitter, RenderableMeshInterface, StaticMesh);
-		if (StaticMesh)
+
+	const INiagaraRenderableMeshArrayInterface* MeshArrayInterface = GetMeshArrayInterface(InEmitter);
+	MeshArrayInterface->ForEachMesh(
+		InEmitter ? InEmitter->GetParentSystemInstance() : nullptr,
+		[&](const FNiagaraMeshRendererMeshProperties& MeshProperties)
 		{
-			TArray<UMaterialInterface*> OrderedMeshMaterials;
-			FNiagaraRenderableStaticMesh(StaticMesh).GetUsedMaterials(OrderedMeshMaterials);
-			if (OrderedMeshMaterials.Num())
+			INiagaraRenderableMeshInterface* RenderableMeshInterface = nullptr;
+			UStaticMesh* StaticMesh = nullptr;
+			NiagaraMeshRendererPropertiesInternal::ResolveRenderableMeshInternal(MeshProperties, InEmitter, RenderableMeshInterface, StaticMesh);
+			if (StaticMesh)
 			{
-				ApplyMaterialOverrides(nullptr, OrderedMeshMaterials);
-				for (UMaterialInterface* MeshMaterial : OrderedMeshMaterials)
+				TArray<UMaterialInterface*> OrderedMeshMaterials;
+				FNiagaraRenderableStaticMesh(StaticMesh).GetUsedMaterials(OrderedMeshMaterials);
+				if (OrderedMeshMaterials.Num())
 				{
-					if (MeshMaterial)
+					ApplyMaterialOverrides(nullptr, OrderedMeshMaterials);
+					for (UMaterialInterface* MeshMaterial : OrderedMeshMaterials)
 					{
-						FPSOPrecacheParams& PSOPrecacheParams = OutParams.AddDefaulted_GetRef();
-						PSOPrecacheParams.MaterialInterface = MeshMaterial;
-						if (!bSupportsManualVertexFetch)
+						if (MeshMaterial)
 						{
-							// Assuming here that all LOD use same vertex decl
-							int32 MeshLODIdx = StaticMesh->GetMinLODIdx();
-							if (StaticMesh->GetRenderData()->LODResources.IsValidIndex(MeshLODIdx))
+							FPSOPrecacheParams& PSOPrecacheParams = OutParams.AddDefaulted_GetRef();
+							PSOPrecacheParams.MaterialInterface = MeshMaterial;
+							if (!bSupportsManualVertexFetch)
 							{
-								FStaticMeshDataType Data;
-								FVertexDeclarationElementList Elements;
-								FNiagaraRenderableStaticMesh::InitVertexFactoryComponents(StaticMesh->GetRenderData()->LODResources[MeshLODIdx].VertexBuffers, nullptr, Data);
-								FNiagaraMeshVertexFactory::GetVertexElements(GMaxRHIFeatureLevel, bSupportsManualVertexFetch, Data, Elements);
-								PSOPrecacheParams.VertexFactoryDataList.Add(FPSOPrecacheVertexFactoryData(VFType, Elements));
+								// Assuming here that all LOD use same vertex decl
+								int32 MeshLODIdx = StaticMesh->GetMinLODIdx();
+								if (StaticMesh->GetRenderData()->LODResources.IsValidIndex(MeshLODIdx))
+								{
+									FStaticMeshDataType Data;
+									FVertexDeclarationElementList Elements;
+									FNiagaraRenderableStaticMesh::InitVertexFactoryComponents(StaticMesh->GetRenderData()->LODResources[MeshLODIdx].VertexBuffers, nullptr, Data);
+									FNiagaraMeshVertexFactory::GetVertexElements(GMaxRHIFeatureLevel, bSupportsManualVertexFetch, Data, Elements);
+									PSOPrecacheParams.VertexFactoryDataList.Add(FPSOPrecacheVertexFactoryData(VFType, Elements));
+								}
+							}
+							else
+							{
+								PSOPrecacheParams.VertexFactoryDataList.Add(FPSOPrecacheVertexFactoryData(VFType));
 							}
-						}
-						else
-						{
-							PSOPrecacheParams.VertexFactoryDataList.Add(FPSOPrecacheVertexFactoryData(VFType));
 						}
 					}
 				}
 			}
 		}
-	}
+	);
 }
 
 void UNiagaraMeshRendererProperties::GetStreamingMeshInfo(const FBoxSphereBounds& OwnerBounds, const FNiagaraEmitterInstance* InEmitter, TArray<FStreamingRenderAssetPrimitiveInfo>& OutStreamingRenderAssets) const
@@ -900,6 +934,12 @@ bool UNiagaraMeshRendererProperties::PopulateRequiredBindings(FNiagaraParameterS
 {
 	bool bAnyAdded = Super::PopulateRequiredBindings(InParameterStore);
 
+	if (MeshesBinding.ResolvedParameter.IsValid())
+	{
+		InParameterStore.AddParameter(MeshesBinding.ResolvedParameter, false);
+		bAnyAdded = true;
+	}
+
 	for (const FNiagaraVariableAttributeBinding* Binding : AttributeBindings)
 	{
 		if (Binding && Binding->CanBindToHostParameterMap())
@@ -1290,6 +1330,7 @@ void UNiagaraMeshRendererProperties::RenameVariable(const FNiagaraVariableBase&
 {
 	Super::RenameVariable(OldVariable, NewVariable, InEmitter);
 #if WITH_EDITORONLY_DATA
+	MeshesBinding.OnRenameVariable(OldVariable, NewVariable, InEmitter.Emitter->GetUniqueEmitterName());
 	MaterialParameters.RenameVariable(OldVariable, NewVariable, InEmitter, GetCurrentSourceMode());
 	for (FNiagaraMeshRendererMeshProperties& Mesh : Meshes)
 	{
@@ -1304,6 +1345,7 @@ void UNiagaraMeshRendererProperties::RemoveVariable(const FNiagaraVariableBase&
 {
 	Super::RemoveVariable(OldVariable, InEmitter);
 #if WITH_EDITORONLY_DATA
+	MeshesBinding.OnRemoveVariable(OldVariable, InEmitter.Emitter->GetUniqueEmitterName());
 	MaterialParameters.RemoveVariable(OldVariable, InEmitter, GetCurrentSourceMode());
 	for (FNiagaraMeshRendererMeshProperties& Mesh : Meshes)
 	{
diff --git a/Engine/Plugins/FX/Niagara/Source/Niagara/Private/NiagaraRendererMeshes.cpp b/Engine/Plugins/FX/Niagara/Source/Niagara/Private/NiagaraRendererMeshes.cpp
--- a/Engine/Plugins/FX/Niagara/Source/Niagara/Private/NiagaraRendererMeshes.cpp
+++ b/Engine/Plugins/FX/Niagara/Source/Niagara/Private/NiagaraRendererMeshes.cpp
@@ -15,6 +15,7 @@
 #include "NiagaraGPURayTracingTransformsShader.h"
 #include "NiagaraGPUSortInfo.h"
 #include "NiagaraGpuComputeDispatchInterface.h"
+#include "NiagaraRenderableMeshArrayInterface.h"
 #include "NiagaraRenderableMeshInterface.h"
 #include "NiagaraSceneProxy.h"
 #include "NiagaraSortingGPU.h"
@@ -210,94 +211,102 @@ void FNiagaraRendererMeshes::Initialize(const UNiagaraRendererProperties* InProp
 	const UNiagaraMeshRendererProperties* Properties = CastChecked<const UNiagaraMeshRendererProperties>(InProps);
 
 	// Initialize the valid mesh slots, and prep them with the data for every mesh, LOD, and section we'll be needing over the lifetime of the renderer
-	const uint32 MaxMeshes = Properties->Meshes.Num();
-	Meshes.Empty(MaxMeshes);
-	MeshUsedMax = 0;
-
-	for (uint32 SourceMeshIndex = 0; SourceMeshIndex < MaxMeshes; ++SourceMeshIndex)
-	{
-		const auto& MeshProperties = Properties->Meshes[SourceMeshIndex];
+	const INiagaraRenderableMeshArrayInterface* MeshArrayInterface = Properties->GetMeshArrayInterface(Emitter);
 
-		// Resolve the renderable mesh
-		FNiagaraRenderableMeshPtr RenderableMesh = MeshProperties.ResolveRenderableMesh(Emitter);
-		if (RenderableMesh == nullptr)
-		{
-			continue;
-		}
+	MeshUsedMax = 0;
 
-		MeshUsedMax = SourceMeshIndex + 1;
+	FNiagaraSystemInstance* SystemInstance = Emitter ? Emitter->GetParentSystemInstance() : nullptr;
 
-		FIntVector2 LODRange = FIntVector2(0, 1);
-		if (MeshProperties.LODMode == ENiagaraMeshLODMode::LODLevel)
+	uint32 SourceMeshIndex = 0;
+	MeshArrayInterface->ForEachMesh(
+		SystemInstance,
+		[this](int32 NumMeshes)
 		{
-			LODRange.X = MeshProperties.LODLevel;
-			LODRange.Y = LODRange.X + 1;
-		}
-		else if (MeshProperties.LODMode == ENiagaraMeshLODMode::PerParticle)
+			Meshes.Empty(NumMeshes);
+		},
+		[this, &SourceMeshIndex, Emitter, Properties](const FNiagaraMeshRendererMeshProperties& MeshProperties)
 		{
-			LODRange = RenderableMesh->GetLODRange();
-			if (MeshProperties.bUseLODRange)
+			FNiagaraRenderableMeshPtr RenderableMesh = MeshProperties.ResolveRenderableMesh(Emitter);
+			if (RenderableMesh != nullptr)
 			{
-				LODRange.X = FMath::Max(MeshProperties.LODRange.X, LODRange.X);
-				LODRange.Y = FMath::Max(MeshProperties.LODRange.Y, LODRange.Y);
-			}
-			bEnableLODCulling |= LODRange.Y - LODRange.X > 1;
-		}
+				MeshUsedMax = SourceMeshIndex + 1;
 
-		for (int32 LOD=LODRange.X; LOD < LODRange.Y; ++LOD)
-		{
-			// We have a valid mesh fill in the details
-			FMeshData& MeshData = Meshes.AddDefaulted_GetRef();
-			MeshData.RenderableMesh = RenderableMesh;
-			MeshData.SourceMeshIndex = SourceMeshIndex;
-			MeshData.LODMode = MeshProperties.LODMode;
-			MeshData.LODLevel = LOD;
-			MeshData.PivotOffset = FVector3f(MeshProperties.PivotOffset);
-			MeshData.PivotOffsetSpace = MeshProperties.PivotOffsetSpace;
-			MeshData.Scale = FVector3f(MeshProperties.Scale);
-			MeshData.Rotation = FQuat4f(MeshProperties.Rotation.Quaternion());
-
-			if (MeshProperties.LODMode == ENiagaraMeshLODMode::LODBias)
-			{
-				RenderableMesh->SetMinLODBias(MeshProperties.LODBias);
-			}
-			else if (MeshProperties.LODMode == ENiagaraMeshLODMode::ByComponentBounds || MeshProperties.LODMode == ENiagaraMeshLODMode::PerParticle)
-			{
-				MeshData.LODDistanceFactor = MeshProperties.LODDistanceFactor;
-				if (MeshProperties.LODMode == ENiagaraMeshLODMode::PerParticle)
+				FIntVector2 LODRange = FIntVector2(0, 1);
+				if (MeshProperties.LODMode == ENiagaraMeshLODMode::LODLevel)
 				{
-					MeshData.LODScreenSize = RenderableMesh->GetLODScreenSize(LOD);
-					MeshData.LODScreenSize.X = LOD < LODRange.Y - 1 ? MeshData.LODScreenSize.X : 0.0f;
-					MeshData.LODScreenSize.Y = LOD > LODRange.X ? MeshData.LODScreenSize.Y : -1.0f;
+					LODRange.X = MeshProperties.LODLevel;
+					LODRange.Y = LODRange.X + 1;
+				}
+				else if (MeshProperties.LODMode == ENiagaraMeshLODMode::PerParticle)
+				{
+					LODRange = RenderableMesh->GetLODRange();
+					if (MeshProperties.bUseLODRange)
+					{
+						LODRange.X = FMath::Max(MeshProperties.LODRange.X, LODRange.X);
+						LODRange.Y = FMath::Max(MeshProperties.LODRange.Y, LODRange.Y);
+					}
+					bEnableLODCulling |= LODRange.Y - LODRange.X > 1;
 				}
-			}
-
-			// Get materials and remap them into the base material list
-			TArray<UMaterialInterface*> UsedMaterials;
-			RenderableMesh->GetUsedMaterials(UsedMaterials);
-			Properties->ApplyMaterialOverrides(Emitter, UsedMaterials);
 
-			for (UMaterialInterface* UsedMaterial : UsedMaterials)
-			{
-				MeshData.MaterialRemapTable.Add(
-					BaseMaterials_GT.IndexOfByPredicate(
-						[&](UMaterialInterface* LookMat)
+				for (int32 LOD=LODRange.X; LOD < LODRange.Y; ++LOD)
+				{
+					// We have a valid mesh fill in the details
+					FMeshData& MeshData = Meshes.AddDefaulted_GetRef();
+					MeshData.RenderableMesh = RenderableMesh;
+					MeshData.SourceMeshIndex = SourceMeshIndex;
+					MeshData.LODMode = MeshProperties.LODMode;
+					MeshData.LODLevel = LOD;
+					MeshData.PivotOffset = FVector3f(MeshProperties.PivotOffset);
+					MeshData.PivotOffsetSpace = MeshProperties.PivotOffsetSpace;
+					MeshData.Scale = FVector3f(MeshProperties.Scale);
+					MeshData.Rotation = FQuat4f(MeshProperties.Rotation.Quaternion());
+
+					if (MeshProperties.LODMode == ENiagaraMeshLODMode::LODBias)
+					{
+						RenderableMesh->SetMinLODBias(MeshProperties.LODBias);
+					}
+					else if (MeshProperties.LODMode == ENiagaraMeshLODMode::ByComponentBounds || MeshProperties.LODMode == ENiagaraMeshLODMode::PerParticle)
+					{
+						MeshData.LODDistanceFactor = MeshProperties.LODDistanceFactor;
+						if (MeshProperties.LODMode == ENiagaraMeshLODMode::PerParticle)
 						{
-							if (LookMat == UsedMaterial)
-							{
-								return true;
-							}
-							if (UMaterialInstanceDynamic* MID = Cast<UMaterialInstanceDynamic>(LookMat))
-							{
-								return UsedMaterial == MID->Parent;
-							}
-							return false;
+							MeshData.LODScreenSize = RenderableMesh->GetLODScreenSize(LOD);
+							MeshData.LODScreenSize.X = LOD < LODRange.Y - 1 ? MeshData.LODScreenSize.X : 0.0f;
+							MeshData.LODScreenSize.Y = LOD > LODRange.X ? MeshData.LODScreenSize.Y : -1.0f;
 						}
-					)
-				);
+					}
+
+					// Get materials and remap them into the base material list
+					TArray<UMaterialInterface*> UsedMaterials;
+					RenderableMesh->GetUsedMaterials(UsedMaterials);
+					Properties->ApplyMaterialOverrides(Emitter, UsedMaterials);
+
+					for (UMaterialInterface* UsedMaterial : UsedMaterials)
+					{
+						MeshData.MaterialRemapTable.Add(
+							BaseMaterials_GT.IndexOfByPredicate(
+								[&](UMaterialInterface* LookMat)
+								{
+									if (LookMat == UsedMaterial)
+									{
+										return true;
+									}
+									if (UMaterialInstanceDynamic* MID = Cast<UMaterialInstanceDynamic>(LookMat))
+									{
+										return UsedMaterial == MID->Parent;
+									}
+									return false;
+								}
+							)
+						);
+					}
+				}
 			}
+
+			// Increment the source mesh index
+			++SourceMeshIndex;
 		}
-	}
+	);
 }
 
 void FNiagaraRendererMeshes::ReleaseRenderThreadResources()
diff --git a/Engine/Plugins/FX/Niagara/Source/Niagara/Public/NiagaraMeshRendererMeshProperties.h b/Engine/Plugins/FX/Niagara/Source/Niagara/Public/NiagaraMeshRendererMeshProperties.h
new file mode 100644
--- /dev/null
+++ b/Engine/Plugins/FX/Niagara/Source/Niagara/Public/NiagaraMeshRendererMeshProperties.h
@@ -0,0 +1,157 @@
+// Copyright Epic Games, Inc. All Rights Reserved.
+
+#pragma once
+
+#include "CoreMinimal.h"
+#include "NiagaraCommon.h"
+#include "NiagaraRenderableMeshInterface.h"
+#include "NiagaraRendererProperties.h"
+#include "NiagaraParameterBinding.h"
+#include "NiagaraMeshRendererMeshProperties.generated.h"
+
+class FNiagaraEmitterInstance;
+class UStaticMesh;
+
+UENUM()
+enum class ENiagaraMeshPivotOffsetSpace : uint8
+{
+	/** The pivot offset is in the mesh's local space (default) */
+	Mesh,
+	/** The pivot offset is in the emitter's local space if the emitter is marked as local-space, or in world space otherwise */
+	Simulation,
+	/** The pivot offset is in world space */
+	World,
+	/** The pivot offset is in the emitter's local space */
+	Local
+};
+
+UENUM()
+enum class ENiagaraMeshLODMode : uint8
+{
+	/*
+	* Uses the provided LOD level to render all mesh particles.
+	* If the LOD is not streamed in or available on the platform the next available lower LOD level will be used.
+	* For example, LOD Level is set to 1 but the first available is LOD 3 then LOD 3 will be used.
+	*/
+	LODLevel,
+
+	/**
+	* Takes the highest available LOD for the platform + LOD bias to render all mesh particles
+	* If the LOD is not streamed in or available on the platform the next available lower LOD level will be used.
+	* For example, LOD bias is set to 1, the current platform has Min LOD of 2 then 3 will be the used LOD.
+	*/
+	LODBias,
+
+	/*
+	* The LOD level is calculated based on screen space size of the component bounds.
+	* All particles will be rendered with the same calculated LOD level.
+	* Increasing 'LOD calculation scale' will result in lower quality LODs being used, this is useful as component bounds generally are larger than the particle mesh bounds.
+	*/
+	ByComponentBounds,
+
+	/*
+	* The LOD level is calcuated per particle using the particle position and mesh sphere bounds.
+	* This involves running a dispatch & draw per LOD level.
+	* Calculates and renders each particle with it's calcualted LOD level.
+	* Increasing 'LOD calculation scale' will result in lower quality LODs being used.
+	*/
+	PerParticle,
+};
+
+USTRUCT(BlueprintType)
+struct FNiagaraMeshRendererMeshPropertiesBase
+{
+	GENERATED_BODY()
+
+	/** The mesh to use when rendering this slot */
+	UPROPERTY(EditAnywhere, Category = "Mesh")
+	TObjectPtr<UStaticMesh> Mesh;
+
+	/** Scale of the mesh */
+	UPROPERTY(EditAnywhere, Category = "Mesh")
+	FVector Scale = FVector::OneVector;
+
+	/** Rotation of the mesh */
+	UPROPERTY(EditAnywhere, Category = "Mesh")
+	FRotator Rotation = FRotator::ZeroRotator;
+
+	/** Offset of the mesh pivot */
+	UPROPERTY(EditAnywhere, Category = "Mesh")
+	FVector PivotOffset = FVector::ZeroVector;
+
+	/** What space is the pivot offset in? */
+	UPROPERTY(EditAnywhere, Category = "Mesh")
+	ENiagaraMeshPivotOffsetSpace PivotOffsetSpace = ENiagaraMeshPivotOffsetSpace::Mesh;
+
+	bool IsNearlyEqual(const FNiagaraMeshRendererMeshPropertiesBase& Rhs, float Tolerance = UE_KINDA_SMALL_NUMBER) const
+	{
+		return
+			Mesh == Rhs.Mesh &&
+			Scale.Equals(Rhs.Scale, Tolerance) &&
+			Rotation.Equals(Rhs.Rotation, Tolerance) &&
+			PivotOffset.Equals(Rhs.PivotOffset, Tolerance) &&
+			PivotOffsetSpace == Rhs.PivotOffsetSpace;
+	}
+
+	bool operator==(const FNiagaraMeshRendererMeshPropertiesBase& Rhs) const
+	{
+		return IsNearlyEqual(Rhs);
+	}
+};
+
+USTRUCT()
+struct FNiagaraMeshRendererMeshProperties : public FNiagaraMeshRendererMeshPropertiesBase
+{
+	GENERATED_BODY()
+
+	NIAGARA_API FNiagaraMeshRendererMeshProperties();
+
+#if WITH_EDITORONLY_DATA
+	UPROPERTY()
+	FNiagaraUserParameterBinding UserParamBinding_DEPRECATED;
+#endif
+
+	/** Binding to supported mesh types. */
+	UPROPERTY(EditAnywhere, Category = "Mesh")
+	FNiagaraParameterBinding MeshParameterBinding;
+
+	UPROPERTY(EditAnywhere, Category = "Mesh")
+	ENiagaraMeshLODMode LODMode = ENiagaraMeshLODMode::LODLevel;
+
+#if WITH_EDITORONLY_DATA
+	/** Absolute LOD level to use */
+	UPROPERTY(EditAnywhere, Category = "Mesh", meta = (UIMin = "0", DisplayName="LOD Level", EditCondition = "LODMode == ENiagaraMeshLODMode::LODLevel", EditConditionHides))
+	FNiagaraParameterBindingWithValue LODLevelBinding;
+
+	/* LOD bias to apply to the LOD calculation. */
+	UPROPERTY(EditAnywhere, Category = "Mesh", meta = (UIMin = "0", DisplayName = "LOD Bias", EditCondition = "LODMode == ENiagaraMeshLODMode::LODBias", EditConditionHides))
+	FNiagaraParameterBindingWithValue LODBiasBinding;
+#endif
+
+	UPROPERTY()
+	int32 LODLevel = 0;
+
+	UPROPERTY()
+	int32 LODBias = 0;
+
+	/** Used in LOD calculation to modify the distance, i.e. increasing the value will make lower poly LODs transition closer to the camera. */
+	UPROPERTY(EditAnywhere, Category = "Mesh", meta = (UIMin = "0", DisplayName = "LOD Distance Factor", EditCondition = "LODMode == ENiagaraMeshLODMode::ByComponentBounds || LODMode == ENiagaraMeshLODMode::PerParticle", EditConditionHides))
+	float LODDistanceFactor = 1.0f;
+
+	/**
+	When enabled you can restrict the LOD range we consider for LOD calculation.
+	This can be useful to reduce the performance impact, as it reduces the number of draw calls required.
+	*/
+	UPROPERTY(EditAnywhere, Category = "Mesh", meta = (DisplayName = "Use LOD Range", EditCondition = "LODMode == ENiagaraMeshLODMode::PerParticle", EditConditionHides))
+	bool bUseLODRange = false;
+
+	/** Used to restrict the range of LODs we include when dynamically calculating the LOD level. */
+	UPROPERTY(EditAnywhere, Category = "Mesh", meta = (DisplayName = "LOD Range", EditCondition = "bUseLODRange && LODMode == ENiagaraMeshLODMode::PerParticle", EditConditionHides))
+	FIntVector2 LODRange;
+
+	/** Resolve renderable mesh. */
+	NIAGARA_API FNiagaraRenderableMeshPtr ResolveRenderableMesh(const FNiagaraEmitterInstance* EmitterInstance) const;
+
+	/** Is the renderable mesh potentially valid or not. */
+	NIAGARA_API bool HasValidRenderableMesh() const;
+};
diff --git a/Engine/Plugins/FX/Niagara/Source/Niagara/Public/NiagaraMeshRendererProperties.h b/Engine/Plugins/FX/Niagara/Source/Niagara/Public/NiagaraMeshRendererProperties.h
--- a/Engine/Plugins/FX/Niagara/Source/Niagara/Public/NiagaraMeshRendererProperties.h
+++ b/Engine/Plugins/FX/Niagara/Source/Niagara/Public/NiagaraMeshRendererProperties.h
@@ -5,6 +5,8 @@
 #include "CoreMinimal.h"
 #include "NiagaraCommon.h"
 #include "NiagaraGPUSortInfo.h"
+#include "NiagaraMeshRendererMeshProperties.h"
+#include "NiagaraRenderableMeshArrayInterface.h"
 #include "NiagaraRenderableMeshInterface.h"
 #include "NiagaraRendererProperties.h"
 #include "NiagaraParameterBinding.h"
@@ -12,6 +14,7 @@
 
 class UMaterialInstanceConstant;
 class FNiagaraEmitterInstance;
+class FNiagaraSystemInstance;
 class FAssetThumbnailPool;
 class SWidget;
 class UStaticMesh;
@@ -30,19 +33,6 @@ enum class ENiagaraMeshFacingMode : uint8
 	CameraPlane
 };
 
-UENUM()
-enum class ENiagaraMeshPivotOffsetSpace : uint8
-{
-	/** The pivot offset is in the mesh's local space (default) */
-	Mesh,
-	/** The pivot offset is in the emitter's local space if the emitter is marked as local-space, or in world space otherwise */
-	Simulation,
-	/** The pivot offset is in world space */
-	World,
-	/** The pivot offset is in the emitter's local space */
-	Local
-};
-
 UENUM()
 enum class ENiagaraMeshLockedAxisSpace : uint8
 {
@@ -54,39 +44,6 @@ enum class ENiagaraMeshLockedAxisSpace : uint8
 	Local
 };
 
-UENUM()
-enum class ENiagaraMeshLODMode : uint8
-{
-	/*
-	* Uses the provided LOD level to render all mesh particles.
-	* If the LOD is not streamed in or available on the platform the next available lower LOD level will be used.
-	* For example, LOD Level is set to 1 but the first available is LOD 3 then LOD 3 will be used.
-	*/
-	LODLevel,
-
-	/**
-	* Takes the highest available LOD for the platform + LOD bias to render all mesh particles
-	* If the LOD is not streamed in or available on the platform the next available lower LOD level will be used.
-	* For example, LOD bias is set to 1, the current platform has Min LOD of 2 then 3 will be the used LOD.
-	*/
-	LODBias,
-
-	/*
-	* The LOD level is calculated based on screen space size of the component bounds.
-	* All particles will be rendered with the same calculated LOD level.
-	* Increasing 'LOD calculation scale' will result in lower quality LODs being used, this is useful as component bounds generally are larger than the particle mesh bounds.
-	*/
-	ByComponentBounds,
-
-	/*
-	* The LOD level is calcuated per particle using the particle position and mesh sphere bounds.
-	* This involves running a dispatch & draw per LOD level.
-	* Calculates and renders each particle with it's calcualted LOD level.
-	* Increasing 'LOD calculation scale' will result in lower quality LODs being used.
-	*/
-	PerParticle,
-};
-
 USTRUCT()
 struct FNiagaraMeshMICOverride
 {
@@ -165,85 +122,8 @@ namespace ENiagaraMeshVFLayout
 	};
 };
 
-USTRUCT()
-struct FNiagaraMeshRendererMeshProperties
-{
-	GENERATED_BODY()
-
-	NIAGARA_API FNiagaraMeshRendererMeshProperties();
-
-	/** The mesh to use when rendering this slot */
-	UPROPERTY(EditAnywhere, Category = "Mesh")
-	TObjectPtr<UStaticMesh> Mesh;
-
-#if WITH_EDITORONLY_DATA
-	UPROPERTY()
-	FNiagaraUserParameterBinding UserParamBinding_DEPRECATED;
-#endif
-
-	/** Binding to supported mesh types. */
-	UPROPERTY(EditAnywhere, Category = "Mesh")
-	FNiagaraParameterBinding MeshParameterBinding;
-
-	UPROPERTY(EditAnywhere, Category = "Mesh")
-	ENiagaraMeshLODMode LODMode = ENiagaraMeshLODMode::LODLevel;
-
-#if WITH_EDITORONLY_DATA
-	/** Absolute LOD level to use */
-	UPROPERTY(EditAnywhere, Category = "Mesh", meta = (UIMin = "0", DisplayName="LOD Level", EditCondition = "LODMode == ENiagaraMeshLODMode::LODLevel", EditConditionHides))
-	FNiagaraParameterBindingWithValue LODLevelBinding;
-
-	/* LOD bias to apply to the LOD calculation. */
-	UPROPERTY(EditAnywhere, Category = "Mesh", meta = (UIMin = "0", DisplayName = "LOD Bias", EditCondition = "LODMode == ENiagaraMeshLODMode::LODBias", EditConditionHides))
-	FNiagaraParameterBindingWithValue LODBiasBinding;
-#endif
-
-	UPROPERTY()
-	int32 LODLevel = 0;
-
-	UPROPERTY()
-	int32 LODBias = 0;
-
-	/** Used in LOD calculation to modify the distance, i.e. increasing the value will make lower poly LODs transition closer to the camera. */
-	UPROPERTY(EditAnywhere, Category = "Mesh", meta = (UIMin = "0", DisplayName = "LOD Distance Factor", EditCondition = "LODMode == ENiagaraMeshLODMode::ByComponentBounds || LODMode == ENiagaraMeshLODMode::PerParticle", EditConditionHides))
-	float LODDistanceFactor = 1.0f;
-
-	/**
-	When enabled you can restrict the LOD range we consider for LOD calculation.
-	This can be useful to reduce the performance impact, as it reduces the number of draw calls required.
-	*/
-	UPROPERTY(EditAnywhere, Category = "Mesh", meta = (DisplayName = "Use LOD Range", EditCondition = "LODMode == ENiagaraMeshLODMode::PerParticle", EditConditionHides))
-	bool bUseLODRange = false;
-
-	/** Used to restrict the range of LODs we include when dynamically calculating the LOD level. */
-	UPROPERTY(EditAnywhere, Category = "Mesh", meta = (DisplayName = "LOD Range", EditCondition = "bUseLODRange && LODMode == ENiagaraMeshLODMode::PerParticle", EditConditionHides))
-	FIntVector2 LODRange;
-
-	/** Scale of the mesh */
-	UPROPERTY(EditAnywhere, Category = "Mesh")
-	FVector Scale;
-
-	/** Rotation of the mesh */
-	UPROPERTY(EditAnywhere, Category = "Mesh")
-	FRotator Rotation;
-
-	/** Offset of the mesh pivot */
-	UPROPERTY(EditAnywhere, Category = "Mesh")
-	FVector PivotOffset;
-
-	/** What space is the pivot offset in? */
-	UPROPERTY(EditAnywhere, Category = "Mesh")
-	ENiagaraMeshPivotOffsetSpace PivotOffsetSpace;
-
-	/** Resolve renderable mesh. */
-	NIAGARA_API FNiagaraRenderableMeshPtr ResolveRenderableMesh(const FNiagaraEmitterInstance* EmitterInstance) const;
-
-	/** Is the renderable mesh potentially valid or not. */
-	NIAGARA_API bool HasValidRenderableMesh() const;
-};
-
 UCLASS(editinlinenew, meta = (DisplayName = "Mesh Renderer", SupportsStateless), MinimalAPI)
-class UNiagaraMeshRendererProperties : public UNiagaraRendererProperties
+class UNiagaraMeshRendererProperties : public UNiagaraRendererProperties, public INiagaraRenderableMeshArrayInterface
 {
 public:
 	GENERATED_BODY()
@@ -297,6 +177,12 @@ public:
 
 	NIAGARA_API void ApplyMaterialOverrides(const FNiagaraEmitterInstance* EmitterInstance, TArray<UMaterialInterface*>& InOutMaterials) const;
 
+	//INiagaraRenderableMeshArrayInterface BEGIN
+	virtual void ForEachMesh(FNiagaraSystemInstance* SystemInstance, TFunction<void(int32)> NumMeshesDelegate, TFunction<void(const FNiagaraMeshRendererMeshProperties&)> IterateDelegate) const override;
+	//INiagaraRenderableMeshArrayInterface END
+
+	const INiagaraRenderableMeshArrayInterface* GetMeshArrayInterface(const FNiagaraEmitterInstance* EmitterInstance) const;
+
 	/**
 	 * The static mesh(es) to be instanced when rendering mesh particles.
 	 *
@@ -307,6 +193,9 @@ public:
 	UPROPERTY(EditAnywhere, Category = "Mesh Rendering", meta = (EditCondition = "!bEnableMeshFlipbook", EditConditionHides))
 	TArray<FNiagaraMeshRendererMeshProperties> Meshes;
 
+	UPROPERTY(EditAnywhere, Category = "Mesh Rendering", meta = (EditCondition = "!bEnableMeshFlipbook", EditConditionHides))
+	FNiagaraParameterBinding MeshesBinding;
+
 	/** Whether or not to draw a single element for the Emitter or to draw the particles.*/
 	UPROPERTY(EditAnywhere, Category = "Mesh Rendering")
 	ENiagaraRendererSourceDataMode SourceMode = ENiagaraRendererSourceDataMode::Particles;
diff --git a/Engine/Plugins/FX/Niagara/Source/Niagara/Public/NiagaraRenderableMeshArrayInterface.h b/Engine/Plugins/FX/Niagara/Source/Niagara/Public/NiagaraRenderableMeshArrayInterface.h
new file mode 100644
--- /dev/null
+++ b/Engine/Plugins/FX/Niagara/Source/Niagara/Public/NiagaraRenderableMeshArrayInterface.h
@@ -0,0 +1,25 @@
+// Copyright Epic Games, Inc. All Rights Reserved.
+
+#pragma once
+
+#include "NiagaraCore.h"
+#include "NiagaraMeshRendererMeshProperties.h"
+#include "NiagaraRenderableMeshArrayInterface.generated.h"
+
+class FNiagaraSystemInstance;
+
+UINTERFACE()
+class UNiagaraRenderableMeshArrayInterface : public UInterface
+{
+	GENERATED_BODY()
+};
+
+// Interface defintion for UObjects
+class INiagaraRenderableMeshArrayInterface
+{
+	GENERATED_BODY()
+
+public:
+	virtual void ForEachMesh(FNiagaraSystemInstance* SystemInstance, TFunction<void(int32)> NumMeshesDelegate, TFunction<void(const FNiagaraMeshRendererMeshProperties&)> IterateDelegate) const = 0;
+	void ForEachMesh(FNiagaraSystemInstance* SystemInstance, TFunction<void(const FNiagaraMeshRendererMeshProperties&)> IterateDelegate) const { ForEachMesh(SystemInstance, [](int32) {}, IterateDelegate); }
+};
