Added a validation rule to limit modules to certain sim targets

#tests Created a test module and added the validation rule in the module metadata. Tested with cpu/gpu emitters and undo/redo with the quick fixes.

diff --git a/Engine/Plugins/FX/Niagara/Source/NiagaraEditor/Private/NiagaraValidationRules.cpp b/Engine/Plugins/FX/Niagara/Source/NiagaraEditor/Private/NiagaraValidationRules.cpp
--- a/Engine/Plugins/FX/Niagara/Source/NiagaraEditor/Private/NiagaraValidationRules.cpp
+++ b/Engine/Plugins/FX/Niagara/Source/NiagaraEditor/Private/NiagaraValidationRules.cpp
@@ -1532,4 +1532,86 @@ void UNiagaraValidationRule_NoMapForOnCpu::CheckValidity(const FNiagaraValidatio
 	}
 }
 
+void UNiagaraValidationRule_ModuleSimTargetRestriction::CheckValidity(const FNiagaraValidationContext& Context,	TArray<FNiagaraValidationResult>& OutResults) const
+{
+	// check to see if we're called from a module
+	if (UNiagaraStackModuleItem* SourceModule = Cast<UNiagaraStackModuleItem>(Context.Source))
+	{
+		if (SourceModule->GetIsEnabled())
+		{
+			ENiagaraScriptUsage ScriptUsage = SourceModule->GetOutputNode()->GetUsage();
+			// system and emitter scripts are always cpu scripts
+			if (FNiagaraUtilities::ConvertScriptUsageToStaticSwitchContext(ScriptUsage) != ENiagaraScriptContextStaticSwitch::Particle && SupportedSimTarget == ENiagaraSimTarget::GPUComputeSim)
+			{
+				OutResults.Emplace(
+					Severity,
+					LOCTEXT("SimTargetModuleCpuError", "This module only supports gpu sim targets"),
+					LOCTEXT("SimTargetModuleCpuDetailedError", "This module is used in a system or emitter script (which are always executed on the cpu), but the module needs to run on the gpu. Place it in a particle script or simulation stage of a gpu emitter."),
+					SourceModule
+				);
+				return;
+			}
+			
+			TSharedPtr<FNiagaraEmitterViewModel> EmitterViewModel = SourceModule->GetEmitterViewModel();
+			if (EmitterViewModel.IsValid())
+			{
+				ENiagaraSimTarget SimTarget = EmitterViewModel->GetEmitter().GetEmitterData()->SimTarget;
+				if (SimTarget != SupportedSimTarget)
+				{
+					FNiagaraValidationResult& Result = OutResults.Emplace_GetRef(
+						Severity,
+						LOCTEXT("SimTargetModuleError", "Module is not compatible with the current emitter sim target"),
+						LOCTEXT("SimTargetModuleErrorDetailed", "This module has a restriction on the emitter sim target and can't run on both cpu/gpu emitters."),
+						SourceModule
+					);
+					Result.Fixes.Emplace(
+						FNiagaraValidationFix(
+							FText::Format(LOCTEXT("SimTargetModuleErrorFix", "Change emitter sim target to {0}"), FText::FromString(SupportedSimTarget == ENiagaraSimTarget::CPUSim ? "CPUSim" : "GPUComputeSim")),
+							FNiagaraValidationFixDelegate::CreateLambda(
+								[WeakSourceModule=MakeWeakObjectPtr(SourceModule), NewSimTarget=SupportedSimTarget]()
+								{
+									if (UNiagaraStackModuleItem* StackModule = WeakSourceModule.Get())
+									{
+										TSharedPtr<FNiagaraEmitterViewModel> EmitterViewModel = StackModule->GetEmitterViewModel();
+										if (EmitterViewModel.IsValid())
+										{
+											FVersionedNiagaraEmitter VersionedEmitter = EmitterViewModel->GetEmitter();
+											if (FVersionedNiagaraEmitterData* VersionedEmitterData = VersionedEmitter.GetEmitterData())
+											{
+												const FScopedTransaction Transaction(LOCTEXT("ChangeSimTarget", "Change emitter sim target"));
+												VersionedEmitter.Emitter->Modify();
+												VersionedEmitterData->SimTarget = NewSimTarget;
+
+												FProperty* SimTargetProperty = FindFProperty<FProperty>(FVersionedNiagaraEmitterData::StaticStruct(), GET_MEMBER_NAME_CHECKED(FVersionedNiagaraEmitterData, SimTarget));
+												FPropertyChangedEvent PropertyChangedEvent(SimTargetProperty);
+												VersionedEmitter.Emitter->PostEditChangeVersionedProperty(PropertyChangedEvent, VersionedEmitter.Version);
+
+												UNiagaraSystem::RequestCompileForEmitter(VersionedEmitter);
+											}
+										}
+									}
+								}
+							)
+						)
+					);
+					Result.Fixes.Emplace(
+						FNiagaraValidationFix(
+							LOCTEXT("SimTargetModuleErrorFixDisable", "Disable module"),
+							FNiagaraValidationFixDelegate::CreateLambda(
+								[WeakSourceModule=MakeWeakObjectPtr(SourceModule)]()
+								{
+									if (UNiagaraStackModuleItem* StackModule = WeakSourceModule.Get())
+									{
+										StackModule->SetEnabled(false);
+									}
+								}
+							)
+						)
+					);
+				}
+			}
+		}
+	}
+}
+
 #undef LOCTEXT_NAMESPACE
diff --git a/Engine/Plugins/FX/Niagara/Source/NiagaraEditor/Public/NiagaraValidationRules.h b/Engine/Plugins/FX/Niagara/Source/NiagaraEditor/Public/NiagaraValidationRules.h
--- a/Engine/Plugins/FX/Niagara/Source/NiagaraEditor/Public/NiagaraValidationRules.h
+++ b/Engine/Plugins/FX/Niagara/Source/NiagaraEditor/Public/NiagaraValidationRules.h
@@ -516,4 +516,21 @@ private:
 		bool bContainsMapForNode = false;
 	};
 	mutable TMap<FObjectKey, FGraphCheckResult> CachedResults;
-};
\ No newline at end of file
+};
+
+/** This validation rule checks that a module is only used in emitters with the configured runtime target. */
+UCLASS(Category = "Validation", DisplayName = "Module SimTarget Restriction")
+class UNiagaraValidationRule_ModuleSimTargetRestriction : public UNiagaraValidationRule
+{
+	GENERATED_BODY()
+
+public:
+	virtual void CheckValidity(const FNiagaraValidationContext& Context, TArray<FNiagaraValidationResult>& OutResults) const override;
+
+	/** How do we want to report the error in the stack */
+	UPROPERTY(EditAnywhere, Category = Validation)
+	ENiagaraValidationSeverity Severity = ENiagaraValidationSeverity::Error;
+	
+	UPROPERTY(EditAnywhere, Category = Validation)
+	ENiagaraSimTarget SupportedSimTarget = ENiagaraSimTarget::CPUSim;
+};
