Fixed a bug where recording a sim cache with "renderer only" attributes would cache all particle attributes when the renderer was in emitter mode.

#tests Tested by using a cache track in a level sequence, then changing the renderer mode to emitter source and checking the cached attributes. Also checked the sim cache engine tests are still green.

#virtualized

diff --git a/Engine/Plugins/FX/Niagara/Source/Niagara/Private/NiagaraSimCacheHelper.h b/Engine/Plugins/FX/Niagara/Source/Niagara/Private/NiagaraSimCacheHelper.h
--- a/Engine/Plugins/FX/Niagara/Source/Niagara/Private/NiagaraSimCacheHelper.h
+++ b/Engine/Plugins/FX/Niagara/Source/Niagara/Private/NiagaraSimCacheHelper.h
@@ -455,6 +455,13 @@ struct FNiagaraSimCacheHelper
 					}
 				}
 			}
+
+			if (ExplicitCaptureAttributes.IsEmpty())
+			{
+				// The attributes to capture only include system/emitter attributes or particle attributes from another emitter.
+				// We add a bogus filter attribute to build a minimal cache layout (as having no filter would record everything).
+				ExplicitCaptureAttributes.Add(FName("__FILTER.NONE__"));
+			}
 		}
 
 		BuildCacheLayout(CreateParameters, CacheLayout, *ParticleDataSet, EmitterHandle.GetName(), MoveTemp(RebaseVariableNames), MoveTemp(InterpVariableNames), ExplicitCaptureAttributes);
