- Always create the mesh renderer if we have a valid mesh array binding

#rnx

diff --git a/Engine/Plugins/FX/Niagara/Source/Niagara/Private/NiagaraMeshRendererProperties.cpp b/Engine/Plugins/FX/Niagara/Source/Niagara/Private/NiagaraMeshRendererProperties.cpp
--- a/Engine/Plugins/FX/Niagara/Source/Niagara/Private/NiagaraMeshRendererProperties.cpp
+++ b/Engine/Plugins/FX/Niagara/Source/Niagara/Private/NiagaraMeshRendererProperties.cpp
@@ -371,22 +371,32 @@ UNiagaraMeshRendererProperties::UNiagaraMeshRendererProperties()
 
 FNiagaraRenderer* UNiagaraMeshRendererProperties::CreateEmitterRenderer(ERHIFeatureLevel::Type FeatureLevel, const FNiagaraEmitterInstance* Emitter, const FNiagaraSystemInstanceController& InController)
 {
-	for (const auto& MeshProperties : Meshes)
+	const bool bMeshesBindingValid = MeshesBinding.ResolvedParameter.IsValid();
+	bool bShouldCreateRenderer = bMeshesBindingValid;
+	if (!bMeshesBindingValid)
 	{
-		if (MeshProperties.HasValidRenderableMesh())
+		for (const auto& MeshProperties : Meshes)
 		{
-			FNiagaraRendererMeshes* NewRenderer = new FNiagaraRendererMeshes(FeatureLevel, this, Emitter);
-			NewRenderer->Initialize(this, Emitter, InController);
-			if (NewRenderer->HasValidMeshes())
+			if (MeshProperties.HasValidRenderableMesh())
 			{
-				return NewRenderer;
+				bShouldCreateRenderer = true;
+				break;
 			}
+		}
+	}
 
-			// There are cases where we might end up with no meshes to render due to LODs or features not being enabled on that platform
-			// so we discard the renderer here, the cost to do this work in HasValidRenderableMesh makes it not worthwhile
-			delete NewRenderer;
-			return nullptr;
+	if (bShouldCreateRenderer)
+	{
+		FNiagaraRendererMeshes* NewRenderer = new FNiagaraRendererMeshes(FeatureLevel, this, Emitter);
+		NewRenderer->Initialize(this, Emitter, InController);
+		if (bMeshesBindingValid || NewRenderer->HasValidMeshes())
+		{
+			return NewRenderer;
 		}
+
+		// There are cases where we might end up with no meshes to render due to LODs or features not being enabled on that platform
+		// so we discard the renderer here, the cost to do this work in HasValidRenderableMesh makes it not worthwhile
+		delete NewRenderer;
 	}
 
 	return nullptr;
