- Niagara Experimental plugin
- This is for in progress work that requires collaboration
- Currently contains the Nanite renderer and Spawn Array data interface

diff --git a/Engine/Plugins/FX/Niagara/Source/Niagara/Classes/NiagaraDataInterfaceArray.h b/Engine/Plugins/FX/Niagara/Source/Niagara/Classes/NiagaraDataInterfaceArray.h
--- a/Engine/Plugins/FX/Niagara/Source/Niagara/Classes/NiagaraDataInterfaceArray.h
+++ b/Engine/Plugins/FX/Niagara/Source/Niagara/Classes/NiagaraDataInterfaceArray.h
@@ -122,11 +122,11 @@ public:
 	//UNiagaraDataInterface Interface
 
 	//~ INiagaraSimCacheCustomStorageInterface interface BEGIN
-	virtual UObject* SimCacheBeginWrite(UObject* SimCache, FNiagaraSystemInstance* NiagaraSystemInstance, const void* OptionalPerInstanceData, FNiagaraSimCacheFeedbackContext& FeedbackContext) const override;
-	virtual bool SimCacheWriteFrame(UObject* StorageObject, int FrameIndex, FNiagaraSystemInstance* SystemInstance, const void* OptionalPerInstanceData, FNiagaraSimCacheFeedbackContext& FeedbackContext) const override;
+	NIAGARA_API virtual UObject* SimCacheBeginWrite(UObject* SimCache, FNiagaraSystemInstance* NiagaraSystemInstance, const void* OptionalPerInstanceData, FNiagaraSimCacheFeedbackContext& FeedbackContext) const override;
+	NIAGARA_API virtual bool SimCacheWriteFrame(UObject* StorageObject, int FrameIndex, FNiagaraSystemInstance* SystemInstance, const void* OptionalPerInstanceData, FNiagaraSimCacheFeedbackContext& FeedbackContext) const override;
 	//virtual bool SimCacheEndWrite(UObject* StorageObject) const override;
-	virtual bool SimCacheReadFrame(UObject* StorageObject, int FrameA, int FrameB, float Interp, FNiagaraSystemInstance* SystemInstance, void* OptionalPerInstanceData) override;
-	virtual bool SimCacheCompareFrame(const UObject* LhsStorageObject, const UObject* RhsStorageObject, int FrameIndex, TOptional<float> Tolerance, FString& OutErrors) const override;
+	NIAGARA_API virtual bool SimCacheReadFrame(UObject* StorageObject, int FrameA, int FrameB, float Interp, FNiagaraSystemInstance* SystemInstance, void* OptionalPerInstanceData) override;
+	NIAGARA_API virtual bool SimCacheCompareFrame(const UObject* LhsStorageObject, const UObject* RhsStorageObject, int FrameIndex, TOptional<float> Tolerance, FString& OutErrors) const override;
 	//~ INiagaraSimCacheCustomStorageInterface interface END
 	NIAGARA_API FString SimCacheVisualizerRead(const UNDIArraySimCacheData* CacheData, const FNDIArraySimCacheDataFrame& FrameData, int Element) const;
 
@@ -175,5 +175,5 @@ public:
 	TArray<uint8> BufferData;
 
 	// Internal use only
-	int32 FindOrAddData(TConstArrayView<uint8> ArrayData);
+	NIAGARA_API int32 FindOrAddData(TConstArrayView<uint8> ArrayData);
 };
diff --git a/Engine/Plugins/FX/Niagara/Source/Niagara/Classes/NiagaraDataInterfaceArrayImpl.h b/Engine/Plugins/FX/Niagara/Source/Niagara/Classes/NiagaraDataInterfaceArrayImpl.h
--- a/Engine/Plugins/FX/Niagara/Source/Niagara/Classes/NiagaraDataInterfaceArrayImpl.h
+++ b/Engine/Plugins/FX/Niagara/Source/Niagara/Classes/NiagaraDataInterfaceArrayImpl.h
@@ -8,7 +8,6 @@
 #include "NiagaraDataInterfaceUtilities.h"
 #include "NiagaraGpuComputeDispatchInterface.h"
 #include "NiagaraGpuReadbackManager.h"
-#include "NiagaraStats.h"
 #include "NiagaraSystemInstance.h"
 
 #include "Async/Async.h"
@@ -17,6 +16,10 @@
 #include "ShaderParameterUtils.h"
 #include "ShaderCompilerCore.h"
 
+//////////////////////////////////////////////////////////////////////////
+// Internal array data interface implementation
+// WARNING: Do not use this is external code as the API is subject to change and is not guaranteed to support backwards compatability
+
 //////////////////////////////////////////////////////////////////////////
 // Helpers
 
@@ -156,7 +159,7 @@ struct FNDIArrayImplHelper : public FNDIArrayImplHelperBase<TArrayType>
 {
 };
 
-struct FNiagaraDataInterfaceArrayImplHelper
+struct FNiagaraDataInterfaceArrayImplInternal
 {
 	struct FFunctionVersion
 	{
@@ -170,32 +173,35 @@ struct FNiagaraDataInterfaceArrayImplHelper
 		};
 	};
 
-	static const TCHAR* HLSLReadTemplateFile;
-	static const TCHAR* HLSLReadWriteTemplateFile;
+	NIAGARA_API static const TCHAR* HLSLReadTemplateFile;
+	NIAGARA_API static const TCHAR* HLSLReadWriteTemplateFile;
 
-	static const FName Function_LengthName;
-	static const FName Function_IsValidIndexName;
-	static const FName Function_LastIndexName;
-	static const FName Function_GetName;
+	NIAGARA_API static const FName Function_LengthName;
+	NIAGARA_API static const FName Function_IsValidIndexName;
+	NIAGARA_API static const FName Function_LastIndexName;
+	NIAGARA_API static const FName Function_GetName;
 
-	static const FName Function_ClearName;
-	static const FName Function_ResizeName;
-	static const FName Function_SetArrayElemName;
-	static const FName Function_AddName;
-	static const FName Function_RemoveLastElemName;
+	NIAGARA_API static const FName Function_ClearName;
+	NIAGARA_API static const FName Function_ResizeName;
+	NIAGARA_API static const FName Function_SetArrayElemName;
+	NIAGARA_API static const FName Function_AddName;
+	NIAGARA_API static const FName Function_RemoveLastElemName;
 
-	static const FName Function_AtomicAddName;
-	static const FName Function_AtomicMinName;
-	static const FName Function_AtomicMaxName;
+	NIAGARA_API static const FName Function_AtomicAddName;
+	NIAGARA_API static const FName Function_AtomicMinName;
+	NIAGARA_API static const FName Function_AtomicMaxName;
 
 #if WITH_EDITORONLY_DATA
-	static bool UpgradeFunctionCall(FNiagaraFunctionSignature& FunctionSignature);
+	NIAGARA_API static void GetFunctions(TArray<FNiagaraFunctionSignature>& OutFunctions, UClass* DIClass, FNiagaraTypeDefinition ValueTypeDef, bool bSupportsCPU, bool bSupportsGPU, bool bSupportsAtomicOps);
+	NIAGARA_API static bool UpgradeFunctionCall(FNiagaraFunctionSignature& FunctionSignature);
 #endif
 	static const TCHAR* GetHLSLTemplateFile(bool bIsRWArray)
 	{
-		return bIsRWArray ? FNiagaraDataInterfaceArrayImplHelper::HLSLReadWriteTemplateFile : FNiagaraDataInterfaceArrayImplHelper::HLSLReadTemplateFile;
+		return bIsRWArray ? HLSLReadWriteTemplateFile : HLSLReadTemplateFile;
 	}
-	static bool IsRWFunction(const FName FunctionName);
+	NIAGARA_API static bool IsRWFunction(const FName FunctionName);
+
+	NIAGARA_API static ERHIAccess GetCountBufferDefaultState();
 };
 
 //////////////////////////////////////////////////////////////////////////
@@ -211,77 +217,42 @@ struct FNDIArrayInstanceData_GameThread
 	TArray<TArrayType>		ArrayData;					// Modified array data
 };
 
-template<typename TArrayType>
-struct FNDIArrayInstanceData_RenderThread
+struct FNDIArrayInstanceData_RenderThreadBase
 {
-	using TVMArrayType = typename FNDIArrayImplHelper<TArrayType>::TVMArrayType;
-
-	~FNDIArrayInstanceData_RenderThread()
-	{
-		if (CountOffset != INDEX_NONE)
-		{
-			ComputeInterface->GetGPUInstanceCounterManager().FreeEntry(CountOffset);
-			CountOffset = INDEX_NONE;
-		}
+	NIAGARA_API ~FNDIArrayInstanceData_RenderThreadBase();
+	NIAGARA_API void Initialize(FRHICommandListImmediate& RHICmdList, FNiagaraGpuComputeDispatchInterface* InComputeInterface, int32 InDefaultElements, bool bRWGpuArray);
+	NIAGARA_API void UpdateDataInternal(FRHICommandList& RHICmdList, int32 ArrayNum, int32 NewNumElements, uint32 ElementSize, EPixelFormat PixelFormat);
+	NIAGARA_API void ReleaseData();
 
-		ReleaseData();
-	}
+	NIAGARA_API void SimCacheWriteFrame(FRHICommandListImmediate& RHICmdList, UNDIArraySimCacheData* CacheData, int32 FrameIndex, int32 ArrayTypeSize, void(*CopyGpuToCpuMemory)(void*, const void*, int32)) const;
 
 	bool IsReadOnly() const { return CountOffset == INDEX_NONE; }
 
-	void Initialize(FRHICommandListImmediate& RHICmdList, FNiagaraGpuComputeDispatchInterface* InComputeInterface, int32 InDefaultElements, bool bRWGpuArray)
-	{
-		ComputeInterface = InComputeInterface;
-		DefaultElements = 0;
-		NumElements = INDEX_NONE;
-		CountOffset = INDEX_NONE;
+	FNiagaraGpuComputeDispatchInterface* ComputeInterface = nullptr;
 
-		if (bRWGpuArray)
-		{
-			DefaultElements = InDefaultElements;
-			CountOffset = ComputeInterface->GetGPUInstanceCounterManager().AcquireOrAllocateEntry(RHICmdList);
-		}
-	}
+	FBufferRHIRef				ArrayBuffer;
+	FUnorderedAccessViewRHIRef	ArrayUAV;
+	FShaderResourceViewRHIRef	ArraySRV;
+	uint32						ArrayNumBytes = 0;
+
+	int32						DefaultElements = 0;		// The default number of elements in the buffer, can be used to reduce allocations / required for RW buffers
+	int32						NumElements = INDEX_NONE;	// Number of elements in the buffer, for RW buffers this is the buffer size since the actual size is in the counter
+	uint32						CountOffset = INDEX_NONE;	// Counter offset for RW buffers
+};
+
+template<typename TArrayType>
+struct FNDIArrayInstanceData_RenderThread : public FNDIArrayInstanceData_RenderThreadBase
+{
+	using TVMArrayType = typename FNDIArrayImplHelper<TArrayType>::TVMArrayType;
 
 	template<typename T = FNDIArrayImplHelper<TArrayType>>
 	typename TEnableIf<T::bSupportsGPU>::Type UpdateDataImpl(FRHICommandList& RHICmdList, TConstArrayView<TArrayType> InArrayData)
 	{
 		const int32 NewNumElements = FMath::Max(DefaultElements, InArrayData.Num());
 
-		// Do we need to update the backing storage for the buffer
-		if (NewNumElements != NumElements)
-		{
-			// Release old data
-			ReleaseData();
-
-			// Allocate new data
-			NumElements = NewNumElements;
-			ArrayNumBytes = (NumElements + 1) * sizeof(TVMArrayType);	// Note +1 because we store the default value at the end of the buffer
-			INC_MEMORY_STAT_BY(STAT_NiagaraGPUDataInterfaceMemory, ArrayNumBytes);
-
-			const EPixelFormat PixelFormat = IsReadOnly() ? FNDIArrayImplHelper<TArrayType>::ReadPixelFormat : FNDIArrayImplHelper<TArrayType>::RWPixelFormat;
-			const int32 TypeStride = GPixelFormats[PixelFormat].BlockBytes;
-
-			// Create Buffer
-			FRHIResourceCreateInfo CreateInfo(TEXT("NiagaraDataInterfaceArray"));
-			const EBufferUsageFlags BufferUsage = BUF_Static | BUF_ShaderResource | BUF_VertexBuffer | BUF_SourceCopy | (IsReadOnly() ? BUF_None : BUF_UnorderedAccess);
-			const ERHIAccess DefaultAccess = IsReadOnly() ? ERHIAccess::SRVCompute : ERHIAccess::UAVCompute;
-			ArrayBuffer = RHICmdList.CreateBuffer(ArrayNumBytes, BufferUsage, TypeStride, DefaultAccess, CreateInfo);
-
-			ArraySRV = RHICmdList.CreateShaderResourceView(
-				ArrayBuffer, 
-				FRHIViewDesc::CreateBufferSRV()
-					.SetType(FRHIViewDesc::EBufferType::Typed)
-					.SetFormat(PixelFormat));
-			if ( !IsReadOnly() )
-			{
-				ArrayUAV = RHICmdList.CreateUnorderedAccessView(
-					ArrayBuffer, 
-					FRHIViewDesc::CreateBufferUAV()
-						.SetType(FRHIViewDesc::EBufferType::Typed)
-						.SetFormat(PixelFormat));
-			}
-		}
+		// Resize if required and update the count value
+		const EPixelFormat PixelFormat = IsReadOnly() ? FNDIArrayImplHelper<TArrayType>::ReadPixelFormat : FNDIArrayImplHelper<TArrayType>::RWPixelFormat;
+		UpdateDataInternal(RHICmdList, InArrayData.Num(), NewNumElements, sizeof(TVMArrayType), PixelFormat);
 
 		// Copy data in new data over
 		{
@@ -296,19 +267,6 @@ struct FNDIArrayInstanceData_RenderThread
 
 			RHICmdList.UnlockBuffer(ArrayBuffer);
 		}
-
-		// Adjust counter value
-		if ( CountOffset != INDEX_NONE )
-		{
-			//-OPT: We could push this into the count manager and batch set as part of the clear process
-			const FNiagaraGPUInstanceCountManager& CounterManager = ComputeInterface->GetGPUInstanceCounterManager();
-			const FRWBuffer& CountBuffer = CounterManager.GetInstanceCountBuffer();
-			
-			const TPair<uint32, uint32> DataToClear(CountOffset, InArrayData.Num());
-			RHICmdList.Transition(FRHITransitionInfo(CountBuffer.UAV, FNiagaraGPUInstanceCountManager::kCountBufferDefaultState, ERHIAccess::UAVCompute));
-			NiagaraClearCounts::ClearCountsUInt(RHICmdList, CountBuffer.UAV, MakeArrayView(&DataToClear, 1));
-			RHICmdList.Transition(FRHITransitionInfo(CountBuffer.UAV, ERHIAccess::UAVCompute, FNiagaraGPUInstanceCountManager::kCountBufferDefaultState));
-		}
 	}
 
 	template<typename T = FNDIArrayImplHelper<TArrayType>>
@@ -320,25 +278,6 @@ struct FNDIArrayInstanceData_RenderThread
 	{
 		UpdateDataImpl(RHICmdList, InArrayData);
 	}
-
-	void ReleaseData()
-	{
-		DEC_MEMORY_STAT_BY(STAT_NiagaraGPUDataInterfaceMemory, ArrayNumBytes);
-		ArrayBuffer.SafeRelease();
-		ArrayUAV.SafeRelease();
-		ArraySRV.SafeRelease();
-	}
-
-	FNiagaraGpuComputeDispatchInterface* ComputeInterface = nullptr;
-
-	FBufferRHIRef				ArrayBuffer;
-	FUnorderedAccessViewRHIRef	ArrayUAV;
-	FShaderResourceViewRHIRef	ArraySRV;
-	uint32						ArrayNumBytes = 0;
-
-	int32						DefaultElements = 0;		// The default number of elements in the buffer, can be used to reduce allocations / required for RW buffers
-	int32						NumElements = INDEX_NONE;	// Number of elements in the buffer, for RW buffers this is the buffer size since the actual size is in the counter
-	uint32						CountOffset = INDEX_NONE;	// Counter offset for RW buffers
 };
 
 template<typename TArrayType, class TOwnerType>
@@ -686,156 +625,30 @@ struct FNDIArrayProxyImpl : public INDIArrayProxyBase
 	// INDIArrayProxyBase
 	virtual void GetFunctions(TArray<FNiagaraFunctionSignature>& OutFunctions) const override
 	{
-		OutFunctions.Reserve(OutFunctions.Num() + 9);
-
-		// Immutable functions
-		FNiagaraFunctionSignature DefaultImmutableSig;
-		DefaultImmutableSig.bMemberFunction = true;
-		DefaultImmutableSig.bRequiresContext = false;
-		DefaultImmutableSig.bSupportsCPU = FNDIArrayImplHelper<TArrayType>::bSupportsCPU;
-		DefaultImmutableSig.bSupportsGPU = FNDIArrayImplHelper<TArrayType>::bSupportsGPU;
-		DefaultImmutableSig.Inputs.Emplace(FNiagaraTypeDefinition(TOwnerType::StaticClass()), TEXT("Array interface"));
-		DefaultImmutableSig.FunctionVersion = FNiagaraDataInterfaceArrayImplHelper::FFunctionVersion::LatestVersion;
-		{
-			FNiagaraFunctionSignature& Sig = OutFunctions.Add_GetRef(DefaultImmutableSig);
-			Sig.Name = FNiagaraDataInterfaceArrayImplHelper::Function_LengthName;
-			Sig.bSupportsCPU = true;
-			Sig.Outputs.Emplace(FNiagaraTypeDefinition::GetIntDef(), TEXT("Num"));
-			Sig.Description = NSLOCTEXT("Niagara", "Array_LengthDesc", "Gets the number of elements in the array.");
-		}
-
-		{
-			FNiagaraFunctionSignature& Sig = OutFunctions.Add_GetRef(DefaultImmutableSig);
-			Sig.Name = FNiagaraDataInterfaceArrayImplHelper::Function_IsValidIndexName;
-			Sig.bSupportsCPU = true;
-			Sig.Inputs.Emplace(FNiagaraTypeDefinition::GetIntDef(), TEXT("Index"));
-			Sig.Outputs.Emplace(FNiagaraTypeDefinition::GetBoolDef(), TEXT("Valid"));
-			Sig.Description = NSLOCTEXT("Niagara", "Array_IsValidIndexDesc", "Tests to see if the index is valid and exists in the array.");
-		}
-
-		{
-			FNiagaraFunctionSignature& Sig = OutFunctions.Add_GetRef(DefaultImmutableSig);
-			Sig.Name = FNiagaraDataInterfaceArrayImplHelper::Function_LastIndexName;
-			Sig.bSupportsCPU = true;
-			Sig.Outputs.Emplace(FNiagaraTypeDefinition::GetIntDef(), TEXT("Index"));
-			Sig.Description = NSLOCTEXT("Niagara", "Array_LastIndexDesc", "Returns the last valid index in the array, will be -1 if no elements.");
-		}
-
-		{
-			FNiagaraFunctionSignature& Sig = OutFunctions.Add_GetRef(DefaultImmutableSig);
-			Sig.Name = FNiagaraDataInterfaceArrayImplHelper::Function_GetName;
-			Sig.Inputs.Emplace(FNiagaraTypeDefinition::GetIntDef(), TEXT("Index"));
-			Sig.Outputs.Emplace(FNDIArrayImplHelper<TArrayType>::GetTypeDefinition(), TEXT("Value"));
-			Sig.Description = NSLOCTEXT("Niagara", "Array_GetDesc", "Gets the value from the array at the given zero based index.");
-		}
-
-		// Mutable functions
-		FNiagaraFunctionSignature DefaultMutableSig = DefaultImmutableSig;
-		DefaultMutableSig.bSupportsGPU = FNDIArrayImplHelper<TArrayType>::bSupportsGPU;
-		DefaultMutableSig.bRequiresExecPin = true;
-		{
-			FNiagaraFunctionSignature& Sig = OutFunctions.Add_GetRef(DefaultMutableSig);
-			Sig.Name = FNiagaraDataInterfaceArrayImplHelper::Function_ClearName;
-			Sig.ModuleUsageBitmask = ENiagaraScriptUsageMask::System | ENiagaraScriptUsageMask::Emitter;
-			Sig.Description = NSLOCTEXT("Niagara", "Array_ClearDesc", "Clears the array, removing all elements");
-		}
-
-		{
-			FNiagaraFunctionSignature& Sig = OutFunctions.Add_GetRef(DefaultMutableSig);
-			Sig.Name = FNiagaraDataInterfaceArrayImplHelper::Function_ResizeName;
-			Sig.ModuleUsageBitmask = ENiagaraScriptUsageMask::System | ENiagaraScriptUsageMask::Emitter;
-			Sig.Inputs.Emplace(FNiagaraTypeDefinition::GetIntDef(), TEXT("Num"));
-			Sig.Description = NSLOCTEXT("Niagara", "Array_ResizeDesc", "Resizes the array to the specified size, initializing new elements with the default value.");
-		}
-
-		{
-			FNiagaraFunctionSignature& Sig = OutFunctions.Add_GetRef(DefaultMutableSig);
-			Sig.Name = FNiagaraDataInterfaceArrayImplHelper::Function_SetArrayElemName;
-			Sig.Inputs.Emplace(FNiagaraTypeDefinition::GetBoolDef(), TEXT("SkipSet"));
-			Sig.Inputs.Emplace(FNiagaraTypeDefinition::GetIntDef(), TEXT("Index"));
-			Sig.Inputs.Emplace(FNDIArrayImplHelper<TArrayType>::GetTypeDefinition(), TEXT("Value"));
-			Sig.Description = NSLOCTEXT("Niagara", "Array_SetArrayElemDesc", "Sets the value at the given zero based index (i.e the first element is 0).");
-			Sig.InputDescriptions.Add(Sig.Inputs[1], NSLOCTEXT("Niagara", "Array_SetArrayElemDesc_SkipSet", "When enabled will not set the array value."));
-		}
-
-		{
-			FNiagaraFunctionSignature& Sig = OutFunctions.Add_GetRef(DefaultMutableSig);
-			Sig.Name = FNiagaraDataInterfaceArrayImplHelper::Function_AddName;
-			Sig.Inputs.Emplace(FNiagaraTypeDefinition::GetBoolDef(), TEXT("SkipAdd"));
-			Sig.Inputs.Emplace(FNDIArrayImplHelper<TArrayType>::GetTypeDefinition(), TEXT("Value"));
-			Sig.Description = NSLOCTEXT("Niagara", "Array_AddDesc", "Optionally add a value onto the end of the array.");
-			Sig.InputDescriptions.Add(Sig.Inputs[1], NSLOCTEXT("Niagara", "Array_AddDesc_SkipAdd", "When enabled we will not add an element to the array."));
-		}
-
-		{
-			FNiagaraFunctionSignature& Sig = OutFunctions.Add_GetRef(DefaultMutableSig);
-			Sig.Name = FNiagaraDataInterfaceArrayImplHelper::Function_RemoveLastElemName;
-			Sig.Inputs.Emplace(FNiagaraTypeDefinition::GetBoolDef(), TEXT("SkipRemove"));
-			Sig.Outputs.Emplace(FNDIArrayImplHelper<TArrayType>::GetTypeDefinition(), TEXT("Value"));
-			Sig.Outputs.Emplace(FNiagaraTypeDefinition::GetBoolDef(), TEXT("IsValid"));
-			Sig.Description = NSLOCTEXT("Niagara", "Array_RemoveLastElemDesc", "Optionally remove the last element from the array.  Returns the default value if no elements are in the array or you skip the remove.");
-			Sig.InputDescriptions.Add(Sig.Inputs[1], NSLOCTEXT("Niagara", "Array_RemoveLastElemDesc_SkipRemove", "When enabled will not remove a value from the array, the return value will therefore be invalid."));
-			Sig.OutputDescriptions.Add(Sig.Outputs[1], NSLOCTEXT("Niagara", "Array_RemoveLastElemDesc_IsValid", "True if we removed a value from the array, False if no entries or we skipped the remove."));
-		}
-
-		if (FNDIArrayImplHelper<TArrayType>::bSupportsAtomicOps)
-		{
-			{
-				FNiagaraFunctionSignature& Sig = OutFunctions.Add_GetRef(DefaultMutableSig);
-				Sig.Name = FNiagaraDataInterfaceArrayImplHelper::Function_AtomicAddName;
-				Sig.Inputs.Emplace(FNiagaraTypeDefinition::GetBoolDef(), TEXT("SkipAdd"));
-				Sig.Inputs.Emplace(FNiagaraTypeDefinition::GetIntDef(), TEXT("Index"));
-				Sig.Inputs.Emplace(FNDIArrayImplHelper<TArrayType>::GetTypeDefinition(), TEXT("Value"));
-				Sig.Outputs.Emplace(FNDIArrayImplHelper<TArrayType>::GetTypeDefinition(), TEXT("PreviousValue"));
-				Sig.Outputs.Emplace(FNDIArrayImplHelper<TArrayType>::GetTypeDefinition(), TEXT("CurrentValue"));
-				Sig.SetDescription(NSLOCTEXT("Niagara", "Array_AtomicAddDesc", "Optionally perform an atomic add on the array element."));
-				Sig.SetInputDescription(Sig.Inputs[1], NSLOCTEXT("Niagara", "Array_AtomicAdd_SkipAdd", "When enabled will not perform the add operation, the return values will therefore be invalid."));
-				Sig.SetOutputDescription(Sig.Inputs[0], NSLOCTEXT("Niagara", "Array_AtomicAdd_PrevValue", "The value before the operation was performed."));
-				Sig.SetOutputDescription(Sig.Inputs[1], NSLOCTEXT("Niagara", "Array_AtomicAdd_CurrValue", "The value after the operation was performed."));
-			}
-			{
-				FNiagaraFunctionSignature& Sig = OutFunctions.Add_GetRef(DefaultMutableSig);
-				Sig.Name = FNiagaraDataInterfaceArrayImplHelper::Function_AtomicMinName;
-				Sig.Inputs.Emplace(FNiagaraTypeDefinition::GetBoolDef(), TEXT("SkipMin"));
-				Sig.Inputs.Emplace(FNiagaraTypeDefinition::GetIntDef(), TEXT("Index"));
-				Sig.Inputs.Emplace(FNDIArrayImplHelper<TArrayType>::GetTypeDefinition(), TEXT("Value"));
-				Sig.Outputs.Emplace(FNDIArrayImplHelper<TArrayType>::GetTypeDefinition(), TEXT("PreviousValue"));
-				Sig.Outputs.Emplace(FNDIArrayImplHelper<TArrayType>::GetTypeDefinition(), TEXT("CurrentValue"));
-				Sig.SetDescription(NSLOCTEXT("Niagara", "Array_AtomicMinDesc", "Optionally perform an atomic min on the array element."));
-				Sig.SetInputDescription(Sig.Inputs[1], NSLOCTEXT("Niagara", "Array_AtomicMin_SkipMin", "When enabled will not perform the min operation, the return values will therefore be invalid."));
-				Sig.SetOutputDescription(Sig.Inputs[0], NSLOCTEXT("Niagara", "Array_AtomicMin_PrevValue", "The value before the operation was performed."));
-				Sig.SetOutputDescription(Sig.Inputs[1], NSLOCTEXT("Niagara", "Array_AtomicMin_CurrValue", "The value after the operation was performed."));
-			}
-			{
-				FNiagaraFunctionSignature& Sig = OutFunctions.Add_GetRef(DefaultMutableSig);
-				Sig.Name = FNiagaraDataInterfaceArrayImplHelper::Function_AtomicMaxName;
-				Sig.Inputs.Emplace(FNiagaraTypeDefinition::GetBoolDef(), TEXT("SkipMax"));
-				Sig.Inputs.Emplace(FNiagaraTypeDefinition::GetIntDef(), TEXT("Index"));
-				Sig.Inputs.Emplace(FNDIArrayImplHelper<TArrayType>::GetTypeDefinition(), TEXT("Value"));
-				Sig.Outputs.Emplace(FNDIArrayImplHelper<TArrayType>::GetTypeDefinition(), TEXT("PreviousValue"));
-				Sig.Outputs.Emplace(FNDIArrayImplHelper<TArrayType>::GetTypeDefinition(), TEXT("CurrentValue"));
-				Sig.SetDescription(NSLOCTEXT("Niagara", "Array_AtomicMaxDesc", "Optionally perform an atomic max on the array element."));
-				Sig.SetInputDescription(Sig.Inputs[1], NSLOCTEXT("Niagara", "Array_AtomicMax_SkipMax", "When enabled will not perform the max operation, the return values will therefore be invalid."));
-				Sig.SetOutputDescription(Sig.Inputs[0], NSLOCTEXT("Niagara", "Array_AtomicMax_PrevValue", "The value before the operation was performed."));
-				Sig.SetOutputDescription(Sig.Inputs[1], NSLOCTEXT("Niagara", "Array_AtomicMax_CurrValue", "The value after the operation was performed."));
-			}
-		}
+		FNiagaraDataInterfaceArrayImplInternal::GetFunctions(
+			OutFunctions,
+			TOwnerType::StaticClass(),
+			FNDIArrayImplHelper<TArrayType>::GetTypeDefinition(),
+			FNDIArrayImplHelper<TArrayType>::bSupportsCPU,
+			FNDIArrayImplHelper<TArrayType>::bSupportsGPU,
+			FNDIArrayImplHelper<TArrayType>::bSupportsAtomicOps
+		);
 	}
 #endif
 
 	void GetVMExternalFunction_Internal(const FVMExternalFunctionBindingInfo& BindingInfo, void* InstanceData, FVMExternalFunction& OutFunc)
 	{
-		if (BindingInfo.Name == FNiagaraDataInterfaceArrayImplHelper::Function_LengthName)
+		if (BindingInfo.Name == FNiagaraDataInterfaceArrayImplInternal::Function_LengthName)
 		{
 			check(BindingInfo.GetNumInputs() == 1 && BindingInfo.GetNumOutputs() == 1);
 			OutFunc = FVMExternalFunction::CreateLambda([this](FVectorVMExternalFunctionContext& Context) { this->VMGetLength(Context); });
 		}
-		else if (BindingInfo.Name == FNiagaraDataInterfaceArrayImplHelper::Function_IsValidIndexName)
+		else if (BindingInfo.Name == FNiagaraDataInterfaceArrayImplInternal::Function_IsValidIndexName)
 		{
 			check(BindingInfo.GetNumInputs() == 2 && BindingInfo.GetNumOutputs() == 1);
 			OutFunc = FVMExternalFunction::CreateLambda([this](FVectorVMExternalFunctionContext& Context) { this->VMIsValidIndex(Context); });
 		}
-		else if (BindingInfo.Name == FNiagaraDataInterfaceArrayImplHelper::Function_LastIndexName)
+		else if (BindingInfo.Name == FNiagaraDataInterfaceArrayImplInternal::Function_LastIndexName)
 		{
 			check(BindingInfo.GetNumInputs() == 1 && BindingInfo.GetNumOutputs() == 1);
 			OutFunc = FVMExternalFunction::CreateLambda([this](FVectorVMExternalFunctionContext& Context) { this->VMGetLastIndex(Context); });
@@ -846,33 +659,33 @@ struct FNDIArrayProxyImpl : public INDIArrayProxyBase
 	typename TEnableIf<T::bSupportsCPU>::Type GetVMExternalFunction_CPUAccessInternal(const FVMExternalFunctionBindingInfo& BindingInfo, void* InstanceData, FVMExternalFunction& OutFunc)
 	{
 		// Immutable functions
-		if (BindingInfo.Name == FNiagaraDataInterfaceArrayImplHelper::Function_GetName)
+		if (BindingInfo.Name == FNiagaraDataInterfaceArrayImplInternal::Function_GetName)
 		{
 			// Note: Outputs is variable based upon type
 			//check(BindingInfo.GetNumInputs() == 2 && BindingInfo.GetNumOutputs() == 1);
 			OutFunc = FVMExternalFunction::CreateLambda([this](FVectorVMExternalFunctionContext& Context) { this->VMGetValue(Context); });
 		}
 		// Mutable functions
-		else if (BindingInfo.Name == FNiagaraDataInterfaceArrayImplHelper::Function_ClearName)
+		else if (BindingInfo.Name == FNiagaraDataInterfaceArrayImplInternal::Function_ClearName)
 		{
 			check(BindingInfo.GetNumInputs() == 1 && BindingInfo.GetNumOutputs() == 0);
 			OutFunc = FVMExternalFunction::CreateLambda([this](FVectorVMExternalFunctionContext& Context) { this->VMClear(Context); });
 		}
-		else if (BindingInfo.Name == FNiagaraDataInterfaceArrayImplHelper::Function_ResizeName)
+		else if (BindingInfo.Name == FNiagaraDataInterfaceArrayImplInternal::Function_ResizeName)
 		{
 			check(BindingInfo.GetNumInputs() == 2 && BindingInfo.GetNumOutputs() == 0);
 			OutFunc = FVMExternalFunction::CreateLambda([this](FVectorVMExternalFunctionContext& Context) { this->VMResize(Context); });
 		}
-		else if (BindingInfo.Name == FNiagaraDataInterfaceArrayImplHelper::Function_SetArrayElemName)
+		else if (BindingInfo.Name == FNiagaraDataInterfaceArrayImplInternal::Function_SetArrayElemName)
 		{
 			OutFunc = FVMExternalFunction::CreateLambda([this](FVectorVMExternalFunctionContext& Context) { this->VMSetValue(Context); });
 		}
-		else if (BindingInfo.Name == FNiagaraDataInterfaceArrayImplHelper::Function_AddName)
+		else if (BindingInfo.Name == FNiagaraDataInterfaceArrayImplInternal::Function_AddName)
 		{
 			// Note: Inputs is variable based upon type
 			OutFunc = FVMExternalFunction::CreateLambda([this](FVectorVMExternalFunctionContext& Context) { this->VMPushValue(Context); });
 		}
-		else if (BindingInfo.Name == FNiagaraDataInterfaceArrayImplHelper::Function_RemoveLastElemName)
+		else if (BindingInfo.Name == FNiagaraDataInterfaceArrayImplInternal::Function_RemoveLastElemName)
 		{
 			// Note: Outputs is variable based upon type
 			OutFunc = FVMExternalFunction::CreateLambda([this](FVectorVMExternalFunctionContext& Context) { this->VMPopValue(Context); });
@@ -887,15 +700,15 @@ struct FNDIArrayProxyImpl : public INDIArrayProxyBase
 	template<typename T = FNDIArrayImplHelper<TArrayType>>
 	typename TEnableIf<T::bSupportsAtomicOps>::Type GetVMExternalFunction_AtomicInternal(const FVMExternalFunctionBindingInfo& BindingInfo, void* InstanceData, FVMExternalFunction& OutFunc)
 	{
-		if (BindingInfo.Name == FNiagaraDataInterfaceArrayImplHelper::Function_AtomicAddName)
+		if (BindingInfo.Name == FNiagaraDataInterfaceArrayImplInternal::Function_AtomicAddName)
 		{
 			OutFunc = FVMExternalFunction::CreateLambda([this](FVectorVMExternalFunctionContext& Context) { this->VMAtomicAdd(Context); });
 		}
-		else if (BindingInfo.Name == FNiagaraDataInterfaceArrayImplHelper::Function_AtomicMinName)
+		else if (BindingInfo.Name == FNiagaraDataInterfaceArrayImplInternal::Function_AtomicMinName)
 		{
 			OutFunc = FVMExternalFunction::CreateLambda([this](FVectorVMExternalFunctionContext& Context) { this->VMAtomicMin(Context); });
 		}
-		else if (BindingInfo.Name == FNiagaraDataInterfaceArrayImplHelper::Function_AtomicMaxName)
+		else if (BindingInfo.Name == FNiagaraDataInterfaceArrayImplInternal::Function_AtomicMaxName)
 		{
 			OutFunc = FVMExternalFunction::CreateLambda([this](FVectorVMExternalFunctionContext& Context) { this->VMAtomicMax(Context); });
 		}
@@ -925,7 +738,7 @@ struct FNDIArrayProxyImpl : public INDIArrayProxyBase
 		return ParamInfo.GeneratedFunctions.ContainsByPredicate(
 			[&](const FNiagaraDataInterfaceGeneratedFunction& Function)
 			{
-				return FNiagaraDataInterfaceArrayImplHelper::IsRWFunction(Function.DefinitionName);
+				return FNiagaraDataInterfaceArrayImplInternal::IsRWFunction(Function.DefinitionName);
 			}
 		);
 	}
@@ -946,7 +759,7 @@ struct FNDIArrayProxyImpl : public INDIArrayProxyBase
 		};
 
 		FString TemplateFile;
-		LoadShaderSourceFile(FNiagaraDataInterfaceArrayImplHelper::GetHLSLTemplateFile(IsRWGpuArray(ParamInfo)), EShaderPlatform::SP_PCD3D_SM5, &TemplateFile, nullptr);
+		LoadShaderSourceFile(FNiagaraDataInterfaceArrayImplInternal::GetHLSLTemplateFile(IsRWGpuArray(ParamInfo)), EShaderPlatform::SP_PCD3D_SM5, &TemplateFile, nullptr);
 		OutHLSL += FString::Format(*TemplateFile, TemplateArgs);
 	}
 
@@ -964,28 +777,28 @@ struct FNDIArrayProxyImpl : public INDIArrayProxyBase
 	{
 		if (FNDIArrayImplHelper<TArrayType>::bSupportsGPU)
 		{
-			if ((FunctionInfo.DefinitionName == FNiagaraDataInterfaceArrayImplHelper::Function_LengthName) ||
-				(FunctionInfo.DefinitionName == FNiagaraDataInterfaceArrayImplHelper::Function_IsValidIndexName) ||
-				(FunctionInfo.DefinitionName == FNiagaraDataInterfaceArrayImplHelper::Function_LastIndexName) ||
-				(FunctionInfo.DefinitionName == FNiagaraDataInterfaceArrayImplHelper::Function_GetName))
+			if ((FunctionInfo.DefinitionName == FNiagaraDataInterfaceArrayImplInternal::Function_LengthName) ||
+				(FunctionInfo.DefinitionName == FNiagaraDataInterfaceArrayImplInternal::Function_IsValidIndexName) ||
+				(FunctionInfo.DefinitionName == FNiagaraDataInterfaceArrayImplInternal::Function_LastIndexName) ||
+				(FunctionInfo.DefinitionName == FNiagaraDataInterfaceArrayImplInternal::Function_GetName))
 			{
 				return true;
 			}
 
-			if ((FunctionInfo.DefinitionName == FNiagaraDataInterfaceArrayImplHelper::Function_ClearName) ||
-				(FunctionInfo.DefinitionName == FNiagaraDataInterfaceArrayImplHelper::Function_ResizeName) ||
-				(FunctionInfo.DefinitionName == FNiagaraDataInterfaceArrayImplHelper::Function_SetArrayElemName) ||
-				(FunctionInfo.DefinitionName == FNiagaraDataInterfaceArrayImplHelper::Function_AddName) ||
-				(FunctionInfo.DefinitionName == FNiagaraDataInterfaceArrayImplHelper::Function_RemoveLastElemName))
+			if ((FunctionInfo.DefinitionName == FNiagaraDataInterfaceArrayImplInternal::Function_ClearName) ||
+				(FunctionInfo.DefinitionName == FNiagaraDataInterfaceArrayImplInternal::Function_ResizeName) ||
+				(FunctionInfo.DefinitionName == FNiagaraDataInterfaceArrayImplInternal::Function_SetArrayElemName) ||
+				(FunctionInfo.DefinitionName == FNiagaraDataInterfaceArrayImplInternal::Function_AddName) ||
+				(FunctionInfo.DefinitionName == FNiagaraDataInterfaceArrayImplInternal::Function_RemoveLastElemName))
 			{
 				return true;
 			}
 
 			if ( FNDIArrayImplHelper<TArrayType>::bSupportsAtomicOps )
 			{
-				if ((FunctionInfo.DefinitionName == FNiagaraDataInterfaceArrayImplHelper::Function_AtomicAddName) ||
-					(FunctionInfo.DefinitionName == FNiagaraDataInterfaceArrayImplHelper::Function_AtomicMinName) ||
-					(FunctionInfo.DefinitionName == FNiagaraDataInterfaceArrayImplHelper::Function_AtomicMaxName) )
+				if ((FunctionInfo.DefinitionName == FNiagaraDataInterfaceArrayImplInternal::Function_AtomicAddName) ||
+					(FunctionInfo.DefinitionName == FNiagaraDataInterfaceArrayImplInternal::Function_AtomicMinName) ||
+					(FunctionInfo.DefinitionName == FNiagaraDataInterfaceArrayImplInternal::Function_AtomicMaxName) )
 				{
 					return true;
 				}
@@ -998,15 +811,15 @@ struct FNDIArrayProxyImpl : public INDIArrayProxyBase
 	{
 		if (FNDIArrayImplHelper<TArrayType>::bSupportsGPU)
 		{
-			InVisitor->UpdateShaderFile(FNiagaraDataInterfaceArrayImplHelper::GetHLSLTemplateFile(false));
-			InVisitor->UpdateShaderFile(FNiagaraDataInterfaceArrayImplHelper::GetHLSLTemplateFile(true));
+			InVisitor->UpdateShaderFile(FNiagaraDataInterfaceArrayImplInternal::GetHLSLTemplateFile(false));
+			InVisitor->UpdateShaderFile(FNiagaraDataInterfaceArrayImplInternal::GetHLSLTemplateFile(true));
 		}
 		return true;
 	}
 
 	virtual bool UpgradeFunctionCall(FNiagaraFunctionSignature& FunctionSignature) const override
 	{
-		return FNiagaraDataInterfaceArrayImplHelper::UpgradeFunctionCall(FunctionSignature);
+		return FNiagaraDataInterfaceArrayImplInternal::UpgradeFunctionCall(FunctionSignature);
 	}
 #endif
 #if WITH_NIAGARA_DEBUGGER
@@ -1077,67 +890,16 @@ struct FNDIArrayProxyImpl : public INDIArrayProxyBase
 				[Proxy_RT=this, InstanceID, CacheData, FrameIndex, ComputeInterface](FRHICommandListImmediate& RHICmdList)
 				{
 					const FNDIArrayInstanceData_RenderThread<TArrayType>* InstanceData_RT = Proxy_RT->PerInstanceData_RenderThread.Find(InstanceID);
-					if (InstanceData_RT == nullptr || InstanceData_RT->ArrayNumBytes == 0)
-					{
-						return;
-					}
-
-					if (CacheData->GpuFrameData.Num() <= FrameIndex)
+					if (InstanceData_RT)
 					{
-						CacheData->GpuFrameData.AddDefaulted(FrameIndex + 1 - CacheData->GpuFrameData.Num());
-					}
-					FNDIArraySimCacheDataFrame& FrameData = CacheData->GpuFrameData[FrameIndex];
-
-					TArray<FNiagaraGpuReadbackManager::FBufferRequest, TInlineAllocator<2>> BufferRequests;
-					TArray<FRHITransitionInfo, TInlineAllocator<2>> TransitionsBefore;
-					TArray<FRHITransitionInfo, TInlineAllocator<2>> TransitionsAfter;
-
-					const ERHIAccess DefaultAccess = InstanceData_RT->IsReadOnly() ? ERHIAccess::SRVCompute : ERHIAccess::UAVCompute;
-
-					BufferRequests.Emplace(InstanceData_RT->ArrayBuffer, 0, InstanceData_RT->ArrayNumBytes);
-					TransitionsBefore.Emplace(InstanceData_RT->ArrayBuffer, DefaultAccess, ERHIAccess::CopySrc);
-					TransitionsAfter.Emplace(InstanceData_RT->ArrayBuffer, ERHIAccess::CopySrc, DefaultAccess);
-
-					if (InstanceData_RT->IsReadOnly())
-					{
-						FrameData.NumElements = InstanceData_RT->NumElements;
-					}
-					else
-					{
-						const FNiagaraGPUInstanceCountManager& CountManager = ComputeInterface->GetGPUInstanceCounterManager();
-						BufferRequests.Emplace(CountManager.GetInstanceCountBuffer().Buffer, uint32(InstanceData_RT->CountOffset * sizeof(uint32)), sizeof(uint32));
-						TransitionsBefore.Emplace(CountManager.GetInstanceCountBuffer().UAV, FNiagaraGPUInstanceCountManager::kCountBufferDefaultState, ERHIAccess::CopySrc);
-						TransitionsAfter.Emplace(CountManager.GetInstanceCountBuffer().UAV, ERHIAccess::CopySrc, FNiagaraGPUInstanceCountManager::kCountBufferDefaultState);
+						InstanceData_RT->SimCacheWriteFrame(
+							RHICmdList,
+							CacheData,
+							FrameIndex,
+							sizeof(TVMArrayType),
+							&FNDIArrayImplHelper<TArrayType>::CopyGpuToCpuMemory
+						);
 					}
-
-					FNiagaraGpuReadbackManager* ReadbackManager = ComputeInterface->GetGpuReadbackManager();
-					RHICmdList.Transition(TransitionsBefore);
-					ReadbackManager->EnqueueReadbacks(
-						RHICmdList,
-						BufferRequests,
-						[CacheData, &FrameData, bReadOnly=InstanceData_RT->IsReadOnly()](TConstArrayView<TPair<void*, uint32>> ReadbackData)
-						{
-							if (!bReadOnly)
-							{
-								FrameData.NumElements = *reinterpret_cast<const uint32*>(ReadbackData[1].Key);
-							}
-							if ( FrameData.NumElements > 0 )
-							{
-								TArray<TVMArrayType> ArrayData;
-								ArrayData.AddUninitialized(FrameData.NumElements);
-								FNDIArrayImplHelper<TArrayType>::CopyGpuToCpuMemory(ArrayData.GetData(), ReadbackData[0].Key, FrameData.NumElements);
-
-								FrameData.DataOffset = CacheData->FindOrAddData(
-									MakeArrayView(
-										reinterpret_cast<const uint8*>(ArrayData.GetData()),
-										ArrayData.Num() * ArrayData.GetTypeSize()
-									)
-								);
-							}
-						}
-					);
-					RHICmdList.Transition(TransitionsAfter);
-					ReadbackManager->WaitCompletion(RHICmdList);
 				}
 			);
 
@@ -1245,7 +1007,7 @@ struct FNDIArrayProxyImpl : public INDIArrayProxyBase
 				DataInterface, SystemInstance,
 				[&](const UNiagaraScript* Script, const FNiagaraDataInterfaceGeneratedFunction& Function) -> bool
 				{
-					bRWGpuArray = FNiagaraDataInterfaceArrayImplHelper::IsRWFunction(Function.DefinitionName);
+					bRWGpuArray = FNiagaraDataInterfaceArrayImplInternal::IsRWFunction(Function.DefinitionName);
 					return !bRWGpuArray;
 				}
 			);
diff --git a/Engine/Plugins/FX/Niagara/Source/Niagara/Private/NiagaraDataInterfaceArrayImpl.cpp b/Engine/Plugins/FX/Niagara/Source/Niagara/Private/NiagaraDataInterfaceArrayImpl.cpp
--- a/Engine/Plugins/FX/Niagara/Source/Niagara/Private/NiagaraDataInterfaceArrayImpl.cpp
+++ b/Engine/Plugins/FX/Niagara/Source/Niagara/Private/NiagaraDataInterfaceArrayImpl.cpp
@@ -2,26 +2,166 @@
 
 #include "NiagaraDataInterfaceArrayImpl.h"
 
-const TCHAR* FNiagaraDataInterfaceArrayImplHelper::HLSLReadTemplateFile = TEXT("/Plugin/FX/Niagara/Private/NiagaraDataInterfaceArrayTemplate.ush");
-const TCHAR* FNiagaraDataInterfaceArrayImplHelper::HLSLReadWriteTemplateFile = TEXT("/Plugin/FX/Niagara/Private/NiagaraDataInterfaceArrayRWTemplate.ush");
+#include "NiagaraStats.h"
 
-const FName FNiagaraDataInterfaceArrayImplHelper::Function_LengthName(TEXT("Length"));
-const FName FNiagaraDataInterfaceArrayImplHelper::Function_IsValidIndexName(TEXT("IsValidIndex"));
-const FName FNiagaraDataInterfaceArrayImplHelper::Function_LastIndexName(TEXT("LastIndex"));
-const FName FNiagaraDataInterfaceArrayImplHelper::Function_GetName(TEXT("Get"));
+const TCHAR* FNiagaraDataInterfaceArrayImplInternal::HLSLReadTemplateFile = TEXT("/Plugin/FX/Niagara/Private/NiagaraDataInterfaceArrayTemplate.ush");
+const TCHAR* FNiagaraDataInterfaceArrayImplInternal::HLSLReadWriteTemplateFile = TEXT("/Plugin/FX/Niagara/Private/NiagaraDataInterfaceArrayRWTemplate.ush");
 
-const FName FNiagaraDataInterfaceArrayImplHelper::Function_ClearName(TEXT("Clear"));
-const FName FNiagaraDataInterfaceArrayImplHelper::Function_ResizeName(TEXT("Resize"));
-const FName FNiagaraDataInterfaceArrayImplHelper::Function_SetArrayElemName(TEXT("SetArrayElem"));
-const FName FNiagaraDataInterfaceArrayImplHelper::Function_AddName(TEXT("Add"));
-const FName FNiagaraDataInterfaceArrayImplHelper::Function_RemoveLastElemName(TEXT("RemoveLastElem"));
+const FName FNiagaraDataInterfaceArrayImplInternal::Function_LengthName(TEXT("Length"));
+const FName FNiagaraDataInterfaceArrayImplInternal::Function_IsValidIndexName(TEXT("IsValidIndex"));
+const FName FNiagaraDataInterfaceArrayImplInternal::Function_LastIndexName(TEXT("LastIndex"));
+const FName FNiagaraDataInterfaceArrayImplInternal::Function_GetName(TEXT("Get"));
 
-const FName FNiagaraDataInterfaceArrayImplHelper::Function_AtomicAddName("AtomicAdd");
-const FName FNiagaraDataInterfaceArrayImplHelper::Function_AtomicMinName("AtomicMin");
-const FName FNiagaraDataInterfaceArrayImplHelper::Function_AtomicMaxName("AtomicMax");
+const FName FNiagaraDataInterfaceArrayImplInternal::Function_ClearName(TEXT("Clear"));
+const FName FNiagaraDataInterfaceArrayImplInternal::Function_ResizeName(TEXT("Resize"));
+const FName FNiagaraDataInterfaceArrayImplInternal::Function_SetArrayElemName(TEXT("SetArrayElem"));
+const FName FNiagaraDataInterfaceArrayImplInternal::Function_AddName(TEXT("Add"));
+const FName FNiagaraDataInterfaceArrayImplInternal::Function_RemoveLastElemName(TEXT("RemoveLastElem"));
+
+const FName FNiagaraDataInterfaceArrayImplInternal::Function_AtomicAddName("AtomicAdd");
+const FName FNiagaraDataInterfaceArrayImplInternal::Function_AtomicMinName("AtomicMin");
+const FName FNiagaraDataInterfaceArrayImplInternal::Function_AtomicMaxName("AtomicMax");
 
 #if WITH_EDITORONLY_DATA
-bool FNiagaraDataInterfaceArrayImplHelper::UpgradeFunctionCall(FNiagaraFunctionSignature& FunctionSignature)
+void FNiagaraDataInterfaceArrayImplInternal::GetFunctions(TArray<FNiagaraFunctionSignature>& OutFunctions, UClass* DIClass, FNiagaraTypeDefinition ValueTypeDef, bool bSupportsCPU, bool bSupportsGPU, bool bSupportsAtomicOps)
+{
+	OutFunctions.Reserve(OutFunctions.Num() + 9);
+
+	// Immutable functions
+	FNiagaraFunctionSignature DefaultImmutableSig;
+	DefaultImmutableSig.bMemberFunction = true;
+	DefaultImmutableSig.bRequiresContext = false;
+	DefaultImmutableSig.bSupportsCPU = bSupportsCPU;
+	DefaultImmutableSig.bSupportsGPU = bSupportsGPU;
+	DefaultImmutableSig.Inputs.Emplace(FNiagaraTypeDefinition(DIClass), TEXT("Array interface"));
+	DefaultImmutableSig.FunctionVersion = FNiagaraDataInterfaceArrayImplInternal::FFunctionVersion::LatestVersion;
+	{
+		FNiagaraFunctionSignature& Sig = OutFunctions.Add_GetRef(DefaultImmutableSig);
+		Sig.Name = FNiagaraDataInterfaceArrayImplInternal::Function_LengthName;
+		Sig.bSupportsCPU = true;
+		Sig.Outputs.Emplace(FNiagaraTypeDefinition::GetIntDef(), TEXT("Num"));
+		Sig.Description = NSLOCTEXT("Niagara", "Array_LengthDesc", "Gets the number of elements in the array.");
+	}
+
+	{
+		FNiagaraFunctionSignature& Sig = OutFunctions.Add_GetRef(DefaultImmutableSig);
+		Sig.Name = FNiagaraDataInterfaceArrayImplInternal::Function_IsValidIndexName;
+		Sig.bSupportsCPU = true;
+		Sig.Inputs.Emplace(FNiagaraTypeDefinition::GetIntDef(), TEXT("Index"));
+		Sig.Outputs.Emplace(FNiagaraTypeDefinition::GetBoolDef(), TEXT("Valid"));
+		Sig.Description = NSLOCTEXT("Niagara", "Array_IsValidIndexDesc", "Tests to see if the index is valid and exists in the array.");
+	}
+
+	{
+		FNiagaraFunctionSignature& Sig = OutFunctions.Add_GetRef(DefaultImmutableSig);
+		Sig.Name = FNiagaraDataInterfaceArrayImplInternal::Function_LastIndexName;
+		Sig.bSupportsCPU = true;
+		Sig.Outputs.Emplace(FNiagaraTypeDefinition::GetIntDef(), TEXT("Index"));
+		Sig.Description = NSLOCTEXT("Niagara", "Array_LastIndexDesc", "Returns the last valid index in the array, will be -1 if no elements.");
+	}
+
+	{
+		FNiagaraFunctionSignature& Sig = OutFunctions.Add_GetRef(DefaultImmutableSig);
+		Sig.Name = FNiagaraDataInterfaceArrayImplInternal::Function_GetName;
+		Sig.Inputs.Emplace(FNiagaraTypeDefinition::GetIntDef(), TEXT("Index"));
+		Sig.Outputs.Emplace(ValueTypeDef, TEXT("Value"));
+		Sig.Description = NSLOCTEXT("Niagara", "Array_GetDesc", "Gets the value from the array at the given zero based index.");
+	}
+
+	// Mutable functions
+	FNiagaraFunctionSignature DefaultMutableSig = DefaultImmutableSig;
+	DefaultMutableSig.bSupportsGPU = bSupportsGPU;
+	DefaultMutableSig.bRequiresExecPin = true;
+	{
+		FNiagaraFunctionSignature& Sig = OutFunctions.Add_GetRef(DefaultMutableSig);
+		Sig.Name = FNiagaraDataInterfaceArrayImplInternal::Function_ClearName;
+		Sig.ModuleUsageBitmask = ENiagaraScriptUsageMask::System | ENiagaraScriptUsageMask::Emitter;
+		Sig.Description = NSLOCTEXT("Niagara", "Array_ClearDesc", "Clears the array, removing all elements");
+	}
+
+	{
+		FNiagaraFunctionSignature& Sig = OutFunctions.Add_GetRef(DefaultMutableSig);
+		Sig.Name = FNiagaraDataInterfaceArrayImplInternal::Function_ResizeName;
+		Sig.ModuleUsageBitmask = ENiagaraScriptUsageMask::System | ENiagaraScriptUsageMask::Emitter;
+		Sig.Inputs.Emplace(FNiagaraTypeDefinition::GetIntDef(), TEXT("Num"));
+		Sig.Description = NSLOCTEXT("Niagara", "Array_ResizeDesc", "Resizes the array to the specified size, initializing new elements with the default value.");
+	}
+
+	{
+		FNiagaraFunctionSignature& Sig = OutFunctions.Add_GetRef(DefaultMutableSig);
+		Sig.Name = FNiagaraDataInterfaceArrayImplInternal::Function_SetArrayElemName;
+		Sig.Inputs.Emplace(FNiagaraTypeDefinition::GetBoolDef(), TEXT("SkipSet"));
+		Sig.Inputs.Emplace(FNiagaraTypeDefinition::GetIntDef(), TEXT("Index"));
+		Sig.Inputs.Emplace(ValueTypeDef, TEXT("Value"));
+		Sig.Description = NSLOCTEXT("Niagara", "Array_SetArrayElemDesc", "Sets the value at the given zero based index (i.e the first element is 0).");
+		Sig.InputDescriptions.Add(Sig.Inputs[1], NSLOCTEXT("Niagara", "Array_SetArrayElemDesc_SkipSet", "When enabled will not set the array value."));
+	}
+
+	{
+		FNiagaraFunctionSignature& Sig = OutFunctions.Add_GetRef(DefaultMutableSig);
+		Sig.Name = FNiagaraDataInterfaceArrayImplInternal::Function_AddName;
+		Sig.Inputs.Emplace(FNiagaraTypeDefinition::GetBoolDef(), TEXT("SkipAdd"));
+		Sig.Inputs.Emplace(ValueTypeDef, TEXT("Value"));
+		Sig.Description = NSLOCTEXT("Niagara", "Array_AddDesc", "Optionally add a value onto the end of the array.");
+		Sig.InputDescriptions.Add(Sig.Inputs[1], NSLOCTEXT("Niagara", "Array_AddDesc_SkipAdd", "When enabled we will not add an element to the array."));
+	}
+
+	{
+		FNiagaraFunctionSignature& Sig = OutFunctions.Add_GetRef(DefaultMutableSig);
+		Sig.Name = FNiagaraDataInterfaceArrayImplInternal::Function_RemoveLastElemName;
+		Sig.Inputs.Emplace(FNiagaraTypeDefinition::GetBoolDef(), TEXT("SkipRemove"));
+		Sig.Outputs.Emplace(ValueTypeDef, TEXT("Value"));
+		Sig.Outputs.Emplace(FNiagaraTypeDefinition::GetBoolDef(), TEXT("IsValid"));
+		Sig.Description = NSLOCTEXT("Niagara", "Array_RemoveLastElemDesc", "Optionally remove the last element from the array.  Returns the default value if no elements are in the array or you skip the remove.");
+		Sig.InputDescriptions.Add(Sig.Inputs[1], NSLOCTEXT("Niagara", "Array_RemoveLastElemDesc_SkipRemove", "When enabled will not remove a value from the array, the return value will therefore be invalid."));
+		Sig.OutputDescriptions.Add(Sig.Outputs[1], NSLOCTEXT("Niagara", "Array_RemoveLastElemDesc_IsValid", "True if we removed a value from the array, False if no entries or we skipped the remove."));
+	}
+
+	if (bSupportsAtomicOps)
+	{
+		{
+			FNiagaraFunctionSignature& Sig = OutFunctions.Add_GetRef(DefaultMutableSig);
+			Sig.Name = FNiagaraDataInterfaceArrayImplInternal::Function_AtomicAddName;
+			Sig.Inputs.Emplace(FNiagaraTypeDefinition::GetBoolDef(), TEXT("SkipAdd"));
+			Sig.Inputs.Emplace(FNiagaraTypeDefinition::GetIntDef(), TEXT("Index"));
+			Sig.Inputs.Emplace(ValueTypeDef, TEXT("Value"));
+			Sig.Outputs.Emplace(ValueTypeDef, TEXT("PreviousValue"));
+			Sig.Outputs.Emplace(ValueTypeDef, TEXT("CurrentValue"));
+			Sig.SetDescription(NSLOCTEXT("Niagara", "Array_AtomicAddDesc", "Optionally perform an atomic add on the array element."));
+			Sig.SetInputDescription(Sig.Inputs[1], NSLOCTEXT("Niagara", "Array_AtomicAdd_SkipAdd", "When enabled will not perform the add operation, the return values will therefore be invalid."));
+			Sig.SetOutputDescription(Sig.Inputs[0], NSLOCTEXT("Niagara", "Array_AtomicAdd_PrevValue", "The value before the operation was performed."));
+			Sig.SetOutputDescription(Sig.Inputs[1], NSLOCTEXT("Niagara", "Array_AtomicAdd_CurrValue", "The value after the operation was performed."));
+		}
+		{
+			FNiagaraFunctionSignature& Sig = OutFunctions.Add_GetRef(DefaultMutableSig);
+			Sig.Name = FNiagaraDataInterfaceArrayImplInternal::Function_AtomicMinName;
+			Sig.Inputs.Emplace(FNiagaraTypeDefinition::GetBoolDef(), TEXT("SkipMin"));
+			Sig.Inputs.Emplace(FNiagaraTypeDefinition::GetIntDef(), TEXT("Index"));
+			Sig.Inputs.Emplace(ValueTypeDef, TEXT("Value"));
+			Sig.Outputs.Emplace(ValueTypeDef, TEXT("PreviousValue"));
+			Sig.Outputs.Emplace(ValueTypeDef, TEXT("CurrentValue"));
+			Sig.SetDescription(NSLOCTEXT("Niagara", "Array_AtomicMinDesc", "Optionally perform an atomic min on the array element."));
+			Sig.SetInputDescription(Sig.Inputs[1], NSLOCTEXT("Niagara", "Array_AtomicMin_SkipMin", "When enabled will not perform the min operation, the return values will therefore be invalid."));
+			Sig.SetOutputDescription(Sig.Inputs[0], NSLOCTEXT("Niagara", "Array_AtomicMin_PrevValue", "The value before the operation was performed."));
+			Sig.SetOutputDescription(Sig.Inputs[1], NSLOCTEXT("Niagara", "Array_AtomicMin_CurrValue", "The value after the operation was performed."));
+		}
+		{
+			FNiagaraFunctionSignature& Sig = OutFunctions.Add_GetRef(DefaultMutableSig);
+			Sig.Name = FNiagaraDataInterfaceArrayImplInternal::Function_AtomicMaxName;
+			Sig.Inputs.Emplace(FNiagaraTypeDefinition::GetBoolDef(), TEXT("SkipMax"));
+			Sig.Inputs.Emplace(FNiagaraTypeDefinition::GetIntDef(), TEXT("Index"));
+			Sig.Inputs.Emplace(ValueTypeDef, TEXT("Value"));
+			Sig.Outputs.Emplace(ValueTypeDef, TEXT("PreviousValue"));
+			Sig.Outputs.Emplace(ValueTypeDef, TEXT("CurrentValue"));
+			Sig.SetDescription(NSLOCTEXT("Niagara", "Array_AtomicMaxDesc", "Optionally perform an atomic max on the array element."));
+			Sig.SetInputDescription(Sig.Inputs[1], NSLOCTEXT("Niagara", "Array_AtomicMax_SkipMax", "When enabled will not perform the max operation, the return values will therefore be invalid."));
+			Sig.SetOutputDescription(Sig.Inputs[0], NSLOCTEXT("Niagara", "Array_AtomicMax_PrevValue", "The value before the operation was performed."));
+			Sig.SetOutputDescription(Sig.Inputs[1], NSLOCTEXT("Niagara", "Array_AtomicMax_CurrValue", "The value after the operation was performed."));
+		}
+	}
+}
+
+bool FNiagaraDataInterfaceArrayImplInternal::UpgradeFunctionCall(FNiagaraFunctionSignature& FunctionSignature)
 {
 	// Early out, nothing to do here
 	if ( FunctionSignature.FunctionVersion == FFunctionVersion::LatestVersion )
@@ -33,13 +173,13 @@ bool FNiagaraDataInterfaceArrayImplHelper::UpgradeFunctionCall(FNiagaraFunctionS
 	{
 		static const TPair<FName, FName> NodeRenames[] =
 		{
-			MakeTuple(FName("GetNum"),			FNiagaraDataInterfaceArrayImplHelper::Function_LengthName),
-			MakeTuple(FName("GetValue"),		FNiagaraDataInterfaceArrayImplHelper::Function_GetName),
-			MakeTuple(FName("Reset"),			FNiagaraDataInterfaceArrayImplHelper::Function_ClearName),
-			MakeTuple(FName("SetNum"),			FNiagaraDataInterfaceArrayImplHelper::Function_ResizeName),
-			MakeTuple(FName("SetValue"),		FNiagaraDataInterfaceArrayImplHelper::Function_SetArrayElemName),
-			MakeTuple(FName("PushValue"),		FNiagaraDataInterfaceArrayImplHelper::Function_AddName),
-			MakeTuple(FName("PopValue"),		FNiagaraDataInterfaceArrayImplHelper::Function_RemoveLastElemName),
+			MakeTuple(FName("GetNum"),			FNiagaraDataInterfaceArrayImplInternal::Function_LengthName),
+			MakeTuple(FName("GetValue"),		FNiagaraDataInterfaceArrayImplInternal::Function_GetName),
+			MakeTuple(FName("Reset"),			FNiagaraDataInterfaceArrayImplInternal::Function_ClearName),
+			MakeTuple(FName("SetNum"),			FNiagaraDataInterfaceArrayImplInternal::Function_ResizeName),
+			MakeTuple(FName("SetValue"),		FNiagaraDataInterfaceArrayImplInternal::Function_SetArrayElemName),
+			MakeTuple(FName("PushValue"),		FNiagaraDataInterfaceArrayImplInternal::Function_AddName),
+			MakeTuple(FName("PopValue"),		FNiagaraDataInterfaceArrayImplInternal::Function_RemoveLastElemName),
 		};
 
 		for (const auto& Pair : NodeRenames)
@@ -51,7 +191,7 @@ bool FNiagaraDataInterfaceArrayImplHelper::UpgradeFunctionCall(FNiagaraFunctionS
 			}
 		}
 
-		if (FunctionSignature.Name == FNiagaraDataInterfaceArrayImplHelper::Function_SetArrayElemName)
+		if (FunctionSignature.Name == FNiagaraDataInterfaceArrayImplInternal::Function_SetArrayElemName)
 		{
 			FunctionSignature.Inputs.EmplaceAt(1, FNiagaraTypeDefinition::GetBoolDef(), TEXT("SkipSet"));
 		}
@@ -63,7 +203,7 @@ bool FNiagaraDataInterfaceArrayImplHelper::UpgradeFunctionCall(FNiagaraFunctionS
 }
 #endif
 
-bool FNiagaraDataInterfaceArrayImplHelper::IsRWFunction(const FName FunctionName)
+bool FNiagaraDataInterfaceArrayImplInternal::IsRWFunction(const FName FunctionName)
 {
 	static const TSet<FName> RWFunctions =
 	{
@@ -79,3 +219,156 @@ bool FNiagaraDataInterfaceArrayImplHelper::IsRWFunction(const FName FunctionName
 	};
 	return RWFunctions.Contains(FunctionName);
 }
+
+ERHIAccess FNiagaraDataInterfaceArrayImplInternal::GetCountBufferDefaultState()
+{
+	return FNiagaraGPUInstanceCountManager::kCountBufferDefaultState;
+}
+
+FNDIArrayInstanceData_RenderThreadBase::~FNDIArrayInstanceData_RenderThreadBase()
+{
+	if (CountOffset != INDEX_NONE)
+	{
+		ComputeInterface->GetGPUInstanceCounterManager().FreeEntry(CountOffset);
+		CountOffset = INDEX_NONE;
+	}
+
+	ReleaseData();
+}
+
+void FNDIArrayInstanceData_RenderThreadBase::Initialize(FRHICommandListImmediate& RHICmdList, FNiagaraGpuComputeDispatchInterface* InComputeInterface, int32 InDefaultElements, bool bRWGpuArray)
+{
+	ComputeInterface	= InComputeInterface;
+	DefaultElements		= 0;
+	NumElements			= INDEX_NONE;
+	CountOffset			= INDEX_NONE;
+
+	if (bRWGpuArray)
+	{
+		DefaultElements = InDefaultElements;
+		CountOffset = ComputeInterface->GetGPUInstanceCounterManager().AcquireOrAllocateEntry(RHICmdList);
+	}
+}
+
+void FNDIArrayInstanceData_RenderThreadBase::UpdateDataInternal(FRHICommandList& RHICmdList, int32 ArrayNum, int32 NewNumElements, uint32 ElementSize, EPixelFormat PixelFormat)
+{
+	// Do we need to update the backing storage for the buffer
+	if (NewNumElements != NumElements)
+	{
+		// Allocate new data
+		NumElements = NewNumElements;
+		ArrayNumBytes = (NumElements + 1) * ElementSize;	// Note +1 because we store the default value at the end of the buffer
+		INC_MEMORY_STAT_BY(STAT_NiagaraGPUDataInterfaceMemory, ArrayNumBytes);
+
+		const int32 TypeStride = GPixelFormats[PixelFormat].BlockBytes;
+
+		// Create Buffer
+		FRHIResourceCreateInfo CreateInfo(TEXT("NiagaraDataInterfaceArray"));
+		const EBufferUsageFlags BufferUsage = BUF_Static | BUF_ShaderResource | BUF_VertexBuffer | BUF_SourceCopy | (IsReadOnly() ? BUF_None : BUF_UnorderedAccess);
+		const ERHIAccess DefaultAccess = IsReadOnly() ? ERHIAccess::SRVCompute : ERHIAccess::UAVCompute;
+		ArrayBuffer = RHICmdList.CreateBuffer(ArrayNumBytes, BufferUsage, TypeStride, DefaultAccess, CreateInfo);
+
+		ArraySRV = RHICmdList.CreateShaderResourceView(
+			ArrayBuffer, 
+			FRHIViewDesc::CreateBufferSRV()
+				.SetType(FRHIViewDesc::EBufferType::Typed)
+				.SetFormat(PixelFormat));
+
+		if ( !IsReadOnly() )
+		{
+			ArrayUAV = RHICmdList.CreateUnorderedAccessView(
+				ArrayBuffer, 
+				FRHIViewDesc::CreateBufferUAV()
+					.SetType(FRHIViewDesc::EBufferType::Typed)
+					.SetFormat(PixelFormat));
+		}
+	}
+
+	// Adjust counter value
+	if (CountOffset != INDEX_NONE)
+	{
+		//-OPT: We could push this into the count manager and batch set as part of the clear process
+		const FNiagaraGPUInstanceCountManager& CounterManager = ComputeInterface->GetGPUInstanceCounterManager();
+		const FRWBuffer& CountBuffer = CounterManager.GetInstanceCountBuffer();
+
+		const TPair<uint32, uint32> DataToClear(CountOffset, ArrayNum);
+		RHICmdList.Transition(FRHITransitionInfo(CountBuffer.UAV, FNiagaraGPUInstanceCountManager::kCountBufferDefaultState, ERHIAccess::UAVCompute));
+		NiagaraClearCounts::ClearCountsUInt(RHICmdList, CountBuffer.UAV, MakeArrayView(&DataToClear, 1));
+		RHICmdList.Transition(FRHITransitionInfo(CountBuffer.UAV, ERHIAccess::UAVCompute, FNiagaraGPUInstanceCountManager::kCountBufferDefaultState));
+	}
+}
+
+void FNDIArrayInstanceData_RenderThreadBase::ReleaseData()
+{
+	DEC_MEMORY_STAT_BY(STAT_NiagaraGPUDataInterfaceMemory, ArrayNumBytes);
+	ArrayBuffer.SafeRelease();
+	ArrayUAV.SafeRelease();
+	ArraySRV.SafeRelease();
+	ArrayNumBytes = 0;
+}
+
+void FNDIArrayInstanceData_RenderThreadBase::SimCacheWriteFrame(FRHICommandListImmediate& RHICmdList, UNDIArraySimCacheData* CacheData, int32 FrameIndex, int32 ArrayTypeSize, void(*CopyGpuToCpuMemory)(void*, const void*, int32)) const
+{
+
+	if (ArrayNumBytes == 0)
+	{
+		return;
+	}
+
+	if (CacheData->GpuFrameData.Num() <= FrameIndex)
+	{
+		CacheData->GpuFrameData.AddDefaulted(FrameIndex + 1 - CacheData->GpuFrameData.Num());
+	}
+	FNDIArraySimCacheDataFrame& FrameData = CacheData->GpuFrameData[FrameIndex];
+
+	TArray<FNiagaraGpuReadbackManager::FBufferRequest, TInlineAllocator<2>> BufferRequests;
+	TArray<FRHITransitionInfo, TInlineAllocator<2>> TransitionsBefore;
+	TArray<FRHITransitionInfo, TInlineAllocator<2>> TransitionsAfter;
+
+	const ERHIAccess DefaultAccess = IsReadOnly() ? ERHIAccess::SRVCompute : ERHIAccess::UAVCompute;
+
+	BufferRequests.Emplace(ArrayBuffer, 0, ArrayNumBytes);
+	TransitionsBefore.Emplace(ArrayBuffer, DefaultAccess, ERHIAccess::CopySrc);
+	TransitionsAfter.Emplace(ArrayBuffer, ERHIAccess::CopySrc, DefaultAccess);
+
+	if (IsReadOnly())
+	{
+		FrameData.NumElements = NumElements;
+	}
+	else
+	{
+		const FNiagaraGPUInstanceCountManager& CountManager = ComputeInterface->GetGPUInstanceCounterManager();
+		BufferRequests.Emplace(CountManager.GetInstanceCountBuffer().Buffer, uint32(CountOffset * sizeof(uint32)), sizeof(uint32));
+		TransitionsBefore.Emplace(CountManager.GetInstanceCountBuffer().UAV, FNiagaraGPUInstanceCountManager::kCountBufferDefaultState, ERHIAccess::CopySrc);
+		TransitionsAfter.Emplace(CountManager.GetInstanceCountBuffer().UAV, ERHIAccess::CopySrc, FNiagaraGPUInstanceCountManager::kCountBufferDefaultState);
+	}
+
+	FNiagaraGpuReadbackManager* ReadbackManager = ComputeInterface->GetGpuReadbackManager();
+	RHICmdList.Transition(TransitionsBefore);
+	ReadbackManager->EnqueueReadbacks(
+		RHICmdList,
+		BufferRequests,
+		[CacheData, &FrameData, bReadOnly=IsReadOnly(), ArrayTypeSize, CopyGpuToCpuMemory](TConstArrayView<TPair<void*, uint32>> ReadbackData)
+		{
+			if (!bReadOnly)
+			{
+				FrameData.NumElements = *reinterpret_cast<const uint32*>(ReadbackData[1].Key);
+			}
+			if ( FrameData.NumElements > 0 )
+			{
+				TArray<uint8> ArrayData;
+				ArrayData.AddUninitialized(FrameData.NumElements * ArrayTypeSize);
+				CopyGpuToCpuMemory(ArrayData.GetData(), ReadbackData[0].Key, FrameData.NumElements);
+
+				FrameData.DataOffset = CacheData->FindOrAddData(
+					MakeArrayView(
+						ArrayData.GetData(),
+						ArrayData.Num() * ArrayTypeSize
+					)
+				);
+			}
+		}
+	);
+	RHICmdList.Transition(TransitionsAfter);
+	ReadbackManager->WaitCompletion(RHICmdList);
+}
