- Add Object Reader source mode option to bind directly to the attach component (i.e. Niagara Component)
- This allows us to get at propties like the Cached Max Draw Distance without requiring BP hook up

#tests Use the new mode to read CachedMaxDrawDistance from the attached component

diff --git a/Engine/Plugins/FX/Niagara/Source/Niagara/Private/DataInterface/NiagaraDataInterfaceUObjectPropertyReader.cpp b/Engine/Plugins/FX/Niagara/Source/Niagara/Private/DataInterface/NiagaraDataInterfaceUObjectPropertyReader.cpp
--- a/Engine/Plugins/FX/Niagara/Source/Niagara/Private/DataInterface/NiagaraDataInterfaceUObjectPropertyReader.cpp
+++ b/Engine/Plugins/FX/Niagara/Source/Niagara/Private/DataInterface/NiagaraDataInterfaceUObjectPropertyReader.cpp
@@ -784,16 +784,32 @@ bool UNiagaraDataInterfaceUObjectPropertyReader::PerInstanceTick(void* PerInstan
 
 	// Resolve our source object
 	UObject* ObjectBinding = nullptr;
-	if (SourceMode != ENDIObjectPropertyReaderSourceMode::AttachParentActor)
 	{
-		ObjectBinding = InstanceData_GT->UObjectBinding.GetValue();
-		ObjectBinding = ObjectBinding ? ObjectBinding : SourceActor.Get();
-	}
-	if (SourceMode != ENDIObjectPropertyReaderSourceMode::Binding && ObjectBinding == nullptr)
-	{
-		if (USceneComponent* AttachComponent = SystemInstance->GetAttachComponent())
+		// Try Object Binding?
+		const bool bTryObjectBinding =
+			SourceMode == ENDIObjectPropertyReaderSourceMode::Binding ||
+			SourceMode == ENDIObjectPropertyReaderSourceMode::BindingThenAttachParentActor ||
+			SourceMode == ENDIObjectPropertyReaderSourceMode::BindingThenAttachParent;
+		if (bTryObjectBinding)
+		{
+			ObjectBinding = InstanceData_GT->UObjectBinding.GetValue();
+			ObjectBinding = ObjectBinding ? ObjectBinding : SourceActor.Get();
+		}
+
+		// If that failed do we look at the parent / parent actor?
+		if (ObjectBinding == nullptr)
 		{
-			ObjectBinding = AttachComponent->GetAttachParentActor();
+			if (SourceMode == ENDIObjectPropertyReaderSourceMode::AttachParent || SourceMode == ENDIObjectPropertyReaderSourceMode::BindingThenAttachParent )
+			{
+				ObjectBinding = SystemInstance->GetAttachComponent();
+			}
+			else if (SourceMode == ENDIObjectPropertyReaderSourceMode::AttachParentActor || SourceMode == ENDIObjectPropertyReaderSourceMode::BindingThenAttachParentActor)
+			{
+				if (USceneComponent* AttachComponent = SystemInstance->GetAttachComponent())
+				{
+					ObjectBinding = AttachComponent->GetAttachParentActor();
+				}
+			}
 		}
 	}
 
diff --git a/Engine/Plugins/FX/Niagara/Source/Niagara/Private/DataInterface/NiagaraDataInterfaceUObjectPropertyReader.h b/Engine/Plugins/FX/Niagara/Source/Niagara/Private/DataInterface/NiagaraDataInterfaceUObjectPropertyReader.h
--- a/Engine/Plugins/FX/Niagara/Source/Niagara/Private/DataInterface/NiagaraDataInterfaceUObjectPropertyReader.h
+++ b/Engine/Plugins/FX/Niagara/Source/Niagara/Private/DataInterface/NiagaraDataInterfaceUObjectPropertyReader.h
@@ -13,10 +13,14 @@ enum class ENDIObjectPropertyReaderSourceMode : uint8
 {
 	/** Source object is found in the order of object binding, source actor. */
 	Binding,
-	/** Source object is found use the attached parent actor. */
+	/** Source object is found use the attached parent actor, i.e. the owner components actor. */
 	AttachParentActor,
 	/** Source object is found in the order of object binding, source actor, then attached parent actor. */
 	BindingThenAttachParentActor,
+	/** Source object is the attach parent, i.e. the owner component. */
+	AttachParent,
+	/** Source object is found in the order of object binding, source actor, then attached parent. */
+	BindingThenAttachParent
 };
 
 
@@ -66,14 +70,14 @@ public:
 	ENDIObjectPropertyReaderSourceMode SourceMode = ENDIObjectPropertyReaderSourceMode::Binding;
 
 	/** User parameter Object binding to read properties from. */
-	UPROPERTY(EditAnywhere, Category = "UObjectPropertyReader", meta = (EditCondition = "SourceMode != ENDIObjectPropertyReaderSourceMode::AttachParentActor", EditConditionHides))
+	UPROPERTY(EditAnywhere, Category = "UObjectPropertyReader", meta = (EditCondition = "SourceMode != ENDIObjectPropertyReaderSourceMode::AttachParent && SourceMode != ENDIObjectPropertyReaderSourceMode::AttachParentActor", EditConditionHides))
 	FNiagaraUserParameterBinding UObjectParameterBinding;
 
 	UPROPERTY(EditAnywhere, Category = "UObjectPropertyReader")
 	TArray<FNiagaraUObjectPropertyReaderRemap> PropertyRemap;
 
 	/** Optional source actor to use, if the user parameter binding is valid this will be ignored. */
-	UPROPERTY(EditAnywhere, Category = "UObjectPropertyReader", meta = (EditCondition = "SourceMode != ENDIObjectPropertyReaderSourceMode::AttachParentActor", EditConditionHides))
+	UPROPERTY(EditAnywhere, Category = "UObjectPropertyReader", meta = (EditCondition = "SourceMode != ENDIObjectPropertyReaderSourceMode::AttachParent && SourceMode != ENDIObjectPropertyReaderSourceMode::AttachParentActor", EditConditionHides))
 	TSoftObjectPtr<AActor> SourceActor;
 
 	/**
