- Make sure bounds calculator is created when using a mesh array binding

#rnx
#tests Create a mesh renderer + array, then removed meshes in mesh renderer and observed bounds

diff --git a/Engine/Plugins/FX/Niagara/Source/Niagara/Private/NiagaraMeshRendererProperties.cpp b/Engine/Plugins/FX/Niagara/Source/Niagara/Private/NiagaraMeshRendererProperties.cpp
--- a/Engine/Plugins/FX/Niagara/Source/Niagara/Private/NiagaraMeshRendererProperties.cpp
+++ b/Engine/Plugins/FX/Niagara/Source/Niagara/Private/NiagaraMeshRendererProperties.cpp
@@ -421,8 +421,6 @@ FNiagaraBoundsCalculator* UNiagaraMeshRendererProperties::CreateBoundsCalculator
 		bLocalSpace = EmitterData->bLocalSpace;
 	}
 
-	//-TODO: Bounds are incorrect when binding to a mesh inside the mesh properties or a mesh array
-
 	for (const auto& MeshProperties : Meshes)
 	{
 		if (MeshProperties.Mesh)
@@ -458,6 +456,12 @@ FNiagaraBoundsCalculator* UNiagaraMeshRendererProperties::CreateBoundsCalculator
 		}
 	}
 
+	//-TODO: This is incorrect when using a binding as we don't know the meshes until later
+	if (!LocalBounds.IsValid && MeshesBinding.ResolvedParameter.IsValid())
+	{
+		LocalBounds = FBox(FVector::OneVector, FVector::OneVector);
+	}
+
 	if (LocalBounds.IsValid)
 	{
 		// Take the bounding center into account with the extents, as it may not be at the origin
