- Fix capture velocity attribute when we have no position attributes to extrapolate from
- Check the capture mode rather than using ExplicitCaptureAttributes == 0 when building the capture attribute list

#rnx
#tests Capture cache in 3 different modes and check with explicit position + velocity extrapolation enabled/disabled

diff --git a/Engine/Plugins/FX/Niagara/Source/Niagara/Private/NiagaraSimCacheHelper.h b/Engine/Plugins/FX/Niagara/Source/Niagara/Private/NiagaraSimCacheHelper.h
--- a/Engine/Plugins/FX/Niagara/Source/Niagara/Private/NiagaraSimCacheHelper.h
+++ b/Engine/Plugins/FX/Niagara/Source/Niagara/Private/NiagaraSimCacheHelper.h
@@ -133,22 +133,25 @@ struct FNiagaraSimCacheHelper
 		// Determine the components to cache
 		int32 TotalCacheComponents = 0;
 		TArray<int32> CacheToDataSetVariables;
+		bool bHasPositionAttribute = false;
 
 		CacheToDataSetVariables.Reserve(CompiledData.Variables.Num());
 		for (int32 i = 0; i < CompiledData.Variables.Num(); ++i)
 		{
 			const FNiagaraVariableLayoutInfo& DataSetVariableLayout = CompiledData.VariableLayouts[i];
 			const FNiagaraVariableBase& DataSetVariable = CompiledData.Variables[i];
-			if (ExplicitCaptureAttributes.Num() == 0 || ExplicitCaptureAttributes.Contains(DataSetVariable.GetName()))
+			if (CreateParameters.AttributeCaptureMode == ENiagaraSimCacheAttributeCaptureMode::All || ExplicitCaptureAttributes.Contains(DataSetVariable.GetName()))
 			{
 				CacheToDataSetVariables.Add(i);
 				TotalCacheComponents += DataSetVariableLayout.GetNumFloatComponents() + DataSetVariableLayout.GetNumHalfComponents() + DataSetVariableLayout.GetNumInt32Components();
+				bHasPositionAttribute |= DataSetVariable.GetType() == FNiagaraTypeDefinition::GetPositionDef();
 			}
 		}
 
 		// We need to preserve the velocity attribute if we want to use velocity based extrapolation of positions
+		// Note:  Only preserve the velocity attribute if it will be used (i.e. we have a position attribute we need to extrapolate)
 		const FNiagaraVariableBase VelocityVariable(FNiagaraTypeDefinition::GetVec3Def(), "Velocity");
-		if (CreateParameters.bAllowVelocityExtrapolation)
+		if (CreateParameters.bAllowVelocityExtrapolation && bHasPositionAttribute)
 		{
 			const int32 VelocityAttributeIndex = CompiledData.Variables.IndexOfByKey(VelocityVariable);
 			if (VelocityAttributeIndex != INDEX_NONE)
@@ -455,13 +458,6 @@ struct FNiagaraSimCacheHelper
 					}
 				}
 			}
-
-			if (ExplicitCaptureAttributes.IsEmpty())
-			{
-				// The attributes to capture only include system/emitter attributes or particle attributes from another emitter.
-				// We add a bogus filter attribute to build a minimal cache layout (as having no filter would record everything).
-				ExplicitCaptureAttributes.Add(FName("__FILTER.NONE__"));
-			}
 		}
 
 		BuildCacheLayout(CreateParameters, CacheLayout, *ParticleDataSet, EmitterHandle.GetName(), MoveTemp(RebaseVariableNames), MoveTemp(InterpVariableNames), ExplicitCaptureAttributes);
