Fix stale/incorrect reporting of culled state to debug hud.

diff --git a/Engine/Plugins/FX/Niagara/Source/Niagara/Private/NiagaraScalabilityManager.cpp b/Engine/Plugins/FX/Niagara/Source/Niagara/Private/NiagaraScalabilityManager.cpp
--- a/Engine/Plugins/FX/Niagara/Source/Niagara/Private/NiagaraScalabilityManager.cpp
+++ b/Engine/Plugins/FX/Niagara/Source/Niagara/Private/NiagaraScalabilityManager.cpp
@@ -138,6 +138,10 @@ void FNiagaraScalabilityManager::Register(UNiagaraComponent* Component)
 
 	Component->ScalabilityManagerHandle = ManagedComponents.Add(Component);
 	State.AddDefaulted();
+	
+#if WITH_NIAGARA_DEBUGGER
+	Component->DebugCachedScalabilityState = State.Last();
+#endif	
 
 	if (HasPendingUpdates())
 	{
