Add a Capture on Keypress feature to the Niagara Outliner.

Misc:
- Utility to allow adding messages to the DebugHud via the world manager.
- Reduced default capture delay to 0. 60 was just annoying in 99% of cases.

diff --git a/Engine/Plugins/FX/Niagara/Source/Niagara/Private/NiagaraWorldManager.cpp b/Engine/Plugins/FX/Niagara/Source/Niagara/Private/NiagaraWorldManager.cpp
--- a/Engine/Plugins/FX/Niagara/Source/Niagara/Private/NiagaraWorldManager.cpp
+++ b/Engine/Plugins/FX/Niagara/Source/Niagara/Private/NiagaraWorldManager.cpp
@@ -2217,6 +2217,16 @@ void FNiagaraWorldManager::PrimePool(UNiagaraSystem* System)
 	}
 }
 
+#if WITH_NIAGARA_DEBUGGER
+void FNiagaraWorldManager::AddDebugHUDMessage(FName Key, const FNiagaraDebugMessage& Message)
+{
+	if(NiagaraDebugHud)
+	{
+		NiagaraDebugHud->AddMessage(Key, Message);
+	}
+}
+#endif
+
 void FNiagaraWorldManager::FlushComputeAndDeferredQueues(bool bWaitForGPU)
 {
 	check(IsInGameThread());
diff --git a/Engine/Plugins/FX/Niagara/Source/Niagara/Public/NiagaraDebuggerCommon.h b/Engine/Plugins/FX/Niagara/Source/Niagara/Public/NiagaraDebuggerCommon.h
--- a/Engine/Plugins/FX/Niagara/Source/Niagara/Public/NiagaraDebuggerCommon.h
+++ b/Engine/Plugins/FX/Niagara/Source/Niagara/Public/NiagaraDebuggerCommon.h
@@ -835,7 +835,7 @@ struct FNiagaraOutlinerCaptureSettings
 
 	/** How many frames to delay capture. If gathering performance data, this is how many frames will be collected. */
 	UPROPERTY(EditAnywhere, Config, Category="Settings")
-	uint32 CaptureDelayFrames = 60;
+	uint32 CaptureDelayFrames = 0;
 	
 	UPROPERTY(EditAnywhere, Config, Category = "Settings")
 	bool bGatherPerfData = true;
@@ -880,7 +880,7 @@ struct FNiagaraSystemSimCacheCaptureRequest
 
 	/** How many frames to delay capture. */
 	UPROPERTY(EditAnywhere, Config, Category="Settings")
-	uint32 CaptureDelayFrames = 60;
+	uint32 CaptureDelayFrames = 0;
 
 	/** How many frames to capture. */
 	UPROPERTY(EditAnywhere, Config, Category = "Settings")
diff --git a/Engine/Plugins/FX/Niagara/Source/Niagara/Public/NiagaraWorldManager.h b/Engine/Plugins/FX/Niagara/Source/Niagara/Public/NiagaraWorldManager.h
--- a/Engine/Plugins/FX/Niagara/Source/Niagara/Public/NiagaraWorldManager.h
+++ b/Engine/Plugins/FX/Niagara/Source/Niagara/Public/NiagaraWorldManager.h
@@ -231,6 +231,7 @@ public:
 
 #if WITH_NIAGARA_DEBUGGER
 	class FNiagaraDebugHud* GetNiagaraDebugHud() { return NiagaraDebugHud.Get(); }
+	NIAGARA_API void AddDebugHUDMessage(FName Key, const FNiagaraDebugMessage& Message);
 #endif
 
 	class FNiagaraDeferredMethodQueue& GetDeferredMethodQueue() { return DeferredMethods; }
diff --git a/Engine/Plugins/FX/Niagara/Source/NiagaraEditor/Private/NiagaraDebugger.cpp b/Engine/Plugins/FX/Niagara/Source/NiagaraEditor/Private/NiagaraDebugger.cpp
--- a/Engine/Plugins/FX/Niagara/Source/NiagaraEditor/Private/NiagaraDebugger.cpp
+++ b/Engine/Plugins/FX/Niagara/Source/NiagaraEditor/Private/NiagaraDebugger.cpp
@@ -12,6 +12,8 @@
 #include "Serialization/MemoryReader.h"
 #include "Serialization/ObjectAndNameAsStringProxyArchive.h"
 
+#include "Framework/Application/SlateApplication.h"
+
 #if WITH_NIAGARA_DEBUGGER
 
 #define LOCTEXT_NAMESPACE "NiagaraDebugger"
@@ -58,6 +60,8 @@ void FNiagaraDebugger::Init()
 
 	GetMutableDefault<UNiagaraDebugHUDSettings>()->OnChangedDelegate.AddSP(this, &FNiagaraDebugger::UpdateDebugHUDSettings);
 	GetMutableDefault<UNiagaraOutliner>()->OnChangedDelegate.AddSP(this, &FNiagaraDebugger::TriggerOutlinerCapture);
+
+	FSlateApplication::Get().RegisterInputPreProcessor(AsShared());
 }
 
 void FNiagaraDebugger::ExecConsoleCommand(const TCHAR* Cmd, bool bRequiresWorld)
@@ -139,6 +143,23 @@ void FNiagaraDebugger::RequestUpdatedClientInfo()
 #endif//WITH_UNREAL_TARGET_DEVELOPER_TOOLS
 }
 
+const FName NiagaraDebuggerCaptureOutlinerOnKeyPressMessageKey(TEXT("NiagaraDebuggerCaptureOutlinerOnKeyPressMessageKey"));
+void FNiagaraDebugger::ToggleOutlinerCaptureOnKeyPress()
+{
+	bTriggerOutlinerCaptureOnKeyPress = !bTriggerOutlinerCaptureOnKeyPress;
+
+	const FString HudMessage = FString::Printf(TEXT("Press 'C' To Capture Niagara Outliner..."));//TODOL Customizable Key?
+
+	if(bTriggerOutlinerCaptureOnKeyPress)
+	{
+		FNiagaraWorldManager::ForAllWorldManagers([&HudMessage](FNiagaraWorldManager& WorldMan) { WorldMan.AddDebugHUDMessage(NiagaraDebuggerCaptureOutlinerOnKeyPressMessageKey, FNiagaraDebugMessage(ENiagaraDebugMessageType::Info, HudMessage, FLT_MAX)); });
+	}
+	else
+	{
+		FNiagaraWorldManager::ForAllWorldManagers([&HudMessage](FNiagaraWorldManager& WorldMan) { WorldMan.AddDebugHUDMessage(NiagaraDebuggerCaptureOutlinerOnKeyPressMessageKey, FNiagaraDebugMessage(ENiagaraDebugMessageType::Info, HudMessage, -1.0f)); });
+	}
+}
+
 void FNiagaraDebugger::TriggerOutlinerCapture()
 {
 	//Send the current settings as a message to all connected clients.
@@ -446,6 +467,25 @@ void FNiagaraDebugger::HandleSimCacheCaptureReply(const FNiagaraSystemSimCacheCa
 	}
 }
 
+//IInputProcessor Interface
+
+bool FNiagaraDebugger::HandleKeyDownEvent(FSlateApplication& SlateApp, const FKeyEvent& InKeyEvent)
+{
+	if (IsCapturingOutlinerOnKeyPress() && InKeyEvent.GetKey() == EKeys::C)
+	{
+		if (UNiagaraOutliner* Outliner = GetOutliner())
+		{
+			Outliner->CaptureSettings.bTriggerCapture = true;
+			Outliner->OnChanged();
+		}
+		ToggleOutlinerCaptureOnKeyPress();
+	}
+	return IInputProcessor::HandleKeyDownEvent(SlateApp, InKeyEvent);;
+}
+
+//IInputProcessor Interface END
+
+
 #undef LOCTEXT_NAMESPACE
 
 
diff --git a/Engine/Plugins/FX/Niagara/Source/NiagaraEditor/Private/NiagaraDebugger.h b/Engine/Plugins/FX/Niagara/Source/NiagaraEditor/Private/NiagaraDebugger.h
--- a/Engine/Plugins/FX/Niagara/Source/NiagaraEditor/Private/NiagaraDebugger.h
+++ b/Engine/Plugins/FX/Niagara/Source/NiagaraEditor/Private/NiagaraDebugger.h
@@ -6,6 +6,7 @@
 #include "NiagaraDebuggerCommon.h"
 #include "NiagaraOutliner.h"
 #include "NiagaraWorldManager.h"
+#include "Framework/Application/IInputProcessor.h"
 
 class ISessionManager;
 class FMessageEndpoint;
@@ -17,7 +18,7 @@ DECLARE_LOG_CATEGORY_EXTERN(LogNiagaraDebugger, Log, All);
 #if WITH_NIAGARA_DEBUGGER
 
 
-class FNiagaraDebugger : public TSharedFromThis<FNiagaraDebugger>
+class FNiagaraDebugger : public TSharedFromThis<FNiagaraDebugger>, public IInputProcessor
 {
 public:
 	struct FClientInfo
@@ -56,6 +57,9 @@ public:
 	void TriggerOutlinerCapture();
 	void TriggerSimCacheCapture(FName ComponentName, int32 CaptureDelay, int32 CaptureFrames);
 
+	void ToggleOutlinerCaptureOnKeyPress();
+	bool IsCapturingOutlinerOnKeyPress(){return bTriggerOutlinerCaptureOnKeyPress;}
+
 	void RequestUpdatedClientInfo();
 
 	FOnConnectionMade& GetOnConnectionMade() { return OnConnectionMadeDelegate; }
@@ -70,6 +74,12 @@ public:
 
 	const FNiagaraSimpleClientInfo& GetSimpleClientInfo() { return SimpleClientInfo; }
 
+	//IInputProcessor Interface
+	virtual void Tick(const float DeltaTime, FSlateApplication& SlateApp, TSharedRef<ICursor> Cursor){};
+	virtual bool HandleKeyDownEvent(FSlateApplication& SlateApp, const FKeyEvent& InKeyEvent)override;
+	virtual const TCHAR* GetDebugName() const { return TEXT("NiagaraDebugger"); }
+	//IInputProcessor Interface END
+
 protected:
 
 	// Handles a connection accepted message and finalizes establishing the connection to a debugger client.
@@ -117,6 +127,8 @@ protected:
 
 	/** Some basic info on the connected client. */
 	FNiagaraSimpleClientInfo SimpleClientInfo;
+
+	bool bTriggerOutlinerCaptureOnKeyPress = false;
 };
 
 template<typename TAction>
diff --git a/Engine/Plugins/FX/Niagara/Source/NiagaraEditor/Private/Widgets/SNiagaraDebugger.cpp b/Engine/Plugins/FX/Niagara/Source/NiagaraEditor/Private/Widgets/SNiagaraDebugger.cpp
--- a/Engine/Plugins/FX/Niagara/Source/NiagaraEditor/Private/Widgets/SNiagaraDebugger.cpp
+++ b/Engine/Plugins/FX/Niagara/Source/NiagaraEditor/Private/Widgets/SNiagaraDebugger.cpp
@@ -332,6 +332,23 @@ namespace NiagaraOutlinerTab
 				);
 			}
 
+			// Capture on key button
+			{
+				ToolbarBuilder.AddToolBarButton(
+					FUIAction(FExecuteAction::CreateLambda([Debugger]()
+						{
+							Debugger->ToggleOutlinerCaptureOnKeyPress();
+						}),
+						FCanExecuteAction(),
+						FIsActionChecked::CreateLambda([Debugger]() { return Debugger->IsCapturingOutlinerOnKeyPress(); })),
+					NAME_None,
+					LOCTEXT("Capture", "Capture (Press 'C')"),
+					LOCTEXT("CaptureTooltip", "Triggers a capture of outliner info from the connected session on a key press."),
+					FSlateIcon(FNiagaraEditorStyle::Get().GetStyleSetName(), "NiagaraEditor.Debugger.Outliner.Capture"),
+					EUserInterfaceActionType::ToggleButton
+				);
+			}
+
 			// Gather Perf Toggle
 			{
 				ToolbarBuilder.AddToolBarButton(
