- Handle modifications to the meshes referenced by the Mesh Array DI
- This will recreate the render state for the associated components to ensure they don't point to invalid data

#rnx
#tests Modify meshes (add materials / reimport / enable nanite etc) on ones being used in an active mesh renderer that is bound to a mesh array

diff --git a/Engine/Plugins/FX/Niagara/Source/Niagara/Classes/NiagaraDataInterfaceArrayImpl.h b/Engine/Plugins/FX/Niagara/Source/Niagara/Classes/NiagaraDataInterfaceArrayImpl.h
--- a/Engine/Plugins/FX/Niagara/Source/Niagara/Classes/NiagaraDataInterfaceArrayImpl.h
+++ b/Engine/Plugins/FX/Niagara/Source/Niagara/Classes/NiagaraDataInterfaceArrayImpl.h
@@ -545,16 +545,22 @@ struct FNDIArrayProxyImpl : public INDIArrayProxyBase
 	void SetArrayDataAndRecreateRenderState(TConstArrayView<TFromArrayType> InArrayData)
 	{
 		SetArrayData(InArrayData);
+		RecreateRenderState();
+	}
 
-		for (auto It=PerInstanceData_GameThread.CreateConstIterator(); It; ++It)
+	void RecreateRenderState()
+	{
+		for (auto It = PerInstanceData_GameThread.CreateConstIterator(); It; ++It)
 		{
 			FNDIArrayInstanceData_GameThread<TArrayType>* InstanceData = It.Value();
 			FNiagaraSystemInstance* SystemInstance = InstanceData->OwnerInstance;
 			USceneComponent* SceneComponent = SystemInstance ? SystemInstance->GetAttachComponent() : nullptr;
-			if (SceneComponent)
+			if (SceneComponent && SceneComponent->IsRenderStateCreated() && !SceneComponent->IsRenderStateRecreating())
 			{
 				SceneComponent->RecreateRenderState_Concurrent();
 			}
+			// Ideally we would replace this with a mark for recreate but things like renderers would need know how to handle changes
+			//SceneComponent->MarkForNeededEndOfFrameRecreate();
 		}
 	}
 
diff --git a/Engine/Plugins/FX/Niagara/Source/Niagara/Internal/DataInterface/NiagaraDataInterfaceArrayMesh.h b/Engine/Plugins/FX/Niagara/Source/Niagara/Internal/DataInterface/NiagaraDataInterfaceArrayMesh.h
--- a/Engine/Plugins/FX/Niagara/Source/Niagara/Internal/DataInterface/NiagaraDataInterfaceArrayMesh.h
+++ b/Engine/Plugins/FX/Niagara/Source/Niagara/Internal/DataInterface/NiagaraDataInterfaceArrayMesh.h
@@ -20,6 +20,24 @@ public:
 
 	NDIARRAY_GENERATE_BODY(UNiagaraDataInterfaceArrayMesh, FNiagaraMeshRendererMeshPropertiesBase, MeshData)
 
+	//UObject Interface Begin
+#if WITH_EDITORONLY_DATA
+	NIAGARA_API virtual void PostLoad() override;
+	NIAGARA_API virtual void BeginDestroy() override;
+	NIAGARA_API virtual void PreEditChange(class FProperty* PropertyThatWillChange) override;
+	NIAGARA_API virtual void PostEditChangeProperty(struct FPropertyChangedEvent& PropertyChangedEvent) override;
+#endif// WITH_EDITORONLY_DATA
+	//UObject Interface End
+
+#if WITH_EDITORONLY_DATA
+	void OnMeshChanged();
+	void OnMeshPostBuild(UStaticMesh*);
+	void OnAssetReimported(UObject* Object);
+
+	void AddMeshChangedDelegates();
+	void RemoveMeshChangedDelegates();
+#endif //WITH_EDITORONLY_DATA
+
 	//INiagaraRenderableMeshArrayInterface Interface Begin
 	virtual void ForEachMesh(FNiagaraSystemInstance* SystemInstance, TFunction<void(int32)> NumMeshesDelegate, TFunction<void(const FNiagaraMeshRendererMeshProperties&)> IterateDelegate) const override;
 	//INiagaraRenderableMeshArrayInterface Interface End
diff --git a/Engine/Plugins/FX/Niagara/Source/Niagara/Private/DataInterface/NiagaraDataInterfaceArrayMesh.cpp b/Engine/Plugins/FX/Niagara/Source/Niagara/Private/DataInterface/NiagaraDataInterfaceArrayMesh.cpp
--- a/Engine/Plugins/FX/Niagara/Source/Niagara/Private/DataInterface/NiagaraDataInterfaceArrayMesh.cpp
+++ b/Engine/Plugins/FX/Niagara/Source/Niagara/Private/DataInterface/NiagaraDataInterfaceArrayMesh.cpp
@@ -7,6 +7,11 @@
 #include "NiagaraSystemInstance.h"
 
 #include "Engine/StaticMesh.h"
+#if WITH_EDITORONLY_DATA
+#include "Editor/EditorEngine.h"
+#include "Subsystems/ImportSubsystem.h"
+#include "Editor.h"
+#endif //WITH_EDITORONLY_DATA
 
 #include UE_INLINE_GENERATED_CPP_BY_NAME(NiagaraDataInterfaceArrayMesh)
 
@@ -48,6 +53,106 @@ struct FNDIArrayImplHelper<FNiagaraMeshRendererMeshPropertiesBase> : public FNDI
 
 NDIARRAY_GENERATE_IMPL(UNiagaraDataInterfaceArrayMesh, FNiagaraMeshRendererMeshPropertiesBase, MeshData)
 
+#if WITH_EDITORONLY_DATA
+void UNiagaraDataInterfaceArrayMesh::PostLoad()
+{
+	Super::PostLoad();
+	AddMeshChangedDelegates();
+}
+
+void UNiagaraDataInterfaceArrayMesh::BeginDestroy()
+{
+	RemoveMeshChangedDelegates();
+	Super::BeginDestroy();
+}
+
+void UNiagaraDataInterfaceArrayMesh::PreEditChange(class FProperty* PropertyThatWillChange)
+{
+	RemoveMeshChangedDelegates();
+	Super::PreEditChange(PropertyThatWillChange);
+}
+
+void UNiagaraDataInterfaceArrayMesh::PostEditChangeProperty(struct FPropertyChangedEvent& PropertyChangedEvent)
+{
+	Super::PostEditChangeProperty(PropertyChangedEvent);
+	AddMeshChangedDelegates();
+}
+
+void UNiagaraDataInterfaceArrayMesh::OnMeshChanged()
+{
+	FProxyType* ArrayProxy = static_cast<FProxyType*>(GetProxy());
+	ArrayProxy->RecreateRenderState();
+}
+
+void UNiagaraDataInterfaceArrayMesh::OnMeshPostBuild(UStaticMesh*)
+{
+	OnMeshChanged();
+}
+
+void UNiagaraDataInterfaceArrayMesh::OnAssetReimported(UObject* Object)
+{
+	for (const FNiagaraMeshRendererMeshPropertiesBase& MeshProperties : MeshData)
+	{
+		if (MeshProperties.Mesh == Object)
+		{
+			OnMeshChanged();
+			break;
+		}
+	}
+}
+
+void UNiagaraDataInterfaceArrayMesh::AddMeshChangedDelegates()
+{
+	if (!GIsEditor || HasAnyFlags(RF_ClassDefaultObject))
+	{
+		return;
+	}
+
+	bool bAnyAdded = false;
+	for (const FNiagaraMeshRendererMeshPropertiesBase& MeshProperties : MeshData)
+	{
+		if (MeshProperties.Mesh)
+		{
+			MeshProperties.Mesh->GetOnMeshChanged().AddUObject(this, &UNiagaraDataInterfaceArrayMesh::OnMeshChanged);
+			MeshProperties.Mesh->OnPreMeshBuild().AddUObject(this, &UNiagaraDataInterfaceArrayMesh::OnMeshPostBuild);
+			MeshProperties.Mesh->OnPostMeshBuild().AddUObject(this, &UNiagaraDataInterfaceArrayMesh::OnMeshPostBuild);
+			bAnyAdded = true;
+		}
+	}
+
+	if (bAnyAdded)
+	{
+		if (UImportSubsystem* ImportSubsystem = GEditor->GetEditorSubsystem<UImportSubsystem>())
+		{
+			ImportSubsystem->OnAssetReimport.AddUObject(this, &UNiagaraDataInterfaceArrayMesh::OnAssetReimported);
+		}
+	}
+}
+
+void UNiagaraDataInterfaceArrayMesh::RemoveMeshChangedDelegates()
+{
+	if (!GIsEditor || HasAnyFlags(RF_ClassDefaultObject))
+	{
+		return;
+	}
+
+	if (UImportSubsystem* ImportSubsystem = GEditor->GetEditorSubsystem<UImportSubsystem>())
+	{
+		ImportSubsystem->OnAssetReimport.RemoveAll(this);
+	}
+
+	for (const FNiagaraMeshRendererMeshPropertiesBase& MeshProperties : MeshData)
+	{
+		if (MeshProperties.Mesh)
+		{
+			MeshProperties.Mesh->GetOnMeshChanged().RemoveAll(this);
+			MeshProperties.Mesh->OnPreMeshBuild().RemoveAll(this);
+			MeshProperties.Mesh->OnPostMeshBuild().RemoveAll(this);
+		}
+	}
+}
+#endif// WITH_EDITORONLY_DATA
+
 void UNiagaraDataInterfaceArrayMesh::ForEachMesh(FNiagaraSystemInstance* SystemInstance, TFunction<void(int32)> NumMeshesDelegate, TFunction<void(const FNiagaraMeshRendererMeshProperties&)> IterateDelegate) const
 {
 	const FNDIArrayInstanceData_GameThread<FNiagaraMeshRendererMeshPropertiesBase>* PerInstanceData = nullptr;
@@ -76,8 +181,14 @@ void UNiagaraDataInterfaceArrayMesh::SetNiagaraArrayMesh(UNiagaraComponent* Niag
 {
 	if (UNiagaraDataInterfaceArrayMesh* ArrayDI = UNiagaraFunctionLibrary::GetDataInterface<UNiagaraDataInterfaceArrayMesh>(NiagaraComponent, OverrideName))
 	{
+	#if WITH_EDITORONLY_DATA
+		ArrayDI->RemoveMeshChangedDelegates();
+	#endif
 		auto* ArrayProxy = static_cast<typename UNiagaraDataInterfaceArrayMesh::FProxyType*>(ArrayDI->GetProxy());
 		ArrayProxy->SetArrayDataAndRecreateRenderState(MakeArrayView(ArrayData));
+	#if WITH_EDITORONLY_DATA
+		ArrayDI->AddMeshChangedDelegates();
+	#endif
 	}
 }
 
@@ -85,7 +196,13 @@ void UNiagaraDataInterfaceArrayMesh::SetNiagaraArrayMeshSM(UNiagaraComponent* Ni
 {
 	if (UNiagaraDataInterfaceArrayMesh* ArrayDI = UNiagaraFunctionLibrary::GetDataInterface<UNiagaraDataInterfaceArrayMesh>(NiagaraComponent, OverrideName))
 	{
+	#if WITH_EDITORONLY_DATA
+		ArrayDI->RemoveMeshChangedDelegates();
+	#endif
 		auto* ArrayProxy = static_cast<typename UNiagaraDataInterfaceArrayMesh::FProxyType*>(ArrayDI->GetProxy());
 		ArrayProxy->SetArrayDataAndRecreateRenderState(MakeArrayView(ArrayData));
+	#if WITH_EDITORONLY_DATA
+		ArrayDI->AddMeshChangedDelegates();
+	#endif
 	}
 }
