UE: Quick fix for non-slate configurations like commandlets

diff --git a/Engine/Plugins/FX/Niagara/Source/NiagaraEditor/Private/NiagaraDebugger.cpp b/Engine/Plugins/FX/Niagara/Source/NiagaraEditor/Private/NiagaraDebugger.cpp
--- a/Engine/Plugins/FX/Niagara/Source/NiagaraEditor/Private/NiagaraDebugger.cpp
+++ b/Engine/Plugins/FX/Niagara/Source/NiagaraEditor/Private/NiagaraDebugger.cpp
@@ -61,7 +61,10 @@ void FNiagaraDebugger::Init()
 	GetMutableDefault<UNiagaraDebugHUDSettings>()->OnChangedDelegate.AddSP(this, &FNiagaraDebugger::UpdateDebugHUDSettings);
 	GetMutableDefault<UNiagaraOutliner>()->OnChangedDelegate.AddSP(this, &FNiagaraDebugger::TriggerOutlinerCapture);
 
-	FSlateApplication::Get().RegisterInputPreProcessor(AsShared());
+	if (FSlateApplication::IsInitialized())
+	{
+		FSlateApplication::Get().RegisterInputPreProcessor(AsShared());
+	}
 }
 
 void FNiagaraDebugger::ExecConsoleCommand(const TCHAR* Cmd, bool bRequiresWorld)
