New Tool: Hierarchy Editor

Moved from Niagara into engine

Changes to the code are fairly minimal
- Renamed classes (removed the Niagara prefix)
- Unified some code paths
- Added comments to help get started

diff --git a/Engine/Plugins/FX/Niagara/Config/BaseNiagara.ini b/Engine/Plugins/FX/Niagara/Config/BaseNiagara.ini
--- a/Engine/Plugins/FX/Niagara/Config/BaseNiagara.ini
+++ b/Engine/Plugins/FX/Niagara/Config/BaseNiagara.ini
@@ -51,6 +51,13 @@
 +PropertyRedirects=(OldName="NiagaraSystem.FlipbookGeneratedSettings", NewName="NiagaraSystem.BakerGeneratedSettings")
 +StructRedirects=(OldName="NiagaraFlipbookTextureSettings", NewName="/Script/Niagara.NiagaraBakerTextureSettings")
 +StructRedirects=(OldName="NiagaraFlipbookTextureSource", NewName="/Script/Niagara.NiagaraBakerTextureSource")
+; Hierarchy Editor Redirects
++StructRedirects=(OldName="/Script/NiagaraEditor.NiagaraHierarchyIdentity",NewName="/Script/HierarchyEditor.HierarchyElementIdentity")
++ClassRedirects=(OldName="/Script/NiagaraEditor.NiagaraHierarchyItemBase",NewName="/Script/HierarchyEditor.HierarchyElement")
++ClassRedirects=(OldName="/Script/NiagaraEditor.NiagaraHierarchyRoot",NewName="/Script/HierarchyEditor.HierarchyRoot")
++ClassRedirects=(OldName="/Script/NiagaraEditor.NiagaraHierarchyItem",NewName="/Script/HierarchyEditor.HierarchyItem")
++ClassRedirects=(OldName="/Script/NiagaraEditor.NiagaraHierarchySection",NewName="/Script/HierarchyEditor.HierarchySection")
++ClassRedirects=(OldName="/Script/NiagaraEditor.NiagaraHierarchyCategory",NewName="/Script/HierarchyEditor.HierarchyCategory")
 
 [/Script/Niagara.NiagaraBakerSettings]
 BakeQualityLevel=Cinematic
diff --git a/Engine/Plugins/FX/Niagara/Source/Niagara/Classes/NiagaraScript.h b/Engine/Plugins/FX/Niagara/Source/Niagara/Classes/NiagaraScript.h
--- a/Engine/Plugins/FX/Niagara/Source/Niagara/Classes/NiagaraScript.h
+++ b/Engine/Plugins/FX/Niagara/Source/Niagara/Classes/NiagaraScript.h
@@ -27,7 +27,7 @@ class UNiagaraDataInterface;
 class FNiagaraCompileRequestDataBase;
 class FNiagaraCompileRequestDuplicateDataBase;
 class UNiagaraConvertInPlaceUtilityBase;
-class UNiagaraHierarchyRoot;
+class UHierarchyRoot;
 
 #define NIAGARA_INVALID_MEMORY (0xBA)
 
diff --git a/Engine/Plugins/FX/Niagara/Source/NiagaraEditor/NiagaraEditor.Build.cs b/Engine/Plugins/FX/Niagara/Source/NiagaraEditor/NiagaraEditor.Build.cs
--- a/Engine/Plugins/FX/Niagara/Source/NiagaraEditor/NiagaraEditor.Build.cs
+++ b/Engine/Plugins/FX/Niagara/Source/NiagaraEditor/NiagaraEditor.Build.cs
@@ -87,6 +87,7 @@ public class NiagaraEditor : ModuleRules
                 "AppFramework",
 				"ToolWidgets",
 				"Sequencer",
+				"HierarchyEditor"
             }
         );
 
diff --git a/Engine/Plugins/FX/Niagara/Source/NiagaraEditor/Private/Customizations/NiagaraComponentDetails.h b/Engine/Plugins/FX/Niagara/Source/NiagaraEditor/Private/Customizations/NiagaraComponentDetails.h
--- a/Engine/Plugins/FX/Niagara/Source/NiagaraEditor/Private/Customizations/NiagaraComponentDetails.h
+++ b/Engine/Plugins/FX/Niagara/Source/NiagaraEditor/Private/Customizations/NiagaraComponentDetails.h
@@ -8,7 +8,7 @@
 #include "NiagaraComponent.h"
 #include "NiagaraSimCacheCapture.h"
 #include "ViewModels/TNiagaraViewModelManager.h"
-#include "ViewModels/HierarchyEditor/NiagaraHierarchyViewModelBase.h"
+#include "HierarchyViewModelBase.h"
 #include "NiagaraComponentDetails.generated.h"
 
 class UNiagaraSystem;
diff --git a/Engine/Plugins/FX/Niagara/Source/NiagaraEditor/Private/Customizations/UserParameters/NiagaraComponentUserParametersBuilder.cpp b/Engine/Plugins/FX/Niagara/Source/NiagaraEditor/Private/Customizations/UserParameters/NiagaraComponentUserParametersBuilder.cpp
--- a/Engine/Plugins/FX/Niagara/Source/NiagaraEditor/Private/Customizations/UserParameters/NiagaraComponentUserParametersBuilder.cpp
+++ b/Engine/Plugins/FX/Niagara/Source/NiagaraEditor/Private/Customizations/UserParameters/NiagaraComponentUserParametersBuilder.cpp
@@ -501,9 +501,9 @@ void FNiagaraComponentUserParametersNodeBuilder::RegisterRebuildOnHierarchyChang
 			if(SystemViewModel.IsValid())
 			{
 				SystemViewModel->GetUserParametersHierarchyViewModel()->OnHierarchyChanged().RemoveAll(this);
-				SystemViewModel->GetUserParametersHierarchyViewModel()->OnHierarchyChanged().Add(UNiagaraHierarchyViewModelBase::FOnHierarchyChanged::FDelegate::CreateSP(this, &FNiagaraComponentUserParametersNodeBuilder::Rebuild));
+				SystemViewModel->GetUserParametersHierarchyViewModel()->OnHierarchyChanged().Add(UHierarchyViewModelBase::FOnHierarchyChanged::FDelegate::CreateSP(this, &FNiagaraComponentUserParametersNodeBuilder::Rebuild));
 				SystemViewModel->GetUserParametersHierarchyViewModel()->OnHierarchyPropertiesChanged().RemoveAll(this);
-				SystemViewModel->GetUserParametersHierarchyViewModel()->OnHierarchyPropertiesChanged().Add(UNiagaraHierarchyViewModelBase::FOnHierarchyPropertiesChanged::FDelegate::CreateSP(this, &FNiagaraComponentUserParametersNodeBuilder::Rebuild));
+				SystemViewModel->GetUserParametersHierarchyViewModel()->OnHierarchyPropertiesChanged().Add(UHierarchyViewModelBase::FOnHierarchyPropertiesChanged::FDelegate::CreateSP(this, &FNiagaraComponentUserParametersNodeBuilder::Rebuild));
 			}
 		}
 	}
diff --git a/Engine/Plugins/FX/Niagara/Source/NiagaraEditor/Private/Customizations/UserParameters/NiagaraSystemUserParametersBuilder.cpp b/Engine/Plugins/FX/Niagara/Source/NiagaraEditor/Private/Customizations/UserParameters/NiagaraSystemUserParametersBuilder.cpp
--- a/Engine/Plugins/FX/Niagara/Source/NiagaraEditor/Private/Customizations/UserParameters/NiagaraSystemUserParametersBuilder.cpp
+++ b/Engine/Plugins/FX/Niagara/Source/NiagaraEditor/Private/Customizations/UserParameters/NiagaraSystemUserParametersBuilder.cpp
@@ -81,9 +81,9 @@ void FNiagaraSystemUserParameterBuilder::GenerateChildContent(IDetailChildrenBui
 			System->GetExposedParameters().OnStructureChanged().Add(FNiagaraParameterStore::FOnStructureChanged::FDelegate::CreateSP(this, &FNiagaraSystemUserParameterBuilder::Rebuild));
 
 			SystemViewModel.Pin()->GetUserParametersHierarchyViewModel()->OnHierarchyChanged().RemoveAll(this);
-			SystemViewModel.Pin()->GetUserParametersHierarchyViewModel()->OnHierarchyChanged().Add(UNiagaraHierarchyViewModelBase::FOnHierarchyChanged::FDelegate::CreateSP(this, &FNiagaraSystemUserParameterBuilder::Rebuild));
+			SystemViewModel.Pin()->GetUserParametersHierarchyViewModel()->OnHierarchyChanged().Add(UHierarchyViewModelBase::FOnHierarchyChanged::FDelegate::CreateSP(this, &FNiagaraSystemUserParameterBuilder::Rebuild));
 			SystemViewModel.Pin()->GetUserParametersHierarchyViewModel()->OnHierarchyPropertiesChanged().RemoveAll(this);
-			SystemViewModel.Pin()->GetUserParametersHierarchyViewModel()->OnHierarchyPropertiesChanged().Add(UNiagaraHierarchyViewModelBase::FOnHierarchyPropertiesChanged::FDelegate::CreateSP(this, &FNiagaraSystemUserParameterBuilder::Rebuild));
+			SystemViewModel.Pin()->GetUserParametersHierarchyViewModel()->OnHierarchyPropertiesChanged().Add(UHierarchyViewModelBase::FOnHierarchyPropertiesChanged::FDelegate::CreateSP(this, &FNiagaraSystemUserParameterBuilder::Rebuild));
 			
 			TSharedPtr<FNiagaraUserParameterPanelViewModel> UserParameterPanelViewModel = SystemViewModel.Pin()->GetUserParameterPanelViewModel();
 			if(ensure(UserParameterPanelViewModel.IsValid()))
@@ -95,7 +95,7 @@ void FNiagaraSystemUserParameterBuilder::GenerateChildContent(IDetailChildrenBui
 			bDelegatesInitialized = true;
 		}
 		
-		UNiagaraHierarchyRoot* Root = Cast<UNiagaraSystemEditorData>(System->GetEditorData())->UserParameterHierarchy;
+		UHierarchyRoot* Root = Cast<UNiagaraSystemEditorData>(System->GetEditorData())->UserParameterHierarchy;
 		GenerateUserParameterRows(ChildrenBuilder, *Root);
 	}
 }
diff --git a/Engine/Plugins/FX/Niagara/Source/NiagaraEditor/Private/Customizations/UserParameters/NiagaraSystemUserParametersBuilder.h b/Engine/Plugins/FX/Niagara/Source/NiagaraEditor/Private/Customizations/UserParameters/NiagaraSystemUserParametersBuilder.h
--- a/Engine/Plugins/FX/Niagara/Source/NiagaraEditor/Private/Customizations/UserParameters/NiagaraSystemUserParametersBuilder.h
+++ b/Engine/Plugins/FX/Niagara/Source/NiagaraEditor/Private/Customizations/UserParameters/NiagaraSystemUserParametersBuilder.h
@@ -6,7 +6,7 @@
 #include "CoreMinimal.h"
 #include "Input/Reply.h"
 #include "ViewModels/TNiagaraViewModelManager.h"
-#include "ViewModels/HierarchyEditor/NiagaraHierarchyViewModelBase.h"
+#include "HierarchyViewModelBase.h"
 #include "ViewModels/HierarchyEditor/NiagaraUserParametersHierarchyViewModel.h"
 
 class IDetailLayoutBuilder;
diff --git a/Engine/Plugins/FX/Niagara/Source/NiagaraEditor/Private/Customizations/UserParameters/NiagaraUserParametersBuilderBase.cpp b/Engine/Plugins/FX/Niagara/Source/NiagaraEditor/Private/Customizations/UserParameters/NiagaraUserParametersBuilderBase.cpp
--- a/Engine/Plugins/FX/Niagara/Source/NiagaraEditor/Private/Customizations/UserParameters/NiagaraUserParametersBuilderBase.cpp
+++ b/Engine/Plugins/FX/Niagara/Source/NiagaraEditor/Private/Customizations/UserParameters/NiagaraUserParametersBuilderBase.cpp
@@ -154,7 +154,7 @@ IDetailPropertyRow* FNiagaraUserParameterNodeBuilder::AddObjectAssetParameterAsR
 	return ChildrenBuilder.AddExternalObjectProperty({AssetHelper}, FName("Path"), AddPropertyParams);
 }
 
-void FNiagaraUserParameterNodeBuilder::GenerateUserParameterRows(IDetailChildrenBuilder& ChildrenBuilder, UNiagaraHierarchyRoot& UserParameterHierarchyRoot)
+void FNiagaraUserParameterNodeBuilder::GenerateUserParameterRows(IDetailChildrenBuilder& ChildrenBuilder, UHierarchyRoot& UserParameterHierarchyRoot)
 {
 	UNiagaraSystem* SystemAsset = GetSystem();
 
@@ -163,13 +163,13 @@ void FNiagaraUserParameterNodeBuilder::GenerateUserParameterRows(IDetailChildren
 		return;
 	}
 
-	TArray<const UNiagaraHierarchySection*> OrderedSections;
-	TMap<const UNiagaraHierarchySection*, TArray<const UNiagaraHierarchyCategory*>> SectionCategoryMap;
+	TArray<const UHierarchySection*> OrderedSections;
+	TMap<const UHierarchySection*, TArray<const UHierarchyCategory*>> SectionCategoryMap;
 
-	UNiagaraHierarchyRoot* Root = &UserParameterHierarchyRoot;
-	for(const UNiagaraHierarchyItemBase* Child : Root->GetChildren())
+	UHierarchyRoot* Root = &UserParameterHierarchyRoot;
+	for(const UHierarchyElement* Child : Root->GetChildren())
 	{
-		if(const UNiagaraHierarchyCategory* HierarchyCategory = Cast<UNiagaraHierarchyCategory>(Child))
+		if(const UHierarchyCategory* HierarchyCategory = Cast<UHierarchyCategory>(Child))
 		{
 			if(HierarchyCategory->GetSection() != nullptr)
 			{
@@ -192,13 +192,13 @@ void FNiagaraUserParameterNodeBuilder::GenerateUserParameterRows(IDetailChildren
 				continue;
 			}
 			
-			TArray<const UNiagaraHierarchyCategory*>& HierarchyCategories = SectionCategoryMap[Section];
+			TArray<const UHierarchyCategory*>& HierarchyCategories = SectionCategoryMap[Section];
 			
 			if(HierarchyCategories.Num() > 0)
 			{
 				// we only want to add a section if any of its categories contain at least one user parameter
 				bool bDoesUserParamExist = false;
-				for(const UNiagaraHierarchyCategory* Category : HierarchyCategories)
+				for(const UHierarchyCategory* Category : HierarchyCategories)
 				{
 					bDoesUserParamExist |= Category->DoesOneChildExist<UNiagaraHierarchyUserParameter>(true);
 
@@ -230,8 +230,8 @@ void FNiagaraUserParameterNodeBuilder::GenerateUserParameterRows(IDetailChildren
 							[
 								SNew(STextBlock)
 								.TextStyle(FAppStyle::Get(), "SmallText")
-								.Text_UObject(Section, &UNiagaraHierarchySection::GetSectionNameAsText)
-								.ToolTipText_UObject(Section, &UNiagaraHierarchySection::GetTooltip)
+								.Text_UObject(Section, &UHierarchySection::GetSectionNameAsText)
+								.ToolTipText_UObject(Section, &UHierarchySection::GetTooltip)
 							]
 						]
 					];
@@ -276,9 +276,9 @@ void FNiagaraUserParameterNodeBuilder::GenerateUserParameterRows(IDetailChildren
 	}
 	
 	// we add the categories now if the active section is set to all (nullptr) or if the section it belongs to is active
-	for(const UNiagaraHierarchyItemBase* Child : Root->GetChildren())
+	for(const UHierarchyElement* Child : Root->GetChildren())
 	{		
-		if(const UNiagaraHierarchyCategory* HierarchyCategory = Cast<UNiagaraHierarchyCategory>(Child))
+		if(const UHierarchyCategory* HierarchyCategory = Cast<UHierarchyCategory>(Child))
 		{
 			if((ActiveSection == nullptr || HierarchyCategory->GetSection() == ActiveSection) && HierarchyCategory->DoesOneChildExist<UNiagaraHierarchyUserParameter>(true))
 			{
@@ -450,7 +450,7 @@ FString FNiagaraUserParameterNodeBuilder::GetReferencerName() const
 	return "UserParametersBuilder";
 }
 
-FNiagaraUserParameterCategoryBuilder::FNiagaraUserParameterCategoryBuilder(const UNiagaraHierarchyCategory* InNiagaraHierarchyCategory, FName InCustomBuilderRowName, FNiagaraUserParameterNodeBuilder& InNodeBuilder) : NodeBuilder(InNodeBuilder)
+FNiagaraUserParameterCategoryBuilder::FNiagaraUserParameterCategoryBuilder(const UHierarchyCategory* InNiagaraHierarchyCategory, FName InCustomBuilderRowName, FNiagaraUserParameterNodeBuilder& InNodeBuilder) : NodeBuilder(InNodeBuilder)
 {
 	HierarchyCategory = InNiagaraHierarchyCategory;
 	CustomBuilderRowName = InCustomBuilderRowName;
@@ -462,8 +462,8 @@ void FNiagaraUserParameterCategoryBuilder::GenerateHeaderRowContent(FDetailWidge
 	[
 		SNew(STextBlock)
 		.TextStyle(FNiagaraEditorStyle::Get(), "NiagaraEditor.ParameterText")
-		.Text_UObject(HierarchyCategory.Get(), &UNiagaraHierarchyCategory::GetCategoryAsText)
-		.ToolTipText_UObject(HierarchyCategory.Get(), &UNiagaraHierarchyCategory::GetTooltip)
+		.Text_UObject(HierarchyCategory.Get(), &UHierarchyCategory::GetCategoryAsText)
+		.ToolTipText_UObject(HierarchyCategory.Get(), &UHierarchyCategory::GetTooltip)
 	];
 
 	NodeRow.FilterString(FText::FromName(CustomBuilderRowName));
@@ -472,9 +472,9 @@ void FNiagaraUserParameterCategoryBuilder::GenerateHeaderRowContent(FDetailWidge
 void FNiagaraUserParameterCategoryBuilder::GenerateChildContent(IDetailChildrenBuilder& ChildrenBuilder)
 {
 	TArray<const UNiagaraHierarchyUserParameter*> ContainedParameters;
-	TArray<const UNiagaraHierarchyCategory*> ContainedCategories;
+	TArray<const UHierarchyCategory*> ContainedCategories;
 
-	for(const UNiagaraHierarchyItemBase* Child : HierarchyCategory->GetChildren())
+	for(const UHierarchyElement* Child : HierarchyCategory->GetChildren())
 	{
 		if(const UNiagaraHierarchyUserParameter* UserParameter = Cast<UNiagaraHierarchyUserParameter>(Child))
 		{
@@ -482,15 +482,15 @@ void FNiagaraUserParameterCategoryBuilder::GenerateChildContent(IDetailChildrenB
 		}
 	}
 
-	for(UNiagaraHierarchyItemBase* Child : HierarchyCategory->GetChildren())
+	for(UHierarchyElement* Child : HierarchyCategory->GetChildren())
 	{
-		if(UNiagaraHierarchyCategory* Category = Cast<UNiagaraHierarchyCategory>(Child))
+		if(UHierarchyCategory* Category = Cast<UHierarchyCategory>(Child))
 		{
 			ContainedCategories.Add(Category);
 		}
 	}
 
-	for(const UNiagaraHierarchyCategory* Category : ContainedCategories)
+	for(const UHierarchyCategory* Category : ContainedCategories)
 	{
 		if(Category->DoesOneChildExist<UNiagaraHierarchyUserParameter>(true))
 		{
diff --git a/Engine/Plugins/FX/Niagara/Source/NiagaraEditor/Private/Customizations/UserParameters/NiagaraUserParametersBuilderBase.h b/Engine/Plugins/FX/Niagara/Source/NiagaraEditor/Private/Customizations/UserParameters/NiagaraUserParametersBuilderBase.h
--- a/Engine/Plugins/FX/Niagara/Source/NiagaraEditor/Private/Customizations/UserParameters/NiagaraUserParametersBuilderBase.h
+++ b/Engine/Plugins/FX/Niagara/Source/NiagaraEditor/Private/Customizations/UserParameters/NiagaraUserParametersBuilderBase.h
@@ -6,7 +6,7 @@
 #include "IDetailCustomNodeBuilder.h"
 #include "NiagaraComponent.h"
 #include "SNiagaraParameterEditor.h"
-#include "ViewModels/HierarchyEditor/NiagaraHierarchyViewModelBase.h"
+#include "HierarchyViewModelBase.h"
 
 class UNiagaraScriptVariable;
 
@@ -33,17 +33,17 @@ public:
 	}
 	
 	/** Generates all rows for the node builder. Includes setting up sections, categories and user parameters */
-	void GenerateUserParameterRows(IDetailChildrenBuilder& ChildrenBuilder, UNiagaraHierarchyRoot& UserParameterHierarchyRoot);
+	void GenerateUserParameterRows(IDetailChildrenBuilder& ChildrenBuilder, UHierarchyRoot& UserParameterHierarchyRoot);
 	/** Generates a single row for a specific user parameter. */
 	void GenerateRowForUserParameter(IDetailChildrenBuilder& ChildrenBuilder, FNiagaraVariable UserParameter);
 	
-	void OnActiveSectionChanged(ECheckBoxState State, const UNiagaraHierarchySection* NewSelection)
+	void OnActiveSectionChanged(ECheckBoxState State, const UHierarchySection* NewSelection)
 	{
 		ActiveSection = NewSelection;
 		OnRebuildChildren.Execute();
 	}
 
-	ECheckBoxState IsSectionActive(const UNiagaraHierarchySection* Section) const
+	ECheckBoxState IsSectionActive(const UHierarchySection* Section) const
 	{
 		return ActiveSection == Section ? ECheckBoxState::Checked : ECheckBoxState::Unchecked;
 	}
@@ -79,8 +79,8 @@ protected:
 	bool bDelegatesInitialized = false;
 	FSimpleDelegate OnRebuildChildren;
 	FName CustomBuilderRowName;
-	const UNiagaraHierarchySection* ActiveSection = nullptr;
-	TArray<UNiagaraHierarchySection*> AvailableSections;
+	const UHierarchySection* ActiveSection = nullptr;
+	TArray<UHierarchySection*> AvailableSections;
 	TMap<FNiagaraVariable, TSharedRef<SNiagaraParameterEditor>> NiagaraParameterEditors;
 	/** The struct on scopes contained here are owned by the builder. When parameter editor value changes, they get copied into the display data, and from there into the parameter store.
 	 * This is done because the types the parameter editors operate on are not the same as the struct on scopes we are editing (bool vs. FNiagaraBool etc.) */
@@ -93,7 +93,7 @@ protected:
 class FNiagaraUserParameterCategoryBuilder : public IDetailCustomNodeBuilder, public TSharedFromThis<FNiagaraUserParameterCategoryBuilder>
 {
 public:
-	FNiagaraUserParameterCategoryBuilder(const UNiagaraHierarchyCategory* InNiagaraHierarchyCategory, FName InCustomBuilderRowName, FNiagaraUserParameterNodeBuilder& InNodeBuilder);
+	FNiagaraUserParameterCategoryBuilder(const UHierarchyCategory* InNiagaraHierarchyCategory, FName InCustomBuilderRowName, FNiagaraUserParameterNodeBuilder& InNodeBuilder);
 	
 	virtual void GenerateHeaderRowContent(FDetailWidgetRow& NodeRow) override;
 	
@@ -107,7 +107,7 @@ public:
 	virtual void GenerateChildContent(IDetailChildrenBuilder& ChildrenBuilder) override;
 
 private:
-	TWeakObjectPtr<const UNiagaraHierarchyCategory> HierarchyCategory;
+	TWeakObjectPtr<const UHierarchyCategory> HierarchyCategory;
 	FName CustomBuilderRowName;
 	FNiagaraUserParameterNodeBuilder& NodeBuilder;
 };
diff --git a/Engine/Plugins/FX/Niagara/Source/NiagaraEditor/Private/NiagaraEditorModule.cpp b/Engine/Plugins/FX/Niagara/Source/NiagaraEditor/Private/NiagaraEditorModule.cpp
--- a/Engine/Plugins/FX/Niagara/Source/NiagaraEditor/Private/NiagaraEditorModule.cpp
+++ b/Engine/Plugins/FX/Niagara/Source/NiagaraEditor/Private/NiagaraEditorModule.cpp
@@ -181,7 +181,6 @@
 #include "NiagaraVolumeRendererProperties.h"
 
 #include "Engine/AssetManager.h"
-#include "ViewModels/HierarchyEditor/NiagaraHierarchyCommands.h"
 #include "Widgets/AssetBrowser/SNiagaraSelectedAssetDetails.h"
 #include "NiagaraRecentAndFavoritesManager.h"
 #include "Widgets/DataChannel/NiagaraDataChannelWizard.h"
@@ -1315,7 +1314,6 @@ void FNiagaraEditorModule::StartupModule()
 		FConsoleCommandDelegate::CreateRaw(this, &FNiagaraEditorModule::ReinitializeStyle));
 
 	FNiagaraEditorCommands::Register();
-	FNiagaraHierarchyEditorCommands::Register();
 
 	NiagaraComponentBroker = MakeShareable(new FNiagaraComponentBroker);
 	FComponentAssetBrokerage::RegisterBroker(NiagaraComponentBroker, UNiagaraComponent::StaticClass(), true, true);
diff --git a/Engine/Plugins/FX/Niagara/Source/NiagaraEditor/Private/NiagaraEditorStyle.cpp b/Engine/Plugins/FX/Niagara/Source/NiagaraEditor/Private/NiagaraEditorStyle.cpp
--- a/Engine/Plugins/FX/Niagara/Source/NiagaraEditor/Private/NiagaraEditorStyle.cpp
+++ b/Engine/Plugins/FX/Niagara/Source/NiagaraEditor/Private/NiagaraEditorStyle.cpp
@@ -434,8 +434,6 @@ void FNiagaraEditorStyle::InitIcons()
 	Set("NiagaraEditor.Stateless.UpdateIcon", new IMAGE_BRUSH("Icons/Update", Icon12x12, FLinearColor::White));
 	Set("NiagaraEditor.Stateless.RenderIcon", new IMAGE_BRUSH("Icons/Render", Icon12x12, FLinearColor::White));
 	
-	Set("NiagaraEditor.HierarchyEditor.RootDropIcon", new IMAGE_BRUSH_SVG("Icons/caret-down", Icon20x20));
-	
 	FSlateBrush* SlateBrush = new FSlateBrush();
 	*SlateBrush = FAppStyle::Get().GetWidgetStyle<FEditableTextBoxStyle>("NormalEditableTextBox").BackgroundImageNormal;
 	SlateBrush->TintColor = FLinearColor(0.01f, 0.01f, 0.01f);
@@ -461,11 +459,6 @@ void FNiagaraEditorStyle::InitTextStyles()
 {
 	const FTextBlockStyle NormalText = FAppStyle::Get().GetWidgetStyle<FTextBlockStyle>("NormalText");
 
-	FInlineEditableTextBlockStyle HierarchyEditorCategoryStyle = FAppStyle::GetWidgetStyle<FInlineEditableTextBlockStyle>("InlineEditableTextBlockStyle");
-	HierarchyEditorCategoryStyle.TextStyle = FAppStyle::Get().GetWidgetStyle<FTextBlockStyle>("DetailsView.CategoryTextStyle");
-	HierarchyEditorCategoryStyle.TextStyle.SetFontSize(10.f);
-	Set("NiagaraEditor.HierarchyEditor.CategoryTextStyle", HierarchyEditorCategoryStyle);
-
 	FSlateFontInfo SummaryViewInputOverrideFontInfo = DEFAULT_FONT("Italic", 9);
 	FTextBlockStyle SummaryViewInputOverrideTextBlockStyle = FTextBlockStyle(NormalText)
 		.SetFont(SummaryViewInputOverrideFontInfo);
@@ -833,67 +826,6 @@ void FNiagaraEditorStyle::InitScalabilityIcons()
 	Set("NiagaraEditor.Scalability.Preview.ResetPlatform", new IMAGE_BRUSH("Icons/Scalability/ResetPreviewPlatform", Icon20x20));
 }
 
-void FNiagaraEditorStyle::InitHierarchyEditor()
-{
-	const FTextBlockStyle NormalText = FAppStyle::GetWidgetStyle<FTextBlockStyle>("NormalText");
-	FSlateFontInfo CategoryFont = FAppStyle::Get().GetFontStyle(TEXT("DetailsView.CategoryFontStyle"));
-	CategoryFont.Size = 11;
-	// const FEditableTextBoxStyle NormalEditableTextBox = FStarshipCoreStyle::GetCoreStyle().GetWidgetStyle<FEditableTextBoxStyle>("NormalEditableTextBox");
-	const FEditableTextBoxStyle NormalEditableTextBox = FAppStyle::Get().GetWidgetStyle<FEditableTextBoxStyle>("NormalEditableTextBox");
-	
-	FEditableTextBoxStyle CategoryEditableText = FEditableTextBoxStyle(NormalEditableTextBox)
-		.SetFont(CategoryFont)
-		.SetForegroundColor(FStyleColors::AccentWhite);
-	
-	FTextBlockStyle CategoryText = FTextBlockStyle(NormalText)
-		.SetFont(CategoryFont);
-	FInlineEditableTextBlockStyle HierarchyCategoryTextStyle = FInlineEditableTextBlockStyle()
-		.SetTextStyle(CategoryText)
-		.SetEditableTextBoxStyle(CategoryEditableText);
-
-	Set("NiagaraEditor.HierarchyEditor.CategoryTextBlock", CategoryText);
-	Set("NiagaraEditor.HierarchyEditor.Category", HierarchyCategoryTextStyle);
-
-	Set("NiagaraEditor.Stack.DropTarget.BorderVertical", new IMAGE_BRUSH("Icons/StackDropTargetBorder_Vertical", FVector2D(2, 8), FLinearColor::White, ESlateBrushTileType::Vertical));
-	Set("NiagaraEditor.Stack.DropTarget.BorderHorizontal", new IMAGE_BRUSH("Icons/StackDropTargetBorder_Horizontal", FVector2D(8, 2), FLinearColor::White, ESlateBrushTileType::Horizontal));
-
-	FButtonStyle SimpleButtonStyle = FAppStyle::Get().GetWidgetStyle<FButtonStyle>("SimpleButton");
-	FButtonStyle ButtonStyle = FButtonStyle(SimpleButtonStyle)
-		.SetNormalForeground(FStyleColors::Foreground)
-		.SetPressedForeground(FStyleColors::ForegroundHover)
-		.SetHoveredForeground(FStyleColors::ForegroundHover)
-		.SetDisabledForeground(FStyleColors::Foreground)
-		.SetNormalPadding(FMargin(8.f, 2.f, 8.f, 2.f))
-		.SetPressedPadding(FMargin(8.f, 3.f, 8.f, 1.f));
-
-	Set("NiagaraEditor.HierarchyEditor.ButtonStyle", ButtonStyle);
-
-	const FSlateColor SelectorColor = FAppStyle::GetSlateColor("SelectorColor");
-	const FSlateColor SelectionColor = FAppStyle::GetSlateColor("SelectionColor");
-	const FSlateColor SelectionColor_Inactive = FAppStyle::GetSlateColor("SelectionColor_Inactive");
-	FLinearColor HighlightColor = FLinearColor(0.068f, 0.068f, 0.068f);
-	
-	Set("NiagaraEditor.Color.Highlight", HighlightColor);
-	
-	FTableRowStyle HierarchyEditorCategoryRowStyle = FTableRowStyle()
-		.SetEvenRowBackgroundBrush(*FAppStyle::GetBrush("DetailsView.CategoryTop"))
-		.SetOddRowBackgroundBrush(*FAppStyle::GetBrush("DetailsView.CategoryTop"))
-		.SetEvenRowBackgroundHoveredBrush(*FAppStyle::GetBrush("DetailsView.CategoryTop_Hovered"))
-		.SetOddRowBackgroundHoveredBrush(*FAppStyle::GetBrush("DetailsView.CategoryTop_Hovered"))
-		.SetSelectorFocusedBrush( CORE_BORDER_BRUSH( "Common/Selector", FMargin(4.f/16.f), SelectorColor))
-		.SetActiveBrush( CORE_IMAGE_BRUSH( "Common/Selection", Icon8x8, SelectionColor))
-		.SetActiveHoveredBrush( CORE_IMAGE_BRUSH( "Common/Selection", Icon8x8, SelectionColor))
-		.SetInactiveBrush( CORE_IMAGE_BRUSH( "Common/Selection", Icon8x8, SelectionColor_Inactive))
-		.SetInactiveHoveredBrush( CORE_IMAGE_BRUSH( "Common/Selection", Icon8x8, SelectionColor_Inactive))
-		.SetActiveHighlightedBrush( CORE_IMAGE_BRUSH("Common/Selection", Icon8x8, HighlightColor))
-		.SetInactiveHighlightedBrush( CORE_IMAGE_BRUSH("Common/Selection", Icon8x8, HighlightColor))
-		.SetDropIndicator_Above(CORE_BOX_BRUSH("Common/DropZoneIndicator_Above", FMargin(10.0f / 16.0f, 10.0f / 16.0f, 0, 0), FStarshipCoreStyle::GetCoreStyle().GetSlateColor("SelectorColor")))
-		.SetDropIndicator_Onto(CORE_BOX_BRUSH("Common/DropZoneIndicator_Onto", FMargin(4.0f / 16.0f), FStarshipCoreStyle::GetCoreStyle().GetSlateColor("SelectorColor")))
-		.SetDropIndicator_Below(CORE_BOX_BRUSH("Common/DropZoneIndicator_Below", FMargin(10.0f / 16.0f, 0, 0, 10.0f / 16.0f), FStarshipCoreStyle::GetCoreStyle().GetSlateColor("SelectorColor")));
-	
-	Set("NiagaraEditor.HierarchyEditor.Row.Category", HierarchyEditorCategoryRowStyle);
-}
-
 void FNiagaraEditorStyle::InitSimCacheEditor()
 {
 	const FSlateColor SelectionColor = FAppStyle::GetSlateColor("SelectionColor");
@@ -994,7 +926,6 @@ FNiagaraEditorStyle::FNiagaraEditorStyle() : FSlateStyleSet("NiagaraEditorStyle"
 	InitScalabilityIcons();
 	InitViewportStyle();
 	InitScratchStyle();
-	InitHierarchyEditor();
 	InitSimCacheEditor();
 }
 
diff --git a/Engine/Plugins/FX/Niagara/Source/NiagaraEditor/Private/NiagaraEditorUtilities.cpp b/Engine/Plugins/FX/Niagara/Source/NiagaraEditor/Private/NiagaraEditorUtilities.cpp
--- a/Engine/Plugins/FX/Niagara/Source/NiagaraEditor/Private/NiagaraEditorUtilities.cpp
+++ b/Engine/Plugins/FX/Niagara/Source/NiagaraEditor/Private/NiagaraEditorUtilities.cpp
@@ -80,7 +80,7 @@
 #include "ViewModels/NiagaraEmitterHandleViewModel.h"
 #include "ViewModels/NiagaraEmitterViewModel.h"
 #include "ViewModels/NiagaraParameterPanelViewModel.h"
-#include "Widgets/SNiagaraHierarchyEditor.h"
+#include "Widgets/SHierarchyEditor.h"
 #include "ViewModels/HierarchyEditor/NiagaraHierarchyScriptParametersViewModel.h"
 #include "Widgets/SToolTip.h"
 
@@ -4695,12 +4695,12 @@ FNiagaraVariable FNiagaraEditorUtilities::UserParameters::DuplicateUserParameter
 	return FNiagaraVariable();
 }
 
-TSharedRef<SWidget> FNiagaraEditorUtilities::HierarchyEditor::Scripts::GenerateRowContentForScriptParameterHierarchyEditor(TSharedRef<FNiagaraHierarchyItemViewModelBase> HierarchyItem)
+TSharedRef<SWidget> FNiagaraEditorUtilities::HierarchyEditor::Scripts::GenerateRowContentForScriptParameterHierarchyEditor(TSharedRef<FHierarchyElementViewModel> HierarchyItem)
 {
-	if(HierarchyItem->GetDataMutable()->IsA<UNiagaraHierarchyCategory>())
+	if(HierarchyItem->GetDataMutable()->IsA<UHierarchyCategory>())
 	{
-		TSharedRef<FNiagaraHierarchyCategoryViewModel> TreeViewCategory = StaticCastSharedRef<FNiagaraHierarchyCategoryViewModel>(HierarchyItem);
-		return SNew(SNiagaraHierarchyCategory, TreeViewCategory);
+		TSharedRef<FHierarchyCategoryViewModel> TreeViewCategory = StaticCastSharedRef<FHierarchyCategoryViewModel>(HierarchyItem);
+		return SNew(SHierarchyCategory, TreeViewCategory);
 	}
 	else if(const UNiagaraHierarchyScriptParameter* ScriptParameter = Cast<UNiagaraHierarchyScriptParameter>(HierarchyItem->GetData()))
 	{
diff --git a/Engine/Plugins/FX/Niagara/Source/NiagaraEditor/Private/NiagaraEmitterEditorData.cpp b/Engine/Plugins/FX/Niagara/Source/NiagaraEditor/Private/NiagaraEmitterEditorData.cpp
--- a/Engine/Plugins/FX/Niagara/Source/NiagaraEditor/Private/NiagaraEmitterEditorData.cpp
+++ b/Engine/Plugins/FX/Niagara/Source/NiagaraEditor/Private/NiagaraEmitterEditorData.cpp
@@ -21,7 +21,7 @@ const FName UNiagaraEmitterEditorData::PrivateMemberNames::SummarySections = GET
 UNiagaraEmitterEditorData::UNiagaraEmitterEditorData(const FObjectInitializer& ObjectInitializer)
 {
 	StackEditorData = ObjectInitializer.CreateDefaultSubobject<UNiagaraStackEditorData>(this, TEXT("StackEditorData"));
-	SummaryViewRoot = CreateDefaultSubobject<UNiagaraHierarchyRoot>(TEXT("SummaryViewRoot"));
+	SummaryViewRoot = CreateDefaultSubobject<UHierarchyRoot>(TEXT("SummaryViewRoot"));
 	
 	if (StackEditorData != nullptr)
 	{
@@ -62,15 +62,15 @@ void UNiagaraEmitterEditorData::PostLoad_TransferSummaryDataToNewFormat()
 		}
 	}
 
-	auto GenerateStableIdentityForCategories = [](UNiagaraHierarchyCategory* Category) -> TArray<FName>
+	auto GenerateStableIdentityForCategories = [](UHierarchyCategory* Category) -> TArray<FName>
 	{
 		TArray<FName> CategoryNameChain;
 		
-		const UNiagaraHierarchyCategory* OuterCategory = Category;
+		const UHierarchyCategory* OuterCategory = Category;
 
 		CategoryNameChain.Add(Category->GetCategoryName());
 
-		for (; OuterCategory != nullptr; OuterCategory = OuterCategory->GetTypedOuter<UNiagaraHierarchyCategory>() )
+		for (; OuterCategory != nullptr; OuterCategory = OuterCategory->GetTypedOuter<UHierarchyCategory>() )
 		{
 			if(OuterCategory != nullptr)
 			{
@@ -83,16 +83,16 @@ void UNiagaraEmitterEditorData::PostLoad_TransferSummaryDataToNewFormat()
 	
 	TArray<UNiagaraNodeFunctionCall*> ModuleNodes = FNiagaraStackGraphUtilities::GetAllModuleNodes(EmitterGraph);
 	
-	TMap<FName, UNiagaraHierarchySection*> Sections;
-	TMap<FName, UNiagaraHierarchyCategory*> IdentityCategoryMap;
+	TMap<FName, UHierarchySection*> Sections;
+	TMap<FName, UHierarchyCategory*> IdentityCategoryMap;
 	TMap<FName, FName> SectionCategoryMapping;
 	if(SummarySections_DEPRECATED.Num() > 0)
 	{
 		for(FNiagaraStackSection& StackSection : SummarySections_DEPRECATED)
 		{
 			FText SectionDisplayName = StackSection.SectionDisplayName.IsEmptyOrWhitespace() ? FText::FromName(StackSection.SectionIdentifier) : StackSection.SectionDisplayName;
-			UNiagaraHierarchySection* Section = SummaryViewRoot->AddSection(SectionDisplayName);
-			FNiagaraHierarchyIdentity SectionIdentity;
+			UHierarchySection* Section = SummaryViewRoot->AddSection(SectionDisplayName);
+			FHierarchyElementIdentity SectionIdentity;
 			SectionIdentity.Names.Add(StackSection.SectionIdentifier);
 			Section->SetIdentity(SectionIdentity);
 			
@@ -101,7 +101,7 @@ void UNiagaraEmitterEditorData::PostLoad_TransferSummaryDataToNewFormat()
 			for(const FText& CategoryIdentity : StackSection.Categories)
 			{
 				FName CategoryName = FName(CategoryIdentity.ToString());
-				UNiagaraHierarchyCategory* HierarchyCategory = SummaryViewRoot->AddChild<UNiagaraHierarchyCategory>();
+				UHierarchyCategory* HierarchyCategory = SummaryViewRoot->AddChild<UHierarchyCategory>();
 				HierarchyCategory->SetSection(Section);
 				HierarchyCategory->SetCategoryName(CategoryName);
 				IdentityCategoryMap.Add(CategoryName, HierarchyCategory);
@@ -115,17 +115,17 @@ void UNiagaraEmitterEditorData::PostLoad_TransferSummaryDataToNewFormat()
 	if(SummaryViewFunctionInputMetadata_DEPRECATED.Num() > 0)
 	{
 		// we have a custom sort order cached explicitly for summary view. This takes precedence.
-		TMap<UNiagaraHierarchyItemBase*, int32> SummarySortOrder;
+		TMap<UHierarchyElement*, int32> SummarySortOrder;
 		// the original metadata defined in the script asset for a given input.
-		TMap<FNiagaraHierarchyIdentity, FNiagaraVariableMetaData> AssetVariableMetadata;
+		TMap<FHierarchyElementIdentity, FNiagaraVariableMetaData> AssetVariableMetadata;
 
 		// since the order in the new summary view data is completely explicit (or managed implicitly by certain classes)
 		// we have to sort inputs during transfer first by asset order, then summary sort order.
 		for(auto It = SummaryViewFunctionInputMetadata_DEPRECATED.CreateConstIterator(); It; ++It)
 		{
 			// The input identity we are constructing from the function call node guid & the variable guid
-			FNiagaraHierarchyIdentity ParentInputIdentity;
-			TArray<FNiagaraHierarchyIdentity> ChildInputIdentities;
+			FHierarchyElementIdentity ParentInputIdentity;
+			TArray<FHierarchyElementIdentity> ChildInputIdentities;
 
 			FName OriginalCategoryName;
 
@@ -170,7 +170,7 @@ void UNiagaraEmitterEditorData::PostLoad_TransferSummaryDataToNewFormat()
 								{
 									if(CandidateChildScriptVariable->Metadata.GetParentAttribute_DEPRECATED().IsEqual(ScriptVariable->Variable.GetName()))
 									{
-										FNiagaraHierarchyIdentity ChildIdentity;
+										FHierarchyElementIdentity ChildIdentity;
 										ChildIdentity.Guids.Add(FunctionCall->NodeGuid);
 										ChildIdentity.Guids.Add(CandidateChildScriptVariable->Metadata.GetVariableGuid());
 										ChildInputIdentities.Add(ChildIdentity);
@@ -228,14 +228,14 @@ void UNiagaraEmitterEditorData::PostLoad_TransferSummaryDataToNewFormat()
 				CategoryName = OriginalCategoryName;
 			}
 
-			UNiagaraHierarchyCategory* Category = nullptr;
+			UHierarchyCategory* Category = nullptr;
 			if(IdentityCategoryMap.Contains(CategoryName))
 			{
 				Category = IdentityCategoryMap[CategoryName];
 			}
 			else
 			{
-				Category = SummaryViewRoot->AddChild<UNiagaraHierarchyCategory>();
+				Category = SummaryViewRoot->AddChild<UHierarchyCategory>();
 				Category->SetCategoryName(CategoryName);
 				IdentityCategoryMap.Add(Category->GetCategoryName(), Category);
 			}
@@ -253,7 +253,7 @@ void UNiagaraEmitterEditorData::PostLoad_TransferSummaryDataToNewFormat()
 				}
 				SummarySortOrder.Add(SummaryModuleInput, It.Value().SortIndex);
 
-				for(const FNiagaraHierarchyIdentity& ChildInputIdentity : ChildInputIdentities)
+				for(const FHierarchyElementIdentity& ChildInputIdentity : ChildInputIdentities)
 				{
 					UNiagaraHierarchyModuleInput* SummaryModuleChildInput = SummaryModuleInput->AddChild<UNiagaraHierarchyModuleInput>();		
 					SummaryModuleChildInput->SetIdentity(ChildInputIdentity);
@@ -269,12 +269,12 @@ void UNiagaraEmitterEditorData::PostLoad_TransferSummaryDataToNewFormat()
 		}
 
 		// now that we have created all inputs, categories and sections, we fix up the order
-		TArray<UNiagaraHierarchyCategory*> TopLevelCategories;
+		TArray<UHierarchyCategory*> TopLevelCategories;
 		SummaryViewRoot->GetChildrenOfType(TopLevelCategories, false);
 
 		// we sort by summary sort order if specified, and use the asset order as fallback.
 		// advanced items always get put at the bottom
-		auto SortInputsPredicate = ([&](UNiagaraHierarchyItemBase& ItemA, UNiagaraHierarchyItemBase& ItemB)
+		auto SortInputsPredicate = ([&](UHierarchyElement& ItemA, UHierarchyElement& ItemB)
 		{
 			// while these should generally be valid, it's possible due to inheritance that a user created fluid system is executing this code despite the fluid plugin being disabled.
 			// in that case, asset variable metadata might not contain these entries.
@@ -310,8 +310,8 @@ void UNiagaraEmitterEditorData::PostLoad_TransferSummaryDataToNewFormat()
 			return SortOrderA < SortOrderB;
 		});
 
-		TMap<UNiagaraHierarchyItemBase*, int32> CategoryMinimumInputSortOrderMap;
-		for(UNiagaraHierarchyCategory* Category : TopLevelCategories)
+		TMap<UHierarchyElement*, int32> CategoryMinimumInputSortOrderMap;
+		for(UHierarchyCategory* Category : TopLevelCategories)
 		{
 			Category->SortChildren(SortInputsPredicate, false);
 
@@ -327,7 +327,7 @@ void UNiagaraEmitterEditorData::PostLoad_TransferSummaryDataToNewFormat()
 			if(Category->GetChildren().Num() > 0)
 			{
 				// since we have sorted the children by sort priority, child 0 has the lowest sort priority number (= highest priority)
-				UNiagaraHierarchyItemBase* MinimumItem = Category->GetChildren()[0];
+				UHierarchyElement* MinimumItem = Category->GetChildren()[0];
 				int32 MinimumCategorySortOrder = SummarySortOrder[MinimumItem];
 				if(MinimumCategorySortOrder == INDEX_NONE)
 				{
@@ -345,7 +345,7 @@ void UNiagaraEmitterEditorData::PostLoad_TransferSummaryDataToNewFormat()
 			}
 		}
 
-		auto SortCategoriesPredicate = ([&](UNiagaraHierarchyItemBase& ItemA, UNiagaraHierarchyItemBase& ItemB)
+		auto SortCategoriesPredicate = ([&](UHierarchyElement& ItemA, UHierarchyElement& ItemB)
 		{
 			return CategoryMinimumInputSortOrderMap[&ItemA] < CategoryMinimumInputSortOrderMap[&ItemB];
 		});
@@ -354,12 +354,12 @@ void UNiagaraEmitterEditorData::PostLoad_TransferSummaryDataToNewFormat()
 
 		// we now perform identity fixup for categories using the chain of category names.
 		// we need to do this so the generated identity is stable instead of random.
-		TArray<UNiagaraHierarchyCategory*> AllCategories;
+		TArray<UHierarchyCategory*> AllCategories;
 		SummaryViewRoot->GetChildrenOfType(AllCategories, true);
 		
-		for(UNiagaraHierarchyCategory* Category : AllCategories)
+		for(UHierarchyCategory* Category : AllCategories)
 		{
-			FNiagaraHierarchyIdentity StableIdentity;
+			FHierarchyElementIdentity StableIdentity;
 			StableIdentity.Names = GenerateStableIdentityForCategories(Category);
 			Category->SetIdentity(StableIdentity);
 		}
@@ -426,7 +426,7 @@ void UNiagaraEmitterEditorData::PostLoad()
 
 	if(SummaryViewRoot == nullptr)
 	{
-		SummaryViewRoot = NewObject<UNiagaraHierarchyRoot>(this, TEXT("SummaryViewRoot"), RF_Transactional);
+		SummaryViewRoot = NewObject<UHierarchyRoot>(this, TEXT("SummaryViewRoot"), RF_Transactional);
 	}
 
 	if(UPackage* Package = GetPackage())
@@ -498,7 +498,7 @@ FSimpleMulticastDelegate& UNiagaraEmitterEditorData::OnSummaryViewStateChanged()
 	return OnSummaryViewStateChangedDelegate;
 }
 
-const TArray<UNiagaraHierarchySection*>& UNiagaraEmitterEditorData::GetSummarySections() const
+const TArray<UHierarchySection*>& UNiagaraEmitterEditorData::GetSummarySections() const
 {
 	return SummaryViewRoot->GetSectionData();
 }
diff --git a/Engine/Plugins/FX/Niagara/Source/NiagaraEditor/Private/NiagaraGraph.cpp b/Engine/Plugins/FX/Niagara/Source/NiagaraEditor/Private/NiagaraGraph.cpp
--- a/Engine/Plugins/FX/Niagara/Source/NiagaraEditor/Private/NiagaraGraph.cpp
+++ b/Engine/Plugins/FX/Niagara/Source/NiagaraEditor/Private/NiagaraGraph.cpp
@@ -440,7 +440,7 @@ UNiagaraGraph::UNiagaraGraph()
 	Schema = UEdGraphSchema_Niagara::StaticClass();
 	ChangeId = FGuid::NewGuid();
 	
-	ParameterHierarchyRoot = CreateDefaultSubobject<UNiagaraHierarchyRoot>("ParameterHierarchyRoot");
+	ParameterHierarchyRoot = CreateDefaultSubobject<UHierarchyRoot>("ParameterHierarchyRoot");
 }
 
 void UNiagaraGraph::PostLoad()
@@ -2375,8 +2375,8 @@ TArray<FGuid> UNiagaraGraph::GetChildScriptVariableGuidsForInput(FGuid VariableG
 
 			if(MetaData.IsSet())
 			{
-				FNiagaraHierarchyIdentity Identity({MetaData.GetValue().GetVariableGuid()}, {});
-				if(UNiagaraHierarchyItemBase* FoundHierarchyParameter = ParameterHierarchyRoot->FindChildWithIdentity(Identity, true))
+				FHierarchyElementIdentity Identity({MetaData.GetValue().GetVariableGuid()}, {});
+				if(UHierarchyElement* FoundHierarchyParameter = ParameterHierarchyRoot->FindChildWithIdentity(Identity, true))
 				{
 					TArray<UNiagaraHierarchyScriptParameter*> ChildrenScriptParameters;
 					FoundHierarchyParameter->GetChildrenOfType(ChildrenScriptParameters, false);
@@ -4342,20 +4342,20 @@ bool UNiagaraGraph::ReferencesStaticVariable(FNiagaraStaticVariableSearchContext
 /** Migrates the deprecated data from the provided OwnerData and populates the ParameterHierarchyRoot */
 void UNiagaraGraph::MigrateParameterScriptDataToHierarchyRoot(FVersionedNiagaraScriptData& OwnerData)
 {
-	UNiagaraHierarchyRoot* ParameterRoot = GetScriptParameterHierarchyRoot();
+	UHierarchyRoot* ParameterRoot = GetScriptParameterHierarchyRoot();
 
-	TMap<FName, UNiagaraHierarchySection*> HierarchySections;
-	TMap<FName, UNiagaraHierarchyCategory*> IdentityCategoryMap;
+	TMap<FName, UHierarchySection*> HierarchySections;
+	TMap<FName, UHierarchyCategory*> IdentityCategoryMap;
 	TMap<FName, FName> SectionCategoryMapping;
 		
 	// First we transfer all sections and create contained categories
 	for(const FNiagaraStackSection& StackSection : OwnerData.InputSections_DEPRECATED)
 	{
 		FText SectionDisplayName = StackSection.SectionDisplayName.IsEmptyOrWhitespace() ? FText::FromName(StackSection.SectionIdentifier) : StackSection.SectionDisplayName;
-		UNiagaraHierarchySection* HierarchySection = ParameterRoot->AddSection(SectionDisplayName);
+		UHierarchySection* HierarchySection = ParameterRoot->AddSection(SectionDisplayName);
 		HierarchySection->SetTooltip(StackSection.Tooltip);
 
-		FNiagaraHierarchyIdentity SectionIdentity;
+		FHierarchyElementIdentity SectionIdentity;
 		SectionIdentity.Names.Add(StackSection.SectionIdentifier);
 		HierarchySection->SetIdentity(SectionIdentity);
 			
@@ -4405,7 +4405,7 @@ void UNiagaraGraph::MigrateParameterScriptDataToHierarchyRoot(FVersionedNiagaraS
 	for(const auto& ParentChildPair : ParentChildInputMap)
 	{
 		// By default, we add to the root directly, unless a category was specified
-		UNiagaraHierarchyItemBase* Owner = ParameterRoot;
+		UHierarchyElement* Owner = ParameterRoot;
 			
 		FText SpecifiedCategoryName = ParentChildPair.Key->Metadata.GetCategoryName_DEPRECATED();
 		if(SpecifiedCategoryName.IsEmptyOrWhitespace() == false)
@@ -4439,12 +4439,12 @@ void UNiagaraGraph::MigrateParameterScriptDataToHierarchyRoot(FVersionedNiagaraS
 	}
 
 	// now that we have created all inputs, categories and sections, we fix up the order
-	TArray<UNiagaraHierarchyCategory*> TopLevelCategories;
+	TArray<UHierarchyCategory*> TopLevelCategories;
 	ParameterRoot->GetChildrenOfType(TopLevelCategories, false);
 		
 	// inputs are generally sorted by editor sort priority, but within a given context (within a category etc.)
 	// child inputs also already are sorted when constructed
-	auto SortInputsPredicate = ([&](UNiagaraHierarchyItemBase& ItemA, UNiagaraHierarchyItemBase& ItemB)
+	auto SortInputsPredicate = ([&](UHierarchyElement& ItemA, UHierarchyElement& ItemB)
 	{
 		UNiagaraHierarchyScriptParameter* ScriptParameterA = Cast<UNiagaraHierarchyScriptParameter>(&ItemA);
 		UNiagaraHierarchyScriptParameter* ScriptParameterB = Cast<UNiagaraHierarchyScriptParameter>(&ItemB);
@@ -4476,8 +4476,8 @@ void UNiagaraGraph::MigrateParameterScriptDataToHierarchyRoot(FVersionedNiagaraS
 
 	// We now sort inputs, and keep track of the minimum sort order per category
 	// This is because the stack determined category order by input sort priority (the category with the parameter of the lowest priority value would come first)
-	TMap<UNiagaraHierarchyItemBase*, int32> CategoryMinimumInputSortOrderMap;
-	for(UNiagaraHierarchyCategory* Category : TopLevelCategories)
+	TMap<UHierarchyElement*, int32> CategoryMinimumInputSortOrderMap;
+	for(UHierarchyCategory* Category : TopLevelCategories)
 	{
 		Category->SortChildren(SortInputsPredicate, false);
 		if(Category->GetChildren().Num() > 0)
@@ -4494,9 +4494,9 @@ void UNiagaraGraph::MigrateParameterScriptDataToHierarchyRoot(FVersionedNiagaraS
 		}
 	}
 
-	auto SortRootChildrenPredicate = ([&](UNiagaraHierarchyItemBase& ItemA, UNiagaraHierarchyItemBase& ItemB)
+	auto SortRootChildrenPredicate = ([&](UHierarchyElement& ItemA, UHierarchyElement& ItemB)
 	{
-		if(ItemA.IsA<UNiagaraHierarchyCategory>() && ItemB.IsA<UNiagaraHierarchyCategory>())
+		if(ItemA.IsA<UHierarchyCategory>() && ItemB.IsA<UHierarchyCategory>())
 		{
 			if(CategoryMinimumInputSortOrderMap[&ItemA] != CategoryMinimumInputSortOrderMap[&ItemB])
 			{
@@ -4507,12 +4507,12 @@ void UNiagaraGraph::MigrateParameterScriptDataToHierarchyRoot(FVersionedNiagaraS
 		}
 
 		// Categories always come after items to ensure consistency with old behavior
-		if(ItemA.IsA<UNiagaraHierarchyCategory>() && ItemB.IsA<UNiagaraHierarchyItem>())
+		if(ItemA.IsA<UHierarchyCategory>() && ItemB.IsA<UHierarchyItem>())
 		{
 			return false;
 		}
 			
-		if(ItemB.IsA<UNiagaraHierarchyCategory>() && ItemA.IsA<UNiagaraHierarchyItem>())
+		if(ItemB.IsA<UHierarchyCategory>() && ItemA.IsA<UHierarchyItem>())
 		{
 			return true;
 		}
diff --git a/Engine/Plugins/FX/Niagara/Source/NiagaraEditor/Private/NiagaraScriptMergeManager.cpp b/Engine/Plugins/FX/Niagara/Source/NiagaraEditor/Private/NiagaraScriptMergeManager.cpp
--- a/Engine/Plugins/FX/Niagara/Source/NiagaraEditor/Private/NiagaraScriptMergeManager.cpp
+++ b/Engine/Plugins/FX/Niagara/Source/NiagaraEditor/Private/NiagaraScriptMergeManager.cpp
@@ -1975,7 +1975,7 @@ bool FNiagaraScriptMergeManager::IsModuleInputDifferentFromBase(const FVersioned
 		ScriptStackDiffResults.ModifiedOtherInputOverrides.FindByPredicate(FindInputOverrideByInputName) != nullptr;
 }
 
-bool FNiagaraScriptMergeManager::DoesSummaryItemExistInBase(const FVersionedNiagaraEmitter& Emitter, FNiagaraHierarchyIdentity Identity)
+bool FNiagaraScriptMergeManager::DoesSummaryItemExistInBase(const FVersionedNiagaraEmitter& Emitter, FHierarchyElementIdentity Identity)
 {
 	SCOPE_CYCLE_COUNTER(STAT_NiagaraEditor_ScriptMergeManager_DoesSummaryItemExistInBase);
 
@@ -2482,33 +2482,33 @@ void FNiagaraScriptMergeManager::DiffEmitterSummary(const UNiagaraEmitterEditorD
 
 	DiffResults.NewShouldShowSummaryViewValue = OtherEditorData->ShouldShowSummaryView();
 	
-	TArray<UNiagaraHierarchyItemBase*> AddedItemsInOther;
-	TArray<UNiagaraHierarchyItemBase*> RemovedItemsInBase;
-	TArray<UNiagaraHierarchySection*> AddedSections;
+	TArray<UHierarchyElement*> AddedItemsInOther;
+	TArray<UHierarchyElement*> RemovedItemsInBase;
+	TArray<UHierarchySection*> AddedSections;
 	
-	UNiagaraHierarchyRoot* BaseRoot = BaseEditorData->GetSummaryRoot();
-	UNiagaraHierarchyRoot* OtherRoot = OtherEditorData->GetSummaryRoot();	
+	UHierarchyRoot* BaseRoot = BaseEditorData->GetSummaryRoot();
+	UHierarchyRoot* OtherRoot = OtherEditorData->GetSummaryRoot();	
 
-	TArray<UNiagaraHierarchyItemBase*> BaseItems;
+	TArray<UHierarchyElement*> BaseItems;
 	BaseRoot->GetChildrenOfType(BaseItems, true);
 
-	TArray<UNiagaraHierarchyItemBase*> OtherItems;
+	TArray<UHierarchyElement*> OtherItems;
 	OtherRoot->GetChildrenOfType(OtherItems, true);
 	
-	FListDiffResults<UNiagaraHierarchyItemBase*> SummaryItemDiff = DiffLists<UNiagaraHierarchyItemBase*, FNiagaraHierarchyIdentity>(
+	FListDiffResults<UHierarchyElement*> SummaryItemDiff = DiffLists<UHierarchyElement*, FHierarchyElementIdentity>(
 		BaseItems,
 		OtherItems,
-		[](UNiagaraHierarchyItemBase* HierarchyItem) { return HierarchyItem->GetPersistentIdentity(); });
+		[](UHierarchyElement* HierarchyItem) { return HierarchyItem->GetPersistentIdentity(); });
 	
 	DiffResults.AddedSummaryEntriesInOther.Append(SummaryItemDiff.AddedOtherValues);
 	
-	TArray<UNiagaraHierarchySection*> BaseSections = BaseRoot->GetSectionData();
-	TArray<UNiagaraHierarchySection*> OtherSections = OtherRoot->GetSectionData();
+	TArray<UHierarchySection*> BaseSections = BaseRoot->GetSectionData();
+	TArray<UHierarchySection*> OtherSections = OtherRoot->GetSectionData();
 
-	for (UNiagaraHierarchySection* OtherSection : OtherSections)
+	for (UHierarchySection* OtherSection : OtherSections)
 	{
-		FNiagaraHierarchyIdentity OtherSectionIdentity = OtherSection->GetPersistentIdentity();
-		if(!BaseSections.ContainsByPredicate([OtherSectionIdentity](UNiagaraHierarchySection* BaseSection)
+		FHierarchyElementIdentity OtherSectionIdentity = OtherSection->GetPersistentIdentity();
+		if(!BaseSections.ContainsByPredicate([OtherSectionIdentity](UHierarchySection* BaseSection)
 		{
 			return BaseSection->GetPersistentIdentity() == OtherSectionIdentity;
 		}))
@@ -3736,23 +3736,23 @@ FNiagaraScriptMergeManager::FApplyDiffResults FNiagaraScriptMergeManager::ApplyE
 	
 	EditorData->Modify();
 
-	UNiagaraHierarchyRoot* BaseRoot = EditorData->GetSummaryRoot();
+	UHierarchyRoot* BaseRoot = EditorData->GetSummaryRoot();
 
 	// this will copy over the section without any child elements
-	for (const UNiagaraHierarchySection* AddedOtherSection : DiffResults.AddedSummarySectionsInOther)
+	for (const UHierarchySection* AddedOtherSection : DiffResults.AddedSummarySectionsInOther)
 	{
 		BaseRoot->DuplicateSectionFromOtherRoot(*AddedOtherSection);
 	}
 
-	for (UNiagaraHierarchyItemBase* AddedInput : DiffResults.AddedSummaryEntriesInOther)
+	for (UHierarchyElement* AddedInput : DiffResults.AddedSummaryEntriesInOther)
 	{
-		if (UNiagaraHierarchyItemBase* SummaryItemParent = AddedInput->GetTypedOuter<UNiagaraHierarchyItemBase>())
+		if (UHierarchyElement* SummaryItemParent = AddedInput->GetTypedOuter<UHierarchyElement>())
 		{
-			FNiagaraHierarchyIdentity ParentIdentity = SummaryItemParent->GetPersistentIdentity();
+			FHierarchyElementIdentity ParentIdentity = SummaryItemParent->GetPersistentIdentity();
 	
 			if(ParentIdentity.IsValid())
 			{
-				const bool bParentIsRoot = SummaryItemParent->IsA<UNiagaraHierarchyRoot>();
+				const bool bParentIsRoot = SummaryItemParent->IsA<UHierarchyRoot>();
 				// if the parent is the root, we can't use the parent identity to copy over the child as the two roots are supposed to be different, so we add it directly
 				if(bParentIsRoot)
 				{
@@ -3763,10 +3763,10 @@ FNiagaraScriptMergeManager::FApplyDiffResults FNiagaraScriptMergeManager::ApplyE
 					*/
 					if(BaseRoot->FindChildWithIdentity(AddedInput->GetPersistentIdentity(), false) == nullptr)
 					{
-						UNiagaraHierarchyItemBase* NewChild = BaseRoot->CopyAndAddItemAsChild(*AddedInput);
+						UHierarchyElement* NewChild = BaseRoot->CopyAndAddItemAsChild(*AddedInput);
 						NewChild->GetChildrenMutable().Empty();
 						// for root level category, we have to fixup the section objects they point to as after duplicating them they will still point to the original section
-						if(UNiagaraHierarchyCategory* AsCategory = Cast<UNiagaraHierarchyCategory>(NewChild))
+						if(UHierarchyCategory* AsCategory = Cast<UHierarchyCategory>(NewChild))
 						{
 							AsCategory->FixupSectionLinkage();
 						}
@@ -3779,7 +3779,7 @@ FNiagaraScriptMergeManager::FApplyDiffResults FNiagaraScriptMergeManager::ApplyE
 					// we only want to add that child if it hasn't been added already
 					if(BaseRoot->FindChildWithIdentity(AddedInput->GetPersistentIdentity(), true) == nullptr)
 					{
-						UNiagaraHierarchyItemBase* NewChild = BaseRoot->CopyAndAddItemUnderParentIdentity(*AddedInput, ParentIdentity);
+						UHierarchyElement* NewChild = BaseRoot->CopyAndAddItemUnderParentIdentity(*AddedInput, ParentIdentity);
 						
 						//if(!ensure(NewChild != nullptr))
 						if(NewChild == nullptr)
diff --git a/Engine/Plugins/FX/Niagara/Source/NiagaraEditor/Private/NiagaraScriptMergeManager.h b/Engine/Plugins/FX/Niagara/Source/NiagaraEditor/Private/NiagaraScriptMergeManager.h
--- a/Engine/Plugins/FX/Niagara/Source/NiagaraEditor/Private/NiagaraScriptMergeManager.h
+++ b/Engine/Plugins/FX/Niagara/Source/NiagaraEditor/Private/NiagaraScriptMergeManager.h
@@ -404,8 +404,8 @@ struct FNiagaraEmitterDiffResults
 	TArray<TSharedRef<FNiagaraRendererMergeAdapter>> ModifiedBaseRenderers;
 	TArray<TSharedRef<FNiagaraRendererMergeAdapter>> ModifiedOtherRenderers;
 
-	TArray<UNiagaraHierarchyItemBase*> AddedSummaryEntriesInOther;	
-	TArray<UNiagaraHierarchySection*> AddedSummarySectionsInOther;
+	TArray<UHierarchyElement*> AddedSummaryEntriesInOther;	
+	TArray<UHierarchySection*> AddedSummarySectionsInOther;
 
 	TOptional<bool> NewShouldShowSummaryViewValue;
 	
@@ -455,7 +455,7 @@ public:
 	bool IsModuleInputDifferentFromBase(const FVersionedNiagaraEmitter& Emitter, const FVersionedNiagaraEmitter& BaseEmitter, ENiagaraScriptUsage ScriptUsage, FGuid ScriptUsageId, FGuid ModuleId, FNiagaraVariableBase Variable);
 
 	/** Check whether a given summary item identity exists in the parent of a given emitter. */
-	bool DoesSummaryItemExistInBase(const FVersionedNiagaraEmitter& Emitter, FNiagaraHierarchyIdentity Identity);
+	bool DoesSummaryItemExistInBase(const FVersionedNiagaraEmitter& Emitter, FHierarchyElementIdentity Identity);
 	
 	FApplyDiffResults ResetModuleInputToBase(const FVersionedNiagaraEmitter& VersionedEmitter, const FVersionedNiagaraEmitter& VersionedBaseEmitter, ENiagaraScriptUsage ScriptUsage, FGuid ScriptUsageId, FGuid ModuleId, FString InputName);
 
diff --git a/Engine/Plugins/FX/Niagara/Source/NiagaraEditor/Private/NiagaraSystemEditorData.cpp b/Engine/Plugins/FX/Niagara/Source/NiagaraEditor/Private/NiagaraSystemEditorData.cpp
--- a/Engine/Plugins/FX/Niagara/Source/NiagaraEditor/Private/NiagaraSystemEditorData.cpp
+++ b/Engine/Plugins/FX/Niagara/Source/NiagaraEditor/Private/NiagaraSystemEditorData.cpp
@@ -68,7 +68,7 @@ UNiagaraSystemEditorData::UNiagaraSystemEditorData(const FObjectInitializer& Obj
 {
 	RootFolder = ObjectInitializer.CreateDefaultSubobject<UNiagaraSystemEditorFolder>(this, TEXT("RootFolder"));
 	StackEditorData = ObjectInitializer.CreateDefaultSubobject<UNiagaraStackEditorData>(this, TEXT("StackEditorData"));
-	UserParameterHierarchy = CreateDefaultSubobject<UNiagaraHierarchyRoot>(TEXT("UserParameterHierarchyRoot"));
+	UserParameterHierarchy = CreateDefaultSubobject<UHierarchyRoot>(TEXT("UserParameterHierarchyRoot"));
 	OwnerTransform.SetLocation(FVector(0.0f, 0.0f, 0.0f));
 	PlaybackRangeMin = 0;
 	PlaybackRangeMax = 10;
@@ -104,7 +104,7 @@ void UNiagaraSystemEditorData::PostLoadFromOwner(UObject* InOwner)
 	}
 	if (UserParameterHierarchy == nullptr)
 	{
-		UserParameterHierarchy = NewObject<UNiagaraHierarchyRoot>(this, TEXT("UserParameterHierarchyRoot"), RF_Transactional);
+		UserParameterHierarchy = NewObject<UHierarchyRoot>(this, TEXT("UserParameterHierarchyRoot"), RF_Transactional);
 	}
 
 	TArray<FNiagaraVariable> UserParameters;
@@ -156,7 +156,7 @@ void UNiagaraSystemEditorData::DeclareConstructClasses(TArray<FTopLevelAssetPath
 	OutConstructClasses.Add(FTopLevelAssetPath(UEdGraph::StaticClass()));
 	OutConstructClasses.Add(FTopLevelAssetPath(UNiagaraOverviewNode::StaticClass()));
 	OutConstructClasses.Add(FTopLevelAssetPath(UNiagaraScriptVariable::StaticClass()));
-	OutConstructClasses.Add(FTopLevelAssetPath(UNiagaraHierarchyRoot::StaticClass()));
+	OutConstructClasses.Add(FTopLevelAssetPath(UHierarchyRoot::StaticClass()));
 }
 #endif
 
diff --git a/Engine/Plugins/FX/Niagara/Source/NiagaraEditor/Private/NiagaraValidationRules.cpp b/Engine/Plugins/FX/Niagara/Source/NiagaraEditor/Private/NiagaraValidationRules.cpp
--- a/Engine/Plugins/FX/Niagara/Source/NiagaraEditor/Private/NiagaraValidationRules.cpp
+++ b/Engine/Plugins/FX/Niagara/Source/NiagaraEditor/Private/NiagaraValidationRules.cpp
@@ -32,7 +32,7 @@
 #include "NiagaraNodeOutput.h"
 #include "NiagaraNodeParameterMapFor.h"
 #include "Materials/MaterialInterface.h"
-
+#include "ScopedTransaction.h"
 #include "DeviceProfiles/DeviceProfile.h"
 #include "DeviceProfiles/DeviceProfileManager.h"
 
diff --git a/Engine/Plugins/FX/Niagara/Source/NiagaraEditor/Private/Toolkits/NiagaraScriptToolkit.cpp b/Engine/Plugins/FX/Niagara/Source/NiagaraEditor/Private/Toolkits/NiagaraScriptToolkit.cpp
--- a/Engine/Plugins/FX/Niagara/Source/NiagaraEditor/Private/Toolkits/NiagaraScriptToolkit.cpp
+++ b/Engine/Plugins/FX/Niagara/Source/NiagaraEditor/Private/Toolkits/NiagaraScriptToolkit.cpp
@@ -41,7 +41,7 @@
 #include "ViewModels/NiagaraParameterPanelViewModel.h"
 #include "ViewModels/NiagaraScriptViewModel.h"
 #include "ViewModels/HierarchyEditor/NiagaraHierarchyScriptParametersViewModel.h"
-#include "Widgets/SNiagaraHierarchyEditor.h"
+#include "Widgets/SHierarchyEditor.h"
 #include "Widgets/Docking/SDockTab.h"
 #include "Widgets/SNiagaraParameterDefinitionsPanel.h"
 #include "Widgets/SNiagaraParameterPanel.h"
@@ -635,7 +635,7 @@ TSharedRef<SDockTab> FNiagaraScriptToolkit::SpawnTabHierarchyEditor(const FSpawn
 			SNew(SBox)
 			.AddMetaData<FTagMetaData>(FTagMetaData(TEXT("UserParameters")))
 			[
-				SNew(SNiagaraHierarchyEditor, ParametersHierarchyViewModel)
+				SNew(SHierarchyEditor, ParametersHierarchyViewModel)
 				.OnGenerateRowContentWidget_Static(&FNiagaraEditorUtilities::HierarchyEditor::Scripts::GenerateRowContentForScriptParameterHierarchyEditor)
 				//.OnGenerateCustomDetailsPanelNameWidget_Static(&GenerateCustomDetailsPanelNameWidgetForUserParameterEditor)
 			]
diff --git a/Engine/Plugins/FX/Niagara/Source/NiagaraEditor/Private/Toolkits/SystemToolkitModes/NiagaraSystemToolkitModeBase.cpp b/Engine/Plugins/FX/Niagara/Source/NiagaraEditor/Private/Toolkits/SystemToolkitModes/NiagaraSystemToolkitModeBase.cpp
--- a/Engine/Plugins/FX/Niagara/Source/NiagaraEditor/Private/Toolkits/SystemToolkitModes/NiagaraSystemToolkitModeBase.cpp
+++ b/Engine/Plugins/FX/Niagara/Source/NiagaraEditor/Private/Toolkits/SystemToolkitModes/NiagaraSystemToolkitModeBase.cpp
@@ -27,7 +27,7 @@
 #include "Widgets/SNiagaraScriptGraph.h"
 #include "Widgets/SNiagaraSelectedObjectsDetails.h"
 #include "Widgets/SNiagaraGeneratedCodeView.h"
-#include "Widgets/SNiagaraHierarchyEditor.h"
+#include "Widgets/SHierarchyEditor.h"
 #include "Widgets/Layout/SWidgetSwitcher.h"
 #include "Widgets/Layout/SBox.h"
 #include "Widgets/Input/SButton.h"
@@ -874,12 +874,12 @@ TSharedRef<SDockTab> FNiagaraSystemToolkitModeBase::SpawnTab_UserParameters(cons
 	return SpawnedTab;
 }
 
-TSharedRef<SWidget> GenerateRowContentForUserParameterHierarchyEditor(TSharedRef<FNiagaraHierarchyItemViewModelBase> HierarchyItem, TSharedPtr<FNiagaraSystemViewModel> SystemViewModel)
+TSharedRef<SWidget> GenerateRowContentForUserParameterHierarchyEditor(TSharedRef<FHierarchyElementViewModel> HierarchyItem, TSharedPtr<FNiagaraSystemViewModel> SystemViewModel)
 {
-	if(HierarchyItem->GetDataMutable()->IsA<UNiagaraHierarchyCategory>())
+	if(HierarchyItem->GetDataMutable()->IsA<UHierarchyCategory>())
 	{
-		TSharedRef<FNiagaraHierarchyCategoryViewModel> TreeViewCategory = StaticCastSharedRef<FNiagaraHierarchyCategoryViewModel>(HierarchyItem);
-		return SNew(SNiagaraHierarchyCategory, TreeViewCategory);
+		TSharedRef<FHierarchyCategoryViewModel> TreeViewCategory = StaticCastSharedRef<FHierarchyCategoryViewModel>(HierarchyItem);
+		return SNew(SHierarchyCategory, TreeViewCategory);
 	}
 	else if(const UNiagaraHierarchyUserParameter* UserParameter = Cast<UNiagaraHierarchyUserParameter>(HierarchyItem->GetData()))
 	{
@@ -897,14 +897,14 @@ TSharedRef<SWidget> GenerateRowContentForUserParameterHierarchyEditor(TSharedRef
 	return SNullWidget::NullWidget;
 }
 
-TSharedRef<SWidget> GenerateCustomDetailsPanelNameWidgetForUserParameterEditor(TSharedPtr<FNiagaraHierarchyItemViewModelBase> HierarchyItem)
+TSharedRef<SWidget> GenerateCustomDetailsPanelNameWidgetForUserParameterEditor(TSharedPtr<FHierarchyElementViewModel> HierarchyItem)
 {
 	if(!HierarchyItem.IsValid())
 	{
 		return SNew(STextBlock).Text(FText::FromString("None selected"));
 	}
 	
-	if(HierarchyItem->GetData()->IsA<UNiagaraHierarchyCategory>() || HierarchyItem->GetData()->IsA<UNiagaraHierarchySection>())
+	if(HierarchyItem->GetData()->IsA<UHierarchyCategory>() || HierarchyItem->GetData()->IsA<UHierarchySection>())
 	{
 		return SNew(SBox).Padding(2.f)
 		[
@@ -913,7 +913,7 @@ TSharedRef<SWidget> GenerateCustomDetailsPanelNameWidgetForUserParameterEditor(T
 			{
 				return FText::FromString(HierarchyItem->ToString());
 			})
-			.TextStyle(&FNiagaraEditorStyle::Get().GetWidgetStyle<FTextBlockStyle>("NiagaraEditor.HierarchyEditor.CategoryTextBlock"))
+			.TextStyle(&FHierarchyEditorStyle::Get().GetWidgetStyle<FTextBlockStyle>("HierarchyEditor.CategoryTextBlock"))
 		];
 	}
 	else if(const UNiagaraHierarchyUserParameter* UserParameter = Cast<UNiagaraHierarchyUserParameter>(HierarchyItem->GetData()))
@@ -934,7 +934,7 @@ TSharedRef<SDockTab> FNiagaraSystemToolkitModeBase::SpawnTab_UserParametersHiera
 			SNew(SBox)
 			.AddMetaData<FTagMetaData>(FTagMetaData(TEXT("UserParameters")))
 			[
-				SNew(SNiagaraHierarchyEditor, SystemToolkit.Pin()->GetSystemViewModel()->GetUserParametersHierarchyViewModel())
+				SNew(SHierarchyEditor, SystemToolkit.Pin()->GetSystemViewModel()->GetUserParametersHierarchyViewModel())
 				.OnGenerateRowContentWidget_Static(&GenerateRowContentForUserParameterHierarchyEditor, SystemToolkit.Pin()->GetSystemViewModel())
 				.OnGenerateCustomDetailsPanelNameWidget_Static(&GenerateCustomDetailsPanelNameWidgetForUserParameterEditor)
 			]
@@ -943,7 +943,7 @@ TSharedRef<SDockTab> FNiagaraSystemToolkitModeBase::SpawnTab_UserParametersHiera
 	return SpawnedTab;
 }
 
-TSharedRef<SWidget> GenerateRowContentForSummaryViewHierarchyEditor(TSharedRef<FNiagaraHierarchyItemViewModelBase> HierarchyItem, TSharedPtr<FNiagaraEmitterViewModel> EmitterViewModel)
+TSharedRef<SWidget> GenerateRowContentForSummaryViewHierarchyEditor(TSharedRef<FHierarchyElementViewModel> HierarchyItem, TSharedPtr<FNiagaraEmitterViewModel> EmitterViewModel)
 {
 	if(const UNiagaraHierarchyModuleInput* ModuleInput = Cast<UNiagaraHierarchyModuleInput>(HierarchyItem->GetData()))
 	{
@@ -962,7 +962,7 @@ TSharedRef<SWidget> GenerateRowContentForSummaryViewHierarchyEditor(TSharedRef<F
 	}
 	else if(const UNiagaraHierarchyModule* Module = Cast<UNiagaraHierarchyModule>(HierarchyItem->GetDataMutable()))
 	{
-		TSharedPtr<FNiagaraHierarchyItemViewModelBase> ItemPtr = HierarchyItem;
+		TSharedPtr<FHierarchyElementViewModel> ItemPtr = HierarchyItem;
 		TSharedPtr<FNiagaraFunctionViewModel> ModuleViewModel = StaticCastSharedPtr<FNiagaraFunctionViewModel>(ItemPtr);
 
 		if(UNiagaraNodeAssignment* AssignmentNode = Cast<UNiagaraNodeAssignment>(ModuleViewModel->GetFunctionCallNode()))
@@ -976,7 +976,7 @@ TSharedRef<SWidget> GenerateRowContentForSummaryViewHierarchyEditor(TSharedRef<F
 	}
 	else if(const UNiagaraHierarchyRenderer* Renderer = Cast<UNiagaraHierarchyRenderer>(HierarchyItem->GetDataMutable()))
 	{
-		TSharedPtr<FNiagaraHierarchyItemViewModelBase> ItemPtr = HierarchyItem;
+		TSharedPtr<FHierarchyElementViewModel> ItemPtr = HierarchyItem;
 		TSharedPtr<FNiagaraHierarchyRendererViewModel> RendererViewModel = StaticCastSharedPtr<FNiagaraHierarchyRendererViewModel>(ItemPtr);
 		return SNew(SHorizontalBox)
 			+ SHorizontalBox::Slot()
@@ -992,7 +992,7 @@ TSharedRef<SWidget> GenerateRowContentForSummaryViewHierarchyEditor(TSharedRef<F
 	}
 	else if(const UNiagaraHierarchyEventHandler* EventHandler = Cast<UNiagaraHierarchyEventHandler>(HierarchyItem->GetDataMutable()))
 	{
-		TSharedPtr<FNiagaraHierarchyItemViewModelBase> ItemPtr = HierarchyItem;
+		TSharedPtr<FHierarchyElementViewModel> ItemPtr = HierarchyItem;
 		TSharedPtr<FNiagaraHierarchyEventHandlerViewModel> EventHandlerViewModel = StaticCastSharedPtr<FNiagaraHierarchyEventHandlerViewModel>(ItemPtr);
 		return SNew(SHorizontalBox)
 			+ SHorizontalBox::Slot()
@@ -1008,7 +1008,7 @@ TSharedRef<SWidget> GenerateRowContentForSummaryViewHierarchyEditor(TSharedRef<F
 	}
 	else if(const UNiagaraHierarchySimStage* SimStage = Cast<UNiagaraHierarchySimStage>(HierarchyItem->GetDataMutable()))
 	{
-		TSharedPtr<FNiagaraHierarchyItemViewModelBase> ItemPtr = HierarchyItem;
+		TSharedPtr<FHierarchyElementViewModel> ItemPtr = HierarchyItem;
 		TSharedPtr<FNiagaraHierarchySimStageViewModel> SimStageViewModel = StaticCastSharedPtr<FNiagaraHierarchySimStageViewModel>(ItemPtr);
 		return SNew(SHorizontalBox)
 			+ SHorizontalBox::Slot()
@@ -1024,16 +1024,16 @@ TSharedRef<SWidget> GenerateRowContentForSummaryViewHierarchyEditor(TSharedRef<F
 	}
 	else if(const UNiagaraHierarchySimStageProperties* SimStageProperties = Cast<UNiagaraHierarchySimStageProperties>(HierarchyItem->GetDataMutable()))
 	{
-		TSharedPtr<FNiagaraHierarchyItemViewModelBase> ItemPtr = HierarchyItem;
+		TSharedPtr<FHierarchyElementViewModel> ItemPtr = HierarchyItem;
 		TSharedPtr<FNiagaraHierarchySimStagePropertiesViewModel> SimStagePropertiesViewModel = StaticCastSharedPtr<FNiagaraHierarchySimStagePropertiesViewModel>(ItemPtr);
 		return SNew(STextBlock).Text(SimStagePropertiesViewModel.ToSharedRef(), &FNiagaraHierarchySimStagePropertiesViewModel::ToStringAsText);
 	}
-	else if(HierarchyItem->GetDataMutable()->IsA<UNiagaraHierarchyCategory>())
+	else if(HierarchyItem->GetDataMutable()->IsA<UHierarchyCategory>())
 	{
-		TSharedRef<FNiagaraHierarchyCategoryViewModel> TreeViewCategory = StaticCastSharedRef<FNiagaraHierarchyCategoryViewModel>(HierarchyItem);
-		return SNew(SNiagaraHierarchyCategory, TreeViewCategory);
+		TSharedRef<FHierarchyCategoryViewModel> TreeViewCategory = StaticCastSharedRef<FHierarchyCategoryViewModel>(HierarchyItem);
+		return SNew(SHierarchyCategory, TreeViewCategory);
 	}
-	else if(const UNiagaraHierarchyItem* Item = Cast<UNiagaraHierarchyItem>(HierarchyItem->GetData()))
+	else if(const UHierarchyItem* Item = Cast<UHierarchyItem>(HierarchyItem->GetData()))
 	{
 		return SNew(STextBlock).Text(FText::FromString(HierarchyItem->ToString()));
 	}
@@ -1041,14 +1041,14 @@ TSharedRef<SWidget> GenerateRowContentForSummaryViewHierarchyEditor(TSharedRef<F
 	return SNullWidget::NullWidget;
 }
 
-TSharedRef<SWidget> GenerateCustomDetailsPanelNameWidgetForSummaryViewEditor(TSharedPtr<FNiagaraHierarchyItemViewModelBase> HierarchyItem)
+TSharedRef<SWidget> GenerateCustomDetailsPanelNameWidgetForSummaryViewEditor(TSharedPtr<FHierarchyElementViewModel> HierarchyItem)
 {
 	if(!HierarchyItem.IsValid())
 	{
 		return SNew(STextBlock).Text(FText::FromString("None selected"));
 	}
 	
-	if(HierarchyItem->GetData()->IsA<UNiagaraHierarchyCategory>() || HierarchyItem->GetData()->IsA<UNiagaraHierarchySection>())
+	if(HierarchyItem->GetData()->IsA<UHierarchyCategory>() || HierarchyItem->GetData()->IsA<UHierarchySection>())
 	{
 		return SNew(SBox).Padding(2.f)
 		[
@@ -1064,7 +1064,7 @@ TSharedRef<SWidget> GenerateCustomDetailsPanelNameWidgetForSummaryViewEditor(TSh
 		TSharedPtr<FNiagaraModuleInputViewModel> ModuleInputViewModel = StaticCastSharedPtr<FNiagaraModuleInputViewModel>(HierarchyItem);
 		return SNew(SNiagaraHierarchyModuleInput, ModuleInputViewModel.ToSharedRef());
 	}
-	else if(HierarchyItem->GetData()->IsA<UNiagaraHierarchyItem>())
+	else if(HierarchyItem->GetData()->IsA<UHierarchyItem>())
 	{
 		return SNew(STextBlock).Text(FText::FromString(HierarchyItem->ToString()));
 	}
@@ -1082,7 +1082,7 @@ TSharedRef<SWidget> FNiagaraSystemToolkitModeBase::CreateSummaryViewWidget() con
 		TSharedPtr<FNiagaraEmitterViewModel> EmitterViewModel = SystemToolkit.Pin()->GetSystemViewModel()->GetEmitterHandleViewModels()[0]->GetEmitterViewModel();
 		UNiagaraSummaryViewViewModel* SummaryViewHierarchyViewModel = EmitterViewModel->GetSummaryHierarchyViewModel();
 		
-		ContentWidget = SNew(SNiagaraHierarchyEditor, SummaryViewHierarchyViewModel)
+		ContentWidget = SNew(SHierarchyEditor, SummaryViewHierarchyViewModel)
 		.OnGenerateRowContentWidget_Static(&GenerateRowContentForSummaryViewHierarchyEditor, EmitterViewModel)
 		.OnGenerateCustomDetailsPanelNameWidget_Static(&GenerateCustomDetailsPanelNameWidgetForSummaryViewEditor);
 	}
@@ -1102,7 +1102,7 @@ TSharedRef<SWidget> FNiagaraSystemToolkitModeBase::CreateSummaryViewWidget() con
 			TSharedPtr<FNiagaraEmitterViewModel> EmitterViewModel = SystemToolkit.Pin()->GetSystemViewModel()->GetEmitterHandleViewModelById(SelectedEmitterHandleGuids[0])->GetEmitterViewModel();
 			UNiagaraSummaryViewViewModel* SummaryViewHierarchyViewModel = EmitterViewModel->GetSummaryHierarchyViewModel();
 		
-			ContentWidget = SNew(SNiagaraHierarchyEditor, SummaryViewHierarchyViewModel)
+			ContentWidget = SNew(SHierarchyEditor, SummaryViewHierarchyViewModel)
 			.OnGenerateRowContentWidget_Static(&GenerateRowContentForSummaryViewHierarchyEditor, EmitterViewModel)
 			.OnGenerateCustomDetailsPanelNameWidget_Static(&GenerateCustomDetailsPanelNameWidgetForSummaryViewEditor);
 		}
@@ -1141,7 +1141,7 @@ TSharedRef<SWidget> FNiagaraSystemToolkitModeBase::CreateScratchPadHierarchyWidg
 	{
 		UNiagaraHierarchyScriptParametersViewModel* ScriptHierarchyViewModel = ActiveScratchPadViewModel->GetHierarchyViewModel();
 		
-		ContentWidget = SNew(SNiagaraHierarchyEditor, ScriptHierarchyViewModel)
+		ContentWidget = SNew(SHierarchyEditor, ScriptHierarchyViewModel)
 		.OnGenerateRowContentWidget_Static(&FNiagaraEditorUtilities::HierarchyEditor::Scripts::GenerateRowContentForScriptParameterHierarchyEditor);
 		//.OnGenerateCustomDetailsPanelNameWidget_Static(&GenerateCustomDetailsPanelNameWidgetForSummaryViewEditor);
 	}
diff --git a/Engine/Plugins/FX/Niagara/Source/NiagaraEditor/Private/ViewModels/HierarchyEditor/NiagaraHierarchyScriptParametersViewModel.cpp b/Engine/Plugins/FX/Niagara/Source/NiagaraEditor/Private/ViewModels/HierarchyEditor/NiagaraHierarchyScriptParametersViewModel.cpp
--- a/Engine/Plugins/FX/Niagara/Source/NiagaraEditor/Private/ViewModels/HierarchyEditor/NiagaraHierarchyScriptParametersViewModel.cpp
+++ b/Engine/Plugins/FX/Niagara/Source/NiagaraEditor/Private/ViewModels/HierarchyEditor/NiagaraHierarchyScriptParametersViewModel.cpp
@@ -6,6 +6,7 @@
 #include "NiagaraEditorUtilities.h"
 #include "NiagaraGraph.h"
 #include "NiagaraScriptSource.h"
+#include "ScopedTransaction.h"
 #include "ViewModels/NiagaraScriptViewModel.h"
 
 #define LOCTEXT_NAMESPACE "NiagaraScriptParameterHierarchyEditor"
@@ -13,7 +14,7 @@
 void UNiagaraHierarchyScriptParameter::Initialize(UNiagaraScriptVariable& InParameterScriptVariable)
 {
 	ParameterScriptVariable = &InParameterScriptVariable;
-	SetIdentity(FNiagaraHierarchyIdentity({InParameterScriptVariable.Metadata.GetVariableGuid()}, {}));
+	SetIdentity(FHierarchyElementIdentity({InParameterScriptVariable.Metadata.GetVariableGuid()}, {}));
 }
 
 void UNiagaraHierarchyScriptParametersViewModel::Initialize(TSharedRef<FNiagaraScriptViewModel> InScriptViewModel)
@@ -21,7 +22,7 @@ void UNiagaraHierarchyScriptParametersViewModel::Initialize(TSharedRef<FNiagaraS
 	ScriptViewModelWeak = InScriptViewModel;
 
 	Cast<UNiagaraScriptSource>(ScriptViewModelWeak.Pin()->GetStandaloneScript().GetScriptData()->GetSource())->NodeGraph->OnParametersChanged().AddUObject(this, &UNiagaraHierarchyScriptParametersViewModel::OnParametersChanged);
-	UNiagaraHierarchyViewModelBase::Initialize();
+	UHierarchyViewModelBase::Initialize();
 
 	UNiagaraHierarchyScriptParameterRefreshContext* ScriptParameterRefreshContext = NewObject<UNiagaraHierarchyScriptParameterRefreshContext>(this);
 	ScriptParameterRefreshContext->SetNiagaraGraph(Cast<UNiagaraScriptSource>(InScriptViewModel->GetStandaloneScript().GetScriptData()->GetSource())->NodeGraph);
@@ -36,7 +37,7 @@ TSharedRef<FNiagaraScriptViewModel> UNiagaraHierarchyScriptParametersViewModel::
 }
 
 
-UNiagaraHierarchyRoot* UNiagaraHierarchyScriptParametersViewModel::GetHierarchyRoot() const
+UHierarchyRoot* UNiagaraHierarchyScriptParametersViewModel::GetHierarchyRoot() const
 {
 	const TArray<FVersionedNiagaraScriptWeakPtr>& Scripts = GetScriptViewModel()->GetScripts();
 	if(ensure(Scripts.Num() >= 1 && Scripts[0].Pin().Script != nullptr) == false)
@@ -53,39 +54,33 @@ UNiagaraHierarchyRoot* UNiagaraHierarchyScriptParametersViewModel::GetHierarchyR
 	return Cast<UNiagaraScriptSource>(ScriptData->GetSource())->NodeGraph->GetScriptParameterHierarchyRoot();
 }
 
-TSubclassOf<UNiagaraHierarchyCategory> UNiagaraHierarchyScriptParametersViewModel::GetCategoryDataClass() const
+TSubclassOf<UHierarchyCategory> UNiagaraHierarchyScriptParametersViewModel::GetCategoryDataClass() const
 {
 	return UNiagaraHierarchyScriptCategory::StaticClass();
 }
 
-TSharedPtr<FNiagaraHierarchyRootViewModel> UNiagaraHierarchyScriptParametersViewModel::CreateRootViewModelForData(UNiagaraHierarchyRoot* Root, bool bIsForHierarchy)
+TSharedPtr<FHierarchyElementViewModel> UNiagaraHierarchyScriptParametersViewModel::CreateCustomViewModelForElement(UHierarchyElement* Element, TSharedPtr<FHierarchyElementViewModel> Parent)
 {
-	return MakeShared<FNiagaraHierarchyScriptRootViewModel>(Root, this, bIsForHierarchy);
-}
-
-TSharedPtr<FNiagaraHierarchyItemViewModelBase> UNiagaraHierarchyScriptParametersViewModel::CreateViewModelForData(UNiagaraHierarchyItemBase* ItemBase, TSharedPtr<FNiagaraHierarchyItemViewModelBase> Parent)
-{
-	if(UNiagaraHierarchyScriptParameter* Item = Cast<UNiagaraHierarchyScriptParameter>(ItemBase))
+	if(UNiagaraHierarchyScriptParameter* Item = Cast<UNiagaraHierarchyScriptParameter>(Element))
 	{
-		return MakeShared<FNiagaraHierarchyScriptParameterViewModel>(Item, Parent, this, Parent.IsValid() ? Parent->IsForHierarchy() : false);
+		return MakeShared<FNiagaraHierarchyScriptParameterViewModel>(Item, Parent.ToSharedRef(), this);
 	}
-	else if(UNiagaraHierarchyCategory* Category = Cast<UNiagaraHierarchyCategory>(ItemBase))
+	else if(UHierarchyRoot* Root = Cast<UHierarchyRoot>(Element))
 	{
-		return MakeShared<FNiagaraHierarchyCategoryViewModel>(Category, Parent, this, Parent.IsValid() ? Parent->IsForHierarchy() : false);
+		// If the root is the hierarchy root, we know it's for the hierarchy. If not, it's the transient source root
+		bool bIsForHierarchy = GetHierarchyRoot() == Element;
+		return MakeShared<FNiagaraHierarchyScriptRootViewModel>(Root, this, bIsForHierarchy);
 	}
 
 	return nullptr;
 }
 
-void UNiagaraHierarchyScriptParametersViewModel::PrepareSourceItems(UNiagaraHierarchyRoot* SourceRoot, TSharedPtr<FNiagaraHierarchyRootViewModel> SourceRootViewModel)
+void UNiagaraHierarchyScriptParametersViewModel::PrepareSourceItems(UHierarchyRoot* SourceRoot, TSharedPtr<FHierarchyRootViewModel> SourceRootViewModel)
 {
 	UNiagaraScriptSourceBase* SourceBase = GetScriptViewModel()->GetStandaloneScript().GetScriptData()->GetSource();
 	UNiagaraScriptSource* ScriptSource = Cast<UNiagaraScriptSource>(SourceBase);
 	const TMap<FNiagaraVariable, TObjectPtr<UNiagaraScriptVariable>>& ScriptVariableMap = ScriptSource->NodeGraph->GetAllMetaData();
 
-	TArray<UNiagaraHierarchyItemBase*> OldChildren = SourceRoot->GetChildrenMutable();
-	SourceRoot->GetChildrenMutable().Empty();
-
 	for(const auto& ScriptVariablePair : ScriptVariableMap)
 	{
 		// We only want to be able to organize module inputs & static switches 
@@ -95,14 +90,6 @@ void UNiagaraHierarchyScriptParametersViewModel::PrepareSourceItems(UNiagaraHier
 		}
 		
 		UNiagaraScriptVariable* ScriptVariable = ScriptVariablePair.Value;
-		if(UNiagaraHierarchyItemBase** ExistingChild = OldChildren.FindByPredicate([ScriptVariable](UNiagaraHierarchyItemBase* ItemBase)
-		{
-			return ItemBase->GetPersistentIdentity().Guids[0] == ScriptVariable->Metadata.GetVariableGuid();
-		}))
-		{
-			SourceRoot->GetChildrenMutable().Add(*ExistingChild);
-			continue;
-		}
 
 		// since the source items are transient we need to create them here and keep them around until the end of the tool's lifetime
 		UNiagaraHierarchyScriptParameter* ScriptParameterHierarchyObject = NewObject<UNiagaraHierarchyScriptParameter>(SourceRoot);
@@ -116,25 +103,24 @@ void UNiagaraHierarchyScriptParametersViewModel::SetupCommands()
 	Super::SetupCommands();
 }
 
-TSharedRef<FNiagaraHierarchyDragDropOp> UNiagaraHierarchyScriptParametersViewModel::CreateDragDropOp(TSharedRef<FNiagaraHierarchyItemViewModelBase> Item)
+TSharedRef<FHierarchyDragDropOp> UNiagaraHierarchyScriptParametersViewModel::CreateDragDropOp(TSharedRef<FHierarchyElementViewModel> Item)
 {
-	if(UNiagaraHierarchyCategory* HierarchyCategory = Cast<UNiagaraHierarchyCategory>(Item->GetDataMutable()))
+	if(UHierarchyCategory* HierarchyCategory = Cast<UHierarchyCategory>(Item->GetDataMutable()))
 	{
-		TSharedRef<FNiagaraHierarchyDragDropOp> CategoryDragDropOp = MakeShared<FNiagaraHierarchyDragDropOp>(Item);
-		CategoryDragDropOp->SetAdditionalLabel(FText::FromString(Item->ToString()));
+		TSharedRef<FHierarchyDragDropOp> CategoryDragDropOp = MakeShared<FHierarchyDragDropOp>(Item);
 		CategoryDragDropOp->Construct();
 		return CategoryDragDropOp;
 	}
 	else if(UNiagaraHierarchyScriptParameter* ScriptParameter = Cast<UNiagaraHierarchyScriptParameter>(Item->GetDataMutable()))
 	{
-		TSharedPtr<FNiagaraHierarchyItemViewModel> ScriptParameterViewModel = StaticCastSharedRef<FNiagaraHierarchyItemViewModel>(Item);
-		TSharedRef<FNiagaraHierarchyDragDropOp> ScriptParameterDragDropOp = MakeShared<FNiagaraHierarchyScriptParameterDragDropOp>(ScriptParameterViewModel);
+		TSharedPtr<FHierarchyItemViewModel> ScriptParameterViewModel = StaticCastSharedRef<FHierarchyItemViewModel>(Item);
+		TSharedRef<FHierarchyDragDropOp> ScriptParameterDragDropOp = MakeShared<FNiagaraHierarchyScriptParameterDragDropOp>(ScriptParameterViewModel);
 		ScriptParameterDragDropOp->Construct();
 		return ScriptParameterDragDropOp;
 	}
 	
 	check(false);
-	return MakeShared<FNiagaraHierarchyDragDropOp>(nullptr);
+	return MakeShared<FHierarchyDragDropOp>(nullptr);
 }
 
 void UNiagaraHierarchyScriptParametersViewModel::FinalizeInternal()
@@ -147,8 +133,6 @@ void UNiagaraHierarchyScriptParametersViewModel::FinalizeInternal()
 			Graph->OnParametersChanged().RemoveAll(this);
 		}
 	}
-
-	Super::FinalizeInternal();
 }
 
 void UNiagaraHierarchyScriptParametersViewModel::OnParametersChanged(TOptional<UNiagaraGraph::FParametersChangedData> ParametersChangedData)
@@ -158,7 +142,7 @@ void UNiagaraHierarchyScriptParametersViewModel::OnParametersChanged(TOptional<U
 
 TSharedRef<SWidget> FNiagaraHierarchyScriptParameterDragDropOp::CreateCustomDecorator() const
 {
-	if(const UNiagaraHierarchyScriptParameter* ScriptParameter = Cast<UNiagaraHierarchyScriptParameter>(DraggedItem.Pin()->GetData()))
+	if(const UNiagaraHierarchyScriptParameter* ScriptParameter = Cast<UNiagaraHierarchyScriptParameter>(DraggedElement.Pin()->GetData()))
 	{
 		return FNiagaraParameterUtilities::GetParameterWidget(ScriptParameter->GetVariable(), false, false);
 	}
@@ -166,7 +150,7 @@ TSharedRef<SWidget> FNiagaraHierarchyScriptParameterDragDropOp::CreateCustomDeco
 	return SNullWidget::NullWidget;
 }
 
-bool FNiagaraHierarchyScriptParameterViewModel::DoesExternalDataStillExist(const UNiagaraHierarchyDataRefreshContext* Context) const
+bool FNiagaraHierarchyScriptParameterViewModel::DoesExternalDataStillExist(const UHierarchyDataRefreshContext* Context) const
 {
 	// During undo/redo it's possible the script variable becomes nullptr. If so, there is no need for this view model either
 	if(Cast<UNiagaraHierarchyScriptParameter>(GetDataMutable())->GetScriptVariable() == nullptr)
@@ -191,7 +175,7 @@ bool FNiagaraHierarchyScriptParameterViewModel::DoesExternalDataStillExist(const
 	return true;
 }
 
-FNiagaraHierarchyItemViewModelBase::FCanPerformActionResults FNiagaraHierarchyScriptParameterViewModel::CanDropOnInternal(TSharedPtr<FNiagaraHierarchyItemViewModelBase> DraggedItem, EItemDropZone ItemDropZone)
+FHierarchyElementViewModel::FCanPerformActionResults FNiagaraHierarchyScriptParameterViewModel::CanDropOnInternal(TSharedPtr<FHierarchyElementViewModel> DraggedItem, EItemDropZone ItemDropZone)
 {
 	// if the input isn't editable, we don't allow any drops on/above/below the item.
 	// Even though it technically works, the merge process will only re-add the item at the end and not preserve order so there is no point in allowing dropping above/below
@@ -202,7 +186,7 @@ FNiagaraHierarchyItemViewModelBase::FCanPerformActionResults FNiagaraHierarchySc
 
 	FCanPerformActionResults AllowDrop(false);
 	
-	TSharedPtr<FNiagaraHierarchyItemViewModelBase> TargetDropItem = AsShared();
+	TSharedPtr<FHierarchyElementViewModel> TargetDropItem = AsShared();
 
 	// we only allow drops if some general conditions are fulfilled
 	if(DraggedItem->GetData() != TargetDropItem->GetData() &&
@@ -231,7 +215,7 @@ FNiagaraHierarchyItemViewModelBase::FCanPerformActionResults FNiagaraHierarchySc
 				// we use default logic only if there is no parent input. Nested children are not allowed to contain anything but other inputs.
 				if(TargetDropItem->GetParent().Pin()->GetData<UNiagaraHierarchyScriptParameter>() == nullptr)
 				{
-					AllowDrop = FNiagaraHierarchyItemViewModel::CanDropOnInternal(DraggedItem, ItemDropZone);
+					AllowDrop = FHierarchyItemViewModel::CanDropOnInternal(DraggedItem, ItemDropZone);
 				}
 			}
 		}
@@ -240,11 +224,11 @@ FNiagaraHierarchyItemViewModelBase::FCanPerformActionResults FNiagaraHierarchySc
 	return AllowDrop;
 }
 
-void FNiagaraHierarchyScriptParameterViewModel::OnDroppedOnInternal(TSharedPtr<FNiagaraHierarchyItemViewModelBase> DroppedItem, EItemDropZone ItemDropZone)
+void FNiagaraHierarchyScriptParameterViewModel::OnDroppedOnInternal(TSharedPtr<FHierarchyElementViewModel> DroppedItem, EItemDropZone ItemDropZone)
 {
 	if(ItemDropZone != EItemDropZone::OntoItem)
 	{
-		return FNiagaraHierarchyItemViewModel::OnDroppedOnInternal(DroppedItem, ItemDropZone);
+		return FHierarchyItemViewModel::OnDroppedOnInternal(DroppedItem, ItemDropZone);
 	}
 	else
 	{
@@ -253,12 +237,12 @@ void FNiagaraHierarchyScriptParameterViewModel::OnDroppedOnInternal(TSharedPtr<F
 		
 		if(DroppedItem->IsForHierarchy() == false)
 		{
-			TSharedRef<FNiagaraHierarchyItemViewModelBase> AddedItemViewModel = DuplicateToThis(DroppedItem);
+			TSharedRef<FHierarchyElementViewModel> AddedItemViewModel = DuplicateToThis(DroppedItem);
 			AddedItemViewModel->SyncViewModelsToData();
 		}
 		else
 		{
-			TSharedRef<FNiagaraHierarchyItemViewModelBase> ReparentedViewModel = ReparentToThis(DroppedItem);
+			TSharedRef<FHierarchyElementViewModel> ReparentedViewModel = ReparentToThis(DroppedItem);
 			ReparentedViewModel->SyncViewModelsToData();
 		}
 
@@ -269,9 +253,9 @@ void FNiagaraHierarchyScriptParameterViewModel::OnDroppedOnInternal(TSharedPtr<F
 
 void FNiagaraHierarchyScriptRootViewModel::SortChildrenData() const
 {
-	GetDataMutable()->GetChildrenMutable().StableSort([](const UNiagaraHierarchyItemBase& ItemA, const UNiagaraHierarchyItemBase& ItemB)
+	GetDataMutable()->GetChildrenMutable().StableSort([](const UHierarchyElement& ItemA, const UHierarchyElement& ItemB)
 		{
-			return ItemA.IsA<UNiagaraHierarchyItem>() && ItemB.IsA<UNiagaraHierarchyCategory>();
+			return ItemA.IsA<UHierarchyItem>() && ItemB.IsA<UHierarchyCategory>();
 		});
 }
 
