New Tool: Hierarchy Editor

diff --git a/Engine/Plugins/FX/Niagara/Source/NiagaraEditor/Private/ViewModels/HierarchyEditor/NiagaraSummaryViewViewModel.cpp b/Engine/Plugins/FX/Niagara/Source/NiagaraEditor/Private/ViewModels/HierarchyEditor/NiagaraSummaryViewViewModel.cpp
--- a/Engine/Plugins/FX/Niagara/Source/NiagaraEditor/Private/ViewModels/HierarchyEditor/NiagaraSummaryViewViewModel.cpp
+++ b/Engine/Plugins/FX/Niagara/Source/NiagaraEditor/Private/ViewModels/HierarchyEditor/NiagaraSummaryViewViewModel.cpp
@@ -15,12 +15,14 @@
 #include "ViewModels/NiagaraEmitterViewModel.h"
 #include "ViewModels/NiagaraScriptGraphViewModel.h"
 #include "ViewModels/NiagaraScriptViewModel.h"
+#include "ScopedTransaction.h"
 #include "ViewModels/NiagaraSystemViewModel.h"
+#include "ToolMenu.h"
 #include "Widgets/Text/SInlineEditableTextBlock.h"
 
 #define LOCTEXT_NAMESPACE "NiagaraSummaryViewHierarchyEditor"
 
-bool GetIsFromBaseEmitter(const FVersionedNiagaraEmitter& Emitter, FNiagaraHierarchyIdentity SummaryItemIdentity)
+bool GetIsFromBaseEmitter(const FVersionedNiagaraEmitter& Emitter, FHierarchyElementIdentity SummaryItemIdentity)
 {
 	TSharedRef<FNiagaraScriptMergeManager> MergeManager = FNiagaraScriptMergeManager::Get();
 	return MergeManager->DoesSummaryItemExistInBase(Emitter, SummaryItemIdentity);
@@ -28,14 +30,14 @@ bool GetIsFromBaseEmitter(const FVersionedNiagaraEmitter& Emitter, FNiagaraHiera
 
 void UNiagaraHierarchyModule::Initialize(const UNiagaraNodeFunctionCall& InModuleNode)
 {
-	FNiagaraHierarchyIdentity ModuleIdentity;
+	FHierarchyElementIdentity ModuleIdentity;
 	ModuleIdentity.Guids.Add(InModuleNode.NodeGuid);
 	SetIdentity(ModuleIdentity);
 }
 
 void UNiagaraHierarchyModuleInput::Initialize(const UNiagaraNodeFunctionCall& InModuleNode, FGuid InputGuid)
 {
-	FNiagaraHierarchyIdentity InputIdentity;
+	FHierarchyElementIdentity InputIdentity;
 	InputIdentity.Guids.Add(InModuleNode.NodeGuid);
 	InputIdentity.Guids.Add(InputGuid);
 	SetIdentity(InputIdentity);
@@ -43,7 +45,7 @@ void UNiagaraHierarchyModuleInput::Initialize(const UNiagaraNodeFunctionCall& In
 
 void UNiagaraHierarchyAssignmentInput::Initialize(const UNiagaraNodeAssignment& AssignmentNode, FName AssignmentTarget)
 {
-	FNiagaraHierarchyIdentity InputIdentity;
+	FHierarchyElementIdentity InputIdentity;
 	InputIdentity.Guids.Add(AssignmentNode.NodeGuid);
 	InputIdentity.Names.Add(AssignmentTarget);
 	SetIdentity(InputIdentity);
@@ -51,7 +53,7 @@ void UNiagaraHierarchyAssignmentInput::Initialize(const UNiagaraNodeAssignment&
 
 void UNiagaraHierarchyEmitterProperties::Initialize(const FVersionedNiagaraEmitter& Emitter)
 {
-	FNiagaraHierarchyIdentity InputIdentity;
+	FHierarchyElementIdentity InputIdentity;
 	InputIdentity.Names.Add(FName(Emitter.Emitter->GetUniqueEmitterName()));
 	InputIdentity.Names.Add("Category");
 	InputIdentity.Names.Add("Properties");
@@ -60,14 +62,14 @@ void UNiagaraHierarchyEmitterProperties::Initialize(const FVersionedNiagaraEmitt
 
 void UNiagaraHierarchyRenderer::Initialize(const UNiagaraRendererProperties& Renderer)
 {
-	FNiagaraHierarchyIdentity RendererIdentity;
+	FHierarchyElementIdentity RendererIdentity;
 	RendererIdentity.Guids.Add(Renderer.GetMergeId());
 	SetIdentity(RendererIdentity);
 }
 
 void UNiagaraHierarchyEventHandler::Initialize(const FNiagaraEventScriptProperties& EventHandler)
 {
-	FNiagaraHierarchyIdentity EventHandlerIdentity;
+	FHierarchyElementIdentity EventHandlerIdentity;
 	EventHandlerIdentity.Guids.Add(EventHandler.Script->GetUsageId());
 	SetIdentity(EventHandlerIdentity);
 }
@@ -77,9 +79,9 @@ void UNiagaraHierarchyEventHandlerProperties::Initialize(const FNiagaraEventScri
 	SetIdentity(MakeIdentity(EventHandler));
 }
 
-FNiagaraHierarchyIdentity UNiagaraHierarchyEventHandlerProperties::MakeIdentity(const FNiagaraEventScriptProperties& EventHandler)
+FHierarchyElementIdentity UNiagaraHierarchyEventHandlerProperties::MakeIdentity(const FNiagaraEventScriptProperties& EventHandler)
 {
-	FNiagaraHierarchyIdentity Identity;
+	FHierarchyElementIdentity Identity;
 	Identity.Guids.Add(EventHandler.Script->GetUsageId());
 	Identity.Names.Add(TEXT("Category"));
 	Identity.Names.Add(TEXT("Properties"));
@@ -88,7 +90,7 @@ FNiagaraHierarchyIdentity UNiagaraHierarchyEventHandlerProperties::MakeIdentity(
 
 void UNiagaraHierarchySimStage::Initialize(const UNiagaraSimulationStageBase& SimStage)
 {
-	FNiagaraHierarchyIdentity SimStageIdentity;
+	FHierarchyElementIdentity SimStageIdentity;
 	SimStageIdentity.Guids.Add(SimStage.GetMergeId());
 	SetIdentity(SimStageIdentity);
 }
@@ -98,9 +100,9 @@ void UNiagaraHierarchySimStageProperties::Initialize(const UNiagaraSimulationSta
 	SetIdentity(MakeIdentity(SimStage));
 }
 
-FNiagaraHierarchyIdentity UNiagaraHierarchySimStageProperties::MakeIdentity(const UNiagaraSimulationStageBase& SimStage)
+FHierarchyElementIdentity UNiagaraHierarchySimStageProperties::MakeIdentity(const UNiagaraSimulationStageBase& SimStage)
 {
-	FNiagaraHierarchyIdentity SimStagePropertiesIdentity;
+	FHierarchyElementIdentity SimStagePropertiesIdentity;
 	SimStagePropertiesIdentity.Guids.Add(SimStage.GetMergeId());
 	SimStagePropertiesIdentity.Names.Add(FName("Category"));
 	SimStagePropertiesIdentity.Names.Add("Properties");
@@ -109,7 +111,7 @@ FNiagaraHierarchyIdentity UNiagaraHierarchySimStageProperties::MakeIdentity(cons
 
 void UNiagaraHierarchyObjectProperty::Initialize(FGuid ObjectGuid, FString PropertyName)
 {
-	FNiagaraHierarchyIdentity PropertyIdentity;
+	FHierarchyElementIdentity PropertyIdentity;
 	PropertyIdentity.Guids.Add(ObjectGuid);
 	PropertyIdentity.Names.Add(FName(PropertyName));
 	SetIdentity(PropertyIdentity);
@@ -126,7 +128,7 @@ void UNiagaraSummaryViewViewModel::Initialize(TSharedRef<FNiagaraEmitterViewMode
 		EmitterViewModel->GetEmitter().Emitter->OnSimStagesChanged().AddUObject(this, &UNiagaraSummaryViewViewModel::OnSimStagesChanged);
 		EmitterViewModel->GetEmitter().Emitter->OnEventHandlersChanged().AddUObject(this, &UNiagaraSummaryViewViewModel::OnEventHandlersChanged);
 
-		UNiagaraHierarchyViewModelBase::Initialize();
+		UHierarchyViewModelBase::Initialize();
 	}
 }
 
@@ -151,7 +153,7 @@ TSharedRef<FNiagaraEmitterViewModel> UNiagaraSummaryViewViewModel::GetEmitterVie
 
 TWeakObjectPtr<UNiagaraNodeFunctionCall> FNiagaraFunctionViewModel::GetFunctionCallNode() const
 {
-	TWeakObjectPtr<UNiagaraHierarchyViewModelBase> ViewModel = GetHierarchyViewModel();
+	TWeakObjectPtr<UHierarchyViewModelBase> ViewModel = GetHierarchyViewModel();
 	UNiagaraSummaryViewViewModel* SummaryViewModel = CastChecked<UNiagaraSummaryViewViewModel>(ViewModel.Get());
 	return SummaryViewModel->GetFunctionCallNode(GetData()->GetPersistentIdentity().Guids[0]);
 }
@@ -167,7 +169,7 @@ void FNiagaraFunctionViewModel::OnScriptApplied(UNiagaraScript* NiagaraScript, F
 
 void FNiagaraFunctionViewModel::ClearCache() const
 {
-	TWeakObjectPtr<UNiagaraHierarchyViewModelBase> ViewModel = GetHierarchyViewModel();
+	TWeakObjectPtr<UHierarchyViewModelBase> ViewModel = GetHierarchyViewModel();
 	UNiagaraSummaryViewViewModel* SummaryViewModel = CastChecked<UNiagaraSummaryViewViewModel>(ViewModel.Get());
 	SummaryViewModel->ClearFunctionCallNodeCache(GetData()->GetPersistentIdentity().Guids[0]);
 }
@@ -182,7 +184,7 @@ FString FNiagaraFunctionViewModel::ToString() const
 	return TEXT("Unknown");
 }
 
-FNiagaraHierarchyItemViewModelBase::FCanPerformActionResults FNiagaraFunctionViewModel::IsEditableByUser()
+FHierarchyElementViewModel::FCanPerformActionResults FNiagaraFunctionViewModel::IsEditableByUser()
 {
 	if(bIsDynamicInput)
 	{
@@ -297,10 +299,10 @@ void FNiagaraFunctionViewModel::RefreshChildrenInputs(bool bClearCache) const
 				
 				for(const FGuid& VariableGuid : VariableGuids)
 				{
-					FNiagaraHierarchyIdentity SearchedChildIdentity;
+					FHierarchyElementIdentity SearchedChildIdentity;
 					SearchedChildIdentity.Guids.Add(FunctionNode->NodeGuid);
 					SearchedChildIdentity.Guids.Add(VariableGuid);
-					const bool bChildExists = GetData()->GetChildren().ContainsByPredicate([SearchedChildIdentity](UNiagaraHierarchyItemBase* CandidateChild)
+					const bool bChildExists = GetData()->GetChildren().ContainsByPredicate([SearchedChildIdentity](UHierarchyElement* CandidateChild)
 					{
 						return CandidateChild->GetPersistentIdentity() == SearchedChildIdentity;
 					});
@@ -318,7 +320,7 @@ void FNiagaraFunctionViewModel::RefreshChildrenInputs(bool bClearCache) const
 
 				for(const FNiagaraVariable& Variable : AssignmentNode->GetAssignmentTargets())
 				{
-					const bool bChildExists = GetDataMutable()->GetChildrenMutable().ContainsByPredicate([Variable, bClearCache, this](UNiagaraHierarchyItemBase* Candidate)
+					const bool bChildExists = GetDataMutable()->GetChildrenMutable().ContainsByPredicate([Variable, bClearCache, this](UHierarchyElement* Candidate)
 					{
 						if(UNiagaraHierarchyAssignmentInput* AssignmentInput = Cast<UNiagaraHierarchyAssignmentInput>(Candidate))
 						{
@@ -339,7 +341,7 @@ void FNiagaraFunctionViewModel::RefreshChildrenInputs(bool bClearCache) const
 	}
 }
 
-FNiagaraHierarchyItemViewModelBase::FCanPerformActionResults FNiagaraFunctionViewModel::CanDropOnInternal(TSharedPtr<FNiagaraHierarchyItemViewModelBase> NiagaraHierarchyItemViewModelBase, EItemDropZone ItemDropZone)
+FHierarchyElementViewModel::FCanPerformActionResults FNiagaraFunctionViewModel::CanDropOnInternal(TSharedPtr<FHierarchyElementViewModel> NiagaraHierarchyItemViewModelBase, EItemDropZone ItemDropZone)
 {
 	if(IsEditableByUser().bCanPerform == false)
 	{
@@ -354,10 +356,10 @@ FNiagaraHierarchyItemViewModelBase::FCanPerformActionResults FNiagaraFunctionVie
 		return Results;
 	}
 	
-	return FNiagaraHierarchyItemViewModel::CanDropOnInternal(NiagaraHierarchyItemViewModelBase, ItemDropZone);
+	return FHierarchyItemViewModel::CanDropOnInternal(NiagaraHierarchyItemViewModelBase, ItemDropZone);
 }
 
-const UNiagaraHierarchySection* FNiagaraFunctionViewModel::GetSectionInternal() const
+const UHierarchySection* FNiagaraFunctionViewModel::GetSectionInternal() const
 {
 	if(bIsForHierarchy)
 	{
@@ -367,7 +369,7 @@ const UNiagaraHierarchySection* FNiagaraFunctionViewModel::GetSectionInternal()
 	return Section.IsValid() ? Section.Get() : nullptr;
 }
 
-FNiagaraHierarchyItemViewModelBase::FCanPerformActionResults FNiagaraModuleInputViewModel::CanDropOnInternal(TSharedPtr<FNiagaraHierarchyItemViewModelBase> DraggedItem, EItemDropZone ItemDropZone)
+FHierarchyElementViewModel::FCanPerformActionResults FNiagaraModuleInputViewModel::CanDropOnInternal(TSharedPtr<FHierarchyElementViewModel> DraggedItem, EItemDropZone ItemDropZone)
 {
 	// if the input isn't editable, we don't allow any drops on/above/below the item.
 	// Even though it technically works, the merge process will only re-add the item at the end and not preserve order so there is no point in allowing dropping above/below
@@ -378,7 +380,7 @@ FNiagaraHierarchyItemViewModelBase::FCanPerformActionResults FNiagaraModuleInput
 
 	FCanPerformActionResults AllowDrop(false);
 	
-	TSharedPtr<FNiagaraHierarchyItemViewModelBase> TargetDropItem = AsShared();
+	TSharedPtr<FHierarchyElementViewModel> TargetDropItem = AsShared();
 
 	// we only allow drops if some general conditions are fulfilled
 	if(DraggedItem->GetData() != TargetDropItem->GetData() &&
@@ -416,7 +418,7 @@ FNiagaraHierarchyItemViewModelBase::FCanPerformActionResults FNiagaraModuleInput
 				// we use default logic only if there is no parent input. Nested children are not allowed to contain anything but other inputs.
 				if(TargetDropItem->GetParent().Pin()->GetData<UNiagaraHierarchyModuleInput>() == nullptr)
 				{
-					AllowDrop = FNiagaraHierarchyItemViewModel::CanDropOnInternal(DraggedItem, ItemDropZone);
+					AllowDrop = FHierarchyItemViewModel::CanDropOnInternal(DraggedItem, ItemDropZone);
 				}
 			}
 		}
@@ -425,11 +427,11 @@ FNiagaraHierarchyItemViewModelBase::FCanPerformActionResults FNiagaraModuleInput
 	return AllowDrop;
 }
 
-void FNiagaraModuleInputViewModel::OnDroppedOnInternal(TSharedPtr<FNiagaraHierarchyItemViewModelBase> DroppedItem, EItemDropZone ItemDropZone)
+void FNiagaraModuleInputViewModel::OnDroppedOnInternal(TSharedPtr<FHierarchyElementViewModel> DroppedItem, EItemDropZone ItemDropZone)
 {
 	if(ItemDropZone != EItemDropZone::OntoItem)
 	{
-		return FNiagaraHierarchyItemViewModel::OnDroppedOnInternal(DroppedItem, ItemDropZone);
+		return FHierarchyItemViewModel::OnDroppedOnInternal(DroppedItem, ItemDropZone);
 	}
 	else
 	{
@@ -439,13 +441,13 @@ void FNiagaraModuleInputViewModel::OnDroppedOnInternal(TSharedPtr<FNiagaraHierar
 		// we empty out the children as technically you can drag a parent input onto another input now. We don't take these child-child inputs with us as we only allow child inputs 1 layer deep
 		if(DroppedItem->IsForHierarchy() == false)
 		{
-			TSharedRef<FNiagaraHierarchyItemViewModelBase> AddedItemViewModel = DuplicateToThis(DroppedItem);
+			TSharedRef<FHierarchyElementViewModel> AddedItemViewModel = DuplicateToThis(DroppedItem);
 			AddedItemViewModel->GetChildrenMutable().Empty();
 			AddedItemViewModel->SyncViewModelsToData();
 		}
 		else
 		{
-			TSharedRef<FNiagaraHierarchyItemViewModelBase> ReparentedViewModel = ReparentToThis(DroppedItem);
+			TSharedRef<FHierarchyElementViewModel> ReparentedViewModel = ReparentToThis(DroppedItem);
 			ReparentedViewModel->GetChildrenMutable().Empty();
 			ReparentedViewModel->SyncViewModelsToData();
 		}
@@ -455,24 +457,21 @@ void FNiagaraModuleInputViewModel::OnDroppedOnInternal(TSharedPtr<FNiagaraHierar
 	}
 }
 
-void FNiagaraModuleInputViewModel::PopulateDynamicContextMenuSection(FToolMenuSection& DynamicSection)
+void FNiagaraModuleInputViewModel::AppendDynamicContextMenuForSingleElement(UToolMenu* ToolMenu)
 {
 	FUIAction Action;
 	Action.ExecuteAction = FExecuteAction::CreateSP(this, &FNiagaraModuleInputViewModel::AddNativeChildrenInputs);
 	Action.CanExecuteAction = FCanExecuteAction::CreateSP(this, &FNiagaraModuleInputViewModel::CanAddNativeChildrenInputs);
 	Action.IsActionVisibleDelegate = FIsActionButtonVisible::CreateSP(this, &FNiagaraModuleInputViewModel::CanAddNativeChildrenInputs);
 	
-	DynamicSection.AddMenuEntry(
-		FName("Add Children Inputs"),
-		LOCTEXT("AddChildrenInputsMenuLabel", "Add Children Inputs"),
-		LOCTEXT("AddChildrenInputsMenuTooltip", "Add children inputs of this input as child inputs."),
-		FSlateIcon(),
-		Action);
+	ToolMenu->AddMenuEntry("Dynamic",
+		FToolMenuEntry::InitMenuEntry(FName("Add Children Inputs"),LOCTEXT("AddChildrenInputsMenuLabel", "Add Children Inputs"),
+		LOCTEXT("AddChildrenInputsMenuTooltip", "Add children inputs of this input as child inputs."), FSlateIcon(),Action));
 }
 
 TWeakObjectPtr<UNiagaraNodeFunctionCall> FNiagaraModuleInputViewModel::GetModuleNode() const
 {
-	TWeakObjectPtr<UNiagaraHierarchyViewModelBase> ViewModel = GetHierarchyViewModel();
+	TWeakObjectPtr<UHierarchyViewModelBase> ViewModel = GetHierarchyViewModel();
 	UNiagaraSummaryViewViewModel* SummaryViewModel = CastChecked<UNiagaraSummaryViewViewModel>(ViewModel.Get());
 	return SummaryViewModel->GetFunctionCallNode(GetData()->GetPersistentIdentity().Guids[0]);
 }
@@ -512,7 +511,7 @@ FString FNiagaraModuleInputViewModel::ToString() const
 		return InputData->InputName.ToString();
 	}
 
-	return FNiagaraHierarchyItemViewModel::ToString();
+	return FHierarchyItemViewModel::ToString();
 }
 
 TArray<FString> FNiagaraModuleInputViewModel::GetSearchTerms() const
@@ -554,9 +553,9 @@ void FNiagaraModuleInputViewModel::RefreshChildDynamicInputs(bool bClearCache)
 
 		if(DynamicInputNode)
 		{
-			FNiagaraHierarchyIdentity DynamicInputIdentity;
+			FHierarchyElementIdentity DynamicInputIdentity;
 			DynamicInputIdentity.Guids.Add(DynamicInputNode->NodeGuid);
-			const bool bChildExists = GetDataMutable()->GetChildrenMutable().ContainsByPredicate([DynamicInputIdentity, bClearCache](UNiagaraHierarchyItemBase* CandidateChild)
+			const bool bChildExists = GetDataMutable()->GetChildrenMutable().ContainsByPredicate([DynamicInputIdentity, bClearCache](UHierarchyElement* CandidateChild)
 				{
 					return CandidateChild->GetPersistentIdentity() == DynamicInputIdentity;
 				});
@@ -575,7 +574,7 @@ FText FNiagaraModuleInputViewModel::GetSummaryInputNameOverride() const
 	return GetData<UNiagaraHierarchyModuleInput>()->GetDisplayNameOverride();
 }
 
-FNiagaraHierarchyItemViewModelBase::FCanPerformActionResults FNiagaraModuleInputViewModel::IsEditableByUser()
+FHierarchyElementViewModel::FCanPerformActionResults FNiagaraModuleInputViewModel::IsEditableByUser()
 {
 	FCanPerformActionResults CanEditResults(IsFromBaseEmitter() == false);
 	CanEditResults.CanPerformMessage = CanEditResults.bCanPerform == false ? LOCTEXT("ModuleInputIsFromBaseEmitter", "This input was added in the parent emitter and can not be edited.") : FText::GetEmpty();
@@ -599,7 +598,7 @@ void FNiagaraModuleInputViewModel::RefreshChildrenDataInternal()
 
 TOptional<FInputData> FNiagaraModuleInputViewModel::FindInputDataInternal() const
 {
-	TWeakObjectPtr<UNiagaraHierarchyViewModelBase> ViewModel = GetHierarchyViewModel();
+	TWeakObjectPtr<UHierarchyViewModelBase> ViewModel = GetHierarchyViewModel();
 	UNiagaraSummaryViewViewModel* SummaryViewModel = CastChecked<UNiagaraSummaryViewViewModel>(ViewModel.Get());
 	return SummaryViewModel->GetInputData(*GetData<UNiagaraHierarchyModuleInput>());
 }
@@ -613,8 +612,8 @@ void FNiagaraModuleInputViewModel::AddNativeChildrenInputs()
 	
 	if(GetModuleNode().IsValid() && GetInputData().IsSet())
 	{
-		TArray<FNiagaraHierarchyIdentity> ChildIdentities = GetNativeChildInputIdentities();
-		TMap<FNiagaraHierarchyIdentity, int32> ChildSortOrderMap;
+		TArray<FHierarchyElementIdentity> ChildIdentities = GetNativeChildInputIdentities();
+		TMap<FHierarchyElementIdentity, int32> ChildSortOrderMap;
 
 		if(UNiagaraGraph* Graph = GetModuleNode()->GetCalledGraph())
 		{
@@ -627,7 +626,7 @@ void FNiagaraModuleInputViewModel::AddNativeChildrenInputs()
 				{
 					if(!ScriptVariable->Metadata.GetParentAttribute_DEPRECATED().IsNone() && ScriptVariable->Metadata.GetParentAttribute_DEPRECATED().IsEqual(GetInputData()->InputName))
 					{
-						FNiagaraHierarchyIdentity ChildIdentity;
+						FHierarchyElementIdentity ChildIdentity;
 						ChildIdentity.Guids.Add(GetData()->GetPersistentIdentity().Guids[0]);
 						ChildIdentity.Guids.Add(ScriptVariable->Metadata.GetVariableGuid());
 						ChildSortOrderMap.Add(ChildIdentity, ScriptVariable->Metadata.GetEditorSortPriority_DEPRECATED());
@@ -636,11 +635,11 @@ void FNiagaraModuleInputViewModel::AddNativeChildrenInputs()
 			}
 		}
 
-		for(const FNiagaraHierarchyIdentity& ChildIdentity : ChildIdentities)
+		for(const FHierarchyElementIdentity& ChildIdentity : ChildIdentities)
 		{
 			if(FindViewModelForChild(ChildIdentity, false) == nullptr)
 			{
-				TSharedPtr<FNiagaraHierarchyItemViewModelBase> ViewModel = HierarchyViewModel->GetHierarchyRootViewModel()->FindViewModelForChild(ChildIdentity, true);
+				TSharedPtr<FHierarchyElementViewModel> ViewModel = HierarchyViewModel->GetHierarchyRootViewModel()->FindViewModelForChild(ChildIdentity, true);
 				if(ViewModel.IsValid())
 				{
 					ReparentToThis(ViewModel);
@@ -658,7 +657,7 @@ void FNiagaraModuleInputViewModel::AddNativeChildrenInputs()
 		TArray<TSharedPtr<FNiagaraModuleInputViewModel>> ChildInputs;
 		GetChildrenViewModelsForType<UNiagaraHierarchyModuleInput, FNiagaraModuleInputViewModel>(ChildInputs);
 
-		auto SortChildrenInputs = [&](UNiagaraHierarchyItemBase& ItemA, UNiagaraHierarchyItemBase& ItemB)
+		auto SortChildrenInputs = [&](UHierarchyElement& ItemA, UHierarchyElement& ItemB)
 		{
 			int32 SortOrderA = ChildSortOrderMap[ItemA.GetPersistentIdentity()];
 			int32 SortOrderB = ChildSortOrderMap[ItemB.GetPersistentIdentity()];
@@ -682,9 +681,9 @@ bool FNiagaraModuleInputViewModel::CanAddNativeChildrenInputs() const
 	return GetNativeChildInputIdentities().Num() > 0;
 }
 
-TArray<FNiagaraHierarchyIdentity> FNiagaraModuleInputViewModel::GetNativeChildInputIdentities() const
+TArray<FHierarchyElementIdentity> FNiagaraModuleInputViewModel::GetNativeChildInputIdentities() const
 {
-	TArray<FNiagaraHierarchyIdentity> ChildIdentities;
+	TArray<FHierarchyElementIdentity> ChildIdentities;
 	if(GetModuleNode().IsValid())
 	{
 		if(UNiagaraGraph* Graph = GetModuleNode()->GetCalledGraph())
@@ -698,7 +697,7 @@ TArray<FNiagaraHierarchyIdentity> FNiagaraModuleInputViewModel::GetNativeChildIn
 				{
 					if(!ScriptVariable->Metadata.GetParentAttribute_DEPRECATED().IsNone() && ScriptVariable->Metadata.GetParentAttribute_DEPRECATED().IsEqual(GetInputData()->InputName))
 					{
-						FNiagaraHierarchyIdentity ChildIdentity;
+						FHierarchyElementIdentity ChildIdentity;
 						ChildIdentity.Guids.Add(GetData()->GetPersistentIdentity().Guids[0]);
 						ChildIdentity.Guids.Add(ScriptVariable->Metadata.GetVariableGuid());
 						ChildIdentities.Add(ChildIdentity);
@@ -711,7 +710,7 @@ TArray<FNiagaraHierarchyIdentity> FNiagaraModuleInputViewModel::GetNativeChildIn
 	return ChildIdentities;
 }
 
-FNiagaraHierarchyItemViewModelBase::FCanPerformActionResults FNiagaraAssignmentInputViewModel::CanDropOnInternal(TSharedPtr<FNiagaraHierarchyItemViewModelBase> DraggedItem, EItemDropZone ItemDropZone)
+FHierarchyElementViewModel::FCanPerformActionResults FNiagaraAssignmentInputViewModel::CanDropOnInternal(TSharedPtr<FHierarchyElementViewModel> DraggedItem, EItemDropZone ItemDropZone)
 {
 	// if the input isn't editable, we don't allow any drops on/above/below the item.
 	// Even though it technically works, the merge process will only re-add the item at the end and not preserve order.
@@ -720,12 +719,12 @@ FNiagaraHierarchyItemViewModelBase::FCanPerformActionResults FNiagaraAssignmentI
 		return false;
 	}
 	
-	return FNiagaraHierarchyItemViewModel::CanDropOnInternal(DraggedItem, ItemDropZone);
+	return FHierarchyItemViewModel::CanDropOnInternal(DraggedItem, ItemDropZone);
 }
 
 TWeakObjectPtr<UNiagaraNodeAssignment> FNiagaraAssignmentInputViewModel::GetAssignmentNode() const
 {
-	TWeakObjectPtr<UNiagaraHierarchyViewModelBase> ViewModel = GetHierarchyViewModel();
+	TWeakObjectPtr<UHierarchyViewModelBase> ViewModel = GetHierarchyViewModel();
 	UNiagaraSummaryViewViewModel* SummaryViewModel = CastChecked<UNiagaraSummaryViewViewModel>(ViewModel.Get());
 	return Cast<UNiagaraNodeAssignment>(SummaryViewModel->GetFunctionCallNode(GetData()->GetPersistentIdentity().Guids[0]));
 }
@@ -748,7 +747,7 @@ FString FNiagaraAssignmentInputViewModel::ToString() const
 		return InputData->InputName.ToString();
 	}
 
-	return FNiagaraHierarchyItemViewModel::ToString();
+	return FHierarchyItemViewModel::ToString();
 }
 
 TArray<FString> FNiagaraAssignmentInputViewModel::GetSearchTerms() const
@@ -772,12 +771,12 @@ bool FNiagaraAssignmentInputViewModel::IsFromBaseEmitter() const
 void FNiagaraAssignmentInputViewModel::ClearCache() const
 {
 	InputDataCache.Reset();
-	TWeakObjectPtr<UNiagaraHierarchyViewModelBase> ViewModel = GetHierarchyViewModel();
+	TWeakObjectPtr<UHierarchyViewModelBase> ViewModel = GetHierarchyViewModel();
 	UNiagaraSummaryViewViewModel* SummaryViewModel = CastChecked<UNiagaraSummaryViewViewModel>(ViewModel.Get());
 	SummaryViewModel->ClearFunctionCallNodeCache(GetData()->GetPersistentIdentity().Guids[0]);
 }
 
-FNiagaraHierarchyItemViewModelBase::FCanPerformActionResults FNiagaraAssignmentInputViewModel::IsEditableByUser()
+FHierarchyElementViewModel::FCanPerformActionResults FNiagaraAssignmentInputViewModel::IsEditableByUser()
 {
 	FCanPerformActionResults CanEditResults(IsFromBaseEmitter() == false);
 	CanEditResults.CanPerformMessage = CanEditResults.bCanPerform == false ? LOCTEXT("ModuleInputIsFromBaseEmitter", "This input was added in the parent emitter and can not be edited.") : FText::GetEmpty();
@@ -796,7 +795,7 @@ FNiagaraHierarchyItemViewModelBase::FCanPerformActionResults FNiagaraAssignmentI
 
 TOptional<FNiagaraStackGraphUtilities::FMatchingFunctionInputData> FNiagaraAssignmentInputViewModel::FindInputDataInternal() const
 {
-	TWeakObjectPtr<UNiagaraHierarchyViewModelBase> ViewModel = GetHierarchyViewModel();
+	TWeakObjectPtr<UHierarchyViewModelBase> ViewModel = GetHierarchyViewModel();
 	UNiagaraSummaryViewViewModel* SummaryViewModel = CastChecked<UNiagaraSummaryViewViewModel>(ViewModel.Get());
 	
 	if(GetAssignmentNode().IsValid())
@@ -818,7 +817,7 @@ bool FNiagaraHierarchySummaryCategoryViewModel::IsFromBaseEmitter() const
 	return IsFromBaseEmitterCache.GetValue();
 }
 
-FNiagaraHierarchyItemViewModelBase::FCanPerformActionResults FNiagaraHierarchySummaryCategoryViewModel::IsEditableByUser()
+FHierarchyElementViewModel::FCanPerformActionResults FNiagaraHierarchySummaryCategoryViewModel::IsEditableByUser()
 {
 	FCanPerformActionResults CanEditResults(IsFromBaseEmitter() == false);
 	CanEditResults.CanPerformMessage = CanEditResults.bCanPerform == false ? LOCTEXT("CategoryIsFromBaseEmitter", "This category was added in the parent emitter and can not be edited. You can add new items.") : FText::GetEmpty();
@@ -836,14 +835,14 @@ bool FNiagaraHierarchyPropertyViewModel::IsFromBaseEmitter() const
 	return IsFromBaseEmitterCache.GetValue();
 }
 
-bool FNiagaraHierarchyPropertyViewModel::DoesExternalDataStillExist(const UNiagaraHierarchyDataRefreshContext* Context) const
+bool FNiagaraHierarchyPropertyViewModel::DoesExternalDataStillExist(const UHierarchyDataRefreshContext* Context) const
 {
 	UNiagaraSummaryViewViewModel* ViewModel = Cast<UNiagaraSummaryViewViewModel>(GetHierarchyViewModel());
 	TMap<FGuid, UObject*> PropertyObjectMap = ViewModel->GetObjectsForProperties();
 	return PropertyObjectMap.Contains(GetData()->GetPersistentIdentity().Guids[0]);
 }
 
-FNiagaraHierarchyItemViewModelBase::FCanPerformActionResults FNiagaraHierarchyPropertyViewModel::IsEditableByUser()
+FHierarchyElementViewModel::FCanPerformActionResults FNiagaraHierarchyPropertyViewModel::IsEditableByUser()
 {
 	FCanPerformActionResults CanEditResults(IsFromBaseEmitter() == false);
 	CanEditResults.CanPerformMessage = CanEditResults.bCanPerform == false ? LOCTEXT("ObjectPropertyIsFromBaseEmitter", "This property was added in the parent emitter and can not be edited.") : FText::GetEmpty();
@@ -869,18 +868,18 @@ bool FNiagaraHierarchyRendererViewModel::IsFromBaseEmitter() const
 
 void FNiagaraHierarchyRendererViewModel::RefreshChildrenDataInternal()
 {
-	TArray<UNiagaraHierarchyItemBase*> NewChildren;
+	TArray<UHierarchyElement*> NewChildren;
 	for (TFieldIterator<FProperty> PropertyIterator(GetRendererProperties()->GetClass(), EFieldIteratorFlags::ExcludeSuper); PropertyIterator; ++PropertyIterator)
 	{
 		if(PropertyIterator->HasAnyPropertyFlags(CPF_Edit))
 		{
 			FString PropertyName = (*PropertyIterator)->GetName();
 
-			FNiagaraHierarchyIdentity PropertyIdentity;
+			FHierarchyElementIdentity PropertyIdentity;
 			PropertyIdentity.Guids.Add(GetRendererProperties()->GetMergeId());
 			PropertyIdentity.Names.Add(FName(PropertyName));
 			
-			auto* FoundItem = GetDataMutable()->GetChildrenMutable().FindByPredicate([PropertyIdentity](UNiagaraHierarchyItemBase* Candidate)
+			auto* FoundItem = GetDataMutable()->GetChildrenMutable().FindByPredicate([PropertyIdentity](UHierarchyElement* Candidate)
 			{
 				return Candidate->GetPersistentIdentity() == PropertyIdentity;
 			});
@@ -904,14 +903,14 @@ void FNiagaraHierarchyRendererViewModel::RefreshChildrenDataInternal()
 	GetDataMutable()->GetChildrenMutable().Append(NewChildren);
 }
 
-FNiagaraHierarchyItemViewModelBase::FCanPerformActionResults FNiagaraHierarchyRendererViewModel::IsEditableByUser()
+FHierarchyElementViewModel::FCanPerformActionResults FNiagaraHierarchyRendererViewModel::IsEditableByUser()
 {
 	FCanPerformActionResults CanEditResults(IsFromBaseEmitter() == false);
 	CanEditResults.CanPerformMessage = CanEditResults.bCanPerform == false ? LOCTEXT("RendererIsFromBaseEmitter", "This renderer was added in the parent emitter and can not be edited.") : FText::GetEmpty();
 	return CanEditResults;
 }
 
-const UNiagaraHierarchySection* FNiagaraHierarchyRendererViewModel::GetSectionInternal() const
+const UHierarchySection* FNiagaraHierarchyRendererViewModel::GetSectionInternal() const
 {
 	return Section.IsValid() ? Section.Get() : nullptr;
 }
@@ -932,14 +931,14 @@ bool FNiagaraHierarchyEmitterPropertiesViewModel::IsFromBaseEmitter() const
 	return IsFromBaseEmitterCache.GetValue();
 }
 
-FNiagaraHierarchyItemViewModelBase::FCanPerformActionResults FNiagaraHierarchyEmitterPropertiesViewModel::IsEditableByUser()
+FHierarchyElementViewModel::FCanPerformActionResults FNiagaraHierarchyEmitterPropertiesViewModel::IsEditableByUser()
 {
 	FCanPerformActionResults CanEditResults(IsFromBaseEmitter() == false);
 	CanEditResults.CanPerformMessage = CanEditResults.bCanPerform == false ? LOCTEXT("EmitterPropertiesIsFromBaseEmitter", "These emitter properties were added in the parent emitter and can not be edited.") : FText::GetEmpty();
 	return CanEditResults;
 }
 
-const UNiagaraHierarchySection* FNiagaraHierarchyEmitterPropertiesViewModel::GetSectionInternal() const
+const UHierarchySection* FNiagaraHierarchyEmitterPropertiesViewModel::GetSectionInternal() const
 {
 	return Section.IsValid() ? Section.Get() : nullptr;
 }
@@ -957,7 +956,7 @@ FString FNiagaraHierarchyEventHandlerViewModel::ToString() const
 
 FNiagaraEventScriptProperties* FNiagaraHierarchyEventHandlerViewModel::GetEventScriptProperties() const
 {
-	TWeakObjectPtr<UNiagaraHierarchyViewModelBase> ViewModel = GetHierarchyViewModel();
+	TWeakObjectPtr<UHierarchyViewModelBase> ViewModel = GetHierarchyViewModel();
 	UNiagaraSummaryViewViewModel* SummaryViewModel = CastChecked<UNiagaraSummaryViewViewModel>(ViewModel.Get());
 	FGuid UsageID = GetData()->GetPersistentIdentity().Guids[0];
 	
@@ -985,12 +984,12 @@ bool FNiagaraHierarchyEventHandlerViewModel::IsFromBaseEmitter() const
 
 void FNiagaraHierarchyEventHandlerViewModel::RefreshChildrenDataInternal()
 {
-	TArray<UNiagaraHierarchyItemBase*> NewChildren;
+	TArray<UHierarchyElement*> NewChildren;
 
 	// First we add the properties item
-	FNiagaraHierarchyIdentity PropertiesIdentity = UNiagaraHierarchyEventHandlerProperties::MakeIdentity(*GetEventScriptProperties());	
+	FHierarchyElementIdentity PropertiesIdentity = UNiagaraHierarchyEventHandlerProperties::MakeIdentity(*GetEventScriptProperties());	
 
-	auto* FoundProperties = GetDataMutable()->GetChildrenMutable().FindByPredicate([PropertiesIdentity](UNiagaraHierarchyItemBase* Candidate)
+	auto* FoundProperties = GetDataMutable()->GetChildrenMutable().FindByPredicate([PropertiesIdentity](UHierarchyElement* Candidate)
 	{
 		return Candidate->GetPersistentIdentity() == PropertiesIdentity;
 	});
@@ -1015,9 +1014,9 @@ void FNiagaraHierarchyEventHandlerViewModel::RefreshChildrenDataInternal()
 	for(UNiagaraNodeFunctionCall* EventHandlerModule : EventHandlerModules)
 	{
 		UNiagaraHierarchyModule* HierarchyEventHandlerModule = nullptr;
-		FNiagaraHierarchyIdentity ModuleIdentity;
+		FHierarchyElementIdentity ModuleIdentity;
 		ModuleIdentity.Guids.Add(EventHandlerModule->NodeGuid);
-		auto* FoundHierarchySimStageModule = GetDataMutable()->GetChildrenMutable().FindByPredicate([ModuleIdentity](UNiagaraHierarchyItemBase* Candidate)
+		auto* FoundHierarchySimStageModule = GetDataMutable()->GetChildrenMutable().FindByPredicate([ModuleIdentity](UHierarchyElement* Candidate)
 		{
 			return Candidate->GetPersistentIdentity() == ModuleIdentity;
 		});
@@ -1039,14 +1038,14 @@ void FNiagaraHierarchyEventHandlerViewModel::RefreshChildrenDataInternal()
 	GetDataMutable()->GetChildrenMutable().Append(NewChildren);
 }
 
-FNiagaraHierarchyItemViewModelBase::FCanPerformActionResults FNiagaraHierarchyEventHandlerViewModel::IsEditableByUser()
+FHierarchyElementViewModel::FCanPerformActionResults FNiagaraHierarchyEventHandlerViewModel::IsEditableByUser()
 {
 	FCanPerformActionResults CanEditResults(IsFromBaseEmitter() == false);
 	CanEditResults.CanPerformMessage = CanEditResults.bCanPerform == false ? LOCTEXT("EventHandlerIsFromBaseEmitter", "This event handler was added in the parent emitter and can not be edited.") : FText::GetEmpty();
 	return CanEditResults;
 }
 
-const UNiagaraHierarchySection* FNiagaraHierarchyEventHandlerViewModel::GetSectionInternal() const
+const UHierarchySection* FNiagaraHierarchyEventHandlerViewModel::GetSectionInternal() const
 {
 	if(Section.IsValid())
 	{
@@ -1058,7 +1057,7 @@ const UNiagaraHierarchySection* FNiagaraHierarchyEventHandlerViewModel::GetSecti
 
 UNiagaraRendererProperties* FNiagaraHierarchyRendererViewModel::GetRendererProperties() const
 {
-	TWeakObjectPtr<UNiagaraHierarchyViewModelBase> ViewModel = GetHierarchyViewModel();
+	TWeakObjectPtr<UHierarchyViewModelBase> ViewModel = GetHierarchyViewModel();
 	UNiagaraSummaryViewViewModel* SummaryViewModel = CastChecked<UNiagaraSummaryViewViewModel>(ViewModel.Get());
 	FGuid RendererGuid = GetData()->GetPersistentIdentity().Guids[0];
 
@@ -1083,7 +1082,7 @@ FString FNiagaraHierarchyEventHandlerPropertiesViewModel::ToString() const
 
 FNiagaraEventScriptProperties* FNiagaraHierarchyEventHandlerPropertiesViewModel::GetEventScriptProperties() const
 {
-	TWeakObjectPtr<UNiagaraHierarchyViewModelBase> ViewModel = GetHierarchyViewModel();
+	TWeakObjectPtr<UHierarchyViewModelBase> ViewModel = GetHierarchyViewModel();
 	UNiagaraSummaryViewViewModel* SummaryViewModel = CastChecked<UNiagaraSummaryViewViewModel>(ViewModel.Get());
 	FGuid UsageID = GetData()->GetPersistentIdentity().Guids[0];
 	
@@ -1111,7 +1110,7 @@ bool FNiagaraHierarchyEventHandlerPropertiesViewModel::IsFromBaseEmitter() const
 
 void FNiagaraHierarchyEventHandlerPropertiesViewModel::RefreshChildrenDataInternal()
 {
-	TArray<UNiagaraHierarchyItemBase*> NewPropertiesChildren;
+	TArray<UHierarchyElement*> NewPropertiesChildren;
 
 	// todo (me) while this works, the stack needs to access the correct FStructOnScope that points to the FNiagaraEventScriptProperties
 	// That can be made to work correctly, but the EventScriptProperties are heavily customized and introduce UI issues
@@ -1127,7 +1126,7 @@ void FNiagaraHierarchyEventHandlerPropertiesViewModel::RefreshChildrenDataIntern
 	// 		PropertyIdentity.Guids.Add(GetEventScriptProperties()->Script->GetUsageId());
 	// 		PropertyIdentity.Names.Add(FName(PropertyName));
 	// 		
-	// 		UNiagaraHierarchyItemBase** FoundPropertyItem = GetDataMutable()->GetChildrenMutable().FindByPredicate([PropertyIdentity](UNiagaraHierarchyItemBase* Candidate)
+	// 		UHierarchyElement** FoundPropertyItem = GetDataMutable()->GetChildrenMutable().FindByPredicate([PropertyIdentity](UHierarchyElement* Candidate)
 	// 		{
 	// 			return Candidate->GetPersistentIdentity() == PropertyIdentity;
 	// 		});
@@ -1151,7 +1150,7 @@ void FNiagaraHierarchyEventHandlerPropertiesViewModel::RefreshChildrenDataIntern
 	GetDataMutable()->GetChildrenMutable().Append(NewPropertiesChildren);
 }
 
-FNiagaraHierarchyItemViewModelBase::FCanPerformActionResults FNiagaraHierarchyEventHandlerPropertiesViewModel::IsEditableByUser()
+FHierarchyElementViewModel::FCanPerformActionResults FNiagaraHierarchyEventHandlerPropertiesViewModel::IsEditableByUser()
 {
 	FCanPerformActionResults CanEditResults(IsFromBaseEmitter() == false);
 	CanEditResults.CanPerformMessage = CanEditResults.bCanPerform == false ? LOCTEXT("EventHandlerPropertiesIsFromBaseEmitter", "This property item was added in the parent emitter and can not be edited.") : FText::GetEmpty();
@@ -1166,7 +1165,7 @@ FString FNiagaraHierarchySimStageViewModel::ToString() const
 
 UNiagaraSimulationStageBase* FNiagaraHierarchySimStageViewModel::GetSimStage() const
 {
-	TWeakObjectPtr<UNiagaraHierarchyViewModelBase> ViewModel = GetHierarchyViewModel();
+	TWeakObjectPtr<UHierarchyViewModelBase> ViewModel = GetHierarchyViewModel();
 	UNiagaraSummaryViewViewModel* SummaryViewModel = CastChecked<UNiagaraSummaryViewViewModel>(ViewModel.Get());
 	FGuid SimStageGuid = GetData()->GetPersistentIdentity().Guids[0];
 
@@ -1197,15 +1196,15 @@ bool FNiagaraHierarchySimStageViewModel::IsFromBaseEmitter() const
 
 void FNiagaraHierarchySimStageViewModel::RefreshChildrenDataInternal()
 {
-	TArray<UNiagaraHierarchyItemBase*> NewChildren;
+	TArray<UHierarchyElement*> NewChildren;
 
 	// First we add the properties item
-	FNiagaraHierarchyIdentity PropertiesIdentity;
+	FHierarchyElementIdentity PropertiesIdentity;
 	PropertiesIdentity.Guids.Add(GetSimStage()->GetMergeId());
 	PropertiesIdentity.Names.Add(FName("Category"));
 	PropertiesIdentity.Names.Add(FName("Properties"));
 
-	auto* FoundProperties = GetDataMutable()->GetChildrenMutable().FindByPredicate([PropertiesIdentity](UNiagaraHierarchyItemBase* Candidate)
+	auto* FoundProperties = GetDataMutable()->GetChildrenMutable().FindByPredicate([PropertiesIdentity](UHierarchyElement* Candidate)
 	{
 		return Candidate->GetPersistentIdentity() == PropertiesIdentity;
 	});
@@ -1230,9 +1229,9 @@ void FNiagaraHierarchySimStageViewModel::RefreshChildrenDataInternal()
 	for(UNiagaraNodeFunctionCall* SimStageModule : SimStageModules)
 	{
 		UNiagaraHierarchyModule* HierarchySimStageModule = nullptr;
-		FNiagaraHierarchyIdentity ModuleIdentity;
+		FHierarchyElementIdentity ModuleIdentity;
 		ModuleIdentity.Guids.Add(SimStageModule->NodeGuid);
-		auto* FoundHierarchySimStageModule = GetDataMutable()->GetChildrenMutable().FindByPredicate([ModuleIdentity](UNiagaraHierarchyItemBase* Candidate)
+		auto* FoundHierarchySimStageModule = GetDataMutable()->GetChildrenMutable().FindByPredicate([ModuleIdentity](UHierarchyElement* Candidate)
 		{
 			return Candidate->GetPersistentIdentity() == ModuleIdentity;
 		});
@@ -1254,14 +1253,14 @@ void FNiagaraHierarchySimStageViewModel::RefreshChildrenDataInternal()
 	GetDataMutable()->GetChildrenMutable().Append(NewChildren);
 }
 
-FNiagaraHierarchyItemViewModelBase::FCanPerformActionResults FNiagaraHierarchySimStageViewModel::IsEditableByUser()
+FHierarchyElementViewModel::FCanPerformActionResults FNiagaraHierarchySimStageViewModel::IsEditableByUser()
 {
 	FCanPerformActionResults CanEditResults(IsFromBaseEmitter() == false);
 	CanEditResults.CanPerformMessage = CanEditResults.bCanPerform == false ? LOCTEXT("SimStageIsFromBaseEmitter", "This simulation stage was added in the parent emitter and can not be edited.") : FText::GetEmpty();
 	return CanEditResults;
 }
 
-const UNiagaraHierarchySection* FNiagaraHierarchySimStageViewModel::GetSectionInternal() const
+const UHierarchySection* FNiagaraHierarchySimStageViewModel::GetSectionInternal() const
 {
 	return Section.IsValid() ? Section.Get() : nullptr;
 }
@@ -1273,7 +1272,7 @@ FString FNiagaraHierarchySimStagePropertiesViewModel::ToString() const
 
 UNiagaraSimulationStageBase* FNiagaraHierarchySimStagePropertiesViewModel::GetSimStage() const
 {
-	TWeakObjectPtr<UNiagaraHierarchyViewModelBase> ViewModel = GetHierarchyViewModel();
+	TWeakObjectPtr<UHierarchyViewModelBase> ViewModel = GetHierarchyViewModel();
 	UNiagaraSummaryViewViewModel* SummaryViewModel = CastChecked<UNiagaraSummaryViewViewModel>(ViewModel.Get());
 	FGuid SimStageGuid = GetData()->GetPersistentIdentity().Guids[0];
 
@@ -1304,7 +1303,7 @@ bool FNiagaraHierarchySimStagePropertiesViewModel::IsFromBaseEmitter() const
 
 void FNiagaraHierarchySimStagePropertiesViewModel::RefreshChildrenDataInternal()
 {
-	TArray<UNiagaraHierarchyItemBase*> NewPropertiesChildren;
+	TArray<UHierarchyElement*> NewPropertiesChildren;
 
 	for (TFieldIterator<FProperty> PropertyIterator(GetSimStage()->GetClass(), EFieldIteratorFlags::IncludeSuper); PropertyIterator; ++PropertyIterator)
 	{
@@ -1312,11 +1311,11 @@ void FNiagaraHierarchySimStagePropertiesViewModel::RefreshChildrenDataInternal()
 		{
 			FString PropertyName = (*PropertyIterator)->GetName();
 
-			FNiagaraHierarchyIdentity PropertyIdentity;
+			FHierarchyElementIdentity PropertyIdentity;
 			PropertyIdentity.Guids.Add(GetSimStage()->GetMergeId());
 			PropertyIdentity.Names.Add(FName(PropertyName));
 			
-			auto* FoundPropertyItem = GetDataMutable()->GetChildrenMutable().FindByPredicate([PropertyIdentity](UNiagaraHierarchyItemBase* Candidate)
+			auto* FoundPropertyItem = GetDataMutable()->GetChildrenMutable().FindByPredicate([PropertyIdentity](UHierarchyElement* Candidate)
 			{
 				return Candidate->GetPersistentIdentity() == PropertyIdentity;
 			});
@@ -1340,7 +1339,7 @@ void FNiagaraHierarchySimStagePropertiesViewModel::RefreshChildrenDataInternal()
 	GetDataMutable()->GetChildrenMutable().Append(NewPropertiesChildren);
 }
 
-FNiagaraHierarchyItemViewModelBase::FCanPerformActionResults FNiagaraHierarchySimStagePropertiesViewModel::IsEditableByUser()
+FHierarchyElementViewModel::FCanPerformActionResults FNiagaraHierarchySimStagePropertiesViewModel::IsEditableByUser()
 {
 	FCanPerformActionResults CanEditResults(IsFromBaseEmitter() == false);
 	CanEditResults.CanPerformMessage = CanEditResults.bCanPerform == false ? LOCTEXT("RendererPropertiesIsFromBaseEmitter", "This renderer's properties were added in the parent emitter and can not be edited.") : FText::GetEmpty();
@@ -1349,7 +1348,7 @@ FNiagaraHierarchyItemViewModelBase::FCanPerformActionResults FNiagaraHierarchySi
 
 TSharedRef<SWidget> FNiagaraHierarchyInputParameterHierarchyDragDropOp::CreateCustomDecorator() const
 {
-	TSharedPtr<FNiagaraModuleInputViewModel> InputViewModel = StaticCastSharedPtr<FNiagaraModuleInputViewModel>(DraggedItem.Pin());
+	TSharedPtr<FNiagaraModuleInputViewModel> InputViewModel = StaticCastSharedPtr<FNiagaraModuleInputViewModel>(DraggedElement.Pin());
 	TOptional<FInputData> InputData = InputViewModel->GetInputData();
 	return FNiagaraParameterUtilities::GetParameterWidget(FNiagaraVariable(InputData->Type, InputData->InputName), false, false);
 }
@@ -1359,73 +1358,70 @@ FString FNiagaraHierarchyPropertyViewModel::ToString() const
 	return GetData()->GetPersistentIdentity().Names[0].ToString();
 }
 
-UNiagaraHierarchyRoot* UNiagaraSummaryViewViewModel::GetHierarchyRoot() const
+UHierarchyRoot* UNiagaraSummaryViewViewModel::GetHierarchyRoot() const
 {
-	UNiagaraHierarchyRoot* RootItem = GetEmitterViewModel()->GetEditorData().GetSummaryRoot();
+	UHierarchyRoot* RootItem = GetEmitterViewModel()->GetEditorData().GetSummaryRoot();
 
 	ensure(RootItem != nullptr);
 	return RootItem;
 }
 
-TSharedPtr<FNiagaraHierarchyItemViewModelBase> UNiagaraSummaryViewViewModel::CreateViewModelForData(UNiagaraHierarchyItemBase* ItemBase, TSharedPtr<FNiagaraHierarchyItemViewModelBase> Parent)
-{
-	bool bIsForHierarchy = Parent.IsValid() ? Parent->IsForHierarchy() : false;
-	
+TSharedPtr<FHierarchyElementViewModel> UNiagaraSummaryViewViewModel::CreateCustomViewModelForElement(UHierarchyElement* ItemBase, TSharedPtr<FHierarchyElementViewModel> Parent)
+{	
 	if(UNiagaraHierarchyModuleInput* SummaryViewItem = Cast<UNiagaraHierarchyModuleInput>(ItemBase))
 	{
-		return MakeShared<FNiagaraModuleInputViewModel>(SummaryViewItem, Parent, this, bIsForHierarchy);
+		return MakeShared<FNiagaraModuleInputViewModel>(SummaryViewItem, Parent.ToSharedRef(), this);
 	}
 	else if(UNiagaraHierarchyModule* Module = Cast<UNiagaraHierarchyModule>(ItemBase))
 	{
-		return MakeShared<FNiagaraFunctionViewModel>(Module, Parent, this, bIsForHierarchy);
+		return MakeShared<FNiagaraFunctionViewModel>(Module, Parent.ToSharedRef(), this);
 	}
 	else if(UNiagaraHierarchyRenderer* Renderer = Cast<UNiagaraHierarchyRenderer>(ItemBase))
 	{
-		return MakeShared<FNiagaraHierarchyRendererViewModel>(Renderer, Parent, this, bIsForHierarchy);
+		return MakeShared<FNiagaraHierarchyRendererViewModel>(Renderer, Parent.ToSharedRef(), this);
 	}
 	else if(UNiagaraHierarchyEmitterProperties* EmitterProperties = Cast<UNiagaraHierarchyEmitterProperties>(ItemBase))
 	{
-		return MakeShared<FNiagaraHierarchyEmitterPropertiesViewModel>(EmitterProperties, Parent, this, bIsForHierarchy);
+		return MakeShared<FNiagaraHierarchyEmitterPropertiesViewModel>(EmitterProperties, Parent.ToSharedRef(), this);
 	}
 	else if(UNiagaraHierarchyEventHandler* EventHandler = Cast<UNiagaraHierarchyEventHandler>(ItemBase))
 	{
-		return MakeShared<FNiagaraHierarchyEventHandlerViewModel>(EventHandler, Parent, this, bIsForHierarchy);
+		return MakeShared<FNiagaraHierarchyEventHandlerViewModel>(EventHandler, Parent.ToSharedRef(), this);
 	}
 	else if(UNiagaraHierarchyEventHandlerProperties* EventHandlerProperties = Cast<UNiagaraHierarchyEventHandlerProperties>(ItemBase))
 	{
-		return MakeShared<FNiagaraHierarchyEventHandlerPropertiesViewModel>(EventHandlerProperties, Parent, this, bIsForHierarchy);
+		return MakeShared<FNiagaraHierarchyEventHandlerPropertiesViewModel>(EventHandlerProperties, Parent.ToSharedRef(), this);
 	}
 	else if(UNiagaraHierarchySimStage* SimStage = Cast<UNiagaraHierarchySimStage>(ItemBase))
 	{
-		return MakeShared<FNiagaraHierarchySimStageViewModel>(SimStage, Parent, this, bIsForHierarchy);
+		return MakeShared<FNiagaraHierarchySimStageViewModel>(SimStage, Parent.ToSharedRef(), this);
 	}
 	else if(UNiagaraHierarchySimStageProperties* SimStageProperties = Cast<UNiagaraHierarchySimStageProperties>(ItemBase))
 	{
-		return MakeShared<FNiagaraHierarchySimStagePropertiesViewModel>(SimStageProperties, Parent, this, bIsForHierarchy);
+		return MakeShared<FNiagaraHierarchySimStagePropertiesViewModel>(SimStageProperties, Parent.ToSharedRef(), this);
 	}
 	else if(UNiagaraHierarchyObjectProperty* ObjectProperty = Cast<UNiagaraHierarchyObjectProperty>(ItemBase))
 	{
-		return MakeShared<FNiagaraHierarchyPropertyViewModel>(ObjectProperty, Parent, this, bIsForHierarchy);
+		return MakeShared<FNiagaraHierarchyPropertyViewModel>(ObjectProperty, Parent.ToSharedRef(), this);
 	}
 	else if(UNiagaraHierarchyAssignmentInput* AssignmentInput = Cast<UNiagaraHierarchyAssignmentInput>(ItemBase))
 	{
-		return MakeShared<FNiagaraAssignmentInputViewModel>(AssignmentInput, Parent, this, bIsForHierarchy);
+		return MakeShared<FNiagaraAssignmentInputViewModel>(AssignmentInput, Parent.ToSharedRef(), this);
 	}
-	else if(UNiagaraHierarchyCategory* Category = Cast<UNiagaraHierarchyCategory>(ItemBase))
+	else if(UHierarchyCategory* Category = Cast<UHierarchyCategory>(ItemBase))
 	{
-		return MakeShared<FNiagaraHierarchySummaryCategoryViewModel>(Category, Parent, this, bIsForHierarchy);
+		return MakeShared<FNiagaraHierarchySummaryCategoryViewModel>(Category, Parent.ToSharedRef(), this);
 	}
-
-	check(false);
+	
 	return nullptr;
 }
 
-void UNiagaraSummaryViewViewModel::PrepareSourceItems(UNiagaraHierarchyRoot* SourceRoot, TSharedPtr<FNiagaraHierarchyRootViewModel> SourceRootViewModel)
+void UNiagaraSummaryViewViewModel::PrepareSourceItems(UHierarchyRoot* SourceRoot, TSharedPtr<FHierarchyRootViewModel> SourceRootViewModel)
 {
 	TSharedPtr<FNiagaraEmitterViewModel> EmitterViewModel = EmitterViewModelWeak.Pin();
 
-	TArray<UNiagaraHierarchyItemBase*> NewItems;
-	TArray<UNiagaraHierarchySection*> NewSections;
+	TArray<UHierarchyElement*> NewItems;
+	TArray<UHierarchySection*> NewSections;
 	
 	FunctionCallCache.Empty();
 		
@@ -1439,27 +1435,13 @@ void UNiagaraSummaryViewViewModel::PrepareSourceItems(UNiagaraHierarchyRoot* Sou
 	TArray<UNiagaraNodeFunctionCall*> SimStageModules = FNiagaraStackGraphUtilities::GetAllSimStagesModuleNodes(EmitterViewModel.ToSharedRef());
 	TArray<UNiagaraNodeFunctionCall*> EventHandlerModules = FNiagaraStackGraphUtilities::GetAllEventHandlerModuleNodes(EmitterViewModel.ToSharedRef());
 
-	FNiagaraHierarchyIdentity EmitterPropertiesIdentity;
+	FHierarchyElementIdentity EmitterPropertiesIdentity;
 	EmitterPropertiesIdentity.Guids.Add(EmitterViewModel->GetEmitter().Version);
 	EmitterPropertiesIdentity.Names.Add("Category");
 	EmitterPropertiesIdentity.Names.Add("Properties");
 
-	UNiagaraHierarchyEmitterProperties* EmitterProperties = nullptr;
-	
-	auto* FoundEmitterProperties = SourceRoot->GetChildrenMutable().FindByPredicate([EmitterPropertiesIdentity](UNiagaraHierarchyItemBase* ItemBase)
-		{
-			return ItemBase->GetPersistentIdentity() == EmitterPropertiesIdentity;
-		});
-
-	if(FoundEmitterProperties == nullptr)
-	{
-		EmitterProperties = SourceRoot->AddChild<UNiagaraHierarchyEmitterProperties>();
-		EmitterProperties->Initialize(EmitterViewModel->GetEmitter());
-	}
-	else
-	{
-		EmitterProperties = CastChecked<UNiagaraHierarchyEmitterProperties>(*FoundEmitterProperties);	
-	}
+	UNiagaraHierarchyEmitterProperties* EmitterProperties =  SourceRoot->AddChild<UNiagaraHierarchyEmitterProperties>();
+	EmitterProperties->Initialize(EmitterViewModel->GetEmitter());
 
 	NewItems.Add(EmitterProperties);
 	
@@ -1473,23 +1455,8 @@ void UNiagaraSummaryViewViewModel::PrepareSourceItems(UNiagaraHierarchyRoot* Sou
 			continue;
 		}
 		
-		UNiagaraHierarchyModule* HierarchyModule = nullptr;
-		FNiagaraHierarchyIdentity ModuleIdentity;
-		ModuleIdentity.Guids.Add(ModuleNode->NodeGuid);
-		auto* FoundItem = SourceRoot->GetChildrenMutable().FindByPredicate([ModuleIdentity](UNiagaraHierarchyItemBase* ItemBase)
-		{
-			return ItemBase->GetPersistentIdentity() == ModuleIdentity;
-		});
-		
-		if(FoundItem == nullptr)		
-		{
-			HierarchyModule = SourceRoot->AddChild<UNiagaraHierarchyModule>();
-			HierarchyModule->Initialize(*ModuleNode);
-		}
-		else
-		{
-			HierarchyModule = CastChecked<UNiagaraHierarchyModule>(*FoundItem);
-		}
+		UNiagaraHierarchyModule* HierarchyModule = SourceRoot->AddChild<UNiagaraHierarchyModule>();		
+		HierarchyModule->Initialize(*ModuleNode);		
 		
 		NewItems.Add(HierarchyModule);
 		UsageMap.FindOrAdd(FNiagaraStackGraphUtilities::GetOutputNodeUsage(*ModuleNode)).Add(HierarchyModule);
@@ -1499,11 +1466,11 @@ void UNiagaraSummaryViewViewModel::PrepareSourceItems(UNiagaraHierarchyRoot* Sou
 	for(const FNiagaraEventScriptProperties& ScriptPropertiesItem : ScriptProperties)
 	{
 		UNiagaraHierarchyEventHandler* HierarchyEventHandler = nullptr;
-		FNiagaraHierarchyIdentity EventHandlerIdentity;
+		FHierarchyElementIdentity EventHandlerIdentity;
 		EventHandlerIdentity.Guids.Add(ScriptPropertiesItem.Script->GetUsageId());
 		EventHandlerIdentity.Guids.Add(ScriptPropertiesItem.SourceEmitterID);
 
-		auto* FoundItem = SourceRoot->GetChildrenMutable().FindByPredicate([EventHandlerIdentity](UNiagaraHierarchyItemBase* ItemBase)
+		auto* FoundItem = SourceRoot->GetChildrenMutable().FindByPredicate([EventHandlerIdentity](UHierarchyElement* ItemBase)
 		{
 			return ItemBase->GetPersistentIdentity() == EventHandlerIdentity;
 		});
@@ -1527,10 +1494,10 @@ void UNiagaraSummaryViewViewModel::PrepareSourceItems(UNiagaraHierarchyRoot* Sou
 	for(UNiagaraSimulationStageBase* SimStage : SimStages)
 	{
 		UNiagaraHierarchySimStage* HierarchySimStage = nullptr;
-		FNiagaraHierarchyIdentity SimStageID;
+		FHierarchyElementIdentity SimStageID;
 		SimStageID.Guids.Add(SimStage->GetMergeId());
 
-		auto* FoundItem = SourceRoot->GetChildrenMutable().FindByPredicate([SimStageID](UNiagaraHierarchyItemBase* ItemBase)
+		auto* FoundItem = SourceRoot->GetChildrenMutable().FindByPredicate([SimStageID](UHierarchyElement* ItemBase)
 		{
 			return ItemBase->GetPersistentIdentity() == SimStageID;
 		});
@@ -1554,10 +1521,10 @@ void UNiagaraSummaryViewViewModel::PrepareSourceItems(UNiagaraHierarchyRoot* Sou
 	for(UNiagaraRendererProperties* RendererProperties : Renderers)
 	{
 		UNiagaraHierarchyRenderer* HierarchyRenderer = nullptr;
-		FNiagaraHierarchyIdentity RendererIdentity;
+		FHierarchyElementIdentity RendererIdentity;
 		RendererIdentity.Guids.Add(RendererProperties->GetMergeId());
 
-		auto* FoundItem = SourceRoot->GetChildrenMutable().FindByPredicate([RendererIdentity](UNiagaraHierarchyItemBase* ItemBase)
+		auto* FoundItem = SourceRoot->GetChildrenMutable().FindByPredicate([RendererIdentity](UHierarchyElement* ItemBase)
 		{
 			return ItemBase->GetPersistentIdentity() == RendererIdentity;
 		});
@@ -1583,18 +1550,18 @@ void UNiagaraSummaryViewViewModel::PrepareSourceItems(UNiagaraHierarchyRoot* Sou
 	SourceRootViewModel->SyncViewModelsToData();
 	
 	// Now we create a section for each usage case that has at least one element and link up the respective sections
-	UNiagaraHierarchySection* EmitterSpawnSection = nullptr;
+	UHierarchySection* EmitterSpawnSection = nullptr;
 	if(UsageMap.Contains(ENiagaraScriptUsage::EmitterSpawnScript))
 	{
 		FText SectionName = FText::FromString("Emitter Spawn");
-		auto* FoundHierarchySection = SourceRootViewModel->GetDataMutable<UNiagaraHierarchyRoot>()->GetSectionDataMutable().FindByPredicate([SectionName](UNiagaraHierarchySection* Candidate)
+		auto* FoundHierarchySection = SourceRootViewModel->GetDataMutable<UHierarchyRoot>()->GetSectionDataMutable().FindByPredicate([SectionName](UHierarchySection* Candidate)
 		{
 			return Candidate->GetSectionNameAsText().EqualTo(SectionName);
 		});
 		
 		if(FoundHierarchySection == nullptr)
 		{
-			EmitterSpawnSection = SourceRootViewModel->GetDataMutable<UNiagaraHierarchyRoot>()->AddSection(SectionName);
+			EmitterSpawnSection = SourceRootViewModel->GetDataMutable<UHierarchyRoot>()->AddSection(SectionName);
 		}
 		else
 		{
@@ -1610,18 +1577,18 @@ void UNiagaraSummaryViewViewModel::PrepareSourceItems(UNiagaraHierarchyRoot* Sou
 		}
 	}
 
-	UNiagaraHierarchySection* EmitterUpdateSection = nullptr;
+	UHierarchySection* EmitterUpdateSection = nullptr;
 	if(UsageMap.Contains(ENiagaraScriptUsage::EmitterUpdateScript))
 	{
 		FText SectionName = FText::FromString("Emitter Update");
-		auto* FoundHierarchySection = SourceRootViewModel->GetDataMutable<UNiagaraHierarchyRoot>()->GetSectionDataMutable().FindByPredicate([SectionName](UNiagaraHierarchySection* Candidate)
+		auto* FoundHierarchySection = SourceRootViewModel->GetDataMutable<UHierarchyRoot>()->GetSectionDataMutable().FindByPredicate([SectionName](UHierarchySection* Candidate)
 		{
 			return Candidate->GetSectionNameAsText().EqualTo(SectionName);
 		});
 		
 		if(FoundHierarchySection == nullptr)
 		{
-			EmitterUpdateSection = SourceRootViewModel->GetDataMutable<UNiagaraHierarchyRoot>()->AddSection(SectionName);
+			EmitterUpdateSection = SourceRootViewModel->GetDataMutable<UHierarchyRoot>()->AddSection(SectionName);
 		}
 		else
 		{
@@ -1637,18 +1604,18 @@ void UNiagaraSummaryViewViewModel::PrepareSourceItems(UNiagaraHierarchyRoot* Sou
 		}
 	}
 
-	UNiagaraHierarchySection* ParticleSpawnSection = nullptr;
+	UHierarchySection* ParticleSpawnSection = nullptr;
 	if(UsageMap.Contains(ENiagaraScriptUsage::ParticleSpawnScript))
 	{
 		FText SectionName = FText::FromString("Particle Spawn");
-		auto* FoundHierarchySection = SourceRootViewModel->GetDataMutable<UNiagaraHierarchyRoot>()->GetSectionDataMutable().FindByPredicate([SectionName](UNiagaraHierarchySection* Candidate)
+		auto* FoundHierarchySection = SourceRootViewModel->GetDataMutable<UHierarchyRoot>()->GetSectionDataMutable().FindByPredicate([SectionName](UHierarchySection* Candidate)
 		{
 			return Candidate->GetSectionNameAsText().EqualTo(SectionName);
 		});
 		
 		if(FoundHierarchySection == nullptr)
 		{
-			ParticleSpawnSection = SourceRootViewModel->GetDataMutable<UNiagaraHierarchyRoot>()->AddSection(SectionName);
+			ParticleSpawnSection = SourceRootViewModel->GetDataMutable<UHierarchyRoot>()->AddSection(SectionName);
 		}
 		else
 		{
@@ -1664,18 +1631,18 @@ void UNiagaraSummaryViewViewModel::PrepareSourceItems(UNiagaraHierarchyRoot* Sou
 		}
 	}
 
-	UNiagaraHierarchySection* ParticleUpdateSection = nullptr;
+	UHierarchySection* ParticleUpdateSection = nullptr;
 	if(UsageMap.Contains(ENiagaraScriptUsage::ParticleUpdateScript))
 	{
 		FText SectionName = FText::FromString("Particle Update");
-		auto* FoundHierarchySection = SourceRootViewModel->GetDataMutable<UNiagaraHierarchyRoot>()->GetSectionDataMutable().FindByPredicate([SectionName](UNiagaraHierarchySection* Candidate)
+		auto* FoundHierarchySection = SourceRootViewModel->GetDataMutable<UHierarchyRoot>()->GetSectionDataMutable().FindByPredicate([SectionName](UHierarchySection* Candidate)
 		{
 			return Candidate->GetSectionNameAsText().EqualTo(SectionName);
 		});
 		
 		if(FoundHierarchySection == nullptr)
 		{
-			ParticleUpdateSection = SourceRootViewModel->GetDataMutable<UNiagaraHierarchyRoot>()->AddSection(SectionName);
+			ParticleUpdateSection = SourceRootViewModel->GetDataMutable<UHierarchyRoot>()->AddSection(SectionName);
 		}
 		else
 		{
@@ -1691,18 +1658,18 @@ void UNiagaraSummaryViewViewModel::PrepareSourceItems(UNiagaraHierarchyRoot* Sou
 		}
 	}
 
-	UNiagaraHierarchySection* EventHandlerSection = nullptr;
+	UHierarchySection* EventHandlerSection = nullptr;
 	if(HierarchyEventHandlers.Num() > 0)
 	{
 		FText SectionName = FText::FromString("Events");
-		auto* FoundHierarchySection = SourceRootViewModel->GetDataMutable<UNiagaraHierarchyRoot>()->GetSectionDataMutable().FindByPredicate([SectionName](UNiagaraHierarchySection* Candidate)
+		auto* FoundHierarchySection = SourceRootViewModel->GetDataMutable<UHierarchyRoot>()->GetSectionDataMutable().FindByPredicate([SectionName](UHierarchySection* Candidate)
 		{
 			return Candidate->GetSectionNameAsText().EqualTo(SectionName);
 		});
 		
 		if(FoundHierarchySection == nullptr)
 		{
-			EventHandlerSection = SourceRootViewModel->GetDataMutable<UNiagaraHierarchyRoot>()->AddSection(SectionName);
+			EventHandlerSection = SourceRootViewModel->GetDataMutable<UHierarchyRoot>()->AddSection(SectionName);
 		}
 		else
 		{
@@ -1718,18 +1685,18 @@ void UNiagaraSummaryViewViewModel::PrepareSourceItems(UNiagaraHierarchyRoot* Sou
 		}
 	}
 
-	UNiagaraHierarchySection* SimulationStagesSection = nullptr;
+	UHierarchySection* SimulationStagesSection = nullptr;
 	if(HierarchySimStages.Num() > 0)
 	{
 		FText SectionName = FText::FromString("Sim Stages");
-		auto* FoundHierarchySection = SourceRootViewModel->GetDataMutable<UNiagaraHierarchyRoot>()->GetSectionDataMutable().FindByPredicate([SectionName](UNiagaraHierarchySection* Candidate)
+		auto* FoundHierarchySection = SourceRootViewModel->GetDataMutable<UHierarchyRoot>()->GetSectionDataMutable().FindByPredicate([SectionName](UHierarchySection* Candidate)
 		{
 			return Candidate->GetSectionNameAsText().EqualTo(SectionName);
 		});
 		
 		if(FoundHierarchySection == nullptr)
 		{
-			SimulationStagesSection = SourceRootViewModel->GetDataMutable<UNiagaraHierarchyRoot>()->AddSection(SectionName);
+			SimulationStagesSection = SourceRootViewModel->GetDataMutable<UHierarchyRoot>()->AddSection(SectionName);
 		}
 		else
 		{
@@ -1745,18 +1712,18 @@ void UNiagaraSummaryViewViewModel::PrepareSourceItems(UNiagaraHierarchyRoot* Sou
 		}
 	}
 
-	UNiagaraHierarchySection* RenderersSection = nullptr; 
+	UHierarchySection* RenderersSection = nullptr; 
 	if(HierarchyRenderers.Num() > 0)
 	{
 		FText SectionName = FText::FromString("Renderers");
-		auto* FoundHierarchySection = SourceRootViewModel->GetDataMutable<UNiagaraHierarchyRoot>()->GetSectionDataMutable().FindByPredicate([SectionName](UNiagaraHierarchySection* Candidate)
+		auto* FoundHierarchySection = SourceRootViewModel->GetDataMutable<UHierarchyRoot>()->GetSectionDataMutable().FindByPredicate([SectionName](UHierarchySection* Candidate)
 		{
 			return Candidate->GetSectionNameAsText().EqualTo(SectionName);
 		});
 		
 		if(FoundHierarchySection == nullptr)
 		{
-			RenderersSection = SourceRootViewModel->GetDataMutable<UNiagaraHierarchyRoot>()->AddSection(SectionName);
+			RenderersSection = SourceRootViewModel->GetDataMutable<UHierarchyRoot>()->AddSection(SectionName);
 		}
 		else
 		{
@@ -1779,7 +1746,7 @@ void UNiagaraSummaryViewViewModel::PrepareSourceItems(UNiagaraHierarchyRoot* Sou
 	// force a sync so we have the view models for the sections available
 	SourceRootViewModel->SyncViewModelsToData();
 
-	for(TSharedPtr<FNiagaraHierarchySectionViewModel> SectionViewModel : SourceRootViewModel->GetSectionViewModels())
+	for(TSharedPtr<FHierarchySectionViewModel> SectionViewModel : SourceRootViewModel->GetSectionViewModels())
 	{
 		if(EmitterSpawnSection != nullptr && SectionViewModel->GetData() == EmitterSpawnSection)
 		{
@@ -1823,32 +1790,30 @@ void UNiagaraSummaryViewViewModel::SetupCommands()
 	// no custom commands yet
 }
 
-TSharedRef<FNiagaraHierarchyDragDropOp> UNiagaraSummaryViewViewModel::CreateDragDropOp(TSharedRef<FNiagaraHierarchyItemViewModelBase> Item)
+TSharedRef<FHierarchyDragDropOp> UNiagaraSummaryViewViewModel::CreateDragDropOp(TSharedRef<FHierarchyElementViewModel> Item)
 {
-	if(UNiagaraHierarchyCategory* HierarchyCategory = Cast<UNiagaraHierarchyCategory>(Item->GetDataMutable()))
+	if(UHierarchyCategory* HierarchyCategory = Cast<UHierarchyCategory>(Item->GetDataMutable()))
 	{
-		TSharedRef<FNiagaraHierarchyDragDropOp> CategoryDragDropOp = MakeShared<FNiagaraHierarchyDragDropOp>(Item);
-		CategoryDragDropOp->SetAdditionalLabel(FText::FromString(Item->ToString()));
+		TSharedRef<FHierarchyDragDropOp> CategoryDragDropOp = MakeShared<FHierarchyDragDropOp>(Item);
 		CategoryDragDropOp->Construct();
 		return CategoryDragDropOp;
 	}
 	else if(UNiagaraHierarchyModuleInput* ModuleInput = Cast<UNiagaraHierarchyModuleInput>(Item->GetDataMutable()))
 	{
 		TSharedPtr<FNiagaraModuleInputViewModel> ModuleInputViewModel = StaticCastSharedRef<FNiagaraModuleInputViewModel>(Item);
-		TSharedRef<FNiagaraHierarchyDragDropOp> ModuleInputDragDropOp = MakeShared<FNiagaraHierarchyInputParameterHierarchyDragDropOp>(ModuleInputViewModel);
+		TSharedRef<FHierarchyDragDropOp> ModuleInputDragDropOp = MakeShared<FNiagaraHierarchyInputParameterHierarchyDragDropOp>(ModuleInputViewModel);
 		ModuleInputDragDropOp->Construct();
 		return ModuleInputDragDropOp;
 	}
-	else if(UNiagaraHierarchyItem* ItemData = Cast<UNiagaraHierarchyItem>(Item->GetDataMutable()))
+	else if(UHierarchyItem* ItemData = Cast<UHierarchyItem>(Item->GetDataMutable()))
 	{
-		TSharedRef<FNiagaraHierarchyDragDropOp> ObjectPropertyDragDropOp = MakeShared<FNiagaraHierarchyDragDropOp>(Item);
-		ObjectPropertyDragDropOp->SetAdditionalLabel(FText::FromString(Item->ToString()));
+		TSharedRef<FHierarchyDragDropOp> ObjectPropertyDragDropOp = MakeShared<FHierarchyDragDropOp>(Item);
 		ObjectPropertyDragDropOp->Construct();
 		return ObjectPropertyDragDropOp;
 	}
 
 	check(false);
-	return MakeShared<FNiagaraHierarchyDragDropOp>(nullptr);
+	return MakeShared<FHierarchyDragDropOp>(nullptr);
 }
 
 TArray<TTuple<UClass*, FOnGetDetailCustomizationInstance>> UNiagaraSummaryViewViewModel::GetInstanceCustomizations()
diff --git a/Engine/Plugins/FX/Niagara/Source/NiagaraEditor/Private/ViewModels/HierarchyEditor/NiagaraUserParametersHierarchyViewModel.cpp b/Engine/Plugins/FX/Niagara/Source/NiagaraEditor/Private/ViewModels/HierarchyEditor/NiagaraUserParametersHierarchyViewModel.cpp
--- a/Engine/Plugins/FX/Niagara/Source/NiagaraEditor/Private/ViewModels/HierarchyEditor/NiagaraUserParametersHierarchyViewModel.cpp
+++ b/Engine/Plugins/FX/Niagara/Source/NiagaraEditor/Private/ViewModels/HierarchyEditor/NiagaraUserParametersHierarchyViewModel.cpp
@@ -16,7 +16,7 @@
 void UNiagaraHierarchyUserParameter::Initialize(UNiagaraScriptVariable& InUserParameterScriptVariable)
 {
 	UserParameterScriptVariable = &InUserParameterScriptVariable;
-	SetIdentity(FNiagaraHierarchyIdentity({InUserParameterScriptVariable.Metadata.GetVariableGuid()}, {}));
+	SetIdentity(FHierarchyElementIdentity({InUserParameterScriptVariable.Metadata.GetVariableGuid()}, {}));
 }
 
 UObject* FNiagaraHierarchyUserParameterViewModel::GetDataForEditing()
@@ -29,7 +29,7 @@ UObject* FNiagaraHierarchyUserParameterViewModel::GetDataForEditing()
 	return ScriptVariable.Get();
 }
 
-bool FNiagaraHierarchyUserParameterViewModel::DoesExternalDataStillExist(const UNiagaraHierarchyDataRefreshContext* Context) const
+bool FNiagaraHierarchyUserParameterViewModel::DoesExternalDataStillExist(const UHierarchyDataRefreshContext* Context) const
 {
 	const UNiagaraHierarchyUserParameterRefreshContext* UserParameterRefreshContext = CastChecked<UNiagaraHierarchyUserParameterRefreshContext>(Context);
 	const UNiagaraSystem* System = UserParameterRefreshContext->GetSystem();
@@ -46,7 +46,7 @@ TSharedRef<FNiagaraSystemViewModel> UNiagaraUserParametersHierarchyViewModel::Ge
 void UNiagaraUserParametersHierarchyViewModel::Initialize(TSharedRef<FNiagaraSystemViewModel> InSystemViewModel)
 {
 	SystemViewModelWeak = InSystemViewModel;
-	UNiagaraHierarchyViewModelBase::Initialize();
+	UHierarchyViewModelBase::Initialize();
 
 	UNiagaraHierarchyUserParameterRefreshContext* UserParameterRefreshContext = NewObject<UNiagaraHierarchyUserParameterRefreshContext>(this);
 	UserParameterRefreshContext->SetSystem(&InSystemViewModel->GetSystem());
@@ -73,37 +73,29 @@ TSharedRef<SWidget> FNiagaraUserParameterHierarchyDragDropOp::CreateCustomDecora
 	return FNiagaraParameterUtilities::GetParameterWidget(GetUserParameter(), true, false);
 }
 
-UNiagaraHierarchyRoot* UNiagaraUserParametersHierarchyViewModel::GetHierarchyRoot() const
+UHierarchyRoot* UNiagaraUserParametersHierarchyViewModel::GetHierarchyRoot() const
 {
-	UNiagaraHierarchyRoot* RootItem = GetSystemViewModel()->GetEditorData().UserParameterHierarchy;
+	UHierarchyRoot* RootItem = GetSystemViewModel()->GetEditorData().UserParameterHierarchy;
 
 	ensure(RootItem != nullptr);
 	return RootItem;
 }
 
-TSharedPtr<FNiagaraHierarchyItemViewModelBase> UNiagaraUserParametersHierarchyViewModel::CreateViewModelForData(UNiagaraHierarchyItemBase* ItemBase, TSharedPtr<FNiagaraHierarchyItemViewModelBase> Parent)
+TSharedPtr<FHierarchyElementViewModel> UNiagaraUserParametersHierarchyViewModel::CreateCustomViewModelForElement(UHierarchyElement* ItemBase, TSharedPtr<FHierarchyElementViewModel> Parent)
 {
 	if(UNiagaraHierarchyUserParameter* UserParameter = Cast<UNiagaraHierarchyUserParameter>(ItemBase))
 	{
-		return MakeShared<FNiagaraHierarchyUserParameterViewModel>(UserParameter, Parent, this, Parent.IsValid() ? Parent->IsForHierarchy() : false);
+		return MakeShared<FNiagaraHierarchyUserParameterViewModel>(UserParameter, Parent.ToSharedRef(), this);
 	}
-	else if(UNiagaraHierarchyCategory* Category = Cast<UNiagaraHierarchyCategory>(ItemBase))
-	{
-		return MakeShared<FNiagaraHierarchyCategoryViewModel>(Category, Parent, this, Parent.IsValid() ? Parent->IsForHierarchy() : false);
-	}
-
-	check(false);
+	
 	return nullptr;
 }
 
-void UNiagaraUserParametersHierarchyViewModel::PrepareSourceItems(UNiagaraHierarchyRoot* SourceRoot, TSharedPtr<FNiagaraHierarchyRootViewModel>)
+void UNiagaraUserParametersHierarchyViewModel::PrepareSourceItems(UHierarchyRoot* SourceRoot, TSharedPtr<FHierarchyRootViewModel>)
 {	
 	TArray<FNiagaraVariable> UserParameters;
 	GetSystemViewModel()->GetSystem().GetExposedParameters().GetUserParameters(UserParameters);
 
-	TArray<UNiagaraHierarchyItemBase*> OldChildren = SourceRoot->GetChildrenMutable();
-	SourceRoot->GetChildrenMutable().Empty();
-
 	// we ensure we have one data child per user parameter, and we get rid of deleted ones
 	for(FNiagaraVariable& UserParameter : UserParameters)
 	{
@@ -111,15 +103,6 @@ void UNiagaraUserParametersHierarchyViewModel::PrepareSourceItems(UNiagaraHierar
 		GetSystemViewModel()->GetSystem().GetExposedParameters().RedirectUserVariable(RedirectedUserParameter);
 		TObjectPtr<UNiagaraScriptVariable> ScriptVariable = FNiagaraEditorUtilities::UserParameters::GetScriptVariableForUserParameter(RedirectedUserParameter, GetSystemViewModel());
 		
-		if(UNiagaraHierarchyItemBase** ExistingChild = OldChildren.FindByPredicate([ScriptVariable](UNiagaraHierarchyItemBase* ItemBase)
-		{
-			return ItemBase->GetPersistentIdentity().Guids[0] == ScriptVariable->Metadata.GetVariableGuid();
-		}))
-		{
-			SourceRoot->GetChildrenMutable().Add(*ExistingChild);
-			continue;
-		}
-		
 		// since the source items are transient we need to create them here and keep them around until the end of the tool's lifetime
 		UNiagaraHierarchyUserParameter* UserParameterHierarchyObject = NewObject<UNiagaraHierarchyUserParameter>(SourceRoot);
 		UserParameterHierarchyObject->Initialize(*ScriptVariable.Get());
@@ -132,7 +115,7 @@ void UNiagaraUserParametersHierarchyViewModel::SetupCommands()
 	// no custom commands yet
 }
 
-TSharedRef<FNiagaraHierarchyDragDropOp> UNiagaraUserParametersHierarchyViewModel::CreateDragDropOp(TSharedRef<FNiagaraHierarchyItemViewModelBase> Item)
+TSharedRef<FHierarchyDragDropOp> UNiagaraUserParametersHierarchyViewModel::CreateDragDropOp(TSharedRef<FHierarchyElementViewModel> Item)
 {
 	if(const UNiagaraHierarchyUserParameter* UserParameter = Cast<UNiagaraHierarchyUserParameter>(Item->GetData()))
 	{
@@ -141,15 +124,14 @@ TSharedRef<FNiagaraHierarchyDragDropOp> UNiagaraUserParametersHierarchyViewModel
 		return ParameterDragDropOp;
 	}
 
-	if(UNiagaraHierarchyCategory* HierarchyCategory = Cast<UNiagaraHierarchyCategory>(Item->GetDataMutable()))
+	if(UHierarchyCategory* HierarchyCategory = Cast<UHierarchyCategory>(Item->GetDataMutable()))
 	{
-		TSharedRef<FNiagaraHierarchyDragDropOp> CategoryDragDropOp = MakeShared<FNiagaraHierarchyDragDropOp>(Item);
-		CategoryDragDropOp->SetAdditionalLabel(FText::FromName(HierarchyCategory->GetCategoryName()));
+		TSharedRef<FHierarchyDragDropOp> CategoryDragDropOp = MakeShared<FHierarchyDragDropOp>(Item);
 		CategoryDragDropOp->Construct();
 		return CategoryDragDropOp;
 	}
 
-	return MakeShared<FNiagaraHierarchyDragDropOp>(nullptr);
+	return MakeShared<FHierarchyDragDropOp>(nullptr);
 }
 
 TArray<TTuple<UClass*, FOnGetDetailCustomizationInstance>> UNiagaraUserParametersHierarchyViewModel::GetInstanceCustomizations()
diff --git a/Engine/Plugins/FX/Niagara/Source/NiagaraEditor/Private/ViewModels/Stack/NiagaraStackEmitterSettingsGroup.cpp b/Engine/Plugins/FX/Niagara/Source/NiagaraEditor/Private/ViewModels/Stack/NiagaraStackEmitterSettingsGroup.cpp
--- a/Engine/Plugins/FX/Niagara/Source/NiagaraEditor/Private/ViewModels/Stack/NiagaraStackEmitterSettingsGroup.cpp
+++ b/Engine/Plugins/FX/Niagara/Source/NiagaraEditor/Private/ViewModels/Stack/NiagaraStackEmitterSettingsGroup.cpp
@@ -22,7 +22,7 @@
 #include "NiagaraNodeFunctionCall.h"
 #include "NiagaraSettings.h"
 #include "Toolkits/SystemToolkitModes/NiagaraSystemToolkitModeBase.h"
-#include "Widgets/SNiagaraHierarchyEditor.h"
+#include "Widgets/SHierarchyEditor.h"
 
 #include UE_INLINE_GENERATED_CPP_BY_NAME(NiagaraStackEmitterSettingsGroup)
 
diff --git a/Engine/Plugins/FX/Niagara/Source/NiagaraEditor/Private/ViewModels/Stack/NiagaraStackEntry.cpp b/Engine/Plugins/FX/Niagara/Source/NiagaraEditor/Private/ViewModels/Stack/NiagaraStackEntry.cpp
--- a/Engine/Plugins/FX/Niagara/Source/NiagaraEditor/Private/ViewModels/Stack/NiagaraStackEntry.cpp
+++ b/Engine/Plugins/FX/Niagara/Source/NiagaraEditor/Private/ViewModels/Stack/NiagaraStackEntry.cpp
@@ -12,7 +12,7 @@
 #include "NiagaraScriptMergeManager.h"
 #include "Misc/SecureHash.h"
 #include "ScopedTransaction.h"
-#include "ViewModels/HierarchyEditor/NiagaraHierarchyViewModelBase.h"
+#include "HierarchyViewModelBase.h"
 #include "ViewModels/HierarchyEditor/NiagaraSummaryViewViewModel.h"
 #include "ViewModels/Stack/NiagaraStackClipboardUtilities.h"
 
@@ -1255,9 +1255,9 @@ void UNiagaraStackEntry::PasteStackNote()
 	}
 }
 
-FNiagaraHierarchyIdentity UNiagaraStackEntry::DetermineSummaryIdentity() const
+FHierarchyElementIdentity UNiagaraStackEntry::DetermineSummaryIdentity() const
 {
-	return FNiagaraHierarchyIdentity();
+	return FHierarchyElementIdentity();
 }
 
 bool UNiagaraStackEntry::IsInSummaryView() const
@@ -1266,7 +1266,7 @@ bool UNiagaraStackEntry::IsInSummaryView() const
 	{
 		if(SupportsSummaryView())
 		{
-			FNiagaraHierarchyIdentity Identity = DetermineSummaryIdentity();
+			FHierarchyElementIdentity Identity = DetermineSummaryIdentity();
 			if(Identity.IsValid())
 			{
 				return GetEmitterViewModel()->GetSummaryHierarchyViewModel()->GetHierarchyRootViewModel()->FindViewModelForChild(Identity, true) != nullptr;
@@ -1329,7 +1329,7 @@ bool UNiagaraStackEntry::ExistsInParentEmitterSummary() const
 		if(GetEmitterViewModel()->GetParentEmitter().Emitter != nullptr)
 		{
 			FVersionedNiagaraEmitter BaseEmitter = GetEmitterViewModel()->GetParentEmitter();
-			FNiagaraHierarchyIdentity Identity = DetermineSummaryIdentity();
+			FHierarchyElementIdentity Identity = DetermineSummaryIdentity();
 			if(BaseEmitter.GetEmitterData() != nullptr && Identity.IsValid())
 			{
 				return Cast<UNiagaraEmitterEditorData>(BaseEmitter.GetEmitterData()->GetEditorData())->GetSummaryRoot()->FindChildWithIdentity(Identity, true) != nullptr;
diff --git a/Engine/Plugins/FX/Niagara/Source/NiagaraEditor/Private/ViewModels/Stack/NiagaraStackEventScriptItemGroup.cpp b/Engine/Plugins/FX/Niagara/Source/NiagaraEditor/Private/ViewModels/Stack/NiagaraStackEventScriptItemGroup.cpp
--- a/Engine/Plugins/FX/Niagara/Source/NiagaraEditor/Private/ViewModels/Stack/NiagaraStackEventScriptItemGroup.cpp
+++ b/Engine/Plugins/FX/Niagara/Source/NiagaraEditor/Private/ViewModels/Stack/NiagaraStackEventScriptItemGroup.cpp
@@ -163,9 +163,9 @@ void UNiagaraStackEventHandlerPropertiesItem::RefreshChildrenInternal(const TArr
 	Super::RefreshChildrenInternal(CurrentChildren, NewChildren, NewIssues);
 }
 
-FNiagaraHierarchyIdentity UNiagaraStackEventHandlerPropertiesItem::DetermineSummaryIdentity() const
+FHierarchyElementIdentity UNiagaraStackEventHandlerPropertiesItem::DetermineSummaryIdentity() const
 {
-	FNiagaraHierarchyIdentity Identity;
+	FHierarchyElementIdentity Identity;
 	Identity.Guids.Add(GetEventScriptUsageId());
 	Identity.Names.Add("Category");
 	Identity.Names.Add("Properties");
@@ -362,9 +362,9 @@ FText UNiagaraStackEventScriptItemGroup::GetInheritanceMessage() const
 	return LOCTEXT("EventGroupInheritanceMessage", "This event handler is inherited from a parent emitter.  Inherited\nevent handlers can only be deleted while editing the parent emitter.");
 }
 
-FNiagaraHierarchyIdentity UNiagaraStackEventScriptItemGroup::DetermineSummaryIdentity() const
+FHierarchyElementIdentity UNiagaraStackEventScriptItemGroup::DetermineSummaryIdentity() const
 {
-	FNiagaraHierarchyIdentity Identity;
+	FHierarchyElementIdentity Identity;
 	Identity.Guids.Add(GetScriptUsageId());
 	return Identity;
 }
diff --git a/Engine/Plugins/FX/Niagara/Source/NiagaraEditor/Private/ViewModels/Stack/NiagaraStackFunctionInput.cpp b/Engine/Plugins/FX/Niagara/Source/NiagaraEditor/Private/ViewModels/Stack/NiagaraStackFunctionInput.cpp
--- a/Engine/Plugins/FX/Niagara/Source/NiagaraEditor/Private/ViewModels/Stack/NiagaraStackFunctionInput.cpp
+++ b/Engine/Plugins/FX/Niagara/Source/NiagaraEditor/Private/ViewModels/Stack/NiagaraStackFunctionInput.cpp
@@ -135,8 +135,8 @@ void UNiagaraStackFunctionInput::Initialize(
 	TOptional<FGuid> VariableGuid = OwningFunctionCallNode->GetFunctionScriptSource()->NodeGraph->GetScriptVariableGuid(FNiagaraVariable(InInputType, InInputParameterHandle));
 	if(VariableGuid.IsSet())
 	{
-		UNiagaraHierarchyRoot* HierarchyScriptParameterRoot = OwningFunctionCallNode->GetFunctionScriptSource()->NodeGraph->GetScriptParameterHierarchyRoot();
-		HierarchyScriptParameter = Cast<UNiagaraHierarchyScriptParameter>(HierarchyScriptParameterRoot->FindChildWithIdentity(FNiagaraHierarchyIdentity({VariableGuid.GetValue()}, {}), true));
+		UHierarchyRoot* HierarchyScriptParameterRoot = OwningFunctionCallNode->GetFunctionScriptSource()->NodeGraph->GetScriptParameterHierarchyRoot();
+		HierarchyScriptParameter = Cast<UNiagaraHierarchyScriptParameter>(HierarchyScriptParameterRoot->FindChildWithIdentity(FHierarchyElementIdentity({VariableGuid.GetValue()}, {}), true));
 	}
 
 	UNiagaraSystem& ParentSystem = GetSystemViewModel()->GetSystem();
@@ -510,9 +510,9 @@ bool UNiagaraStackFunctionInput::SupportsSummaryView() const
 	return true;
 }
 
-FNiagaraHierarchyIdentity UNiagaraStackFunctionInput::DetermineSummaryIdentity() const
+FHierarchyElementIdentity UNiagaraStackFunctionInput::DetermineSummaryIdentity() const
 {
-	FNiagaraHierarchyIdentity Identity;
+	FHierarchyElementIdentity Identity;
 	Identity.Guids.Add(GetInputFunctionCallNode().NodeGuid);
 	if(Cast<UNiagaraNodeAssignment>(&GetInputFunctionCallNode()) != nullptr)
 	{
diff --git a/Engine/Plugins/FX/Niagara/Source/NiagaraEditor/Private/ViewModels/Stack/NiagaraStackGraphUtilities.cpp b/Engine/Plugins/FX/Niagara/Source/NiagaraEditor/Private/ViewModels/Stack/NiagaraStackGraphUtilities.cpp
--- a/Engine/Plugins/FX/Niagara/Source/NiagaraEditor/Private/ViewModels/Stack/NiagaraStackGraphUtilities.cpp
+++ b/Engine/Plugins/FX/Niagara/Source/NiagaraEditor/Private/ViewModels/Stack/NiagaraStackGraphUtilities.cpp
@@ -48,7 +48,7 @@
 #include "NiagaraSettings.h"
 #include "NiagaraSimulationStageBase.h"
 #include "NiagaraEmitter.h"
-#include "ViewModels/HierarchyEditor/NiagaraHierarchyViewModelBase.h"
+#include "HierarchyViewModelBase.h"
 #include "ViewModels/HierarchyEditor/NiagaraSummaryViewViewModel.h"
 #include "TraversalCache/StaticContextBuilder.h"
 #include "TraversalCache/TraversalCache.h"
@@ -1192,7 +1192,7 @@ void FNiagaraStackGraphUtilities::GatherInputRelationsForStack(FInputDataCollect
 	TSet<const UEdGraphPin*> HiddenPins;
     TArray<const UEdGraphPin*> InputPins;
 	
-	UNiagaraHierarchyRoot* SummaryRoot = EmitterViewModel->GetEditorData().GetSummaryRoot();
+	UHierarchyRoot* SummaryRoot = EmitterViewModel->GetEditorData().GetSummaryRoot();
 
 	TArray<UNiagaraNodeOutput*> OutputNodes = GetAllOutputNodes(EmitterViewModel);
 
diff --git a/Engine/Plugins/FX/Niagara/Source/NiagaraEditor/Private/ViewModels/Stack/NiagaraStackInputCategory.cpp b/Engine/Plugins/FX/Niagara/Source/NiagaraEditor/Private/ViewModels/Stack/NiagaraStackInputCategory.cpp
--- a/Engine/Plugins/FX/Niagara/Source/NiagaraEditor/Private/ViewModels/Stack/NiagaraStackInputCategory.cpp
+++ b/Engine/Plugins/FX/Niagara/Source/NiagaraEditor/Private/ViewModels/Stack/NiagaraStackInputCategory.cpp
@@ -98,7 +98,7 @@ bool UNiagaraStackCategory::FilterForIsInlineEditConditionToggle(const UNiagaraS
 	return StackFunctionInputChild == nullptr || StackFunctionInputChild->GetIsInlineEditConditionToggle() == false;
 }
 
-void UNiagaraStackScriptHierarchyCategory::Initialize(FRequiredEntryData InRequiredEntryData, const UNiagaraHierarchyCategory& InHierarchyCategory, FString InOwningStackItemEditorDataKey, FString InStackEditorDataKey)
+void UNiagaraStackScriptHierarchyCategory::Initialize(FRequiredEntryData InRequiredEntryData, const UHierarchyCategory& InHierarchyCategory, FString InOwningStackItemEditorDataKey, FString InStackEditorDataKey)
 {
 	Super::Initialize(InRequiredEntryData, InOwningStackItemEditorDataKey, InStackEditorDataKey);
 	HierarchyCategory = &InHierarchyCategory;
@@ -218,7 +218,7 @@ FText UNiagaraStackScriptHierarchyCategory::GetTooltipText() const
 
 void UNiagaraStackScriptHierarchyCategory::RefreshChildrenInternal(const TArray<UNiagaraStackEntry*>& CurrentChildren, TArray<UNiagaraStackEntry*>& NewChildren, TArray<FStackIssue>& NewIssues)
 {
-	for(const UNiagaraHierarchyItemBase* ChildHierarchyItem : HierarchyCategory->GetChildren())
+	for(const UHierarchyElement* ChildHierarchyItem : HierarchyCategory->GetChildren())
 	{
 		if(const UNiagaraHierarchyScriptParameter* HierarchyParameter = Cast<UNiagaraHierarchyScriptParameter>(ChildHierarchyItem))
 		{
@@ -248,7 +248,7 @@ void UNiagaraStackScriptHierarchyCategory::RefreshChildrenInternal(const TArray<
 			
 			NewChildren.Add(InputChild);
 		}
-		if(const UNiagaraHierarchyCategory* ChildHierarchyCategory = Cast<UNiagaraHierarchyCategory>(ChildHierarchyItem))
+		if(const UHierarchyCategory* ChildHierarchyCategory = Cast<UHierarchyCategory>(ChildHierarchyItem))
 		{
 			// Try to find an already existing category to reuse
 			UNiagaraStackScriptHierarchyCategory* ChildCategory = FindCurrentChildOfTypeByPredicate<UNiagaraStackScriptHierarchyCategory>(CurrentChildren,
@@ -272,7 +272,7 @@ void UNiagaraStackScriptHierarchyCategory::RefreshChildrenInternal(const TArray<
 	Super::RefreshChildrenInternal(CurrentChildren, NewChildren, NewIssues);
 }
 
-void UNiagaraStackSummaryCategory::Initialize(FRequiredEntryData InRequiredEntryData, TSharedPtr<FNiagaraHierarchyCategoryViewModel> InCategoryViewModel, FString InOwnerStackItemEditorDataKey)
+void UNiagaraStackSummaryCategory::Initialize(FRequiredEntryData InRequiredEntryData, TSharedPtr<FHierarchyCategoryViewModel> InCategoryViewModel, FString InOwnerStackItemEditorDataKey)
 {
 	CategoryViewModelWeakPtr = InCategoryViewModel;
 	
@@ -290,17 +290,17 @@ void UNiagaraStackSummaryCategory::RefreshChildrenInternal(const TArray<UNiagara
 	FInputDataCollection State;
 	GatherInputRelationsForStack(State, GetEmitterViewModel().ToSharedRef());
 
-	TArray<TSharedPtr<FNiagaraHierarchyItemViewModelBase>> ChildrenViewModels;
-	CategoryViewModelWeakPtr.Pin()->GetChildrenViewModelsForType<UNiagaraHierarchyItemBase, FNiagaraHierarchyItemViewModelBase>(ChildrenViewModels, false);
+	TArray<TSharedPtr<FHierarchyElementViewModel>> ChildrenViewModels;
+	CategoryViewModelWeakPtr.Pin()->GetChildrenViewModelsForType<UHierarchyElement, FHierarchyElementViewModel>(ChildrenViewModels, false);
 
-	TArray<TSharedPtr<FNiagaraHierarchyItemViewModelBase>> AllChildrenViewModels;
-	CategoryViewModelWeakPtr.Pin()->GetChildrenViewModelsForType<UNiagaraHierarchyItemBase, FNiagaraHierarchyItemViewModelBase>(AllChildrenViewModels, true);
+	TArray<TSharedPtr<FHierarchyElementViewModel>> AllChildrenViewModels;
+	CategoryViewModelWeakPtr.Pin()->GetChildrenViewModelsForType<UHierarchyElement, FHierarchyElementViewModel>(AllChildrenViewModels, true);
 	
 	// first we gather all function call nodes so we can create cache for them instead of looking it up for each input individually
 	TSet<UNiagaraNodeFunctionCall*> UsedFunctionCallNodes;
-	for(TSharedPtr<FNiagaraHierarchyItemViewModelBase> HierarchyViewModel : AllChildrenViewModels)
+	for(TSharedPtr<FHierarchyElementViewModel> HierarchyViewModel : AllChildrenViewModels)
 	{
-		UNiagaraHierarchyItemBase* Data = HierarchyViewModel->GetDataMutable();
+		UHierarchyElement* Data = HierarchyViewModel->GetDataMutable();
 		TOptional<FGuid> FunctionCallGuid;
 		if(UNiagaraHierarchyModuleInput* ModuleInput = Cast<UNiagaraHierarchyModuleInput>(Data))
 		{
@@ -339,9 +339,9 @@ void UNiagaraStackSummaryCategory::RefreshChildrenInternal(const TArray<UNiagara
 		}
 	}
 	
-	for(TSharedPtr<FNiagaraHierarchyItemViewModelBase> HierarchyViewModel : ChildrenViewModels)
+	for(TSharedPtr<FHierarchyElementViewModel> HierarchyViewModel : ChildrenViewModels)
 	{
-		UNiagaraHierarchyItemBase* Data = HierarchyViewModel->GetDataMutable();
+		UHierarchyElement* Data = HierarchyViewModel->GetDataMutable();
 
 		if(UNiagaraHierarchyModuleInput* ModuleInput = Cast<UNiagaraHierarchyModuleInput>(Data))
 		{
@@ -536,7 +536,7 @@ void UNiagaraStackSummaryCategory::RefreshChildrenInternal(const TArray<UNiagara
 			const TArray<FNiagaraEventScriptProperties>& EventScriptProperties = GetEmitterViewModel()->GetEmitter().GetEmitterData()->GetEventHandlers();
 			const FNiagaraEventScriptProperties* EventScriptPropertiesItem = EventScriptProperties.FindByPredicate([EventHandlerIdentity = EventHandlerProperties->GetPersistentIdentity()](const FNiagaraEventScriptProperties& Candidate)
 			{
-				FNiagaraHierarchyIdentity CandidateIdentity = UNiagaraHierarchyEventHandlerProperties::MakeIdentity(Candidate);
+				FHierarchyElementIdentity CandidateIdentity = UNiagaraHierarchyEventHandlerProperties::MakeIdentity(Candidate);
 				return CandidateIdentity == EventHandlerIdentity;
 			});
 
@@ -653,9 +653,9 @@ void UNiagaraStackSummaryCategory::RefreshChildrenInternal(const TArray<UNiagara
 				NewChildren.Add(SimStageProperties);	
 			}
 		}
-		else if(UNiagaraHierarchyCategory* HierarchyCategory = Cast<UNiagaraHierarchyCategory>(Data))
+		else if(UHierarchyCategory* HierarchyCategory = Cast<UHierarchyCategory>(Data))
 		{			
-			TSharedPtr<FNiagaraHierarchyCategoryViewModel> CategoryViewModel = StaticCastSharedPtr<FNiagaraHierarchyCategoryViewModel>(HierarchyViewModel);
+			TSharedPtr<FHierarchyCategoryViewModel> CategoryViewModel = StaticCastSharedPtr<FHierarchyCategoryViewModel>(HierarchyViewModel);
 			
 			UNiagaraStackSummaryCategory* StackCategory = FindCurrentChildOfTypeByPredicate<UNiagaraStackSummaryCategory>(CurrentChildren,
 			[&](UNiagaraStackSummaryCategory* StackCategoryCandidate)
@@ -726,12 +726,12 @@ void UNiagaraStackSummaryCategory::RefreshChildrenInternal(const TArray<UNiagara
 
 bool UNiagaraStackSummaryCategory::IsTopLevelCategory() const
 {
-	return CategoryViewModelWeakPtr.Pin()->GetData()->GetOuter()->IsA<UNiagaraHierarchyRoot>();
+	return CategoryViewModelWeakPtr.Pin()->GetData()->GetOuter()->IsA<UHierarchyRoot>();
 }
 
 FText UNiagaraStackSummaryCategory::GetTooltipText() const
 {
-	return Cast<UNiagaraHierarchyCategory>(CategoryViewModelWeakPtr.Pin()->GetData())->GetTooltip();
+	return Cast<UHierarchyCategory>(CategoryViewModelWeakPtr.Pin()->GetData())->GetTooltip();
 }
 
 int32 UNiagaraStackSummaryCategory::GetChildIndentLevel() const
diff --git a/Engine/Plugins/FX/Niagara/Source/NiagaraEditor/Private/ViewModels/Stack/NiagaraStackModuleItem.cpp b/Engine/Plugins/FX/Niagara/Source/NiagaraEditor/Private/ViewModels/Stack/NiagaraStackModuleItem.cpp
--- a/Engine/Plugins/FX/Niagara/Source/NiagaraEditor/Private/ViewModels/Stack/NiagaraStackModuleItem.cpp
+++ b/Engine/Plugins/FX/Niagara/Source/NiagaraEditor/Private/ViewModels/Stack/NiagaraStackModuleItem.cpp
@@ -1728,9 +1728,9 @@ FText UNiagaraStackModuleItem::GetInheritanceMessage() const
 	return LOCTEXT("ModuleItemInheritanceMessage", "This module is inherited from a parent emitter.  Inherited modules\ncan only be moved, deleted, and versioned while editing the parent emitter.");
 }
 
-FNiagaraHierarchyIdentity UNiagaraStackModuleItem::DetermineSummaryIdentity() const
+FHierarchyElementIdentity UNiagaraStackModuleItem::DetermineSummaryIdentity() const
 {
-	FNiagaraHierarchyIdentity Identity;
+	FHierarchyElementIdentity Identity;
 	Identity.Guids.Add(GetModuleNode().NodeGuid);
 	return Identity;
 }
diff --git a/Engine/Plugins/FX/Niagara/Source/NiagaraEditor/Private/ViewModels/Stack/NiagaraStackObject.cpp b/Engine/Plugins/FX/Niagara/Source/NiagaraEditor/Private/ViewModels/Stack/NiagaraStackObject.cpp
--- a/Engine/Plugins/FX/Niagara/Source/NiagaraEditor/Private/ViewModels/Stack/NiagaraStackObject.cpp
+++ b/Engine/Plugins/FX/Niagara/Source/NiagaraEditor/Private/ViewModels/Stack/NiagaraStackObject.cpp
@@ -15,7 +15,7 @@
 #include "IDetailTreeNode.h"
 
 #include "NiagaraPlatformSet.h"
-#include "ViewModels/HierarchyEditor/NiagaraHierarchyViewModelBase.h"
+#include "HierarchyViewModelBase.h"
 #include "ViewModels/Stack/NiagaraStackObjectIssueGenerator.h"
 
 #include UE_INLINE_GENERATED_CPP_BY_NAME(NiagaraStackObject)
@@ -119,9 +119,9 @@ UObject* UNiagaraStackObject::GetDisplayedObject() const
 	return WeakObject.Get();
 }
 
-FNiagaraHierarchyIdentity UNiagaraStackObject::DetermineSummaryIdentity() const
+FHierarchyElementIdentity UNiagaraStackObject::DetermineSummaryIdentity() const
 {
-	FNiagaraHierarchyIdentity Identity;
+	FHierarchyElementIdentity Identity;
 	Identity.Guids.Add(ObjectGuid.GetValue());
 	return Identity;
 }
diff --git a/Engine/Plugins/FX/Niagara/Source/NiagaraEditor/Private/ViewModels/Stack/NiagaraStackPropertyRow.cpp b/Engine/Plugins/FX/Niagara/Source/NiagaraEditor/Private/ViewModels/Stack/NiagaraStackPropertyRow.cpp
--- a/Engine/Plugins/FX/Niagara/Source/NiagaraEditor/Private/ViewModels/Stack/NiagaraStackPropertyRow.cpp
+++ b/Engine/Plugins/FX/Niagara/Source/NiagaraEditor/Private/ViewModels/Stack/NiagaraStackPropertyRow.cpp
@@ -8,7 +8,7 @@
 #include "PropertyHandle.h"
 #include "IDetailPropertyRow.h"
 #include "ScopedTransaction.h"
-#include "ViewModels/HierarchyEditor/NiagaraHierarchyViewModelBase.h"
+#include "HierarchyViewModelBase.h"
 
 #include UE_INLINE_GENERATED_CPP_BY_NAME(NiagaraStackPropertyRow)
 
@@ -342,9 +342,9 @@ bool UNiagaraStackPropertyRow::SupportsSummaryView() const
 	return OwnerGuid.IsSet() && OwnerGuid->IsValid() && DetailTreeNode->GetNodeType() == EDetailNodeType::Item;
 }
 
-FNiagaraHierarchyIdentity UNiagaraStackPropertyRow::DetermineSummaryIdentity() const
+FHierarchyElementIdentity UNiagaraStackPropertyRow::DetermineSummaryIdentity() const
 {
-	FNiagaraHierarchyIdentity Identity;
+	FHierarchyElementIdentity Identity;
 	Identity.Guids.Add(OwnerGuid.GetValue());
 	Identity.Names.Add(DetailTreeNode->GetNodeName());
 	return Identity;
diff --git a/Engine/Plugins/FX/Niagara/Source/NiagaraEditor/Private/ViewModels/Stack/NiagaraStackRendererItem.cpp b/Engine/Plugins/FX/Niagara/Source/NiagaraEditor/Private/ViewModels/Stack/NiagaraStackRendererItem.cpp
--- a/Engine/Plugins/FX/Niagara/Source/NiagaraEditor/Private/ViewModels/Stack/NiagaraStackRendererItem.cpp
+++ b/Engine/Plugins/FX/Niagara/Source/NiagaraEditor/Private/ViewModels/Stack/NiagaraStackRendererItem.cpp
@@ -399,9 +399,9 @@ FText UNiagaraStackRendererItem::GetInheritanceMessage() const
 	return LOCTEXT("RendererItemInheritanceMessage", "This renderer is inherited from a parent emitter.  Inherited\nrenderers can only be deleted while editing the parent emitter.");
 }
 
-FNiagaraHierarchyIdentity UNiagaraStackRendererItem::DetermineSummaryIdentity() const
+FHierarchyElementIdentity UNiagaraStackRendererItem::DetermineSummaryIdentity() const
 {
-	FNiagaraHierarchyIdentity Identity;
+	FHierarchyElementIdentity Identity;
 	Identity.Guids.Add(GetRendererProperties()->GetMergeId());
 	return Identity;
 }
diff --git a/Engine/Plugins/FX/Niagara/Source/NiagaraEditor/Private/ViewModels/Stack/NiagaraStackScriptHierarchyRoot.cpp b/Engine/Plugins/FX/Niagara/Source/NiagaraEditor/Private/ViewModels/Stack/NiagaraStackScriptHierarchyRoot.cpp
--- a/Engine/Plugins/FX/Niagara/Source/NiagaraEditor/Private/ViewModels/Stack/NiagaraStackScriptHierarchyRoot.cpp
+++ b/Engine/Plugins/FX/Niagara/Source/NiagaraEditor/Private/ViewModels/Stack/NiagaraStackScriptHierarchyRoot.cpp
@@ -49,23 +49,23 @@ void UNiagaraStackScriptHierarchyRoot::FinalizeInternal()
 	Super::FinalizeInternal();
 }
 
-const TArray<UNiagaraHierarchySection*>& UNiagaraStackScriptHierarchyRoot::GetSections() const
+const TArray<UHierarchySection*>& UNiagaraStackScriptHierarchyRoot::GetSections() const
 {
 	if(OwningFunctionCallNode->GetCalledUsage() == ENiagaraScriptUsage::Module)
 	{
 		return GetScriptParameterHierarchyRoot()->GetSectionData();
 	}
 	
-	static TArray<UNiagaraHierarchySection*> Dummy;
+	static TArray<UHierarchySection*> Dummy;
 	return Dummy;
 }
 
-const UNiagaraHierarchySection* UNiagaraStackScriptHierarchyRoot::GetActiveSection() const
+const UHierarchySection* UNiagaraStackScriptHierarchyRoot::GetActiveSection() const
 {
 	return ActiveSection.IsValid() ? ActiveSection.Get() : nullptr;
 }
 
-void UNiagaraStackScriptHierarchyRoot::SetActiveSection(const UNiagaraHierarchySection* InActiveSection)
+void UNiagaraStackScriptHierarchyRoot::SetActiveSection(const UHierarchySection* InActiveSection)
 {
 	ActiveSection = InActiveSection;
 	
@@ -81,7 +81,7 @@ void UNiagaraStackScriptHierarchyRoot::SetActiveSection(const UNiagaraHierarchyS
 	RefreshFilteredChildren();
 }
 
-UNiagaraHierarchyRoot* UNiagaraStackScriptHierarchyRoot::GetScriptParameterHierarchyRoot() const
+UHierarchyRoot* UNiagaraStackScriptHierarchyRoot::GetScriptParameterHierarchyRoot() const
 {
 	return OwningFunctionCallNode->GetFunctionScriptSource()->NodeGraph->GetScriptParameterHierarchyRoot();
 }
@@ -195,7 +195,7 @@ void UNiagaraStackScriptHierarchyRoot::RefreshChildrenInternal(const TArray<UNia
 		}
 	}
 	
-	for(const UNiagaraHierarchyItemBase* ChildHierarchyItem : GetScriptParameterHierarchyRoot()->GetChildren())
+	for(const UHierarchyElement* ChildHierarchyItem : GetScriptParameterHierarchyRoot()->GetChildren())
 	{
 		if(const UNiagaraHierarchyScriptParameter* HierarchyParameter = Cast<UNiagaraHierarchyScriptParameter>(ChildHierarchyItem))
 		{
@@ -227,7 +227,7 @@ void UNiagaraStackScriptHierarchyRoot::RefreshChildrenInternal(const TArray<UNia
 			NewChildren.AddUnique(InputChild);
 		}
 		
-		if(const UNiagaraHierarchyCategory* HierarchyCategory = Cast<UNiagaraHierarchyCategory>(ChildHierarchyItem))
+		if(const UHierarchyCategory* HierarchyCategory = Cast<UHierarchyCategory>(ChildHierarchyItem))
 		{
 			// Try to find an already existing category to reuse
 			UNiagaraStackScriptHierarchyCategory* StackCategory = FindCurrentChildOfTypeByPredicate<UNiagaraStackScriptHierarchyCategory>(CurrentChildren,
@@ -281,9 +281,9 @@ void UNiagaraStackScriptHierarchyRoot::RefreshChildrenInternal(const TArray<UNia
 	}
 }
 
-const UNiagaraHierarchySection* UNiagaraStackScriptHierarchyRoot::FindSectionByName(FName SectionName)
+const UHierarchySection* UNiagaraStackScriptHierarchyRoot::FindSectionByName(FName SectionName)
 {
-	UNiagaraHierarchySection* const* FoundSection = GetSections().FindByPredicate([SectionName](UNiagaraHierarchySection* SectionCandidate)
+	UHierarchySection* const* FoundSection = GetSections().FindByPredicate([SectionName](UHierarchySection* SectionCandidate)
 	{
 		return SectionCandidate->GetSectionName() == SectionName;
 	});
diff --git a/Engine/Plugins/FX/Niagara/Source/NiagaraEditor/Private/ViewModels/Stack/NiagaraStackSimulationStageGroup.cpp b/Engine/Plugins/FX/Niagara/Source/NiagaraEditor/Private/ViewModels/Stack/NiagaraStackSimulationStageGroup.cpp
--- a/Engine/Plugins/FX/Niagara/Source/NiagaraEditor/Private/ViewModels/Stack/NiagaraStackSimulationStageGroup.cpp
+++ b/Engine/Plugins/FX/Niagara/Source/NiagaraEditor/Private/ViewModels/Stack/NiagaraStackSimulationStageGroup.cpp
@@ -158,7 +158,7 @@ void UNiagaraStackSimulationStagePropertiesItem::RefreshChildrenInternal(const T
 	Super::RefreshChildrenInternal(CurrentChildren, NewChildren, NewIssues);
 }
 
-FNiagaraHierarchyIdentity UNiagaraStackSimulationStagePropertiesItem::DetermineSummaryIdentity() const
+FHierarchyElementIdentity UNiagaraStackSimulationStagePropertiesItem::DetermineSummaryIdentity() const
 {
 	return UNiagaraHierarchySimStageProperties::MakeIdentity(*GetSimulationStage().Get());
 }
@@ -452,9 +452,9 @@ FText UNiagaraStackSimulationStageGroup::GetInheritanceMessage() const
 	return LOCTEXT("SimulationStageGroupInheritanceMessage", "This simulation stage is inherited from a parent emitter.  Inherited\nsimulation stages can only be deleted while editing the parent emitter.");
 }
 
-FNiagaraHierarchyIdentity UNiagaraStackSimulationStageGroup::DetermineSummaryIdentity() const
+FHierarchyElementIdentity UNiagaraStackSimulationStageGroup::DetermineSummaryIdentity() const
 {
-	FNiagaraHierarchyIdentity Identity;
+	FHierarchyElementIdentity Identity;
 	Identity.Guids.Add(GetSimulationStage()->GetMergeId());
 	return Identity;
 }
