New Tool: Hierarchy Editor

diff --git a/Engine/Plugins/FX/Niagara/Source/NiagaraEditor/Private/ViewModels/Stack/NiagaraStackSummaryViewInputCollection.cpp b/Engine/Plugins/FX/Niagara/Source/NiagaraEditor/Private/ViewModels/Stack/NiagaraStackSummaryViewInputCollection.cpp
--- a/Engine/Plugins/FX/Niagara/Source/NiagaraEditor/Private/ViewModels/Stack/NiagaraStackSummaryViewInputCollection.cpp
+++ b/Engine/Plugins/FX/Niagara/Source/NiagaraEditor/Private/ViewModels/Stack/NiagaraStackSummaryViewInputCollection.cpp
@@ -96,8 +96,8 @@ void UNiagaraStackSummaryViewCollection::RefreshChildrenInternal(const TArray<UN
 	TSharedPtr<FNiagaraScriptViewModel> ScriptViewModelPinned = ViewModel->GetSharedScriptViewModel();
 	checkf(ScriptViewModelPinned.IsValid(), TEXT("Can not refresh children when the script view model has been deleted."));
 
-	const UNiagaraHierarchyRoot* Root = ViewModel->GetEditorData().GetSummaryRoot();
-	TSharedPtr<FNiagaraHierarchyRootViewModel> RootViewModel = ViewModel->GetSummaryHierarchyViewModel()->GetHierarchyRootViewModel();
+	const UHierarchyRoot* Root = ViewModel->GetEditorData().GetSummaryRoot();
+	TSharedPtr<FHierarchyRootViewModel> RootViewModel = ViewModel->GetSummaryHierarchyViewModel()->GetHierarchyRootViewModel();
 	//-TODO:Stateless: Do we need stateless support here?
 	if (RootViewModel == nullptr)
 	{
@@ -110,17 +110,17 @@ void UNiagaraStackSummaryViewCollection::RefreshChildrenInternal(const TArray<UN
 	FNiagaraStackGraphUtilities::FInputDataCollection State;
 	FNiagaraStackGraphUtilities::GatherInputRelationsForStack(State, GetEmitterViewModel().ToSharedRef());
 	
-	TArray<TSharedPtr<FNiagaraHierarchyItemViewModelBase>> ChildrenViewModels;
-	RootViewModel->GetChildrenViewModelsForType<UNiagaraHierarchyItemBase, FNiagaraHierarchyItemViewModelBase>(ChildrenViewModels, false);
+	TArray<TSharedPtr<FHierarchyElementViewModel>> ChildrenViewModels;
+	RootViewModel->GetChildrenViewModelsForType<UHierarchyElement, FHierarchyElementViewModel>(ChildrenViewModels, false);
 
-	TArray<TSharedPtr<FNiagaraHierarchyItemViewModelBase>> AllChildrenViewModels;
-	RootViewModel->GetChildrenViewModelsForType<UNiagaraHierarchyItemBase, FNiagaraHierarchyItemViewModelBase>(AllChildrenViewModels, true);
+	TArray<TSharedPtr<FHierarchyElementViewModel>> AllChildrenViewModels;
+	RootViewModel->GetChildrenViewModelsForType<UHierarchyElement, FHierarchyElementViewModel>(AllChildrenViewModels, true);
 
 	// first we gather all function call nodes so we can create cache for them instead of looking it up for each input individually
 	TSet<UNiagaraNodeFunctionCall*> UsedFunctionCallNodes;
-	for(TSharedPtr<FNiagaraHierarchyItemViewModelBase> HierarchyViewModel : AllChildrenViewModels)
+	for(TSharedPtr<FHierarchyElementViewModel> HierarchyViewModel : AllChildrenViewModels)
 	{
-		UNiagaraHierarchyItemBase* Data = HierarchyViewModel->GetDataMutable();
+		UHierarchyElement* Data = HierarchyViewModel->GetDataMutable();
 		TOptional<FGuid> FunctionCallGuid;
 		if(UNiagaraHierarchyModuleInput* ModuleInput = Cast<UNiagaraHierarchyModuleInput>(Data))
 		{
@@ -159,9 +159,9 @@ void UNiagaraStackSummaryViewCollection::RefreshChildrenInternal(const TArray<UN
 		}
 	}
 	
-	for(TSharedPtr<FNiagaraHierarchyItemViewModelBase> HierarchyViewModel : ChildrenViewModels)
+	for(TSharedPtr<FHierarchyElementViewModel> HierarchyViewModel : ChildrenViewModels)
 	{
-		UNiagaraHierarchyItemBase* Data = HierarchyViewModel->GetDataMutable();
+		UHierarchyElement* Data = HierarchyViewModel->GetDataMutable();
 
 		if(UNiagaraHierarchyModuleInput* ModuleInput = Cast<UNiagaraHierarchyModuleInput>(Data))
 		{
@@ -390,7 +390,7 @@ void UNiagaraStackSummaryViewCollection::RefreshChildrenInternal(const TArray<UN
 			const TArray<FNiagaraEventScriptProperties>& EventScriptProperties = GetEmitterViewModel()->GetEmitter().GetEmitterData()->GetEventHandlers();
 			const FNiagaraEventScriptProperties* EventScriptPropertiesItem = EventScriptProperties.FindByPredicate([EventHandlerIdentity = EventHandlerProperties->GetPersistentIdentity()](const FNiagaraEventScriptProperties& Candidate)
 			{
-				FNiagaraHierarchyIdentity CandidateIdentity = UNiagaraHierarchyEventHandlerProperties::MakeIdentity(Candidate);
+				FHierarchyElementIdentity CandidateIdentity = UNiagaraHierarchyEventHandlerProperties::MakeIdentity(Candidate);
 				return CandidateIdentity == EventHandlerIdentity;
 			});
 
@@ -472,9 +472,9 @@ void UNiagaraStackSummaryViewCollection::RefreshChildrenInternal(const TArray<UN
 				NewChildren.Add(SimStageProperties);	
 			}
 		}
-		else if(UNiagaraHierarchyCategory* HierarchyCategory = Cast<UNiagaraHierarchyCategory>(Data))
+		else if(UHierarchyCategory* HierarchyCategory = Cast<UHierarchyCategory>(Data))
 		{			
-			TSharedPtr<FNiagaraHierarchyCategoryViewModel> CategoryViewModel = StaticCastSharedPtr<FNiagaraHierarchyCategoryViewModel>(HierarchyViewModel);
+			TSharedPtr<FHierarchyCategoryViewModel> CategoryViewModel = StaticCastSharedPtr<FHierarchyCategoryViewModel>(HierarchyViewModel);
 			
 			UNiagaraStackSummaryCategory* StackCategory = FindCurrentChildOfTypeByPredicate<UNiagaraStackSummaryCategory>(CurrentChildren,
 			[&](UNiagaraStackSummaryCategory* StackCategoryCandidate)
@@ -562,22 +562,22 @@ void UNiagaraStackSummaryViewCollection::RefreshChildrenInternal(const TArray<UN
 
 void UNiagaraStackSummaryViewCollection::GetSectionsInternal(TArray<FNiagaraStackSection>& OutStackSections) const
 {
-	const TArray<UNiagaraHierarchySection*>& HierarchySections = GetHierarchySections();
+	const TArray<UHierarchySection*>& HierarchySections = GetHierarchySections();
 
-	TMap<const UNiagaraHierarchySection*, TArray<UNiagaraHierarchyCategory*>> SectionToCategoryMap;
+	TMap<const UHierarchySection*, TArray<UHierarchyCategory*>> SectionToCategoryMap;
 
 	TSharedPtr<FNiagaraEmitterViewModel> ViewModel = GetEmitterViewModel();
-	const UNiagaraHierarchyRoot* Root = ViewModel->GetEditorData().GetSummaryRoot();
+	const UHierarchyRoot* Root = ViewModel->GetEditorData().GetSummaryRoot();
 	
-	TArray<UNiagaraHierarchyCategory*> HierarchyCategories;
-	Root->GetChildrenOfType<UNiagaraHierarchyCategory>(HierarchyCategories, false);
+	TArray<UHierarchyCategory*> HierarchyCategories;
+	Root->GetChildrenOfType<UHierarchyCategory>(HierarchyCategories, false);
 	
-	for(UNiagaraHierarchyCategory* TopLevelHierarchyCategories : HierarchyCategories)
+	for(UHierarchyCategory* TopLevelHierarchyCategories : HierarchyCategories)
 	{
 		SectionToCategoryMap.FindOrAdd(TopLevelHierarchyCategories->GetSection()).Add(TopLevelHierarchyCategories);
 	}
 	
-	for(const UNiagaraHierarchySection* HierarchySection : HierarchySections)
+	for(const UHierarchySection* HierarchySection : HierarchySections)
 	{
 		FNiagaraStackSection Section;
 		Section.SectionDisplayName = HierarchySection->GetSectionNameAsText();
@@ -585,7 +585,7 @@ void UNiagaraStackSummaryViewCollection::GetSectionsInternal(TArray<FNiagaraStac
 		
 		if(SectionToCategoryMap.Contains(HierarchySection))
 		{
-			for(UNiagaraHierarchyCategory* Category : SectionToCategoryMap[HierarchySection])
+			for(UHierarchyCategory* Category : SectionToCategoryMap[HierarchySection])
 			{
 				Section.Categories.Add(FText::FromString(Category->ToString()));
 			}
@@ -595,14 +595,14 @@ void UNiagaraStackSummaryViewCollection::GetSectionsInternal(TArray<FNiagaraStac
 	}
 }
 
-const TArray<UNiagaraHierarchySection*>& UNiagaraStackSummaryViewCollection::GetHierarchySections() const
+const TArray<UHierarchySection*>& UNiagaraStackSummaryViewCollection::GetHierarchySections() const
 {
 	if (GetEmitterViewModel().IsValid())
 	{
 		return GetEmitterViewModel()->GetEditorData().GetSummarySections();
 	}
 
-	static TArray<UNiagaraHierarchySection*> Dummy;
+	static TArray<UHierarchySection*> Dummy;
 	return Dummy;
 }
 
diff --git a/Engine/Plugins/FX/Niagara/Source/NiagaraEditor/Private/Widgets/SNiagaraScriptInputPreviewPanel.cpp b/Engine/Plugins/FX/Niagara/Source/NiagaraEditor/Private/Widgets/SNiagaraScriptInputPreviewPanel.cpp
--- a/Engine/Plugins/FX/Niagara/Source/NiagaraEditor/Private/Widgets/SNiagaraScriptInputPreviewPanel.cpp
+++ b/Engine/Plugins/FX/Niagara/Source/NiagaraEditor/Private/Widgets/SNiagaraScriptInputPreviewPanel.cpp
@@ -2,6 +2,8 @@
 
 #include "Widgets/SNiagaraScriptInputPreviewPanel.h"
 
+#include "HierarchyEditorStyle.h"
+#include "HierarchyViewModelBase.h"
 #include "NiagaraEditorStyle.h"
 #include "NiagaraEditorUtilities.h"
 #include "NiagaraObjectSelection.h"
@@ -10,7 +12,6 @@
 #include "SlateOptMacros.h"
 #include "Framework/Application/SlateApplication.h"
 #include "Toolkits/NiagaraScriptToolkit.h"
-#include "ViewModels/HierarchyEditor/NiagaraHierarchyViewModelBase.h"
 #include "ViewModels/HierarchyEditor/NiagaraHierarchyScriptParametersViewModel.h"
 #include "Widgets/Input/SButton.h"
 #include "Widgets/Input/SComboButton.h"
@@ -55,7 +56,7 @@ void SNiagaraScriptInputPreviewPanel::Construct(const FArguments& InArgs, TShare
 		]
 		+ SVerticalBox::Slot()
 		[
-			SAssignNew(TreeView, STreeView<UNiagaraHierarchyItemBase*>)
+			SAssignNew(TreeView, STreeView<UHierarchyElement*>)
 			.TreeItemsSource(&RootArray)
 			.OnGenerateRow(this, &SNiagaraScriptInputPreviewPanel::OnGenerateRow)
 			.OnGetChildren(this, &SNiagaraScriptInputPreviewPanel::OnGetChildren)
@@ -78,7 +79,7 @@ void SNiagaraScriptInputPreviewPanel::Refresh()
 		return;
 	}
 	
-	TWeakObjectPtr<UNiagaraHierarchyRoot> Root = ScriptToolkit.Pin()->GetHierarchyViewModel()->GetHierarchyRoot();
+	TWeakObjectPtr<UHierarchyRoot> Root = ScriptToolkit.Pin()->GetHierarchyViewModel()->GetHierarchyRoot();
 	if(Root.IsValid())
 	{
 		TransientLeftoverParameters.Empty();
@@ -124,25 +125,25 @@ void SNiagaraScriptInputPreviewPanel::Refresh()
 	}
 }
 
-TSharedRef<ITableRow> SNiagaraScriptInputPreviewPanel::OnGenerateRow(UNiagaraHierarchyItemBase* Item, const TSharedRef<STableViewBase>& TableViewBase) const
+TSharedRef<ITableRow> SNiagaraScriptInputPreviewPanel::OnGenerateRow(UHierarchyElement* Item, const TSharedRef<STableViewBase>& TableViewBase) const
 {
-	TSharedPtr<STableRow<UNiagaraHierarchyItemBase*>> TableRow;
-	if(UNiagaraHierarchyCategory* Category = Cast<UNiagaraHierarchyCategory>(Item))
+	TSharedPtr<STableRow<UHierarchyElement*>> TableRow;
+	if(UHierarchyCategory* Category = Cast<UHierarchyCategory>(Item))
 	{		
-		TableRow = SNew(STableRow<UNiagaraHierarchyItemBase*>, TableViewBase)
-		.ToolTipText_UObject(Category, &UNiagaraHierarchyCategory::GetTooltip)
-		.Style(FNiagaraEditorStyle::Get(), "NiagaraEditor.HierarchyEditor.Row.Category")
+		TableRow = SNew(STableRow<UHierarchyElement*>, TableViewBase)
+		.ToolTipText_UObject(Category, &UHierarchyCategory::GetTooltip)
+		.Style(FHierarchyEditorStyle::Get(), "HierarchyEditor.Row.Category")
 		.Padding(FMargin(0.f, 6.f))
 		[
 			SNew(SRichTextBlock)
-			.Text_UObject(Item, &UNiagaraHierarchyItemBase::ToText)
+			.Text_UObject(Item, &UHierarchyElement::ToText)
 			.TextStyle(&FNiagaraEditorStyle::Get().GetWidgetStyle<FTextBlockStyle>("NiagaraEditor.Parameters.HeaderText"))
 		];
 	}
 	else if(UNiagaraHierarchyScriptParameter* ScriptParameter = Cast<UNiagaraHierarchyScriptParameter>(Item))
 	{
-		TWeakObjectPtr<UNiagaraHierarchyItemBase> WeakItem = ScriptParameter;
-		TableRow = SNew(STableRow<UNiagaraHierarchyItemBase*>, TableViewBase)
+		TWeakObjectPtr<UHierarchyElement> WeakItem = ScriptParameter;
+		TableRow = SNew(STableRow<UHierarchyElement*>, TableViewBase)
 		.ToolTipText_UObject(ScriptParameter, &UNiagaraHierarchyScriptParameter::GetTooltip)
 		.Padding(FMargin(0.f, 2.f))
 		[
@@ -153,7 +154,7 @@ TSharedRef<ITableRow> SNiagaraScriptInputPreviewPanel::OnGenerateRow(UNiagaraHie
 	return TableRow.ToSharedRef();
 }
 
-void SNiagaraScriptInputPreviewPanel::OnGetChildren(UNiagaraHierarchyItemBase* Item, TArray<UNiagaraHierarchyItemBase*>& OutChildren) const
+void SNiagaraScriptInputPreviewPanel::OnGetChildren(UHierarchyElement* Item, TArray<UHierarchyElement*>& OutChildren) const
 {
 	OutChildren.Append(Item->GetChildrenMutable());
 }
@@ -179,7 +180,7 @@ void SNiagaraScriptInputPreviewPanel::OnSearchTextChanged(const FText& Text)
 		FString TextAsString = Text.ToString();;
 
 		TArray<FSearchItem> SearchItems;
-		for(UNiagaraHierarchyItemBase* RootItem : RootArray)
+		for(UHierarchyElement* RootItem : RootArray)
 		{
 			GenerateSearchItems(RootItem, {}, SearchItems);
 			
@@ -230,12 +231,12 @@ void SNiagaraScriptInputPreviewPanel::OnSearchButtonClicked(SSearchBox::SearchDi
 	}
 }
 
-void SNiagaraScriptInputPreviewPanel::GenerateSearchItems(UNiagaraHierarchyItemBase* InRoot, TArray<UNiagaraHierarchyItemBase*> ParentChain, TArray<FSearchItem>& OutSearchItems)
+void SNiagaraScriptInputPreviewPanel::GenerateSearchItems(UHierarchyElement* InRoot, TArray<UHierarchyElement*> ParentChain, TArray<FSearchItem>& OutSearchItems)
 {
-	const TArray<UNiagaraHierarchyItemBase*>& Children = InRoot->GetChildren();
+	const TArray<UHierarchyElement*>& Children = InRoot->GetChildren();
 	ParentChain.Add(InRoot);
 	OutSearchItems.Add(FSearchItem{ParentChain});
-	for(UNiagaraHierarchyItemBase* Child : Children)
+	for(UHierarchyElement* Child : Children)
 	{
 		GenerateSearchItems(Child, ParentChain, OutSearchItems);
 	}
@@ -247,7 +248,7 @@ void SNiagaraScriptInputPreviewPanel::ExpandSourceSearchResults()
 
 	for(FSearchItem& SearchResult : SearchResults)
 	{
-		for(UNiagaraHierarchyItemBase* EntryInPath : SearchResult.Path)
+		for(UHierarchyElement* EntryInPath : SearchResult.Path)
 		{
 			TreeView->SetItemExpansion(EntryInPath, true);
 		}
diff --git a/Engine/Plugins/FX/Niagara/Source/NiagaraEditor/Public/NiagaraClipboard.h b/Engine/Plugins/FX/Niagara/Source/NiagaraEditor/Public/NiagaraClipboard.h
--- a/Engine/Plugins/FX/Niagara/Source/NiagaraEditor/Public/NiagaraClipboard.h
+++ b/Engine/Plugins/FX/Niagara/Source/NiagaraEditor/Public/NiagaraClipboard.h
@@ -11,7 +11,7 @@
 #include "UObject/SoftObjectPtr.h"
 #include "NiagaraClipboard.generated.h"
 
-class UNiagaraHierarchyCategory;
+class UHierarchyCategory;
 class IPropertyHandle;
 class UNiagaraDataInterface;
 class UNiagaraScript;
diff --git a/Engine/Plugins/FX/Niagara/Source/NiagaraEditor/Public/NiagaraEditorStyle.h b/Engine/Plugins/FX/Niagara/Source/NiagaraEditor/Public/NiagaraEditorStyle.h
--- a/Engine/Plugins/FX/Niagara/Source/NiagaraEditor/Public/NiagaraEditorStyle.h
+++ b/Engine/Plugins/FX/Niagara/Source/NiagaraEditor/Public/NiagaraEditorStyle.h
@@ -54,7 +54,6 @@ private:
 	void InitScalabilityColors();
 	void InitScalabilityIcons();
 	void InitScratchStyle();
-	void InitHierarchyEditor();
 	void InitSimCacheEditor();
 
 	static TSharedPtr<FNiagaraEditorStyle> NiagaraEditorStyle;
diff --git a/Engine/Plugins/FX/Niagara/Source/NiagaraEditor/Public/NiagaraEditorUtilities.h b/Engine/Plugins/FX/Niagara/Source/NiagaraEditor/Public/NiagaraEditorUtilities.h
--- a/Engine/Plugins/FX/Niagara/Source/NiagaraEditor/Public/NiagaraEditorUtilities.h
+++ b/Engine/Plugins/FX/Niagara/Source/NiagaraEditor/Public/NiagaraEditorUtilities.h
@@ -12,7 +12,7 @@
 #include "NiagaraEditorSettings.h"
 #include "ViewModels/NiagaraSystemScalabilityViewModel.h"
 #include "ViewModels/NiagaraSystemViewModel.h"
-#include "ViewModels/HierarchyEditor/NiagaraHierarchyViewModelBase.h"
+#include "HierarchyViewModelBase.h"
 
 class UNiagaraClipboardContent;
 class UNiagaraNodeInput;
@@ -482,7 +482,7 @@ namespace FNiagaraEditorUtilities
 	{		
 		namespace Scripts
 		{
-			TSharedRef<SWidget> GenerateRowContentForScriptParameterHierarchyEditor(TSharedRef<FNiagaraHierarchyItemViewModelBase> HierarchyItem);
+			TSharedRef<SWidget> GenerateRowContentForScriptParameterHierarchyEditor(TSharedRef<FHierarchyElementViewModel> HierarchyItem);
 		}
 	}
 }
diff --git a/Engine/Plugins/FX/Niagara/Source/NiagaraEditor/Public/NiagaraEmitterEditorData.h b/Engine/Plugins/FX/Niagara/Source/NiagaraEditor/Public/NiagaraEmitterEditorData.h
--- a/Engine/Plugins/FX/Niagara/Source/NiagaraEditor/Public/NiagaraEmitterEditorData.h
+++ b/Engine/Plugins/FX/Niagara/Source/NiagaraEditor/Public/NiagaraEmitterEditorData.h
@@ -5,7 +5,7 @@
 #include "NiagaraEditorDataBase.h"
 #include "NiagaraEditorCommon.h"
 #include "NiagaraStackSection.h"
-#include "ViewModels/HierarchyEditor/NiagaraHierarchyViewModelBase.h"
+#include "HierarchyViewModelBase.h"
 #include "NiagaraEmitterEditorData.generated.h"
 
 class UNiagaraStackEditorData;
@@ -44,8 +44,8 @@ public:
 	NIAGARAEDITOR_API void SetShowSummaryView(bool bInShouldShowSummaryView);
 	NIAGARAEDITOR_API void ToggleShowSummaryView();
 	
-	UNiagaraHierarchyRoot* GetSummaryRoot() const { return SummaryViewRoot; }
-	const TArray<UNiagaraHierarchySection*>& GetSummarySections() const;
+	UHierarchyRoot* GetSummaryRoot() const { return SummaryViewRoot; }
+	const TArray<UHierarchySection*>& GetSummarySections() const;
 
 	UTexture2D* GetThumbnail() const { return EmitterThumbnail; }
 	void SetThumbnail(UTexture2D* InThumbnail) { EmitterThumbnail = InThumbnail; OnPersistentDataChanged().Broadcast(); }
@@ -70,7 +70,7 @@ private:
 	TArray<FNiagaraStackSection> SummarySections_DEPRECATED;
 
 	UPROPERTY()
-	TObjectPtr<UNiagaraHierarchyRoot> SummaryViewRoot = nullptr;
+	TObjectPtr<UHierarchyRoot> SummaryViewRoot = nullptr;
 
 	UPROPERTY()
 	TObjectPtr<UTexture2D> EmitterThumbnail = nullptr;
diff --git a/Engine/Plugins/FX/Niagara/Source/NiagaraEditor/Public/NiagaraGraph.h b/Engine/Plugins/FX/Niagara/Source/NiagaraEditor/Public/NiagaraGraph.h
--- a/Engine/Plugins/FX/Niagara/Source/NiagaraEditor/Public/NiagaraGraph.h
+++ b/Engine/Plugins/FX/Niagara/Source/NiagaraEditor/Public/NiagaraGraph.h
@@ -13,7 +13,7 @@
 class UNiagaraNodeStaticSwitch;
 class UNiagaraParameterDefinitions;
 class UNiagaraScriptVariable;
-class UNiagaraHierarchyRoot;
+class UHierarchyRoot;
 struct FNiagaraCompileHashVisitorDebugInfo;
 struct FNiagaraScriptVariableData;
 struct FNiagaraStaticVariableSearchContext;
@@ -320,7 +320,7 @@ public:
 	NIAGARAEDITOR_API const FScriptVariableMap& GetAllMetaData() const;
 	NIAGARAEDITOR_API FScriptVariableMap& GetAllMetaData();
 
-	NIAGARAEDITOR_API UNiagaraHierarchyRoot* GetScriptParameterHierarchyRoot() const { return ParameterHierarchyRoot; }
+	NIAGARAEDITOR_API UHierarchyRoot* GetScriptParameterHierarchyRoot() const { return ParameterHierarchyRoot; }
 
 	UNiagaraScriptVariable* GetScriptVariable(FNiagaraVariable Parameter) const;
 	NIAGARAEDITOR_API UNiagaraScriptVariable* GetScriptVariable(FName ParameterName) const;
@@ -532,7 +532,7 @@ private:
 	TMap<FNiagaraVariable, TObjectPtr<UNiagaraScriptVariable>> VariableToScriptVariable;
 
 	UPROPERTY()
-	TObjectPtr<UNiagaraHierarchyRoot> ParameterHierarchyRoot;
+	TObjectPtr<UHierarchyRoot> ParameterHierarchyRoot;
 	
 	/** A map of parameters in the graph to their referencers. */
 	UPROPERTY(Transient)
diff --git a/Engine/Plugins/FX/Niagara/Source/NiagaraEditor/Public/NiagaraSystemEditorData.h b/Engine/Plugins/FX/Niagara/Source/NiagaraEditor/Public/NiagaraSystemEditorData.h
--- a/Engine/Plugins/FX/Niagara/Source/NiagaraEditor/Public/NiagaraSystemEditorData.h
+++ b/Engine/Plugins/FX/Niagara/Source/NiagaraEditor/Public/NiagaraSystemEditorData.h
@@ -5,7 +5,7 @@
 #include "Misc/FrameRate.h"
 #include "NiagaraEditorDataBase.h"
 #include "NiagaraEditorData.h"
-#include "ViewModels/HierarchyEditor/NiagaraHierarchyViewModelBase.h"
+#include "HierarchyViewModelBase.h"
 
 #include "NiagaraSystemEditorData.generated.h"
 
@@ -74,7 +74,7 @@ class UNiagaraSystemEditorData : public UNiagaraEditorDataBase
 public:
 	NIAGARAEDITOR_API UNiagaraSystemEditorData(const FObjectInitializer& ObjectInitializer);
 
-	NIAGARAEDITOR_API void PostInitProperties();
+	NIAGARAEDITOR_API virtual void PostInitProperties() override;
 
 	NIAGARAEDITOR_API virtual void PostLoadFromOwner(UObject* InOwner) override;
 #if WITH_EDITORONLY_DATA
@@ -135,7 +135,7 @@ public:
 	
 	/** Contains the root ids for organizing user parameters. */
 	UPROPERTY()
-	TObjectPtr<UNiagaraHierarchyRoot> UserParameterHierarchy;
+	TObjectPtr<UHierarchyRoot> UserParameterHierarchy;
 	
 private:
 	NIAGARAEDITOR_API void UpdatePlaybackRangeFromEmitters(UNiagaraSystem& OwnerSystem);
diff --git a/Engine/Plugins/FX/Niagara/Source/NiagaraEditor/Public/ViewModels/HierarchyEditor/NiagaraHierarchyScriptParametersViewModel.h b/Engine/Plugins/FX/Niagara/Source/NiagaraEditor/Public/ViewModels/HierarchyEditor/NiagaraHierarchyScriptParametersViewModel.h
--- a/Engine/Plugins/FX/Niagara/Source/NiagaraEditor/Public/ViewModels/HierarchyEditor/NiagaraHierarchyScriptParametersViewModel.h
+++ b/Engine/Plugins/FX/Niagara/Source/NiagaraEditor/Public/ViewModels/HierarchyEditor/NiagaraHierarchyScriptParametersViewModel.h
@@ -3,7 +3,7 @@
 #pragma once
 
 #include "NiagaraGraph.h"
-#include "ViewModels/HierarchyEditor/NiagaraHierarchyViewModelBase.h"
+#include "HierarchyViewModelBase.h"
 #include "NiagaraScriptVariable.h"
 #include "ViewModels/NiagaraScriptViewModel.h"
 #include "NiagaraHierarchyScriptParametersViewModel.generated.h"
@@ -12,7 +12,7 @@ class UNiagaraGraph;
 
 /** The refresh context is used to determine if hierarchy script variables should be removed. */
 UCLASS()
-class UNiagaraHierarchyScriptParameterRefreshContext : public UNiagaraHierarchyDataRefreshContext
+class UNiagaraHierarchyScriptParameterRefreshContext : public UHierarchyDataRefreshContext
 {
 	GENERATED_BODY()
 
@@ -26,7 +26,7 @@ private:
 
 /** A hierarchy script parameter is an optional object embedded in the hierarchy. */
 UCLASS()
-class UNiagaraHierarchyScriptParameter : public UNiagaraHierarchyItem
+class UNiagaraHierarchyScriptParameter : public UHierarchyItem
 {
 	GENERATED_BODY()
 
@@ -45,7 +45,7 @@ private:
 
 /** The category class used for the script hierarchy editor. It lets us add additional data later on. */
 UCLASS()
-class UNiagaraHierarchyScriptCategory : public UNiagaraHierarchyCategory
+class UNiagaraHierarchyScriptCategory : public UHierarchyCategory
 {
 	GENERATED_BODY()
 };
@@ -54,21 +54,20 @@ class UNiagaraHierarchyScriptCategory : public UNiagaraHierarchyCategory
  * The view model that defines the script editor's hierarchy editor for input parameters.
  */
 UCLASS()
-class NIAGARAEDITOR_API UNiagaraHierarchyScriptParametersViewModel : public UNiagaraHierarchyViewModelBase
+class NIAGARAEDITOR_API UNiagaraHierarchyScriptParametersViewModel : public UHierarchyViewModelBase
 {
 	GENERATED_BODY()
 public:
 	void Initialize(TSharedRef<FNiagaraScriptViewModel> InScriptViewModel);
 	TSharedRef<FNiagaraScriptViewModel> GetScriptViewModel() const;
 
-	/** UNiagaraHierarchyViewModelBase */
-	virtual UNiagaraHierarchyRoot* GetHierarchyRoot() const override;
-	virtual TSubclassOf<UNiagaraHierarchyCategory> GetCategoryDataClass() const override;
-	virtual TSharedPtr<FNiagaraHierarchyRootViewModel> CreateRootViewModelForData(UNiagaraHierarchyRoot* Root, bool bIsForHierarchy) override;
-	virtual TSharedPtr<FNiagaraHierarchyItemViewModelBase> CreateViewModelForData(UNiagaraHierarchyItemBase* ItemBase, TSharedPtr<FNiagaraHierarchyItemViewModelBase> Parent) override;
-	virtual void PrepareSourceItems(UNiagaraHierarchyRoot* SourceRoot, TSharedPtr<FNiagaraHierarchyRootViewModel> SourceRootViewModel) override;
+	/** UHierarchyViewModelBase */
+	virtual UHierarchyRoot* GetHierarchyRoot() const override;
+	virtual TSubclassOf<UHierarchyCategory> GetCategoryDataClass() const override;
+	virtual TSharedPtr<FHierarchyElementViewModel> CreateCustomViewModelForElement(UHierarchyElement* Element, TSharedPtr<FHierarchyElementViewModel> Parent) override;
+	virtual void PrepareSourceItems(UHierarchyRoot* SourceRoot, TSharedPtr<FHierarchyRootViewModel> SourceRootViewModel) override;
 	virtual void SetupCommands() override;
-	virtual TSharedRef<FNiagaraHierarchyDragDropOp> CreateDragDropOp(TSharedRef<FNiagaraHierarchyItemViewModelBase> Item) override;
+	virtual TSharedRef<FHierarchyDragDropOp> CreateDragDropOp(TSharedRef<FHierarchyElementViewModel> Item) override;
 	virtual bool SupportsDetailsPanel() override { return true; }
 	virtual void FinalizeInternal() override;
 private:
@@ -77,39 +76,39 @@ private:
 	TWeakPtr<FNiagaraScriptViewModel> ScriptViewModelWeak;
 };
 
-class FNiagaraHierarchyScriptParameterDragDropOp : public FNiagaraHierarchyDragDropOp
+class FNiagaraHierarchyScriptParameterDragDropOp : public FHierarchyDragDropOp
 {
 public:
-	DRAG_DROP_OPERATOR_TYPE(FNiagaraHierarchyScriptParameterDragDropOp, FNiagaraHierarchyDragDropOp)
+	DRAG_DROP_OPERATOR_TYPE(FNiagaraHierarchyScriptParameterDragDropOp, FHierarchyDragDropOp)
 
-	FNiagaraHierarchyScriptParameterDragDropOp(TSharedPtr<FNiagaraHierarchyItemViewModel> InputViewModel) : FNiagaraHierarchyDragDropOp(InputViewModel) {}
+	FNiagaraHierarchyScriptParameterDragDropOp(TSharedPtr<FHierarchyItemViewModel> InputViewModel) : FHierarchyDragDropOp(InputViewModel) {}
 	
 	virtual TSharedRef<SWidget> CreateCustomDecorator() const override;
 };
 
-struct FNiagaraHierarchyScriptParameterViewModel : public FNiagaraHierarchyItemViewModel
+struct FNiagaraHierarchyScriptParameterViewModel : public FHierarchyItemViewModel
 {
-	FNiagaraHierarchyScriptParameterViewModel(UNiagaraHierarchyScriptParameter* ScriptParameter, TSharedPtr<FNiagaraHierarchyItemViewModelBase> InParent, TWeakObjectPtr<UNiagaraHierarchyScriptParametersViewModel> ViewModel, bool bInIsForHierarchy)
-		: FNiagaraHierarchyItemViewModel(ScriptParameter, InParent, ViewModel, bInIsForHierarchy)	{}
+	FNiagaraHierarchyScriptParameterViewModel(UNiagaraHierarchyScriptParameter* ScriptParameter, TSharedRef<FHierarchyElementViewModel> InParent, TWeakObjectPtr<UNiagaraHierarchyScriptParametersViewModel> ViewModel)
+		: FHierarchyItemViewModel(ScriptParameter, InParent, ViewModel)	{}
 
 protected:
 	virtual bool RepresentsExternalData() const override { return true; }
-	virtual bool DoesExternalDataStillExist(const UNiagaraHierarchyDataRefreshContext* Context) const override;
+	virtual bool DoesExternalDataStillExist(const UHierarchyDataRefreshContext* Context) const override;
 	virtual UObject* GetDataForEditing() override { return Cast<UNiagaraHierarchyScriptParameter>(GetDataMutable())->GetScriptVariable(); }
 	/** We want to be able to edit in the details panel regardless of source or hierarchy item. */
 	virtual bool AllowEditingInDetailsPanel() const override { return true; }
 	virtual bool CanRenameInternal() override { return false; }
 	
-	virtual FCanPerformActionResults CanDropOnInternal(TSharedPtr<FNiagaraHierarchyItemViewModelBase>, EItemDropZone ItemDropZone) override;
-	virtual void OnDroppedOnInternal(TSharedPtr<FNiagaraHierarchyItemViewModelBase> DroppedItem, EItemDropZone ItemDropZone) override;
+	virtual FCanPerformActionResults CanDropOnInternal(TSharedPtr<FHierarchyElementViewModel>, EItemDropZone ItemDropZone) override;
+	virtual void OnDroppedOnInternal(TSharedPtr<FHierarchyElementViewModel> DroppedItem, EItemDropZone ItemDropZone) override;
 	
 	virtual bool CanHaveChildren() const override { return bIsForHierarchy == true; }
 };
 
-struct FNiagaraHierarchyScriptRootViewModel : public FNiagaraHierarchyRootViewModel
+struct FNiagaraHierarchyScriptRootViewModel : public FHierarchyRootViewModel
 {
-	FNiagaraHierarchyScriptRootViewModel(UNiagaraHierarchyRoot* Root, TWeakObjectPtr<UNiagaraHierarchyScriptParametersViewModel> ViewModel, bool bInIsForHierarchy)
-		: FNiagaraHierarchyRootViewModel(Root, ViewModel, bInIsForHierarchy)	{}
+	FNiagaraHierarchyScriptRootViewModel(UHierarchyRoot* Root, TWeakObjectPtr<UNiagaraHierarchyScriptParametersViewModel> ViewModel, bool bInIsForHierarchy)
+		: FHierarchyRootViewModel(Root, ViewModel, bInIsForHierarchy)	{}
 
 protected:
 	// In the script, loose parameters are always added before categories (reverse of the default implementation)
diff --git a/Engine/Plugins/FX/Niagara/Source/NiagaraEditor/Public/ViewModels/HierarchyEditor/NiagaraSummaryViewViewModel.h b/Engine/Plugins/FX/Niagara/Source/NiagaraEditor/Public/ViewModels/HierarchyEditor/NiagaraSummaryViewViewModel.h
--- a/Engine/Plugins/FX/Niagara/Source/NiagaraEditor/Public/ViewModels/HierarchyEditor/NiagaraSummaryViewViewModel.h
+++ b/Engine/Plugins/FX/Niagara/Source/NiagaraEditor/Public/ViewModels/HierarchyEditor/NiagaraSummaryViewViewModel.h
@@ -6,14 +6,14 @@
 #include "NiagaraNodeFunctionCall.h"
 #include "NiagaraNodeAssignment.h"
 #include "NiagaraSimulationStageBase.h"
-#include "ViewModels/HierarchyEditor/NiagaraHierarchyViewModelBase.h"
+#include "HierarchyViewModelBase.h"
 #include "NiagaraSummaryViewViewModel.generated.h"
 
 class UNiagaraStackEventHandlerPropertiesItem;
 class FNiagaraEmitterViewModel;
 
 UCLASS()
-class UNiagaraHierarchySummaryDataRefreshContext : public UNiagaraHierarchyDataRefreshContext
+class UNiagaraHierarchySummaryDataRefreshContext : public UHierarchyDataRefreshContext
 {
 	GENERATED_BODY()
 
@@ -25,7 +25,7 @@ public:
 };
 
 UCLASS(MinimalAPI)
-class UNiagaraHierarchyModule : public UNiagaraHierarchyItem
+class UNiagaraHierarchyModule : public UHierarchyItem
 {
 	GENERATED_BODY()
 public:
@@ -33,7 +33,7 @@ public:
 };
 
 UCLASS(MinimalAPI)
-class UNiagaraHierarchyModuleInput : public UNiagaraHierarchyItem
+class UNiagaraHierarchyModuleInput : public UHierarchyItem
 {
 	GENERATED_BODY()
 
@@ -58,7 +58,7 @@ private:
 };
 
 UCLASS(MinimalAPI)
-class UNiagaraHierarchyAssignmentInput : public UNiagaraHierarchyItem
+class UNiagaraHierarchyAssignmentInput : public UHierarchyItem
 {
 	GENERATED_BODY()
 
@@ -76,7 +76,7 @@ private:
 };
 
 UCLASS(MinimalAPI)
-class UNiagaraHierarchyEmitterProperties : public UNiagaraHierarchyItem
+class UNiagaraHierarchyEmitterProperties : public UHierarchyItem
 {
 	GENERATED_BODY()
 public:
@@ -84,7 +84,7 @@ public:
 };
 
 UCLASS(MinimalAPI)
-class UNiagaraHierarchyRenderer : public UNiagaraHierarchyItem
+class UNiagaraHierarchyRenderer : public UHierarchyItem
 {
 	GENERATED_BODY()
 public:
@@ -92,7 +92,7 @@ public:
 };
 
 UCLASS(MinimalAPI)
-class UNiagaraHierarchyEventHandler : public UNiagaraHierarchyItem
+class UNiagaraHierarchyEventHandler : public UHierarchyItem
 {
 	GENERATED_BODY()
 public:
@@ -100,17 +100,17 @@ public:
 };
 
 UCLASS(MinimalAPI)
-class UNiagaraHierarchyEventHandlerProperties : public UNiagaraHierarchyItem
+class UNiagaraHierarchyEventHandlerProperties : public UHierarchyItem
 {
 	GENERATED_BODY()
 public:
 	NIAGARAEDITOR_API void Initialize(const FNiagaraEventScriptProperties& EventHandler);
 	
-	static NIAGARAEDITOR_API FNiagaraHierarchyIdentity MakeIdentity(const FNiagaraEventScriptProperties& EventHandler);
+	static NIAGARAEDITOR_API FHierarchyElementIdentity MakeIdentity(const FNiagaraEventScriptProperties& EventHandler);
 };
 
 UCLASS(MinimalAPI)
-class UNiagaraHierarchySimStage : public UNiagaraHierarchyItem
+class UNiagaraHierarchySimStage : public UHierarchyItem
 {
 	GENERATED_BODY()
 public:
@@ -118,17 +118,26 @@ public:
 };
 
 UCLASS(MinimalAPI)
-class UNiagaraHierarchySimStageProperties : public UNiagaraHierarchyItem
+class UNiagaraHierarchySimStageProperties : public UHierarchyItem
 {
 	GENERATED_BODY()
 public:
 	NIAGARAEDITOR_API void Initialize(const UNiagaraSimulationStageBase& SimStage);
 	
-	static NIAGARAEDITOR_API FNiagaraHierarchyIdentity MakeIdentity(const UNiagaraSimulationStageBase& SimStage);
+	static NIAGARAEDITOR_API FHierarchyElementIdentity MakeIdentity(const UNiagaraSimulationStageBase& SimStage);
 };
 
 UCLASS(MinimalAPI)
-class UNiagaraSummaryViewViewModel : public UNiagaraHierarchyViewModelBase
+class UNiagaraHierarchyObjectProperty : public UHierarchyItem
+{
+	GENERATED_BODY()
+public:
+	/** To know what object this ObjectProperty is referring to, a persistent guid that can be mapped back to an object is required. */
+	NIAGARAEDITOR_API void Initialize(FGuid ObjectGuid, FString PropertyName);
+};
+
+UCLASS(MinimalAPI)
+class UNiagaraSummaryViewViewModel : public UHierarchyViewModelBase
 {
 	GENERATED_BODY()
 public:
@@ -142,13 +151,13 @@ public:
 	
 	NIAGARAEDITOR_API TSharedRef<FNiagaraEmitterViewModel> GetEmitterViewModel() const;
 	
-	NIAGARAEDITOR_API virtual UNiagaraHierarchyRoot* GetHierarchyRoot() const override;
-	NIAGARAEDITOR_API virtual TSharedPtr<FNiagaraHierarchyItemViewModelBase> CreateViewModelForData(UNiagaraHierarchyItemBase* ItemBase, TSharedPtr<FNiagaraHierarchyItemViewModelBase> Parent) override;
+	NIAGARAEDITOR_API virtual UHierarchyRoot* GetHierarchyRoot() const override;
+	NIAGARAEDITOR_API virtual TSharedPtr<FHierarchyElementViewModel> CreateCustomViewModelForElement(UHierarchyElement* ItemBase, TSharedPtr<FHierarchyElementViewModel> Parent) override;
 	
-	NIAGARAEDITOR_API virtual void PrepareSourceItems(UNiagaraHierarchyRoot* SourceRoot, TSharedPtr<FNiagaraHierarchyRootViewModel>) override;
+	NIAGARAEDITOR_API virtual void PrepareSourceItems(UHierarchyRoot* SourceRoot, TSharedPtr<FHierarchyRootViewModel>) override;
 	NIAGARAEDITOR_API virtual void SetupCommands() override;
 	
-	NIAGARAEDITOR_API virtual TSharedRef<FNiagaraHierarchyDragDropOp> CreateDragDropOp(TSharedRef<FNiagaraHierarchyItemViewModelBase> Item) override;
+	NIAGARAEDITOR_API virtual TSharedRef<FHierarchyDragDropOp> CreateDragDropOp(TSharedRef<FHierarchyElementViewModel> Item) override;
 	
 	virtual bool SupportsDetailsPanel() override { return true; }
 	NIAGARAEDITOR_API virtual TArray<TTuple<UClass*, FOnGetDetailCustomizationInstance>> GetInstanceCustomizations() override;
@@ -171,10 +180,10 @@ protected:
 };
 
 /** The view model for both module nodes & dynamic input nodes */
-struct FNiagaraFunctionViewModel : public FNiagaraHierarchyItemViewModel
+struct FNiagaraFunctionViewModel : public FHierarchyItemViewModel
 {
-	FNiagaraFunctionViewModel(UNiagaraHierarchyModule* HierarchyModule, TSharedPtr<FNiagaraHierarchyItemViewModelBase> InParent, TWeakObjectPtr<UNiagaraSummaryViewViewModel> ViewModel, bool bInIsForHierarchy)
-	: FNiagaraHierarchyItemViewModel(HierarchyModule, InParent, ViewModel, bInIsForHierarchy)
+	FNiagaraFunctionViewModel(UNiagaraHierarchyModule* HierarchyModule, TSharedRef<FHierarchyElementViewModel> InParent, TWeakObjectPtr<UNiagaraSummaryViewViewModel> ViewModel)
+	: FHierarchyItemViewModel(HierarchyModule, InParent, ViewModel)
 	{ }
 
 	virtual ~FNiagaraFunctionViewModel() override;
@@ -183,7 +192,7 @@ struct FNiagaraFunctionViewModel : public FNiagaraHierarchyItemViewModel
 	
 	bool IsFromBaseEmitter() const;
 
-	void SetSection(UNiagaraHierarchySection& InSection) { Section = &InSection; }
+	void SetSection(UHierarchySection& InSection) { Section = &InSection; }
 
 	bool IsDynamicInput() const { return bIsDynamicInput; }
 private:	
@@ -196,11 +205,11 @@ private:
 	virtual FCanPerformActionResults IsEditableByUser() override;
 	virtual bool CanHaveChildren() const override { return bIsForHierarchy == false; }
 	virtual bool CanRenameInternal() override { return false; }
-	virtual FCanPerformActionResults CanDropOnInternal(TSharedPtr<FNiagaraHierarchyItemViewModelBase>, EItemDropZone ItemDropZone) override;
-	virtual const UNiagaraHierarchySection* GetSectionInternal() const override;
+	virtual FCanPerformActionResults CanDropOnInternal(TSharedPtr<FHierarchyElementViewModel>, EItemDropZone ItemDropZone) override;
+	virtual const UHierarchySection* GetSectionInternal() const override;
 	
 	virtual bool RepresentsExternalData() const override { return true; }
-	virtual bool DoesExternalDataStillExist(const UNiagaraHierarchyDataRefreshContext* Context) const override { ClearCache(); return GetFunctionCallNode().IsValid(); }
+	virtual bool DoesExternalDataStillExist(const UHierarchyDataRefreshContext* Context) const override { ClearCache(); return GetFunctionCallNode().IsValid(); }
 	
 	void OnScriptApplied(UNiagaraScript* NiagaraScript, FGuid Guid);
 
@@ -208,7 +217,7 @@ private:
 private:
 	FDelegateHandle OnScriptAppliedHandle;
 	mutable TOptional<bool> IsFromBaseEmitterCache;
-	TWeakObjectPtr<UNiagaraHierarchySection> Section;
+	TWeakObjectPtr<UHierarchySection> Section;
 	bool bIsDynamicInput = false;
 };
 
@@ -222,10 +231,10 @@ struct FInputData
 	UNiagaraNodeFunctionCall* FunctionCallNode;
 };
 
-struct FNiagaraModuleInputViewModel : public FNiagaraHierarchyItemViewModel
+struct FNiagaraModuleInputViewModel : public FHierarchyItemViewModel
 {	
-	FNiagaraModuleInputViewModel(UNiagaraHierarchyModuleInput* ModuleInput, TSharedPtr<FNiagaraHierarchyItemViewModelBase> InParent, TWeakObjectPtr<UNiagaraSummaryViewViewModel> ViewModel, bool bInIsForHierarchy)
-	: FNiagaraHierarchyItemViewModel(ModuleInput, InParent, ViewModel, bInIsForHierarchy) {}
+	FNiagaraModuleInputViewModel(UNiagaraHierarchyModuleInput* ModuleInput, TSharedRef<FHierarchyElementViewModel> InParent, TWeakObjectPtr<UNiagaraSummaryViewViewModel> ViewModel)
+	: FHierarchyItemViewModel(ModuleInput, InParent, ViewModel) {}
 
 public:
 	TOptional<FInputData> GetInputData() const;
@@ -246,29 +255,29 @@ protected:
 	virtual bool CanHaveChildren() const override;
 	virtual FCanPerformActionResults IsEditableByUser() override;
 	virtual bool RepresentsExternalData() const override { return true; }
-	virtual bool DoesExternalDataStillExist(const UNiagaraHierarchyDataRefreshContext* Context) const override { ClearCache(); return GetInputData().IsSet(); }
+	virtual bool DoesExternalDataStillExist(const UHierarchyDataRefreshContext* Context) const override { ClearCache(); return GetInputData().IsSet(); }
 	virtual void RefreshChildrenDataInternal() override;
 
-	virtual FCanPerformActionResults CanDropOnInternal(TSharedPtr<FNiagaraHierarchyItemViewModelBase>, EItemDropZone ItemDropZone) override;
-	virtual void OnDroppedOnInternal(TSharedPtr<FNiagaraHierarchyItemViewModelBase> DroppedItem, EItemDropZone ItemDropZone) override;
+	virtual FCanPerformActionResults CanDropOnInternal(TSharedPtr<FHierarchyElementViewModel>, EItemDropZone ItemDropZone) override;
+	virtual void OnDroppedOnInternal(TSharedPtr<FHierarchyElementViewModel> DroppedItem, EItemDropZone ItemDropZone) override;
 
-	virtual void PopulateDynamicContextMenuSection(FToolMenuSection& DynamicSection) override;
+	virtual void AppendDynamicContextMenuForSingleElement(UToolMenu* ToolMenu) override;
 private:
 	TOptional<FInputData> FindInputDataInternal() const;
 	void AddNativeChildrenInputs();
 	bool CanAddNativeChildrenInputs() const;
-	TArray<FNiagaraHierarchyIdentity> GetNativeChildInputIdentities() const;
+	TArray<FHierarchyElementIdentity> GetNativeChildInputIdentities() const;
 private:
 	mutable TOptional<FInputData> InputDataCache;
 	mutable TOptional<bool> IsFromBaseEmitterCache;
 };
 
-struct FNiagaraAssignmentInputViewModel : public FNiagaraHierarchyItemViewModel
+struct FNiagaraAssignmentInputViewModel : public FHierarchyItemViewModel
 {	
-	FNiagaraAssignmentInputViewModel(UNiagaraHierarchyAssignmentInput* ModuleInput, TSharedPtr<FNiagaraHierarchyItemViewModelBase> InParent, TWeakObjectPtr<UNiagaraSummaryViewViewModel> ViewModel, bool bInIsForHierarchy)
-	: FNiagaraHierarchyItemViewModel(ModuleInput, InParent, ViewModel, bInIsForHierarchy) {}
+	FNiagaraAssignmentInputViewModel(UNiagaraHierarchyAssignmentInput* ModuleInput, TSharedRef<FHierarchyElementViewModel> InParent, TWeakObjectPtr<UNiagaraSummaryViewViewModel> ViewModel)
+	: FHierarchyItemViewModel(ModuleInput, InParent, ViewModel) {}
 
-	virtual FCanPerformActionResults CanDropOnInternal(TSharedPtr<FNiagaraHierarchyItemViewModelBase>, EItemDropZone ItemDropZone) override;
+	virtual FCanPerformActionResults CanDropOnInternal(TSharedPtr<FHierarchyElementViewModel>, EItemDropZone ItemDropZone) override;
 	
 	TWeakObjectPtr<UNiagaraNodeAssignment> GetAssignmentNode() const;
 	TOptional<FNiagaraStackGraphUtilities::FMatchingFunctionInputData> GetInputData() const;
@@ -284,7 +293,7 @@ struct FNiagaraAssignmentInputViewModel : public FNiagaraHierarchyItemViewModel
 protected:
 	virtual FCanPerformActionResults IsEditableByUser() override;
 	virtual bool RepresentsExternalData() const override { return true; }
-	virtual bool DoesExternalDataStillExist(const UNiagaraHierarchyDataRefreshContext* Context) const override { ClearCache(); return GetInputData().IsSet(); }
+	virtual bool DoesExternalDataStillExist(const UHierarchyDataRefreshContext* Context) const override { ClearCache(); return GetInputData().IsSet(); }
 private:
 	TOptional<FNiagaraStackGraphUtilities::FMatchingFunctionInputData> FindInputDataInternal() const;
 private:
@@ -293,10 +302,10 @@ private:
 	mutable TOptional<bool> IsFromBaseEmitterCache;
 };
 
-struct FNiagaraHierarchySummaryCategoryViewModel : public FNiagaraHierarchyCategoryViewModel
+struct FNiagaraHierarchySummaryCategoryViewModel : public FHierarchyCategoryViewModel
 {
-	FNiagaraHierarchySummaryCategoryViewModel(UNiagaraHierarchyCategory* Category, TSharedPtr<FNiagaraHierarchyItemViewModelBase> InParent, TWeakObjectPtr<UNiagaraSummaryViewViewModel> ViewModel, bool bInIsForHierarchy)
-	: FNiagaraHierarchyCategoryViewModel(Category, InParent, ViewModel, bInIsForHierarchy) {}
+	FNiagaraHierarchySummaryCategoryViewModel(UHierarchyCategory* Category, TSharedRef<FHierarchyElementViewModel> InParent, TWeakObjectPtr<UNiagaraSummaryViewViewModel> ViewModel)
+	: FHierarchyCategoryViewModel(Category, InParent, ViewModel) {}
 	
 	bool IsFromBaseEmitter() const;
 	
@@ -306,34 +315,34 @@ protected:
 	mutable TOptional<bool> IsFromBaseEmitterCache;
 };
 
-struct FNiagaraHierarchyPropertyViewModel : public FNiagaraHierarchyItemViewModel
+struct FNiagaraHierarchyPropertyViewModel : public FHierarchyItemViewModel
 {
-	FNiagaraHierarchyPropertyViewModel(UNiagaraHierarchyObjectProperty* ObjectProperty, TSharedPtr<FNiagaraHierarchyItemViewModelBase> InParent, TWeakObjectPtr<UNiagaraSummaryViewViewModel> ViewModel, bool bInIsForHierarchy)
-		: FNiagaraHierarchyItemViewModel(ObjectProperty, InParent, ViewModel, bInIsForHierarchy) {}
+	FNiagaraHierarchyPropertyViewModel(UNiagaraHierarchyObjectProperty* ObjectProperty, TSharedRef<FHierarchyElementViewModel> InParent, TWeakObjectPtr<UNiagaraSummaryViewViewModel> ViewModel)
+		: FHierarchyItemViewModel(ObjectProperty, InParent, ViewModel) {}
 
 	virtual FString ToString() const override;
 
 	bool IsFromBaseEmitter() const;
 protected:
 	virtual bool RepresentsExternalData() const override { return true; }
-	virtual bool DoesExternalDataStillExist(const UNiagaraHierarchyDataRefreshContext* Context) const override;
+	virtual bool DoesExternalDataStillExist(const UHierarchyDataRefreshContext* Context) const override;
 	
 	virtual FCanPerformActionResults IsEditableByUser() override;
 	
 	mutable TOptional<bool> IsFromBaseEmitterCache;	
 };
 
-struct FNiagaraHierarchyRendererViewModel : public FNiagaraHierarchyItemViewModel
+struct FNiagaraHierarchyRendererViewModel : public FHierarchyItemViewModel
 {
-	FNiagaraHierarchyRendererViewModel(UNiagaraHierarchyRenderer* Renderer, TSharedPtr<FNiagaraHierarchyItemViewModelBase> InParent, TWeakObjectPtr<UNiagaraSummaryViewViewModel> ViewModel, bool bInIsForHierarchy)
-	: FNiagaraHierarchyItemViewModel(Renderer, InParent, ViewModel, bInIsForHierarchy) {}
+	FNiagaraHierarchyRendererViewModel(UNiagaraHierarchyRenderer* Renderer, TSharedRef<FHierarchyElementViewModel> InParent, TWeakObjectPtr<UNiagaraSummaryViewViewModel> ViewModel)
+	: FHierarchyItemViewModel(Renderer, InParent, ViewModel) {}
 
 	virtual FString ToString() const override;
 	virtual bool CanRenameInternal() override { return false; }
 
 	UNiagaraRendererProperties* GetRendererProperties() const;
 
-	void SetSection(UNiagaraHierarchySection& InSection) { Section = &InSection; }
+	void SetSection(UHierarchySection& InSection) { Section = &InSection; }
 
 	bool IsFromBaseEmitter() const;
 protected:
@@ -343,24 +352,24 @@ protected:
 	virtual bool CanHaveChildren() const override { return bIsForHierarchy == false; }
 	
 	virtual bool RepresentsExternalData() const override { return true; }
-	virtual bool DoesExternalDataStillExist(const UNiagaraHierarchyDataRefreshContext* Context) const override { return GetRendererProperties() != nullptr;}
+	virtual bool DoesExternalDataStillExist(const UHierarchyDataRefreshContext* Context) const override { return GetRendererProperties() != nullptr;}
 
-	virtual const UNiagaraHierarchySection* GetSectionInternal() const override;
+	virtual const UHierarchySection* GetSectionInternal() const override;
 	
-	TWeakObjectPtr<UNiagaraHierarchySection> Section;
+	TWeakObjectPtr<UHierarchySection> Section;
 	mutable TOptional<bool> IsFromBaseEmitterCache;
 };
 
 /** Emitter Properties currently don't list their individual properties since it's a mix of data of FVersionedNiagaraEmitterData & actual properties on the emitter which requires customization. */
-struct FNiagaraHierarchyEmitterPropertiesViewModel : public FNiagaraHierarchyItemViewModel
+struct FNiagaraHierarchyEmitterPropertiesViewModel : public FHierarchyItemViewModel
 {
-	FNiagaraHierarchyEmitterPropertiesViewModel(UNiagaraHierarchyEmitterProperties* EmitterProperties, TSharedPtr<FNiagaraHierarchyItemViewModelBase> InParent, TWeakObjectPtr<UNiagaraSummaryViewViewModel> ViewModel, bool bInIsForHierarchy)
-	: FNiagaraHierarchyItemViewModel(EmitterProperties, InParent, ViewModel, bInIsForHierarchy) {}
+	FNiagaraHierarchyEmitterPropertiesViewModel(UNiagaraHierarchyEmitterProperties* EmitterProperties, TSharedRef<FHierarchyElementViewModel> InParent, TWeakObjectPtr<UNiagaraSummaryViewViewModel> ViewModel)
+	: FHierarchyItemViewModel(EmitterProperties, InParent, ViewModel) {}
 
 	virtual FString ToString() const override;
 	virtual bool CanRenameInternal() override { return false; }
 
-	void SetSection(UNiagaraHierarchySection& InSection) { Section = &InSection; }
+	void SetSection(UHierarchySection& InSection) { Section = &InSection; }
 
 	bool IsFromBaseEmitter() const;
 protected:	
@@ -368,25 +377,25 @@ protected:
 	virtual bool CanHaveChildren() const override { return bIsForHierarchy == false; }
 	
 	virtual bool RepresentsExternalData() const override { return true; }
-	virtual bool DoesExternalDataStillExist(const UNiagaraHierarchyDataRefreshContext* Context) const override { return true; }
+	virtual bool DoesExternalDataStillExist(const UHierarchyDataRefreshContext* Context) const override { return true; }
 
-	virtual const UNiagaraHierarchySection* GetSectionInternal() const override;
+	virtual const UHierarchySection* GetSectionInternal() const override;
 	
-	TWeakObjectPtr<UNiagaraHierarchySection> Section;
+	TWeakObjectPtr<UHierarchySection> Section;
 	mutable TOptional<bool> IsFromBaseEmitterCache;
 };
 
-struct FNiagaraHierarchyEventHandlerViewModel : public FNiagaraHierarchyItemViewModel
+struct FNiagaraHierarchyEventHandlerViewModel : public FHierarchyItemViewModel
 {
-	FNiagaraHierarchyEventHandlerViewModel(UNiagaraHierarchyEventHandler* EventHandler, TSharedPtr<FNiagaraHierarchyItemViewModelBase> InParent, TWeakObjectPtr<UNiagaraSummaryViewViewModel> ViewModel, bool bInIsForHierarchy)
-	: FNiagaraHierarchyItemViewModel(EventHandler, InParent, ViewModel, bInIsForHierarchy) {}
+	FNiagaraHierarchyEventHandlerViewModel(UNiagaraHierarchyEventHandler* EventHandler, TSharedRef<FHierarchyElementViewModel> InParent, TWeakObjectPtr<UNiagaraSummaryViewViewModel> ViewModel)
+	: FHierarchyItemViewModel(EventHandler, InParent, ViewModel) {}
 
 	virtual FString ToString() const override;
 	virtual bool CanRenameInternal() override { return false; }
 
 	FNiagaraEventScriptProperties* GetEventScriptProperties() const;
 
-	void SetSection(UNiagaraHierarchySection& InSection) { Section = &InSection; }
+	void SetSection(UHierarchySection& InSection) { Section = &InSection; }
 
 	bool IsFromBaseEmitter() const;
 protected:
@@ -396,18 +405,18 @@ protected:
 	virtual bool CanHaveChildren() const override { return bIsForHierarchy == false; }
 	
 	virtual bool RepresentsExternalData() const override { return true; }
-	virtual bool DoesExternalDataStillExist(const UNiagaraHierarchyDataRefreshContext* Context) const override { return GetEventScriptProperties() != nullptr;}
+	virtual bool DoesExternalDataStillExist(const UHierarchyDataRefreshContext* Context) const override { return GetEventScriptProperties() != nullptr;}
 
-	virtual const UNiagaraHierarchySection* GetSectionInternal() const override;
+	virtual const UHierarchySection* GetSectionInternal() const override;
 	
-	TWeakObjectPtr<UNiagaraHierarchySection> Section;
+	TWeakObjectPtr<UHierarchySection> Section;
 	mutable TOptional<bool> IsFromBaseEmitterCache;
 };
 
-struct FNiagaraHierarchyEventHandlerPropertiesViewModel : public FNiagaraHierarchyItemViewModel
+struct FNiagaraHierarchyEventHandlerPropertiesViewModel : public FHierarchyItemViewModel
 {
-	FNiagaraHierarchyEventHandlerPropertiesViewModel(UNiagaraHierarchyEventHandlerProperties* EventHandlerProperties, TSharedPtr<FNiagaraHierarchyItemViewModelBase> InParent, TWeakObjectPtr<UNiagaraSummaryViewViewModel> ViewModel, bool bInIsForHierarchy)
-	: FNiagaraHierarchyItemViewModel(EventHandlerProperties, InParent, ViewModel, bInIsForHierarchy) {}
+	FNiagaraHierarchyEventHandlerPropertiesViewModel(UNiagaraHierarchyEventHandlerProperties* EventHandlerProperties, TSharedRef<FHierarchyElementViewModel> InParent, TWeakObjectPtr<UNiagaraSummaryViewViewModel> ViewModel)
+	: FHierarchyItemViewModel(EventHandlerProperties, InParent, ViewModel) {}
 
 	virtual FString ToString() const override;
 	virtual bool CanRenameInternal() override { return false; }
@@ -422,22 +431,22 @@ protected:
 	virtual bool CanHaveChildren() const override { return bIsForHierarchy == false; }
 	
 	virtual bool RepresentsExternalData() const override { return true; }
-	virtual bool DoesExternalDataStillExist(const UNiagaraHierarchyDataRefreshContext* Context) const override { return GetEventScriptProperties() != nullptr;}
+	virtual bool DoesExternalDataStillExist(const UHierarchyDataRefreshContext* Context) const override { return GetEventScriptProperties() != nullptr;}
 
 	mutable TOptional<bool> IsFromBaseEmitterCache;
 };
 
-struct FNiagaraHierarchySimStageViewModel : public FNiagaraHierarchyItemViewModel
+struct FNiagaraHierarchySimStageViewModel : public FHierarchyItemViewModel
 {
-	FNiagaraHierarchySimStageViewModel(UNiagaraHierarchySimStage* SimStage, TSharedPtr<FNiagaraHierarchyItemViewModelBase> InParent, TWeakObjectPtr<UNiagaraSummaryViewViewModel> ViewModel, bool bInIsForHierarchy)
-	: FNiagaraHierarchyItemViewModel(SimStage, InParent, ViewModel, bInIsForHierarchy) {}
+	FNiagaraHierarchySimStageViewModel(UNiagaraHierarchySimStage* SimStage, TSharedRef<FHierarchyElementViewModel> InParent, TWeakObjectPtr<UNiagaraSummaryViewViewModel> ViewModel)
+	: FHierarchyItemViewModel(SimStage, InParent, ViewModel) {}
 
 	virtual FString ToString() const override;
 	virtual bool CanRenameInternal() override { return false; }
 
 	UNiagaraSimulationStageBase* GetSimStage() const;
 
-	void SetSection(UNiagaraHierarchySection& InSection) { Section = &InSection; }
+	void SetSection(UHierarchySection& InSection) { Section = &InSection; }
 
 	bool IsFromBaseEmitter() const;
 protected:
@@ -447,18 +456,18 @@ protected:
 	virtual bool CanHaveChildren() const override { return bIsForHierarchy == false; }
 	
 	virtual bool RepresentsExternalData() const override { return true; }
-	virtual bool DoesExternalDataStillExist(const UNiagaraHierarchyDataRefreshContext* Context) const override { return GetSimStage() != nullptr;}
+	virtual bool DoesExternalDataStillExist(const UHierarchyDataRefreshContext* Context) const override { return GetSimStage() != nullptr;}
 
-	virtual const UNiagaraHierarchySection* GetSectionInternal() const override;
+	virtual const UHierarchySection* GetSectionInternal() const override;
 	
-	TWeakObjectPtr<UNiagaraHierarchySection> Section;
+	TWeakObjectPtr<UHierarchySection> Section;
 	mutable TOptional<bool> IsFromBaseEmitterCache;
 };
 
-struct FNiagaraHierarchySimStagePropertiesViewModel : public FNiagaraHierarchyItemViewModel
+struct FNiagaraHierarchySimStagePropertiesViewModel : public FHierarchyItemViewModel
 {
-	FNiagaraHierarchySimStagePropertiesViewModel(UNiagaraHierarchySimStageProperties* SimStage, TSharedPtr<FNiagaraHierarchyItemViewModelBase> InParent, TWeakObjectPtr<UNiagaraSummaryViewViewModel> ViewModel, bool bInIsForHierarchy)
-	: FNiagaraHierarchyItemViewModel(SimStage, InParent, ViewModel, bInIsForHierarchy) {}
+	FNiagaraHierarchySimStagePropertiesViewModel(UNiagaraHierarchySimStageProperties* SimStage, TSharedRef<FHierarchyElementViewModel> InParent, TWeakObjectPtr<UNiagaraSummaryViewViewModel> ViewModel)
+	: FHierarchyItemViewModel(SimStage, InParent, ViewModel) {}
 
 	virtual FString ToString() const override;
 	virtual bool CanRenameInternal() override { return false; }
@@ -473,17 +482,17 @@ protected:
 	virtual bool CanHaveChildren() const override { return bIsForHierarchy == false; }
 	
 	virtual bool RepresentsExternalData() const override { return true; }
-	virtual bool DoesExternalDataStillExist(const UNiagaraHierarchyDataRefreshContext* Context) const override { return GetSimStage() != nullptr;}
+	virtual bool DoesExternalDataStillExist(const UHierarchyDataRefreshContext* Context) const override { return GetSimStage() != nullptr;}
 
 	mutable TOptional<bool> IsFromBaseEmitterCache;
 };
 
-class FNiagaraHierarchyInputParameterHierarchyDragDropOp : public FNiagaraHierarchyDragDropOp
+class FNiagaraHierarchyInputParameterHierarchyDragDropOp : public FHierarchyDragDropOp
 {
 public:
-	DRAG_DROP_OPERATOR_TYPE(FNiagaraInputParameterHierarchyDragDropOp, FNiagaraHierarchyDragDropOp)
+	DRAG_DROP_OPERATOR_TYPE(FNiagaraInputParameterHierarchyDragDropOp, FHierarchyDragDropOp)
 
-	FNiagaraHierarchyInputParameterHierarchyDragDropOp(TSharedPtr<FNiagaraModuleInputViewModel> InputViewModel) : FNiagaraHierarchyDragDropOp(InputViewModel) {}
+	FNiagaraHierarchyInputParameterHierarchyDragDropOp(TSharedPtr<FNiagaraModuleInputViewModel> InputViewModel) : FHierarchyDragDropOp(InputViewModel) {}
 	
 	virtual TSharedRef<SWidget> CreateCustomDecorator() const override;
 };
diff --git a/Engine/Plugins/FX/Niagara/Source/NiagaraEditor/Public/ViewModels/HierarchyEditor/NiagaraUserParametersHierarchyViewModel.h b/Engine/Plugins/FX/Niagara/Source/NiagaraEditor/Public/ViewModels/HierarchyEditor/NiagaraUserParametersHierarchyViewModel.h
--- a/Engine/Plugins/FX/Niagara/Source/NiagaraEditor/Public/ViewModels/HierarchyEditor/NiagaraUserParametersHierarchyViewModel.h
+++ b/Engine/Plugins/FX/Niagara/Source/NiagaraEditor/Public/ViewModels/HierarchyEditor/NiagaraUserParametersHierarchyViewModel.h
@@ -2,14 +2,14 @@
 
 #pragma once
 
-#include "ViewModels/HierarchyEditor/NiagaraHierarchyViewModelBase.h"
+#include "HierarchyViewModelBase.h"
 #include "NiagaraScriptVariable.h"
 #include "NiagaraUserParametersHierarchyViewModel.generated.h"
 
 class FNiagaraSystemViewModel;
 
 UCLASS()
-class UNiagaraHierarchyUserParameterRefreshContext : public UNiagaraHierarchyDataRefreshContext
+class UNiagaraHierarchyUserParameterRefreshContext : public UHierarchyDataRefreshContext
 {
 	GENERATED_BODY()
 
@@ -22,7 +22,7 @@ private:
 };
 
 UCLASS()
-class UNiagaraHierarchyUserParameter : public UNiagaraHierarchyItem
+class UNiagaraHierarchyUserParameter : public UHierarchyItem
 {
 	GENERATED_BODY()
 
@@ -42,10 +42,10 @@ private:
 	TObjectPtr<UNiagaraScriptVariable> UserParameterScriptVariable;
 };
 
-struct FNiagaraHierarchyUserParameterViewModel : public FNiagaraHierarchyItemViewModel
+struct FNiagaraHierarchyUserParameterViewModel : public FHierarchyItemViewModel
 {
-	FNiagaraHierarchyUserParameterViewModel(UNiagaraHierarchyUserParameter* InItem, TSharedPtr<FNiagaraHierarchyItemViewModelBase> InParent, TWeakObjectPtr<UNiagaraHierarchyViewModelBase> InHierarchyViewModel, bool bInIsForHierarchy)
-	: FNiagaraHierarchyItemViewModel(InItem, InParent, InHierarchyViewModel, bInIsForHierarchy) {}
+	FNiagaraHierarchyUserParameterViewModel(UNiagaraHierarchyUserParameter* InItem, TSharedRef<FHierarchyElementViewModel> InParent, TWeakObjectPtr<UHierarchyViewModelBase> InHierarchyViewModel)
+	: FHierarchyItemViewModel(InItem, InParent, InHierarchyViewModel) {}
 	
 	/** For editing in the details panel we want to handle the script variable that is represented by the hierarchy item, not the hierarchy item itself. */
 	virtual UObject* GetDataForEditing() override;
@@ -53,11 +53,11 @@ struct FNiagaraHierarchyUserParameterViewModel : public FNiagaraHierarchyItemVie
 	virtual bool AllowEditingInDetailsPanel() const override { return true; }
 
 	virtual bool RepresentsExternalData() const override { return true; }
-	virtual bool DoesExternalDataStillExist(const UNiagaraHierarchyDataRefreshContext* Context) const override;
+	virtual bool DoesExternalDataStillExist(const UHierarchyDataRefreshContext* Context) const override;
 };
 
 UCLASS()
-class UNiagaraUserParametersHierarchyViewModel : public UNiagaraHierarchyViewModelBase
+class UNiagaraUserParametersHierarchyViewModel : public UHierarchyViewModelBase
 {
 	GENERATED_BODY()
 public:
@@ -70,13 +70,13 @@ public:
 
 	void Initialize(TSharedRef<FNiagaraSystemViewModel> InSystemViewModel);
 	
-	virtual UNiagaraHierarchyRoot* GetHierarchyRoot() const override;
-	virtual TSharedPtr<FNiagaraHierarchyItemViewModelBase> CreateViewModelForData(UNiagaraHierarchyItemBase* ItemBase, TSharedPtr<FNiagaraHierarchyItemViewModelBase> Parent) override;
+	virtual UHierarchyRoot* GetHierarchyRoot() const override;
+	virtual TSharedPtr<FHierarchyElementViewModel> CreateCustomViewModelForElement(UHierarchyElement* ItemBase, TSharedPtr<FHierarchyElementViewModel> Parent) override;
 	
-	virtual void PrepareSourceItems(UNiagaraHierarchyRoot* SourceRoot, TSharedPtr<FNiagaraHierarchyRootViewModel>) override;
+	virtual void PrepareSourceItems(UHierarchyRoot* SourceRoot, TSharedPtr<FHierarchyRootViewModel>) override;
 	virtual void SetupCommands() override;
 	
-	virtual TSharedRef<FNiagaraHierarchyDragDropOp> CreateDragDropOp(TSharedRef<FNiagaraHierarchyItemViewModelBase> Item) override;
+	virtual TSharedRef<FHierarchyDragDropOp> CreateDragDropOp(TSharedRef<FHierarchyElementViewModel> Item) override;
 	
 	virtual bool SupportsDetailsPanel() override { return true; }
 	virtual TArray<TTuple<UClass*, FOnGetDetailCustomizationInstance>> GetInstanceCustomizations() override;
@@ -86,16 +86,16 @@ private:
 	TWeakPtr<FNiagaraSystemViewModel> SystemViewModelWeak;
 };
 
-class FNiagaraUserParameterHierarchyDragDropOp : public FNiagaraHierarchyDragDropOp
+class FNiagaraUserParameterHierarchyDragDropOp : public FHierarchyDragDropOp
 {
 public:
-	DRAG_DROP_OPERATOR_TYPE(FNiagaraUserParameterHierarchyDragDropOp, FNiagaraHierarchyDragDropOp)
+	DRAG_DROP_OPERATOR_TYPE(FNiagaraUserParameterHierarchyDragDropOp, FHierarchyDragDropOp)
 
-	FNiagaraUserParameterHierarchyDragDropOp(TSharedPtr<FNiagaraHierarchyItemViewModelBase> UserParameterItem) : FNiagaraHierarchyDragDropOp(UserParameterItem) {}
+	FNiagaraUserParameterHierarchyDragDropOp(TSharedPtr<FHierarchyElementViewModel> UserParameterItem) : FHierarchyDragDropOp(UserParameterItem) {}
 
 	FNiagaraVariable GetUserParameter() const
 	{
-		const UNiagaraHierarchyUserParameter* HierarchyUserParameter = CastChecked<UNiagaraHierarchyUserParameter>(DraggedItem.Pin()->GetData());
+		const UNiagaraHierarchyUserParameter* HierarchyUserParameter = CastChecked<UNiagaraHierarchyUserParameter>(DraggedElement.Pin()->GetData());
 		return HierarchyUserParameter->GetUserParameter();
 	}
 	
diff --git a/Engine/Plugins/FX/Niagara/Source/NiagaraEditor/Public/ViewModels/Stack/NiagaraStackEntry.h b/Engine/Plugins/FX/Niagara/Source/NiagaraEditor/Public/ViewModels/Stack/NiagaraStackEntry.h
--- a/Engine/Plugins/FX/Niagara/Source/NiagaraEditor/Public/ViewModels/Stack/NiagaraStackEntry.h
+++ b/Engine/Plugins/FX/Niagara/Source/NiagaraEditor/Public/ViewModels/Stack/NiagaraStackEntry.h
@@ -577,7 +577,7 @@ public:
 	NIAGARAEDITOR_API void PasteStackNote();
 	
 	virtual bool SupportsSummaryView() const { return false; }
-	NIAGARAEDITOR_API virtual struct FNiagaraHierarchyIdentity DetermineSummaryIdentity() const;
+	NIAGARAEDITOR_API virtual struct FHierarchyElementIdentity DetermineSummaryIdentity() const;
 	NIAGARAEDITOR_API bool IsInSummaryView() const;
 	NIAGARAEDITOR_API bool IsAnyParentInSummaryView() const;
 	NIAGARAEDITOR_API bool IsAnyChildInSummaryView(bool bRecursive = false) const;
diff --git a/Engine/Plugins/FX/Niagara/Source/NiagaraEditor/Public/ViewModels/Stack/NiagaraStackEventScriptItemGroup.h b/Engine/Plugins/FX/Niagara/Source/NiagaraEditor/Public/ViewModels/Stack/NiagaraStackEventScriptItemGroup.h
--- a/Engine/Plugins/FX/Niagara/Source/NiagaraEditor/Public/ViewModels/Stack/NiagaraStackEventScriptItemGroup.h
+++ b/Engine/Plugins/FX/Niagara/Source/NiagaraEditor/Public/ViewModels/Stack/NiagaraStackEventScriptItemGroup.h
@@ -53,7 +53,7 @@ protected:
 	NIAGARAEDITOR_API virtual void RefreshChildrenInternal(const TArray<UNiagaraStackEntry*>& CurrentChildren, TArray<UNiagaraStackEntry*>& NewChildren, TArray<FStackIssue>& NewIssues) override;
 
 	virtual bool SupportsSummaryView() const override { return true; }
-	NIAGARAEDITOR_API virtual FNiagaraHierarchyIdentity DetermineSummaryIdentity() const override;
+	NIAGARAEDITOR_API virtual FHierarchyElementIdentity DetermineSummaryIdentity() const override;
 private:
 	NIAGARAEDITOR_API void EventHandlerPropertiesChanged();
 
@@ -113,7 +113,7 @@ protected:
 	NIAGARAEDITOR_API virtual void RefreshChildrenInternal(const TArray<UNiagaraStackEntry*>& CurrentChildren, TArray<UNiagaraStackEntry*>& NewChildren, TArray<FStackIssue>& NewIssues) override;
 
 	virtual bool SupportsSummaryView() const override { return true; }
-	NIAGARAEDITOR_API virtual FNiagaraHierarchyIdentity DetermineSummaryIdentity() const override;
+	NIAGARAEDITOR_API virtual FHierarchyElementIdentity DetermineSummaryIdentity() const override;
 private:
 	NIAGARAEDITOR_API bool HasBaseEventHandler() const;
 	
diff --git a/Engine/Plugins/FX/Niagara/Source/NiagaraEditor/Public/ViewModels/Stack/NiagaraStackFunctionInput.h b/Engine/Plugins/FX/Niagara/Source/NiagaraEditor/Public/ViewModels/Stack/NiagaraStackFunctionInput.h
--- a/Engine/Plugins/FX/Niagara/Source/NiagaraEditor/Public/ViewModels/Stack/NiagaraStackFunctionInput.h
+++ b/Engine/Plugins/FX/Niagara/Source/NiagaraEditor/Public/ViewModels/Stack/NiagaraStackFunctionInput.h
@@ -124,7 +124,7 @@ public:
 	NIAGARAEDITOR_API virtual void Paste(const UNiagaraClipboardContent* ClipboardContent, FText& OutPasteWarning) override;
 	NIAGARAEDITOR_API virtual bool HasOverridenContent() const override;
 	NIAGARAEDITOR_API virtual bool SupportsSummaryView() const override;
-	NIAGARAEDITOR_API virtual FNiagaraHierarchyIdentity DetermineSummaryIdentity() const override;
+	NIAGARAEDITOR_API virtual FHierarchyElementIdentity DetermineSummaryIdentity() const override;
 	virtual bool GetCanExpand() const override;
 	virtual bool KeepExpanderIndentation() const override { return true; }
 
diff --git a/Engine/Plugins/FX/Niagara/Source/NiagaraEditor/Public/ViewModels/Stack/NiagaraStackGraphUtilities.h b/Engine/Plugins/FX/Niagara/Source/NiagaraEditor/Public/ViewModels/Stack/NiagaraStackGraphUtilities.h
--- a/Engine/Plugins/FX/Niagara/Source/NiagaraEditor/Public/ViewModels/Stack/NiagaraStackGraphUtilities.h
+++ b/Engine/Plugins/FX/Niagara/Source/NiagaraEditor/Public/ViewModels/Stack/NiagaraStackGraphUtilities.h
@@ -10,7 +10,6 @@
 #include "ViewModels/Stack/NiagaraStackEntry.h"
 #include "AssetRegistry/AssetData.h"
 
-struct FNiagaraHierarchyIdentity;
 class UEdGraph;
 class UEdGraphPin;
 class UNiagaraGraph;
diff --git a/Engine/Plugins/FX/Niagara/Source/NiagaraEditor/Public/ViewModels/Stack/NiagaraStackInputCategory.h b/Engine/Plugins/FX/Niagara/Source/NiagaraEditor/Public/ViewModels/Stack/NiagaraStackInputCategory.h
--- a/Engine/Plugins/FX/Niagara/Source/NiagaraEditor/Public/ViewModels/Stack/NiagaraStackInputCategory.h
+++ b/Engine/Plugins/FX/Niagara/Source/NiagaraEditor/Public/ViewModels/Stack/NiagaraStackInputCategory.h
@@ -6,7 +6,7 @@
 #include "NiagaraStackScriptHierarchyRoot.h"
 #include "ViewModels/Stack/NiagaraStackItem.h"
 #include "NiagaraTypes.h"
-#include "ViewModels/HierarchyEditor/NiagaraHierarchyViewModelBase.h"
+#include "HierarchyViewModelBase.h"
 #include "ViewModels/HierarchyEditor/NiagaraSummaryViewViewModel.h"
 #include "NiagaraStackInputCategory.generated.h"
 
@@ -56,7 +56,7 @@ class UNiagaraStackScriptHierarchyCategory : public UNiagaraStackCategory
 	GENERATED_BODY()
 
 public:
-	NIAGARAEDITOR_API void Initialize(FRequiredEntryData InRequiredEntryData, const UNiagaraHierarchyCategory& InHierarchyCategory, FString InOwningStackItemEditorDataKey, FString InStackEditorDataKey);
+	NIAGARAEDITOR_API void Initialize(FRequiredEntryData InRequiredEntryData, const UHierarchyCategory& InHierarchyCategory, FString InOwningStackItemEditorDataKey, FString InStackEditorDataKey);
 
 	/** UNiagaraStackEntry */
 	virtual bool SupportsCopy() const override { return true; }
@@ -77,13 +77,13 @@ public:
 	
 	void SetScriptInstanceData(const FNiagaraScriptInstanceData& InScriptInstanceData) { ScriptInstanceData = InScriptInstanceData; }
 	
-	const UNiagaraHierarchyCategory* GetHierarchyCategory() const { return HierarchyCategory.Get(); }
+	const UHierarchyCategory* GetHierarchyCategory() const { return HierarchyCategory.Get(); }
 protected:
 	virtual void RefreshChildrenInternal(const TArray<UNiagaraStackEntry*>& CurrentChildren, TArray<UNiagaraStackEntry*>& NewChildren, TArray<FStackIssue>& NewIssues) override;
 
-	virtual bool IsTopLevelCategory() const override { return HierarchyCategory->GetOuter()->IsA<UNiagaraHierarchyRoot>(); }
+	virtual bool IsTopLevelCategory() const override { return HierarchyCategory->GetOuter()->IsA<UHierarchyRoot>(); }
 protected:
-	TWeakObjectPtr<const UNiagaraHierarchyCategory> HierarchyCategory;
+	TWeakObjectPtr<const UHierarchyCategory> HierarchyCategory;
 
 	TWeakObjectPtr<UNiagaraNodeFunctionCall> OwningModuleNode;
 	TWeakObjectPtr<UNiagaraNodeFunctionCall> OwningFunctionCallNode;
@@ -101,14 +101,14 @@ public:
 
 	NIAGARAEDITOR_API void Initialize(
 		FRequiredEntryData InRequiredEntryData,
-		TSharedPtr<FNiagaraHierarchyCategoryViewModel> InCategory,
+		TSharedPtr<FHierarchyCategoryViewModel> InCategory,
 		FString InOwnerStackItemEditorDataKey);
 
 	//~ UNiagaraStackEntry interface
 	NIAGARAEDITOR_API virtual FText GetDisplayName() const override;
 	virtual bool GetIsEnabled() const override { return true; }
 	
-	TWeakPtr<FNiagaraHierarchyCategoryViewModel> GetHierarchyCategory() const { return CategoryViewModelWeakPtr; }
+	TWeakPtr<FHierarchyCategoryViewModel> GetHierarchyCategory() const { return CategoryViewModelWeakPtr; }
 
 protected:
 	NIAGARAEDITOR_API virtual void RefreshChildrenInternal(const TArray<UNiagaraStackEntry*>& CurrentChildren, TArray<UNiagaraStackEntry*>& NewChildren, TArray<FStackIssue>& NewIssues) override;
@@ -117,5 +117,5 @@ protected:
 	NIAGARAEDITOR_API virtual int32 GetChildIndentLevel() const override;
 
 private:
-	TWeakPtr<FNiagaraHierarchyCategoryViewModel> CategoryViewModelWeakPtr;
+	TWeakPtr<FHierarchyCategoryViewModel> CategoryViewModelWeakPtr;
 };
diff --git a/Engine/Plugins/FX/Niagara/Source/NiagaraEditor/Public/ViewModels/Stack/NiagaraStackModuleItem.h b/Engine/Plugins/FX/Niagara/Source/NiagaraEditor/Public/ViewModels/Stack/NiagaraStackModuleItem.h
--- a/Engine/Plugins/FX/Niagara/Source/NiagaraEditor/Public/ViewModels/Stack/NiagaraStackModuleItem.h
+++ b/Engine/Plugins/FX/Niagara/Source/NiagaraEditor/Public/ViewModels/Stack/NiagaraStackModuleItem.h
@@ -100,7 +100,7 @@ public:
 	NIAGARAEDITOR_API virtual FText GetInheritanceMessage() const override;
 
 	virtual bool SupportsSummaryView() const override { return true; }
-	NIAGARAEDITOR_API virtual FNiagaraHierarchyIdentity DetermineSummaryIdentity() const override;
+	NIAGARAEDITOR_API virtual FHierarchyElementIdentity DetermineSummaryIdentity() const override;
 	
 	NIAGARAEDITOR_API bool IsScratchModule() const;
 
diff --git a/Engine/Plugins/FX/Niagara/Source/NiagaraEditor/Public/ViewModels/Stack/NiagaraStackScriptHierarchyRoot.h b/Engine/Plugins/FX/Niagara/Source/NiagaraEditor/Public/ViewModels/Stack/NiagaraStackScriptHierarchyRoot.h
--- a/Engine/Plugins/FX/Niagara/Source/NiagaraEditor/Public/ViewModels/Stack/NiagaraStackScriptHierarchyRoot.h
+++ b/Engine/Plugins/FX/Niagara/Source/NiagaraEditor/Public/ViewModels/Stack/NiagaraStackScriptHierarchyRoot.h
@@ -58,11 +58,11 @@ public:
 
 	UNiagaraNodeFunctionCall& GetOwningModuleNode() const { return *ModuleNode.Get(); }
 	UNiagaraNodeFunctionCall& GetOwningFunctionCallNode() const { return *OwningFunctionCallNode.Get(); }
-	UNiagaraHierarchyRoot* GetScriptParameterHierarchyRoot() const;
+	UHierarchyRoot* GetScriptParameterHierarchyRoot() const;
 	
-	NIAGARAEDITOR_API const TArray<UNiagaraHierarchySection*>& GetSections() const;
-	NIAGARAEDITOR_API const UNiagaraHierarchySection* GetActiveSection() const;
-	NIAGARAEDITOR_API void SetActiveSection(const UNiagaraHierarchySection* InActiveSection);
+	NIAGARAEDITOR_API const TArray<UHierarchySection*>& GetSections() const;
+	NIAGARAEDITOR_API const UHierarchySection* GetActiveSection() const;
+	NIAGARAEDITOR_API void SetActiveSection(const UHierarchySection* InActiveSection);
 
 	void SetShouldDisplayLabel(bool bInShouldDisplayLabel);
 	NIAGARAEDITOR_API bool GetShouldDisplayLabel() const;
@@ -80,7 +80,7 @@ protected:
 	 * a lot of information is static (categories, sorting etc.), whereas visibility is determined by the active configuration */
 	void RefreshInstanceData();
 	
-	const UNiagaraHierarchySection* FindSectionByName(FName SectionName);
+	const UHierarchySection* FindSectionByName(FName SectionName);
 	NIAGARAEDITOR_API virtual bool FilterByActiveSection(const UNiagaraStackEntry& Child) const;
 	NIAGARAEDITOR_API bool FilterForVisibleCondition(const UNiagaraStackEntry& NiagaraStackEntry) const;
 	bool FilterForIsInlineEditConditionToggle(const UNiagaraStackEntry& NiagaraStackEntry) const;
@@ -97,7 +97,7 @@ private:
 	
 	void OnScriptApplied(UNiagaraScript* NiagaraScript, FGuid Guid);
 protected:
-	TWeakObjectPtr<const UNiagaraHierarchySection> ActiveSection;
+	TWeakObjectPtr<const UHierarchySection> ActiveSection;
 
 	TWeakObjectPtr<UNiagaraNodeFunctionCall> ModuleNode;
 	TWeakObjectPtr<UNiagaraNodeFunctionCall> OwningFunctionCallNode;
diff --git a/Engine/Plugins/FX/Niagara/Source/NiagaraEditor/Public/ViewModels/Stack/NiagaraStackSimulationStageGroup.h b/Engine/Plugins/FX/Niagara/Source/NiagaraEditor/Public/ViewModels/Stack/NiagaraStackSimulationStageGroup.h
--- a/Engine/Plugins/FX/Niagara/Source/NiagaraEditor/Public/ViewModels/Stack/NiagaraStackSimulationStageGroup.h
+++ b/Engine/Plugins/FX/Niagara/Source/NiagaraEditor/Public/ViewModels/Stack/NiagaraStackSimulationStageGroup.h
@@ -36,7 +36,7 @@ protected:
 	NIAGARAEDITOR_API virtual void RefreshChildrenInternal(const TArray<UNiagaraStackEntry*>& CurrentChildren, TArray<UNiagaraStackEntry*>& NewChildren, TArray<FStackIssue>& NewIssues) override;
 
 	virtual bool SupportsSummaryView() const override { return true; }
-	NIAGARAEDITOR_API virtual FNiagaraHierarchyIdentity DetermineSummaryIdentity() const override;
+	NIAGARAEDITOR_API virtual FHierarchyElementIdentity DetermineSummaryIdentity() const override;
 private:
 	NIAGARAEDITOR_API void SimulationStagePropertiesChanged();
 
@@ -102,7 +102,7 @@ protected:
 	NIAGARAEDITOR_API virtual TOptional<FDropRequestResponse> DropInternal(const FDropRequest& DropRequest) override;
 
 	virtual bool SupportsSummaryView() const override { return true; }
-	NIAGARAEDITOR_API virtual FNiagaraHierarchyIdentity DetermineSummaryIdentity() const override;
+	NIAGARAEDITOR_API virtual FHierarchyElementIdentity DetermineSummaryIdentity() const override;
 private:
 	NIAGARAEDITOR_API void SimulationStagePropertiesChanged();
 
diff --git a/Engine/Plugins/FX/Niagara/Source/NiagaraEditor/Public/ViewModels/Stack/NiagaraStackSummaryViewInputCollection.h b/Engine/Plugins/FX/Niagara/Source/NiagaraEditor/Public/ViewModels/Stack/NiagaraStackSummaryViewInputCollection.h
--- a/Engine/Plugins/FX/Niagara/Source/NiagaraEditor/Public/ViewModels/Stack/NiagaraStackSummaryViewInputCollection.h
+++ b/Engine/Plugins/FX/Niagara/Source/NiagaraEditor/Public/ViewModels/Stack/NiagaraStackSummaryViewInputCollection.h
@@ -4,7 +4,7 @@
 
 #include "NiagaraEmitter.h"
 #include "NiagaraStackValueCollection.h"
-#include "ViewModels/HierarchyEditor/NiagaraHierarchyViewModelBase.h"
+#include "HierarchyViewModelBase.h"
 #include "NiagaraStackSummaryViewInputCollection.generated.h"
 
 UCLASS(MinimalAPI)
@@ -30,7 +30,7 @@ protected:
 	NIAGARAEDITOR_API virtual void RefreshChildrenInternal(const TArray<UNiagaraStackEntry*>& CurrentChildren, TArray<UNiagaraStackEntry*>& NewChildren, TArray<FStackIssue>& NewIssues) override;
 
 	NIAGARAEDITOR_API virtual void GetSectionsInternal(TArray<FNiagaraStackSection>& OutStackSections) const override;
-	NIAGARAEDITOR_API const TArray<UNiagaraHierarchySection*>& GetHierarchySections() const;
+	NIAGARAEDITOR_API const TArray<UHierarchySection*>& GetHierarchySections() const;
 
 	NIAGARAEDITOR_API virtual bool FilterByActiveSection(const UNiagaraStackEntry& Child) const override;
 private:
diff --git a/Engine/Plugins/FX/Niagara/Source/NiagaraEditor/Public/Widgets/SNiagaraScriptInputPreviewPanel.h b/Engine/Plugins/FX/Niagara/Source/NiagaraEditor/Public/Widgets/SNiagaraScriptInputPreviewPanel.h
--- a/Engine/Plugins/FX/Niagara/Source/NiagaraEditor/Public/Widgets/SNiagaraScriptInputPreviewPanel.h
+++ b/Engine/Plugins/FX/Niagara/Source/NiagaraEditor/Public/Widgets/SNiagaraScriptInputPreviewPanel.h
@@ -5,7 +5,7 @@
 #include "CoreMinimal.h"
 #include "NiagaraGraph.h"
 #include "NiagaraScriptVariable.h"
-#include "ViewModels/HierarchyEditor/NiagaraHierarchyViewModelBase.h"
+#include "HierarchyViewModelBase.h"
 #include "Widgets/SCompoundWidget.h"
 #include "Widgets/Input/SSearchBox.h"
 #include "Widgets/Views/STreeView.h"
@@ -31,17 +31,17 @@ public:
 	void SetupDelegates();
 	void RemoveDelegates();
 private:
-	TSharedRef<ITableRow> OnGenerateRow(UNiagaraHierarchyItemBase* Item, const TSharedRef<STableViewBase>& TableViewBase) const;
-	void OnGetChildren(UNiagaraHierarchyItemBase* Item, TArray<UNiagaraHierarchyItemBase*>& OutChildren) const;
+	TSharedRef<ITableRow> OnGenerateRow(UHierarchyElement* Item, const TSharedRef<STableViewBase>& TableViewBase) const;
+	void OnGetChildren(UHierarchyElement* Item, TArray<UHierarchyElement*>& OutChildren) const;
 	void OnParametersChanged(TOptional<UNiagaraGraph::FParametersChangedData> ParametersChangedData);
 
 	FReply SummonHierarchyEditor() const;
 
 	struct FSearchItem
 	{
-		TArray<UNiagaraHierarchyItemBase*> Path;
+		TArray<UHierarchyElement*> Path;
 
-		UNiagaraHierarchyItemBase* GetEntry() const
+		UHierarchyElement* GetEntry() const
 		{
 			return Path.Num() > 0 ? 
 				Path[Path.Num() - 1] : 
@@ -57,7 +57,7 @@ private:
 	void OnSearchTextChanged(const FText& Text);
 	void OnSearchTextCommitted(const FText& Text, ETextCommit::Type CommitType);
 	void OnSearchButtonClicked(SSearchBox::SearchDirection SearchDirection);
-	void GenerateSearchItems(UNiagaraHierarchyItemBase* Root, TArray<UNiagaraHierarchyItemBase*> ParentChain, TArray<FSearchItem>& OutSearchItems);
+	void GenerateSearchItems(UHierarchyElement* Root, TArray<UHierarchyElement*> ParentChain, TArray<FSearchItem>& OutSearchItems);
 	void ExpandSourceSearchResults();
 	void SelectNextSourceSearchResult();
 	void SelectPreviousSourceSearchResult();
@@ -71,11 +71,11 @@ private:
 	virtual void PostUndo(bool bSuccess) override;
 	virtual void PostRedo(bool bSuccess) override;
 private:
-	TArray<UNiagaraHierarchyItemBase*> RootArray;
+	TArray<UHierarchyElement*> RootArray;
 	TWeakPtr<FNiagaraScriptToolkit> ScriptToolkit;
 	TWeakPtr<FNiagaraObjectSelection> VariableObjectSelection;
 	TSharedPtr<SSearchBox> SearchBox;
-	TSharedPtr<STreeView<UNiagaraHierarchyItemBase*>> TreeView;
+	TSharedPtr<STreeView<UHierarchyElement*>> TreeView;
 
 	/** We construct and maintain an array of parameters _not_ in the hierarchy to ensure we display all parameters without requiring hierarchy setup. */
 	TArray<TObjectPtr<UNiagaraHierarchyScriptParameter>> TransientLeftoverParameters; 
diff --git a/Engine/Plugins/FX/Niagara/Source/NiagaraEditorWidgets/NiagaraEditorWidgets.Build.cs b/Engine/Plugins/FX/Niagara/Source/NiagaraEditorWidgets/NiagaraEditorWidgets.Build.cs
--- a/Engine/Plugins/FX/Niagara/Source/NiagaraEditorWidgets/NiagaraEditorWidgets.Build.cs
+++ b/Engine/Plugins/FX/Niagara/Source/NiagaraEditorWidgets/NiagaraEditorWidgets.Build.cs
@@ -35,7 +35,8 @@ public class NiagaraEditorWidgets : ModuleRules
 			"DesktopPlatform",
 			"ToolWidgets",
 			"KismetWidgets",
-			"ToolMenus"
+			"ToolMenus",
+			"HierarchyEditor"
 		});
 
 		PrivateIncludePathModuleNames.AddRange(new string[] {
diff --git a/Engine/Plugins/FX/Niagara/Source/NiagaraEditorWidgets/Private/Stack/SNiagaraStackScriptHierarchyRoot.cpp b/Engine/Plugins/FX/Niagara/Source/NiagaraEditorWidgets/Private/Stack/SNiagaraStackScriptHierarchyRoot.cpp
--- a/Engine/Plugins/FX/Niagara/Source/NiagaraEditorWidgets/Private/Stack/SNiagaraStackScriptHierarchyRoot.cpp
+++ b/Engine/Plugins/FX/Niagara/Source/NiagaraEditorWidgets/Private/Stack/SNiagaraStackScriptHierarchyRoot.cpp
@@ -12,7 +12,7 @@ BEGIN_SLATE_FUNCTION_BUILD_OPTIMIZATION
 
 #define LOCTEXT_NAMESPACE "NiagaraEditor"
 
-void SNiagaraStackHierarchySection::Construct(const FArguments& InArgs, const UNiagaraHierarchySection* InSection)
+void SNiagaraStackHierarchySection::Construct(const FArguments& InArgs, const UHierarchySection* InSection)
 {
 	Section = InSection;
 	
@@ -79,12 +79,12 @@ EVisibility SNiagaraStackScriptHierarchyRoot::GetLabelVisibility() const
 	return ModuleHierarchyRoot->GetShouldDisplayLabel() ? EVisibility::Visible : EVisibility::Collapsed;
 }
 
-ECheckBoxState SNiagaraStackScriptHierarchyRoot::IsSectionChecked(const UNiagaraHierarchySection* NiagaraHierarchySection) const
+ECheckBoxState SNiagaraStackScriptHierarchyRoot::IsSectionChecked(const UHierarchySection* NiagaraHierarchySection) const
 {
 	return ModuleHierarchyRoot->GetActiveSection() == NiagaraHierarchySection ? ECheckBoxState::Checked : ECheckBoxState::Unchecked;
 }
 
-void SNiagaraStackScriptHierarchyRoot::OnCheckStateChanged(ECheckBoxState CheckBoxState, const UNiagaraHierarchySection* NiagaraHierarchySection)
+void SNiagaraStackScriptHierarchyRoot::OnCheckStateChanged(ECheckBoxState CheckBoxState, const UHierarchySection* NiagaraHierarchySection)
 {
 	if(CheckBoxState == ECheckBoxState::Checked)
 	{
@@ -95,7 +95,7 @@ void SNiagaraStackScriptHierarchyRoot::OnCheckStateChanged(ECheckBoxState CheckB
 void SNiagaraStackScriptHierarchyRoot::ConstructSectionButtons()
 {
 	SectionSelectorBox->ClearChildren();
-	for (const UNiagaraHierarchySection* Section : ModuleHierarchyRoot->GetSections())
+	for (const UHierarchySection* Section : ModuleHierarchyRoot->GetSections())
 	{
 		SectionSelectorBox->AddSlot()
 		[
@@ -108,8 +108,8 @@ void SNiagaraStackScriptHierarchyRoot::ConstructSectionButtons()
 	SectionSelectorBox->AddSlot()
 	[
 		SNew(SNiagaraStackHierarchySection, nullptr)
-		.OnCheckStateChanged(this, &SNiagaraStackScriptHierarchyRoot::OnCheckStateChanged, (const UNiagaraHierarchySection*) nullptr)
-		.IsChecked(this, &SNiagaraStackScriptHierarchyRoot::IsSectionChecked, (const UNiagaraHierarchySection*) nullptr)
+		.OnCheckStateChanged(this, &SNiagaraStackScriptHierarchyRoot::OnCheckStateChanged, (const UHierarchySection*) nullptr)
+		.IsChecked(this, &SNiagaraStackScriptHierarchyRoot::IsSectionChecked, (const UHierarchySection*) nullptr)
 	];
 }
 
diff --git a/Engine/Plugins/FX/Niagara/Source/NiagaraEditorWidgets/Private/Stack/SNiagaraStackScriptHierarchyRoot.h b/Engine/Plugins/FX/Niagara/Source/NiagaraEditorWidgets/Private/Stack/SNiagaraStackScriptHierarchyRoot.h
--- a/Engine/Plugins/FX/Niagara/Source/NiagaraEditorWidgets/Private/Stack/SNiagaraStackScriptHierarchyRoot.h
+++ b/Engine/Plugins/FX/Niagara/Source/NiagaraEditorWidgets/Private/Stack/SNiagaraStackScriptHierarchyRoot.h
@@ -17,13 +17,13 @@ public:
 	SLATE_ATTRIBUTE(ECheckBoxState, IsChecked)
 SLATE_END_ARGS()
 	
-	void Construct(const FArguments& InArgs, const UNiagaraHierarchySection* InSection);
+	void Construct(const FArguments& InArgs, const UHierarchySection* InSection);
 
 private:
 	FText GetSectionNameAsText() const;
 	FText GetTooltipText() const;
 private:
-	TWeakObjectPtr<const UNiagaraHierarchySection> Section;
+	TWeakObjectPtr<const UHierarchySection> Section;
 };
 
 class SNiagaraStackScriptHierarchyRoot : public SCompoundWidget
@@ -39,8 +39,8 @@ private:
 	
 	EVisibility GetLabelVisibility() const;
 
-	ECheckBoxState IsSectionChecked(const UNiagaraHierarchySection* NiagaraHierarchySection) const;
-	void OnCheckStateChanged(ECheckBoxState CheckBoxState, const UNiagaraHierarchySection* NiagaraHierarchySection);
+	ECheckBoxState IsSectionChecked(const UHierarchySection* NiagaraHierarchySection) const;
+	void OnCheckStateChanged(ECheckBoxState CheckBoxState, const UHierarchySection* NiagaraHierarchySection);
 
 	void HierarchyStructureChanged(ENiagaraStructureChangedFlags StructureChangedFlags);
 
diff --git a/Engine/Plugins/FX/Niagara/Source/NiagaraEditorWidgets/Private/Stack/SNiagaraStackTableRow.cpp b/Engine/Plugins/FX/Niagara/Source/NiagaraEditorWidgets/Private/Stack/SNiagaraStackTableRow.cpp
--- a/Engine/Plugins/FX/Niagara/Source/NiagaraEditorWidgets/Private/Stack/SNiagaraStackTableRow.cpp
+++ b/Engine/Plugins/FX/Niagara/Source/NiagaraEditorWidgets/Private/Stack/SNiagaraStackTableRow.cpp
@@ -445,7 +445,7 @@ FReply SNiagaraStackTableRow::OnMouseButtonUp(const FGeometry& MyGeometry, const
 		{			
 			if(StackEntry->SupportsSummaryView())
 			{
-				FNiagaraHierarchyIdentity Identity = StackEntry->DetermineSummaryIdentity();
+				FHierarchyElementIdentity Identity = StackEntry->DetermineSummaryIdentity();
 				if(Identity.IsValid())
 				{
 					FUIAction ShowHideSummaryViewAction(
@@ -652,19 +652,19 @@ void SNiagaraStackTableRow::ToggleShowInSummaryView() const
 {
 	if(StackEntry->GetEmitterViewModel().IsValid())
 	{
-		FNiagaraHierarchyIdentity Identity = StackEntry->DetermineSummaryIdentity();
+		FHierarchyElementIdentity Identity = StackEntry->DetermineSummaryIdentity();
 		if(Identity.IsValid())
 		{
 			bool bItemExistsInHierarchy = StackEntry->GetEmitterViewModel()->GetSummaryHierarchyViewModel()->GetHierarchyRootViewModel()->FindViewModelForChild(Identity, true) != nullptr;
 			if(bItemExistsInHierarchy)
 			{
-				StackEntry->GetEmitterViewModel()->GetSummaryHierarchyViewModel()->DeleteItemWithIdentity(Identity);
+				StackEntry->GetEmitterViewModel()->GetSummaryHierarchyViewModel()->DeleteElementWithIdentity(Identity);
 			}
 			else
 			{
 				UClass* HierarchyClass = DetermineHierarchyClassForSummaryView();
 				check(HierarchyClass);
-				StackEntry->GetEmitterViewModel()->GetSummaryHierarchyViewModel()->AddItem(HierarchyClass, Identity);
+				StackEntry->GetEmitterViewModel()->GetSummaryHierarchyViewModel()->AddElementUnderRoot(HierarchyClass, Identity);
 			}
 		}
 	}
@@ -689,7 +689,7 @@ FText SNiagaraStackTableRow::GetToggleShowSummaryActionTooltip() const
 
 	FString ResultMessage;
 	FText CantToggleMessage = LOCTEXT("CantToggleSummaryViewBaseMessage", "You can not add or remove this item due to:\n");
-	FNiagaraHierarchyIdentity Identity = StackEntry->DetermineSummaryIdentity();
+	FHierarchyElementIdentity Identity = StackEntry->DetermineSummaryIdentity();
 	
 	bool bExistsInParentEmitter = StackEntry->ExistsInParentEmitterSummary();
 	bool bIsAnyParentInSummaryView = StackEntry->IsAnyParentInSummaryView();
@@ -728,10 +728,10 @@ void SNiagaraStackTableRow::NavigateToSummaryView() const
 {
 	StackEntry->GetSystemViewModel()->FocusTab(FName("NiagaraSystemEditor_SummaryViewHierarchyEditor"), true);
 	StackEntry->GetEmitterViewModel()->GetSummaryHierarchyViewModel()->SetActiveHierarchySection(nullptr);
-	StackEntry->GetEmitterViewModel()->GetSummaryHierarchyViewModel()->NavigateToItemInHierarchy(StackEntry->DetermineSummaryIdentity());
+	StackEntry->GetEmitterViewModel()->GetSummaryHierarchyViewModel()->NavigateToElementInHierarchy(StackEntry->DetermineSummaryIdentity());
 }
 
-TSubclassOf<UNiagaraHierarchyItemBase> SNiagaraStackTableRow::DetermineHierarchyClassForSummaryView() const
+TSubclassOf<UHierarchyElement> SNiagaraStackTableRow::DetermineHierarchyClassForSummaryView() const
 {
 	if(UNiagaraStackFunctionInput* FunctionInput = Cast<UNiagaraStackFunctionInput>(StackEntry))
 	{
diff --git a/Engine/Plugins/FX/Niagara/Source/NiagaraEditorWidgets/Private/Stack/SNiagaraStackTableRow.h b/Engine/Plugins/FX/Niagara/Source/NiagaraEditorWidgets/Private/Stack/SNiagaraStackTableRow.h
--- a/Engine/Plugins/FX/Niagara/Source/NiagaraEditorWidgets/Private/Stack/SNiagaraStackTableRow.h
+++ b/Engine/Plugins/FX/Niagara/Source/NiagaraEditorWidgets/Private/Stack/SNiagaraStackTableRow.h
@@ -11,7 +11,7 @@
 #include "Layout/Visibility.h"
 #include "Animation/CurveSequence.h"
 
-class UNiagaraHierarchyItemBase;
+class UHierarchyElement;
 class UNiagaraStackViewModel;
 class UNiagaraStackEntry;
 class FNiagaraStackCommandContext;
@@ -107,7 +107,7 @@ private:
 
 	void NavigateToSummaryView() const;
 	
-	TSubclassOf<UNiagaraHierarchyItemBase> DetermineHierarchyClassForSummaryView() const;
+	TSubclassOf<UHierarchyElement> DetermineHierarchyClassForSummaryView() const;
 private:
 	UNiagaraStackViewModel* StackViewModel;
 	UNiagaraStackEntry* StackEntry;
diff --git a/Engine/Source/Editor/HierarchyEditor/HierarchyEditor.Build.cs b/Engine/Source/Editor/HierarchyEditor/HierarchyEditor.Build.cs
new file mode 100644
--- /dev/null
+++ b/Engine/Source/Editor/HierarchyEditor/HierarchyEditor.Build.cs
@@ -0,0 +1,32 @@
+﻿// Copyright Epic Games, Inc. All Rights Reserved.
+
+using UnrealBuildTool;
+
+public class HierarchyEditor : ModuleRules
+{
+    public HierarchyEditor(ReadOnlyTargetRules Target) : base(Target)
+    {
+        PCHUsage = ModuleRules.PCHUsageMode.UseExplicitOrSharedPCHs;
+
+        PublicDependencyModuleNames.AddRange(
+            new string[]
+            {
+                "Core", 
+                "CoreUObject",
+                "EditorWidgets",
+                "SlateCore"
+            }
+        );
+
+        PrivateDependencyModuleNames.AddRange(
+            new string[]
+            {
+                "Engine",
+                "Slate",
+                "UnrealEd",
+                "ToolMenus",
+                "InputCore"
+            }
+        );
+    }
+}
\ No newline at end of file
diff --git a/Engine/Source/Editor/HierarchyEditor/Private/HierarchyEditorCommands.cpp b/Engine/Source/Editor/HierarchyEditor/Private/HierarchyEditorCommands.cpp
new file mode 100644
--- /dev/null
+++ b/Engine/Source/Editor/HierarchyEditor/Private/HierarchyEditorCommands.cpp
@@ -0,0 +1,12 @@
+// Copyright Epic Games, Inc. All Rights Reserved.
+
+#include "HierarchyEditorCommands.h"
+
+#define LOCTEXT_NAMESPACE "HierarchyEditorCommands"
+
+void FHierarchyEditorCommands::RegisterCommands()
+{
+	UI_COMMAND(FindInHierarchy, "Find in Hierarchy", "Find the element in the hierarchy", EUserInterfaceActionType::Button, FInputChord(EKeys::F));
+}
+
+#undef LOCTEXT_NAMESPACE
diff --git a/Engine/Source/Editor/HierarchyEditor/Private/HierarchyEditorMisc.cpp b/Engine/Source/Editor/HierarchyEditor/Private/HierarchyEditorMisc.cpp
new file mode 100644
--- /dev/null
+++ b/Engine/Source/Editor/HierarchyEditor/Private/HierarchyEditorMisc.cpp
@@ -0,0 +1,37 @@
+﻿// Copyright Epic Games, Inc. All Rights Reserved.
+
+#include "HierarchyEditorMisc.h"
+
+FName UE::HierarchyEditor::GetUniqueName(FName CandidateName, const TSet<FName>& ExistingNames)
+{
+	// This utility function needs to generate a unique name while only considering the text portion of the name and
+	// not the index, so generate names with 0 indices before using them for comparison.
+	TSet<FName> ExistingNamesWithIndexZero;
+	for (FName ExistingName : ExistingNames)
+	{
+		ExistingNamesWithIndexZero.Add(FName(ExistingName, 0));
+	}
+	FName CandidateNameWithIndexZero = FName(CandidateName, 0);
+
+	if (ExistingNamesWithIndexZero.Contains(CandidateNameWithIndexZero) == false)
+	{
+		return CandidateName;
+	}
+
+	FString CandidateNameString = CandidateNameWithIndexZero.ToString();
+	FString BaseNameString = CandidateNameString;
+	if (CandidateNameString.Len() >= 3 && CandidateNameString.Right(3).IsNumeric())
+	{
+		BaseNameString = CandidateNameString.Left(CandidateNameString.Len() - 3);
+	}
+
+	FName UniqueName = FName(*BaseNameString);
+	int32 NameIndex = 1;
+	while (ExistingNamesWithIndexZero.Contains(UniqueName))
+	{
+		UniqueName = FName(*FString::Printf(TEXT("%s%03i"), *BaseNameString, NameIndex));
+		NameIndex++;
+	}
+
+	return UniqueName;
+}
diff --git a/Engine/Source/Editor/HierarchyEditor/Private/HierarchyEditorModule.cpp b/Engine/Source/Editor/HierarchyEditor/Private/HierarchyEditorModule.cpp
new file mode 100644
--- /dev/null
+++ b/Engine/Source/Editor/HierarchyEditor/Private/HierarchyEditorModule.cpp
@@ -0,0 +1,26 @@
+﻿// Copyright Epic Games, Inc. All Rights Reserved.
+
+#include "HierarchyEditorModule.h"
+
+#include "HierarchyEditorCommands.h"
+#include "HierarchyEditorStyle.h"
+
+#define LOCTEXT_NAMESPACE "HierarchyEditor"
+
+DEFINE_LOG_CATEGORY(LogHierarchyEditor)
+
+void FHierarchyEditorModule::StartupModule()
+{
+    FHierarchyEditorStyle::Register();
+	FHierarchyEditorCommands::Register();
+}
+
+void FHierarchyEditorModule::ShutdownModule()
+{
+	FHierarchyEditorCommands::Unregister();
+    FHierarchyEditorStyle::Shutdown();
+}
+
+#undef LOCTEXT_NAMESPACE
+    
+IMPLEMENT_MODULE(FHierarchyEditorModule, HierarchyEditor)
diff --git a/Engine/Source/Editor/HierarchyEditor/Private/HierarchyEditorStyle.cpp b/Engine/Source/Editor/HierarchyEditor/Private/HierarchyEditorStyle.cpp
new file mode 100644
--- /dev/null
+++ b/Engine/Source/Editor/HierarchyEditor/Private/HierarchyEditorStyle.cpp
@@ -0,0 +1,126 @@
+// Copyright Epic Games, Inc. All Rights Reserved.
+
+#include "HierarchyEditorStyle.h"
+
+#include "Framework/Application/SlateApplication.h"
+#include "Styling/SlateStyleMacros.h"
+#include "Styling/AppStyle.h"
+#include "Styling/SlateStyleRegistry.h"
+#include "Styling/SlateTypes.h"
+#include "Styling/CoreStyle.h"
+#include "Styling/StarshipCoreStyle.h"
+#include "Styling/StyleColors.h"
+
+TSharedPtr<FHierarchyEditorStyle> FHierarchyEditorStyle::HierarchyEditorStyle = nullptr;
+
+void FHierarchyEditorStyle::Register()
+{
+	FSlateStyleRegistry::RegisterSlateStyle(Get());
+}
+
+void FHierarchyEditorStyle::Unregister()
+{
+	FSlateStyleRegistry::UnRegisterSlateStyle(Get());
+}
+
+void FHierarchyEditorStyle::Shutdown()
+{
+	Unregister();
+	HierarchyEditorStyle.Reset();
+}
+
+void FHierarchyEditorStyle::InitHierarchyEditor()
+{
+	const FTextBlockStyle NormalText = FAppStyle::GetWidgetStyle<FTextBlockStyle>("NormalText");
+	FSlateFontInfo CategoryFont = FAppStyle::Get().GetFontStyle(TEXT("DetailsView.CategoryFontStyle"));
+	CategoryFont.Size = 11;
+	
+	const FEditableTextBoxStyle NormalEditableTextBox = FAppStyle::Get().GetWidgetStyle<FEditableTextBoxStyle>("NormalEditableTextBox");
+	
+	FEditableTextBoxStyle CategoryEditableText = FEditableTextBoxStyle(NormalEditableTextBox)
+		.SetFont(CategoryFont)
+		.SetForegroundColor(FStyleColors::AccentWhite);
+	
+	FTextBlockStyle CategoryText = FTextBlockStyle(NormalText)
+		.SetFont(CategoryFont);
+	FInlineEditableTextBlockStyle HierarchyCategoryTextStyle = FInlineEditableTextBlockStyle()
+		.SetTextStyle(CategoryText)
+		.SetEditableTextBoxStyle(CategoryEditableText);
+
+	Set("HierarchyEditor.CategoryTextBlock", CategoryText);
+	Set("HierarchyEditor.Category", HierarchyCategoryTextStyle);
+
+	Set("HierarchyEditor.Stack.DropTarget.BorderVertical", new IMAGE_BRUSH("Icons/StackDropTargetBorder_Vertical", FVector2D(2, 8), FLinearColor::White, ESlateBrushTileType::Vertical));
+	Set("HierarchyEditor.Stack.DropTarget.BorderHorizontal", new IMAGE_BRUSH("Icons/StackDropTargetBorder_Horizontal", FVector2D(8, 2), FLinearColor::White, ESlateBrushTileType::Horizontal));
+
+	FButtonStyle SimpleButtonStyle = FAppStyle::Get().GetWidgetStyle<FButtonStyle>("SimpleButton");
+	FButtonStyle ButtonStyle = FButtonStyle(SimpleButtonStyle)
+		.SetNormalForeground(FStyleColors::Foreground)
+		.SetPressedForeground(FStyleColors::ForegroundHover)
+		.SetHoveredForeground(FStyleColors::ForegroundHover)
+		.SetDisabledForeground(FStyleColors::Foreground)
+		.SetNormalPadding(FMargin(8.f, 2.f, 8.f, 2.f))
+		.SetPressedPadding(FMargin(8.f, 3.f, 8.f, 1.f));
+
+	Set("HierarchyEditor.ButtonStyle", ButtonStyle);
+
+	const FSlateColor SelectorColor = FAppStyle::GetSlateColor("SelectorColor");
+	const FSlateColor SelectionColor = FAppStyle::GetSlateColor("SelectionColor");
+	const FSlateColor SelectionColor_Inactive = FAppStyle::GetSlateColor("SelectionColor_Inactive");
+	FLinearColor HighlightColor = FLinearColor(0.068f, 0.068f, 0.068f);
+	
+	Set("HierarchyEditor.Color.Highlight", HighlightColor);
+	
+	FTableRowStyle HierarchyEditorCategoryRowStyle = FTableRowStyle()
+		.SetEvenRowBackgroundBrush(*FAppStyle::GetBrush("DetailsView.CategoryTop"))
+		.SetOddRowBackgroundBrush(*FAppStyle::GetBrush("DetailsView.CategoryTop"))
+		.SetEvenRowBackgroundHoveredBrush(*FAppStyle::GetBrush("DetailsView.CategoryTop_Hovered"))
+		.SetOddRowBackgroundHoveredBrush(*FAppStyle::GetBrush("DetailsView.CategoryTop_Hovered"))
+		.SetSelectorFocusedBrush( CORE_BORDER_BRUSH( "Common/Selector", FMargin(4.f/16.f), SelectorColor))
+		.SetActiveBrush( CORE_IMAGE_BRUSH( "Common/Selection", CoreStyleConstants::Icon8x8, SelectionColor))
+		.SetActiveHoveredBrush( CORE_IMAGE_BRUSH( "Common/Selection", CoreStyleConstants::Icon8x8, SelectionColor))
+		.SetInactiveBrush( CORE_IMAGE_BRUSH( "Common/Selection", CoreStyleConstants::Icon8x8, SelectionColor_Inactive))
+		.SetInactiveHoveredBrush( CORE_IMAGE_BRUSH( "Common/Selection", CoreStyleConstants::Icon8x8, SelectionColor_Inactive))
+		.SetActiveHighlightedBrush( CORE_IMAGE_BRUSH("Common/Selection", CoreStyleConstants::Icon8x8, HighlightColor))
+		.SetInactiveHighlightedBrush( CORE_IMAGE_BRUSH("Common/Selection", CoreStyleConstants::Icon8x8, HighlightColor))
+		.SetDropIndicator_Above(CORE_BOX_BRUSH("Common/DropZoneIndicator_Above", FMargin(10.0f / 16.0f, 10.0f / 16.0f, 0, 0), FStarshipCoreStyle::GetCoreStyle().GetSlateColor("SelectorColor")))
+		.SetDropIndicator_Onto(CORE_BOX_BRUSH("Common/DropZoneIndicator_Onto", FMargin(4.0f / 16.0f), FStarshipCoreStyle::GetCoreStyle().GetSlateColor("SelectorColor")))
+		.SetDropIndicator_Below(CORE_BOX_BRUSH("Common/DropZoneIndicator_Below", FMargin(10.0f / 16.0f, 0, 0, 10.0f / 16.0f), FStarshipCoreStyle::GetCoreStyle().GetSlateColor("SelectorColor")));
+	
+	Set("HierarchyEditor.Row.Category", HierarchyEditorCategoryRowStyle);
+	
+	FInlineEditableTextBlockStyle HierarchyEditorCategoryStyle = FAppStyle::GetWidgetStyle<FInlineEditableTextBlockStyle>("InlineEditableTextBlockStyle");
+	HierarchyEditorCategoryStyle.TextStyle = FAppStyle::Get().GetWidgetStyle<FTextBlockStyle>("DetailsView.CategoryTextStyle");
+	HierarchyEditorCategoryStyle.TextStyle.SetFontSize(10.f);
+	Set("HierarchyEditor.CategoryTextStyle", HierarchyEditorCategoryStyle);
+	
+	Set("HierarchyEditor.RootDropIcon", new IMAGE_BRUSH_SVG("Starship/Common/caret-down", CoreStyleConstants::Icon20x20));
+}
+
+FHierarchyEditorStyle::FHierarchyEditorStyle() : FSlateStyleSet("HierarchyEditorStyle")
+{
+	SetContentRoot(FPaths::EngineContentDir() / TEXT("Slate"));
+	InitHierarchyEditor();
+}
+
+void FHierarchyEditorStyle::ReloadTextures()
+{
+	FSlateApplication::Get().GetRenderer()->ReloadTextureResources();
+}
+
+const FHierarchyEditorStyle& FHierarchyEditorStyle::Get()
+{
+	if(!HierarchyEditorStyle.IsValid())
+	{
+		HierarchyEditorStyle = MakeShareable(new FHierarchyEditorStyle());
+	}
+	
+	return *HierarchyEditorStyle;
+}
+
+void FHierarchyEditorStyle::ReinitializeStyle()
+{
+	Unregister();
+	HierarchyEditorStyle.Reset();
+	Register();	
+}
