- Remove standalone flag for SVT baker preview material

#rnx

diff --git a/Engine/Plugins/FX/Niagara/Source/NiagaraEditor/Private/NiagaraBakerRenderer.cpp b/Engine/Plugins/FX/Niagara/Source/NiagaraEditor/Private/NiagaraBakerRenderer.cpp
--- a/Engine/Plugins/FX/Niagara/Source/NiagaraEditor/Private/NiagaraBakerRenderer.cpp
+++ b/Engine/Plugins/FX/Niagara/Source/NiagaraEditor/Private/NiagaraBakerRenderer.cpp
@@ -660,7 +660,8 @@ void FNiagaraBakerRenderer::RenderSparseVolumeTexture(UTextureRenderTarget2D* Re
 
 		// #todo(dmp): we had to duplicate the material itself because we cannot make a mid and send that to HV
 		// HV internally makes a MID from whatever is bound, and the MID of a MID workflow appears broken	
-		UMaterial *DuplicateMat = DuplicateObject<UMaterial>(Mat, SVTPreviewComponent);
+		UMaterial* DuplicateMat = DuplicateObject<UMaterial>(Mat, SVTPreviewComponent);
+		DuplicateMat->ClearFlags(RF_Standalone);
 
 		FGuid ExprGuid;
 		DuplicateMat->SetStaticComponentMaskParameterValueEditorOnly("Temperature Mask", false, true, false, false, ExprGuid);
