Limit Input tracking to only when needed.

diff --git a/Engine/Plugins/FX/Niagara/Source/NiagaraEditor/Private/NiagaraDebugger.cpp b/Engine/Plugins/FX/Niagara/Source/NiagaraEditor/Private/NiagaraDebugger.cpp
--- a/Engine/Plugins/FX/Niagara/Source/NiagaraEditor/Private/NiagaraDebugger.cpp
+++ b/Engine/Plugins/FX/Niagara/Source/NiagaraEditor/Private/NiagaraDebugger.cpp
@@ -60,11 +60,6 @@ void FNiagaraDebugger::Init()
 
 	GetMutableDefault<UNiagaraDebugHUDSettings>()->OnChangedDelegate.AddSP(this, &FNiagaraDebugger::UpdateDebugHUDSettings);
 	GetMutableDefault<UNiagaraOutliner>()->OnChangedDelegate.AddSP(this, &FNiagaraDebugger::TriggerOutlinerCapture);
-
-	if (FSlateApplication::IsInitialized())
-	{
-		FSlateApplication::Get().RegisterInputPreProcessor(AsShared());
-	}
 }
 
 void FNiagaraDebugger::ExecConsoleCommand(const TCHAR* Cmd, bool bRequiresWorld)
@@ -155,10 +150,20 @@ void FNiagaraDebugger::ToggleOutlinerCaptureOnKeyPress()
 
 	if(bTriggerOutlinerCaptureOnKeyPress)
 	{
+		if (FSlateApplication::IsInitialized())
+		{
+			FSlateApplication::Get().RegisterInputPreProcessor(AsShared());
+		}
+
 		FNiagaraWorldManager::ForAllWorldManagers([&HudMessage](FNiagaraWorldManager& WorldMan) { WorldMan.AddDebugHUDMessage(NiagaraDebuggerCaptureOutlinerOnKeyPressMessageKey, FNiagaraDebugMessage(ENiagaraDebugMessageType::Info, HudMessage, FLT_MAX)); });
 	}
 	else
 	{
+		if (FSlateApplication::IsInitialized())
+		{
+			FSlateApplication::Get().UnregisterInputPreProcessor(AsShared());
+		}
+
 		FNiagaraWorldManager::ForAllWorldManagers([&HudMessage](FNiagaraWorldManager& WorldMan) { WorldMan.AddDebugHUDMessage(NiagaraDebuggerCaptureOutlinerOnKeyPressMessageKey, FNiagaraDebugMessage(ENiagaraDebugMessageType::Info, HudMessage, -1.0f)); });
 	}
 }
