Adding FResourceBulkDataArrayView to replace a number of custom BulkData implementations that did the same thing.

diff --git a/Engine/Plugins/FX/Niagara/Source/Niagara/Private/NiagaraDataInterfaceLandscape.cpp b/Engine/Plugins/FX/Niagara/Source/Niagara/Private/NiagaraDataInterfaceLandscape.cpp
--- a/Engine/Plugins/FX/Niagara/Source/Niagara/Private/NiagaraDataInterfaceLandscape.cpp
+++ b/Engine/Plugins/FX/Niagara/Source/Niagara/Private/NiagaraDataInterfaceLandscape.cpp
@@ -153,23 +153,6 @@ private:
 	TArray<float> HeightValues;
 	TArray<uint8> PhysMatValues;
 
-	struct FTextureBulkData : public FResourceBulkDataInterface
-	{
-		void Init(const void* InData, uint32 InDataSize) { Data = InData; DataSize = InDataSize; }
-		void Clear() { Data = nullptr; DataSize = 0; }
-		virtual const void* GetResourceBulkData() const override { return Data; }
-		virtual uint32 GetResourceBulkDataSize() const override { return DataSize; }
-		virtual void Discard() override { }
-		const void* Data = nullptr;
-		uint32 DataSize = 0;
-	};
-
-	FTextureBulkData* GetHeightBulkData() { return HeightBulkData.DataSize > 0 ? &HeightBulkData : nullptr; }
-	FTextureBulkData* GetPhysMatBulkData() { return PhysMatBulkData.DataSize > 0 ? &PhysMatBulkData : nullptr; }
-
-	FTextureBulkData HeightBulkData;
-	FTextureBulkData PhysMatBulkData;
-
 #if STATS
 	int32 GpuMemoryUsage = 0;
 #endif
@@ -678,14 +661,16 @@ FLandscapeTextureResource::FLandscapeTextureResource(const FIntPoint& InCellCoun
 
 void FLandscapeTextureResource::InitRHI(FRHICommandListBase&)
 {
-	if (HeightBulkData.DataSize > 0)
+	if (HeightValues.Num())
 	{
-		HeightTexture.Initialize(TEXT("FLandscapeTextureResource_HeightTexture"), sizeof(float), CellCount.X, CellCount.Y, EPixelFormat::PF_R32_FLOAT, FTextureReadBuffer2D::DefaultTextureInitFlag, &HeightBulkData);
+		FResourceBulkDataArrayView BulkData(HeightValues);
+		HeightTexture.Initialize(TEXT("FLandscapeTextureResource_HeightTexture"), sizeof(float), CellCount.X, CellCount.Y, EPixelFormat::PF_R32_FLOAT, FTextureReadBuffer2D::DefaultTextureInitFlag, &BulkData);
 	}
 
-	if (PhysMatBulkData.DataSize > 0)
+	if (PhysMatValues.Num())
 	{
-		PhysMatTexture.Initialize(TEXT("FLandscapeTextureResource_PhysMatTexture"), sizeof(uint8), CellCount.X, CellCount.Y, EPixelFormat::PF_R8_UINT, FTextureReadBuffer2D::DefaultTextureInitFlag, &PhysMatBulkData);
+		FResourceBulkDataArrayView BulkData(PhysMatValues);
+		PhysMatTexture.Initialize(TEXT("FLandscapeTextureResource_PhysMatTexture"), sizeof(uint8), CellCount.X, CellCount.Y, EPixelFormat::PF_R8_UINT, FTextureReadBuffer2D::DefaultTextureInitFlag, &BulkData);
 	}
 
 	ReleaseSourceData();
@@ -719,7 +704,6 @@ TArray<float>& FLandscapeTextureResource::EditHeightValues(int32 SampleCount)
 	const float DefaultHeight = 0.0f;
 	HeightValues.Reset(SampleCount);
 	HeightValues.Init(DefaultHeight, SampleCount);
-	HeightBulkData.Init(HeightValues.GetData(), HeightValues.Num() * HeightValues.GetTypeSize());
 
 	return HeightValues;
 }
@@ -729,7 +713,6 @@ TArray<uint8>& FLandscapeTextureResource::EditPhysMatValues(int32 SampleCount)
 	const uint8 DefaultPhysMat = INDEX_NONE;
 	PhysMatValues.Reset(SampleCount);
 	PhysMatValues.Init(DefaultPhysMat, SampleCount);
-	PhysMatBulkData.Init(PhysMatValues.GetData(), PhysMatValues.Num() * PhysMatValues.GetTypeSize());
 
 	return PhysMatValues;
 }
diff --git a/Engine/Plugins/PCG/Source/PCG/Private/Compute/DataInterfaces/PCGLandscapeDataInterface.cpp b/Engine/Plugins/PCG/Source/PCG/Private/Compute/DataInterfaces/PCGLandscapeDataInterface.cpp
--- a/Engine/Plugins/PCG/Source/PCG/Private/Compute/DataInterfaces/PCGLandscapeDataInterface.cpp
+++ b/Engine/Plugins/PCG/Source/PCG/Private/Compute/DataInterfaces/PCGLandscapeDataInterface.cpp
@@ -40,9 +40,10 @@ public:
 
 	virtual void InitRHI(FRHICommandListBase& RHICmdList) override
 	{
-		if (HeightBulkData.DataSize > 0)
+		if (HeightValues.Num())
 		{
-			HeightTexture.Initialize(TEXT("FLandscapeTextureResource_HeightTexture"), sizeof(float), CellCount.X, CellCount.Y, EPixelFormat::PF_R32_FLOAT, FTextureReadBuffer2D::DefaultTextureInitFlag, &HeightBulkData);
+			FResourceBulkDataArrayView BulkData(HeightValues);
+			HeightTexture.Initialize(TEXT("FLandscapeTextureResource_HeightTexture"), sizeof(float), CellCount.X, CellCount.Y, EPixelFormat::PF_R32_FLOAT, FTextureReadBuffer2D::DefaultTextureInitFlag, &BulkData);
 		}
 
 		ReleaseSourceData();
@@ -58,7 +59,6 @@ public:
 		const float DefaultHeight = 0.0f;
 		HeightValues.Reset(SampleCount);
 		HeightValues.Init(DefaultHeight, SampleCount);
-		HeightBulkData.Init(HeightValues.GetData(), HeightValues.Num() * HeightValues.GetTypeSize());
 
 		return HeightValues;
 	}
@@ -68,20 +68,6 @@ private:
 	FIntPoint CellCount;
 
 	TArray<float> HeightValues;
-
-	struct FTextureBulkData : public FResourceBulkDataInterface
-	{
-		void Init(const void* InData, uint32 InDataSize) { Data = InData; DataSize = InDataSize; }
-		void Clear() { Data = nullptr; DataSize = 0; }
-		virtual const void* GetResourceBulkData() const override { return Data; }
-		virtual uint32 GetResourceBulkDataSize() const override { return DataSize; }
-		virtual void Discard() override { }
-		const void* Data = nullptr;
-		uint32 DataSize = 0;
-	};
-
-	FTextureBulkData* GetHeightBulkData() { return HeightBulkData.DataSize > 0 ? &HeightBulkData : nullptr; }
-	FTextureBulkData HeightBulkData;
 };
 
 /** Manages the resources created by this DI. */
diff --git a/Engine/Plugins/Runtime/nDisplay/Source/DisplayCluster/Private/Render/Containers/DisplayClusterRender_TextureResource.cpp b/Engine/Plugins/Runtime/nDisplay/Source/DisplayCluster/Private/Render/Containers/DisplayClusterRender_TextureResource.cpp
--- a/Engine/Plugins/Runtime/nDisplay/Source/DisplayCluster/Private/Render/Containers/DisplayClusterRender_TextureResource.cpp
+++ b/Engine/Plugins/Runtime/nDisplay/Source/DisplayCluster/Private/Render/Containers/DisplayClusterRender_TextureResource.cpp
@@ -4,36 +4,6 @@
 #include "Containers/ResourceArray.h"
 #include "RenderingThread.h"
 
-/**
- * Container for texture data
- */
-class FDisplayClusterResourceBulkData : public FResourceBulkDataInterface
-{
-public:
-	FDisplayClusterResourceBulkData(const void* InData, uint32_t InDataSize)
-		: Data(InData)
-		, DataSize(InDataSize)
-	{ }
-
-public:
-	virtual const void* GetResourceBulkData() const
-	{
-		return Data;
-	}
-
-	virtual uint32 GetResourceBulkDataSize() const
-	{
-		return DataSize;
-	}
-
-	virtual void Discard()
-	{ }
-
-private:
-	const void* Data;
-	uint32_t    DataSize;
-};
-
 //---------------------------------------------------------------------------------------
 // FDisplayClusterRender_TextureResource
 //---------------------------------------------------------------------------------------
@@ -149,7 +119,7 @@ void FDisplayClusterRender_TextureResource::InitRHI(FRHICommandListBase&)
 	}
 
 	const uint32 DataSize = CalculateImageBytes(Width, Height, 1, PixelFormat);
-	FDisplayClusterResourceBulkData BulkDataInterface(TextureData, DataSize);
+	FResourceBulkDataArrayView BulkDataInterface(TextureData, DataSize);
 
 	// Create a texture from the data in memory.
 	const FRHITextureCreateDesc Desc =
diff --git a/Engine/Source/Runtime/Core/Public/Containers/ResourceArray.h b/Engine/Source/Runtime/Core/Public/Containers/ResourceArray.h
--- a/Engine/Source/Runtime/Core/Public/Containers/ResourceArray.h
+++ b/Engine/Source/Runtime/Core/Public/Containers/ResourceArray.h
@@ -129,6 +129,50 @@ public:
 	}
 };
 
+/** Utility to do a simple upload of data from an array managed by the caller. */
+struct FResourceBulkDataArrayView : public FResourceBulkDataInterface
+{
+	const void* const Data;
+	const uint32 SizeInBytes;
+
+	FResourceBulkDataArrayView() = delete;
+
+	[[nodiscard]] explicit FResourceBulkDataArrayView(const void* InData, uint32 InSizeInBytes)
+		: Data(InData)
+		, SizeInBytes(InSizeInBytes)
+	{
+	}
+
+	template<typename ElementType>
+	[[nodiscard]] explicit FResourceBulkDataArrayView(TConstArrayView<ElementType> View)
+		: Data(View.GetData())
+		, SizeInBytes(View.Num()* View.GetTypeSize())
+	{
+	}
+
+	template<typename ElementType, typename AllocatorType>
+	[[nodiscard]] explicit FResourceBulkDataArrayView(const TArray<ElementType, AllocatorType>& InArray)
+		: FResourceBulkDataArrayView(MakeArrayView(InArray))
+	{
+	}
+
+	// FResourceArrayUploadInterface
+	[[nodiscard]] virtual const void* GetResourceBulkData() const final
+	{
+		return Data;
+	}
+
+	[[nodiscard]] virtual uint32 GetResourceBulkDataSize() const final
+	{
+		return SizeInBytes;
+	}
+
+	virtual void Discard() final
+	{
+		// do nothing
+	}
+};
+
 
 /**
  * Allows for direct GPU mem allocation for texture resource.
diff --git a/Engine/Source/Runtime/Engine/Private/VectorField.cpp b/Engine/Source/Runtime/Engine/Private/VectorField.cpp
--- a/Engine/Source/Runtime/Engine/Private/VectorField.cpp
+++ b/Engine/Source/Runtime/Engine/Private/VectorField.cpp
@@ -129,54 +129,6 @@ void UVectorField::InitInstance(class FVectorFieldInstance* Instance, bool bPrev
 	UVectorFieldStatic implementation.
 ------------------------------------------------------------------------------*/
 
-
-/**
- * Bulk data interface for initializing a static vector field volume texture.
- */
-class FVectorFieldStaticResourceBulkDataInterface : public FResourceBulkDataInterface
-{
-public:
-
-	/** Initialization constructor. */
-	FVectorFieldStaticResourceBulkDataInterface( void* InBulkData, uint32 InBulkDataSize )
-		: BulkData(InBulkData)
-		, BulkDataSize(InBulkDataSize)
-	{
-	}
-
-	/** 
-	 * @return ptr to the resource memory which has been preallocated
-	 */
-	virtual const void* GetResourceBulkData() const override
-	{
-		check(BulkData != NULL);
-		return BulkData;
-	}
-
-	/** 
-	 * @return size of resource memory
-	 */
-	virtual uint32 GetResourceBulkDataSize() const override
-	{
-		check(BulkDataSize > 0);
-		return BulkDataSize;
-	}
-
-	/**
-	 * Free memory after it has been used to initialize RHI resource 
-	 */
-	virtual void Discard() override
-	{
-	}
-
-private:
-
-	/** Pointer to the bulk data. */
-	void* BulkData;
-	/** Size of the bulk data in bytes. */
-	uint32 BulkDataSize;
-};
-
 /**
  * Resource for static vector fields.
  */
@@ -217,7 +169,7 @@ public:
 		if (VolumeData && GSupportsTexture3D)
 		{
 			const uint32 DataSize = SizeX * SizeY * SizeZ * sizeof(FFloat16Color);
-			FVectorFieldStaticResourceBulkDataInterface BulkDataInterface(VolumeData, DataSize);
+			FResourceBulkDataArrayView BulkDataInterface(VolumeData, DataSize);
 
 			const FRHITextureCreateDesc Desc =
 				FRHITextureCreateDesc::Create3D(TEXT("FVectorFieldStaticResource"), SizeX, SizeY, SizeZ, PF_FloatRGBA)
