- Update emitter count rule to allow inclusion / exclusion of stateful & stateless emitters

#rnx
#tests Test run with different counts / emitter types

diff --git a/Engine/Plugins/FX/Niagara/Source/Niagara/Classes/NiagaraEmitterHandle.h b/Engine/Plugins/FX/Niagara/Source/Niagara/Classes/NiagaraEmitterHandle.h
--- a/Engine/Plugins/FX/Niagara/Source/Niagara/Classes/NiagaraEmitterHandle.h
+++ b/Engine/Plugins/FX/Niagara/Source/Niagara/Classes/NiagaraEmitterHandle.h
@@ -7,6 +7,7 @@
 #include "NiagaraTypes.h"
 #include "NiagaraEmitterHandle.generated.h"
 
+struct FNiagaraPlatformSet;
 class UNiagaraSystem;
 class UNiagaraStatelessEmitter;
 
@@ -65,10 +66,13 @@ public:
 	/** Test to see if the emitter is allowed by scalability or not */
 	NIAGARA_API bool IsAllowedByScalability() const;
 
+#if WITH_EDITORONLY_DATA
 	/** Test to see if the emitter is enabled on the provided quality level or not. */
 	NIAGARA_API bool IsEnabledOnEffectQualityLevel(int32 QualityLevel) const;
 
-#if WITH_EDITORONLY_DATA
+	/** Get the PlatformSet for this emitter. */
+	NIAGARA_API const FNiagaraPlatformSet* GetPlatformSet() const;
+
 	bool IsIsolated() const {	return bIsolated; }
 	void SetIsolated(bool bInIsolated) { bIsolated = bInIsolated; }
 #endif
diff --git a/Engine/Plugins/FX/Niagara/Source/Niagara/Private/NiagaraEmitterHandle.cpp b/Engine/Plugins/FX/Niagara/Source/Niagara/Private/NiagaraEmitterHandle.cpp
--- a/Engine/Plugins/FX/Niagara/Source/Niagara/Private/NiagaraEmitterHandle.cpp
+++ b/Engine/Plugins/FX/Niagara/Source/Niagara/Private/NiagaraEmitterHandle.cpp
@@ -195,6 +195,19 @@ bool FNiagaraEmitterHandle::IsEnabledOnEffectQualityLevel(int32 QualityLevel) co
 	}
 }
 
+const FNiagaraPlatformSet* FNiagaraEmitterHandle::GetPlatformSet() const
+{
+	if (EmitterMode == ENiagaraEmitterMode::Standard)
+	{
+		FVersionedNiagaraEmitterData* EmitterData = VersionedInstance.GetEmitterData();
+		return EmitterData ? &EmitterData->Platforms : nullptr;
+	}
+	else
+	{
+		return StatelessEmitter ? &StatelessEmitter->GetPlatformSet() : nullptr;
+	}
+}
+
 void FNiagaraEmitterHandle::SetEmitterMode(UNiagaraSystem& InOwningSystem, ENiagaraEmitterMode InEmitterMode)
 {
 	if (EmitterMode != InEmitterMode)
diff --git a/Engine/Plugins/FX/Niagara/Source/NiagaraEditor/Private/NiagaraValidationRules.cpp b/Engine/Plugins/FX/Niagara/Source/NiagaraEditor/Private/NiagaraValidationRules.cpp
--- a/Engine/Plugins/FX/Niagara/Source/NiagaraEditor/Private/NiagaraValidationRules.cpp
+++ b/Engine/Plugins/FX/Niagara/Source/NiagaraEditor/Private/NiagaraValidationRules.cpp
@@ -463,15 +463,26 @@ void UNiagaraValidationRule_EmitterCount::CheckValidity(const FNiagaraValidation
 		{
 			continue;
 		}
-		FVersionedNiagaraEmitterData* EmitterData = EmitterHandleModel.Get().GetEmitterHandle()->GetEmitterData();
-	
+
+		const FNiagaraPlatformSet* EmitterPlatformSet = EmitterHandle->GetPlatformSet();
+		if (!EmitterPlatformSet)
+		{
+			continue;
+		}
+
+		const bool bIsStateful  = EmitterHandle->GetEmitterMode() == ENiagaraEmitterMode::Standard;
+		const bool bIsStateless = EmitterHandle->GetEmitterMode() == ENiagaraEmitterMode::Stateless;
+
 		for (int32 i=0; i < NumEmitterCountLimits; ++i)
 		{
-			const TArray<FNiagaraPlatformSetConflictInfo> Conflicts = NiagaraValidation::GatherPlatformSetConflicts(&EmitterCountLimits[i].Platforms, &EmitterData->Platforms);
-			if (Conflicts.Num() > 0)
+			if ((EmitterCountLimits[i].bIncludeStateful && bIsStateful) || (EmitterCountLimits[i].bIncludeStateless && bIsStateless))
 			{
-				ConflictsPerLimit[i].Append(Conflicts);
-				++EmitterCountPerLimit[i];
+				const TArray<FNiagaraPlatformSetConflictInfo> Conflicts = NiagaraValidation::GatherPlatformSetConflicts(&EmitterCountLimits[i].Platforms, EmitterPlatformSet);
+				if (Conflicts.Num() > 0)
+				{
+					ConflictsPerLimit[i].Append(Conflicts);
+					++EmitterCountPerLimit[i];
+				}
 			}
 		}
 	}
diff --git a/Engine/Plugins/FX/Niagara/Source/NiagaraEditor/Public/NiagaraValidationRules.h b/Engine/Plugins/FX/Niagara/Source/NiagaraEditor/Public/NiagaraValidationRules.h
--- a/Engine/Plugins/FX/Niagara/Source/NiagaraEditor/Public/NiagaraValidationRules.h
+++ b/Engine/Plugins/FX/Niagara/Source/NiagaraEditor/Public/NiagaraValidationRules.h
@@ -135,6 +135,12 @@ struct FNiagaraValidationRule_EmitterCountAndPlatformSet
 	UPROPERTY(EditAnywhere, Category = Validation)
 	FNiagaraPlatformSet Platforms;
 
+	UPROPERTY(EditAnywhere, Category = Validation, meta = (DisplayName = "Include Regular Emitters"))
+	bool bIncludeStateful = true;
+
+	UPROPERTY(EditAnywhere, Category = Validation, meta = (DisplayName = "Include Lightweight Emitters"))
+	bool bIncludeStateless = true;
+
 	UPROPERTY(EditAnywhere, Category = Validation)
 	int32 EmitterCountLimit = 8;
 };
