- Add binding support for particle lifetime to stateless

#rnx

diff --git a/Engine/Plugins/FX/Niagara/Shaders/Private/Stateless/NiagaraStatelessSimulationTemplate.ush b/Engine/Plugins/FX/Niagara/Shaders/Private/Stateless/NiagaraStatelessSimulationTemplate.ush
--- a/Engine/Plugins/FX/Niagara/Shaders/Private/Stateless/NiagaraStatelessSimulationTemplate.ush
+++ b/Engine/Plugins/FX/Niagara/Shaders/Private/Stateless/NiagaraStatelessSimulationTemplate.ush
@@ -41,7 +41,6 @@ uint	Common_RandomSeed;
 float	Common_SimulationTime;
 float	Common_SimulationDeltaTime;
 float	Common_SimulationInvDeltaTime;
-float2	Common_LifetimeScaleBias;
 
 uint				Common_GPUCountBufferOffset;
 RWBuffer<int>		Common_GPUCountBuffer;
@@ -51,6 +50,8 @@ DECLARE_SCALAR_ARRAY(uint,	SpawnInfo_ParticleOffset, NIAGARA_MAX_GPU_SPAWN_INFOS
 DECLARE_SCALAR_ARRAY(uint,	SpawnInfo_UniqueOffset, NIAGARA_MAX_GPU_SPAWN_INFOS);
 DECLARE_SCALAR_ARRAY(float,	SpawnInfo_Time, NIAGARA_MAX_GPU_SPAWN_INFOS);
 DECLARE_SCALAR_ARRAY(float,	SpawnInfo_Rate, NIAGARA_MAX_GPU_SPAWN_INFOS);
+DECLARE_SCALAR_ARRAY(float,	SpawnInfo_LifetimeScale, NIAGARA_MAX_GPU_SPAWN_INFOS);
+DECLARE_SCALAR_ARRAY(float,	SpawnInfo_LifetimeBias, NIAGARA_MAX_GPU_SPAWN_INFOS);
 
 /////////////////////////////////////////////////////////////////////////////////////////////////////////////////////////
 
@@ -71,6 +72,7 @@ void StatelessMain(uint3 GroupId : SV_GroupID, int GroupThreadIndex : SV_GroupIn
 	Particle.Age = -1.0f;
 
 	// Find which spawn info we belong to and fill out details
+	uint ModuleSeedOffset = 0;
 	{
 		uint SpawnInfoIndex = LinearThreadId;
 		for ( int i=0; i < NIAGARA_MAX_GPU_SPAWN_INFOS; ++i )
@@ -80,11 +82,15 @@ void StatelessMain(uint3 GroupId : SV_GroupID, int GroupThreadIndex : SV_GroupIn
 			const uint SpawnInfoUniqueOffset	= GET_SCALAR_ARRAY_ELEMENT(SpawnInfo_UniqueOffset, i);
 			const float SpawnInfoTime			= GET_SCALAR_ARRAY_ELEMENT(SpawnInfo_Time, i);
 			const float SpawnInfoRate			= GET_SCALAR_ARRAY_ELEMENT(SpawnInfo_Rate, i);
+			const float SpawnInfoLifetimeScale	= GET_SCALAR_ARRAY_ELEMENT(SpawnInfo_LifetimeScale, i);
+			const float SpawnInfoLifetimeBias	= GET_SCALAR_ARRAY_ELEMENT(SpawnInfo_LifetimeBias, i);
 			if ( SpawnInfoIndex < SpawnInfoNumActive )
 			{
 				const uint SpawnParticleIndex = SpawnInfoIndex + SpawnInfoParticleOffset;
 				Particle.UniqueIndex	= SpawnInfoIndex + SpawnInfoUniqueOffset + SpawnInfoParticleOffset;
 				Particle.Age			= Common_SimulationTime - (SpawnInfoTime + float(SpawnParticleIndex) * SpawnInfoRate);
+				GRandomSeedInternal		= FNiagaraStatelessDefinitions::MakeRandomSeed(Common_RandomSeed, Particle.UniqueIndex, ModuleSeedOffset++, 0);
+				Particle.Lifetime		= RandomScaleBiasFloat(0, SpawnInfoLifetimeScale, SpawnInfoLifetimeBias);
 				break;
 			}
 
@@ -97,10 +103,6 @@ void StatelessMain(uint3 GroupId : SV_GroupID, int GroupThreadIndex : SV_GroupIn
 	}
 	
 	// Initialize variables / determine lifetime / etc
-	uint ModuleSeedOffset = 0;
-	GRandomSeedInternal = FNiagaraStatelessDefinitions::MakeRandomSeed(Common_RandomSeed, Particle.UniqueIndex, ModuleSeedOffset++, 0);
-
-	Particle.Lifetime = RandomScaleBiasFloat(0, Common_LifetimeScaleBias.x, Common_LifetimeScaleBias.y);
 	if ( Particle.Lifetime <= 0.0f || Particle.Age >= Particle.Lifetime )
 	{
 		return;
diff --git a/Engine/Plugins/FX/Niagara/Source/Niagara/Internal/Stateless/NiagaraStatelessSpawnInfo.h b/Engine/Plugins/FX/Niagara/Source/Niagara/Internal/Stateless/NiagaraStatelessSpawnInfo.h
--- a/Engine/Plugins/FX/Niagara/Source/Niagara/Internal/Stateless/NiagaraStatelessSpawnInfo.h
+++ b/Engine/Plugins/FX/Niagara/Source/Niagara/Internal/Stateless/NiagaraStatelessSpawnInfo.h
@@ -63,4 +63,6 @@ struct FNiagaraStatelessRuntimeSpawnInfo
 	float	SpawnTimeEnd	= 0.0f;
 	float	Rate			= 0.0f;
 	int32	Amount			= 0;		// Note: When a burst this is the absolute amount burst, when rate this is the amount over the spawn duration
+	float	LifetimeMin		= 0.0f;
+	float	LifetimeMax		= 0.0f;
 };
diff --git a/Engine/Plugins/FX/Niagara/Source/Niagara/Private/Stateless/NiagaraStatelessEmitter.cpp b/Engine/Plugins/FX/Niagara/Source/Niagara/Private/Stateless/NiagaraStatelessEmitter.cpp
--- a/Engine/Plugins/FX/Niagara/Source/Niagara/Private/Stateless/NiagaraStatelessEmitter.cpp
+++ b/Engine/Plugins/FX/Niagara/Source/Niagara/Private/Stateless/NiagaraStatelessEmitter.cpp
@@ -58,12 +58,6 @@ namespace NiagaraStatelessInternal
 			return false;
 		}
 
-		// Lifetime is not valid so we have nothing to every render with
-		if ((EmitterData.LifetimeRange.Min <= 0.0f) && (EmitterData.LifetimeRange.Max <= 0.0f))
-		{
-			return false;
-		}
-
 		// Validate we have output components if we don't there's no point simulating the emitter
 		if (!EmitterData.ParticleDataSetCompiledData.Get() || !EmitterData.ParticleDataSetCompiledData->Variables.Num())
 		{
@@ -284,18 +278,6 @@ void UNiagaraStatelessEmitter::CacheFromCompiledData()
 	StatelessEmitterData->EmitterState.LoopDuration.Min = FMath::Max(StatelessEmitterData->EmitterState.LoopDuration.Min, UE_KINDA_SMALL_NUMBER);
 	StatelessEmitterData->EmitterState.LoopDuration.Max = FMath::Max(StatelessEmitterData->EmitterState.LoopDuration.Max, UE_KINDA_SMALL_NUMBER);
 
-	// Find lifetime values
-	//-Note: We could abstact this out to be a more general modules implements Lifetime but this mirrors core Niagara where you always have Initialize Particle in 99% of cases
-	if (StatelessEmitterData->EmitterTemplate)
-	{
-		UNiagaraStatelessModule_InitializeParticle* InitializeParticleModule = nullptr;
-		Modules.FindItemByClass(&InitializeParticleModule);
-		if (ensure(InitializeParticleModule))
-		{
-			StatelessEmitterData->LifetimeRange = InitializeParticleModule->LifetimeDistribution.CalculateRange();
-		}
-	}
-
 	// Fill in common data, removing any spawn infos that are invalid to the simulation
 	StatelessEmitterData->bDeterministic = bDeterministic;
 	StatelessEmitterData->RandomSeed = RandomSeed ^ 0xdefa081;
@@ -398,6 +380,23 @@ void UNiagaraStatelessEmitter::CacheFromCompiledData()
 			}
 		}
 
+		// Find lifetime values
+		//-Note: We could abstact this out to be a more general modules implements Lifetime but this mirrors core Niagara where you always have Initialize Particle in 99% of cases
+		UNiagaraStatelessModule_InitializeParticle* InitializeParticleModule = nullptr;
+		if (StatelessEmitterData->EmitterTemplate)
+		{
+			Modules.FindItemByClass(&InitializeParticleModule);
+			if (ensure(InitializeParticleModule))
+			{
+				StatelessEmitterData->LifetimeRange = EmitterBuildContext.ConvertDistributionToRange(InitializeParticleModule->LifetimeDistribution, 0.0f);
+
+				if ((StatelessEmitterData->LifetimeRange.Min <= 0.0f) && (StatelessEmitterData->LifetimeRange.Max <= 0.0f) && StatelessEmitterData->LifetimeRange.ParameterOffset == INDEX_NONE)
+				{
+					StatelessEmitterData->bCanEverExecute = false;
+				}
+			}
+		}
+
 		// Populate any renderer bindings from the emitter state and spawn infos
 		EmitterBuildContext.ConvertDistributionToRange(StatelessEmitterData->EmitterState.LoopDuration, 0.0f);
 		EmitterBuildContext.ConvertDistributionToRange(StatelessEmitterData->EmitterState.LoopDelay, 0.0f);
@@ -598,7 +597,6 @@ NiagaraStateless::FCommonShaderParameters* UNiagaraStatelessEmitter::AllocateSha
 	);
 
 	NiagaraStateless::FCommonShaderParameters* CommonParameters = SetShaderParametersContext.GetParameterNestedStruct<NiagaraStateless::FCommonShaderParameters>();
-	CommonParameters->Common_LifetimeScaleBias	= FVector2f(StatelessEmitterData->LifetimeRange.Max - StatelessEmitterData->LifetimeRange.Min, StatelessEmitterData->LifetimeRange.Min);
 
 	for (UNiagaraStatelessModule* Module : Modules)
 	{
diff --git a/Engine/Plugins/FX/Niagara/Source/Niagara/Private/Stateless/NiagaraStatelessEmitterData.cpp b/Engine/Plugins/FX/Niagara/Source/Niagara/Private/Stateless/NiagaraStatelessEmitterData.cpp
--- a/Engine/Plugins/FX/Niagara/Source/Niagara/Private/Stateless/NiagaraStatelessEmitterData.cpp
+++ b/Engine/Plugins/FX/Niagara/Source/Niagara/Private/Stateless/NiagaraStatelessEmitterData.cpp
@@ -85,7 +85,7 @@ uint32 FNiagaraStatelessEmitterData::CalculateActiveParticles(int32 InRandomSeed
 	{
 		uint32 SpawnInfoTotalParticles = 0;
 
-		const bool bIsValidForAge = !Age.IsSet() || (Age.GetValue() >= SpawnInfo.SpawnTimeStart && Age.GetValue() < SpawnInfo.SpawnTimeEnd + LifetimeRange.Max);
+		const bool bIsValidForAge = !Age.IsSet() || (Age.GetValue() >= SpawnInfo.SpawnTimeStart && Age.GetValue() < SpawnInfo.SpawnTimeEnd + SpawnInfo.LifetimeMax);
 		if (bIsValidForAge && GpuSpawnIndex < NiagaraStateless::MaxGpuSpawnInfos)
 		{
 			uint32	NumActive		= SpawnInfo.Amount;
@@ -102,7 +102,7 @@ uint32 FNiagaraStatelessEmitterData::CalculateActiveParticles(int32 InRandomSeed
 					if (Age.IsSet())
 					{
 						const uint32 MaxActive = SpawnInfo.Amount;
-						ParticleOffset = FMath::FloorToInt(FMath::Max(Age.GetValue() - SpawnInfo.SpawnTimeStart - LifetimeRange.Max, 0.0f) * SpawnInfo.Rate);
+						ParticleOffset = FMath::FloorToInt(FMath::Max(Age.GetValue() - SpawnInfo.SpawnTimeStart - SpawnInfo.LifetimeMax, 0.0f) * SpawnInfo.Rate);
 						ParticleOffset = FMath::Min(ParticleOffset, MaxActive);
 						NumActive = FMath::FloorToInt(FMath::Max(Age.GetValue() - SpawnInfo.SpawnTimeStart, 0.0f) * SpawnInfo.Rate);
 						NumActive = FMath::Min(NumActive, MaxActive);
@@ -125,6 +125,8 @@ uint32 FNiagaraStatelessEmitterData::CalculateActiveParticles(int32 InRandomSeed
 					GET_SCALAR_ARRAY_ELEMENT(SpawnParameters->SpawnInfo_UniqueOffset, GpuSpawnIndex)	= SpawnInfo.UniqueOffset;
 					GET_SCALAR_ARRAY_ELEMENT(SpawnParameters->SpawnInfo_Time, GpuSpawnIndex)			= SpawnTimeStart;
 					GET_SCALAR_ARRAY_ELEMENT(SpawnParameters->SpawnInfo_Rate, GpuSpawnIndex)			= SpawnRate;
+					GET_SCALAR_ARRAY_ELEMENT(SpawnParameters->SpawnInfo_LifetimeScale, GpuSpawnIndex)	= SpawnInfo.LifetimeMax - SpawnInfo.LifetimeMin;
+					GET_SCALAR_ARRAY_ELEMENT(SpawnParameters->SpawnInfo_LifetimeBias, GpuSpawnIndex)	= SpawnInfo.LifetimeMin;
 					++GpuSpawnIndex;
 				}
 				TotalActiveParticles += NumActive;
@@ -141,6 +143,8 @@ uint32 FNiagaraStatelessEmitterData::CalculateActiveParticles(int32 InRandomSeed
 			GET_SCALAR_ARRAY_ELEMENT(SpawnParameters->SpawnInfo_UniqueOffset, GpuSpawnIndex)	= 0;
 			GET_SCALAR_ARRAY_ELEMENT(SpawnParameters->SpawnInfo_Time, GpuSpawnIndex)			= 0.0f;
 			GET_SCALAR_ARRAY_ELEMENT(SpawnParameters->SpawnInfo_Rate, GpuSpawnIndex)			= 0.0f;
+			GET_SCALAR_ARRAY_ELEMENT(SpawnParameters->SpawnInfo_LifetimeScale, GpuSpawnIndex)	= 0.0f;
+			GET_SCALAR_ARRAY_ELEMENT(SpawnParameters->SpawnInfo_LifetimeBias, GpuSpawnIndex)	= 0.0f;
 			++GpuSpawnIndex;
 		}
 	}
diff --git a/Engine/Plugins/FX/Niagara/Source/Niagara/Private/Stateless/NiagaraStatelessEmitterInstance.cpp b/Engine/Plugins/FX/Niagara/Source/Niagara/Private/Stateless/NiagaraStatelessEmitterInstance.cpp
--- a/Engine/Plugins/FX/Niagara/Source/Niagara/Private/Stateless/NiagaraStatelessEmitterInstance.cpp
+++ b/Engine/Plugins/FX/Niagara/Source/Niagara/Private/Stateless/NiagaraStatelessEmitterInstance.cpp
@@ -37,6 +37,21 @@ namespace NiagaraStatelessEmitterInstancePrivate
 		}
 		return RandomStream.RandRange(Distribution.Min, Distribution.Max);
 	}
+
+	bool EvaluateLifetime(const FNiagaraStatelessRangeFloat& Distribution, FRandomStream& RandomStream, const FNiagaraParameterStore& ParameterStore, float& OutMin, float& OutMax)
+	{
+		if (Distribution.ParameterOffset != INDEX_NONE)
+		{
+			FMemory::Memcpy(&OutMin, ParameterStore.GetParameterData(Distribution.ParameterOffset), sizeof(float));
+			OutMax = OutMin;
+		}
+		else
+		{
+			OutMin = Distribution.Min;
+			OutMax = Distribution.Max;
+		}
+		return (OutMin > 0.0f) || (OutMax > 0.0f);
+	}
 }
 
 FNiagaraDataBuffer* NiagaraStateless::FEmitterInstance_RT::GetDataToRender(FRHICommandListBase& RHICmdList, bool bIsLowLatencyTranslucent) const
@@ -618,6 +633,13 @@ void FNiagaraStatelessEmitterInstance::InitSpawnInfosForLoop(float Initializatio
 			continue;
 		}
 
+		float LifetimeMin = 0.0f;
+		float LifetimeMax = 0.0f;
+		if ( EvaluateLifetime(EmitterData->LifetimeRange, RandomStream, RendererBindings, LifetimeMin, LifetimeMax) == false )
+		{
+			continue;
+		}
+
 		const float SpawnAgeStart	= FMath::Min(InitializationAge + CurrentLoopDelay - SpawnInfo.ResidualSpawnTime, CurrentLoopAgeEnd);
 		const float ActiveDuration	= CurrentLoopAgeEnd - SpawnAgeStart;
 		const int32 NumSpawned		= FMath::FloorToInt(ActiveDuration * SpawnRate);
@@ -631,7 +653,12 @@ void FNiagaraStatelessEmitterInstance::InitSpawnInfosForLoop(float Initializatio
 			if ( SpawnInfos.Num() > 0 )
 			{
 				FNiagaraStatelessRuntimeSpawnInfo& ExistingInfo = SpawnInfos.Last();
-				if ( (ExistingInfo.Type == ENiagaraStatelessSpawnInfoType::Rate) && (ExistingInfo.Rate == SpawnRate) && (ExistingInfo.SpawnTimeEnd == SpawnAgeStart))
+				if ( (ExistingInfo.Type == ENiagaraStatelessSpawnInfoType::Rate) &&
+					 (ExistingInfo.Rate == SpawnRate) &&
+					 (ExistingInfo.SpawnTimeEnd == SpawnAgeStart) &&
+					 (ExistingInfo.LifetimeMin == LifetimeMin) &&
+					 (ExistingInfo.LifetimeMax == LifetimeMax)
+				)
 				{
 					if (ExistingInfo.UniqueOffset + ExistingInfo.Amount == UniqueIndexOffset)
 					{
@@ -652,6 +679,8 @@ void FNiagaraStatelessEmitterInstance::InitSpawnInfosForLoop(float Initializatio
 				NewSpawnInfo.SpawnTimeEnd	= SpawnAgeEnd;
 				NewSpawnInfo.Rate			= SpawnRate;
 				NewSpawnInfo.Amount			= NumSpawned;
+				NewSpawnInfo.LifetimeMin	= LifetimeMin;
+				NewSpawnInfo.LifetimeMax	= LifetimeMax;
 			}
 
 			UniqueIndexOffset += NumSpawned;
@@ -692,12 +721,21 @@ void FNiagaraStatelessEmitterInstance::InitSpawnInfosForLoop(float Initializatio
 			continue;
 		}
 
+		float LifetimeMin = 0.0f;
+		float LifetimeMax = 0.0f;
+		if ( EvaluateLifetime(EmitterData->LifetimeRange, RandomStream, RendererBindings, LifetimeMin, LifetimeMax) == false )
+		{
+			continue;
+		}
+
 		FNiagaraStatelessRuntimeSpawnInfo& NewSpawnInfo = SpawnInfos.AddDefaulted_GetRef();
 		NewSpawnInfo.Type			= ENiagaraStatelessSpawnInfoType::Burst;
 		NewSpawnInfo.UniqueOffset	= UniqueIndexOffset;
 		NewSpawnInfo.SpawnTimeStart	= SpawnTime;
 		NewSpawnInfo.SpawnTimeEnd	= SpawnTime;
 		NewSpawnInfo.Amount			= SpawnAmount;
+		NewSpawnInfo.LifetimeMin	= LifetimeMin;
+		NewSpawnInfo.LifetimeMax	= LifetimeMax;
 
 		UniqueIndexOffset += SpawnAmount;
 		bSpawnInfosDirty = true;
@@ -722,11 +760,10 @@ void FNiagaraStatelessEmitterInstance::TickSpawnInfos()
 		}
 	}
 
-	const float MaxLifetime = EmitterData->LifetimeRange.Max;
 	SpawnInfos.RemoveAll(
-		[this, &MaxLifetime](const FNiagaraStatelessRuntimeSpawnInfo& SpawnInfo)
+		[this](const FNiagaraStatelessRuntimeSpawnInfo& SpawnInfo)
 		{
-			return Age >= SpawnInfo.SpawnTimeEnd + MaxLifetime;
+			return Age >= SpawnInfo.SpawnTimeEnd + SpawnInfo.LifetimeMax;
 		}
 	);
 }
@@ -740,7 +777,6 @@ void FNiagaraStatelessEmitterInstance::CropSpawnInfos()
 
 	ActiveSpawnRates.Empty(SpawnInfos.Num());
 
-	const float MaxLifetime = EmitterData->LifetimeRange.Max;
 	for (auto it = SpawnInfos.CreateIterator(); it; ++it)
 	{
 		FNiagaraStatelessRuntimeSpawnInfo& SpawnInfo = *it;
@@ -749,7 +785,7 @@ void FNiagaraStatelessEmitterInstance::CropSpawnInfos()
 			SpawnInfo.SpawnTimeEnd = FMath::Min(SpawnInfo.SpawnTimeEnd, Age);
 			SpawnInfo.Amount = FMath::FloorToInt((SpawnInfo.SpawnTimeEnd - SpawnInfo.SpawnTimeStart) * SpawnInfo.Rate);
 		}
-		if (Age < SpawnInfo.SpawnTimeStart || Age >= SpawnInfo.SpawnTimeEnd + MaxLifetime)
+		if (Age < SpawnInfo.SpawnTimeStart || Age >= SpawnInfo.SpawnTimeEnd + SpawnInfo.LifetimeMax)
 		{
 			it.RemoveCurrent();
 		}
diff --git a/Engine/Plugins/FX/Niagara/Source/Niagara/Private/Stateless/NiagaraStatelessParticleSimContext.cpp b/Engine/Plugins/FX/Niagara/Source/Niagara/Private/Stateless/NiagaraStatelessParticleSimContext.cpp
--- a/Engine/Plugins/FX/Niagara/Source/Niagara/Private/Stateless/NiagaraStatelessParticleSimContext.cpp
+++ b/Engine/Plugins/FX/Niagara/Source/Niagara/Private/Stateless/NiagaraStatelessParticleSimContext.cpp
@@ -169,20 +169,28 @@ namespace NiagaraStateless
 		{
 			uint32 Particle_UniqueIndex = 0;
 			float Particle_Age = -1.0f;
+			float Particle_Lifetime = 0.0f;
 			{
 				uint32 SpawnInfoIndex = iParticle;
 				for (int32 GpuSpawnIndex = 0; GpuSpawnIndex < NiagaraStateless::MaxGpuSpawnInfos; ++GpuSpawnIndex)
 				{
-					const uint32 SpawnInfoNumActive = GET_SCALAR_ARRAY_ELEMENT(SpawnParameters.SpawnInfo_NumActive, GpuSpawnIndex);
-					const uint32 SpawnInfoParticleOffset = GET_SCALAR_ARRAY_ELEMENT(SpawnParameters.SpawnInfo_ParticleOffset, GpuSpawnIndex);
-					const uint32 SpawnInfoUniqueOffset = GET_SCALAR_ARRAY_ELEMENT(SpawnParameters.SpawnInfo_UniqueOffset, GpuSpawnIndex);
-					const float  SpawnInfoTime = GET_SCALAR_ARRAY_ELEMENT(SpawnParameters.SpawnInfo_Time, GpuSpawnIndex);
-					const float  SpawnInfoRate = GET_SCALAR_ARRAY_ELEMENT(SpawnParameters.SpawnInfo_Rate, GpuSpawnIndex);
+					const uint32 SpawnInfoNumActive			= GET_SCALAR_ARRAY_ELEMENT(SpawnParameters.SpawnInfo_NumActive, GpuSpawnIndex);
+					const uint32 SpawnInfoParticleOffset	= GET_SCALAR_ARRAY_ELEMENT(SpawnParameters.SpawnInfo_ParticleOffset, GpuSpawnIndex);
+					const uint32 SpawnInfoUniqueOffset		= GET_SCALAR_ARRAY_ELEMENT(SpawnParameters.SpawnInfo_UniqueOffset, GpuSpawnIndex);
+					const float  SpawnInfoTime				= GET_SCALAR_ARRAY_ELEMENT(SpawnParameters.SpawnInfo_Time, GpuSpawnIndex);
+					const float  SpawnInfoRate				= GET_SCALAR_ARRAY_ELEMENT(SpawnParameters.SpawnInfo_Rate, GpuSpawnIndex);
+					const float  SpawnInfoLifetimeScale		= GET_SCALAR_ARRAY_ELEMENT(SpawnParameters.SpawnInfo_LifetimeScale, GpuSpawnIndex);
+					const float  SpawnInfoLifetimeBias		= GET_SCALAR_ARRAY_ELEMENT(SpawnParameters.SpawnInfo_LifetimeBias, GpuSpawnIndex);
 					if (SpawnInfoIndex < SpawnInfoNumActive)
 					{
 						const uint32 SpawnParticleIndex = SpawnInfoIndex + SpawnInfoParticleOffset;
-						Particle_UniqueIndex = SpawnInfoIndex + SpawnInfoUniqueOffset + SpawnInfoParticleOffset;
-						Particle_Age = EmitterAge - (SpawnInfoTime + float(SpawnParticleIndex) * SpawnInfoRate);
+						Particle_UniqueIndex	= SpawnInfoIndex + SpawnInfoUniqueOffset + SpawnInfoParticleOffset;
+						Particle_Age			= EmitterAge - (SpawnInfoTime + float(SpawnParticleIndex) * SpawnInfoRate);
+
+						// Write unique index as it's needed for the random operation
+						GetParticleUniqueIndex()[NumInstances] = Particle_UniqueIndex;
+
+						Particle_Lifetime		= RandomScaleBiasFloat(NumInstances, 0, SpawnInfoLifetimeScale, SpawnInfoLifetimeBias);
 						break;
 					}
 
@@ -194,10 +202,6 @@ namespace NiagaraStateless
 				}
 			}
 
-			// Write unique index as it's needed for the random operation
-			GetParticleUniqueIndex()[NumInstances] = Particle_UniqueIndex;
-
-			float Particle_Lifetime = RandomScaleBiasFloat(NumInstances, 0, EmitterData->LifetimeRange.GetScale(), EmitterData->LifetimeRange.Min);
 			if (Particle_Lifetime <= 0.0f || Particle_Age >= Particle_Lifetime)
 			{
 				continue;
diff --git a/Engine/Plugins/FX/Niagara/Source/NiagaraShader/Internal/Stateless/NiagaraStatelessSimulationShader.h b/Engine/Plugins/FX/Niagara/Source/NiagaraShader/Internal/Stateless/NiagaraStatelessSimulationShader.h
--- a/Engine/Plugins/FX/Niagara/Source/NiagaraShader/Internal/Stateless/NiagaraStatelessSimulationShader.h
+++ b/Engine/Plugins/FX/Niagara/Source/NiagaraShader/Internal/Stateless/NiagaraStatelessSimulationShader.h
@@ -17,6 +17,8 @@ namespace NiagaraStateless
 		SHADER_PARAMETER_SCALAR_ARRAY(uint32,	SpawnInfo_UniqueOffset, [MaxGpuSpawnInfos])
 		SHADER_PARAMETER_SCALAR_ARRAY(float,	SpawnInfo_Time, [MaxGpuSpawnInfos])
 		SHADER_PARAMETER_SCALAR_ARRAY(float,	SpawnInfo_Rate, [MaxGpuSpawnInfos])
+		SHADER_PARAMETER_SCALAR_ARRAY(float,	SpawnInfo_LifetimeScale, [MaxGpuSpawnInfos])
+		SHADER_PARAMETER_SCALAR_ARRAY(float,	SpawnInfo_LifetimeBias, [MaxGpuSpawnInfos])
 	END_SHADER_PARAMETER_STRUCT()
 
 	BEGIN_SHADER_PARAMETER_STRUCT(FCommonShaderParameters, NIAGARASHADER_API)
@@ -24,7 +26,6 @@ namespace NiagaraStateless
 		SHADER_PARAMETER(float,					Common_SimulationTime)
 		SHADER_PARAMETER(float,					Common_SimulationDeltaTime)
 		SHADER_PARAMETER(float,					Common_SimulationInvDeltaTime)
-		SHADER_PARAMETER(FVector2f,				Common_LifetimeScaleBias)
 
 		SHADER_PARAMETER(uint32,				Common_OutputBufferStride)
 		SHADER_PARAMETER(uint32,				Common_GPUCountBufferOffset)
