Fix issue with NDC impacts on moving objects not working.
Also fixes issue with impacts seeming to follow the player.

Recent fix to avoid creating a new NS inside the activation flow of another meant DI full init was delayed which mean some things were cleaned up too early and thoroughly confused the NDC handler for impacts.
Now preventing cleanup until the handler systems have complete.

diff --git a/Engine/Plugins/FX/Niagara/Source/Niagara/Private/NiagaraDataChannel_Islands.cpp b/Engine/Plugins/FX/Niagara/Source/Niagara/Private/NiagaraDataChannel_Islands.cpp
--- a/Engine/Plugins/FX/Niagara/Source/Niagara/Private/NiagaraDataChannel_Islands.cpp
+++ b/Engine/Plugins/FX/Niagara/Source/Niagara/Private/NiagaraDataChannel_Islands.cpp
@@ -307,14 +307,14 @@ FNDCIsland::FNDCIsland()
 
 FNDCIsland::~FNDCIsland()
 {
-	for (UNiagaraComponent* Comp : NiagaraSystems)
+	for (UNiagaraComponent* Comp : SpawnedComponents)
 	{
 		if (IsValid(Comp))
 		{
 			Comp->ReleaseToPool();
 		}
 	}
-	NiagaraSystems.Empty();
+	SpawnedComponents.Empty();
 }
 
 void FNDCIsland::Init(UNiagaraDataChannelHandler_Islands* InOwner)
@@ -369,7 +369,7 @@ void FNDCIsland::OnAcquired(FVector Location)
 			SpawnParams.WorldContextObject = Owner->GetWorld();
 			if(UNiagaraComponent* NewComp = UNiagaraFunctionLibrary::SpawnSystemAtLocationWithParams(SpawnParams))
 			{
-				NiagaraSystems.Add(NewComp);
+				SpawnedComponents.Add(NewComp);
 				NewComp->SetSystemFixedBounds(HandlerSystemBounds);
 				NewComp->Activate();
 			}
@@ -385,11 +385,11 @@ void FNDCIsland::OnAcquired(FVector Location)
 void FNDCIsland::OnReleased()
 {
 	Data = nullptr;
-	for (UNiagaraComponent* Comp : NiagaraSystems)
+	for (UNiagaraComponent* Comp : SpawnedComponents)
 	{
 		Comp->ReleaseToPool();
 	}
-	NiagaraSystems.Reset();
+	SpawnedComponents.Reset();
 }
 
 void FNDCIsland::BeginFrame()
@@ -407,7 +407,7 @@ void FNDCIsland::Tick(const ETickingGroup& TickGroup)
 	int32 AddedData = Data->ConsumePublishRequests(Owner, TickGroup);
 	if (IsBeingUsed() && AddedData > 0)
 	{
-		for (UNiagaraComponent* Comp : NiagaraSystems)
+		for (UNiagaraComponent* Comp : SpawnedComponents)
 		{
 			if (Comp->IsComplete())
 			{
@@ -439,12 +439,14 @@ bool FNDCIsland::Intersects(const FBoxSphereBounds& CheckBounds)
 
 bool FNDCIsland::IsHandlerSystem(USceneComponent* Component)
 {
-	return NiagaraSystems.Contains(Component);
+	return SpawnedComponents.Contains(Component);
 }
 
 bool FNDCIsland::IsBeingUsed()const
 {
-	return Data.IsUnique() == false;//If someone is making use of our data then we're still in use.
+	bool bSystemsComplete = SpawnedComponents.ContainsByPredicate([](const TObjectPtr<UNiagaraComponent>& Comp) { return !IsValid(Comp) || Comp->IsComplete(); });
+
+	return Data.IsValid() && (!Data.IsUnique() || !bSystemsComplete);
 }
 
 bool FNDCIsland::TryGrow(FVector Point, FVector PerElementExtents, FVector MaxIslandExtents)
@@ -478,7 +480,7 @@ bool FNDCIsland::TryGrow(FVector Point, FVector PerElementExtents, FVector MaxIs
 
 	//If new bounds are valid, grow our bounds and update the 
 	FBox HandlerSystemBounds(-Bounds.BoxExtent, Bounds.BoxExtent);
-	for (auto It = NiagaraSystems.CreateIterator(); It; ++It)
+	for (auto It = SpawnedComponents.CreateIterator(); It; ++It)
 	{
 		UNiagaraComponent* Comp = *It;
 		if (IsValid(Comp))
diff --git a/Engine/Plugins/FX/Niagara/Source/Niagara/Public/NiagaraDataChannel_Islands.h b/Engine/Plugins/FX/Niagara/Source/Niagara/Public/NiagaraDataChannel_Islands.h
--- a/Engine/Plugins/FX/Niagara/Source/Niagara/Public/NiagaraDataChannel_Islands.h
+++ b/Engine/Plugins/FX/Niagara/Source/Niagara/Public/NiagaraDataChannel_Islands.h
@@ -73,8 +73,8 @@ private:
 	FBoxSphereBounds Bounds = FBoxSphereBounds(EForceInit::ForceInit);
 
 	/** Niagara components spawned for this island. */
-	UPROPERTY()
-	TArray<TObjectPtr<UNiagaraComponent>> NiagaraSystems;
+	UPROPERTY(Transient)
+	TArray<TObjectPtr<UNiagaraComponent>> SpawnedComponents;
 	
 	/** The underlying storage for this island. */
 	FNiagaraDataChannelDataPtr Data = nullptr;
