Fix loctext conflict

diff --git a/Engine/Plugins/FX/Niagara/Source/NiagaraEditor/Private/Widgets/SNiagaraDebugger.cpp b/Engine/Plugins/FX/Niagara/Source/NiagaraEditor/Private/Widgets/SNiagaraDebugger.cpp
--- a/Engine/Plugins/FX/Niagara/Source/NiagaraEditor/Private/Widgets/SNiagaraDebugger.cpp
+++ b/Engine/Plugins/FX/Niagara/Source/NiagaraEditor/Private/Widgets/SNiagaraDebugger.cpp
@@ -342,8 +342,8 @@ namespace NiagaraOutlinerTab
 						FCanExecuteAction(),
 						FIsActionChecked::CreateLambda([Debugger]() { return Debugger->IsCapturingOutlinerOnKeyPress(); })),
 					NAME_None,
-					LOCTEXT("Capture", "Capture (Press 'C')"),
-					LOCTEXT("CaptureTooltip", "Triggers a capture of outliner info from the connected session on a key press."),
+					LOCTEXT("CaptureOnKeyPress", "Capture (Press 'C')"),
+					LOCTEXT("CaptureOnKeyPressTooltip", "Triggers a capture of outliner info from the connected session on a key press."),
 					FSlateIcon(FNiagaraEditorStyle::Get().GetStyleSetName(), "NiagaraEditor.Debugger.Outliner.Capture"),
 					EUserInterfaceActionType::ToggleButton
 				);
