Cleanup RF_Public usage for NiagaraDataInterfaces.  Removes a number of duplicated/unreachable objects from cooked builds lowering memory consumption.  Tries to unify the evaluation of whether we should mark a DI as public through the static UNiagaraDataInterface::BuildObjectFlagsForOwner() which currently will mark the DI as public if the owner is public or it's a template.

diff --git a/Engine/Plugins/FX/Niagara/Source/Niagara/Classes/NiagaraDataInterface.h b/Engine/Plugins/FX/Niagara/Source/Niagara/Classes/NiagaraDataInterface.h
--- a/Engine/Plugins/FX/Niagara/Source/Niagara/Classes/NiagaraDataInterface.h
+++ b/Engine/Plugins/FX/Niagara/Source/Niagara/Classes/NiagaraDataInterface.h
@@ -576,11 +576,13 @@ public:
 	NIAGARA_API virtual ~UNiagaraDataInterface() override;
 
 	// UObject Interface
-	NIAGARA_API virtual void PostLoad() override;
 #if WITH_EDITOR
 	NIAGARA_API virtual void PostEditChangeProperty(struct FPropertyChangedEvent& PropertyChangedEvent) override;
+#endif
+	NIAGARA_API virtual void Serialize(FStructuredArchive::FRecord Record) override;
 	// UObject Interface END
 
+#if WITH_EDITOR
 	/** Does this data interface need setup and teardown for each stage when working a sim stage sim source? */
 	virtual bool SupportsSetupAndTeardownHLSL() const { return false; }
 	/** Generate the necessary HLSL to set up data when being added as a sim stage sim source. */
@@ -848,6 +850,8 @@ public:
 	*/
 	virtual void SetShaderParameters(const FNiagaraDataInterfaceSetShaderParametersContext& Context) const {}
 
+	static NIAGARA_API EObjectFlags BuildObjectFlagsForOwner(UObject* OwnerObject, EObjectFlags DefaultFlags = RF_NoFlags);
+
 protected:
 	virtual void PushToRenderThreadImpl() {}
 
diff --git a/Engine/Plugins/FX/Niagara/Source/Niagara/Private/NiagaraComponent.cpp b/Engine/Plugins/FX/Niagara/Source/Niagara/Private/NiagaraComponent.cpp
--- a/Engine/Plugins/FX/Niagara/Source/Niagara/Private/NiagaraComponent.cpp
+++ b/Engine/Plugins/FX/Niagara/Source/Niagara/Private/NiagaraComponent.cpp
@@ -2987,6 +2987,7 @@ void UNiagaraComponent::Serialize(FStructuredArchive::FRecord Record)
 	if (Record.GetUnderlyingArchive().IsLoading())
 	{
 		FixDataInterfaceOuters();
+		FixDataInterfaceObjectFlags();
 
 		// Check the the event handler here as well as in post load since post load isn't always called due to how components are managed when spawning
 		// actors with templates.
@@ -3453,7 +3454,9 @@ void UNiagaraComponent::ApplyOverridesToParameterStore()
 
 void UNiagaraComponent::FixDataInterfaceOuters()
 {
-	auto FixParameterOverrides = [this](TMap<FNiagaraVariableBase, FNiagaraVariant>& ParameterOverrides, TMap<FNiagaraVariableBase, UNiagaraDataInterface*>& OutFixedDataInterfaces)
+	const EObjectFlags DataInterfaceObjectFlags = UNiagaraDataInterface::BuildObjectFlagsForOwner(this, RF_Transactional);
+
+	auto FixParameterOverrides = [this, DataInterfaceObjectFlags](TMap<FNiagaraVariableBase, FNiagaraVariant>& ParameterOverrides, TMap<FNiagaraVariableBase, UNiagaraDataInterface*>& OutFixedDataInterfaces)
 	{
 		for (TPair<FNiagaraVariableBase, FNiagaraVariant>& VariableValuePair : ParameterOverrides)
 		{
@@ -3461,7 +3464,7 @@ void UNiagaraComponent::FixDataInterfaceOuters()
 			{
 				if (VariableValuePair.Value.GetDataInterface()->GetOuter() != this)
 				{
-					UNiagaraDataInterface* FixedDataInterface = NewObject<UNiagaraDataInterface>(this, VariableValuePair.Value.GetDataInterface()->GetClass(), NAME_None, RF_Transactional | RF_Public);
+					UNiagaraDataInterface* FixedDataInterface = NewObject<UNiagaraDataInterface>(this, VariableValuePair.Value.GetDataInterface()->GetClass(), NAME_None, DataInterfaceObjectFlags);
 					VariableValuePair.Value.GetDataInterface()->CopyTo(FixedDataInterface);
 					VariableValuePair.Value.SetDataInterface(FixedDataInterface);
 					OutFixedDataInterfaces.Add(VariableValuePair.Key, FixedDataInterface);
@@ -3491,13 +3494,62 @@ void UNiagaraComponent::FixDataInterfaceOuters()
 		UNiagaraDataInterface* OverrideParameterDataInterface = OverrideParameters.GetDataInterface(i);
 		if (OverrideParameterDataInterface != nullptr && OverrideParameterDataInterface->GetOuter() != this)
 		{
-			UNiagaraDataInterface* FixedDataInterface = NewObject<UNiagaraDataInterface>(this, OverrideParameterDataInterface->GetClass(), NAME_None, RF_Transactional | RF_Public);
+			UNiagaraDataInterface* FixedDataInterface = NewObject<UNiagaraDataInterface>(this, OverrideParameterDataInterface->GetClass(), NAME_None, DataInterfaceObjectFlags);
 			OverrideParameterDataInterface->CopyTo(FixedDataInterface);
 			OverrideParameters.SetDataInterface(FixedDataInterface, i);
 		}
 	}
 }
 
+namespace NiagaraComponentImpl
+{
+
+void SetDataInterfaceFlags(UNiagaraDataInterface* DataInterface, bool bPublicDataInterfaces)
+{
+	if (DataInterface != nullptr)
+	{
+		if (DataInterface->HasAnyFlags(RF_Public) != bPublicDataInterfaces)
+		{
+			if (bPublicDataInterfaces)
+			{
+				DataInterface->SetFlags(RF_Public);
+			}
+			else
+			{
+				DataInterface->ClearFlags(RF_Public);
+			}
+		}
+	}
+}
+
+};
+
+void UNiagaraComponent::FixDataInterfaceObjectFlags()
+{
+	const bool bPublicDataInterfaces = !!(UNiagaraDataInterface::BuildObjectFlagsForOwner(this) & RF_Public);
+
+	auto FixParameterOverrides = [](TMap<FNiagaraVariableBase, FNiagaraVariant>& ParameterOverrides, bool bPublicDataInterfaces)
+	{
+		for (TPair<FNiagaraVariableBase, FNiagaraVariant>& VariableValuePair : ParameterOverrides)
+		{
+			if (VariableValuePair.Key.IsDataInterface())
+			{
+				NiagaraComponentImpl::SetDataInterfaceFlags(VariableValuePair.Value.GetDataInterface(), bPublicDataInterfaces);
+			}
+		}
+	};
+
+	// Fix data interfaces in the template and instance overrides.
+	FixParameterOverrides(TemplateParameterOverrides, bPublicDataInterfaces);
+	FixParameterOverrides(InstanceParameterOverrides, bPublicDataInterfaces);
+
+	// Fix any additional data interfaces in the override parameters which are still incorrect.
+	for (int32 i = 0; i < OverrideParameters.Num(); i++)
+	{
+		NiagaraComponentImpl::SetDataInterfaceFlags(OverrideParameters.GetDataInterface(i), bPublicDataInterfaces);
+	}
+}
+
 #endif
 
 void UNiagaraComponent::CopyParametersFromAsset(bool bResetExistingOverrideParameters)
@@ -3534,6 +3586,8 @@ void UNiagaraComponent::CopyParametersFromAsset(bool bResetExistingOverrideParam
 	}
 
 	// DataInterfaces are unique per component so we need to make new instances if they are not already
+	const EObjectFlags DataInterfaceObjectFlags = UNiagaraDataInterface::BuildObjectFlagsForOwner(this, RF_Transactional);
+
 	for (int32 i = 0; i < AssetExposedParameters.GetDataInterfaces().Num(); ++i)
 	{
 		UNiagaraDataInterface* AssetDataInterface = AssetExposedParameters.GetDataInterface(i);
@@ -3552,7 +3606,7 @@ void UNiagaraComponent::CopyParametersFromAsset(bool bResetExistingOverrideParam
 		}
 		else
 		{
-			OverrideDataInterface = NewObject<UNiagaraDataInterface>(this, AssetDataInterface->GetClass(), NAME_None, RF_Transactional | RF_Public);
+			OverrideDataInterface = NewObject<UNiagaraDataInterface>(this, AssetDataInterface->GetClass(), NAME_None, DataInterfaceObjectFlags);
 			AssetDataInterface->CopyTo(OverrideDataInterface);
 			OverrideParameters.SetDataInterface(OverrideDataInterface, i);
 		}
@@ -3593,6 +3647,8 @@ void UNiagaraComponent::SynchronizeWithSourceSystem()
 #if WITH_EDITORONLY_DATA
 void FixInvalidDataInterfaceOverrides(TMap<FNiagaraVariableBase, FNiagaraVariant>& ParameterOverrides, const FString& OverrideSource, UNiagaraComponent* OwningComponent)
 {
+	const EObjectFlags DataInterfaceObjectFlags = UNiagaraDataInterface::BuildObjectFlagsForOwner(OwningComponent, RF_Transactional);
+
 	for (TPair<FNiagaraVariableBase, FNiagaraVariant>& VariableValuePair : ParameterOverrides)
 	{
 		FNiagaraVariableBase& Variable = VariableValuePair.Key;
@@ -3611,7 +3667,7 @@ void FixInvalidDataInterfaceOverrides(TMap<FNiagaraVariableBase, FNiagaraVariant
 			}
 			if (DataInterface == nullptr)
 			{
-				DataInterface = NewObject<UNiagaraDataInterface>(OwningComponent, Variable.GetType().GetClass(), NAME_None, RF_Transactional | RF_Public);
+				DataInterface = NewObject<UNiagaraDataInterface>(OwningComponent, Variable.GetType().GetClass(), NAME_None, DataInterfaceObjectFlags);
 				if (UNiagaraSystem* NiagaraSystem = OwningComponent->GetAsset())
 				{
 					const int32* DataInterfaceOffset = NiagaraSystem->GetExposedParameters().FindParameterOffset(Variable);
@@ -3629,6 +3685,8 @@ void FixInvalidDataInterfaceOverrides(TMap<FNiagaraVariableBase, FNiagaraVariant
 
 void UNiagaraComponent::FixInvalidUserParameterOverrideData()
 {
+	const EObjectFlags DataInterfaceObjectFlags = UNiagaraDataInterface::BuildObjectFlagsForOwner(this, RF_Transactional);
+
 	FixInvalidDataInterfaceOverrides(TemplateParameterOverrides, TEXT("Template"), this);
 	FixInvalidDataInterfaceOverrides(InstanceParameterOverrides, TEXT("Instance"), this);
 
@@ -3683,7 +3741,7 @@ void UNiagaraComponent::FixInvalidUserParameterOverrideData()
 
 			if (DataInterface == nullptr)
 			{
-				DataInterface = NewObject<UNiagaraDataInterface>(this, ExistingVar.GetType().GetClass(), NAME_None, RF_Transactional | RF_Public);
+				DataInterface = NewObject<UNiagaraDataInterface>(this, ExistingVar.GetType().GetClass(), NAME_None, DataInterfaceObjectFlags);
 				if ( SourceVars.Contains(ExistingVar) )
 				{
 					if ( UNiagaraDataInterface* SourceDataInterface = Asset->GetExposedParameters().GetDataInterface(ExistingVar) )
diff --git a/Engine/Plugins/FX/Niagara/Source/Niagara/Private/NiagaraDataInterface.cpp b/Engine/Plugins/FX/Niagara/Source/Niagara/Private/NiagaraDataInterface.cpp
--- a/Engine/Plugins/FX/Niagara/Source/Niagara/Private/NiagaraDataInterface.cpp
+++ b/Engine/Plugins/FX/Niagara/Source/Niagara/Private/NiagaraDataInterface.cpp
@@ -163,12 +163,6 @@ void UNiagaraDataInterface::GetAssetTagsForContext(const UObject* InAsset,  FGui
 	}
 }
 
-void UNiagaraDataInterface::PostLoad()
-{
-	Super::PostLoad();
-	SetFlags(RF_Public);
-}
-
 #if WITH_EDITOR
 void UNiagaraDataInterface::PostEditChangeProperty(struct FPropertyChangedEvent& PropertyChangedEvent)
 {
@@ -177,6 +171,24 @@ void UNiagaraDataInterface::PostEditChangeProperty(struct FPropertyChangedEvent&
 }
 #endif
 
+void UNiagaraDataInterface::Serialize(FStructuredArchive::FRecord Record)
+{
+	Super::Serialize(Record);
+#if WITH_EDITORONLY_DATA
+	if (Record.GetUnderlyingArchive().IsLoading())
+	{
+		if (BuildObjectFlagsForOwner(GetOuter()) & RF_Public)
+		{
+			SetFlags(RF_Public);
+		}
+		else
+		{
+			ClearFlags(RF_Public);
+		}
+	}
+#endif
+}
+
 bool UNiagaraDataInterface::CopyTo(UNiagaraDataInterface* Destination) const 
 {
 	bool result = CopyToInternal(Destination);
@@ -326,6 +338,20 @@ FSimpleMulticastDelegate& UNiagaraDataInterface::OnErrorsRefreshed()
 
 #endif
 
+// DataInterfaces owned by public objects or templates will be made public.  We might be able to further
+// refine this, but for now it provides a significant reduction in public DI which ensures that DI created
+// in the editor that are unnecessary in game will be stripped during the cook process and reduce our static
+// memory footprint.
+EObjectFlags UNiagaraDataInterface::BuildObjectFlagsForOwner(UObject* OwnerObject, EObjectFlags DefaultFlags)
+{
+	if (OwnerObject && (OwnerObject->HasAnyFlags(RF_Public) || OwnerObject->IsTemplate()))
+	{
+		return DefaultFlags | RF_Public;
+	}
+
+	return DefaultFlags;
+}
+
 //////////////////////////////////////////////////////////////////////////
 //FNDIGpuComputeContext
 
diff --git a/Engine/Plugins/FX/Niagara/Source/Niagara/Private/NiagaraParameterStore.cpp b/Engine/Plugins/FX/Niagara/Source/Niagara/Private/NiagaraParameterStore.cpp
--- a/Engine/Plugins/FX/Niagara/Source/Niagara/Private/NiagaraParameterStore.cpp
+++ b/Engine/Plugins/FX/Niagara/Source/Niagara/Private/NiagaraParameterStore.cpp
@@ -698,7 +698,8 @@ bool FNiagaraParameterStore::AddParameter(const FNiagaraVariable& Param, bool bI
 	if (ParamType.IsDataInterface())
 	{
 		Offset = DataInterfaces.AddZeroed();
-		DataInterfaces[Offset] = bInitInterfaces ? NewObject<UNiagaraDataInterface>(Owner.Get(), const_cast<UClass*>(ParamType.GetClass()), NAME_None, RF_Transactional | RF_Public) : nullptr;
+		const EObjectFlags DataInterfaceObjectFlags = UNiagaraDataInterface::BuildObjectFlagsForOwner(Owner.Get(), RF_Transactional);
+		DataInterfaces[Offset] = bInitInterfaces ? NewObject<UNiagaraDataInterface>(Owner.Get(), const_cast<UClass*>(ParamType.GetClass()), NAME_None, DataInterfaceObjectFlags) : nullptr;
 		bInterfacesDirty = true;
 	}
 	else if (ParamType.IsUObject())
@@ -937,7 +938,8 @@ void FNiagaraParameterStore::SanityCheckData(bool bInitInterfaces)
 				}
 				if (DataInterfaces[SrcIndex] == nullptr && bInitInterfaces && Owner.IsValid())
 				{
-					DataInterfaces[SrcIndex] = NewObject<UNiagaraDataInterface>(Owner.Get(), const_cast<UClass*>(Parameter.GetType().GetClass()), NAME_None, RF_Transactional | RF_Public);
+					const EObjectFlags DataInterfaceObjectFlags = UNiagaraDataInterface::BuildObjectFlagsForOwner(Owner.Get(), RF_Transactional);
+					DataInterfaces[SrcIndex] = NewObject<UNiagaraDataInterface>(Owner.Get(), const_cast<UClass*>(Parameter.GetType().GetClass()), NAME_None, DataInterfaceObjectFlags);
 					UE_LOG(LogNiagara, Verbose, TEXT("Had to initialize data interface! %s on %s"), *Parameter.GetName().ToString(), Owner != nullptr ? *Owner->GetPathName() : TEXT("Unknown owner"));
 
 					OwnerDirtied = true;
diff --git a/Engine/Plugins/FX/Niagara/Source/Niagara/Private/NiagaraScript.cpp b/Engine/Plugins/FX/Niagara/Source/Niagara/Private/NiagaraScript.cpp
--- a/Engine/Plugins/FX/Niagara/Source/Niagara/Private/NiagaraScript.cpp
+++ b/Engine/Plugins/FX/Niagara/Source/Niagara/Private/NiagaraScript.cpp
@@ -2942,7 +2942,8 @@ UNiagaraDataInterface* UNiagaraScript::CopyDataInterface(UNiagaraDataInterface*
 {
 	if (Src)
 	{
-		UNiagaraDataInterface* DI = NewObject<UNiagaraDataInterface>(Owner, const_cast<UClass*>(Src->GetClass()), NAME_None, RF_Transactional | RF_Public);
+		const EObjectFlags DataInterfaceObjectFlags = UNiagaraDataInterface::BuildObjectFlagsForOwner(Owner, RF_Transactional);
+		UNiagaraDataInterface* DI = NewObject<UNiagaraDataInterface>(Owner, const_cast<UClass*>(Src->GetClass()), NAME_None, DataInterfaceObjectFlags);
 		Src->CopyTo(DI);
 		return DI;
 	}
diff --git a/Engine/Plugins/FX/Niagara/Source/Niagara/Public/NiagaraComponent.h b/Engine/Plugins/FX/Niagara/Source/Niagara/Public/NiagaraComponent.h
--- a/Engine/Plugins/FX/Niagara/Source/Niagara/Public/NiagaraComponent.h
+++ b/Engine/Plugins/FX/Niagara/Source/Niagara/Public/NiagaraComponent.h
@@ -699,6 +699,8 @@ private:
 	 * subobject handling on components when they're spawned from templates which can cause invalid cross package references.
  	 */
 	NIAGARA_API void FixDataInterfaceOuters();
+
+	void FixDataInterfaceObjectFlags();
 #endif
 
 public:
diff --git a/Engine/Plugins/FX/Niagara/Source/NiagaraEditor/Private/NiagaraEditorUtilities.cpp b/Engine/Plugins/FX/Niagara/Source/NiagaraEditor/Private/NiagaraEditorUtilities.cpp
--- a/Engine/Plugins/FX/Niagara/Source/NiagaraEditor/Private/NiagaraEditorUtilities.cpp
+++ b/Engine/Plugins/FX/Niagara/Source/NiagaraEditor/Private/NiagaraEditorUtilities.cpp
@@ -168,7 +168,7 @@ void FNiagaraEditorUtilities::InitializeParameterInputNode(UNiagaraNodeInput& In
 	else if(Type.IsDataInterface())
 	{
 		InputNode.Input.AllocateData(); // Frees previously used memory if we're switching from a struct to a class type.
-		InputNode.SetDataInterface(NewObject<UNiagaraDataInterface>(&InputNode, Type.GetClass(), NAME_None, RF_Transactional | RF_Public));
+		InputNode.SetDataInterface(NewObject<UNiagaraDataInterface>(&InputNode, Type.GetClass(), NAME_None, UNiagaraDataInterface::BuildObjectFlagsForOwner(&InputNode, RF_Transactional)));
 	}
 	else if (Type.IsUObject())
 	{
diff --git a/Engine/Plugins/FX/Niagara/Source/NiagaraEditor/Private/NiagaraNodeInput.cpp b/Engine/Plugins/FX/Niagara/Source/NiagaraEditor/Private/NiagaraNodeInput.cpp
--- a/Engine/Plugins/FX/Niagara/Source/NiagaraEditor/Private/NiagaraNodeInput.cpp
+++ b/Engine/Plugins/FX/Niagara/Source/NiagaraEditor/Private/NiagaraNodeInput.cpp
@@ -146,7 +146,8 @@ void UNiagaraNodeInput::AllocateDefaultPins()
 	{
 		if (DataInterface == nullptr)
 		{
-			DataInterface = NewObject<UNiagaraDataInterface>(this, InputType.GetClass(), NAME_None, RF_Transactional | RF_Public);
+			const EObjectFlags DataInterfaceObjectFlags = UNiagaraDataInterface::BuildObjectFlagsForOwner(this, RF_Transactional);
+			DataInterface = NewObject<UNiagaraDataInterface>(this, InputType.GetClass(), NAME_None, DataInterfaceObjectFlags);
 		}
 	}
 
diff --git a/Engine/Plugins/FX/Niagara/Source/NiagaraEditor/Private/ViewModels/Stack/NiagaraStackGraphUtilities.cpp b/Engine/Plugins/FX/Niagara/Source/NiagaraEditor/Private/ViewModels/Stack/NiagaraStackGraphUtilities.cpp
--- a/Engine/Plugins/FX/Niagara/Source/NiagaraEditor/Private/ViewModels/Stack/NiagaraStackGraphUtilities.cpp
+++ b/Engine/Plugins/FX/Niagara/Source/NiagaraEditor/Private/ViewModels/Stack/NiagaraStackGraphUtilities.cpp
@@ -2177,7 +2177,8 @@ void FNiagaraStackGraphUtilities::SetDataInterfaceValueForFunctionInput(UEdGraph
 	UNiagaraNodeInput* InputNode = InputNodeCreator.CreateNode();
 	FNiagaraEditorUtilities::InitializeParameterInputNode(*InputNode, FNiagaraTypeDefinition(DataObjectType), CastChecked<UNiagaraGraph>(Graph), *InputNodeInputName);
 
-	OutDataObject = NewObject<UNiagaraDataInterface>(InputNode, DataObjectType, *ObjectTools::SanitizeObjectName(InputNodeInputName), RF_Transactional | RF_Public);
+	const EObjectFlags DataInterfaceObjectFlags = UNiagaraDataInterface::BuildObjectFlagsForOwner(InputNode, RF_Transactional);
+	OutDataObject = NewObject<UNiagaraDataInterface>(InputNode, DataObjectType, *ObjectTools::SanitizeObjectName(InputNodeInputName), DataInterfaceObjectFlags);
 	InputNode->SetDataInterface(OutDataObject);
 
 	InputNodeCreator.Finalize();
