- Allow desired age modes to run async when possible
- Can be disabled using the console variable fx.Niagara.AllowDesiredAgeAsync

diff --git a/Engine/Plugins/FX/Niagara/Source/Niagara/Private/NiagaraComponent.cpp b/Engine/Plugins/FX/Niagara/Source/Niagara/Private/NiagaraComponent.cpp
--- a/Engine/Plugins/FX/Niagara/Source/Niagara/Private/NiagaraComponent.cpp
+++ b/Engine/Plugins/FX/Niagara/Source/Niagara/Private/NiagaraComponent.cpp
@@ -72,6 +72,14 @@ static FAutoConsoleVariableRef CVarNiagaraSoloAllowAsyncWorkToEndOfFrame(
 	ECVF_Default
 );
 
+static bool GNiagaraAllowDesiredAgeAsync = true;
+static FAutoConsoleVariableRef CVarNiagaraAllowDesiredAgeAsync(
+	TEXT("fx.Niagara.AllowDesiredAgeAsync"),
+	GNiagaraAllowDesiredAgeAsync,
+	TEXT("When enabled Desired Age / Desired Age No Seek tick modes will be allows to run async if we single tick in a frame, multi tick will still run on the game thread."),
+	ECVF_Default
+);
+
 static bool GbNiagaraComponentsEnabled = true;
 static FAutoConsoleVariableRef CVarNiagaraComponentsEnabled(
 	TEXT("fx.NiagaraComponentsEnabled"),
@@ -942,6 +950,8 @@ void UNiagaraComponent::TickComponent(float DeltaSeconds, enum ELevelTick TickTy
 				FNiagaraSystemSimulation* SystemSim = SystemInstanceController->GetSoloSystemSimulation().Get();
 				if (SystemSim)
 				{
+					FGraphEventRef CompletionEvent = (GNiagaraAllowDesiredAgeAsync && ThisTickFunction && ThisTickFunction->IsCompletionHandleValid()) ? ThisTickFunction->GetCompletionHandle() : nullptr;
+
 					const bool bFixedTickDelta = Asset->HasFixedTickDelta();
 					const float TickDelta = bFixedTickDelta ? Asset->GetFixedTickDeltaTime() : SeekDelta;
 					if (bLockDesiredAgeDeltaTimeToSeekDelta || AgeDiff > TickDelta)
@@ -951,11 +961,13 @@ void UNiagaraComponent::TickComponent(float DeltaSeconds, enum ELevelTick TickTy
 						const double EndMaxSimTime = bUseMaxSimTime ? FPlatformTime::Seconds() + MaxSimTime : 0.0f;
 
 						TicksToProcess = FMath::FloorToInt(AgeDiff / TickDelta);
+						CompletionEvent = TicksToProcess == 1 ? CompletionEvent : nullptr;
+
 						while ( TicksToProcess > 0 )
 						{
 							// Cannot do multiple tick off the game thread here without additional work. So we pass in null for the completion event which will force GT execution.
 							// If this becomes a perf problem I can add a new path for the tick code to handle multiple ticks.
-							SystemInstanceController->ManualTick(TickDelta, nullptr);
+							SystemInstanceController->ManualTick(TickDelta, CompletionEvent);
 							if (IsComplete())
 							{
 								break;
@@ -978,7 +990,7 @@ void UNiagaraComponent::TickComponent(float DeltaSeconds, enum ELevelTick TickTy
 						// Otherwise just tick by the age difference.
 						if (!bFixedTickDelta || AgeDiff >= TickDelta)
 						{
-							SystemInstanceController->ManualTick(AgeDiff, nullptr);
+							SystemInstanceController->ManualTick(AgeDiff, CompletionEvent);
 						}
 					}
 				}
@@ -991,6 +1003,8 @@ void UNiagaraComponent::TickComponent(float DeltaSeconds, enum ELevelTick TickTy
 		}
 		else if (AgeUpdateMode == ENiagaraAgeUpdateMode::DesiredAgeNoSeek)
 		{
+			FGraphEventRef CompletionEvent = (GNiagaraAllowDesiredAgeAsync && ThisTickFunction && ThisTickFunction->IsCompletionHandleValid()) ? ThisTickFunction->GetCompletionHandle() : nullptr;
+
 			int32 MaxForwardFrames = 5; // HACK - for some reason sequencer jumps forwards by multiple frames on pause, so this is being added to allow for FX to stay alive when being controlled by sequencer in the editor.  This should be lowered once that issue is fixed.
 			float AgeDiff = DesiredAge - LastHandledDesiredAge;
 			float TickDelta = Asset->HasFixedTickDelta() ? Asset->GetFixedTickDeltaTime() : SeekDelta;
@@ -1001,19 +1015,19 @@ void UNiagaraComponent::TickComponent(float DeltaSeconds, enum ELevelTick TickTy
 					// When going back in time for a frame or more, reset and simulate a single frame.  We ignore small negative changes to delta
 					// time which can happen when controlling time with the timeline and the time snaps to a previous time when paused.
 					SystemInstanceController->Reset(FNiagaraSystemInstance::EResetMode::ResetAll);
-					SystemInstanceController->ManualTick(TickDelta, nullptr);
+					SystemInstanceController->ManualTick(TickDelta, CompletionEvent);
 				}
 			}
 			else if (AgeDiff < MaxForwardFrames * TickDelta)
 			{
 				// Allow ticks between 0 and MaxForwardFrames, but don't ever send more then 2 x the seek delta.
-				SystemInstanceController->ManualTick(FMath::Min(AgeDiff, 2 * TickDelta), nullptr);
+				SystemInstanceController->ManualTick(FMath::Min(AgeDiff, 2 * TickDelta), CompletionEvent);
 			}
 			else
 			{
 				// When going forward in time for more than MaxForwardFrames, reset and simulate a single frame.
 				SystemInstanceController->Reset(FNiagaraSystemInstance::EResetMode::ResetAll);
-				SystemInstanceController->ManualTick(TickDelta, nullptr);
+				SystemInstanceController->ManualTick(TickDelta, CompletionEvent);
 			}
 			LastHandledDesiredAge = DesiredAge;
 		}
