- Remove standalone / public flag when merging emitters
- Fixes an issue where it was impossible to delete a system

#rnx

diff --git a/Engine/Plugins/FX/Niagara/Source/Niagara/Classes/NiagaraEmitter.h b/Engine/Plugins/FX/Niagara/Source/Niagara/Classes/NiagaraEmitter.h
--- a/Engine/Plugins/FX/Niagara/Source/Niagara/Classes/NiagaraEmitter.h
+++ b/Engine/Plugins/FX/Niagara/Source/Niagara/Classes/NiagaraEmitter.h
@@ -908,7 +908,10 @@ public:
 	/** Merges in any changes from the parent emitter into this emitter. */
 	NIAGARA_API TArray<INiagaraMergeManager::FMergeEmitterResults> MergeChangesFromParent();	
 
-	/** Duplicates this emitter, but prevents the duplicate from merging in changes from the parent emitter.  The resulting duplicate will have no parent information. */
+	/**
+	Duplicates this emitter, but prevents the duplicate from merging in changes from the parent emitter.
+	The resulting duplicate will have no parent information and will clear RF_Standalone | RF_Public flags.
+	*/
 	NIAGARA_API UNiagaraEmitter* DuplicateWithoutMerging(UObject* InOuter);
 
 	NIAGARA_API void SetEditorData(UNiagaraEditorDataBase* InEditorData, const FGuid& VersionGuid);
diff --git a/Engine/Plugins/FX/Niagara/Source/Niagara/Private/NiagaraEmitter.cpp b/Engine/Plugins/FX/Niagara/Source/Niagara/Private/NiagaraEmitter.cpp
--- a/Engine/Plugins/FX/Niagara/Source/Niagara/Private/NiagaraEmitter.cpp
+++ b/Engine/Plugins/FX/Niagara/Source/Niagara/Private/NiagaraEmitter.cpp
@@ -397,7 +397,6 @@ TArray<INiagaraMergeManager::FMergeEmitterResults> UNiagaraEmitter::MergeChanges
 
 			// Update the last merged source and clear it's stand alone and public flags since it's not an asset.
 			VersionedParentAtLastMerge.Emitter = VersionedParent.Emitter->DuplicateWithoutMerging(this);
-			VersionedParentAtLastMerge.Emitter->ClearFlags(RF_Standalone | RF_Public);
 			VersionedParentAtLastMerge.Emitter->DisableVersioning(VersionedParent.Version);
 			VersionedParentAtLastMerge.Version = VersionedParent.Version;
 		}
@@ -470,6 +469,7 @@ UNiagaraEmitter* UNiagaraEmitter::DuplicateWithoutMerging(UObject* InOuter)
 	{
 		TParentGuard ParentGuard(VersionData);
 		Duplicate = Cast<UNiagaraEmitter>(StaticDuplicateObject(this, InOuter));
+		Duplicate->ClearFlags(RF_Standalone | RF_Public);
 	}
 	return Duplicate;
 }
@@ -3670,7 +3670,6 @@ void UNiagaraEmitter::SetParent(const FVersionedNiagaraEmitter& InParent)
 		FVersionedNiagaraEmitterData& EmitterData = VersionData[0];
 		EmitterData.VersionedParent = InParent;
 		EmitterData.VersionedParentAtLastMerge.Emitter = InParent.Emitter->DuplicateWithoutMerging(this);
-		EmitterData.VersionedParentAtLastMerge.Emitter->ClearFlags(RF_Standalone | RF_Public);
 		EmitterData.VersionedParentAtLastMerge.Emitter->DisableVersioning(InParent.Version);
 		EmitterData.VersionedParentAtLastMerge.Version = InParent.Version;
 
