Changed module dependency insertion to always target the end of the script stage, if the dependency is across stages

Previously, PostDependency requirements would be added at Index 0, which would cause modules to be added before InitializeParticle or ParticleState etc.

diff --git a/Engine/Plugins/FX/Niagara/Source/NiagaraEditor/Private/ViewModels/Stack/NiagaraStackModuleItem.cpp b/Engine/Plugins/FX/Niagara/Source/NiagaraEditor/Private/ViewModels/Stack/NiagaraStackModuleItem.cpp
--- a/Engine/Plugins/FX/Niagara/Source/NiagaraEditor/Private/ViewModels/Stack/NiagaraStackModuleItem.cpp
+++ b/Engine/Plugins/FX/Niagara/Source/NiagaraEditor/Private/ViewModels/Stack/NiagaraStackModuleItem.cpp
@@ -387,6 +387,8 @@ namespace NiagaraStackModuleItemIssues {
 		else
 		{
 			// Otherwise we need to search for a compatible script to insert the module into.
+			// We always use INDEX_NONE as target index (add to end), because PreDependency and PostDependency across stages doesn't mean much
+			// If a module in ParticleSpawn requires a module in ParticleUpdate, only other ParticleUpdate scripts are relevant for its placement
 			TArray<UNiagaraScript*> OrderedScripts;
 			DependentSystemViewModel->GetOrderedScriptsForEmitterHandleId(EmitterHandleId, OrderedScripts);
 			int32 DependentScriptIndex = OrderedScripts.IndexOfByPredicate([&LastDependentModuleData](UNiagaraScript* Script)
@@ -409,7 +411,7 @@ namespace NiagaraStackModuleItemIssues {
 					if (CompatibleOutputNode != nullptr)
 					{
 						OutTargetOutputNode = CompatibleOutputNode;
-						OutTargetIndex = 0;
+						OutTargetIndex = INDEX_NONE;
 					}
 				}
 			}
