- Add Near Plane to camera DI

diff --git a/Engine/Plugins/FX/Niagara/Shaders/Private/NiagaraDataInterfaceCamera.ush b/Engine/Plugins/FX/Niagara/Shaders/Private/NiagaraDataInterfaceCamera.ush
--- a/Engine/Plugins/FX/Niagara/Shaders/Private/NiagaraDataInterfaceCamera.ush
+++ b/Engine/Plugins/FX/Niagara/Shaders/Private/NiagaraDataInterfaceCamera.ush
@@ -29,7 +29,8 @@ void GetCameraPropertiesCPUASC47GPU_{ParameterName}(out float3 Out_CameraPositio
 }
 
 void GetViewPropertiesGPU_{ParameterName}(out float3 Out_ViewPositionWorld, out float3 Out_ViewForwardVector, out float3 Out_ViewUpVector, out float3 Out_ViewRightVector, out float4 Out_ViewSizeAndInverseSize,
-	out float4 Out_ScreenToViewSpace, out float2 Out_Current_TAAJitter, out float2 Out_Previous_TAAJitter, out float3 Out_PreViewTranslation, out float4 Out_BufferSizeAndInverseSize, out float2 Out_ViewportOffset)
+	out float4 Out_ScreenToViewSpace, out float2 Out_Current_TAAJitter, out float2 Out_Previous_TAAJitter, out float3 Out_PreViewTranslation, out float4 Out_BufferSizeAndInverseSize, out float2 Out_ViewportOffset,
+	out float Out_NearPlane)
 {
 	FLWCVector3 ViewPos = PrimaryView.TileOffset.WorldViewOrigin;
 	LWCSetTile(ViewPos, LWCGetTile(ViewPos) - GetEngineOwnerLWCTile()); // convert to simulation space
@@ -44,6 +45,7 @@ void GetViewPropertiesGPU_{ParameterName}(out float3 Out_ViewPositionWorld, out
 	Out_PreViewTranslation = DFHackToFloat(PrimaryView.PreViewTranslation);
 	Out_BufferSizeAndInverseSize = View.BufferSizeAndInvSize;
 	Out_ViewportOffset = View.ViewRectMin.xy;
+	Out_NearPlane = View.NearPlane;
 }
 
 void GetClipSpaceTransformsGPU_{ParameterName}(out float4x4 Out_WorldToClipTransform, out float4x4 Out_TranslatedWorldToClipTransform, out float4x4 Out_ClipToWorldTransform, out float4x4 Out_ClipToViewTransform,
diff --git a/Engine/Plugins/FX/Niagara/Source/Niagara/Private/NiagaraDataInterfaceCamera.cpp b/Engine/Plugins/FX/Niagara/Source/Niagara/Private/NiagaraDataInterfaceCamera.cpp
--- a/Engine/Plugins/FX/Niagara/Source/Niagara/Private/NiagaraDataInterfaceCamera.cpp
+++ b/Engine/Plugins/FX/Niagara/Source/Niagara/Private/NiagaraDataInterfaceCamera.cpp
@@ -29,6 +29,7 @@ struct FNiagaraCameraDIFunctionVersion
         AddedPreViewTranslation = 1,
 		AddedBufferSizeInvSizeForGPU = 2,
 		LWCConversion = 3,
+		AddGPUNearPlane = 4,
 
         VersionPlusOne,
         LatestVersion = VersionPlusOne - 1
@@ -214,6 +215,7 @@ void UNiagaraDataInterfaceCamera::GetFunctionsInternal(TArray<FNiagaraFunctionSi
 	Sig.AddOutput(FNiagaraVariable(FNiagaraTypeDefinition::GetVec3Def(), TEXT("PreViewTranslation")), LOCTEXT("PreViewTranslationDescription", "Returns the translation to apply for the various 'Translated XX to XX' transforms."));
 	Sig.AddOutput(FNiagaraVariable(FNiagaraTypeDefinition::GetVec4Def(), TEXT("Buffer Size And Inverse Size")), LOCTEXT("BufferSizeAndInverseSizeDescription", "Returns a vec4 - the x and y values are width and height of the buffer; the z and w values are the inverse width and height."));
 	Sig.AddOutput(FNiagaraVariable(FNiagaraTypeDefinition::GetVec2Def(), TEXT("Viewport Offset")), LOCTEXT("ViewportOffsetDescription", "Returns the viewport offset in the buffer."));
+	Sig.AddOutput(FNiagaraVariable(FNiagaraTypeDefinition::GetFloatDef(), TEXT("Near Plane")), LOCTEXT("NearPlaneDescription", "Returns the near plane value."));
 	OutFunctions.Add(Sig);
 
 
@@ -385,8 +387,9 @@ bool UNiagaraDataInterfaceCamera::UpgradeFunctionCall(FNiagaraFunctionSignature&
 		bWasChanged = true;
 	}
 
-	// LWC upgrades
-	if (FunctionSignature.FunctionVersion < FNiagaraCameraDIFunctionVersion::LWCConversion)
+	// LWCConversion
+	// AddGPUNearPlane
+	if (FunctionSignature.FunctionVersion < FNiagaraCameraDIFunctionVersion::LatestVersion)
 	{
 		TArray<FNiagaraFunctionSignature> AllFunctions;
 		GetFunctionsInternal(AllFunctions);
