- IWYU fix

diff --git a/Engine/Plugins/FX/Niagara/Source/NiagaraEditorWidgets/Private/Stack/SNiagaraSimTargetToggle.cpp b/Engine/Plugins/FX/Niagara/Source/NiagaraEditorWidgets/Private/Stack/SNiagaraSimTargetToggle.cpp
--- a/Engine/Plugins/FX/Niagara/Source/NiagaraEditorWidgets/Private/Stack/SNiagaraSimTargetToggle.cpp
+++ b/Engine/Plugins/FX/Niagara/Source/NiagaraEditorWidgets/Private/Stack/SNiagaraSimTargetToggle.cpp
@@ -8,6 +8,7 @@
 #include "Widgets/Input/SButton.h"
 #include "Widgets/Images/SImage.h"
 #include "ViewModels/NiagaraEmitterViewModel.h"
+#include "ScopedTransaction.h"
 
 #define LOCTEXT_NAMESPACE "SNiagaraSimTargetToggle"
 
