- Renderer vertex data export
- Creates a permutation for Niagara Vertex Factories that can export the geometry into the target UAV
- These permutations are only enabled in editor builds by default, but can be enabled on cooked targets
- Call NiagaraRendererReadback::EnqueueReadback with a list of components you want to capture, when complete the callback will execute on the game thread with a list of sections and vertices that are captures.  The vertex data is not deduplicated, so we have 3 vertices per triangle.
- Mesh vertex factory is not supported currently

diff --git a/Engine/Plugins/FX/Niagara/Shaders/Private/NiagaraRibbonVertexFactory.ush b/Engine/Plugins/FX/Niagara/Shaders/Private/NiagaraRibbonVertexFactory.ush
--- a/Engine/Plugins/FX/Niagara/Shaders/Private/NiagaraRibbonVertexFactory.ush
+++ b/Engine/Plugins/FX/Niagara/Shaders/Private/NiagaraRibbonVertexFactory.ush
@@ -1129,4 +1129,56 @@ FVertexFactoryInterpolantsVSToPS VertexFactoryAssignInterpolants(FVertexFactoryR
 
 #endif
 
+#if WITH_NIAGARA_VERTEX_FACTORY_EXPORT
+FVertexFactoryInput NiagaraVertexFactoryExport_InitVertexFactoryInput(uint TriangleIndex, int VertexIndex, uint PrimitiveId, uint DrawInstanceId)
+{
+	FVertexFactoryInput Input;
+
+	Input.InterpVtxId = TriangleIndex * 3 + VertexIndex;
+	Input.InterpVtxId = NiagaraRibbonVFLooseParameters.IndexBuffer[Input.InterpVtxId];
+
+#if FEATURE_LEVEL >= FEATURE_LEVEL_SM4 && ONEPASS_POINTLIGHT_SHADOW && USING_VERTEX_SHADER_LAYER
+	Input.InstanceId = 0;
+#endif
+	return Input;
+}
+
+uint NiagaraVertexFactoryExport_GetIndirectVertexCount()
+{
+	return NiagaraRibbonVFLooseParameters.IndirectDrawOutput[NiagaraRibbonVFLooseParameters.IndirectDrawOutputOffset + INDEX_GEN_DRAW_INDIRECT_OFFSET + 0];
+}
+
+float3 NiagaraVertexFactoryExport_GetPosition(FVertexFactoryInput Input, FVertexFactoryIntermediates Intermediates)
+{
+	return ConstructWorldPosition(Input, Intermediates.Current, LWCPromote(float3(0,0,0)));
+}
+
+float4 NiagaraVertexFactoryExport_GetColor(FVertexFactoryInput Input, FVertexFactoryIntermediates Intermediates)
+{
+#if NEEDS_PARTICLE_COLOR
+	return Intermediates.Color;
+#endif
+	return float4(1,1,1,1);
+}
+
+float3x3 NiagaraVertexFactoryExport_GetTangentBasis(FVertexFactoryInput Input, FVertexFactoryIntermediates Intermediates)
+{
+	return float3x3(Intermediates.Current.TangentX, Intermediates.Current.TangentY, Intermediates.Current.TangentZ);
+}
+
+float2 NiagaraVertexFactoryExport_GetTexCoord(FVertexFactoryInput Input, FVertexFactoryIntermediates Intermediates, uint CoordIndex)
+{
+	switch (CoordIndex)
+	{
+	#if NUM_MATERIAL_TEXCOORDS_VERTEX > 0
+		case 0:	return Intermediates.TexCoord.xy;
+	#endif
+	#if NUM_MATERIAL_TEXCOORDS_VERTEX > 1
+		case 1:	return Intermediates.TexCoord.zw;
+	#endif
+		default: return float2(0,0);
+	}
+}
+#endif //WITH_NIAGARA_VERTEX_FACTORY_EXPORT
+
 #include "/Engine/Private/VertexFactoryDefaultInterface.ush"
diff --git a/Engine/Plugins/FX/Niagara/Shaders/Private/NiagaraSpriteVertexFactory.ush b/Engine/Plugins/FX/Niagara/Shaders/Private/NiagaraSpriteVertexFactory.ush
--- a/Engine/Plugins/FX/Niagara/Shaders/Private/NiagaraSpriteVertexFactory.ush
+++ b/Engine/Plugins/FX/Niagara/Shaders/Private/NiagaraSpriteVertexFactory.ush
@@ -245,6 +245,26 @@ FPrimitiveSceneData GetPrimitiveData(FVertexFactoryIntermediates Intermediates)
 	return Intermediates.SceneData.Primitive;
 }
 
+float2 GetUV(FVertexFactoryIntermediates Intermediates, uint CoordIndex)
+{
+	switch (CoordIndex)
+	{
+	#if NUM_MATERIAL_TEXCOORDS_VERTEX >= 4
+		case 3:		return Intermediates.TexCoord_Unflipped.zw;
+	#endif
+	#if NUM_MATERIAL_TEXCOORDS_VERTEX >= 3
+		case 2:		return Intermediates.TexCoord_Unflipped.xy;
+	#endif
+	#if NUM_MATERIAL_TEXCOORDS_VERTEX >= 2
+		case 1:		return Intermediates.TexCoord.zw;
+	#endif
+	#if NUM_MATERIAL_TEXCOORDS_VERTEX >= 1
+		case 0:		return Intermediates.TexCoord.xy;
+	#endif
+		default:	return float2(0,0);
+	}
+}
+
 #if NUM_TEX_COORD_INTERPOLATORS
 bool UVIndexUseZW(int UVIndex)
 {
@@ -453,20 +473,18 @@ FMaterialVertexParameters GetMaterialVertexParameters(
 	Result.Particle.Random = 0.0f;
 #endif
 
-#if NUM_MATERIAL_TEXCOORDS_VERTEX
-	#if NUM_MATERIAL_TEXCOORDS_VERTEX >= 1
-		Result.TexCoords[0] = Intermediates.TexCoord.xy;
-	#if NUM_MATERIAL_TEXCOORDS_VERTEX >= 2
-		Result.TexCoords[1] = Intermediates.TexCoord.zw;
-	#if NUM_MATERIAL_TEXCOORDS_VERTEX >= 3
-		Result.TexCoords[2] = Intermediates.TexCoord_Unflipped.xy;
-	#if NUM_MATERIAL_TEXCOORDS_VERTEX >= 4
-		Result.TexCoords[3] = Intermediates.TexCoord_Unflipped.zw;
-	#endif	// >= 4
-	#endif	// >= 3
-	#endif	// >= 2
-	#endif	// >= 1
-#endif	//NUM_MATERIAL_TEXCOORDS_VERTEX
+#if NUM_MATERIAL_TEXCOORDS_VERTEX >= 1
+	Result.TexCoords[0] = GetUV(Intermediates, 0);
+#endif
+#if NUM_MATERIAL_TEXCOORDS_VERTEX >= 2
+	Result.TexCoords[1] = GetUV(Intermediates, 1);
+#endif
+#if NUM_MATERIAL_TEXCOORDS_VERTEX >= 3
+	Result.TexCoords[2] = GetUV(Intermediates, 2);
+#endif
+#if NUM_MATERIAL_TEXCOORDS_VERTEX >= 4
+	Result.TexCoords[3] = GetUV(Intermediates, 3);
+#endif
 
 #if USE_PARTICLE_SUBUVS
 	Result.Particle.SubUVCoords[0] = Intermediates.TexCoord.xy;
@@ -1322,4 +1340,42 @@ uint VertexFactoryGetPrimitiveId(FVertexFactoryInterpolantsVSToPS Interpolants)
 	return 0;
 }
 
+#if WITH_NIAGARA_VERTEX_FACTORY_EXPORT
+FVertexFactoryInput NiagaraVertexFactoryExport_InitVertexFactoryInput(uint TriangleIndex, int VertexIndex, uint PrimitiveId, uint DrawInstanceId)
+{
+	return LoadVertexFactoryInputForHGS(TriangleIndex, VertexIndex);
+}
+
+uint NiagaraVertexFactoryExport_GetIndirectVertexCount()
+{
+	return 6 * NiagaraSpriteVFLooseParameters.IndirectArgsBuffer[NiagaraSpriteVFLooseParameters.IndirectArgsOffset + 1];
+}
+
+float3 NiagaraVertexFactoryExport_GetPosition(FVertexFactoryInput Input, FVertexFactoryIntermediates Intermediates)
+{
+	//FLWCVector3 Origin = LWCGetOrigin(DFToTileOffset(GetPrimitiveDataFromUniformBuffer().LocalToWorld));
+	//return LWCToFloat(LWCSubtract(Intermediates.VertexWorldPosition, Origin));
+
+	return LWCToFloat(Intermediates.VertexWorldPosition);
+}
+
+float4 NiagaraVertexFactoryExport_GetColor(FVertexFactoryInput Input, FVertexFactoryIntermediates Intermediates)
+{
+#if NEEDS_PARTICLE_COLOR
+	return Intermediates.Color;
+#endif
+	return float4(1,1,1,1);
+}
+
+float3x3 NiagaraVertexFactoryExport_GetTangentBasis(FVertexFactoryInput Input, FVertexFactoryIntermediates Intermediates)
+{
+	return CalcTangentBasis(Intermediates);
+}
+
+float2 NiagaraVertexFactoryExport_GetTexCoord(FVertexFactoryInput Input, FVertexFactoryIntermediates Intermediates, uint CoordIndex)
+{
+	return GetUV(Intermediates, CoordIndex);
+}
+#endif //WITH_NIAGARA_VERTEX_FACTORY_EXPORT
+
 #include "/Engine/Private/VertexFactoryDefaultInterface.ush"
diff --git a/Engine/Plugins/FX/Niagara/Shaders/Private/NiagaraVertexFactoryExport.usf b/Engine/Plugins/FX/Niagara/Shaders/Private/NiagaraVertexFactoryExport.usf
new file mode 100644
--- /dev/null
+++ b/Engine/Plugins/FX/Niagara/Shaders/Private/NiagaraVertexFactoryExport.usf
@@ -0,0 +1,124 @@
+// Copyright Epic Games, Inc. All Rights Reserved.
+
+// Change to force shader compilation of this shader
+
+#include "/Engine/Private/Common.ush"
+#include "/Engine/Private/ComputeShaderUtils.ush"
+#include "/Engine/Generated/Material.ush"
+#include "/Engine/Generated/VertexFactory.ush"
+ 
+#ifndef THREAD_GROUP_SIZE
+	#define THREAD_GROUP_SIZE 64
+#endif
+
+uint IsIndirectDraw;
+uint NumInstances;
+uint NumVerticesPerInstance;
+uint bApplyWPO;
+uint SectionInfoOutputOffset;
+
+uint VertexStride;
+uint VertexPositionOffset;
+uint VertexColorOffset;
+uint VertexTangentBasisOffset;
+uint VertexTexCoordOffset;
+uint VertexTexCoordNum;
+uint VertexOutputOffset;
+RWByteAddressBuffer	RWVertexData;
+
+static const uint INVALID_COMPONENT = 0xffffffff;
+
+void OutputPosition(uint VertexIndex, float3 Value)
+{
+	if ( VertexPositionOffset != INVALID_COMPONENT )
+	{
+		const uint OutputVertexIndex = VertexOutputOffset + VertexIndex;
+		const uint Offset = (OutputVertexIndex * VertexStride) + VertexPositionOffset;
+		RWVertexData.Store3(Offset, asuint(Value));
+	}
+}
+
+void OutputColor(uint VertexIndex, float4 Value)
+{
+	if ( VertexColorOffset != INVALID_COMPONENT )
+	{
+		const uint OutputVertexIndex = VertexOutputOffset + VertexIndex;
+		const uint Offset = (OutputVertexIndex * VertexStride) + VertexColorOffset;
+		RWVertexData.Store4(Offset, asuint(Value));
+	}
+}
+
+void OutputTangentBasis(uint VertexIndex, float3x3 Value)
+{
+	if ( VertexTangentBasisOffset != INVALID_COMPONENT )
+	{
+		const uint OutputVertexIndex = VertexOutputOffset + VertexIndex;
+		const uint Offset = (OutputVertexIndex * VertexStride) + VertexTangentBasisOffset;
+		RWVertexData.Store3(Offset, asuint(Value[0]));
+		RWVertexData.Store3(Offset + 12, asuint(Value[1]));
+		RWVertexData.Store3(Offset + 24, asuint(Value[2]));
+	}
+}
+
+void OutputTexCoord(uint VertexIndex, uint TexCoord, float2 Value)
+{
+	if ( VertexTexCoordOffset != INVALID_COMPONENT && TexCoord < VertexTexCoordNum )
+	{
+		const uint OutputVertexIndex = VertexOutputOffset + VertexIndex;
+		const uint Offset = (OutputVertexIndex * VertexStride) + VertexTexCoordOffset + (8 * TexCoord);
+		RWVertexData.Store2(Offset, asuint(Value));
+	}
+}
+
+[numthreads(THREAD_GROUP_SIZE, 1, 1)]
+void VertexFactoryExportCS(uint3 GroupId : SV_GroupID, int GroupThreadIndex : SV_GroupIndex)
+{
+	const uint VertexIndex = GetUnWrappedDispatchThreadId(GroupId, GroupThreadIndex, THREAD_GROUP_SIZE);
+	const uint MaxOrNumVertices = NumInstances * NumVerticesPerInstance;
+	const uint NumVertices = IsIndirectDraw == 0 ? MaxOrNumVertices : NiagaraVertexFactoryExport_GetIndirectVertexCount();
+	if (VertexIndex == 0)
+	{
+		RWVertexData.Store(SectionInfoOutputOffset, NumVertices);
+	}
+
+	if (VertexIndex >= NumVertices)
+	{
+		// Helpful for debugging as it is obvious the position is NaN
+		if (VertexIndex < MaxOrNumVertices)
+		{
+			uint4 NaN = uint4(0xFFFFFFFF, 0xFFFFFFFF, 0xFFFFFFFF, 0xFFFFFFFF);
+			OutputPosition(VertexIndex, asfloat(NaN.xyz));
+		}
+		return;
+	}
+
+	// Genreate the vertex data
+	ResolvedView = ResolveView();
+
+	const uint	PrimitiveId = 0;
+	const int	InstanceId = VertexIndex / NumVerticesPerInstance;
+	const uint	InstanceVertexIndex = VertexIndex % NumVerticesPerInstance;
+	
+	FVertexFactoryInput Input = NiagaraVertexFactoryExport_InitVertexFactoryInput(InstanceVertexIndex / 3, InstanceVertexIndex % 3, PrimitiveId, InstanceId);
+
+	FVertexFactoryIntermediates VFIntermediates = GetVertexFactoryIntermediates(Input);
+
+	float4 TranslatedWorldPosition = VertexFactoryGetWorldPosition(Input, VFIntermediates);
+	float3x3 TangentToLocal = VertexFactoryGetTangentToLocal(Input, VFIntermediates);
+	FMaterialVertexParameters VertexParameters = GetMaterialVertexParameters(Input, VFIntermediates, TranslatedWorldPosition.xyz, TangentToLocal);
+
+	float3 MaterialWPO = 0;
+	if (bApplyWPO)
+	{
+		MaterialWPO = MakeFinite(GetMaterialWorldPositionOffset(VertexParameters));
+	}
+
+	//-TODO: LWC issues
+	OutputPosition(VertexIndex, NiagaraVertexFactoryExport_GetPosition(Input, VFIntermediates) + MaterialWPO);
+	OutputColor(VertexIndex, NiagaraVertexFactoryExport_GetColor(Input, VFIntermediates));
+	OutputTangentBasis(VertexIndex, NiagaraVertexFactoryExport_GetTangentBasis(Input, VFIntermediates));
+	for ( uint i=0; i < VertexTexCoordNum; ++i )
+	{
+		OutputTexCoord(VertexIndex, i, NiagaraVertexFactoryExport_GetTexCoord(Input, VFIntermediates, i));
+	}
+}
diff --git a/Engine/Plugins/FX/Niagara/Source/Niagara/Internal/NiagaraRendererReadback.h b/Engine/Plugins/FX/Niagara/Source/Niagara/Internal/NiagaraRendererReadback.h
new file mode 100644
--- /dev/null
+++ b/Engine/Plugins/FX/Niagara/Source/Niagara/Internal/NiagaraRendererReadback.h
@@ -0,0 +1,81 @@
+// Copyright Epic Games, Inc. All Rights Reserved.
+
+#pragma once
+
+#include "PrimitiveComponentId.h"
+#include "MeshBatch.h"
+
+class UNiagaraComponent;
+class FNiagaraSceneProxy;
+class FPrimitiveSceneProxy;
+class FSceneView;
+
+#ifndef WITH_NIAGARA_RENDERER_READBACK
+	// Console variable CVarNiagarVertexFactoryExportEnabledMode is used to control if the permutations are generated for the taget platform
+	// You will need to modify both that and this define in order to run on none editor targets
+	#define WITH_NIAGARA_RENDERER_READBACK WITH_EDITOR
+#endif
+
+#if WITH_NIAGARA_RENDERER_READBACK
+struct FNiagaraRendererReadbackParameters
+{
+	bool bExportPosition		= true;
+	bool bExportTangentBasis	= true;
+	bool bExportColor			= true;
+	uint32 ExportNumTexCoords	= 1;
+
+	bool bApplyWPO				= false;
+};
+
+struct FNiagaraRendererReadbackResult
+{
+	struct FSection
+	{
+		TWeakObjectPtr<UMaterialInterface>	WeakMaterialInterface;
+		uint32								FirstTriangle = 0;
+		uint32								NumTriangles = 0;
+	};
+
+	bool HasPosition() const { return VertexPositionOffset != INDEX_NONE; }
+	bool HasColor() const { return VertexColorOffset != INDEX_NONE; }
+	bool HasTangentBasis() const { return VertexTangentBasisOffset != INDEX_NONE; }
+	bool HasTexCoords() const { return VertexTexCoordOffset != INDEX_NONE; }
+
+	NIAGARA_API FVector3f GetPosition(uint32 Vertex) const;
+	NIAGARA_API FLinearColor GetColor(uint32 Vertex) const;
+	NIAGARA_API FVector3f GetTangentX(uint32 Vertex) const;
+	NIAGARA_API FVector3f GetTangentY(uint32 Vertex) const;
+	NIAGARA_API FVector3f GetTangentZ(uint32 Vertex) const;
+	NIAGARA_API FVector2f GetTexCoord(uint32 Vertex, uint32 TexCoordIndex) const;
+
+	uint32				NumVertices = 0;
+	uint32				VertexStride = 0;
+	uint32				VertexPositionOffset = INDEX_NONE;
+	uint32				VertexColorOffset = INDEX_NONE;
+	uint32				VertexTangentBasisOffset = INDEX_NONE;
+	uint32				VertexTexCoordOffset = INDEX_NONE;
+	uint32				VertexTexCoordNum = 0;
+	TArray<uint8>		VertexData;
+	TArray<FSection>	Sections;
+
+	TArray<FString>		Errors;
+};
+
+using FNiagaraRendererReadbackComplete = TFunction<void(const FNiagaraRendererReadbackResult& Result)>;
+
+namespace NiagaraRendererReadback
+{
+	extern uint32 GIsCapturing;
+
+	// Capture the next frame's renderer data for the provided components
+	// These functions are expected to be called on the Game Thread
+	// The Callback will also be invoked on the GameThread once complete
+	NIAGARA_API void EnqueueReadback(TConstArrayView<UNiagaraComponent*> Components, FNiagaraRendererReadbackComplete Callback, const FNiagaraRendererReadbackParameters& Parameters = FNiagaraRendererReadbackParameters());
+	NIAGARA_API void EnqueueReadback(UNiagaraComponent* Component, FNiagaraRendererReadbackComplete Callback, const FNiagaraRendererReadbackParameters& Parameters = FNiagaraRendererReadbackParameters());
+
+	inline void BeginCapture() { ++GIsCapturing; }
+	inline bool IsCapturing() { return GIsCapturing != 0; }
+	void CaptureMeshBatch(const FSceneView* View, const FNiagaraSceneProxy* SceneProxy, const FMeshBatch& MeshBatch, uint32 NumInstances, uint32 NumVerticesPerInstance);
+	inline void EndCapture() { --GIsCapturing; }
+}
+#endif //WITH_NIAGARA_RENDERER_READBACK
diff --git a/Engine/Plugins/FX/Niagara/Source/Niagara/Private/NiagaraRendererReadback.cpp b/Engine/Plugins/FX/Niagara/Source/Niagara/Private/NiagaraRendererReadback.cpp
new file mode 100644
--- /dev/null
+++ b/Engine/Plugins/FX/Niagara/Source/Niagara/Private/NiagaraRendererReadback.cpp
@@ -0,0 +1,524 @@
+// Copyright Epic Games, Inc. All Rights Reserved.
+
+//-TODO: Add sprite cutout support
+//-TODO: Add way to specify local / world and how we should handle LWC
+//-TODO: Validate cooked doesn't contain permutations + validate can run
+
+#include "NiagaraRendererReadback.h"
+#include "NiagaraComponent.h"
+#include "NiagaraGpuComputeDispatchInterface.h"
+#include "NiagaraGpuComputeDataManager.h"
+#include "NiagaraGpuReadbackManager.h"
+#include "NiagaraSceneProxy.h"
+#include "NiagaraVertexFactoryExport.h"
+
+#include "Materials/Material.h"
+#include "Materials/MaterialRenderProxy.h"
+#include "Async/Async.h"
+#include "MeshMaterialShader.h"
+#include "MeshPassProcessor.h"
+#include "RenderGraphUtils.h"
+#include "SceneInterface.h"
+#include "SceneTexturesConfig.h"
+#include "SceneView.h"
+
+#if WITH_NIAGARA_RENDERER_READBACK
+namespace NiagaraRendererReadback
+{
+	uint32 GIsCapturing = 0;
+
+	BEGIN_SHADER_PARAMETER_STRUCT(FReadbackPassParams, )
+		SHADER_PARAMETER_RDG_BUFFER_UAV(RWByteAddressBuffer,				RWVertexBuffer)
+		SHADER_PARAMETER_STRUCT_REF(FViewUniformShaderParameters,			View)
+		//SHADER_PARAMETER_RDG_UNIFORM_BUFFER(FSceneUniformParameters,	Scene)
+	END_SHADER_PARAMETER_STRUCT()
+
+	//////////////////////////////////////////////////////////////////////////////////////////////////////////////////////////////////////////////////////////////////////////////
+
+	struct FReadbackMeshBatch
+	{
+		uint32						OutputVertexOffset = 0;
+		uint32						SectionIndex = 0;
+		const FSceneView*			View = nullptr;
+		const FPrimitiveSceneProxy* SceneProxy = nullptr;
+		FMeshBatch					MeshBatch;					// We might be able to use a pointer here since it's all within the same frame
+		uint32						NumInstances = 0;
+		uint32						NumVerticesPerInstance = 0;
+	};
+
+	struct FRendererReadbackRequest
+	{
+		bool													bApplyWPO = false;
+		FNiagaraRendererReadbackResult							Result;
+
+		UE::FMutex												MeshBatchesMutex;
+		TArray<FReadbackMeshBatch>								MeshBatches;
+
+		FNiagaraRendererReadbackComplete						CompleteCallback;
+		TArray<FPrimitiveComponentId, TInlineAllocator<1>>		PrimitivesToCapture;
+	};
+
+	using FRendererReadbackRequestPtr = TSharedPtr<FRendererReadbackRequest, ESPMode::ThreadSafe>;
+
+	void ExecuteCompleteCallback(FRendererReadbackRequestPtr Request)
+	{
+		AsyncTask(
+			ENamedThreads::GameThread,
+			[Request]()
+			{
+				Request->CompleteCallback(Request->Result);
+			}
+		);
+	}
+
+	//////////////////////////////////////////////////////////////////////////////////////////////////////////////////////////////////////////////////////////////////////////////
+
+	void ExportMeshBatch(FRHICommandList& RHICmdList, FScene* Scene, ERHIFeatureLevel::Type FeatureLevel, const FReadbackMeshBatch& ReadbackMeshBatch, bool bApplyWPO, const FNiagaraRendererReadbackResult& VertexLayout, FRDGBufferUAV* VertexBufferUAV)
+	{
+		const FMeshBatch& MeshBatch = ReadbackMeshBatch.MeshBatch;
+
+		for (const FMaterialRenderProxy* MaterialRenderProxy=MeshBatch.MaterialRenderProxy; MaterialRenderProxy; MaterialRenderProxy=MaterialRenderProxy->GetFallback(FeatureLevel))
+		{
+			const FMaterial* Material = MaterialRenderProxy->GetMaterialNoFallback(FeatureLevel);
+			if (!Material || !Material->GetRenderingThreadShaderMap())
+			{
+				continue;
+			}
+
+			FMaterialShaderTypes ShaderTypes;
+			ShaderTypes.AddShaderType<FNiagaraVertexFactoryExportCS>();
+
+			FMaterialShaders MaterialShaders;
+			if (Material->TryGetShaders(ShaderTypes, MeshBatch.VertexFactory->GetType(), MaterialShaders) == false)
+			{
+				continue;
+			}
+
+			const FSceneView* View = ReadbackMeshBatch.View;
+			const FPrimitiveSceneProxy* SceneProxy = ReadbackMeshBatch.SceneProxy;
+
+			TShaderRef<FNiagaraVertexFactoryExportCS> Shader;
+			MaterialShaders.TryGetShader(SF_Compute, Shader);
+							
+			FMeshProcessorShaders MeshProcessorShaders;
+			MeshProcessorShaders.ComputeShader = Shader;
+
+			FMeshDrawShaderBindings ShaderBindings;
+			ShaderBindings.Initialize(MeshProcessorShaders);
+
+			FMeshMaterialShaderElementData ShaderElementData;
+			ShaderElementData.InitializeMeshMaterialData(View, SceneProxy, MeshBatch, -1, false);
+
+			FMeshDrawSingleShaderBindings SingleShaderBindings = ShaderBindings.GetSingleShaderBindings(SF_Compute);
+			Shader->GetShaderBindings(Scene, FeatureLevel, SceneProxy, *MaterialRenderProxy, *Material, ShaderElementData, SingleShaderBindings);
+
+			const FMeshBatchElement& MeshBatchElement = MeshBatch.Elements[0];
+
+			FVertexInputStreamArray DummyArray;
+			FMeshMaterialShader::GetElementShaderBindings(Shader, Scene, View, MeshBatch.VertexFactory, EVertexInputStreamType::Default, FeatureLevel, SceneProxy, MeshBatch, MeshBatchElement, ShaderElementData, SingleShaderBindings, DummyArray);
+
+			const uint32 SectionInfoOffset = (VertexLayout.NumVertices * VertexLayout.VertexStride) + (ReadbackMeshBatch.SectionIndex * 4);
+
+			SingleShaderBindings.Add(Shader->IsIndirectDraw,			MeshBatchElement.IndirectArgsBuffer == nullptr ? 0 : 1);
+			SingleShaderBindings.Add(Shader->NumInstances,				ReadbackMeshBatch.NumInstances);
+			SingleShaderBindings.Add(Shader->NumVerticesPerInstance,	ReadbackMeshBatch.NumVerticesPerInstance);
+			SingleShaderBindings.Add(Shader->bApplyWPO,					bApplyWPO ? 1 : 0);
+
+			SingleShaderBindings.Add(Shader->VertexStride,				VertexLayout.VertexStride);
+			SingleShaderBindings.Add(Shader->VertexPositionOffset,		VertexLayout.VertexPositionOffset);
+			SingleShaderBindings.Add(Shader->VertexColorOffset,			VertexLayout.VertexColorOffset);
+			SingleShaderBindings.Add(Shader->VertexTangentBasisOffset,	VertexLayout.VertexTangentBasisOffset);
+			SingleShaderBindings.Add(Shader->VertexTexCoordOffset,		VertexLayout.VertexTexCoordOffset);
+			SingleShaderBindings.Add(Shader->VertexTexCoordNum,			VertexLayout.VertexTexCoordNum);
+			SingleShaderBindings.Add(Shader->VertexOutputOffset,		ReadbackMeshBatch.OutputVertexOffset);
+			SingleShaderBindings.Add(Shader->SectionInfoOutputOffset,	SectionInfoOffset);
+
+		#if MESH_DRAW_COMMAND_DEBUG_DATA
+			ShaderBindings.Finalize(&MeshProcessorShaders);
+		#endif
+								
+			FRHIComputeShader* ComputeShader = Shader.GetComputeShader();
+			SetComputePipelineState(RHICmdList, ComputeShader);
+
+			FShaderBindingState ShaderBindingState;
+			FRHIBatchedShaderParameters& BatchedParameters = RHICmdList.GetScratchShaderParameters();
+			SetUAVParameter(BatchedParameters, Shader->RWVertexData, VertexBufferUAV->GetRHI());
+			ShaderBindings.SetParameters(BatchedParameters, &ShaderBindingState);
+			RHICmdList.SetBatchedShaderParameters(ComputeShader, BatchedParameters);
+
+			const uint32 MaxOrNumVertices = ReadbackMeshBatch.NumInstances * ReadbackMeshBatch.NumVerticesPerInstance;
+			const FIntVector NumWrappedThreadGroups = FComputeShaderUtils::GetGroupCountWrapped(MaxOrNumVertices, FNiagaraVertexFactoryExportCS::ThreadGroupSize);
+			RHICmdList.DispatchComputeShader(NumWrappedThreadGroups.X, NumWrappedThreadGroups.Y, NumWrappedThreadGroups.Z);
+
+			return;
+		}
+	}
+
+	//////////////////////////////////////////////////////////////////////////////////////////////////////////////////////////////////////////////////////////////////////////////
+
+	class FRendererReadbackComputeManager final : public FNiagaraGpuComputeDataManager
+	{
+	public:
+		FRendererReadbackComputeManager(FNiagaraGpuComputeDispatchInterface* InOwnerInterface)
+			: FNiagaraGpuComputeDataManager(InOwnerInterface)
+		{
+			InOwnerInterface->GetOnPreInitViewsEvent().AddRaw(this, &FRendererReadbackComputeManager::OnPreRender);
+			InOwnerInterface->GetOnPostRenderEvent().AddRaw(this, &FRendererReadbackComputeManager::OnPostRender);
+		}
+
+		virtual ~FRendererReadbackComputeManager()
+		{
+		}
+
+		static FName GetManagerName()
+		{
+			static FName ManagerName("FRendererReadbackComputeManager");
+			return ManagerName;
+		}
+
+		void OnPreRender(FRDGBuilder& GraphBuilder)
+		{
+			if (ReadbackRequests.Num() == 0)
+			{
+				return;
+			}
+			++NiagaraRendererReadback::GIsCapturing;
+		}
+
+		void OnPostRender(FRDGBuilder& GraphBuilder)
+		{
+			if (ReadbackRequests.Num() == 0)
+			{
+				return;
+			}
+
+			check(NiagaraRendererReadback::GIsCapturing > 0);
+			--NiagaraRendererReadback::GIsCapturing;
+
+			FScene* Scene = GetOwnerInterface()->GetScene();
+			const ERHIFeatureLevel::Type FeatureLevel = GetOwnerInterface()->GetSceneInterface()->GetFeatureLevel();
+
+			for (const FRendererReadbackRequestPtr& Request : ReadbackRequests)
+			{
+				for (FReadbackMeshBatch& ReadbackMeshBatch : Request->MeshBatches)
+				{
+					ReadbackMeshBatch.OutputVertexOffset = Request->Result.NumVertices;
+					ReadbackMeshBatch.SectionIndex = Request->Result.Sections.Num();
+
+					const FMeshBatch& MeshBatch = ReadbackMeshBatch.MeshBatch;
+					const FMeshBatchElement& MeshBatchElement = MeshBatch.Elements[0];
+
+					const uint32 MaxOrNumVertices = ReadbackMeshBatch.NumInstances * ReadbackMeshBatch.NumVerticesPerInstance;
+
+					FNiagaraRendererReadbackResult::FSection& Section = Request->Result.Sections.AddDefaulted_GetRef();
+					Section.WeakMaterialInterface	= MeshBatch.MaterialRenderProxy->GetMaterialInterface();
+					Section.FirstTriangle			= Request->Result.NumVertices / 3;
+					Section.NumTriangles			= MaxOrNumVertices / 3;
+					Request->Result.NumVertices += MaxOrNumVertices;
+				}
+
+				if (Request->Result.NumVertices == 0)
+				{
+					ExecuteCompleteCallback(Request);
+					continue;
+				}
+
+				// Allocate output buffer
+				const uint32 VertexSize = Request->Result.NumVertices * Request->Result.VertexStride;
+				const uint32 SectionSize = Request->Result.Sections.Num() * sizeof(uint32);
+
+				FRDGBufferDesc VertexBufferDesc = FRDGBufferDesc::CreateByteAddressDesc(VertexSize + SectionSize);
+				VertexBufferDesc.Usage |= EBufferUsageFlags::SourceCopy;
+				FRDGBufferRef VertexBuffer = GraphBuilder.CreateBuffer(VertexBufferDesc, TEXT("NiagaraRendererReadback"), ERDGBufferFlags::None);
+				FRDGBufferUAVRef VertexBufferUAV = GraphBuilder.CreateUAV(VertexBuffer, PF_Unknown);
+
+				// Make sure we clear the UAV so the section data we write out is valid even when we can't find an appropriate export shader
+				AddClearUAVPass(GraphBuilder, VertexBufferUAV, 0);
+
+				// Capture data
+				{
+					const FSceneView* SceneView = Request->MeshBatches[0].View;
+
+					FReadbackPassParams* PassParameters = GraphBuilder.AllocParameters<FReadbackPassParams>();
+					PassParameters->RWVertexBuffer	= VertexBufferUAV;
+					PassParameters->View			= SceneView->ViewUniformBuffer;
+
+					GraphBuilder.AddPass(
+						RDG_EVENT_NAME("NiagaraRendererReadback"),
+						PassParameters,
+						ERDGPassFlags::Compute | ERDGPassFlags::NeverCull,
+						[Request, PassParameters, Scene, FeatureLevel, SceneView](/*FRDGAsyncTask, */FRHICommandListImmediate& RHICmdList)
+						{
+							FMemMark Mark(FMemStack::Get());
+
+							for (const FReadbackMeshBatch& ReadbackMeshBatch : Request->MeshBatches)
+							{
+								ExportMeshBatch(RHICmdList, Scene, FeatureLevel, ReadbackMeshBatch, Request->bApplyWPO, Request->Result, PassParameters->RWVertexBuffer);
+							}
+						}
+					);
+				}
+
+				// Queue readback of the buffer
+				FNiagaraGpuReadbackManager* ReadbackManager = GetOwnerInterface()->GetGpuReadbackManager();
+				ReadbackManager->EnqueueReadback(
+					GraphBuilder,
+					VertexBuffer,
+					[Request, VertexSize, SectionSize](TConstArrayView<TPair<void*, uint32>> ReadbackData)
+					{
+						check(ReadbackData.Num() == 1);
+						check(ReadbackData[0].Value == VertexSize + SectionSize);
+
+						const uint8* VertexData = static_cast<const uint8*>(ReadbackData[0].Key);
+						const uint32* SectionInfo = reinterpret_cast<const uint32*>(VertexData + VertexSize);
+
+						// Get Vertex Count and Allocate Space
+						uint32 VertexCount = 0;
+						for ( int32 i=0; i < Request->Result.Sections.Num(); ++i )
+						{
+							VertexCount += SectionInfo[i];
+						}
+						Request->Result.NumVertices = VertexCount;
+
+						Request->Result.VertexData.AddUninitialized(VertexCount * Request->Result.VertexStride);
+						uint8* OutputVertexData = Request->Result.VertexData.GetData();
+
+						// Copy over the vertex data section by section
+						VertexCount = 0;
+						for (int32 i=0; i < Request->Result.Sections.Num(); ++i)
+						{
+							const uint32 SectionVertexCount = SectionInfo[i];
+							FMemory::Memcpy(OutputVertexData, VertexData, SectionVertexCount * Request->Result.VertexStride);
+
+							FNiagaraRendererReadbackResult::FSection& Section = Request->Result.Sections[i];
+							OutputVertexData += SectionVertexCount * Request->Result.VertexStride;
+							VertexData += Section.NumTriangles * 3 * Request->Result.VertexStride;
+
+							Section.FirstTriangle = VertexCount / 3;
+							Section.NumTriangles = SectionVertexCount / 3;
+
+							VertexCount += SectionVertexCount;
+						}
+
+						ExecuteCompleteCallback(Request);
+					}
+				);
+			}
+
+		#if 1
+			ReadbackRequests.Empty();
+		#else
+			// For debugging forces contant readback once requested
+			TArray<FRendererReadbackRequestPtr> TempRequests;
+			Swap(TempRequests, ReadbackRequests);
+
+			for (const FRendererReadbackRequestPtr& TempRequest : TempRequests)
+			{
+				FRendererReadbackRequestPtr NewRequest		= MakeShared<FRendererReadbackRequest>();
+				NewRequest->CompleteCallback				= TempRequest->CompleteCallback;
+				NewRequest->PrimitivesToCapture				= TempRequest->PrimitivesToCapture;
+				NewRequest->Result.VertexStride				= TempRequest->Result.VertexStride;
+				NewRequest->Result.VertexPositionOffset		= TempRequest->Result.VertexPositionOffset;
+				NewRequest->Result.VertexColorOffset		= TempRequest->Result.VertexColorOffset;
+				NewRequest->Result.VertexTangentBasisOffset	= TempRequest->Result.VertexTangentBasisOffset;
+				NewRequest->Result.VertexTexCoordOffset		= TempRequest->Result.VertexTexCoordOffset;
+				NewRequest->Result.VertexTexCoordNum		= TempRequest->Result.VertexTexCoordNum;
+				ReadbackRequests.Emplace(NewRequest);
+			}
+		#endif
+		}
+
+		void AddMeshBatch(const FSceneView* View, const FPrimitiveSceneProxy* SceneProxy, const FMeshBatch& MeshBatch, uint32 NumInstances, uint32 NumVerticesPerInstance)
+		{
+			for (const FRendererReadbackRequestPtr& Request : ReadbackRequests)
+			{
+				if (!Request->PrimitivesToCapture.Contains(SceneProxy->GetPrimitiveComponentId()))
+				{
+					continue;
+				}
+
+				if (MeshBatch.Elements.Num() != 1)
+				{
+					Request->Result.Errors.AddUnique(TEXT("Skipped mesh batch as only a single element is supported."));
+					return;
+				}
+
+				if (!ensureMsgf(MeshBatch.Type == PT_TriangleList, TEXT("Only PT_TriangleList are supported for renderer readback")))
+				{
+					Request->Result.Errors.AddUnique(TEXT("Skipped mesh batch as only PT_TriangleList is supported."));
+					return;
+				}
+
+				if (NumInstances * NumVerticesPerInstance == 0)
+				{
+					return;
+				}
+
+				if (!FNiagaraVertexFactoryExportCS::SupportsVertexFactoryType(MeshBatch.VertexFactory->GetType()))
+				{
+					Request->Result.Errors.AddUnique(FString::Printf(TEXT("Skipped mesh batch due to unsupported vertex factory '%s'"), MeshBatch.VertexFactory->GetType()->GetName()));
+					return;
+				}
+
+				UE::TScopeLock ScopeLock(Request->MeshBatchesMutex);
+				FReadbackMeshBatch& ReadbackMeshBatch = Request->MeshBatches.AddDefaulted_GetRef();
+				ReadbackMeshBatch.View						= View;
+				ReadbackMeshBatch.SceneProxy				= SceneProxy;
+				ReadbackMeshBatch.MeshBatch					= MeshBatch;
+				ReadbackMeshBatch.NumInstances				= NumInstances;
+				ReadbackMeshBatch.NumVerticesPerInstance	= NumVerticesPerInstance;
+			}
+		}
+
+		TArray<FRendererReadbackRequestPtr> ReadbackRequests;
+	};
+
+	template<typename T>
+	T GetVertexValue(const FNiagaraRendererReadbackResult& Result, uint32 Vertex, uint32 ComponentOffset, const T& DefaultValue)
+	{
+		T Value = DefaultValue;
+		if (Vertex < Result.NumVertices)
+		{
+			const uint32 Offset = (Result.VertexStride * Vertex) + ComponentOffset;
+			FMemory::Memcpy(&Value, Result.VertexData.GetData() + Offset, sizeof(T));
+		}
+		return Value;
+	}
+} //namespace NiagaraRendererReadback
+
+//////////////////////////////////////////////////////////////////////////////////////////////////////////////////////////////////////////////////////////////////////////////
+
+void NiagaraRendererReadback::EnqueueReadback(TConstArrayView<UNiagaraComponent*> Components, FNiagaraRendererReadbackComplete Callback, const FNiagaraRendererReadbackParameters& Parameters)
+{
+	using namespace NiagaraRendererReadback;
+
+	FNiagaraGpuComputeDispatchInterface* ComputeInterface = nullptr;
+	FRendererReadbackRequestPtr ReadbackRequest = MakeShared<FRendererReadbackRequest>();
+	ReadbackRequest->CompleteCallback = MoveTemp(Callback);
+
+	for (UNiagaraComponent* Component : Components)
+	{
+		UWorld* ComponentWorld = ::IsValid(Component) ? Component->GetWorld() : nullptr;
+		if (ComponentWorld == nullptr)
+		{
+			continue;
+		}
+
+		FNiagaraGpuComputeDispatchInterface* ComponentComputeInterface = FNiagaraGpuComputeDispatchInterface::Get(ComponentWorld);
+		ComputeInterface = ComputeInterface == nullptr ? ComponentComputeInterface : ComputeInterface;
+
+		if (ComponentComputeInterface == nullptr)
+		{
+			ReadbackRequest->Result.Errors.AddUnique(TEXT("Request failed due to no compute interface."));
+			ExecuteCompleteCallback(ReadbackRequest);
+			return;
+		}
+		else if (ComputeInterface != ComponentComputeInterface)
+		{
+			ReadbackRequest->Result.Errors.AddUnique(TEXT("Request failed due to components with mismatching compute interface."));
+			ExecuteCompleteCallback(ReadbackRequest);
+			return;
+		}
+
+		ReadbackRequest->PrimitivesToCapture.Add(Component->GetPrimitiveSceneId());
+	}
+
+	// Nothing added then there is nothing to do complete
+	if (ReadbackRequest->PrimitivesToCapture.Num() == 0)
+	{
+		ReadbackRequest->Result.Errors.AddUnique(TEXT("No primitives found to capture."));
+		ExecuteCompleteCallback(ReadbackRequest);
+		return;
+	}
+
+	// Copy over parameters
+	ReadbackRequest->bApplyWPO = Parameters.bApplyWPO;
+
+	// Build Vertex Output
+	{
+		FNiagaraRendererReadbackResult& Result = ReadbackRequest->Result;
+		if (Parameters.bExportPosition)
+		{
+			Result.VertexPositionOffset = Result.VertexStride;
+			Result.VertexStride += sizeof(FVector3f);
+		}
+		if (Parameters.bExportTangentBasis)
+		{
+			Result.VertexTangentBasisOffset = Result.VertexStride;
+			Result.VertexStride += sizeof(FVector3f) * 3;
+		}
+		if (Parameters.bExportColor)
+		{
+			Result.VertexColorOffset = Result.VertexStride;
+			Result.VertexStride += sizeof(FVector4f);
+		}
+		if (Parameters.ExportNumTexCoords > 0)
+		{
+			Result.VertexTexCoordNum = Parameters.ExportNumTexCoords;
+			Result.VertexTexCoordOffset = Result.VertexStride;
+			Result.VertexStride += sizeof(FVector2f) * Parameters.ExportNumTexCoords;
+		}
+	}
+
+	ENQUEUE_RENDER_COMMAND(EnqueueRendererReadback)(
+		[ComputeInterface, ReadbackRequest](FRHICommandListImmediate&)
+		{
+			FRendererReadbackComputeManager& ReadbackManager = ComputeInterface->GetOrCreateDataManager<FRendererReadbackComputeManager>();
+			ReadbackManager.ReadbackRequests.Emplace(ReadbackRequest);
+		}
+	);
+}
+
+FVector3f FNiagaraRendererReadbackResult::GetPosition(uint32 Vertex) const
+{
+	return NiagaraRendererReadback::GetVertexValue(*this, Vertex, VertexPositionOffset, FVector3f::ZeroVector);
+}
+
+FLinearColor FNiagaraRendererReadbackResult::GetColor(uint32 Vertex) const
+{
+	return NiagaraRendererReadback::GetVertexValue(*this, Vertex, VertexColorOffset, FLinearColor::Black);
+}
+
+FVector3f FNiagaraRendererReadbackResult::GetTangentX(uint32 Vertex) const
+{
+	return NiagaraRendererReadback::GetVertexValue(*this, Vertex, VertexTangentBasisOffset, FVector3f::XAxisVector);
+}
+
+FVector3f FNiagaraRendererReadbackResult::GetTangentY(uint32 Vertex) const
+{
+	return NiagaraRendererReadback::GetVertexValue(*this, Vertex, VertexTangentBasisOffset + sizeof(FVector3f), FVector3f::YAxisVector);
+}
+
+FVector3f FNiagaraRendererReadbackResult::GetTangentZ(uint32 Vertex) const
+{
+	return NiagaraRendererReadback::GetVertexValue(*this, Vertex, VertexTangentBasisOffset + (sizeof(FVector3f) * 2), FVector3f::ZAxisVector);
+}
+
+FVector2f FNiagaraRendererReadbackResult::GetTexCoord(uint32 Vertex, uint32 TexCoordIndex) const
+{
+	if (TexCoordIndex >= VertexTexCoordNum)
+	{
+		return FVector2f::ZeroVector;
+	}
+	return NiagaraRendererReadback::GetVertexValue(*this, Vertex, VertexTexCoordOffset + (sizeof(FVector2f) * TexCoordIndex), FVector2f::ZeroVector);
+}
+
+void NiagaraRendererReadback::EnqueueReadback(UNiagaraComponent* Component, TFunction<void(const FNiagaraRendererReadbackResult& Result)> Callback, const FNiagaraRendererReadbackParameters& Parameters)
+{
+	EnqueueReadback(MakeArrayView(&Component, 1), MoveTemp(Callback), Parameters);
+}
+
+void NiagaraRendererReadback::CaptureMeshBatch(const FSceneView* View, const FNiagaraSceneProxy* SceneProxy, const FMeshBatch& MeshBatch, uint32 NumInstances, uint32 NumVerticesPerInstance)
+{
+	using namespace NiagaraRendererReadback;
+
+	FNiagaraGpuComputeDispatchInterface* ComputeDispatchInterface = SceneProxy->GetComputeDispatchInterface();
+	if (ComputeDispatchInterface == nullptr)
+	{
+		return;
+	}
+
+	FRendererReadbackComputeManager& ReadbackManager = ComputeDispatchInterface->GetOrCreateDataManager<FRendererReadbackComputeManager>();
+	ReadbackManager.AddMeshBatch(View, SceneProxy, MeshBatch, NumInstances, NumVerticesPerInstance);
+}
+#endif //WITH_NIAGARA_RENDERER_READBACK
diff --git a/Engine/Plugins/FX/Niagara/Source/Niagara/Private/NiagaraRendererRibbons.cpp b/Engine/Plugins/FX/Niagara/Source/Niagara/Private/NiagaraRendererRibbons.cpp
--- a/Engine/Plugins/FX/Niagara/Source/Niagara/Private/NiagaraRendererRibbons.cpp
+++ b/Engine/Plugins/FX/Niagara/Source/Niagara/Private/NiagaraRendererRibbons.cpp
@@ -10,6 +10,7 @@
 #include "NiagaraEmitterInstance.h"
 #include "NiagaraSceneProxy.h"
 #include "NiagaraGpuComputeDataManager.h"
+#include "NiagaraRendererReadback.h"
 #include "NiagaraStats.h"
 #include "NiagaraSystemInstance.h"
 #include "NiagaraComponent.h"
@@ -555,7 +556,7 @@ struct FNiagaraRibbonGPUInitComputeBuffers
 		
 		if (SortBuffer.NumBytes < NeededSize * sizeof(int32))
 		{
-			SortBuffer.Initialize(RHICmdList, TEXT("NiagarGPUInit-SortedIndices"), sizeof(uint32), NeededSize, EPixelFormat::PF_R32_UINT, InitialAccess);
+			SortBuffer.Initialize(RHICmdList, TEXT("NiagarGPUInit-SortedIndices"), sizeof(uint32), NeededSize, EPixelFormat::PF_R32_UINT, InitialAccess | ERHIAccess::VertexOrIndexBuffer);
 		}
 
 		const uint32 TessellationBufferSize = NeededSize * (bWantsTessellation? (bWantsTessellationTwist? FTessellationStatsEntry::NumElements : FTessellationStatsEntryNoTwist::NumElements) : 0);
@@ -997,6 +998,14 @@ void FNiagaraRendererRibbons::GetDynamicMeshElements(const TArray<const FSceneVi
 			SetupMeshBatchAndCollectorResourceForView(RHICmdList, RenderingViewResources->IndexGenerationSettings, DynamicData, SourceParticleData, View, ViewFamily, SceneProxy, RenderingResources.RibbonResources, RenderingViewResources, MeshBatch, DynamicData->bUseGPUInit);
 
 			Collector.AddMesh(ViewIndex, MeshBatch);
+
+		#if WITH_NIAGARA_RENDERER_READBACK
+			if (NiagaraRendererReadback::IsCapturing())
+			{
+				const uint32 NumVertices = DynamicData->bUseGPUInit ? RenderingViewResources->IndexGenerationSettings.TotalNumIndices : RenderingViewResources->IndexGenerationSettings.CPUTriangleCount * 3;
+				NiagaraRendererReadback::CaptureMeshBatch(View, SceneProxy, MeshBatch, 1, NumVertices);
+			}
+		#endif //WITH_NIAGARA_RENDERER_READBACK
 		}
 	}
 
diff --git a/Engine/Plugins/FX/Niagara/Source/Niagara/Private/NiagaraRendererSprites.cpp b/Engine/Plugins/FX/Niagara/Source/Niagara/Private/NiagaraRendererSprites.cpp
--- a/Engine/Plugins/FX/Niagara/Source/Niagara/Private/NiagaraRendererSprites.cpp
+++ b/Engine/Plugins/FX/Niagara/Source/Niagara/Private/NiagaraRendererSprites.cpp
@@ -11,6 +11,7 @@
 #include "NiagaraEmitterInstance.h"
 #include "NiagaraGpuComputeDispatchInterface.h"
 #include "NiagaraGPUSortInfo.h"
+#include "NiagaraRendererReadback.h"
 #include "NiagaraSceneProxy.h"
 #include "NiagaraSettings.h"
 #include "NiagaraSortingGPU.h"
@@ -1073,6 +1074,13 @@ void FNiagaraRendererSprites::GetDynamicMeshElements(const TArray<const FSceneVi
 				FMeshBatch& MeshBatch = Collector.AllocateMesh();
 				CreateMeshBatchForView(RHICmdList, ParticleSpriteRenderData, MeshBatch, *View, *SceneProxy, VertexFactory, NumInstances, GPUCountBufferOffset, ParticleSpriteRenderData.bNeedsCull);
 				Collector.AddMesh(ViewIndex, MeshBatch);
+
+			#if WITH_NIAGARA_RENDERER_READBACK
+				if (NiagaraRendererReadback::IsCapturing())
+				{
+					NiagaraRendererReadback::CaptureMeshBatch(View, SceneProxy, MeshBatch, NumInstances, NumIndicesPerInstance);
+				}
+			#endif //WITH_NIAGARA_RENDERER_READBACK
 			}
 		}
 	}
diff --git a/Engine/Plugins/FX/Niagara/Source/NiagaraVertexFactories/Internal/NiagaraVertexFactoryExport.h b/Engine/Plugins/FX/Niagara/Source/NiagaraVertexFactories/Internal/NiagaraVertexFactoryExport.h
new file mode 100644
--- /dev/null
+++ b/Engine/Plugins/FX/Niagara/Source/NiagaraVertexFactories/Internal/NiagaraVertexFactoryExport.h
@@ -0,0 +1,70 @@
+// Copyright Epic Games, Inc. All Rights Reserved.
+
+#pragma once
+
+#include "Serialization/MemoryLayout.h"
+#include "MeshMaterialShader.h"
+#include "RHIStaticStates.h"
+
+class FVertexFactoryType;
+
+class FNiagaraVertexFactoryExportCS : public FMeshMaterialShader
+{
+	DECLARE_EXPORTED_SHADER_TYPE(FNiagaraVertexFactoryExportCS, MeshMaterial, NIAGARAVERTEXFACTORIES_API);
+
+public:
+	static constexpr uint32 ThreadGroupSize = 64;
+
+	NIAGARAVERTEXFACTORIES_API FNiagaraVertexFactoryExportCS(const FMeshMaterialShaderType::CompiledShaderInitializerType& Initializer);
+	NIAGARAVERTEXFACTORIES_API FNiagaraVertexFactoryExportCS() = default;
+
+	NIAGARAVERTEXFACTORIES_API static bool IsEnabled();
+	NIAGARAVERTEXFACTORIES_API static bool SupportsVertexFactoryType(const FVertexFactoryType* VertexFactory);
+
+	static bool ShouldCompilePermutation(const FMeshMaterialShaderPermutationParameters& Parameters);
+	static void ModifyCompilationEnvironment(const FMaterialShaderPermutationParameters& Parameters, FShaderCompilerEnvironment& OutEnvironment);
+
+	void GetShaderBindings(
+		const FScene* Scene,
+		ERHIFeatureLevel::Type FeatureLevel,
+		const FPrimitiveSceneProxy* PrimitiveSceneProxy,
+		const FMaterialRenderProxy& MaterialRenderProxy,
+		const FMaterial& Material,
+		const FMeshMaterialShaderElementData& ShaderElementData,
+		FMeshDrawSingleShaderBindings& ShaderBindings) const
+	{
+		FMeshMaterialShader::GetShaderBindings(Scene, FeatureLevel, PrimitiveSceneProxy, MaterialRenderProxy, Material, ShaderElementData, ShaderBindings);
+	}
+
+	void GetElementShaderBindings(
+		const FShaderMapPointerTable& PointerTable,
+		const FScene* Scene,
+		const FSceneView* ViewIfDynamicMeshCommand,
+		const FVertexFactory* VertexFactory,
+		const EVertexInputStreamType InputStreamType,
+		ERHIFeatureLevel::Type FeatureLevel,
+		const FPrimitiveSceneProxy* PrimitiveSceneProxy,
+		const FMeshBatch& MeshBatch,
+		const FMeshBatchElement& BatchElement,
+		const FMeshMaterialShaderElementData& ShaderElementData,
+		FMeshDrawSingleShaderBindings& ShaderBindings,
+		FVertexInputStreamArray& VertexStreams) const
+	{
+		FMeshMaterialShader::GetElementShaderBindings(PointerTable, Scene, ViewIfDynamicMeshCommand, VertexFactory, InputStreamType, FeatureLevel, PrimitiveSceneProxy, MeshBatch, BatchElement, ShaderElementData, ShaderBindings, VertexStreams);
+	}
+
+	LAYOUT_FIELD(FShaderParameter,			IsIndirectDraw);
+	LAYOUT_FIELD(FShaderParameter,			NumInstances);
+	LAYOUT_FIELD(FShaderParameter,			NumVerticesPerInstance);
+	LAYOUT_FIELD(FShaderParameter,			bApplyWPO);
+	LAYOUT_FIELD(FShaderParameter,			SectionInfoOutputOffset);
+
+	LAYOUT_FIELD(FShaderParameter,			VertexStride);
+	LAYOUT_FIELD(FShaderParameter,			VertexPositionOffset);
+	LAYOUT_FIELD(FShaderParameter,			VertexColorOffset);
+	LAYOUT_FIELD(FShaderParameter,			VertexTangentBasisOffset);
+	LAYOUT_FIELD(FShaderParameter,			VertexTexCoordOffset);
+	LAYOUT_FIELD(FShaderParameter,			VertexTexCoordNum);
+	LAYOUT_FIELD(FShaderParameter,			VertexOutputOffset);
+	LAYOUT_FIELD(FShaderResourceParameter,	RWVertexData);
+};
diff --git a/Engine/Plugins/FX/Niagara/Source/NiagaraVertexFactories/NiagaraVertexFactories.Build.cs b/Engine/Plugins/FX/Niagara/Source/NiagaraVertexFactories/NiagaraVertexFactories.Build.cs
--- a/Engine/Plugins/FX/Niagara/Source/NiagaraVertexFactories/NiagaraVertexFactories.Build.cs
+++ b/Engine/Plugins/FX/Niagara/Source/NiagaraVertexFactories/NiagaraVertexFactories.Build.cs
@@ -11,7 +11,8 @@ public class NiagaraVertexFactories : ModuleRules
 				"Core",
 				"Engine",		
 				"RenderCore",
-                "RHI",
+				"Renderer",
+				"RHI",
             }
         );
 
diff --git a/Engine/Plugins/FX/Niagara/Source/NiagaraVertexFactories/Private/NiagaraVertexFactoryExport.cpp b/Engine/Plugins/FX/Niagara/Source/NiagaraVertexFactories/Private/NiagaraVertexFactoryExport.cpp
new file mode 100644
--- /dev/null
+++ b/Engine/Plugins/FX/Niagara/Source/NiagaraVertexFactories/Private/NiagaraVertexFactoryExport.cpp
@@ -0,0 +1,61 @@
+// Copyright Epic Games, Inc. All Rights Reserved.
+
+#include "NiagaraVertexFactoryExport.h"
+#include "SceneTexturesConfig.h"
+
+static TAutoConsoleVariable<int32> CVarNiagarVertexFactoryExportEnabledMode(
+	TEXT("fx.Niagara.VertexFactoryExport.EnabledMode"),
+	1,
+	TEXT("Determins compilation mode for vertex factory export shader permutations\n")
+	TEXT("0 - Disabled on all platforms / targets\n")
+	TEXT("1 - Enabled for editor only (default)\n")
+	TEXT("2 - Enabled on all platforms / targets"),
+	ECVF_ReadOnly | ECVF_RenderThreadSafe
+);
+
+FNiagaraVertexFactoryExportCS::FNiagaraVertexFactoryExportCS(const FMeshMaterialShaderType::CompiledShaderInitializerType& Initializer)
+	: FMeshMaterialShader(Initializer)
+{
+	PassUniformBuffer.Bind(Initializer.ParameterMap, FSceneTextureUniformParameters::FTypeInfo::GetStructMetadata()->GetShaderVariableName());
+
+	IsIndirectDraw.Bind(Initializer.ParameterMap, TEXT("IsIndirectDraw"));
+	NumInstances.Bind(Initializer.ParameterMap, TEXT("NumInstances"));
+	NumVerticesPerInstance.Bind(Initializer.ParameterMap, TEXT("NumVerticesPerInstance"));
+	bApplyWPO.Bind(Initializer.ParameterMap, TEXT("bApplyWPO"));
+	SectionInfoOutputOffset.Bind(Initializer.ParameterMap, TEXT("SectionInfoOutputOffset"));
+
+	VertexStride.Bind(Initializer.ParameterMap, TEXT("VertexStride"));
+	VertexPositionOffset.Bind(Initializer.ParameterMap, TEXT("VertexPositionOffset"));
+	VertexColorOffset.Bind(Initializer.ParameterMap, TEXT("VertexColorOffset"));
+	VertexTangentBasisOffset.Bind(Initializer.ParameterMap, TEXT("VertexTangentBasisOffset"));
+	VertexTexCoordOffset.Bind(Initializer.ParameterMap, TEXT("VertexTexCoordOffset"));
+	VertexTexCoordNum.Bind(Initializer.ParameterMap, TEXT("VertexTexCoordNum"));
+	VertexOutputOffset.Bind(Initializer.ParameterMap, TEXT("VertexOutputOffset"));
+	RWVertexData.Bind(Initializer.ParameterMap, TEXT("RWVertexData"));
+}
+
+bool FNiagaraVertexFactoryExportCS::SupportsVertexFactoryType(const FVertexFactoryType* VertexFactoryType)
+{
+	return
+		//VertexFactoryType == FindVertexFactoryType(TEXT("FNiagaraMeshVertexFactory")) ||
+		VertexFactoryType == FindVertexFactoryType(TEXT("FNiagaraRibbonVertexFactory")) ||
+		VertexFactoryType == FindVertexFactoryType(TEXT("FNiagaraSpriteVertexFactory"));
+}
+
+bool FNiagaraVertexFactoryExportCS::ShouldCompilePermutation(const FMeshMaterialShaderPermutationParameters& Parameters)
+{
+	const int32 EnabledMode = CVarNiagarVertexFactoryExportEnabledMode.GetValueOnAnyThread();
+	const bool bEnabled =
+		(EnabledMode == 2) ||
+		(EnabledMode == 1 && EnumHasAllFlags(Parameters.Flags, EShaderPermutationFlags::HasEditorOnlyData));
+
+	return bEnabled && SupportsVertexFactoryType(Parameters.VertexFactoryType);
+}
+
+void FNiagaraVertexFactoryExportCS::ModifyCompilationEnvironment(const FMaterialShaderPermutationParameters& Parameters, FShaderCompilerEnvironment& OutEnvironment)
+{
+	OutEnvironment.SetDefine(TEXT("THREAD_GROUP_SIZE"), ThreadGroupSize);
+	OutEnvironment.SetDefine(TEXT("WITH_NIAGARA_VERTEX_FACTORY_EXPORT"), 1);
+}
+
+IMPLEMENT_MATERIAL_SHADER_TYPE(, FNiagaraVertexFactoryExportCS, TEXT("/Plugin/FX/Niagara/Private/NiagaraVertexFactoryExport.usf"), TEXT("VertexFactoryExportCS"), SF_Compute);
