- Fixes for ray tracing & sprite cutouts

#tests Assign cutout with no sub image to a sprite renderer and observe raytracing scene

diff --git a/Engine/Plugins/FX/Niagara/Shaders/Private/NiagaraSpriteVertexFactory.ush b/Engine/Plugins/FX/Niagara/Shaders/Private/NiagaraSpriteVertexFactory.ush
--- a/Engine/Plugins/FX/Niagara/Shaders/Private/NiagaraSpriteVertexFactory.ush
+++ b/Engine/Plugins/FX/Niagara/Shaders/Private/NiagaraSpriteVertexFactory.ush
@@ -43,21 +43,30 @@ struct FVertexFactoryInput
 	uint InstanceId					: SV_InstanceID;
 };
 
+uint GetCutoutNumVertices()
+{
+	return NiagaraSpriteVFLooseParameters.CutoutParameters & ~0x1;
+}
+
+bool GetCutoutUseSubImage()
+{
+	return (NiagaraSpriteVFLooseParameters.CutoutParameters & 0x1) != 0;
+}
+
 // RHI_RAYTRACING
 #if COMPUTESHADER || RAYHITGROUPSHADER
 FVertexFactoryInput LoadVertexFactoryInputForHGS(uint TriangleIndex, int VertexIndex)
 {
 	// Different numbers of cutout vertices correspond to different index buffers
 	// For 8 verts, use GSixTriangleParticleIndexBuffer
-	if (NiagaraSpriteVFLooseParameters.NumCutoutVerticesPerFrame == 8)
+	if (GetCutoutNumVertices() == 8)
 	{		
 		FVertexFactoryInput Input;
 		
 		// Hard coded GSixTriangleParticleIndexBuffer
 		uint IndexBuffer[18] = { 0, 1, 2, 0, 2, 3, 0, 3, 4, 0, 4, 5, 0, 5, 6, 0, 6, 7 };
 		uint VertexId = IndexBuffer[(TriangleIndex * 3 + VertexIndex) % 18];
-		float2 TexCoords[4] = { float2(0.0f, 0.0f), float2(0.0f, 1.0f), float2(1.0f, 1.0f), float2(1.0f, 0.0f) };
-		Input.TexCoord = TexCoords[VertexId];
+		Input.TexCoord = NiagaraSpriteVFLooseParameters.CutoutGeometry[VertexId];
 		Input.VertexId = VertexId;
 		Input.InstanceId = TriangleIndex / 6;
 	
@@ -70,8 +79,15 @@ FVertexFactoryInput LoadVertexFactoryInputForHGS(uint TriangleIndex, int VertexI
 		
 		uint IndexBuffer[6] = { 0, 2, 3, 0, 1, 2 };
 		uint VertexId = IndexBuffer[(TriangleIndex * 3 + VertexIndex) % 6];
-		float2 TexCoords[4] = { float2(0.0f, 0.0f), float2(0.0f, 1.0f), float2(1.0f, 1.0f), float2(1.0f, 0.0f) };
-		Input.TexCoord = TexCoords[VertexId];
+		if (GetCutoutNumVertices() > 0)
+		{
+			Input.TexCoord = NiagaraSpriteVFLooseParameters.CutoutGeometry[VertexId];
+		}
+		else
+		{
+			float2 TexCoords[4] = { float2(0.0f, 0.0f), float2(0.0f, 1.0f), float2(1.0f, 1.0f), float2(1.0f, 0.0f) };
+			Input.TexCoord = TexCoords[VertexId];
+		}
 		Input.VertexId = VertexId;
 		Input.InstanceId = TriangleIndex / 2;
 	
@@ -85,7 +101,9 @@ FVertexFactoryInput LoadVertexFactoryInputForDynamicUpdate(uint TriangleIndex, i
 }
 uint GetNumRayTracingDynamicMeshVerticesIndirect()
 {
-	return 6 * NiagaraSpriteVFLooseParameters.IndirectArgsBuffer[NiagaraSpriteVFLooseParameters.IndirectArgsOffset + 1];
+	const uint NumVertices = NiagaraSpriteVFLooseParameters.IndirectArgsBuffer[NiagaraSpriteVFLooseParameters.IndirectArgsOffset + 0];
+	const uint NumInstances = NiagaraSpriteVFLooseParameters.IndirectArgsBuffer[NiagaraSpriteVFLooseParameters.IndirectArgsOffset + 1];
+	return NumVertices * NumInstances;
 }
 #endif
 
@@ -656,10 +674,11 @@ void ComputeBillboardUVs(FVertexFactoryInput Input, float2 ParticleSize, float S
 {
 #if FEATURE_LEVEL >= FEATURE_LEVEL_SM4
 	BRANCH
-	if (NiagaraSpriteVFLooseParameters.NumCutoutVerticesPerFrame > 0)
+	if (GetCutoutUseSubImage())
 	{
 		// Avoid uint divide which is extremely slow on GCN
-		const uint VertexIndexMask = NiagaraSpriteVFLooseParameters.NumCutoutVerticesPerFrame == 8 ? 7 : 3; // BVC_FourVertices || BVC_EightVertices
+		const uint NumCutOutVertices = GetCutoutNumVertices();
+		const uint VertexIndexMask = NumCutOutVertices == 8 ? 7 : 3; // BVC_FourVertices || BVC_EightVertices
 		uint CutoutVertexIndex = Input.VertexId & VertexIndexMask;
 		const float NumFrames = NiagaraSpriteVF.SubImageSize.x * NiagaraSpriteVF.SubImageSize.y;
 		const uint SubImageIndexInt = (uint)fmod(SubImageIndex, NumFrames);
@@ -670,11 +689,11 @@ void ComputeBillboardUVs(FVertexFactoryInput Input, float2 ParticleSize, float S
 		FLATTEN
 		if (UVFlip.x * UVFlip.y < 0.0f)
 		{
-			CutoutVertexIndex = NiagaraSpriteVFLooseParameters.NumCutoutVerticesPerFrame - 1 - CutoutVertexIndex;
+			CutoutVertexIndex = NumCutOutVertices - 1 - CutoutVertexIndex;
 		}
 
 		// Fetch the billboard space positions from the appropriate frame of the preprocessed cutout geometry
-		UVForPosition = NiagaraSpriteVFLooseParameters.CutoutGeometry[SubImageIndexInt * NiagaraSpriteVFLooseParameters.NumCutoutVerticesPerFrame + CutoutVertexIndex];
+		UVForPosition = NiagaraSpriteVFLooseParameters.CutoutGeometry[SubImageIndexInt * NumCutOutVertices + CutoutVertexIndex];
 		UVForTexturingUnflipped = UVForPosition;
 		UVForTexturing = UVForPosition;
 		
@@ -1108,7 +1127,7 @@ float3 VertexFactoryGetPositionForVertexLighting(FVertexFactoryInput Input, FVer
 {
 #if (TRANSLUCENCY_LIGHTING_VOLUMETRIC_PERVERTEX_DIRECTIONAL || TRANSLUCENCY_LIGHTING_VOLUMETRIC_PERVERTEX_NONDIRECTIONAL)
 	// Do per-vertex lighting with particle position instead of vertex position when we're using SubUV cutouts, since the vertex positions vary per-frame which would cause popping
-	return NiagaraSpriteVFLooseParameters.NumCutoutVerticesPerFrame > 0 ? Intermediates.TranslatedWorldPosition : TranslatedWorldPosition;
+	return GetCutoutUseSubImage() ? Intermediates.TranslatedWorldPosition : TranslatedWorldPosition;
 #else
 	return TranslatedWorldPosition;
 #endif
@@ -1343,12 +1362,18 @@ uint VertexFactoryGetPrimitiveId(FVertexFactoryInterpolantsVSToPS Interpolants)
 #if WITH_NIAGARA_VERTEX_FACTORY_EXPORT
 FVertexFactoryInput NiagaraVertexFactoryExport_InitVertexFactoryInput(uint TriangleIndex, int VertexIndex, uint PrimitiveId, uint DrawInstanceId)
 {
+	const uint VertsPerInstance = GetCutoutNumVertices() == 0 ? 6 : ((GetCutoutNumVertices() - 2) * 3);
+	const uint TrisPerInstance = VertsPerInstance / 3;
+	TriangleIndex += TrisPerInstance * DrawInstanceId;
+	VertexIndex += VertsPerInstance * DrawInstanceId;
 	return LoadVertexFactoryInputForHGS(TriangleIndex, VertexIndex);
 }
 
 uint NiagaraVertexFactoryExport_GetIndirectVertexCount()
 {
-	return 6 * NiagaraSpriteVFLooseParameters.IndirectArgsBuffer[NiagaraSpriteVFLooseParameters.IndirectArgsOffset + 1];
+	const uint NumVertices = NiagaraSpriteVFLooseParameters.IndirectArgsBuffer[NiagaraSpriteVFLooseParameters.IndirectArgsOffset + 0];
+	const uint NumInstances = NiagaraSpriteVFLooseParameters.IndirectArgsBuffer[NiagaraSpriteVFLooseParameters.IndirectArgsOffset + 1];
+	return NumVertices * NumInstances;
 }
 
 float3 NiagaraVertexFactoryExport_GetPosition(FVertexFactoryInput Input, FVertexFactoryIntermediates Intermediates)
diff --git a/Engine/Plugins/FX/Niagara/Source/Niagara/Private/NiagaraRendererSprites.cpp b/Engine/Plugins/FX/Niagara/Source/Niagara/Private/NiagaraRendererSprites.cpp
--- a/Engine/Plugins/FX/Niagara/Source/Niagara/Private/NiagaraRendererSprites.cpp
+++ b/Engine/Plugins/FX/Niagara/Source/Niagara/Private/NiagaraRendererSprites.cpp
@@ -515,12 +515,11 @@ void FNiagaraRendererSprites::SetupVertexFactory(FRHICommandListBase& RHICmdList
 	const bool bUseCutout = CutoutVertexBuffer.VertexBufferRHI.IsValid();
 	if (bUseCutout)
 	{
-		if (bUseSubImage)
-		{
-			VertexFactory.SetCutoutParameters(NumCutoutVertexPerSubImage, CutoutVertexBuffer.VertexBufferSRV);
-		}
-		else // Otherwise simply replace the input stream with the single cutout geometry
+		VertexFactory.SetCutoutParameters(bUseSubImage, NumCutoutVertexPerSubImage);
+		VertexFactory.SetCutoutGeometry(CutoutVertexBuffer.VertexBufferSRV);
+		if (bUseSubImage == false)
 		{
+			// Replace the vertex buffer used for cutouts which do not use sub image animation, this is more optimal in the shader
 			VertexFactory.SetVertexBufferOverride(&CutoutVertexBuffer);
 		}
 	}
@@ -868,7 +867,7 @@ void FNiagaraRendererSprites::CreateMeshBatchForView(
 	FMaterialRenderProxy* MaterialRenderProxy = ParticleSpriteRenderData.DynamicDataSprites->Material;
 	check(MaterialRenderProxy);
 
-	VFLooseParams.NumCutoutVerticesPerFrame = VertexFactory.GetNumCutoutVerticesPerFrame();
+	VFLooseParams.CutoutParameters = VertexFactory.GetCutoutParameters();
 	VFLooseParams.CutoutGeometry = VertexFactory.GetCutoutGeometrySRV() ? VertexFactory.GetCutoutGeometrySRV() : GFNiagaraNullCutoutVertexBuffer.VertexBufferSRV.GetReference();
 	VFLooseParams.ParticleAlignmentMode = VertexFactory.GetAlignmentMode();
 	VFLooseParams.ParticleFacingMode = VertexFactory.GetFacingMode();
@@ -903,7 +902,7 @@ void FNiagaraRendererSprites::CreateMeshBatchForView(
 		VFLooseParams.IndirectArgsOffset = 0;
 	}
 
-	VertexFactory.LooseParameterUniformBuffer = FNiagaraSpriteVFLooseParametersRef::CreateUniformBufferImmediate(VFLooseParams, UniformBuffer_SingleFrame);
+	VertexFactory.SetLooseParameterUniformBuffer(FNiagaraSpriteVFLooseParametersRef::CreateUniformBufferImmediate(VFLooseParams, UniformBuffer_SingleFrame));
 
 	MeshBatch.VertexFactory = &VertexFactory;
 	MeshBatch.CastShadow = SceneProxy.CastsDynamicShadow() && bCastShadows;
diff --git a/Engine/Plugins/FX/Niagara/Source/NiagaraVertexFactories/Private/NiagaraSpriteVertexFactory.cpp b/Engine/Plugins/FX/Niagara/Source/NiagaraVertexFactories/Private/NiagaraSpriteVertexFactory.cpp
--- a/Engine/Plugins/FX/Niagara/Source/NiagaraVertexFactories/Private/NiagaraSpriteVertexFactory.cpp
+++ b/Engine/Plugins/FX/Niagara/Source/NiagaraVertexFactories/Private/NiagaraSpriteVertexFactory.cpp
@@ -25,19 +25,6 @@ public:
 	void Bind(const FShaderParameterMap& ParameterMap)
 	{
 		FNiagaraVertexFactoryShaderParametersBase::Bind(ParameterMap);
-
-		NumCutoutVerticesPerFrame.Bind(ParameterMap, TEXT("NumCutoutVerticesPerFrame"));
-		CutoutGeometry.Bind(ParameterMap, TEXT("CutoutGeometry"));
-
-		//  		NiagaraParticleDataInt.Bind(ParameterMap, TEXT("NiagaraParticleDataInt"));
-		//  		Int32DataOffset.Bind(ParameterMap, TEXT("NiagaraInt32DataOffset"));
-		//  		Int32DataStride.Bind(ParameterMap, TEXT("NiagaraInt3DataStride"));
-
-		ParticleAlignmentMode.Bind(ParameterMap, TEXT("ParticleAlignmentMode"));
-		ParticleFacingMode.Bind(ParameterMap, TEXT("ParticleFacingMode"));
-
-		SortedIndices.Bind(ParameterMap, TEXT("SortedIndices"));
-		SortedIndicesOffset.Bind(ParameterMap, TEXT("SortedIndicesOffset"));
 	}
 
 	void GetElementShaderBindings(
@@ -55,37 +42,8 @@ public:
 
 		const FNiagaraSpriteVertexFactory* SpriteVF = static_cast<const FNiagaraSpriteVertexFactory*>(VertexFactory);
 		ShaderBindings.Add(Shader->GetUniformBufferParameter<FNiagaraSpriteUniformParameters>(), SpriteVF->GetSpriteUniformBuffer());
-
-		ShaderBindings.Add(Shader->GetUniformBufferParameter<FNiagaraSpriteVFLooseParameters>(), SpriteVF->LooseParameterUniformBuffer);
-
-		ShaderBindings.Add(NumCutoutVerticesPerFrame, SpriteVF->GetNumCutoutVerticesPerFrame());
-		FRHIShaderResourceView* NullSRV = GFNiagaraNullCutoutVertexBuffer.VertexBufferSRV;
-		ShaderBindings.Add(CutoutGeometry, SpriteVF->GetCutoutGeometrySRV() ? SpriteVF->GetCutoutGeometrySRV() : NullSRV);
-
-		ShaderBindings.Add(ParticleAlignmentMode, SpriteVF->GetAlignmentMode());
-		ShaderBindings.Add(ParticleFacingMode, SpriteVF->GetFacingMode());
-
-		ShaderBindings.Add(SortedIndices, SpriteVF->GetSortedIndicesSRV() ? SpriteVF->GetSortedIndicesSRV() : GFNiagaraNullSortedIndicesVertexBuffer.VertexBufferSRV.GetReference());
-		ShaderBindings.Add(SortedIndicesOffset, SpriteVF->GetSortedIndicesOffset());
+		ShaderBindings.Add(Shader->GetUniformBufferParameter<FNiagaraSpriteVFLooseParameters>(), SpriteVF->GetLooseParameterUniformBuffer());
 	}
-
-private:
-
-	LAYOUT_FIELD(FShaderParameter, NumCutoutVerticesPerFrame);
-
-	LAYOUT_FIELD(FShaderParameter, ParticleAlignmentMode);
-	LAYOUT_FIELD(FShaderParameter, ParticleFacingMode);
-
-	LAYOUT_FIELD(FShaderResourceParameter, CutoutGeometry);
-
-
-	//  	LAYOUT_FIELD(FShaderResourceParameter, NiagaraParticleDataInt);
-	//  	LAYOUT_FIELD(FShaderParameter, Int32DataOffset);
-	//  	LAYOUT_FIELD(FShaderParameter, Int32DataStride);
-
-	LAYOUT_FIELD(FShaderResourceParameter, SortedIndices);
-	LAYOUT_FIELD(FShaderParameter, SortedIndicesOffset);
-
 };
 
 IMPLEMENT_TYPE_LAYOUT(FNiagaraSpriteVertexFactoryShaderParametersVS);
diff --git a/Engine/Plugins/FX/Niagara/Source/NiagaraVertexFactories/Public/NiagaraSpriteVertexFactory.h b/Engine/Plugins/FX/Niagara/Source/NiagaraVertexFactories/Public/NiagaraSpriteVertexFactory.h
--- a/Engine/Plugins/FX/Niagara/Source/NiagaraVertexFactories/Public/NiagaraSpriteVertexFactory.h
+++ b/Engine/Plugins/FX/Niagara/Source/NiagaraVertexFactories/Public/NiagaraSpriteVertexFactory.h
@@ -92,7 +92,7 @@ BEGIN_GLOBAL_SHADER_PARAMETER_STRUCT( FNiagaraSpriteUniformParameters, NIAGARAVE
 typedef TUniformBufferRef<FNiagaraSpriteUniformParameters> FNiagaraSpriteUniformBufferRef;
 
 BEGIN_GLOBAL_SHADER_PARAMETER_STRUCT(FNiagaraSpriteVFLooseParameters, NIAGARAVERTEXFACTORIES_API)
-	SHADER_PARAMETER(uint32, NumCutoutVerticesPerFrame)
+	SHADER_PARAMETER(uint32, CutoutParameters)
 	SHADER_PARAMETER(uint32, NiagaraFloatDataStride)
 	SHADER_PARAMETER(uint32, ParticleAlignmentMode)
 	SHADER_PARAMETER(uint32, ParticleFacingMode)
@@ -117,27 +117,8 @@ class FNiagaraSpriteVertexFactory : public FNiagaraVertexFactoryBase
 public:
 
 	/** Default constructor. */
-	FNiagaraSpriteVertexFactory(ENiagaraVertexFactoryType InType, ERHIFeatureLevel::Type InFeatureLevel )
-		: FNiagaraVertexFactoryBase(InType, InFeatureLevel)
-		, LooseParameterUniformBuffer(nullptr)
-		, NumCutoutVerticesPerFrame(0)
-		, CutoutGeometrySRV(nullptr)
-		, AlignmentMode(0)
-		, FacingMode(0)
-		, SortedIndicesSRV(nullptr)
-		, SortedIndicesOffset(0)
-	{}
-
-	FNiagaraSpriteVertexFactory()
-		: FNiagaraVertexFactoryBase(NVFT_MAX, ERHIFeatureLevel::Num)
-		, LooseParameterUniformBuffer(nullptr)
-		, NumCutoutVerticesPerFrame(0)
-		, CutoutGeometrySRV(nullptr)
-		, AlignmentMode(0)
-		, FacingMode(0)
-		, SortedIndicesSRV(nullptr)
-		, SortedIndicesOffset(0)
-	{}
+	FNiagaraSpriteVertexFactory(ENiagaraVertexFactoryType InType, ERHIFeatureLevel::Type InFeatureLevel ) : FNiagaraVertexFactoryBase(InType, InFeatureLevel) {}
+	FNiagaraSpriteVertexFactory() : FNiagaraVertexFactoryBase(NVFT_MAX, ERHIFeatureLevel::Num) {}
 
 	// FRenderResource interface.
 	NIAGARAVERTEXFACTORIES_API virtual void InitRHI(FRHICommandListBase& RHICmdList) override;
@@ -161,29 +142,17 @@ public:
 
 	NIAGARAVERTEXFACTORIES_API void SetTexCoordBuffer(const FVertexBuffer* InTexCoordBuffer);
 	
-	/**
-	 * Set the uniform buffer for this vertex factory.
-	 */
-	FORCEINLINE void SetSpriteUniformBuffer( const FNiagaraSpriteUniformBufferRef& InSpriteUniformBuffer )
-	{
-		SpriteUniformBuffer = InSpriteUniformBuffer;
-	}
-
-	/**
-	 * Retrieve the uniform buffer for this vertex factory.
-	 */
-	FORCEINLINE FRHIUniformBuffer* GetSpriteUniformBuffer() const
-	{
-		return SpriteUniformBuffer;
-	}
+	void SetSpriteUniformBuffer(const FNiagaraSpriteUniformBufferRef& InSpriteUniformBuffer ) {SpriteUniformBuffer = InSpriteUniformBuffer; }
+	FRHIUniformBuffer* GetSpriteUniformBuffer() const { return SpriteUniformBuffer; }
 
-	void SetCutoutParameters(int32 InNumCutoutVerticesPerFrame, FRHIShaderResourceView* InCutoutGeometrySRV)
+	void SetCutoutParameters(bool bUseSubImage, int32 InNumCutoutVertexPerSubImage)
 	{
-		NumCutoutVerticesPerFrame = InNumCutoutVerticesPerFrame;
-		CutoutGeometrySRV = InCutoutGeometrySRV;
+		check((InNumCutoutVertexPerSubImage & 1) == 0);
+		CutoutParameters = InNumCutoutVertexPerSubImage | (bUseSubImage ? 1 : 0);
 	}
+	uint32 GetCutoutParameters() const { return CutoutParameters; }
 
-	inline int32 GetNumCutoutVerticesPerFrame() const { return NumCutoutVerticesPerFrame; }
+	void SetCutoutGeometry(FRHIShaderResourceView* InCutoutGeometrySRV) { CutoutGeometrySRV = InCutoutGeometrySRV; }
 	inline FRHIShaderResourceView* GetCutoutGeometrySRV() const { return CutoutGeometrySRV; }
 
 	void SetSortedIndices(const FShaderResourceViewRHIRef& InSortedIndicesSRV, uint32 InSortedIndicesOffset)
@@ -192,41 +161,19 @@ public:
 		SortedIndicesOffset = InSortedIndicesOffset;
 	}
 
-	FORCEINLINE FRHIShaderResourceView* GetSortedIndicesSRV() const
-	{
-		return SortedIndicesSRV;
-	}
+	FRHIShaderResourceView* GetSortedIndicesSRV() const { return SortedIndicesSRV; }
+	int32 GetSortedIndicesOffset() const { return SortedIndicesOffset; }
 
-	FORCEINLINE int32 GetSortedIndicesOffset() const
-	{
-		return SortedIndicesOffset;
-	}
+	void SetFacingMode(uint32 InMode) { FacingMode = InMode; }
+	uint32 GetFacingMode() const { return FacingMode; }
 
-	void SetFacingMode(uint32 InMode)
-	{
-		FacingMode = InMode;
-	}
+	void SetAlignmentMode(uint32 InMode) { AlignmentMode = InMode; }
+	uint32 GetAlignmentMode() const { return AlignmentMode; }
 
-	uint32 GetFacingMode() const
-	{
-		return FacingMode;
-	}
+	void SetVertexBufferOverride(const FVertexBuffer* InVertexBufferOverride) { VertexBufferOverride = InVertexBufferOverride; }
 
-	void SetAlignmentMode(uint32 InMode)
-	{
-		AlignmentMode = InMode;
-	}
-
-	uint32 GetAlignmentMode() const
-	{
-		return AlignmentMode;
-	}
-
-	void SetVertexBufferOverride(const FVertexBuffer* InVertexBufferOverride)
-	{
-		VertexBufferOverride = InVertexBufferOverride;
-	}
-	FUniformBufferRHIRef LooseParameterUniformBuffer;
+	void SetLooseParameterUniformBuffer(FUniformBufferRHIRef InBuffer) { LooseParameterUniformBuffer = InBuffer; }
+	FRHIUniformBuffer* GetLooseParameterUniformBuffer() const { return LooseParameterUniformBuffer; }
 
 protected:
 	/** Initialize streams for this vertex factory. */
@@ -234,16 +181,16 @@ protected:
 
 private:
 
-	const FVertexBuffer* VertexBufferOverride = nullptr;
+	const FVertexBuffer*		VertexBufferOverride = nullptr;
+	FUniformBufferRHIRef		LooseParameterUniformBuffer;
 
 	/** Uniform buffer with sprite parameters. */
-	FUniformBufferRHIRef SpriteUniformBuffer;
-
-	int32 NumCutoutVerticesPerFrame;
-	FShaderResourceViewRHIRef CutoutGeometrySRV;
-	uint32 AlignmentMode;
-	uint32 FacingMode;
-
-	FShaderResourceViewRHIRef SortedIndicesSRV;
-	uint32 SortedIndicesOffset;
+	FUniformBufferRHIRef		SpriteUniformBuffer;
+	uint32						CutoutParameters = 0;
+	FShaderResourceViewRHIRef	CutoutGeometrySRV;
+	uint32						AlignmentMode = 0;
+	uint32						FacingMode = 0;
+
+	FShaderResourceViewRHIRef	SortedIndicesSRV;
+	uint32						SortedIndicesOffset = 0;
 };
