- Fix name collisions when adding empty emitters and a combination of stateless & stateful emitters

diff --git a/Engine/Plugins/FX/Niagara/Source/NiagaraEditor/Private/NiagaraEditorUtilities.cpp b/Engine/Plugins/FX/Niagara/Source/NiagaraEditor/Private/NiagaraEditorUtilities.cpp
--- a/Engine/Plugins/FX/Niagara/Source/NiagaraEditor/Private/NiagaraEditorUtilities.cpp
+++ b/Engine/Plugins/FX/Niagara/Source/NiagaraEditor/Private/NiagaraEditorUtilities.cpp
@@ -2107,8 +2107,17 @@ const FGuid FNiagaraEditorUtilities::AddEmitterToSystem(UNiagaraSystem& InSystem
 	FNiagaraEmitterHandle EmitterHandle;
 	if (SystemEditorData->GetOwningSystemIsPlaceholder() == false)
 	{
+		// FNiagaraUtilities::GetUniqueName will produce a unique emitter name (i.e. Empty -> Empty001 -> Empty002, etc)
+		// but it will not produce a unique number so if that UObject is already taken by another object type we need to find a unique number
+		// The number is not important to the stack / emitter name as it's not displayed but is important to avoid colliding between UNiagaraEmitter and UNiagaraStatelessEmitter types
+		FName EmitterName = FNiagaraUtilities::GetUniqueName(InEmitterToAdd.GetFName(), EmitterHandleNames);
+		while (FindObjectFast<UObject>(&InSystem, EmitterName))
+		{
+			EmitterName.SetNumber(EmitterName.GetNumber() + 1);
+		}
+
 		InSystem.Modify();
-		EmitterHandle = InSystem.AddEmitterHandle(InEmitterToAdd, FNiagaraUtilities::GetUniqueName(InEmitterToAdd.GetFName(), EmitterHandleNames), EmitterVersion);
+		EmitterHandle = InSystem.AddEmitterHandle(InEmitterToAdd, EmitterName, EmitterVersion);
 	}
 	else if (bCreateCopy)
 	{
