- Static mesh export for Niagara Baker

#tests Bake supported types (mesh + sprite) to a static mesh using the baker

diff --git a/Engine/Plugins/FX/Niagara/Source/Niagara/Classes/NiagaraBakerOutputSimCache.h b/Engine/Plugins/FX/Niagara/Source/Niagara/Classes/NiagaraBakerOutputSimCache.h
--- a/Engine/Plugins/FX/Niagara/Source/Niagara/Classes/NiagaraBakerOutputSimCache.h
+++ b/Engine/Plugins/FX/Niagara/Source/Niagara/Classes/NiagaraBakerOutputSimCache.h
@@ -6,7 +6,7 @@
 #include "NiagaraSimCache.h"
 #include "NiagaraBakerOutputSimCache.generated.h"
 
-UCLASS(meta = (DisplayName = "Simulation Cache Output (Experimental)"), MinimalAPI)
+UCLASS(meta = (DisplayName = "Simulation Cache Output"), MinimalAPI)
 class UNiagaraBakerOutputSimCache : public UNiagaraBakerOutput
 {
 	GENERATED_BODY()
diff --git a/Engine/Plugins/FX/Niagara/Source/Niagara/Classes/NiagaraBakerOutputStaticMesh.h b/Engine/Plugins/FX/Niagara/Source/Niagara/Classes/NiagaraBakerOutputStaticMesh.h
new file mode 100644
--- /dev/null
+++ b/Engine/Plugins/FX/Niagara/Source/Niagara/Classes/NiagaraBakerOutputStaticMesh.h
@@ -0,0 +1,33 @@
+// Copyright Epic Games, Inc. All Rights Reserved.
+
+#pragma once
+
+#include "NiagaraBakerOutput.h"
+#include "NiagaraBakerOutputStaticMesh.generated.h"
+
+UCLASS(meta = (DisplayName = "Static Mesh Output"), MinimalAPI)
+class UNiagaraBakerOutputStaticMesh : public UNiagaraBakerOutput
+{
+	GENERATED_BODY()
+
+public:
+	UNiagaraBakerOutputStaticMesh(const FObjectInitializer& Init)
+	{
+	}
+
+	UPROPERTY(EditAnywhere, Category = "Settings")
+	FString FramesAssetPathFormat = TEXT("{AssetFolder}/{AssetName}_BakedFrame_{OutputName}_{FrameIndex}");
+
+	//UPROPERTY(EditAnywhere, Category = "Settings")
+	//FNiagaraRendererReadbackParameters CreateParameters;
+
+	NIAGARA_API virtual bool Equals(const UNiagaraBakerOutput& Other) const override;
+
+#if WITH_EDITOR
+	NIAGARA_API FString MakeOutputName() const override;
+#endif
+
+#if WITH_EDITORONLY_DATA
+	NIAGARA_API virtual void PostEditChangeProperty(struct FPropertyChangedEvent& PropertyChangedEvent) override;
+#endif
+};
diff --git a/Engine/Plugins/FX/Niagara/Source/Niagara/Private/NiagaraBakerOutputStaticMesh.cpp b/Engine/Plugins/FX/Niagara/Source/Niagara/Private/NiagaraBakerOutputStaticMesh.cpp
new file mode 100644
--- /dev/null
+++ b/Engine/Plugins/FX/Niagara/Source/Niagara/Private/NiagaraBakerOutputStaticMesh.cpp
@@ -0,0 +1,28 @@
+// Copyright Epic Games, Inc. All Rights Reserved.
+
+#include "NiagaraBakerOutputStaticMesh.h"
+
+#include UE_INLINE_GENERATED_CPP_BY_NAME(NiagaraBakerOutputStaticMesh)
+
+bool UNiagaraBakerOutputStaticMesh::Equals(const UNiagaraBakerOutput& OtherBase) const
+{
+	const UNiagaraBakerOutputStaticMesh& Other = *CastChecked<UNiagaraBakerOutputStaticMesh>(&OtherBase);
+	return
+		Super::Equals(Other);
+}
+
+#if WITH_EDITOR
+FString UNiagaraBakerOutputStaticMesh::MakeOutputName() const
+{
+	return FString::Printf(TEXT("StaticMesh_%d"), GetFName().GetNumber());
+}
+#endif
+
+#if WITH_EDITORONLY_DATA
+void UNiagaraBakerOutputStaticMesh::PostEditChangeProperty(struct FPropertyChangedEvent& PropertyChangedEvent)
+{
+	Super::PostEditChangeProperty(PropertyChangedEvent);
+
+}
+#endif
+
diff --git a/Engine/Plugins/FX/Niagara/Source/NiagaraEditor/NiagaraEditor.Build.cs b/Engine/Plugins/FX/Niagara/Source/NiagaraEditor/NiagaraEditor.Build.cs
--- a/Engine/Plugins/FX/Niagara/Source/NiagaraEditor/NiagaraEditor.Build.cs
+++ b/Engine/Plugins/FX/Niagara/Source/NiagaraEditor/NiagaraEditor.Build.cs
@@ -38,6 +38,8 @@ public class NiagaraEditor : ModuleRules
                 "AdvancedPreviewScene",
 				"Projects",
                 "MainFrame",
+				"MeshDescription",
+				"StaticMeshDescription",
 				"ToolMenus",
 				"Renderer",
 				"EditorWidgets",
diff --git a/Engine/Plugins/FX/Niagara/Source/NiagaraEditor/Private/Customizations/NiagaraBakerOutputCustomization.cpp b/Engine/Plugins/FX/Niagara/Source/NiagaraEditor/Private/Customizations/NiagaraBakerOutputCustomization.cpp
--- a/Engine/Plugins/FX/Niagara/Source/NiagaraEditor/Private/Customizations/NiagaraBakerOutputCustomization.cpp
+++ b/Engine/Plugins/FX/Niagara/Source/NiagaraEditor/Private/Customizations/NiagaraBakerOutputCustomization.cpp
@@ -2,6 +2,7 @@
 
 #include "NiagaraBakerOutputCustomization.h"
 #include "NiagaraBakerOutputSimCache.h"
+#include "NiagaraBakerOutputStaticMesh.h"
 #include "NiagaraBakerOutputTexture2D.h"
 #include "NiagaraBakerOutputVolumeTexture.h"
 #include "NiagaraBakerOutputSparseVolumeTexture.h"
@@ -144,6 +145,64 @@ void FNiagaraBakerOutputSimCacheDetails::CustomizeDetails(IDetailLayoutBuilder&
 
 //////////////////////////////////////////////////////////////////////////
 
+TSharedRef<IDetailCustomization> FNiagaraBakerOutputStaticMeshDetails::MakeInstance()
+{
+	return MakeShared<FNiagaraBakerOutputStaticMeshDetails>();
+}
+
+FText FNiagaraBakerOutputStaticMeshDetails::BrowseAssetToolTipText(TWeakObjectPtr<UNiagaraBakerOutputStaticMesh> WeakOutput)
+{
+	UNiagaraBakerOutputStaticMesh* Output = WeakOutput.Get();
+	return FText::Format(LOCTEXT("BrowseAssetToolTipFormat", "Browse to asset '{0}'"), FText::FromString(Output ? Output->GetAssetPath(Output->FramesAssetPathFormat, 0) : FString()));
+}
+
+FReply FNiagaraBakerOutputStaticMeshDetails::BrowseToAsset(TWeakObjectPtr<UNiagaraBakerOutputStaticMesh> WeakOutput)
+{
+	if (UNiagaraBakerOutputStaticMesh* Output = WeakOutput.Get())
+	{
+		const FString AssetPath = Output->GetAssetPath(Output->FramesAssetPathFormat, 0);
+		FocusContentBrowserToAsset(AssetPath);
+	}
+	return FReply::Handled();
+}
+
+void FNiagaraBakerOutputStaticMeshDetails::CustomizeDetails(IDetailLayoutBuilder& DetailBuilder)
+{
+	// We only support customization on 1 object
+	TArray<TWeakObjectPtr<UObject>> ObjectsCustomized;
+	DetailBuilder.GetObjectsBeingCustomized(ObjectsCustomized);
+	if (ObjectsCustomized.Num() != 1 || !ObjectsCustomized[0]->IsA<UNiagaraBakerOutputStaticMesh>())
+	{
+		return;
+	}
+	TWeakObjectPtr<UNiagaraBakerOutputStaticMesh> WeakOutput = CastChecked<UNiagaraBakerOutputStaticMesh>(ObjectsCustomized[0]);
+
+	static FName NAME_Settings("Settings");
+	IDetailCategoryBuilder& DetailCategory = DetailBuilder.EditCategory(NAME_Settings);
+
+	TArray<TSharedRef<IPropertyHandle>> CategoryProperties;
+	DetailCategory.GetDefaultProperties(CategoryProperties);
+
+	for (TSharedRef<IPropertyHandle>& PropertyHandle : CategoryProperties)
+	{
+		const FName PropertyName = PropertyHandle->GetProperty() ? PropertyHandle->GetProperty()->GetFName() : NAME_None;
+		if (PropertyName == GET_MEMBER_NAME_CHECKED(UNiagaraBakerOutputStaticMesh, FramesAssetPathFormat))
+		{
+			BuildAssetPathWidget(
+				DetailCategory, PropertyHandle,
+				TAttribute<FText>::FGetter::CreateStatic(&FNiagaraBakerOutputStaticMeshDetails::BrowseAssetToolTipText, WeakOutput),
+				FOnClicked::CreateStatic(&FNiagaraBakerOutputStaticMeshDetails::BrowseToAsset, WeakOutput)
+			);
+		}
+		else
+		{
+			DetailCategory.AddProperty(PropertyHandle);
+		}
+	}
+}
+
+//////////////////////////////////////////////////////////////////////////
+
 TSharedRef<IDetailCustomization> FNiagaraBakerOutputTexture2DDetails::MakeInstance()
 {
 	return MakeShared<FNiagaraBakerOutputTexture2DDetails>();
diff --git a/Engine/Plugins/FX/Niagara/Source/NiagaraEditor/Private/Customizations/NiagaraBakerOutputCustomization.h b/Engine/Plugins/FX/Niagara/Source/NiagaraEditor/Private/Customizations/NiagaraBakerOutputCustomization.h
--- a/Engine/Plugins/FX/Niagara/Source/NiagaraEditor/Private/Customizations/NiagaraBakerOutputCustomization.h
+++ b/Engine/Plugins/FX/Niagara/Source/NiagaraEditor/Private/Customizations/NiagaraBakerOutputCustomization.h
@@ -10,6 +10,7 @@ class IDetailCategoryBuilder;
 class IPropertyHandle;
 
 class UNiagaraBakerOutputSimCache;
+class UNiagaraBakerOutputStaticMesh;
 class UNiagaraBakerOutputTexture2D;
 class UNiagaraBakerOutputVolumeTexture;
 class UNiagaraBakerOutputSparseVolumeTexture;
@@ -31,6 +32,16 @@ struct FNiagaraBakerOutputSimCacheDetails : public FNiagaraBakerOutputDetails
 	virtual void CustomizeDetails(IDetailLayoutBuilder& DetailBuilder) override;
 };
 
+struct FNiagaraBakerOutputStaticMeshDetails : public FNiagaraBakerOutputDetails
+{
+	static TSharedRef<IDetailCustomization> MakeInstance();
+
+	static FText BrowseAssetToolTipText(TWeakObjectPtr<UNiagaraBakerOutputStaticMesh> WeakOutput);
+	static FReply BrowseToAsset(TWeakObjectPtr<UNiagaraBakerOutputStaticMesh> WeakOutput);
+
+	virtual void CustomizeDetails(IDetailLayoutBuilder& DetailBuilder) override;
+};
+
 struct FNiagaraBakerOutputTexture2DDetails : public FNiagaraBakerOutputDetails
 {
 	static TSharedRef<IDetailCustomization> MakeInstance();
diff --git a/Engine/Plugins/FX/Niagara/Source/NiagaraEditor/Private/NiagaraBakerOutputRegistry.cpp b/Engine/Plugins/FX/Niagara/Source/NiagaraEditor/Private/NiagaraBakerOutputRegistry.cpp
--- a/Engine/Plugins/FX/Niagara/Source/NiagaraEditor/Private/NiagaraBakerOutputRegistry.cpp
+++ b/Engine/Plugins/FX/Niagara/Source/NiagaraEditor/Private/NiagaraBakerOutputRegistry.cpp
@@ -3,10 +3,12 @@
 #include "NiagaraBakerOutputRegistry.h"
 #include "IDetailsView.h"
 #include "NiagaraBakerOutputSimCache.h"
+#include "NiagaraBakerOutputStaticMesh.h"
 #include "NiagaraBakerOutputTexture2D.h"
 #include "NiagaraBakerOutputVolumeTexture.h"
 #include "NiagaraBakerOutputSparseVolumeTexture.h"
 #include "NiagaraBakerRendererOutputSimCache.h"
+#include "NiagaraBakerRendererOutputStaticMesh.h"
 #include "NiagaraBakerRendererOutputTexture2D.h"
 #include "NiagaraBakerRendererOutputVolumeTexture.h"
 #include "NiagaraBakerRendererOutputSparseVolumeTexture.h"
@@ -26,6 +28,13 @@ FNiagaraBakerOutputRegistry& FNiagaraBakerOutputRegistry::Get()
 			Entry.CreateRenderer		= []() -> FNiagaraBakerOutputRenderer* { return new FNiagaraBakerRendererOutputSimCache(); };
 			Entry.CreateCustomization	= FOnGetDetailCustomizationInstance::CreateStatic(&FNiagaraBakerOutputSimCacheDetails::MakeInstance);
 		}
+		// Static Mesh
+		{
+			FRegistryEntry& Entry = Singleton->Registry.AddDefaulted_GetRef();
+			Entry.WeakClass				= UNiagaraBakerOutputStaticMesh::StaticClass();
+			Entry.CreateRenderer		= []() -> FNiagaraBakerOutputRenderer* { return new FNiagaraBakerRendererOutputStaticMesh(); };
+			Entry.CreateCustomization	= FOnGetDetailCustomizationInstance::CreateStatic(&FNiagaraBakerOutputStaticMeshDetails::MakeInstance);
+		}
 		// Texture 2D
 		{
 			FRegistryEntry& Entry		= Singleton->Registry.AddDefaulted_GetRef();
diff --git a/Engine/Plugins/FX/Niagara/Source/NiagaraEditor/Private/NiagaraBakerRenderer.cpp b/Engine/Plugins/FX/Niagara/Source/NiagaraEditor/Private/NiagaraBakerRenderer.cpp
--- a/Engine/Plugins/FX/Niagara/Source/NiagaraEditor/Private/NiagaraBakerRenderer.cpp
+++ b/Engine/Plugins/FX/Niagara/Source/NiagaraEditor/Private/NiagaraBakerRenderer.cpp
@@ -14,6 +14,7 @@
 #include "NiagaraDataInterfaceGrid2DCollection.h"
 
 #include "Components/SceneCaptureComponent2D.h"
+#include "Components/StaticMeshComponent.h"
 #include "Engine/TextureRenderTarget2D.h"
 #include "Misc/FileHelper.h"
 #include "Misc/Paths.h"
@@ -37,9 +38,59 @@
 #include "SparseVolumeTexture/SparseVolumeTexture.h"
 #include "Components/HeterogeneousVolumeComponent.h"
 
+DEFINE_LOG_CATEGORY(LogNiagaraBaker);
+
 //////////////////////////////////////////////////////////////////////////
 
-DEFINE_LOG_CATEGORY(LogNiagaraBaker);
+namespace NiagaraBakerRendererPrivate
+{
+	static void CreatePreviewScene(UNiagaraSystem* NiagaraSystem, TObjectPtr<UNiagaraComponent>& OutComponent, TSharedPtr<FAdvancedPreviewScene>& OutPreviewScene)
+	{
+		check(NiagaraSystem);
+		OutComponent = NewObject<UNiagaraComponent>(GetTransientPackage(), NAME_None, RF_Transient);
+		OutComponent->CastShadow = 1;
+		OutComponent->bCastDynamicShadow = 1;
+		OutComponent->SetAllowScalability(false);
+		OutComponent->SetAsset(NiagaraSystem);
+		OutComponent->SetForceSolo(true);
+		OutComponent->SetAgeUpdateMode(ENiagaraAgeUpdateMode::DesiredAge);
+		OutComponent->SetCanRenderWhileSeeking(true);
+		OutComponent->SetMaxSimTime(0.0f);
+		OutComponent->Activate(true);
+
+		OutPreviewScene = MakeShareable(new FAdvancedPreviewScene(FPreviewScene::ConstructionValues()));
+		OutPreviewScene->SetFloorVisibility(false);
+		OutPreviewScene->AddComponent(OutComponent, OutComponent->GetRelativeTransform());
+	}
+
+	template<typename TComponentType>
+	static void CreatePreviewScene(TObjectPtr<TComponentType>& OutComponent, TSharedPtr<FAdvancedPreviewScene>& OutPreviewScene)
+	{
+		OutComponent = NewObject<TComponentType>(GetTransientPackage(), NAME_None, RF_Transient);
+
+		OutPreviewScene = MakeShareable(new FAdvancedPreviewScene(FPreviewScene::ConstructionValues()));
+		OutPreviewScene->SetFloorVisibility(false);
+		OutPreviewScene->AddComponent(OutComponent, OutComponent->GetRelativeTransform());
+	}
+
+	template<typename TComponentType>
+	static void DestroyPreviewScene(TObjectPtr<TComponentType>& InOutComponent, TSharedPtr<FAdvancedPreviewScene>& InOutPreviewScene)
+	{
+		if (InOutPreviewScene && InOutComponent)
+		{
+			InOutPreviewScene->RemoveComponent(InOutComponent);
+			InOutPreviewScene = nullptr;
+		}
+
+		if (InOutComponent)
+		{
+			InOutComponent->DestroyComponent();
+			InOutComponent = nullptr;
+		}
+	}
+}
+
+//////////////////////////////////////////////////////////////////////////
 
 const FString FNiagaraBakerOutputBindingHelper::STRING_SceneCaptureSource("SceneCaptureSource");
 const FString FNiagaraBakerOutputBindingHelper::STRING_BufferVisualization("BufferVisualization");
@@ -276,6 +327,8 @@ void FNiagaraBakerOutputBindingHelper::GetParticleAttributeBindings(TArray<FNiag
 FNiagaraBakerRenderer::FNiagaraBakerRenderer(UNiagaraSystem* InNiagaraSystem)
 	: NiagaraSystem(InNiagaraSystem)
 {
+	using namespace NiagaraBakerRendererPrivate;
+
 	SceneCaptureComponent = NewObject<USceneCaptureComponent2D>(GetTransientPackage(), NAME_None, RF_Transient);
 	SceneCaptureComponent->bTickInEditor = false;
 	SceneCaptureComponent->SetComponentTickEnabled(false);
@@ -288,20 +341,11 @@ FNiagaraBakerRenderer::FNiagaraBakerRenderer(UNiagaraSystem* InNiagaraSystem)
 
 FNiagaraBakerRenderer::~FNiagaraBakerRenderer()
 {
+	using namespace NiagaraBakerRendererPrivate;
 	DestroyPreviewScene(PreviewComponent, AdvancedPreviewScene);
 	DestroyPreviewScene(SimCachePreviewComponent, SimCacheAdvancedPreviewScene);
-
-	// Clean up SVT preview scene
-	if (SVTPreviewScene && SVTPreviewComponent)
-	{
-		SVTPreviewScene->RemoveComponent(SVTPreviewComponent);
-		SVTPreviewScene = nullptr;
-	}
-	if (SVTPreviewComponent)
-	{
-		SVTPreviewComponent->DestroyComponent();
-		SVTPreviewComponent = nullptr;
-	}
+	DestroyPreviewScene(StaticMeshPreviewComponent, StaticMeshPreviewScene);
+	DestroyPreviewScene(SVTPreviewComponent, SVTPreviewScene);
 }
 
 void FNiagaraBakerRenderer::SetAbsoluteTime(float AbsoluteTime, bool bShouldTickComponent)
@@ -614,6 +658,8 @@ void FNiagaraBakerRenderer::RenderParticleAttribute(UTextureRenderTarget2D* Rend
 
 void FNiagaraBakerRenderer::RenderSimCache(UTextureRenderTarget2D* RenderTarget, UNiagaraSimCache* SimCache) const
 {
+	using namespace NiagaraBakerRendererPrivate;
+
 	UNiagaraBakerSettings* BakerSettings = GetBakerSettings();
 	if (!SimCache)
 	{
@@ -714,6 +760,31 @@ void FNiagaraBakerRenderer::RenderSparseVolumeTexture(UTextureRenderTarget2D* Re
 	RenderSceneCapture(RenderTarget, SVTPreviewComponent, ESceneCaptureSource::SCS_SceneColorHDR);
 }
 
+void FNiagaraBakerRenderer::RenderStaticMesh(UTextureRenderTarget2D* RenderTarget, UStaticMesh* StaticMesh) const
+{
+	using namespace NiagaraBakerRendererPrivate;
+
+	UNiagaraBakerSettings* BakerSettings = GetBakerSettings();
+	if (!StaticMesh)
+	{
+		return;
+	}
+
+	if (StaticMeshPreviewComponent == nullptr)
+	{
+		CreatePreviewScene<UStaticMeshComponent>(StaticMeshPreviewComponent, StaticMeshPreviewScene);
+	}
+
+	StaticMeshPreviewComponent->SetStaticMesh(StaticMesh);
+
+	UWorld* World = StaticMeshPreviewComponent->GetWorld();
+	World->SendAllEndOfFrameUpdates();
+
+	RenderSceneCapture(RenderTarget, StaticMeshPreviewComponent, ESceneCaptureSource::SCS_SceneColorHDR);
+
+	StaticMeshPreviewComponent->SetStaticMesh(nullptr);
+}
+
 UWorld* FNiagaraBakerRenderer::GetWorld() const
 {
 	return PreviewComponent->GetWorld();
@@ -797,39 +868,6 @@ bool FNiagaraBakerRenderer::ExportImage(FStringView FilePath, FIntPoint ImageSiz
 	return FFileHelper::SaveArrayToFile(TempData, FilePath.GetData());
 }
 
-void FNiagaraBakerRenderer::CreatePreviewScene(UNiagaraSystem* NiagaraSystem, TObjectPtr<UNiagaraComponent>& OutComponent, TSharedPtr<FAdvancedPreviewScene>& OutPreviewScene)
-{
-	check(NiagaraSystem);
-	OutComponent = NewObject<UNiagaraComponent>(GetTransientPackage(), NAME_None, RF_Transient);
-	OutComponent->CastShadow = 1;
-	OutComponent->bCastDynamicShadow = 1;
-	OutComponent->SetAllowScalability(false);
-	OutComponent->SetAsset(NiagaraSystem);
-	OutComponent->SetForceSolo(true);
-	OutComponent->SetAgeUpdateMode(ENiagaraAgeUpdateMode::DesiredAge);
-	OutComponent->SetCanRenderWhileSeeking(true);
-	OutComponent->SetMaxSimTime(0.0f);
-	OutComponent->Activate(true);
-
-	OutPreviewScene = MakeShareable(new FAdvancedPreviewScene(FPreviewScene::ConstructionValues()));
-	OutPreviewScene->SetFloorVisibility(false);
-	OutPreviewScene->AddComponent(OutComponent, OutComponent->GetRelativeTransform());
-}
-
-void FNiagaraBakerRenderer::DestroyPreviewScene(TObjectPtr<UNiagaraComponent>& InOutComponent, TSharedPtr<FAdvancedPreviewScene>& InOutPreviewScene)
-{
-	if ( InOutPreviewScene && InOutComponent )
-	{
-		InOutPreviewScene->RemoveComponent(InOutComponent);
-		InOutPreviewScene = nullptr;
-	}
-
-	if ( InOutComponent )
-	{
-		InOutComponent->DestroyComponent();
-		InOutComponent = nullptr;
-	}
-}
 bool FNiagaraBakerRenderer::ExportVolume(FStringView FilePath, FIntVector ImageSize, TArrayView<FFloat16Color> ImageData)
 {
 	const FString FileExtension = FPaths::GetExtension(FilePath.GetData(), true);
@@ -845,8 +883,6 @@ bool FNiagaraBakerRenderer::ExportVolume(FStringView FilePath, FIntVector ImageS
 	}
 }
 
-
-
 bool FVolumeDataInterfaceHelper::Initialize(const TArray<FString>& InputDataInterfacePath, UNiagaraComponent* InNiagaraComponent)
 {
 	NiagaraComponent = InNiagaraComponent;
diff --git a/Engine/Plugins/FX/Niagara/Source/NiagaraEditor/Private/NiagaraBakerRenderer.h b/Engine/Plugins/FX/Niagara/Source/NiagaraEditor/Private/NiagaraBakerRenderer.h
--- a/Engine/Plugins/FX/Niagara/Source/NiagaraEditor/Private/NiagaraBakerRenderer.h
+++ b/Engine/Plugins/FX/Niagara/Source/NiagaraEditor/Private/NiagaraBakerRenderer.h
@@ -21,6 +21,8 @@ class FAdvancedPreviewScene;
 class FCanvas;
 class UAnimatedSparseVolumeTexture;
 class UHeterogeneousVolumeComponent;
+class UStaticMesh;
+class UStaticMeshComponent;
 
 struct FNiagaraBakerFeedbackContext
 {
@@ -112,6 +114,7 @@ public:
 	void RenderParticleAttribute(UTextureRenderTarget2D* RenderTarget, FName BindingName) const;
 	void RenderSimCache(UTextureRenderTarget2D* RenderTarget, UNiagaraSimCache* SimCache) const;
 	void RenderSparseVolumeTexture(UTextureRenderTarget2D* RenderTarget, const FNiagaraBakerOutputFrameIndices Indices, UAnimatedSparseVolumeTexture *SVT) const;
+	void RenderStaticMesh(UTextureRenderTarget2D* RenderTarget, UStaticMesh* StaticMesh) const;
 
 	UWorld* GetWorld() const;
 	float GetWorldTime() const;
@@ -133,9 +136,6 @@ public:
 
 	static bool ExportImage(FStringView FilePath, FIntPoint ImageSize, TArrayView<FFloat16Color> ImageData);
 	static bool ExportVolume(FStringView FilePath, FIntVector ImageSize, TArrayView<FFloat16Color> ImageData);
-private:
-	static void CreatePreviewScene(UNiagaraSystem* NiagaraSystem, TObjectPtr<UNiagaraComponent>& OutComponent, TSharedPtr<FAdvancedPreviewScene>& OutPreviewScene);
-	static void DestroyPreviewScene(TObjectPtr<UNiagaraComponent>& InOutComponent, TSharedPtr<FAdvancedPreviewScene>& InOutPreviewScene);
 
 private:
 	TObjectPtr<UNiagaraSystem> NiagaraSystem = nullptr;
@@ -144,9 +144,12 @@ private:
 	TObjectPtr<USceneCaptureComponent2D> SceneCaptureComponent = nullptr;
 
 	mutable TObjectPtr<UNiagaraComponent> SimCachePreviewComponent = nullptr;
-	mutable TObjectPtr<UHeterogeneousVolumeComponent> SVTPreviewComponent = nullptr;
-
 	mutable TSharedPtr<FAdvancedPreviewScene> SimCacheAdvancedPreviewScene;
+
+	mutable TObjectPtr<UStaticMeshComponent> StaticMeshPreviewComponent = nullptr;
+	mutable TSharedPtr<FAdvancedPreviewScene> StaticMeshPreviewScene;
+
+	mutable TObjectPtr<UHeterogeneousVolumeComponent> SVTPreviewComponent = nullptr;
 	mutable TSharedPtr<FAdvancedPreviewScene> SVTPreviewScene;
 };
 
diff --git a/Engine/Plugins/FX/Niagara/Source/NiagaraEditor/Private/NiagaraBakerRendererOutputStaticMesh.cpp b/Engine/Plugins/FX/Niagara/Source/NiagaraEditor/Private/NiagaraBakerRendererOutputStaticMesh.cpp
new file mode 100644
--- /dev/null
+++ b/Engine/Plugins/FX/Niagara/Source/NiagaraEditor/Private/NiagaraBakerRendererOutputStaticMesh.cpp
@@ -0,0 +1,211 @@
+// Copyright Epic Games, Inc. All Rights Reserved.
+
+#include "NiagaraBakerRendererOutputStaticMesh.h"
+#include "NiagaraBakerOutputStaticMesh.h"
+
+#include "NiagaraBakerSettings.h"
+#include "NiagaraComponent.h"
+#include "NiagaraGpuComputeDispatchInterface.h"
+#include "NiagaraRendererReadback.h"
+#include "NiagaraSystemInstance.h"
+
+#include "Engine/StaticMesh.h"
+#include "Engine/TextureRenderTarget2D.h"
+#include "MeshDescription.h"
+#include "MeshAttributes.h"
+#include "StaticMeshAttributes.h"
+#include "StaticMeshOperations.h"
+
+UNiagaraBakerStaticMeshFactoryNew::UNiagaraBakerStaticMeshFactoryNew()
+{
+	SupportedClass = UStaticMesh::StaticClass();
+}
+
+UObject* UNiagaraBakerStaticMeshFactoryNew::FactoryCreateNew(UClass* InClass, UObject* InParent, FName InName, EObjectFlags Flags, UObject* Context, FFeedbackContext* Warn)
+{
+	return NewObject<UStaticMesh>(InParent, InClass, InName, Flags);
+}
+
+TArray<FNiagaraBakerOutputBinding> FNiagaraBakerRendererOutputStaticMesh::GetRendererBindings(UNiagaraBakerOutput* InBakerOutput) const
+{
+	return TArray<FNiagaraBakerOutputBinding>();
+}
+
+FIntPoint FNiagaraBakerRendererOutputStaticMesh::GetPreviewSize(UNiagaraBakerOutput* InBakerOutput, FIntPoint InAvailableSize) const
+{
+	return InAvailableSize;
+}
+
+void FNiagaraBakerRendererOutputStaticMesh::RenderPreview(UNiagaraBakerOutput* InBakerOutput, const FNiagaraBakerRenderer& BakerRenderer, UTextureRenderTarget2D* RenderTarget, TOptional<FString>& OutErrorString) const
+{
+	BakerRenderer.RenderSceneCapture(RenderTarget, ESceneCaptureSource::SCS_SceneColorHDR);
+}
+
+FIntPoint FNiagaraBakerRendererOutputStaticMesh::GetGeneratedSize(UNiagaraBakerOutput* InBakerOutput, FIntPoint InAvailableSize) const
+{
+	return InAvailableSize;
+}
+
+void FNiagaraBakerRendererOutputStaticMesh::RenderGenerated(UNiagaraBakerOutput* InBakerOutput, const FNiagaraBakerRenderer& BakerRenderer, UTextureRenderTarget2D* RenderTarget, TOptional<FString>& OutErrorString) const
+{
+	static FString StaticMeshNotFoundError(TEXT("StaticMesh asset not found.\nPlease bake to see the result."));
+
+	UNiagaraBakerOutputStaticMesh* BakerOutput = CastChecked<UNiagaraBakerOutputStaticMesh>(InBakerOutput);
+	UNiagaraBakerSettings* BakerGeneratedSettings = BakerOutput->GetTypedOuter<UNiagaraBakerSettings>();
+
+	const float WorldTime = BakerRenderer.GetWorldTime();
+	const FNiagaraBakerOutputFrameIndices FrameIndices = BakerGeneratedSettings->GetOutputFrameIndices(BakerOutput, WorldTime);
+
+	UStaticMesh* StaticMesh = BakerOutput->GetAsset<UStaticMesh>(BakerOutput->FramesAssetPathFormat, FrameIndices.FrameIndexA);
+	if (StaticMesh == nullptr)
+	{
+		OutErrorString = StaticMeshNotFoundError;
+		return;
+	}
+
+	BakerRenderer.RenderStaticMesh(RenderTarget, StaticMesh);
+}
+
+bool FNiagaraBakerRendererOutputStaticMesh::BeginBake(FNiagaraBakerFeedbackContext& FeedbackContext, UNiagaraBakerOutput* InBakerOutput)
+{
+#if WITH_NIAGARA_RENDERER_READBACK
+	BakeRenderTarget = NewObject<UTextureRenderTarget2D>();
+	BakeRenderTarget->AddToRoot();
+	BakeRenderTarget->ClearColor = FLinearColor::Transparent;
+	BakeRenderTarget->TargetGamma = 1.0f;
+	BakeRenderTarget->InitCustomFormat(128, 128, PF_FloatRGBA, false);
+
+	return true;
+#else
+	FeedbackContext.Errors.Add(TEXT("Niagara Renderer Readback not enabled, failed to bake"));
+	return false;
+#endif
+}
+
+void FNiagaraBakerRendererOutputStaticMesh::BakeFrame(FNiagaraBakerFeedbackContext& FeedbackContext, UNiagaraBakerOutput* InBakerOutput, int FrameIndex, const FNiagaraBakerRenderer& BakerRenderer)
+{
+	UNiagaraBakerOutputStaticMesh* BakerOutput = CastChecked<UNiagaraBakerOutputStaticMesh>(InBakerOutput);
+
+	UNiagaraComponent* PreviewComponent = BakerRenderer.GetPreviewComponent();
+	if (!PreviewComponent || PreviewComponent->IsComplete())
+	{
+		return;
+	}
+
+#if WITH_NIAGARA_RENDERER_READBACK
+	NiagaraRendererReadback::EnqueueReadback(
+		PreviewComponent,
+		[BakerOutput, FrameIndex](const FNiagaraRendererReadbackResult& Result)
+		{
+			// Failed or no data
+			if (Result.NumVertices == 0)
+			{
+				return;
+			}
+
+			// Create asset
+			const FString AssetFullName = BakerOutput->GetAssetPath(BakerOutput->FramesAssetPathFormat, FrameIndex);
+			UStaticMesh* StaticMesh = UNiagaraBakerOutput::GetOrCreateAsset<UStaticMesh, UNiagaraBakerStaticMeshFactoryNew>(AssetFullName);
+			if (StaticMesh == nullptr)
+			{
+				return;
+			}
+
+			const uint32 NumTexCoords = Result.VertexTexCoordNum;
+			const uint32 NumTriangles = Result.NumVertices / 3;
+
+			// Create Mesh Description
+			FMeshDescription MeshDescription;
+			TVertexAttributesRef<FVector3f> VertexPositions = MeshDescription.GetVertexPositions();
+			TVertexInstanceAttributesRef<FVector3f> VertexInstanceNormals = MeshDescription.VertexInstanceAttributes().RegisterAttribute<FVector3f>(MeshAttribute::VertexInstance::Normal, 1, FVector3f::ZeroVector, EMeshAttributeFlags::Mandatory);
+			TVertexInstanceAttributesRef<FVector3f> VertexInstanceTangents = MeshDescription.VertexInstanceAttributes().RegisterAttribute<FVector3f>(MeshAttribute::VertexInstance::Tangent, 1, FVector3f::ZeroVector, EMeshAttributeFlags::Mandatory);
+			TVertexInstanceAttributesRef<float> VertexInstanceBinormalSigns = MeshDescription.VertexInstanceAttributes().RegisterAttribute<float>(MeshAttribute::VertexInstance::BinormalSign, 1, 1.0f, EMeshAttributeFlags::Mandatory);
+			TVertexInstanceAttributesRef<FVector4f> VertexInstanceColors = MeshDescription.VertexInstanceAttributes().RegisterAttribute<FVector4f>(MeshAttribute::VertexInstance::Color, 1, FVector4f(1.0f, 1.0f, 1.0f, 1.0f), EMeshAttributeFlags::Lerpable | EMeshAttributeFlags::Mandatory);
+			TVertexInstanceAttributesRef<FVector2f> VertexInstanceUVs = MeshDescription.VertexInstanceAttributes().RegisterAttribute<FVector2f>(MeshAttribute::VertexInstance::TextureCoordinate, NumTexCoords, FVector2f::ZeroVector, EMeshAttributeFlags::Lerpable | EMeshAttributeFlags::Mandatory);
+
+			MeshDescription.EdgeAttributes().RegisterAttribute<bool>(MeshAttribute::Edge::IsHard, 1, false, EMeshAttributeFlags::Mandatory);
+			TPolygonGroupAttributesRef<FName> PolygonGroupSlotNames = MeshDescription.PolygonGroupAttributes().RegisterAttribute<FName>(MeshAttribute::PolygonGroup::ImportedMaterialSlotName, 1, NAME_None, EMeshAttributeFlags::Mandatory); //The unique key to match the mesh material slot
+
+			// Reserve space
+			MeshDescription.ReserveNewVertices(Result.NumVertices);
+			MeshDescription.ReserveNewVertexInstances(NumTriangles);
+			MeshDescription.ReserveNewEdges(NumTriangles);
+			MeshDescription.ReserveNewPolygons(NumTriangles);
+			MeshDescription.ReserveNewPolygonGroups(Result.Sections.Num());
+
+			// Build vertices
+			TArray<FVertexInstanceID> VertexInstances;
+			VertexInstances.Reserve(Result.NumVertices);
+			for (uint32 iVertex=0; iVertex < Result.NumVertices; ++iVertex)
+			{
+				const FVertexID VertexID = MeshDescription.CreateVertex();
+				check(VertexID.GetValue() == iVertex);
+				VertexPositions[VertexID] = Result.GetPosition(iVertex);
+
+				const FVector3f TangentX = Result.GetTangentX(iVertex);
+				const FVector3f TangentY = Result.GetTangentY(iVertex);
+				const FVector3f TangentZ = Result.GetTangentZ(iVertex);
+				const float TangentSign = FVector3f::DotProduct(FVector3f::CrossProduct(TangentX, TangentZ), TangentY) < 0.0f ? -1.0f : 1.0f;
+
+				FVertexInstanceID VertexInstanceID = MeshDescription.CreateVertexInstance(VertexID);
+				VertexInstances.Add(VertexInstanceID);
+				VertexInstanceNormals.Set(VertexInstanceID, TangentX);
+				VertexInstanceTangents.Set(VertexInstanceID, TangentZ);
+				VertexInstanceBinormalSigns.Set(VertexInstanceID, TangentSign);
+				VertexInstanceColors.Set(VertexInstanceID, Result.GetColor(iVertex));
+				for (uint32 iTexCoord=0; iTexCoord < NumTexCoords; ++iTexCoord)
+				{
+					VertexInstanceUVs.Set(VertexInstanceID, iTexCoord, Result.GetTexCoord(iVertex, iTexCoord));
+				}
+			}
+
+			// Build sections / triangles
+			TArray<FStaticMaterial> StaticMaterials;
+			for (int32 iSection=0; iSection < Result.Sections.Num(); ++iSection)
+			{
+				const FNiagaraRendererReadbackResult::FSection& Section = Result.Sections[iSection];
+
+				const FName MaterialSlotName(*FString::Printf(TEXT("Section%d"), iSection));
+				const FPolygonGroupID PolyGroupId = MeshDescription.CreatePolygonGroup();
+				PolygonGroupSlotNames[PolyGroupId] = MaterialSlotName;
+
+				FStaticMaterial& StaticMaterial			= StaticMaterials.AddDefaulted_GetRef();
+				StaticMaterial.MaterialInterface		= Section.WeakMaterialInterface.Get();
+				StaticMaterial.MaterialSlotName			= MaterialSlotName;
+				StaticMaterial.ImportedMaterialSlotName	= MaterialSlotName;
+
+				for (uint32 iTriangle=0; iTriangle < Section.NumTriangles; ++iTriangle)
+				{
+					const uint32 BaseIndex = (Section.FirstTriangle + iTriangle) * 3;
+					const FPolygonID PolygonID = MeshDescription.CreatePolygon(PolyGroupId, { VertexInstances[BaseIndex + 0], VertexInstances[BaseIndex + 1], VertexInstances[BaseIndex + 2] });
+					MeshDescription.ComputePolygonTriangulation(PolygonID);
+				}
+			}
+
+			StaticMesh->SetStaticMaterials(StaticMaterials);
+
+			// Build Mesh
+			UStaticMesh::FBuildMeshDescriptionsParams Params;
+			Params.bFastBuild = WITH_EDITOR ? false : true;
+			Params.bUseHashAsGuid = true;
+			Params.bMarkPackageDirty = false;
+			Params.bCommitMeshDescription = true;
+			Params.bAllowCpuAccess = false;
+
+			StaticMesh->BuildFromMeshDescriptions({ &MeshDescription }, Params);
+		}
+	);
+
+	BakerRenderer.RenderSceneCapture(BakeRenderTarget, ESceneCaptureSource::SCS_SceneColorHDR);
+#endif
+}
+
+void FNiagaraBakerRendererOutputStaticMesh::EndBake(FNiagaraBakerFeedbackContext& FeedbackContext, UNiagaraBakerOutput* InBakerOutput)
+{
+	if (BakeRenderTarget)
+	{
+		BakeRenderTarget->RemoveFromRoot();
+		BakeRenderTarget->MarkAsGarbage();
+		BakeRenderTarget = nullptr;
+	}
+}
diff --git a/Engine/Plugins/FX/Niagara/Source/NiagaraEditor/Private/NiagaraBakerRendererOutputStaticMesh.h b/Engine/Plugins/FX/Niagara/Source/NiagaraEditor/Private/NiagaraBakerRendererOutputStaticMesh.h
new file mode 100644
--- /dev/null
+++ b/Engine/Plugins/FX/Niagara/Source/NiagaraEditor/Private/NiagaraBakerRendererOutputStaticMesh.h
@@ -0,0 +1,39 @@
+// Copyright Epic Games, Inc. All Rights Reserved.
+
+#pragma once
+
+#include "NiagaraBakerRenderer.h"
+
+#include "Factories/Factory.h"
+
+#include "NiagaraBakerRendererOutputStaticMesh.generated.h"
+
+class UTextureRenderTarget2D;
+
+UCLASS()
+class UNiagaraBakerStaticMeshFactoryNew : public UFactory
+{
+	GENERATED_BODY()
+
+	UNiagaraBakerStaticMeshFactoryNew();
+	virtual UObject* FactoryCreateNew(UClass* InClass, UObject* InParent, FName InName, EObjectFlags Flags, UObject* Context, FFeedbackContext* Warn) override;
+};
+
+class FNiagaraBakerRendererOutputStaticMesh : public FNiagaraBakerOutputRenderer
+{
+public:
+	virtual TArray<FNiagaraBakerOutputBinding> GetRendererBindings(UNiagaraBakerOutput* InBakerOutput) const override;
+
+	virtual FIntPoint GetPreviewSize(UNiagaraBakerOutput* InBakerOutput, FIntPoint InAvailableSize) const override;
+	virtual void RenderPreview(UNiagaraBakerOutput* InBakerOutput, const FNiagaraBakerRenderer& BakerRenderer, UTextureRenderTarget2D* RenderTarget, TOptional<FString>& OutErrorString) const override;
+
+	virtual FIntPoint GetGeneratedSize(UNiagaraBakerOutput* InBakerOutput, FIntPoint InAvailableSize) const override;
+	virtual void RenderGenerated(UNiagaraBakerOutput* InBakerOutput, const FNiagaraBakerRenderer& BakerRenderer, UTextureRenderTarget2D* RenderTarget, TOptional<FString>& OutErrorString) const override;
+
+	virtual bool BeginBake(FNiagaraBakerFeedbackContext& FeedbackContext, UNiagaraBakerOutput* InBakerOutput) override;
+	virtual void BakeFrame(FNiagaraBakerFeedbackContext& FeedbackContext, UNiagaraBakerOutput* InBakerOutput, int FrameIndex, const FNiagaraBakerRenderer& BakerRenderer) override;
+	virtual void EndBake(FNiagaraBakerFeedbackContext& FeedbackContext, UNiagaraBakerOutput* InBakerOutput) override;
+
+private:
+	UTextureRenderTarget2D* BakeRenderTarget = nullptr;
+};
