- IWYU fixes

#rnx

diff --git a/Engine/Plugins/FX/Niagara/Source/NiagaraEditorWidgets/Private/Stack/SNiagaraDeterminismToggle.cpp b/Engine/Plugins/FX/Niagara/Source/NiagaraEditorWidgets/Private/Stack/SNiagaraDeterminismToggle.cpp
--- a/Engine/Plugins/FX/Niagara/Source/NiagaraEditorWidgets/Private/Stack/SNiagaraDeterminismToggle.cpp
+++ b/Engine/Plugins/FX/Niagara/Source/NiagaraEditorWidgets/Private/Stack/SNiagaraDeterminismToggle.cpp
@@ -10,6 +10,7 @@
 #include "ViewModels/NiagaraEmitterViewModel.h"
 #include "ViewModels/NiagaraEmitterHandleViewModel.h"
 #include "Widgets/Text/STextBlock.h"
+#include "ScopedTransaction.h"
 
 #define LOCTEXT_NAMESPACE "SNiagaraDeterminismToggle"
 
diff --git a/Engine/Plugins/FX/Niagara/Source/NiagaraEditorWidgets/Private/Stack/SNiagaraLocalSpaceToggle.cpp b/Engine/Plugins/FX/Niagara/Source/NiagaraEditorWidgets/Private/Stack/SNiagaraLocalSpaceToggle.cpp
--- a/Engine/Plugins/FX/Niagara/Source/NiagaraEditorWidgets/Private/Stack/SNiagaraLocalSpaceToggle.cpp
+++ b/Engine/Plugins/FX/Niagara/Source/NiagaraEditorWidgets/Private/Stack/SNiagaraLocalSpaceToggle.cpp
@@ -7,6 +7,7 @@
 #include "Widgets/Images/SImage.h"
 #include "ViewModels/NiagaraEmitterViewModel.h"
 #include "ViewModels/NiagaraEmitterHandleViewModel.h"
+#include "ScopedTransaction.h"
 
 #define LOCTEXT_NAMESPACE "SNiagaraLocalSpaceToggle"
 
