Add a cvar to force ribbon creating 32 bits index buffers. Potential work around for a cpu/gpu crash when GNiagaraRibbonTessellationEnabled is enabled.

#replayrun
replayrun 6787e28ce0110cc231d640b5 (preflight was succesfull, but server was not selected so it generated an error)

diff --git a/Engine/Plugins/FX/Niagara/Source/Niagara/Private/NiagaraRendererRibbons.cpp b/Engine/Plugins/FX/Niagara/Source/Niagara/Private/NiagaraRendererRibbons.cpp
--- a/Engine/Plugins/FX/Niagara/Source/Niagara/Private/NiagaraRendererRibbons.cpp
+++ b/Engine/Plugins/FX/Niagara/Source/Niagara/Private/NiagaraRendererRibbons.cpp
@@ -177,6 +177,14 @@ static TAutoConsoleVariable<int32> CVarRayTracingNiagaraRibbonsGPU(
 	1,
 	TEXT("If we allow GPU ribbon raytracing"));
 
+static bool GNiagaraRibbonForceIndex32 = false;
+static FAutoConsoleVariableRef CNiagaraRibbonForceIndex32(
+	TEXT("Niagara.Ribbon.ForceIndex32"),
+	GNiagaraRibbonForceIndex32,
+	TEXT("Force creating 32 bits index buffers for the ribbons."),
+	ECVF_Default
+);
+
 // max absolute error 9.0x10^-3
 // Eberly's polynomial degree 1 - respect bounds
 // input [-1, 1] and output [0, PI]
@@ -2011,7 +2019,7 @@ void FNiagaraRendererRibbons::GenerateIndexBufferForView(
 		else
 		{
 			RenderingViewResources->IndexBuffer = MakeShared<FNiagaraRibbonIndexBuffer>();
-			if (GeneratedData.TotalBitCount <= 16)
+			if (GeneratedData.TotalBitCount <= 16 && !GNiagaraRibbonForceIndex32)
 			{
 				FGlobalDynamicIndexBuffer::FAllocationEx IndexAllocation = Collector.GetDynamicIndexBuffer().Allocate<uint16>(GeneratedData.TotalNumIndices);
 				RenderingViewResources->IndexBuffer->Initialize(RHICmdList, IndexAllocation);
