Tentative fix for crash during validation of niagara systems

diff --git a/Engine/Plugins/FX/Niagara/Source/NiagaraEditor/Private/ViewModels/NiagaraSystemScriptViewModel.cpp b/Engine/Plugins/FX/Niagara/Source/NiagaraEditor/Private/ViewModels/NiagaraSystemScriptViewModel.cpp
--- a/Engine/Plugins/FX/Niagara/Source/NiagaraEditor/Private/ViewModels/NiagaraSystemScriptViewModel.cpp
+++ b/Engine/Plugins/FX/Niagara/Source/NiagaraEditor/Private/ViewModels/NiagaraSystemScriptViewModel.cpp
@@ -16,16 +16,21 @@ FNiagaraSystemScriptViewModel::FNiagaraSystemScriptViewModel(bool bInIsForDataPr
 {
 }
 
-void FNiagaraSystemScriptViewModel::Initialize(UNiagaraSystem& InSystem)
+void FNiagaraSystemScriptViewModel::Initialize(TSharedRef<FNiagaraSystemViewModel> InSystemViewModel)
 {
-	System = &InSystem;
-	SetScript(FVersionedNiagaraScript(System->GetSystemSpawnScript()));
-	System->OnSystemCompiled().AddSP(this, &FNiagaraSystemScriptViewModel::OnSystemVMCompiled);
+	SystemViewModelWeak = InSystemViewModel;
+
+	UNiagaraSystem* System = GetSystem();
+	if (ensure(System))
+	{
+		SetScript(FVersionedNiagaraScript(System->GetSystemSpawnScript()));
+		System->OnSystemCompiled().AddSP(this, &FNiagaraSystemScriptViewModel::OnSystemVMCompiled);
+	}
 }
 
 FNiagaraSystemScriptViewModel::~FNiagaraSystemScriptViewModel()
 {
-	if (System.IsValid())
+	if (UNiagaraSystem* System = GetSystem())
 	{
 		System->OnSystemCompiled().RemoveAll(this);
 	}
@@ -33,9 +38,12 @@ FNiagaraSystemScriptViewModel::~FNiagaraSystemScriptViewModel()
 
 void FNiagaraSystemScriptViewModel::OnSystemVMCompiled(UNiagaraSystem* InSystem)
 {
-	if (InSystem != System.Get())
 	{
-		return;
+		UNiagaraSystem* System = GetSystem();
+		if (!ensure(System) || InSystem != System)
+		{
+			return;
+		}
 	}
 
 	TArray<ENiagaraScriptCompileStatus> InCompileStatuses;
@@ -147,42 +155,51 @@ FNiagaraSystemScriptViewModel::FOnSystemCompiled& FNiagaraSystemScriptViewModel:
 
 void FNiagaraSystemScriptViewModel::CompileSystem(bool bForce)
 {
-	System->RequestCompile(bForce);
+	UNiagaraSystem* System = GetSystem();
+	if (ensure(System))
+	{
+		System->RequestCompile(bForce);
+	}
 }
 
 ENiagaraScriptCompileStatus FNiagaraSystemScriptViewModel::GetLatestCompileStatus(FGuid VersionGuid)
 {
-	TArray<UNiagaraScript*> SystemScripts;
-	SystemScripts.Add(System->GetSystemSpawnScript());
-	SystemScripts.Add(System->GetSystemUpdateScript());
-
-	for (const FNiagaraEmitterHandle& Handle : System->GetEmitterHandles())
+	UNiagaraSystem* System = GetSystem();
+	if (ensure(System))
 	{
-		if (Handle.GetIsEnabled())
+		TArray<UNiagaraScript*> SystemScripts;
+		SystemScripts.Add(System->GetSystemSpawnScript());
+		SystemScripts.Add(System->GetSystemUpdateScript());
+
+		for (const FNiagaraEmitterHandle& Handle : System->GetEmitterHandles())
 		{
-			Handle.GetEmitterData()->GetScripts(SystemScripts, true);
+			if (Handle.GetIsEnabled())
+			{
+				Handle.GetEmitterData()->GetScripts(SystemScripts, true);
+			}
 		}
-	}
 
-	bool bDirty = false;
-	for (int32 i = 0; i < SystemScripts.Num(); i++)
-	{
-		if (!SystemScripts[i])
+		bool bDirty = false;
+		for (int32 i = 0; i < SystemScripts.Num(); i++)
 		{
-			continue;
+			if (!SystemScripts[i])
+			{
+				continue;
+			}
+
+			if (SystemScripts[i]->IsCompilable() && !SystemScripts[i]->AreScriptAndSourceSynchronized(VersionGuid))
+			{
+				bDirty = true;
+				break;
+			}
 		}
 
-		if (SystemScripts[i]->IsCompilable() && !SystemScripts[i]->AreScriptAndSourceSynchronized(VersionGuid))
+		if (bDirty)
 		{
-			bDirty = true;
-			break;
+			return ENiagaraScriptCompileStatus::NCS_Dirty;
 		}
 	}
 
-	if (bDirty)
-	{
-		return ENiagaraScriptCompileStatus::NCS_Dirty;
-	}
 	return LastCompileStatus;
 }
 
diff --git a/Engine/Plugins/FX/Niagara/Source/NiagaraEditor/Private/ViewModels/NiagaraSystemScriptViewModel.h b/Engine/Plugins/FX/Niagara/Source/NiagaraEditor/Private/ViewModels/NiagaraSystemScriptViewModel.h
--- a/Engine/Plugins/FX/Niagara/Source/NiagaraEditor/Private/ViewModels/NiagaraSystemScriptViewModel.h
+++ b/Engine/Plugins/FX/Niagara/Source/NiagaraEditor/Private/ViewModels/NiagaraSystemScriptViewModel.h
@@ -3,6 +3,7 @@
 #pragma once
 
 #include "ViewModels/NiagaraScriptViewModel.h"
+#include "ViewModels/NiagaraSystemViewModel.h"
 
 class FNiagaraSystemViewModel;
 class FNiagaraScriptViewModel;
@@ -17,7 +18,7 @@ public:
 public:
 	explicit FNiagaraSystemScriptViewModel(bool bInIsForDataProcessingOnly);
 
-	void Initialize(UNiagaraSystem& InSystem);
+	void Initialize(TSharedRef<FNiagaraSystemViewModel> InSystemViewModel);
 
 	virtual ~FNiagaraSystemScriptViewModel() override;
 
@@ -28,11 +29,22 @@ public:
 	virtual ENiagaraScriptCompileStatus GetLatestCompileStatus(FGuid VersionGuid = FGuid()) override;
 
 private:
+
 	void OnSystemVMCompiled(UNiagaraSystem* InSystem);
 
+	FORCEINLINE UNiagaraSystem* GetSystem() const
+	{
+		if (const TSharedPtr<FNiagaraSystemViewModel> SystemViewModel = SystemViewModelWeak.Pin())
+		{
+			return &SystemViewModel->GetSystem();
+		}
+
+		return nullptr;
+	}
+
 private:
 	/** The System who's script is getting viewed and edited by this view model. */
-	TWeakObjectPtr<UNiagaraSystem> System;
+	TWeakPtr<FNiagaraSystemViewModel> SystemViewModelWeak;
 
 	FOnSystemCompiled OnSystemCompiledDelegate;
 };
diff --git a/Engine/Plugins/FX/Niagara/Source/NiagaraEditor/Private/ViewModels/NiagaraSystemViewModel.cpp b/Engine/Plugins/FX/Niagara/Source/NiagaraEditor/Private/ViewModels/NiagaraSystemViewModel.cpp
--- a/Engine/Plugins/FX/Niagara/Source/NiagaraEditor/Private/ViewModels/NiagaraSystemViewModel.cpp
+++ b/Engine/Plugins/FX/Niagara/Source/NiagaraEditor/Private/ViewModels/NiagaraSystemViewModel.cpp
@@ -138,7 +138,7 @@ void FNiagaraSystemViewModel::Initialize(UNiagaraSystem& InSystem, FNiagaraSyste
 	SelectionViewModel->OnEmitterHandleIdSelectionChanged().AddSP(this, &FNiagaraSystemViewModel::SystemSelectionChanged);
 
 	SystemScriptViewModel = MakeShared<FNiagaraSystemScriptViewModel>(bIsForDataProcessingOnly);
-	SystemScriptViewModel->Initialize(GetSystem());
+	SystemScriptViewModel->Initialize(this->AsShared());
 
 	OverviewGraphViewModel = MakeShared<FNiagaraOverviewGraphViewModel>();
 	OverviewGraphViewModel->Initialize(this->AsShared());
