Allow for int-enum conversion in niagara python upgrade scripts

#tests no additional testing was done

diff --git a/Engine/Plugins/FX/Niagara/Source/NiagaraEditor/Private/UpgradeNiagaraScriptResults.cpp b/Engine/Plugins/FX/Niagara/Source/NiagaraEditor/Private/UpgradeNiagaraScriptResults.cpp
--- a/Engine/Plugins/FX/Niagara/Source/NiagaraEditor/Private/UpgradeNiagaraScriptResults.cpp
+++ b/Engine/Plugins/FX/Niagara/Source/NiagaraEditor/Private/UpgradeNiagaraScriptResults.cpp
@@ -5,6 +5,7 @@
 #include "NiagaraClipboard.h"
 #include "NiagaraEmitterHandle.h"
 #include "NiagaraNodeFunctionCall.h"
+#include "Logging/StructuredLog.h"
 #include "ViewModels/NiagaraEmitterHandleViewModel.h"
 #include "ViewModels/NiagaraEmitterViewModel.h"
 #include "ViewModels/Stack/NiagaraStackModuleItem.h"
@@ -96,6 +97,10 @@ int32 UNiagaraPythonScriptModuleInput::AsInt() const
 	{
 		return NiagaraScriptResults::GetValue<int32>(Input);
 	}
+	if (IsSet() && Input->InputType.IsEnum())
+	{
+		return NiagaraScriptResults::GetValue<int32>(Input);
+	}
 	return 0;
 }
 
@@ -107,7 +112,7 @@ bool UNiagaraPythonScriptModuleInput::AsBool() const
 		{
 			return false;
 		}
-		FNiagaraBool* BoolStruct = (FNiagaraBool*)Input->Local.GetData();
+		const FNiagaraBool* BoolStruct = reinterpret_cast<const FNiagaraBool*>(Input->Local.GetData());
 		return BoolStruct->GetValue();
 	}
 	return false;
@@ -126,7 +131,7 @@ FVector UNiagaraPythonScriptModuleInput::AsVec3() const
 {
 	if (IsSet() && Input->InputType == FNiagaraTypeDefinition::GetVec3Def())
 	{
-		return (FVector)NiagaraScriptResults::GetValue<FVector3f>(Input);
+		return FVector(NiagaraScriptResults::GetValue<FVector3f>(Input));
 	}
 	return FVector();
 }
@@ -291,6 +296,18 @@ void UUpgradeNiagaraScriptResults::SetEnumInput(const FString& InputName, FStrin
 	}
 }
 
+void UUpgradeNiagaraScriptResults::SetEnumInputFromInt(const FString& InputName, int32 Value)
+{
+	UNiagaraPythonScriptModuleInput* ModuleInput = GetNewInput(FName(InputName));
+	if (ModuleInput && ModuleInput->Input->InputType.IsEnum())
+	{
+		if (!ModuleInput->Input->InputType.GetEnum()->IsValidEnumValue(Value))
+		{
+			UE_LOGFMT(LogNiagaraEditor, Error, "Value {Value} is not a valid enum value for input {InputName}", Value, InputName);
+		}
+		NiagaraScriptResults::SetValue(ModuleInput, Value);
+	}
+}
 void UUpgradeNiagaraScriptResults::SetLinkedInput(const FString& InputName, FString Value)
 {
 	if (UNiagaraPythonScriptModuleInput* ModuleInput = GetNewInput(FName(InputName)))
diff --git a/Engine/Plugins/FX/Niagara/Source/NiagaraEditor/Public/UpgradeNiagaraScriptResults.h b/Engine/Plugins/FX/Niagara/Source/NiagaraEditor/Public/UpgradeNiagaraScriptResults.h
--- a/Engine/Plugins/FX/Niagara/Source/NiagaraEditor/Public/UpgradeNiagaraScriptResults.h
+++ b/Engine/Plugins/FX/Niagara/Source/NiagaraEditor/Public/UpgradeNiagaraScriptResults.h
@@ -139,6 +139,9 @@ public:
 
 	UFUNCTION(BlueprintCallable, Category = "Scripting")
     NIAGARAEDITOR_API void SetEnumInput(const FString& InputName, FString Value);
+
+	UFUNCTION(BlueprintCallable, Category = "Scripting")
+	NIAGARAEDITOR_API void SetEnumInputFromInt(const FString& InputName, int32 Value);
     
     UFUNCTION(BlueprintCallable, Category = "Scripting")
     NIAGARAEDITOR_API void SetLinkedInput(const FString& InputName, FString Value);
