- Display valid binding types on tooltips

#tests Hover over various bindings, simulation stages, mesh renderer mesh array binding, etc

diff --git a/Engine/Plugins/FX/Niagara/Source/NiagaraEditor/Private/Customizations/NiagaraParameterBindingCustomization.cpp b/Engine/Plugins/FX/Niagara/Source/NiagaraEditor/Private/Customizations/NiagaraParameterBindingCustomization.cpp
--- a/Engine/Plugins/FX/Niagara/Source/NiagaraEditor/Private/Customizations/NiagaraParameterBindingCustomization.cpp
+++ b/Engine/Plugins/FX/Niagara/Source/NiagaraEditor/Private/Customizations/NiagaraParameterBindingCustomization.cpp
@@ -34,6 +34,77 @@
 
 #define LOCTEXT_NAMESPACE "NiagaraParameterBindingCustomization"
 
+namespace NiagaraParameterBindingCustomizationPrivate
+{
+	FText JoinTypeArray(FText Delim, TConstArrayView<UClass*> TypeArray)
+	{
+		TArray<FText> TypeNamesText;
+		for (UClass* Class : TypeArray)
+		{
+			TypeNamesText.Add(Class->GetDisplayNameText());
+		}
+		return FText::Join(Delim, TypeNamesText);
+	}
+
+	FText JoinTypeArray(FText Delim, TConstArrayView<FNiagaraTypeDefinition> TypeArray)
+	{
+		TArray<FText> TypeNamesText;
+		for (const FNiagaraTypeDefinition& TypeDef : TypeArray)
+		{
+			TypeNamesText.Add(TypeDef.GetNameText());
+		}
+		return FText::Join(Delim, TypeNamesText);
+	}
+
+	template<typename TArrayType>
+	void BuildTypesText(FText& InOutString, FText TypesDescription, TArrayType TypeArray)
+	{
+		if (TypeArray.Num() == 0)
+		{
+			return;
+		}
+
+		InOutString = FText::Join(
+			FText::GetEmpty(),
+			TArray<FText>(
+				{
+					InOutString,
+					LOCTEXT("NewLine", "\n"),
+					TypesDescription,
+					JoinTypeArray(LOCTEXT("CommaSpace", ", "), TypeArray)
+				}
+			)
+		);
+	}
+
+	FText GetAllowedBindingTypesText(FNiagaraParameterBinding* ParameterBinding)
+	{
+		FText BoundTypes;
+		BuildTypesText(
+			BoundTypes,
+			LOCTEXT("BindingTooltip_DIFormat", "Data Interface Types: "),
+			ParameterBinding->GetAllowedDataInterfaces()
+		);
+		BuildTypesText(
+			BoundTypes,
+			LOCTEXT("BindingTooltip_ObjectFormat", "Object Types: "),
+			ParameterBinding->GetAllowedObjects()
+		);
+		BuildTypesText(
+			BoundTypes,
+			LOCTEXT("BindingTooltip_ObjectFormat", "Interface Types: "),
+			ParameterBinding->GetAllowedInterfaces()
+		);
+		BuildTypesText(
+			BoundTypes,
+			LOCTEXT("BindingTooltip_ObjectFormat", "Attribute Types: "),
+			ParameterBinding->GetAllowedTypeDefinitions()
+		);
+
+		return BoundTypes;
+	}
+}
+
 FNiagaraParameterBindingAction::FNiagaraParameterBindingAction()
 	: FEdGraphSchemaAction(FText(), FText(), FText(), 0)
 {
@@ -492,30 +563,39 @@ FText FNiagaraParameterBindingCustomization::GetTooltipText() const
 {
 	if ( FNiagaraParameterBinding* ParameterBinding = GetParameterBinding() )
 	{
+		TArray<FText> TooltipTexts;
 		if (ParameterBinding->AliasedParameter.GetName().IsNone())
 		{
 			if (ParameterBinding->HasDefaultValueEditorOnly())
 			{
 				const FNiagaraTypeDefinition ConstantType = ParameterBinding->GetDefaultAliasedParameter().GetType();
-				return FText::Format(LOCTEXT("BindingTooltip_BoundConstant", "Bound to constant value \"{0}\" type \"{1}\""),
-					FText::FromString(ConstantType.ToString(ParameterBinding->GetDefaultValueEditorOnly().GetData())),
-					FText::FromString(ConstantType.GetName())
+				TooltipTexts.Add(
+					FText::Format(
+						LOCTEXT("BindingTooltip_BoundConstant", "Bound to constant value \"{0}\" type \"{1}\""),
+						FText::FromString(ConstantType.ToString(ParameterBinding->GetDefaultValueEditorOnly().GetData())),
+						FText::FromString(ConstantType.GetName())
+					)
 				);
 			}
 			else
 			{
-				return LOCTEXT("BindingTooltip_UnboundParameter", "Not bound to any parameter.");
+				TooltipTexts.Add(LOCTEXT("BindingTooltip_UnboundParameter", "Not bound to any parameter."));
 			}
 		}
 		else
 		{
-			return FText::Format(LOCTEXT("BindingTooltip_BoundParameter", "Bound to parameter \"{0}\" type \"{1}\""),
-				FText::FromName(ParameterBinding->AliasedParameter.GetName()),
-				FText::FromString(ParameterBinding->AliasedParameter.GetType().GetName())
+			TooltipTexts.Add(
+				FText::Format(LOCTEXT("BindingTooltip_BoundParameter", "Bound to parameter \"{0}\" type \"{1}\""),
+					FText::FromName(ParameterBinding->AliasedParameter.GetName()),
+					FText::FromString(ParameterBinding->AliasedParameter.GetType().GetName())
+				)
 			);
 		}
+
+		TooltipTexts.Add(NiagaraParameterBindingCustomizationPrivate::GetAllowedBindingTypesText(ParameterBinding));
+		return FText::Join(LOCTEXT("NewLine", "\n"), TooltipTexts);
 	}
-	return FText::FromString(TEXT("Missing"));
+	return LOCTEXT("BindingTooltip_Invalid", "Binding is invalid");
 }
 
 FName FNiagaraParameterBindingCustomization::GetVariableName() const
diff --git a/Engine/Plugins/FX/Niagara/Source/NiagaraEditor/Private/Customizations/NiagaraTypeCustomizations.cpp b/Engine/Plugins/FX/Niagara/Source/NiagaraEditor/Private/Customizations/NiagaraTypeCustomizations.cpp
--- a/Engine/Plugins/FX/Niagara/Source/NiagaraEditor/Private/Customizations/NiagaraTypeCustomizations.cpp
+++ b/Engine/Plugins/FX/Niagara/Source/NiagaraEditor/Private/Customizations/NiagaraTypeCustomizations.cpp
@@ -626,12 +626,29 @@ FText FNiagaraUserParameterBindingCustomization::GetCurrentText() const
 
 FText FNiagaraUserParameterBindingCustomization::GetTooltipText() const
 {
-	if (IsValid() && TargetUserParameterBinding && TargetUserParameterBinding->Parameter.IsValid())
+	if (IsValid() && TargetUserParameterBinding)
 	{
-		FText TooltipDesc = FText::Format(LOCTEXT("UserParameterBindingTooltip", "Bound to the user parameter \"{0}\""), FText::FromName(TargetUserParameterBinding->Parameter.GetName()));
-		return TooltipDesc;
+		const FNiagaraTypeDefinition TypeDef = TargetUserParameterBinding->Parameter.GetType();
+		if (TypeDef.IsValid())
+		{
+			FText TooltipFormat;
+			if (TargetUserParameterBinding->Parameter.IsValid())
+			{
+				TooltipFormat = LOCTEXT("UserParameterTooltip_Bound", "Bound to parameter \"{0}\".\nCan be bound to parameters of type \"{1}\".");
+			}
+			else
+			{
+				TooltipFormat = LOCTEXT("UserParameterTooltip_Unbound", "Currently not bound to any parameter.\nCan be bound to parameters of type \"{1}\".");
+			}
+
+			return FText::Format(
+				TooltipFormat,
+				FText::FromName(TargetUserParameterBinding->Parameter.GetName()),
+				TypeDef.GetNameText()
+			);
+		}
 	}
-	return FText::FromString(TEXT("Missing"));
+	return LOCTEXT("InvalidBinding", "Binding is invalid");
 }
 
 TSharedRef<SWidget> FNiagaraUserParameterBindingCustomization::OnGetMenuContent() const
@@ -1492,10 +1509,16 @@ FText FNiagaraDataInterfaceBindingCustomization::GetCurrentText() const
 
 FText FNiagaraDataInterfaceBindingCustomization::GetTooltipText() const
 {
-	if (BaseStage && TargetDataInterfaceBinding && TargetDataInterfaceBinding->BoundVariable.IsValid())
+	if (BaseStage && TargetDataInterfaceBinding)
 	{
-		FText TooltipDesc = FText::Format(LOCTEXT("DataInterfaceBindingTooltip", "Bound to the user parameter \"{0}\""), FText::FromName(TargetDataInterfaceBinding->BoundVariable.GetName()));
-		return TooltipDesc;
+		if (TargetDataInterfaceBinding->BoundVariable.IsValid())
+		{
+			return FText::Format(LOCTEXT("DataInterfaceBindingTooltip", "Bound to the user parameter \"{0}\""), FText::FromName(TargetDataInterfaceBinding->BoundVariable.GetName()));
+		}
+		else
+		{
+			return LOCTEXT("DataInterfaceBindingTooltip_Unbound", "Currently not bound to any data interface");
+		}
 	}
 	return FText::FromString(TEXT("Missing"));
 }
