- Add Loop Count Limit to Stateless Spawn Burst

diff --git a/Engine/Plugins/FX/Niagara/Source/Niagara/Internal/Stateless/NiagaraStatelessSpawnInfo.h b/Engine/Plugins/FX/Niagara/Source/Niagara/Internal/Stateless/NiagaraStatelessSpawnInfo.h
--- a/Engine/Plugins/FX/Niagara/Source/Niagara/Internal/Stateless/NiagaraStatelessSpawnInfo.h
+++ b/Engine/Plugins/FX/Niagara/Source/Niagara/Internal/Stateless/NiagaraStatelessSpawnInfo.h
@@ -32,12 +32,12 @@ struct FNiagaraStatelessSpawnInfo
 	float SpawnTime = 0.0f;
 
 	// Number of particles to spawn at the provided time
-	UPROPERTY(EditAnywhere, Category = "Spawn", meta = (EditConditionHides, EditCondition = "Type == ENiagaraStatelessSpawnInfoType::Burst", ClampMin = "0"))
+	UPROPERTY(EditAnywhere, Category = "Spawn", meta = (EditConditionHides, EditCondition = "Type == ENiagaraStatelessSpawnInfoType::Burst", ClampMin = "0", DisplayName = "Spawn Count"))
 	FNiagaraDistributionRangeInt Amount = FNiagaraDistributionRangeInt(1);
 
 	// Number of particles to spawn per second
 	// Note: In the case of a random range or binding the value is only evaluated at the start of each loop.  This varies from regular emitters.
-	UPROPERTY(EditAnywhere, Category = "Spawn", meta = (EditConditionHides, EditCondition = "Type == ENiagaraStatelessSpawnInfoType::Rate", ClampMin = "0.0"))
+	UPROPERTY(EditAnywhere, Category = "Spawn", meta = (EditConditionHides, EditCondition = "Type == ENiagaraStatelessSpawnInfoType::Rate", ClampMin = "0.0", DisplayName = "Spawn Rate"))
 	FNiagaraDistributionRangeFloat Rate = FNiagaraDistributionRangeFloat(60.0f);
 
 	UPROPERTY(EditAnywhere, Category = "Spawn", meta = (InlineEditConditionToggle))
@@ -46,12 +46,22 @@ struct FNiagaraStatelessSpawnInfo
 	UPROPERTY(EditAnywhere, Category = "Spawn", meta = (InlineEditConditionToggle))
 	uint8 bSpawnProbabilityEnabled : 1 = false;
 
+	UPROPERTY(EditAnywhere, Category = "Spawn", meta = (InlineEditConditionToggle))
+	uint8 bLoopCountLimitEnabled : 1 = false;
+
 	// 0 - 1 value that can be viewed as a percentage chance that the spawn will generated particles or not.
 	// A value of 0.5 can be viewed as a 50% chance that the spawn will trigger.
 	// Note: Rate spawning will only evaluate the probability at the start of each loop.  This varies from regular emitters which evaluate each time particles are spawned.
 	UPROPERTY(EditAnywhere, Category = "Spawn", meta = (EditCondition = "bSpawnProbabilityEnabled", ClampMin = "0.0", ClampMax = "1.0", UIMin = "0.0", UIMax = "1.0"))
 	FNiagaraDistributionRangeFloat SpawnProbability = FNiagaraDistributionRangeFloat(1.0f);
 
+
+	// For Emitters that loop multiple times (or infinitely), this can constrain the bursting to only happen within this loop count limit.
+	// For example, an emitter might loop 5 times, but bursting should only happen on the first two loops.  Setting Loop Count Limit to 2 will allow the burst to fire twice, once after each of the first two loops, and then stop.
+	//-TODO:EditConditionHides, EditCondition = "Type == ENiagaraStatelessSpawnInfoType::Burst",
+	UPROPERTY(EditAnywhere, Category = "Spawn", meta = (EditCondition = "bLoopCountLimitEnabled", ClampMin = "1", UIMin = "0"))
+	FNiagaraDistributionRangeInt LoopCountLimit = FNiagaraDistributionRangeInt(1);
+
 	bool IsValid(TOptional<float> LoopDuration) const;
 };
 
diff --git a/Engine/Plugins/FX/Niagara/Source/Niagara/Private/Stateless/NiagaraStatelessEmitter.cpp b/Engine/Plugins/FX/Niagara/Source/Niagara/Private/Stateless/NiagaraStatelessEmitter.cpp
--- a/Engine/Plugins/FX/Niagara/Source/Niagara/Private/Stateless/NiagaraStatelessEmitter.cpp
+++ b/Engine/Plugins/FX/Niagara/Source/Niagara/Private/Stateless/NiagaraStatelessEmitter.cpp
@@ -411,6 +411,7 @@ void UNiagaraStatelessEmitter::CacheFromCompiledData()
 				EmitterBuildContext.ConvertDistributionToRange(SpawnInfo.Amount, 0);
 			}
 			EmitterBuildContext.ConvertDistributionToRange(SpawnInfo.SpawnProbability, 0.0f);
+			EmitterBuildContext.ConvertDistributionToRange(SpawnInfo.LoopCountLimit, 0);
 		}
 
 		// Prepare renderer bindings this avoids having to do this per instance spawned
diff --git a/Engine/Plugins/FX/Niagara/Source/Niagara/Private/Stateless/NiagaraStatelessEmitterInstance.cpp b/Engine/Plugins/FX/Niagara/Source/Niagara/Private/Stateless/NiagaraStatelessEmitterInstance.cpp
--- a/Engine/Plugins/FX/Niagara/Source/Niagara/Private/Stateless/NiagaraStatelessEmitterInstance.cpp
+++ b/Engine/Plugins/FX/Niagara/Source/Niagara/Private/Stateless/NiagaraStatelessEmitterInstance.cpp
@@ -10,12 +10,21 @@
 
 #include "Shader.h"
 
+static bool bPerfTest = false;
+static FAutoConsoleVariableRef CVarTest(
+	TEXT("fx.NiagaraStateless.PerfTest"),
+	bPerfTest,
+	TEXT(""),
+	ECVF_Scalability | ECVF_Default
+);
+
 namespace NiagaraStatelessEmitterInstancePrivate
 {
 	static const float DefaultLoopDuration		= 0.001f;
 	static const float DefaultLoopDelay			= 0.0f;
 	static const float DefaultSpawnRate			= 0.0f;
 	static const float DefaultSpawnProbability	= 0.0f;
+	static const int32 DefaultLoopCountLimit	= 1;
 	static const int32 DefaultSpawnAmount		= 0;
 
 	// Other areas work on built distributions so this is a helper function to work with a raw distribution
@@ -538,6 +547,11 @@ void FNiagaraStatelessEmitterInstance::SendRenderData()
 	const FTransform3f ParentTransform = FTransform3f(ParentSystemInstance->GetWorldTransform());
 	bNeedsShaderParametersUpdate |= EmitterTransforms.UpdateTransforms(ParentTransform);
 
+	if (bPerfTest && !bNeedsShaderParametersUpdate && !bSpawnInfosDirty)
+	{
+		return;
+	}
+
 	if (bNeedsShaderParametersUpdate)
 	{
 		DataForRenderThread.ShaderParameters = WeakStatelessEmitter->AllocateShaderParameters(EmitterTransforms, RendererBindings);
@@ -698,6 +712,15 @@ void FNiagaraStatelessEmitterInstance::InitSpawnInfosForLoop(float Initializatio
 			continue;
 		}
 
+		if (SpawnInfo.bLoopCountLimitEnabled)
+		{
+			const int32 LoopCountLimit = EvaluateDistribution(SpawnInfo.LoopCountLimit, RandomStream, RendererBindings, DefaultLoopCountLimit);
+			if (LoopCount >= LoopCountLimit)
+			{
+				continue;
+			}
+		}
+
 		if (SpawnInfo.bSpawnProbabilityEnabled)
 		{
 			float SpawnProbability = EvaluateDistribution(SpawnInfo.SpawnProbability, RandomStream, RendererBindings, DefaultSpawnProbability);
diff --git a/Engine/Plugins/FX/Niagara/Source/NiagaraEditor/Private/Stateless/NiagaraSpawnInfoPropertyCustomization.cpp b/Engine/Plugins/FX/Niagara/Source/NiagaraEditor/Private/Stateless/NiagaraSpawnInfoPropertyCustomization.cpp
new file mode 100644
--- /dev/null
+++ b/Engine/Plugins/FX/Niagara/Source/NiagaraEditor/Private/Stateless/NiagaraSpawnInfoPropertyCustomization.cpp
@@ -0,0 +1,52 @@
+// Copyright Epic Games, Inc. All Rights Reserved.
+
+#include "Stateless/NiagaraSpawnInfoPropertyCustomization.h"
+#include "Stateless/NiagaraStatelessSpawnInfo.h"
+
+#include "DetailLayoutBuilder.h"
+#include "DetailCategoryBuilder.h"
+#include "IDetailChildrenBuilder.h"
+#include "IDetailGroup.h"
+#include "DetailWidgetRow.h"
+
+TSharedRef<IDetailCustomization> FNiagaraSpawnInfoDetailCustomization::MakeInstance()
+{
+	return MakeShared<FNiagaraSpawnInfoDetailCustomization>();
+}
+
+void FNiagaraSpawnInfoDetailCustomization::CustomizeDetails(IDetailLayoutBuilder& DetailBuilder)
+{
+	IDetailCategoryBuilder& DetailCategory = DetailBuilder.EditCategory("Spawn");
+	{
+		TArray<TSharedRef<IPropertyHandle>> Properties;
+		DetailCategory.GetDefaultProperties(Properties, true, true);
+
+		for (TSharedPtr<IPropertyHandle> Property : Properties)
+		{
+			IDetailPropertyRow& DetailPropertyRow = DetailCategory.AddProperty(Property);
+			if (Property->GetProperty()->GetFName() == GET_MEMBER_NAME_CHECKED(FNiagaraStatelessSpawnInfo, LoopCountLimit))
+			{
+				DetailPropertyRow.Visibility(TAttribute<EVisibility>::Create(TAttribute<EVisibility>::FGetter::CreateSP(this, &FNiagaraSpawnInfoDetailCustomization::GetLoopCountLimitVisibility)));
+			}
+			else if (Property->GetProperty()->GetFName() == GET_MEMBER_NAME_CHECKED(FNiagaraStatelessSpawnInfo, Type))
+			{
+				SpawnTypeProperty = Property;
+			}
+		}
+	}
+}
+
+EVisibility FNiagaraSpawnInfoDetailCustomization::GetLoopCountLimitVisibility() const
+{
+	if ( SpawnTypeProperty )
+	{
+		uint8 UntypedValue = -1;
+		const FPropertyAccess::Result GetResult = SpawnTypeProperty->GetValue(UntypedValue);
+		if (GetResult == FPropertyAccess::Result::Success)
+		{
+			return ENiagaraStatelessSpawnInfoType(UntypedValue) == ENiagaraStatelessSpawnInfoType::Burst ? EVisibility::Visible : EVisibility::Hidden;
+		}
+	}
+
+	return EVisibility::Hidden;
+}
diff --git a/Engine/Plugins/FX/Niagara/Source/NiagaraEditor/Private/Stateless/NiagaraSpawnInfoPropertyCustomization.h b/Engine/Plugins/FX/Niagara/Source/NiagaraEditor/Private/Stateless/NiagaraSpawnInfoPropertyCustomization.h
new file mode 100644
--- /dev/null
+++ b/Engine/Plugins/FX/Niagara/Source/NiagaraEditor/Private/Stateless/NiagaraSpawnInfoPropertyCustomization.h
@@ -0,0 +1,20 @@
+// Copyright Epic Games, Inc. All Rights Reserved.
+
+#pragma once
+
+#include "PropertyHandle.h"
+#include "IDetailCustomization.h"
+
+struct FNiagaraStatelessSpawnInfo;
+
+class FNiagaraSpawnInfoDetailCustomization : public IDetailCustomization
+{
+public:
+	static TSharedRef<IDetailCustomization> MakeInstance();
+	virtual void CustomizeDetails(IDetailLayoutBuilder& DetailBuilder) override;
+
+	EVisibility GetLoopCountLimitVisibility() const;
+
+private:
+	TSharedPtr<IPropertyHandle> SpawnTypeProperty;
+};
diff --git a/Engine/Plugins/FX/Niagara/Source/NiagaraEditor/Private/ViewModels/Stack/NiagaraStackStatelessEmitterSpawnGroup.cpp b/Engine/Plugins/FX/Niagara/Source/NiagaraEditor/Private/ViewModels/Stack/NiagaraStackStatelessEmitterSpawnGroup.cpp
--- a/Engine/Plugins/FX/Niagara/Source/NiagaraEditor/Private/ViewModels/Stack/NiagaraStackStatelessEmitterSpawnGroup.cpp
+++ b/Engine/Plugins/FX/Niagara/Source/NiagaraEditor/Private/ViewModels/Stack/NiagaraStackStatelessEmitterSpawnGroup.cpp
@@ -8,6 +8,7 @@
 #include "ScopedTransaction.h"
 #include "Stateless/NiagaraDistributionPropertyCustomization.h"
 #include "Stateless/NiagaraDistributionIntPropertyCustomization.h"
+#include "Stateless/NiagaraSpawnInfoPropertyCustomization.h"
 #include "Stateless/NiagaraStatelessEmitter.h"
 #include "Stateless/NiagaraStatelessModule.h"
 #include "Styling/AppStyle.h"
@@ -400,6 +401,8 @@ void UNiagaraStackStatelessEmitterSpawnItem::RefreshChildrenInternal(const TArra
 			SpawnInfoObject->SetOnFilterDetailNodes(
 				FNiagaraStackObjectShared::FOnFilterDetailNodes::CreateStatic(&UNiagaraStackStatelessEmitterSpawnItem::FilterDetailNodes),
 				UNiagaraStackObject::EDetailNodeFilterMode::FilterAllNodes);
+			SpawnInfoObject->RegisterInstancedCustomPropertyLayout(FNiagaraStatelessSpawnInfo::StaticStruct(), FOnGetDetailCustomizationInstance::CreateStatic(&FNiagaraSpawnInfoDetailCustomization::MakeInstance));
+
 			//SpawnInfoObject->RegisterInstancedCustomPropertyTypeLayout(FNiagaraDistributionFloat::StaticStruct()->GetFName(), FOnGetPropertyTypeCustomizationInstance::CreateStatic(&FNiagaraDistributionPropertyCustomization::MakeFloatInstance, StatelessEmitterObject));
 			//SpawnInfoObject->RegisterInstancedCustomPropertyTypeLayout(FNiagaraDistributionVector2::StaticStruct()->GetFName(), FOnGetPropertyTypeCustomizationInstance::CreateStatic(&FNiagaraDistributionPropertyCustomization::MakeVector2Instance));
 			//SpawnInfoObject->RegisterInstancedCustomPropertyTypeLayout(FNiagaraDistributionVector3::StaticStruct()->GetFName(), FOnGetPropertyTypeCustomizationInstance::CreateStatic(&FNiagaraDistributionPropertyCustomization::MakeVector3Instance));
