- Change display name for OrientToAxis to match stateful

#rnx

diff --git a/Engine/Plugins/FX/Niagara/Source/Niagara/Internal/Stateless/Modules/NiagaraStatelessModule_InitialMeshOrientation.h b/Engine/Plugins/FX/Niagara/Source/Niagara/Internal/Stateless/Modules/NiagaraStatelessModule_InitialMeshOrientation.h
--- a/Engine/Plugins/FX/Niagara/Source/Niagara/Internal/Stateless/Modules/NiagaraStatelessModule_InitialMeshOrientation.h
+++ b/Engine/Plugins/FX/Niagara/Source/Niagara/Internal/Stateless/Modules/NiagaraStatelessModule_InitialMeshOrientation.h
@@ -14,7 +14,7 @@ enum class ENSMInitialMeshOrientationMode
 	None,
 	Random,
 	//System,
-	OrientToAxis,
+	OrientToAxis UMETA(DisplayName="Orient To Vector"),
 	//OrientToMatrix,
 	//OrientToQuaternion,
 };
