Reduce oversubscription with the NiagaraTypeRegistry.  Changes the lock to being a RWLock rather than a mutex as the mutex isn't currently well supported by AUTORTFM.  Also introduces an atomic counter (for the unlikely scenario of types being invalidated) to further reduce how often we need to lock things.
#tests Compared performance and results of a compile of niagara assets before/after and validated that we no longer got oversubscribed using the new compilation mode

diff --git a/Engine/Plugins/FX/Niagara/Source/Niagara/Private/NiagaraTypeRegistry.cpp b/Engine/Plugins/FX/Niagara/Source/Niagara/Private/NiagaraTypeRegistry.cpp
--- a/Engine/Plugins/FX/Niagara/Source/Niagara/Private/NiagaraTypeRegistry.cpp
+++ b/Engine/Plugins/FX/Niagara/Source/Niagara/Private/NiagaraTypeRegistry.cpp
@@ -26,6 +26,10 @@ namespace NiagaraTypeRegistryLocal
 	// the TypeRegistry's RegisteredTypes array.  Allows us to quickly access
 	// the type definition from a handle without having to hit the lock
 	static thread_local TMap<uint32, int32> RegisteredTypeIndexCacheTLS;
+	static thread_local uint32 RegisteredTypeIndexCacheGeneration = INDEX_NONE;
+
+	static thread_local uint32 RegistryProxyRefCount = 0;
+	static thread_local TArray<FNiagaraTypeDefinition> RegistryProxyContents;
 
 	static FDelegateHandle PluginUnmountedHandle;
 
@@ -38,12 +42,27 @@ namespace NiagaraTypeRegistryLocal
 	}
 }
 
+FNiagaraTypeRegistryTLSProxy::FNiagaraTypeRegistryTLSProxy()
+	: CallingThreadId(FPlatformTLS::GetCurrentThreadId())
+{
+	++NiagaraTypeRegistryLocal::RegistryProxyRefCount;
+}
+
+FNiagaraTypeRegistryTLSProxy::~FNiagaraTypeRegistryTLSProxy()
+{
+	checkf(CallingThreadId == FPlatformTLS::GetCurrentThreadId(), TEXT("FNiagaraTypeRegistryTLSProxy destroyed from different thread (%x) from where it was constructed (%x)."), FPlatformTLS::GetCurrentThreadId(), CallingThreadId);
+	if ((--NiagaraTypeRegistryLocal::RegistryProxyRefCount) == 0)
+	{
+		NiagaraTypeRegistryLocal::RegistryProxyContents.Reset();
+	}
+}
+
 UNiagaraDataInterfaceBase* FNiagaraTypeRegistry::GetDefaultDataInterfaceByName(const FString& DIClassName)
 {
 	UClass* DIClass = nullptr;
 	{
 		FNiagaraTypeRegistry& Registry = Get();
-		UE::TScopeLock Lock(Registry.RegisteredTypesLock);
+		UE::TReadScopeLock Lock(Registry.RegisteredTypesLock);
 		for (const TUniquePtr<FNiagaraTypeDefinition>& TypePtr : Registry.RegisteredTypes)
 		{
 			if (TypePtr.IsValid() && TypePtr->IsDataInterface())
@@ -84,10 +103,10 @@ UNiagaraDataInterfaceBase* FNiagaraTypeRegistry::GetDefaultDataInterfaceByName(c
 // 
 // It might be worth investigating changing RegisteredTypeIndexMap to be hashed based on the name of the class/enum/struct rather than the pointer.  That
 // way if we reload a type old data, with their existing index, will still be treated as valid.  In this current implementation old FNiagaraVariables will
-// have invalid types.  But, for now I think it would be an error to unload a type without having the variables that dependend on it also getting deleted.
+// have invalid types.  But, for now I think it would be an error to unload a type without having the variables that depended on it also getting deleted.
 void FNiagaraTypeRegistry::AddReferencedObjects(FReferenceCollector& Collector)
 {
-	UE::TScopeLock Lock(RegisteredTypesLock);
+	UE::TReadScopeLock Lock(RegisteredTypesLock);
 	for (const TUniquePtr<FNiagaraTypeDefinition>& RegisteredType : RegisteredTypes)
 	{
 		if (RegisteredType.IsValid())
@@ -156,51 +175,62 @@ void FNiagaraTypeRegistry::InvalidateTypesByPath(const FString& AssetPath)
 		}
 	}
 
-	TArray<FNiagaraTypeDefinition> InvalidatedTypes;
-	InvalidatedTypes.Reserve(TypesToInvalidate.Num());
-
+	if (!TypesToInvalidate.IsEmpty())
 	{
-		UE::TScopeLock Lock(Registry.RegisteredTypesLock);
+		++Registry.RegisteredTypesGeneration;
+
+		TArray<FNiagaraTypeDefinition> InvalidatedTypes;
+		InvalidatedTypes.Reserve(TypesToInvalidate.Num());
 
-		for (int32 TypeToInvalidate : TypesToInvalidate)
 		{
-			TUniquePtr<FNiagaraTypeDefinition>& TypePtr = Registry.RegisteredTypes[TypeToInvalidate];
-			InvalidatedTypes.Add(*TypePtr);
+			UE::TWriteScopeLock Lock(Registry.RegisteredTypesLock);
 
-			// we leave the type definition around, but invalidate it so that anyone holding onto a reference will still have
-			// valid memory but we aren't going to be holding onto a class that has been invalidated
-			TypePtr->Invalidate();
-		}
-	}
+			for (int32 TypeToInvalidate : TypesToInvalidate)
+			{
+				TUniquePtr<FNiagaraTypeDefinition>& TypePtr = Registry.RegisteredTypes[TypeToInvalidate];
+				InvalidatedTypes.Add(*TypePtr);
 
-	{
-		FReadScopeLock Lock(Registry.RegistrationLock);
+				// we leave the type definition around, but invalidate it so that anyone holding onto a reference will still have
+				// valid memory but we aren't going to be holding onto a class that has been invalidated
+				TypePtr->Invalidate();
+			}
+		}
 
-		for (const FNiagaraTypeDefinition& TypeToRemove : InvalidatedTypes)
 		{
-			// we only support unregistering implicitly added types which don't include support for the different
-			// flag based categorizations
-			ensure(!Registry.RegisteredUserVariableTypes.Contains(TypeToRemove));
-			ensure(!Registry.RegisteredUserVariableTypes.Contains(TypeToRemove));
-			ensure(!Registry.RegisteredSystemVariableTypes.Contains(TypeToRemove));
-			ensure(!Registry.RegisteredEmitterVariableTypes.Contains(TypeToRemove));
-			ensure(!Registry.RegisteredParticleVariableTypes.Contains(TypeToRemove));
-			ensure(!Registry.RegisteredParamTypes.Contains(TypeToRemove));
-			ensure(!Registry.RegisteredPayloadTypes.Contains(TypeToRemove));
-			ensure(!Registry.RegisteredUserDefinedTypes.Contains(TypeToRemove));
-			ensure(!Registry.RegisteredNumericTypes.Contains(TypeToRemove));
-			ensure(!Registry.RegisteredIndexTypes.Contains(TypeToRemove));
-
-			// we leave the following properties alone.  These could be considered orphaned, but in the current setup they
-			// would not be repopulated if the type is re-registered (implicitly)
-			//RegisteredStructConversionMap
-			//RegisteredSoftObjectPaths
+			FReadScopeLock Lock(Registry.RegistrationLock);
+
+			for (const FNiagaraTypeDefinition& TypeToRemove : InvalidatedTypes)
+			{
+				// we only support unregistering implicitly added types which don't include support for the different
+				// flag based categorizations
+				ensure(!Registry.RegisteredUserVariableTypes.Contains(TypeToRemove));
+				ensure(!Registry.RegisteredUserVariableTypes.Contains(TypeToRemove));
+				ensure(!Registry.RegisteredSystemVariableTypes.Contains(TypeToRemove));
+				ensure(!Registry.RegisteredEmitterVariableTypes.Contains(TypeToRemove));
+				ensure(!Registry.RegisteredParticleVariableTypes.Contains(TypeToRemove));
+				ensure(!Registry.RegisteredParamTypes.Contains(TypeToRemove));
+				ensure(!Registry.RegisteredPayloadTypes.Contains(TypeToRemove));
+				ensure(!Registry.RegisteredUserDefinedTypes.Contains(TypeToRemove));
+				ensure(!Registry.RegisteredNumericTypes.Contains(TypeToRemove));
+				ensure(!Registry.RegisteredIndexTypes.Contains(TypeToRemove));
+
+				// we leave the following properties alone.  These could be considered orphaned, but in the current setup they
+				// would not be repopulated if the type is re-registered (implicitly)
+				//RegisteredStructConversionMap
+				//RegisteredSoftObjectPaths
+			}
 		}
 	}
 }
 
 int32 FNiagaraTypeRegistry::RegisterTypeInternal(const FNiagaraTypeDefinition& NewType)
 {
+	if (NiagaraTypeRegistryLocal::RegisteredTypeIndexCacheGeneration != RegisteredTypesGeneration)
+	{
+		NiagaraTypeRegistryLocal::RegisteredTypeIndexCacheGeneration = RegisteredTypesGeneration;
+		NiagaraTypeRegistryLocal::RegisteredTypeIndexCacheTLS.Reset();
+	}
+
 	uint32 TypeHash = 0;
 
 	if (NewType.RegisteredTypeDefIndex == INDEX_NONE)
@@ -210,8 +240,13 @@ int32 FNiagaraTypeRegistry::RegisterTypeInternal(const FNiagaraTypeDefinition& N
 		// first see if we're in the TLS cache
 		if (const int32* ExistingIndex = NiagaraTypeRegistryLocal::RegisteredTypeIndexCacheTLS.Find(TypeHash))
 		{
-			// be sure to update the cached value on the type definition as well
-			NewType.RegisteredTypeDefIndex = *ExistingIndex;
+			if (*ExistingIndex != INDEX_NONE)
+			{
+				// be sure to update the cached value on the type definition as well
+				NewType.RegisteredTypeDefIndex = *ExistingIndex;
+
+				return NewType.RegisteredTypeDefIndex;
+			}
 		}
 	}
 
@@ -232,13 +267,23 @@ int32 FNiagaraTypeRegistry::RegisterTypeInternal(const FNiagaraTypeDefinition& N
 	{
 		if (NewType.IsValid())
 		{
-			UE::TScopeLock Lock(RegisteredTypesLock);
-			if (ensure(RegisteredTypes.IsValidIndex(NewType.RegisteredTypeDefIndex)))
+			bool bNeedsRegistration = false;
+
+			{
+				UE::TReadScopeLock ReadLock(RegisteredTypesLock);
+				if (ensure(RegisteredTypes.IsValidIndex(NewType.RegisteredTypeDefIndex)))
+				{
+					bNeedsRegistration = !RegisteredTypes[NewType.RegisteredTypeDefIndex]->IsValid();
+					check(bNeedsRegistration || (NewType == *RegisteredTypes[NewType.RegisteredTypeDefIndex]));
+				}
+			}
+
+			if (bNeedsRegistration)
 			{
-				TUniquePtr<FNiagaraTypeDefinition>& TypePtr = RegisteredTypes[NewType.RegisteredTypeDefIndex];
-				if (!TypePtr->IsValid())
+				UE::TWriteScopeLock Lock(RegisteredTypesLock);
+				if (ensure(RegisteredTypes.IsValidIndex(NewType.RegisteredTypeDefIndex)))
 				{
-					*TypePtr = NewType;
+					*RegisteredTypes[NewType.RegisteredTypeDefIndex] = NewType;
 				}
 			}
 		}
@@ -247,7 +292,7 @@ int32 FNiagaraTypeRegistry::RegisterTypeInternal(const FNiagaraTypeDefinition& N
 	{
 		// we need to create a new entry for a registered type
 		{
-			UE::TScopeLock Lock(RegisteredTypesLock);
+			UE::TWriteScopeLock Lock(RegisteredTypesLock);
 			NewType.RegisteredTypeDefIndex = RegisteredTypes.Num();
 			RegisteredTypes.Add(MakeUnique<FNiagaraTypeDefinition>(NewType));
 		}
@@ -257,10 +302,10 @@ int32 FNiagaraTypeRegistry::RegisterTypeInternal(const FNiagaraTypeDefinition& N
 			RegisteredSoftObjectPaths.Emplace(NewType.ClassStructOrEnum);
 			RegisteredTypeIndexMap.Add(TypeHash, NewType.RegisteredTypeDefIndex);
 		}
-
-		NiagaraTypeRegistryLocal::RegisteredTypeIndexCacheTLS.Add(TypeHash, NewType.RegisteredTypeDefIndex);
 	}
 
+	NiagaraTypeRegistryLocal::RegisteredTypeIndexCacheTLS.Add(TypeHash, NewType.RegisteredTypeDefIndex);
+
 	return NewType.RegisteredTypeDefIndex;
 }
 
@@ -268,9 +313,18 @@ const FNiagaraTypeDefinition& FNiagaraTypeRegistry::GetRegisteredType(int32 Type
 {
 	if (TypeIndex != INDEX_NONE)
 	{
-		const FNiagaraTypeRegistry& Registry = FNiagaraTypeRegistry::Get();
+		// check if we've got a RegistryProxy that can already provide the registered type
+		if (NiagaraTypeRegistryLocal::RegistryProxyRefCount && NiagaraTypeRegistryLocal::RegistryProxyContents.IsValidIndex(TypeIndex))
+		{
+			const FNiagaraTypeDefinition& ProxyTypeDefinition = NiagaraTypeRegistryLocal::RegistryProxyContents[TypeIndex];
+			if (ProxyTypeDefinition.IsValid())
+			{
+				return ProxyTypeDefinition;
+			}
+		}
 
-		UE::TScopeLock Lock(Registry.RegisteredTypesLock);
+		const FNiagaraTypeRegistry& Registry = FNiagaraTypeRegistry::Get();
+		UE::TReadScopeLock Lock(Registry.RegisteredTypesLock);
 
 		if (Registry.RegisteredTypes.IsValidIndex(TypeIndex))
 		{
@@ -279,6 +333,15 @@ const FNiagaraTypeDefinition& FNiagaraTypeRegistry::GetRegisteredType(int32 Type
 			// been unloaded.
 			if (RegisteredType.IsValid() && RegisteredType->IsValid())
 			{
+				if (NiagaraTypeRegistryLocal::RegistryProxyRefCount)
+				{
+					if (!NiagaraTypeRegistryLocal::RegistryProxyContents.IsValidIndex(TypeIndex))
+					{
+						NiagaraTypeRegistryLocal::RegistryProxyContents.SetNum(TypeIndex + 1);
+					}
+					NiagaraTypeRegistryLocal::RegistryProxyContents[TypeIndex] = *RegisteredType;
+				}
+
 				return *RegisteredType;
 			}
 		}
@@ -295,7 +358,7 @@ TArray<FNiagaraTypeDefinition> FNiagaraTypeRegistry::GetRegisteredTypes()
 	TArray<FNiagaraTypeDefinition> Types;
 
 	{
-		UE::TScopeLock Lock(Registry.RegisteredTypesLock);
+		UE::TReadScopeLock Lock(Registry.RegisteredTypesLock);
 		Types.Reserve(Registry.RegisteredTypes.Num());
 		Algo::TransformIf(Registry.RegisteredTypes, Types,
 			[](const TUniquePtr<FNiagaraTypeDefinition>& TypePtr) -> bool
diff --git a/Engine/Plugins/FX/Niagara/Source/Niagara/Public/NiagaraTypeRegistry.h b/Engine/Plugins/FX/Niagara/Source/Niagara/Public/NiagaraTypeRegistry.h
--- a/Engine/Plugins/FX/Niagara/Source/Niagara/Public/NiagaraTypeRegistry.h
+++ b/Engine/Plugins/FX/Niagara/Source/Niagara/Public/NiagaraTypeRegistry.h
@@ -3,7 +3,7 @@
 #pragma once
 
 #include "CoreMinimal.h"
-#include "Async/TransactionallySafeMutex.h"
+#include "Misc/TransactionallySafeRWLock.h"
 #include "Misc/ScopeLock.h"
 #include "Misc/SecureHash.h"
 #include "UObject/GCObject.h"
@@ -116,11 +116,12 @@ private:
 	int32 RegisterTypeInternal(const FNiagaraTypeDefinition& NewType);
 	const FNiagaraTypeDefinition& GetRegisteredType(int32 TypeIndex) const;
 
-	// The spinlock is used for accessing the RegisteredTypes array in performance critical areas (resolve TypeDefinitionHandles being the
-	// most common).
-	mutable UE::FTransactionallySafeMutex RegisteredTypesLock; 
+	mutable FTransactionallySafeRWLock RegisteredTypesLock;
 	FRegisteredTypesArray RegisteredTypes;
 
+	// atomic index that is incremented each time types may have been invalidated
+	std::atomic<uint32> RegisteredTypesGeneration = {0};
+
 	// Covers manipulation of all of the containers populated during registration (except for RegisteredTypes which is handled
 	// by it's own lock so that we can provide a fast path for the common case of reading the array)
 	FRWLock RegistrationLock;
@@ -158,7 +159,7 @@ private:
 		TArray<FNiagaraTypeDefinition> LocalCopy;
 
 		{
-			UE::TScopeLock Lock(RegisteredTypesLock);
+			UE::TReadScopeLock Lock(RegisteredTypesLock);
 			LocalCopy.Reserve(RegisteredTypes.Num());
 			Algo::TransformIf(RegisteredTypes, LocalCopy,
 				[Select](const TUniquePtr<FNiagaraTypeDefinition>& DefPtr) -> bool
@@ -180,3 +181,16 @@ private:
 		}
 	}
 };
+
+// helper class to trigger access to a TLS of NiagaraTypeDefinition.  While constructed, the current thread
+// will get lock-free access (minus the first call) to a registered type.  Destruction clears the TLS to avoid
+// any potential fallout of holding onto stale pointers or taking up too m
+class FNiagaraTypeRegistryTLSProxy final
+{
+public:
+	NIAGARA_API FNiagaraTypeRegistryTLSProxy();
+	NIAGARA_API ~FNiagaraTypeRegistryTLSProxy();
+
+private:
+	const uint32 CallingThreadId;
+};
\ No newline at end of file
diff --git a/Engine/Plugins/FX/Niagara/Source/Niagara/Public/NiagaraTypes.h b/Engine/Plugins/FX/Niagara/Source/Niagara/Public/NiagaraTypes.h
--- a/Engine/Plugins/FX/Niagara/Source/Niagara/Public/NiagaraTypes.h
+++ b/Engine/Plugins/FX/Niagara/Source/Niagara/Public/NiagaraTypes.h
@@ -837,16 +837,16 @@ public:
 		{
 			return FName();
 		}
-		return ClassStructOrEnum->GetFName();
+		return ClassStructOrEnum.GetFName();
 	}
 
-	FString GetName()const
+	FString GetName() const
 	{
 		if (IsValid() == false)
 		{
 			return TEXT("Invalid");
 		}
-		return ClassStructOrEnum->GetName();
+		return ClassStructOrEnum.GetName();
 	}
 
 	UStruct* GetStruct() const
