Populate a cache of function signatures in the CDO so that we don't have to invoke the full GetFunctions() whenever we want to check something in the translator.  From profiling the compilation of a complicated effect this can reduce thread time by 10s.
#tests Ran compilation tests comparing performance results and correctness

diff --git a/Engine/Plugins/FX/Niagara/Source/Niagara/Classes/NiagaraDataInterface.h b/Engine/Plugins/FX/Niagara/Source/Niagara/Classes/NiagaraDataInterface.h
--- a/Engine/Plugins/FX/Niagara/Source/Niagara/Classes/NiagaraDataInterface.h
+++ b/Engine/Plugins/FX/Niagara/Source/Niagara/Classes/NiagaraDataInterface.h
@@ -941,6 +941,11 @@ private:
 #if WITH_EDITOR
 	FSimpleMulticastDelegate OnErrorsRefreshedDelegate;
 #endif
+
+#if WITH_EDITORONLY_DATA
+	mutable FRWLock CachedFunctionSignatureLock;
+	mutable TArray<FNiagaraFunctionSignature> CachedFunctionSignatures;
+#endif
 };
 
 /** Helper class for decoding NDI parameters into a usable struct type. */
diff --git a/Engine/Plugins/FX/Niagara/Source/Niagara/Private/NiagaraDataInterface.cpp b/Engine/Plugins/FX/Niagara/Source/Niagara/Private/NiagaraDataInterface.cpp
--- a/Engine/Plugins/FX/Niagara/Source/Niagara/Private/NiagaraDataInterface.cpp
+++ b/Engine/Plugins/FX/Niagara/Source/Niagara/Private/NiagaraDataInterface.cpp
@@ -294,14 +294,35 @@ void UNiagaraDataInterface::GetFeedback(UNiagaraDataInterface* DataInterface, TA
 void UNiagaraDataInterface::GetFunctionSignatures(TArray<FNiagaraFunctionSignature>& OutFunctions) const
 {
 #if WITH_EDITORONLY_DATA
-	// Until we can eliminate GetFunctions() we just call it, and let it's base implementation call
-	// GetFunctionsInternal().  When GetFunctions() is deleted we can just call GetFunctionsInternal()
-	// directly here.
-	//GetFunctionsInternal(OutFunctions);
-
-	PRAGMA_DISABLE_DEPRECATION_WARNINGS // to be removed when GetFunctions is removed
-	const_cast<UNiagaraDataInterface*>(this)->GetFunctions(OutFunctions);
-	PRAGMA_ENABLE_DEPRECATION_WARNINGS // to be removed when GetFunctions is removed
+	// only cache things if we're dealing with the CDO
+	if (HasAnyFlags(EObjectFlags::RF_ClassDefaultObject))
+	{
+		{
+			UE::TReadScopeLock ReadLock(CachedFunctionSignatureLock);
+			if (!CachedFunctionSignatures.IsEmpty())
+			{
+				OutFunctions = CachedFunctionSignatures;
+				return;
+			}
+		}
+
+		{
+			UE::TWriteScopeLock WriteLock(CachedFunctionSignatureLock);
+
+			if (CachedFunctionSignatures.IsEmpty())
+			{
+				PRAGMA_DISABLE_DEPRECATION_WARNINGS // to be removed when GetFunctions is removed
+				const_cast<UNiagaraDataInterface*>(this)->GetFunctions(CachedFunctionSignatures);
+				PRAGMA_ENABLE_DEPRECATION_WARNINGS // to be removed when GetFunctions is removed
+			}
+
+			OutFunctions = CachedFunctionSignatures;
+		}
+	}
+	else
+	{
+		return GetClass()->GetDefaultObject<UNiagaraDataInterface>()->GetFunctionSignatures(OutFunctions);
+	}
 #endif
 }
 
