Optimizations to the Attribute Trimming portion of niagara translation.
-Incorporate the output node UsageId into the search criteria when building traversals (especially important for simulations with sim stages)
-Add a sorting pass to the CustomHLSL tokens to reduce search costs during the trimming
-Some simple stringbuilder wins
#tests Verified translated shaders all match at a reduced cost

diff --git a/Engine/Plugins/FX/Niagara/Source/NiagaraEditor/Private/NiagaraAttributeTrimmer.cpp b/Engine/Plugins/FX/Niagara/Source/NiagaraEditor/Private/NiagaraAttributeTrimmer.cpp
--- a/Engine/Plugins/FX/Niagara/Source/NiagaraEditor/Private/NiagaraAttributeTrimmer.cpp
+++ b/Engine/Plugins/FX/Niagara/Source/NiagaraEditor/Private/NiagaraAttributeTrimmer.cpp
@@ -21,11 +21,11 @@ template<typename GraphBridge>
 class FNiagaraAttributeTrimmerHelper<GraphBridge>::FFunctionInputResolver
 {
 public:
-	FFunctionInputResolver(const FGraph* Graph, ENiagaraScriptUsage Usage)
+	FFunctionInputResolver(const FGraph* Graph, ENiagaraScriptUsage Usage, const FGuid& UsageId)
 	{
 		TMap<FName, FModuleScopedPin> UnresolvedInputs;
 
-		Traverse(TEXT(""), Graph, Usage, UnresolvedInputs);
+		Traverse(TEXT(""), Graph, Usage, UsageId, UnresolvedInputs);
 	}
 
 	FModuleScopedPin ResolveInput(const FModuleScopedPin& Input) const
@@ -39,14 +39,17 @@ public:
 	}
 
 private:
-	void Traverse(FString Namespace, const FGraph* Graph, ENiagaraScriptUsage Usage, TMap<FName, FModuleScopedPin>& UnresolvedInputs)
+	void Traverse(FStringView Namespace, const FGraph* Graph, ENiagaraScriptUsage Usage, const FGuid& UsageId, TMap<FName, FModuleScopedPin>& UnresolvedInputs)
 	{
-		const FString NamespacePrefix = (Namespace.Len() > 0) ? (Namespace + TEXT(".")) : TEXT("");
-
-		TArray<const FOutputNode*> OutputNodes;
-		GraphBridge::FindOutputNodes(Graph, Usage, OutputNodes);
+		FNameBuilder NamespacePrefix;
+		if (Namespace.Len() > 0)
+		{
+			NamespacePrefix << Namespace;
+			NamespacePrefix << TEXT(".");
+		}
+		const FName NamespaceFName(Namespace);
 
-		for (const FOutputNode* OutputNode : OutputNodes)
+		if (const FOutputNode* OutputNode = GraphBridge::FindOutputNode(Graph, Usage, UsageId))
 		{
 			TArray<const FNode*> TraversedNodes;
 			GraphBridge::BuildTraversal(Graph, OutputNode, TraversedNodes);
@@ -57,15 +60,19 @@ private:
 				{
 					if (const FGraph* FunctionGraph = GraphBridge::GetFunctionNodeGraph(FunctionNode))
 					{
+						FNameBuilder NamespacedFunction;
+						NamespacedFunction << NamespacePrefix.ToView();
+						NamespacedFunction << GraphBridge::GetFunctionName(FunctionNode);
+
 						TMap<FName, FModuleScopedPin> FunctionInputs;
-						Traverse(NamespacePrefix + GraphBridge::GetFunctionName(FunctionNode), FunctionGraph, GraphBridge::GetFunctionUsage(FunctionNode), FunctionInputs);
+						Traverse(NamespacedFunction.ToView(), FunctionGraph, GraphBridge::GetFunctionUsage(FunctionNode), FGuid(), FunctionInputs);
 
 						// resolve any of the UnresolvedInputs that we've accumulated
 						for (const FInputPin* NodePin : GraphBridge::GetInputPins(FunctionNode))
 						{
 							if (FModuleScopedPin* InnerPin = FunctionInputs.Find(NodePin->PinName))
 							{
-								FunctionInputMap.Add(*InnerPin, FModuleScopedPin(NodePin, *Namespace));
+								FunctionInputMap.Add(*InnerPin, FModuleScopedPin(NodePin, NamespaceFName));
 								FunctionInputs.Remove(NodePin->PinName);
 							}
 						}
@@ -79,7 +86,7 @@ private:
 						check(OutputPinCount == 0);
 						++OutputPinCount;
 
-						UnresolvedInputs.Add(GraphBridge::GetInputVariable(InputNode).GetName(), FModuleScopedPin(InputNodePin, *Namespace));
+						UnresolvedInputs.Add(GraphBridge::GetInputVariable(InputNode).GetName(), FModuleScopedPin(InputNodePin, NamespaceFName));
 					}
 				}
 			}
@@ -93,10 +100,10 @@ template<typename GraphBridge>
 class FNiagaraAttributeTrimmerHelper<GraphBridge>::FImpureFunctionParser
 {
 public:
-	FImpureFunctionParser(const FGraph* Graph, ENiagaraScriptUsage Usage, FTrimAttributeCache& InCache)
+	FImpureFunctionParser(const FGraph* Graph, ENiagaraScriptUsage Usage, const FGuid& UsageId, FTrimAttributeCache& InCache)
 		: Cache(InCache)
 	{
-		Traverse(TEXT(""), Graph, Usage);
+		Traverse(TEXT(""), Graph, Usage, UsageId);
 	}
 
 	const TArray<FModuleScopedPin>& ReadFunctionInputs() const
@@ -105,14 +112,11 @@ public:
 	}
 
 private:
-	void Traverse(FString Namespace, const FGraph* Graph, ENiagaraScriptUsage Usage)
+	void Traverse(FString Namespace, const FGraph* Graph, ENiagaraScriptUsage Usage, const FGuid& UsageId)
 	{
 		const FString NamespacePrefix = (Namespace.Len() > 0) ? (Namespace + TEXT(".")) : TEXT("");
 
-		TArray<const FOutputNode*> OutputNodes;
-		GraphBridge::FindOutputNodes(Graph, Usage, OutputNodes);
-
-		for (const FOutputNode* OutputNode : OutputNodes)
+		if (const FOutputNode* OutputNode = GraphBridge::FindOutputNode(Graph, Usage, UsageId))
 		{
 			TArray<const FNode*> TraversedNodes;
 			GraphBridge::BuildTraversal(Graph, OutputNode, TraversedNodes);
@@ -123,7 +127,7 @@ private:
 				{
 					if (const FGraph* FunctionGraph = GraphBridge::GetFunctionNodeGraph(FunctionNode))
 					{
-						Traverse(NamespacePrefix + GraphBridge::GetFunctionName(FunctionNode), FunctionGraph, GraphBridge::GetFunctionUsage(FunctionNode));
+						Traverse(NamespacePrefix + GraphBridge::GetFunctionName(FunctionNode), FunctionGraph, GraphBridge::GetFunctionUsage(FunctionNode), FGuid());
 					}
 
 					bool bEvaluatePins = false;
@@ -677,7 +681,7 @@ void FNiagaraAttributeTrimmerHelper<GraphBridge>::TrimAttributes_Aggressive(cons
 
 	for (const FParamMapHistory* ParamMap : LocalParamHistories)
 	{
-		FFunctionInputResolver InputResolver(GraphBridge::GetGraph(CompileDuplicateData), ParamMap->OriginatingScriptUsage);
+		FFunctionInputResolver InputResolver(GraphBridge::GetGraph(CompileDuplicateData), ParamMap->OriginatingScriptUsage, ParamMap->UsageGuid);
 
 		const int32 VariableCount = ParamMap->Variables.Num();
 
@@ -702,7 +706,7 @@ void FNiagaraAttributeTrimmerHelper<GraphBridge>::TrimAttributes_Aggressive(cons
 			}
 		}
 
-		FImpureFunctionParser ImpureFunctionParser(GraphBridge::GetGraph(CompileDuplicateData), ParamMap->OriginatingScriptUsage, Cache);
+		FImpureFunctionParser ImpureFunctionParser(GraphBridge::GetGraph(CompileDuplicateData), ParamMap->OriginatingScriptUsage, ParamMap->UsageGuid, Cache);
 		for (const FModuleScopedPin& RequiredFunctionInput : ImpureFunctionParser.ReadFunctionInputs())
 		{
 			FDependencyChain Dependencies;
diff --git a/Engine/Plugins/FX/Niagara/Source/NiagaraEditor/Private/NiagaraCompilationDigestBridgeImpl.cpp b/Engine/Plugins/FX/Niagara/Source/NiagaraEditor/Private/NiagaraCompilationDigestBridgeImpl.cpp
--- a/Engine/Plugins/FX/Niagara/Source/NiagaraEditor/Private/NiagaraCompilationDigestBridgeImpl.cpp
+++ b/Engine/Plugins/FX/Niagara/Source/NiagaraEditor/Private/NiagaraCompilationDigestBridgeImpl.cpp
@@ -57,12 +57,7 @@ void FNiagaraCompilationDigestBridge::CustomHlslReferencesTokens(const FCustomHl
 
 	for (int32 i = 0; i < TokenStrings.Num(); i++)
 	{
-		FNameBuilder NameBuilder(TokenStrings[i]);
-		FStringView NameString(NameBuilder.ToView());
-		Results[i] = CustomNode->Tokens.ContainsByPredicate([NameString](const FString& Token) -> bool
-			{
-				return Token.Contains(NameString);
-			});
+		Results[i] = CustomNode->TokensContainsName(TokenStrings[i]);
 	}
 }
 
@@ -126,6 +121,11 @@ void FNiagaraCompilationDigestBridge::FindOutputNodes(const FGraph* Graph, ENiag
 	Graph->FindOutputNodes(ScriptUsage, OutputNodes);
 }
 
+const FNiagaraCompilationDigestBridge::FOutputNode* FNiagaraCompilationDigestBridge::FindOutputNode(const FGraph* Graph, ENiagaraScriptUsage ScriptUsage, const FGuid& UsageId)
+{
+	return Graph->FindOutputNode(ScriptUsage, UsageId);
+}
+
 void FNiagaraCompilationDigestBridge::FindOutputNodes(const FGraph* Graph, TArray<const FOutputNode*>& OutputNodes)
 {
 	Graph->FindOutputNodes(OutputNodes);
diff --git a/Engine/Plugins/FX/Niagara/Source/NiagaraEditor/Private/NiagaraCompilationGraphBridgeImpl.cpp b/Engine/Plugins/FX/Niagara/Source/NiagaraEditor/Private/NiagaraCompilationGraphBridgeImpl.cpp
--- a/Engine/Plugins/FX/Niagara/Source/NiagaraEditor/Private/NiagaraCompilationGraphBridgeImpl.cpp
+++ b/Engine/Plugins/FX/Niagara/Source/NiagaraEditor/Private/NiagaraCompilationGraphBridgeImpl.cpp
@@ -194,9 +194,9 @@ TArray<FNiagaraVariableBase> FNiagaraCompilationGraphBridge::GraphGetStaticSwitc
 
 	StaticSwitchInputs.Reserve(Variables.Num());
 	Algo::Transform(Variables, StaticSwitchInputs, [](const FNiagaraVariable& InVar) -> FNiagaraVariableBase
-		{
-			return InVar;
-		});
+	{
+		return InVar;
+	});
 
 	return StaticSwitchInputs;
 }
@@ -207,9 +207,9 @@ void FNiagaraCompilationGraphBridge::FindOutputNodes(const FGraph* Graph, ENiaga
 	Graph->FindOutputNodes(ScriptUsage, MutableOutputNodes);
 	OutputNodes.Reserve(MutableOutputNodes.Num());
 	Algo::Transform(MutableOutputNodes, OutputNodes, [](FOutputNode* OutputNode) -> const FOutputNode*
-		{
-			return OutputNode;
-		});
+	{
+		return OutputNode;
+	});
 }
 
 void FNiagaraCompilationGraphBridge::FindOutputNodes(const FGraph* Graph, TArray<const FOutputNode*>& OutputNodes)
@@ -218,9 +218,21 @@ void FNiagaraCompilationGraphBridge::FindOutputNodes(const FGraph* Graph, TArray
 	Graph->FindOutputNodes(MutableOutputNodes);
 	OutputNodes.Reserve(MutableOutputNodes.Num());
 	Algo::Transform(MutableOutputNodes, OutputNodes, [](FOutputNode* OutputNode) -> const FOutputNode*
-		{
-			return OutputNode;
-		});
+	{
+		return OutputNode;
+	});
+}
+
+const FNiagaraCompilationGraphBridge::FOutputNode* FNiagaraCompilationGraphBridge::FindOutputNode(const FGraph* Graph, ENiagaraScriptUsage ScriptUsage, const FGuid& UsageId)
+{
+	TArray<FOutputNode*> MutableOutputNodes;
+	Graph->FindOutputNodes(MutableOutputNodes);
+	UNiagaraNodeOutput** FoundEntry = Algo::FindByPredicate(MutableOutputNodes, [ScriptUsage, UsageId](FOutputNode* OutputNode) -> bool
+	{
+		return OutputNode->GetUsage() == ScriptUsage && OutputNode->GetUsageId() == UsageId;
+	});
+
+	return FoundEntry ? *FoundEntry : nullptr;
 }
 
 void FNiagaraCompilationGraphBridge::BuildTraversal(const FGraph* Graph, const FNode* OutputNode, TArray<const FNode*>& TraversedNodes)
@@ -229,9 +241,9 @@ void FNiagaraCompilationGraphBridge::BuildTraversal(const FGraph* Graph, const F
 	Graph->BuildTraversal(MutableTraversedNodes, const_cast<FNode*>(OutputNode), true);
 	TraversedNodes.Reserve(MutableTraversedNodes.Num());
 	Algo::Transform(MutableTraversedNodes, TraversedNodes, [](UNiagaraNode* Node) -> const FNode*
-		{
-			return Node;
-		});
+	{
+		return Node;
+	});
 }
 
 const UNiagaraGraph* FNiagaraCompilationGraphBridge::GetEmitterGraph(const FEmitterNode* EmitterNode)
diff --git a/Engine/Plugins/FX/Niagara/Source/NiagaraEditor/Private/NiagaraGraphDigest.cpp b/Engine/Plugins/FX/Niagara/Source/NiagaraEditor/Private/NiagaraGraphDigest.cpp
--- a/Engine/Plugins/FX/Niagara/Source/NiagaraEditor/Private/NiagaraGraphDigest.cpp
+++ b/Engine/Plugins/FX/Niagara/Source/NiagaraEditor/Private/NiagaraGraphDigest.cpp
@@ -3040,6 +3040,28 @@ void FNiagaraCompilationNodeFunctionCall::MultiFindParameterMapDefaultValues(ENi
 	}
 }
 
+struct FNiagaraCompilationNodeCustomHlsl::FIndexedTokenLess
+{
+	FIndexedTokenLess(TConstArrayView<FString> InStringsView)
+	: StringsView(InStringsView)
+	{}
+
+	bool operator()(int32 LhsIndex, int32 RhsIndex) const
+	{
+		return StringsView[LhsIndex].Compare(StringsView[RhsIndex]) < 0;
+	}
+
+	TConstArrayView<FString> StringsView;
+};
+
+struct FNiagaraCompilationNodeCustomHlsl::FStringViewLess
+{
+	bool operator()(FStringView LhsValue, FStringView RhsValue) const
+	{
+		return LhsValue.Compare(RhsValue) < 0;
+	}
+};
+
 FNiagaraCompilationNodeCustomHlsl::FNiagaraCompilationNodeCustomHlsl(const UNiagaraNodeCustomHlsl* InNode, FNiagaraCompilationGraphCreateContext& Context)
 	: FNiagaraCompilationNodeFunctionCall(InNode, Context, ENodeType::CustomHlsl)
 {
@@ -3048,11 +3070,19 @@ FNiagaraCompilationNodeCustomHlsl::FNiagaraCompilationNodeCustomHlsl(const UNiag
 	CustomHlsl = InNode->GetCustomHlsl();
 	TArray<FStringView> TokenViews;
 	InNode->GetTokens(TokenViews, false, false);
-	Tokens.Reserve(TokenViews.Num());
+
+	const int32 TokenCount = TokenViews.Num();
+	Tokens.Reserve(TokenCount);
+	SortedTokenIndices.Reserve(TokenCount);
+
+	int32 TokenIt = 0;
 	for (const FStringView View : TokenViews)
 	{
 		Tokens.Push(FString(View));
+		SortedTokenIndices.Add(TokenIt++);
 	}
+	Algo::StableSort(SortedTokenIndices, FIndexedTokenLess(Tokens));
+
 	InNode->GetIncludeFilePaths(CustomIncludePaths);
 
 	bCallsImpureFunctions = InNode->CallsImpureDataInterfaceFunctions();
@@ -3064,6 +3094,7 @@ FNiagaraCompilationNodeCustomHlsl::FNiagaraCompilationNodeCustomHlsl(const FNiag
 	, Signature(InNode.Signature)
 	, CustomHlsl(InNode.CustomHlsl)
 	, Tokens(InNode.Tokens)
+	, SortedTokenIndices(InNode.SortedTokenIndices)
 	, CustomIncludePaths(InNode.CustomIncludePaths)
 	, bCallsImpureFunctions(InNode.bCallsImpureFunctions)
 {
@@ -3146,6 +3177,19 @@ void FNiagaraCompilationNodeCustomHlsl::BuildParameterMapHistory(FParameterMapHi
 	}
 }
 
+bool FNiagaraCompilationNodeCustomHlsl::TokensContainsName(FName InName) const
+{
+	FNameBuilder NameBuilder(InName);
+	FStringView NameBuilderView = NameBuilder.ToView();
+
+	int32 SearchStartIndex = Algo::LowerBoundBy(SortedTokenIndices, NameBuilderView, [this](int32 TokenIndex) -> FStringView
+	{
+		return Tokens[TokenIndex];
+	}, FStringViewLess());
+
+	return Tokens.IsValidIndex(SearchStartIndex) && Tokens[SortedTokenIndices[SearchStartIndex]].StartsWith(NameBuilderView);
+}
+
 FNiagaraCompilationNodeIf::FNiagaraCompilationNodeIf(const UNiagaraNodeIf* InNode, FNiagaraCompilationGraphCreateContext& Context)
 	: FNiagaraCompilationNode(ENodeType::If, InNode, Context)
 {
diff --git a/Engine/Plugins/FX/Niagara/Source/NiagaraEditor/Private/NiagaraGraphDigest.h b/Engine/Plugins/FX/Niagara/Source/NiagaraEditor/Private/NiagaraGraphDigest.h
--- a/Engine/Plugins/FX/Niagara/Source/NiagaraEditor/Private/NiagaraGraphDigest.h
+++ b/Engine/Plugins/FX/Niagara/Source/NiagaraEditor/Private/NiagaraGraphDigest.h
@@ -460,12 +460,19 @@ public:
 
 	virtual void BuildParameterMapHistory(FParameterMapHistoryBuilder& Builder, bool bRecursive, bool bFilterForCompilation) const override;
 
+	bool TokensContainsName(FName InName) const;
+
 	ENiagaraScriptUsage CustomScriptUsage;
 	FNiagaraFunctionSignature Signature;
 	FString CustomHlsl;
 	TArray<FString> Tokens;
+	TArray<int32> SortedTokenIndices;
 	TArray<FNiagaraCustomHlslInclude> CustomIncludePaths;
 	bool bCallsImpureFunctions = false;
+
+private:
+	struct FIndexedTokenLess;
+	struct FStringViewLess;
 };
 
 class FNiagaraCompilationNodeIf : public FNiagaraCompilationNode
diff --git a/Engine/Plugins/FX/Niagara/Source/NiagaraEditor/Public/NiagaraCompilationBridge.h b/Engine/Plugins/FX/Niagara/Source/NiagaraEditor/Public/NiagaraCompilationBridge.h
--- a/Engine/Plugins/FX/Niagara/Source/NiagaraEditor/Public/NiagaraCompilationBridge.h
+++ b/Engine/Plugins/FX/Niagara/Source/NiagaraEditor/Public/NiagaraCompilationBridge.h
@@ -159,6 +159,7 @@ struct FNiagaraCompilationGraphBridge
 	static TArray<FNiagaraVariableBase> GraphGetStaticSwitchInputs(const FGraph* Graph);
 	static void FindOutputNodes(const FGraph* Graph, ENiagaraScriptUsage ScriptUsage, TArray<const FOutputNode*>& OutputNodes);
 	static void FindOutputNodes(const FGraph* Graph, TArray<const FOutputNode*>& OutputNodes);
+	static const FOutputNode* FindOutputNode(const FGraph* Graph, ENiagaraScriptUsage ScriptUsage, const FGuid& UsageId);
 	static void BuildTraversal(const FGraph* Graph, const FNode* OutputNode, TArray<const FNode*>& TraversedNodes);
 	static const FGraph* GetEmitterGraph(const FEmitterNode* EmitterNode);
 	static FString GetEmitterUniqueName(const FEmitterNode* EmitterNode);
@@ -289,6 +290,7 @@ struct FNiagaraCompilationDigestBridge
 	static TArray<FNiagaraVariableBase> GraphGetStaticSwitchInputs(const FGraph* Graph);
 	static void FindOutputNodes(const FGraph* Graph, ENiagaraScriptUsage ScriptUsage, TArray<const FOutputNode*>& OutputNodes);
 	static void FindOutputNodes(const FGraph* Graph, TArray<const FOutputNode*>& OutputNodes);
+	static const FOutputNode* FindOutputNode(const FGraph* Graph, ENiagaraScriptUsage ScriptUsage, const FGuid& UsageId);
 	static void BuildTraversal(const FGraph* Graph, const FNode* OutputNode, TArray<const FNode*>& TraversedNodes);
 	static const FGraph* GetEmitterGraph(const FEmitterNode* EmitterNode);
 	static FNiagaraEmitterID GetEmitterID(const FEmitterNode* EmitterNode);
