Changes to the new compilation mode to make it explicit the different ways we can get the RI parameters, and ensure that the different techniques are mutually exclusive (prior to this we would run RI collection prior to beginning and then also within the task).  This also should hopefully address the failures that were occuring because we could have been invalidating the RI parameters we generated on the gamethread with the data on the task thread.

#tests validated shader generation before/after

diff --git a/Engine/Plugins/FX/Niagara/Source/NiagaraEditor/Private/NiagaraCompilationPrivate.h b/Engine/Plugins/FX/Niagara/Source/NiagaraEditor/Private/NiagaraCompilationPrivate.h
--- a/Engine/Plugins/FX/Niagara/Source/NiagaraEditor/Private/NiagaraCompilationPrivate.h
+++ b/Engine/Plugins/FX/Niagara/Source/NiagaraEditor/Private/NiagaraCompilationPrivate.h
@@ -343,3 +343,20 @@ protected:
 private:
 	TArray<TWeakObjectPtr<UNiagaraScriptSource>> TrackedScriptSourceCopies;
 };
+
+enum class ENiagaraCompileRIParamMode : uint8
+{
+	// the issue here is that we aren't strictly allowed to call ComputeVMCompilationId on the worker threads which leaves us in a
+	// state where we could have broken RI parameters.  We ensure that the RI parameters are up to date before we begin.  This is a
+	// costly operation, but sadly seems necessary for some content.
+	GameThread = 0,
+
+	// for now we need RI generation to be synchronous to match the behavior of the default
+	// compilation mode.  In the future we need to get RI preparation out of the compilation
+	// and instead we are provided with overridden RI values and don't need to feed back the
+	// entire list of RI parameters.
+	WorkerThread_Synchronous,
+
+	// desirable behavior that currently doesn't work
+	WorkerThread_Asynchronous,
+};
diff --git a/Engine/Plugins/FX/Niagara/Source/NiagaraEditor/Private/NiagaraCompilationTasks.cpp b/Engine/Plugins/FX/Niagara/Source/NiagaraEditor/Private/NiagaraCompilationTasks.cpp
--- a/Engine/Plugins/FX/Niagara/Source/NiagaraEditor/Private/NiagaraCompilationTasks.cpp
+++ b/Engine/Plugins/FX/Niagara/Source/NiagaraEditor/Private/NiagaraCompilationTasks.cpp
@@ -368,6 +368,7 @@ namespace NiagaraCompilationTasksImpl
 				BuildString = FString(MetaData.Find(ANSITEXTVIEW("Build")).AsString());
 				DateTime = MetaData.Find(ANSITEXTVIEW("Time")).AsDateTime();
 				CompileIdString = FString(MetaData.Find(ANSITEXTVIEW("CompileId")).AsString());
+				CompileTaskId = FGuid(FString(MetaData.Find(ANSITEXTVIEW("CompileTaskId")).AsString()));
 
 				const FCbArray RIParamArray = MetaData.Find(ANSITEXTVIEW("RIParam")).AsArray();
 				for (FCbFieldViewIterator It = RIParamArray.CreateViewIterator(); It; ++It)
@@ -385,12 +386,13 @@ namespace NiagaraCompilationTasksImpl
 			}
 		}
 
-		FScriptCompileMetaData(const FNiagaraSystemCompilationTask::FCompileTaskInfo& CompileTask)
+		FScriptCompileMetaData(const FNiagaraSystemCompilationTask* SystemCompileTask, const FNiagaraSystemCompilationTask::FCompileTaskInfo& CompileTask)
 		{
 			Version = EMetaDataVersion::Latest;
 			BuildString = FApp::GetBuildVersion();
 			DateTime = FDateTime::UtcNow();
 			CompileTask.CompileId.AppendKeyString(CompileIdString);
+			CompileTaskId = SystemCompileTask->GetCompilationTaskId();
 
 			BakedRapidIterationParameters.Reserve(CompileTask.BakedRapidIterationParameters.Num());
 			Algo::Transform(CompileTask.BakedRapidIterationParameters, BakedRapidIterationParameters, [](const FNiagaraVariable& InVar) -> TTuple<FName, FNiagaraVariant>
@@ -408,6 +410,7 @@ namespace NiagaraCompilationTasksImpl
 				Writer.AddString(ANSITEXTVIEW("Build"), BuildString);
 				Writer.AddDateTime(ANSITEXTVIEW("Time"), DateTime);
 				Writer.AddString(ANSITEXTVIEW("CompileId"), CompileIdString);
+				Writer.AddString(ANSITEXTVIEW("CompileTaskId"), CompileTaskId.ToString());
 
 				Writer.BeginArray(ANSITEXTVIEW("RIParam"));
 				{
@@ -432,19 +435,26 @@ namespace NiagaraCompilationTasksImpl
 		enum class EMetaDataVersion : int32
 		{
 			Invalid = 0,
-			Latest = 1
+
+			// new entries
+			WithCompilationTaskId,
+
+			// 
+			Last,
+			Latest = Last - 1,
 		};
 
 		EMetaDataVersion Version = EMetaDataVersion::Invalid;
 		FString BuildString;
 		FDateTime DateTime;
 		FString CompileIdString;
+		FGuid CompileTaskId;
 		TArray<TTuple<FName, FNiagaraVariant>> BakedRapidIterationParameters;
 	};
 
-	static FCbObject GenerateScriptCompileMetaData(const FNiagaraSystemCompilationTask::FCompileTaskInfo& CompileTask)
+	static FCbObject GenerateScriptCompileMetaData(const FNiagaraSystemCompilationTask* SystemCompileTask, const FNiagaraSystemCompilationTask::FCompileTaskInfo& CompileTask)
 	{
-		return FScriptCompileMetaData(CompileTask).ToCompactBinary();
+		return FScriptCompileMetaData(SystemCompileTask, CompileTask).ToCompactBinary();
 	}
 
 	static void ParseScriptCompileMetaData(const FCbObject& MetaData)
@@ -602,18 +612,20 @@ namespace NiagaraCompilationTasksImpl
 	}
 } // NiagaraCompilationTasksImpl
 
-FNiagaraSystemCompilationTask::FNiagaraSystemCompilationTask(FNiagaraCompilationTaskHandle InTaskHandle, UNiagaraSystem* InSystem)
+FNiagaraSystemCompilationTask::FNiagaraSystemCompilationTask(FNiagaraCompilationTaskHandle InTaskHandle, UNiagaraSystem* InSystem, ENiagaraCompileRIParamMode InRIParamMode)
 	: TaskHandle(InTaskHandle)
+	, RIParamMode(InRIParamMode)
 	, DDCRequestOwner(UE::DerivedData::EPriority::Normal)
 	, System_GT(InSystem)
 {
 	bGenerateCompileMetaData = GetDefault<UNiagaraEditorSettings>()->GetGenerateMetaDataOnCompile();
+
+	CompilationTaskId = FGuid::NewGuid();
 }
 
 void FNiagaraSystemCompilationTask::Abort()
 {
 	bAborting = true;
-	CompileCompletionEvent.Trigger();
 }
 
 FString FNiagaraSystemCompilationTask::GetDescription() const
@@ -1209,11 +1221,6 @@ void FNiagaraSystemCompilationTask::FCompileTaskInfo::GenerateOptimizedVMByteCod
 
 void FNiagaraSystemCompilationTask::FCompileTaskInfo::ProcessCompilation(FNiagaraSystemCompilationTask* SystemCompileTask, const FCompileGroupInfo& GroupInfo)
 {
-	ON_SCOPE_EXIT
-	{
-		CompileResultsProcessedEvent->Trigger();
-	};
-
 	if (ScriptCompileType == EScriptCompileType::CompileForVm)
 	{
 		GenerateOptimizedVMByteCode(SystemCompileTask, GroupInfo);
@@ -1444,7 +1451,7 @@ void FNiagaraSystemCompilationTask::FDispatchDataCachePutRequests::Launch(FNiaga
 					RecordBuilder.AddValue(NiagaraCompilationTasksImpl::ScriptDataValue, PutValue);
 					if (bIncludeMetaData)
 					{
-						RecordBuilder.SetMeta(NiagaraCompilationTasksImpl::GenerateScriptCompileMetaData(CompileTask));
+						RecordBuilder.SetMeta(NiagaraCompilationTasksImpl::GenerateScriptCompileMetaData(SystemCompileTask, CompileTask));
 					}
 
 					FCachePutRequest& PutRequest = PutRequests.AddZeroed_GetRef();
@@ -1504,6 +1511,12 @@ void FNiagaraSystemCompilationTask::DigestSystemInfo()
 
 	using namespace NiagaraCompilationTasksImpl;
 
+	if (RIParamMode == ENiagaraCompileRIParamMode::GameThread)
+	{
+		TRACE_CPUPROFILER_EVENT_SCOPE(NiagaraAsyncTask_PrepareRIParam);
+		System_GT->PrepareRapidIterationParametersForCompilation();
+	}
+
 	FNiagaraDigestDatabase& DigestDatabase = FNiagaraDigestDatabase::Get();
 	const UNiagaraScript* SystemSpawnScript = System_GT->GetSystemSpawnScript();
 	const UNiagaraScript* SystemUpdateScript = System_GT->GetSystemUpdateScript();
@@ -1521,7 +1534,13 @@ void FNiagaraSystemCompilationTask::DigestSystemInfo()
 	SystemInfo.ConstantResolver = FNiagaraFixedConstantResolver(FCompileConstantResolver(System_GT.Get(), ENiagaraScriptUsage::SystemSpawnScript));
 	SystemInfo.OwnedScriptKeys = { SystemSpawnScript, SystemUpdateScript };
 	SystemInfo.SystemSourceGraph = DigestDatabase.CreateGraphDigest(SystemGraph, ChangeIdBuilder);
-	
+
+	if (RIParamMode == ENiagaraCompileRIParamMode::GameThread)
+	{
+		TSharedPtr<FNiagaraGraphCachedDataBase, ESPMode::ThreadSafe> CachedTraversalSystemData = System_GT->GetCachedTraversalData();
+		SystemInfo.StaticVariableResults = static_cast<FNiagaraGraphCachedBuiltHistory*>(CachedTraversalSystemData.Get())->StaticVariables;
+	}
+	else
 	{
 		TArray<FNiagaraVariable> StaticVariablesFromEmitters;
 
@@ -1565,6 +1584,18 @@ void FNiagaraSystemCompilationTask::DigestSystemInfo()
 				EmitterInfo.DigestedEmitterIndex = DigestedEmitterIndex;
 				EmitterInfo.SourceGraph = DigestDatabase.CreateGraphDigest(EmitterGraph, ChangeIdBuilder);
 
+				if (RIParamMode == ENiagaraCompileRIParamMode::GameThread)
+				{
+					EmitterInfo.StaticVariableResults = SystemInfo.StaticVariableResults;
+
+					TSharedPtr<FNiagaraGraphCachedDataBase, ESPMode::ThreadSafe> CachedTraversalEmitterData = Handle.GetInstance().Emitter->GetCachedTraversalData(Handle.GetInstance().Version);
+					EmitterInfo.StaticVariableResults.Append(static_cast<FNiagaraGraphCachedBuiltHistory*>(CachedTraversalEmitterData.Get())->StaticVariables);
+				}
+				else
+				{
+					EmitterData->GatherStaticVariables(EmitterInfo.InitialStaticVariables);
+				}
+
 				// be sure to incorporate our constant resolver into the top level SystemInfo.ConstantResolver
 				SystemInfo.ConstantResolver.AddChildResolver(EmitterInfo.EmitterHandleId, EmitterInfo.ConstantResolver);
 			}
@@ -1612,8 +1643,6 @@ void FNiagaraSystemCompilationTask::DigestSystemInfo()
 					}
 				}
 			}
-
-			EmitterData->GatherStaticVariables(EmitterInfo.InitialStaticVariables);
 		}
 	}
 
@@ -1904,13 +1933,13 @@ bool FNiagaraSystemCompilationTask::AreResultsPending() const
 void FNiagaraSystemCompilationTask::WaitTillCompileCompletion()
 {
 	const FTimespan WaitTimeout = FTimespan::FromMilliseconds(50.0);
-	while (!CompileCompletionEvent.IsCompleted())
+	while (!MasterCompilationTask.IsCompleted())
 	{
 		// if the wait doesn't complete the task then we need to make sure to poke the
 		// compilation manager since there we may be waiting on it
-		if (!CompileCompletionEvent.Wait(WaitTimeout))
+		if (!MasterCompilationTask.Wait(WaitTimeout))
 		{
-			FAssetCompilingManager::Get().ProcessAsyncTasks();
+			FNiagaraSystemCompilingManager::Get().AdvanceAsyncTasks();
 		}
 	}
 }
@@ -2044,6 +2073,11 @@ bool FNiagaraSystemCompilationTask::GenerateCompileMetaData() const
 	return bGenerateCompileMetaData;
 }
 
+const FGuid& FNiagaraSystemCompilationTask::GetCompilationTaskId() const
+{
+	return CompilationTaskId;
+}
+
 bool FNiagaraSystemCompilationTask::GetStageName(int32 EmitterIndex, const FNiagaraCompilationNodeOutput* OutputNode, FName& OutStageName) const
 {
 	OutStageName = NAME_None;
@@ -2588,7 +2622,6 @@ void FNiagaraSystemCompilationTask::BuildAndApplyRapidIterationParameters()
 
 void FNiagaraSystemCompilationTask::IssuePostResultsProcessedTasks()
 {
-	CompileCompletionEvent.Trigger();
 	CompilationState = EState::ResultsProcessed;
 	PutRequestHelper = MakeUnique<FDispatchDataCachePutRequests>();
 	PutRequestHelper->Launch(this);
@@ -2599,12 +2632,12 @@ void FNiagaraSystemCompilationTask::IssuePostResultsProcessedTasks()
 	}, PutRequestHelper->CompletionEvent);
 }
 
-void FNiagaraSystemCompilationTask::IssueCompilationTasks()
+UE::Tasks::FTask FNiagaraSystemCompilationTask::IssueCompilationTasks()
 {
 	using namespace UE::Tasks;
 
 	TArray<FTask> IssuedCompilationTasks;
-	TArray<FTaskEvent> CompilationTasks;
+	TArray<FTask> CompilationTasks;
 
 	if (HasOutstandingCompileTasks())
 	{
@@ -2633,7 +2666,6 @@ void FNiagaraSystemCompilationTask::IssueCompilationTasks()
 					if (CompileTask.IsOutstanding(*this))
 					{
 						CompileTask.CompileResultsReadyEvent = MakeUnique<FTaskEvent>(UE_SOURCE_LOCATION);
-						CompileTask.CompileResultsProcessedEvent = MakeUnique<FTaskEvent>(UE_SOURCE_LOCATION);
 
 						IssuedCompilationTasks.Add(Launch(UE_SOURCE_LOCATION, [Self = AsShared(), &GroupInfo, &CompileTask]
 						{
@@ -2671,14 +2703,12 @@ void FNiagaraSystemCompilationTask::IssueCompilationTasks()
 							}
 						}, InstantiateGraphTask));
 
-						Launch(UE_SOURCE_LOCATION, [Self = AsShared(), &GroupInfo, &CompileTask]
+						CompilationTasks.Add(Launch(UE_SOURCE_LOCATION, [Self = AsShared(), &GroupInfo, &CompileTask]
 						{
 							COOK_STAT(auto Timer = NiagaraSystemCookStats::UsageStats.TimeSyncWork(); Timer.TrackCyclesOnly(););
 
 							CompileTask.ProcessCompilation(&Self.Get(), GroupInfo);
-						}, *CompileTask.CompileResultsReadyEvent);
-
-						CompilationTasks.Add(*CompileTask.CompileResultsProcessedEvent);
+						}, *CompileTask.CompileResultsReadyEvent));
 					}
 				}
 			}
@@ -2690,7 +2720,7 @@ void FNiagaraSystemCompilationTask::IssueCompilationTasks()
 		Self->CompilationState = EState::WaitingForProcessing;
 	}, IssuedCompilationTasks);
 
-	Launch(UE_SOURCE_LOCATION, [Self = AsShared()]
+	return Launch(UE_SOURCE_LOCATION, [Self = AsShared()]
 	{
 		Self->IssuePostResultsProcessedTasks();
 	}, CompilationTasks);
@@ -2728,17 +2758,13 @@ UE::Tasks::FTask FNiagaraSystemCompilationTask::BeginTasks()
 	// that we can apply them to the script if something has changed.
 	if (SystemInfo.bUseRapidIterationParams)
 	{
-		// for now we need RI generation to be synchronous to match the behavior of the default
-		// compilation mode.  In the future we need to get RI preparation out of the compilation
-		// and instead we are provided with overridden RI values and don't need to feed back the
-		// entire list of RI parameters.
 		constexpr bool bSyncPrepareRapidIterationParameter = true;
 
-		if (bSyncPrepareRapidIterationParameter)
+		if (RIParamMode == ENiagaraCompileRIParamMode::WorkerThread_Synchronous)
 		{
 			BuildAndApplyRapidIterationParameters();
 		}
-		else
+		else if (RIParamMode == ENiagaraCompileRIParamMode::WorkerThread_Asynchronous)
 		{
 			SystemTaskPrerequisites.Add(BuildRapidIterationParametersAsync());
 		}
@@ -2752,7 +2778,7 @@ UE::Tasks::FTask FNiagaraSystemCompilationTask::BeginTasks()
 		{
 			TArray<FTask> CompilationTaskPrerequisites;
 
-			if (!Self->SystemInfo.bUseRapidIterationParams)
+			if (Self->RIParamMode != ENiagaraCompileRIParamMode::GameThread && !Self->SystemInfo.bUseRapidIterationParams)
 			{
 				FTask BuildRIParamTask = Self->BuildRapidIterationParametersAsync();
 				Self->PostRIParameterGetRequestHelper = MakeUnique<FDispatchAndProcessDataCacheGetRequests>();
@@ -2767,7 +2793,7 @@ UE::Tasks::FTask FNiagaraSystemCompilationTask::BeginTasks()
 
 			FTask InnerTask = Launch(UE_SOURCE_LOCATION, [Self]
 			{
-				Self->IssueCompilationTasks();
+				Self->MasterCompilationTask = Self->IssueCompilationTasks();
 			}, CompilationTaskPrerequisites);
 		}
 		else
diff --git a/Engine/Plugins/FX/Niagara/Source/NiagaraEditor/Private/NiagaraCompilationTasks.h b/Engine/Plugins/FX/Niagara/Source/NiagaraEditor/Private/NiagaraCompilationTasks.h
--- a/Engine/Plugins/FX/Niagara/Source/NiagaraEditor/Private/NiagaraCompilationTasks.h
+++ b/Engine/Plugins/FX/Niagara/Source/NiagaraEditor/Private/NiagaraCompilationTasks.h
@@ -35,7 +35,7 @@ struct FNiagaraSystemCompilationTask : public TSharedFromThis<FNiagaraSystemComp
 		EShaderPlatform ShaderPlatform;
 	};
 
-	FNiagaraSystemCompilationTask(FNiagaraCompilationTaskHandle InTaskHandle, UNiagaraSystem* InSystem);
+	FNiagaraSystemCompilationTask(FNiagaraCompilationTaskHandle InTaskHandle, UNiagaraSystem* InSystem, ENiagaraCompileRIParamMode InRIParamMode);
 
 	void Abort();
 
@@ -47,7 +47,7 @@ struct FNiagaraSystemCompilationTask : public TSharedFromThis<FNiagaraSystemComp
 
 	UE::Tasks::FTask BuildRapidIterationParametersAsync();
 	void BuildAndApplyRapidIterationParameters();
-	void IssueCompilationTasks();
+	UE::Tasks::FTask IssueCompilationTasks();
 	void IssuePostResultsProcessedTasks();
 	bool HasOutstandingCompileTasks() const;
 	FString GetDescription() const;
@@ -195,7 +195,6 @@ struct FNiagaraSystemCompilationTask : public TSharedFromThis<FNiagaraSystemComp
 		FDDCTaskInfo DDCTaskInfo;
 
 		TUniquePtr<UE::Tasks::FTaskEvent> CompileResultsReadyEvent;
-		TUniquePtr<UE::Tasks::FTaskEvent> CompileResultsProcessedEvent;
 		TMap<FName, UNiagaraDataInterface*> NamedDataInterfaces;
 		TArray<FNiagaraVariable> BakedRapidIterationParameters; // todo - remove and have the data pulled directly from the precompile data
 		double TaskStartTime = 0;
@@ -242,6 +241,8 @@ struct FNiagaraSystemCompilationTask : public TSharedFromThis<FNiagaraSystemComp
 	int32 TasksAwaitingDDCGetResults = 0;
 	bool bCompileForEdit = false;
 
+	const ENiagaraCompileRIParamMode RIParamMode;
+
 	void Tick();
 	bool Poll(FNiagaraSystemAsyncCompileResults& Results) const;
 	bool CanRemove() const;
@@ -254,6 +255,7 @@ struct FNiagaraSystemCompilationTask : public TSharedFromThis<FNiagaraSystemComp
 	void GetAvailableCollections(TArray<FNiagaraCompilationNPCHandle>& OutCollections) const;
 
 	bool GenerateCompileMetaData() const;
+	const FGuid& GetCompilationTaskId() const;
 
 	UE::DerivedData::FRequestOwner DDCRequestOwner;
 
@@ -274,9 +276,13 @@ private:
 
 	TMap<TObjectKey<UNiagaraParameterCollection>, FNiagaraCompilationNPCHandle> DigestedParameterCollections;
 
-	UE::Tasks::FTaskEvent CompileCompletionEvent = UE::Tasks::FTaskEvent(UE_SOURCE_LOCATION);
+	UE::Tasks::FTask MasterCompilationTask;
 
 	bool bAborting = false;
 
 	bool bGenerateCompileMetaData = false;
+
+	// Guid just to track this specific compilation task.  Used in the DDC metadata to help see if in the case of a problem whether
+	// the data pulled from DDC all came from the same job.
+	FGuid CompilationTaskId;
 };
diff --git a/Engine/Plugins/FX/Niagara/Source/NiagaraEditor/Private/NiagaraSystemCompilingManager.cpp b/Engine/Plugins/FX/Niagara/Source/NiagaraEditor/Private/NiagaraSystemCompilingManager.cpp
--- a/Engine/Plugins/FX/Niagara/Source/NiagaraEditor/Private/NiagaraSystemCompilingManager.cpp
+++ b/Engine/Plugins/FX/Niagara/Source/NiagaraEditor/Private/NiagaraSystemCompilingManager.cpp
@@ -7,10 +7,11 @@
 #include "NiagaraCompilationTasks.h"
 #include "NiagaraEditorModule.h"
 #include "NiagaraShaderType.h"
+#include "Misc/ScopeRWLock.h"
 #include "ProfilingDebugging/CookStats.h"
+#include "ShaderCompiler.h"
 #include "UObject/UObjectIterator.h"
 #include "UObject/UObjectThreadContext.h"
-#include "Misc/ScopeRWLock.h"
 
 #define LOCTEXT_NAMESPACE "NiagaraCompilationManager"
 
@@ -102,11 +103,14 @@ FNiagaraShaderMapId BuildShaderMapId(FNiagaraShaderType* ShaderType, const ITarg
 	return ShaderMapId;
 }
 
-void EnsureTargetPlatformsLoaded(const FNiagaraSystemCompilationTask& CompilationTask)
+void EnsureGlobalsInitialized(const FNiagaraSystemCompilationTask& CompilationTask)
 {
 	// in order to work around the fact that the target platform API is not thread safe we need to make sure that the target platform and the
 	// various shader formats it handles has been initialized
 	GetTargetPlatformManagerRef().ShaderFormatVersion(TEXT("VVM_1_0"));
+
+	// additionally, make sure that some shader compiler internals are also initialized appropriately
+	GShaderCompilingManager->GetAbsoluteShaderDebugInfoDirectory();
 }
 
 };
@@ -129,7 +133,12 @@ FTextFormat FNiagaraSystemCompilingManager::GetAssetNameFormat() const
 
 TArrayView<FName> FNiagaraSystemCompilingManager::GetDependentTypeNames() const
 {
-	return {};
+	static FName DependentTypeNames[] =
+	{
+		FShaderCompilingManager::GetStaticAssetTypeName(),
+	};
+
+	return TArrayView<FName>(DependentTypeNames);
 }
 
 int32 FNiagaraSystemCompilingManager::GetNumRemainingAssets() const
@@ -464,7 +473,7 @@ FNiagaraCompilationTaskHandle FNiagaraSystemCompilingManager::AddSystem(UNiagara
 			QueuedRequests.Add(RequestHandle);
 			CompilationTask = SystemRequestMap.Add(
 				RequestHandle,
-				MakeShared<FNiagaraSystemCompilationTask, ESPMode::ThreadSafe>(RequestHandle, System));
+				MakeShared<FNiagaraSystemCompilationTask, ESPMode::ThreadSafe>(RequestHandle, System, CompileOptions.RIParamMode));
 		}
 
 		CompilationTask->PrepareStartTime = FPlatformTime::Seconds();
@@ -481,7 +490,7 @@ FNiagaraCompilationTaskHandle FNiagaraSystemCompilingManager::AddSystem(UNiagara
 			CompilationTask->AddScript(ScriptToCompile.EmitterIndex, ScriptToCompile.Script, ScriptToCompile.CompileId, ScriptToCompile.bRequiresCompilation, ScriptToCompile.ShaderRequests);
 		}
 
-		NiagaraSystemCompilingManagerImpl::EnsureTargetPlatformsLoaded(*CompilationTask);
+		NiagaraSystemCompilingManagerImpl::EnsureGlobalsInitialized(*CompilationTask);
 
 		CompilationTask->QueueStartTime = FPlatformTime::Seconds();
 	}
@@ -540,6 +549,12 @@ void FNiagaraSystemCompilingManager::AbortSystemCompile(FNiagaraCompilationTaskH
 	}
 }
 
+void FNiagaraSystemCompilingManager::AdvanceAsyncTasks()
+{
+	GShaderCompilingManager->ProcessAsyncResults(true /*TimeSlice*/, false /*BlockOnGlobalShaderCompletion*/);
+	ProcessAsyncTasks(true /*bLimitExecutionTime*/);
+}
+
 void FNiagaraSystemCompilingManager::QueueGameThreadFunction(FGameThreadFunction GameThreadTask)
 {
 	FWriteScopeLock Write(GameThreadFunctionLock);
@@ -634,11 +649,6 @@ FNiagaraCompilationTaskHandle FNiagaraEditorModule::RequestCompileSystem(UNiagar
 	const TArray<TWeakObjectPtr<UNiagaraParameterCollection>>& Collections = ParameterCollectionAssetCache.Get();
 	CompileOptions.ParameterCollections = ParameterCollectionAssetCache.Get();
 
-	// the issue here is that we aren't strictly allowed to call ComputeVMCompilationId on the worker threads which leaves us in a
-	// state where we could have broken RI parameters.  We ensure that the RI parameters are up to date before we begin.  This is a
-	// costly operation, but sadly seems necessary for some content.
-	System->PrepareRapidIterationParametersForCompilation();
-
 	return FNiagaraSystemCompilingManager::Get().AddSystem(System, CompileOptions);
 }
 
diff --git a/Engine/Plugins/FX/Niagara/Source/NiagaraEditor/Private/NiagaraSystemCompilingManager.h b/Engine/Plugins/FX/Niagara/Source/NiagaraEditor/Private/NiagaraSystemCompilingManager.h
--- a/Engine/Plugins/FX/Niagara/Source/NiagaraEditor/Private/NiagaraSystemCompilingManager.h
+++ b/Engine/Plugins/FX/Niagara/Source/NiagaraEditor/Private/NiagaraSystemCompilingManager.h
@@ -21,6 +21,7 @@ public:
 		EShaderPlatform PreviewShaderPlatform = EShaderPlatform::SP_NumPlatforms;
 		ERHIFeatureLevel::Type PreviewFeatureLevel = ERHIFeatureLevel::Num;
 		bool bForced = false;
+		ENiagaraCompileRIParamMode RIParamMode = ENiagaraCompileRIParamMode::GameThread;
 	};
 
 	static FNiagaraSystemCompilingManager& Get();
@@ -29,6 +30,8 @@ public:
 	bool PollSystemCompile(FNiagaraCompilationTaskHandle TaskHandle, bool bPeek, bool bWait, FNiagaraSystemAsyncCompileResults& Results);
 	void AbortSystemCompile(FNiagaraCompilationTaskHandle TaskHandle);
 
+	void AdvanceAsyncTasks();
+
 	using FGameThreadFunction = TFunction<void()>;
 	void QueueGameThreadFunction(FGameThreadFunction GameThreadTask);
 
