Cleanup variable name and some comments from previous changes
#rnx
#tests none

diff --git a/Engine/Plugins/FX/Niagara/Source/Niagara/Public/NiagaraTypeRegistry.h b/Engine/Plugins/FX/Niagara/Source/Niagara/Public/NiagaraTypeRegistry.h
--- a/Engine/Plugins/FX/Niagara/Source/Niagara/Public/NiagaraTypeRegistry.h
+++ b/Engine/Plugins/FX/Niagara/Source/Niagara/Public/NiagaraTypeRegistry.h
@@ -184,7 +184,7 @@ private:
 
 // helper class to trigger access to a TLS of NiagaraTypeDefinition.  While constructed, the current thread
 // will get lock-free access (minus the first call) to a registered type.  Destruction clears the TLS to avoid
-// any potential fallout of holding onto stale pointers or taking up too m
+// any potential fallout of holding onto stale pointers or taking up too much memory bloat
 class FNiagaraTypeRegistryTLSProxy final
 {
 public:
diff --git a/Engine/Plugins/FX/Niagara/Source/NiagaraEditor/Private/NiagaraCompilationTasks.cpp b/Engine/Plugins/FX/Niagara/Source/NiagaraEditor/Private/NiagaraCompilationTasks.cpp
--- a/Engine/Plugins/FX/Niagara/Source/NiagaraEditor/Private/NiagaraCompilationTasks.cpp
+++ b/Engine/Plugins/FX/Niagara/Source/NiagaraEditor/Private/NiagaraCompilationTasks.cpp
@@ -1933,11 +1933,11 @@ bool FNiagaraSystemCompilationTask::AreResultsPending() const
 void FNiagaraSystemCompilationTask::WaitTillCompileCompletion()
 {
 	const FTimespan WaitTimeout = FTimespan::FromMilliseconds(50.0);
-	while (!MasterCompilationTask.IsCompleted())
+	while (!ParentCompilationTask.IsCompleted())
 	{
 		// if the wait doesn't complete the task then we need to make sure to poke the
 		// compilation manager since there we may be waiting on it
-		if (!MasterCompilationTask.Wait(WaitTimeout))
+		if (!ParentCompilationTask.Wait(WaitTimeout))
 		{
 			FNiagaraSystemCompilingManager::Get().AdvanceAsyncTasks();
 		}
@@ -2793,7 +2793,7 @@ UE::Tasks::FTask FNiagaraSystemCompilationTask::BeginTasks()
 
 			FTask InnerTask = Launch(UE_SOURCE_LOCATION, [Self]
 			{
-				Self->MasterCompilationTask = Self->IssueCompilationTasks();
+				Self->ParentCompilationTask = Self->IssueCompilationTasks();
 			}, CompilationTaskPrerequisites);
 		}
 		else
diff --git a/Engine/Plugins/FX/Niagara/Source/NiagaraEditor/Private/NiagaraCompilationTasks.h b/Engine/Plugins/FX/Niagara/Source/NiagaraEditor/Private/NiagaraCompilationTasks.h
--- a/Engine/Plugins/FX/Niagara/Source/NiagaraEditor/Private/NiagaraCompilationTasks.h
+++ b/Engine/Plugins/FX/Niagara/Source/NiagaraEditor/Private/NiagaraCompilationTasks.h
@@ -276,7 +276,7 @@ private:
 
 	TMap<TObjectKey<UNiagaraParameterCollection>, FNiagaraCompilationNPCHandle> DigestedParameterCollections;
 
-	UE::Tasks::FTask MasterCompilationTask;
+	UE::Tasks::FTask ParentCompilationTask;
 
 	bool bAborting = false;
 
