- Fix duplicate LOCTEXT keys

#rnx

diff --git a/Engine/Plugins/FX/Niagara/Source/NiagaraEditor/Private/Customizations/NiagaraParameterBindingCustomization.cpp b/Engine/Plugins/FX/Niagara/Source/NiagaraEditor/Private/Customizations/NiagaraParameterBindingCustomization.cpp
--- a/Engine/Plugins/FX/Niagara/Source/NiagaraEditor/Private/Customizations/NiagaraParameterBindingCustomization.cpp
+++ b/Engine/Plugins/FX/Niagara/Source/NiagaraEditor/Private/Customizations/NiagaraParameterBindingCustomization.cpp
@@ -92,12 +92,12 @@ namespace NiagaraParameterBindingCustomizationPrivate
 		);
 		BuildTypesText(
 			BoundTypes,
-			LOCTEXT("BindingTooltip_ObjectFormat", "Interface Types: "),
+			LOCTEXT("BindingTooltip_InterfacesFormat", "Interface Types: "),
 			ParameterBinding->GetAllowedInterfaces()
 		);
 		BuildTypesText(
 			BoundTypes,
-			LOCTEXT("BindingTooltip_ObjectFormat", "Attribute Types: "),
+			LOCTEXT("BindingTooltip_AttributesFormat", "Attribute Types: "),
 			ParameterBinding->GetAllowedTypeDefinitions()
 		);
 
