- Use shader parameter flags to determine if we cook the SVT blend to dense permutation or not

#rnx
#tests Cook and launched a test project

diff --git a/Engine/Plugins/FX/Niagara/Source/NiagaraShader/Private/NiagaraSVTShaders.cpp b/Engine/Plugins/FX/Niagara/Source/NiagaraShader/Private/NiagaraSVTShaders.cpp
--- a/Engine/Plugins/FX/Niagara/Source/NiagaraShader/Private/NiagaraSVTShaders.cpp
+++ b/Engine/Plugins/FX/Niagara/Source/NiagaraShader/Private/NiagaraSVTShaders.cpp
@@ -16,11 +16,7 @@ void FNiagaraBlendSVTsToDenseBufferCS::ModifyCompilationEnvironment(const FGloba
 
 bool FNiagaraBlendSVTsToDenseBufferCS::ShouldCompilePermutation(const FGlobalShaderPermutationParameters& Parameters)
 {
-#if WITH_EDITOR
-	return true;
-#else
-	return false;
-#endif
+	return EnumHasAllFlags(Parameters.Flags, EShaderPermutationFlags::HasEditorOnlyData);
 }
 
-IMPLEMENT_GLOBAL_SHADER(FNiagaraBlendSVTsToDenseBufferCS, "/Plugin/FX/Niagara/Private/NiagaraBlendSVTsToDenseBuffer.usf", "PerformBlendCS", SF_Compute);
\ No newline at end of file
+IMPLEMENT_GLOBAL_SHADER(FNiagaraBlendSVTsToDenseBufferCS, "/Plugin/FX/Niagara/Private/NiagaraBlendSVTsToDenseBuffer.usf", "PerformBlendCS", SF_Compute);
