No longer hold onto TLAS scene for HWRT trace provider as it may go stale when we're dealing with dispatching of multiple frames.
#tests Evaluate HWRT in toggling back and forth between application

diff --git a/Engine/Plugins/FX/Niagara/Source/NiagaraShader/Private/NiagaraAsyncGpuTraceProviderHwrt.cpp b/Engine/Plugins/FX/Niagara/Source/NiagaraShader/Private/NiagaraAsyncGpuTraceProviderHwrt.cpp
--- a/Engine/Plugins/FX/Niagara/Source/NiagaraShader/Private/NiagaraAsyncGpuTraceProviderHwrt.cpp
+++ b/Engine/Plugins/FX/Niagara/Source/NiagaraShader/Private/NiagaraAsyncGpuTraceProviderHwrt.cpp
@@ -338,18 +338,9 @@ bool FNiagaraAsyncGpuTraceProviderHwrt::IsSupported()
 
 bool FNiagaraAsyncGpuTraceProviderHwrt::IsAvailable() const
 {
-	// Disable HWRT collision if there are multiple GPUs and this is a pending tick.  The ray tracing scene may be in an indeterminate state
-	// on a given GPU outside the context of the scene renderer, where pending ticks are handled.  To handle this properly, we would either
-	// need to provide a mechanism to force Niagara to run on the GPU where the ray tracing scene is valid, or refresh the ray tracing scene
-	// on GPU 0.  Either solution is a bit complicated, so we'll wait to implement one until it becomes important.  Virtual Production is the
-	// main client of multi-GPU, and they generally don't de-focus the application to create a scenario where pending ticks accumulate,
-	// so they are unlikely to encounter this issue.
-	//
-	// Note that pending ticks may run into inconsistencies with multiple views, even on a single GPU, because collision will be run against the
-	// last view, rather than the first view, which is normally where Niagara runs.  Longer term, it might be interesting to support running
-	// Niagara against all view's ray tracing scenes, and if you wanted to go that route, the solution would be to preserve the otherwise
-	// transient TLAS data for all views.  That could solve this issue as well.
-	if (!GNiagaraAsyncGpuTraceHwrtEnabled || (GNumExplicitGPUsForRendering > 1 && Dispatcher->IsOutsideSceneRenderer()))
+	// Never allow HWRT to run if the dispatcher is outside of the scene renderer. The TLAS is only expected to be valid for a single frame 
+	// and this risks using a stale TLAS from many frames ago which can contain pointers to resources that no longer exist
+	if (!GNiagaraAsyncGpuTraceHwrtEnabled ||Dispatcher->IsOutsideSceneRenderer())
 	{
 		return false;
 	}
@@ -372,7 +363,6 @@ void FNiagaraAsyncGpuTraceProviderHwrt::PostRenderOpaque(FRHICommandList& RHICmd
 
 	if (UE::FXRenderingUtils::RayTracing::HasRayTracingScene(Scene))
 	{
-		RayTracingSceneView = UE::FXRenderingUtils::RayTracing::GetRayTracingSceneView(RHICmdList, Scene);
 		ViewUniformBuffer = ReferenceView.ViewUniformBuffer;
 
 		if (!SupportsNiagaraAsyncGpuTraceHwrtInline() && GRHISupportsRayTracingShaders)
@@ -416,7 +406,6 @@ void FNiagaraAsyncGpuTraceProviderHwrt::PostRenderOpaque(FRHICommandList& RHICmd
 void FNiagaraAsyncGpuTraceProviderHwrt::IssueTraces(FRHICommandList& RHICmdList, const FDispatchRequest& Request, TUniformBufferRef<FSceneUniformParameters> SceneUniformBufferRHI, FCollisionGroupHashMap* CollisionGroupHash)
 {
 	check(IsAvailable());
-	check(RayTracingSceneView);
 
 	if (Request.MaxTraceCount == 0)
 	{
@@ -439,7 +428,8 @@ void FNiagaraAsyncGpuTraceProviderHwrt::IssueTraces(FRHICommandList& RHICmdList,
 		Params.HashToCollisionGroups = CollisionGroupHash->HashToCollisionGroups.SRV;
 	}
 
-	Params.TLAS = RayTracingSceneView;
+	FSceneInterface* Scene = Dispatcher->GetSceneInterface();
+	Params.TLAS = UE::FXRenderingUtils::RayTracing::GetRayTracingSceneView(RHICmdList, Scene);
 	Params.Rays = Request.TracesBuffer->SRV;
 	Params.RaysOffset = Request.TracesOffset;
 	Params.CollisionOutput = Request.ResultsBuffer->UAV;
@@ -506,7 +496,6 @@ void FNiagaraAsyncGpuTraceProviderHwrt::Reset()
 {
 	RayTracingPipelineState = nullptr;
 	RayTracingSBT = nullptr;
-	RayTracingSceneView = nullptr;
 	ViewUniformBuffer = TUniformBufferRef<FViewUniformShaderParameters>();
 }
 
diff --git a/Engine/Plugins/FX/Niagara/Source/NiagaraShader/Private/NiagaraAsyncGpuTraceProviderHwrt.h b/Engine/Plugins/FX/Niagara/Source/NiagaraShader/Private/NiagaraAsyncGpuTraceProviderHwrt.h
--- a/Engine/Plugins/FX/Niagara/Source/NiagaraShader/Private/NiagaraAsyncGpuTraceProviderHwrt.h
+++ b/Engine/Plugins/FX/Niagara/Source/NiagaraShader/Private/NiagaraAsyncGpuTraceProviderHwrt.h
@@ -26,7 +26,6 @@ public:
 private:
 	FRayTracingPipelineState* RayTracingPipelineState = nullptr;
 	FRHIShaderBindingTable* RayTracingSBT = nullptr; // TODO: Hold reference?
-	FShaderResourceViewRHIRef RayTracingSceneView = nullptr;
 	TUniformBufferRef<FViewUniformShaderParameters> ViewUniformBuffer;
 };
 
