Fix regression in new compilation mode which would result in a hang if you tried to save the niagara asset while there was an active compilation
#tests Make a change to a niagara system resulting in a compilation, while that is in progress try to save the system
#rnx

diff --git a/Engine/Plugins/FX/Niagara/Source/NiagaraEditor/Private/NiagaraCompilationTasks.cpp b/Engine/Plugins/FX/Niagara/Source/NiagaraEditor/Private/NiagaraCompilationTasks.cpp
--- a/Engine/Plugins/FX/Niagara/Source/NiagaraEditor/Private/NiagaraCompilationTasks.cpp
+++ b/Engine/Plugins/FX/Niagara/Source/NiagaraEditor/Private/NiagaraCompilationTasks.cpp
@@ -626,6 +626,7 @@ FNiagaraSystemCompilationTask::FNiagaraSystemCompilationTask(FNiagaraCompilation
 void FNiagaraSystemCompilationTask::Abort()
 {
 	bAborting = true;
+	CompileCompletionEvent.Trigger();
 }
 
 FString FNiagaraSystemCompilationTask::GetDescription() const
@@ -1221,6 +1222,11 @@ void FNiagaraSystemCompilationTask::FCompileTaskInfo::GenerateOptimizedVMByteCod
 
 void FNiagaraSystemCompilationTask::FCompileTaskInfo::ProcessCompilation(FNiagaraSystemCompilationTask* SystemCompileTask, const FCompileGroupInfo& GroupInfo)
 {
+	ON_SCOPE_EXIT
+	{
+		CompileResultsProcessedEvent->Trigger();
+	};
+
 	if (ScriptCompileType == EScriptCompileType::CompileForVm)
 	{
 		GenerateOptimizedVMByteCode(SystemCompileTask, GroupInfo);
@@ -1933,11 +1939,11 @@ bool FNiagaraSystemCompilationTask::AreResultsPending() const
 void FNiagaraSystemCompilationTask::WaitTillCompileCompletion()
 {
 	const FTimespan WaitTimeout = FTimespan::FromMilliseconds(50.0);
-	while (!ParentCompilationTask.IsCompleted())
+	while (!CompileCompletionEvent.IsCompleted())
 	{
 		// if the wait doesn't complete the task then we need to make sure to poke the
 		// compilation manager since there we may be waiting on it
-		if (!ParentCompilationTask.Wait(WaitTimeout))
+		if (!CompileCompletionEvent.Wait(WaitTimeout))
 		{
 			FNiagaraSystemCompilingManager::Get().AdvanceAsyncTasks();
 		}
@@ -2622,6 +2628,7 @@ void FNiagaraSystemCompilationTask::BuildAndApplyRapidIterationParameters()
 
 void FNiagaraSystemCompilationTask::IssuePostResultsProcessedTasks()
 {
+	CompileCompletionEvent.Trigger();
 	CompilationState = EState::ResultsProcessed;
 	PutRequestHelper = MakeUnique<FDispatchDataCachePutRequests>();
 	PutRequestHelper->Launch(this);
@@ -2632,12 +2639,12 @@ void FNiagaraSystemCompilationTask::IssuePostResultsProcessedTasks()
 	}, PutRequestHelper->CompletionEvent);
 }
 
-UE::Tasks::FTask FNiagaraSystemCompilationTask::IssueCompilationTasks()
+void FNiagaraSystemCompilationTask::IssueCompilationTasks()
 {
 	using namespace UE::Tasks;
 
 	TArray<FTask> IssuedCompilationTasks;
-	TArray<FTask> CompilationTasks;
+	TArray<FTaskEvent> CompilationTasks;
 
 	if (HasOutstandingCompileTasks())
 	{
@@ -2666,6 +2673,7 @@ UE::Tasks::FTask FNiagaraSystemCompilationTask::IssueCompilationTasks()
 					if (CompileTask.IsOutstanding(*this))
 					{
 						CompileTask.CompileResultsReadyEvent = MakeUnique<FTaskEvent>(UE_SOURCE_LOCATION);
+						CompileTask.CompileResultsProcessedEvent = MakeUnique<FTaskEvent>(UE_SOURCE_LOCATION);
 
 						IssuedCompilationTasks.Add(Launch(UE_SOURCE_LOCATION, [Self = AsShared(), &GroupInfo, &CompileTask]
 						{
@@ -2703,12 +2711,14 @@ UE::Tasks::FTask FNiagaraSystemCompilationTask::IssueCompilationTasks()
 							}
 						}, InstantiateGraphTask));
 
-						CompilationTasks.Add(Launch(UE_SOURCE_LOCATION, [Self = AsShared(), &GroupInfo, &CompileTask]
+						Launch(UE_SOURCE_LOCATION, [Self = AsShared(), &GroupInfo, &CompileTask]
 						{
 							COOK_STAT(auto Timer = NiagaraSystemCookStats::UsageStats.TimeSyncWork(); Timer.TrackCyclesOnly(););
 
 							CompileTask.ProcessCompilation(&Self.Get(), GroupInfo);
-						}, *CompileTask.CompileResultsReadyEvent));
+						}, *CompileTask.CompileResultsReadyEvent);
+
+						CompilationTasks.Add(*CompileTask.CompileResultsProcessedEvent);
 					}
 				}
 			}
@@ -2720,7 +2730,7 @@ UE::Tasks::FTask FNiagaraSystemCompilationTask::IssueCompilationTasks()
 		Self->CompilationState = EState::WaitingForProcessing;
 	}, IssuedCompilationTasks);
 
-	return Launch(UE_SOURCE_LOCATION, [Self = AsShared()]
+	Launch(UE_SOURCE_LOCATION, [Self = AsShared()]
 	{
 		Self->IssuePostResultsProcessedTasks();
 	}, CompilationTasks);
@@ -2793,7 +2803,7 @@ UE::Tasks::FTask FNiagaraSystemCompilationTask::BeginTasks()
 
 			FTask InnerTask = Launch(UE_SOURCE_LOCATION, [Self]
 			{
-				Self->ParentCompilationTask = Self->IssueCompilationTasks();
+				Self->IssueCompilationTasks();
 			}, CompilationTaskPrerequisites);
 		}
 		else
diff --git a/Engine/Plugins/FX/Niagara/Source/NiagaraEditor/Private/NiagaraCompilationTasks.h b/Engine/Plugins/FX/Niagara/Source/NiagaraEditor/Private/NiagaraCompilationTasks.h
--- a/Engine/Plugins/FX/Niagara/Source/NiagaraEditor/Private/NiagaraCompilationTasks.h
+++ b/Engine/Plugins/FX/Niagara/Source/NiagaraEditor/Private/NiagaraCompilationTasks.h
@@ -47,7 +47,7 @@ struct FNiagaraSystemCompilationTask : public TSharedFromThis<FNiagaraSystemComp
 
 	UE::Tasks::FTask BuildRapidIterationParametersAsync();
 	void BuildAndApplyRapidIterationParameters();
-	UE::Tasks::FTask IssueCompilationTasks();
+	void IssueCompilationTasks();
 	void IssuePostResultsProcessedTasks();
 	bool HasOutstandingCompileTasks() const;
 	FString GetDescription() const;
@@ -195,6 +195,7 @@ struct FNiagaraSystemCompilationTask : public TSharedFromThis<FNiagaraSystemComp
 		FDDCTaskInfo DDCTaskInfo;
 
 		TUniquePtr<UE::Tasks::FTaskEvent> CompileResultsReadyEvent;
+		TUniquePtr<UE::Tasks::FTaskEvent> CompileResultsProcessedEvent;
 		TMap<FName, UNiagaraDataInterface*> NamedDataInterfaces;
 		TArray<FNiagaraVariable> BakedRapidIterationParameters; // todo - remove and have the data pulled directly from the precompile data
 		double TaskStartTime = 0;
@@ -276,7 +277,7 @@ private:
 
 	TMap<TObjectKey<UNiagaraParameterCollection>, FNiagaraCompilationNPCHandle> DigestedParameterCollections;
 
-	UE::Tasks::FTask ParentCompilationTask;
+	UE::Tasks::FTaskEvent CompileCompletionEvent = UE::Tasks::FTaskEvent(UE_SOURCE_LOCATION);
 
 	bool bAborting = false;
 
