Add support for Displacement & Mask4 channels in the RVT DI

diff --git a/Engine/Plugins/FX/Niagara/Shaders/Private/NiagaraDataInterfaceVirtualTexture.ush b/Engine/Plugins/FX/Niagara/Shaders/Private/NiagaraDataInterfaceVirtualTexture.ush
--- a/Engine/Plugins/FX/Niagara/Shaders/Private/NiagaraDataInterfaceVirtualTexture.ush
+++ b/Engine/Plugins/FX/Niagara/Shaders/Private/NiagaraDataInterfaceVirtualTexture.ush
@@ -16,5 +16,6 @@
 #define ERuntimeVirtualTextureMaterialType_BaseColor_Normal_Specular_YCoCg			4
 #define ERuntimeVirtualTextureMaterialType_BaseColor_Normal_Specular_Mask_YCoCg		5
 #define ERuntimeVirtualTextureMaterialType_WorldHeight								6
-#define ERuntimeVirtualTextureMaterialType_Count									7
+#define ERuntimeVirtualTextureMaterialType_Displacement								7
+#define ERuntimeVirtualTextureMaterialType_Count									8
 
diff --git a/Engine/Plugins/FX/Niagara/Shaders/Private/NiagaraDataInterfaceVirtualTextureTemplate.ush b/Engine/Plugins/FX/Niagara/Shaders/Private/NiagaraDataInterfaceVirtualTextureTemplate.ush
--- a/Engine/Plugins/FX/Niagara/Shaders/Private/NiagaraDataInterfaceVirtualTextureTemplate.ush
+++ b/Engine/Plugins/FX/Niagara/Shaders/Private/NiagaraDataInterfaceVirtualTextureTemplate.ush
@@ -19,7 +19,7 @@ uint4				{ParameterName}_PageTableUniforms[2];
 float2				{ParameterName}_WorldHeightUnpack;
 SamplerState		{ParameterName}_SharedSampler;
 
-void GetAttributesValid_{ParameterName}(out bool BaseColorValid, out bool SpecularValid, out bool RoughnessValid, out bool NormalValid, out bool WorldHeightValid, out bool MaskValid)
+void GetAttributesValid_{ParameterName}(out bool BaseColorValid, out bool SpecularValid, out bool RoughnessValid, out bool NormalValid, out bool WorldHeightValid, out bool MaskValid, out bool DisplacementValid, out bool Mask4Valid)
 {
 	BaseColorValid = false;
 	SpecularValid = false;
@@ -27,6 +27,8 @@ void GetAttributesValid_{ParameterName}(out bool BaseColorValid, out bool Specul
 	NormalValid = false;
 	WorldHeightValid = false;
 	MaskValid = false;
+	DisplacementValid = false;
+	Mask4Valid = false;
 
 	switch ( {ParameterName}_MaterialType )
 	{
@@ -34,13 +36,16 @@ void GetAttributesValid_{ParameterName}(out bool BaseColorValid, out bool Specul
 			BaseColorValid = true;
 			break;
 
+		case ERuntimeVirtualTextureMaterialType_Mask4:
+			Mask4Valid = true;
+			break;
+
 		case ERuntimeVirtualTextureMaterialType_BaseColor_Normal_Roughness:
 			BaseColorValid = true;
 			RoughnessValid = true;
 			NormalValid = true;
 			break;
 
-		case ERuntimeVirtualTextureMaterialType_Mask4:
 		case ERuntimeVirtualTextureMaterialType_BaseColor_Normal_Specular:
 			BaseColorValid = true;
 			SpecularValid = true;
@@ -66,6 +71,10 @@ void GetAttributesValid_{ParameterName}(out bool BaseColorValid, out bool Specul
 		case ERuntimeVirtualTextureMaterialType_WorldHeight:
 			WorldHeightValid = true;
 			break;
+
+		case ERuntimeVirtualTextureMaterialType_Displacement:
+			DisplacementValid = true;
+			break;
 	}
 }
 
@@ -75,7 +84,7 @@ float4 SampleRVTLayer_{ParameterName}(float2 SampleUV, Texture2D InTexture, Text
 	return TextureVirtualSample(InTexture, {ParameterName}_SharedSampler, PageTable, 0, VTUniform_Unpack(InTextureUniforms));
 }
 
-void SampleRVT_{ParameterName}(in float3 WorldPosition, out bool bInsideVolume, out float3 BaseColor, out float Specular, out float Roughness, out float3 Normal, out float WorldHeight, out float Mask)
+void SampleRVT_{ParameterName}(in float3 WorldPosition, out bool bInsideVolume, out float3 BaseColor, out float Specular, out float Roughness, out float3 Normal, out float WorldHeight, out float Mask, out float Displacement, out float4 Mask4)
 {
 	bInsideVolume = false;
 	BaseColor = float3(0.0f, 0.0f, 0.0f);
@@ -84,6 +93,8 @@ void SampleRVT_{ParameterName}(in float3 WorldPosition, out bool bInsideVolume,
 	Normal = float3(0.0f, 0.0f, 1.0f);
 	WorldHeight = 0.0f;
 	Mask = 1.0f;
+	Displacement = 0.0f;
+	Mask4 = float4(0.0f, 0.0f, 0.0f, 0.0f);
 
 	// Get Sample Location
 	FDFVector3 LWCWorldPosition = DFFromTileOffset_Hack(MakeLWCVector3(GetEngineOwnerLWCTile(), WorldPosition));
@@ -109,6 +120,12 @@ void SampleRVT_{ParameterName}(in float3 WorldPosition, out bool bInsideVolume,
 			break;
 		}
 
+		case ERuntimeVirtualTextureMaterialType_Mask4:
+		{
+			Mask4 = LayerSample[0];
+			break;
+		}
+
 		case ERuntimeVirtualTextureMaterialType_BaseColor_Normal_Roughness:
 		{
 			BaseColor = LayerSample[0].rgb;
@@ -117,7 +134,6 @@ void SampleRVT_{ParameterName}(in float3 WorldPosition, out bool bInsideVolume,
 			break;
 		}
 
-		case ERuntimeVirtualTextureMaterialType_Mask4:
 		case ERuntimeVirtualTextureMaterialType_BaseColor_Normal_Specular:
 		{
 			BaseColor = LayerSample[0].rgb;
@@ -151,5 +167,11 @@ void SampleRVT_{ParameterName}(in float3 WorldPosition, out bool bInsideVolume,
 			WorldHeight = VirtualTextureUnpackHeight(LayerSample[0], {ParameterName}_WorldHeightUnpack);
 			break;
 		}
+
+		case ERuntimeVirtualTextureMaterialType_Displacement:
+		{
+			Displacement = LayerSample[0].x;
+			break;
+		}
 	}
 }
diff --git a/Engine/Plugins/FX/Niagara/Source/Niagara/Private/DataInterface/NiagaraDataInterfaceVirtualTexture.cpp b/Engine/Plugins/FX/Niagara/Source/Niagara/Private/DataInterface/NiagaraDataInterfaceVirtualTexture.cpp
--- a/Engine/Plugins/FX/Niagara/Source/Niagara/Private/DataInterface/NiagaraDataInterfaceVirtualTexture.cpp
+++ b/Engine/Plugins/FX/Niagara/Source/Niagara/Private/DataInterface/NiagaraDataInterfaceVirtualTexture.cpp
@@ -51,6 +51,7 @@ namespace NDIVirtualTextureLocal
 		{
 			InitialVersion = 0,
 			AddInsideVolumeToSampleRVT = 1,
+			AddNewChannels = 2,
 
 			VersionPlusOne,
 			LatestVersion = VersionPlusOne - 1
@@ -194,6 +195,8 @@ void UNiagaraDataInterfaceVirtualTexture::GetFunctionsInternal(TArray<FNiagaraFu
 		Sig.Outputs.Emplace(FNiagaraTypeDefinition::GetBoolDef(), TEXT("NormalValid"));
 		Sig.Outputs.Emplace(FNiagaraTypeDefinition::GetBoolDef(), TEXT("WorldHeightValid"));
 		Sig.Outputs.Emplace(FNiagaraTypeDefinition::GetBoolDef(), TEXT("MaskValid"));
+		Sig.Outputs.Emplace(FNiagaraTypeDefinition::GetBoolDef(), TEXT("DisplacementValid"));
+		Sig.Outputs.Emplace(FNiagaraTypeDefinition::GetBoolDef(), TEXT("Mask4Valid"));
 	}
 
 	{
@@ -213,6 +216,8 @@ void UNiagaraDataInterfaceVirtualTexture::GetFunctionsInternal(TArray<FNiagaraFu
 		Sig.Outputs.Emplace(FNiagaraTypeDefinition::GetVec3Def(), TEXT("Normal"));
 		Sig.Outputs.Emplace(FNiagaraTypeDefinition::GetFloatDef(), TEXT("WorldHeight"));
 		Sig.Outputs.Emplace(FNiagaraTypeDefinition::GetFloatDef(), TEXT("Mask"));
+		Sig.Outputs.Emplace(FNiagaraTypeDefinition::GetFloatDef(), TEXT("Displacement"));
+		Sig.Outputs.Emplace(FNiagaraTypeDefinition::GetVec4Def(), TEXT("Mask4"));
 	}
 }
 #endif
@@ -361,6 +366,21 @@ bool UNiagaraDataInterfaceVirtualTexture::UpgradeFunctionCall(FNiagaraFunctionSi
 			bModified = true;
 		}
 	}
+	if (Sig.FunctionVersion < FNDIFunctionVersion::AddNewChannels)
+	{
+		if (Sig.Name == SampleRVTName)
+		{
+			Sig.Outputs.Emplace(FNiagaraTypeDefinition::GetFloatDef(), TEXT("Displacement"));
+			Sig.Outputs.Emplace(FNiagaraTypeDefinition::GetVec4Def(), TEXT("Mask4"));
+			bModified = true;
+		}
+		else if (Sig.Name == GetAttributesValidName)
+		{
+			Sig.Outputs.Emplace(FNiagaraTypeDefinition::GetBoolDef(), TEXT("DisplacementValid"));
+			Sig.Outputs.Emplace(FNiagaraTypeDefinition::GetBoolDef(), TEXT("Mask4Valid"));
+			bModified = true;
+		}
+	}
 
 	Sig.SetFunctionVersion(FNDIFunctionVersion::LatestVersion);
 	return bModified;
