Respect the need to immediately grab the landscape data in the generated data so that the first tick of the system (like spawning) will incorporate it.  Otherwise, the data is only available following the next tick group (when the DIGeneratedData is ticked)

diff --git a/Engine/Plugins/FX/Niagara/Source/Niagara/Private/NiagaraDataInterfaceLandscape.cpp b/Engine/Plugins/FX/Niagara/Source/Niagara/Private/NiagaraDataInterfaceLandscape.cpp
--- a/Engine/Plugins/FX/Niagara/Source/Niagara/Private/NiagaraDataInterfaceLandscape.cpp
+++ b/Engine/Plugins/FX/Niagara/Source/Niagara/Private/NiagaraDataInterfaceLandscape.cpp
@@ -802,6 +802,11 @@ void FNDI_Landscape_SharedResource::RegisterUser(const FNDI_SharedResourceUsage&
 		if (ShaderPhysicsDataUserCount++ == 0)
 		{
 			NextState = EResourceState::Initialized;
+			if (bNeedsDataImmediately)
+			{
+				bool bLandscapeRemoved = false;
+				UpdateState(bLandscapeRemoved);
+			}
 		}
 	}
 }
