Fix leaking data buffers for emitters that don't write to NDC data each frame.

diff --git a/Engine/Plugins/FX/Niagara/Source/Niagara/Private/DataInterface/NiagaraDataInterfaceDataChannelWrite.cpp b/Engine/Plugins/FX/Niagara/Source/Niagara/Private/DataInterface/NiagaraDataInterfaceDataChannelWrite.cpp
--- a/Engine/Plugins/FX/Niagara/Source/Niagara/Private/DataInterface/NiagaraDataInterfaceDataChannelWrite.cpp
+++ b/Engine/Plugins/FX/Niagara/Source/Niagara/Private/DataInterface/NiagaraDataInterfaceDataChannelWrite.cpp
@@ -471,6 +471,10 @@ struct FNDIDataChannelWriteInstanceData
 #endif
 				DataChannelData->Publish(PublishRequest);
 			}
+			else
+			{
+				DestinationData->Unlock();
+			}
 			
 			AtomicNumInstances = 0;
 		}
