- Fix crash if a select node is not hooked up correctly

#rnx
#tests Add a scratchpad module, on the map set add an output attribute, hook a if/select into the output then hit apply

diff --git a/Engine/Plugins/FX/Niagara/Source/NiagaraEditor/Private/NiagaraGraphDigest.cpp b/Engine/Plugins/FX/Niagara/Source/NiagaraEditor/Private/NiagaraGraphDigest.cpp
--- a/Engine/Plugins/FX/Niagara/Source/NiagaraEditor/Private/NiagaraGraphDigest.cpp
+++ b/Engine/Plugins/FX/Niagara/Source/NiagaraEditor/Private/NiagaraGraphDigest.cpp
@@ -4762,7 +4762,10 @@ void FNiagaraCompilationNodeSelect::Compile(FTranslator* Translator, TArray<int3
 	Outputs.Init(INDEX_NONE, OutputVarCount);
 	for (int32 ConnectedVarIndex = 0; ConnectedVarIndex < ConnectedOutputVars.Num(); ++ConnectedVarIndex)
 	{
-		Outputs[ConnectedOutputVarIndices[ConnectedVarIndex]] = SelectOutputs[ConnectedVarIndex];
+		if (SelectOutputs.IsValidIndex(ConnectedVarIndex))
+		{
+			Outputs[ConnectedOutputVarIndices[ConnectedVarIndex]] = SelectOutputs[ConnectedVarIndex];
+		}
 	}
 }
 
