- Add curves to sprite / mesh rotation rates
- Curve distributions can run in sample & accumulate mode when generating the LUT
- Accumulate mode will generate a LUT which is the sum of all the values to that point

#virtualized

diff --git a/Engine/Plugins/FX/Niagara/Shaders/Private/Stateless/Modules/NiagaraStatelessModule_MeshRotationRate.ush b/Engine/Plugins/FX/Niagara/Shaders/Private/Stateless/Modules/NiagaraStatelessModule_MeshRotationRate.ush
--- a/Engine/Plugins/FX/Niagara/Shaders/Private/Stateless/Modules/NiagaraStatelessModule_MeshRotationRate.ush
+++ b/Engine/Plugins/FX/Niagara/Shaders/Private/Stateless/Modules/NiagaraStatelessModule_MeshRotationRate.ush
@@ -2,9 +2,10 @@
 
 #include "../NiagaraStatelessCommon.ush"
 
-int	MeshRotationRate_ModuleEnabled;
-float3	MeshRotationRate_Scale;
-float3	MeshRotationRate_Bias;
+int										MeshRotationRate_ModuleEnabled;
+float3									MeshRotationRate_Scale;
+float3									MeshRotationRate_Bias;
+FNiagaraStatelessBuiltDistributionType	MeshRotationRate_RateScaleParameters;
 
 void MeshRotationRate_Simulate(inout FStatelessParticle Particle)
 {
@@ -13,8 +14,10 @@ void MeshRotationRate_Simulate(inout FStatelessParticle Particle)
 		return;
 	}
 	const float3 RotationRate = RandomScaleBiasFloat(0, MeshRotationRate_Scale, MeshRotationRate_Bias);
+	const float3 MultiplyRate = SampleCurveFloat3(MeshRotationRate_RateScaleParameters, Particle.NormalizedAge);
+	const float3 PreviousMultiplyRate = SampleCurveFloat3(MeshRotationRate_RateScaleParameters, Particle.PreviousNormalizedAge);
 
-	Particle.MeshOrientation			= MultiplyQuat(Particle.MeshOrientation,			RotatorToQuat(Particle.Age * RotationRate));
-	Particle.PreviousMeshOrientation	= MultiplyQuat(Particle.PreviousMeshOrientation,	RotatorToQuat(Particle.PreviousAge * RotationRate));
+	Particle.MeshOrientation			= MultiplyQuat(Particle.MeshOrientation,			RotatorToQuat(Particle.Lifetime * RotationRate * MultiplyRate));
+	Particle.PreviousMeshOrientation	= MultiplyQuat(Particle.PreviousMeshOrientation,	RotatorToQuat(Particle.Lifetime * RotationRate * PreviousMultiplyRate));
 }
 
diff --git a/Engine/Plugins/FX/Niagara/Shaders/Private/Stateless/Modules/NiagaraStatelessModule_SpriteRotationRate.ush b/Engine/Plugins/FX/Niagara/Shaders/Private/Stateless/Modules/NiagaraStatelessModule_SpriteRotationRate.ush
--- a/Engine/Plugins/FX/Niagara/Shaders/Private/Stateless/Modules/NiagaraStatelessModule_SpriteRotationRate.ush
+++ b/Engine/Plugins/FX/Niagara/Shaders/Private/Stateless/Modules/NiagaraStatelessModule_SpriteRotationRate.ush
@@ -2,13 +2,22 @@
 
 #include "../NiagaraStatelessCommon.ush"
 
-float		SpriteRotationRate_Scale;
-float		SpriteRotationRate_Bias;
+int										SpriteRotationRate_Enabled;
+float									SpriteRotationRate_Scale;
+float									SpriteRotationRate_Bias;
+FNiagaraStatelessBuiltDistributionType	SpriteRotationRate_RateScaleParameters;
 
 void SpriteRotationRate_Simulate(inout FStatelessParticle Particle)
 {
+	if (SpriteRotationRate_Enabled == 0)
+	{
+		 return;
+	}
+		
 	const float RotationRate = RandomScaleBiasFloat(0, SpriteRotationRate_Scale, SpriteRotationRate_Bias);
-
-	Particle.SpriteRotation += Particle.Age * RotationRate;
-	Particle.PreviousSpriteRotation += Particle.PreviousAge * RotationRate;
+	const float MultiplyRate = SampleCurveFloat(SpriteRotationRate_RateScaleParameters, Particle.NormalizedAge);
+	const float PreviousMultiplyRate = SampleCurveFloat(SpriteRotationRate_RateScaleParameters, Particle.PreviousNormalizedAge);
+		
+	Particle.SpriteRotation += Particle.Lifetime * RotationRate * MultiplyRate;
+	Particle.PreviousSpriteRotation += Particle.Lifetime * RotationRate * PreviousMultiplyRate;
 }
diff --git a/Engine/Plugins/FX/Niagara/Shaders/Private/Stateless/NiagaraStatelessCommon.ush b/Engine/Plugins/FX/Niagara/Shaders/Private/Stateless/NiagaraStatelessCommon.ush
--- a/Engine/Plugins/FX/Niagara/Shaders/Private/Stateless/NiagaraStatelessCommon.ush
+++ b/Engine/Plugins/FX/Niagara/Shaders/Private/Stateless/NiagaraStatelessCommon.ush
@@ -171,7 +171,25 @@ float  StatelessLerp(float  lhs, float  rhs, float  u, bool bUniform) { return l
 float2 StatelessLerp(float2 lhs, float2 rhs, float2 u, bool bUniform) { return bUniform ? lerp(lhs.x, rhs.x, u.x).xx   : lerp(lhs, rhs, u); }
 float3 StatelessLerp(float3 lhs, float3 rhs, float3 u, bool bUniform) { return bUniform ? lerp(lhs.x, rhs.x, u.x).xxx  : lerp(lhs, rhs, u); }
 float4 StatelessLerp(float4 lhs, float4 rhs, float4 u, bool bUniform) { return bUniform ? lerp(lhs.x, rhs.x, u.x).xxxx : lerp(lhs, rhs, u); }
-	
+
+float SampleCurveFloat (FNiagaraStatelessBuiltDistributionType Distribution, float Time)
+{
+	const uint DataOffset = FNiagaraStatelessBuiltDistribution::GetDataOffset(Distribution);
+	const float Offset = FNiagaraStatelessBuiltDistribution::ConvertTimeToLookup(Distribution, Time);
+	const float Value0 = GetStaticFloat(DataOffset, floor(Offset));
+	const float Value1 = GetStaticFloat(DataOffset, ceil(Offset));
+	return lerp(Value0, Value1, frac(Offset));
+}
+
+float3 SampleCurveFloat3(FNiagaraStatelessBuiltDistributionType Distribution, float Time)
+{
+	const uint DataOffset = FNiagaraStatelessBuiltDistribution::GetDataOffset(Distribution);
+	const float Offset = FNiagaraStatelessBuiltDistribution::ConvertTimeToLookup(Distribution, Time);
+	const float3 Value0 = GetStaticFloat3(DataOffset, floor(Offset));
+	const float3 Value1 = GetStaticFloat3(DataOffset, ceil(Offset));
+	return lerp(Value0, Value1, frac(Offset));
+}
+
 uint GetComponentOffset(uint OutputIndex, int ComponentOffset)
 {
 	return ComponentOffset * Common_OutputBufferStride + OutputIndex;
diff --git a/Engine/Plugins/FX/Niagara/Source/Niagara/Internal/Stateless/Modules/NiagaraStatelessModule_MeshRotationRate.h b/Engine/Plugins/FX/Niagara/Source/Niagara/Internal/Stateless/Modules/NiagaraStatelessModule_MeshRotationRate.h
--- a/Engine/Plugins/FX/Niagara/Source/Niagara/Internal/Stateless/Modules/NiagaraStatelessModule_MeshRotationRate.h
+++ b/Engine/Plugins/FX/Niagara/Source/Niagara/Internal/Stateless/Modules/NiagaraStatelessModule_MeshRotationRate.h
@@ -19,6 +19,7 @@ class UNiagaraStatelessModule_MeshRotationRate : public UNiagaraStatelessModule
 	{
 		bool							ModuleEnabled = false;
 		FNiagaraStatelessRangeVector3	RotationRange;
+		FUintVector3					RateScaleParameters = FUintVector3::ZeroValue;
 		int32							MeshOrientationVariableOffset = INDEX_NONE;
 		int32							PreviousMeshOrientationVariableOffset = INDEX_NONE;
 	};
@@ -26,9 +27,15 @@ class UNiagaraStatelessModule_MeshRotationRate : public UNiagaraStatelessModule
 public:
 	using FParameters = NiagaraStateless::FMeshRotationRateModule_ShaderParameters;
 
+	UPROPERTY(EditAnywhere, Category = "Parameters", meta = (InlineEditConditionToggle))
+	bool bUseRateScale = false;
+
 	UPROPERTY(EditAnywhere, Category = "Parameters", meta = (DisplayName = "Rotation Rate", Units="deg"))
 	FNiagaraDistributionRangeVector3 RotationRateDistribution = FNiagaraDistributionRangeVector3(FVector3f::ZeroVector);
 
+	UPROPERTY(EditAnywhere, Category = "Parameters", meta = (DisplayName = "Rate Scale", EditCondition = "bUseRateScale"))
+	FNiagaraDistributionCurveVector3 RateScaleDistribution = FNiagaraDistributionCurveVector3(ENiagaraDistributionCurveLUTMode::Accumulate);
+
 	virtual void BuildEmitterData(const FNiagaraStatelessEmitterDataBuildContext& BuildContext) const override
 	{
 		FModuleBuiltData* BuiltData = BuildContext.AllocateBuiltData<FModuleBuiltData>();
@@ -47,6 +54,14 @@ public:
 			BuiltData->RotationRange = BuildContext.ConvertDistributionToRange(RotationRateDistribution, FVector3f::ZeroVector);
 			BuiltData->RotationRange.Min *= 1.0f / 360.0f;
 			BuiltData->RotationRange.Max *= 1.0f / 360.0f;
+			if (bUseRateScale)
+			{
+				BuiltData->RateScaleParameters = BuildContext.AddDistributionAsCurve(RateScaleDistribution, 1.0f);
+			}
+			else
+			{
+				BuiltData->RateScaleParameters = BuildContext.AddDistributionAsCurve(FNiagaraDistributionCurveVector3(ENiagaraDistributionCurveLUTMode::Accumulate), 1.0f);
+			}
 
 			BuildContext.AddParticleSimulationExecSimulate(&UNiagaraStatelessModule_MeshRotationRate::ParticleSimulate);
 		}
@@ -63,6 +78,7 @@ public:
 		const FModuleBuiltData* ModuleBuiltData = SetShaderParameterContext.ReadBuiltData<FModuleBuiltData>();
 		Parameters->MeshRotationRate_ModuleEnabled = ModuleBuiltData->ModuleEnabled;
 		SetShaderParameterContext.ConvertRangeToScaleBias(ModuleBuiltData->RotationRange, Parameters->MeshRotationRate_Scale, Parameters->MeshRotationRate_Bias);
+		Parameters->MeshRotationRate_RateScaleParameters = ModuleBuiltData->RateScaleParameters;
 	}
 
 	static void ParticleSimulate(const NiagaraStateless::FParticleSimulationContext& ParticleSimulationContext)
@@ -72,20 +88,21 @@ public:
 		const FModuleBuiltData* ModuleBuiltData = ParticleSimulationContext.ReadBuiltData<FModuleBuiltData>();
 		const FParameters* ShaderParameters = ParticleSimulationContext.ReadParameterNestedStruct<FParameters>();
 
-		const float* AgeData = ParticleSimulationContext.GetParticleAge();
-		const float* PreviousAgeData = ParticleSimulationContext.GetParticlePreviousAge();
+		const float* LifetimeData = ParticleSimulationContext.GetParticleLifetime();
+		const float* AgeData = ParticleSimulationContext.GetParticleNormalizedAge();
+		const float* PreviousAgeData = ParticleSimulationContext.GetParticlePreviousNormalizedAge();
 
 		for (uint32 i = 0; i < ParticleSimulationContext.GetNumInstances(); ++i)
 		{
 			const FVector3f RotationRate = ParticleSimulationContext.RandomScaleBiasFloat(i, 0, ShaderParameters->MeshRotationRate_Scale, ShaderParameters->MeshRotationRate_Bias);
-			const float Age = AgeData[i];
-			const float PreviousAge = PreviousAgeData[i];
+			const FVector3f MultiplyRate = ParticleSimulationContext.SampleCurve<FVector3f>(ShaderParameters->MeshRotationRate_RateScaleParameters, AgeData[i]);
+			const FVector3f PreviousMultiplyRate = ParticleSimulationContext.SampleCurve<FVector3f>(ShaderParameters->MeshRotationRate_RateScaleParameters, PreviousAgeData[i]);
 
 			FQuat4f MeshOrientation = ParticleSimulationContext.ReadParticleVariable(ModuleBuiltData->MeshOrientationVariableOffset, i, FQuat4f::Identity);
 			FQuat4f PreviousMeshOrientation = ParticleSimulationContext.ReadParticleVariable(ModuleBuiltData->PreviousMeshOrientationVariableOffset, i, FQuat4f::Identity);
 
-			MeshOrientation *= ParticleSimulationContext.RotatorToQuat(RotationRate * Age);
-			PreviousMeshOrientation *= ParticleSimulationContext.RotatorToQuat(RotationRate * PreviousAge);
+			MeshOrientation *= ParticleSimulationContext.RotatorToQuat(LifetimeData[i] * RotationRate * MultiplyRate);
+			PreviousMeshOrientation *= ParticleSimulationContext.RotatorToQuat(LifetimeData[i] * RotationRate * PreviousMultiplyRate);
 
 			ParticleSimulationContext.WriteParticleVariable(ModuleBuiltData->MeshOrientationVariableOffset, i, MeshOrientation);
 			ParticleSimulationContext.WriteParticleVariable(ModuleBuiltData->PreviousMeshOrientationVariableOffset, i, PreviousMeshOrientation);
diff --git a/Engine/Plugins/FX/Niagara/Source/Niagara/Internal/Stateless/Modules/NiagaraStatelessModule_SpriteRotationRate.h b/Engine/Plugins/FX/Niagara/Source/Niagara/Internal/Stateless/Modules/NiagaraStatelessModule_SpriteRotationRate.h
--- a/Engine/Plugins/FX/Niagara/Source/Niagara/Internal/Stateless/Modules/NiagaraStatelessModule_SpriteRotationRate.h
+++ b/Engine/Plugins/FX/Niagara/Source/Niagara/Internal/Stateless/Modules/NiagaraStatelessModule_SpriteRotationRate.h
@@ -17,7 +17,9 @@ class UNiagaraStatelessModule_SpriteRotationRate : public UNiagaraStatelessModul
 
 	struct FModuleBuiltData
 	{
+		int32							ModuleEnabled = false;
 		FNiagaraStatelessRangeFloat		RotationRange;
+		FUintVector3					RateScaleParameters = FUintVector3::ZeroValue;
 		int32							SpriteRotationVariableOffset = INDEX_NONE;
 		int32							PreviousSpriteRotationVariableOffset = INDEX_NONE;
 	};
@@ -25,9 +27,15 @@ class UNiagaraStatelessModule_SpriteRotationRate : public UNiagaraStatelessModul
 public:
 	using FParameters = NiagaraStateless::FSpriteRotationRateModule_ShaderParameters;
 
+	UPROPERTY(EditAnywhere, Category = "Parameters", meta = (InlineEditConditionToggle))
+	bool bUseRateScale = false;
+
 	UPROPERTY(EditAnywhere, Category = "Parameters", meta = (DisplayName = "Rotation Rate", Units = "deg"))
 	FNiagaraDistributionRangeFloat RotationRateDistribution = FNiagaraDistributionRangeFloat(FNiagaraStatelessGlobals::GetDefaultSpriteRotationValue());
 
+	UPROPERTY(EditAnywhere, Category = "Parameters", meta = (DisplayName = "Rate Scale", EditCondition = "bUseRateScale"))
+	FNiagaraDistributionCurveFloat RateScaleDistribution = FNiagaraDistributionCurveFloat(ENiagaraDistributionCurveLUTMode::Accumulate);
+
 	virtual void BuildEmitterData(const FNiagaraStatelessEmitterDataBuildContext& BuildContext) const override
 	{
 		FModuleBuiltData* BuiltData = BuildContext.AllocateBuiltData<FModuleBuiltData>();
@@ -44,8 +52,17 @@ public:
 		{
 			return;
 		}
-		
+
+		BuiltData->ModuleEnabled = true;
 		BuiltData->RotationRange = BuildContext.ConvertDistributionToRange(RotationRateDistribution, 0.0f);
+		if (bUseRateScale)
+		{
+			BuiltData->RateScaleParameters = BuildContext.AddDistributionAsCurve(RateScaleDistribution, 1.0f);
+		}
+		else
+		{
+			BuiltData->RateScaleParameters = BuildContext.AddDistributionAsCurve(FNiagaraDistributionCurveFloat(ENiagaraDistributionCurveLUTMode::Accumulate), 1.0f);
+		}
 
 		BuildContext.AddParticleSimulationExecSimulate(&UNiagaraStatelessModule_SpriteRotationRate::ParticleSimulate);
 	}
@@ -60,7 +77,9 @@ public:
 		FParameters* Parameters = SetShaderParameterContext.GetParameterNestedStruct<FParameters>();
 		const FModuleBuiltData* ModuleBuiltData = SetShaderParameterContext.ReadBuiltData<FModuleBuiltData>();
 
+		Parameters->SpriteRotationRate_Enabled = ModuleBuiltData->ModuleEnabled;
 		SetShaderParameterContext.ConvertRangeToScaleBias(ModuleBuiltData->RotationRange, Parameters->SpriteRotationRate_Scale, Parameters->SpriteRotationRate_Bias);
+		Parameters->SpriteRotationRate_RateScaleParameters = ModuleBuiltData->RateScaleParameters;
 	}
 
 	static void ParticleSimulate(const NiagaraStateless::FParticleSimulationContext& ParticleSimulationContext)
@@ -70,18 +89,21 @@ public:
 		const FModuleBuiltData* ModuleBuiltData = ParticleSimulationContext.ReadBuiltData<FModuleBuiltData>();
 		const FParameters* ShaderParameters = ParticleSimulationContext.ReadParameterNestedStruct<FParameters>();
 
-		const float* AgeData = ParticleSimulationContext.GetParticleAge();
-		const float* PreviousAgeData = ParticleSimulationContext.GetParticlePreviousAge();
+		const float* LifetimeData = ParticleSimulationContext.GetParticleLifetime();
+		const float* AgeData = ParticleSimulationContext.GetParticleNormalizedAge();
+		const float* PreviousAgeData = ParticleSimulationContext.GetParticlePreviousNormalizedAge();
 
 		for (uint32 i = 0; i < ParticleSimulationContext.GetNumInstances(); ++i)
 		{
 			const float RotationRate = ParticleSimulationContext.RandomScaleBiasFloat(i, 0, ShaderParameters->SpriteRotationRate_Scale, ShaderParameters->SpriteRotationRate_Bias);
+			const float MultiplyRate = ParticleSimulationContext.SampleCurve<float>(ShaderParameters->SpriteRotationRate_RateScaleParameters, AgeData[i]);
+			const float PreviousMultiplyRate = ParticleSimulationContext.SampleCurve<float>(ShaderParameters->SpriteRotationRate_RateScaleParameters, PreviousAgeData[i]);
 
 			float SpriteRotation = ParticleSimulationContext.ReadParticleVariable(ModuleBuiltData->SpriteRotationVariableOffset, i, 0.0f);
 			float PreviousSpriteRotation = ParticleSimulationContext.ReadParticleVariable(ModuleBuiltData->PreviousSpriteRotationVariableOffset, i, 0.0f);
 
-			SpriteRotation += AgeData[i] * RotationRate;
-			PreviousSpriteRotation += PreviousAgeData[i] * RotationRate;
+			SpriteRotation += LifetimeData[i] * RotationRate * MultiplyRate;
+			PreviousSpriteRotation += LifetimeData[i] * RotationRate * PreviousMultiplyRate;
 
 			ParticleSimulationContext.WriteParticleVariable(ModuleBuiltData->SpriteRotationVariableOffset, i, SpriteRotation);
 			ParticleSimulationContext.WriteParticleVariable(ModuleBuiltData->PreviousSpriteRotationVariableOffset, i, PreviousSpriteRotation);
diff --git a/Engine/Plugins/FX/Niagara/Source/Niagara/Internal/Stateless/NiagaraStatelessEmitterDataBuildContext.h b/Engine/Plugins/FX/Niagara/Source/Niagara/Internal/Stateless/NiagaraStatelessEmitterDataBuildContext.h
--- a/Engine/Plugins/FX/Niagara/Source/Niagara/Internal/Stateless/NiagaraStatelessEmitterDataBuildContext.h
+++ b/Engine/Plugins/FX/Niagara/Source/Niagara/Internal/Stateless/NiagaraStatelessEmitterDataBuildContext.h
@@ -114,6 +114,23 @@ public:
 		return BuiltDistribution;
 	}
 
+	// Adds a distribution and forces it to generate as a curve for lookup
+	template<typename TDistribution, typename TType>
+	FNiagaraStatelessBuiltDistributionType AddDistributionAsCurve(const TDistribution& Distribution, TType DefaultValue) const
+	{
+		FNiagaraStatelessBuiltDistributionType BuiltDistribution = FNiagaraStatelessBuiltDistribution::GetDefault();
+		if (ensure(Distribution.Mode == ENiagaraDistributionMode::UniformCurve || Distribution.Mode == ENiagaraDistributionMode::NonUniformCurve) && Distribution.Values.Num() > 1)
+		{
+			FNiagaraStatelessBuiltDistribution::SetLookupParameters(BuiltDistribution, AddStaticData(Distribution.Values), Distribution.Values.Num(), Distribution.ValuesTimeRange);
+		}
+		else
+		{
+			const TType DefaultValues[] = {DefaultValue , DefaultValue};
+			FNiagaraStatelessBuiltDistribution::SetLookupParameters(BuiltDistribution, AddStaticData(DefaultValues));
+		}
+		return BuiltDistribution;
+	}
+
 	template<typename TRange, typename TDistribution, typename TDefaultValue>
 	TRange ConvertDistributionToRangeHelper(const TDistribution& Distribution, const TDefaultValue& DefaultValue) const
 	{
diff --git a/Engine/Plugins/FX/Niagara/Source/Niagara/Internal/Stateless/NiagaraStatelessParticleSimContext.h b/Engine/Plugins/FX/Niagara/Source/Niagara/Internal/Stateless/NiagaraStatelessParticleSimContext.h
--- a/Engine/Plugins/FX/Niagara/Source/Niagara/Internal/Stateless/NiagaraStatelessParticleSimContext.h
+++ b/Engine/Plugins/FX/Niagara/Source/Niagara/Internal/Stateless/NiagaraStatelessParticleSimContext.h
@@ -124,6 +124,16 @@ public:
 	FVector3f		Lerp(const FVector3f& Lhs, const FVector3f& Rhs, const FVector3f& U, bool bUniform) const { const FVector3f V = FMath::Lerp(Lhs, Rhs, U); return bUniform ? FVector3f(V.X, V.X, V.X) : V; }
 	FVector4f		Lerp(const FVector4f& Lhs, const FVector4f& Rhs, const FVector4f& U, bool bUniform) const { const FVector4f V = FMath::Lerp(Lhs, Rhs, U); return bUniform ? FVector4f(V.X, V.X, V.X, V.X) : V; }
 
+	template<typename TType>
+	TType SampleCurve(const FNiagaraStatelessBuiltDistributionType& Distribution, float Time) const
+	{
+		const uint32 DataOffset = FNiagaraStatelessBuiltDistribution::GetDataOffset(Distribution);
+		const float Offset = FNiagaraStatelessBuiltDistribution::ConvertTimeToLookup(Distribution, Time);
+		const TType Value0 = GetStaticFloat<TType>(DataOffset, FMath::FloorToInt(Offset));
+		const TType Value1 = GetStaticFloat<TType>(DataOffset, FMath::CeilToInt(Offset));
+		return FMath::Lerp(Value0, Value1, FMath::Fractional(Offset));
+	}
+
 	template<typename TType>
 	TType					TRandomFloat(uint32 iInstance, uint32 RandomSeedOffset) const;
 	template<> float		TRandomFloat<float>(uint32 iInstance, uint32 RandomSeedOffset) const { return RandomFloat(iInstance, RandomSeedOffset); }
diff --git a/Engine/Plugins/FX/Niagara/Source/Niagara/Private/NiagaraRendererProperties.cpp b/Engine/Plugins/FX/Niagara/Source/Niagara/Private/NiagaraRendererProperties.cpp
--- a/Engine/Plugins/FX/Niagara/Source/Niagara/Private/NiagaraRendererProperties.cpp
+++ b/Engine/Plugins/FX/Niagara/Source/Niagara/Private/NiagaraRendererProperties.cpp
@@ -18,6 +18,7 @@
 #include "Materials/Material.h"
 #include "Materials/MaterialInterface.h"
 #include "Materials/MaterialInstanceConstant.h"
+#include "MaterialDomain.h"
 #include "SceneInterface.h"
 #include "SceneUtils.h"
 #include "Styling/SlateIconFinder.h"
diff --git a/Engine/Plugins/FX/Niagara/Source/Niagara/Private/Stateless/NiagaraStatelessDistribution.cpp b/Engine/Plugins/FX/Niagara/Source/Niagara/Private/Stateless/NiagaraStatelessDistribution.cpp
--- a/Engine/Plugins/FX/Niagara/Source/Niagara/Private/Stateless/NiagaraStatelessDistribution.cpp
+++ b/Engine/Plugins/FX/Niagara/Source/Niagara/Private/Stateless/NiagaraStatelessDistribution.cpp
@@ -22,7 +22,7 @@ namespace NiagaraStatelessDistributionPrivate
 		ECVF_Default
 	);
 
-	TArray<float> CurvesToLUT(TArrayView<const FRichCurve> Curves, int32 NumSamples, const FVector2f& TimeRange)
+	TArray<float> CurvesToLUT(TArrayView<const FRichCurve> Curves, int32 NumSamples, const FVector2f& TimeRange, ENiagaraDistributionCurveLUTMode LUTMode)
 	{
 		const int32 NumChannels = Curves.Num();
 
@@ -31,18 +31,37 @@ namespace NiagaraStatelessDistributionPrivate
 
 		TArray<float> LUT;
 		LUT.SetNumUninitialized(NumChannels * NumSamples);
-		for (int32 iSample=0; iSample < NumSamples; ++iSample)
+		if (LUTMode == ENiagaraDistributionCurveLUTMode::Sample)
 		{
-			const float Time = (float(iSample) * SampleToDuration) + TimeRange.X;
-			for (int32 iChannel=0; iChannel < NumChannels; ++iChannel)
+			for (int32 iSample = 0; iSample < NumSamples; ++iSample)
 			{
-				LUT[(iSample * NumChannels) + iChannel] = Curves[iChannel].Eval(Time);
+				const float Time = (float(iSample) * SampleToDuration) + TimeRange.X;
+				for (int32 iChannel = 0; iChannel < NumChannels; ++iChannel)
+				{
+					LUT[(iSample * NumChannels) + iChannel] = Curves[iChannel].Eval(Time);
+				}
+			}
+		}
+		else //if (LUTMode == ENiagaraDistributionCurveLUTMode::Accumulate)
+		{
+			const float SampleDelta = NumSamples > 1 ? float(NumSamples - 1) : 1.0f;
+
+			TArray<float> ChannelSum;
+			ChannelSum.AddZeroed(Curves.Num());
+			for (int32 iSample = 0; iSample < NumSamples; ++iSample)
+			{
+				const float Time = (float(iSample) * SampleToDuration) + TimeRange.X;
+				for (int32 iChannel = 0; iChannel < NumChannels; ++iChannel)
+				{
+					LUT[(iSample * NumChannels) + iChannel] = ChannelSum[iChannel];
+					ChannelSum[iChannel] += Curves[iChannel].Eval(Time) / SampleDelta;
+				}
 			}
 		}
 		return LUT;
 	}
 
-	FVector2f CurvesFindTimeRange(TArrayView<const FRichCurve> Curves, int32 NumSamples)
+	FVector2f CurvesFindTimeRange(TArrayView<const FRichCurve> Curves, int32 NumSamples, ENiagaraDistributionCurveLUTMode LUTMode)
 	{
 		const int32 NumChannels = Curves.Num();
 		check(NumChannels > 0);
@@ -61,7 +80,7 @@ namespace NiagaraStatelessDistributionPrivate
 		// Reduce start / end times if the sampled results match
 		if (GReduceLUTTimeRange)
 		{
-			TArray<float> LUT = CurvesToLUT(Curves, NumSamples, TimeRange);
+			TArray<float> LUT = CurvesToLUT(Curves, NumSamples, TimeRange, LUTMode);
 
 			auto DoKeysMatch =
 				[&LUT, &NumChannels](int iCurrKey, int iNextKey)
@@ -138,21 +157,24 @@ namespace NiagaraStatelessDistributionPrivate
 		return true;
 	}
 
-	TArray<float> CurvesToOptimizedLUT(TArrayView<const FRichCurve> Curves, int32 MaxLutSampleCount, FVector2f& OutTimeRange)
+	TArray<float> CurvesToOptimizedLUT(TArrayView<const FRichCurve> Curves, int32 MaxLutSampleCount, ENiagaraDistributionCurveLUTMode LUTMode, FVector2f& OutTimeRange)
 	{
-		OutTimeRange = CurvesFindTimeRange(Curves, MaxLutSampleCount);
-		if (FMath::IsNearlyEqual(OutTimeRange.X, OutTimeRange.Y))
+		if (LUTMode == ENiagaraDistributionCurveLUTMode::Sample)
 		{
-			TArray<float> LUT = CurvesToLUT(Curves, 2, OutTimeRange);
-			return LUT;
+			OutTimeRange = CurvesFindTimeRange(Curves, MaxLutSampleCount, LUTMode);
+			if (FMath::IsNearlyEqual(OutTimeRange.X, OutTimeRange.Y))
+			{
+				TArray<float> LUT = CurvesToLUT(Curves, 2, OutTimeRange, LUTMode);
+				return LUT;
+			}
 		}
 
-		TArray<float> LUT = CurvesToLUT(Curves, MaxLutSampleCount, OutTimeRange);
+		TArray<float> LUT = CurvesToLUT(Curves, MaxLutSampleCount, OutTimeRange, LUTMode);
 		if (GOptimizeLUTs)
 		{
 			for (int32 iSamples = 2; iSamples < MaxLutSampleCount; ++iSamples)
 			{
-				TArray<float> NewLUT = CurvesToLUT(Curves, iSamples, OutTimeRange);
+				TArray<float> NewLUT = CurvesToLUT(Curves, iSamples, OutTimeRange, LUTMode);
 				if (AreLUTsAlmostEqual(NewLUT, LUT, Curves.Num()))
 				{
 					return NewLUT;
@@ -202,6 +224,7 @@ namespace NiagaraStatelessDistributionPrivate
 	template<typename FValueContainerSetNum, typename FValueAndChannelAccessor>
 	void UpdateDistributionValues(ENiagaraDistributionMode InMode, const TArray<float>& InChannelConstantsAndRanges, const TArray<FRichCurve>& InChannelCurves, int32 InChannelCount, FVector2f& OutTimeRange, FValueContainerSetNum InContainerNum, FValueAndChannelAccessor InValueAndChannelAccessor, int32 MaxLutSampleCount)
 	{
+		const ENiagaraDistributionCurveLUTMode LUTMode = ENiagaraDistributionCurveLUTMode::Sample;
 		switch(InMode)
 		{
 			case ENiagaraDistributionMode::UniformConstant:
@@ -261,7 +284,7 @@ namespace NiagaraStatelessDistributionPrivate
 				if (InChannelCurves.Num() >= ExpectedChannels)
 				{
 					MaxLutSampleCount = FMath::Max(MaxLutSampleCount, 2);
-					const TArray<float> LUT = CurvesToOptimizedLUT(MakeArrayView(InChannelCurves.GetData(), ExpectedChannels), MaxLutSampleCount, OutTimeRange);
+					const TArray<float> LUT = CurvesToOptimizedLUT(MakeArrayView(InChannelCurves.GetData(), ExpectedChannels), MaxLutSampleCount, LUTMode, OutTimeRange);
 					const int32 NumSamples = LUT.Num() / ExpectedChannels;
 
 					InContainerNum(NumSamples);
@@ -278,6 +301,30 @@ namespace NiagaraStatelessDistributionPrivate
 			}
 		}
 	}
+
+	template<typename FValueContainerSetNum, typename FValueAndChannelAccessor>
+	void UpdateDistributionCurveValues(ENiagaraDistributionMode InMode, const TArray<float>& InChannelConstantsAndRanges, const TArray<FRichCurve>& InChannelCurves, int32 InChannelCount, FVector2f& OutTimeRange, FValueContainerSetNum InContainerNum, FValueAndChannelAccessor InValueAndChannelAccessor, int32 MaxLutSampleCount, ENiagaraDistributionCurveLUTMode LUTMode)
+	{
+		check(InMode == ENiagaraDistributionMode::UniformCurve || InMode == ENiagaraDistributionMode::NonUniformCurve);
+
+		const int32 ExpectedChannels = InMode == ENiagaraDistributionMode::NonUniformCurve ? InChannelCount : 1;
+		if (InChannelCurves.Num() >= ExpectedChannels)
+		{
+			MaxLutSampleCount = FMath::Max(MaxLutSampleCount, 2);
+			const TArray<float> LUT = CurvesToOptimizedLUT(MakeArrayView(InChannelCurves.GetData(), ExpectedChannels), MaxLutSampleCount, LUTMode, OutTimeRange);
+			const int32 NumSamples = LUT.Num() / ExpectedChannels;
+
+			InContainerNum(NumSamples);
+			for (int32 SampleIndex = 0; SampleIndex < NumSamples; ++SampleIndex)
+			{
+				for (int32 ChannelIndex=0; ChannelIndex < InChannelCount; ++ChannelIndex)
+				{
+					const int32 LUTChannelIndex = InMode == ENiagaraDistributionMode::NonUniformCurve ? ChannelIndex : 0;
+					InValueAndChannelAccessor(SampleIndex, ChannelIndex) = LUT[(SampleIndex * ExpectedChannels) + LUTChannelIndex];
+				}
+			}
+		}
+	}
 #endif
 } //NiagaraStatelessCommon
 
@@ -375,8 +422,6 @@ void FNiagaraDistributionFloat::InitConstant(float Value)
 
 void FNiagaraDistributionFloat::InitCurve(std::initializer_list<float> CurvePoints)
 {
-	const int32 NumValues = CurvePoints.size();
-
 	Mode = ENiagaraDistributionMode::UniformCurve;
 	Values = CurvePoints;
 
@@ -549,6 +594,54 @@ FNiagaraStatelessRangeColor FNiagaraDistributionColor::CalculateRange(const FLin
 	return Range;
 }
 
+FNiagaraDistributionCurveFloat::FNiagaraDistributionCurveFloat()
+ : FNiagaraDistributionCurveFloat(ENiagaraDistributionCurveLUTMode::Sample)
+{
+}
+
+FNiagaraDistributionCurveFloat::FNiagaraDistributionCurveFloat(ENiagaraDistributionCurveLUTMode InLUTMode)
+#if WITH_EDITORONLY_DATA
+	: LUTMode(InLUTMode)
+#endif
+{
+	Mode = ENiagaraDistributionMode::UniformCurve;
+	if (InLUTMode == ENiagaraDistributionCurveLUTMode::Sample)
+	{
+		Values = { 1.0f, 1.0f };
+	}
+	else
+	{
+		Values = { 0.0f, 1.0f };
+	}
+#if WITH_EDITORONLY_DATA
+	NiagaraStatelessDistributionPrivate::CreateCurveFromValues({ 1.0f, 1.0f }, ChannelCurves, ChannelConstantsAndRanges);
+#endif
+}
+
+FNiagaraDistributionCurveVector3::FNiagaraDistributionCurveVector3()
+	: FNiagaraDistributionCurveVector3(ENiagaraDistributionCurveLUTMode::Sample)
+{
+}
+
+FNiagaraDistributionCurveVector3::FNiagaraDistributionCurveVector3(ENiagaraDistributionCurveLUTMode InLUTMode)
+#if WITH_EDITORONLY_DATA
+	: LUTMode(InLUTMode)
+#endif
+{
+	Mode = ENiagaraDistributionMode::UniformCurve;
+	if (InLUTMode == ENiagaraDistributionCurveLUTMode::Sample)
+	{
+		Values = { FVector3f::OneVector, FVector3f::OneVector };
+	}
+	else
+	{
+		Values = { FVector3f::ZeroVector, FVector3f::OneVector };
+	}
+#if WITH_EDITORONLY_DATA
+	NiagaraStatelessDistributionPrivate::CreateCurveFromValues({ FVector3f::OneVector, FVector3f::OneVector }, ChannelCurves, ChannelConstantsAndRanges);
+#endif
+}
+
 #if WITH_EDITORONLY_DATA
 void FNiagaraDistributionBase::PostEditChangeProperty(UObject* OwnerObject, FPropertyChangedEvent& PropertyChangedEvent)
 {
@@ -691,6 +784,35 @@ void FNiagaraDistributionColor::UpdateValuesFromDistribution()
 	);
 }
 
+void FNiagaraDistributionCurveFloat::UpdateValuesFromDistribution()
+{
+	NiagaraStatelessDistributionPrivate::UpdateDistributionCurveValues(
+		Mode,
+		ChannelConstantsAndRanges,
+		ChannelCurves,
+		1,
+		ValuesTimeRange,
+		[this](int32 NumValues) { Values.SetNum(NumValues); },
+		[this](int32 ValueIndex, int32 ChannelIndex) -> float& { return Values[ValueIndex]; },
+		MaxLutSampleCount,
+		LUTMode
+	);
+}
+
+void FNiagaraDistributionCurveVector3::UpdateValuesFromDistribution()
+{
+	NiagaraStatelessDistributionPrivate::UpdateDistributionCurveValues(
+		Mode,
+		ChannelConstantsAndRanges,
+		ChannelCurves,
+		3,
+		ValuesTimeRange,
+		[this](int32 NumValues) { Values.SetNum(NumValues); },
+		[this](int32 ValueIndex, int32 ChannelIndex) -> float& { return Values[ValueIndex].Component(ChannelIndex); },
+		MaxLutSampleCount,
+		LUTMode
+	);
+}
 #endif
 
 bool FNiagaraDistributionRangeFloat::SerializeFromMismatchedTag(const struct FPropertyTag& Tag, FStructuredArchive::FSlot Slot)
diff --git a/Engine/Plugins/FX/Niagara/Source/Niagara/Public/Stateless/NiagaraStatelessDistribution.h b/Engine/Plugins/FX/Niagara/Source/Niagara/Public/Stateless/NiagaraStatelessDistribution.h
--- a/Engine/Plugins/FX/Niagara/Source/Niagara/Public/Stateless/NiagaraStatelessDistribution.h
+++ b/Engine/Plugins/FX/Niagara/Source/Niagara/Public/Stateless/NiagaraStatelessDistribution.h
@@ -20,6 +20,12 @@ enum class ENiagaraDistributionMode
 	NonUniformCurve
 };
 
+enum class ENiagaraDistributionCurveLUTMode
+{
+	Sample,		// Each sample in the LUT represents the curve evaulation
+	Accumulate,	// Each sample in the LUT represents the acculumation of the curve evaluations
+};
+
 USTRUCT()
 struct FNiagaraDistributionBase
 {
@@ -60,6 +66,7 @@ struct FNiagaraDistributionBase
 	}
 
 	virtual bool AllowBinding() const { return true; }
+	virtual bool AllowConstant() const { return true; }
 	virtual bool AllowCurves() const { return true; }
 	virtual bool DisplayAsColor() const { return false; }
 	virtual int32 GetBaseNumberOfChannels() const { return 0; }
@@ -339,6 +346,66 @@ struct FNiagaraDistributionColor : public FNiagaraDistributionBase
 #endif
 };
 
+USTRUCT()
+struct FNiagaraDistributionCurveFloat : public FNiagaraDistributionBase
+{
+	GENERATED_BODY()
+
+	NIAGARA_API FNiagaraDistributionCurveFloat();
+	NIAGARA_API explicit FNiagaraDistributionCurveFloat(ENiagaraDistributionCurveLUTMode InLUTMode);
+
+	UPROPERTY(EditAnywhere, Category = "Parameters")
+	TArray<float> Values;
+
+	UPROPERTY(EditAnywhere, Category = "Parameters")
+	FVector2f ValuesTimeRange = FVector2f(0.0f, 1.0f);
+
+#if WITH_EDITORONLY_DATA
+	bool operator==(const FNiagaraDistributionCurveFloat& Other) const
+	{
+		return (FNiagaraDistributionBase)*this == (FNiagaraDistributionBase)Other && Values == Other.Values;
+	}
+
+	virtual bool AllowBinding() const override { return false; }
+	virtual bool AllowConstant() const override { return false; }
+	virtual int32 GetBaseNumberOfChannels() const override { return 1; }
+	NIAGARA_API virtual void UpdateValuesFromDistribution() override;
+
+private:
+	ENiagaraDistributionCurveLUTMode LUTMode = ENiagaraDistributionCurveLUTMode::Sample;
+#endif
+};
+
+USTRUCT()
+struct FNiagaraDistributionCurveVector3 : public FNiagaraDistributionBase
+{
+	GENERATED_BODY()
+
+	NIAGARA_API FNiagaraDistributionCurveVector3();
+	NIAGARA_API explicit FNiagaraDistributionCurveVector3(ENiagaraDistributionCurveLUTMode InLUTMode);
+
+	UPROPERTY(EditAnywhere, Category = "Parameters")
+	TArray<FVector3f> Values;
+
+	UPROPERTY(EditAnywhere, Category = "Parameters")
+	FVector2f ValuesTimeRange = FVector2f(0.0f, 1.0f);
+
+#if WITH_EDITORONLY_DATA
+	bool operator==(const FNiagaraDistributionCurveVector3& Other) const
+	{
+		return (FNiagaraDistributionBase)*this == (FNiagaraDistributionBase)Other && Values == Other.Values;
+	}
+
+	virtual bool AllowBinding() const override { return false; }
+	virtual bool AllowConstant() const override { return false; }
+	virtual int32 GetBaseNumberOfChannels() const override { return 3; }
+	NIAGARA_API virtual void UpdateValuesFromDistribution() override;
+
+private:
+	ENiagaraDistributionCurveLUTMode LUTMode = ENiagaraDistributionCurveLUTMode::Sample;
+#endif
+};
+
 template<>
 struct TStructOpsTypeTraits<FNiagaraDistributionRangeFloat> : public TStructOpsTypeTraitsBase2<FNiagaraDistributionRangeFloat>
 {
diff --git a/Engine/Plugins/FX/Niagara/Source/NiagaraEditor/Private/Stateless/NiagaraDistributionPropertyCustomization.cpp b/Engine/Plugins/FX/Niagara/Source/NiagaraEditor/Private/Stateless/NiagaraDistributionPropertyCustomization.cpp
--- a/Engine/Plugins/FX/Niagara/Source/NiagaraEditor/Private/Stateless/NiagaraDistributionPropertyCustomization.cpp
+++ b/Engine/Plugins/FX/Niagara/Source/NiagaraEditor/Private/Stateless/NiagaraDistributionPropertyCustomization.cpp
@@ -48,9 +48,9 @@ public:
 			const FName DisableRangeDistributionName("DisableRangeDistribution");
 			const FName DisplayAsColorDistributionName("DisplayAsColorDistribution");
 			bAllowBinding = InDistribution->AllowBinding() && InDistribution->GetBindingTypeDef().IsValid() && (PropertyHandle ? PropertyHandle->HasMetaData(DisableBindingDistributionName) == false : false);
-			bAllowUniform = PropertyHandle->HasMetaData(DisableUniformDistributionName) == false;
-			bAllowNonUniform = PropertyHandle->HasMetaData(DisableNonUniformDistributionName) == false;
-			bAllowRange = PropertyHandle->HasMetaData(DisableRangeDistributionName) == false;
+			bAllowUniform = InDistribution->AllowConstant() && PropertyHandle->HasMetaData(DisableUniformDistributionName) == false;
+			bAllowNonUniform = InDistribution->AllowConstant() && PropertyHandle->HasMetaData(DisableNonUniformDistributionName) == false;
+			bAllowRange = InDistribution->AllowConstant() && PropertyHandle->HasMetaData(DisableRangeDistributionName) == false;
 			bAllowCurves = InDistribution->AllowCurves() && PropertyHandle->HasMetaData(DisableCurveDistributionName) == false;
 			bDisplayAsColor = InDistribution->DisplayAsColor() || (InNumChannels >= 3 && PropertyHandle->HasMetaData(DisplayAsColorDistributionName));
 			PropertyHandle->SetOnPropertyValueChanged(FSimpleDelegate::CreateSP(this->AsShared(), &FNiagaraDistributionAdapter::DistributionPropertyChanged));
@@ -186,10 +186,7 @@ public:
 			if (bAllowCurves)
 			{
 				OutSupportedModes.Add(ENiagaraDistributionEditorMode::UniformCurve);
-				if (bAllowNonUniform)
-				{
-					OutSupportedModes.Add(ENiagaraDistributionEditorMode::NonUniformCurve);
-				}
+				OutSupportedModes.Add(ENiagaraDistributionEditorMode::NonUniformCurve);
 			}
 		}
 	}
diff --git a/Engine/Plugins/FX/Niagara/Source/NiagaraEditor/Private/ViewModels/Stack/NiagaraStackStatelessEmitterSimulateGroup.cpp b/Engine/Plugins/FX/Niagara/Source/NiagaraEditor/Private/ViewModels/Stack/NiagaraStackStatelessEmitterSimulateGroup.cpp
--- a/Engine/Plugins/FX/Niagara/Source/NiagaraEditor/Private/ViewModels/Stack/NiagaraStackStatelessEmitterSimulateGroup.cpp
+++ b/Engine/Plugins/FX/Niagara/Source/NiagaraEditor/Private/ViewModels/Stack/NiagaraStackStatelessEmitterSimulateGroup.cpp
@@ -386,6 +386,8 @@ void UNiagaraStackStatelessModuleItem::RefreshChildrenInternal(const TArray<UNia
 			ModuleObject->RegisterInstancedCustomPropertyTypeLayout(FNiagaraDistributionRangeVector3::StaticStruct()->GetFName(), FOnGetPropertyTypeCustomizationInstance::CreateStatic(&FNiagaraDistributionPropertyCustomization::MakeVector3Instance));
 			ModuleObject->RegisterInstancedCustomPropertyTypeLayout(FNiagaraDistributionRangeColor::StaticStruct()->GetFName(), FOnGetPropertyTypeCustomizationInstance::CreateStatic(&FNiagaraDistributionPropertyCustomization::MakeColorInstance));
 			ModuleObject->RegisterInstancedCustomPropertyTypeLayout(FNiagaraDistributionRangeInt::StaticStruct()->GetFName(), FOnGetPropertyTypeCustomizationInstance::CreateStatic(&FNiagaraDistributionIntPropertyCustomization::MakeIntInstance));
+			ModuleObject->RegisterInstancedCustomPropertyTypeLayout(FNiagaraDistributionCurveFloat::StaticStruct()->GetFName(), FOnGetPropertyTypeCustomizationInstance::CreateStatic(&FNiagaraDistributionPropertyCustomization::MakeFloatInstance));
+			ModuleObject->RegisterInstancedCustomPropertyTypeLayout(FNiagaraDistributionCurveVector3::StaticStruct()->GetFName(), FOnGetPropertyTypeCustomizationInstance::CreateStatic(&FNiagaraDistributionPropertyCustomization::MakeVector3Instance));
 			ModuleObjectWeak = ModuleObject;
 		}
 		NewChildren.Add(ModuleObject);
diff --git a/Engine/Plugins/FX/Niagara/Source/NiagaraShader/Internal/Stateless/NiagaraStatelessModuleShaderParameters.h b/Engine/Plugins/FX/Niagara/Source/NiagaraShader/Internal/Stateless/NiagaraStatelessModuleShaderParameters.h
--- a/Engine/Plugins/FX/Niagara/Source/NiagaraShader/Internal/Stateless/NiagaraStatelessModuleShaderParameters.h
+++ b/Engine/Plugins/FX/Niagara/Source/NiagaraShader/Internal/Stateless/NiagaraStatelessModuleShaderParameters.h
@@ -93,9 +93,10 @@ namespace NiagaraStateless
 	END_SHADER_PARAMETER_STRUCT()
 
 	BEGIN_SHADER_PARAMETER_STRUCT(FMeshRotationRateModule_ShaderParameters, NIAGARASHADER_API)
-		SHADER_PARAMETER(int32,			MeshRotationRate_ModuleEnabled)
-		SHADER_PARAMETER(FVector3f,		MeshRotationRate_Scale)
-		SHADER_PARAMETER(FVector3f,		MeshRotationRate_Bias)
+		SHADER_PARAMETER(int32,										MeshRotationRate_ModuleEnabled)
+		SHADER_PARAMETER(FVector3f,									MeshRotationRate_Scale)
+		SHADER_PARAMETER(FVector3f,									MeshRotationRate_Bias)
+		SHADER_PARAMETER(FNiagaraStatelessBuiltDistributionType,	MeshRotationRate_RateScaleParameters)
 	END_SHADER_PARAMETER_STRUCT()
 
 	BEGIN_SHADER_PARAMETER_STRUCT(FScaleSpriteSizeModule_ShaderParameters, NIAGARASHADER_API)
@@ -116,8 +117,10 @@ namespace NiagaraStateless
 	END_SHADER_PARAMETER_STRUCT()
 
 	BEGIN_SHADER_PARAMETER_STRUCT(FSpriteRotationRateModule_ShaderParameters, NIAGARASHADER_API)
-		SHADER_PARAMETER(float,		SpriteRotationRate_Scale)
-		SHADER_PARAMETER(float,		SpriteRotationRate_Bias)
+		SHADER_PARAMETER(int32,										SpriteRotationRate_Enabled)
+		SHADER_PARAMETER(float,										SpriteRotationRate_Scale)
+		SHADER_PARAMETER(float,										SpriteRotationRate_Bias)
+		SHADER_PARAMETER(FNiagaraStatelessBuiltDistributionType,	SpriteRotationRate_RateScaleParameters)
 	END_SHADER_PARAMETER_STRUCT()
 
 	BEGIN_SHADER_PARAMETER_STRUCT(FSubUVAnimationModule_ShaderParameters, NIAGARASHADER_API)
