A number of optimzations to the Niagara translator
-store a local copy of the keyword replacement map rather than accessing the module each time we need to sanitize a symbol
-enable the TLS for type resolution during translation
-try to reduce the amount of string allocations through the use of local stringbuilders and string views instead of constant FString construction
#tests validated that the generated hlsl is identical for all niagara systems

diff --git a/Engine/Plugins/FX/Niagara/Source/NiagaraEditor/Private/NiagaraHlslTranslator.cpp b/Engine/Plugins/FX/Niagara/Source/NiagaraEditor/Private/NiagaraHlslTranslator.cpp
--- a/Engine/Plugins/FX/Niagara/Source/NiagaraEditor/Private/NiagaraHlslTranslator.cpp
+++ b/Engine/Plugins/FX/Niagara/Source/NiagaraEditor/Private/NiagaraHlslTranslator.cpp
@@ -37,6 +37,7 @@
 #include "NiagaraTrace.h"
 #include "NiagaraTraversalStateContext.h"
 #include "ShaderCore.h"
+#include "String/ParseTokens.h"
 #include "Misc/FileHelper.h"
 #include "Misc/CoreMiscDefines.h"
 #include "NiagaraDataInterfaceUtilities.h"
@@ -151,6 +152,320 @@ void FNiagaraShaderQueueTickable::Tick(float DeltaSeconds)
 	ProcessQueue();
 }
 
+namespace NiagaraHlslTranslatorImpl
+{
+
+// helper function to evaluate whether the given character is an ascii alpha/number
+bool IsAsciiAlnum(TCHAR ch)
+{
+	return ((ch >= TCHAR('0') && ch <= TCHAR('9')) ||
+		(ch >= TCHAR('A') && ch <= TCHAR('Z')) ||
+		(ch >= TCHAR('a') && ch <= TCHAR('z')));
+}
+
+// helper function to append an integer to a string builder (mimics FString::AppendInt)
+void AppendInt(FStringBuilderBase& OutBuilder, TCHAR ch)
+{
+	const TCHAR* DigitToChar = TEXT("0123456789");
+	constexpr int32 MaxDigitCount = 16;
+	TCHAR TempDigits[MaxDigitCount];
+	int32 CurDigitIndex = MaxDigitCount;
+
+	uint32 chAsUint = static_cast<uint32>(ch);
+
+	do
+	{
+		TempDigits[--CurDigitIndex] = DigitToChar[chAsUint % 10];
+		chAsUint /= 10;
+	} while (chAsUint);
+
+	OutBuilder.Append(TempDigits + CurDigitIndex, MaxDigitCount - CurDigitIndex);
+}
+
+// helper function to mimic FString::ReplaceInline but uses StringBuilders/Views
+void ReplaceInline(FStringBuilderBase& StringBuilder, int32 SearchStartIndex, FStringView SearchString, FStringView ReplaceString)
+{
+	const int32 SearchStringLen = SearchString.Len();
+	const int32 ReplaceStringLen = ReplaceString.Len();
+
+	while (SearchStartIndex < StringBuilder.Len())
+	{
+		SearchStartIndex = StringBuilder.ToView().Find(SearchString, SearchStartIndex);
+		if (SearchStartIndex != INDEX_NONE)
+		{
+			StringBuilder.ReplaceAt(SearchStartIndex, SearchStringLen, ReplaceString);
+			SearchStartIndex += ReplaceStringLen;
+		}
+		else
+		{
+			break;
+		}
+	}
+}
+
+void AppendConvertedToAscii(FStringView InNameView, FStringBuilderBase& OutBuilder)
+{
+	for (TCHAR ch : InNameView)
+	{
+		if (IsAsciiAlnum(ch) || FChar::IsUnderscore(ch))
+		{
+			// Do nothing.. these are valid chars..
+			OutBuilder << ch;
+		}
+		// skip whitespace characters
+		else if (!FChar::IsWhitespace(ch))
+		{
+			// Need to replace the bad characters..
+			OutBuilder << TEXT("ASC");
+
+			// replacement for FString::AppendInt
+			AppendInt(OutBuilder, ch);
+		}
+	}
+}
+
+void AppendSanitizedSymbolString(FStringView SymbolName, bool bCollapseNamespaces, const TMap<FString, FString>& ReplacementMap, FStringBuilderBase& StringBuilder)
+{
+	if (SymbolName.Len() == 0)
+	{
+		return;
+	}
+
+	// Split up into individual namespaces...
+	TArray<FStringView, TInlineAllocator<16>> NamePartStrings;
+	UE::String::ParseTokens(SymbolName, TEXT("."), [&NamePartStrings](FStringView Token) { NamePartStrings.Add(Token); });
+
+	bool bIsFirstSplit = true;
+	const int32 OutputStartIndex = StringBuilder.Len();
+
+	// Rules for variable namespaces..
+	for (FStringView SplitName : NamePartStrings)
+	{
+		SplitName = SplitName.TrimStartAndEnd();
+
+		if (!bIsFirstSplit)
+		{
+			StringBuilder << (bCollapseNamespaces ? TEXT("_") : TEXT("."));
+		}
+
+		if (SplitName[0] >= TCHAR('0') && SplitName[0] <= TCHAR('9')) // Cannot start with a numeric digit
+		{
+			StringBuilder << TEXT("INTEGER_");
+		}
+
+		if (const FString* FoundReplacementStr = ReplacementMap.FindByHash(GetTypeHash(SplitName), SplitName)) // Look for the string in the keyword protections array.
+		{
+			StringBuilder << *FoundReplacementStr;
+		}
+		else
+		{
+			AppendConvertedToAscii(SplitName, StringBuilder);
+		}
+
+		bIsFirstSplit = false;
+	}
+
+	ReplaceInline(StringBuilder, OutputStartIndex, TEXT("__"), TEXT("ASC95ASC95"));
+}
+
+void AppendSanitizedSymbolName(FName SymbolName, bool bCollapseNamespaces, const TMap<FString, FString>& ReplacementMap, FStringBuilderBase& StringBuilder)
+{
+	FNameBuilder NameBuilder(SymbolName);
+	AppendSanitizedSymbolString(NameBuilder.ToView(), bCollapseNamespaces, ReplacementMap, StringBuilder);
+}
+
+void AppendFunctionSignatureSymbol(const FNiagaraFunctionSignature& Sig, FStringBuilderBase& SigStr)
+{
+	Sig.Name.AppendString(SigStr);
+	if (!Sig.OwnerName.IsNone() && Sig.OwnerName.IsValid())
+	{
+		SigStr << TEXT("_");
+		SigStr << Sig.OwnerName.ToString().Replace(TEXT("."), TEXT("_"));
+	}
+	else
+	{
+		SigStr << TEXT("_Func_");
+	}
+
+	if (Sig.bRequiresExecPin)
+	{
+		SigStr << TEXT("_UEImpureCall"); // Let the cross compiler know that we intend to keep this.
+	}
+
+	for (const TTuple<FName, FName>& Specifier : Sig.FunctionSpecifiers)
+	{
+		SigStr << TEXT("_");
+		Specifier.Key.AppendString(SigStr);
+
+		const int32 SpecifierValueStartIndex = SigStr.Len();
+		Specifier.Value.AppendString(SigStr);
+		ReplaceInline(SigStr, SpecifierValueStartIndex, TEXT("."), TEXT("_"));
+	}
+
+	auto AddVarsToSig = [&SigStr](TConstArrayView<FNiagaraVariableBase> Vars)
+	{
+		for (const FNiagaraVariableBase& Var : Vars)
+		{
+			SigStr << TEXT("_");
+			Var.GetName().AppendString(SigStr);
+			Var.GetType().GetFName().AppendString(SigStr);
+		}
+	};
+
+	TArray<FNiagaraVariableBase> VariadicParams;
+	VariadicParams.Reserve(Sig.NumOptionalInputs() + Sig.NumOptionalOutputs());
+	Sig.GetVariadicInputs(VariadicParams);
+	Sig.GetVariadicOutputs(VariadicParams);
+	AddVarsToSig(VariadicParams);
+}
+
+void AppendStructHlslTypeName(const FNiagaraTypeDefinition& Type, FStringBuilderBase& StringBuilder)
+{
+	check(FNiagaraTypeHelper::IsLWCType(Type) == false);
+
+	if (Type.IsValid() == false)
+	{
+		StringBuilder.Append(TEXT("undefined"));
+	}
+	else if (Type == FNiagaraTypeDefinition::GetFloatDef())
+	{
+		StringBuilder.Append(TEXT("float"));
+	}
+	else if (Type == FNiagaraTypeDefinition::GetVec2Def())
+	{
+		StringBuilder.Append(TEXT("float2"));
+	}
+	else if (Type == FNiagaraTypeDefinition::GetVec3Def() || Type == FNiagaraTypeDefinition::GetPositionDef())
+	{
+		StringBuilder.Append(TEXT("float3"));
+	}
+	else if (Type == FNiagaraTypeDefinition::GetVec4Def() || Type == FNiagaraTypeDefinition::GetColorDef() || Type == FNiagaraTypeDefinition::GetQuatDef())
+	{
+		StringBuilder.Append(TEXT("float4"));
+	}
+	else if (Type == FNiagaraTypeDefinition::GetMatrix4Def())
+	{
+		StringBuilder.Append(TEXT("float4x4"));
+	}
+	else if (Type.IsSameBaseDefinition(FNiagaraTypeDefinition::GetIntDef()) || Type.GetEnum() != nullptr)
+	{
+		StringBuilder.Append(TEXT("int"));
+	}
+	else if (Type.IsSameBaseDefinition(FNiagaraTypeDefinition::GetBoolDef()))
+	{
+		StringBuilder.Append(TEXT("bool"));
+	}
+	else if (Type == FNiagaraTypeDefinition::GetParameterMapDef())
+	{
+		StringBuilder.Append(TEXT("FParamMap0"));
+	}
+	else
+	{
+		Type.GetFName().AppendString(StringBuilder);
+	}
+}
+
+FNiagaraVariable ConvertToSimulationVariable(const FNiagaraVariable& Param)
+{
+	if (FNiagaraTypeHelper::IsLWCType(Param.GetType()))
+	{
+		UScriptStruct* Struct = FNiagaraTypeHelper::FindNiagaraFriendlyTopLevelStruct(Param.GetType().GetScriptStruct(), ENiagaraStructConversion::Simulation);
+		FNiagaraVariable SimParam(FNiagaraTypeDefinition(Struct), Param.GetName());
+		if (Param.IsDataAllocated())
+		{
+			SimParam.AllocateData();
+			FNiagaraTypeRegistry::GetStructConverter(Param.GetType()).ConvertDataToSimulation(SimParam.GetData(), Param.GetData());
+		}
+		return SimParam;
+	}
+	return Param;
+}
+
+void AppendFunctionHlslPrototypeString(FStringView InVariableName, const FNiagaraFunctionSignature& FunctionSignature, const TMap<FString, FString>& ReplacementMap, FStringBuilderBase& StringBuilder)
+{
+	if (FunctionSignature.bMemberFunction)
+	{
+		StringBuilder.Append(InVariableName);
+		StringBuilder.Append(TEXT("."));
+		AppendSanitizedSymbolName(FunctionSignature.Name, false /*bCollapse*/, ReplacementMap, StringBuilder);
+
+		// Build specifiers
+		if (FunctionSignature.FunctionSpecifiers.Num())
+		{
+			bool bNeedsComma = false;
+			StringBuilder.AppendChar(TEXT('<'));
+			for (auto SpecifierIt = FunctionSignature.FunctionSpecifiers.CreateConstIterator(); SpecifierIt; ++SpecifierIt)
+			{
+				if (bNeedsComma)
+				{
+					StringBuilder.Append(TEXT(", "));
+				}
+				bNeedsComma = true;
+
+				SpecifierIt.Key().AppendString(StringBuilder);
+				StringBuilder.AppendChar(TEXT('='));
+				StringBuilder.AppendChar(TEXT('"'));
+				if (SpecifierIt.Value().IsNone())
+				{
+					StringBuilder.Append(TEXT("Value"));
+				}
+				else
+				{
+					SpecifierIt.Value().AppendString(StringBuilder);
+				}
+				StringBuilder.AppendChar(TEXT('"'));
+			}
+			StringBuilder.AppendChar(TEXT('>'));
+		}
+
+		// Build function parameters
+		{
+			StringBuilder.Append(TEXT("("));
+			bool bNeedsComma = false;
+
+			// Inputs
+			for (int i = 1; i < FunctionSignature.Inputs.Num(); ++i)
+			{
+				FNiagaraVariable InputVar = ConvertToSimulationVariable(FunctionSignature.Inputs[i]);
+				if (bNeedsComma)
+				{
+					StringBuilder.Append(TEXT(", "));
+				}
+				bNeedsComma = true;
+
+				StringBuilder.Append(TEXT("in "));
+				AppendStructHlslTypeName(InputVar.GetType(), StringBuilder);
+				StringBuilder.Append(TEXT(" In_"));
+				AppendSanitizedSymbolName(InputVar.GetName(), false /*bCollapse*/, ReplacementMap, StringBuilder);
+			}
+
+			// Outputs
+			for (int i = 0; i < FunctionSignature.Outputs.Num(); ++i)
+			{
+				FNiagaraVariable OutputVar = ConvertToSimulationVariable(FunctionSignature.Outputs[i]);
+				if (bNeedsComma)
+				{
+					StringBuilder.Append(TEXT(", "));
+				}
+				bNeedsComma = true;
+
+				StringBuilder.Append(TEXT("out "));
+				AppendStructHlslTypeName(OutputVar.GetType(), StringBuilder);
+				StringBuilder.Append(TEXT(" Out_"));
+				AppendSanitizedSymbolName(OutputVar.GetName(), false /*bCollapse*/, ReplacementMap, StringBuilder);
+			}
+			StringBuilder.Append(TEXT(");"));
+		}
+	}
+	else
+	{
+		ensureAlwaysMsgf(false, TEXT("None member functions not supported currently"));
+	}
+}
+
+};
+
+
 ENiagaraScriptCompileStatus FNiagaraTranslateResults::TranslateResultsToSummary(const FNiagaraTranslateResults* TranslateResults)
 {
 	ENiagaraScriptCompileStatus SummaryStatus = ENiagaraScriptCompileStatus::NCS_Unknown;
@@ -514,7 +829,7 @@ void TNiagaraHlslTranslator<GraphBridge>::GenerateFunctionSignature(ENiagaraScri
 				{
 					checkf(InputVar.GetType().IsSameBaseDefinition(FNiagaraTypeDefinition::GetBoolDef()), TEXT("Only boolean types are currently supported for literal constants."));
 					FString LiteralConstantAlias = InputVar.GetName().ToString() + TEXT("_") + (InputVar.GetValue<bool>() ? TEXT("true") : TEXT("false"));
-					InName += TEXT("_") + GetSanitizedSymbolName(LiteralConstantAlias.Replace(TEXT("."), TEXT("_")));
+					InName += TEXT("_") + BuildSanitizedSymbolString(LiteralConstantAlias.Replace(TEXT("."), TEXT("_")));
 					ConstantInputIndicesToRemove.Add(i);
 				}
 				else
@@ -574,13 +889,13 @@ void TNiagaraHlslTranslator<GraphBridge>::GenerateFunctionSignature(ENiagaraScri
 		SignatureName.Reserve(1024);
 		if (ModuleAliasStr != nullptr)
 		{
-			SignatureName = GetSanitizedSymbolName(*ModuleAliasStr);
+			SignatureName = BuildSanitizedSymbolString(*ModuleAliasStr);
 		}
 		if (EmitterAliasStr != nullptr)
 		{
 			FString Prefix = ModuleAliasStr != nullptr ? TEXT("_") : TEXT("");
 			SignatureName += Prefix;
-			SignatureName += GetSanitizedSymbolName(*EmitterAliasStr);
+			SignatureName += BuildSanitizedSymbolString(*EmitterAliasStr);
 		}
 		if (InFunctionNameSuffix.IsEmpty() == false)
 		{
@@ -619,6 +934,9 @@ FNiagaraHlslTranslator::FNiagaraHlslTranslator()
 {
 	const UNiagaraSettings* Settings = GetDefault<UNiagaraSettings>();
 	bEnforceStrictTypesValidations = Settings->bEnforceStrictStackTypes;
+
+	const UNiagaraEditorSettings* EditorSettings = GetDefault<UNiagaraEditorSettings>();
+	KeywordReplacementMap = EditorSettings->GetHLSLKeywordReplacementsMap();
 }
 
 
@@ -714,8 +1032,8 @@ void TNiagaraHlslTranslator<GraphBridge>::BuildMissingDefaults()
 				{
 					FNiagaraVariable SourceForInitialValue = FNiagaraParameterUtilities::GetSourceForInitialValue(Var);
 					FString ParameterMapInstanceName = GetParameterMapInstanceName(0);
-					FString Value = FString::Printf(TEXT("%s.%s = %s.%s;\n"), *ParameterMapInstanceName, *GetSanitizedSymbolName(Var.GetName().ToString()),
-						*ParameterMapInstanceName, *GetSanitizedSymbolName(SourceForInitialValue.GetName().ToString()));
+					FString Value = FString::Printf(TEXT("%s.%s = %s.%s;\n"), *ParameterMapInstanceName, *BuildSanitizedSymbolName(Var.GetName()),
+						*ParameterMapInstanceName, *BuildSanitizedSymbolName(SourceForInitialValue.GetName()));
 					AddBodyChunk(Value);
 					continue;
 				}
@@ -895,7 +1213,7 @@ FString TNiagaraHlslTranslator<GraphBridge>::BuildParameterMapHlslDefinitions(TA
 		const FNiagaraVariable& NiagaraVariable = UniqueVariables[UniqueVarIdx];
 		if (FNiagaraParameterUtilities::IsAttribute(NiagaraVariable))
 		{
-			FString VariableName = GetSanitizedSymbolName(NiagaraVariable.GetName().ToString());
+			FString VariableName = BuildSanitizedSymbolName(NiagaraVariable.GetName());
 			ReplaceNamespaceInline(VariableName, PARAM_MAP_ATTRIBUTE_STR, PARAM_MAP_INDICES_STR);
 			RegisterNames.Add(VariableName);
 		}
@@ -931,7 +1249,7 @@ FString TNiagaraHlslTranslator<GraphBridge>::BuildParameterMapHlslDefinitions(TA
 		}
 
 		TArray<FString> StructNameArray;
-		FString SanitizedVarName = GetSanitizedSymbolName(Variable.GetName().ToString());
+		FString SanitizedVarName = BuildSanitizedSymbolName(Variable.GetName());
 		int32 NumFound = SanitizedVarName.ParseIntoArray(StructNameArray, TEXT("."));
 		if (NumFound == 1) // Meaning no split above
 		{
@@ -950,14 +1268,14 @@ FString TNiagaraHlslTranslator<GraphBridge>::BuildParameterMapHlslDefinitions(TA
 				}
 
 				FString TypeName = GetStructHlslTypeName(Variable.GetType());
-				FString VarName = GetSanitizedSymbolName(*FinalName);
+				FString VarName = BuildSanitizedSymbolString(FinalName);
 				if (NumFound > StructNameArray.Num() + 1 && StructNameArray.Num() != 0)
 				{
-					TypeName = FString::Printf(TEXT("FParamMap%d_%s_%s"), UniqueParamMapIdx, *FString::Join(StructNameArray, TEXT("_")), *GetSanitizedSymbolName(*FinalName));
+					TypeName = FString::Printf(TEXT("FParamMap%d_%s_%s"), UniqueParamMapIdx, *FString::Join(StructNameArray, TEXT("_")), *VarName);
 				}
 				else if (StructNameArray.Num() == 0)
 				{
-					TypeName = FString::Printf(TEXT("FParamMap%d_%s"), UniqueParamMapIdx, *GetSanitizedSymbolName(*FinalName));
+					TypeName = FString::Printf(TEXT("FParamMap%d_%s"), UniqueParamMapIdx, *VarName);
 				}
 				TPair<FString, FString> Pair(TypeName, VarName);
 				ParamStructNameToMembers.FindOrAdd(StructType).AddUnique(Pair);
@@ -1064,7 +1382,7 @@ void FNiagaraHlslTranslator::BuildConstantBuffer(ENiagaraCodeChunkMode ChunkMode
 {
 	for (const FNiagaraVariable& Variable : T::GetVariables())
 	{
-		const FString SymbolName = GetSanitizedSymbolName(Variable.GetName().ToString(), true);
+		const FString SymbolName = BuildSanitizedSymbolName(Variable.GetName(), true);
 		AddChunkToConstantBuffer(SymbolName, Variable, ChunkMode);
 	}
 }
@@ -1229,6 +1547,8 @@ FNiagaraTranslateResults TNiagaraHlslTranslator<GraphBridge>::Translate(const FN
 	check(CompileData);
 	check(CompileDuplicateData);
 
+	FNiagaraTypeRegistryTLSProxy TypeRegistryProxy;
+
 	CompileOptions = InCompileOptions;
 	TranslationOptions = InTranslateOptions;
 	CompilationTarget = TranslationOptions.SimTarget;
@@ -1275,7 +1595,7 @@ FNiagaraTranslateResults TNiagaraHlslTranslator<GraphBridge>::Translate(const FN
 		OtherOutputParamMapHistoriesSanitizedVariables[i].Reserve(OtherOutputParamMapHistories[i].Variables.Num());
 		for (const FNiagaraVariable& Var : OtherOutputParamMapHistories[i].Variables)
 		{
-			OtherOutputParamMapHistoriesSanitizedVariables[i].Emplace(Var.GetType(), *GetSanitizedSymbolName(Var.GetName().ToString()));
+			OtherOutputParamMapHistoriesSanitizedVariables[i].Emplace(Var.GetType(), *BuildSanitizedSymbolName(Var.GetName()));
 		}
 	}
 
@@ -2510,7 +2830,7 @@ void TNiagaraHlslTranslator<GraphBridge>::HandleSimStageSetupAndTeardown(int32 I
 				FString SetupOutHLSL;
 				ProcessCustomHlsl(SetupGeneratedHLSL, TranslationStage.ScriptUsage, Sig, Inputs, nullptr, SetupOutHLSL, SignatureOut);
 				SetupFunctionHLSL += GetFunctionSignature(SignatureOut) + TEXT("\n{\n") + SetupOutHLSL + TEXT("\n}\n");
-				SetupBody += FString::Printf(TEXT("\n\t%s(Context);\n"), *GetFunctionSignatureSymbol(SignatureOut));
+				SetupBody += FString::Printf(TEXT("\n\t%s(Context);\n"), *BuildFunctionSignatureString(SignatureOut));
 			}
 
 			FString TeardownGeneratedHLSL;
@@ -2521,7 +2841,7 @@ void TNiagaraHlslTranslator<GraphBridge>::HandleSimStageSetupAndTeardown(int32 I
 				FString TeardownOutHLSL;
 				ProcessCustomHlsl(TeardownGeneratedHLSL, TranslationStage.ScriptUsage, Sig, Inputs, nullptr, TeardownOutHLSL, SignatureOut);
 				TeardownFunctionHLSL += GetFunctionSignature(SignatureOut) + TEXT("\n{\n") + TeardownOutHLSL + TEXT("\n}\n");
-				TeardownBody += FString::Printf(TEXT("\n\t%s(Context);\n"), *GetFunctionSignatureSymbol(SignatureOut));
+				TeardownBody += FString::Printf(TEXT("\n\t%s(Context);\n"), *BuildFunctionSignatureString(SignatureOut));
 			}
 
 		}
@@ -2533,7 +2853,7 @@ void TNiagaraHlslTranslator<GraphBridge>::HandleSimStageSetupAndTeardown(int32 I
 			TArray<FString> AttributeHLSLNames;
 			for (const FNiagaraVariable& Var : ReadVars)
 			{
-				AttributeHLSLNames.Emplace(TEXT("Map.") + GetSanitizedSymbolName(Var.GetName().ToString()));
+				AttributeHLSLNames.Emplace(TEXT("Map.") + BuildSanitizedSymbolName(Var.GetName()));
 			}
 
 			if (CDO->GenerateIterationSourceNamespaceReadAttributesHLSL(DIInstanceInfo,  IterationSourceVar, MakeArrayView(Sig.Inputs), MakeArrayView(ReadVars), MakeArrayView(AttributeHLSLNames), bSpawnOnly, bSpawnOnly, bPartialWrites, GeneratedErrors, AttributeReadGeneratedHLSL) && AttributeReadGeneratedHLSL.Len() > 0)
@@ -2543,7 +2863,7 @@ void TNiagaraHlslTranslator<GraphBridge>::HandleSimStageSetupAndTeardown(int32 I
 				FString AttributeReadOutHLSL;
 				ProcessCustomHlsl(AttributeReadGeneratedHLSL, TranslationStage.ScriptUsage, Sig, Inputs, nullptr, AttributeReadOutHLSL, SignatureOut);
 				SetupFunctionHLSL += GetFunctionSignature(SignatureOut) + TEXT("\n{\n") + AttributeReadOutHLSL + TEXT("\n}\n");
-				SetupBody += FString::Printf(TEXT("\n\t%s(Context);\n"), *GetFunctionSignatureSymbol(SignatureOut));
+				SetupBody += FString::Printf(TEXT("\n\t%s(Context);\n"), *BuildFunctionSignatureString(SignatureOut));
 
 			}
 		}
@@ -2554,7 +2874,7 @@ void TNiagaraHlslTranslator<GraphBridge>::HandleSimStageSetupAndTeardown(int32 I
 			TArray<FString> AttributeHLSLNames;
 			for (const FNiagaraVariable& Var : WriteVars)
 			{
-				AttributeHLSLNames.Emplace(TEXT("Map.") + GetSanitizedSymbolName(Var.GetName().ToString()));
+				AttributeHLSLNames.Emplace(TEXT("Map.") + BuildSanitizedSymbolName(Var.GetName()));
 			}
 
 			if (CDO->GenerateIterationSourceNamespaceWriteAttributesHLSL(DIInstanceInfo, IterationSourceVar, MakeArrayView(Sig.Inputs), MakeArrayView(WriteVars), MakeArrayView(AttributeHLSLNames), MakeArrayView(AllVars), bSpawnOnly, bPartialWrites, GeneratedErrors, AttributeWriteGeneratedHLSL) && AttributeWriteGeneratedHLSL.Len() > 0)
@@ -2564,7 +2884,7 @@ void TNiagaraHlslTranslator<GraphBridge>::HandleSimStageSetupAndTeardown(int32 I
 				FString AttributeWriteOutHLSL;
 				ProcessCustomHlsl(AttributeWriteGeneratedHLSL, TranslationStage.ScriptUsage, Sig, Inputs, nullptr, AttributeWriteOutHLSL, SignatureOut);
 				TeardownFunctionHLSL += GetFunctionSignature(SignatureOut) + TEXT("\n{\n") + AttributeWriteOutHLSL + TEXT("\n}\n");
-				TeardownBody += FString::Printf(TEXT("\n\t%s(Context);\n"), *GetFunctionSignatureSymbol(SignatureOut));
+				TeardownBody += FString::Printf(TEXT("\n\t%s(Context);\n"), *BuildFunctionSignatureString(SignatureOut));
 			}
 		}
 
@@ -2762,8 +3082,8 @@ void FNiagaraHlslTranslator::DefinePreviousParametersFunction(FString& HlslOutpu
 				if (FNiagaraParameterUtilities::IsPreviousValue(Var))
 				{
 					FString CurMap = TranslationStages[0].PassNamespace;
-					FString Value = FString::Printf(TEXT("Context.%s.%s = Context.%s.%s;\n"), *CurMap, *GetSanitizedSymbolName(Var.GetName().ToString()),
-						*CurMap, *GetSanitizedSymbolName(FNiagaraParameterUtilities::GetSourceForPreviousValue(Var).GetName().ToString()));
+					FString Value = FString::Printf(TEXT("Context.%s.%s = Context.%s.%s;\n"), *CurMap, *BuildSanitizedSymbolName(Var.GetName()),
+						*CurMap, *BuildSanitizedSymbolName(FNiagaraParameterUtilities::GetSourceForPreviousValue(Var).GetName()));
 					HlslOutputString += Value + TEXT("\n");
 				}
 			}
@@ -2782,9 +3102,9 @@ void TNiagaraHlslTranslator<GraphBridge>::DefineInterpolatedParametersFunction(F
 			continue;
 		}
 
-		FString Emitter_InterpSpawnStartDt = GetSanitizedSymbolName(ActiveHistoryForFunctionCalls.ResolveAliases(SYS_PARAM_EMITTER_INTERP_SPAWN_START_DT).GetName().ToString());
+		FString Emitter_InterpSpawnStartDt = BuildSanitizedSymbolName(ActiveHistoryForFunctionCalls.ResolveAliases(SYS_PARAM_EMITTER_INTERP_SPAWN_START_DT).GetName());
 		Emitter_InterpSpawnStartDt = Emitter_InterpSpawnStartDt.Replace(TEXT("."), TEXT("_"));//TODO: This should be rolled into GetSanitisedSymbolName but currently some usages require the '.' intact. Fix those up!.
-		FString Emitter_SpawnInterval = GetSanitizedSymbolName(ActiveHistoryForFunctionCalls.ResolveAliases(SYS_PARAM_EMITTER_SPAWN_INTERVAL).GetName().ToString());
+		FString Emitter_SpawnInterval = BuildSanitizedSymbolName(ActiveHistoryForFunctionCalls.ResolveAliases(SYS_PARAM_EMITTER_SPAWN_INTERVAL).GetName());
 		Emitter_SpawnInterval = Emitter_SpawnInterval.Replace(TEXT("."), TEXT("_"));//TODO: This should be rolled into GetSanitisedSymbolName but currently some usages require the '.' intact. Fix those up!.
 
 		HlslOutputString += TEXT("void InterpolateParameters(inout FSimulationContext Context)\n{\n");
@@ -2987,7 +3307,7 @@ void FNiagaraHlslTranslator::DefineDataSetWriteFunction(FString &HlslOutputStrin
 void FNiagaraHlslTranslator::ConvertCompileInfoToParamInfo(const FNiagaraScriptDataInterfaceCompileInfo& Info, FNiagaraDataInterfaceGPUParamInfo& DIInstanceInfo, TArray<FNiagaraFunctionSignature>& GeneratedFunctionSignatures)
 {
 	FString OwnerIDString = Info.Name.ToString();
-	FString SanitizedOwnerIDString = GetSanitizedSymbolName(OwnerIDString.Replace(TEXT("."), TEXT("_")));
+	FString SanitizedOwnerIDString = BuildSanitizedSymbolString(OwnerIDString.Replace(TEXT("."), TEXT("_")));
 
 	DIInstanceInfo.DataInterfaceHLSLSymbol = SanitizedOwnerIDString;
 	DIInstanceInfo.DIClassName = Info.Type.GetClass()->GetName();
@@ -3021,7 +3341,7 @@ void FNiagaraHlslTranslator::ConvertCompileInfoToParamInfo(const FNiagaraScriptD
 
 		FNiagaraDataInterfaceGeneratedFunction& DIFunc = DIInstanceInfo.GeneratedFunctions.AddDefaulted_GetRef();
 		DIFunc.DefinitionName = Sig.Name;
-		DIFunc.InstanceName = GetFunctionSignatureSymbol(Sig);
+		DIFunc.InstanceName = BuildFunctionSignatureString(Sig);
 		DIFunc.MiscUsageBitMask = Sig.MiscUsageBitMask;
 		DIFunc.Specifiers.Empty(Sig.FunctionSpecifiers.Num());
 		for (const TTuple<FName, FName>& Specifier : Sig.FunctionSpecifiers)
@@ -3229,7 +3549,7 @@ void TNiagaraHlslTranslator<GraphBridge>::DefineMainGPUFunctions(
 			const TArray<FNiagaraVariable>& NiagaraVariables = DataSetVariables[DataSetReads[DataSetID]];
 			for (const FNiagaraVariable& Var : NiagaraVariables)
 			{
-				FString VarFmt = ContextName + GetSanitizedSymbolName(Var.GetName().ToString()) + TEXT("{0} = {4};\n");
+				FString VarFmt = ContextName + BuildSanitizedSymbolName(Var.GetName()) + TEXT("{0} = {4};\n");
 				GatherVariableForDataSetAccess(Var, VarFmt, IntCounter, FloatCounter, HalfCounter, DataSetIndex, TEXT(""), HlslOutput);
 			}
 		}
@@ -3294,7 +3614,7 @@ void TNiagaraHlslTranslator<GraphBridge>::DefineMainGPUFunctions(
 				const TArray<FNiagaraVariable>& NiagaraVariables = DataSetVariables[DataSetReads[DataSetID]];
 				for (const FNiagaraVariable& Var : NiagaraVariables)
 				{
-					const FString VarName = ContextName + GetSanitizedSymbolName(Var.GetName().ToString());
+					const FString VarName = ContextName + BuildSanitizedSymbolName(Var.GetName());
 					FString VarFmt;
 					bool bWrite = true;
 
@@ -3330,8 +3650,8 @@ void TNiagaraHlslTranslator<GraphBridge>::DefineMainGPUFunctions(
 			for (FNiagaraVariable VarPrevious : GatheredPreviousVariables)
 			{
 				FNiagaraVariable SrcVar = FNiagaraParameterUtilities::GetSourceForPreviousValue(VarPrevious);
-				const FString VarName = ContextName + GetSanitizedSymbolName(SrcVar.GetName().ToString());
-				const FString VarPrevName = ContextName + GetSanitizedSymbolName(VarPrevious.GetName().ToString());
+				const FString VarName = ContextName + BuildSanitizedSymbolName(SrcVar.GetName());
+				const FString VarPrevName = ContextName + BuildSanitizedSymbolName(VarPrevious.GetName());
 				HlslOutput += VarPrevName + TEXT(" = ") + VarName + TEXT(";\n");
 			}
 
@@ -3395,8 +3715,8 @@ void TNiagaraHlslTranslator<GraphBridge>::DefineMainGPUFunctions(
 						if (FNiagaraParameterUtilities::IsPreviousValue(Var))
 						{
 							FNiagaraVariable SrcVar = FNiagaraParameterUtilities::GetSourceForPreviousValue(Var);
-							const FString VarName = GetSanitizedSymbolName(SrcVar.GetName().ToString());
-							const FString VarPrevName = GetSanitizedSymbolName(Var.GetName().ToString());
+							const FString VarName = BuildSanitizedSymbolName(SrcVar.GetName());
+							const FString VarPrevName = BuildSanitizedSymbolName(Var.GetName());
 							HlslOutput += TEXT("\t\tContext.") + TranslationStages[i].PassNamespace + TEXT(".") + VarPrevName + TEXT(" = Context.") + TranslationStages[i-1].PassNamespace + TEXT(".") + VarName + TEXT(";\n");
 						}
 					}
@@ -3488,7 +3808,7 @@ void TNiagaraHlslTranslator<GraphBridge>::DefineMainGPUFunctions(
 					const bool bWriteToHLSL = !TranslationStages[i].bPartialParticleUpdate || TranslationStages[i].SetParticleAttributes.Contains(Var);
 
 					// If coming from a parameter map, use the one on the context, otherwise use the output.
-					FString VarFmt = TEXT("\t\t\tOutputData{1}(0, {2}, {3}, ") + ContextName + GetSanitizedSymbolName(Var.GetName().ToString()) + TEXT("{0});\n");
+					FString VarFmt = TEXT("\t\t\tOutputData{1}(0, {2}, {3}, ") + ContextName + BuildSanitizedSymbolName(Var.GetName()) + TEXT("{0});\n");
 					GatherVariableForDataSetAccess(Var, VarFmt, IntCounter, FloatCounter, HalfCounter, -1, TEXT("WriteIndex"), HlslOutput, bWriteToHLSL);
 				}
 			}
@@ -3890,8 +4210,8 @@ void FNiagaraHlslTranslator::DefineMain(FString &OutHlslOutput,
 						if (FNiagaraParameterUtilities::IsPreviousValue(Var))
 						{
 							FNiagaraVariable SrcVar = FNiagaraParameterUtilities::GetSourceForPreviousValue(Var);
-							const FString VarName =  GetSanitizedSymbolName(SrcVar.GetName().ToString());
-							const FString VarPrevName =  GetSanitizedSymbolName(Var.GetName().ToString());
+							const FString VarName = BuildSanitizedSymbolName(SrcVar.GetName());
+							const FString VarPrevName = BuildSanitizedSymbolName(Var.GetName());
 							OutHlslOutput += TEXT("\t\tContext.") + TranslationStages[StageIdx + 1].PassNamespace + TEXT(".") + VarPrevName + TEXT(" = Context.") + TranslationStages[StageIdx].PassNamespace + TEXT(".")+ VarName +  TEXT(";\n");
 						}
 					}
@@ -4005,11 +4325,11 @@ void TNiagaraHlslTranslator<GraphBridge>::DefineDataSetVariableWrites(FString &O
 		FString Fmt;
 		if (TranslationStages.Num() > 1)
 		{
-			Fmt = TEXT("\tOutputData{1}(0, {2}, {3}, Context.") + TranslationStages[TranslationStages.Num() - 1].PassNamespace + TEXT(".") + GetSanitizedSymbolName(Var.GetName().ToString()) + TEXT("{0});\n");
+			Fmt = TEXT("\tOutputData{1}(0, {2}, {3}, Context.") + TranslationStages[TranslationStages.Num() - 1].PassNamespace + TEXT(".") + BuildSanitizedSymbolName(Var.GetName()) + TEXT("{0});\n");
 		}
 		else
 		{
-			Fmt = TEXT("\tOutputData{1}(0, {2}, {3}, Context.Map.") + GetSanitizedSymbolName(Var.GetName().ToString()) + TEXT("{0});\n");
+			Fmt = TEXT("\tOutputData{1}(0, {2}, {3}, Context.Map.") + BuildSanitizedSymbolName(Var.GetName()) + TEXT("{0});\n");
 		}
 		GatherVariableForDataSetAccess(Var, Fmt, WriteOffsetInt, WriteOffsetFloat, WriteOffsetHalf, -1, TEXT("TmpWriteIndex"), OutHlslOutput);
 	}
@@ -4055,7 +4375,7 @@ void FNiagaraHlslTranslator::DefineDataSetVariableReads(FString &OutHlslOutput,
 
 		for (const FNiagaraVariable &Var : ReadVars)
 		{
-			Fmt = ContextName + GetSanitizedSymbolName(Var.GetName().ToString()) + TEXT("{0} = {4};\n");
+			Fmt = ContextName + BuildSanitizedSymbolName(Var.GetName()) + TEXT("{0} = {4};\n");
 			GatherVariableForDataSetAccess(Var, Fmt, ReadOffsetInt, ReadOffsetFloat, ReadOffsetHalf, DataSetIndex, TEXT(""), VarReads);
 		}
 
@@ -4087,7 +4407,7 @@ void FNiagaraHlslTranslator::DefineDataSetVariableReads(FString &OutHlslOutput,
 		for (const FNiagaraVariable &Var : ReadVars)
 		{
 			bool bWrite = true;
-			const FString VariableName = ContextName + GetSanitizedSymbolName(Var.GetName().ToString());
+			const FString VariableName = ContextName + BuildSanitizedSymbolName(Var.GetName());
 			// If the NiagaraClearEachFrame value is set on the data set, we don't bother reading it in each frame as we know that it is is invalid. However,
 			// this is only used for the base data set. Other reads are potentially from events and are therefore perfectly valid.
 			if (DataSetIndex == 0 && Var.GetType().GetScriptStruct() != nullptr && Var.GetType().GetScriptStruct()->GetMetaData(TEXT("NiagaraClearEachFrame")).Equals(TEXT("true"), ESearchCase::IgnoreCase))
@@ -4120,8 +4440,8 @@ void FNiagaraHlslTranslator::DefineDataSetVariableReads(FString &OutHlslOutput,
 		for (FNiagaraVariable VarPrevious : GatheredPreviousVariables)
 		{
 			FNiagaraVariable SrcVar = FNiagaraParameterUtilities::GetSourceForPreviousValue(VarPrevious);
-			const FString VarName = ContextName + GetSanitizedSymbolName(SrcVar.GetName().ToString());
-			const FString VarPrevName = ContextName + GetSanitizedSymbolName(VarPrevious.GetName().ToString());
+			const FString VarName = ContextName + BuildSanitizedSymbolName(SrcVar.GetName());
+			const FString VarPrevName = ContextName + BuildSanitizedSymbolName(VarPrevious.GetName());
 			HlslOutput += VarPrevName + TEXT(" = ") + VarName + TEXT(";\n");
 		}
 
@@ -4236,64 +4556,33 @@ void FNiagaraHlslTranslator::DecomposeVariableAccess(UStruct* Struct, bool bRead
 	}
 };
 
-FString FNiagaraHlslTranslator::GetSanitizedSymbolName(FStringView SymbolName, bool bCollapsNamespaces)
+FString FNiagaraHlslTranslator::GetSanitizedSymbolName(FStringView SymbolName, bool bCollapseNamespaces)
 {
-	if (SymbolName.Len() == 0)
-	{
-		return FString(SymbolName);
-	}
-
-	const UNiagaraEditorSettings* Settings = GetDefault<UNiagaraEditorSettings>();
-	check(Settings);
-	const TMap<FString, FString>& ReplacementsForInvalid = Settings->GetHLSLKeywordReplacementsMap();
-
-	FString Ret = FString(SymbolName);
-
-	// Split up into individual namespaces...
-	TArray<FString> SplitName;
-	Ret.ParseIntoArray(SplitName, TEXT("."));
-
-	// Rules for variable namespaces..
-	for (int32 i = 0; i < SplitName.Num(); i++)
-	{
-		if (SplitName[i][0] >= TCHAR('0') && SplitName[i][0] <= TCHAR('9')) // Cannot start with a numeric digit
-		{
-			SplitName[i] = TEXT("INTEGER_") + SplitName[i];
-		}
-
-		const FString* FoundReplacementStr = ReplacementsForInvalid.Find(SplitName[i]); // Look for the string in the keyword protections array.
-		if (FoundReplacementStr)
-		{
-			SplitName[i] = *FoundReplacementStr;
-		}
-
-		SplitName[i].ReplaceInline(TEXT("\t"), TEXT(""));
-		SplitName[i].ReplaceInline(TEXT(" "), TEXT(""));
+	const UNiagaraEditorSettings* EditorSettings = GetDefault<UNiagaraEditorSettings>();
 
-		// Handle internationalization of characters..
-		SplitName[i] = ConvertToAsciiString(SplitName[i]);
-	}
+	TStringBuilder<1024> StringBuilder;
+	NiagaraHlslTranslatorImpl::AppendSanitizedSymbolString(SymbolName, bCollapseNamespaces, EditorSettings->GetHLSLKeywordReplacementsMap(), StringBuilder);
 
-	// Gather back into single string..
-	Ret = FString::Join(SplitName, TEXT("."));
+	return StringBuilder.ToString();
+}
 
-	/*
-	Ret.ReplaceInline(TEXT("\\"), TEXT("_"));
-	Ret.ReplaceInline(TEXT("/"), TEXT("_"));
-	Ret.ReplaceInline(TEXT(","), TEXT("_"));
-	Ret.ReplaceInline(TEXT("-"), TEXT("_"));
-	Ret.ReplaceInline(TEXT(":"), TEXT("_"));
-	Ret = Ret.ConvertTabsToSpaces(0);
-	*/
+FString FNiagaraHlslTranslator::BuildSanitizedSymbolString(FStringView SymbolName, bool bCollapseNamespaces) const
+{
+	TStringBuilder<1024> StringBuilder;
+	AppendSanitizedSymbolString(StringBuilder, SymbolName, bCollapseNamespaces);
 
-	if (bCollapsNamespaces)
-	{
-		Ret.ReplaceInline(TEXT("."), TEXT("_"));
-	}
+	return StringBuilder.ToString();
+}
 
-	Ret.ReplaceInline(TEXT("__"), TEXT("ASC95ASC95")); // Opengl reserves "__" within a name
+FString FNiagaraHlslTranslator::BuildSanitizedSymbolName(FName SymbolName, bool bCollapseNamespaces) const
+{
+	FNameBuilder NameBuilder(SymbolName);
+	return BuildSanitizedSymbolString(NameBuilder.ToView(), bCollapseNamespaces);
+}
 
-	return Ret;
+void FNiagaraHlslTranslator::AppendSanitizedSymbolString(FStringBuilderBase& Builder, FStringView SymbolName, bool bCollapseNamespaces) const
+{
+	return NiagaraHlslTranslatorImpl::AppendSanitizedSymbolString(SymbolName, bCollapseNamespaces, KeywordReplacementMap, Builder);
 }
 
 FString FNiagaraHlslTranslator::GetSanitizedDIFunctionName(const FString& FunctionName)
@@ -4345,33 +4634,19 @@ FString FNiagaraHlslTranslator::GetSanitizedFunctionNameSuffix(FString Name)
 	return ConvertToAsciiString(Ret);
 }
 
+
 FString FNiagaraHlslTranslator::ConvertToAsciiString(FString Str)
 {
-	FString AsciiString;
-	AsciiString.Reserve(Str.Len() * 6); // Assign room for every current char to be 'ASCXXX'
-	for (int32 j = 0; j < Str.Len(); j++)
-	{
-		if ((Str[j] >= TCHAR('0') && Str[j] <= TCHAR('9')) ||
-			(Str[j] >= TCHAR('A') && Str[j] <= TCHAR('Z')) ||
-			(Str[j] >= TCHAR('a') && Str[j] <= TCHAR('z')) ||
-			Str[j] == TCHAR('_') || Str[j] == TCHAR(' '))
-		{
-			// Do nothing.. these are valid chars..
-			AsciiString.AppendChar(Str[j]);
-		}
-		else
-		{
-			// Need to replace the bad characters..
-			AsciiString.Append(TEXT("ASC"));
-			AsciiString.AppendInt((int32)Str[j]);
-		}
-	}
-	return AsciiString;
+	FNameBuilder Builder;
+	NiagaraHlslTranslatorImpl::AppendConvertedToAscii(Str, Builder);
+
+	return Builder.ToString();
 }
 
 FString FNiagaraHlslTranslator::GetUniqueSymbolName(FName BaseName)
 {
-	FString RetString = GetSanitizedSymbolName(BaseName.ToString());
+	FNameBuilder BaseNameBuilder(BaseName);
+	FString RetString = BuildSanitizedSymbolString(BaseNameBuilder.ToView());
 	FName RetName = *RetString;
 	uint32* NameCount = SymbolCounts.Find(RetName);
 	if (NameCount == nullptr)
@@ -4530,7 +4805,7 @@ int32 FNiagaraHlslTranslator::AddChunkToConstantBuffer(const FString& SymbolName
 	check(FNiagaraTypeHelper::IsLWCType(Type) == false);
 	int32 ChunkIndex = CodeChunks.AddDefaulted();
 	FNiagaraCodeChunk& Chunk = CodeChunks[ChunkIndex];
-	Chunk.SymbolName = GetSanitizedSymbolName(SymbolName.Replace(TEXT("."), TEXT("_")));
+	Chunk.SymbolName = BuildSanitizedSymbolString(SymbolName, true);
 	Chunk.Type = Type;
 	Chunk.Original = InVariable;
 	Chunk.Mode = ChunkMode;
@@ -4558,7 +4833,10 @@ int32 FNiagaraHlslTranslator::RegisterUniformChunkToPack(const FString& SymbolNa
 		const int32 ConflictIndex = Parameters.Parameters.IndexOfByPredicate([&](const FNiagaraVariable& ExistingParameter)
 		{
 			FNameBuilder ExistingParameterName(ExistingParameter.GetName());
-			return GetSanitizedSymbolName(ExistingParameterName).Equals(SymbolName);
+			FNameBuilder SanitizedParameterName;
+			AppendSanitizedSymbolString(SanitizedParameterName, ExistingParameterName);
+
+			return ExistingParameterName.ToView().Equals(SymbolName);
 		});
 
 		if (ConflictIndex != INDEX_NONE)
@@ -4635,10 +4913,8 @@ void FNiagaraHlslTranslator::PackRegisteredUniformChunk(FNiagaraParameters& Para
 				FName InstanceName = PaddingName;
 				InstanceName.SetNumber(UniquePaddingIndex++);
 
-				FNameBuilder InstanceNameBuilder(InstanceName);
-
 				FNiagaraVariable PaddingVariable(FNiagaraTypeDefinition::GetIntDef(), InstanceName);
-				FString SymbolName = GetSanitizedSymbolName(InstanceNameBuilder);
+				FString SymbolName = BuildSanitizedSymbolName(InstanceName);
 
 				const int32 PaddingVariableIndex = Parameters.Parameters.Num();
 				RegisterUniformChunkToPack(SymbolName, PaddingVariable, true /*AddPadding*/, Parameters, ConflictingVariable);
@@ -4723,7 +4999,7 @@ int32 FNiagaraHlslTranslator::AddSourceChunk(FString SymbolName, const FNiagaraT
 		check(FNiagaraTypeHelper::IsLWCType(Type) == false);
 		Ret = CodeChunks.AddDefaulted();
 		FNiagaraCodeChunk& Chunk = CodeChunks[Ret];
-		Chunk.SymbolName = bSanitize ? GetSanitizedSymbolName(SymbolName) : SymbolName;
+		Chunk.SymbolName = bSanitize ? BuildSanitizedSymbolString(SymbolName) : SymbolName;
 		Chunk.Type = Type;
 		Chunk.Mode = ENiagaraCodeChunkMode::Source;
 
@@ -4733,30 +5009,30 @@ int32 FNiagaraHlslTranslator::AddSourceChunk(FString SymbolName, const FNiagaraT
 }
 
 
-int32 FNiagaraHlslTranslator::AddBodyComment(const FString& Comment)
+int32 FNiagaraHlslTranslator::AddBodyComment(FStringView Comment)
 {
-	return AddBodyChunk(TEXT(""), Comment, FNiagaraTypeDefinition::GetIntDef(), false, false);
+	return AddBodyChunk(TEXT(""), Comment, FNiagaraTypeDefinition::GetIntDef(), TArrayView<int32>(), false, false);
 }
 
-int32 FNiagaraHlslTranslator::AddBodyChunk(const FString& Value)
+int32 FNiagaraHlslTranslator::AddBodyChunk(FStringView Value)
 {
 	return AddBodyChunk(TEXT(""), Value, FNiagaraTypeDefinition::GetIntDef(), INDEX_NONE, false, false);
 }
 
-int32 FNiagaraHlslTranslator::AddBodyChunk(FString SymbolName, FString Definition, const FNiagaraTypeDefinition& Type, TArray<int32>& SourceChunks, bool bDecl, bool bIsTerminated)
+int32 FNiagaraHlslTranslator::AddBodyChunk(FStringView SymbolName, FStringView Definition, const FNiagaraTypeDefinition& Type, TArrayView<int32> SourceChunks, bool bDecl, bool bIsTerminated)
 {
 	check(CurrentBodyChunkMode == ENiagaraCodeChunkMode::Body || CurrentBodyChunkMode == ENiagaraCodeChunkMode::SpawnBody || CurrentBodyChunkMode == ENiagaraCodeChunkMode::UpdateBody ||
 		(CurrentBodyChunkMode >= ENiagaraCodeChunkMode::SimulationStageBody && CurrentBodyChunkMode < ENiagaraCodeChunkMode::SimulationStageBodyMax));
 	check(FNiagaraTypeHelper::IsLWCType(Type) == false);
 	int32 Ret = CodeChunks.AddDefaulted();
 	FNiagaraCodeChunk& Chunk = CodeChunks[Ret];
-	Chunk.SymbolName = GetSanitizedSymbolName(SymbolName);
+	Chunk.SymbolName = BuildSanitizedSymbolString(SymbolName);
 	Chunk.Definition = Definition;
 	Chunk.Type = Type;
 	Chunk.bDecl = bDecl;
 	Chunk.bIsTerminated = bIsTerminated;
 	Chunk.Mode = CurrentBodyChunkMode;
-	Chunk.SourceChunks = SourceChunks;
+	Chunk.SourceChunks.Append(SourceChunks);
 
 	ChunksByMode[(int32)CurrentBodyChunkMode].Add(Ret);
 	return Ret;
@@ -4764,41 +5040,14 @@ int32 FNiagaraHlslTranslator::AddBodyChunk(FString SymbolName, FString Definitio
 
 
 
-int32 FNiagaraHlslTranslator::AddBodyChunk(FString SymbolName, FString Definition, const FNiagaraTypeDefinition& Type, int32 SourceChunk, bool bDecl, bool bIsTerminated)
+int32 FNiagaraHlslTranslator::AddBodyChunk(FStringView SymbolName, FStringView Definition, const FNiagaraTypeDefinition& Type, int32 SourceChunk, bool bDecl, bool bIsTerminated)
 {
-	check(CurrentBodyChunkMode == ENiagaraCodeChunkMode::Body || CurrentBodyChunkMode == ENiagaraCodeChunkMode::SpawnBody || CurrentBodyChunkMode == ENiagaraCodeChunkMode::UpdateBody ||
-		(CurrentBodyChunkMode >= ENiagaraCodeChunkMode::SimulationStageBody && CurrentBodyChunkMode < ENiagaraCodeChunkMode::SimulationStageBodyMax));
-	check(FNiagaraTypeHelper::IsLWCType(Type) == false);
-	int32 Ret = CodeChunks.AddDefaulted();
-	FNiagaraCodeChunk& Chunk = CodeChunks[Ret];
-	Chunk.SymbolName = GetSanitizedSymbolName(SymbolName);
-	Chunk.Definition = Definition;
-	Chunk.Type = Type;
-	Chunk.bDecl = bDecl;
-	Chunk.bIsTerminated = bIsTerminated;
-	Chunk.Mode = CurrentBodyChunkMode;
-	Chunk.SourceChunks.Add(SourceChunk);
-
-	ChunksByMode[(int32)CurrentBodyChunkMode].Add(Ret);
-	return Ret;
+	return AddBodyChunk(SymbolName, Definition, Type, MakeArrayView( &SourceChunk, 1), bDecl, bIsTerminated);
 }
 
-int32 FNiagaraHlslTranslator::AddBodyChunk(FString SymbolName, FString Definition, const FNiagaraTypeDefinition& Type, bool bDecl, bool bIsTerminated)
+int32 FNiagaraHlslTranslator::AddBodyChunk(FStringView SymbolName, FStringView Definition, const FNiagaraTypeDefinition& Type, bool bDecl, bool bIsTerminated)
 {
-	check(CurrentBodyChunkMode == ENiagaraCodeChunkMode::Body || CurrentBodyChunkMode == ENiagaraCodeChunkMode::SpawnBody || CurrentBodyChunkMode == ENiagaraCodeChunkMode::UpdateBody ||
-			(CurrentBodyChunkMode >= ENiagaraCodeChunkMode::SimulationStageBody && CurrentBodyChunkMode < ENiagaraCodeChunkMode::SimulationStageBodyMax));
-	check(FNiagaraTypeHelper::IsLWCType(Type) == false);
-	int32 Ret = CodeChunks.AddDefaulted();
-	FNiagaraCodeChunk& Chunk = CodeChunks[Ret];
-	Chunk.SymbolName = GetSanitizedSymbolName(SymbolName);
-	Chunk.Definition = Definition;
-	Chunk.Type = Type;
-	Chunk.bDecl = bDecl;
-	Chunk.bIsTerminated = bIsTerminated;
-	Chunk.Mode = CurrentBodyChunkMode;
-
-	ChunksByMode[(int32)CurrentBodyChunkMode].Add(Ret);
-	return Ret;
+	return AddBodyChunk(SymbolName, Definition, Type, TArrayView<int32>(), bDecl, bIsTerminated);
 }
 
 bool FNiagaraHlslTranslator::ShouldInterpolateParameter(const FNiagaraVariable& Parameter)
@@ -4850,7 +5099,7 @@ bool FNiagaraHlslTranslator::ShouldInterpolateParameter(const FNiagaraVariable&
 
 FString FNiagaraHlslTranslator::GetInterpolateHlsl(const FNiagaraVariable& Parameter, const FString& PrevMapName, const FNiagaraCodeChunk& Chunk) const
 {
-	FString FoundName = GetSanitizedSymbolName(Parameter.GetName().ToString());
+	FString FoundName = BuildSanitizedSymbolName(Parameter.GetName());
 	if (Parameter.GetType() == FNiagaraTypeDefinition::GetQuatDef() && CompileOptions.AdditionalDefines.Contains(FNiagaraCompileOptions::AccurateQuatInterpolation))
 	{
 		return FString::Format(TEXT("\tContext.{0}.{1} = NiagaraQuatSLerp({2}{3}{4}, {3}{4}, SpawnInterp);\n"), {PrevMapName, FoundName, INTERPOLATED_PARAMETER_PREFIX, Chunk.SymbolName, Chunk.ComponentMask});
@@ -4999,7 +5248,7 @@ int32 TNiagaraHlslTranslator<GraphBridge>::GetParameter(const FNiagaraVariable&
 				return FuncParam;
 			}
 			//If this is a valid function parameter, use that.
-			FString SymbolName = TEXT("In_") + GetSanitizedSymbolName(Parameter.GetName().ToString());
+			FString SymbolName = TEXT("In_") + BuildSanitizedSymbolName(Parameter.GetName());
 			return AddSourceChunk(SymbolName, ConvertToSimulationVariable(Parameter).GetType());
 		}
 	}
@@ -5017,7 +5266,7 @@ int32 TNiagaraHlslTranslator<GraphBridge>::GetParameter(const FNiagaraVariable&
 	// We don't pass in the input node here (really there could be multiple nodes for the same parameter)
 	// so we have to match up the input parameter map variable value through the pre-traversal histories 
 	// so that we know which parameter map we are referencing.
-	FString SymbolName = GetSanitizedSymbolName(Parameter.GetName().ToString());
+	FString SymbolName = BuildSanitizedSymbolName(Parameter.GetName());
 	if (Parameter.GetType() == FNiagaraTypeDefinition::GetParameterMapDef())
 	{
 		if (ParamMapHistories.Num() == 0)
@@ -5534,7 +5783,7 @@ void TNiagaraHlslTranslator<GraphBridge>::Output(const FOutputNode* OutputNode,
 
 			if (Attributes[i].GetType() != FNiagaraTypeDefinition::GetParameterMapDef())
 			{
-				FString SymbolName = *GetSanitizedSymbolName((TEXT("Out_") + Attributes[i].GetName().ToString()));
+				FString SymbolName = *BuildSanitizedSymbolString(TEXT("Out_") + Attributes[i].GetName().ToString());
 				ENiagaraCodeChunkMode OldMode = CurrentBodyChunkMode;
 				CurrentBodyChunkMode = ENiagaraCodeChunkMode::Body;
 				AddBodyChunk(SymbolName, TEXT("{0}"), Attributes[i].GetType(), Inputs[i], false);
@@ -5567,7 +5816,7 @@ void TNiagaraHlslTranslator<GraphBridge>::Output(const FOutputNode* OutputNode,
 				{
 					FNiagaraVariable VarNamespaced = FNiagaraParameterUtilities::BasicAttributeToNamespacedAttribute(Var);
 					FString ParameterMapInstanceName = GetParameterMapInstanceName(0);
-					int32 ChunkIdx = AddBodyChunk(ParameterMapInstanceName + TEXT(".") + GetSanitizedSymbolName(VarNamespaced.GetName().ToString()), TEXT("{0}"), VarNamespaced.GetType(), Input, false);
+					int32 ChunkIdx = AddBodyChunk(ParameterMapInstanceName + TEXT(".") + BuildSanitizedSymbolName(VarNamespaced.GetName()), TEXT("{0}"), VarNamespaced.GetType(), Input, false);
 
 					// Make sure that we end up in the list of Attributes that have been written to by this script.
 					if (ParamMapDefinedAttributesToUniformChunks.Find(Var.GetName()) == nullptr)
@@ -5607,7 +5856,7 @@ int32 FNiagaraHlslTranslator::GetAttribute(const FNiagaraVariable& Attribute)
 			FString ParameterMapInstanceName = GetParameterMapInstanceName(0);
 
 			FNiagaraVariable NamespacedVar = Attribute;
-			FString SymbolName = *(ParameterMapInstanceName + TEXT(".") + GetSanitizedSymbolName(NamespacedVar.GetName().ToString()));
+			FString SymbolName = *(ParameterMapInstanceName + TEXT(".") + BuildSanitizedSymbolName(NamespacedVar.GetName()));
 			return AddSourceChunk(SymbolName, Attribute.GetType());
 		}
 		else
@@ -5977,7 +6226,7 @@ void TNiagaraHlslTranslator<GraphBridge>::ParameterMapSet(const FParamMapSetNode
 					TranslationStages[ActiveStageIdx].bPartialParticleUpdate = false;
 					TranslationStages[ActiveStageIdx].bWritesAlive = true;
 				}
-				AddBodyChunk(ParameterMapInstanceName + TEXT(".") + GetSanitizedSymbolName(Var.GetName().ToString()), TEXT("{0}"), Var.GetType(), Input, false);
+				AddBodyChunk(ParameterMapInstanceName + TEXT(".") + BuildSanitizedSymbolName(Var.GetName()), TEXT("{0}"), Var.GetType(), Input, false);
 			}
 		}
 	}
@@ -5987,7 +6236,7 @@ FString FNiagaraHlslTranslator::GetUniqueEmitterName() const
 {
 	if (CompileOptions.TargetUsage == ENiagaraScriptUsage::SystemSpawnScript || CompileOptions.TargetUsage == ENiagaraScriptUsage::SystemUpdateScript)
 	{
-		return GetEmitterUniqueName();
+		return FString(GetEmitterUniqueName());
 	}
 	else
 	{
@@ -6139,8 +6388,7 @@ bool TNiagaraHlslTranslator<GraphBridge>::ParameterMapRegisterExternalConstantNa
 		return false;
 	}
 
-	FString VarName = InVariable.GetName().ToString();
-	FString SymbolName = GetSanitizedSymbolName(VarName);
+	FString SymbolName = BuildSanitizedSymbolName(InVariable.GetName());
 	FString ParameterMapInstanceName = GetParameterMapInstanceName(InParamMapHistoryIdx);
 
 	bool bMissingParameter = false;
@@ -6251,7 +6499,7 @@ bool TNiagaraHlslTranslator<GraphBridge>::ParameterMapRegisterExternalConstantNa
 			else
 			{
 				//Add this separately as the same uniform can appear in the pre sim chunks more than once in different param maps.
-				PerStageMainPreSimulateChunks[ActiveStageIdx].AddUnique(FString::Printf(TEXT("%s.%s = %s;"), *ParameterMapInstanceName, *GetSanitizedSymbolName(VarName), *GetCodeAsSource(UniformChunk)));
+				PerStageMainPreSimulateChunks[ActiveStageIdx].AddUnique(FString::Printf(TEXT("%s.%s = %s;"), *ParameterMapInstanceName, *SymbolName, *GetCodeAsSource(UniformChunk)));
 			}
 		}
 		else if (bIsPerInstanceBulkSystemParam && !ExternalVariablesForBulkUsage.Contains(InVariable))
@@ -6519,8 +6767,7 @@ void TNiagaraHlslTranslator<GraphBridge>::ValidateParticleIDUsage()
 template<typename GraphBridge>
 bool TNiagaraHlslTranslator<GraphBridge>::ParameterMapRegisterNamespaceAttributeVariable(const FNiagaraVariable& InVariable, const FNode* InNode, int32 InParamMapHistoryIdx, int32& Output)
 {
-	FString VarName = InVariable.GetName().ToString();
-	FString SymbolNameNamespaced = GetSanitizedSymbolName(VarName);
+	FString SymbolNameNamespaced = BuildSanitizedSymbolName(InVariable.GetName());
 	FString ParameterMapInstanceName = GetParameterMapInstanceName(InParamMapHistoryIdx);
 	FNiagaraVariable NamespaceVar = InVariable;
 
@@ -6530,7 +6777,7 @@ bool TNiagaraHlslTranslator<GraphBridge>::ParameterMapRegisterNamespaceAttribute
 	{
 		if (false == ParamMapDefinedAttributesToUniformChunks.Contains(BasicVar.GetName()))
 		{
-			FString SymbolNameDefined = GetSanitizedSymbolName(BasicVar.GetName().ToString());
+			FString SymbolNameDefined = BuildSanitizedSymbolName(BasicVar.GetName());
 			int32 UniformChunk = INDEX_NONE;
 			int32 Idx = InstanceRead.Variables.Find(NamespaceVar);
 			if (Idx != INDEX_NONE)
@@ -6691,7 +6938,7 @@ void TNiagaraHlslTranslator<GraphBridge>::Emitter(const FEmitterNode* EmitterNod
 			ParamMapDefinedEngineVars.Add(EmitterIDVar);
 			History.AddVariable(EmitterIDVar, SYS_PARAM_ENGINE_EMITTER_ID, NAME_None, nullptr);
 			FNiagaraEmitterID EmitterID = GraphBridge::GetEmitterID(EmitterNode);
-			AddBodyChunk(TEXT(""), FString::Printf(TEXT("%s.%s.ID = %d"), *GetParameterMapInstanceName(ParamMapHistoryIdx), *GetSanitizedSymbolName(EmitterIDVar.GetName().ToString()), EmitterID.ID), FNiagaraTypeDefinition::GetIntDef(), false);
+			AddBodyChunk(TEXT(""), FString::Printf(TEXT("%s.%s.ID = %d"), *GetParameterMapInstanceName(ParamMapHistoryIdx), *BuildSanitizedSymbolName(EmitterIDVar.GetName()), EmitterID.ID), FNiagaraTypeDefinition::GetIntDef(), false);
 
 			for (int32 i = 0; i < ParamMapHistories[ParamMapHistoryIdx].Variables.Num(); i++)
 			{
@@ -6862,8 +7109,7 @@ void TNiagaraHlslTranslator<GraphBridge>::HandleParameterRead(int32 ParamMapHist
 	}
 
 	// If this is a System parameter, just wire in the system appropriate system attribute.
-	FString VarName = Var.GetName().ToString();
-	FString SymbolName = GetSanitizedSymbolName(VarName);
+	FString SymbolName = BuildSanitizedSymbolName(Var.GetName());
 
 	bool bIsPerInstanceAttribute = false;
 	bool bIsCandidateForRapidIteration = false;
@@ -7118,8 +7364,8 @@ void TNiagaraHlslTranslator<GraphBridge>::HandleParameterRead(int32 ParamMapHist
 				if (bFoundExistingSet)
 				{
 					LastSetChunkIdx = AddBodyChunk(
-						ParameterMapInstanceName + TEXT(".") + GetSanitizedSymbolName(Var.GetName().ToString()),
-						ParameterMapInstanceName + TEXT(".") + GetSanitizedSymbolName(SourceForInitialValue.GetName().ToString()), Var.GetType(),
+						ParameterMapInstanceName + TEXT(".") + BuildSanitizedSymbolName(Var.GetName()),
+						ParameterMapInstanceName + TEXT(".") + BuildSanitizedSymbolName(SourceForInitialValue.GetName()), Var.GetType(),
 						false
 					);
 
@@ -7263,7 +7509,7 @@ void TNiagaraHlslTranslator<GraphBridge>::HandleParameterRead(int32 ParamMapHist
 				{
 					if (Var.GetType().GetClass() == nullptr) // Only need to do this wiring for things that aren't data interfaces.
 					{
-						AddBodyChunk(ParameterMapInstanceName + TEXT(".") + GetSanitizedSymbolName(Var.GetName().ToString()), TEXT("{0}"), Var.GetType(), LastSetChunkIdx, false);
+						AddBodyChunk(ParameterMapInstanceName + TEXT(".") + BuildSanitizedSymbolName(Var.GetName()), TEXT("{0}"), Var.GetType(), LastSetChunkIdx, false);
 					}
 				}
 			}
@@ -7313,7 +7559,7 @@ void TNiagaraHlslTranslator<GraphBridge>::HandleParameterRead(int32 ParamMapHist
 	}
 	else
 	{
-		OutputChunkId = AddSourceChunk(ParameterMapInstanceName + TEXT(".") + GetSanitizedSymbolName(Var.GetName().ToString()), Var.GetType());
+		OutputChunkId = AddSourceChunk(ParameterMapInstanceName + TEXT(".") + BuildSanitizedSymbolName(Var.GetName()), Var.GetType());
 		RecordParamMapDefinedAttributeToNamespaceVar(Var, DefaultPin);
 	}
 }
@@ -7373,7 +7619,7 @@ void TNiagaraHlslTranslator<GraphBridge>::ReadDataSet(const FNiagaraDataSetID Da
 			{
 				Error(FText::Format(LOCTEXT("GetConstantFailTypeVar", "Cannot handle type {0}! Variable: {1}"), Var.GetType().GetNameText(), FText::FromName(Var.GetName())), nullptr, nullptr);
 			}
-			DataSetReadForInput->CodeChunks.Add(AddSourceChunk(DataSetAccessSymbol + TEXT(".") + GetSanitizedSymbolName(Var.GetName().ToString()), Var.GetType()));
+			DataSetReadForInput->CodeChunks.Add(AddSourceChunk(DataSetAccessSymbol + TEXT(".") + BuildSanitizedSymbolName(Var.GetName()), Var.GetType()));
 		}
 		Outputs.Add(ParamMapHistoryIdx);
 		Outputs.Append(DataSetReadForInput->CodeChunks);
@@ -7449,7 +7695,7 @@ void TNiagaraHlslTranslator<GraphBridge>::WriteDataSet(const FNiagaraDataSetID D
 	{
 		const FNiagaraVariable& Var = Variables[i];
 		int32 Input = Inputs[i + 2];//input 0 is the valid input (no entry in variables array), so we need of offset all other inputs by 1.
-		DataSetWriteForInput->CodeChunks.Add(AddBodyChunk(DataSetAccessName + TEXT(".") + GetSanitizedSymbolName(Var.GetName().ToString()), TEXT("{0}"), Var.GetType(), Input, false));
+		DataSetWriteForInput->CodeChunks.Add(AddBodyChunk(DataSetAccessName + TEXT(".") + BuildSanitizedSymbolName(Var.GetName()), TEXT("{0}"), Var.GetType(), Input, false));
 	}
 }
 
@@ -8125,7 +8371,7 @@ void TNiagaraHlslTranslator<GraphBridge>::FinalResolveNamespacedTokens(const FSt
 					}
 
 
-					Tokens[i] = ParameterMapInstanceNamespace + GetSanitizedSymbolName(Var.GetName().ToString());
+					Tokens[i] = ParameterMapInstanceNamespace + BuildSanitizedSymbolName(Var.GetName());
 					break;
 				}
 			}
@@ -8457,7 +8703,7 @@ void TNiagaraHlslTranslator<GraphBridge>::ProcessCustomHlsl(const FString& InCus
 				}
 
 				// Now we can build the function name and replace the method call token with the final function name.
-				FunctionNameToken = GetFunctionSignatureSymbol(Sig);
+				FunctionNameToken = BuildFunctionSignatureString(Sig);
 				if (Sig.bRequiresExecPin)
 				{
 					Sig.Inputs.Insert(FNiagaraVariable(FNiagaraTypeDefinition::GetParameterMapDef(), TEXT("InExecPin")), 0);
@@ -8500,7 +8746,7 @@ void TNiagaraHlslTranslator<GraphBridge>::ProcessCustomHlsl(const FString& InCus
 
 			if (bPermuteSignatureByDataInterface)
 			{
-				OutSignature.Name = FName(OutSignature.Name.ToString() + GetSanitizedSymbolName(Info.Name.ToString(), true));
+				OutSignature.Name = FName(OutSignature.Name.ToString() + BuildSanitizedSymbolName(Info.Name, true));
 			}
 
 			SigInputs.Add(Input);
@@ -8627,7 +8873,7 @@ void TNiagaraHlslTranslator<GraphBridge>::HandleCustomHlslNode(const FCustomHlsl
 
 	// Determine the important outputs
 	OutScriptUsage = GraphBridge::GetCustomHlslUsage(CustomFunctionHlsl);
-	OutName = GetSanitizedSymbolName(CustomFunctionHlsl->Signature.Name.ToString() + CustomFunctionHlsl->NodeGuid.ToString());
+	OutName = BuildSanitizedSymbolString(CustomFunctionHlsl->Signature.Name.ToString() + CustomFunctionHlsl->NodeGuid.ToString());
 	OutSignature = CustomFunctionHlsl->Signature;
 	OutFullName = GraphBridge::GetFunctionFullName(CustomFunctionHlsl);
 	OutSignature.Name = *OutName; // Force the name to be set to include the node guid for safety...
@@ -9127,14 +9373,14 @@ void FNiagaraHlslTranslator::GenerateFunctionCall(ENiagaraScriptUsage ScriptUsag
 	if (ScriptUsage == ENiagaraScriptUsage::Module)
 	{
 		bEnteredStatScope = true;
-		EnterStatsScope(FNiagaraStatScope(*GetFunctionSignatureSymbol(FunctionSignature), *(FunctionSignature.GetNameString())));
+		EnterStatsScope(FNiagaraStatScope(*BuildFunctionSignatureString(FunctionSignature), *(FunctionSignature.GetNameString())));
 	}
 
 	TArray<FString> MissingParameters;
 	int32 ParamIdx = 0;
 	TArray<int32> Params;
 	Params.Reserve(Inputs.Num() + Outputs.Num());
-	FString DefStr = GetFunctionSignatureSymbol(FunctionSignature) + TEXT("(");
+	FString DefStr = BuildFunctionSignatureString(FunctionSignature) + TEXT("(");
 	for (int32 i = 0; i < FunctionSignature.Inputs.Num(); ++i)
 	{
 		const FNiagaraVariable& InputVar = FunctionSignature.Inputs[i];
@@ -9226,7 +9472,7 @@ void FNiagaraHlslTranslator::GenerateFunctionCall(ENiagaraScriptUsage ScriptUsag
 			}
 			else
 			{
-				FString OutputStr = FString::Printf(TEXT("%sOutput_%s"), *GetFunctionSignatureSymbol(FunctionSignature), *OutVar.GetName().ToString());
+				FString OutputStr = FString::Printf(TEXT("%sOutput_%s"), *BuildFunctionSignatureString(FunctionSignature), *OutVar.GetName().ToString());
 				Output = AddBodyChunk(GetUniqueSymbolName(*OutputStr), TEXT(""), Type);
 				ParamOutput = Output;
 			}
@@ -9270,7 +9516,7 @@ void FNiagaraHlslTranslator::GenerateFunctionCall(ENiagaraScriptUsage ScriptUsag
 		for (FString MissingParam : MissingParameters)
 		{
 			FText Fmt = LOCTEXT("ErrorCompilingParameterFmt", "Error compiling parameter {0} in function call {1}");
-			FText ErrorText = FText::Format(Fmt, FText::FromString(MissingParam), FText::FromString(GetFunctionSignatureSymbol(FunctionSignature)));
+			FText ErrorText = FText::Format(Fmt, FText::FromString(MissingParam), FText::FromString(BuildFunctionSignatureString(FunctionSignature)));
 			Error(ErrorText);
 		}
 		return;
@@ -9286,118 +9532,53 @@ void FNiagaraHlslTranslator::GenerateFunctionCall(ENiagaraScriptUsage ScriptUsag
 
 FString FNiagaraHlslTranslator::GetFunctionSignatureSymbol(const FNiagaraFunctionSignature& Sig)
 {
-	FString SigStr = Sig.GetNameString();
-	if (!Sig.OwnerName.IsNone() && Sig.OwnerName.IsValid())
-	{
-		SigStr += TEXT("_") + Sig.OwnerName.ToString().Replace(TEXT("."), TEXT("_"));;
-	}
-	else
-	{
-		SigStr += TEXT("_Func_");
-	}
-	if (Sig.bRequiresExecPin)
-	{
-		SigStr += TEXT("_UEImpureCall"); // Let the cross compiler know that we intend to keep this.
-	}
-
-	for (const TTuple<FName, FName>& Specifier : Sig.FunctionSpecifiers)
-	{
-		SigStr += TEXT("_") + Specifier.Key.ToString() + Specifier.Value.ToString().Replace(TEXT("."), TEXT("_"));
-	}
+	const UNiagaraEditorSettings* EditorSettings = GetDefault<UNiagaraEditorSettings>();
 
-	auto AddVarsToSig = [&SigStr](TConstArrayView<FNiagaraVariableBase> Vars)
-	{
-		for(const FNiagaraVariableBase& Var : Vars)
-		{
-			SigStr += TEXT("_") + Var.GetName().ToString() + Var.GetType().GetName();
-		}
-	};
+	TStringBuilder<1024> SigStr;
+	NiagaraHlslTranslatorImpl::AppendFunctionSignatureSymbol(Sig, SigStr);
 
-	TArray<FNiagaraVariableBase> VariadicParams;
-	VariadicParams.Reserve(Sig.NumOptionalInputs() + Sig.NumOptionalOutputs());
-	Sig.GetVariadicInputs(VariadicParams);
-	Sig.GetVariadicOutputs(VariadicParams);
-	AddVarsToSig(VariadicParams);
+	TStringBuilder<1024> SanitizedSigStr;
+	NiagaraHlslTranslatorImpl::AppendSanitizedSymbolString(SigStr, false, EditorSettings->GetHLSLKeywordReplacementsMap(), SanitizedSigStr);
 
-	return GetSanitizedSymbolName(SigStr);
+	return SanitizedSigStr.ToString();
 }
 
-FString FNiagaraHlslTranslator::GenerateFunctionHlslPrototype(FStringView InVariableName, const FNiagaraFunctionSignature& FunctionSignature)
+void FNiagaraHlslTranslator::AppendFunctionSignatureString(const FNiagaraFunctionSignature& Sig, FStringBuilderBase& StringBuilder) const
 {
-	TStringBuilder<512> StringBuilder;
-	if (FunctionSignature.bMemberFunction)
-	{
-		StringBuilder.Append(InVariableName);
-		StringBuilder.Append(TEXT("."));
-		StringBuilder.Append(GetSanitizedSymbolName(FunctionSignature.Name.ToString()));
+	TStringBuilder<1024> SigStr;
+	NiagaraHlslTranslatorImpl::AppendFunctionSignatureSymbol(Sig, SigStr);
 
-		// Build specifiers
-		if (FunctionSignature.FunctionSpecifiers.Num())
-		{
-			bool bNeedsComma = false;
-			StringBuilder.AppendChar(TEXT('<'));
-			for (auto SpecifierIt = FunctionSignature.FunctionSpecifiers.CreateConstIterator(); SpecifierIt; ++SpecifierIt)
-			{
-				if (bNeedsComma)
-				{
-					StringBuilder.Append(TEXT(", "));
-				}
-				bNeedsComma = true;
+	NiagaraHlslTranslatorImpl::AppendSanitizedSymbolString(SigStr, false, KeywordReplacementMap, StringBuilder);
+}
 
-				StringBuilder.Append(SpecifierIt.Key().ToString());
-				StringBuilder.AppendChar(TEXT('='));
-				StringBuilder.AppendChar(TEXT('"'));
-				StringBuilder.Append(SpecifierIt.Value().IsNone() ? TEXT("Value") : *SpecifierIt.Value().ToString());
-				StringBuilder.AppendChar(TEXT('"'));
-			}
-			StringBuilder.AppendChar(TEXT('>'));
-		}
+FString FNiagaraHlslTranslator::BuildFunctionSignatureString(const FNiagaraFunctionSignature& Sig) const
+{
+	TStringBuilder<1024> SanitizedSigStr;
+	AppendFunctionSignatureString(Sig, SanitizedSigStr);
 
-		// Build function parameters
-		{
-			StringBuilder.Append(TEXT("("));
-			bool bNeedsComma = false;
+	return SanitizedSigStr.ToString();
+}
 
-			// Inputs
-			for (int i=1; i < FunctionSignature.Inputs.Num(); ++i)
-			{
-				FNiagaraVariable InputVar = ConvertToSimulationVariable(FunctionSignature.Inputs[i]);
-				if (bNeedsComma)
-				{
-					StringBuilder.Append(TEXT(", "));
-				}
-				bNeedsComma = true;
 
-				StringBuilder.Append(TEXT("in "));
-				StringBuilder.Append(GetStructHlslTypeName(InputVar.GetType()));
-				StringBuilder.Append(TEXT(" In_"));
-				StringBuilder.Append(GetSanitizedSymbolName(InputVar.GetName().ToString()));
-			}
+FString FNiagaraHlslTranslator::BuildFunctionHlslPrototypeString(FStringView InVariableName, const FNiagaraFunctionSignature& FunctionSignature) const
+{
+	TStringBuilder<1024> StringBuilder;
+	AppendFunctionHlslPrototypeString(InVariableName, FunctionSignature, StringBuilder);
 
-			// Outputs
-			for (int i=0; i < FunctionSignature.Outputs.Num(); ++i)
-			{
-				FNiagaraVariable OutputVar = ConvertToSimulationVariable(FunctionSignature.Outputs[i]);
-				if (bNeedsComma)
-				{
-					StringBuilder.Append(TEXT(", "));
-				}
-				bNeedsComma = true;
+	return StringBuilder.ToString();
+}
 
-				StringBuilder.Append(TEXT("out "));
-				StringBuilder.Append(GetStructHlslTypeName(OutputVar.GetType()));
-				StringBuilder.Append(TEXT(" Out_"));
-				StringBuilder.Append(GetSanitizedSymbolName(OutputVar.GetName().ToString()));
-			}
-			StringBuilder.Append(TEXT(");"));
-		}
-	}
-	else
-	{
-		ensureAlwaysMsgf(false, TEXT("None member functions not supported currently"));
-	}
+void FNiagaraHlslTranslator::AppendFunctionHlslPrototypeString(FStringView InVariableName, const FNiagaraFunctionSignature& FunctionSignature, FStringBuilderBase& StringBuilder) const
+{
+	NiagaraHlslTranslatorImpl::AppendFunctionHlslPrototypeString(InVariableName, FunctionSignature, KeywordReplacementMap, StringBuilder);
+}
 
-	return FString(StringBuilder.ToString());
+FString FNiagaraHlslTranslator::GenerateFunctionHlslPrototype(FStringView InVariableName, const FNiagaraFunctionSignature& FunctionSignature)
+{
+	const UNiagaraEditorSettings* EditorSettings = GetDefault<UNiagaraEditorSettings>();
+	TStringBuilder<1024> StringBuilder;
+	NiagaraHlslTranslatorImpl::AppendFunctionHlslPrototypeString(InVariableName, FunctionSignature, EditorSettings->GetHLSLKeywordReplacementsMap(), StringBuilder);
+	return StringBuilder.ToString();
 }
 
 FName FNiagaraHlslTranslator::GetDataInterfaceName(FName BaseName, const FString& UniqueEmitterName, bool bIsParameterMapDataInterface)
@@ -9438,9 +9619,11 @@ FString FNiagaraHlslTranslator::GetFunctionIncludeStatement(const FNiagaraCustom
 
 FString FNiagaraHlslTranslator::GetFunctionSignature(const FNiagaraFunctionSignature& Sig)
 {
-	FString SigStr = TEXT("void ") + GetFunctionSignatureSymbol(Sig);
+	TStringBuilder<1024> SigStr;
+	SigStr << TEXT("void ");
+	AppendFunctionSignatureString(Sig, SigStr);
 
-	SigStr += TEXT("(");
+	SigStr << TEXT("(");
 	int32 ParamIdx = 0;
 	for (int32 i = 0; i < Sig.Inputs.Num(); ++i)
 	{
@@ -9456,11 +9639,11 @@ FString FNiagaraHlslTranslator::GetFunctionSignature(const FNiagaraFunctionSigna
 			{
 				if (ParamIdx > 0)
 				{
-					SigStr += TEXT(", ");
+					SigStr << TEXT(", ");
 				}
 
 				FNiagaraVariable SimInput = ConvertToSimulationVariable(Input);
-				SigStr += GetStructHlslTypeName(SimInput.GetType()) + TEXT(" In_") + GetSanitizedSymbolName(Input.GetName().ToString(), true);
+				SigStr << GetStructHlslTypeName(SimInput.GetType()) + TEXT(" In_") + BuildSanitizedSymbolName(Input.GetName(), true);
 				++ParamIdx;
 			}
 		}
@@ -9480,11 +9663,11 @@ FString FNiagaraHlslTranslator::GetFunctionSignature(const FNiagaraFunctionSigna
 			{
 				if (ParamIdx > 0)
 				{
-					SigStr += TEXT(", ");
+					SigStr << TEXT(", ");
 				}
 
 				FNiagaraVariable SimOutput = ConvertToSimulationVariable(Output);
-				SigStr += TEXT("out ") + GetStructHlslTypeName(SimOutput.GetType()) + TEXT(" ") + GetSanitizedSymbolName(TEXT("Out_") + Output.GetName().ToString());
+				SigStr << TEXT("out ") + GetStructHlslTypeName(SimOutput.GetType()) + TEXT(" ") + BuildSanitizedSymbolString(TEXT("Out_") + Output.GetName().ToString());
 				++ParamIdx;
 			}
 		}
@@ -9493,11 +9676,12 @@ FString FNiagaraHlslTranslator::GetFunctionSignature(const FNiagaraFunctionSigna
 	{
 		if (ParamIdx > 0)
 		{
-			SigStr += TEXT(", ");
+			SigStr << TEXT(", ");
 		}
-		SigStr += TEXT("inout FSimulationContext Context");
+		SigStr << TEXT("inout FSimulationContext Context");
 	}
-	return SigStr + TEXT(")");
+	SigStr << TEXT(")");
+	return SigStr.ToString();
 }
 
 FNiagaraTypeDefinition FNiagaraHlslTranslator::GetChildType(const FNiagaraTypeDefinition& BaseType, const FName& PropertyName)
@@ -9909,7 +10093,7 @@ void TNiagaraHlslTranslator<GraphBridge>::FindConstantValue(int32 InputCompileRe
 						for (const FNiagaraVariable& UniformVar : CompilationOutput.ScriptData.Parameters.Parameters)
 						{
 							VarSymbol.Reset();
-							VarSymbol = GetSanitizedSymbolName(UniformVar.GetName().ToString(), true);
+							VarSymbol = BuildSanitizedSymbolName(UniformVar.GetName(), true);
 							if (VarSymbol == CodeChunks[InputCompileResult].SymbolName)
 							{
 								//UE_LOG(LogNiagaraEditor, Log, TEXT("Match Found: %s == %s"), *VarSymbol, *UniformVar.GetName().ToString());
@@ -10415,7 +10599,7 @@ void TNiagaraHlslTranslator<GraphBridge>::TrimAttributes(const FNiagaraCompileOp
 			TMap<FString, FNiagaraVariableBase> SanitizedNames;
 			for (const FNiagaraVariableBase& Attribute : Attributes)
 			{
-				const FString SanitizedName = GetSanitizedSymbolName(Attribute.GetName().ToString());
+				const FString SanitizedName = BuildSanitizedSymbolName(Attribute.GetName());
 				if (const FNiagaraVariableBase* ExistingVariable = SanitizedNames.Find(SanitizedName))
 				{
 					HasOverlappingNames = true;
@@ -10509,7 +10693,7 @@ TUniquePtr<INiagaraHlslTranslator> INiagaraHlslTranslator::CreateTranslator(cons
 }
 
 template<typename GraphBridge>
-const FString& TNiagaraHlslTranslator<GraphBridge>::GetEmitterUniqueName() const
+FStringView TNiagaraHlslTranslator<GraphBridge>::GetEmitterUniqueName() const
 {
 	return CompileData->EmitterUniqueName;
 }
diff --git a/Engine/Plugins/FX/Niagara/Source/NiagaraEditor/Private/NiagaraHlslTranslator.h b/Engine/Plugins/FX/Niagara/Source/NiagaraEditor/Private/NiagaraHlslTranslator.h
--- a/Engine/Plugins/FX/Niagara/Source/NiagaraEditor/Private/NiagaraHlslTranslator.h
+++ b/Engine/Plugins/FX/Niagara/Source/NiagaraEditor/Private/NiagaraHlslTranslator.h
@@ -337,11 +337,11 @@ protected:
 	void PackRegisteredUniformChunk(FNiagaraParameters& Parameters);
 
 	/* Add a chunk that is written to the body of the shader code. */
-	int32 AddBodyChunk(FString SymbolName, FString Definition, const FNiagaraTypeDefinition& Type, TArray<int32>& SourceChunks, bool bDecl = true, bool bIsTerminated = true);
-	int32 AddBodyChunk(FString SymbolName, FString Definition, const FNiagaraTypeDefinition& Type, int32 SourceChunk, bool bDecl = true, bool bIsTerminated = true);
-	int32 AddBodyChunk(FString SymbolName, FString Definition, const FNiagaraTypeDefinition& Type, bool bDecl = true, bool bIsTerminated = true);
-	int32 AddBodyComment(const FString& Comment);
-	int32 AddBodyChunk(const FString& Definition);
+	int32 AddBodyChunk(FStringView SymbolName, FStringView Definition, const FNiagaraTypeDefinition& Type, TArrayView<int32> SourceChunks, bool bDecl = true, bool bIsTerminated = true);
+	int32 AddBodyChunk(FStringView SymbolName, FStringView Definition, const FNiagaraTypeDefinition& Type, int32 SourceChunk, bool bDecl = true, bool bIsTerminated = true);
+	int32 AddBodyChunk(FStringView SymbolName, FStringView Definition, const FNiagaraTypeDefinition& Type, bool bDecl = true, bool bIsTerminated = true);
+	int32 AddBodyComment(FStringView Comment);
+	int32 AddBodyChunk(FStringView Definition);
 
 	FString GetUniqueEmitterName() const;
 	void ConvertCompileInfoToParamInfo(const FNiagaraScriptDataInterfaceCompileInfo& InCompileInfo, FNiagaraDataInterfaceGPUParamInfo& OutGPUParamInfo, TArray<FNiagaraFunctionSignature>& GeneratedFunctionSignatures);
@@ -407,9 +407,6 @@ public:
 	static NIAGARAEDITOR_API FString GetFunctionSignatureSymbol(const FNiagaraFunctionSignature& Sig);
 
 	static NIAGARAEDITOR_API FName GetDataInterfaceName(FName BaseName, const FString& UniqueEmitterName, bool bIsParameterMapDataInterface);
-
-
-
 	static NIAGARAEDITOR_API FString GenerateFunctionHlslPrototype(FStringView InVariableName, const FNiagaraFunctionSignature& FunctionSignature);
 
 	NIAGARAEDITOR_API void SetConstantByStaticVariable(int32& OutValue, const FNiagaraVariable& Var, FString* DebugString = nullptr);
@@ -473,11 +470,21 @@ protected:
 	bool RequiresInterpolation() const;
 	bool IsWriteAllowedForNamespace(const FNiagaraVariable& Var, ENiagaraScriptUsage TargetUsage, FText& ErrorMsg);
 
+	void AppendSanitizedSymbolString(FStringBuilderBase& Builder, FStringView SymbolName, bool bCollapseNamespaces = false) const;
+	FString BuildSanitizedSymbolString(FStringView SymbolName, bool bCollapseNamespaces = false) const;
+	FString BuildSanitizedSymbolName(FName SymbolName, bool bCollapseNamespaces = false) const;
+
+	void AppendFunctionSignatureString(const FNiagaraFunctionSignature& Sig, FStringBuilderBase& Builder) const;
+	FString BuildFunctionSignatureString(const FNiagaraFunctionSignature& Sig) const;
+
+	void AppendFunctionHlslPrototypeString(FStringView InVariableName, const FNiagaraFunctionSignature& FunctionSignature, FStringBuilderBase& Builder) const;
+	FString BuildFunctionHlslPrototypeString(FStringView InVariableName, const FNiagaraFunctionSignature& FunctionSignature) const;
+
 	template<typename T>
 	void BuildConstantBuffer(ENiagaraCodeChunkMode ChunkMode);
 
 	virtual FNiagaraEmitterID GetEmitterID() const = 0;
-	virtual const FString& GetEmitterUniqueName() const = 0;
+	virtual FStringView GetEmitterUniqueName() const = 0;
 	virtual TConstArrayView<FNiagaraVariable> GetStaticVariables() const = 0;
 	virtual UNiagaraDataInterface* GetDataInterfaceCDO(UClass* DIClass) const = 0;
 
@@ -576,6 +583,9 @@ protected:
 	TArray<FNiagaraVariable> InitialNamespaceVariablesMissingDefault;
 	// Variables that need to be initialized in the body or at the end of spawn.
 	TArray<FNiagaraVariable> DeferredVariablesMissingDefault;
+
+	// local copy of UNiagaraEditorSettings::GetHLSLKeywordReplacmentsMap() for the duration of the translation process
+	TMap<FString, FString> KeywordReplacementMap;
 };
 
 // all data and function implementations that rely on specifics of the graph implementation
@@ -736,7 +746,7 @@ protected:
 
 	void UpdateStaticSwitchConstants(const FPin* Pin);
 
-	virtual const FString& GetEmitterUniqueName() const override;
+	virtual FStringView GetEmitterUniqueName() const override;
 	virtual FNiagaraEmitterID GetEmitterID() const override;
 	virtual TConstArrayView<FNiagaraVariable> GetStaticVariables() const override;
 	virtual UNiagaraDataInterface* GetDataInterfaceCDO(UClass* DIClass) const override;
