Fix discrepency between RVT and Landscape DIs when it comes to sampling the base color (specifically around sRGB).  Tries to leverage some public functions of the RVT system to clean things up, but at the end of the day there's still some private details that aren't cleanly exposed
#tests Ran through level provided by Dan to compare material output vs landscape DI output vs RVT DI output
prflight 6792bd5d161368bfee9949e3

diff --git a/Engine/Plugins/FX/Niagara/Shaders/Private/NiagaraDataInterfaceLandscape.ush b/Engine/Plugins/FX/Niagara/Shaders/Private/NiagaraDataInterfaceLandscape.ush
--- a/Engine/Plugins/FX/Niagara/Shaders/Private/NiagaraDataInterfaceLandscape.ush
+++ b/Engine/Plugins/FX/Niagara/Shaders/Private/NiagaraDataInterfaceLandscape.ush
@@ -17,7 +17,8 @@ Texture2D<uint4>	{ParameterName}_BaseColorVirtualTexturePageTable;
 SamplerState		{ParameterName}_BaseColorVirtualTextureSampler;
 float3				{ParameterName}_BaseColorVirtualTextureLWCTile;
 float4x4			{ParameterName}_BaseColorVirtualTextureWorldToUvTransform;
-uint				{ParameterName}_BaseColorVirtualTextureSRGB;
+uint				{ParameterName}_BaseColorVirtualTextureUnpackSRGB;
+uint				{ParameterName}_BaseColorVirtualTextureUnpackYCoCg;
 uint				{ParameterName}_BaseColorVirtualTextureEnabled;
 uint4				{ParameterName}_BaseColorVirtualTexturePackedUniform0;
 uint4				{ParameterName}_BaseColorVirtualTexturePackedUniform1;
@@ -95,7 +96,18 @@ void GetBaseColor_{ParameterName}(float3 InWorldPos, out float3 OutBaseColor, ou
 			0,
 			VTUniform_Unpack({ParameterName}_BaseColorVirtualTextureUniforms));
 
-		OutBaseColor = {ParameterName}_BaseColorVirtualTextureSRGB != 0 ? VirtualTextureUnpackBaseColorSRGB(PackedValue) : VirtualTextureUnpackBaseColorYCoCg(PackedValue);
+		if ({ParameterName}_BaseColorVirtualTextureUnpackSRGB)
+		{
+			OutBaseColor = VirtualTextureUnpackBaseColorSRGB(PackedValue);
+		}
+		else if ({ParameterName}_BaseColorVirtualTextureUnpackYCoCg)
+		{
+			OutBaseColor = VirtualTextureUnpackBaseColorYCoCg(PackedValue);
+		}
+		else
+		{
+			OutBaseColor = PackedValue.rgb;
+		}
 		OutIsValid = true;
 	}
 	else
diff --git a/Engine/Plugins/FX/Niagara/Shaders/Private/NiagaraDataInterfaceVirtualTextureTemplate.ush b/Engine/Plugins/FX/Niagara/Shaders/Private/NiagaraDataInterfaceVirtualTextureTemplate.ush
--- a/Engine/Plugins/FX/Niagara/Shaders/Private/NiagaraDataInterfaceVirtualTextureTemplate.ush
+++ b/Engine/Plugins/FX/Niagara/Shaders/Private/NiagaraDataInterfaceVirtualTextureTemplate.ush
@@ -128,7 +128,7 @@ void SampleRVT_{ParameterName}(in float3 WorldPosition, out bool bInsideVolume,
 
 		case ERuntimeVirtualTextureMaterialType_BaseColor_Normal_Roughness:
 		{
-			BaseColor = LayerSample[0].rgb;
+			BaseColor = VirtualTextureUnpackBaseColorSRGB(LayerSample[0]);
 			Roughness = LayerSample[1].y;
 			Normal = VirtualTextureUnpackNormalBGR565(LayerSample[1]);
 			break;
diff --git a/Engine/Plugins/FX/Niagara/Source/Niagara/Private/DataInterface/NiagaraDataInterfaceVirtualTexture.cpp b/Engine/Plugins/FX/Niagara/Source/Niagara/Private/DataInterface/NiagaraDataInterfaceVirtualTexture.cpp
--- a/Engine/Plugins/FX/Niagara/Source/Niagara/Private/DataInterface/NiagaraDataInterfaceVirtualTexture.cpp
+++ b/Engine/Plugins/FX/Niagara/Source/Niagara/Private/DataInterface/NiagaraDataInterfaceVirtualTexture.cpp
@@ -283,11 +283,6 @@ bool UNiagaraDataInterfaceVirtualTexture::PerInstanceTick(void* PerInstanceData,
 					InstanceData.UVUniforms[1] = RT_Texture->GetUniformParameter(ERuntimeVirtualTextureShaderUniform_WorldToUVTransform1);
 					InstanceData.UVUniforms[2] = RT_Texture->GetUniformParameter(ERuntimeVirtualTextureShaderUniform_WorldToUVTransform2);
 
-					const bool bBaseColorSRGB = 
-						InstanceData.MaterialType == ERuntimeVirtualTextureMaterialType::BaseColor ||
-						InstanceData.MaterialType == ERuntimeVirtualTextureMaterialType::BaseColor_Normal_Roughness ||
-						InstanceData.MaterialType == ERuntimeVirtualTextureMaterialType::BaseColor_Normal_Specular;
-
 					const FVector4 WorldHeightUnpack = RT_Texture->GetUniformParameter(ERuntimeVirtualTextureShaderUniform_WorldHeightUnpack);
 					InstanceData.WorldHeightUnpack = FVector2f(float(WorldHeightUnpack.X), float(WorldHeightUnpack.Y));		//LWC Precision Loss
 
@@ -295,7 +290,7 @@ bool UNiagaraDataInterfaceVirtualTexture::PerInstanceTick(void* PerInstanceData,
 					{
 						FNDInstanceData_RenderThread::FLayerTexture& LayerTexture = InstanceData.LayerTexture[iLayer];
 						LayerTexture.Texture = VirtualTexture->GetPhysicalTexture(iLayer);
-						LayerTexture.TextureSRV = VirtualTexture->GetPhysicalTextureSRV(iLayer, iLayer == 0 && bBaseColorSRGB);
+						LayerTexture.TextureSRV = VirtualTexture->GetPhysicalTextureSRV(iLayer, RT_Texture->IsLayerSRGB(iLayer));
 						LayerTexture.PageTable = VirtualTexture->GetPageTableTexture(0);//-TODO:iLayer);
 
 						VirtualTexture->GetPackedUniform(&LayerTexture.TextureUniforms, iLayer);
diff --git a/Engine/Plugins/FX/Niagara/Source/Niagara/Private/NiagaraDataInterfaceLandscape.cpp b/Engine/Plugins/FX/Niagara/Source/Niagara/Private/NiagaraDataInterfaceLandscape.cpp
--- a/Engine/Plugins/FX/Niagara/Source/Niagara/Private/NiagaraDataInterfaceLandscape.cpp
+++ b/Engine/Plugins/FX/Niagara/Source/Niagara/Private/NiagaraDataInterfaceLandscape.cpp
@@ -75,7 +75,8 @@ namespace NiagaraDataInterfaceLandscape
 		SHADER_PARAMETER_SAMPLER(SamplerState,		BaseColorVirtualTextureSampler)
 		SHADER_PARAMETER(FVector3f,					BaseColorVirtualTextureLWCTile)
 		SHADER_PARAMETER(FMatrix44f,				BaseColorVirtualTextureWorldToUvTransform)
-		SHADER_PARAMETER(uint32,					BaseColorVirtualTextureSRGB)
+		SHADER_PARAMETER(uint32,					BaseColorVirtualTextureUnpackSRGB)
+		SHADER_PARAMETER(uint32,					BaseColorVirtualTextureUnpackYCoCg)
 		SHADER_PARAMETER(uint32,					BaseColorVirtualTextureEnabled)
 		SHADER_PARAMETER(FUintVector4,				BaseColorVirtualTexturePackedUniform0)
 		SHADER_PARAMETER(FUintVector4,				BaseColorVirtualTexturePackedUniform1)
@@ -115,6 +116,14 @@ namespace NiagaraDataInterfaceLandscape
 		SHADER_PARAMETER_TEXTURE(Texture2D<uint>,	CachedPhysMatTexture)
 		SHADER_PARAMETER(FIntPoint,					CachedPhysMatTextureDimension)
 	END_SHADER_PARAMETER_STRUCT()
+
+
+	enum class EBaseColorUnpackType : int32
+	{
+		None = 0,		// no unpacking is required 
+		SRGBUnpack,		// base color is manually packed as SRGB
+		YCoCgUnpack,	// base color is manually packed as YCoCg
+	};
 }
 
 const FName UNiagaraDataInterfaceLandscape::GetBaseColorName(TEXT("GetBaseColor"));
@@ -222,6 +231,7 @@ struct FNDILandscapeData_GameThread
 	TWeakObjectPtr<ULandscapeHeightfieldCollisionComponent> CollisionComponent;
 	FNDI_Landscape_SharedResourceHandle SharedResourceHandle;
 	bool BaseColorVirtualTextureSRGB = false;
+	NiagaraDataInterfaceLandscape::EBaseColorUnpackType BaseColorVirtualTextureUnpackType = NiagaraDataInterfaceLandscape::EBaseColorUnpackType::None;
 	int32 BaseColorVirtualTextureIndex = INDEX_NONE;
 	int32 HeightVirtualTextureIndex = INDEX_NONE;
 	int32 NormalVirtualTextureIndex = INDEX_NONE;
@@ -238,6 +248,7 @@ struct FNDILandscapeData_GameThread
 		Landscape = nullptr;
 		SharedResourceHandle = FNDI_Landscape_SharedResourceHandle();
 		BaseColorVirtualTextureSRGB = false;
+		BaseColorVirtualTextureUnpackType = NiagaraDataInterfaceLandscape::EBaseColorUnpackType::None;
 		BaseColorVirtualTextureIndex = INDEX_NONE;
 		HeightVirtualTextureIndex = INDEX_NONE;
 		NormalVirtualTextureIndex = INDEX_NONE;
@@ -261,7 +272,7 @@ struct FNDILandscapeData_GameToRenderThread
 	FVector4 CachedHeightTextureUvScaleBias = FVector4::Zero();
 	FVector2D CachedHeightTextureWorldGridSize = FVector2D(1.0f, 1.0f);
 
-	bool BaseColorVirtualTextureSRGB = false;
+	NiagaraDataInterfaceLandscape::EBaseColorUnpackType BaseColorVirtualTextureUnpackType = NiagaraDataInterfaceLandscape::EBaseColorUnpackType::None;
 };
 
 // Landscape data used on the render thread
@@ -394,7 +405,8 @@ struct FNDILandscapeData_RenderThread
 			FVector3f((FVector4f)BaseColorVirtualPage.WorldToUvParameters[1]),
 			FVector3f((FVector4f)BaseColorVirtualPage.WorldToUvParameters[2]),
 			FVector3f(0, 0, 0));
-		ShaderParameters->BaseColorVirtualTextureSRGB = LandscapeData.BaseColorVirtualTextureSRGB;
+		ShaderParameters->BaseColorVirtualTextureUnpackSRGB = LandscapeData.BaseColorVirtualTextureUnpackType == NiagaraDataInterfaceLandscape::EBaseColorUnpackType::SRGBUnpack;
+		ShaderParameters->BaseColorVirtualTextureUnpackYCoCg = LandscapeData.BaseColorVirtualTextureUnpackType == NiagaraDataInterfaceLandscape::EBaseColorUnpackType::YCoCgUnpack;
 		ShaderParameters->BaseColorVirtualTextureEnabled = 1;
 		ShaderParameters->BaseColorVirtualTexturePackedUniform0 = BaseColorVirtualPage.PageTableUniforms[0];
 		ShaderParameters->BaseColorVirtualTexturePackedUniform1 = BaseColorVirtualPage.PageTableUniforms[1];
@@ -410,7 +422,8 @@ struct FNDILandscapeData_RenderThread
 		ShaderParameters->BaseColorVirtualTexture = GBlackTextureWithSRV->ShaderResourceViewRHI;
 		ShaderParameters->BaseColorVirtualTexturePageTable = GBlackUintTexture->TextureRHI;
 		ShaderParameters->BaseColorVirtualTextureWorldToUvTransform = FMatrix44f::Identity;
-		ShaderParameters->BaseColorVirtualTextureSRGB = 0;
+		ShaderParameters->BaseColorVirtualTextureUnpackSRGB = 0;
+		ShaderParameters->BaseColorVirtualTextureUnpackYCoCg = 0;
 		ShaderParameters->BaseColorVirtualTextureEnabled = 0;
 		ShaderParameters->BaseColorVirtualTexturePackedUniform0 = DummyUint4;
 		ShaderParameters->BaseColorVirtualTexturePackedUniform1 = DummyUint4;
@@ -1576,7 +1589,7 @@ void UNiagaraDataInterfaceLandscape::ProvidePerInstanceDataForRenderThread(void*
 		TargetData->CachedHeightTextureWorldGridSize = SourceResource.TextureWorldGridSize;
 	}
 
-	TargetData->BaseColorVirtualTextureSRGB = SourceData.BaseColorVirtualTextureSRGB;
+	TargetData->BaseColorVirtualTextureUnpackType = SourceData.BaseColorVirtualTextureUnpackType;
 }
 
 #if WITH_EDITORONLY_DATA
@@ -1828,19 +1841,13 @@ void UNiagaraDataInterfaceLandscape::ApplyLandscape(const FNiagaraSystemInstance
 
 				switch (VirtualMaterialType)
 				{
-				case ERuntimeVirtualTextureMaterialType::BaseColor:
-					if (InstanceData.BaseColorVirtualTextureIndex == INDEX_NONE)
-					{
-						InstanceData.BaseColorVirtualTextureIndex = TextureIt;
-						InstanceData.BaseColorVirtualTextureSRGB = true;
-					}
-					break;
 				case ERuntimeVirtualTextureMaterialType::WorldHeight:
 					if (InstanceData.HeightVirtualTextureIndex == INDEX_NONE)
 					{
 						InstanceData.HeightVirtualTextureIndex = TextureIt;
 					}
 					break;
+
 				case ERuntimeVirtualTextureMaterialType::BaseColor_Normal_Roughness:
 				case ERuntimeVirtualTextureMaterialType::BaseColor_Normal_Specular:
 				case ERuntimeVirtualTextureMaterialType::BaseColor_Normal_Specular_YCoCg:
@@ -1850,15 +1857,27 @@ void UNiagaraDataInterfaceLandscape::ApplyLandscape(const FNiagaraSystemInstance
 						InstanceData.NormalVirtualTextureIndex = TextureIt;
 						InstanceData.NormalVirtualTextureMode = VirtualMaterialType;
 					}
+					// intentionally falling through to the BaseColor option
+				case ERuntimeVirtualTextureMaterialType::BaseColor:
 					if (InstanceData.BaseColorVirtualTextureIndex == INDEX_NONE)
 					{
 						InstanceData.BaseColorVirtualTextureIndex = TextureIt;
-						InstanceData.BaseColorVirtualTextureSRGB =
-							(VirtualMaterialType == ERuntimeVirtualTextureMaterialType::BaseColor_Normal_Roughness) ||
-							(VirtualMaterialType == ERuntimeVirtualTextureMaterialType::BaseColor_Normal_Specular);
 					}
 					break;
 				}
+
+				if (InstanceData.BaseColorVirtualTextureIndex == TextureIt)
+				{
+					InstanceData.BaseColorVirtualTextureSRGB = Vt->IsLayerSRGB(0);
+					if (Vt->IsLayerYCoCg(0))
+					{
+						InstanceData.BaseColorVirtualTextureUnpackType = NiagaraDataInterfaceLandscape::EBaseColorUnpackType::YCoCgUnpack;
+					}
+					else if (VirtualMaterialType == ERuntimeVirtualTextureMaterialType::BaseColor_Normal_Roughness)
+					{
+						InstanceData.BaseColorVirtualTextureUnpackType = NiagaraDataInterfaceLandscape::EBaseColorUnpackType::SRGBUnpack;
+					}
+				}
 			}
 		}
 	}
