- Material Random attribute support for stateless

diff --git a/Engine/Plugins/FX/Niagara/Shaders/Private/Stateless/Modules/NiagaraStatelessModule_InitializeParticle.ush b/Engine/Plugins/FX/Niagara/Shaders/Private/Stateless/Modules/NiagaraStatelessModule_InitializeParticle.ush
--- a/Engine/Plugins/FX/Niagara/Shaders/Private/Stateless/Modules/NiagaraStatelessModule_InitializeParticle.ush
+++ b/Engine/Plugins/FX/Niagara/Shaders/Private/Stateless/Modules/NiagaraStatelessModule_InitializeParticle.ush
@@ -27,8 +27,7 @@ void InitializeParticle_Simulate(inout FStatelessParticle Particle)
 
 	FStatelessDistributionSamplerFloat3 InitialPosition;
 	InitialPosition.Init(0, InitializeParticle_InitialPosition);
-		
-	Particle.UniqueID		= Particle.UniqueIndex;
+
 	Particle.Position		= InitialPosition.GetValue(0.0f);
 	Particle.Color			= RandomScaleBiasFloat(1, InitializeParticle_ColorScale, InitializeParticle_ColorBias);
 	Particle.RibbonWidth	= RandomScaleBiasFloat(2, InitializeParticle_RibbonWidthScale, InitializeParticle_RibbonWidthBias);
diff --git a/Engine/Plugins/FX/Niagara/Shaders/Private/Stateless/NiagaraStatelessCommon.ush b/Engine/Plugins/FX/Niagara/Shaders/Private/Stateless/NiagaraStatelessCommon.ush
--- a/Engine/Plugins/FX/Niagara/Shaders/Private/Stateless/NiagaraStatelessCommon.ush
+++ b/Engine/Plugins/FX/Niagara/Shaders/Private/Stateless/NiagaraStatelessCommon.ush
@@ -68,6 +68,7 @@ struct FStatelessParticle
 {
 	bool bAlive;					// If the particle should remain alive or not (only evaluated as specific points)
 	uint UniqueIndex;				// Unique particle index, sequential based on previous spawn info accumulation
+	float MaterialRandom;			// Random float passed to materials
 	float Lifetime;					// Overall lifetime for the particle
 	float Age;						// Current age for the particle
 	float NormalizedAge;			// Current normalized age for the particle
@@ -79,10 +80,13 @@ struct FStatelessParticle
 	
 	/////////////////////////////////////////////////////////////////////////////////////////////////////////////////////////
 	// TBD this might produce sub optimal code, needs investigate and is temporary
-	#define PARTICLE_ATTRIBUTE(TYPE, NAME) TYPE NAME##;
-		PARTICLE_ATTRIBUTES_OUTPUT
-		PARTICLE_ATTRIBUTES_TRANSIENT
+	#define PARTICLE_ATTRIBUTE_OUTPUT(TYPE, NAME) TYPE NAME##;
+	#define PARTICLE_ATTRIBUTE_OUTPUT_IMPLICIT(TYPE, NAME)
+	#define PARTICLE_ATTRIBUTE_TRANSIENT(TYPE, NAME) TYPE NAME##;
+		PARTICLE_ATTRIBUTES
 	#undef PARTICLE_ATTRIBUTE
+	#undef PARTICLE_ATTRIBUTE_OUTPUT_IMPLICIT
+	#undef PARTICLE_ATTRIBUTE_TRANSIENT
 	/////////////////////////////////////////////////////////////////////////////////////////////////////////////////////////
 };
 
diff --git a/Engine/Plugins/FX/Niagara/Shaders/Private/Stateless/NiagaraStatelessSimulationDefault.usf b/Engine/Plugins/FX/Niagara/Shaders/Private/Stateless/NiagaraStatelessSimulationDefault.usf
--- a/Engine/Plugins/FX/Niagara/Shaders/Private/Stateless/NiagaraStatelessSimulationDefault.usf
+++ b/Engine/Plugins/FX/Niagara/Shaders/Private/Stateless/NiagaraStatelessSimulationDefault.usf
@@ -1,38 +1,39 @@
 // Copyright Epic Games, Inc. All Rights Reserved.
 
-// Defines potential output variables, usage in renderers will determine what we actually output
-#define PARTICLE_ATTRIBUTES_OUTPUT \
-	PARTICLE_ATTRIBUTE(int,		UniqueID) \
-	PARTICLE_ATTRIBUTE(float3,	Position) \
-	PARTICLE_ATTRIBUTE(float,	CameraOffset) \
-	PARTICLE_ATTRIBUTE(float4,	Color) \
-	PARTICLE_ATTRIBUTE(float4,	DynamicMaterialParameter0) \
-	PARTICLE_ATTRIBUTE(float4,	DynamicMaterialParameter1) \
-	PARTICLE_ATTRIBUTE(float4,	DynamicMaterialParameter2) \
-	PARTICLE_ATTRIBUTE(float4,	DynamicMaterialParameter3) \
-	PARTICLE_ATTRIBUTE(int,		MeshIndex) \
-	PARTICLE_ATTRIBUTE(float4,	MeshOrientation) \
-	PARTICLE_ATTRIBUTE(float,	RibbonWidth) \
-	PARTICLE_ATTRIBUTE(float3,	Scale) \
-	PARTICLE_ATTRIBUTE(float2,	SpriteSize) \
-	PARTICLE_ATTRIBUTE(float3,	SpriteFacing) \
-	PARTICLE_ATTRIBUTE(float3,	SpriteAlignment) \
-	PARTICLE_ATTRIBUTE(float,	SpriteRotation) \
-	PARTICLE_ATTRIBUTE(float,	SubImageIndex) \
-	PARTICLE_ATTRIBUTE(float3,	Velocity) \
-	PARTICLE_ATTRIBUTE(float,	PreviousCameraOffset) \
-	PARTICLE_ATTRIBUTE(float3,	PreviousPosition) \
-	PARTICLE_ATTRIBUTE(float4,	PreviousMeshOrientation) \
-	PARTICLE_ATTRIBUTE(float,	PreviousRibbonWidth) \
-	PARTICLE_ATTRIBUTE(float3,	PreviousScale) \
-	PARTICLE_ATTRIBUTE(float2,	PreviousSpriteSize) \
-	PARTICLE_ATTRIBUTE(float3,	PreviousSpriteFacing) \
-	PARTICLE_ATTRIBUTE(float3,	PreviousSpriteAlignment) \
-	PARTICLE_ATTRIBUTE(float,	PreviousSpriteRotation) \
-	PARTICLE_ATTRIBUTE(float3,	PreviousVelocity) \
-
-// Defines transient variables, i.e. ones we will not output
-#define PARTICLE_ATTRIBUTES_TRANSIENT \
+// Defines all paricles attributes
+// - PARTICLE_ATTRIBUTE_OUTPUT - Defines an output variable
+// - PARTICLE_ATTRIBUTE_OUTPUT_IMPLICIT - Defines an output variable which is already part of the main particle payload (i.e. always present)
+// - PARTICLE_ATTRIBUTE_TRANSIENT - Defines an output that is temporary and not output
+#define PARTICLE_ATTRIBUTES \
+	PARTICLE_ATTRIBUTE_OUTPUT_IMPLICIT(int,		UniqueIndex) \
+	PARTICLE_ATTRIBUTE_OUTPUT_IMPLICIT(float,	MaterialRandom) \
+	PARTICLE_ATTRIBUTE_OUTPUT(float3,	Position) \
+	PARTICLE_ATTRIBUTE_OUTPUT(float,	CameraOffset) \
+	PARTICLE_ATTRIBUTE_OUTPUT(float4,	Color) \
+	PARTICLE_ATTRIBUTE_OUTPUT(float4,	DynamicMaterialParameter0) \
+	PARTICLE_ATTRIBUTE_OUTPUT(float4,	DynamicMaterialParameter1) \
+	PARTICLE_ATTRIBUTE_OUTPUT(float4,	DynamicMaterialParameter2) \
+	PARTICLE_ATTRIBUTE_OUTPUT(float4,	DynamicMaterialParameter3) \
+	PARTICLE_ATTRIBUTE_OUTPUT(int,		MeshIndex) \
+	PARTICLE_ATTRIBUTE_OUTPUT(float4,	MeshOrientation) \
+	PARTICLE_ATTRIBUTE_OUTPUT(float,	RibbonWidth) \
+	PARTICLE_ATTRIBUTE_OUTPUT(float3,	Scale) \
+	PARTICLE_ATTRIBUTE_OUTPUT(float2,	SpriteSize) \
+	PARTICLE_ATTRIBUTE_OUTPUT(float3,	SpriteFacing) \
+	PARTICLE_ATTRIBUTE_OUTPUT(float3,	SpriteAlignment) \
+	PARTICLE_ATTRIBUTE_OUTPUT(float,	SpriteRotation) \
+	PARTICLE_ATTRIBUTE_OUTPUT(float,	SubImageIndex) \
+	PARTICLE_ATTRIBUTE_OUTPUT(float3,	Velocity) \
+	PARTICLE_ATTRIBUTE_OUTPUT(float,	PreviousCameraOffset) \
+	PARTICLE_ATTRIBUTE_OUTPUT(float3,	PreviousPosition) \
+	PARTICLE_ATTRIBUTE_OUTPUT(float4,	PreviousMeshOrientation) \
+	PARTICLE_ATTRIBUTE_OUTPUT(float,	PreviousRibbonWidth) \
+	PARTICLE_ATTRIBUTE_OUTPUT(float3,	PreviousScale) \
+	PARTICLE_ATTRIBUTE_OUTPUT(float2,	PreviousSpriteSize) \
+	PARTICLE_ATTRIBUTE_OUTPUT(float3,	PreviousSpriteFacing) \
+	PARTICLE_ATTRIBUTE_OUTPUT(float3,	PreviousSpriteAlignment) \
+	PARTICLE_ATTRIBUTE_OUTPUT(float,	PreviousSpriteRotation) \
+	PARTICLE_ATTRIBUTE_OUTPUT(float3,	PreviousVelocity) \
 
 // Defines modules and order of execution
 #define PARTICLE_MODULES \
diff --git a/Engine/Plugins/FX/Niagara/Shaders/Private/Stateless/NiagaraStatelessSimulationExample1.usf b/Engine/Plugins/FX/Niagara/Shaders/Private/Stateless/NiagaraStatelessSimulationExample1.usf
--- a/Engine/Plugins/FX/Niagara/Shaders/Private/Stateless/NiagaraStatelessSimulationExample1.usf
+++ b/Engine/Plugins/FX/Niagara/Shaders/Private/Stateless/NiagaraStatelessSimulationExample1.usf
@@ -1,38 +1,39 @@
 // Copyright Epic Games, Inc. All Rights Reserved.
 
-// Defines potential output variables, usage in renderers will determine what we actually output
-#define PARTICLE_ATTRIBUTES_OUTPUT \
-	PARTICLE_ATTRIBUTE(int,		UniqueID) \
-	PARTICLE_ATTRIBUTE(float3,	Position) \
-	PARTICLE_ATTRIBUTE(float4,	Color) \
-	PARTICLE_ATTRIBUTE(float,	RibbonWidth) \
-	PARTICLE_ATTRIBUTE(float3,	Scale) \
-	PARTICLE_ATTRIBUTE(float2,	SpriteSize) \
-	PARTICLE_ATTRIBUTE(float,	SpriteRotation) \
-	PARTICLE_ATTRIBUTE(float3,	Velocity) \
-	PARTICLE_ATTRIBUTE(float3,	PreviousPosition) \
-	PARTICLE_ATTRIBUTE(float,	PreviousRibbonWidth) \
-	PARTICLE_ATTRIBUTE(float3,	PreviousScale) \
-	PARTICLE_ATTRIBUTE(float2,	PreviousSpriteSize) \
-	PARTICLE_ATTRIBUTE(float,	PreviousSpriteRotation) \
-	PARTICLE_ATTRIBUTE(float3,	PreviousVelocity) \
-
-// Defines transient variables, i.e. ones we will not output
-#define PARTICLE_ATTRIBUTES_TRANSIENT \
-	PARTICLE_ATTRIBUTE(float,	CameraOffset) \
-	PARTICLE_ATTRIBUTE(float4,	DynamicMaterialParameter0) \
-	PARTICLE_ATTRIBUTE(float4,	DynamicMaterialParameter1) \
-	PARTICLE_ATTRIBUTE(float4,	DynamicMaterialParameter2) \
-	PARTICLE_ATTRIBUTE(float4,	DynamicMaterialParameter3) \
-	PARTICLE_ATTRIBUTE(int,		MeshIndex) \
-	PARTICLE_ATTRIBUTE(float4,	MeshOrientation) \
-	PARTICLE_ATTRIBUTE(float3,	SpriteFacing) \
-	PARTICLE_ATTRIBUTE(float3,	SpriteAlignment) \
-	PARTICLE_ATTRIBUTE(float,	SubImageIndex) \
-	PARTICLE_ATTRIBUTE(float,	PreviousCameraOffset) \
-	PARTICLE_ATTRIBUTE(float4,	PreviousMeshOrientation) \
-	PARTICLE_ATTRIBUTE(float3,	PreviousSpriteFacing) \
-	PARTICLE_ATTRIBUTE(float3,	PreviousSpriteAlignment) \
+// Defines all paricles attributes
+// - PARTICLE_ATTRIBUTE_OUTPUT - Defines an output variable
+// - PARTICLE_ATTRIBUTE_OUTPUT_IMPLICIT - Defines an output variable which is already part of the main particle payload (i.e. always present)
+// - PARTICLE_ATTRIBUTE_TRANSIENT - Defines an output that is temporary and not output
+#define PARTICLE_ATTRIBUTES \
+	PARTICLE_ATTRIBUTE_OUTPUT_IMPLICIT(int,		UniqueIndex) \
+	PARTICLE_ATTRIBUTE_OUTPUT_IMPLICIT(float,	MaterialRandom) \
+	PARTICLE_ATTRIBUTE_OUTPUT(float3,	Position) \
+	PARTICLE_ATTRIBUTE_OUTPUT(float4,	Color) \
+	PARTICLE_ATTRIBUTE_OUTPUT(float,	RibbonWidth) \
+	PARTICLE_ATTRIBUTE_OUTPUT(float3,	Scale) \
+	PARTICLE_ATTRIBUTE_OUTPUT(float2,	SpriteSize) \
+	PARTICLE_ATTRIBUTE_OUTPUT(float,	SpriteRotation) \
+	PARTICLE_ATTRIBUTE_OUTPUT(float3,	Velocity) \
+	PARTICLE_ATTRIBUTE_OUTPUT(float3,	PreviousPosition) \
+	PARTICLE_ATTRIBUTE_OUTPUT(float,	PreviousRibbonWidth) \
+	PARTICLE_ATTRIBUTE_OUTPUT(float3,	PreviousScale) \
+	PARTICLE_ATTRIBUTE_OUTPUT(float2,	PreviousSpriteSize) \
+	PARTICLE_ATTRIBUTE_OUTPUT(float,	PreviousSpriteRotation) \
+	PARTICLE_ATTRIBUTE_OUTPUT(float3,	PreviousVelocity) \
+	PARTICLE_ATTRIBUTE_TRANSIENT(float,		CameraOffset) \
+	PARTICLE_ATTRIBUTE_TRANSIENT(float4,	DynamicMaterialParameter0) \
+	PARTICLE_ATTRIBUTE_TRANSIENT(float4,	DynamicMaterialParameter1) \
+	PARTICLE_ATTRIBUTE_TRANSIENT(float4,	DynamicMaterialParameter2) \
+	PARTICLE_ATTRIBUTE_TRANSIENT(float4,	DynamicMaterialParameter3) \
+	PARTICLE_ATTRIBUTE_TRANSIENT(int,		MeshIndex) \
+	PARTICLE_ATTRIBUTE_TRANSIENT(float4,	MeshOrientation) \
+	PARTICLE_ATTRIBUTE_TRANSIENT(float3,	SpriteFacing) \
+	PARTICLE_ATTRIBUTE_TRANSIENT(float3,	SpriteAlignment) \
+	PARTICLE_ATTRIBUTE_TRANSIENT(float,		SubImageIndex) \
+	PARTICLE_ATTRIBUTE_TRANSIENT(float,		PreviousCameraOffset) \
+	PARTICLE_ATTRIBUTE_TRANSIENT(float4,	PreviousMeshOrientation) \
+	PARTICLE_ATTRIBUTE_TRANSIENT(float3,	PreviousSpriteFacing) \
+	PARTICLE_ATTRIBUTE_TRANSIENT(float3,	PreviousSpriteAlignment) \
 
 // Defines modules and order of execution
 #define PARTICLE_MODULES \
diff --git a/Engine/Plugins/FX/Niagara/Shaders/Private/Stateless/NiagaraStatelessSimulationTemplate.ush b/Engine/Plugins/FX/Niagara/Shaders/Private/Stateless/NiagaraStatelessSimulationTemplate.ush
--- a/Engine/Plugins/FX/Niagara/Shaders/Private/Stateless/NiagaraStatelessSimulationTemplate.ush
+++ b/Engine/Plugins/FX/Niagara/Shaders/Private/Stateless/NiagaraStatelessSimulationTemplate.ush
@@ -6,9 +6,13 @@
 
 /////////////////////////////////////////////////////////////////////////////////////////////////////////////////////////
 
-#define PARTICLE_ATTRIBUTE(TYPE, NAME) int Permutation_##NAME##Component;
-	PARTICLE_ATTRIBUTES_OUTPUT
+#define PARTICLE_ATTRIBUTE_OUTPUT(TYPE, NAME) int Permutation_##NAME##Component;
+#define PARTICLE_ATTRIBUTE_OUTPUT_IMPLICIT(TYPE, NAME) int Permutation_##NAME##Component;
+#define PARTICLE_ATTRIBUTE_TRANSIENT(TYPE, NAME)
+	PARTICLE_ATTRIBUTES
 #undef PARTICLE_ATTRIBUTE
+#undef PARTICLE_ATTRIBUTE_OUTPUT_IMPLICIT
+#undef PARTICLE_ATTRIBUTE_TRANSIENT
 
 /////////////////////////////////////////////////////////////////////////////////////////////////////////////////////////
 // TBD this might produce sub optimal code, needs investigate and is temporary
@@ -107,15 +111,18 @@ void StatelessMain(uint3 GroupId : SV_GroupID, int GroupThreadIndex : SV_GroupIn
 	{
 		return;
 	}
-	Particle.NormalizedAge = Particle.Age / Particle.Lifetime;
 
-	Particle.PreviousAge = max(Particle.Age - Common_SimulationDeltaTime, 0.0f);
-	Particle.PreviousNormalizedAge = Particle.PreviousAge / Particle.Lifetime;
+	Particle.NormalizedAge	= Particle.Age / Particle.Lifetime;
+
+	Particle.PreviousAge			= max(Particle.Age - Common_SimulationDeltaTime, 0.0f);
+	Particle.PreviousNormalizedAge	= Particle.PreviousAge / Particle.Lifetime;
 
 	Particle.DeltaTime		= Common_SimulationDeltaTime;
 	Particle.InvDeltaTime	= Common_SimulationInvDeltaTime;
 	
 	Particle.bAlive = true;
+
+	Particle.MaterialRandom	= RandomFloat(1);
 	
 	// Run Particle Simulate
 	#define PARTICLE_MODULE(NAME) \
@@ -130,14 +137,18 @@ void StatelessMain(uint3 GroupId : SV_GroupID, int GroupThreadIndex : SV_GroupIn
 	{
 		return;
 	}
-
+	
 	// Output Particles that are still alive
 	uint OutputIndex;
 	
 	//WaveInterlockedAddScalar_(Common_GPUCountBuffer[Common_GPUCountBufferOffset], 1, Particle.OutputIndex);
 	InterlockedAdd(Common_GPUCountBuffer[Common_GPUCountBufferOffset], 1, OutputIndex);
 
-	#define PARTICLE_ATTRIBUTE(TYPE, NAME) if ( IsValidComponent(Permutation_##NAME##Component) ) { OutputComponentData(OutputIndex, Permutation_##NAME##Component, Particle.##NAME); }
-		PARTICLE_ATTRIBUTES_OUTPUT
+	#define PARTICLE_ATTRIBUTE_OUTPUT(TYPE, NAME) if ( IsValidComponent(Permutation_##NAME##Component) ) { OutputComponentData(OutputIndex, Permutation_##NAME##Component, Particle.##NAME); }
+	#define PARTICLE_ATTRIBUTE_OUTPUT_IMPLICIT(TYPE, NAME) if ( IsValidComponent(Permutation_##NAME##Component) ) { OutputComponentData(OutputIndex, Permutation_##NAME##Component, Particle.##NAME); }
+	#define PARTICLE_ATTRIBUTE_TRANSIENT(TYPE, NAME)
+		PARTICLE_ATTRIBUTES
 	#undef PARTICLE_ATTRIBUTE
+	#undef PARTICLE_ATTRIBUTE_OUTPUT_IMPLICIT
+	#undef PARTICLE_ATTRIBUTE_TRANSIENT
 }
diff --git a/Engine/Plugins/FX/Niagara/Source/Niagara/Internal/Stateless/Modules/NiagaraStatelessModule_InitializeParticle.h b/Engine/Plugins/FX/Niagara/Source/Niagara/Internal/Stateless/Modules/NiagaraStatelessModule_InitializeParticle.h
--- a/Engine/Plugins/FX/Niagara/Source/Niagara/Internal/Stateless/Modules/NiagaraStatelessModule_InitializeParticle.h
+++ b/Engine/Plugins/FX/Niagara/Source/Niagara/Internal/Stateless/Modules/NiagaraStatelessModule_InitializeParticle.h
@@ -171,7 +171,6 @@ public:
 	virtual void GetOutputVariables(TArray<FNiagaraVariableBase>& OutVariables) const override
 	{
 		const FNiagaraStatelessGlobals& StatelessGlobals = FNiagaraStatelessGlobals::Get();
-		OutVariables.AddUnique(StatelessGlobals.UniqueIDVariable);
 		OutVariables.AddUnique(StatelessGlobals.PositionVariable);
 		OutVariables.AddUnique(StatelessGlobals.ColorVariable);
 		OutVariables.AddUnique(StatelessGlobals.SpriteSizeVariable);
diff --git a/Engine/Plugins/FX/Niagara/Source/Niagara/Internal/Stateless/NiagaraStatelessCommon.h b/Engine/Plugins/FX/Niagara/Source/Niagara/Internal/Stateless/NiagaraStatelessCommon.h
--- a/Engine/Plugins/FX/Niagara/Source/Niagara/Internal/Stateless/NiagaraStatelessCommon.h
+++ b/Engine/Plugins/FX/Niagara/Source/Niagara/Internal/Stateless/NiagaraStatelessCommon.h
@@ -60,6 +60,7 @@ struct FNiagaraStatelessGlobals
 	FNiagaraVariableBase	DynamicMaterialParameters1Variable;
 	FNiagaraVariableBase	DynamicMaterialParameters2Variable;
 	FNiagaraVariableBase	DynamicMaterialParameters3Variable;
+	FNiagaraVariableBase	MaterialRandomVariable;
 	FNiagaraVariableBase	MeshIndexVariable;
 	FNiagaraVariableBase	MeshOrientationVariable;
 	FNiagaraVariableBase	PositionVariable;
diff --git a/Engine/Plugins/FX/Niagara/Source/Niagara/Internal/Stateless/NiagaraStatelessParticleSimContext.h b/Engine/Plugins/FX/Niagara/Source/Niagara/Internal/Stateless/NiagaraStatelessParticleSimContext.h
--- a/Engine/Plugins/FX/Niagara/Source/Niagara/Internal/Stateless/NiagaraStatelessParticleSimContext.h
+++ b/Engine/Plugins/FX/Niagara/Source/Niagara/Internal/Stateless/NiagaraStatelessParticleSimContext.h
@@ -26,6 +26,7 @@ enum class EParticleComponent
 	PreviousAge,
 	PreviousNormalizedAge,
 	UniqueIndex,
+	MaterialRandom,
 	Num
 };
 
@@ -34,6 +35,8 @@ class FParticleSimulationContext
 public:
 	explicit FParticleSimulationContext(const FNiagaraStatelessEmitterData* EmitterData, const void* InShaderParametersData, TConstArrayView<uint8> DynamicBufferData);
 
+	static TConstArrayView<FNiagaraVariableBase> GetRequiredComponents();
+
 	void Simulate(int32 EmitterRandomSeed, float EmitterAge, float DeltaTime, TConstArrayView<FNiagaraStatelessRuntimeSpawnInfo> SpawnInfos, FNiagaraDataBuffer* DestinationData);
 	void SimulateGPU(FRHICommandListBase& RHICmdList, int32 EmitterRandomSeed, float EmitterAge, float DeltaTime, TConstArrayView<FNiagaraStatelessRuntimeSpawnInfo> SpawnInfos, FNiagaraDataBuffer* DestinationData);
 
@@ -55,6 +58,7 @@ public:
 	float* GetParticlePreviousAge() { return static_cast<float*>(RequiredComponents[int32(EParticleComponent::PreviousAge)]); }
 	float* GetParticlePreviousNormalizedAge() { return static_cast<float*>(RequiredComponents[int32(EParticleComponent::PreviousNormalizedAge)]); }
 	int32* GetParticleUniqueIndex() { return static_cast<int32*>(RequiredComponents[int32(EParticleComponent::UniqueIndex)]); }
+	float* GetParticleMaterialRandom() { return static_cast<float*>(RequiredComponents[int32(EParticleComponent::MaterialRandom)]); }
 
 	const float* GetParticleLifetime() const { return static_cast<float*>(RequiredComponents[int32(EParticleComponent::Lifetime)]); }
 	const float* GetParticleAge() const { return static_cast<float*>(RequiredComponents[int32(EParticleComponent::Age)]); }
diff --git a/Engine/Plugins/FX/Niagara/Source/Niagara/Internal/Stateless/NiagaraStatelessParticleSimExecData.h b/Engine/Plugins/FX/Niagara/Source/Niagara/Internal/Stateless/NiagaraStatelessParticleSimExecData.h
--- a/Engine/Plugins/FX/Niagara/Source/Niagara/Internal/Stateless/NiagaraStatelessParticleSimExecData.h
+++ b/Engine/Plugins/FX/Niagara/Source/Niagara/Internal/Stateless/NiagaraStatelessParticleSimExecData.h
@@ -40,11 +40,21 @@ public:
 		uint16	Offset	: 15 = 0;
 	};
 
+	struct FRequiredComponentOffset
+	{
+		const bool IsTransient() const { return bTransient != 0; }
+		const uint32 GetOffset() const { return Offset; }
+	
+		uint16	bTransient : 1 = 0;
+		uint16	Offset : 15 = 0;
+	};
+
 	explicit FParticleSimulationExecData(const FNiagaraDataSetCompiledData& ParticleDataSetCompiledData);
 
-	TArray<FVariableOffset>	VariableComponentOffsets;		// Stored offsets per variable
-	int32					UniqueIDIndex = INDEX_NONE;		// Variable index for Unique ID, since we always generate this we need to know if we should output of keep internal
-	TArray<FCallback>		SimulateFunctions;				// Series of functions to simulate particles
+	TArray<FVariableOffset>				VariableComponentOffsets;		// Stored offsets per variable
+	TArray<FRequiredComponentOffset>	RequiredComponentOffsets;		// Stored offsets per required variable, these may be output or transient
+	int32								RequiredComponentByteSize = 0;	// Required byte size for all required components that are transient
+	TArray<FCallback>					SimulateFunctions;				// Series of functions to simulate particles
 };
 
 } //namespace NiagaraStateless
diff --git a/Engine/Plugins/FX/Niagara/Source/Niagara/Private/Stateless/NiagaraStatelessCommon.cpp b/Engine/Plugins/FX/Niagara/Source/Niagara/Private/Stateless/NiagaraStatelessCommon.cpp
--- a/Engine/Plugins/FX/Niagara/Source/Niagara/Private/Stateless/NiagaraStatelessCommon.cpp
+++ b/Engine/Plugins/FX/Niagara/Source/Niagara/Private/Stateless/NiagaraStatelessCommon.cpp
@@ -110,6 +110,7 @@ namespace NiagaraStatelessCommon
 		GGlobals.DynamicMaterialParameters1Variable = FNiagaraVariableBase(FNiagaraTypeDefinition::GetVec4Def(), TEXT("DynamicMaterialParameter1"));
 		GGlobals.DynamicMaterialParameters2Variable = FNiagaraVariableBase(FNiagaraTypeDefinition::GetVec4Def(), TEXT("DynamicMaterialParameter2"));
 		GGlobals.DynamicMaterialParameters3Variable = FNiagaraVariableBase(FNiagaraTypeDefinition::GetVec4Def(), TEXT("DynamicMaterialParameter3"));
+		GGlobals.MaterialRandomVariable				= FNiagaraVariableBase(FNiagaraTypeDefinition::GetFloatDef(), TEXT("MaterialRandom"));
 		GGlobals.MeshIndexVariable					= FNiagaraVariableBase(FNiagaraTypeDefinition::GetIntDef(), TEXT("MeshIndex"));
 		GGlobals.MeshOrientationVariable			= FNiagaraVariableBase(FNiagaraTypeDefinition::GetQuatDef(), TEXT("MeshOrientation"));
 		GGlobals.PositionVariable					= FNiagaraVariableBase(FNiagaraTypeDefinition::GetPositionDef(), TEXT("Position"));
diff --git a/Engine/Plugins/FX/Niagara/Source/Niagara/Private/Stateless/NiagaraStatelessEmitter.cpp b/Engine/Plugins/FX/Niagara/Source/Niagara/Private/Stateless/NiagaraStatelessEmitter.cpp
--- a/Engine/Plugins/FX/Niagara/Source/Niagara/Private/Stateless/NiagaraStatelessEmitter.cpp
+++ b/Engine/Plugins/FX/Niagara/Source/Niagara/Private/Stateless/NiagaraStatelessEmitter.cpp
@@ -436,6 +436,19 @@ void UNiagaraStatelessEmitter::BuildCompiledDataSet()
 	{
 		// Gather a list of all the output variables from the modules, this can change based on what is enabled / disabled
 		TArray<FNiagaraVariableBase> AvailableVariables;
+
+		// Implicitly defined variables
+		const FNiagaraStatelessGlobals& StatelessGlobals = FNiagaraStatelessGlobals::Get();
+		if ( EmitterTemplate->GetOututputComponents().Contains(StatelessGlobals.UniqueIDVariable) )
+		{
+			AvailableVariables.AddUnique(StatelessGlobals.UniqueIDVariable);
+		}
+		if (EmitterTemplate->GetOututputComponents().Contains(StatelessGlobals.MaterialRandomVariable))
+		{
+			AvailableVariables.AddUnique(StatelessGlobals.MaterialRandomVariable);
+		}
+
+		// Module based variables
 		for (UNiagaraStatelessModule* Module : Modules)
 		{
 			if (Module->IsModuleEnabled())
diff --git a/Engine/Plugins/FX/Niagara/Source/Niagara/Private/Stateless/NiagaraStatelessEmitterTemplate.cpp b/Engine/Plugins/FX/Niagara/Source/Niagara/Private/Stateless/NiagaraStatelessEmitterTemplate.cpp
--- a/Engine/Plugins/FX/Niagara/Source/Niagara/Private/Stateless/NiagaraStatelessEmitterTemplate.cpp
+++ b/Engine/Plugins/FX/Niagara/Source/Niagara/Private/Stateless/NiagaraStatelessEmitterTemplate.cpp
@@ -102,6 +102,7 @@ void UNiagaraStatelessEmitterDefault::InitModulesAndAttributes()
 	OutputComponents =
 	{
 		StatelessGlobals.UniqueIDVariable,
+		StatelessGlobals.MaterialRandomVariable,
 		StatelessGlobals.PositionVariable,
 		StatelessGlobals.CameraOffsetVariable,
 		StatelessGlobals.ColorVariable,
@@ -151,7 +152,8 @@ void UNiagaraStatelessEmitterDefault::SetShaderParameters(uint8* ShaderParameter
 	using namespace NiagaraStateless;
 	FSimulationShaderDefaultCS::FParameters* ShaderParameters	= reinterpret_cast<FSimulationShaderDefaultCS::FParameters*>(ShaderParametersBase);
 	int iComponent = 0;
-	ShaderParameters->Permutation_UniqueIDComponent					= ComponentOffsets[iComponent++];
+	ShaderParameters->Permutation_UniqueIndexComponent				= ComponentOffsets[iComponent++];
+	ShaderParameters->Permutation_MaterialRandomComponent			= ComponentOffsets[iComponent++];
 	ShaderParameters->Permutation_PositionComponent					= ComponentOffsets[iComponent++];
 	ShaderParameters->Permutation_CameraOffsetComponent				= ComponentOffsets[iComponent++];
 	ShaderParameters->Permutation_ColorComponent					= ComponentOffsets[iComponent++];
diff --git a/Engine/Plugins/FX/Niagara/Source/Niagara/Private/Stateless/NiagaraStatelessParticleSimContext.cpp b/Engine/Plugins/FX/Niagara/Source/Niagara/Private/Stateless/NiagaraStatelessParticleSimContext.cpp
--- a/Engine/Plugins/FX/Niagara/Source/Niagara/Private/Stateless/NiagaraStatelessParticleSimContext.cpp
+++ b/Engine/Plugins/FX/Niagara/Source/Niagara/Private/Stateless/NiagaraStatelessParticleSimContext.cpp
@@ -8,18 +8,6 @@
 
 namespace NiagaraStateless
 {
-	constexpr uint32 ParticleComponentOffsets[int32(EParticleComponent::Num)] =
-	{
-		0,	// uint32		- Alive
-		1,	// float		- Lifetime
-		2,	// float		- Age
-		3,	// float		- NormalizedAge
-		4,	// float		- PreviousAge
-		5,	// float		- PreviousNormalizedAge
-		6,	// int32		- UniqueIndex
-	};
-	constexpr uint32 ParticleComponentTotalSize = 7;
-
 	FUintVector4 Rand4DPCG32(FUintVector4 v)
 	{
 		// Linear congruential step.
@@ -59,6 +47,22 @@ namespace NiagaraStateless
 		check(EmitterData->ParticleSimExecData);
 	}
 
+	TConstArrayView<FNiagaraVariableBase> FParticleSimulationContext::GetRequiredComponents()
+	{
+		static const FNiagaraVariableBase RequireComponents[int(EParticleComponent::Num)] =
+		{
+			{},														// uint32		- Alive
+			{},														// float		- Lifetime
+			{},														// float		- Age
+			{},														// float		- NormalizedAge
+			{},														// float		- PreviousAge
+			{},														// float		- PreviousNormalizedAge
+			FNiagaraStatelessGlobals::Get().UniqueIDVariable,		// int32		- UniqueIndex
+			FNiagaraStatelessGlobals::Get().MaterialRandomVariable,	// float		- MaterialRandom
+		};
+		return MakeArrayView(RequireComponents);
+	}
+
 	void FParticleSimulationContext::Simulate(int32 InEmitterRandomSeed, float EmitterAge, float InDeltaTime, TConstArrayView<FNiagaraStatelessRuntimeSpawnInfo> SpawnInfos, FNiagaraDataBuffer* DestinationData)
 	{
 		NumInstances = 0;
@@ -126,22 +130,6 @@ namespace NiagaraStateless
 		EmitterRandomSeed = InEmitterRandomSeed;
 		ModuleRandomSeed = 0;
 
-		// Setup Required Components
-		{
-			static TArray<uint8>		RequiredStorage;
-			static uint32				RequiredStorageStride;
-			if (RequiredStorageStride < int32(BufferStride))
-			{
-				RequiredStorageStride = BufferStride;
-				RequiredStorage.SetNumUninitialized(ParticleComponentTotalSize * BufferStride);
-			}
-
-			for (int32 i = 0; i < int32(EParticleComponent::Num); ++i)
-			{
-				RequiredComponents[i] = RequiredStorage.GetData() + (ParticleComponentOffsets[i] * BufferStride);
-			}
-		}
-
 		// Setup optional components
 		const FParticleSimulationExecData* ExecData = EmitterData->ParticleSimExecData;
 		VariableComponents.AddDefaulted(ExecData->VariableComponentOffsets.Num());
@@ -158,10 +146,21 @@ namespace NiagaraStateless
 			}
 		}
 
-		// Redirect Unique ID (if required)
-		if (ExecData->UniqueIDIndex != INDEX_NONE)
+		// Setup required components and temporary memory
+		const uint32 TransientStorageSize = ExecData->RequiredComponentByteSize * BufferStride;
+		uint8* TransientStorage = static_cast<uint8*>(FMemory::Malloc(TransientStorageSize));
+		ON_SCOPE_EXIT { FMemory::Free(TransientStorage); };
+
+		for (int32 i = 0; i < int32(EParticleComponent::Num); ++i)
 		{
-			RequiredComponents[int(EParticleComponent::UniqueIndex)] = VariableComponents[ExecData->UniqueIDIndex];
+			if (ExecData->RequiredComponentOffsets[i].IsTransient())
+			{
+				RequiredComponents[i] = TransientStorage + (ExecData->RequiredComponentOffsets[i].GetOffset() * BufferStride);
+			}
+			else
+			{					
+				RequiredComponents[i] = VariableComponents[ExecData->RequiredComponentOffsets[i].GetOffset()];
+			}
 		}
 
 		// Setup particles
@@ -216,6 +215,7 @@ namespace NiagaraStateless
 			GetParticleNormalizedAge()[NumInstances] = Particle_Age / Particle_Lifetime;
 			GetParticlePreviousAge()[NumInstances] = PreviousAge;
 			GetParticlePreviousNormalizedAge()[NumInstances] = PreviousAge / Particle_Lifetime;
+			GetParticleMaterialRandom()[NumInstances] = RandomFloat(NumInstances, 1);
 
 			++NumInstances;
 		}
diff --git a/Engine/Plugins/FX/Niagara/Source/Niagara/Private/Stateless/NiagaraStatelessParticleSimExecData.cpp b/Engine/Plugins/FX/Niagara/Source/Niagara/Private/Stateless/NiagaraStatelessParticleSimExecData.cpp
--- a/Engine/Plugins/FX/Niagara/Source/Niagara/Private/Stateless/NiagaraStatelessParticleSimExecData.cpp
+++ b/Engine/Plugins/FX/Niagara/Source/Niagara/Private/Stateless/NiagaraStatelessParticleSimExecData.cpp
@@ -1,6 +1,7 @@
 // Copyright Epic Games, Inc. All Rights Reserved.
 
 #include "Stateless/NiagaraStatelessParticleSimExecData.h"
+#include "Stateless/NiagaraStatelessParticleSimContext.h"
 
 #include "NiagaraDataSetCompiledData.h"
 
@@ -26,11 +27,6 @@ namespace NiagaraStateless
 				check(VariableLayout.GetNumFloatComponents() == 0 && VariableLayout.GetNumHalfComponents() == 0);
 				VariableComponentOffsets[i].Type	= 1;
 				VariableComponentOffsets[i].Offset	= VariableLayout.GetInt32ComponentStart();
-
-				if (Variable.GetName() == FNiagaraStatelessGlobals::Get().UniqueIDVariable.GetName())
-				{
-					UniqueIDIndex = i;
-				}
 			}
 			else
 			{
@@ -38,6 +34,25 @@ namespace NiagaraStateless
 				checkNoEntry();
 			}
 		}
+		
+		TConstArrayView<FNiagaraVariableBase> RequiredComponents = FParticleSimulationContext::GetRequiredComponents();
+		RequiredComponentOffsets.Reserve(RequiredComponents.Num());
+		for (const FNiagaraVariableBase& RequiredComponent : RequiredComponents)
+		{
+			const int32 OutputIndex = RequiredComponent.IsValid() ? ParticleDataSetCompiledData.Variables.IndexOfByKey(RequiredComponent) : INDEX_NONE;
+			FRequiredComponentOffset& RequiredComponentOffset = RequiredComponentOffsets.AddDefaulted_GetRef();
+			if (OutputIndex == INDEX_NONE)
+			{
+				RequiredComponentOffset.bTransient	= 1;
+				RequiredComponentOffset.Offset		= RequiredComponentByteSize;
+				RequiredComponentByteSize += sizeof(float);
+			}
+			else
+			{
+				RequiredComponentOffset.bTransient	= 0;
+				RequiredComponentOffset.Offset		= OutputIndex;
+			}
+		}
 	}
 
 } //namespace NiagaraStateless
diff --git a/Engine/Plugins/FX/Niagara/Source/NiagaraShader/Internal/Stateless/NiagaraStatelessEmitterShaders.h b/Engine/Plugins/FX/Niagara/Source/NiagaraShader/Internal/Stateless/NiagaraStatelessEmitterShaders.h
--- a/Engine/Plugins/FX/Niagara/Source/NiagaraShader/Internal/Stateless/NiagaraStatelessEmitterShaders.h
+++ b/Engine/Plugins/FX/Niagara/Source/NiagaraShader/Internal/Stateless/NiagaraStatelessEmitterShaders.h
@@ -44,7 +44,8 @@ namespace NiagaraStateless
 			ADD_STATELESS_MODULE(FSpriteRotationRateModule_ShaderParameters)
 			ADD_STATELESS_MODULE(FSubUVAnimationModule_ShaderParameters)
 
-			SHADER_PARAMETER(int, Permutation_UniqueIDComponent)
+			SHADER_PARAMETER(int, Permutation_UniqueIndexComponent)
+			SHADER_PARAMETER(int, Permutation_MaterialRandomComponent)
 			SHADER_PARAMETER(int, Permutation_PositionComponent)
 			SHADER_PARAMETER(int, Permutation_CameraOffsetComponent)
 			SHADER_PARAMETER(int, Permutation_ColorComponent)
