Hierarchy Editor:
Added support for specifying an outer for the transient source root
This allows elements under the source root to rely on its outer being correct.

Fixed parameter renaming not propagating to the hierarchy elements in Niagara Scripts.
Hierarchy Script Parameters were pointing at the actual script variables that were getting replaced by new script variable, meaning the hierarchy parameters represented outdated data, and got cleaned up/deleted.
Changed how hierarchy script parameters retrieve their matching script variables; by going up the outer chain to the niagara graph, then checking for a script variable with matching identity. This ensures we only check for current data. This requires that all transient hierarchy script parameters need to be outered to the owning niagara graph. By marking up the objects as transient, we can parent them correctly without risking saving them.

Also changed how the parameter changed event propagates data. Using TInstancedStruct we can react more dynamically

diff --git a/Engine/Plugins/FX/Niagara/Source/NiagaraEditor/Private/NiagaraGraph.cpp b/Engine/Plugins/FX/Niagara/Source/NiagaraEditor/Private/NiagaraGraph.cpp
--- a/Engine/Plugins/FX/Niagara/Source/NiagaraEditor/Private/NiagaraGraph.cpp
+++ b/Engine/Plugins/FX/Niagara/Source/NiagaraEditor/Private/NiagaraGraph.cpp
@@ -1008,7 +1008,7 @@ void UNiagaraGraph::ChangeParameterType(const TArray<FNiagaraVariable>& Paramete
 	}
 	
 	InvalidateCachedParameterData();
-	OnParametersChangedDelegate.Broadcast(TOptional<FParametersChangedData>());
+	OnParametersChangedDelegate.Broadcast(TInstancedStruct<FNiagaraParametersChangedData>::Make());
 }
 
 void UNiagaraGraph::ReplaceScriptReferences(UNiagaraScript* OldScript, UNiagaraScript* NewScript)
@@ -2386,7 +2386,6 @@ TArray<FGuid> UNiagaraGraph::GetChildScriptVariableGuidsForInput(FGuid VariableG
 						ChildrenVariableGuids.Add(ChildrenScriptParameter->GetScriptVariable()->Metadata.GetVariableGuid());
 					}
 				}
-			
 			}
 			else
 			{
@@ -2456,7 +2455,8 @@ UNiagaraScriptVariable* UNiagaraGraph::AddParameter(const UNiagaraScriptVariable
 		}
  		ParameterToReferencesMap.Emplace(NewScriptVariable->Variable, {true /*bInCreated*/});
 		VariableToScriptVariable.Add(NewScriptVariable->Variable, NewScriptVariable);
-		OnParametersChangedDelegate.Broadcast(FParametersChangedData{NewScriptVariable});
+		TInstancedStruct<FNiagaraParametersChangedData> ChangeData = TInstancedStruct<FNiagaraParameterCreatedData>::Make(NewScriptVariable);
+		OnParametersChangedDelegate.Broadcast(ChangeData);
 		NotifyGraphChanged();
 		return NewScriptVariable;
 	}
@@ -2483,7 +2483,8 @@ UNiagaraScriptVariable* UNiagaraGraph::CreateScriptVariableInternal(const FNiaga
 	}
 
 	VariableToScriptVariable.Add(Parameter, NewScriptVariable);
-	OnParametersChangedDelegate.Broadcast(FParametersChangedData{NewScriptVariable});
+	TInstancedStruct<FNiagaraParametersChangedData> ChangeData = TInstancedStruct<FNiagaraParameterCreatedData>::Make(NewScriptVariable);
+	OnParametersChangedDelegate.Broadcast(ChangeData);
 
 	return NewScriptVariable;
 }
@@ -2568,7 +2569,7 @@ void UNiagaraGraph::RemoveParameter(const FNiagaraVariable& Parameter, bool bAll
 	for (FNiagaraVariable& Var : VarsToRemove)
 	{
 		VariableToScriptVariable.Remove(Var);
-		OnParametersChangedDelegate.Broadcast(TOptional<FParametersChangedData>());
+		OnParametersChangedDelegate.Broadcast(TOptional<TInstancedStruct<FNiagaraParametersChangedData>>());
 	}
 }
 
@@ -2843,7 +2844,11 @@ bool UNiagaraGraph::RenameParameter(const FNiagaraVariable& Parameter, FName New
 		*bOutMerged = bParameterMerged;
 	}
 
-	OnParametersChangedDelegate.Broadcast(TOptional<FParametersChangedData>());
+	TInstancedStruct<FNiagaraParametersChangedData> RenameData = TInstancedStruct<FNiagaraParameterRenamedData>::Make();
+	FNiagaraParameterRenamedData& Data = RenameData.GetMutable<FNiagaraParameterRenamedData>();
+	Data.OldScriptVariable = OldScriptVariable;
+	Data.NewScriptVariable = VariableToScriptVariable[NewParameter];
+	OnParametersChangedDelegate.Broadcast(RenameData);
 
 	return true;
 }
diff --git a/Engine/Plugins/FX/Niagara/Source/NiagaraEditor/Private/ViewModels/HierarchyEditor/NiagaraHierarchyScriptParametersViewModel.cpp b/Engine/Plugins/FX/Niagara/Source/NiagaraEditor/Private/ViewModels/HierarchyEditor/NiagaraHierarchyScriptParametersViewModel.cpp
--- a/Engine/Plugins/FX/Niagara/Source/NiagaraEditor/Private/ViewModels/HierarchyEditor/NiagaraHierarchyScriptParametersViewModel.cpp
+++ b/Engine/Plugins/FX/Niagara/Source/NiagaraEditor/Private/ViewModels/HierarchyEditor/NiagaraHierarchyScriptParametersViewModel.cpp
@@ -11,12 +11,23 @@
 
 #define LOCTEXT_NAMESPACE "NiagaraScriptParameterHierarchyEditor"
 
-void UNiagaraHierarchyScriptParameter::Initialize(UNiagaraScriptVariable& InParameterScriptVariable)
+void UNiagaraHierarchyScriptParameter::Initialize(const UNiagaraScriptVariable& InParameterScriptVariable)
 {
-	ParameterScriptVariable = &InParameterScriptVariable;
 	SetIdentity(FHierarchyElementIdentity({InParameterScriptVariable.Metadata.GetVariableGuid()}, {}));
 }
 
+UNiagaraScriptVariable* UNiagaraHierarchyScriptParameter::GetScriptVariable() const
+{
+	UNiagaraGraph* OwningGraph = GetTypedOuter<UNiagaraGraph>();
+	if(ensureMsgf(OwningGraph != nullptr, TEXT("When retrieving the script variable matching this hierarchy parameter, there should always be an owning graph")))
+	{
+		// If the script variable does not exist, this is an indicator it has been deleted and this hierarchy parameter should be deleted
+		return OwningGraph->GetScriptVariable(GetPersistentIdentity().Guids[0]);
+	}
+
+	return nullptr;
+}
+
 void UNiagaraHierarchyScriptParametersViewModel::Initialize(TSharedRef<FNiagaraScriptViewModel> InScriptViewModel)
 {
 	ScriptViewModelWeak = InScriptViewModel;
@@ -54,6 +65,18 @@ UHierarchyRoot* UNiagaraHierarchyScriptParametersViewModel::GetHierarchyRoot() c
 	return Cast<UNiagaraScriptSource>(ScriptData->GetSource())->NodeGraph->GetScriptParameterHierarchyRoot();
 }
 
+UObject* UNiagaraHierarchyScriptParametersViewModel::GetOuterForSourceRoot() const
+{
+	// The graph always needs to be valid as the outer
+	if(UNiagaraGraph* Graph = Cast<UNiagaraScriptSource>(ScriptViewModelWeak.Pin()->GetStandaloneScript().GetScriptData()->GetSource())->NodeGraph)
+	{
+		return Graph;
+	}
+
+	checkNoEntry();
+	return nullptr;
+}
+
 TSubclassOf<UHierarchyCategory> UNiagaraHierarchyScriptParametersViewModel::GetCategoryDataClass() const
 {
 	return UNiagaraHierarchyScriptCategory::StaticClass();
@@ -90,9 +113,8 @@ void UNiagaraHierarchyScriptParametersViewModel::PrepareSourceItems(UHierarchyRo
 		}
 		
 		UNiagaraScriptVariable* ScriptVariable = ScriptVariablePair.Value;
-
-		// since the source items are transient we need to create them here and keep them around until the end of the tool's lifetime
-		UNiagaraHierarchyScriptParameter* ScriptParameterHierarchyObject = NewObject<UNiagaraHierarchyScriptParameter>(SourceRoot);
+		
+		UNiagaraHierarchyScriptParameter* ScriptParameterHierarchyObject = NewObject<UNiagaraHierarchyScriptParameter>(SourceRoot, NAME_None, RF_Transient);
 		ScriptParameterHierarchyObject->Initialize(*ScriptVariable);
 		SourceRoot->GetChildrenMutable().Add(ScriptParameterHierarchyObject);
 	}
@@ -135,8 +157,8 @@ void UNiagaraHierarchyScriptParametersViewModel::FinalizeInternal()
 	}
 }
 
-void UNiagaraHierarchyScriptParametersViewModel::OnParametersChanged(TOptional<UNiagaraGraph::FParametersChangedData> ParametersChangedData)
-{
+void UNiagaraHierarchyScriptParametersViewModel::OnParametersChanged(TOptional<TInstancedStruct<FNiagaraParametersChangedData>> ParametersChangedData)
+{	
 	ForceFullRefresh();
 }
 
@@ -152,19 +174,13 @@ TSharedRef<SWidget> FNiagaraHierarchyScriptParameterDragDropOp::CreateCustomDeco
 
 bool FNiagaraHierarchyScriptParameterViewModel::DoesExternalDataStillExist(const UHierarchyDataRefreshContext* Context) const
 {
-	// During undo/redo it's possible the script variable becomes nullptr. If so, there is no need for this view model either
-	if(Cast<UNiagaraHierarchyScriptParameter>(GetDataMutable())->GetScriptVariable() == nullptr)
-	{
-		return false;
-	}
+	// GetScriptVariable checks the owning graph if it still exists
+	const UNiagaraScriptVariable* ScriptVariable = Cast<UNiagaraHierarchyScriptParameter>(GetDataMutable())->GetScriptVariable();
 	
-	const UNiagaraHierarchyScriptParameterRefreshContext* RefreshContext = CastChecked<UNiagaraHierarchyScriptParameterRefreshContext>(Context);
-	if(RefreshContext->GetNiagaraGraph()->GetAllMetaData().Contains(Cast<UNiagaraHierarchyScriptParameter>(GetDataMutable())->GetVariable()) == false)
+	if(ScriptVariable == nullptr)
 	{
 		return false;
-	}
-
-	const UNiagaraScriptVariable* ScriptVariable = Cast<UNiagaraHierarchyScriptParameter>(GetDataMutable())->GetScriptVariable();
+	}	
 
 	// We make sure that the variable not only still exists but also qualifies for the hierarchy (namespace can change for example)
 	if(ScriptVariable->GetIsStaticSwitch() == false && ScriptVariable->Variable.IsInNameSpace(FNiagaraConstants::ModuleNamespace) == false)
diff --git a/Engine/Plugins/FX/Niagara/Source/NiagaraEditor/Private/Widgets/SNiagaraScriptInputPreviewPanel.cpp b/Engine/Plugins/FX/Niagara/Source/NiagaraEditor/Private/Widgets/SNiagaraScriptInputPreviewPanel.cpp
--- a/Engine/Plugins/FX/Niagara/Source/NiagaraEditor/Private/Widgets/SNiagaraScriptInputPreviewPanel.cpp
+++ b/Engine/Plugins/FX/Niagara/Source/NiagaraEditor/Private/Widgets/SNiagaraScriptInputPreviewPanel.cpp
@@ -108,7 +108,8 @@ void SNiagaraScriptInputPreviewPanel::Refresh()
 		// We construct a transient list of script parameters for all the leftover parameters
 		for(UNiagaraScriptVariable* LeftoverScriptVariable : AllInputScriptVariables)
 		{
-			UNiagaraHierarchyScriptParameter* LeftoverTransientScriptParameter = NewObject<UNiagaraHierarchyScriptParameter>();
+			// They need to have an outer that has the graph as the outer to properly function, so we create transient script parameters under the hierarchy root
+			UNiagaraHierarchyScriptParameter* LeftoverTransientScriptParameter = NewObject<UNiagaraHierarchyScriptParameter>(Root.Get(), NAME_None, RF_Transient);
 			LeftoverTransientScriptParameter->Initialize(*LeftoverScriptVariable);
 			TransientLeftoverParameters.Add(LeftoverTransientScriptParameter);
 		}
@@ -358,7 +359,7 @@ void SNiagaraScriptInputPreviewPanel::RemoveDelegates()
 	}
 }
 
-void SNiagaraScriptInputPreviewPanel::OnParametersChanged(TOptional<UNiagaraGraph::FParametersChangedData> ParametersChangedData)
+void SNiagaraScriptInputPreviewPanel::OnParametersChanged(TOptional<TInstancedStruct<FNiagaraParametersChangedData>> ParametersChangedData)
 {
 	Refresh();
 }
diff --git a/Engine/Plugins/FX/Niagara/Source/NiagaraEditor/Public/NiagaraGraph.h b/Engine/Plugins/FX/Niagara/Source/NiagaraEditor/Public/NiagaraGraph.h
--- a/Engine/Plugins/FX/Niagara/Source/NiagaraEditor/Public/NiagaraGraph.h
+++ b/Engine/Plugins/FX/Niagara/Source/NiagaraEditor/Public/NiagaraGraph.h
@@ -8,6 +8,7 @@
 #include "NiagaraScript.h"
 #include "NiagaraScriptVariable.h"
 #include "NiagaraVariableMetaData.h"
+#include "StructUtils/InstancedStruct.h"
 #include "NiagaraGraph.generated.h"
 
 class UNiagaraNodeStaticSwitch;
@@ -179,18 +180,39 @@ struct FNiagaraFindInputNodeOptions
 	FGuid TargetScriptUsageId;
 };
 
+USTRUCT()
+struct FNiagaraParametersChangedData
+{
+	GENERATED_BODY()
+};
+
+USTRUCT()
+struct FNiagaraParameterCreatedData : public FNiagaraParametersChangedData
+{
+	GENERATED_BODY()
+
+	FNiagaraParameterCreatedData() = default;
+	explicit FNiagaraParameterCreatedData(UNiagaraScriptVariable* ScriptVariable) : NewScriptVariable(ScriptVariable) { }
+	
+	TWeakObjectPtr<const UNiagaraScriptVariable> NewScriptVariable;
+};
+
+USTRUCT()
+struct FNiagaraParameterRenamedData : public FNiagaraParametersChangedData
+{
+	GENERATED_BODY()
+	
+	TWeakObjectPtr<const UNiagaraScriptVariable> OldScriptVariable;
+	TWeakObjectPtr<const UNiagaraScriptVariable> NewScriptVariable;
+};
+
 UCLASS(MinimalAPI)
 class UNiagaraGraph : public UEdGraph
 {
 	GENERATED_BODY()
 
-public:
-	struct FParametersChangedData
-	{
-		const UNiagaraScriptVariable* ChangedParameter;
-	};
-	
-	DECLARE_MULTICAST_DELEGATE_OneParam(FOnParametersChanged, TOptional<FParametersChangedData>)
+public:		
+	DECLARE_MULTICAST_DELEGATE_OneParam(FOnParametersChanged, TOptional<TInstancedStruct<FNiagaraParametersChangedData>> ParametersChangedData)
 	DECLARE_MULTICAST_DELEGATE(FOnDataInterfaceChanged);
 	DECLARE_MULTICAST_DELEGATE_OneParam(FOnSubObjectSelectionChanged, const UObject*);
 	DECLARE_DELEGATE_RetVal_OneParam(TArray<UNiagaraParameterDefinitions*> /*AvailableParameterDefinitions*/, FOnGetParameterDefinitionsForDetailsCustomization, bool /*bSkipSubscribedParameterDefinitions*/)
diff --git a/Engine/Plugins/FX/Niagara/Source/NiagaraEditor/Public/ViewModels/HierarchyEditor/NiagaraHierarchyScriptParametersViewModel.h b/Engine/Plugins/FX/Niagara/Source/NiagaraEditor/Public/ViewModels/HierarchyEditor/NiagaraHierarchyScriptParametersViewModel.h
--- a/Engine/Plugins/FX/Niagara/Source/NiagaraEditor/Public/ViewModels/HierarchyEditor/NiagaraHierarchyScriptParametersViewModel.h
+++ b/Engine/Plugins/FX/Niagara/Source/NiagaraEditor/Public/ViewModels/HierarchyEditor/NiagaraHierarchyScriptParametersViewModel.h
@@ -31,16 +31,13 @@ class UNiagaraHierarchyScriptParameter : public UHierarchyItem
 	GENERATED_BODY()
 
 public:
-	void Initialize(UNiagaraScriptVariable& InParameterScriptVariable);
+	void Initialize(const UNiagaraScriptVariable& InParameterScriptVariable);
 
 	virtual FString ToString() const override { return GetVariable().GetName().ToString(); }
-	FText GetTooltip() const { return ParameterScriptVariable->Metadata.Description; }
+	FText GetTooltip() const { return GetScriptVariable()->Metadata.Description; }
 	
-	UNiagaraScriptVariable* GetScriptVariable() const { return ParameterScriptVariable; }
-	FNiagaraVariable GetVariable() const { return ParameterScriptVariable->Variable; }
-private:
-	UPROPERTY()
-	TObjectPtr<UNiagaraScriptVariable> ParameterScriptVariable;
+	UNiagaraScriptVariable* GetScriptVariable() const;
+	FNiagaraVariable GetVariable() const { return GetScriptVariable()->Variable; }
 };
 
 /** The category class used for the script hierarchy editor. It lets us add additional data later on. */
@@ -63,6 +60,8 @@ public:
 
 	/** UHierarchyViewModelBase */
 	virtual UHierarchyRoot* GetHierarchyRoot() const override;
+	/** We set the outer for the source root to the graph, so that both source & hierarchy items have the same outer. */
+	virtual UObject* GetOuterForSourceRoot() const override;
 	virtual TSubclassOf<UHierarchyCategory> GetCategoryDataClass() const override;
 	virtual TSharedPtr<FHierarchyElementViewModel> CreateCustomViewModelForElement(UHierarchyElement* Element, TSharedPtr<FHierarchyElementViewModel> Parent) override;
 	virtual void PrepareSourceItems(UHierarchyRoot* SourceRoot, TSharedPtr<FHierarchyRootViewModel> SourceRootViewModel) override;
@@ -71,7 +70,7 @@ public:
 	virtual bool SupportsDetailsPanel() override { return true; }
 	virtual void FinalizeInternal() override;
 private:
-	void OnParametersChanged(TOptional<UNiagaraGraph::FParametersChangedData> ParametersChangedData);
+	void OnParametersChanged(TOptional<TInstancedStruct<FNiagaraParametersChangedData>> ParametersChangedData);
 private:
 	TWeakPtr<FNiagaraScriptViewModel> ScriptViewModelWeak;
 };
diff --git a/Engine/Plugins/FX/Niagara/Source/NiagaraEditor/Public/Widgets/SNiagaraScriptInputPreviewPanel.h b/Engine/Plugins/FX/Niagara/Source/NiagaraEditor/Public/Widgets/SNiagaraScriptInputPreviewPanel.h
--- a/Engine/Plugins/FX/Niagara/Source/NiagaraEditor/Public/Widgets/SNiagaraScriptInputPreviewPanel.h
+++ b/Engine/Plugins/FX/Niagara/Source/NiagaraEditor/Public/Widgets/SNiagaraScriptInputPreviewPanel.h
@@ -33,7 +33,7 @@ public:
 private:
 	TSharedRef<ITableRow> OnGenerateRow(UHierarchyElement* Item, const TSharedRef<STableViewBase>& TableViewBase) const;
 	void OnGetChildren(UHierarchyElement* Item, TArray<UHierarchyElement*>& OutChildren) const;
-	void OnParametersChanged(TOptional<UNiagaraGraph::FParametersChangedData> ParametersChangedData);
+	void OnParametersChanged(TOptional<TInstancedStruct<FNiagaraParametersChangedData>> ParametersChangedData);
 
 	FReply SummonHierarchyEditor() const;
 
diff --git a/Engine/Source/Editor/HierarchyEditor/Private/Widgets/SHierarchyEditor.cpp b/Engine/Source/Editor/HierarchyEditor/Private/Widgets/SHierarchyEditor.cpp
--- a/Engine/Source/Editor/HierarchyEditor/Private/Widgets/SHierarchyEditor.cpp
+++ b/Engine/Source/Editor/HierarchyEditor/Private/Widgets/SHierarchyEditor.cpp
@@ -483,8 +483,8 @@ void SHierarchyEditor::Construct(const FArguments& InArgs, TObjectPtr<UHierarchy
 		HierarchyViewModel->Initialize();
 	}
 	
-	SourceRoot = NewObject<UHierarchyRoot>(GetTransientPackage(), MakeUniqueObjectName(GetTransientPackage(), UHierarchyRoot::StaticClass()));
-	TSharedPtr<FHierarchyElementViewModel> ViewModel = HierarchyViewModel->CreateViewModelForElement(SourceRoot, nullptr);
+	SourceRoot.Reset(NewObject<UHierarchyRoot>(HierarchyViewModel->GetOuterForSourceRoot(), MakeUniqueObjectName(HierarchyViewModel->GetOuterForSourceRoot(), UHierarchyRoot::StaticClass()), RF_Transient));
+	TSharedPtr<FHierarchyElementViewModel> ViewModel = HierarchyViewModel->CreateViewModelForElement(SourceRoot.Get(), nullptr);
 	SourceRootViewModel = StaticCastSharedPtr<FHierarchyRootViewModel>(ViewModel);	
 	SourceRootViewModel->Initialize();
 	SourceRootViewModel->AddChildFilter(FHierarchyElementViewModel::FOnFilterChild::CreateSP(this, &SHierarchyEditor::FilterForSourceSection));
@@ -793,7 +793,7 @@ SHierarchyEditor::~SHierarchyEditor()
 void SHierarchyEditor::RefreshSourceItems()
 {
 	SourceRoot->EmptyAllData();
-	HierarchyViewModel->PrepareSourceItems(SourceRoot, SourceRootViewModel);
+	HierarchyViewModel->PrepareSourceItems(SourceRoot.Get(), SourceRootViewModel);
 	SourceRootViewModel->SyncViewModelsToData();
 	RefreshSourceView(false);
 	RefreshSectionsView();
@@ -1715,16 +1715,6 @@ FSlateColor SHierarchyEditor::GetRootIconColor() const
 	return FLinearColor(0.2f, 0.2f, 0.2f, 0.5f);
 }
 
-FString SHierarchyEditor::GetReferencerName() const
-{
-	return TEXT(" Hierarchy");
-}
-
-void SHierarchyEditor::AddReferencedObjects(FReferenceCollector& Collector)
-{
-	Collector.AddReferencedObject(SourceRoot);
-}
-
 void SHierarchyEditor::NotifyPostChange(const FPropertyChangedEvent& PropertyChangedEvent,	FProperty* PropertyThatChanged)
 {
 	HierarchyViewModel->OnHierarchyPropertiesChanged().Broadcast();
diff --git a/Engine/Source/Editor/HierarchyEditor/Public/HierarchyViewModelBase.h b/Engine/Source/Editor/HierarchyEditor/Public/HierarchyViewModelBase.h
--- a/Engine/Source/Editor/HierarchyEditor/Public/HierarchyViewModelBase.h
+++ b/Engine/Source/Editor/HierarchyEditor/Public/HierarchyViewModelBase.h
@@ -12,6 +12,8 @@
 #include "TickableEditorObject.h"
 #include "Engine/TimerHandle.h"
 #include "Templates/SubclassOf.h"
+#include "UObject/Package.h"
+#include "UObject/UObjectGlobals.h"
 #include "HierarchyViewModelBase.generated.h"
 
 /** HIERARCHY EDITOR
@@ -458,6 +460,8 @@ public:
 public:
 	/** The hierarchy root the widget is editing. This should point to persistent data stored somewhere else as the serialized root of the hierarchy. */
 	virtual UHierarchyRoot* GetHierarchyRoot() const PURE_VIRTUAL(UHierarchyViewModelBase::GetHierarchyRoot, return nullptr;);
+	/** The outer for the transient source root creation can be overridden. */
+	virtual UObject* GetOuterForSourceRoot() const { return GetTransientPackage(); }
 	/** Prepares the items we want to create a hierarchy for. Primary purpose is to add children to the source root to gather the items to display in the source panel.
 	 * The root view model is also given as a way to forcefully sync view models to access additional functionality, if needed */
 	virtual void PrepareSourceItems(UHierarchyRoot* SourceRoot, TSharedPtr<FHierarchyRootViewModel> SourceRootViewModel) PURE_VIRTUAL(UHierarchyViewModelBase::PrepareSourceItems,);
diff --git a/Engine/Source/Editor/HierarchyEditor/Public/Widgets/SHierarchyEditor.h b/Engine/Source/Editor/HierarchyEditor/Public/Widgets/SHierarchyEditor.h
--- a/Engine/Source/Editor/HierarchyEditor/Public/Widgets/SHierarchyEditor.h
+++ b/Engine/Source/Editor/HierarchyEditor/Public/Widgets/SHierarchyEditor.h
@@ -81,7 +81,7 @@ private:
 	mutable bool bDraggedOn = false;
 };
 
-class HIERARCHYEDITOR_API SHierarchyEditor : public SCompoundWidget, public FGCObject, public FNotifyHook
+class HIERARCHYEDITOR_API SHierarchyEditor : public SCompoundWidget, public FNotifyHook
 {
 	DECLARE_DELEGATE_RetVal_OneParam(TSharedRef<SWidget>, FOnGenerateRowContentWidget, TSharedRef<FHierarchyElementViewModel> HierarchyElement);
 	DECLARE_DELEGATE_RetVal_OneParam(TSharedRef<SWidget>, FOnGenerateCustomDetailsPanelNameWidget, TSharedPtr<FHierarchyElementViewModel> HierarchyElement);
@@ -209,9 +209,6 @@ private:
 	void OnRootDragLeave(const FDragDropEvent& DragDropEvent) const;
 	FSlateColor GetRootIconColor() const;
 
-	virtual FString GetReferencerName() const override; 
-	virtual void AddReferencedObjects(FReferenceCollector& Collector) override;
-
 	virtual void NotifyPostChange(const FPropertyChangedEvent& PropertyChangedEvent, FProperty* PropertyThatChanged) override;
 private:
 	TWeakObjectPtr<UHierarchyViewModelBase> HierarchyViewModel;
@@ -225,7 +222,7 @@ private:
 	mutable TWeakPtr<FHierarchyElementViewModel> SelectedDetailsPanelItemViewModel;
 
 private:
-	TObjectPtr<UHierarchyRoot> SourceRoot = nullptr;
+	TStrongObjectPtr<UHierarchyRoot> SourceRoot;
 	TSharedPtr<FHierarchyRootViewModel> SourceRootViewModel;
 	TSharedPtr<STreeView<TSharedPtr<FHierarchyElementViewModel>>> SourceTreeView;
 	TSharedPtr<STreeView<TSharedPtr<FHierarchyElementViewModel>>> HierarchyTreeView;
