Fixing dangling garbage reference in NiagaraDataInterfaceStaticMesh

diff --git a/Engine/Plugins/FX/Niagara/Source/Niagara/Internal/DataInterface/NiagaraDataInterfaceStaticMesh.h b/Engine/Plugins/FX/Niagara/Source/Niagara/Internal/DataInterface/NiagaraDataInterfaceStaticMesh.h
--- a/Engine/Plugins/FX/Niagara/Source/Niagara/Internal/DataInterface/NiagaraDataInterfaceStaticMesh.h
+++ b/Engine/Plugins/FX/Niagara/Source/Niagara/Internal/DataInterface/NiagaraDataInterfaceStaticMesh.h
@@ -81,7 +81,7 @@ protected:
 
 	/** The source component from which to sample. Takes precedence over the direct mesh. Not exposed to the user, only indirectly accessible from blueprints. */
 	UPROPERTY(Transient)
-	TObjectPtr<UStaticMeshComponent> SourceComponent;
+	TWeakObjectPtr<UStaticMeshComponent> SourceComponent;
 public:
 
 	/** Array of filters the can be used to limit sampling to certain sections of the mesh. */
diff --git a/Engine/Plugins/FX/Niagara/Source/Niagara/Private/DataInterface/NiagaraDataInterfaceStaticMesh.cpp b/Engine/Plugins/FX/Niagara/Source/Niagara/Private/DataInterface/NiagaraDataInterfaceStaticMesh.cpp
--- a/Engine/Plugins/FX/Niagara/Source/Niagara/Private/DataInterface/NiagaraDataInterfaceStaticMesh.cpp
+++ b/Engine/Plugins/FX/Niagara/Source/Niagara/Private/DataInterface/NiagaraDataInterfaceStaticMesh.cpp
@@ -3438,9 +3438,9 @@ UStaticMesh* UNiagaraDataInterfaceStaticMesh::ResolveStaticMesh(UObject*& OutCom
 	const bool bTryDefaultMesh = SourceMode == ENDIStaticMesh_SourceMode::Default || SourceMode == ENDIStaticMesh_SourceMode::DefaultMeshOnly;
 
 	UStaticMesh* OutMesh = nullptr;
-	if (bTrySource && ::IsValid(SourceComponent))
+	if (bTrySource && SourceComponent.IsValid())
 	{
-		OutComponent = SourceComponent;
+		OutComponent = SourceComponent.Get();
 		OutMesh = SourceComponent->GetStaticMesh();
 	}
 	else if (bTrySource && SoftSourceActor.Get())
@@ -3514,7 +3514,7 @@ void UNiagaraDataInterfaceStaticMesh::BindSourceDelegates()
 	{
 		Source->OnEndPlay.AddDynamic(this, &UNiagaraDataInterfaceStaticMesh::OnSourceEndPlay);
 	}
-	else if (SourceComponent)
+	else if (SourceComponent.IsValid())
 	{
 		UE_CLOG(!UObjectBaseUtility::IsGarbageEliminationEnabled(),
 			LogNiagara, Warning, TEXT("%s: Unable to bind OnEndPlay for actor-less source component %s, this may extend the lifetime of the component"),
