- Niagara Script Stat Changes
- Add drop shadow, change error to orange to make it easier to read
- Fix recomile button to work
- Add Copy to CSV button

#rnx
#tests Use script stats window, copy data out, etc.

diff --git a/Engine/Plugins/FX/Niagara/Source/NiagaraEditor/Private/ViewModels/NiagaraScriptStatsViewModel.cpp b/Engine/Plugins/FX/Niagara/Source/NiagaraEditor/Private/ViewModels/NiagaraScriptStatsViewModel.cpp
--- a/Engine/Plugins/FX/Niagara/Source/NiagaraEditor/Private/ViewModels/NiagaraScriptStatsViewModel.cpp
+++ b/Engine/Plugins/FX/Niagara/Source/NiagaraEditor/Private/ViewModels/NiagaraScriptStatsViewModel.cpp
@@ -14,6 +14,7 @@
 #include "NiagaraShader.h"
 #include "Styling/AppStyle.h"
 #include "Framework/MultiBox/MultiBoxBuilder.h"
+#include "HAL/PlatformApplicationMisc.h"
 #include "Interfaces/ITargetPlatformManagerModule.h"
 #include "Types/SlateEnums.h"
 #include "Widgets/SCompoundWidget.h"
@@ -90,7 +91,7 @@ namespace NiagaraScriptStatsLocal
 		{
 			return FColor::Yellow;
 		}
-		return bHasError ? FColor::Red : FColor::Green;
+		return bHasError ? FColor::Orange : FColor::Green;
 	}
 
 	bool ScriptStatusToWrapText(bool bIsCompiling, bool bHasError)
@@ -306,6 +307,8 @@ public:
 				.ColorAndOpacity(this, &SNiagaraScriptStatsViewRow::GetCellColor, ColumnName)
 				.Text(this, &SNiagaraScriptStatsViewRow::GetCellText, ColumnName, false)
 				.AutoWrapText(this, &SNiagaraScriptStatsViewRow::GetCellWrapText, ColumnName)
+				.ShadowOffset(FVector2D::UnitVector)
+				.ShadowColorAndOpacity(FLinearColor::Black)
 			];
 	}
 
@@ -351,6 +354,42 @@ private:
 
 /////////////////////////////////////////////////////////////////////////////////////
 
+class SNiagaraScriptStatsListView : public SListView<TSharedPtr<FGuid>>
+{
+public:
+	SLATE_BEGIN_ARGS(SNiagaraScriptStatsListView) {}
+		SLATE_ARGUMENT(TWeakPtr<FNiagaraScriptStatsViewModel>, WeakViewModel)
+		SLATE_ARGUMENT(TSharedPtr<SHeaderRow>, HeaderRow)
+	SLATE_END_ARGS()
+
+	void Construct(const FArguments& InArgs)
+	{
+		WeakViewModel = InArgs._WeakViewModel;
+
+		auto ViewModel = WeakViewModel.Pin();
+		check(ViewModel != nullptr);
+
+		SListView::Construct(SListView::FArguments()
+			.ListItemsSource(ViewModel->GetGridRowIDs())
+			.OnGenerateRow(this, &SNiagaraScriptStatsListView::MakeRowWidget)
+			.Visibility(EVisibility::Visible)
+			.HeaderRow(InArgs._HeaderRow)
+		);
+	}
+
+	TSharedRef<ITableRow> MakeRowWidget(const TSharedPtr<FGuid> PtrRowID, const TSharedRef<STableViewBase>& OwnerTable) const
+	{
+		return SNew(SNiagaraScriptStatsViewRow, OwnerTable)
+			.PtrRowID(PtrRowID)
+			.WeakViewModel(WeakViewModel);
+	}
+
+private:
+	TWeakPtr<FNiagaraScriptStatsViewModel>	WeakViewModel;
+};
+
+/////////////////////////////////////////////////////////////////////////////////////
+
 class SNiagaraScriptStatsWidget : public SCompoundWidget
 {
 public:
@@ -416,8 +455,17 @@ public:
 						SNew(SButton)
 						.VAlign(VAlign_Center)
 						.OnClicked(this, &SNiagaraScriptStatsWidget::OnForceRecompile)
-						.Text(FText::FromString(TEXT("Recompile")))
-						.ToolTipText(FText::FromString(TEXT("Forces a recompile to run")))
+						.Text(LOCTEXT("Recompile", "Recompile"))
+						.ToolTipText(LOCTEXT("RecompileTooltip", "Forces a recompile to run"))
+					]
+					+SHorizontalBox::Slot()
+					.AutoWidth()
+					[
+						SNew(SButton)
+						.VAlign(VAlign_Center)
+						.OnClicked(this, &SNiagaraScriptStatsWidget::OnCopyToClipboard)
+						.Text(LOCTEXT("CopyToCSV", "Copy To CSV"))
+						.ToolTipText(LOCTEXT("CopyToCSVTooltip", "Copies the data to the clipboard, formatted as CSV"))
 					]
 				]
 			]
@@ -450,11 +498,8 @@ public:
 									.Orientation(Orient_Horizontal)
 									+ SScrollBox::Slot()
 									[
-										SAssignNew(StatsInfoList, SListView<TSharedPtr<FGuid>>)
-										.ListItemsSource(ViewModel->GetGridRowIDs())
-										.OnGenerateRow(this, &SNiagaraScriptStatsWidget::MakeRowWidget)
-										.Visibility(EVisibility::Visible)
-										.SelectionMode(ESelectionMode::Single)
+										SAssignNew(StatsInfoList, SNiagaraScriptStatsListView)
+										.WeakViewModel(WeakViewModel)
 										.HeaderRow(PlatformColumnHeader)
 									]
 								]
@@ -582,10 +627,31 @@ public:
 		return FReply::Handled();
 	}
 
+	FReply OnCopyToClipboard() const
+	{
+		auto ViewModel = WeakViewModel.Pin();
+		if (ViewModel.IsValid())
+		{
+			FString ClipboardString;
+
+			for (const TSharedPtr<FGuid>& RowID : *ViewModel->GetGridRowIDs())
+			{
+				for (const SHeaderRow::FColumn& Column : PlatformColumnHeader->GetColumns())
+				{
+					const auto CellDetails = ViewModel->GetCellDetails(Column.ColumnId, *RowID);
+					ClipboardString.Append(CellDetails.CellText.ToString());
+					ClipboardString.Append("\n");
+				}
+			}
+			FPlatformApplicationMisc::ClipboardCopy(*ClipboardString);
+		}
+		return FReply::Handled();
+	}
+
 private:
 	TSharedPtr<SHeaderRow> PlatformColumnHeader;
 	TWeakPtr<FNiagaraScriptStatsViewModel> WeakViewModel;
-	TSharedPtr<SListView<TSharedPtr<FGuid>>> StatsInfoList;
+	TSharedPtr<SNiagaraScriptStatsListView> StatsInfoList;
 };
 
 /////////////////////////////////////////////////////////////////////////////////////
@@ -609,7 +675,7 @@ void FNiagaraScriptStatsViewModel::Initialize(TWeakPtr<FNiagaraSystemViewModel>
 		SystemViewModel->OnEmitterHandleViewModelsChanged().AddSP(this, &FNiagaraScriptStatsViewModel::RefreshView);
 		SystemViewModel->GetSelectionViewModel()->OnSystemIsSelectedChanged().AddSP(this, &FNiagaraScriptStatsViewModel::RefreshView);
 		SystemViewModel->GetSelectionViewModel()->OnEmitterHandleIdSelectionChanged().AddSP(this, &FNiagaraScriptStatsViewModel::RefreshView);
-		SystemViewModel->OnSystemCompiled().AddSP(this, &FNiagaraScriptStatsViewModel::OnForceRecompile);
+		SystemViewModel->OnSystemCompiled().AddSP(this, &FNiagaraScriptStatsViewModel::OnSystemCompiled);
 	}
 
 	RefreshView();
@@ -859,8 +925,17 @@ void FNiagaraScriptStatsViewModel::CancelCompilations()
 void FNiagaraScriptStatsViewModel::OnForceRecompile()
 {
 	CancelCompilations();
+
+	if (auto SystemViewModel = WeakSystemViewModel.Pin())
+	{
+		UNiagaraSystem* NiagaraSystem = &SystemViewModel->GetSystem();
+		NiagaraSystem->UNiagaraSystem::RequestCompile(true);
+	}
+}
+
+void FNiagaraScriptStatsViewModel::OnSystemCompiled()
+{
 	RefreshView();
 }
 
 #undef LOCTEXT_NAMESPACE
-
diff --git a/Engine/Plugins/FX/Niagara/Source/NiagaraEditor/Private/ViewModels/NiagaraScriptStatsViewModel.h b/Engine/Plugins/FX/Niagara/Source/NiagaraEditor/Private/ViewModels/NiagaraScriptStatsViewModel.h
--- a/Engine/Plugins/FX/Niagara/Source/NiagaraEditor/Private/ViewModels/NiagaraScriptStatsViewModel.h
+++ b/Engine/Plugins/FX/Niagara/Source/NiagaraEditor/Private/ViewModels/NiagaraScriptStatsViewModel.h
@@ -65,6 +65,8 @@ public:
 
 	void OnForceRecompile();
 
+	void OnSystemCompiled();
+
 protected:
 	void BuildShaderPlatformDetails();
 
