[Backout] - CL39482572
#rnx
Original CL Desc
-----------------------------------------------------------------
Fixing dangling garbage reference in NiagaraDataInterfaceStaticMesh

diff --git a/Engine/Plugins/FX/Niagara/Source/Niagara/Internal/DataInterface/NiagaraDataInterfaceStaticMesh.h b/Engine/Plugins/FX/Niagara/Source/Niagara/Internal/DataInterface/NiagaraDataInterfaceStaticMesh.h
--- a/Engine/Plugins/FX/Niagara/Source/Niagara/Internal/DataInterface/NiagaraDataInterfaceStaticMesh.h
+++ b/Engine/Plugins/FX/Niagara/Source/Niagara/Internal/DataInterface/NiagaraDataInterfaceStaticMesh.h
@@ -81,7 +81,7 @@ protected:
 
 	/** The source component from which to sample. Takes precedence over the direct mesh. Not exposed to the user, only indirectly accessible from blueprints. */
 	UPROPERTY(Transient)
-	TWeakObjectPtr<UStaticMeshComponent> SourceComponent;
+	TObjectPtr<UStaticMeshComponent> SourceComponent;
 public:
 
 	/** Array of filters the can be used to limit sampling to certain sections of the mesh. */
diff --git a/Engine/Plugins/FX/Niagara/Source/Niagara/Private/DataInterface/NiagaraDataInterfaceStaticMesh.cpp b/Engine/Plugins/FX/Niagara/Source/Niagara/Private/DataInterface/NiagaraDataInterfaceStaticMesh.cpp
--- a/Engine/Plugins/FX/Niagara/Source/Niagara/Private/DataInterface/NiagaraDataInterfaceStaticMesh.cpp
+++ b/Engine/Plugins/FX/Niagara/Source/Niagara/Private/DataInterface/NiagaraDataInterfaceStaticMesh.cpp
@@ -3438,9 +3438,9 @@ UStaticMesh* UNiagaraDataInterfaceStaticMesh::ResolveStaticMesh(UObject*& OutCom
 	const bool bTryDefaultMesh = SourceMode == ENDIStaticMesh_SourceMode::Default || SourceMode == ENDIStaticMesh_SourceMode::DefaultMeshOnly;
 
 	UStaticMesh* OutMesh = nullptr;
-	if (bTrySource && SourceComponent.IsValid())
+	if (bTrySource && ::IsValid(SourceComponent))
 	{
-		OutComponent = SourceComponent.Get();
+		OutComponent = SourceComponent;
 		OutMesh = SourceComponent->GetStaticMesh();
 	}
 	else if (bTrySource && SoftSourceActor.Get())
@@ -3514,7 +3514,7 @@ void UNiagaraDataInterfaceStaticMesh::BindSourceDelegates()
 	{
 		Source->OnEndPlay.AddDynamic(this, &UNiagaraDataInterfaceStaticMesh::OnSourceEndPlay);
 	}
-	else if (SourceComponent.IsValid())
+	else if (SourceComponent)
 	{
 		UE_CLOG(!UObjectBaseUtility::IsGarbageEliminationEnabled(),
 			LogNiagara, Warning, TEXT("%s: Unable to bind OnEndPlay for actor-less source component %s, this may extend the lifetime of the component"),
