- Sim cache fix for GPU emitter resetting when detaching full sim cache

diff --git a/Engine/Plugins/FX/Niagara/Source/Niagara/Private/NiagaraSimCacheHelper.h b/Engine/Plugins/FX/Niagara/Source/Niagara/Private/NiagaraSimCacheHelper.h
--- a/Engine/Plugins/FX/Niagara/Source/Niagara/Private/NiagaraSimCacheHelper.h
+++ b/Engine/Plugins/FX/Niagara/Source/Niagara/Private/NiagaraSimCacheHelper.h
@@ -846,12 +846,17 @@ struct FNiagaraSimCacheHelper
 		}
 
 		FNiagaraGpuComputeDispatchInterface* DispathInterface = EmitterInstance.GetParentSystemInstance()->GetComputeDispatchInterface();
+		FNiagaraComputeExecutionContext* GPUExecContext = EmitterInstance.GetGPUContext();
+		++GPUExecContext->ParticleCountReadFence;
+		GPUExecContext->bResetPending_GT = false;
+
 		ENQUEUE_RENDER_COMMAND(NiagaraSimCacheGpuReadFrame)(
-			[DispathInterface, GPUExecContext=EmitterInstance.GetGPUContext(), FrameFraction=InFrameFraction, FrameDeltaSeconds=InFrameDeltaSeconds, SimDeltaSeconds=InSimDeltaSeconds, RebaseATransform=InRebaseATransform, RebaseBTransform=InRebaseBTransform, CacheLayout=&InCacheLayout, CacheBufferA=&InCacheBufferA, CacheBufferB=&InCacheBufferB, DataSet=&InDataSet, PendingCommandsCounter=&InPendingCommandsCounter](FRHICommandListImmediate& RHICmdList)
+			[DispathInterface, GPUExecContext, ParticleCountFence=GPUExecContext->ParticleCountReadFence, FrameFraction=InFrameFraction, FrameDeltaSeconds=InFrameDeltaSeconds, SimDeltaSeconds=InSimDeltaSeconds, RebaseATransform=InRebaseATransform, RebaseBTransform=InRebaseBTransform, CacheLayout=&InCacheLayout, CacheBufferA=&InCacheBufferA, CacheBufferB=&InCacheBufferB, DataSet=&InDataSet, PendingCommandsCounter=&InPendingCommandsCounter](FRHICommandListImmediate& RHICmdList)
 			{
 				const uint32 NumInstances = FMath::Min(CacheBufferA->NumInstances, DataSet->GetMaxAllocationCount());
 
 				// Set Instance Count
+				FNiagaraDataBuffer& DataBuffer = DataSet->GetCurrentDataChecked();
 				{
 					FNiagaraGPUInstanceCountManager& CountManager = DispathInterface->GetGPUInstanceCounterManager();
 					if (GPUExecContext->CountOffset_RT == INDEX_NONE)
@@ -859,6 +864,16 @@ struct FNiagaraSimCacheHelper
 						GPUExecContext->CountOffset_RT = CountManager.AcquireOrAllocateEntry(RHICmdList);
 					}
 
+					uint32 CurrentCountOffset = DataBuffer.GetGPUInstanceCountBufferOffset();
+					if (CurrentCountOffset != INDEX_NONE && CurrentCountOffset != GPUExecContext->CountOffset_RT)
+					{
+						CountManager.FreeEntry(CurrentCountOffset);
+						DataBuffer.SetGPUInstanceCountBufferOffset(GPUExecContext->CountOffset_RT);
+					}
+
+					GPUExecContext->CurrentNumInstances_RT = NumInstances;
+					GPUExecContext->ParticleCountWriteFence = ParticleCountFence;
+
 					const FRWBuffer& CountBuffer = CountManager.GetInstanceCountBuffer();
 					const TPair<uint32, int32> DataToSet(GPUExecContext->CountOffset_RT, NumInstances);
 					RHICmdList.Transition(FRHITransitionInfo(CountBuffer.UAV, FNiagaraGPUInstanceCountManager::kCountBufferDefaultState, ERHIAccess::UAVCompute));
@@ -867,7 +882,6 @@ struct FNiagaraSimCacheHelper
 				}
 
 				// Copy instance counts
-				FNiagaraDataBuffer& DataBuffer = DataSet->GetCurrentDataChecked();
 				DataBuffer.AllocateGPU(RHICmdList, NumInstances, DispathInterface->GetFeatureLevel(), TEXT("NiagaraSimCache"));
 				DataBuffer.SetNumInstances(NumInstances);
 				DataBuffer.SetGPUDataReadyStage(ENiagaraGpuComputeTickStage::PreInitViews);
