- Fix missing initial state for some buffers from CL 39357771

#rnx

diff --git a/Engine/Plugins/FX/Niagara/Source/Niagara/Private/NiagaraDataInterfaceSkeletalMesh.cpp b/Engine/Plugins/FX/Niagara/Source/Niagara/Private/NiagaraDataInterfaceSkeletalMesh.cpp
--- a/Engine/Plugins/FX/Niagara/Source/Niagara/Private/NiagaraDataInterfaceSkeletalMesh.cpp
+++ b/Engine/Plugins/FX/Niagara/Source/Niagara/Private/NiagaraDataInterfaceSkeletalMesh.cpp
@@ -980,8 +980,9 @@ void FSkeletalMeshGpuSpawnStaticBuffers::InitRHI(FRHICommandListBase& RHICmdList
 		BufferTriangleUniformSamplerProbAliasSRV = RHICmdList.CreateShaderResourceView(
 			BufferTriangleUniformSamplerProbAliasRHI, 
 			FRHIViewDesc::CreateBufferSRV()
-				.SetType(FRHIViewDesc::EBufferType::Typed)
-				.SetFormat(PF_R32_UINT));
+			.SetType(FRHIViewDesc::EBufferType::Typed)
+			.SetFormat(PF_R32_UINT)
+		);
 #if STATS
 		GPUMemoryUsage += SizeByte;
 #endif
@@ -1002,14 +1003,16 @@ void FSkeletalMeshGpuSpawnStaticBuffers::InitRHI(FRHICommandListBase& RHICmdList
 			const FRHIBufferCreateDesc CreateDesc =
 				FRHIBufferCreateDesc::CreateVertex(TEXT("SampleRegionsProbAliasBuffer"), SampleRegionsProbAlias.Num() * SampleRegionsProbAlias.GetTypeSize())
 				.AddUsage(EBufferUsageFlags::Static | EBufferUsageFlags::ShaderResource)
-				.SetInitActionResourceArray(&SampleRegionsProbAlias);
+				.SetInitActionResourceArray(&SampleRegionsProbAlias)
+				.DetermineInitialState();
 
 			SampleRegionsProbAliasBuffer = RHICmdList.CreateBuffer(CreateDesc);
 			SampleRegionsProbAliasSRV = RHICmdList.CreateShaderResourceView(
 				SampleRegionsProbAliasBuffer, 
 				FRHIViewDesc::CreateBufferSRV()
-					.SetType(FRHIViewDesc::EBufferType::Typed)
-					.SetFormat(PF_R32_UINT));
+				.SetType(FRHIViewDesc::EBufferType::Typed)
+				.SetFormat(PF_R32_UINT)
+			);
 #if STATS
 			GPUMemoryUsage += SampleRegionsProbAlias.Num() * SampleRegionsProbAlias.GetTypeSize();
 #endif
@@ -1023,14 +1026,16 @@ void FSkeletalMeshGpuSpawnStaticBuffers::InitRHI(FRHICommandListBase& RHICmdList
 			const FRHIBufferCreateDesc CreateDesc =
 				FRHIBufferCreateDesc::CreateVertex(TEXT("SampleRegionsTriangleIndicesBuffer"), SampleRegionsTriangleIndicies.Num() * SampleRegionsTriangleIndicies.GetTypeSize())
 				.AddUsage(EBufferUsageFlags::Static | EBufferUsageFlags::ShaderResource)
-				.SetInitActionResourceArray(&SampleRegionsTriangleIndicies);
+				.SetInitActionResourceArray(&SampleRegionsTriangleIndicies)
+				.DetermineInitialState();
 
 			SampleRegionsTriangleIndicesBuffer = RHICmdList.CreateBuffer(CreateDesc);
-		SampleRegionsTriangleIndicesSRV = RHICmdList.CreateShaderResourceView(
-			SampleRegionsTriangleIndicesBuffer, 
-			FRHIViewDesc::CreateBufferSRV()
+			SampleRegionsTriangleIndicesSRV = RHICmdList.CreateShaderResourceView(
+				SampleRegionsTriangleIndicesBuffer, 
+				FRHIViewDesc::CreateBufferSRV()
 				.SetType(FRHIViewDesc::EBufferType::Typed)
-				.SetFormat(PF_R32_UINT));
+				.SetFormat(PF_R32_UINT)
+			);
 #if STATS
 			GPUMemoryUsage += SampleRegionsTriangleIndicies.Num() * SampleRegionsTriangleIndicies.GetTypeSize();
 #endif
@@ -1040,14 +1045,16 @@ void FSkeletalMeshGpuSpawnStaticBuffers::InitRHI(FRHICommandListBase& RHICmdList
 			const FRHIBufferCreateDesc CreateDesc =
 				FRHIBufferCreateDesc::CreateVertex(TEXT("SampleRegionsVerticesBuffer"), SampleRegionsVertices.Num() * SampleRegionsVertices.GetTypeSize())
 				.AddUsage(EBufferUsageFlags::Static | EBufferUsageFlags::ShaderResource)
-				.SetInitActionResourceArray(&SampleRegionsVertices);
+				.SetInitActionResourceArray(&SampleRegionsVertices)
+				.DetermineInitialState();
 
 			SampleRegionsVerticesBuffer = RHICmdList.CreateBuffer(CreateDesc);
-		SampleRegionsVerticesSRV = RHICmdList.CreateShaderResourceView(
-			SampleRegionsVerticesBuffer, 
-			FRHIViewDesc::CreateBufferSRV()
+			SampleRegionsVerticesSRV = RHICmdList.CreateShaderResourceView(
+				SampleRegionsVerticesBuffer, 
+				FRHIViewDesc::CreateBufferSRV()
 				.SetType(FRHIViewDesc::EBufferType::Typed)
-				.SetFormat(PF_R32_UINT));
+				.SetFormat(PF_R32_UINT)
+			);
 #if STATS
 			GPUMemoryUsage += SampleRegionsVertices.Num() * SampleRegionsVertices.GetTypeSize();
 #endif
@@ -1114,8 +1121,9 @@ void FSkeletalMeshGpuSpawnStaticBuffers::InitRHI(FRHICommandListBase& RHICmdList
 		FilteredAndUnfilteredBonesSRV = RHICmdList.CreateShaderResourceView(
 			FilteredAndUnfilteredBonesBuffer,
 			FRHIViewDesc::CreateBufferSRV()
-				.SetType(FRHIViewDesc::EBufferType::Typed)
-				.SetFormat(PF_R16_UINT));
+			.SetType(FRHIViewDesc::EBufferType::Typed)
+			.SetFormat(PF_R16_UINT)
+		);
 	}
 	else
 	{
@@ -1201,8 +1209,9 @@ void FSkeletalMeshGpuDynamicBufferProxy::InitRHI(FRHICommandListBase& RHICmdList
 		Buffer.SectionSRV = RHICmdList.CreateShaderResourceView(
 			Buffer.SectionBuffer, 
 			FRHIViewDesc::CreateBufferSRV()
-				.SetType(FRHIViewDesc::EBufferType::Typed)
-				.SetFormat(PF_A32B32G32R32F));
+			.SetType(FRHIViewDesc::EBufferType::Typed)
+			.SetFormat(PF_A32B32G32R32F)
+		);
 
 		const FRHIBufferCreateDesc SamplingBufferDesc =
 			FRHIBufferCreateDesc::CreateVertex<FVector4f>(TEXT("SkeletalMeshGpuDynamicBuffer"), 3 * (SamplingBoneCount + SamplingSocketCount))
@@ -1213,8 +1222,9 @@ void FSkeletalMeshGpuDynamicBufferProxy::InitRHI(FRHICommandListBase& RHICmdList
 		Buffer.SamplingSRV = RHICmdList.CreateShaderResourceView(
 			Buffer.SamplingBuffer, 
 			FRHIViewDesc::CreateBufferSRV()
-				.SetType(FRHIViewDesc::EBufferType::Typed)
-				.SetFormat(PF_A32B32G32R32F));
+			.SetType(FRHIViewDesc::EBufferType::Typed)
+			.SetFormat(PF_A32B32G32R32F)
+		);
 
 #if STATS
 		GPUMemoryUsage += sizeof(FVector4f) * 3 * SectionBoneCount;
