- Fix for crash when calling SetIsEnabled directly on renderers

diff --git a/Engine/Plugins/FX/Niagara/Source/Niagara/Private/NiagaraRendererProperties.cpp b/Engine/Plugins/FX/Niagara/Source/Niagara/Private/NiagaraRendererProperties.cpp
--- a/Engine/Plugins/FX/Niagara/Source/Niagara/Private/NiagaraRendererProperties.cpp
+++ b/Engine/Plugins/FX/Niagara/Source/Niagara/Private/NiagaraRendererProperties.cpp
@@ -1137,18 +1137,21 @@ void UNiagaraRendererProperties::SetIsEnabled(bool bInIsEnabled)
 {
 	if (bIsEnabled != bInIsEnabled)
 	{
-#if WITH_EDITORONLY_DATA
-		// Changing the enabled state will add or remove its renderer binding data stored on the emitters RenderBindings
-		// parameter store, so we need to reset to clear any binding references or add new ones
-		FVersionedNiagaraEmitter SrcEmitter = GetOuterEmitter();
-		if (SrcEmitter.Emitter)
+		FNiagaraSystemUpdateContext UpdateContext;
+	#if WITH_EDITORONLY_DATA
+		UpdateContext.Add(GetOuterEmitter(), true);
+	#else
+		// Shouldn't really be called at runtime, but let's ensure we handle the case
+		UpdateContext.Add(GetTypedOuter<UNiagaraSystem>(), true);
+	#endif
+
+		bIsEnabled = bInIsEnabled;
+
+		if (UNiagaraSystem* NiagaraSystem = GetTypedOuter<UNiagaraSystem>())
 		{
-			FNiagaraSystemUpdateContext(SrcEmitter, true);
+			NiagaraSystem->ComputeRenderersDrawOrder();
 		}
-#endif
 	}
-
-	bIsEnabled = bInIsEnabled;
 }
 
 void UNiagaraRendererProperties::UpdateSourceModeDerivates(ENiagaraRendererSourceDataMode InSourceMode, bool bFromPropertyEdit)
