- Add Stateless Scale Ribbon Width module

diff --git a/Engine/Plugins/FX/Niagara/Shaders/Private/Stateless/Modules/NiagaraStatelessModule_ScaleRibbonWidth.ush b/Engine/Plugins/FX/Niagara/Shaders/Private/Stateless/Modules/NiagaraStatelessModule_ScaleRibbonWidth.ush
new file mode 100644
--- /dev/null
+++ b/Engine/Plugins/FX/Niagara/Shaders/Private/Stateless/Modules/NiagaraStatelessModule_ScaleRibbonWidth.ush
@@ -0,0 +1,20 @@
+// Copyright Epic Games, Inc. All Rights Reserved.
+
+#include "../NiagaraStatelessCommon.ush"
+
+FNiagaraStatelessBuiltDistributionType	ScaleRibbonWidth_Distribution;
+float	ScaleRibbonWidth_CurveScale;
+
+void ScaleRibbonWidth_Simulate(inout FStatelessParticle Particle)
+{
+	FStatelessDistributionSamplerFloat ScaleDistribution;
+	ScaleDistribution.Init(0, ScaleRibbonWidth_Distribution);
+
+	if (ScaleDistribution.IsValid())
+	{
+		const float ScaleFactor = ScaleRibbonWidth_CurveScale;
+
+		Particle.RibbonWidth			*= ScaleDistribution.GetValue(Particle.NormalizedAge) * ScaleFactor;
+		Particle.PreviousRibbonWidth	*= ScaleDistribution.GetValue(Particle.PreviousNormalizedAge) * ScaleFactor;
+	}
+}
diff --git a/Engine/Plugins/FX/Niagara/Shaders/Private/Stateless/NiagaraStatelessSimulationDefault.usf b/Engine/Plugins/FX/Niagara/Shaders/Private/Stateless/NiagaraStatelessSimulationDefault.usf
--- a/Engine/Plugins/FX/Niagara/Shaders/Private/Stateless/NiagaraStatelessSimulationDefault.usf
+++ b/Engine/Plugins/FX/Niagara/Shaders/Private/Stateless/NiagaraStatelessSimulationDefault.usf
@@ -54,6 +54,7 @@
 	PARTICLE_MODULE(MeshIndex) \
 	PARTICLE_MODULE(MeshRotationRate) \
 	PARTICLE_MODULE(ScaleColor) \
+	PARTICLE_MODULE(ScaleRibbonWidth) \
 	PARTICLE_MODULE(ScaleSpriteSize) \
 	PARTICLE_MODULE(ScaleSpriteSizeBySpeed) \
 	PARTICLE_MODULE(ScaleMeshSize) \
diff --git a/Engine/Plugins/FX/Niagara/Shaders/Private/Stateless/NiagaraStatelessSimulationTemplate.ush b/Engine/Plugins/FX/Niagara/Shaders/Private/Stateless/NiagaraStatelessSimulationTemplate.ush
--- a/Engine/Plugins/FX/Niagara/Shaders/Private/Stateless/NiagaraStatelessSimulationTemplate.ush
+++ b/Engine/Plugins/FX/Niagara/Shaders/Private/Stateless/NiagaraStatelessSimulationTemplate.ush
@@ -26,6 +26,7 @@
 #include "Modules/NiagaraStatelessModule_RotateAroundPoint.ush"
 #include "Modules/NiagaraStatelessModule_ShapeLocation.ush"
 #include "Modules/NiagaraStatelessModule_ScaleColor.ush"
+#include "Modules/NiagaraStatelessModule_ScaleRibbonWidth.ush"
 #include "Modules/NiagaraStatelessModule_ScaleSpriteSize.ush"
 #include "Modules/NiagaraStatelessModule_ScaleSpriteSizeBySpeed.ush"
 #include "Modules/NiagaraStatelessModule_ScaleMeshSize.ush"
diff --git a/Engine/Plugins/FX/Niagara/Source/Niagara/Internal/Stateless/Modules/NiagaraStatelessModule_ScaleRibbonWidth.h b/Engine/Plugins/FX/Niagara/Source/Niagara/Internal/Stateless/Modules/NiagaraStatelessModule_ScaleRibbonWidth.h
new file mode 100644
--- /dev/null
+++ b/Engine/Plugins/FX/Niagara/Source/Niagara/Internal/Stateless/Modules/NiagaraStatelessModule_ScaleRibbonWidth.h
@@ -0,0 +1,130 @@
+// Copyright Epic Games, Inc. All Rights Reserved.
+
+#pragma once
+
+#include "Stateless/NiagaraStatelessCommon.h"
+#include "Stateless/NiagaraStatelessModule.h"
+#include "Stateless/NiagaraStatelessEmitterDataBuildContext.h"
+#include "Stateless/NiagaraStatelessModuleShaderParameters.h"
+#include "Stateless/NiagaraStatelessParticleSimContext.h"
+
+#include "NiagaraStatelessModule_ScaleRibbonWidth.generated.h"
+
+// Multiply Particle.RibbonWidth by the module calculated scale value
+// This can be a constant, random or curve indexed by Particle.NormalizedAge
+UCLASS(MinimalAPI, EditInlineNew, meta = (DisplayName = "Scale Ribbon Width"))
+class UNiagaraStatelessModule_ScaleRibbonWidth : public UNiagaraStatelessModule
+{
+	GENERATED_BODY()
+
+	struct FModuleBuiltData
+	{
+		FUintVector3	DistributionParameters = FUintVector3::ZeroValue;
+		float			CurveScale = 1.0f;
+		int32			CurveScaleOffset = INDEX_NONE;
+
+		int32			RibbonWidthVariableOffset = INDEX_NONE;
+		int32			PreviousRibbonWidthVariableOffset = INDEX_NONE;
+	};
+
+public:
+	using FParameters = NiagaraStateless::FScaleRibbonWidthModule_ShaderParameters;
+
+	UPROPERTY(EditAnywhere, Category = "Parameters", meta = (DisplayName = "Scale"))
+	FNiagaraDistributionFloat ScaleDistribution = FNiagaraDistributionFloat(1.0f);
+
+	UPROPERTY(EditAnywhere, Category = "Parameters", meta = (EditConditionHides, EditCondition = "UseScaleCurveRange()"))
+	FNiagaraParameterBindingWithValue ScaleCurveRange;
+
+	virtual void PostInitProperties() override
+	{
+		Super::PostInitProperties();
+
+	#if WITH_EDITORONLY_DATA
+		if (HasAnyFlags(RF_ClassDefaultObject) == false)
+		{
+			ScaleCurveRange.SetUsage(ENiagaraParameterBindingUsage::NotParticle);
+			ScaleCurveRange.SetAllowedTypeDefinitions({ FNiagaraTypeDefinition::GetVec2Def() });
+			ScaleCurveRange.SetDefaultParameter(FNiagaraTypeDefinition::GetFloatDef(), 1.0f);
+		}
+	#endif
+	}
+
+	virtual void BuildEmitterData(const FNiagaraStatelessEmitterDataBuildContext& BuildContext) const override
+	{
+		FModuleBuiltData* BuiltData = BuildContext.AllocateBuiltData<FModuleBuiltData>();
+
+		const FNiagaraStatelessGlobals& StatelessGlobals = FNiagaraStatelessGlobals::Get();
+		BuiltData->RibbonWidthVariableOffset			= BuildContext.FindParticleVariableIndex(StatelessGlobals.RibbonWidthVariable);
+		BuiltData->PreviousRibbonWidthVariableOffset	= BuildContext.FindParticleVariableIndex(StatelessGlobals.PreviousRibbonWidthVariable);
+
+		const bool bAttributesUsed = (BuiltData->RibbonWidthVariableOffset != INDEX_NONE || BuiltData->PreviousRibbonWidthVariableOffset != INDEX_NONE);
+		if (IsModuleEnabled() && bAttributesUsed)
+		{
+			BuiltData->DistributionParameters = BuildContext.AddDistribution(ScaleDistribution);
+			if (UseScaleCurveRange())
+			{
+				BuiltData->CurveScaleOffset = BuildContext.AddRendererBinding(ScaleCurveRange.ResolvedParameter);
+				BuiltData->CurveScale = ScaleCurveRange.GetDefaultValue<float>();
+			}
+
+			BuildContext.AddParticleSimulationExecSimulate(&UNiagaraStatelessModule_ScaleRibbonWidth::ParticleSimulate);
+		}
+	}
+
+	virtual void BuildShaderParameters(FNiagaraStatelessShaderParametersBuilder& ShaderParametersBuilder) const override
+	{
+		ShaderParametersBuilder.AddParameterNestedStruct<FParameters>();
+	}
+
+	virtual void SetShaderParameters(const FNiagaraStatelessSetShaderParameterContext& SetShaderParameterContext) const override
+	{
+		const FModuleBuiltData* ModuleBuiltData = SetShaderParameterContext.ReadBuiltData<FModuleBuiltData>();
+
+		FParameters* Parameters						= SetShaderParameterContext.GetParameterNestedStruct<FParameters>();
+		Parameters->ScaleRibbonWidth_Distribution	= ModuleBuiltData->DistributionParameters;
+		Parameters->ScaleRibbonWidth_CurveScale		= SetShaderParameterContext.GetRendererParameterValue(ModuleBuiltData->CurveScaleOffset, ModuleBuiltData->CurveScale);
+	}
+
+	static void ParticleSimulate(const NiagaraStateless::FParticleSimulationContext& ParticleSimulationContext)
+	{
+		using namespace NiagaraStateless;
+
+		const FModuleBuiltData* ModuleBuiltData = ParticleSimulationContext.ReadBuiltData<FModuleBuiltData>();
+
+		const float ScaleFactor					= ParticleSimulationContext.GetParameterBufferFloat(ModuleBuiltData->CurveScaleOffset, ModuleBuiltData->CurveScale);
+		const float* NormalizedAgeData			= ParticleSimulationContext.GetParticleNormalizedAge();
+		const float* PreviousNormalizedAgeData	= ParticleSimulationContext.GetParticlePreviousNormalizedAge();
+
+		for (uint32 i = 0; i < ParticleSimulationContext.GetNumInstances(); ++i)
+		{
+			const FStatelessDistributionSampler<float> SpriteScaleSampler(ParticleSimulationContext, ModuleBuiltData->DistributionParameters, i, 0);
+
+			const float NormalizedAge			= NormalizedAgeData[i];
+			const float PreviousNormalizedAge	= PreviousNormalizedAgeData[i];
+
+			const float RibbonWidth			= ParticleSimulationContext.ReadParticleVariable(ModuleBuiltData->RibbonWidthVariableOffset, i, 0.0f);
+			const float PreviousRibbonWidth = ParticleSimulationContext.ReadParticleVariable(ModuleBuiltData->PreviousRibbonWidthVariableOffset, i, 0.0f);
+			const float Scale				= SpriteScaleSampler.GetValue(ParticleSimulationContext, NormalizedAge) * ScaleFactor;
+			const float PreviousScale		= SpriteScaleSampler.GetValue(ParticleSimulationContext, PreviousNormalizedAge) * ScaleFactor;
+
+			ParticleSimulationContext.WriteParticleVariable(ModuleBuiltData->RibbonWidthVariableOffset, i, RibbonWidth * Scale);
+			ParticleSimulationContext.WriteParticleVariable(ModuleBuiltData->PreviousRibbonWidthVariableOffset, i, PreviousRibbonWidth * PreviousScale);
+		}
+	}
+
+	UFUNCTION()
+	bool UseScaleCurveRange() const { return ScaleDistribution.IsCurve(); }
+
+#if WITH_EDITOR
+	virtual bool CanDisableModule() const override { return true; }
+#endif
+#if WITH_EDITORONLY_DATA
+	virtual void GetOutputVariables(TArray<FNiagaraVariableBase>& OutVariables) const override
+	{
+		const FNiagaraStatelessGlobals& StatelessGlobals = FNiagaraStatelessGlobals::Get();
+		OutVariables.AddUnique(StatelessGlobals.RibbonWidthVariable);
+		OutVariables.AddUnique(StatelessGlobals.PreviousRibbonWidthVariable);
+	}
+#endif
+};
diff --git a/Engine/Plugins/FX/Niagara/Source/Niagara/Private/Stateless/NiagaraStatelessEmitterTemplate.cpp b/Engine/Plugins/FX/Niagara/Source/Niagara/Private/Stateless/NiagaraStatelessEmitterTemplate.cpp
--- a/Engine/Plugins/FX/Niagara/Source/Niagara/Private/Stateless/NiagaraStatelessEmitterTemplate.cpp
+++ b/Engine/Plugins/FX/Niagara/Source/Niagara/Private/Stateless/NiagaraStatelessEmitterTemplate.cpp
@@ -23,6 +23,7 @@
 #include "Stateless/Modules/NiagaraStatelessModule_ScaleColor.h"
 #include "Stateless/Modules/NiagaraStatelessModule_ScaleMeshSize.h"
 #include "Stateless/Modules/NiagaraStatelessModule_ScaleMeshSizeBySpeed.h"
+#include "Stateless/Modules/NiagaraStatelessModule_ScaleRibbonWidth.h"
 #include "Stateless/Modules/NiagaraStatelessModule_ScaleSpriteSize.h"
 #include "Stateless/Modules/NiagaraStatelessModule_ScaleSpriteSizeBySpeed.h"
 #include "Stateless/Modules/NiagaraStatelessModule_ShapeLocation.h"
@@ -89,6 +90,7 @@ void UNiagaraStatelessEmitterDefault::InitModulesAndAttributes()
 		UNiagaraStatelessModule_MeshIndex::StaticClass(),
 		UNiagaraStatelessModule_MeshRotationRate::StaticClass(),
 		UNiagaraStatelessModule_ScaleColor::StaticClass(),
+		UNiagaraStatelessModule_ScaleRibbonWidth::StaticClass(),
 		UNiagaraStatelessModule_ScaleSpriteSize::StaticClass(),
 		UNiagaraStatelessModule_ScaleSpriteSizeBySpeed::StaticClass(),
 		UNiagaraStatelessModule_ScaleMeshSize::StaticClass(),
diff --git a/Engine/Plugins/FX/Niagara/Source/NiagaraShader/Internal/Stateless/NiagaraStatelessEmitterShaders.h b/Engine/Plugins/FX/Niagara/Source/NiagaraShader/Internal/Stateless/NiagaraStatelessEmitterShaders.h
--- a/Engine/Plugins/FX/Niagara/Source/NiagaraShader/Internal/Stateless/NiagaraStatelessEmitterShaders.h
+++ b/Engine/Plugins/FX/Niagara/Source/NiagaraShader/Internal/Stateless/NiagaraStatelessEmitterShaders.h
@@ -36,6 +36,7 @@ namespace NiagaraStateless
 			ADD_STATELESS_MODULE(FMeshIndexModule_ShaderParameters)
 			ADD_STATELESS_MODULE(FMeshRotationRateModule_ShaderParameters)
 			ADD_STATELESS_MODULE(FScaleColorModule_ShaderParameters)
+			ADD_STATELESS_MODULE(FScaleRibbonWidthModule_ShaderParameters)
 			ADD_STATELESS_MODULE(FScaleSpriteSizeModule_ShaderParameters)
 			ADD_STATELESS_MODULE(FScaleSpriteSizeBySpeedModule_ShaderParameters)
 			ADD_STATELESS_MODULE(FScaleMeshSizeModule_ShaderParameters)
diff --git a/Engine/Plugins/FX/Niagara/Source/NiagaraShader/Internal/Stateless/NiagaraStatelessModuleShaderParameters.h b/Engine/Plugins/FX/Niagara/Source/NiagaraShader/Internal/Stateless/NiagaraStatelessModuleShaderParameters.h
--- a/Engine/Plugins/FX/Niagara/Source/NiagaraShader/Internal/Stateless/NiagaraStatelessModuleShaderParameters.h
+++ b/Engine/Plugins/FX/Niagara/Source/NiagaraShader/Internal/Stateless/NiagaraStatelessModuleShaderParameters.h
@@ -92,6 +92,11 @@ namespace NiagaraStateless
 		SHADER_PARAMETER(int32,		MeshIndex_TableNumElements)
 	END_SHADER_PARAMETER_STRUCT()
 
+	BEGIN_SHADER_PARAMETER_STRUCT(FScaleRibbonWidthModule_ShaderParameters, NIAGARASHADER_API)
+		SHADER_PARAMETER(FNiagaraStatelessBuiltDistributionType,	ScaleRibbonWidth_Distribution)
+		SHADER_PARAMETER(float,										ScaleRibbonWidth_CurveScale)
+	END_SHADER_PARAMETER_STRUCT()
+
 	BEGIN_SHADER_PARAMETER_STRUCT(FMeshRotationRateModule_ShaderParameters, NIAGARASHADER_API)
 		SHADER_PARAMETER(int32,										MeshRotationRate_ModuleEnabled)
 		SHADER_PARAMETER(FVector3f,									MeshRotationRate_Scale)
