Added some more debug information to the NiagaraDataInterfaceLandscape via DebugHud.  Also removed incomplete SimCache implementation.
#tests Checked out the DebugHud results for various systems in EngineTest

diff --git a/Engine/Plugins/FX/Niagara/Source/Niagara/Classes/NiagaraDataInterfaceLandscape.h b/Engine/Plugins/FX/Niagara/Source/Niagara/Classes/NiagaraDataInterfaceLandscape.h
--- a/Engine/Plugins/FX/Niagara/Source/Niagara/Classes/NiagaraDataInterfaceLandscape.h
+++ b/Engine/Plugins/FX/Niagara/Source/Niagara/Classes/NiagaraDataInterfaceLandscape.h
@@ -36,7 +36,7 @@ enum class ENDILandscape_SourceMode : uint8
 
 /** Data Interface allowing sampling of a Landscape */
 UCLASS(EditInlineNew, Category = "Landscape", CollapseCategories, meta = (DisplayName = "Landscape Sample"), MinimalAPI)
-class UNiagaraDataInterfaceLandscape : public UNiagaraDataInterface, public INiagaraSimCacheCustomStorageInterface
+class UNiagaraDataInterfaceLandscape : public UNiagaraDataInterface
 {
 	GENERATED_UCLASS_BODY()
 
@@ -66,12 +66,9 @@ public:
 	NIAGARA_API virtual bool Equals(const UNiagaraDataInterface* Other) const override;
 	//UNiagaraDataInterface Interface End
 
-	//~ INiagaraSimCacheCustomStorageInterface interface BEGIN
-	NIAGARA_API virtual UObject* SimCacheBeginWrite(UObject* SimCache, FNiagaraSystemInstance* NiagaraSystemInstance, const void* OptionalPerInstanceData, FNiagaraSimCacheFeedbackContext& FeedbackContext) const override;
-	NIAGARA_API virtual bool SimCacheWriteFrame(UObject* StorageObject, int FrameIndex, FNiagaraSystemInstance* SystemInstance, const void* OptionalPerInstanceData, FNiagaraSimCacheFeedbackContext& FeedbackContext) const override;
-	//NIAGARA_API virtual bool SimCacheEndWrite(UObject* StorageObject) const override;
-	NIAGARA_API virtual bool SimCacheReadFrame(UObject* StorageObject, int FrameA, int FrameB, float Interp, FNiagaraSystemInstance* SystemInstance, void* OptionalPerInstanceData) override;
-	//~ UNiagaraDataInterface interface END
+#if WITH_NIAGARA_DEBUGGER
+	virtual void DrawDebugHud(FNDIDrawDebugHudContext& DebugHudContext) const override;
+#endif
 
 	// GPU sim functionality
 #if WITH_EDITORONLY_DATA
diff --git a/Engine/Plugins/FX/Niagara/Source/Niagara/Private/NiagaraDataInterfaceLandscape.cpp b/Engine/Plugins/FX/Niagara/Source/Niagara/Private/NiagaraDataInterfaceLandscape.cpp
--- a/Engine/Plugins/FX/Niagara/Source/Niagara/Private/NiagaraDataInterfaceLandscape.cpp
+++ b/Engine/Plugins/FX/Niagara/Source/Niagara/Private/NiagaraDataInterfaceLandscape.cpp
@@ -204,6 +204,8 @@ public:
 	FIntPoint CellCount = FIntPoint(ForceInitToZero);
 	FVector2D TextureWorldGridSize = FVector2D(1.0f, 1.0f);
 
+	TConstArrayView<FIntPoint> ReadCapturedRegions() const { return ResourceKey.CapturedRegions; }
+
 private:
 	enum class EResourceState : uint8
 	{
@@ -1247,243 +1249,101 @@ bool UNiagaraDataInterfaceLandscape::Equals(const UNiagaraDataInterface* Other)
 		&& OtherLandscape->bVirtualTexturesSupported == bVirtualTexturesSupported;
 }
 
-UObject* UNiagaraDataInterfaceLandscape::SimCacheBeginWrite(UObject* SimCache, FNiagaraSystemInstance* NiagaraSystemInstance, const void* OptionalPerInstanceData, FNiagaraSimCacheFeedbackContext& FeedbackContext) const
-{
-	return NewObject<UNDILandscapeSimCacheData>(SimCache);
-}
-
-bool UNiagaraDataInterfaceLandscape::SimCacheWriteFrame(UObject* StorageObject, int FrameIndex, FNiagaraSystemInstance* SystemInstance, const void* PerInstanceData, FNiagaraSimCacheFeedbackContext& FeedbackContext) const
+#if WITH_NIAGARA_DEBUGGER
+void UNiagaraDataInterfaceLandscape::DrawDebugHud(FNDIDrawDebugHudContext& DebugHudContext) const
 {
-	if (UNDILandscapeSimCacheData* SimCacheData = Cast<UNDILandscapeSimCacheData>(StorageObject))
+	const FNDILandscapeData_GameThread* InstanceData_GT = DebugHudContext.GetSystemInstance()->FindTypedDataInterfaceInstanceData<FNDILandscapeData_GameThread>(this);
+	if (InstanceData_GT == nullptr)
 	{
-		SimCacheData->HeightFieldTextures.SetNum(1 + FrameIndex);
-
-		const FNDILandscapeData_GameThread& SourceData = *reinterpret_cast<const FNDILandscapeData_GameThread*>(PerInstanceData);
+		return;
+	}
 
-		if (SourceData.SharedResourceHandle)
+	auto SafeActorLabel = [](const AActor* Actor) -> FString
+	{
+		if (Actor)
 		{
-			const FNDI_Landscape_SharedResource& SourceResource = SourceData.SharedResourceHandle.ReadResource();
-			if (const FRHITexture* HeightTexture = SourceResource.LandscapeTextures->GetHeightTexture())
-			{
-				UTexture2D* FrameTexture = NewObject<UTexture2D>(SimCacheData);
-
-				// copy texture
-
-				SimCacheData->HeightFieldTextures[FrameIndex] = FrameTexture;
-			}
+			return Actor->GetActorNameOrLabel();
 		}
-	}
-
-	return true;
-}
-
-bool UNiagaraDataInterfaceLandscape::SimCacheReadFrame(UObject* StorageObject, int FrameA, int FrameB, float Interp, FNiagaraSystemInstance* SystemInstance, void* OptionalPerInstanceData)
-{
-	return true;
-	#if 0
-	FRenderTargetVolumeRWInstanceData_GameThread* InstanceData_GT = reinterpret_cast<FRenderTargetVolumeRWInstanceData_GameThread*>(OptionalPerInstanceData);
+		else
+		{
+			return TEXT("none");
+		}
+	};
 
-	if (const* UTexture2D HeightFieldTexture = Cast<UTexture2D>(StorageObject))
+	TStringBuilder<256> OutputString;
+	OutputString.Appendf(TEXT("Landscape(%s)"),	*SafeActorLabel(InstanceData_GT->Landscape.Get()));
+	if (InstanceData_GT->BaseColorVirtualTextureIndex != INDEX_NONE)
 	{
-		
+		OutputString.Append(TEXT(" BaseColor:RVT"));
 	}
-
-	return true;
-	#endif
-	#if 0
-	else if (Cast<UAnimatedSparseVolumeTexture>(StorageObject))
+	if (InstanceData_GT->HeightVirtualTextureIndex != INDEX_NONE)
 	{
-		UAnimatedSparseVolumeTexture* SVT = Cast<UAnimatedSparseVolumeTexture>(StorageObject);
-
-		const int32 MipLevel = 0;
-		const bool bBlocking = true;
-		USparseVolumeTextureFrame* SVTFrame = USparseVolumeTextureFrame::GetFrameAndIssueStreamingRequest(SVT, FMath::RoundToInt(FrameA + Interp), MipLevel, bBlocking);
-
-		// The streaming manager normally ticks in FDeferredShadingSceneRenderer::Render(), but the SVT->DenseTexture conversion compute shader happens in a render command before that.
-		// At execution time of that command, the streamer hasn't had the chance to do any streaming yet, so we force another tick here.
-		// Assuming blocking requests are used, this guarantees that the requested frame is fully streamed in (if there is memory available).
-		UE::SVT::GetStreamingManager().Update_GameThread();
-
-		FIntVector VolumeResolution = SVT->GetVolumeResolution();
-
-		InstanceData_GT->Size = VolumeResolution;
-		InstanceData_GT->Format = SVT->GetFormat(0);
-		InstanceData_GT->Filter = TextureFilter::TF_Default;
-
-		PerInstanceTick(InstanceData_GT, SystemInstance, 0.0f);
-		PerInstanceTickPostSimulate(InstanceData_GT, SystemInstance, 0.0f);
-
-		FNiagaraDataInterfaceProxyRenderTargetVolumeProxy* RT_Proxy = GetProxyAs<FNiagaraDataInterfaceProxyRenderTargetVolumeProxy>();
+		OutputString.Append(TEXT(" Height:RVT"));
+	}
+	if (InstanceData_GT->NormalVirtualTextureIndex != INDEX_NONE)
+	{
+		OutputString.Append(TEXT(" Normal:RVT"));
+	}
+	if (InstanceData_GT->RequiresCollisionCacheGpu || InstanceData_GT->RequiresPhysMatCacheGpu)
+	{
+		OutputString.Append(TEXT(" HasCollisionCache"));
+	}
+	if (InstanceData_GT->SystemRequiresBaseColorGpu || InstanceData_GT->SystemRequiresHeightsGpu || InstanceData_GT->SystemRequiresNormalsGpu)
+	{
+		OutputString.Appendf(TEXT(" GpuReqs(BaseColor:%s, Height:%s, Normal:%s)"),
+			InstanceData_GT->SystemRequiresBaseColorGpu ? TEXT("Yes") : TEXT("No"),
+			InstanceData_GT->SystemRequiresHeightsGpu ? TEXT("Yes") : TEXT("No"),
+			InstanceData_GT->SystemRequiresNormalsGpu ? TEXT("Yes") : TEXT("No"));
+	}
 
+	if (DebugHudContext.IsVerbose() && InstanceData_GT->SharedResourceHandle)
+	{
+		if (ALandscape* LandscapeActor = InstanceData_GT->Landscape.Get())
+		{
+			if (ULandscapeInfo* LandscapeInfo = LandscapeActor->GetLandscapeInfo())
+			{
+				// getting the list of publicly facing LandscapeProxies corresponding to the captured regions is a bit of a pain
+				TSet<const ULandscapeHeightfieldCollisionComponent*> RelevantCollisionComponents;
+				TSet<const ALandscapeProxy*> RelevantLandscapeProxies;
 
-		// Execute compute shader
-		ENQUEUE_RENDER_COMMAND(NDIRenderTargetVolumeUpdate)
-			(
-				[RT_Proxy, RT_InstanceID = SystemInstance->GetId(), RT_Format = InstanceData_GT->Format,
-				RT_VolumeResolution = VolumeResolution, RT_SVTRenderResources = SVTFrame ? SVTFrame->GetTextureRenderResources() : nullptr,
-				FeatureLevel = SystemInstance->GetFeatureLevel()](FRHICommandListImmediate& RHICmdList)
+				const FNDI_Landscape_SharedResource& LandscapeData = InstanceData_GT->SharedResourceHandle.ReadResource();
+				for (const FIntPoint& CapturedRegion : LandscapeData.ReadCapturedRegions())
 				{
-					if (RT_SVTRenderResources == nullptr)
+					if (const ULandscapeHeightfieldCollisionComponent* RefCollisionComponent = LandscapeInfo->XYtoCollisionComponentMap.FindRef(CapturedRegion))
 					{
-						return;
+						RelevantCollisionComponents.Add(RefCollisionComponent);
 					}
+				}
 
-		FUintVector4 CurrentPackedUniforms0 = FUintVector4();
-		FUintVector4 CurrentPackedUniforms1 = FUintVector4();
-		RT_SVTRenderResources->GetPackedUniforms(CurrentPackedUniforms0, CurrentPackedUniforms1);
-
-		if (FRenderTargetVolumeRWInstanceData_RenderThread* InstanceData_RT = RT_Proxy->SystemInstancesToProxyData_RT.Find(RT_InstanceID))
-		{
-			FRDGBuilder GraphBuilder(RHICmdList);
-
-			FIntVector ThreadGroupSize = FNiagaraShader::GetDefaultThreadGroupSize(ENiagaraGpuDispatchType::ThreeD);
-
-			TShaderMapRef<FNiagaraCopySVTToDenseBufferCS> ComputeShader(GetGlobalShaderMap(FeatureLevel));
-
-			const FIntVector NumThreadGroups(
-				FMath::DivideAndRoundUp(RT_VolumeResolution.X, ThreadGroupSize.X),
-				FMath::DivideAndRoundUp(RT_VolumeResolution.Y, ThreadGroupSize.Y),
-				FMath::DivideAndRoundUp(RT_VolumeResolution.Z, ThreadGroupSize.Z)
-			);
-
-			FNiagaraCopySVTToDenseBufferCS::FParameters* PassParameters = GraphBuilder.AllocParameters<FNiagaraCopySVTToDenseBufferCS::FParameters>();
-
-			if (InstanceData_RT->RenderTarget.IsValid() && InstanceData_RT->TransientRDGTexture == nullptr)
-			{
-				InstanceData_RT->TransientRDGTexture = GraphBuilder.RegisterExternalTexture(InstanceData_RT->RenderTarget);
-				InstanceData_RT->TransientRDGUAV = GraphBuilder.CreateUAV(InstanceData_RT->TransientRDGTexture);
-
-				// #todo(dmp): needed?							
-				//Context.GetRDGExternalAccessQueue().Add(InstanceData_RT->TransientRDGTexture);
-			}
-
-			if (InstanceData_RT->RenderTarget.IsValid() && InstanceData_RT->TransientRDGUAV != nullptr)
-			{
-				PassParameters->DestinationBuffer = InstanceData_RT->TransientRDGUAV;
-			}
-			else
-			{
-				PassParameters->DestinationBuffer = GraphBuilder.CreateUAV(GraphBuilder.CreateTexture(
-					FRDGTextureDesc::Create3D(FIntVector(1, 1, 1), RT_Format, FClearValueBinding::Black, ETextureCreateFlags::ShaderResource | ETextureCreateFlags::UAV),
-					TEXT("NiagaraEmptyTextureUAV::Texture3D")
-				),
-					ERDGUnorderedAccessViewFlags::SkipBarrier
-				);
-			}
-
-			FRHITexture* PageTableTexture = RT_SVTRenderResources->GetPageTableTexture();
-			FRHITexture* TextureA = RT_SVTRenderResources->GetPhysicalTileDataATexture();
-
-			PassParameters->TileDataTextureSampler = TStaticSamplerState<SF_Point, AM_Clamp, AM_Clamp, AM_Clamp>::GetRHI();
-			PassParameters->SparseVolumeTexturePageTable = PageTableTexture ? PageTableTexture : GBlackUintVolumeTexture->TextureRHI.GetReference();
-			PassParameters->SparseVolumeTextureA = TextureA ? TextureA : GBlackVolumeTexture->TextureRHI.GetReference();
-
-			PassParameters->PackedSVTUniforms0 = CurrentPackedUniforms0;
-			PassParameters->PackedSVTUniforms1 = CurrentPackedUniforms1;
-
-			PassParameters->TextureSize = RT_VolumeResolution;
-			PassParameters->MipLevel = 0;
-
-			GraphBuilder.AddPass(
-				// Friendly name of the pass for profilers using printf semantics.
-				RDG_EVENT_NAME("Copy SVT to Volume RT"),
-				// Parameters provided to RDG.
-				PassParameters,
-				// Issues compute commands.
-				ERDGPassFlags::Compute,
-				// This is deferred until Execute. May execute in parallel with other passes.
-				[PassParameters, ComputeShader, NumThreadGroups](FRDGAsyncTask, FRHIComputeCommandList& RHICmdList)
+				LandscapeInfo->ForEachLandscapeProxy([&RelevantCollisionComponents, &RelevantLandscapeProxies](ALandscapeProxy* LandscapeProxy) -> bool
 				{
-					FComputeShaderUtils::Dispatch(RHICmdList, ComputeShader, *PassParameters, NumThreadGroups);
-				});
-
-			// Execute the graph.
-			GraphBuilder.Execute();
-
-			InstanceData_RT->TransientRDGTexture = nullptr;
-			InstanceData_RT->TransientRDGSRV = nullptr;
-			InstanceData_RT->TransientRDGUAV = nullptr;
-		}
+					for (const ULandscapeHeightfieldCollisionComponent* CollisionComponent : LandscapeProxy->CollisionComponents)
+					{
+						if (RelevantCollisionComponents.Contains(CollisionComponent))
+						{
+							RelevantLandscapeProxies.Add(LandscapeProxy);
+						}
+					}
+					return true;
 				});
-	}
-	else
-	{
-		UNDIRenderTargetVolumeSimCacheData* SimCacheData = CastChecked<UNDIRenderTargetVolumeSimCacheData>(StorageObject);
-
-		const int FrameIndex = Interp >= 0.5f ? FrameB : FrameA;
-		if (!SimCacheData->Frames.IsValidIndex(FrameIndex))
-		{
-			return false;
-		}
-
-		FNDIRenderTargetVolumeSimCacheFrame* CacheFrame = &SimCacheData->Frames[FrameIndex];
 
-		InstanceData_GT->Size = CacheFrame->Size;
-		InstanceData_GT->Format = CacheFrame->Format;
-		InstanceData_GT->Filter = TextureFilter::TF_Default;
-
-		PerInstanceTick(InstanceData_GT, SystemInstance, 0.0f);
-		PerInstanceTickPostSimulate(InstanceData_GT, SystemInstance, 0.0f);
+				if (!RelevantLandscapeProxies.IsEmpty())
+				{
+					OutputString.Append(TEXT(" Proxies:"));
 
-		if (InstanceData_GT->TargetTexture && CacheFrame->GetPixelData() != nullptr)
-		{
-			FNiagaraDataInterfaceProxyRenderTargetVolumeProxy* RT_Proxy = GetProxyAs<FNiagaraDataInterfaceProxyRenderTargetVolumeProxy>();
-			FTextureRenderTargetResource* RT_TargetTexture = InstanceData_GT->TargetTexture->GameThread_GetRenderTargetResource();
-			ENQUEUE_RENDER_COMMAND(NDIRenderTargetVolumeUpdate)
-				(
-					[RT_Proxy, RT_InstanceID = SystemInstance->GetId(), RT_TargetTexture, RT_CacheFrame = CacheFrame, RT_CompressionType = SimCacheData->CompressionType, RT_Format = InstanceData_GT->Format](FRHICommandListImmediate& RHICmdList)
+					for (const ALandscapeProxy* RelevantProxy : RelevantLandscapeProxies)
 					{
-						if (FRenderTargetVolumeRWInstanceData_RenderThread* InstanceData_RT = RT_Proxy->SystemInstancesToProxyData_RT.Find(RT_InstanceID))
-						{
-							FUpdateTextureRegion3D UpdateRegion(FIntVector::ZeroValue, FIntVector::ZeroValue, InstanceData_RT->Size);
-
-							FUpdateTexture3DData UpdateTexture = RHICmdList.BeginUpdateTexture3D(InstanceData_RT->RenderTarget->GetRHI(), 0, UpdateRegion);
-
-							uint32 BlockBytes = GPixelFormats[RT_Format].BlockBytes;
-
-							const int32 SrcRowPitch = InstanceData_RT->Size.X * BlockBytes;
-							const int32 SrcDepthPitch = InstanceData_RT->Size.Y * SrcRowPitch;
-							if (UpdateTexture.RowPitch == SrcRowPitch && UpdateTexture.DepthPitch == SrcDepthPitch)
-							{
-								if (RT_CacheFrame->CompressedSize > 0)
-								{
-									FCompression::UncompressMemory(RT_CompressionType, UpdateTexture.Data, UpdateTexture.DataSizeBytes, RT_CacheFrame->GetPixelData(), RT_CacheFrame->CompressedSize);
-								}
-								else
-								{
-									FMemory::Memcpy(UpdateTexture.Data, RT_CacheFrame->GetPixelData(), InstanceData_RT->Size.X * InstanceData_RT->Size.Y * sizeof(FFloat16Color));
-								}
-							}
-							else
-							{
-								TArray<uint8> Decompressed;
-								if (RT_CacheFrame->CompressedSize > 0)
-								{
-									Decompressed.AddUninitialized(BlockBytes * InstanceData_RT->Size.X * InstanceData_RT->Size.Y * InstanceData_RT->Size.Z);
-									FCompression::UncompressMemory(RT_CompressionType, Decompressed.GetData(), Decompressed.Num(), RT_CacheFrame->GetPixelData(), RT_CacheFrame->CompressedSize);
-								}
-
-								const uint8* SrcData = Decompressed.Num() > 0 ? Decompressed.GetData() : RT_CacheFrame->GetPixelData();
-								for (int32 z = 0; z < InstanceData_RT->Size.Z; ++z)
-								{
-									uint8* DstData = UpdateTexture.Data + (z * UpdateTexture.DepthPitch);
-									for (int32 y = 0; y < InstanceData_RT->Size.Y; ++y)
-									{
-										FMemory::Memcpy(DstData, SrcData, SrcRowPitch);
-										SrcData += SrcRowPitch;
-										DstData += UpdateTexture.RowPitch;
-									}
-								}
-							}
-							RHICmdList.EndUpdateTexture3D(UpdateTexture);
-						}
+						OutputString.Append(TEXT("\n"));
+						OutputString.Appendf(TEXT("\t%s"), *SafeActorLabel(RelevantProxy));
 					}
-			);
+				}
+			}
 		}
 	}
-	return true;
-	#endif
+
+	DebugHudContext.GetOutputString().Append(OutputString);
 }
+#endif
 
 #if WITH_EDITORONLY_DATA
 void UNiagaraDataInterfaceLandscape::GetFunctionsInternal(TArray<FNiagaraFunctionSignature>& OutFunctions) const
