- Rename OverrideName to Parameter Name for consistency

diff --git a/Engine/Plugins/Experimental/Water/Source/Runtime/Public/NiagaraWaterFunctionLibrary.h b/Engine/Plugins/Experimental/Water/Source/Runtime/Public/NiagaraWaterFunctionLibrary.h
--- a/Engine/Plugins/Experimental/Water/Source/Runtime/Public/NiagaraWaterFunctionLibrary.h
+++ b/Engine/Plugins/Experimental/Water/Source/Runtime/Public/NiagaraWaterFunctionLibrary.h
@@ -17,9 +17,9 @@ class WATER_API UNiagaraWaterFunctionLibrary : public UBlueprintFunctionLibrary
 	/** Sets the water body on the Niagra water data interface on a Niagara particle system*/
 	UE_DEPRECATED(all, "Use SetWaterBodyComponent instead")
 	UFUNCTION(BlueprintCallable, Category = Water, meta = (DisplayName = "Set Water Body"), meta = (DeprecatedFunction, DeprecationMessage="Use SetWaterBodyComponent instead"))
-	static void SetWaterBody(UNiagaraComponent* NiagaraSystem, const FString& OverrideName, AWaterBody* WaterBody) {}
+	static void SetWaterBody(UNiagaraComponent* NiagaraSystem, UPARAM(DisplayName = "Parameter Name") const FString & OverrideName, AWaterBody* WaterBody) {}
 
 	/** Sets the water body component on the Niagra water data interface on a Niagara particle system*/
 	UFUNCTION(BlueprintCallable, Category = Water)
-	static void SetWaterBodyComponent(UNiagaraComponent* NiagaraSystem, const FString& OverrideName, UWaterBodyComponent* WaterBodyComponent);
+	static void SetWaterBodyComponent(UNiagaraComponent* NiagaraSystem, UPARAM(DisplayName = "Parameter Name") const FString& OverrideName, UWaterBodyComponent* WaterBodyComponent);
 };
diff --git a/Engine/Plugins/FX/Niagara/Source/Niagara/Classes/NiagaraDataInterfaceArrayFunctionLibrary.h b/Engine/Plugins/FX/Niagara/Source/Niagara/Classes/NiagaraDataInterfaceArrayFunctionLibrary.h
--- a/Engine/Plugins/FX/Niagara/Source/Niagara/Classes/NiagaraDataInterfaceArrayFunctionLibrary.h
+++ b/Engine/Plugins/FX/Niagara/Source/Niagara/Classes/NiagaraDataInterfaceArrayFunctionLibrary.h
@@ -18,150 +18,150 @@ class UNiagaraDataInterfaceArrayFunctionLibrary : public UBlueprintFunctionLibra
 public:
 	/** Sets Niagara Array Float Data. */
 	UFUNCTION(BlueprintCallable, Category = Niagara, meta=(DisplayName = "Niagara Set Float Array"))
-	static NIAGARA_API void SetNiagaraArrayFloat(UNiagaraComponent* NiagaraSystem, FName OverrideName, const TArray<float>& ArrayData);
+	static NIAGARA_API void SetNiagaraArrayFloat(UNiagaraComponent* NiagaraSystem, UPARAM(DisplayName = "Parameter Name") FName OverrideName, const TArray<float>& ArrayData);
 	/** Sets Niagara Array FVector2D Data. */
 	UFUNCTION(BlueprintCallable, Category = Niagara, meta = (DisplayName = "Niagara Set Vector 2D Array"))
-	static NIAGARA_API void SetNiagaraArrayVector2D(UNiagaraComponent* NiagaraSystem, FName OverrideName, const TArray<FVector2D>& ArrayData);
+	static NIAGARA_API void SetNiagaraArrayVector2D(UNiagaraComponent* NiagaraSystem, UPARAM(DisplayName = "Parameter Name") FName OverrideName, const TArray<FVector2D>& ArrayData);
 	/** Sets Niagara Array FVector Data. */
 	UFUNCTION(BlueprintCallable, Category = Niagara, meta = (DisplayName = "Niagara Set Vector Array"))
-	static NIAGARA_API void SetNiagaraArrayVector(UNiagaraComponent* NiagaraSystem, FName OverrideName, const TArray<FVector>& ArrayData);
+	static NIAGARA_API void SetNiagaraArrayVector(UNiagaraComponent* NiagaraSystem, UPARAM(DisplayName = "Parameter Name") FName OverrideName, const TArray<FVector>& ArrayData);
 	/** Sets Niagara Array FVector Data. */
 	UFUNCTION(BlueprintCallable, Category = Niagara, meta = (DisplayName = "Niagara Set Position Array"))
-	static NIAGARA_API void SetNiagaraArrayPosition(UNiagaraComponent* NiagaraSystem, FName OverrideName, const TArray<FVector>& ArrayData);
+	static NIAGARA_API void SetNiagaraArrayPosition(UNiagaraComponent* NiagaraSystem, UPARAM(DisplayName = "Parameter Name") FName OverrideName, const TArray<FVector>& ArrayData);
 	/** Sets Niagara Array FVector4 Data. */
 	UFUNCTION(BlueprintCallable, Category = Niagara, meta = (DisplayName = "Niagara Set Vector 4 Array"))
-	static NIAGARA_API void SetNiagaraArrayVector4(UNiagaraComponent* NiagaraSystem, FName OverrideName, const TArray<FVector4>& ArrayData);
+	static NIAGARA_API void SetNiagaraArrayVector4(UNiagaraComponent* NiagaraSystem, UPARAM(DisplayName = "Parameter Name") FName OverrideName, const TArray<FVector4>& ArrayData);
 	/** Sets Niagara Array FLinearColor Data. */
 	UFUNCTION(BlueprintCallable, Category = Niagara, meta = (DisplayName = "Niagara Set Color Array"))
-	static NIAGARA_API void SetNiagaraArrayColor(UNiagaraComponent* NiagaraSystem, FName OverrideName, const TArray<FLinearColor>& ArrayData);
+	static NIAGARA_API void SetNiagaraArrayColor(UNiagaraComponent* NiagaraSystem, UPARAM(DisplayName = "Parameter Name") FName OverrideName, const TArray<FLinearColor>& ArrayData);
 	/** Sets Niagara Array FQuat Data. */
 	UFUNCTION(BlueprintCallable, Category = Niagara, meta = (DisplayName = "Niagara Set Quaternion Array"))
-	static NIAGARA_API void SetNiagaraArrayQuat(UNiagaraComponent* NiagaraSystem, FName OverrideName, const TArray<FQuat>& ArrayData);
+	static NIAGARA_API void SetNiagaraArrayQuat(UNiagaraComponent* NiagaraSystem, UPARAM(DisplayName = "Parameter Name") FName OverrideName, const TArray<FQuat>& ArrayData);
 	/**
 	 * Sets Niagara Array FMatrix Data.
 	 * @param bApplyLWCRebase When enabled the matrix translation will have the simulation tile offset subtracted from it
 	 */
 	UFUNCTION(BlueprintCallable, Category = Niagara, meta = (DisplayName = "Niagara Set Matrix Array"))
-	static NIAGARA_API void SetNiagaraArrayMatrix(UNiagaraComponent* NiagaraSystem, FName OverrideName, const TArray<FMatrix>& ArrayData, bool bApplyLWCRebase = true);
+	static NIAGARA_API void SetNiagaraArrayMatrix(UNiagaraComponent* NiagaraSystem, UPARAM(DisplayName = "Parameter Name") FName OverrideName, const TArray<FMatrix>& ArrayData, bool bApplyLWCRebase = true);
 	/** Sets Niagara Array Int32 Data. */
 	UFUNCTION(BlueprintCallable, Category = Niagara, meta = (DisplayName = "Niagara Set Int32 Array"))
-	static NIAGARA_API void SetNiagaraArrayInt32(UNiagaraComponent* NiagaraSystem, FName OverrideName, const TArray<int32>& ArrayData);
+	static NIAGARA_API void SetNiagaraArrayInt32(UNiagaraComponent* NiagaraSystem, UPARAM(DisplayName = "Parameter Name") FName OverrideName, const TArray<int32>& ArrayData);
 	/** Sets Niagara Array UInt8 Data. */
 	UFUNCTION(BlueprintCallable, Category = Niagara, meta = (DisplayName = "Niagara Set UInt8 Array"))
-	static NIAGARA_API void SetNiagaraArrayUInt8(UNiagaraComponent* NiagaraSystem, FName OverrideName, const TArray<int32>& ArrayData);
+	static NIAGARA_API void SetNiagaraArrayUInt8(UNiagaraComponent* NiagaraSystem, UPARAM(DisplayName = "Parameter Name") FName OverrideName, const TArray<int32>& ArrayData);
 	/** Sets Niagara Array Bool Data. */
 	UFUNCTION(BlueprintCallable, Category = Niagara, meta = (DisplayName = "Niagara Set Bool Array"))
-	static NIAGARA_API void SetNiagaraArrayBool(UNiagaraComponent* NiagaraSystem, FName OverrideName, const TArray<bool>& ArrayData);
+	static NIAGARA_API void SetNiagaraArrayBool(UNiagaraComponent* NiagaraSystem, UPARAM(DisplayName = "Parameter Name") FName OverrideName, const TArray<bool>& ArrayData);
 
 	/** Gets a copy of Niagara Float Data. */
 	UFUNCTION(BlueprintCallable, Category = Niagara, meta = (DisplayName = "Niagara Get Float Array"))
-	static NIAGARA_API TArray<float> GetNiagaraArrayFloat(UNiagaraComponent* NiagaraSystem, FName OverrideName);
+	static NIAGARA_API TArray<float> GetNiagaraArrayFloat(UNiagaraComponent* NiagaraSystem, UPARAM(DisplayName = "Parameter Name") FName OverrideName);
 	/** Gets a copy of Niagara FVector2D Data. */
 	UFUNCTION(BlueprintCallable, Category = Niagara, meta = (DisplayName = "Niagara Get Vector 2D Array"))
-	static NIAGARA_API TArray<FVector2D> GetNiagaraArrayVector2D(UNiagaraComponent* NiagaraSystem, FName OverrideName);
+	static NIAGARA_API TArray<FVector2D> GetNiagaraArrayVector2D(UNiagaraComponent* NiagaraSystem, UPARAM(DisplayName = "Parameter Name") FName OverrideName);
 	/** Gets a copy of Niagara FVector Data. */
 	UFUNCTION(BlueprintCallable, Category = Niagara, meta = (DisplayName = "Niagara Get Vector Array"))
-	static NIAGARA_API TArray<FVector> GetNiagaraArrayVector(UNiagaraComponent* NiagaraSystem, FName OverrideName);
+	static NIAGARA_API TArray<FVector> GetNiagaraArrayVector(UNiagaraComponent* NiagaraSystem, UPARAM(DisplayName = "Parameter Name") FName OverrideName);
 	/** Gets a copy of Niagara Position Data. */
 	UFUNCTION(BlueprintCallable, Category = Niagara, meta = (DisplayName = "Niagara Get Position Array"))
-	static NIAGARA_API TArray<FVector> GetNiagaraArrayPosition(UNiagaraComponent* NiagaraSystem, FName OverrideName);
+	static NIAGARA_API TArray<FVector> GetNiagaraArrayPosition(UNiagaraComponent* NiagaraSystem, UPARAM(DisplayName = "Parameter Name") FName OverrideName);
 	/** Gets a copy of Niagara FVector4 Data. */
 	UFUNCTION(BlueprintCallable, Category = Niagara, meta = (DisplayName = "Niagara Get Vector 4 Array"))
-	static NIAGARA_API TArray<FVector4> GetNiagaraArrayVector4(UNiagaraComponent* NiagaraSystem, FName OverrideName);
+	static NIAGARA_API TArray<FVector4> GetNiagaraArrayVector4(UNiagaraComponent* NiagaraSystem, UPARAM(DisplayName = "Parameter Name") FName OverrideName);
 	/** Gets a copy of Niagara FLinearColor Data. */
 	UFUNCTION(BlueprintCallable, Category = Niagara, meta = (DisplayName = "Niagara Get Color Array"))
-	static NIAGARA_API TArray<FLinearColor> GetNiagaraArrayColor(UNiagaraComponent* NiagaraSystem, FName OverrideName);
+	static NIAGARA_API TArray<FLinearColor> GetNiagaraArrayColor(UNiagaraComponent* NiagaraSystem, UPARAM(DisplayName = "Parameter Name") FName OverrideName);
 	/** Gets a copy of Niagara FQuat Data. */
 	UFUNCTION(BlueprintCallable, Category = Niagara, meta = (DisplayName = "Niagara Get Quaternion Array"))
-	static NIAGARA_API TArray<FQuat> GetNiagaraArrayQuat(UNiagaraComponent* NiagaraSystem, FName OverrideName);
+	static NIAGARA_API TArray<FQuat> GetNiagaraArrayQuat(UNiagaraComponent* NiagaraSystem, UPARAM(DisplayName = "Parameter Name") FName OverrideName);
 	/**
 	 * Gets a copy of Niagara FMatrix Data.
 	 * @param bApplyLWCRebase When enabled the matrix translation will have the simulation tile offset added to it
 	 */
 	UFUNCTION(BlueprintCallable, Category = Niagara, meta = (DisplayName = "Niagara Get Matrix Array"))
-	static NIAGARA_API TArray<FMatrix> GetNiagaraArrayMatrix(UNiagaraComponent* NiagaraSystem, FName OverrideName, bool bApplyLWCRebase = true);
+	static NIAGARA_API TArray<FMatrix> GetNiagaraArrayMatrix(UNiagaraComponent* NiagaraSystem, UPARAM(DisplayName = "Parameter Name") FName OverrideName, bool bApplyLWCRebase = true);
 	/** Gets a copy of Niagara Int32 Data. */
 	UFUNCTION(BlueprintCallable, Category = Niagara, meta = (DisplayName = "Niagara Get Int32 Array"))
-	static NIAGARA_API TArray<int32> GetNiagaraArrayInt32(UNiagaraComponent* NiagaraSystem, FName OverrideName);
+	static NIAGARA_API TArray<int32> GetNiagaraArrayInt32(UNiagaraComponent* NiagaraSystem, UPARAM(DisplayName = "Parameter Name") FName OverrideName);
 	/** Gets a copy of Niagara UInt8 Data. */
 	UFUNCTION(BlueprintCallable, Category = Niagara, meta = (DisplayName = "Niagara Get UInt8 Array"))
-	static NIAGARA_API TArray<int32> GetNiagaraArrayUInt8(UNiagaraComponent* NiagaraSystem, FName OverrideName);
+	static NIAGARA_API TArray<int32> GetNiagaraArrayUInt8(UNiagaraComponent* NiagaraSystem, UPARAM(DisplayName = "Parameter Name") FName OverrideName);
 	/** Gets a copy of Niagara Bool Data. */
 	UFUNCTION(BlueprintCallable, Category = Niagara, meta = (DisplayName = "Niagara Get Bool Array"))
-	static NIAGARA_API TArray<bool> GetNiagaraArrayBool(UNiagaraComponent* NiagaraSystem, FName OverrideName);
+	static NIAGARA_API TArray<bool> GetNiagaraArrayBool(UNiagaraComponent* NiagaraSystem, UPARAM(DisplayName = "Parameter Name") FName OverrideName);
 
 	/** Sets a single value within a Niagara Array Float. */
 	UFUNCTION(BlueprintCallable, Category = Niagara, meta = (DisplayName = "Niagara Set Float Array Value"))
-	static NIAGARA_API void SetNiagaraArrayFloatValue(UNiagaraComponent* NiagaraSystem, FName OverrideName, int Index, float Value, bool bSizeToFit);
+	static NIAGARA_API void SetNiagaraArrayFloatValue(UNiagaraComponent* NiagaraSystem, UPARAM(DisplayName = "Parameter Name") FName OverrideName, int Index, float Value, bool bSizeToFit);
 	/** Sets a single value within a Niagara Array FVector2D. */
 	UFUNCTION(BlueprintCallable, Category = Niagara, meta = (DisplayName = "Niagara Set Vector 2D Array Value"))
-	static NIAGARA_API void SetNiagaraArrayVector2DValue(UNiagaraComponent* NiagaraSystem, FName OverrideName, int Index, const FVector2D& Value, bool bSizeToFit);
+	static NIAGARA_API void SetNiagaraArrayVector2DValue(UNiagaraComponent* NiagaraSystem, UPARAM(DisplayName = "Parameter Name") FName OverrideName, int Index, const FVector2D& Value, bool bSizeToFit);
 	/** Sets a single value within a Niagara Array FVector. */
 	UFUNCTION(BlueprintCallable, Category = Niagara, meta = (DisplayName = "Niagara Set Vector Array Value"))
-	static NIAGARA_API void SetNiagaraArrayVectorValue(UNiagaraComponent* NiagaraSystem, FName OverrideName, int Index, const FVector& Value, bool bSizeToFit);
+	static NIAGARA_API void SetNiagaraArrayVectorValue(UNiagaraComponent* NiagaraSystem, UPARAM(DisplayName = "Parameter Name") FName OverrideName, int Index, const FVector& Value, bool bSizeToFit);
 	UFUNCTION(BlueprintCallable, Category = Niagara, meta = (DisplayName = "Niagara Set Position Array Value"))
-	static NIAGARA_API void SetNiagaraArrayPositionValue(UNiagaraComponent* NiagaraSystem, FName OverrideName, int Index, const FVector& Value, bool bSizeToFit);
+	static NIAGARA_API void SetNiagaraArrayPositionValue(UNiagaraComponent* NiagaraSystem, UPARAM(DisplayName = "Parameter Name") FName OverrideName, int Index, const FVector& Value, bool bSizeToFit);
 	/** Sets a single value within a Niagara Array FVector4. */
 	UFUNCTION(BlueprintCallable, Category = Niagara, meta = (DisplayName = "Niagara Set Vector 4 Array Value"))
-	static NIAGARA_API void SetNiagaraArrayVector4Value(UNiagaraComponent* NiagaraSystem, FName OverrideName, int Index, const FVector4& Value, bool bSizeToFit);
+	static NIAGARA_API void SetNiagaraArrayVector4Value(UNiagaraComponent* NiagaraSystem, UPARAM(DisplayName = "Parameter Name") FName OverrideName, int Index, const FVector4& Value, bool bSizeToFit);
 	/** Sets a single value within a Niagara Array FLinearColor. */
 	UFUNCTION(BlueprintCallable, Category = Niagara, meta = (DisplayName = "Niagara Set Color Array Value"))
-	static NIAGARA_API void SetNiagaraArrayColorValue(UNiagaraComponent* NiagaraSystem, FName OverrideName, int Index, const FLinearColor& Value, bool bSizeToFit);
+	static NIAGARA_API void SetNiagaraArrayColorValue(UNiagaraComponent* NiagaraSystem, UPARAM(DisplayName = "Parameter Name") FName OverrideName, int Index, const FLinearColor& Value, bool bSizeToFit);
 	/** Sets a single value within a Niagara Array FQuat. */
 	UFUNCTION(BlueprintCallable, Category = Niagara, meta = (DisplayName = "Niagara Set Quaternion Array Value"))
-	static NIAGARA_API void SetNiagaraArrayQuatValue(UNiagaraComponent* NiagaraSystem, FName OverrideName, int Index, const FQuat& Value, bool bSizeToFit);
+	static NIAGARA_API void SetNiagaraArrayQuatValue(UNiagaraComponent* NiagaraSystem, UPARAM(DisplayName = "Parameter Name") FName OverrideName, int Index, const FQuat& Value, bool bSizeToFit);
 	/**
 	 * Sets a single value within a Niagara Array FMatrix.
 	 * @param bApplyLWCRebase When enabled the matrix translation will have the simulation tile offset subtracted from it
 	 */
 	UFUNCTION(BlueprintCallable, Category = Niagara, meta = (DisplayName = "Niagara Set Matrix Array Value"))
-	static NIAGARA_API void SetNiagaraArrayMatrixValue(UNiagaraComponent* NiagaraSystem, FName OverrideName, int Index, const FMatrix& Value, bool bSizeToFit, bool bApplyLWCRebase = true);
+	static NIAGARA_API void SetNiagaraArrayMatrixValue(UNiagaraComponent* NiagaraSystem, UPARAM(DisplayName = "Parameter Name") FName OverrideName, int Index, const FMatrix& Value, bool bSizeToFit, bool bApplyLWCRebase = true);
 	/** Sets a single value within a Niagara Array Int32. */
 	UFUNCTION(BlueprintCallable, Category = Niagara, meta = (DisplayName = "Niagara Set Int32 Array Value"))
-	static NIAGARA_API void SetNiagaraArrayInt32Value(UNiagaraComponent* NiagaraSystem, FName OverrideName, int Index, int32 Value, bool bSizeToFit);
+	static NIAGARA_API void SetNiagaraArrayInt32Value(UNiagaraComponent* NiagaraSystem, UPARAM(DisplayName = "Parameter Name") FName OverrideName, int Index, int32 Value, bool bSizeToFit);
 	/** Sets a single value within a Niagara Array UInt8. */
 	UFUNCTION(BlueprintCallable, Category = Niagara, meta = (DisplayName = "Niagara Set UInt8 Array Value"))
-	static NIAGARA_API void SetNiagaraArrayUInt8Value(UNiagaraComponent* NiagaraSystem, FName OverrideName, int Index, int32 Value, bool bSizeToFit);
+	static NIAGARA_API void SetNiagaraArrayUInt8Value(UNiagaraComponent* NiagaraSystem, UPARAM(DisplayName = "Parameter Name") FName OverrideName, int Index, int32 Value, bool bSizeToFit);
 	/** Sets a single value within a Niagara Array Bool. */
 	UFUNCTION(BlueprintCallable, Category = Niagara, meta = (DisplayName = "Niagara Set Bool Array Value"))
-	static NIAGARA_API void SetNiagaraArrayBoolValue(UNiagaraComponent* NiagaraSystem, FName OverrideName, int Index, const bool& Value, bool bSizeToFit);
+	static NIAGARA_API void SetNiagaraArrayBoolValue(UNiagaraComponent* NiagaraSystem, UPARAM(DisplayName = "Parameter Name") FName OverrideName, int Index, const bool& Value, bool bSizeToFit);
 
 	/** Gets a single value within a Niagara Array Float. */
 	UFUNCTION(BlueprintCallable, Category = Niagara, meta = (DisplayName = "Niagara Get Float Array Value"))
-	static NIAGARA_API float GetNiagaraArrayFloatValue(UNiagaraComponent* NiagaraSystem, FName OverrideName, int Index);
+	static NIAGARA_API float GetNiagaraArrayFloatValue(UNiagaraComponent* NiagaraSystem, UPARAM(DisplayName = "Parameter Name") FName OverrideName, int Index);
 	/** Gets a single value within a Niagara Array FVector2D. */
 	UFUNCTION(BlueprintCallable, Category = Niagara, meta = (DisplayName = "Niagara Get Vector 2D Array Value"))
-	static NIAGARA_API FVector2D GetNiagaraArrayVector2DValue(UNiagaraComponent* NiagaraSystem, FName OverrideName, int Index);
+	static NIAGARA_API FVector2D GetNiagaraArrayVector2DValue(UNiagaraComponent* NiagaraSystem, UPARAM(DisplayName = "Parameter Name") FName OverrideName, int Index);
 	/** Gets a single value within a Niagara Array FVector. */
 	UFUNCTION(BlueprintCallable, Category = Niagara, meta = (DisplayName = "Niagara Get Vector Array Value"))
-	static NIAGARA_API FVector GetNiagaraArrayVectorValue(UNiagaraComponent* NiagaraSystem, FName OverrideName, int Index);
+	static NIAGARA_API FVector GetNiagaraArrayVectorValue(UNiagaraComponent* NiagaraSystem, UPARAM(DisplayName = "Parameter Name") FName OverrideName, int Index);
 	/** Gets a single value within a Niagara Array Position. */
 	UFUNCTION(BlueprintCallable, Category = Niagara, meta = (DisplayName = "Niagara Get Position Array Value"))
-	static NIAGARA_API FVector GetNiagaraArrayPositionValue(UNiagaraComponent* NiagaraSystem, FName OverrideName, int Index);
+	static NIAGARA_API FVector GetNiagaraArrayPositionValue(UNiagaraComponent* NiagaraSystem, UPARAM(DisplayName = "Parameter Name") FName OverrideName, int Index);
 	/** Gets a single value within a Niagara Array FVector4. */
 	UFUNCTION(BlueprintCallable, Category = Niagara, meta = (DisplayName = "Niagara Get Vector 4 Array Value"))
-	static NIAGARA_API FVector4 GetNiagaraArrayVector4Value(UNiagaraComponent* NiagaraSystem, FName OverrideName, int Index);
+	static NIAGARA_API FVector4 GetNiagaraArrayVector4Value(UNiagaraComponent* NiagaraSystem, UPARAM(DisplayName = "Parameter Name") FName OverrideName, int Index);
 	/** Gets a single value within a Niagara Array FLinearColor. */
 	UFUNCTION(BlueprintCallable, Category = Niagara, meta = (DisplayName = "Niagara Get Color Array Value"))
-	static NIAGARA_API FLinearColor GetNiagaraArrayColorValue(UNiagaraComponent* NiagaraSystem, FName OverrideName, int Index);
+	static NIAGARA_API FLinearColor GetNiagaraArrayColorValue(UNiagaraComponent* NiagaraSystem, UPARAM(DisplayName = "Parameter Name") FName OverrideName, int Index);
 	/** Gets a single value within a Niagara Array FQuat. */
 	UFUNCTION(BlueprintCallable, Category = Niagara, meta = (DisplayName = "Niagara Get Quaternion Array Value"))
-	static NIAGARA_API FQuat GetNiagaraArrayQuatValue(UNiagaraComponent* NiagaraSystem, FName OverrideName, int Index);
+	static NIAGARA_API FQuat GetNiagaraArrayQuatValue(UNiagaraComponent* NiagaraSystem, UPARAM(DisplayName = "Parameter Name") FName OverrideName, int Index);
 	/**
 	 * Gets a single value within a Niagara Array FMatrix.
 	 * @param bApplyLWCRebase When enabled the matrix translation will have the simulation tile offset added to it
 	 */
 	UFUNCTION(BlueprintCallable, Category = Niagara, meta = (DisplayName = "Niagara Get Matrix Array Value"))
-	static NIAGARA_API FMatrix GetNiagaraArrayMatrixValue(UNiagaraComponent* NiagaraSystem, FName OverrideName, int Index, bool bApplyLWCRebase = true);
+	static NIAGARA_API FMatrix GetNiagaraArrayMatrixValue(UNiagaraComponent* NiagaraSystem, UPARAM(DisplayName = "Parameter Name") FName OverrideName, int Index, bool bApplyLWCRebase = true);
 	/** Gets a single value within a Niagara Array Int32. */
 	UFUNCTION(BlueprintCallable, Category = Niagara, meta = (DisplayName = "Niagara Get Int32 Array Value"))
-	static NIAGARA_API int32 GetNiagaraArrayInt32Value(UNiagaraComponent* NiagaraSystem, FName OverrideName, int Index);
+	static NIAGARA_API int32 GetNiagaraArrayInt32Value(UNiagaraComponent* NiagaraSystem, UPARAM(DisplayName = "Parameter Name") FName OverrideName, int Index);
 	/** Gets a single value within a Niagara Array UInt8. */
 	UFUNCTION(BlueprintCallable, Category = Niagara, meta = (DisplayName = "Niagara Get UInt8 Array Value"))
-	static NIAGARA_API int32 GetNiagaraArrayUInt8Value(UNiagaraComponent* NiagaraSystem, FName OverrideName, int Index);
+	static NIAGARA_API int32 GetNiagaraArrayUInt8Value(UNiagaraComponent* NiagaraSystem, UPARAM(DisplayName = "Parameter Name") FName OverrideName, int Index);
 	/** Gets a single value within a Niagara Array Bool. */
 	UFUNCTION(BlueprintCallable, Category = Niagara, meta = (DisplayName = "Niagara Get Bool Array Value"))
-	static NIAGARA_API bool GetNiagaraArrayBoolValue(UNiagaraComponent* NiagaraSystem, FName OverrideName, int Index);
+	static NIAGARA_API bool GetNiagaraArrayBoolValue(UNiagaraComponent* NiagaraSystem, UPARAM(DisplayName = "Parameter Name") FName OverrideName, int Index);
 
 	//////////////////////////////////////////////////////////////////////////////////////////////////////////////////
 	// None BP compatable set functions
diff --git a/Engine/Plugins/FX/Niagara/Source/Niagara/Internal/DataInterface/NiagaraDataInterfaceArrayMesh.h b/Engine/Plugins/FX/Niagara/Source/Niagara/Internal/DataInterface/NiagaraDataInterfaceArrayMesh.h
--- a/Engine/Plugins/FX/Niagara/Source/Niagara/Internal/DataInterface/NiagaraDataInterfaceArrayMesh.h
+++ b/Engine/Plugins/FX/Niagara/Source/Niagara/Internal/DataInterface/NiagaraDataInterfaceArrayMesh.h
@@ -43,8 +43,8 @@ public:
 	//INiagaraRenderableMeshArrayInterface Interface End
 
 	UFUNCTION(BlueprintCallable, Category = Niagara, meta = (DisplayName = "Niagara Set Mesh Array", UnsafeDuringActorConstruction = "true"))
-	static NIAGARA_API void SetNiagaraArrayMesh(UNiagaraComponent* NiagaraComponent, FName OverrideName, const TArray<FNiagaraMeshRendererMeshPropertiesBase>& ArrayData);
+	static NIAGARA_API void SetNiagaraArrayMesh(UNiagaraComponent* NiagaraComponent, UPARAM(DisplayName = "Parameter Name") FName OverrideName, const TArray<FNiagaraMeshRendererMeshPropertiesBase>& ArrayData);
 
 	UFUNCTION(BlueprintCallable, Category = Niagara, meta = (DisplayName = "Niagara Set Mesh Array (Static Mesh)", UnsafeDuringActorConstruction = "true"))
-	static NIAGARA_API void SetNiagaraArrayMeshSM(UNiagaraComponent* NiagaraComponent, FName OverrideName, const TArray<UStaticMesh*>& ArrayData);
+	static NIAGARA_API void SetNiagaraArrayMeshSM(UNiagaraComponent* NiagaraComponent, UPARAM(DisplayName = "Parameter Name") FName OverrideName, const TArray<UStaticMesh*>& ArrayData);
 };
diff --git a/Engine/Plugins/FX/Niagara/Source/Niagara/Public/NiagaraDataInterfaceRigidMeshCollisionQuery.h b/Engine/Plugins/FX/Niagara/Source/Niagara/Public/NiagaraDataInterfaceRigidMeshCollisionQuery.h
--- a/Engine/Plugins/FX/Niagara/Source/Niagara/Public/NiagaraDataInterfaceRigidMeshCollisionQuery.h
+++ b/Engine/Plugins/FX/Niagara/Source/Niagara/Public/NiagaraDataInterfaceRigidMeshCollisionQuery.h
@@ -307,5 +307,5 @@ class UNiagaraDIRigidMeshCollisionFunctionLibrary : public UBlueprintFunctionLib
 public:
 
 	UFUNCTION(BlueprintCallable, Category = Niagara, meta = (DisplayName = "Niagara Set Source Actors"))
-	static NIAGARA_API void SetSourceActors(UNiagaraComponent* NiagaraSystem, FName OverrideName, const TArray<AActor*>& SourceActors);
+	static NIAGARA_API void SetSourceActors(UNiagaraComponent* NiagaraSystem, UPARAM(DisplayName = "Parameter Name") FName OverrideName, const TArray<AActor*>& SourceActors);
 };
diff --git a/Engine/Plugins/FX/Niagara/Source/Niagara/Public/NiagaraFunctionLibrary.h b/Engine/Plugins/FX/Niagara/Source/Niagara/Public/NiagaraFunctionLibrary.h
--- a/Engine/Plugins/FX/Niagara/Source/Niagara/Public/NiagaraFunctionLibrary.h
+++ b/Engine/Plugins/FX/Niagara/Source/Niagara/Public/NiagaraFunctionLibrary.h
@@ -48,29 +48,29 @@ public:
 
 	/** Sets a Niagara StaticMesh parameter by name, overriding locally if necessary.*/
 	UFUNCTION(BlueprintCallable, Category = Niagara, meta = (DisplayName = "Set Niagara Static Mesh Component"))
-	static NIAGARA_API void OverrideSystemUserVariableStaticMeshComponent(UNiagaraComponent* NiagaraSystem, const FString& OverrideName, UStaticMeshComponent* StaticMeshComponent);
+	static NIAGARA_API void OverrideSystemUserVariableStaticMeshComponent(UNiagaraComponent* NiagaraSystem, UPARAM(DisplayName = "Parameter Name") const FString& OverrideName, UStaticMeshComponent* StaticMeshComponent);
 
 	UFUNCTION(BlueprintCallable, Category = Niagara, meta = (DisplayName = "Set Niagara Static Mesh Directly"))
-	static NIAGARA_API void OverrideSystemUserVariableStaticMesh(UNiagaraComponent* NiagaraSystem, const FString& OverrideName, UStaticMesh* StaticMesh);
+	static NIAGARA_API void OverrideSystemUserVariableStaticMesh(UNiagaraComponent* NiagaraSystem, UPARAM(DisplayName = "Parameter Name") const FString& OverrideName, UStaticMesh* StaticMesh);
 
 	/** Get the skeletal mesh data interface by name .*/
 	static NIAGARA_API class UNiagaraDataInterfaceSkeletalMesh* GetSkeletalMeshDataInterface(UNiagaraComponent* NiagaraSystem, const FString& OverrideName);
 
 	/** Sets a Niagara StaticMesh parameter by name, overriding locally if necessary.*/
 	UFUNCTION(BlueprintCallable, Category = Niagara, meta = (DisplayName = "Set Niagara Skeletal Mesh Component"))
-	static NIAGARA_API void OverrideSystemUserVariableSkeletalMeshComponent(UNiagaraComponent* NiagaraSystem, const FString& OverrideName, USkeletalMeshComponent* SkeletalMeshComponent);
+	static NIAGARA_API void OverrideSystemUserVariableSkeletalMeshComponent(UNiagaraComponent* NiagaraSystem, UPARAM(DisplayName = "Parameter Name") const FString& OverrideName, USkeletalMeshComponent* SkeletalMeshComponent);
 
 	/** Sets the SamplingRegion to use on the skeletal mesh data interface, this is destructive as it modifies the data interface. */
 	UFUNCTION(BlueprintCallable, Category = Niagara)
-	static NIAGARA_API void SetSkeletalMeshDataInterfaceSamplingRegions(UNiagaraComponent* NiagaraSystem, const FString& OverrideName, const TArray<FName>& SamplingRegions);
+	static NIAGARA_API void SetSkeletalMeshDataInterfaceSamplingRegions(UNiagaraComponent* NiagaraSystem, UPARAM(DisplayName = "Parameter Name") const FString& OverrideName, const TArray<FName>& SamplingRegions);
 
 	/** Sets the Filtered Bones to use on the skeletal mesh data interface, this is destructive as it modifies the data interface. */
 	UFUNCTION(BlueprintCallable, Category = Niagara)
-	static NIAGARA_API void SetSkeletalMeshDataInterfaceFilteredBones(UNiagaraComponent* NiagaraSystem, const FString& OverrideName, const TArray<FName>& FilteredBones);
+	static NIAGARA_API void SetSkeletalMeshDataInterfaceFilteredBones(UNiagaraComponent* NiagaraSystem, UPARAM(DisplayName = "Parameter Name") const FString& OverrideName, const TArray<FName>& FilteredBones);
 
 	/** Sets the Filtered Sockets to use on the skeletal mesh data interface, this is destructive as it modifies the data interface. */
 	UFUNCTION(BlueprintCallable, Category = Niagara)
-	static NIAGARA_API void SetSkeletalMeshDataInterfaceFilteredSockets(UNiagaraComponent* NiagaraSystem, const FString& OverrideName, const TArray<FName>& FilteredSockets);
+	static NIAGARA_API void SetSkeletalMeshDataInterfaceFilteredSockets(UNiagaraComponent* NiagaraSystem, UPARAM(DisplayName = "Parameter Name") const FString& OverrideName, const TArray<FName>& FilteredSockets);
 
 	/** Sets managed mode parameters for the Scene capture 2D  data interface, this is destructive as it modifies the data interface. */
 	UFUNCTION(BlueprintCallable, Category = Niagara)
@@ -87,15 +87,15 @@ public:
 
 	/** Overrides the Texture Object for a Niagara Texture Data Interface User Parameter.*/
 	UFUNCTION(BlueprintCallable, Category = Niagara)
-	static NIAGARA_API void SetTextureObject(UNiagaraComponent* NiagaraSystem, const FString& OverrideName, UTexture* Texture);
+	static NIAGARA_API void SetTextureObject(UNiagaraComponent* NiagaraSystem, UPARAM(DisplayName = "Parameter Name") const FString& OverrideName, UTexture* Texture);
 
 	/** Overrides the 2D Array Texture for a Niagara 2D Array Texture Data Interface User Parameter.*/
 	UFUNCTION(BlueprintCallable, Category = Niagara)
-	static NIAGARA_API void SetTexture2DArrayObject(UNiagaraComponent* NiagaraSystem, const FString& OverrideName, class UTexture2DArray* Texture);
+	static NIAGARA_API void SetTexture2DArrayObject(UNiagaraComponent* NiagaraSystem, UPARAM(DisplayName = "Parameter Name") const FString& OverrideName, class UTexture2DArray* Texture);
 
 	/** Overrides the Volume Texture for a Niagara Volume Texture Data Interface User Parameter.*/
 	UFUNCTION(BlueprintCallable, Category = Niagara)
-	static NIAGARA_API void SetVolumeTextureObject(UNiagaraComponent* NiagaraSystem, const FString& OverrideName, UVolumeTexture* Texture);
+	static NIAGARA_API void SetVolumeTextureObject(UNiagaraComponent* NiagaraSystem, UPARAM(DisplayName = "Parameter Name") const FString& OverrideName, UVolumeTexture* Texture);
 	
 	/** Finds an array interface of the given class. */
 	static NIAGARA_API UNiagaraDataInterface* GetDataInterface(UClass* DIClass, UNiagaraComponent* NiagaraSystem, FName OverrideName);
