- Don't release the count if one exists, it's already been released by the compute dispatch

#rnx
#tests Cache a sequencer system using a GPU emitter, check cache restart works as expected on GPU

diff --git a/Engine/Plugins/FX/Niagara/Source/Niagara/Private/NiagaraSimCacheHelper.h b/Engine/Plugins/FX/Niagara/Source/Niagara/Private/NiagaraSimCacheHelper.h
--- a/Engine/Plugins/FX/Niagara/Source/Niagara/Private/NiagaraSimCacheHelper.h
+++ b/Engine/Plugins/FX/Niagara/Source/Niagara/Private/NiagaraSimCacheHelper.h
@@ -864,12 +864,7 @@ struct FNiagaraSimCacheHelper
 						GPUExecContext->CountOffset_RT = CountManager.AcquireOrAllocateEntry(RHICmdList);
 					}
 
-					uint32 CurrentCountOffset = DataBuffer.GetGPUInstanceCountBufferOffset();
-					if (CurrentCountOffset != INDEX_NONE && CurrentCountOffset != GPUExecContext->CountOffset_RT)
-					{
-						CountManager.FreeEntry(CurrentCountOffset);
-						DataBuffer.SetGPUInstanceCountBufferOffset(GPUExecContext->CountOffset_RT);
-					}
+					DataBuffer.SetGPUInstanceCountBufferOffset(GPUExecContext->CountOffset_RT);
 
 					GPUExecContext->CurrentNumInstances_RT = NumInstances;
 					GPUExecContext->ParticleCountWriteFence = ParticleCountFence;
