- Change check to ensure in the editor if the renderers do not match the draw order

#rnx
#tests Force arrays to mismatch in code to verify it's handled ok

diff --git a/Engine/Plugins/FX/Niagara/Source/Niagara/Private/NiagaraSystemRenderData.cpp b/Engine/Plugins/FX/Niagara/Source/Niagara/Private/NiagaraSystemRenderData.cpp
--- a/Engine/Plugins/FX/Niagara/Source/Niagara/Private/NiagaraSystemRenderData.cpp
+++ b/Engine/Plugins/FX/Niagara/Source/Niagara/Private/NiagaraSystemRenderData.cpp
@@ -311,7 +311,17 @@ void FNiagaraSystemRenderData::RecacheRenderers(const FNiagaraSystemInstance& Sy
 	}
 
 	// If we have renderers then the draw order on the system should match, when compiling the number of renderers can be zero
+#if WITH_EDITOR
+	// In editor builds we may have cases where modifications to the enabled state is inside undo / redo and we have no pre/post calls to execute an update context.
+	// We can handle these invalid cases by clearing the renderers and logging an ensure, when we have a solid repro we can investigate fixing, or we can separate out editor vs runtime data.
+	if (!ensureMsgf(EmitterRenderers_GT.Num() == 0|| (EmitterRenderers_GT.Num() == RendererDrawOrder.Num()), TEXT("EmitterRenderers Num %d does not match System DrawOrder %d"), EmitterRenderers_GT.Num(), RendererDrawOrder.Num()))
+	{
+		EmitterRenderers_GT.Empty();
+	}
+#else
+	// In cooked builds the enabled state will not be mutated so any failures here are fatals
 	checkf((EmitterRenderers_GT.Num() == 0) || (EmitterRenderers_GT.Num() == RendererDrawOrder.Num()), TEXT("EmitterRenderers Num %d does not match System DrawOrder %d"), EmitterRenderers_GT.Num(), RendererDrawOrder.Num());
+#endif
 
 	UE::Tasks::FTaskEvent TaskEvent{ UE_SOURCE_LOCATION };
 
