- Make sure the count buffer is updated even if we don't simulation GPU particles
- External users such as data interfaces / stateless need the counts to be allocated

#rnx

diff --git a/Engine/Plugins/FX/Niagara/Source/Niagara/Private/NiagaraGpuComputeDispatch.cpp b/Engine/Plugins/FX/Niagara/Source/Niagara/Private/NiagaraGpuComputeDispatch.cpp
--- a/Engine/Plugins/FX/Niagara/Source/Niagara/Private/NiagaraGpuComputeDispatch.cpp
+++ b/Engine/Plugins/FX/Niagara/Source/Niagara/Private/NiagaraGpuComputeDispatch.cpp
@@ -1988,10 +1988,16 @@ void FNiagaraGpuComputeDispatch::PreInitViews(FRDGBuilder& GraphBuilder, bool bA
 
 			ExecuteTicks(GraphBuilder, TConstStridedView<FSceneView>(), ENiagaraGpuComputeTickStage::PreInitViews);
 		}
+		else
+		{
+			// We need to update the count buffer with the requested estimate from our proxies as they may still run compute work
+			// Not doing this means a system which assumes it has a free count can fail due to the buffer being zero sized for example
+			GPUInstanceCounterManager.ResizeBuffers(GraphBuilder.RHICmdList, ProxyGpuCountBufferEstimate);
+		}
 	}
 	else
 	{
-		GPUInstanceCounterManager.ResizeBuffers(GraphBuilder.RHICmdList, 0);
+		GPUInstanceCounterManager.ResizeBuffers(GraphBuilder.RHICmdList, ProxyGpuCountBufferEstimate);
 		FinishDispatches();
 	}
 
