- Stateless - trim disabled modules on cook
- When loading on cooked target we replace the module with the CDO to avoid having to handle nullptr
- The IsEnabled call on the module will return false if the module is the CDO
- With some simple test assets this went from 255kb -> 220kb or ~14%

#tests Run ET / cooked test

diff --git a/Engine/Plugins/FX/Niagara/Source/Niagara/Internal/Stateless/NiagaraStatelessEmitter.h b/Engine/Plugins/FX/Niagara/Source/Niagara/Internal/Stateless/NiagaraStatelessEmitter.h
--- a/Engine/Plugins/FX/Niagara/Source/Niagara/Internal/Stateless/NiagaraStatelessEmitter.h
+++ b/Engine/Plugins/FX/Niagara/Source/Niagara/Internal/Stateless/NiagaraStatelessEmitter.h
@@ -39,6 +39,7 @@ public:
 
 public:
 	//~Begin UObject Interface
+	virtual void Serialize(FArchive& Ar) override;
 	virtual void PostLoad() override;
 	virtual bool NeedsLoadForTargetPlatform(const ITargetPlatform* TargetPlatform) const override;
 #if WITH_EDITOR
diff --git a/Engine/Plugins/FX/Niagara/Source/Niagara/Internal/Stateless/NiagaraStatelessModule.h b/Engine/Plugins/FX/Niagara/Source/Niagara/Internal/Stateless/NiagaraStatelessModule.h
--- a/Engine/Plugins/FX/Niagara/Source/Niagara/Internal/Stateless/NiagaraStatelessModule.h
+++ b/Engine/Plugins/FX/Niagara/Source/Niagara/Internal/Stateless/NiagaraStatelessModule.h
@@ -30,7 +30,7 @@ protected:
 #endif
 
 public:
-	bool IsModuleEnabled() const { return bModuleEnabled; }
+	NIAGARA_API bool IsModuleEnabled() const;
 #if WITH_EDITOR
 	bool IsDebugDrawEnabled() const { return bModuleEnabled && bDebugDrawEnabled; }
 
diff --git a/Engine/Plugins/FX/Niagara/Source/Niagara/Private/Stateless/NiagaraStatelessEmitter.cpp b/Engine/Plugins/FX/Niagara/Source/Niagara/Private/Stateless/NiagaraStatelessEmitter.cpp
--- a/Engine/Plugins/FX/Niagara/Source/Niagara/Private/Stateless/NiagaraStatelessEmitter.cpp
+++ b/Engine/Plugins/FX/Niagara/Source/Niagara/Private/Stateless/NiagaraStatelessEmitter.cpp
@@ -93,6 +93,42 @@ namespace NiagaraStatelessInternal
 	}
 }
 
+void UNiagaraStatelessEmitter::Serialize(FArchive& Ar)
+{
+#if WITH_EDITORONLY_DATA
+	// When cooking we will remove any modules that are not enabled as we can replace with the CDO at runtime saving memory
+	TOptional<TGuardValue<TArray<TObjectPtr<UNiagaraStatelessModule>>>> ModulesGuard;
+	if (Ar.IsCooking() && Ar.CookingTarget()->RequiresCookedData())
+	{
+		TArray<TObjectPtr<UNiagaraStatelessModule>> TrimmedModules;
+		TrimmedModules.Reserve(Modules.Num());
+		for (UNiagaraStatelessModule* Module : Modules)
+		{
+			TrimmedModules.Add(Module->IsModuleEnabled() ? Module : nullptr);
+		}
+		ModulesGuard.Emplace(Modules, TrimmedModules);
+	}
+#endif
+
+	Super::Serialize(Ar);
+
+	// On cooked make sure any null modules are replaced with the CDO
+	if (Ar.IsLoading() && FPlatformProperties::RequiresCookedData())
+	{
+		if (const UNiagaraStatelessEmitterTemplate* EmitterTemplate = GetEmitterTemplate())
+		{
+			TConstArrayView<TObjectPtr<UClass>> ModuleClass = EmitterTemplate->GetModules();
+			for (int32 i = 0; i < Modules.Num(); ++i)
+			{
+				if (Modules[i] == nullptr)
+				{
+					Modules[i] = CastChecked<UNiagaraStatelessModule>(ModuleClass[i]->GetDefaultObject());
+				}
+			}
+		}
+	}
+}
+
 void UNiagaraStatelessEmitter::PostLoad()
 {
 	Super::PostLoad();
diff --git a/Engine/Plugins/FX/Niagara/Source/Niagara/Private/Stateless/NiagaraStatelessEmitterTemplate.cpp b/Engine/Plugins/FX/Niagara/Source/Niagara/Private/Stateless/NiagaraStatelessEmitterTemplate.cpp
--- a/Engine/Plugins/FX/Niagara/Source/Niagara/Private/Stateless/NiagaraStatelessEmitterTemplate.cpp
+++ b/Engine/Plugins/FX/Niagara/Source/Niagara/Private/Stateless/NiagaraStatelessEmitterTemplate.cpp
@@ -45,13 +45,17 @@ void UNiagaraStatelessEmitterTemplate::PostInitProperties()
 	Super::PostInitProperties();
 
 	// We can end up hitting PostInitProperties before the Niagara Module has initialized bindings this needs, mark this object for deferred init and early out.
-	if (FModuleManager::Get().IsModuleLoaded("Niagara") == false)
+	if (GetClass()->HasAnyClassFlags(CLASS_Abstract) == false)
 	{
-		NiagaraStatelessEmitterTemplateImpl::ObjectsToDeferredInit.Add(this);
-	}
-	else
-	{
-		InitModulesAndAttributes();
+		if (FModuleManager::Get().IsModuleLoaded("Niagara") == false)
+		{
+			NiagaraStatelessEmitterTemplateImpl::ObjectsToDeferredInit.Add(this);
+		}
+		else
+		{
+			InitModulesAndAttributes();
+			checkf(Modules.Num() != 0, TEXT("StatelessTemplate(%s) has no modules, this is likely a code error"), *GetNameSafe(this));
+		}
 	}
 }
 
@@ -62,15 +66,16 @@ void UNiagaraStatelessEmitterTemplate::InitCDOPropertiesAfterModuleStartup()
 		if (UNiagaraStatelessEmitterTemplate* EmitterObject = WeakObject.Get())
 		{
 			EmitterObject->InitModulesAndAttributes();
+			checkf(EmitterObject->Modules.Num() != 0, TEXT("StatelessTemplate(%s) has no modules, this is likely a code error"), *GetNameSafe(EmitterObject));
 		}
 	}
+	NiagaraStatelessEmitterTemplateImpl::ObjectsToDeferredInit.Empty();
 }
 
 //////////////////////////////////////////////////////////////////////////////////////////////////////////////////////////////////////////////////////////
 
 void UNiagaraStatelessEmitterDefault::InitModulesAndAttributes()
 {
-#if WITH_EDITORONLY_DATA
 	Modules =
 	{
 		// Initializer Modules
@@ -100,6 +105,7 @@ void UNiagaraStatelessEmitterDefault::InitModulesAndAttributes()
 		UNiagaraStatelessModule_SubUVAnimation::StaticClass(),
 	};
 
+#if WITH_EDITORONLY_DATA
 	const FNiagaraStatelessGlobals& StatelessGlobals = FNiagaraStatelessGlobals::Get();
 	OutputComponents =
 	{
diff --git a/Engine/Plugins/FX/Niagara/Source/Niagara/Private/Stateless/NiagaraStatelessModule.cpp b/Engine/Plugins/FX/Niagara/Source/Niagara/Private/Stateless/NiagaraStatelessModule.cpp
--- a/Engine/Plugins/FX/Niagara/Source/Niagara/Private/Stateless/NiagaraStatelessModule.cpp
+++ b/Engine/Plugins/FX/Niagara/Source/Niagara/Private/Stateless/NiagaraStatelessModule.cpp
@@ -7,6 +7,11 @@ const FName UNiagaraStatelessModule::PrivateMemberNames::bModuleEnabled = GET_ME
 const FName UNiagaraStatelessModule::PrivateMemberNames::bDebugDrawEnabled = GET_MEMBER_NAME_CHECKED(UNiagaraStatelessModule, bDebugDrawEnabled);
 #endif
 
+bool UNiagaraStatelessModule::IsModuleEnabled() const
+{
+	return bModuleEnabled && !HasAnyFlags(RF_ClassDefaultObject);
+}
+
 #if WITH_EDITOR
 bool UNiagaraStatelessModule::CanEditChange(const FProperty* InProperty) const
 {
