- Add a mesh LOD option to use the component origin as the source of the LOD calculation
- This is more stable if you have dynamic bounds and allows for a single mesh to be rendered vs per particle LODs

#tests

diff --git a/Engine/Plugins/FX/Niagara/Source/Niagara/Private/NiagaraRendererMeshes.cpp b/Engine/Plugins/FX/Niagara/Source/Niagara/Private/NiagaraRendererMeshes.cpp
--- a/Engine/Plugins/FX/Niagara/Source/Niagara/Private/NiagaraRendererMeshes.cpp
+++ b/Engine/Plugins/FX/Niagara/Source/Niagara/Private/NiagaraRendererMeshes.cpp
@@ -265,7 +265,7 @@ void FNiagaraRendererMeshes::Initialize(const UNiagaraRendererProperties* InProp
 					{
 						RenderableMesh->SetMinLODBias(MeshProperties.LODBias);
 					}
-					else if (MeshProperties.LODMode == ENiagaraMeshLODMode::ByComponentBounds || MeshProperties.LODMode == ENiagaraMeshLODMode::PerParticle)
+					else if (MeshProperties.LODMode == ENiagaraMeshLODMode::ByComponentBounds || MeshProperties.LODMode == ENiagaraMeshLODMode::ComponentOrigin || MeshProperties.LODMode == ENiagaraMeshLODMode::PerParticle)
 					{
 						MeshData.LODDistanceFactor = MeshProperties.LODDistanceFactor;
 						if (MeshProperties.LODMode == ENiagaraMeshLODMode::PerParticle)
@@ -1493,6 +1493,13 @@ void FNiagaraRendererMeshes::GetDynamicMeshElements(const TArray<const FSceneVie
 					const int32 LODLevel = MeshData.RenderableMesh->ComputeLOD(ProxyBounds.Origin, ProxyBounds.SphereRadius, *View, MeshData.LODDistanceFactor);
 					MeshData.RenderableMesh->GetLODModelData(LODModel, LODLevel);
 				}
+				else if (MeshData.LODMode == ENiagaraMeshLODMode::ComponentOrigin)
+				{
+					const FBox MeshLocalBounds = ScaleFBox(MeshData.RenderableMesh->GetLocalBounds(), MeshData.Scale, MeshBoundsScale);
+					const FBox MeshWorldBounds = MeshLocalBounds.TransformBy(SceneProxy->GetLocalToWorld());
+					const int32 LODLevel = MeshData.RenderableMesh->ComputeLOD(MeshWorldBounds.GetCenter(), float(MeshWorldBounds.GetExtent().Size()), *View, MeshData.LODDistanceFactor);
+					MeshData.RenderableMesh->GetLODModelData(LODModel, LODLevel);
+				}
 				else
 				{
 					MeshData.RenderableMesh->GetLODModelData(LODModel, MeshData.LODLevel);
diff --git a/Engine/Plugins/FX/Niagara/Source/Niagara/Public/NiagaraMeshRendererMeshProperties.h b/Engine/Plugins/FX/Niagara/Source/Niagara/Public/NiagaraMeshRendererMeshProperties.h
--- a/Engine/Plugins/FX/Niagara/Source/Niagara/Public/NiagaraMeshRendererMeshProperties.h
+++ b/Engine/Plugins/FX/Niagara/Source/Niagara/Public/NiagaraMeshRendererMeshProperties.h
@@ -29,31 +29,38 @@ UENUM()
 enum class ENiagaraMeshLODMode : uint8
 {
 	/*
-	* Uses the provided LOD level to render all mesh particles.
-	* If the LOD is not streamed in or available on the platform the next available lower LOD level will be used.
-	* For example, LOD Level is set to 1 but the first available is LOD 3 then LOD 3 will be used.
+	Uses the provided LOD level to render all mesh particles.
+	If the LOD is not streamed in or available on the platform the next available lower LOD level will be used.
+	For example, LOD Level is set to 1 but the first available is LOD 3 then LOD 3 will be used.
 	*/
 	LODLevel,
 
 	/**
-	* Takes the highest available LOD for the platform + LOD bias to render all mesh particles
-	* If the LOD is not streamed in or available on the platform the next available lower LOD level will be used.
-	* For example, LOD bias is set to 1, the current platform has Min LOD of 2 then 3 will be the used LOD.
+	Takes the highest available LOD for the platform + LOD bias to render all mesh particles
+	If the LOD is not streamed in or available on the platform the next available lower LOD level will be used.
+	For example, LOD bias is set to 1, the current platform has Min LOD of 2 then 3 will be the used LOD.
 	*/
 	LODBias,
 
 	/*
-	* The LOD level is calculated based on screen space size of the component bounds.
-	* All particles will be rendered with the same calculated LOD level.
-	* Increasing 'LOD calculation scale' will result in lower quality LODs being used, this is useful as component bounds generally are larger than the particle mesh bounds.
+	The LOD level is calculated based on screen space size of the component bounds.
+	All particles will be rendered with the same calculated LOD level.
+	Increasing 'LOD calculation scale' will result in lower quality LODs being used, this is useful as component bounds generally are larger than the particle mesh bounds.
 	*/
 	ByComponentBounds,
 
 	/*
-	* The LOD level is calcuated per particle using the particle position and mesh sphere bounds.
-	* This involves running a dispatch & draw per LOD level.
-	* Calculates and renders each particle with it's calcualted LOD level.
-	* Increasing 'LOD calculation scale' will result in lower quality LODs being used.
+	The LOD level will be calculated like we have a single particle at the component origin, i.e. it should match a static mesh with the exact same transform.
+	All particles will be rendered with the same calculated LOD level.
+	Increasing 'LOD calculation scale' will result in lower quality LODs being used
+	*/
+	ComponentOrigin,
+
+	/*
+	The LOD level is calcuated per particle using the particle position and mesh sphere bounds.
+	This involves running a dispatch & draw per LOD level.
+	Calculates and renders each particle with it's calcualted LOD level.
+	Increasing 'LOD calculation scale' will result in lower quality LODs being used.
 	*/
 	PerParticle,
 };
@@ -135,7 +142,7 @@ struct FNiagaraMeshRendererMeshProperties : public FNiagaraMeshRendererMeshPrope
 	int32 LODBias = 0;
 
 	/** Used in LOD calculation to modify the distance, i.e. increasing the value will make lower poly LODs transition closer to the camera. */
-	UPROPERTY(EditAnywhere, Category = "Mesh", meta = (UIMin = "0", DisplayName = "LOD Distance Factor", EditCondition = "LODMode == ENiagaraMeshLODMode::ByComponentBounds || LODMode == ENiagaraMeshLODMode::PerParticle", EditConditionHides))
+	UPROPERTY(EditAnywhere, Category = "Mesh", meta = (UIMin = "0", DisplayName = "LOD Distance Factor", EditCondition = "LODMode == ENiagaraMeshLODMode::ByComponentBounds || LODMode == ENiagaraMeshLODMode::ComponentOrigin || LODMode == ENiagaraMeshLODMode::PerParticle", EditConditionHides))
 	float LODDistanceFactor = 1.0f;
 
 	/**
