Force inlining of the AppendToIndexBufferCPU function to work around a platform specific crash related to RIP not saving the last 8bits when calling RET. The function is only called in one location with 2 different template version (uint32 and uint16). No impact where noticed on program size for all platforms from a replayrun test.

#tests 679a7f9e3654027ab7ac1964

diff --git a/Engine/Plugins/FX/Niagara/Source/Niagara/Private/NiagaraRendererRibbons.cpp b/Engine/Plugins/FX/Niagara/Source/Niagara/Private/NiagaraRendererRibbons.cpp
--- a/Engine/Plugins/FX/Niagara/Source/Niagara/Private/NiagaraRendererRibbons.cpp
+++ b/Engine/Plugins/FX/Niagara/Source/Niagara/Private/NiagaraRendererRibbons.cpp
@@ -2074,7 +2074,7 @@ void FNiagaraRendererRibbons::GenerateIndexBufferCPU(FNiagaraIndexGenerationInpu
 }
 
 template <typename TValue>
-TValue* FNiagaraRendererRibbons::AppendToIndexBufferCPU(TValue* OutIndices, const FNiagaraIndexGenerationInput& GeneratedData, const FNiagaraRibbonShapeGeometryData& ShapeState, const TArrayView<uint32>& SegmentData, bool bInvertOrder)
+FORCEINLINE TValue* FNiagaraRendererRibbons::AppendToIndexBufferCPU(TValue* OutIndices, const FNiagaraIndexGenerationInput& GeneratedData, const FNiagaraRibbonShapeGeometryData& ShapeState, const TArrayView<uint32>& SegmentData, bool bInvertOrder)
 {
 	if (SegmentData.Num() == 0)
 	{
