Fix for system static variables getting caught in an ensure when trying to parse rapid iteration parameter names.  More of a workaround till we can clean up the static variable handling
#tests use a 'Fill Position Array' module in a system and verify that there's no huge stall from multiple ensures

diff --git a/Engine/Plugins/FX/Niagara/Source/Niagara/Private/NiagaraCommon.cpp b/Engine/Plugins/FX/Niagara/Source/Niagara/Private/NiagaraCommon.cpp
--- a/Engine/Plugins/FX/Niagara/Source/Niagara/Private/NiagaraCommon.cpp
+++ b/Engine/Plugins/FX/Niagara/Source/Niagara/Private/NiagaraCommon.cpp
@@ -1362,6 +1362,12 @@ FNiagaraAliasContext& FNiagaraAliasContext::ChangeStackContext(const FString& In
 	return *this;
 }
 
+FNiagaraAliasContext& FNiagaraAliasContext::ChangeRapidIterationParameterMode(ERapidIterationParameterMode InMode)
+{
+	RapidIterationParameterMode = InMode;
+	return *this;
+}
+
 void AliasRapidIterationConstant(const FNiagaraAliasContext& InContext, TArray<FStringView, TInlineAllocator<16>>& InOutSplitName, int32& OutAssignmentNamespaceIndex)
 {
 	if(ensureMsgf(InContext.GetRapidIterationParameterMode() != FNiagaraAliasContext::ERapidIterationParameterMode::None, TEXT("Can not resolve a rapid iteration variable without specifying the mode in the context.")))
@@ -1371,7 +1377,16 @@ void AliasRapidIterationConstant(const FNiagaraAliasContext& InContext, TArray<F
 		int32 MinParts;
 		int32 EmitterNameIndex;
 		int32 ModuleNameIndex;
-		if (InContext.GetRapidIterationParameterMode() == FNiagaraAliasContext::ERapidIterationParameterMode::SystemScript)
+
+		const FNiagaraAliasContext::ERapidIterationParameterMode Mode = InContext.GetRapidIterationParameterMode();
+
+		// static variables break the regular conventions here because they can include both system scoped variables as well as particle/emitter scoped
+		// variables, so we try to account for this here.  This really needs to be sorted when we can clean up how we manage our static variables and
+		// have them no longer being considered rapid iteration parameters.
+		const bool bTreatAsSystemVariable = (Mode == FNiagaraAliasContext::ERapidIterationParameterMode::SystemScript)
+			|| (Mode == FNiagaraAliasContext::ERapidIterationParameterMode::StaticVariables && InOutSplitName.Num() == 3);
+
+		if (bTreatAsSystemVariable)
 		{
 			MinParts = 3;
 			EmitterNameIndex = INDEX_NONE;
diff --git a/Engine/Plugins/FX/Niagara/Source/Niagara/Public/NiagaraCommon.h b/Engine/Plugins/FX/Niagara/Source/Niagara/Public/NiagaraCommon.h
--- a/Engine/Plugins/FX/Niagara/Source/Niagara/Public/NiagaraCommon.h
+++ b/Engine/Plugins/FX/Niagara/Source/Niagara/Public/NiagaraCommon.h
@@ -1446,7 +1446,9 @@ struct FNiagaraAliasContext
 		/** Rapid iteration parameters will be from an emitter or particle script. */
 		EmitterOrParticleScript,
 		/** Rapid iteration parameters will not be handled. NOTE: Using this mode will cause an assert if the context is used with a rapid iteration parameter constant name. */
-		None
+		None,
+		/** Special case of evaluating static variables which can include variables from both system scripts and particle/emitter scripts */
+		StaticVariables,
 	};
 
 	FNiagaraAliasContext(ERapidIterationParameterMode InRapidIterationParameterMode = ERapidIterationParameterMode::None)
@@ -1482,6 +1484,7 @@ struct FNiagaraAliasContext
 	NIAGARA_API FNiagaraAliasContext& ChangeStackContext(const FString& InStackContextName);
 
 	ERapidIterationParameterMode GetRapidIterationParameterMode() const { return RapidIterationParameterMode; }
+	NIAGARA_API FNiagaraAliasContext& ChangeRapidIterationParameterMode(ERapidIterationParameterMode InMode);
 
 	const TOptional<FString>& GetEmitterName() const { return EmitterName; }
 	const TOptional<TPair<FString, FString>>& GetEmitterMapping() const { return EmitterMapping; }
diff --git a/Engine/Plugins/FX/Niagara/Source/NiagaraEditor/Private/NiagaraEditorUtilities.cpp b/Engine/Plugins/FX/Niagara/Source/NiagaraEditor/Private/NiagaraEditorUtilities.cpp
--- a/Engine/Plugins/FX/Niagara/Source/NiagaraEditor/Private/NiagaraEditorUtilities.cpp
+++ b/Engine/Plugins/FX/Niagara/Source/NiagaraEditor/Private/NiagaraEditorUtilities.cpp
@@ -5128,7 +5128,7 @@ FText FNiagaraEditorUtilities::AssetBrowser::GetAssetTagSectionNameFromSource(EA
 
 void FNiagaraParameterUtilities::FilterToRelevantStaticVariables(TConstArrayView<FNiagaraVariable> InVars, TArray<FNiagaraVariable>& OutVars, FName InOldEmitterAlias, FName InNewEmitterAlias, bool bFilterByEmitterAliasAndConvertToUnaliased)
 {
-	FNiagaraAliasContext RenameContext(ENiagaraScriptUsage::ParticleSpawnScript);
+	FNiagaraAliasContext RenameContext(FNiagaraAliasContext::ERapidIterationParameterMode::StaticVariables);
 	RenameContext.ChangeEmitterName(InOldEmitterAlias.ToString(), InNewEmitterAlias.ToString());
 
 	for (const FNiagaraVariable& InVar : InVars)
diff --git a/Engine/Plugins/FX/Niagara/Source/NiagaraEditor/Private/NiagaraGraphDigest.cpp b/Engine/Plugins/FX/Niagara/Source/NiagaraEditor/Private/NiagaraGraphDigest.cpp
--- a/Engine/Plugins/FX/Niagara/Source/NiagaraEditor/Private/NiagaraGraphDigest.cpp
+++ b/Engine/Plugins/FX/Niagara/Source/NiagaraEditor/Private/NiagaraGraphDigest.cpp
@@ -2243,6 +2243,8 @@ void FNiagaraCompilationNodeEmitter::BuildParameterMapHistory(FParameterMapHisto
 				Builder.Histories[ParamMapIdx].PinToConstantValues.Append(History.PinToConstantValues);
 			}
 
+			ResolveAliasesContext.ChangeRapidIterationParameterMode(FNiagaraAliasContext::ERapidIterationParameterMode::StaticVariables);
+
 			// We only want to push out appropriately scoped static variables that should be in the system builder, not in-betweens like "Module.MyInputVar"
 			// or per-particle or others vars. Really only want Emitter or System parameters here.
 			for (int32 StaticVarIdx = 0; StaticVarIdx < ChildBuilder.StaticVariables.Num(); StaticVarIdx++)
diff --git a/Engine/Plugins/FX/Niagara/Source/NiagaraEditor/Private/NiagaraNodeEmitter.cpp b/Engine/Plugins/FX/Niagara/Source/NiagaraEditor/Private/NiagaraNodeEmitter.cpp
--- a/Engine/Plugins/FX/Niagara/Source/NiagaraEditor/Private/NiagaraNodeEmitter.cpp
+++ b/Engine/Plugins/FX/Niagara/Source/NiagaraEditor/Private/NiagaraNodeEmitter.cpp
@@ -427,6 +427,7 @@ void UNiagaraNodeEmitter::BuildParameterMapHistory(FNiagaraParameterMapHistoryBu
 				OutHistory.Histories[ParamMapIdx].PinToConstantValues.Append(History.PinToConstantValues);
 			}
 
+			ResolveAliasesContext.ChangeRapidIterationParameterMode(FNiagaraAliasContext::ERapidIterationParameterMode::StaticVariables);
 			// We only want to push out appropriately scoped static variables that should be in the system builder, not in-betweens like "Module.MyInputVar"
 			// or per-particle or others vars. Really only want Emitter or System parameters here.
 			for (int32 StaticVarIdx = 0; StaticVarIdx < ChildBuilder.StaticVariables.Num(); StaticVarIdx++)
diff --git a/Engine/Plugins/FX/Niagara/Source/NiagaraEditor/Private/ViewModels/Stack/NiagaraStackGraphUtilities.cpp b/Engine/Plugins/FX/Niagara/Source/NiagaraEditor/Private/ViewModels/Stack/NiagaraStackGraphUtilities.cpp
--- a/Engine/Plugins/FX/Niagara/Source/NiagaraEditor/Private/ViewModels/Stack/NiagaraStackGraphUtilities.cpp
+++ b/Engine/Plugins/FX/Niagara/Source/NiagaraEditor/Private/ViewModels/Stack/NiagaraStackGraphUtilities.cpp
@@ -1669,17 +1669,15 @@ void FNiagaraStackGraphUtilities::GetStackFunctionStaticSwitchPinsLegacy(const U
 	TArray<FNiagaraVariable> StaticVars;
 	TSharedPtr<FNiagaraGraphCachedDataBase, ESPMode::ThreadSafe> CachedTraversalData;
 	FString EmitterName;
-	FNiagaraAliasContext::ERapidIterationParameterMode Mode = FNiagaraAliasContext::ERapidIterationParameterMode::None;
+	FNiagaraAliasContext::ERapidIterationParameterMode Mode = FNiagaraAliasContext::ERapidIterationParameterMode::StaticVariables;
 	if (OuterEmitter.GetEmitterData())
 	{
 		CachedTraversalData = OuterEmitter.Emitter->GetCachedTraversalData(OuterEmitter.Version);
 		EmitterName = OuterEmitter.Emitter->GetUniqueEmitterName();
-		Mode = FNiagaraAliasContext::ERapidIterationParameterMode::EmitterOrParticleScript;
 	}
 	else if (System)
 	{
 		CachedTraversalData = System->GetCachedTraversalData();
-		Mode = FNiagaraAliasContext::ERapidIterationParameterMode::SystemScript;
 	}
 	if (CachedTraversalData.IsValid())
 	{
