- Fix for invalid array data size being pushed into the cache storage
- Fix for potentially reallocating the array on the RT, there's a small window but it could happen

diff --git a/Engine/Plugins/FX/Niagara/Source/Niagara/Classes/NiagaraDataInterfaceArray.h b/Engine/Plugins/FX/Niagara/Source/Niagara/Classes/NiagaraDataInterfaceArray.h
--- a/Engine/Plugins/FX/Niagara/Source/Niagara/Classes/NiagaraDataInterfaceArray.h
+++ b/Engine/Plugins/FX/Niagara/Source/Niagara/Classes/NiagaraDataInterfaceArray.h
@@ -176,4 +176,7 @@ public:
 
 	// Internal use only
 	NIAGARA_API int32 FindOrAddData(TConstArrayView<uint8> ArrayData);
+
+protected:
+	UE::FMutex FindOrAddGuard;
 };
diff --git a/Engine/Plugins/FX/Niagara/Source/Niagara/Private/NiagaraDataInterfaceArray.cpp b/Engine/Plugins/FX/Niagara/Source/Niagara/Private/NiagaraDataInterfaceArray.cpp
--- a/Engine/Plugins/FX/Niagara/Source/Niagara/Private/NiagaraDataInterfaceArray.cpp
+++ b/Engine/Plugins/FX/Niagara/Source/Niagara/Private/NiagaraDataInterfaceArray.cpp
@@ -172,6 +172,8 @@ FString UNiagaraDataInterfaceArray::SimCacheVisualizerRead(const UNDIArraySimCac
 
 int32 UNDIArraySimCacheData::FindOrAddData(TConstArrayView<uint8> ArrayData)
 {
+	UE::TScopeLock ScopeLock(FindOrAddGuard);
+
 	const int32 Last = BufferData.Num() - ArrayData.Num();
 	for (int32 i = 0; i <= Last; ++i)
 	{
diff --git a/Engine/Plugins/FX/Niagara/Source/Niagara/Private/NiagaraDataInterfaceArrayImpl.cpp b/Engine/Plugins/FX/Niagara/Source/Niagara/Private/NiagaraDataInterfaceArrayImpl.cpp
--- a/Engine/Plugins/FX/Niagara/Source/Niagara/Private/NiagaraDataInterfaceArrayImpl.cpp
+++ b/Engine/Plugins/FX/Niagara/Source/Niagara/Private/NiagaraDataInterfaceArrayImpl.cpp
@@ -367,7 +367,7 @@ void FNDIArrayInstanceData_RenderThreadBase::SimCacheWriteFrame(FRHICommandListI
 				FrameData.DataOffset = CacheData->FindOrAddData(
 					MakeArrayView(
 						ArrayData.GetData(),
-						ArrayData.Num() * ArrayTypeSize
+						ArrayData.Num()
 					)
 				);
 			}
