- Add child components to the baker
- This is to allow for things like the component renderer to be visibile and not having to toggle Niagara only off

#tests Use the baker with the component renderer, adding something like a skeletal mesh component

diff --git a/Engine/Plugins/FX/Niagara/Source/NiagaraEditor/Private/NiagaraBakerRenderer.cpp b/Engine/Plugins/FX/Niagara/Source/NiagaraEditor/Private/NiagaraBakerRenderer.cpp
--- a/Engine/Plugins/FX/Niagara/Source/NiagaraEditor/Private/NiagaraBakerRenderer.cpp
+++ b/Engine/Plugins/FX/Niagara/Source/NiagaraEditor/Private/NiagaraBakerRenderer.cpp
@@ -428,9 +428,18 @@ void FNiagaraBakerRenderer::RenderSceneCapture(UTextureRenderTarget2D* RenderTar
 
 	if (BakerSettings->bRenderComponentOnly)
 	{
+		const TArray<TObjectPtr<USceneComponent>>& AttachChildren = BakedDataComponent->GetAttachChildren();
+
 		SceneCaptureComponent->PrimitiveRenderMode = ESceneCapturePrimitiveRenderMode::PRM_UseShowOnlyList;
-		SceneCaptureComponent->ShowOnlyComponents.Empty(1);
+		SceneCaptureComponent->ShowOnlyComponents.Empty(1 + AttachChildren.Num());
 		SceneCaptureComponent->ShowOnlyComponents.Add(BakedDataComponent);
+		for (TWeakObjectPtr<USceneComponent> WeakChildComponent : AttachChildren)
+		{
+			if (UPrimitiveComponent* ChildComponent = Cast<UPrimitiveComponent>(WeakChildComponent.Get()))
+			{
+				SceneCaptureComponent->ShowOnlyComponents.Add(ChildComponent);
+			}
+		}
 	}
 	else
 	{
