Change the name of the generated function for CustomHlsl nodes to incorporate more than just the NodeGuid as that can be non-unique depending on graph oddities (duplicating a graph, ...).
#tests Added two different versions of the same module (World Position to Screen UV) and watch it now compile

diff --git a/Engine/Plugins/FX/Niagara/Source/NiagaraEditor/Private/NiagaraHlslTranslator.cpp b/Engine/Plugins/FX/Niagara/Source/NiagaraEditor/Private/NiagaraHlslTranslator.cpp
--- a/Engine/Plugins/FX/Niagara/Source/NiagaraEditor/Private/NiagaraHlslTranslator.cpp
+++ b/Engine/Plugins/FX/Niagara/Source/NiagaraEditor/Private/NiagaraHlslTranslator.cpp
@@ -8873,12 +8873,35 @@ void TNiagaraHlslTranslator<GraphBridge>::HandleCustomHlslNode(const FCustomHlsl
 
 	// Determine the important outputs
 	OutScriptUsage = GraphBridge::GetCustomHlslUsage(CustomFunctionHlsl);
-	OutName = BuildSanitizedSymbolString(CustomFunctionHlsl->Signature.Name.ToString() + CustomFunctionHlsl->NodeGuid.ToString());
+
+	FString CustomHlslSourceString = GraphBridge::GetCustomHlslString(CustomFunctionHlsl);
+
+	// because the NodeGuid isn't necessarily sufficiently unique we're going to use a hash of the NodeGuid, the source and the inputs & outputs of the node to
+	// be used as a unique identifier for the name
+	FString UniqueIdentifier;
+
+	{
+		FSHA1 HashState;
+		HashState.Update((const uint8*)&CustomFunctionHlsl->NodeGuid, sizeof(CustomFunctionHlsl->NodeGuid));
+		HashState.UpdateWithString(*CustomHlslSourceString, CustomHlslSourceString.Len());
+
+		TArray<const FPin*> Pins = GraphBridge::GetPins(CustomFunctionHlsl);
+		for (const GraphBridge::FPin* Pin : Pins)
+		{
+			FNiagaraVariable PinVariable = GraphBridge::GetPinVariable(Pin, false /*bNeedsValue*/, ENiagaraStructConversion::UserFacing);
+			const uint32 PinVariableTypeHash = GetTypeHash(PinVariable);
+			HashState.Update((const uint8*)&PinVariableTypeHash, sizeof(PinVariableTypeHash));
+		}
+
+		UniqueIdentifier = HashState.Finalize().ToString();
+	}
+
+	OutName = BuildSanitizedSymbolString(CustomFunctionHlsl->Signature.Name.ToString() + UniqueIdentifier);
 	OutSignature = CustomFunctionHlsl->Signature;
 	OutFullName = GraphBridge::GetFunctionFullName(CustomFunctionHlsl);
 	OutSignature.Name = *OutName; // Force the name to be set to include the node guid for safety...
 	bOutCustomHlsl = true;
-	OutCustomHlsl = GraphBridge::GetCustomHlslString(CustomFunctionHlsl);
+	OutCustomHlsl = CustomHlslSourceString;
 	OutCustomHlsl.ReplaceInline(TEXT("\r\n"), TEXT("\n"), ESearchCase::CaseSensitive);
 	GraphBridge::GetCustomHlslIncludePaths(CustomFunctionHlsl, OutCustomHlslIncludeFilePaths);
 
diff --git a/Engine/Source/Runtime/Core/Private/UObject/FortniteMainBranchObjectVersion.cpp b/Engine/Source/Runtime/Core/Private/UObject/FortniteMainBranchObjectVersion.cpp
--- a/Engine/Source/Runtime/Core/Private/UObject/FortniteMainBranchObjectVersion.cpp
+++ b/Engine/Source/Runtime/Core/Private/UObject/FortniteMainBranchObjectVersion.cpp
@@ -10,8 +10,8 @@ TMap<FGuid, FGuid> FFortniteMainBranchObjectVersion::GetSystemGuids()
 	SystemGuids.Add(DevGuids.LANDSCAPE_MOBILE_COOK_VERSION, FGuid("32D02EF867C74B71A0D4E0FA41392732"));
 	SystemGuids.Add(DevGuids.MATERIALSHADERMAP_DERIVEDDATA_VER, FGuid("0E6202C151DE402CA451802671AE9347"));
 	SystemGuids.Add(DevGuids.NANITE_DERIVEDDATA_VER, FGuid("42C46B24E1586986880AFE62B0F575AB"));
-	SystemGuids.Add(DevGuids.NIAGARASHADERMAP_DERIVEDDATA_VER, FGuid("016F98966CE4455EBA36999548E22A39"));
-	SystemGuids.Add(DevGuids.Niagara_LatestScriptCompileVersion, FGuid("393081597FA7234C92F00869DE262D73"));
+	SystemGuids.Add(DevGuids.NIAGARASHADERMAP_DERIVEDDATA_VER, FGuid("3AE847DB7C504839B140718982B04087"));
+	SystemGuids.Add(DevGuids.Niagara_LatestScriptCompileVersion, FGuid("670A2EBDEE1A4794BB77BBFF331EA7EF"));
 	SystemGuids.Add(DevGuids.SkeletalMeshDerivedDataVersion, FGuid("2F90180769024483BB7D90AACEB2798D"));
 	SystemGuids.Add(DevGuids.STATICMESH_DERIVEDDATA_VER, FGuid("656E338C78C34641A5D3BFAE303997AF"));
 	SystemGuids.Add(DevGuids.MaterialTranslationDDCVersion, FGuid("9119259dbf0449a3a53ad39df9fd13a4"));
