Fix up compile error with translator on some platforms
#tests compiled on pc

diff --git a/Engine/Plugins/FX/Niagara/Source/NiagaraEditor/Private/NiagaraHlslTranslator.cpp b/Engine/Plugins/FX/Niagara/Source/NiagaraEditor/Private/NiagaraHlslTranslator.cpp
--- a/Engine/Plugins/FX/Niagara/Source/NiagaraEditor/Private/NiagaraHlslTranslator.cpp
+++ b/Engine/Plugins/FX/Niagara/Source/NiagaraEditor/Private/NiagaraHlslTranslator.cpp
@@ -8886,7 +8886,7 @@ void TNiagaraHlslTranslator<GraphBridge>::HandleCustomHlslNode(const FCustomHlsl
 		HashState.UpdateWithString(*CustomHlslSourceString, CustomHlslSourceString.Len());
 
 		TArray<const FPin*> Pins = GraphBridge::GetPins(CustomFunctionHlsl);
-		for (const GraphBridge::FPin* Pin : Pins)
+		for (const FPin* Pin : Pins)
 		{
 			FNiagaraVariable PinVariable = GraphBridge::GetPinVariable(Pin, false /*bNeedsValue*/, ENiagaraStructConversion::UserFacing);
 			const uint32 PinVariableTypeHash = GetTypeHash(PinVariable);
