- When rebasing is enabled push the sim cache data into the system instance if it has moved, even if the time has not advanced
- This fixes issues with local space emitters when time does not advance
- It does mean world space emitters diverge, but they already work like local space when a cache is active (as in they will move if the component moves).

diff --git a/Engine/Plugins/FX/Niagara/Source/Niagara/Classes/NiagaraSimCache.h b/Engine/Plugins/FX/Niagara/Source/Niagara/Classes/NiagaraSimCache.h
--- a/Engine/Plugins/FX/Niagara/Source/Niagara/Classes/NiagaraSimCache.h
+++ b/Engine/Plugins/FX/Niagara/Source/Niagara/Classes/NiagaraSimCache.h
@@ -433,6 +433,9 @@ public:
 	UFUNCTION(BlueprintCallable, Category=NiagaraSimCache)
 	ENiagaraSimCacheAttributeCaptureMode GetAttributeCaptureMode() const { return CreateParameters.AttributeCaptureMode; }
 
+	/** Returns true if we have any rebasing enabled. */
+	bool IsRebaseEnabled() const { return CreateParameters.bAllowRebasing; }
+
 	NIAGARA_API bool BeginWrite(FNiagaraSimCacheCreateParameters InCreateParameters, UNiagaraComponent* NiagaraComponent);
 	NIAGARA_API bool BeginWrite(FNiagaraSimCacheCreateParameters InCreateParameters, UNiagaraComponent* NiagaraComponent, FNiagaraSimCacheFeedbackContext& FeedbackContext);
 
diff --git a/Engine/Plugins/FX/Niagara/Source/Niagara/Private/NiagaraSystemInstance.cpp b/Engine/Plugins/FX/Niagara/Source/Niagara/Private/NiagaraSystemInstance.cpp
--- a/Engine/Plugins/FX/Niagara/Source/Niagara/Private/NiagaraSystemInstance.cpp
+++ b/Engine/Plugins/FX/Niagara/Source/Niagara/Private/NiagaraSystemInstance.cpp
@@ -2297,17 +2297,36 @@ void FNiagaraSystemInstance::SimCacheTick_GameThread(UNiagaraSimCache* SimCache,
 		return;
 	}
 
-	// Early out if our age and desired age are virtually equal as there's nothing to do
-	if (FMath::IsNearlyEqual(Age, DesiredAge) )
+	TOptional<FTransform> OwnerTransform;
+	if (USceneComponent* TransformComponent = AttachComponent.Get())
 	{
-		return;
+		OwnerTransform = AttachComponent->GetComponentToWorld();
+		OwnerTransform->AddToTranslation(FVector(LWCTile) * -FLargeWorldRenderScalar::GetTileSize());
+	}
+
+	// Early out if our age and desired age are virtually equal
+	if (FMath::IsNearlyEqual(Age, DesiredAge))
+	{
+		// If rebasing is enabled and the transform has moved then force an update otherwise local space emitters won't move which differs from uncached.
+		// World space emitters will also move on transform change, which doesn't match uncached but is matches how it would change on next tick
+		// Local space emitter do not need to update since they will automatically move.
+		// Note: We could expose this as an option when using a sim cache if we desire more control
+		if (!SimCache->IsRebaseEnabled() || !OwnerTransform.IsSet() || WorldTransform.Equals(*OwnerTransform))
+		{
+			return;
+		}
 	}
 
 	CachedDeltaSeconds = DeltaSeconds;
 	//FixedBounds_CNC = FixedBounds_GT;
 
-	// We still need to tick the override parameters as these might pass through to material bindings
+	if (OwnerTransform.IsSet())
+	{
+		WorldTransform = OwnerTransform.GetValue();
+	}
 	//TickInstanceParameters_GameThread(DeltaSeconds);
+
+	// We still need to tick the override parameters as these might pass through to material bindings
 	if (OverrideParameters)
 	{
 		OverrideParameters->ResolvePositions(GetLWCConverter());
