- Remove scaling from SimCache when rebasing and fix ordering of operations
- This does not precisely mirror an uncached simulation as the way scale is applied in the stack can be arbitary, so it keeps the locations looking the same which is more accurate vs moving them

diff --git a/Engine/Plugins/FX/Niagara/Source/Niagara/Private/NiagaraSimCache.cpp b/Engine/Plugins/FX/Niagara/Source/Niagara/Private/NiagaraSimCache.cpp
--- a/Engine/Plugins/FX/Niagara/Source/Niagara/Private/NiagaraSimCache.cpp
+++ b/Engine/Plugins/FX/Niagara/Source/Niagara/Private/NiagaraSimCache.cpp
@@ -1034,10 +1034,19 @@ bool UNiagaraSimCache::ReadFrame(int32 FrameIndex, float FrameFraction, FNiagara
 	FTransform RebaseBTransform = FTransform::Identity;
 	if ( USceneComponent* AttachComponent = SystemInstance->GetAttachComponent() )
 	{
+		// Note: We remove the scaling as how that is applied to uncached emitters can not be represented by applying a transform.
+		//       If we desire to control this an option can be exposed on the cache.
 		FTransform LocalToWorld = AttachComponent->GetComponentToWorld();
 		LocalToWorld.AddToTranslation(FVector(SystemInstance->GetLWCTile()) * -FLargeWorldRenderScalar::GetTileSize());
-		RebaseATransform = LocalToWorld * CacheFrame.LocalToWorld.Inverse();
-		RebaseBTransform = LocalToWorld * CacheFrames[NextFrameIndex].LocalToWorld.Inverse();
+		LocalToWorld.RemoveScaling();
+
+		RebaseATransform = CacheFrame.LocalToWorld;
+		RebaseATransform.RemoveScaling();
+		RebaseBTransform = CacheFrames[NextFrameIndex].LocalToWorld;
+		RebaseBTransform.RemoveScaling();
+
+		RebaseATransform = RebaseATransform.Inverse() * LocalToWorld;
+		RebaseBTransform = RebaseBTransform.Inverse() * LocalToWorld;
 	}
 
 	const float FrameDeltaSeconds = CacheFrames[NextFrameIndex].SimulationAge - CacheFrame.SimulationAge;
