Substrate - fix niagara GPU resource assert on mobile using legacy GBuffer (not Substrate GBuffer)

#tests PC editor dx12 and mobile preview
#rnx

diff --git a/Engine/Plugins/FX/Niagara/Shaders/Private/NiagaraDataInterfaceGBufferTemplate.ush b/Engine/Plugins/FX/Niagara/Shaders/Private/NiagaraDataInterfaceGBufferTemplate.ush
--- a/Engine/Plugins/FX/Niagara/Shaders/Private/NiagaraDataInterfaceGBufferTemplate.ush
+++ b/Engine/Plugins/FX/Niagara/Shaders/Private/NiagaraDataInterfaceGBufferTemplate.ush
@@ -2,21 +2,25 @@
 
 #include "/Engine/Private/DeferredShadingCommon.ush"
 
+#define NDIGBUFFER_HAS_GBUFFER				(FEATURE_LEVEL >= FEATURE_LEVEL_SM5)
+#define NDIGBUFFER_USE_GBUFFER_VELOCITY		(false && NDIGBUFFER_HAS_GBUFFER)
+#define NDIGBUFFER_IS_MOBILE				(FEATURE_LEVEL <= FEATURE_LEVEL_ES3_1)
+
+#define NIAGARA_SAMPLE_SUBSTRATE (SUBSTRATE_ENABLED && !NDIGBUFFER_IS_MOBILE)
+
+#if NIAGARA_SAMPLE_SUBSTRATE
 #define NEEDS_SCENE_TEXTURES 1
 #define SCENE_TEXTURES_DISABLED 0
 #define SubstratePublicStruct SubstratePublic
 #include "/Engine/Public/SubstratePublic.ush"
-
-#define NDIGBUFFER_HAS_GBUFFER				(FEATURE_LEVEL >= FEATURE_LEVEL_SM5)
-#define NDIGBUFFER_USE_GBUFFER_VELOCITY		(false && NDIGBUFFER_HAS_GBUFFER)
-#define NDIGBUFFER_IS_MOBILE				(FEATURE_LEVEL <= FEATURE_LEVEL_ES3_1)
+#endif
 
 ////////////////////////////////////////////////////////////////////////////////////////////////////////////////////////////////////////////////////////////////
 
 Texture2D {ParameterName}_VelocityTexture;
 SamplerState {ParameterName}_VelocityTextureSampler;
 
-#if SUBSTRATE_ENABLED
+#if NIAGARA_SAMPLE_SUBSTRATE
 FSubstrateGBuffer GetNiagaraSubstrateGBuffer_{ParameterName}(in float2 ScreenUV)
 {
 	uint2 ScreenPos = uint2(ScreenUV * View.BufferSizeAndInvSize.xy);
@@ -39,7 +43,7 @@ void DecodeDiffuseColor_{ParameterName}(in float2 ScreenUV, in bool bApplyViewpo
 	ScreenUV = ApplyViewportOffset_{ParameterName}(ScreenUV, bApplyViewportOffset);
 	IsValid = true;
 
-	#if SUBSTRATE_ENABLED
+	#if NIAGARA_SAMPLE_SUBSTRATE
 		uint2 ScreenPos = uint2(ScreenUV * View.BufferSizeAndInvSize.xy);
 		OutValue = SubstratePublic_GetWorldNormal(ScreenPos);
 	#else
@@ -58,7 +62,7 @@ void DecodeWorldNormal_{ParameterName}(in float2 ScreenUV, in bool bApplyViewpor
 	ScreenUV = ApplyViewportOffset_{ParameterName}(ScreenUV, bApplyViewportOffset);
 	IsValid = true;
 
-	#if SUBSTRATE_ENABLED
+	#if NIAGARA_SAMPLE_SUBSTRATE
 		uint2 ScreenPos = uint2(ScreenUV * View.BufferSizeAndInvSize.xy);
 		OutValue = SubstratePublic_GetWorldNormal(ScreenPos);
 	#else
@@ -140,7 +144,7 @@ void DecodeBaseColor_{ParameterName}(in float2 ScreenUV, in bool bApplyViewportO
 	ScreenUV = ApplyViewportOffset_{ParameterName}(ScreenUV, bApplyViewportOffset);
 	IsValid = true;
 
-	#if SUBSTRATE_ENABLED
+	#if NIAGARA_SAMPLE_SUBSTRATE
 		OutValue = GetNiagaraSubstrateGBuffer_{ParameterName}(ScreenUV).BaseColor;
 	#else
 		FGBufferData GBufferData = GetGBufferData(ScreenUV);
@@ -159,7 +163,7 @@ void DecodeMetallic_{ParameterName}(in float2 ScreenUV, in bool bApplyViewportOf
 	ScreenUV = ApplyViewportOffset_{ParameterName}(ScreenUV, bApplyViewportOffset);
 	IsValid = true;
 
-	#if SUBSTRATE_ENABLED
+	#if NIAGARA_SAMPLE_SUBSTRATE
 		OutValue = GetNiagaraSubstrateGBuffer_{ParameterName}(ScreenUV).Metallic;
 	#else
 		FGBufferData GBufferData = GetGBufferData(ScreenUV);
@@ -178,7 +182,7 @@ void DecodeSpecular_{ParameterName}(in float2 ScreenUV, in bool bApplyViewportOf
 	ScreenUV = ApplyViewportOffset_{ParameterName}(ScreenUV, bApplyViewportOffset);
 	IsValid = true;
 
-	#if SUBSTRATE_ENABLED
+	#if NIAGARA_SAMPLE_SUBSTRATE
 		OutValue = GetNiagaraSubstrateGBuffer_{ParameterName}(ScreenUV).Specular;
 	#else
 		FGBufferData GBufferData = GetGBufferData(ScreenUV);
@@ -197,7 +201,7 @@ void DecodeRoughness_{ParameterName}(in float2 ScreenUV, in bool bApplyViewportO
 	ScreenUV = ApplyViewportOffset_{ParameterName}(ScreenUV, bApplyViewportOffset);
 	IsValid = true;
 
-	#if SUBSTRATE_ENABLED
+	#if NIAGARA_SAMPLE_SUBSTRATE
 		uint2 ScreenPos = uint2(ScreenUV * View.BufferSizeAndInvSize.xy);
 		OutValue = SubstratePublic_GetRoughness(ScreenPos);
 	#else
@@ -316,7 +320,7 @@ void DecodeShadingModelID_{ParameterName}(in float2 ScreenUV, in bool bApplyView
 	ScreenUV = ApplyViewportOffset_{ParameterName}(ScreenUV, bApplyViewportOffset);
 	IsValid = true;
 
-	#if SUBSTRATE_ENABLED
+	#if NIAGARA_SAMPLE_SUBSTRATE
 		OutValue = int(GetNiagaraSubstrateGBuffer_{ParameterName}(ScreenUV).ShadingModelID);
 	#else
 		FGBufferData GBufferData = GetGBufferData(ScreenUV);
diff --git a/Engine/Shaders/Private/MaterialTemplate.ush b/Engine/Shaders/Private/MaterialTemplate.ush
--- a/Engine/Shaders/Private/MaterialTemplate.ush
+++ b/Engine/Shaders/Private/MaterialTemplate.ush
@@ -3236,9 +3236,10 @@ float4 SceneTextureLookup(float2 UV, int SceneTextureIndex, bool bFiltered)
 
 	// 2. Post-process inputs which are dependent on material data
 	//    Note: For PostProcess material we always return the first layer dataCalcSceneCustomStencil
-	#if SUBSTRATE_ENABLED && IS_DBUFFER_DECAL
+	#if SUBSTRATE_ENABLED && !SHADING_PATH_MOBILE && IS_DBUFFER_DECAL
 	{
 		// A special decal path to avoid requiring the binding of substrate data in the decal passes. Decals can only access depth, stencil and depth-from-normals or reprojected.
+		// But not on mobile that is not using dbuffer decal but the legacy GBuffer blending.
 		switch (SceneTextureIndex)
 		{
 		case PPI_DiffuseColor:			return 0.0f;
@@ -3260,7 +3261,7 @@ float4 SceneTextureLookup(float2 UV, int SceneTextureIndex, bool bFiltered)
 		case PPI_IsFirstPerson:			return 0.0f;
 		}
 	}
-	#elif SUBSTRATE_ENABLED
+	#elif SUBSTRATE_ENABLED && !SHADING_PATH_MOBILE // Not on mobile that is not using dbuffer decal but the legacy GBuffer blending.
 	{
 		FSubstrateGBuffer SubstrateGBuffer = SubstratePublic_GetSubstrateGBufferFromFirstBSDF(PixelPos);
 		BRANCH
