- Add previous / current FOV data to GPU view properties

diff --git a/Engine/Plugins/FX/Niagara/Shaders/Private/NiagaraDataInterfaceCamera.ush b/Engine/Plugins/FX/Niagara/Shaders/Private/NiagaraDataInterfaceCamera.ush
--- a/Engine/Plugins/FX/Niagara/Shaders/Private/NiagaraDataInterfaceCamera.ush
+++ b/Engine/Plugins/FX/Niagara/Shaders/Private/NiagaraDataInterfaceCamera.ush
@@ -30,7 +30,7 @@ void GetCameraPropertiesCPUASC47GPU_{ParameterName}(out float3 Out_CameraPositio
 
 void GetViewPropertiesGPU_{ParameterName}(out float3 Out_ViewPositionWorld, out float3 Out_ViewForwardVector, out float3 Out_ViewUpVector, out float3 Out_ViewRightVector, out float4 Out_ViewSizeAndInverseSize,
 	out float4 Out_ScreenToViewSpace, out float2 Out_Current_TAAJitter, out float2 Out_Previous_TAAJitter, out float3 Out_PreViewTranslation, out float4 Out_BufferSizeAndInverseSize, out float2 Out_ViewportOffset,
-	out float Out_NearPlane)
+	out float Out_NearPlane, out float2 Out_FOVCurrent, out float2 Out_FOVPrevious)
 {
 	FLWCVector3 ViewPos = PrimaryView.TileOffset.WorldViewOrigin;
 	LWCSetTile(ViewPos, LWCGetTile(ViewPos) - GetEngineOwnerLWCTile()); // convert to simulation space
@@ -46,6 +46,8 @@ void GetViewPropertiesGPU_{ParameterName}(out float3 Out_ViewPositionWorld, out
 	Out_BufferSizeAndInverseSize = View.BufferSizeAndInvSize;
 	Out_ViewportOffset = View.ViewRectMin.xy;
 	Out_NearPlane = View.NearPlane;
+	Out_FOVCurrent = View.FieldOfViewWideAngles.xy;
+	Out_FOVPrevious = View.PrevFieldOfViewWideAngles.xy;
 }
 
 void GetClipSpaceTransformsGPU_{ParameterName}(out float4x4 Out_WorldToClipTransform, out float4x4 Out_TranslatedWorldToClipTransform, out float4x4 Out_ClipToWorldTransform, out float4x4 Out_ClipToViewTransform,
diff --git a/Engine/Plugins/FX/Niagara/Source/Niagara/Private/NiagaraDataInterfaceCamera.cpp b/Engine/Plugins/FX/Niagara/Source/Niagara/Private/NiagaraDataInterfaceCamera.cpp
--- a/Engine/Plugins/FX/Niagara/Source/Niagara/Private/NiagaraDataInterfaceCamera.cpp
+++ b/Engine/Plugins/FX/Niagara/Source/Niagara/Private/NiagaraDataInterfaceCamera.cpp
@@ -30,6 +30,7 @@ struct FNiagaraCameraDIFunctionVersion
 		AddedBufferSizeInvSizeForGPU = 2,
 		LWCConversion = 3,
 		AddGPUNearPlane = 4,
+		AddGPUFOVData = 5,
 
         VersionPlusOne,
         LatestVersion = VersionPlusOne - 1
@@ -216,6 +217,8 @@ void UNiagaraDataInterfaceCamera::GetFunctionsInternal(TArray<FNiagaraFunctionSi
 	Sig.AddOutput(FNiagaraVariable(FNiagaraTypeDefinition::GetVec4Def(), TEXT("Buffer Size And Inverse Size")), LOCTEXT("BufferSizeAndInverseSizeDescription", "Returns a vec4 - the x and y values are width and height of the buffer; the z and w values are the inverse width and height."));
 	Sig.AddOutput(FNiagaraVariable(FNiagaraTypeDefinition::GetVec2Def(), TEXT("Viewport Offset")), LOCTEXT("ViewportOffsetDescription", "Returns the viewport offset in the buffer."));
 	Sig.AddOutput(FNiagaraVariable(FNiagaraTypeDefinition::GetFloatDef(), TEXT("Near Plane")), LOCTEXT("NearPlaneDescription", "Returns the near plane value."));
+	Sig.AddOutput(FNiagaraVariable(FNiagaraTypeDefinition::GetVec2Def(), TEXT("Field Of View (Current Frame)")), LOCTEXT("FOVCurrentDescription", "Returns the current camera field of view (X & Y) in radians."));
+	Sig.AddOutput(FNiagaraVariable(FNiagaraTypeDefinition::GetVec2Def(), TEXT("Field Of View (Previous Frame)")), LOCTEXT("FOVPreviousDescription", "Returns the previous camera field of view (X & Y) in radians."));
 	OutFunctions.Add(Sig);
 
 
@@ -387,8 +390,7 @@ bool UNiagaraDataInterfaceCamera::UpgradeFunctionCall(FNiagaraFunctionSignature&
 		bWasChanged = true;
 	}
 
-	// LWCConversion
-	// AddGPUNearPlane
+	// LWCConversion + AddGPUNearPlane + AddGPUFOVData
 	if (FunctionSignature.FunctionVersion < FNiagaraCameraDIFunctionVersion::LatestVersion)
 	{
 		TArray<FNiagaraFunctionSignature> AllFunctions;
