- Capture all static variables into the sim cache debug data

#rnx

diff --git a/Engine/Plugins/FX/Niagara/Source/Niagara/Private/NiagaraSimCacheDebugData.cpp b/Engine/Plugins/FX/Niagara/Source/Niagara/Private/NiagaraSimCacheDebugData.cpp
--- a/Engine/Plugins/FX/Niagara/Source/Niagara/Private/NiagaraSimCacheDebugData.cpp
+++ b/Engine/Plugins/FX/Niagara/Source/Niagara/Private/NiagaraSimCacheDebugData.cpp
@@ -36,6 +36,22 @@ namespace NSCDebugDataPrivate
 			AddParameterStore(Owner, FrameData, StoreName, ExecContext->CombinedParamStore);
 		}
 	}
+
+#if WITH_EDITORONLY_DATA
+	template<typename TStringType>
+	void AddStaticVariables(UObject* Owner, FNiagaraSimCacheDebugDataFrame& FrameData, TStringType StoreName, const UNiagaraScript* Script)
+	{
+		if (Script && Script->GetVMExecutableData().StaticVariablesWritten.Num() > 0)
+		{
+			FNiagaraParameterStore& DestStore = FrameData.DebugParameterStores.FindOrAdd(StoreName);
+
+			for (const FNiagaraVariable& StaticVar : Script->GetVMExecutableData().StaticVariablesWritten)
+			{
+				DestStore.AddParameter(StaticVar, true, false);
+			}
+		}
+	}
+#endif
 }
 
 UNiagaraSimCacheDebugData::UNiagaraSimCacheDebugData(const FObjectInitializer& ObjectInitializer)
@@ -64,6 +80,14 @@ void UNiagaraSimCacheDebugData::CaptureFrame(FNiagaraSimCacheHelper& Helper, int
 	{
 		AddParameterStore(this, FrameData, TEXT("System Spawn"), SystemSimulation->GetSpawnExecutionContext());
 		AddParameterStore(this, FrameData, TEXT("System Update"), SystemSimulation->GetUpdateExecutionContext());
+
+	#if WITH_EDITORONLY_DATA
+		if (UNiagaraSystem* NiagaraSystem = SystemSimulation->GetSystem())
+		{
+			AddStaticVariables(this, FrameData, TEXT("Static Variables"), NiagaraSystem->GetSystemSpawnScript());
+			AddStaticVariables(this, FrameData, TEXT("Static Variables"), NiagaraSystem->GetSystemUpdateScript());
+		}
+	#endif
 	}
 
 	// Add Emitter Parameters
