Initialize new derived data in FNiagaraTypeDefinition after being imported from text.
Fixes a few NDC related issues and failing Engine Tests.

diff --git a/Engine/Plugins/FX/Niagara/Source/Niagara/Private/NiagaraModule.cpp b/Engine/Plugins/FX/Niagara/Source/Niagara/Private/NiagaraModule.cpp
--- a/Engine/Plugins/FX/Niagara/Source/Niagara/Private/NiagaraModule.cpp
+++ b/Engine/Plugins/FX/Niagara/Source/Niagara/Private/NiagaraModule.cpp
@@ -1533,6 +1533,18 @@ void FNiagaraTypeDefinition::InvalidateCachedIndex(bool bInvalidate)
 	RegisteredTypeDefIndex = bInvalidate ? INDEX_NONE : RegisteredTypeDefIndex;
 }
 
+bool FNiagaraTypeDefinition::ImportTextItem(const TCHAR*& Buffer, int32 PortFlags, UObject* Parent, FOutputDevice* ErrorText)
+{
+	//We provide a custom import text only so that we can generate derived data (GeneratePathNameHash()).
+	//We call back into the standard ImportText but prevent it from using the custom import and re-entering this function with the last arg.
+	StaticStruct()->ImportText(Buffer, this, nullptr, PPF_None, ErrorText, StaticStruct()->GetName(), false);
+
+	//Now we're imported we can generate our derived values.
+	GeneratePathNameHash();
+
+	return true;//Outer call to ImportText should exit.
+}
+
 bool FNiagaraTypeDefinition::Serialize(FArchive& Ar)
 {
 	Ar.UsingCustomVersion(FNiagaraCustomVersion::GUID);
diff --git a/Engine/Plugins/FX/Niagara/Source/Niagara/Public/NiagaraTypes.h b/Engine/Plugins/FX/Niagara/Source/Niagara/Public/NiagaraTypes.h
--- a/Engine/Plugins/FX/Niagara/Source/Niagara/Public/NiagaraTypes.h
+++ b/Engine/Plugins/FX/Niagara/Source/Niagara/Public/NiagaraTypes.h
@@ -974,6 +974,7 @@ public:
 	uint16 UnderlyingType;
 
 	NIAGARA_API bool Serialize(FArchive& Ar);
+	NIAGARA_API bool ImportTextItem(const TCHAR*& Buffer, int32 PortFlags, UObject* Parent, FOutputDevice* ErrorText);
 
 private:
 	// for allowing access to the cached RegisteredTypeDefIndex
@@ -1187,6 +1188,7 @@ struct TStructOpsTypeTraits<FNiagaraTypeDefinition> : public TStructOpsTypeTrait
 	enum
 	{
 		WithSerializer = true,
+		WithImportTextItem = true,
 	};
 };
 
