Fix assert that can fire in cases where we have an NDC reading emitter that completes and one that loops.
Fires a check incorrectly that they apparently disagree on how many NDC items there were.

diff --git a/Engine/Plugins/FX/Niagara/Source/Niagara/Private/DataInterface/NiagaraDataInterfaceDataChannelRead.cpp b/Engine/Plugins/FX/Niagara/Source/Niagara/Private/DataInterface/NiagaraDataInterfaceDataChannelRead.cpp
--- a/Engine/Plugins/FX/Niagara/Source/Niagara/Private/DataInterface/NiagaraDataInterfaceDataChannelRead.cpp
+++ b/Engine/Plugins/FX/Niagara/Source/Niagara/Private/DataInterface/NiagaraDataInterfaceDataChannelRead.cpp
@@ -726,7 +726,7 @@ void UNiagaraDataInterfaceDataChannelRead::PostStageTick(FNDICpuPostStageContext
 	{
 		FNDIDataChannelRead_EmitterInstanceData& EmitterInstData = EmitterInstanceDataPair.Value;
 		FNiagaraEmitterInstance* TargetEmitter = EmitterInstanceDataPair.Key;
-		if(TargetEmitter == nullptr)
+		if(TargetEmitter == nullptr || EmitterInstData.NDCSpawnCounts.Num() == 0)
 		{
 			EmitterInstData.Reset();
 			continue;
@@ -736,6 +736,7 @@ void UNiagaraDataInterfaceDataChannelRead::PostStageTick(FNDICpuPostStageContext
 		{
 			InstanceData->NDCElementCountAtSpawn = EmitterInstData.NDCSpawnCounts.Num();
 		}
+		//All emitters with a non-zero spawn count should agree about how many NDC entries there were.
 		check(EmitterInstData.NDCSpawnCounts.Num() == InstanceData->NDCElementCountAtSpawn);
 
 		TArray<uint32> PerNDCSpawnCounts;
