Several perf improvements for NDC writes.

- New helper struct for NDC writing from Cpp. Avoids slower BP API.
- Allows optional validation/checking for required vs optonal data.
- Improves perf for NDC writes ~5x
- Method to accumulate writes into the same NDC data in larger shared buffers. Avoids naieve write code generating lots of tiny inefficient buffers to consume.
- Improves perf for NDC consuming data in pathalogical cases by ~1000x
- NDC data game data now only refernces it's layout data rather than the NDC itself. Cleaner and avoids weakptr derefs
- Improved some type helpers for general use and for us in template code.
- LWC types in FNiagaraTypeHelper now follow same init and access patern as regular core typedefs.
- GetTypeDef() and GetSimType() utilities to access typedefs from CPP types
- TNiagaraTypeHelper to help accessing correct typedefs in templated code.
- Removed vestigial publish requests array from Island NDC.
- Adding new API to NDC RW engine tests.

diff --git a/Engine/Plugins/FX/Niagara/Source/Niagara/Private/DataInterface/NiagaraDataInterfaceDataChannelRead.cpp b/Engine/Plugins/FX/Niagara/Source/Niagara/Private/DataInterface/NiagaraDataInterfaceDataChannelRead.cpp
--- a/Engine/Plugins/FX/Niagara/Source/Niagara/Private/DataInterface/NiagaraDataInterfaceDataChannelRead.cpp
+++ b/Engine/Plugins/FX/Niagara/Source/Niagara/Private/DataInterface/NiagaraDataInterfaceDataChannelRead.cpp
@@ -1233,7 +1233,10 @@ void UNiagaraDataInterfaceDataChannelRead::GetVMExternalFunction(const FVMExtern
 		{
 			FName NDCVarName = NDCVar->GetName();
 			FNiagaraTypeDefinition NDCVarType = NDCVar->GetType();
-			if (NDCVarType == FNiagaraTypeDefinition::GetIntDef()) OutFunc = FVMExternalFunction::CreateLambda([this, NDCVarName](FVectorVMExternalFunctionContext& Context) { this->SpawnDirect<int32>(Context, NDCVarName); });
+			if (NDCVarType == FNiagaraTypeDefinition::GetIntDef()) 
+			{
+				OutFunc = FVMExternalFunction::CreateLambda([this, NDCVarName](FVectorVMExternalFunctionContext& Context) { this->SpawnDirect<int32>(Context, NDCVarName); });
+			}
 			else
 			{
 				UE_LOG(LogTemp, Display, TEXT("Failed to bind Data Interface function as this is not a valid variable type for SpawnDirect.\nDI: %s.\nReceived Name: %s\nNDC: %s"), *GetPathNameSafe(this), *BindingInfo.Name.ToString(), *GetPathNameSafe(Channel));
@@ -1842,32 +1845,6 @@ template<> FVector2f NDCValueDefault<FVector2f>() { return FVector2f(1.0f); }
 template<> FVector3f NDCValueDefault<FVector3f>() { return FVector3f(1.0f); }
 template<> FVector4f NDCValueDefault<FVector4f>() { return FVector4f(1.0f); }
 
-template <typename T> FNiagaraTypeDefinition GetTypeDef() { return FNiagaraTypeDefinition(); }
-template <typename T> FNiagaraTypeDefinition GetSimTypeDef() { return FNiagaraTypeDefinition(); }
-
-template<> FNiagaraTypeDefinition GetTypeDef<int32>() { return FNiagaraTypeDefinition::GetIntDef(); }
-template<> FNiagaraTypeDefinition GetTypeDef<double>() { return FNiagaraTypeHelper::GetDoubleDef(); }
-template<> FNiagaraTypeDefinition GetTypeDef<FVector2D>() { return FNiagaraTypeHelper::GetVector2DDef(); }
-template<> FNiagaraTypeDefinition GetTypeDef<FVector>() { return FNiagaraTypeHelper::GetVectorDef(); }
-template<> FNiagaraTypeDefinition GetTypeDef<FVector4>() { return FNiagaraTypeHelper::GetVector4Def(); }
-template<> FNiagaraTypeDefinition GetTypeDef<FNiagaraPosition>() { return FNiagaraTypeDefinition::GetPositionDef(); }
-
-template<> FNiagaraTypeDefinition GetSimTypeDef<int32>() { return FNiagaraTypeDefinition::GetIntDef(); }
-template<> FNiagaraTypeDefinition GetSimTypeDef<double>() { return FNiagaraTypeDefinition::GetFloatDef(); }
-template<> FNiagaraTypeDefinition GetSimTypeDef<FVector2D>() { return FNiagaraTypeDefinition::GetVec2Def(); }
-template<> FNiagaraTypeDefinition GetSimTypeDef<FVector>() { return FNiagaraTypeDefinition::GetVec3Def(); }
-template<> FNiagaraTypeDefinition GetSimTypeDef<FVector4>() { return FNiagaraTypeDefinition::GetVec4Def(); }
-template<> FNiagaraTypeDefinition GetSimTypeDef<FNiagaraPosition>() { return FNiagaraTypeDefinition::GetVec3Def(); }
-
-template<typename TLWCType>
-struct TLWCToSimType {};
-
-template<> struct TLWCToSimType<int32> { typedef int32 TSimType; };
-template<> struct TLWCToSimType<double> { typedef float TSimType;  };
-template<> struct TLWCToSimType<FVector2D> { typedef FVector2f TSimType; };
-template<> struct TLWCToSimType<FVector> { typedef FVector3f TSimType; };
-template<> struct TLWCToSimType<FVector4> { typedef FVector4f TSimType; };
-template<> struct TLWCToSimType<FNiagaraPosition> { typedef FVector3f TSimType; };
 
 template<typename T>
 void UNiagaraDataInterfaceDataChannelRead::SpawnDirect(FVectorVMExternalFunctionContext& Context, FName NDCVarName)
@@ -1875,7 +1852,7 @@ void UNiagaraDataInterfaceDataChannelRead::SpawnDirect(FVectorVMExternalFunction
 	//This should only be called from emitter scripts and since it has per instance data then we process them individually.
 	check(Context.GetNumInstances() == 1);
 
-	using TSimType = typename TLWCToSimType<T>::TSimType;
+	using TSimType = typename TNiagaraTypeHelper<T>::TSimType;
 	VectorVM::FUserPtrHandler<FNDIDataChannelReadInstanceData> InstData(Context);
 
 	FNDIInputParam<FNiagaraBool> InEnabled(Context);
@@ -1951,7 +1928,7 @@ void UNiagaraDataInterfaceDataChannelRead::ScaleSpawnCount(FVectorVMExternalFunc
 	//This should only be called from emitter scripts and since it has per instance data then we process them individually.
 	check(Context.GetNumInstances() == 1);
 
-	using TSimType = typename TLWCToSimType<T>::TSimType;
+	using TSimType = typename TNiagaraTypeHelper<T>::TSimType;
 	VectorVM::FUserPtrHandler<FNDIDataChannelReadInstanceData> InstData(Context);
 
 	FNDIInputParam<FNiagaraBool> InEnabled(Context);
diff --git a/Engine/Plugins/FX/Niagara/Source/Niagara/Private/DataInterface/NiagaraDataInterfaceDataChannelWrite.cpp b/Engine/Plugins/FX/Niagara/Source/Niagara/Private/DataInterface/NiagaraDataInterfaceDataChannelWrite.cpp
--- a/Engine/Plugins/FX/Niagara/Source/Niagara/Private/DataInterface/NiagaraDataInterfaceDataChannelWrite.cpp
+++ b/Engine/Plugins/FX/Niagara/Source/Niagara/Private/DataInterface/NiagaraDataInterfaceDataChannelWrite.cpp
@@ -854,8 +854,7 @@ bool UNiagaraDataInterfaceDataChannelWrite::SimCacheWriteFrame(UObject* StorageO
 
 		if (InstanceData->DataChannelData && ShouldPublish())
 		{
-			FNiagaraDataChannelGameData GameData;
-			GameData.Init(Channel->Get());
+			FNiagaraDataChannelGameData GameData(Channel->Get()->GetLayoutInfo());
 
 			if(InstanceData->DestinationData && InstanceData->DestinationData->GetNumInstances() > 0)
 			{
@@ -961,8 +960,7 @@ bool UNiagaraDataInterfaceDataChannelWrite::SimCacheReadFrame(UObject* StorageOb
 			PublishRequest.DebugSource = FString::Format(TEXT("{0} (Sim cache {1})"), {SystemInstance->GetSystem()->GetName(), GetPathNameSafe(StorageObject->GetOuter())});
 #endif
 
-			PublishRequest.GameData = MakeShared<FNiagaraDataChannelGameData>();
-			PublishRequest.GameData->Init(Channel->Get());
+			PublishRequest.GameData = MakeShared<FNiagaraDataChannelGameData>(Channel->Get()->GetLayoutInfo());
 			PublishRequest.GameData->SetNum(Frame.NumElements);
 			for (int32 i = 0; i < Frame.VariableData.Num(); i++)
 			{
diff --git a/Engine/Plugins/FX/Niagara/Source/Niagara/Private/NiagaraDataChannel.cpp b/Engine/Plugins/FX/Niagara/Source/Niagara/Private/NiagaraDataChannel.cpp
--- a/Engine/Plugins/FX/Niagara/Source/Niagara/Private/NiagaraDataChannel.cpp
+++ b/Engine/Plugins/FX/Niagara/Source/Niagara/Private/NiagaraDataChannel.cpp
@@ -49,6 +49,8 @@ namespace NDCCVars
 	bool AutoUploadGPUSpawnData = true;
 	static FAutoConsoleVariableRef CVarAutoUploadGPUSpawnData(TEXT("fx.Niagara.DataChannels.AutoUploadGPUSpawnData"), AutoUploadGPUSpawnData, TEXT("When true we will automatically upload any CPU NDC data to the GPU if it has been used to spawn GPU particles."), ECVF_Default);
 
+	bool EnableAccumulatedWrites = true;
+	static FAutoConsoleVariableRef CVarEnableAccumulatedWrites(TEXT("fx.Niagara.DataChannels.EnableAccumulatedWrites"), EnableAccumulatedWrites, TEXT("True if we attempt to accumulate writes into larger shared buffers."), ECVF_Default);
 };
 
 //////////////////////////////////////////////////////////////////////////
@@ -74,11 +76,14 @@ void FNiagaraDataChannelGameDataLayout::Init(TConstArrayView<FNiagaraDataChannel
 
 //////////////////////////////////////////////////////////////////////////
 
-void FNiagaraDataChannelGameData::Init(const UNiagaraDataChannel* InDataChannel)
+FNiagaraDataChannelGameData::FNiagaraDataChannelGameData(const FNiagaraDataChannelLayoutInfoPtr& InLayoutInfo)
 {
-	check(InDataChannel);
-	DataChannel = InDataChannel;
-	LayoutInfo = InDataChannel->GetLayoutInfo();
+	Init(InLayoutInfo);
+}
+
+void FNiagaraDataChannelGameData::Init(const FNiagaraDataChannelLayoutInfoPtr& InLayoutInfo)
+{
+	LayoutInfo = InLayoutInfo;
 
 	const FNiagaraDataChannelGameDataLayout& Layout = LayoutInfo->GetGameDataLayout();
 	VariableData.Empty();
@@ -93,15 +98,27 @@ void FNiagaraDataChannelGameData::Init(const UNiagaraDataChannel* InDataChannel)
 void FNiagaraDataChannelGameData::Empty()
 {
 	NumElements = 0;
+	PrevNumElements = 0;
+	MaxElements = 0;
 	for (FNiagaraDataChannelVariableBuffer& VarData : VariableData)
 	{
 		VarData.Empty();
 	}
 }
 
+void FNiagaraDataChannelGameData::Reset()
+{
+	NumElements = 0;
+	PrevNumElements = 0;
+	for (FNiagaraDataChannelVariableBuffer& VarData : VariableData)
+	{
+		VarData.Reset();
+	}
+}
+
 void FNiagaraDataChannelGameData::BeginFrame()
 {
-	bool bKeepPrevious = DataChannel->KeepPreviousFrameData();
+	bool bKeepPrevious = LayoutInfo->KeepPreviousFrameData();
 	PrevNumElements = bKeepPrevious ? NumElements : 0;
 	NumElements = 0;
 	for (FNiagaraDataChannelVariableBuffer& VarData : VariableData)
@@ -113,6 +130,7 @@ void FNiagaraDataChannelGameData::BeginFrame()
 void FNiagaraDataChannelGameData::SetNum(int32 NewNum)
 {
 	NumElements = NewNum;
+	MaxElements = FMath::Max(NewNum, MaxElements);
 	for (auto& Buffer : VariableData)
 	{
 		Buffer.SetNum(NewNum);
@@ -121,7 +139,7 @@ void FNiagaraDataChannelGameData::SetNum(int32 NewNum)
 
 void FNiagaraDataChannelGameData::Reserve(int32 NewNum)
 {
-	//NumElements = NewNum;
+	MaxElements = FMath::Max(NewNum, MaxElements);
 	for (auto& Buffer : VariableData)
 	{
 		Buffer.Reserve(NewNum);
@@ -301,7 +319,7 @@ void FNiagaraDataChannelGameData::WriteToDataSet(FNiagaraDataBuffer* DestBuffer,
 }
 void FNiagaraDataChannelGameData::AppendFromGameData(const FNiagaraDataChannelGameData& GameData)
 {
-	check(GetDataChannel() == GameData.GetDataChannel());
+	check(GetLayoutInfo() == GameData.GetLayoutInfo());
 
 	NumElements += GameData.Num();
 	TConstArrayView<FNiagaraDataChannelVariableBuffer> SrcBuffers = GameData.GetVariableBuffers();
@@ -979,6 +997,8 @@ int32 FNiagaraDataChannelData::ConsumePublishRequests(UNiagaraDataChannelHandler
 {
 	check(IsValid(Owner));
 
+	FlushAllPendingGameData();
+
 	//There should be no access on other threads at this point anyway but lock just to be safe.
 	FScopeLock Lock(&PublishCritSec);
 
@@ -1188,6 +1208,83 @@ bool FNiagaraDataChannelData::HasData()const
 	return GameData->Num() > 0 || (CPUSimData && CPUSimData->GetCurrentData()->GetNumInstances() > 0) || PublishRequestsForGPU.Num() > 0;
 }
 
+FNiagaraDataChannelGameDataPtr FNiagaraDataChannelData::GetGameDataForWriteGT(int32 Count, bool bVisibleToGame, bool bVisibleToCPU, bool bVisibleToGPU)
+{
+	//We keep an array of cached data ptrs, one for each of the most recent combinations of visible flags.
+	//This dynamic choice of sending data to Game/CPU/GPU is honestly annoying. I think I might look at removing and having only as a system wide thing on the NDC	
+	auto GetPendingBufferIdx = [](bool bVisibleToGame, bool bVisibleToCPU, bool bVisibleToGPU)
+	{
+		uint32 PendingBufferIndex = bVisibleToGame ? 1 : 0;
+		PendingBufferIndex += (bVisibleToCPU ? 1 : 0) << 1;
+		PendingBufferIndex += (bVisibleToGPU ? 1 : 0) << 2;
+		return PendingBufferIndex;
+	};
+
+	if(PendingDestGameData.Num() == 0)
+	{
+		PendingDestGameData.SetNum(GetPendingBufferIdx(true,true,true) + 1);
+	}
+
+	//See if we have a pending game data that can accommodate the new write.
+	
+	int32 PendingIndex = GetPendingBufferIdx(bVisibleToGame, bVisibleToCPU, bVisibleToGPU);
+	FNiagaraDataChannelGameDataPtr& CachedData = PendingDestGameData[PendingIndex];
+	if(NDCCVars::EnableAccumulatedWrites)
+	{
+		if(CachedData.IsValid() && CachedData->Max() >= CachedData->Num() + Count)
+		{
+			return CachedData;
+		}
+		else
+		{
+			int32 AllocCount = FMath::Max(Count, 128);
+			FlushPendingGameData(PendingIndex);
+			PendingDestGameData[PendingIndex] = MakeShared<FNiagaraDataChannelGameData>(LayoutInfo);
+			PendingDestGameData[PendingIndex]->Reserve(AllocCount);
+			return PendingDestGameData[PendingIndex];	
+		}
+	}
+	else
+	{
+		FNiagaraDataChannelPublishRequest PublishRequest;
+		PublishRequest.bVisibleToGame = bVisibleToGame;
+		PublishRequest.bVisibleToCPUSims = bVisibleToCPU;
+		PublishRequest.bVisibleToGPUSims = bVisibleToGPU;
+		PublishRequest.GameData = MakeShared<FNiagaraDataChannelGameData>(LayoutInfo);
+#if !UE_BUILD_SHIPPING
+		PublishRequest.DebugSource = TEXT("Accumulated Pending Game Data");
+#endif
+		Publish(PublishRequest); 
+		return PublishRequest.GameData;
+	}
+}
+
+void FNiagaraDataChannelData::FlushPendingGameData(int32 Index)
+{
+	if(PendingDestGameData[Index].IsValid() && PendingDestGameData[Index]->Num() > 0)
+	{
+		FNiagaraDataChannelPublishRequest PublishRequest;
+		PublishRequest.bVisibleToGame = Index & 1;
+		PublishRequest.bVisibleToCPUSims = Index & 2;
+		PublishRequest.bVisibleToGPUSims = Index & 4;
+		PublishRequest.GameData = PendingDestGameData[Index];
+#if !UE_BUILD_SHIPPING
+		PublishRequest.DebugSource = TEXT("Accumulated Pending Game Data");
+#endif
+		Publish(PublishRequest);
+	}
+
+	PendingDestGameData[Index] = nullptr;
+}
+
+void FNiagaraDataChannelData::FlushAllPendingGameData()
+{
+	for(int32 i=0; i < PendingDestGameData.Num(); ++i)
+	{
+		FlushPendingGameData(i);
+	}
+}
+
 void FNiagaraDataChannelData::DestroyRenderThreadProxy(FNiagaraGpuComputeDispatchInterface* ComputeDispatchInterface)
 {
 	//We can leave the proxy to be destroyed by the dispatcher ( or in some edge cases perhaps a DI proxy ).
@@ -1229,6 +1326,8 @@ FNiagaraDataChannelLayoutInfo::FNiagaraDataChannelLayoutInfo(const UNiagaraDataC
 	check(DataChannel);
 	GameDataLayout.Init(DataChannel->GetVariables());
 
+	bKeepPreviousFrameData = DataChannel->KeepPreviousFrameData();
+
 	CompiledData.SimTarget = ENiagaraSimTarget::CPUSim;
 	CompiledDataGPU.SimTarget = ENiagaraSimTarget::GPUComputeSim;
 	for (const FNiagaraDataChannelVariable& NDCVar : DataChannel->GetVariables())
@@ -1400,9 +1499,7 @@ const FNiagaraDataChannelLayoutInfoPtr UNiagaraDataChannel::GetLayoutInfo()const
 
 FNiagaraDataChannelGameDataPtr UNiagaraDataChannel::CreateGameData()const
 {
-	FNiagaraDataChannelGameDataPtr NewData = MakeShared<FNiagaraDataChannelGameData>();
-	NewData->Init(this);
-	return NewData;
+	return MakeShared<FNiagaraDataChannelGameData>(GetLayoutInfo());
 }
 
 bool UNiagaraDataChannel::IsValid()const
@@ -1646,6 +1743,7 @@ void FNiagaraDataChannelDebugUtilities::LogWrite(const FNiagaraDataChannelPublis
 			Data->FrameNumber = GFrameCounter;
 		}
 		FChannelWriteRequest DebugData;
+		DebugData.Channel = DataChannel;
 		DebugData.DebugSource = WriteRequest.DebugSource;
 		DebugData.bVisibleToGame = WriteRequest.bVisibleToGame;
 		DebugData.bVisibleToCPUSims = WriteRequest.bVisibleToCPUSims;
@@ -1657,8 +1755,7 @@ void FNiagaraDataChannelDebugUtilities::LogWrite(const FNiagaraDataChannelPublis
 		}
 		else if (ensure(WriteRequest.Data))
 		{
-			DebugData.Data = MakeShared<FNiagaraDataChannelGameData>();
-			DebugData.Data->Init(DataChannel);
+			DebugData.Data = MakeShared<FNiagaraDataChannelGameData>(DataChannel->GetLayoutInfo());
 			DebugData.Data->AppendFromDataSet(WriteRequest.Data, WriteRequest.LwcTile);
 		}
 		Data->WriteRequests.Add(DebugData);
@@ -1688,8 +1785,7 @@ void FNiagaraDataChannelDebugUtilities::LogWrite(const FNiagaraDataChannelPublis
 			NumInsts = WriteRequest.Data->GetNumInstances();
 			if (NDCCVars::LogWritesToOutputLog > 1)
 			{
-				FNiagaraDataChannelGameData TempData;
-				TempData.Init(DataChannel);
+				FNiagaraDataChannelGameData TempData(DataChannel->GetLayoutInfo());
 				TempData.AppendFromDataSet(WriteRequest.Data, WriteRequest.LwcTile);
 				DataString = ToJson(&TempData);
 			}
@@ -1731,7 +1827,7 @@ void FNiagaraDataChannelDebugUtilities::DumpAllWritesToLog()
 			, Request.bVisibleToGPUSims ? TEXT("X") : TEXT(" ")
 			, *TickGroupToString(Request.TickGroup)
 			, Request.Data->Num()
-			, *GetPathNameSafe(Request.Data->GetDataChannel())
+			, *GetPathNameSafe(Request.Channel.Get())
 			, *DataString);
 		}
 		UE_LOG(LogNiagara, Log, TEXT("----------------------------------------------"));
@@ -1754,14 +1850,25 @@ FString FNiagaraDataChannelDebugUtilities::ToJson(FNiagaraDataChannelGameData* D
 	TSharedRef<TJsonWriter<>> JsonWriter = TJsonWriterFactory<>::Create(&JsonString);
 	JsonWriter->WriteArrayStart();
 
-	TConstArrayView<FNiagaraDataChannelVariable> ChannelVariables = Data->GetDataChannel()->GetVariables();
+	TConstArrayView<FNiagaraDataChannelVariableBuffer> VaraibleBuffers = Data->GetVariableBuffers();
+
+	const FNiagaraDataChannelGameDataLayout& GameDataLayout = Data->GetLayoutInfo()->GetGameDataLayout();
 	for (int32 i = 0; i < Data->Num(); i++)
 	{
 		TSharedPtr<FJsonObject> JsonObject = MakeShared<FJsonObject>();
-		for (const FNiagaraDataChannelVariable& Var : ChannelVariables)
+		for (auto& LayoutVarPair : GameDataLayout.VariableIndices)
 		{
+			const FNiagaraVariableBase& Var = LayoutVarPair.Key;
+			int32 VarIndex = LayoutVarPair.Value;
+
 			FString VarName = Var.GetName().ToString();
-			FNiagaraDataChannelVariableBuffer* Buffer = Data->FindVariableBuffer(Var);
+			const FNiagaraDataChannelVariableBuffer* Buffer = nullptr;
+			
+			if(VaraibleBuffers.IsValidIndex(VarIndex))
+			{
+				Buffer = &VaraibleBuffers[VarIndex];
+			}
+			
 			if (Var.GetType() == FNiagaraTypeHelper::GetDoubleDef() && Buffer)
 			{
 				double Value;
diff --git a/Engine/Plugins/FX/Niagara/Source/Niagara/Private/NiagaraDataChannelAccessor.cpp b/Engine/Plugins/FX/Niagara/Source/Niagara/Private/NiagaraDataChannelAccessor.cpp
--- a/Engine/Plugins/FX/Niagara/Source/Niagara/Private/NiagaraDataChannelAccessor.cpp
+++ b/Engine/Plugins/FX/Niagara/Source/Niagara/Private/NiagaraDataChannelAccessor.cpp
@@ -2,6 +2,7 @@
 
 #include "NiagaraDataChannelAccessor.h"
 
+#include "Engine/Engine.h"
 #include "UObject/UObjectGlobals.h"
 #include "UObject/Package.h"
 
@@ -9,9 +10,25 @@
 #include "NiagaraDataChannelCommon.h"
 #include "NiagaraDataChannel.h"
 #include "NiagaraDataChannelHandler.h"
+#include "NiagaraDataChannelManager.h"
 
 #include "NiagaraWorldManager.h"
 
+
+namespace NDCCVars
+{
+#if DEBUG_NDC_ACCESS
+	int32 MissingNDCAccessorMode = 2;
+	static FAutoConsoleVariableRef CVarMissingNDCAccessorMode(
+		TEXT("fx.Niagara.DataChannels.MissingNDCAccessorMode"), 
+		MissingNDCAccessorMode, 
+		TEXT("Controls validation behavior for missing NDC Variabls being accessed from code. 0=Silent, 1=Log Warning, 2=Log Warning and ensure once, 3=Log Warning and ensure always."),
+		ECVF_Default);
+#else
+	constexpr int32 MissingNDCAccessorMode = 0;
+#endif
+};
+
 //////////////////////////////////////////////////////////////////////////
 
 void UNiagaraDataChannelReader::Cleanup()
@@ -142,6 +159,8 @@ void UNiagaraDataChannelWriter::Cleanup()
 
 bool UNiagaraDataChannelWriter::InitWrite(FNiagaraDataChannelSearchParameters SearchParams, int32 Count, bool bVisibleToGame, bool bVisibleToCPU, bool bVisibleToGPU, const FString& DebugSource)
 {
+	check(IsInGameThread());
+
 	if (Count == 0)
 	{
 		UE_LOG(LogNiagara, Warning, TEXT("Call to UNiagaraDataChannelWriter::InitWrite with Count == 0. Ignored."));
@@ -152,19 +171,11 @@ bool UNiagaraDataChannelWriter::InitWrite(FNiagaraDataChannelSearchParameters Se
 
 	if(FNiagaraDataChannelDataPtr DestData = Owner->FindData(SearchParams, ENiagaraResourceAccess::WriteOnly))
 	{
-		//TODO- Dont create a whole new game Data here. Rather grab a pre-made staging/input One from the Data PTR With the relevant publish Flags.
-		Data = Owner->GetDataChannel()->CreateGameData();
-		Data->SetNum(Count);
-
-		FNiagaraDataChannelPublishRequest PublishRequest;
-		PublishRequest.bVisibleToGame = bVisibleToGame;
-		PublishRequest.bVisibleToCPUSims = bVisibleToCPU;
-		PublishRequest.bVisibleToGPUSims = bVisibleToGPU;
-		PublishRequest.GameData = Data;
-#if !UE_BUILD_SHIPPING
-		PublishRequest.DebugSource = DebugSource;
-#endif
-		DestData->Publish(PublishRequest);
+		//Attempt to use an existing cached dest data.
+		Data = DestData->GetGameDataForWriteGT(Count, bVisibleToGame, bVisibleToCPU, bVisibleToGPU);	
+		StartIndex = Data->Num();
+		Data->SetNum(StartIndex + Count);
+
 		return true;
 	}
 
@@ -182,60 +193,263 @@ int32 UNiagaraDataChannelWriter::Num()const
 
 void UNiagaraDataChannelWriter::WriteFloat(FName VarName, int32 Index, double InData)
 {
-	WriteData(FNiagaraVariableBase(FNiagaraTypeHelper::GetDoubleDef(), VarName), Index, InData);
+	WriteData(FNiagaraVariableBase(FNiagaraTypeHelper::GetDoubleDef(), VarName), StartIndex + Index, InData);
 }
 
 void UNiagaraDataChannelWriter::WriteVector2D(FName VarName, int32 Index, FVector2D InData)
 {
-	WriteData(FNiagaraVariableBase(FNiagaraTypeHelper::GetVector2DDef(), VarName), Index, InData);
+	WriteData(FNiagaraVariableBase(FNiagaraTypeHelper::GetVector2DDef(), VarName), StartIndex + Index, InData);
 }
 
 void UNiagaraDataChannelWriter::WriteVector(FName VarName, int32 Index, FVector InData)
 {
-	WriteData(FNiagaraVariableBase(FNiagaraTypeHelper::GetVectorDef(), VarName), Index, InData);
+	WriteData(FNiagaraVariableBase(FNiagaraTypeHelper::GetVectorDef(), VarName), StartIndex + Index, InData);
 }
 
 void UNiagaraDataChannelWriter::WriteVector4(FName VarName, int32 Index, FVector4 InData)
 {
-	WriteData(FNiagaraVariableBase(FNiagaraTypeHelper::GetVector4Def(), VarName), Index, InData);
+	WriteData(FNiagaraVariableBase(FNiagaraTypeHelper::GetVector4Def(), VarName), StartIndex + Index, InData);
 }
 
 void UNiagaraDataChannelWriter::WriteQuat(FName VarName, int32 Index, FQuat InData)
 {
-	WriteData(FNiagaraVariableBase(FNiagaraTypeHelper::GetQuatDef(), VarName), Index, InData);
+	WriteData(FNiagaraVariableBase(FNiagaraTypeHelper::GetQuatDef(), VarName), StartIndex + Index, InData);
 }
 
 void UNiagaraDataChannelWriter::WriteLinearColor(FName VarName, int32 Index, FLinearColor InData)
 {
-	WriteData(FNiagaraVariableBase(FNiagaraTypeDefinition::GetColorDef(), VarName), Index, InData);
+	WriteData(FNiagaraVariableBase(FNiagaraTypeDefinition::GetColorDef(), VarName), StartIndex + Index, InData);
 }
 
 void UNiagaraDataChannelWriter::WriteInt(FName VarName, int32 Index, int32 InData)
 {
-	WriteData(FNiagaraVariableBase(FNiagaraTypeDefinition::GetIntDef(), VarName), Index, InData);
+	WriteData(FNiagaraVariableBase(FNiagaraTypeDefinition::GetIntDef(), VarName), StartIndex + Index, InData);
 }
 
 void UNiagaraDataChannelWriter::WriteEnum(FName VarName, int32 Index, uint8 InData)
 {
-	WriteInt(VarName, Index, InData);
+	WriteInt(VarName, StartIndex + Index, InData);
 }
 
 void UNiagaraDataChannelWriter::WriteBool(FName VarName, int32 Index, bool InData)
 {
-	WriteData(FNiagaraVariableBase(FNiagaraTypeDefinition::GetBoolDef(), VarName), Index, FNiagaraBool(InData));
+	WriteData(FNiagaraVariableBase(FNiagaraTypeDefinition::GetBoolDef(), VarName), StartIndex + Index, FNiagaraBool(InData));
 }
 
 void UNiagaraDataChannelWriter::WriteSpawnInfo(FName VarName, int32 Index, FNiagaraSpawnInfo InData)
 {
-	WriteData(FNiagaraVariableBase(FNiagaraTypeDefinition(FNiagaraSpawnInfo::StaticStruct()), VarName), Index, InData);
+	WriteData(FNiagaraVariableBase(FNiagaraTypeDefinition(FNiagaraSpawnInfo::StaticStruct()), VarName), StartIndex + Index, InData);
 }
 
 void UNiagaraDataChannelWriter::WritePosition(FName VarName, int32 Index, FVector InData)
 {
-	WriteData(FNiagaraVariableBase(FNiagaraTypeDefinition::GetPositionDef(), VarName), Index, InData);
+	WriteData(FNiagaraVariableBase(FNiagaraTypeDefinition::GetPositionDef(), VarName), StartIndex + Index, InData);
 }
 
 void UNiagaraDataChannelWriter::WriteID(FName VarName, int32 Index, FNiagaraID InData)
 {
-	WriteData(FNiagaraVariableBase(FNiagaraTypeDefinition::GetIDDef(), VarName), Index, InData);
+	WriteData(FNiagaraVariableBase(FNiagaraTypeDefinition::GetIDDef(), VarName), StartIndex + Index, InData);
+}
+
+
+//////////////////////////////////////////////////////////////////////////
+//FNDCVarAccessorBase - Base class for all NDC accessor helpers. See full description in header.
+
+FNDCVarAccessorBase::FNDCVarAccessorBase(FNDCAccessorBase& Owner, FNiagaraVariableBase InVariable, bool bInIsRequired)
+	: Variable(InVariable)
+	, VarOffset(INDEX_NONE)
+	, bIsRequired(bInIsRequired)
+{
+	Owner.VariableAccessors.Add(this);
+}
+
+void FNDCVarAccessorBase::Init(const UNiagaraDataChannel* DataChannel)
+{
+	VarOffset = INDEX_NONE;
+	if (const FNiagaraDataChannelLayoutInfoPtr LayoutInfo = DataChannel->GetLayoutInfo())
+	{
+		const FNiagaraDataChannelGameDataLayout& GameDataLayout = LayoutInfo->GetGameDataLayout();
+
+#if DEBUG_NDC_ACCESS
+		DebugCachedLayout = LayoutInfo;
+		WeakNDC = DataChannel;
+#endif
+	
+		if (const int32* Index = GameDataLayout.VariableIndices.Find(Variable))
+		{
+			VarOffset = *Index;
+		}
+		else if(Variable.GetType().IsEnum())
+		{
+			//In the case of enums we allow for using raw int32s internally to Niagara also.
+			FNiagaraVariableBase IntVar(FNiagaraTypeDefinition::GetIntDef(), Variable.GetName());
+			if (const int32* IntIndex = GameDataLayout.VariableIndices.Find(IntVar))
+			{
+				VarOffset = *IntIndex;
+			}
+		}
+	}
 }
+
+//////////////////////////////////////////////////////////////////////////
+//FNDCAccessorBase - Base class for all user code NDC Accessor classes. See full description in header.
+
+void FNDCAccessorBase::Init(const UNiagaraDataChannel* DataChannel)
+{
+	check(DataChannel);
+	TArray<FNiagaraVariableBase> MissingVars;
+	for (FNDCVarAccessorBase* Var : VariableAccessors)
+	{
+		Var->Init(DataChannel);
+
+		if(NDCCVars::MissingNDCAccessorMode > 0)
+		{
+			if (Var->VarOffset == INDEX_NONE && Var->bIsRequired)
+			{
+				MissingVars.Reserve(DataChannel->GetVariables().Num());
+				MissingVars.Add(Var->Variable);
+			}
+		}
+	}
+
+	if(MissingVars.Num() > 0 && NDCCVars::MissingNDCAccessorMode > 0)
+	{
+		//This variable is not present in the NDC. If it was required then warn and/or ensure if appropriate.
+		if (NDCCVars::MissingNDCAccessorMode >= 1)
+		{
+			FStringBuilderBase ErrorBuilder;
+			ErrorBuilder.Appendf(TEXT("Niagara Data Channel being accessed from C++ is missing expected required variables.\nNDC:%s\nVariables:\n"), *GetNameSafe(DataChannel->GetAsset()));
+			for(FNiagaraVariableBase& MissingVar : MissingVars)
+			{
+				ErrorBuilder << *MissingVar.GetType().GetName() << TEXT(" ") << MissingVar.GetName() << TEXT("\n");
+			}
+			ErrorBuilder.Appendf(TEXT("\nVariables in NDC: %s\n"), *GetNameSafe(DataChannel->GetAsset()));
+			for(const FNiagaraDataChannelVariable& NDCVar : DataChannel->GetVariables())
+			{
+				ErrorBuilder << *NDCVar.GetType().GetName() << TEXT(" ") << NDCVar.GetName() << TEXT("\n");
+			}
+			UE_LOG(LogNiagara, Warning, TEXT("%s"), ErrorBuilder.ToString());
+		}
+		
+		if (NDCCVars::MissingNDCAccessorMode == 2)
+		{
+			ensureMsgf(false, TEXT("NDC Variable is missing but required."));
+		}
+		else if (NDCCVars::MissingNDCAccessorMode == 3)
+		{
+			ensureAlwaysMsgf(false, TEXT("NDC Variable is missing but required."));
+		}
+	}
+}
+
+//////////////////////////////////////////////////////////////////////////
+//FNDCWriterBase - Base class for user code NDCWriters. See full description in header.
+
+bool FNDCWriterBase::BeginWrite(const UObject* WorldContextObject, const UNiagaraDataChannel* DataChannel, const FNiagaraDataChannelSearchParameters& SearchParams, int32 InCount, bool bVisibleToGame, bool bVisibleToCPU, bool bVisibleToGPU)
+{
+	ensureMsgf(Data.IsValid() == false, TEXT("NDC Writer with non-null data on BeginWrite. Possible EndWrite was not called after previous write."));
+
+	const FNiagaraDataChannelLayoutInfoPtr NDCLayout = DataChannel->GetLayoutInfo();
+
+	if (CachedLayout != NDCLayout)
+	{
+		Init(DataChannel);
+		CachedLayout = NDCLayout;
+	}
+
+#if DEBUG_NDC_ACCESS
+	for (FNDCVarAccessorBase* Var : VariableAccessors)
+	{
+		Var->ValidateLayout();
+	}
+#endif
+
+	check(IsInGameThread());
+	UWorld* World = (WorldContextObject != nullptr) ? GEngine->GetWorldFromContextObject(WorldContextObject, EGetWorldErrorMode::LogAndReturnNull) : nullptr;
+	if (World && InCount > 0 && CachedLayout.IsValid())
+	{
+		if (FNiagaraWorldManager* WorldMan = FNiagaraWorldManager::Get(World))
+		{
+			if (UNiagaraDataChannelHandler* Handler = WorldMan->GetDataChannelManager().FindDataChannelHandler(DataChannel))
+			{
+				if (FNiagaraDataChannelDataPtr DestData = Handler->FindData(SearchParams, ENiagaraResourceAccess::WriteOnly))
+				{
+					//Attempt to use an existing cached dest data.
+					Data = DestData->GetGameDataForWriteGT(InCount, bVisibleToGame, bVisibleToCPU, bVisibleToGPU);
+
+					//We are potentially writing into the middle of an existing buffer here so track base index and count.
+					StartIndex = Data->Num();
+					Count = InCount;
+					Data->SetNum(StartIndex + InCount);
+
+					return true;
+				}
+			}
+		}
+	}
+	return false;
+}
+
+void FNDCWriterBase::EndWrite()
+{
+	Data = nullptr;
+	StartIndex = INDEX_NONE;
+	Count = INDEX_NONE;
+}
+
+//////////////////////////////////////////////////////////////////////////
+//NDCReaderBase - Base class for user code NDCReaders. See full description in header.
+
+bool FNDCReaderBase::BeginRead(const UObject* WorldContextObject, const UNiagaraDataChannel* DataChannel, const FNiagaraDataChannelSearchParameters& SearchParams, bool bInPreviousFrame)
+{
+	ensureMsgf(Data.IsValid() == false, TEXT("NDC Reader with non-null data on BeginRead. Possible EndRead was not called after previous read."));
+
+	const FNiagaraDataChannelLayoutInfoPtr NDCLayout = DataChannel->GetLayoutInfo();
+
+	if (CachedLayout != NDCLayout)
+	{
+		Init(DataChannel);
+		CachedLayout = NDCLayout;
+	}
+
+#if DEBUG_NDC_ACCESS
+	for (FNDCVarAccessorBase* Var : VariableAccessors)
+	{
+		Var->ValidateLayout();
+	}
+#endif
+
+	check(IsInGameThread());
+	UWorld* World = (WorldContextObject != nullptr) ? GEngine->GetWorldFromContextObject(WorldContextObject, EGetWorldErrorMode::LogAndReturnNull) : nullptr;
+	if (World && CachedLayout.IsValid())
+	{
+		if (FNiagaraWorldManager* WorldMan = FNiagaraWorldManager::Get(World))
+		{
+			if (UNiagaraDataChannelHandler* Handler = WorldMan->GetDataChannelManager().FindDataChannelHandler(DataChannel))
+			{
+				if (FNiagaraDataChannelDataPtr NDCData = Handler->FindData(SearchParams, ENiagaraResourceAccess::ReadOnly))
+				{
+					Data = NDCData->GetGameData()->AsShared();
+					bPreviousFrame = bInPreviousFrame;
+					return true;
+				}
+			}
+		}
+	}
+	return false;
+}
+
+void FNDCReaderBase::EndRead()
+{
+	Data = nullptr;
+	bPreviousFrame = false;
+}
+
+int32 FNDCReaderBase::Num()const
+{
+	if(Data)
+	{
+		return bPreviousFrame ? Data->PrevNum() : Data->Num();
+	}
+	return INDEX_NONE;
+}
\ No newline at end of file
diff --git a/Engine/Plugins/FX/Niagara/Source/Niagara/Private/NiagaraModule.cpp b/Engine/Plugins/FX/Niagara/Source/Niagara/Private/NiagaraModule.cpp
--- a/Engine/Plugins/FX/Niagara/Source/Niagara/Private/NiagaraModule.cpp
+++ b/Engine/Plugins/FX/Niagara/Source/Niagara/Private/NiagaraModule.cpp
@@ -318,7 +318,6 @@ void INiagaraModule::StartupModule()
 	LLM_SCOPE(ELLMTag::Niagara);
 	FNiagaraTypeDefinition::Init();
 	FNiagaraTypeRegistry::ProcessRegistryQueue();
-	FNiagaraTypeHelper::RegisterLWCTypes();
 
 	if (IsRunningCookCommandlet())
 	{
diff --git a/Engine/Plugins/FX/Niagara/Source/Niagara/Private/NiagaraTypes.cpp b/Engine/Plugins/FX/Niagara/Source/Niagara/Private/NiagaraTypes.cpp
--- a/Engine/Plugins/FX/Niagara/Source/Niagara/Private/NiagaraTypes.cpp
+++ b/Engine/Plugins/FX/Niagara/Source/Niagara/Private/NiagaraTypes.cpp
@@ -744,19 +744,13 @@ FNiagaraTypeDefinition FNiagaraTypeHelper::GetLWCType(const FNiagaraTypeDefiniti
 
 void FNiagaraTypeHelper::InitLWCTypes()
 {
-	if (Vector2DDef.IsValid() == false)
-	{
-		UPackage* CoreUObjectPkg = FindObjectChecked<UPackage>(nullptr, TEXT("/Script/CoreUObject"));
-		Vector2DDef = FNiagaraTypeDefinition(FindObjectChecked<UScriptStruct>(CoreUObjectPkg, TEXT("Vector2D")), FNiagaraTypeDefinition::EAllowUnfriendlyStruct::Allow);
-		VectorDef = FNiagaraTypeDefinition(FindObjectChecked<UScriptStruct>(CoreUObjectPkg, TEXT("Vector")), FNiagaraTypeDefinition::EAllowUnfriendlyStruct::Allow);
-		Vector4Def = FNiagaraTypeDefinition(FindObjectChecked<UScriptStruct>(CoreUObjectPkg, TEXT("Vector4")), FNiagaraTypeDefinition::EAllowUnfriendlyStruct::Allow);
-		QuatDef = FNiagaraTypeDefinition(FindObjectChecked<UScriptStruct>(CoreUObjectPkg, TEXT("Quat")), FNiagaraTypeDefinition::EAllowUnfriendlyStruct::Allow);
-		DoubleDef = FNiagaraTypeDefinition(FNiagaraDouble::StaticStruct(), FNiagaraTypeDefinition::EAllowUnfriendlyStruct::Allow);
-	}
-}
+	UPackage* CoreUObjectPkg = FindObjectChecked<UPackage>(nullptr, TEXT("/Script/CoreUObject"));
+	Vector2DDef = FNiagaraTypeDefinition(FindObjectChecked<UScriptStruct>(CoreUObjectPkg, TEXT("Vector2D")), FNiagaraTypeDefinition::EAllowUnfriendlyStruct::Allow);
+	VectorDef = FNiagaraTypeDefinition(FindObjectChecked<UScriptStruct>(CoreUObjectPkg, TEXT("Vector")), FNiagaraTypeDefinition::EAllowUnfriendlyStruct::Allow);
+	Vector4Def = FNiagaraTypeDefinition(FindObjectChecked<UScriptStruct>(CoreUObjectPkg, TEXT("Vector4")), FNiagaraTypeDefinition::EAllowUnfriendlyStruct::Allow);
+	QuatDef = FNiagaraTypeDefinition(FindObjectChecked<UScriptStruct>(CoreUObjectPkg, TEXT("Quat")), FNiagaraTypeDefinition::EAllowUnfriendlyStruct::Allow);
+	DoubleDef = FNiagaraTypeDefinition(FNiagaraDouble::StaticStruct(), FNiagaraTypeDefinition::EAllowUnfriendlyStruct::Allow);
 
-void FNiagaraTypeHelper::RegisterLWCTypes()
-{
 	FNiagaraTypeRegistry::Register(Vector2DDef, ENiagaraTypeRegistryFlags::None);
 	FNiagaraTypeRegistry::Register(VectorDef, ENiagaraTypeRegistryFlags::None);
 	FNiagaraTypeRegistry::Register(Vector4Def, ENiagaraTypeRegistryFlags::None);
@@ -764,36 +758,6 @@ void FNiagaraTypeHelper::RegisterLWCTypes()
 	FNiagaraTypeRegistry::Register(DoubleDef, ENiagaraTypeRegistryFlags::None);
 }
 
-FNiagaraTypeDefinition FNiagaraTypeHelper::GetVector2DDef()
-{
-	InitLWCTypes();
-	return Vector2DDef;
-}
-
-FNiagaraTypeDefinition FNiagaraTypeHelper::GetVectorDef()
-{
-	InitLWCTypes();
-	return VectorDef;
-}
-
-FNiagaraTypeDefinition FNiagaraTypeHelper::GetVector4Def()
-{
-	InitLWCTypes();
-	return Vector4Def;
-}
-
-FNiagaraTypeDefinition FNiagaraTypeHelper::GetQuatDef()
-{
-	InitLWCTypes();
-	return QuatDef;
-}
-
-FNiagaraTypeDefinition FNiagaraTypeHelper::GetDoubleDef()
-{
-	InitLWCTypes();
-	return DoubleDef;
-}
-
 UScriptStruct* FNiagaraTypeHelper::FindNiagaraFriendlyTopLevelStruct(UScriptStruct* InStruct, ENiagaraStructConversion StructConversion)
 {
 	if (!InStruct)
diff --git a/Engine/Plugins/FX/Niagara/Source/Niagara/Public/NiagaraDataChannel.h b/Engine/Plugins/FX/Niagara/Source/Niagara/Public/NiagaraDataChannel.h
--- a/Engine/Plugins/FX/Niagara/Source/Niagara/Public/NiagaraDataChannel.h
+++ b/Engine/Plugins/FX/Niagara/Source/Niagara/Public/NiagaraDataChannel.h
@@ -138,6 +138,8 @@ struct FNiagaraDataChannelLayoutInfo : public TSharedFromThis<FNiagaraDataChanne
 	const FNiagaraDataSetCompiledData& GetDataSetCompiledDataGPU()const { return CompiledDataGPU; }
 	const FNiagaraDataChannelGameDataLayout& GetGameDataLayout()const { return GameDataLayout; }
 
+	/** If true, we keep our previous frame's data. Some users will prefer a frame of latency to tick dependency. */
+	bool KeepPreviousFrameData() const { return bKeepPreviousFrameData; }
 private:
 
 	/**
@@ -149,6 +151,8 @@ private:
 
 	/** Layout information for any data stored at the "Game" level. i.e. From game code/BP. AoS layout and LWC types. */
 	FNiagaraDataChannelGameDataLayout GameDataLayout;
+
+	bool bKeepPreviousFrameData = false;
 };
 
 DECLARE_MULTICAST_DELEGATE_OneParam(FOnDataChannelCreated, const UNiagaraDataChannel*);
@@ -179,11 +183,12 @@ public:
 	/** Create the appropriate handler object for this data channel. */
 	NIAGARA_API virtual UNiagaraDataChannelHandler* CreateHandler(UWorld* OwningWorld) const PURE_VIRTUAL(UNiagaraDataChannel::CreateHandler, {return nullptr;} );
 	
-	const FNiagaraDataChannelLayoutInfoPtr GetLayoutInfo()const;
+	NIAGARA_API const FNiagaraDataChannelLayoutInfoPtr GetLayoutInfo()const;
 
+	//Creates a new GameData for this NDC.
 	NIAGARA_API FNiagaraDataChannelGameDataPtr CreateGameData() const;
 
-	bool IsValid() const;
+	NIAGARA_API bool IsValid() const;
 
 	#if WITH_NIAGARA_DEBUGGER
 	void SetVerboseLogging(bool bValue){ bVerboseLogging = bValue; }
diff --git a/Engine/Plugins/FX/Niagara/Source/Niagara/Public/NiagaraDataChannelAccessor.h b/Engine/Plugins/FX/Niagara/Source/Niagara/Public/NiagaraDataChannelAccessor.h
--- a/Engine/Plugins/FX/Niagara/Source/Niagara/Public/NiagaraDataChannelAccessor.h
+++ b/Engine/Plugins/FX/Niagara/Source/Niagara/Public/NiagaraDataChannelAccessor.h
@@ -88,6 +88,8 @@ private:
 	/** Local data buffers we're writing into. */
 	FNiagaraDataChannelGameDataPtr Data = nullptr;
 
+	//Starting index into our dest data. 
+	int32 StartIndex = INDEX_NONE;
 
 public:
 	template<typename T>
@@ -153,3 +155,290 @@ public:
 
 };
 
+
+
+//////////////////////////////////////////////////////////////////////////
+
+#define DEBUG_NDC_ACCESS !UE_BUILD_SHIPPING && !UE_BUILD_TEST
+
+/*
+* Below are several utility classes for reading from and writing to Niagara Data Channels from C++ code.
+* Usage:
+* Create a child of FNDCReaderBase or FNDCWriterBase in your own code.
+* Add variable members to this class via the NDCVarWriter or NDCVarReader macros.
+* The types and names used here should correspond to names and types of the data inside the NDC.
+* Variables defined with FNDCVarWriter(...) will be expected and will trigger a warning if missing from the target NDC.
+* Variables defined with FNDCVarWriterOptional(...) will be safely ignored with minimal overhead if missing from the NDC.
+* Note: Take care to provide positions as FNiagaraPosition types rather than Vector so that they are interpreted by Niagara correctly as LWC positions.
+* Note: Enum types can be used to target integer types inside the NDC.
+* 
+* 1.
+* An examples Writer struct:
+* struct FNDCExampleWriter : public FNDCWriterBase
+* {
+*		FNDCVarWriter(FNiagaraPosition, Position);
+*		FNDCVarWriter(FVector, Velocity);
+* 		FNDCVarWriter(bool, bSomeValue);* 
+* };
+* 
+* 2.
+* Make an instance of this class and initialize it with your desired Data Channel.
+* It can be re-initialized with a new Data Channel in future if needed.
+* Alternatively, the NDCWriter will (re)initialize itself if the Data Channel passed into BeginWrite is different (or has changed).
+* class FYourUserClass
+* {
+*	...
+*	FNDCExampleWriter ExampleWriter;
+*	void Init(UNiagaraDataChannel* NDC)
+*	{
+*		ExampleWriter.Init(NDC);
+*	}
+*	...
+* };
+* 
+* 3.
+* Use this writer object to write data into the Data Channel.
+* First you must call BeginWrite() to initialize the writer for the current search parameters.
+* 
+* void FYouUserClass::YourFunction()
+* {
+*	FNiagaraDataChannelSearchParameters SearchParameters;
+*	SearchParameters.Location = MyLocation; //We find the correct NDC data to write into by using this position.
+*	//SearchParameters.Owner = MyComponent; //Alternatively we can provide a component for either the location or some NDC types may pull other data from the component to find the correct NDC data.
+* 
+*	int32 NumItems = YourData.Num();
+*	if(ExampleWriter.BeginWrite(MyWorld, DataChannel, SearchParameters, NumItems, bVisibleToGame, bVisibleToCPU, bVisibleToGPU)
+*	{
+*		for(int32 i=0; i < NumItems; ++i)
+*		{	
+* 			//The macro creates named write functions for you in your writer class.
+			ExampleWriter.WritePosition(YourData[i].Position);
+			ExampleWriter.WriteVelocity(YourData[i].Velocity);
+			ExampleWriter.WritebSomeValue(YourData[i].bBoolValue);			
+*		}
+*	}
+*	//Finally call end write to release and references the writer holds to it's destination data etc.
+*	ExampleWriter.EndWrite();
+* 
+*	Alternatively you can use the scoped writer to use this much as above but without needing to remember to call EndWrite();
+*	FNDCScopedWriter ScopedNDCWriter(ExampleWriter);
+*	//We use the -> operator for the scoped writer to access it's internal NDCWriter.
+*	if(ScopedNDCWriter->BeginWrite(MyWorld, DataChannel, SearchParameters, NumItems, bVisibleToGame, bVisibleToCPU, bVisibleToGPU)
+*	{
+*		for(int32 i=0; i < NumItems; ++i)
+*		{
+			ScopedNDCWriter->WritePosition(YourData[i].Position);
+			ScopedNDCWriter->WriteVelocity(YourData[i].Velocity);
+			ScopedNDCWriter->WritebSomeValue(YourData[i].bBoolValue);
+*		}
+*	}	
+* }	//EndWrite is called automatically when we exit the scope.
+* 
+* Usage for Readers follows the same patterns as above.
+*/
+
+struct FNDCVarAccessorBase;
+
+/*
+Base class for NDC accessor utilities. Handles common bookkeeping for NDCWriter and NDCReader utility structs. 
+See full description above.
+*/
+struct FNDCAccessorBase
+{
+friend struct FNDCVarAccessorBase;
+
+protected:
+	//Cached Layout for the NDC data. If the data channel changes layout then this will trigger a re-init of this accessor.
+	FNiagaraDataChannelLayoutInfoPtr CachedLayout;
+
+	//Array of our variable accessors, allow us to easily update all variable layout info.
+	TArray<FNDCVarAccessorBase*> VariableAccessors;
+
+	//Base index from which to access data in the found NDC buffer.
+	int32 StartIndex = INDEX_NONE;
+
+public:
+
+	NIAGARA_API bool IsInitialized()const { return CachedLayout.IsValid() && !CachedLayout.IsUnique(); }
+
+	//Initialize the writer and update cached layout information.
+	NIAGARA_API void Init(const UNiagaraDataChannel* DataChannel);
+};
+
+
+/*
+Utility for accessing a specific NDC variable data from C++.
+Caches offset for variable in NDC layout to speed up access.
+Owning NDCReader/Writer will re-init if layout changes.
+See full description above.
+*/
+struct FNDCVarAccessorBase
+{
+	friend FNDCAccessorBase;
+protected:
+	FNiagaraVariableBase Variable;
+
+	int32 VarOffset:31;
+	int32 bIsRequired:1;
+
+public:
+
+	NIAGARA_API FNDCVarAccessorBase(FNDCAccessorBase& Owner, FNiagaraVariableBase InVariable, bool bInIsRequired);
+	NIAGARA_API void Init(const UNiagaraDataChannel* DataChannel);
+	NIAGARA_API bool IsValid()const { return VarOffset != INDEX_NONE; }
+
+#if DEBUG_NDC_ACCESS
+	FNiagaraDataChannelLayoutInfoPtr DebugCachedLayout;
+	TWeakObjectPtr<const UNiagaraDataChannel> WeakNDC;
+	void ValidateLayout()const { checkf(WeakNDC->IsValid() && WeakNDC->GetLayoutInfo() == DebugCachedLayout, TEXT("NDC Variable being accessed with stale layout info.")); }
+#else
+	FORCEINLINE void ValidateLayout()const {}
+#endif
+
+	template<typename T>
+	FORCEINLINE bool WriteInternal(const FNiagaraDataChannelGameDataPtr& Data, int32 Index, const T& InValue)
+	{
+		if (VarOffset != INDEX_NONE)
+		{
+			return Data->Write<T>(VarOffset, Index, InValue);
+		}
+		return false;
+	}	
+	
+	template<typename T>
+	FORCEINLINE bool ReadInternal(const FNiagaraDataChannelGameDataPtr& Data, int32 Index, T& OutValue, bool bPreviousFrame)
+	{
+		if (VarOffset != INDEX_NONE)
+		{
+			return Data->Read<T>(VarOffset, Index, OutValue, bPreviousFrame);
+		}
+		return false;
+	}
+};
+
+/** Base class for NDC Reader Utility classes. See full description above. */
+struct NIAGARA_API FNDCWriterBase : public FNDCAccessorBase
+{
+protected:
+	//Current NDC target data we're accessing.
+	FNiagaraDataChannelGameDataPtr Data;
+
+	//StartIndex we're writing into. Can be non-zero in cases where we're writing into the middle of an existing buffer.
+	int32 StartIndex = INDEX_NONE;
+
+	//Number of items we're expecting/allowed to write.
+	int32 Count = INDEX_NONE;
+
+public:
+
+	//Finds the correct target NDC data using the current SearchParameters.
+	//Allocates enough space in the target NDC data.
+	//Refreshes internal layout info if needed.
+	bool BeginWrite(const UObject* WorldContextObject, const UNiagaraDataChannel* DataChannel, const FNiagaraDataChannelSearchParameters& SearchParams, int32 Count, bool bVisibleToGame, bool bVisibleToCPU, bool bVisibleToGPU);
+
+	//Finalizes the write and free/reset any data.
+	void EndWrite();
+};
+
+/** Utility class for simplified writing of variables to a Niagara Data Channel. See full description above. */
+template<typename TNDCWriterType> 
+struct FNDCScopedWriter
+{
+private:
+	TNDCWriterType& Writer;
+public:
+	FNDCScopedWriter(TNDCWriterType& InWriter) : Writer(InWriter){ }
+	~FNDCScopedWriter()	{ Writer.EndWrite(); }
+	TNDCWriterType* operator->(){ return &Writer; }
+};
+
+/** Base class for NDC Reader Utility classes. See full description above. */
+struct NIAGARA_API FNDCReaderBase : public FNDCAccessorBase
+{
+protected:
+	FNiagaraDataChannelGameDataPtr Data;
+	bool bPreviousFrame = false;
+
+public:
+
+	//Finds the correct source NDC data using the current SearchParameters.
+	bool BeginRead(const UObject* WorldContextObject, const UNiagaraDataChannel* DataChannel, const FNiagaraDataChannelSearchParameters& SearchParams, bool bInPreviousFrame);
+
+	//Finalizes the read and free/reset any data.
+	void EndRead();
+
+	int32 Num()const;
+
+};
+
+/** Utility class for simplified reading of variables from a Niagara Data Channel. See full description above. */
+template<typename TNDCReaderType>
+struct FNDCScopedReader
+{
+private:
+	FNDCReaderBase& Reader;
+public:
+	FNDCScopedReader(TNDCReaderType& InReader) : Reader(InReader){ }
+	~FNDCScopedReader() { Reader.EndRead();	}
+	TNDCReaderType* operator->() { return &Reader; }
+};
+
+
+/** Utility class for writing variables to a Niagara Data Channel. See full description above. */
+template<typename T>
+struct FNDCVarWriter : public FNDCVarAccessorBase
+{
+public:
+	FNDCVarWriter(FNDCWriterBase& Owner, FName VarName, bool bInIsRequired)
+		: FNDCVarAccessorBase(Owner, FNiagaraVariableBase(FNiagaraTypeHelper::GetTypeDef<T>(), VarName), bInIsRequired)
+	{}
+
+	FORCEINLINE bool Write(const FNiagaraDataChannelGameDataPtr& Data, int32 Index, const typename TNiagaraTypeHelper<T>::TLWCType& InValue, int32 StartIndex, int32 Count)
+	{
+		if(Index >= 0 && Index < Count)
+		{
+			return WriteInternal<typename TNiagaraTypeHelper<T>::TLWCType>(Data, StartIndex + Index, InValue);
+		}	
+		return false;
+	}
+};
+
+/** Utility class for reading variables from a Niagara Data Channel. See full description above. */
+template<typename T>
+struct FNDCVarReader : public FNDCVarAccessorBase
+{
+public:
+	FNDCVarReader(FNDCReaderBase& Owner, FName VarName, bool bInIsRequired)
+		: FNDCVarAccessorBase(Owner, FNiagaraVariableBase(FNiagaraTypeHelper::GetTypeDef<T>(), VarName), bInIsRequired)
+	{}
+
+	FORCEINLINE bool Read(const FNiagaraDataChannelGameDataPtr& Data, int32 Index, typename TNiagaraTypeHelper<T>::TLWCType& OutValue, bool bPreviousFrame)
+	{
+		return ReadInternal<typename TNiagaraTypeHelper<T>::TLWCType>(Data, Index, OutValue, bPreviousFrame);
+	}
+};
+
+
+/** Defines an NDC Variable that we will write from C++. This variable is required and will generate errors if not present in the NDC. */
+#define NDCVarWriter(VarType, VarName)\
+FNDCVarWriter<VarType> VarName##Writer = FNDCVarWriter<VarType>(*this, #VarName, true);\
+bool Write##VarName(int32 Index, const TNiagaraTypeHelper<VarType>::TLWCType& InData){ return VarName##Writer.Write(Data, Index, InData, StartIndex, Count); }\
+const FNDCVarWriter<VarType>& Get##VarName##Writer()const{ return VarName##Writer; }
+
+/** Defines an NDC Variable that we will write from C++. This variable is optional and will be ignored  if not present in the NDC. */
+#define NDCVarWriterOptional(VarType, VarName)\
+FNDCVarWriter<VarType> VarName##Writer = FNDCVarWriter<VarType>(*this, #VarName, false);\
+bool Write##VarName(int32 Index, const TNiagaraTypeHelper<VarType>::TLWCType& InData){ return VarName##Writer.Write(Data, Index, InData, StartIndex, Count); }\
+const FNDCVarWriter<VarType>& Get##VarName##Writer()const{ return VarName##Writer; }
+
+/** Defines an NDC Variable that we will read from C++. This variable is required and will generate errors if not present in the NDC. */
+#define NDCVarReader(VarType, VarName)\
+FNDCVarReader<VarType> VarName##Reader = FNDCVarReader<VarType>(*this, #VarName, true);\
+bool Read##VarName(int32 Index, TNiagaraTypeHelper<VarType>::TLWCType& OutData){ return VarName##Reader.Read(Data, Index, OutData, bPreviousFrame); }\
+const FNDCVarReader<VarType>& Get##VarName##Reader()const { return VarName##Reader; }
+
+/** Defines an NDC Variable that we will read from C++. This variable is optional and will be ignored  if not present in the NDC. */
+#define NDCVarReaderOptional(VarType, VarName)\
+FNDCVarReader<VarType> VarName##Reader = FNDCVarReader<VarType>(*this, #VarName, false);\
+bool Read##VarName(int32 Index, TNiagaraTypeHelper<VarType>::TLWCType& OutData){ return VarName##Reader.Read(Data, Index, OutData, bPreviousFrame); }\
+const FNDCVarReader<VarType>& Get##VarName##Reader()const { return VarName##Reader; }
diff --git a/Engine/Plugins/FX/Niagara/Source/Niagara/Public/NiagaraDataChannelCommon.h b/Engine/Plugins/FX/Niagara/Source/Niagara/Public/NiagaraDataChannelCommon.h
--- a/Engine/Plugins/FX/Niagara/Source/Niagara/Public/NiagaraDataChannelCommon.h
+++ b/Engine/Plugins/FX/Niagara/Source/Niagara/Public/NiagaraDataChannelCommon.h
@@ -109,10 +109,16 @@ struct FNiagaraDataChannelData final : public TSharedFromThis<FNiagaraDataChanne
 	//Return true if data has been written to this NDC data for the current frame.
 	bool HasData()const;
 
+	//Returns a game data buffer into which we can write Count valuse on the Game Thread.
+	FNiagaraDataChannelGameDataPtr GetGameDataForWriteGT(int32 Count, bool bVisibleToGame, bool bVisibleToCPU, bool bVisibleToGPU);
+
 private:
 
 	void CreateRenderThreadProxy(UNiagaraDataChannelHandler* Owner);
 
+	void FlushPendingGameData(int32 Index);
+	void FlushAllPendingGameData();
+
 	/** DataChannel data accessible from Game/BP. AoS Layout. LWC types. */
 	FNiagaraDataChannelGameDataPtr GameData;
 
@@ -152,7 +158,10 @@ private:
 	/** 
 	Track number of explicitly registered readers that spawn GPU particles from this data.
 	If we're spawning GPU particles using the CPU data (Spawn Conditional etc) then we have to send all CPU data to the GPU every frame.
-	Can possibly extend this to be a more automatic, registration based approach to shipping NDC data around rather than exmplicit flags on write.
+	Can possibly extend this to be a more automatic, registration based approach to shipping NDC data around rather than explicit flags on write.
 	*/
 	std::atomic<int32> NumGPUSpawningReaders;
+
+	//We keep a set of incoming game data for all flag combinations and accumulate data into these rather than keeping all as separate
+	TArray<FNiagaraDataChannelGameDataPtr, TInlineAllocator<8>> PendingDestGameData;
 };
diff --git a/Engine/Plugins/FX/Niagara/Source/Niagara/Public/NiagaraDataChannelPublic.h b/Engine/Plugins/FX/Niagara/Source/Niagara/Public/NiagaraDataChannelPublic.h
--- a/Engine/Plugins/FX/Niagara/Source/Niagara/Public/NiagaraDataChannelPublic.h
+++ b/Engine/Plugins/FX/Niagara/Source/Niagara/Public/NiagaraDataChannelPublic.h
@@ -173,6 +173,7 @@ public:
 private:
 	struct FChannelWriteRequest
 	{
+		TWeakObjectPtr<const UNiagaraDataChannel> Channel;
 		TSharedPtr<FNiagaraDataChannelGameData> Data;
 		bool bVisibleToGame = false;
 		bool bVisibleToCPUSims = false;
@@ -215,7 +216,17 @@ struct FNiagaraDataChannelVariableBuffer
 		}
 	}
 
-	void Empty() { Data.Empty(); }
+	void Empty() 
+	{
+		PrevData.Empty();
+		Data.Empty(); 
+	}
+
+	void Reset()
+	{
+		PrevData.Reset();
+		Data.Reset();
+	}
 
 	void BeginFrame(bool bKeepPrevious)
 	{
@@ -287,21 +298,23 @@ private:
 	TArray<FNiagaraDataChannelVariableBuffer> VariableData;
 
 	int32 NumElements = 0;
+	int32 MaxElements = 0;
 	int32 PrevNumElements = 0;
 
-	TWeakObjectPtr<const UNiagaraDataChannel> DataChannel;
-
 	//Keep reference to the layout this data was built with.
 	FNiagaraDataChannelLayoutInfoPtr LayoutInfo;
 
 public:
-
-	NIAGARA_API void Init(const UNiagaraDataChannel* InDataChannel);
+	
+	FNiagaraDataChannelGameData(const FNiagaraDataChannelLayoutInfoPtr& InLayoutInfo);
+	NIAGARA_API void Init(const FNiagaraDataChannelLayoutInfoPtr& InLayoutInfo);
 	NIAGARA_API void Empty();
+	NIAGARA_API void Reset();
 	NIAGARA_API void BeginFrame();
 	NIAGARA_API void SetNum(int32 NewNum);
 	NIAGARA_API void Reserve(int32 NewNum);
 	int32 Num() const { return NumElements; }
+	int32 Max() const { return MaxElements; }
 	int32 PrevNum()const { return PrevNumElements; }
 
 	int32 Add(int32 Count=1)
@@ -319,8 +332,30 @@ public:
 
 	NIAGARA_API void AppendFromDataSet(const FNiagaraDataBuffer* SrcBuffer, FVector3f SimulationLwcTile);
 
-	const UNiagaraDataChannel* GetDataChannel()const { return DataChannel.Get(); }
+	const FNiagaraDataChannelLayoutInfoPtr& GetLayoutInfo()const { return LayoutInfo; }
 
 	const TConstArrayView<FNiagaraDataChannelVariableBuffer> GetVariableBuffers()const { return VariableData; }
 	void SetFromSimCache(const FNiagaraVariableBase& SourceVar, TConstArrayView<uint8> Data, int32 Size);
+
+	template<typename T>
+	bool Write(int32 VarIndex, int32 Index, const T& InValue)
+	{
+		if(VariableData.IsValidIndex(VarIndex))
+		{
+			FNiagaraDataChannelVariableBuffer& VarBuffer = VariableData[VarIndex];
+			return VarBuffer.Write<T>(Index, InValue);
+		}
+		return false;
+	}
+
+	template<typename T>
+	bool Read(int32 VarIndex, int32 Index, T& OutValue, bool bPreviousFrame)
+	{
+		if (VariableData.IsValidIndex(VarIndex))
+		{
+			FNiagaraDataChannelVariableBuffer& VarBuffer = VariableData[VarIndex];
+			return VarBuffer.Read<T>(Index, OutValue, bPreviousFrame);
+		}
+		return false;
+	}
 };
diff --git a/Engine/Plugins/FX/Niagara/Source/Niagara/Public/NiagaraDataChannel_Islands.h b/Engine/Plugins/FX/Niagara/Source/Niagara/Public/NiagaraDataChannel_Islands.h
--- a/Engine/Plugins/FX/Niagara/Source/Niagara/Public/NiagaraDataChannel_Islands.h
+++ b/Engine/Plugins/FX/Niagara/Source/Niagara/Public/NiagaraDataChannel_Islands.h
@@ -78,9 +78,6 @@ private:
 	
 	/** The underlying storage for this island. */
 	FNiagaraDataChannelDataPtr Data = nullptr;
-	
-	/** Publish requests from game code/BP or other Niagara Systems wishing to write data into the data channel for this island. */
-	TArray<FNiagaraDataChannelPublishRequest> PublishRequests;
 };
 
 template<>
diff --git a/Engine/Plugins/FX/Niagara/Source/Niagara/Public/NiagaraTypes.h b/Engine/Plugins/FX/Niagara/Source/Niagara/Public/NiagaraTypes.h
--- a/Engine/Plugins/FX/Niagara/Source/Niagara/Public/NiagaraTypes.h
+++ b/Engine/Plugins/FX/Niagara/Source/Niagara/Public/NiagaraTypes.h
@@ -438,15 +438,29 @@ public:
 	These cannot be properly serialized currently so must be constructed at runtime from their Niagara Simulation types.
 	*/
 	static NIAGARA_API FNiagaraTypeDefinition GetLWCType(const FNiagaraTypeDefinition& InType);
-	
-	static NIAGARA_API FNiagaraTypeDefinition GetVector2DDef();
-	static NIAGARA_API FNiagaraTypeDefinition GetVectorDef();
-	static NIAGARA_API FNiagaraTypeDefinition GetVector4Def();
-	static NIAGARA_API FNiagaraTypeDefinition GetQuatDef();
-	static NIAGARA_API FNiagaraTypeDefinition GetDoubleDef();
+
+	static const FNiagaraTypeDefinition& GetVector2DDef() { return Vector2DDef; }
+	static const FNiagaraTypeDefinition& GetVectorDef() { return VectorDef; }
+	static const FNiagaraTypeDefinition& GetVector4Def() { return Vector4Def; }
+	static const FNiagaraTypeDefinition& GetQuatDef() { return QuatDef; }
+	static const FNiagaraTypeDefinition& GetDoubleDef() { return DoubleDef; }
 
 	static NIAGARA_API void InitLWCTypes();
-	static NIAGARA_API void RegisterLWCTypes();
+
+	//Utility functions for accessing typedefs for specific C++ types.
+
+	template <typename T, std::enable_if_t<!std::is_enum_v<T>, int> = 0>
+	static FNiagaraTypeDefinition GetTypeDef();
+	template <typename T, std::enable_if_t<!std::is_enum_v<T>, int> = 0>
+	static FNiagaraTypeDefinition GetSimTypeDef();
+
+	//Special case for any enum types to access the correct TypeDefs.
+	template <typename T, std::enable_if_t<std::is_enum_v<T>, int> = 0>
+	static FNiagaraTypeDefinition GetTypeDef() ;
+	template <typename T, std::enable_if_t<std::is_enum_v<T>, int> = 0>
+	static FNiagaraTypeDefinition GetSimTypeDef() ;
+
+
 private:
 	struct FRemapEntry
 	{
@@ -470,11 +484,11 @@ private:
 	static std::atomic<bool> RemapTableDirty;
 
 	/** LWC Counterparts to Niagara's base simulation types. */
-	static FNiagaraTypeDefinition Vector2DDef;
-	static FNiagaraTypeDefinition VectorDef;
-	static FNiagaraTypeDefinition Vector4Def;
-	static FNiagaraTypeDefinition QuatDef;
-	static FNiagaraTypeDefinition DoubleDef;
+	NIAGARA_API static FNiagaraTypeDefinition Vector2DDef;
+	NIAGARA_API static FNiagaraTypeDefinition VectorDef;
+	NIAGARA_API static FNiagaraTypeDefinition Vector4Def;
+	NIAGARA_API static FNiagaraTypeDefinition QuatDef;
+	NIAGARA_API static FNiagaraTypeDefinition DoubleDef;
 
 	static FString ConvertedSWCStructSuffix;
 };
@@ -1816,3 +1830,78 @@ struct FVersionedNiagaraEmitterWeakPtr
 	return GetTypeHash(Item.Emitter);
 }
 };
+
+//////////////////////////////////////////////////////////////////////////
+//Utilities for getting the correct LWC Niagara Type Def.
+
+template <typename T, std::enable_if_t<!std::is_enum_v<T>, int>>
+FNiagaraTypeDefinition FNiagaraTypeHelper::GetTypeDef() { return FNiagaraTypeDefinition(StaticStruct<T>()); }
+
+template <typename T, std::enable_if_t<!std::is_enum_v<T>, int>>
+FNiagaraTypeDefinition FNiagaraTypeHelper::GetSimTypeDef() { return FNiagaraTypeDefinition(StaticStruct<T>()); }
+
+//Special case for any enum types to access the correct TypeDefs.
+template <typename T, std::enable_if_t<std::is_enum_v<T>, int>>
+FNiagaraTypeDefinition FNiagaraTypeHelper::GetTypeDef()
+{
+	return FNiagaraTypeDefinition(StaticEnum<T>());
+}
+
+template <typename T, std::enable_if_t<std::is_enum_v<T>, int>>
+FNiagaraTypeDefinition FNiagaraTypeHelper::GetSimTypeDef()
+{
+	return FNiagaraTypeDefinition(StaticEnum<T>());
+}
+
+template<> inline FNiagaraTypeDefinition FNiagaraTypeHelper::GetTypeDef<FNiagaraBool>() { return FNiagaraTypeDefinition::GetBoolDef(); }
+template<> inline FNiagaraTypeDefinition FNiagaraTypeHelper::GetTypeDef<bool>() { return FNiagaraTypeDefinition::GetBoolDef(); }
+template<> inline FNiagaraTypeDefinition FNiagaraTypeHelper::GetTypeDef<int32>() { return FNiagaraTypeDefinition::GetIntDef(); }
+template<> inline FNiagaraTypeDefinition FNiagaraTypeHelper::GetTypeDef<float>() { return FNiagaraTypeHelper::GetDoubleDef(); }
+template<> inline FNiagaraTypeDefinition FNiagaraTypeHelper::GetTypeDef<double>() { return FNiagaraTypeHelper::GetDoubleDef(); }
+template<> inline FNiagaraTypeDefinition FNiagaraTypeHelper::GetTypeDef<FVector2D>() { return FNiagaraTypeHelper::GetVector2DDef(); }
+template<> inline FNiagaraTypeDefinition FNiagaraTypeHelper::GetTypeDef<FVector>() { return FNiagaraTypeHelper::GetVectorDef(); }
+template<> inline FNiagaraTypeDefinition FNiagaraTypeHelper::GetTypeDef<FVector4>() { return FNiagaraTypeHelper::GetVector4Def(); }
+template<> inline FNiagaraTypeDefinition FNiagaraTypeHelper::GetTypeDef<FQuat>() { return FNiagaraTypeHelper::GetQuatDef(); }
+template<> inline FNiagaraTypeDefinition FNiagaraTypeHelper::GetTypeDef<FNiagaraPosition>() { return FNiagaraTypeDefinition::GetPositionDef(); }
+template<> inline FNiagaraTypeDefinition FNiagaraTypeHelper::GetTypeDef<FLinearColor>() { return FNiagaraTypeDefinition::GetColorDef(); }
+template<> inline FNiagaraTypeDefinition FNiagaraTypeHelper::GetTypeDef<FNiagaraID>() { return FNiagaraTypeDefinition::GetIDDef(); }
+
+//Utilities for getting the correct Niagara Sim TypeDef used internally to Niagara.
+template<> inline FNiagaraTypeDefinition FNiagaraTypeHelper::GetSimTypeDef<FNiagaraBool>() { return FNiagaraTypeDefinition::GetBoolDef(); }
+template<> inline FNiagaraTypeDefinition FNiagaraTypeHelper::GetSimTypeDef<bool>() { return FNiagaraTypeDefinition::GetBoolDef(); }
+template<> inline FNiagaraTypeDefinition FNiagaraTypeHelper::GetSimTypeDef<int32>() { return FNiagaraTypeDefinition::GetIntDef(); }
+template<> inline FNiagaraTypeDefinition FNiagaraTypeHelper::GetSimTypeDef<float>() { return FNiagaraTypeDefinition::GetFloatDef(); }
+template<> inline FNiagaraTypeDefinition FNiagaraTypeHelper::GetSimTypeDef<double>() { return FNiagaraTypeDefinition::GetFloatDef(); }
+template<> inline FNiagaraTypeDefinition FNiagaraTypeHelper::GetSimTypeDef<FVector2D>() { return FNiagaraTypeDefinition::GetVec2Def(); }
+template<> inline FNiagaraTypeDefinition FNiagaraTypeHelper::GetSimTypeDef<FVector>() { return FNiagaraTypeDefinition::GetVec3Def(); }
+template<> inline FNiagaraTypeDefinition FNiagaraTypeHelper::GetSimTypeDef<FVector4>() { return FNiagaraTypeDefinition::GetVec4Def(); }
+template<> inline FNiagaraTypeDefinition FNiagaraTypeHelper::GetSimTypeDef<FQuat>() { return FNiagaraTypeDefinition::GetQuatDef(); }
+template<> inline FNiagaraTypeDefinition FNiagaraTypeHelper::GetSimTypeDef<FNiagaraPosition>() { return FNiagaraTypeDefinition::GetVec3Def(); }
+template<> inline FNiagaraTypeDefinition FNiagaraTypeHelper::GetSimTypeDef<FLinearColor>() { return FNiagaraTypeDefinition::GetColorDef(); }
+template<> inline FNiagaraTypeDefinition FNiagaraTypeHelper::GetSimTypeDef<FNiagaraID>() { return FNiagaraTypeDefinition::GetIDDef(); }
+
+//Utils for converting to Niagara Sim or LWC Types inside template code.
+template<typename T, typename Enable = void> 
+struct TNiagaraTypeHelper{ typedef T TSimType; typedef T TLWCType; };
+
+template<> struct TNiagaraTypeHelper<bool> { typedef FNiagaraBool TSimType; typedef FNiagaraBool TLWCType; };
+template<> struct TNiagaraTypeHelper<float> { typedef float TSimType; typedef double TLWCType; };
+template<> struct TNiagaraTypeHelper<double> { typedef float TSimType; typedef double TLWCType; };
+template<> struct TNiagaraTypeHelper<FVector2D> { typedef FVector2f TSimType; typedef FVector2D TLWCType; };
+template<> struct TNiagaraTypeHelper<FVector> { typedef FVector3f TSimType; typedef FVector TLWCType; };
+template<> struct TNiagaraTypeHelper<FVector4> { typedef FVector4f TSimType; typedef FVector4 TLWCType; };
+template<> struct TNiagaraTypeHelper<FQuat> { typedef FQuat4f TSimType; typedef FQuat TLWCType; };
+
+template<> struct TNiagaraTypeHelper<FVector2f> { typedef FVector2f TSimType; typedef FVector2D TLWCType; };
+template<> struct TNiagaraTypeHelper<FVector3f> { typedef FVector3f TSimType; typedef FVector TLWCType; };
+template<> struct TNiagaraTypeHelper<FVector4f> { typedef FVector4f TSimType; typedef FVector4 TLWCType; };
+template<> struct TNiagaraTypeHelper<FQuat4f> { typedef FQuat4f TSimType; typedef FQuat TLWCType; };
+
+template<> struct TNiagaraTypeHelper<FNiagaraPosition> { typedef FVector3f TSimType; typedef FVector TLWCType; };
+
+//Always treat enums as int32s internally to Niagara for both Sim and Game.
+template<typename T> 
+struct TNiagaraTypeHelper<T, std::enable_if_t<std::is_enum_v<T>>> 
+{
+	typedef int32 TSimType; typedef int32 TLWCType; 
+};
\ No newline at end of file
