Fixed access of invalid script variables in script data post load migration to hierarchy data crashing.

diff --git a/Engine/Plugins/FX/Niagara/Source/NiagaraEditor/Private/NiagaraGraph.cpp b/Engine/Plugins/FX/Niagara/Source/NiagaraEditor/Private/NiagaraGraph.cpp
--- a/Engine/Plugins/FX/Niagara/Source/NiagaraEditor/Private/NiagaraGraph.cpp
+++ b/Engine/Plugins/FX/Niagara/Source/NiagaraEditor/Private/NiagaraGraph.cpp
@@ -4453,14 +4453,22 @@ void UNiagaraGraph::MigrateParameterScriptDataToHierarchyRoot(FVersionedNiagaraS
 	{
 		UNiagaraHierarchyScriptParameter* ScriptParameterA = Cast<UNiagaraHierarchyScriptParameter>(&ItemA);
 		UNiagaraHierarchyScriptParameter* ScriptParameterB = Cast<UNiagaraHierarchyScriptParameter>(&ItemB);
-
+		
 		if(ScriptParameterA == nullptr || ScriptParameterB == nullptr)
 		{
 			return false;
 		}
+		
+		UNiagaraScriptVariable* ScriptVariableA = ScriptParameterA->GetScriptVariable();
+		UNiagaraScriptVariable* ScriptVariableB = ScriptParameterB->GetScriptVariable();
+
+		if(ScriptVariableA == nullptr || ScriptVariableB == nullptr)
+		{
+			return false;
+		}
 
-		const FNiagaraVariableMetaData& MetaDataA = ScriptParameterA->GetScriptVariable()->Metadata;
-		const FNiagaraVariableMetaData& MetaDataB = ScriptParameterB->GetScriptVariable()->Metadata;
+		const FNiagaraVariableMetaData& MetaDataA = ScriptVariableA->Metadata;
+		const FNiagaraVariableMetaData& MetaDataB = ScriptVariableB->Metadata;
 
 		if(MetaDataA.bAdvancedDisplay && !MetaDataB.bAdvancedDisplay)
 		{
@@ -4526,9 +4534,22 @@ void UNiagaraGraph::MigrateParameterScriptDataToHierarchyRoot(FVersionedNiagaraS
 		{
 			UNiagaraHierarchyScriptParameter* ScriptParameterA = Cast<UNiagaraHierarchyScriptParameter>(&ItemA);
 			UNiagaraHierarchyScriptParameter* ScriptParameterB = Cast<UNiagaraHierarchyScriptParameter>(&ItemB);
+			
+			if(ScriptParameterA == nullptr || ScriptParameterB == nullptr)
+			{
+				return false;
+			}
+			
+			UNiagaraScriptVariable* ScriptVariableA = ScriptParameterA->GetScriptVariable();
+			UNiagaraScriptVariable* ScriptVariableB = ScriptParameterB->GetScriptVariable();
 
-			const FNiagaraVariableMetaData& MetaDataA = ScriptParameterA->GetScriptVariable()->Metadata;
-			const FNiagaraVariableMetaData& MetaDataB = ScriptParameterB->GetScriptVariable()->Metadata;
+			if(ScriptVariableA == nullptr || ScriptVariableB == nullptr)
+			{
+				return false;
+			}
+			
+			const FNiagaraVariableMetaData& MetaDataA = ScriptVariableA->Metadata;
+			const FNiagaraVariableMetaData& MetaDataB = ScriptVariableB->Metadata;
 
 			if(MetaDataA.bAdvancedDisplay && !MetaDataB.bAdvancedDisplay)
 			{
