- Add GPU support for local space test to particle reader

#rnx
#tests Create two emitters and have one read the local state of the other, testing the output for combinations of CPU & GPU

diff --git a/Engine/Plugins/FX/Niagara/Source/Niagara/Classes/NiagaraDataInterfaceParticleRead.h b/Engine/Plugins/FX/Niagara/Source/Niagara/Classes/NiagaraDataInterfaceParticleRead.h
--- a/Engine/Plugins/FX/Niagara/Source/Niagara/Classes/NiagaraDataInterfaceParticleRead.h
+++ b/Engine/Plugins/FX/Niagara/Source/Niagara/Classes/NiagaraDataInterfaceParticleRead.h
@@ -14,6 +14,7 @@ class UNiagaraDataInterfaceParticleRead : public UNiagaraDataInterfaceRWBase
 	GENERATED_UCLASS_BODY()
 
 	BEGIN_SHADER_PARAMETER_STRUCT(FShaderParameters, )
+		SHADER_PARAMETER(uint32,			IsLocalSpace)
 		SHADER_PARAMETER(int,				NumSpawnedParticles)
 		SHADER_PARAMETER(int,				SpawnedParticlesAcquireTag)
 		SHADER_PARAMETER(uint32,			InstanceCountOffset)
diff --git a/Engine/Plugins/FX/Niagara/Source/Niagara/Private/NiagaraDataInterfaceParticleRead.cpp b/Engine/Plugins/FX/Niagara/Source/Niagara/Private/NiagaraDataInterfaceParticleRead.cpp
--- a/Engine/Plugins/FX/Niagara/Source/Niagara/Private/NiagaraDataInterfaceParticleRead.cpp
+++ b/Engine/Plugins/FX/Niagara/Source/Niagara/Private/NiagaraDataInterfaceParticleRead.cpp
@@ -62,6 +62,7 @@ static const FName GetParticleIndexFromIDTableName("Get Particle Index From ID T
 
 static const FName ParticleReadIDName(TEXT("ID"));
 
+static const TCHAR* IsLocalSpaceName = TEXT("_IsLocalSpace");
 static const TCHAR* NumSpawnedParticlesName = TEXT("_NumSpawnedParticles");
 static const TCHAR* SpawnedParticlesAcquireTagName = TEXT("_SpawnedParticlesAcquireTag");
 static const TCHAR* InstanceCountOffsetName = TEXT("_InstanceCountOffset");
@@ -216,6 +217,7 @@ struct FNDIParticleRead_RenderInstanceData
 	const FNiagaraDataSet* CachedDataSet = nullptr;
 	TArray<int32> AttributeIndices;
 	TArray<int32> AttributeCompressed;
+	uint32 IsLocalSpace = 0;
 	int32 AcquireTagRegisterIndex = -1;
 	uint32 bSourceEmitterNotGPUErrorShown : 1;
 	uint32 bWarnFailedToFindAcquireTag : 1;
@@ -248,11 +250,12 @@ struct FNiagaraDataInterfaceProxyParticleRead : public FNiagaraDataInterfaceProx
 		return sizeof(FNDIParticleRead_GameToRenderData);
 	}
 
-	void CreateRenderThreadSystemData(const FNiagaraSystemInstanceID& InstanceID, FString InDebugSourceName)
+	void CreateRenderThreadSystemData(const FNiagaraSystemInstanceID& InstanceID, bool bIsLocalSpace, FString InDebugSourceName)
 	{
 		check(IsInRenderingThread());
 		check(!SystemsRenderData.Contains(InstanceID));
 		FNDIParticleRead_RenderInstanceData& RTInstanceData = SystemsRenderData.Add(InstanceID);
+		RTInstanceData.IsLocalSpace = bIsLocalSpace ? 1 : 0;
 		RTInstanceData.DebugSourceName = InDebugSourceName;
 	}
 
@@ -468,10 +471,11 @@ bool UNiagaraDataInterfaceParticleRead::InitPerInstanceData(void* PerInstanceDat
 	}
 
 	FNiagaraDataInterfaceProxyParticleRead* ThisProxy = GetProxyAs<FNiagaraDataInterfaceProxyParticleRead>();
+	const bool bIsLocalSpace = PIData->EmitterInstance ? PIData->EmitterInstance->IsLocalSpace() : false;
 	ENQUEUE_RENDER_COMMAND(FNDIParticleReadCreateRTInstance)(
-		[ThisProxy, InstanceID = SystemInstance->GetId(), RT_DebugSourceName=MoveTemp(DebugSourceName)](FRHICommandList& CmdList)
+		[ThisProxy, InstanceID = SystemInstance->GetId(), bIsLocalSpace, RT_DebugSourceName=MoveTemp(DebugSourceName)](FRHICommandList& CmdList)
 		{
-			ThisProxy->CreateRenderThreadSystemData(InstanceID, RT_DebugSourceName);
+			ThisProxy->CreateRenderThreadSystemData(InstanceID, bIsLocalSpace, RT_DebugSourceName);
 		}
 	);
 
@@ -1681,7 +1685,7 @@ void UNiagaraDataInterfaceParticleRead::GetParameterDefinitionHLSL(const FNiagar
 	// which is a hassle. Instead, use int4 explicitly, and access individual components in the generated code.
 	// Note that we have to have at least one here because hlsl doesn't support arrays of size 0.
 	const int AttributeInt4Count = FMath::Max(1, FMath::DivideAndRoundUp(ParamInfo.GeneratedFunctions.Num(), 4));
-
+	OutHLSL.Appendf(TEXT("uint			%s%s;"), *ParamInfo.DataInterfaceHLSLSymbol, IsLocalSpaceName);
 	OutHLSL.Appendf(TEXT("int			%s%s;"), *ParamInfo.DataInterfaceHLSLSymbol, NumSpawnedParticlesName);
 	OutHLSL.Appendf(TEXT("int			%s%s;"), *ParamInfo.DataInterfaceHLSLSymbol, SpawnedParticlesAcquireTagName);
 	OutHLSL.Appendf(TEXT("uint			%s%s;"), *ParamInfo.DataInterfaceHLSLSymbol, InstanceCountOffsetName);
@@ -1862,6 +1866,23 @@ static bool GenerateGetFunctionHLSL(const FNiagaraDataInterfaceGPUParamInfo& Par
 
 bool UNiagaraDataInterfaceParticleRead::GetFunctionHLSL(const FNiagaraDataInterfaceGPUParamInfo& ParamInfo, const FNiagaraDataInterfaceGeneratedFunction& FunctionInfo, int FunctionInstanceIndex, FString& OutHLSL)
 {
+	if (FunctionInfo.DefinitionName == GetLocalSpaceName)
+	{
+		static const TCHAR* FuncTemplate = TEXT(
+			"void {FunctionName}(out bool Out_LocalSpace)\n"
+			"{\n"
+			"    Out_LocalSpace = {IsLocalSpace} != 0;\n"
+			"}\n\n"
+		);
+
+		TMap<FString, FStringFormatArg> FuncTemplateArgs;
+		FuncTemplateArgs.Add(TEXT("FunctionName"), FunctionInfo.InstanceName);
+		FuncTemplateArgs.Add(TEXT("IsLocalSpace"), ParamInfo.DataInterfaceHLSLSymbol + IsLocalSpaceName);
+
+		OutHLSL += FString::Format(FuncTemplate, FuncTemplateArgs);
+		return true;
+	}
+
 	//
 	// Spawn info and particle count
 	//
@@ -2120,7 +2141,6 @@ void UNiagaraDataInterfaceParticleRead::SetShaderParameters(const FNiagaraDataIn
 				ShaderParameters->SpawnedParticlesAcquireTag	= 0;
 				ShaderParameters->SpawnedIDsBuffer				= FNiagaraRenderer::GetDummyIntBuffer();
 			}
-
 			ShaderParameters->InstanceCountOffset		= INDEX_NONE;
 			ShaderParameters->ParticleStrideFloat		= 0;
 			ShaderParameters->ParticleStrideInt			= 0;
@@ -2138,10 +2158,12 @@ void UNiagaraDataInterfaceParticleRead::SetShaderParameters(const FNiagaraDataIn
 	FNDIParticleRead_RenderInstanceData* InstanceData = DIProxy.GetRenderDataForSystem(Context.GetSystemInstanceID());
 	if (InstanceData == nullptr)
 	{
+		ShaderParameters->IsLocalSpace = 0;
 		SetErrorParams(true);
 		return;
 	}
 
+	ShaderParameters->IsLocalSpace = InstanceData->IsLocalSpace;
 	if (InstanceData->SourceEmitterGPUContext == nullptr)
 	{
 		if (!InstanceData->bSourceEmitterNotGPUErrorShown)
