Niagara - UX - Added a Origin Axis option under Preview Settings menu in Niagara Preview to visualize grid 0 location

diff --git a/Engine/Plugins/FX/Niagara/Source/NiagaraEditor/Private/NiagaraEditorCommands.cpp b/Engine/Plugins/FX/Niagara/Source/NiagaraEditor/Private/NiagaraEditorCommands.cpp
--- a/Engine/Plugins/FX/Niagara/Source/NiagaraEditor/Private/NiagaraEditorCommands.cpp
+++ b/Engine/Plugins/FX/Niagara/Source/NiagaraEditor/Private/NiagaraEditorCommands.cpp
@@ -22,6 +22,7 @@ void FNiagaraEditorCommands::RegisterCommands()
 
 	UI_COMMAND(ResetSimulation, "Reset", "Resets the current simulation", EUserInterfaceActionType::Button, FInputChord(EKeys::SpaceBar));
 	UI_COMMAND(TogglePreviewGrid, "Grid", "Toggles the preview pane's grid.", EUserInterfaceActionType::ToggleButton, FInputChord() );
+	UI_COMMAND(ToggleOriginAxis, "Origin Axis", "Toggles the preview pane's origin axis.", EUserInterfaceActionType::ToggleButton, FInputChord());
 	UI_COMMAND(ToggleInstructionCounts, "InstructionCounts", "Display Instruction Counts", EUserInterfaceActionType::ToggleButton, FInputChord());
 	UI_COMMAND(ToggleParticleCounts, "ParticleCounts", "Display Particle Counts", EUserInterfaceActionType::ToggleButton, FInputChord());
 	UI_COMMAND(ToggleEmitterExecutionOrder, "Emitter Execution Order", "Display Emitter Execution Order", EUserInterfaceActionType::ToggleButton, FInputChord());
diff --git a/Engine/Plugins/FX/Niagara/Source/NiagaraEditor/Private/NiagaraEditorStyle.cpp b/Engine/Plugins/FX/Niagara/Source/NiagaraEditor/Private/NiagaraEditorStyle.cpp
--- a/Engine/Plugins/FX/Niagara/Source/NiagaraEditor/Private/NiagaraEditorStyle.cpp
+++ b/Engine/Plugins/FX/Niagara/Source/NiagaraEditor/Private/NiagaraEditorStyle.cpp
@@ -405,6 +405,7 @@ void FNiagaraEditorStyle::InitSelectedEmitter()
 
 void FNiagaraEditorStyle::InitToolbarIcons()
 {
+	Set("NiagaraEditor.ToggleOriginAxis", new IMAGE_BRUSH_SVG("Icons/OriginAxis_16", Icon16x16));
 	Set("NiagaraEditor.ApplyScratchPadChanges", new IMAGE_BRUSH_SVG("Icons/Commands/ApplyScratch", Icon40x40));
 	Set("NiagaraEditor.OverviewNode.IsolatedColor", FLinearColor::Yellow);
 	Set("NiagaraEditor.OverviewNode.NotIsolatedColor", FLinearColor::Transparent);
diff --git a/Engine/Plugins/FX/Niagara/Source/NiagaraEditor/Private/Widgets/NiagaraSystemEditorViewportToolbarSections.cpp b/Engine/Plugins/FX/Niagara/Source/NiagaraEditor/Private/Widgets/NiagaraSystemEditorViewportToolbarSections.cpp
--- a/Engine/Plugins/FX/Niagara/Source/NiagaraEditor/Private/Widgets/NiagaraSystemEditorViewportToolbarSections.cpp
+++ b/Engine/Plugins/FX/Niagara/Source/NiagaraEditor/Private/Widgets/NiagaraSystemEditorViewportToolbarSections.cpp
@@ -1,9 +1,10 @@
-﻿// Copyright Epic Games, Inc. All Rights Reserved.
+// Copyright Epic Games, Inc. All Rights Reserved.
 
 #include "NiagaraSystemEditorViewportToolbarSections.h"
 
 #include "EditorViewportCommands.h"
 #include "NiagaraEditorCommands.h"
+#include "NiagaraEditorStyle.h"
 #include "SNiagaraSystemViewport.h"
 #include "ToolMenu.h"
 #include "ToolMenus.h"
@@ -120,6 +121,7 @@ void UE::NiagaraSystemEditor::FillSettingsSubmenu(UToolMenu* InMenu)
 
 	AddMotionSettingsToSection(ViewportControlsSection);
 }
+
 void UE::NiagaraSystemEditor::AddMotionSettingsToSection(FToolMenuSection& InSection)
 {
 	InSection.AddSubMenu(
@@ -194,6 +196,17 @@ void UE::NiagaraSystemEditor::ExtendPreviewSceneSettingsSubmenu(FName InSubmenuN
 	);
 
 	UE::NiagaraSystemEditor::AddMotionSettingsToSection(PreviewControlsSection);
+
+	//Adds toggle origin axis to menu
+	FToolMenuSection& ProfileOptionsSection = Submenu->FindOrAddSection(
+		"AssetViewerProfileOptionsSection", LOCTEXT("AssetViewerProfileOptionsSectionLabel", "Preview Scene Options")
+	);
+	ProfileOptionsSection.AddMenuEntry(
+		FNiagaraEditorCommands::Get().ToggleOriginAxis, 
+		TAttribute<FText>(), 
+		TAttribute<FText>(),
+		FSlateIcon(FNiagaraEditorStyle::Get().GetStyleSetName(), "NiagaraEditor.ToggleOriginAxis")
+	);
 }
 
 #undef LOCTEXT_NAMESPACE
diff --git a/Engine/Plugins/FX/Niagara/Source/NiagaraEditor/Private/Widgets/SNiagaraSystemViewport.cpp b/Engine/Plugins/FX/Niagara/Source/NiagaraEditor/Private/Widgets/SNiagaraSystemViewport.cpp
--- a/Engine/Plugins/FX/Niagara/Source/NiagaraEditor/Private/Widgets/SNiagaraSystemViewport.cpp
+++ b/Engine/Plugins/FX/Niagara/Source/NiagaraEditor/Private/Widgets/SNiagaraSystemViewport.cpp
@@ -42,6 +42,8 @@
 #include "ViewModels/NiagaraSystemViewModel.h"
 #include "Widgets/Docking/SDockTab.h"
 #include "Widgets/Layout/SBox.h"
+#include "Components/LineBatchComponent.h"
+
 
 #define LOCTEXT_NAMESPACE "SNiagaraSystemViewport"
 
@@ -307,6 +309,24 @@ void FNiagaraSystemViewportClient::Draw(FViewport* InViewport, FCanvas* Canvas)
 			UFont* TinyFont = GEngine->GetTinyFont();
 			DrawStatelessInfo(Component, InViewport, Canvas, TinyFont, TinyFont->GetMaxCharHeight() * 1.1f);
 		}
+		if (NiagaraViewport->GetDrawElement(SNiagaraSystemViewport::EDrawElements::OriginAxis) && Component)
+		{
+			UWorld* World = Component->GetWorld();
+			ULineBatchComponent* LineBatcher = World->GetLineBatcher(UWorld::ELineBatcherType::World);
+
+			const FVector XAxis(1, 0, 0);
+			const FVector YAxis(0, 1, 0);
+			const FVector ZAxis(0, 0, 1);
+
+			FMatrix ArrowMatrix = FMatrix(XAxis, YAxis, ZAxis, FVector::ZeroVector);
+			LineBatcher->DrawDirectionalArrow(ArrowMatrix, FColorList::Red, 50.f, 4.0f, SDPG_World);
+
+			ArrowMatrix = FMatrix(YAxis, ZAxis, XAxis, FVector::ZeroVector);
+			LineBatcher->DrawDirectionalArrow(ArrowMatrix, FColorList::Green, 50.f, 4.0f, SDPG_World);
+
+			ArrowMatrix = FMatrix(ZAxis, XAxis, YAxis, FVector::ZeroVector);
+			LineBatcher->DrawDirectionalArrow(ArrowMatrix, FColorList::Blue, 50.f, 4.0f, SDPG_World);
+		}
 	}
 
 	if (bCaptureScreenShot && ScreenShotOwner.IsValid() && OnScreenShotCaptured.IsBound())
@@ -1051,6 +1071,17 @@ void SNiagaraSystemViewport::BindCommands()
 		FIsActionChecked::CreateSP( this, &SNiagaraSystemViewport::IsTogglePreviewGridChecked )
 	);
 
+	CommandList->MapAction(
+		Commands.ToggleOriginAxis,		
+		FExecuteAction::CreateLambda([Viewport = this]()
+		{
+			Viewport->ToggleDrawElement(EDrawElements::OriginAxis);
+			Viewport->RefreshViewport();
+		}),
+		FCanExecuteAction(),
+		FIsActionChecked::CreateLambda([Viewport = this]() -> bool { return Viewport->GetDrawElement(EDrawElements::OriginAxis); })
+	);
+
 	CommandList->MapAction(
 		Commands.ToggleInstructionCounts,
 		FExecuteAction::CreateLambda([Viewport=this]()
diff --git a/Engine/Plugins/FX/Niagara/Source/NiagaraEditor/Private/Widgets/SNiagaraSystemViewport.h b/Engine/Plugins/FX/Niagara/Source/NiagaraEditor/Private/Widgets/SNiagaraSystemViewport.h
--- a/Engine/Plugins/FX/Niagara/Source/NiagaraEditor/Private/Widgets/SNiagaraSystemViewport.h
+++ b/Engine/Plugins/FX/Niagara/Source/NiagaraEditor/Private/Widgets/SNiagaraSystemViewport.h
@@ -80,6 +80,7 @@ public:
 	/** Draw flag types */
 	enum EDrawElements
 	{
+		OriginAxis				= 0x010,
 		Bounds					= 0x020,
 		InstructionCounts		= 0x040,
 		ParticleCounts			= 0x080,
diff --git a/Engine/Plugins/FX/Niagara/Source/NiagaraEditor/Public/NiagaraEditorCommands.h b/Engine/Plugins/FX/Niagara/Source/NiagaraEditor/Public/NiagaraEditorCommands.h
--- a/Engine/Plugins/FX/Niagara/Source/NiagaraEditor/Public/NiagaraEditorCommands.h
+++ b/Engine/Plugins/FX/Niagara/Source/NiagaraEditor/Public/NiagaraEditorCommands.h
@@ -37,6 +37,7 @@ public:
 	
 	/** Toggles the preview pane's grid */
 	TSharedPtr<FUICommandInfo> TogglePreviewGrid;
+	TSharedPtr<FUICommandInfo> ToggleOriginAxis;
 	TSharedPtr<FUICommandInfo> ToggleInstructionCounts;
 	TSharedPtr<FUICommandInfo> ToggleParticleCounts;
 	TSharedPtr<FUICommandInfo> ToggleEmitterExecutionOrder;
