Niagara - Update vector UIs to handle the new axis name and color changes.
+ Refactor the vector type utilities to use the editor supplied vector editor which automatically handles the axis names and colors.
+ Refactor the vector curve UIs to use the editor supplied axis names and colors where appropriate.

diff --git a/Engine/Plugins/FX/Niagara/Source/NiagaraEditor/Private/NiagaraEditorUtilities.cpp b/Engine/Plugins/FX/Niagara/Source/NiagaraEditor/Private/NiagaraEditorUtilities.cpp
--- a/Engine/Plugins/FX/Niagara/Source/NiagaraEditor/Private/NiagaraEditorUtilities.cpp
+++ b/Engine/Plugins/FX/Niagara/Source/NiagaraEditor/Private/NiagaraEditorUtilities.cpp
@@ -4674,6 +4674,23 @@ UNiagaraDataInterface* FNiagaraEditorUtilities::GetResolvedRuntimeInstanceForEdi
 	return nullptr;
 }
 
+EAxisList::Type FNiagaraEditorUtilities::VectorComponentToAxis(int32 NumComponents, int32 ComponentIndex)
+{
+	if (NumComponents == 3)
+	{
+		switch (ComponentIndex)
+		{
+		case 0:
+			return EAxisList::X;
+		case 1:
+			return EAxisList::Y;
+		case 2:
+			return EAxisList::Z;
+		}
+	}
+	return EAxisList::None;
+}
+
 FNiagaraVariable FNiagaraEditorUtilities::UserParameters::DuplicateUserParameter(FNiagaraVariable ParameterToDuplicate, UNiagaraSystem& System)
 {
 	FScopedTransaction AddTransaction(LOCTEXT("DuplicateUserParameter", "Duplicate User Parameter"));
diff --git a/Engine/Plugins/FX/Niagara/Source/NiagaraEditor/Private/Stateless/NiagaraDistributionPropertyCustomization.cpp b/Engine/Plugins/FX/Niagara/Source/NiagaraEditor/Private/Stateless/NiagaraDistributionPropertyCustomization.cpp
--- a/Engine/Plugins/FX/Niagara/Source/NiagaraEditor/Private/Stateless/NiagaraDistributionPropertyCustomization.cpp
+++ b/Engine/Plugins/FX/Niagara/Source/NiagaraEditor/Private/Stateless/NiagaraDistributionPropertyCustomization.cpp
@@ -5,6 +5,7 @@
 #include "DetailWidgetRow.h"
 #include "Editor.h"
 #include "EditorUndoClient.h"
+#include "Misc/AxisDisplayInfo.h"
 #include "PropertyHandle.h"
 #include "ScopedTransaction.h"
 #include "Stateless/NiagaraStatelessDistribution.h"
@@ -90,6 +91,12 @@ public:
 		}
 		else
 		{
+			EAxisList::Type ComponentAxis = FNiagaraEditorUtilities::VectorComponentToAxis(EditorNumChannels, ChannelIndex);
+			if (ComponentAxis != EAxisList::None)
+			{
+				return AxisDisplayInfo::GetAxisDisplayNameShort(ComponentAxis);
+			}
+
 			switch(ChannelIndex)
 			{
 			case 0:
@@ -109,6 +116,12 @@ public:
 	{
 		if (EditorNumChannels != 1)
 		{
+			EAxisList::Type ComponentAxis = bDisplayAsColor ? EAxisList::None : FNiagaraEditorUtilities::VectorComponentToAxis(EditorNumChannels, ChannelIndex);
+			if (ComponentAxis != EAxisList::None)
+			{
+				return AxisDisplayInfo::GetAxisColor(ComponentAxis);
+			}
+
 			switch (ChannelIndex)
 			{
 			case 0:
diff --git a/Engine/Plugins/FX/Niagara/Source/NiagaraEditor/Private/TypeEditorUtilities/NiagaraVectorTypeEditorUtilities.cpp b/Engine/Plugins/FX/Niagara/Source/NiagaraEditor/Private/TypeEditorUtilities/NiagaraVectorTypeEditorUtilities.cpp
--- a/Engine/Plugins/FX/Niagara/Source/NiagaraEditor/Private/TypeEditorUtilities/NiagaraVectorTypeEditorUtilities.cpp
+++ b/Engine/Plugins/FX/Niagara/Source/NiagaraEditor/Private/TypeEditorUtilities/NiagaraVectorTypeEditorUtilities.cpp
@@ -7,8 +7,8 @@
 #include "NiagaraEditorStyle.h"
 #include "NiagaraTypes.h"
 
-#include "Widgets/Input/SNumericEntryBox.h"
 #include "Misc/DefaultValueHelper.h"
+#include "Widgets/Input/SVectorInputBox.h"
 
 class SNiagaraVectorParameterEditorBase : public SNiagaraParameterEditor
 {
@@ -23,19 +23,9 @@ public:
 			.MinimumDesiredWidth(DefaultInputSize * InArgs._ComponentCount)
 			.MaximumDesiredWidth(DefaultInputSize * InArgs._ComponentCount));
 
-		TSharedRef<SHorizontalBox> ComponentBox = SNew(SHorizontalBox);
-		for (int32 ComponentIndex = 0; ComponentIndex < InArgs._ComponentCount; ++ComponentIndex)
-		{
-			ComponentBox->AddSlot()
-			.Padding(ComponentIndex == 0 ? 0 : 3, 0, 0, 0)
-			[
-				ConstructComponentWidget(ComponentIndex, DisplayUnit, WidgetCustomization)
-			];
-		}
-
 		ChildSlot
 		[
-			ComponentBox
+			ConstructVectorWidget(InArgs._ComponentCount, DisplayUnit, WidgetCustomization)
 		];
 	}
 
@@ -47,48 +37,81 @@ protected:
 	virtual void SetValue(int32 Index, float Value) = 0;
 
 private:
-	TSharedRef<SWidget> ConstructComponentWidget(int32 Index, EUnit DisplayUnit, const FNiagaraInputParameterCustomization& WidgetCustomization)
-	{
-		const static TArray Colors = {
-			&SNumericEntryBox<int32>::RedLabelBackgroundColor,
-			&SNumericEntryBox<int32>::GreenLabelBackgroundColor,
-			&SNumericEntryBox<int32>::BlueLabelBackgroundColor,
-			&FLinearColor::White
-		};
-		TOptional<float> MinValue;
-		TOptional<float> MaxValue;
+	typedef SNumericVectorInputBox<float, UE::Math::TVector2<float>, 2> SNumericVectorInputBox2;
+	typedef SNumericVectorInputBox<float, UE::Math::TVector<float>, 3> SNumericVectorInputBox3;
+	typedef SNumericVectorInputBox<float, UE::Math::TVector4<float>, 4> SNumericVectorInputBox4;
+
+	template<typename TArgs, typename TVector>
+	TArgs CreateBaseArgs(EUnit DisplayUnit, const FNiagaraInputParameterCustomization& WidgetCustomization)
+	{
+		TOptional<TVector> MinValue;
+		TOptional<TVector> MaxValue;
 		if (WidgetCustomization.bHasMinValue)
 		{
-			MinValue = WidgetCustomization.MinValue;
+			MinValue = TVector(WidgetCustomization.MinValue);
 		}
 		if (WidgetCustomization.bHasMaxValue)
 		{
-			MaxValue = WidgetCustomization.MaxValue;
+			MaxValue = TVector(WidgetCustomization.MaxValue);
 		}
-		
-		TSharedRef<SWidget>  LabelWidget = SNumericEntryBox<float>::BuildNarrowColorLabel(*Colors[Index]);
-		
-		return SNew(SNumericEntryBox<float>)
-		.Font(FAppStyle::Get().GetFontStyle("PropertyWindow.NormalFont"))
-		.MinValue(MinValue)
-		.MaxValue(MaxValue)
-		.MaxSliderValue(MaxValue)
-		.MinSliderValue(MinValue)
-		.Delta(WidgetCustomization.bHasStepWidth ? WidgetCustomization.StepWidth : 0)
-		.Value(this, &SNiagaraVectorParameterEditorBase::GetValueInternal, Index)
-		.OnValueChanged(this, &SNiagaraVectorParameterEditorBase::ValueChanged, Index)
-		.OnValueCommitted(this, &SNiagaraVectorParameterEditorBase::ValueCommitted, Index)
-		.OnBeginSliderMovement(this, &SNiagaraVectorParameterEditorBase::BeginSliderMovement)
-		.OnEndSliderMovement(this, &SNiagaraVectorParameterEditorBase::EndSliderMovement)
-		.TypeInterface(GetTypeInterface<float>(DisplayUnit))
-		.BroadcastValueChangesPerKey(!GetDefault<UNiagaraEditorSettings>()->GetUpdateStackValuesOnCommitOnly() && !WidgetCustomization.bBroadcastValueChangesOnCommitOnly)
-		.AllowSpin(true)
-		.LabelPadding( FMargin(3) )
-		.LabelLocation( SNumericEntryBox<float>::ELabelLocation::Inside )
-		.Label()
-		[
-			LabelWidget
-		];
+
+		return TArgs()
+			.Font(FAppStyle::Get().GetFontStyle("PropertyWindow.NormalFont"))
+			.AllowSpin(true)
+			.bColorAxisLabels(true)
+			.SpinDelta(WidgetCustomization.bHasStepWidth ? WidgetCustomization.StepWidth : 0)
+			.OnBeginSliderMovement(this, &SNiagaraVectorParameterEditorBase::BeginSliderMovement)
+			.OnEndSliderMovement(this, &SNiagaraVectorParameterEditorBase::EndSliderMovement)
+			.TypeInterface(MakeShareable(new TNumericUnitTypeInterface<float>(DisplayUnit)))
+			.MinVector(MinValue)
+			.MaxVector(MaxValue)
+			.MinSliderVector(MinValue)
+			.MaxSliderVector(MaxValue)
+			.X(this, &SNiagaraVectorParameterEditorBase::GetValueInternal, 0)
+			.OnXChanged(this, &SNiagaraVectorParameterEditorBase::ValueChanged, 0)
+			.OnXCommitted(this, &SNiagaraVectorParameterEditorBase::ValueCommitted, 0)
+			.Y(this, &SNiagaraVectorParameterEditorBase::GetValueInternal, 1)
+			.OnYChanged(this, &SNiagaraVectorParameterEditorBase::ValueChanged, 1)
+			.OnYCommitted(this, &SNiagaraVectorParameterEditorBase::ValueCommitted, 1);
+	}
+
+	template<typename TArguments>
+	void SetZArgs(TArguments& Args)
+	{
+		Args.Z(this, &SNiagaraVectorParameterEditorBase::GetValueInternal, 2)
+			.OnZChanged(this, &SNiagaraVectorParameterEditorBase::ValueChanged, 2)
+			.OnZCommitted(this, &SNiagaraVectorParameterEditorBase::ValueCommitted, 2);
+	}
+
+	template<typename TArguments>
+	void SetWArgs(TArguments& Args)
+	{
+		Args.W(this, &SNiagaraVectorParameterEditorBase::GetValueInternal, 3)
+			.OnWChanged(this, &SNiagaraVectorParameterEditorBase::ValueChanged, 3)
+			.OnWCommitted(this, &SNiagaraVectorParameterEditorBase::ValueCommitted, 3);
+	}
+
+	TSharedRef<SWidget> ConstructVectorWidget(int32 ChannelCount, EUnit DisplayUnit, const FNiagaraInputParameterCustomization& WidgetCustomization)
+	{
+		if (ChannelCount == 2)
+		{
+			SNumericVectorInputBox2::FArguments Vector2Args = CreateBaseArgs<SNumericVectorInputBox2::FArguments, FVector2f>(DisplayUnit, WidgetCustomization);
+			return SArgumentNew(Vector2Args, SNumericVectorInputBox2);
+		}
+		if (ChannelCount == 3)
+		{
+			SNumericVectorInputBox3::FArguments Vector3Args = CreateBaseArgs<SNumericVectorInputBox3::FArguments, FVector3f>(DisplayUnit, WidgetCustomization);
+			SetZArgs(Vector3Args);
+			return SArgumentNew(Vector3Args, SNumericVectorInputBox3);
+		}
+		if (ChannelCount == 4)
+		{
+			SNumericVectorInputBox4::FArguments Vector4Args = CreateBaseArgs<SNumericVectorInputBox4::FArguments, FVector4f>(DisplayUnit, WidgetCustomization);
+			SetZArgs(Vector4Args);
+			SetWArgs(Vector4Args);
+			return SArgumentNew(Vector4Args, SNumericVectorInputBox4);
+		}
+		return SNullWidget::NullWidget;
 	}
 
 	void BeginSliderMovement()
diff --git a/Engine/Plugins/FX/Niagara/Source/NiagaraEditor/Private/ViewModels/NiagaraCurveSelectionViewModel.cpp b/Engine/Plugins/FX/Niagara/Source/NiagaraEditor/Private/ViewModels/NiagaraCurveSelectionViewModel.cpp
--- a/Engine/Plugins/FX/Niagara/Source/NiagaraEditor/Private/ViewModels/NiagaraCurveSelectionViewModel.cpp
+++ b/Engine/Plugins/FX/Niagara/Source/NiagaraEditor/Private/ViewModels/NiagaraCurveSelectionViewModel.cpp
@@ -3,7 +3,9 @@
 #include "ViewModels/NiagaraCurveSelectionViewModel.h"
 
 #include "EdGraphSchema_Niagara.h"
+#include "Misc/AxisDisplayInfo.h"
 #include "NiagaraDataInterfaceCurveBase.h"
+#include "NiagaraEditorUtilities.h"
 #include "NiagaraEmitter.h"
 #include "NiagaraEmitterEditorData.h"
 #include "NiagaraGraph.h"
@@ -103,6 +105,16 @@ void FNiagaraCurveSelectionTreeNode::SetSecondDisplayName(FText InSecondDisplayN
 	SecondDisplayName = InSecondDisplayName;
 }
 
+FText FNiagaraCurveSelectionTreeNode::GetToolTipText() const
+{
+	return ToolTipText;
+}
+
+void FNiagaraCurveSelectionTreeNode::SetToolTipText(FText InToolTipText)
+{
+	ToolTipText = InToolTipText;
+}
+
 ENiagaraCurveSelectionNodeStyleMode FNiagaraCurveSelectionTreeNode::GetStyleMode() const
 {
 	return StyleMode;
@@ -378,17 +390,37 @@ TSharedRef<FNiagaraCurveSelectionTreeNode> UNiagaraCurveSelectionViewModel::Crea
 	else
 	{
 		TArray<TSharedRef<FNiagaraCurveSelectionTreeNode>> CurveNodes;
+		int32 CurveComponentIndex = 0;
 		for (UNiagaraDataInterfaceCurveBase::FCurveData& CurveData : InterfaceCurveData)
 		{
+			FText CurveDisplayName;
+			FText CurveToolTipText;
+			FLinearColor CurveDisplayColor;
+
+			EAxisList::Type CurveAxis = FNiagaraEditorUtilities::VectorComponentToAxis(InterfaceCurveData.Num(), CurveComponentIndex);
+			if (CurveAxis != EAxisList::None)
+			{
+				CurveDisplayName = AxisDisplayInfo::GetAxisDisplayNameShort(CurveAxis);
+				CurveToolTipText = AxisDisplayInfo::GetAxisDisplayName(CurveAxis);
+				CurveDisplayColor = AxisDisplayInfo::GetAxisColor(CurveAxis);
+			}
+			else
+			{
+				CurveDisplayName = FText::FromName(CurveData.Name);
+				CurveDisplayColor = CurveData.Color;
+			}
+
 			TSharedRef<FNiagaraCurveSelectionTreeNode> CurveNode = MakeShared<FNiagaraCurveSelectionTreeNode>();
 			CurveNode->SetDataId(FNiagaraCurveSelectionTreeNodeDataId::FromUniqueName(CurveData.Name));
-			CurveNode->SetDisplayName(FText::FromName(CurveData.Name));
+			CurveNode->SetDisplayName(CurveDisplayName);
+			CurveNode->SetToolTipText(CurveToolTipText);
 			CurveNode->SetStyle(ENiagaraCurveSelectionNodeStyleMode::CurveComponent, NAME_None, NAME_None, false);
-			CurveNode->SetCurveData(&CurveDataInterface, CurveData.Curve, CurveData.Name, CurveData.Color);
+			CurveNode->SetCurveData(&CurveDataInterface, CurveData.Curve, CurveData.Name, CurveDisplayColor);
 			CurveNode->SetDisplayedObjectKey(FObjectKey(&CurveDataInterface));
 			CurveNode->SetShowInTree(true);
 			CurveNode->GetOnCurveChanged().AddUObject(this, &UNiagaraCurveSelectionViewModel::DataInterfaceCurveChanged, TWeakObjectPtr<UNiagaraDataInterfaceCurveBase>(&CurveDataInterface));
 			CurveNodes.Add(CurveNode);
+			CurveComponentIndex++;
 		}
 		DataInterfaceNode->SetChildNodes(CurveNodes);
 	}
diff --git a/Engine/Plugins/FX/Niagara/Source/NiagaraEditor/Public/NiagaraEditorUtilities.h b/Engine/Plugins/FX/Niagara/Source/NiagaraEditor/Public/NiagaraEditorUtilities.h
--- a/Engine/Plugins/FX/Niagara/Source/NiagaraEditor/Public/NiagaraEditorUtilities.h
+++ b/Engine/Plugins/FX/Niagara/Source/NiagaraEditor/Public/NiagaraEditorUtilities.h
@@ -425,6 +425,8 @@ namespace FNiagaraEditorUtilities
 
 	NIAGARAEDITOR_API UNiagaraDataInterface* GetResolvedRuntimeInstanceForEditorDataInterfaceInstance(const UNiagaraSystem& OwningSystem, UNiagaraDataInterface& EditorDataInterfaceInstance);
 
+	NIAGARAEDITOR_API EAxisList::Type VectorComponentToAxis(int32 NumComponents, int32 ComponentIndex);
+
 	namespace UserParameters
 	{
 		FNiagaraVariable DuplicateUserParameter(FNiagaraVariable ParameterToDuplicate, UNiagaraSystem& System);
diff --git a/Engine/Plugins/FX/Niagara/Source/NiagaraEditor/Public/ViewModels/NiagaraCurveSelectionViewModel.h b/Engine/Plugins/FX/Niagara/Source/NiagaraEditor/Public/ViewModels/NiagaraCurveSelectionViewModel.h
--- a/Engine/Plugins/FX/Niagara/Source/NiagaraEditor/Public/ViewModels/NiagaraCurveSelectionViewModel.h
+++ b/Engine/Plugins/FX/Niagara/Source/NiagaraEditor/Public/ViewModels/NiagaraCurveSelectionViewModel.h
@@ -64,6 +64,10 @@ public:
 
 	NIAGARAEDITOR_API void SetSecondDisplayName(FText InSecondDisplayName);
 
+	NIAGARAEDITOR_API FText GetToolTipText() const;
+
+	NIAGARAEDITOR_API void SetToolTipText(FText InToolTipText);
+
 	NIAGARAEDITOR_API ENiagaraCurveSelectionNodeStyleMode GetStyleMode() const;
 
 	NIAGARAEDITOR_API FName GetExecutionCategory() const;
@@ -135,6 +139,7 @@ private:
 	FGuid NodeUniqueId;
 	FText DisplayName;
 	FText SecondDisplayName;
+	FText ToolTipText;
 	TWeakPtr<FNiagaraCurveSelectionTreeNode> ParentWeak;
 	TArray<TSharedRef<FNiagaraCurveSelectionTreeNode>> ChildNodes;
 
diff --git a/Engine/Plugins/FX/Niagara/Source/NiagaraEditorWidgets/Private/DetailCustomizations/NiagaraDataInterfaceCurveDetails.cpp b/Engine/Plugins/FX/Niagara/Source/NiagaraEditorWidgets/Private/DetailCustomizations/NiagaraDataInterfaceCurveDetails.cpp
--- a/Engine/Plugins/FX/Niagara/Source/NiagaraEditorWidgets/Private/DetailCustomizations/NiagaraDataInterfaceCurveDetails.cpp
+++ b/Engine/Plugins/FX/Niagara/Source/NiagaraEditorWidgets/Private/DetailCustomizations/NiagaraDataInterfaceCurveDetails.cpp
@@ -10,6 +10,7 @@
 #include "NiagaraEditorWidgetsModule.h"
 #include "NiagaraSystem.h"
 #include "NiagaraEditorSettings.h"
+#include "NiagaraEditorUtilities.h"
 #include "ViewModels/NiagaraSystemViewModel.h"
 #include "ViewModels/NiagaraCurveSelectionViewModel.h"
 
@@ -49,6 +50,7 @@
 #include "Curves/RichCurve.h"
 #include "ScopedTransaction.h"
 #include "SColorGradientEditor.h"
+#include "Misc/AxisDisplayInfo.h"
 
 #include "EditorFontGlyphs.h"
 #include "Widgets/Input/SSegmentedControl.h"
@@ -57,9 +59,10 @@
 
 struct FNiagaraCurveDetailsTreeItem : public ICurveEditorTreeItem, TSharedFromThis<FNiagaraCurveDetailsTreeItem>
 {
-	FNiagaraCurveDetailsTreeItem(UObject* InCurveOwnerObject, UNiagaraDataInterfaceCurveBase::FCurveData InCurveData)
+	FNiagaraCurveDetailsTreeItem(UObject* InCurveOwnerObject, UNiagaraDataInterfaceCurveBase::FCurveData InCurveData, EAxisList::Type InCurveAxis)
 		: CurveOwnerObject(InCurveOwnerObject)
 		, CurveData(InCurveData)
+		, CurveAxis(InCurveAxis)
 	{
 	}
 
@@ -77,7 +80,23 @@ struct FNiagaraCurveDetailsTreeItem : public ICurveEditorTreeItem, TSharedFromTh
 	{
 		if (InColumnName == ColumnNames.Label)
 		{
-			return SNew(SHorizontalBox)
+			FText CurveDisplayName;
+			FText CurveToolTipText;
+			FLinearColor CurveDisplayColor;
+
+			if (CurveAxis != EAxisList::None)
+			{
+				CurveDisplayName = AxisDisplayInfo::GetAxisDisplayNameShort(CurveAxis);
+				CurveToolTipText = AxisDisplayInfo::GetAxisDisplayName(CurveAxis);
+				CurveDisplayColor = AxisDisplayInfo::GetAxisColor(CurveAxis);
+			}
+			else
+			{
+				CurveDisplayName = FText::FromName(CurveData.Name);
+				CurveDisplayColor = CurveData.Color;
+			}
+
+			TSharedRef<SHorizontalBox> LabelBox = SNew(SHorizontalBox)
 				+ SHorizontalBox::Slot()
 				.AutoWidth()
 				.VAlign(VAlign_Center)
@@ -86,7 +105,7 @@ struct FNiagaraCurveDetailsTreeItem : public ICurveEditorTreeItem, TSharedFromTh
 					SNew(STextBlock)
 					.Font(FAppStyle::Get().GetFontStyle("FontAwesome.8"))
 					.Text(FEditorFontGlyphs::Circle)
-					.ColorAndOpacity(FSlateColor(CurveData.Color))
+					.ColorAndOpacity(FSlateColor(CurveDisplayColor))
 				]
 				+ SHorizontalBox::Slot()
 				.VAlign(VAlign_Center)
@@ -96,10 +115,17 @@ struct FNiagaraCurveDetailsTreeItem : public ICurveEditorTreeItem, TSharedFromTh
 					.MinDesiredHeight(22)
 					[
 						SNew(STextBlock)
-						.Text(FText::FromName(CurveData.Name))
+						.Text(CurveDisplayName)
 						.TextStyle(FNiagaraEditorWidgetsStyle::Get(), "NiagaraEditor.CurveOverview.CurveComponentText")
 					]
 				];
+
+			if (CurveToolTipText.IsEmptyOrWhitespace() == false)
+			{
+				LabelBox->SetToolTipText(CurveToolTipText);
+			}
+
+			return LabelBox;
 		}
 		else if (InColumnName == ColumnNames.PinHeader)
 		{
@@ -143,6 +169,7 @@ private:
 private:
 	TWeakObjectPtr<UObject> CurveOwnerObject;
 	UNiagaraDataInterfaceCurveBase::FCurveData CurveData;
+	EAxisList::Type CurveAxis;
 	FSimpleMulticastDelegate CurveChangedDelegate;
 };
 
@@ -858,9 +885,11 @@ public:
 		TArray<UNiagaraDataInterfaceCurveBase::FCurveData> CurveDatas;
 		CurveOwnerObject->GetCurveData(CurveDatas);
 		TArray<FCurveEditorTreeItemID> TreeItemIds;
+		int32 CurveComponentIndex = 0;
 		for (const UNiagaraDataInterfaceCurveBase::FCurveData& CurveData : CurveDatas)
 		{
-			TSharedRef<FNiagaraCurveDetailsTreeItem> TreeItem = MakeShared<FNiagaraCurveDetailsTreeItem>(CurveOwnerObject, CurveData);
+			EAxisList::Type CurveAxis = FNiagaraEditorUtilities::VectorComponentToAxis(CurveDatas.Num(), CurveComponentIndex);
+			TSharedRef<FNiagaraCurveDetailsTreeItem> TreeItem = MakeShared<FNiagaraCurveDetailsTreeItem>(CurveOwnerObject, CurveData, CurveAxis);
 			TreeItem->OnCurveChanged().AddSP(this, &SNiagaraDataInterfaceCurveEditor::CurveChanged, TWeakPtr<FNiagaraCurveDetailsTreeItem>(TreeItem));
 			FCurveEditorTreeItem* NewItem = CurveEditor->AddTreeItem(FCurveEditorTreeItemID::Invalid());
 			NewItem->SetStrongItem(TreeItem);
@@ -869,6 +898,7 @@ public:
 			{
 				CurveEditor->PinCurve(CurveModel);
 			}
+			CurveComponentIndex++;
 		}
 
 		float KeySelectorLeftPadding = CurveProperties.Num() == 1 ? 0 : 7;
diff --git a/Engine/Plugins/FX/Niagara/Source/NiagaraEditorWidgets/Private/SNiagaraCurveOverview.cpp b/Engine/Plugins/FX/Niagara/Source/NiagaraEditorWidgets/Private/SNiagaraCurveOverview.cpp
--- a/Engine/Plugins/FX/Niagara/Source/NiagaraEditorWidgets/Private/SNiagaraCurveOverview.cpp
+++ b/Engine/Plugins/FX/Niagara/Source/NiagaraEditorWidgets/Private/SNiagaraCurveOverview.cpp
@@ -286,6 +286,11 @@ struct FNiagaraCurveOverviewTreeItem : public ICurveEditorTreeItem
 				AddWidgetsToLabelBox(CurveSelectionTreeNode);
 			}
 
+			if (CurveSelectionTreeNode->GetToolTipText().IsEmptyOrWhitespace() == false)
+			{
+				LabelBox->SetToolTipText(CurveSelectionTreeNode->GetToolTipText());
+			}
+
 			return LabelBox;
 		}
 		else if (InColumnName == ColumnNames.PinHeader)
