- Handle edge case where the number of renderers does not match the previously cached amount
- I don't have a repro for this, but from making an artifical hack I can repro the crash in both increased and decreased numbers

#rnx
#tests Hack the code to make a mismatch in the renderer count

diff --git a/Engine/Plugins/FX/Niagara/Source/Niagara/Private/NiagaraSystemRenderData.cpp b/Engine/Plugins/FX/Niagara/Source/Niagara/Private/NiagaraSystemRenderData.cpp
--- a/Engine/Plugins/FX/Niagara/Source/Niagara/Private/NiagaraSystemRenderData.cpp
+++ b/Engine/Plugins/FX/Niagara/Source/Niagara/Private/NiagaraSystemRenderData.cpp
@@ -182,15 +182,22 @@ void FNiagaraSystemRenderData::GenerateSetDynamicDataCommands(FSetDynamicDataCom
 		EmitterInst->ForEachEnabledRenderer(
 			[&](const UNiagaraRendererProperties* Properties)
 			{
+			#if WITH_EDITOR
+				if ( RendererIndex >= NumEmitterRenderers )
+				{
+					++RendererIndex;
+					return;
+				}
+			#endif
 				FNiagaraRenderer* Renderer = EmitterRenderers_GT[RendererIndex];
 				FNiagaraDynamicDataBase* NewData = nullptr;
 
 				if (Renderer && Properties->GetIsActive())
 				{
 					bool bRendererEditorEnabled = true;
-#if WITH_EDITORONLY_DATA
+				#if WITH_EDITORONLY_DATA
 					bRendererEditorEnabled = EmitterInst->IsDisabledFromIsolation() == false;
-#endif
+				#endif
 					if (bRendererEditorEnabled && !EmitterInst->IsComplete() && !SystemInstance->IsComplete())
 					{
 						NewData = Renderer->GenerateDynamicData(&SceneProxy, Properties, EmitterInst);
@@ -220,7 +227,10 @@ void FNiagaraSystemRenderData::GenerateSetDynamicDataCommands(FSetDynamicDataCom
 	{
 		// This can happen in the editor when modifying the number or renderers while the system is running and the render thread is already processing the data.
 		// in this case we just skip drawing this frame since the system will be reinitialized.
-		Commands.SetNum(Commands.Num() - RendererIndex);
+		if ( RendererIndex > NumEmitterRenderers )
+		{
+			Commands.SetNum(NumEmitterRenderers);
+		}
 	}
 #endif
 }
