- Reduce locks inside stateless data get function
- Stress test reduction (100 instances of 1 emitter 1 particle )
- CPU reduction on GPU path ~30%
- CPU reduction on CPU path ~75%

#tests Before / after comparison with many emitters running
#rnx

diff --git a/Engine/Plugins/FX/Niagara/Source/Niagara/Private/Stateless/NiagaraStatelessComputeManager.cpp b/Engine/Plugins/FX/Niagara/Source/Niagara/Private/Stateless/NiagaraStatelessComputeManager.cpp
--- a/Engine/Plugins/FX/Niagara/Source/Niagara/Private/Stateless/NiagaraStatelessComputeManager.cpp
+++ b/Engine/Plugins/FX/Niagara/Source/Niagara/Private/Stateless/NiagaraStatelessComputeManager.cpp
@@ -225,14 +225,7 @@ FNiagaraDataBuffer* FNiagaraStatelessComputeManager::GetDataBuffer(FRHICommandLi
 
 	SCOPE_CYCLE_COUNTER(STAT_NiagaraStateless_GetDataBuffer);
 
-	//-OPT: This lock is very conservative, ideally we only have it around the relevant parts
-	UE::TScopeLock ScopeLock(GetDataBufferGuard);
-
-	if (TUniquePtr<FStatelessDataCache>* ExistingData = UsedData.Find(EmitterKey))
-	{
-		return (*ExistingData)->DataBuffer;
-	}
-
+	// Is there any data to generate?
 	if (EmitterInstance->ExecutionState == ENiagaraExecutionState::Complete || EmitterInstance->ExecutionState == ENiagaraExecutionState::Disabled)
 	{
 		return nullptr;
@@ -262,29 +255,60 @@ FNiagaraDataBuffer* FNiagaraStatelessComputeManager::GetDataBuffer(FRHICommandLi
 
 	FNiagaraGpuComputeDispatchInterface* ComputeInterface = GetOwnerInterface();
 
-	const uint32 DataSetLayoutHash = EmitterInstance->EmitterData->ParticleDataSetCompiledData->GetLayoutHash();
+	// Until we add an extension to the render to notify about GDME start / end we can not allow GPU generation requests outside of PreInitViews / PreRender
+	// The shadow rendering, for example, will call GDME outside of this and it will result in crashes
+	// Therefore we fallback to CPU generation in these cases, if available
+	const bool bAllowGPUGeneration = bAllowDeferredGeneration;
+	const EComputeExecutionPath ComputeExecutionPath = DetermineComputeExecutionPath(EmitterData, ActiveParticles, bAllowGPUGeneration);
+	uint32 CountOffset = INDEX_NONE;
 
+	// Allocate / Reuse an existing buffer
 	FStatelessDataCache* CacheData = nullptr;
-	if (GUseDataBufferCache)
 	{
-		for (int32 i=0; i < FreeData.Num(); ++i)
+		const uint32 DataSetLayoutHash = EmitterInstance->EmitterData->ParticleDataSetCompiledData->GetLayoutHash();
+
+		UE::TScopeLock ScopeLock(GetDataBufferGuard);
+		if (TUniquePtr<FStatelessDataCache>* ExistingData = UsedData.Find(EmitterKey))
+		{
+			return (*ExistingData)->DataBuffer;
+		}
+
+		if (GUseDataBufferCache)
 		{
-			if (FreeData[i]->DataSetLayoutHash == DataSetLayoutHash)
+			for (int32 i = 0; i < FreeData.Num(); ++i)
 			{
-				CacheData = FreeData[i].Release();
-				FreeData.RemoveAtSwap(i, EAllowShrinking::No);
-				break;
+				if (FreeData[i]->DataSetLayoutHash == DataSetLayoutHash)
+				{
+					CacheData = FreeData[i].Release();
+					FreeData.RemoveAtSwap(i, EAllowShrinking::No);
+					break;
+				}
 			}
 		}
-	}
 
-	if ( CacheData == nullptr )
-	{
-		CacheData = new FStatelessDataCache();
-		CacheData->DataSetLayoutHash = DataSetLayoutHash;
-		CacheData->DataSetCompiledData = EmitterInstance->EmitterData->ParticleDataSetCompiledData;
-		CacheData->DataSet.Init(CacheData->DataSetCompiledData.Get());
-		CacheData->DataBuffer = new FNiagaraDataBuffer(&CacheData->DataSet);
+		if (CacheData == nullptr)
+		{
+			CacheData = new FStatelessDataCache();
+			CacheData->DataSetLayoutHash = DataSetLayoutHash;
+			CacheData->DataSetCompiledData = EmitterInstance->EmitterData->ParticleDataSetCompiledData;
+			CacheData->DataSet.Init(CacheData->DataSetCompiledData.Get());
+			CacheData->DataBuffer = new FNiagaraDataBuffer(&CacheData->DataSet);
+		}
+
+		// Note: We can not free the data when in paralell GDME mode as multiple lock / unlock operations could result in the wrong one providing the final data
+		// The plus side to adding back into the cache data is that multiple calls to get the same emitter's data will resolve to no active particles.
+		UsedData.Emplace(EmitterKey, CacheData);
+
+		if (ComputeExecutionPath == EComputeExecutionPath::GPU)
+		{
+			FNiagaraGPUInstanceCountManager& CountManager = ComputeInterface->GetGPUInstanceCounterManager();
+			CountOffset = CountManager.AcquireEntry();
+			if (CountOffset != INDEX_NONE)
+			{
+				GPUGenerationRequests.Emplace(CacheData->DataBuffer, EmitterInstance, ActiveParticles);
+				CountsToRelease.Add(CountOffset);
+			}
+		}
 	}
 
 #if NIAGARA_MEMORY_TRACKING && STATS
@@ -297,37 +321,23 @@ FNiagaraDataBuffer* FNiagaraStatelessComputeManager::GetDataBuffer(FRHICommandLi
 	INC_MEMORY_STAT_BY(STAT_NiagaraStateless_BufferMemory, CacheData->DataBuffer->GetAllocationSizeBytes());
 #endif
 
-	// Until we add an extension to the render to notify about GDME start / end we can not allow GPU generation requests outside of PreInitViews / PreRender
-	// The shadow rendering, for example, will call GDME outside of this and it will result in crashes
-	// Therefore we fallback to CPU generation in these cases, if available
-	const bool bAllowGPUGeneration = bAllowDeferredGeneration;
-
-	const EComputeExecutionPath ComputeExecutionPath = DetermineComputeExecutionPath(EmitterData, ActiveParticles, bAllowGPUGeneration);
+	bool bDidGenerateData = false;
 
 	switch (ComputeExecutionPath)
 	{
 		case EComputeExecutionPath::CPU:
 		{
-			if (!GenerateCPUDataForGPUSim(RHICmdList, EmitterInstance, CacheData->DataBuffer))
-			{
-				// Note: We can not free the data when in paralell GDME mode as multiple lock / unlock operations could result in the wrong one providing the final data
-				// The plus side to adding back into the cache data is that multiple calls to get the same emitter's data will resolve to no active particles.
-				UsedData.Emplace(EmitterKey, CacheData);
-				return nullptr;
-			}
+			bDidGenerateData = GenerateCPUDataForGPUSim(RHICmdList, EmitterInstance, CacheData->DataBuffer);
 			break;
 		}
 
 		case EComputeExecutionPath::GPU:
 		{
-			FNiagaraGPUInstanceCountManager& CountManager = ComputeInterface->GetGPUInstanceCounterManager();
-			const uint32 CountOffset = CountManager.AcquireEntry();
 			if (CountOffset != INDEX_NONE)
 			{
 				CacheData->DataBuffer->SetNumInstances(ActiveParticles);
 				CacheData->DataBuffer->SetGPUInstanceCountBufferOffset(CountOffset);
-				GPUGenerationRequests.Emplace(CacheData->DataBuffer, EmitterInstance, ActiveParticles);
-				CountsToRelease.Add(CountOffset);
+				bDidGenerateData = true;
 			}
 			// If we failed to alocate a count we will need to go through the CPU path (if available)
 			// This should never happen as we reserve a count up front via the compute proxy
@@ -340,27 +350,17 @@ FNiagaraDataBuffer* FNiagaraStatelessComputeManager::GetDataBuffer(FRHICommandLi
 				}
 
 				const bool bAllowCPUExec = EnumHasAnyFlags(EmitterInstance->EmitterData->FeatureMask, ENiagaraStatelessFeatureMask::ExecuteCPU);
-				const bool bDidGenerate = bAllowCPUExec && GenerateCPUDataForGPUSim(RHICmdList, EmitterInstance, CacheData->DataBuffer);
-				if (!bDidGenerate)
-				{
-					// Note: We can not free the data when in paralell GDME mode as multiple lock / unlock operations could result in the wrong one providing the final data
-					// The plus side to adding back into the cache data is that multiple calls to get the same emitter's data will resolve to no active particles.
-					UsedData.Emplace(EmitterKey, CacheData);
-					return nullptr;
-				}
+				bDidGenerateData = bAllowCPUExec && GenerateCPUDataForGPUSim(RHICmdList, EmitterInstance, CacheData->DataBuffer);
 			}
 			break;
 		}
 
 		default:
 			ensureMsgf(false, TEXT("No execution path was found for stateless emitter, data will not be generated"));
-			// We can add back to free data since we never locked / unlocked the buffer
-			FreeData.Emplace(CacheData);
-			return nullptr;
+			break;
 	}
 
-	UsedData.Emplace(EmitterKey, CacheData);
-	return CacheData->DataBuffer;
+	return bDidGenerateData ? CacheData->DataBuffer : nullptr;
 }
 
 void FNiagaraStatelessComputeManager::GenerateDataBufferForDebugging(FRHICommandListImmediate& RHICmdList, FNiagaraDataBuffer* DataBuffer, const NiagaraStateless::FEmitterInstance_RT* EmitterInstance) const
