[HWRT/Niagara] Updated GetRayTraceLODModelData(...) to get all data from RenderData->RayTracingProxy.

diff --git a/Engine/Plugins/FX/Niagara/Source/Niagara/Private/NiagaraMeshRendererProperties.cpp b/Engine/Plugins/FX/Niagara/Source/Niagara/Private/NiagaraMeshRendererProperties.cpp
--- a/Engine/Plugins/FX/Niagara/Source/Niagara/Private/NiagaraMeshRendererProperties.cpp
+++ b/Engine/Plugins/FX/Niagara/Source/Niagara/Private/NiagaraMeshRendererProperties.cpp
@@ -90,27 +90,46 @@ public:
 		}
 	}
 
+#if RHI_RAYTRACING
 	virtual void GetRayTraceLODModelData(FLODModelData& OutLODModelData, int32 LODLevel) const override
 	{
-		GetLODModelData(OutLODModelData, LODLevel);
-		if (OutLODModelData.LODIndex == INDEX_NONE)
-		{
-			return;
-		}
-
-		const FStaticMeshLODResources& LODResources = RenderData->LODResources[OutLODModelData.LODIndex];
-		OutLODModelData.RayTracingGeometry = LODResources.RayTracingGeometry;
+		LODLevel = FMath::Max(MinLOD, LODLevel);
 
-		if (FStaticMeshRayTracingProxy* RayTracingProxy = RenderData->RayTracingProxy)
+		if (RenderData->RayTracingProxy != nullptr)
 		{
-			const int32 RayTracingLOD = RayTracingProxy->bUsingRenderingLODs ? OutLODModelData.LODIndex : 0;
+			FStaticMeshRayTracingProxyLODArray& RayTracingLODs = RenderData->RayTracingProxy->LODs;
+
+			const int32 NumLODs = RayTracingLODs.Num();
+			const int32 RayTracingMinLOD = RenderData->RayTracingProxy->bUsingRenderingLODs ? RenderData->GetCurrentFirstLODIdx(LODLevel) : 0;
 
-			if (RayTracingProxy->LODs.IsValidIndex(RayTracingLOD))
+			for (int32 LODIndex = RayTracingMinLOD; LODIndex < NumLODs; ++LODIndex)
 			{
-				OutLODModelData.RayTracingGeometry = RayTracingProxy->LODs[RayTracingLOD].RayTracingGeometry;
+				const FStaticMeshRayTracingProxyLOD& RayTracingLOD = RayTracingLODs[LODIndex];
+				const FRayTracingGeometry& RayTracingGeometry = *RayTracingLOD.RayTracingGeometry;
+
+				if (RayTracingGeometry.IsValid() && !RayTracingGeometry.IsEvicted() && !RayTracingGeometry.HasPendingBuildRequest())
+				{
+					OutLODModelData.LODIndex = LODIndex;
+					OutLODModelData.RayTracingGeometry = &RayTracingGeometry;
+					OutLODModelData.NumVertices = RayTracingLOD.VertexBuffers->StaticMeshVertexBuffer.GetNumVertices();
+					OutLODModelData.NumIndices = RayTracingLOD.IndexBuffer->GetNumIndices();
+					OutLODModelData.Sections = *RayTracingLOD.Sections;
+					OutLODModelData.IndexBuffer = RayTracingLOD.IndexBuffer;
+					OutLODModelData.VertexFactoryUserData = RenderData->RayTracingProxy->LODVertexFactories->IsValidIndex(OutLODModelData.LODIndex) ? (*RenderData->RayTracingProxy->LODVertexFactories)[OutLODModelData.LODIndex].VertexFactory.GetUniformBuffer() : nullptr;
+
+					return;
+				}
 			}
 		}
+		else
+		{
+#if DO_CHECK
+			const UStaticMesh* StaticMesh = WeakStaticMesh.Get();
+			check(!IsRayTracingAllowed() || StaticMesh == nullptr || !StaticMesh->bSupportRayTracing);
+#endif
+		}
 	}
+#endif
 
 	virtual FIntVector2 GetLODRange() const override
 	{
