Couple of tweaks to niagara compilation
-increases the maximum number of concurrent compiles to 24 by default
-adds some periodic logging of the number of compilation tasks that are piling up in the queue
-makes deferred compilation of systems the default for -game, instead of just for the editor (fx.Niagara.OnDemandCompileEnabled now supports 0 for disabled, 1 for editor, and 2 for editor & game)
#tests Cooked enginetests with a maximum of 2 concurrent compiles to verify the reporting

diff --git a/Engine/Plugins/FX/Niagara/Source/Niagara/Private/NiagaraSystem.cpp b/Engine/Plugins/FX/Niagara/Source/Niagara/Private/NiagaraSystem.cpp
--- a/Engine/Plugins/FX/Niagara/Source/Niagara/Private/NiagaraSystem.cpp
+++ b/Engine/Plugins/FX/Niagara/Source/Niagara/Private/NiagaraSystem.cpp
@@ -181,11 +181,11 @@ FAutoConsoleVariableRef CVarNiagaraPrecachePSOAtAssetLoadingTime(
 	ECVF_Default);
 
 #if WITH_EDITORONLY_DATA
-static int GNiagaraOnDemandCompileEnabled = 1;
+static int GNiagaraOnDemandCompileEnabled = 2;
 static FAutoConsoleVariableRef CVarNiagaraOnDemandCompileEnabled(
 	TEXT("fx.Niagara.OnDemandCompileEnabled"),
 	GNiagaraOnDemandCompileEnabled,
-	TEXT("Compiles Niagara Systems on demand rather than on post load."),
+	TEXT("Compiles Niagara Systems on demand rather than on post load.  1 - Will defer when in editor; 2 - Will defer in editor & game"),
 	ECVF_Default
 );
 
@@ -788,7 +788,9 @@ void UNiagaraSystem::UpdateSystemAfterLoad()
 				// commandlet. We don't need it for normal post-loading.
 				Modify();
 			}
-			if ( !GIsEditor || GNiagaraOnDemandCompileEnabled == 0 )
+
+			const bool CompileOnLoad = !GNiagaraOnDemandCompileEnabled || (GIsEditor && GNiagaraOnDemandCompileEnabled == 1);
+			if (CompileOnLoad)
 			{
 				RequestCompile(false);
 			}
diff --git a/Engine/Plugins/FX/Niagara/Source/NiagaraEditor/Private/NiagaraSystemCompilingManager.cpp b/Engine/Plugins/FX/Niagara/Source/NiagaraEditor/Private/NiagaraSystemCompilingManager.cpp
--- a/Engine/Plugins/FX/Niagara/Source/NiagaraEditor/Private/NiagaraSystemCompilingManager.cpp
+++ b/Engine/Plugins/FX/Niagara/Source/NiagaraEditor/Private/NiagaraSystemCompilingManager.cpp
@@ -15,7 +15,7 @@
 
 #define LOCTEXT_NAMESPACE "NiagaraCompilationManager"
 
-static int GNiagaraCompilationMaxActiveTaskCount = 16;
+static int GNiagaraCompilationMaxActiveTaskCount = 24;
 static FAutoConsoleVariableRef CVarNiagaraCompilationMaxActiveTaskCount(
 	TEXT("fx.Niagara.Compilation.MaxActiveTaskCount"),
 	GNiagaraCompilationMaxActiveTaskCount,
@@ -567,10 +567,30 @@ bool FNiagaraSystemCompilingManager::ConditionalLaunchTask()
 	{
 		FReadScopeLock Read(QueueLock);
 
-		if (QueuedRequests.IsEmpty() || ActiveTasks.Num() >= GNiagaraCompilationMaxActiveTaskCount)
+		const int32 ActiveRequestCount = ActiveTasks.Num();
+		const int32 QueuedRequestCount = QueuedRequests.Num();
+
+		if (QueuedRequestCount == 0)
 		{
 			return false;
 		}
+		else if (ActiveRequestCount >= GNiagaraCompilationMaxActiveTaskCount)
+		{
+			if (QueuedRequestCount > PeakReportedPendingTaskCount)
+			{
+				static const double MinTimeBetweenReports = 5.0;
+				const double CurrentTime = FPlatformTime::Seconds();
+				if (CurrentTime - LastTaskCountReportTime >= MinTimeBetweenReports)
+				{
+					PeakReportedPendingTaskCount = QueuedRequestCount;
+					LastTaskCountReportTime = CurrentTime;
+
+					UE_LOG(LogNiagaraEditor, Display, TEXT("Wanted to launch a task, but there's already %d active tasks - keeping %d requests queued. Consider increasing fx.Niagara.Compilation.MaxActiveTaskCount which is currently %d"),
+						ActiveRequestCount, QueuedRequestCount, GNiagaraCompilationMaxActiveTaskCount);
+				}
+			}
+			return false;
+		}
 	}
 
 	{
diff --git a/Engine/Plugins/FX/Niagara/Source/NiagaraEditor/Private/NiagaraSystemCompilingManager.h b/Engine/Plugins/FX/Niagara/Source/NiagaraEditor/Private/NiagaraSystemCompilingManager.h
--- a/Engine/Plugins/FX/Niagara/Source/NiagaraEditor/Private/NiagaraSystemCompilingManager.h
+++ b/Engine/Plugins/FX/Niagara/Source/NiagaraEditor/Private/NiagaraSystemCompilingManager.h
@@ -68,4 +68,7 @@ protected:
 
 	mutable FRWLock GameThreadFunctionLock;
 	TArray<FGameThreadFunction> GameThreadFunctions;
+
+	int32 PeakReportedPendingTaskCount = 0;
+	double LastTaskCountReportTime = 0;
 };
