Followup fixes for ondemand niagara compiling and logging of queued compile requests
#tests verified cook will actually report the peak count of queued compiles and ran with -game to test on demand compilations

diff --git a/Engine/Plugins/FX/Niagara/Source/Niagara/Private/NiagaraSystem.cpp b/Engine/Plugins/FX/Niagara/Source/Niagara/Private/NiagaraSystem.cpp
--- a/Engine/Plugins/FX/Niagara/Source/Niagara/Private/NiagaraSystem.cpp
+++ b/Engine/Plugins/FX/Niagara/Source/Niagara/Private/NiagaraSystem.cpp
@@ -789,7 +789,8 @@ void UNiagaraSystem::UpdateSystemAfterLoad()
 				Modify();
 			}
 
-			const bool CompileOnLoad = !GNiagaraOnDemandCompileEnabled || (GIsEditor && GNiagaraOnDemandCompileEnabled == 1);
+			const bool CompileOnLoad = !GNiagaraOnDemandCompileEnabled || (!GIsEditor && GNiagaraOnDemandCompileEnabled == 1);
+
 			if (CompileOnLoad)
 			{
 				RequestCompile(false);
diff --git a/Engine/Plugins/FX/Niagara/Source/NiagaraEditor/Private/NiagaraSystemCompilingManager.cpp b/Engine/Plugins/FX/Niagara/Source/NiagaraEditor/Private/NiagaraSystemCompilingManager.cpp
--- a/Engine/Plugins/FX/Niagara/Source/NiagaraEditor/Private/NiagaraSystemCompilingManager.cpp
+++ b/Engine/Plugins/FX/Niagara/Source/NiagaraEditor/Private/NiagaraSystemCompilingManager.cpp
@@ -576,17 +576,19 @@ bool FNiagaraSystemCompilingManager::ConditionalLaunchTask()
 		}
 		else if (ActiveRequestCount >= GNiagaraCompilationMaxActiveTaskCount)
 		{
-			if (QueuedRequestCount > PeakReportedPendingTaskCount)
+			PeakPendingTaskCount = FMath::Max(PeakPendingTaskCount, QueuedRequestCount);
+
+			if (PeakPendingTaskCount > PeakReportedPendingTaskCount)
 			{
-				static const double MinTimeBetweenReports = 5.0;
+				constexpr double MinTimeBetweenReports = 5.0;
 				const double CurrentTime = FPlatformTime::Seconds();
 				if (CurrentTime - LastTaskCountReportTime >= MinTimeBetweenReports)
 				{
-					PeakReportedPendingTaskCount = QueuedRequestCount;
+					PeakReportedPendingTaskCount = PeakPendingTaskCount;
 					LastTaskCountReportTime = CurrentTime;
 
 					UE_LOG(LogNiagaraEditor, Display, TEXT("Wanted to launch a task, but there's already %d active tasks - keeping %d requests queued. Consider increasing fx.Niagara.Compilation.MaxActiveTaskCount which is currently %d"),
-						ActiveRequestCount, QueuedRequestCount, GNiagaraCompilationMaxActiveTaskCount);
+						ActiveRequestCount, PeakPendingTaskCount, GNiagaraCompilationMaxActiveTaskCount);
 				}
 			}
 			return false;
diff --git a/Engine/Plugins/FX/Niagara/Source/NiagaraEditor/Private/NiagaraSystemCompilingManager.h b/Engine/Plugins/FX/Niagara/Source/NiagaraEditor/Private/NiagaraSystemCompilingManager.h
--- a/Engine/Plugins/FX/Niagara/Source/NiagaraEditor/Private/NiagaraSystemCompilingManager.h
+++ b/Engine/Plugins/FX/Niagara/Source/NiagaraEditor/Private/NiagaraSystemCompilingManager.h
@@ -70,5 +70,6 @@ protected:
 	TArray<FGameThreadFunction> GameThreadFunctions;
 
 	int32 PeakReportedPendingTaskCount = 0;
+	int32 PeakPendingTaskCount = 0;
 	double LastTaskCountReportTime = 0;
 };
