- Change sim cache read function to take in a context object
- The context object contains more information, such as transforms used for rebasing

#rnx

diff --git a/Engine/Plugins/FX/Niagara/Source/Niagara/Classes/NiagaraDataInterfaceArray.h b/Engine/Plugins/FX/Niagara/Source/Niagara/Classes/NiagaraDataInterfaceArray.h
--- a/Engine/Plugins/FX/Niagara/Source/Niagara/Classes/NiagaraDataInterfaceArray.h
+++ b/Engine/Plugins/FX/Niagara/Source/Niagara/Classes/NiagaraDataInterfaceArray.h
@@ -125,7 +125,7 @@ public:
 	NIAGARA_API virtual UObject* SimCacheBeginWrite(UObject* SimCache, FNiagaraSystemInstance* NiagaraSystemInstance, const void* OptionalPerInstanceData, FNiagaraSimCacheFeedbackContext& FeedbackContext) const override;
 	NIAGARA_API virtual bool SimCacheWriteFrame(UObject* StorageObject, int FrameIndex, FNiagaraSystemInstance* SystemInstance, const void* OptionalPerInstanceData, FNiagaraSimCacheFeedbackContext& FeedbackContext) const override;
 	//virtual bool SimCacheEndWrite(UObject* StorageObject) const override;
-	NIAGARA_API virtual bool SimCacheReadFrame(UObject* StorageObject, int FrameA, int FrameB, float Interp, FNiagaraSystemInstance* SystemInstance, void* OptionalPerInstanceData) override;
+	NIAGARA_API virtual bool SimCacheReadFrame(const FNiagaraSimCacheDataInterfaceReadContext& ReadContext) override;
 	NIAGARA_API virtual bool SimCacheCompareFrame(const UObject* LhsStorageObject, const UObject* RhsStorageObject, int FrameIndex, TOptional<float> Tolerance, FString& OutErrors) const override;
 	//~ INiagaraSimCacheCustomStorageInterface interface END
 	NIAGARA_API FString SimCacheVisualizerRead(const UNDIArraySimCacheData* CacheData, const FNDIArraySimCacheDataFrame& FrameData, int Element) const;
diff --git a/Engine/Plugins/FX/Niagara/Source/Niagara/Classes/NiagaraDataInterfaceRenderTarget2D.h b/Engine/Plugins/FX/Niagara/Source/Niagara/Classes/NiagaraDataInterfaceRenderTarget2D.h
--- a/Engine/Plugins/FX/Niagara/Source/Niagara/Classes/NiagaraDataInterfaceRenderTarget2D.h
+++ b/Engine/Plugins/FX/Niagara/Source/Niagara/Classes/NiagaraDataInterfaceRenderTarget2D.h
@@ -128,7 +128,7 @@ public:
 	NIAGARA_API virtual UObject* SimCacheBeginWrite(UObject* SimCache, FNiagaraSystemInstance* NiagaraSystemInstance, const void* OptionalPerInstanceData, FNiagaraSimCacheFeedbackContext& FeedbackContext) const override;
 	NIAGARA_API virtual bool SimCacheWriteFrame(UObject* StorageObject, int FrameIndex, FNiagaraSystemInstance* SystemInstance, const void* OptionalPerInstanceData, FNiagaraSimCacheFeedbackContext& FeedbackContext) const override;
 	NIAGARA_API virtual bool SimCacheEndWrite(UObject* StorageObject) const override;
-	NIAGARA_API virtual bool SimCacheReadFrame(UObject* StorageObject, int FrameA, int FrameB, float Interp, FNiagaraSystemInstance* SystemInstance, void* OptionalPerInstanceData) override;
+	NIAGARA_API virtual bool SimCacheReadFrame(const FNiagaraSimCacheDataInterfaceReadContext& ReadContext) override;
 	NIAGARA_API virtual bool SimCacheCompareFrame(const UObject* LhsStorageObject, const UObject* RhsStorageObject, int FrameIndex, TOptional<float> Tolerance, FString& OutErrors) const override;
 	NIAGARA_API virtual TSharedPtr<FJsonObject> SimCacheToJson(const UObject* StorageObject, int FrameIndex, TOptional<FString> TargetFolder, TOptional<FString> FilenamePrefix) const override;
 	//~ UNiagaraDataInterface interface END
diff --git a/Engine/Plugins/FX/Niagara/Source/Niagara/Classes/NiagaraDataInterfaceRenderTargetVolume.h b/Engine/Plugins/FX/Niagara/Source/Niagara/Classes/NiagaraDataInterfaceRenderTargetVolume.h
--- a/Engine/Plugins/FX/Niagara/Source/Niagara/Classes/NiagaraDataInterfaceRenderTargetVolume.h
+++ b/Engine/Plugins/FX/Niagara/Source/Niagara/Classes/NiagaraDataInterfaceRenderTargetVolume.h
@@ -114,7 +114,7 @@ public:
 	NIAGARA_API virtual UObject* SimCacheBeginWrite(UObject* SimCache, FNiagaraSystemInstance* NiagaraSystemInstance, const void* OptionalPerInstanceData, FNiagaraSimCacheFeedbackContext& FeedbackContext) const override;
 	NIAGARA_API virtual bool SimCacheWriteFrame(UObject* StorageObject, int FrameIndex, FNiagaraSystemInstance* SystemInstance, const void* OptionalPerInstanceData, FNiagaraSimCacheFeedbackContext& FeedbackContext) const override;
 	NIAGARA_API virtual bool SimCacheEndWrite(UObject* StorageObject) const override;
-	NIAGARA_API virtual bool SimCacheReadFrame(UObject* StorageObject, int FrameA, int FrameB, float Interp, FNiagaraSystemInstance* SystemInstance, void* OptionalPerInstanceData) override;
+	NIAGARA_API virtual bool SimCacheReadFrame(const FNiagaraSimCacheDataInterfaceReadContext& ReadContext) override;
 	//~ UNiagaraDataInterface interface END
 
 	NIAGARA_API void VMGetSize(FVectorVMExternalFunctionContext& Context);
diff --git a/Engine/Plugins/FX/Niagara/Source/Niagara/Classes/NiagaraSimCacheCustomStorageInterface.h b/Engine/Plugins/FX/Niagara/Source/Niagara/Classes/NiagaraSimCacheCustomStorageInterface.h
--- a/Engine/Plugins/FX/Niagara/Source/Niagara/Classes/NiagaraSimCacheCustomStorageInterface.h
+++ b/Engine/Plugins/FX/Niagara/Source/Niagara/Classes/NiagaraSimCacheCustomStorageInterface.h
@@ -10,6 +10,64 @@
 struct FNiagaraSimCacheFeedbackContext;
 class FNiagaraSystemInstance;
 
+// Context used when reading data for a data interface
+// Contains all the information requires to find the data / interpolate / rebase information
+struct FNiagaraSimCacheDataInterfaceReadContext
+{
+	friend class UNiagaraSimCache;
+
+public:
+	// Get the stored data interface object, will assert on nullptr
+	template<typename T> T* GetStorageObject() const  { return CastChecked<T>(StorageObject); }
+	// Get the stored data interface object, can return nullptr if the cast fails
+	template<typename T> T* GetOptionalStorageObject() const { return Cast<T>(StorageObject); }
+
+	// Get the system instance the cache read is attached to
+	FNiagaraSystemInstance* GetSystemInstance() const { return SystemInstance; }
+
+	// Get the system instance data interface per instance data, will assert on nullptr
+	template<typename T> T* GetPerInstanceData() const { check(OptionalPerInstanceData); return static_cast<T*>(OptionalPerInstanceData); }
+
+	// Test to see if we should interpolate the data or not
+	// The default value is expected to be the 'just works' case for when the user did not set force include / exlucde and interpolation is enabled
+	bool ShouldInterpolate(bool DefaultValue) const { return InterpolateData.Get(DefaultValue); }
+
+	// Get the currrent frame we are interpolating from
+	int32 GetFrameIndexA() const { return(FrameIndexA); }
+	// Get the currrent frame we are interpolating to
+	int32 GetFrameIndexB() const { return(FrameIndexB); }
+	// Get the currrent frame interpolation fraction
+	int32 GetFrameInterp() const { return(FrameInterp); }
+
+	// Test to see if we should rebase the data or not
+	// The default value is expected to be the 'just works' case for when the user did not set force include / exlucde and rebasing is enabled
+	bool ShouldRebaseData(bool DefaultValue) const { return RebaseData.Get(DefaultValue); }
+
+	// Get the transform to rebase data, this is not world space but LWC tile space
+	const FTransform& GetRebaseTransformA() const { return RebaseTransformA; }
+	// Get the transform to rebase data, this is not world space but LWC tile space
+	const FTransform& GetRebaseTransformB() const { return RebaseTransformB; }
+
+	// Get the LWC tile the instance was in during caching for frame A
+	FVector3f GetLWCTileA() const { return LWCTileA; }
+	// Get the LWC tile the instance was in during caching for frame B
+	FVector3f GetLWCTileB() const { return LWCTileB; }
+
+private:
+	UObject*				StorageObject = nullptr;
+	FNiagaraSystemInstance*	SystemInstance = nullptr;
+	void*					OptionalPerInstanceData = nullptr;
+	TOptional<bool>			InterpolateData;
+	int						FrameIndexA = 0;
+	int						FrameIndexB = 0;
+	float					FrameInterp = 0.0f;
+	TOptional<bool>			RebaseData;
+	FTransform				RebaseTransformA;
+	FTransform				RebaseTransformB;
+	FVector3f				LWCTileA = FVector3f::ZeroVector;
+	FVector3f				LWCTileB = FVector3f::ZeroVector;
+};
+
 // Interface for UObjects to implement renderable mesh
 UINTERFACE(MinimalAPI)
 class UNiagaraSimCacheCustomStorageInterface : public UInterface
@@ -50,6 +108,14 @@ public:
 	/**
 	Read a frame of data from the simulation cache.
 	*/
+	virtual bool SimCacheReadFrame(const FNiagaraSimCacheDataInterfaceReadContext& ReadContext)
+	{
+		PRAGMA_DISABLE_DEPRECATION_WARNINGS
+		return SimCacheReadFrame(ReadContext.GetStorageObject<UObject>(), ReadContext.GetFrameIndexA(), ReadContext.GetFrameIndexB(), ReadContext.GetFrameInterp(), ReadContext.GetSystemInstance(), ReadContext.GetPerInstanceData<void>());
+		PRAGMA_ENABLE_DEPRECATION_WARNINGS
+	}
+
+	UE_DEPRECATED(5.6, "This function has been deprecated in favor of one which uses a context")
 	virtual bool SimCacheReadFrame(UObject* StorageObject, int FrameA, int FrameB, float Interp, FNiagaraSystemInstance* SystemInstance, void* OptionalPerInstanceData) { return true; }
 
 	/**
diff --git a/Engine/Plugins/FX/Niagara/Source/Niagara/Internal/DataInterface/NiagaraDataInterfaceDataChannelWrite.h b/Engine/Plugins/FX/Niagara/Source/Niagara/Internal/DataInterface/NiagaraDataInterfaceDataChannelWrite.h
--- a/Engine/Plugins/FX/Niagara/Source/Niagara/Internal/DataInterface/NiagaraDataInterfaceDataChannelWrite.h
+++ b/Engine/Plugins/FX/Niagara/Source/Niagara/Internal/DataInterface/NiagaraDataInterfaceDataChannelWrite.h
@@ -210,7 +210,7 @@ public:
 	virtual UObject* SimCacheBeginWrite(UObject* SimCache, FNiagaraSystemInstance* NiagaraSystemInstance, const void* OptionalPerInstanceData, FNiagaraSimCacheFeedbackContext& FeedbackContext) const override;
 	virtual bool SimCacheWriteFrame(UObject* StorageObject, int FrameIndex, FNiagaraSystemInstance* SystemInstance, const void* OptionalPerInstanceData, FNiagaraSimCacheFeedbackContext& FeedbackContext) const override;
 	virtual bool SimCacheEndWrite(UObject* StorageObject) const override;
-	virtual bool SimCacheReadFrame(UObject* StorageObject, int FrameA, int FrameB, float Interp, FNiagaraSystemInstance* SystemInstance, void* OptionalPerInstanceData) override;
+	virtual bool SimCacheReadFrame(const FNiagaraSimCacheDataInterfaceReadContext& ReadContext) override;
 	virtual void SimCachePostReadFrame(void* OptionalPerInstanceData, FNiagaraSystemInstance* SystemInstance) override;
 	virtual bool SimCacheCompareFrame(const UObject* LhsStorageObject, const UObject* RhsStorageObject, int FrameIndex, TOptional<float> Tolerance, FString& OutErrors) const override;
 protected:
diff --git a/Engine/Plugins/FX/Niagara/Source/Niagara/Internal/DataInterface/NiagaraDataInterfaceMemoryBuffer.h b/Engine/Plugins/FX/Niagara/Source/Niagara/Internal/DataInterface/NiagaraDataInterfaceMemoryBuffer.h
--- a/Engine/Plugins/FX/Niagara/Source/Niagara/Internal/DataInterface/NiagaraDataInterfaceMemoryBuffer.h
+++ b/Engine/Plugins/FX/Niagara/Source/Niagara/Internal/DataInterface/NiagaraDataInterfaceMemoryBuffer.h
@@ -70,7 +70,7 @@ public:
 	virtual UObject* SimCacheBeginWrite(UObject* SimCache, FNiagaraSystemInstance* NiagaraSystemInstance, const void* OptionalPerInstanceData, FNiagaraSimCacheFeedbackContext& FeedbackContext) const override;
 	virtual bool SimCacheWriteFrame(UObject* StorageObject, int FrameIndex, FNiagaraSystemInstance* SystemInstance, const void* OptionalPerInstanceData, FNiagaraSimCacheFeedbackContext& FeedbackContext) const override;
 	//virtual bool SimCacheEndWrite(UObject* StorageObject) const override;
-	virtual bool SimCacheReadFrame(UObject* StorageObject, int FrameA, int FrameB, float Interp, FNiagaraSystemInstance* SystemInstance, void* OptionalPerInstanceData) override;
+	virtual bool SimCacheReadFrame(const FNiagaraSimCacheDataInterfaceReadContext& ReadContext) override;
 	//~ UNiagaraDataInterface interface END
 };
 
diff --git a/Engine/Plugins/FX/Niagara/Source/Niagara/Internal/DataInterface/NiagaraDataInterfaceSimpleCounter.h b/Engine/Plugins/FX/Niagara/Source/Niagara/Internal/DataInterface/NiagaraDataInterfaceSimpleCounter.h
--- a/Engine/Plugins/FX/Niagara/Source/Niagara/Internal/DataInterface/NiagaraDataInterfaceSimpleCounter.h
+++ b/Engine/Plugins/FX/Niagara/Source/Niagara/Internal/DataInterface/NiagaraDataInterfaceSimpleCounter.h
@@ -62,7 +62,7 @@ public:
 	//~ INiagaraSimCacheCustomStorageInterface interface BEGIN
 	virtual UObject* SimCacheBeginWrite(UObject* SimCache, FNiagaraSystemInstance* NiagaraSystemInstance, const void* OptionalPerInstanceData, FNiagaraSimCacheFeedbackContext& FeedbackContext) const override;
 	virtual bool SimCacheWriteFrame(UObject* StorageObject, int FrameIndex, FNiagaraSystemInstance* SystemInstance, const void* OptionalPerInstanceData, FNiagaraSimCacheFeedbackContext& FeedbackContext) const override;
-	virtual bool SimCacheReadFrame(UObject* StorageObject, int FrameA, int FrameB, float Interp, FNiagaraSystemInstance* SystemInstance, void* OptionalPerInstanceData) override;
+	virtual bool SimCacheReadFrame(const FNiagaraSimCacheDataInterfaceReadContext& ReadContext) override;
 	virtual bool SimCacheCompareFrame(const UObject* LhsStorageObject, const UObject* RhsStorageObject, int FrameIndex, TOptional<float> Tolerance, FString& OutErrors) const override;
 	//~ UNiagaraDataInterface interface END
 
diff --git a/Engine/Plugins/FX/Niagara/Source/Niagara/Private/DataInterface/NiagaraDataInterfaceDataChannelWrite.cpp b/Engine/Plugins/FX/Niagara/Source/Niagara/Private/DataInterface/NiagaraDataInterfaceDataChannelWrite.cpp
--- a/Engine/Plugins/FX/Niagara/Source/Niagara/Private/DataInterface/NiagaraDataInterfaceDataChannelWrite.cpp
+++ b/Engine/Plugins/FX/Niagara/Source/Niagara/Private/DataInterface/NiagaraDataInterfaceDataChannelWrite.cpp
@@ -926,13 +926,15 @@ bool UNiagaraDataInterfaceDataChannelWrite::SimCacheWriteFrame(UObject* StorageO
 	return false;
 }
 
-bool UNiagaraDataInterfaceDataChannelWrite::SimCacheReadFrame(UObject* StorageObject, int FrameA, int FrameB, float Interp, FNiagaraSystemInstance* SystemInstance, void* OptionalPerInstanceData)
+bool UNiagaraDataInterfaceDataChannelWrite::SimCacheReadFrame(const FNiagaraSimCacheDataInterfaceReadContext& ReadContext)
 {
 	if (Channel == nullptr || Channel->Get() == nullptr)
 	{
 		return false;
 	}
 
+	FNiagaraSystemInstance* SystemInstance = ReadContext.GetSystemInstance();
+
 	FNiagaraDataChannelDataPtr DataChannelData;
 	if (FNiagaraWorldManager* WorldMan = FNiagaraWorldManager::Get(SystemInstance->GetWorld()))
 	{
@@ -947,25 +949,59 @@ bool UNiagaraDataInterfaceDataChannelWrite::SimCacheReadFrame(UObject* StorageOb
 		return false;
 	}
 	
-	if (UNDIDataChannelWriteSimCacheData* Storage = Cast<UNDIDataChannelWriteSimCacheData>(StorageObject))
+	if (UNDIDataChannelWriteSimCacheData* Storage = ReadContext.GetOptionalStorageObject<UNDIDataChannelWriteSimCacheData>())
 	{
-		if (Storage->FrameData.IsValidIndex(FrameA))
+		const int32 FrameIndex = ReadContext.GetFrameIndexA();
+		if (Storage->FrameData.IsValidIndex(FrameIndex))
 		{
-			FNDIDataChannelWriteSimCacheFrame& Frame = Storage->FrameData[FrameA];
+			FNDIDataChannelWriteSimCacheFrame& Frame = Storage->FrameData[FrameIndex];
 			FNiagaraDataChannelPublishRequest PublishRequest;
 			PublishRequest.bVisibleToGame = Frame.bVisibleToGame;
 			PublishRequest.bVisibleToCPUSims = Frame.bVisibleToCPUSims;
 			PublishRequest.bVisibleToGPUSims = Frame.bVisibleToGPUSims;
 #if WITH_NIAGARA_DEBUGGER
-			PublishRequest.DebugSource = FString::Format(TEXT("{0} (Sim cache {1})"), {SystemInstance->GetSystem()->GetName(), GetPathNameSafe(StorageObject->GetOuter())});
+			PublishRequest.DebugSource = FString::Format(TEXT("{0} (Sim cache {1})"), {SystemInstance->GetSystem()->GetName(), GetPathNameSafe(Storage->GetOuter())});
 #endif
 
 			PublishRequest.GameData = MakeShared<FNiagaraDataChannelGameData>(Channel->Get()->GetLayoutInfo());
 			PublishRequest.GameData->SetNum(Frame.NumElements);
-			for (int32 i = 0; i < Frame.VariableData.Num(); i++)
+			PublishRequest.LwcTile = ReadContext.GetLWCTileA();
+
+			if (ReadContext.ShouldRebaseData(true))
+			{
+				const FNiagaraTypeDefinition& PositionTypeDef = FNiagaraTypeDefinition::GetPositionDef();
+				for (int32 i = 0; i < Frame.VariableData.Num(); ++i)
+				{
+					const FNDIDataChannelWriteSimCacheFrameBuffer& Buffer = Frame.VariableData[i];
+					if (Buffer.SourceVar.GetType() == PositionTypeDef)
+					{
+						check(Buffer.Size == sizeof(FVector));
+
+						//-OPT: Could set directly to avoid the copy
+						TArray<uint8> TempData;
+						TempData.AddUninitialized(Buffer.Data.Num());
+						for (int32 Element=0; Element < Frame.NumElements; ++Element)
+						{
+							FVector Position;
+							FMemory::Memcpy(&Position, Buffer.Data.GetData() + (sizeof(FVector) * Element), sizeof(FVector));
+							Position  = ReadContext.GetRebaseTransformA().TransformPosition(Position);
+							FMemory::Memcpy(TempData.GetData() + (sizeof(FVector) * Element), &Position, sizeof(FVector));
+						}
+						PublishRequest.GameData->SetFromSimCache(Buffer.SourceVar, TempData, sizeof(FVector));
+					}
+					else
+					{
+						PublishRequest.GameData->SetFromSimCache(Buffer.SourceVar, Buffer.Data, Buffer.Size);
+					}
+				}
+			}
+			else
 			{
-				const FNDIDataChannelWriteSimCacheFrameBuffer& Buffer = Frame.VariableData[i];
-				PublishRequest.GameData->SetFromSimCache(Buffer.SourceVar, Buffer.Data, Buffer.Size);
+				for (int32 i = 0; i < Frame.VariableData.Num(); ++i)
+				{
+					const FNDIDataChannelWriteSimCacheFrameBuffer& Buffer = Frame.VariableData[i];
+					PublishRequest.GameData->SetFromSimCache(Buffer.SourceVar, Buffer.Data, Buffer.Size);
+				}
 			}
 			
 			DataChannelData->Publish(PublishRequest);
diff --git a/Engine/Plugins/FX/Niagara/Source/Niagara/Private/DataInterface/NiagaraDataInterfaceMemoryBuffer.cpp b/Engine/Plugins/FX/Niagara/Source/Niagara/Private/DataInterface/NiagaraDataInterfaceMemoryBuffer.cpp
--- a/Engine/Plugins/FX/Niagara/Source/Niagara/Private/DataInterface/NiagaraDataInterfaceMemoryBuffer.cpp
+++ b/Engine/Plugins/FX/Niagara/Source/Niagara/Private/DataInterface/NiagaraDataInterfaceMemoryBuffer.cpp
@@ -1009,14 +1009,13 @@ bool UNiagaraDataInterfaceMemoryBuffer::SimCacheWriteFrame(UObject* StorageObjec
 	return true;
 }
 
-bool UNiagaraDataInterfaceMemoryBuffer::SimCacheReadFrame(UObject* StorageObject, int FrameA, int FrameB, float Interp, FNiagaraSystemInstance* SystemInstance, void* OptionalPerInstanceData)
+bool UNiagaraDataInterfaceMemoryBuffer::SimCacheReadFrame(const FNiagaraSimCacheDataInterfaceReadContext& ReadContext)
 {
 	using namespace NDIMemoryBufferLocal;
 
-	check(OptionalPerInstanceData && StorageObject);
-	UNDIMemoryBufferSimCacheData* CacheData = CastChecked<UNDIMemoryBufferSimCacheData>(StorageObject);
+	UNDIMemoryBufferSimCacheData* CacheData = ReadContext.GetStorageObject<UNDIMemoryBufferSimCacheData>();
 
-	const int FrameIndex = FrameA;
+	const int FrameIndex = ReadContext.GetFrameIndexA();
 	if (!CacheData->FrameData.IsValidIndex(FrameIndex))
 	{
 		return false;
@@ -1024,7 +1023,7 @@ bool UNiagaraDataInterfaceMemoryBuffer::SimCacheReadFrame(UObject* StorageObject
 
 	const FNDIMemoryBufferSimCacheDataFrame& FrameData = CacheData->FrameData[FrameIndex];
 
-	FNDIInstanceData_GameThread* InstanceData_GT = static_cast<FNDIInstanceData_GameThread*>(OptionalPerInstanceData);
+	FNDIInstanceData_GameThread* InstanceData_GT = ReadContext.GetPerInstanceData<FNDIInstanceData_GameThread>();
 	if (FrameData.CpuBufferSize > 0 && InstanceData_GT->bUsedWithCpu)
 	{		
 		InstanceData_GT->NumElements = FrameData.CpuBufferSize;
@@ -1034,6 +1033,7 @@ bool UNiagaraDataInterfaceMemoryBuffer::SimCacheReadFrame(UObject* StorageObject
 
 	if (FrameData.GpuBufferSize > 0 && InstanceData_GT->bUsedWithGpu)
 	{
+		FNiagaraSystemInstance* SystemInstance = ReadContext.GetSystemInstance();
 		TConstArrayView<uint32> GpuData(&CacheData->BufferData[FrameData.GpuDataOffset], FrameData.GpuBufferSize);
 		ENQUEUE_RENDER_COMMAND(FNDIMemoryBuffer_UpdateData)
 		(
diff --git a/Engine/Plugins/FX/Niagara/Source/Niagara/Private/DataInterface/NiagaraDataInterfaceSimpleCounter.cpp b/Engine/Plugins/FX/Niagara/Source/Niagara/Private/DataInterface/NiagaraDataInterfaceSimpleCounter.cpp
--- a/Engine/Plugins/FX/Niagara/Source/Niagara/Private/DataInterface/NiagaraDataInterfaceSimpleCounter.cpp
+++ b/Engine/Plugins/FX/Niagara/Source/Niagara/Private/DataInterface/NiagaraDataInterfaceSimpleCounter.cpp
@@ -531,21 +531,19 @@ bool UNiagaraDataInterfaceSimpleCounter::SimCacheWriteFrame(UObject* StorageObje
 	return true;
 }
 
-bool UNiagaraDataInterfaceSimpleCounter::SimCacheReadFrame(UObject* StorageObject, int FrameA, int FrameB, float Interp, FNiagaraSystemInstance* SystemInstance, void* OptionalPerInstanceData)
+bool UNiagaraDataInterfaceSimpleCounter::SimCacheReadFrame(const FNiagaraSimCacheDataInterfaceReadContext& ReadContext)
 {
-	check(OptionalPerInstanceData && StorageObject);
-
-	const int32 ValueOffset = FrameA * 2;
+	const int32 ValueOffset = ReadContext.GetFrameIndexA() * 2;
 	const int32 ExpectedValues = ValueOffset + 2;
 
-	UNDISimpleCounterSimCacheData* CacheData = CastChecked<UNDISimpleCounterSimCacheData>(StorageObject);
+	UNDISimpleCounterSimCacheData* CacheData = ReadContext.GetStorageObject<UNDISimpleCounterSimCacheData>();
 	if (CacheData->Values.Num() < ExpectedValues)
 	{
 		return false;
 	}
 
 	// Set CPU Data
-	FNDISimpleCounterInstanceData_GameThread* InstanceData_GT = static_cast<FNDISimpleCounterInstanceData_GameThread*>(OptionalPerInstanceData);
+	FNDISimpleCounterInstanceData_GameThread* InstanceData_GT = ReadContext.GetPerInstanceData<FNDISimpleCounterInstanceData_GameThread>();
 	InstanceData_GT->Counter = CacheData->Values[ValueOffset + 0];
 
 	// Set GPU Data
@@ -553,7 +551,7 @@ bool UNiagaraDataInterfaceSimpleCounter::SimCacheReadFrame(UObject* StorageObjec
 	{
 		ENQUEUE_RENDER_COMMAND(FNDISimpleCounter_PushToRender)
 		(
-			[Proxy_RT=GetProxyAs<FNDISimpleCounterProxy>(), InstanceID=SystemInstance->GetId(), NewValue=CacheData->Values[ValueOffset + 1]](FRHICommandListImmediate& RHICmdList)
+			[Proxy_RT=GetProxyAs<FNDISimpleCounterProxy>(), InstanceID=ReadContext.GetSystemInstance()->GetId(), NewValue = CacheData->Values[ValueOffset + 1]](FRHICommandListImmediate& RHICmdList)
 			{
 				if ( FNDISimpleCounterInstanceData_RenderThread* InstanceData_RT=Proxy_RT->PerInstanceData_RenderThread.Find(InstanceID) )
 				{
diff --git a/Engine/Plugins/FX/Niagara/Source/Niagara/Private/NiagaraDataInterfaceArray.cpp b/Engine/Plugins/FX/Niagara/Source/Niagara/Private/NiagaraDataInterfaceArray.cpp
--- a/Engine/Plugins/FX/Niagara/Source/Niagara/Private/NiagaraDataInterfaceArray.cpp
+++ b/Engine/Plugins/FX/Niagara/Source/Niagara/Private/NiagaraDataInterfaceArray.cpp
@@ -108,13 +108,10 @@ bool UNiagaraDataInterfaceArray::SimCacheWriteFrame(UObject* StorageObject, int
 	return ArrayProxy->SimCacheWriteFrame(CacheData, FrameIndex, SystemInstance);
 }
 
-bool UNiagaraDataInterfaceArray::SimCacheReadFrame(UObject* StorageObject, int FrameA, int FrameB, float Interp, FNiagaraSystemInstance* SystemInstance, void* OptionalPerInstanceData)
+bool UNiagaraDataInterfaceArray::SimCacheReadFrame(const FNiagaraSimCacheDataInterfaceReadContext& ReadContext)
 {
-	check(OptionalPerInstanceData && StorageObject);
-	UNDIArraySimCacheData* CacheData = CastChecked<UNDIArraySimCacheData>(StorageObject);
-
 	INDIArrayProxyBase* ArrayProxy = GetProxyAs<INDIArrayProxyBase>();
-	return ArrayProxy->SimCacheReadFrame(CacheData, FrameA, SystemInstance);
+	return ArrayProxy->SimCacheReadFrame(ReadContext.GetStorageObject<UNDIArraySimCacheData>(), ReadContext.GetFrameIndexA(), ReadContext.GetSystemInstance());
 }
 
 bool UNiagaraDataInterfaceArray::SimCacheCompareFrame(const UObject* LhsStorageObject, const UObject* RhsStorageObject, int FrameIndex, TOptional<float> InTolerance, FString& OutErrors) const
diff --git a/Engine/Plugins/FX/Niagara/Source/Niagara/Private/NiagaraDataInterfaceRenderTarget2D.cpp b/Engine/Plugins/FX/Niagara/Source/Niagara/Private/NiagaraDataInterfaceRenderTarget2D.cpp
--- a/Engine/Plugins/FX/Niagara/Source/Niagara/Private/NiagaraDataInterfaceRenderTarget2D.cpp
+++ b/Engine/Plugins/FX/Niagara/Source/Niagara/Private/NiagaraDataInterfaceRenderTarget2D.cpp
@@ -663,13 +663,13 @@ NIAGARA_API bool UNiagaraDataInterfaceRenderTarget2D::SimCacheEndWrite(UObject*
 	return true;
 }
 
-NIAGARA_API bool UNiagaraDataInterfaceRenderTarget2D::SimCacheReadFrame(UObject* StorageObject, int FrameA, int FrameB, float Interp, FNiagaraSystemInstance* SystemInstance, void* OptionalPerInstanceData)
+NIAGARA_API bool UNiagaraDataInterfaceRenderTarget2D::SimCacheReadFrame(const FNiagaraSimCacheDataInterfaceReadContext& ReadContext)
 {
-	UNDIRenderTargetSimCacheData* SimCacheData = CastChecked<UNDIRenderTargetSimCacheData>(StorageObject);
+	UNDIRenderTargetSimCacheData* SimCacheData = ReadContext.GetStorageObject<UNDIRenderTargetSimCacheData>();
 
-	FRenderTarget2DRWInstanceData_GameThread* InstanceData_GT = reinterpret_cast<FRenderTarget2DRWInstanceData_GameThread*>(OptionalPerInstanceData);
+	FRenderTarget2DRWInstanceData_GameThread* InstanceData_GT = ReadContext.GetPerInstanceData<FRenderTarget2DRWInstanceData_GameThread>();
 
-	const int FrameIndex = Interp >= 0.5f ? FrameB : FrameA;
+	const int FrameIndex = ReadContext.GetFrameInterp() >= 0.5f ? ReadContext.GetFrameIndexB() : ReadContext.GetFrameIndexA();
 	if (!SimCacheData->IsValidFrame(FrameIndex))
 	{
 		return false;
@@ -684,6 +684,8 @@ NIAGARA_API bool UNiagaraDataInterfaceRenderTarget2D::SimCacheReadFrame(UObject*
 		InstanceData_GT->Filter = TextureFilter::TF_Default;
 	}
 
+	FNiagaraSystemInstance* SystemInstance = ReadContext.GetSystemInstance();
+
 	PerInstanceTick(InstanceData_GT, SystemInstance, 0.0f);
 	PerInstanceTickPostSimulate(InstanceData_GT, SystemInstance, 0.0f);
 
diff --git a/Engine/Plugins/FX/Niagara/Source/Niagara/Private/NiagaraDataInterfaceRenderTargetVolume.cpp b/Engine/Plugins/FX/Niagara/Source/Niagara/Private/NiagaraDataInterfaceRenderTargetVolume.cpp
--- a/Engine/Plugins/FX/Niagara/Source/Niagara/Private/NiagaraDataInterfaceRenderTargetVolume.cpp
+++ b/Engine/Plugins/FX/Niagara/Source/Niagara/Private/NiagaraDataInterfaceRenderTargetVolume.cpp
@@ -726,32 +726,31 @@ bool UNiagaraDataInterfaceRenderTargetVolume::SimCacheEndWrite(UObject* StorageO
 	return true;
 }
 
-bool UNiagaraDataInterfaceRenderTargetVolume::SimCacheReadFrame(UObject* StorageObject, int FrameA, int FrameB, float Interp, FNiagaraSystemInstance* SystemInstance, void* OptionalPerInstanceData)
+bool UNiagaraDataInterfaceRenderTargetVolume::SimCacheReadFrame(const FNiagaraSimCacheDataInterfaceReadContext& ReadContext)
 {
-	FRenderTargetVolumeRWInstanceData_GameThread* InstanceData_GT = reinterpret_cast<FRenderTargetVolumeRWInstanceData_GameThread*>(OptionalPerInstanceData);
+	FRenderTargetVolumeRWInstanceData_GameThread* InstanceData_GT = ReadContext.GetPerInstanceData<FRenderTargetVolumeRWInstanceData_GameThread>();
 
-	if (Cast<UVolumeCache>(StorageObject))
+	if (ReadContext.GetOptionalStorageObject<UVolumeCache>())
 	{
 		UE_LOG(LogNiagara, Error, TEXT("vdb caches are no longer supported.  Regenerate your Sim Cache, which will use SVT to store volume data."));
 
 	}
-	else if (Cast<UAnimatedSparseVolumeTexture>(StorageObject))
+	else if (UAnimatedSparseVolumeTexture* SVT = ReadContext.GetOptionalStorageObject<UAnimatedSparseVolumeTexture>())
 	{			
-		UAnimatedSparseVolumeTexture* SVT = Cast<UAnimatedSparseVolumeTexture>(StorageObject);
-		
 		const int32 MipLevel = 0;
 		const float FrameRate = 0.0f;
 		const bool bBlocking = true;
 		const bool bHasFrameRate = false;
-		USparseVolumeTextureFrame *SVTFrame = USparseVolumeTextureFrame::GetFrameAndIssueStreamingRequest(SVT, GetTypeHash(this), MipLevel, FMath::RoundToInt(FrameA + Interp), MipLevel, bBlocking, bHasFrameRate);
+		const float SVYFrameIndex = FMath::RoundToFloat(float(ReadContext.GetFrameIndexA()) + ReadContext.GetFrameInterp());
+		USparseVolumeTextureFrame *SVTFrame = USparseVolumeTextureFrame::GetFrameAndIssueStreamingRequest(SVT, GetTypeHash(this), FrameRate, SVYFrameIndex, MipLevel, bBlocking, bHasFrameRate);
 		
 		// The streaming manager normally ticks in FDeferredShadingSceneRenderer::Render(), but the SVT->DenseTexture conversion compute shader happens in a render command before that.
 		// At execution time of that command, the streamer hasn't had the chance to do any streaming yet, so we force another tick here.
 		// Assuming blocking requests are used, this guarantees that the requested frame is fully streamed in (if there is memory available).
 		UE::SVT::GetStreamingManager().Update_GameThread();
 
-		FIntVector VolumeResolution = SVT->GetVolumeResolution();						
-
+		FNiagaraSystemInstance* SystemInstance = ReadContext.GetSystemInstance();
+		const FIntVector VolumeResolution = SVT->GetVolumeResolution();						
 		InstanceData_GT->Size = VolumeResolution;
 		InstanceData_GT->Format = SVT->GetFormat(0);
 		InstanceData_GT->Filter = TextureFilter::TF_Default;
@@ -761,7 +760,6 @@ bool UNiagaraDataInterfaceRenderTargetVolume::SimCacheReadFrame(UObject* Storage
 
 		FNiagaraDataInterfaceProxyRenderTargetVolumeProxy* RT_Proxy = GetProxyAs<FNiagaraDataInterfaceProxyRenderTargetVolumeProxy>();
 
-
 		// Execute compute shader
 		ENQUEUE_RENDER_COMMAND(NDIRenderTargetVolumeUpdate)
 			(
@@ -852,16 +850,16 @@ bool UNiagaraDataInterfaceRenderTargetVolume::SimCacheReadFrame(UObject* Storage
 					}					
 				});	
 	}
-	else
+	else if (UNDIRenderTargetSimCacheData* SimCacheData = ReadContext.GetOptionalStorageObject<UNDIRenderTargetSimCacheData>())
 	{
-		UNDIRenderTargetSimCacheData* SimCacheData = CastChecked<UNDIRenderTargetSimCacheData>(StorageObject);
-
-		const int FrameIndex = Interp >= 0.5f ? FrameB : FrameA;
+		const int FrameIndex = ReadContext.GetFrameInterp() >= 0.5f ? ReadContext.GetFrameIndexB() : ReadContext.GetFrameIndexA();
 		if (!SimCacheData->IsValidFrame(FrameIndex))
 		{
 			return false;
 		}
 
+		FNiagaraSystemInstance* SystemInstance = ReadContext.GetSystemInstance();
+
 		InstanceData_GT->Size   = SimCacheData->GetTextureSize(FrameIndex);
 		InstanceData_GT->Format = SimCacheData->GetTextureFormat(FrameIndex);
 		InstanceData_GT->Filter = TextureFilter::TF_Default;
diff --git a/Engine/Plugins/FX/Niagara/Source/Niagara/Private/NiagaraSimCache.cpp b/Engine/Plugins/FX/Niagara/Source/Niagara/Private/NiagaraSimCache.cpp
--- a/Engine/Plugins/FX/Niagara/Source/Niagara/Private/NiagaraSimCache.cpp
+++ b/Engine/Plugins/FX/Niagara/Source/Niagara/Private/NiagaraSimCache.cpp
@@ -1084,6 +1084,18 @@ bool UNiagaraSimCache::ReadFrame(int32 FrameIndex, float FrameFraction, FNiagara
 	if (DataInterfaceStorage.IsEmpty() == false)
 	{
 		TSet<FNiagaraVariableBase> VisitedDataInterfaces;
+		VisitedDataInterfaces.Reserve(DataInterfaceStorage.Num());
+
+		FNiagaraSimCacheDataInterfaceReadContext DIReadContext;
+		DIReadContext.SystemInstance		= Helper.SystemInstance;
+		DIReadContext.FrameIndexA			= FrameIndex;
+		DIReadContext.FrameIndexB			= NextFrameIndex;
+		DIReadContext.FrameInterp			= FrameFraction;
+		DIReadContext.RebaseTransformA		= RebaseATransform;
+		DIReadContext.RebaseTransformB		= RebaseBTransform;
+		DIReadContext.LWCTileA				= CacheFrames[FrameIndex].LWCTile;
+		DIReadContext.LWCTileB				= CacheFrames[NextFrameIndex].LWCTile;
+
 		bool bDataInterfacesSucess = true;
 
 		ForEachDataInterface(
@@ -1100,8 +1112,52 @@ bool UNiagaraSimCache::ReadFrame(int32 FrameIndex, float FrameFraction, FNiagara
 						}
 						VisitedDataInterfaces.Add(Variable);
 
-						void* PerInstanceData = Helper.SystemInstance->FindDataInterfaceInstanceData(DataInterface);
-						bDataInterfacesSucess &= SimCacheCustomStorageInterface->SimCacheReadFrame(StorageObject, FrameIndex, NextFrameIndex, FrameFraction, Helper.SystemInstance, PerInstanceData);
+						DIReadContext.StorageObject				= StorageObject;
+						DIReadContext.OptionalPerInstanceData	= Helper.SystemInstance->FindDataInterfaceInstanceData(DataInterface);
+
+						// Determine the interpolation state for the data interface
+						if (CreateParameters.bAllowInterpolation)
+						{
+							if (CreateParameters.InterpolationIncludeAttributes.Contains(Variable.GetName()))
+							{
+								DIReadContext.InterpolateData = true;
+							}
+							else if (CreateParameters.InterpolationExcludeAttributes.Contains(Variable.GetName()))
+							{
+								DIReadContext.InterpolateData = false;
+							}
+							else
+							{
+								DIReadContext.InterpolateData.Reset();
+							}
+						}
+						else
+						{
+							DIReadContext.InterpolateData = false;
+						}
+
+						// Determine the rebasing state for the data interface
+						if (CreateParameters.bAllowRebasing)
+						{
+							if (CreateParameters.RebaseIncludeAttributes.Contains(Variable.GetName()))
+							{
+								DIReadContext.RebaseData = true;
+							}
+							else if (CreateParameters.RebaseExcludeAttributes.Contains(Variable.GetName()))
+							{
+								DIReadContext.RebaseData = false;
+							}
+							else
+							{
+								DIReadContext.RebaseData.Reset();
+							}
+						}
+						else
+						{
+							DIReadContext.RebaseData = false;
+						}
+
+						bDataInterfacesSucess &= SimCacheCustomStorageInterface->SimCacheReadFrame(DIReadContext);
 					}
 				}
 				return true;
