- Fix PVS warning - int vs float for interp value

#rnx

diff --git a/Engine/Plugins/FX/Niagara/Source/Niagara/Classes/NiagaraSimCacheCustomStorageInterface.h b/Engine/Plugins/FX/Niagara/Source/Niagara/Classes/NiagaraSimCacheCustomStorageInterface.h
--- a/Engine/Plugins/FX/Niagara/Source/Niagara/Classes/NiagaraSimCacheCustomStorageInterface.h
+++ b/Engine/Plugins/FX/Niagara/Source/Niagara/Classes/NiagaraSimCacheCustomStorageInterface.h
@@ -37,7 +37,7 @@ public:
 	// Get the currrent frame we are interpolating to
 	int32 GetFrameIndexB() const { return(FrameIndexB); }
 	// Get the currrent frame interpolation fraction
-	int32 GetFrameInterp() const { return(FrameInterp); }
+	float GetFrameInterp() const { return(FrameInterp); }
 
 	// Test to see if we should rebase the data or not
 	// The default value is expected to be the 'just works' case for when the user did not set force include / exlucde and rebasing is enabled
