- Add LWC tile to debug hud

#rnx
#tests debug a system at a tile beyond 0,0,0

diff --git a/Engine/Plugins/FX/Niagara/Source/Niagara/Private/NiagaraDebugHud.cpp b/Engine/Plugins/FX/Niagara/Source/Niagara/Private/NiagaraDebugHud.cpp
--- a/Engine/Plugins/FX/Niagara/Source/Niagara/Private/NiagaraDebugHud.cpp
+++ b/Engine/Plugins/FX/Niagara/Source/Niagara/Private/NiagaraDebugHud.cpp
@@ -3065,9 +3065,18 @@ void FNiagaraDebugHud::DrawComponents(FNiagaraWorldManager* WorldManager, UCanva
 							case ENiagaraOcclusionQueryMode::AlwaysDisabled:	StringBuilder.Append(TEXT(" OcclusionQuery(AlwaysDisable)")); break;
 						}
 						StringBuilder.Append(TEXT("\n"));
+
+						if ( NiagaraSystem->bSupportLargeWorldCoordinates )
+						{
+							const FVector3f LWCTile = SystemInstance->GetLWCTile();
+							if (!LWCTile.Equals(FVector3f::ZeroVector) )
+							{
+								StringBuilder.Appendf(TEXT("LWCTile - %.1f, %.1f, %.1f"), LWCTile.X, LWCTile.Y, LWCTile.Z);
+							}
+						}
 					}
 
-					int64 TotalBytes = SystemInstanceController->GetTotalBytesUsed();
+					const int64 TotalBytes = SystemInstanceController->GetTotalBytesUsed();
 					StringBuilder.Appendf(TEXT("Memory - %6.2fMB\n"), float(double(TotalBytes) / (1024.0*1024.0)));
 				}
 
