- Add blueprint RW to Niagara mesh properties so you can build an array inside BP

#rnx
#tests Create an mesh array in BP and set on a NS

diff --git a/Engine/Plugins/FX/Niagara/Source/Niagara/Public/NiagaraMeshRendererMeshProperties.h b/Engine/Plugins/FX/Niagara/Source/Niagara/Public/NiagaraMeshRendererMeshProperties.h
--- a/Engine/Plugins/FX/Niagara/Source/Niagara/Public/NiagaraMeshRendererMeshProperties.h
+++ b/Engine/Plugins/FX/Niagara/Source/Niagara/Public/NiagaraMeshRendererMeshProperties.h
@@ -71,23 +71,23 @@ struct FNiagaraMeshRendererMeshPropertiesBase
 	GENERATED_BODY()
 
 	/** The mesh to use when rendering this slot */
-	UPROPERTY(EditAnywhere, Category = "Mesh")
+	UPROPERTY(BlueprintReadWrite, EditAnywhere, Category = "Mesh")
 	TObjectPtr<UStaticMesh> Mesh;
 
 	/** Scale of the mesh */
-	UPROPERTY(EditAnywhere, Category = "Mesh")
+	UPROPERTY(BlueprintReadWrite, EditAnywhere, Category = "Mesh")
 	FVector Scale = FVector::OneVector;
 
 	/** Rotation of the mesh */
-	UPROPERTY(EditAnywhere, Category = "Mesh")
+	UPROPERTY(BlueprintReadWrite, EditAnywhere, Category = "Mesh")
 	FRotator Rotation = FRotator::ZeroRotator;
 
 	/** Offset of the mesh pivot */
-	UPROPERTY(EditAnywhere, Category = "Mesh")
+	UPROPERTY(BlueprintReadWrite, EditAnywhere, Category = "Mesh")
 	FVector PivotOffset = FVector::ZeroVector;
 
 	/** What space is the pivot offset in? */
-	UPROPERTY(EditAnywhere, Category = "Mesh")
+	UPROPERTY(BlueprintReadWrite, EditAnywhere, Category = "Mesh")
 	ENiagaraMeshPivotOffsetSpace PivotOffsetSpace = ENiagaraMeshPivotOffsetSpace::Mesh;
 
 	bool IsNearlyEqual(const FNiagaraMeshRendererMeshPropertiesBase& Rhs, float Tolerance = UE_KINDA_SMALL_NUMBER) const
