Avoid potential wasted work reinitializing LWC types.

diff --git a/Engine/Plugins/FX/Niagara/Source/Niagara/Private/NiagaraTypes.cpp b/Engine/Plugins/FX/Niagara/Source/Niagara/Private/NiagaraTypes.cpp
--- a/Engine/Plugins/FX/Niagara/Source/Niagara/Private/NiagaraTypes.cpp
+++ b/Engine/Plugins/FX/Niagara/Source/Niagara/Private/NiagaraTypes.cpp
@@ -744,18 +744,21 @@ FNiagaraTypeDefinition FNiagaraTypeHelper::GetLWCType(const FNiagaraTypeDefiniti
 
 void FNiagaraTypeHelper::InitLWCTypes()
 {
-	UPackage* CoreUObjectPkg = FindObjectChecked<UPackage>(nullptr, TEXT("/Script/CoreUObject"));
-	Vector2DDef = FNiagaraTypeDefinition(FindObjectChecked<UScriptStruct>(CoreUObjectPkg, TEXT("Vector2D")), FNiagaraTypeDefinition::EAllowUnfriendlyStruct::Allow);
-	VectorDef = FNiagaraTypeDefinition(FindObjectChecked<UScriptStruct>(CoreUObjectPkg, TEXT("Vector")), FNiagaraTypeDefinition::EAllowUnfriendlyStruct::Allow);
-	Vector4Def = FNiagaraTypeDefinition(FindObjectChecked<UScriptStruct>(CoreUObjectPkg, TEXT("Vector4")), FNiagaraTypeDefinition::EAllowUnfriendlyStruct::Allow);
-	QuatDef = FNiagaraTypeDefinition(FindObjectChecked<UScriptStruct>(CoreUObjectPkg, TEXT("Quat")), FNiagaraTypeDefinition::EAllowUnfriendlyStruct::Allow);
-	DoubleDef = FNiagaraTypeDefinition(FNiagaraDouble::StaticStruct(), FNiagaraTypeDefinition::EAllowUnfriendlyStruct::Allow);
-
-	FNiagaraTypeRegistry::Register(Vector2DDef, ENiagaraTypeRegistryFlags::None);
-	FNiagaraTypeRegistry::Register(VectorDef, ENiagaraTypeRegistryFlags::None);
-	FNiagaraTypeRegistry::Register(Vector4Def, ENiagaraTypeRegistryFlags::None);
-	FNiagaraTypeRegistry::Register(QuatDef, ENiagaraTypeRegistryFlags::None);
-	FNiagaraTypeRegistry::Register(DoubleDef, ENiagaraTypeRegistryFlags::None);
+	if(Vector2DDef.IsValid() == false)
+	{
+		UPackage* CoreUObjectPkg = FindObjectChecked<UPackage>(nullptr, TEXT("/Script/CoreUObject"));
+		Vector2DDef = FNiagaraTypeDefinition(FindObjectChecked<UScriptStruct>(CoreUObjectPkg, TEXT("Vector2D")), FNiagaraTypeDefinition::EAllowUnfriendlyStruct::Allow);
+		VectorDef = FNiagaraTypeDefinition(FindObjectChecked<UScriptStruct>(CoreUObjectPkg, TEXT("Vector")), FNiagaraTypeDefinition::EAllowUnfriendlyStruct::Allow);
+		Vector4Def = FNiagaraTypeDefinition(FindObjectChecked<UScriptStruct>(CoreUObjectPkg, TEXT("Vector4")), FNiagaraTypeDefinition::EAllowUnfriendlyStruct::Allow);
+		QuatDef = FNiagaraTypeDefinition(FindObjectChecked<UScriptStruct>(CoreUObjectPkg, TEXT("Quat")), FNiagaraTypeDefinition::EAllowUnfriendlyStruct::Allow);
+		DoubleDef = FNiagaraTypeDefinition(FNiagaraDouble::StaticStruct(), FNiagaraTypeDefinition::EAllowUnfriendlyStruct::Allow);
+
+		FNiagaraTypeRegistry::Register(Vector2DDef, ENiagaraTypeRegistryFlags::None);
+		FNiagaraTypeRegistry::Register(VectorDef, ENiagaraTypeRegistryFlags::None);
+		FNiagaraTypeRegistry::Register(Vector4Def, ENiagaraTypeRegistryFlags::None);
+		FNiagaraTypeRegistry::Register(QuatDef, ENiagaraTypeRegistryFlags::None);
+		FNiagaraTypeRegistry::Register(DoubleDef, ENiagaraTypeRegistryFlags::None);
+	}
 }
 
 UScriptStruct* FNiagaraTypeHelper::FindNiagaraFriendlyTopLevelStruct(UScriptStruct* InStruct, ENiagaraStructConversion StructConversion)
