Sprite Renderers set to SubUV mode pops up dismissible warning message if Cutout Texture is not set

diff --git a/Engine/Plugins/FX/Niagara/Source/Niagara/Private/NiagaraSpriteRendererProperties.cpp b/Engine/Plugins/FX/Niagara/Source/Niagara/Private/NiagaraSpriteRendererProperties.cpp
--- a/Engine/Plugins/FX/Niagara/Source/Niagara/Private/NiagaraSpriteRendererProperties.cpp
+++ b/Engine/Plugins/FX/Niagara/Source/Niagara/Private/NiagaraSpriteRendererProperties.cpp
@@ -648,6 +648,20 @@ void UNiagaraSpriteRendererProperties::GetRendererFeedback(const FVersionedNiaga
 		GetUsedMaterials(nullptr, Materials);
 		MaterialParameters.GetFeedback(Materials, OutWarnings);
 	}
+
+	if (bSubImageBlend)
+	{
+		if (!bUseMaterialCutoutTexture && CutoutTexture == nullptr)
+		{
+			OutWarnings.Emplace(
+				LOCTEXT("MissingCutoutTexture", "Missing Cutout Texture for Sub UV Blending, which might lead to unnecessary overdraw and reduce performance.\nEnable Use Material Cutout Texture or set Cutout Texture"),
+				LOCTEXT("RequiredCutoutTexture", "Sub UV Blending enabled but Cutout Texture is not set."),
+				FText(),
+				FNiagaraRendererFeedbackFix(),
+				true
+			);
+		}
+	}
 }
 
 void UNiagaraSpriteRendererProperties::GetRendererTooltipWidgets(const FNiagaraEmitterInstance* InEmitter, TArray<TSharedPtr<SWidget>>& OutWidgets, TSharedPtr<FAssetThumbnailPool> InThumbnailPool) const
