- When generating the HLSL prototypes do not include functions which can not run on either CPU or GPU

#rnx
#tests Make a mesh array, right click the pin

diff --git a/Engine/Plugins/FX/Niagara/Source/NiagaraEditor/Private/EdGraphSchema_Niagara.cpp b/Engine/Plugins/FX/Niagara/Source/NiagaraEditor/Private/EdGraphSchema_Niagara.cpp
--- a/Engine/Plugins/FX/Niagara/Source/NiagaraEditor/Private/EdGraphSchema_Niagara.cpp
+++ b/Engine/Plugins/FX/Niagara/Source/NiagaraEditor/Private/EdGraphSchema_Niagara.cpp
@@ -2211,6 +2211,11 @@ TArray<TTuple<FString, FString, FText>> UEdGraphSchema_Niagara::GetDataInterface
 					continue;
 				}
 
+				if (!FunctionSignature.bSupportsGPU && !FunctionSignature.bSupportsCPU)
+				{
+					continue;
+				}
+
 				FString Prototype = FNiagaraHlslTranslator::GenerateFunctionHlslPrototype(DataInterfaceName, FunctionSignature);
 				if ( Prototype.Len() > 0 )
 				{
