- Fix for pending spawn paused instances switching solo mode being in the wrong simulation

diff --git a/Engine/Plugins/FX/Niagara/Source/Niagara/Private/NiagaraSystemInstance.cpp b/Engine/Plugins/FX/Niagara/Source/Niagara/Private/NiagaraSystemInstance.cpp
--- a/Engine/Plugins/FX/Niagara/Source/Niagara/Private/NiagaraSystemInstance.cpp
+++ b/Engine/Plugins/FX/Niagara/Source/Niagara/Private/NiagaraSystemInstance.cpp
@@ -2950,17 +2950,13 @@ FNiagaraEmitterInstance* FNiagaraSystemInstance::GetEmitterByID(FNiagaraEmitterI
 
 void FNiagaraSystemInstance::SetForceSolo(bool bInForceSolo)
 {
-	// We may be forced into solo mode so if we match nothing to do here
-	bForceSolo = bInForceSolo;
-	if ( bSolo == bInForceSolo )
+	if (bForceSolo == bInForceSolo)
 	{
 		return;
 	}
 
-	if (bForceSolo != bSolo)
-	{
-		SetSolo(bForceSolo);
-	}
+	bForceSolo = bInForceSolo;
+	SetSolo(bForceSolo);
 }
 
 void FNiagaraSystemInstance::EvaluateBoundFunction(FName FunctionName, bool& UsedOnCpu, bool& UsedOnGpu) const
diff --git a/Engine/Plugins/FX/Niagara/Source/Niagara/Private/NiagaraSystemSimulation.cpp b/Engine/Plugins/FX/Niagara/Source/Niagara/Private/NiagaraSystemSimulation.cpp
--- a/Engine/Plugins/FX/Niagara/Source/Niagara/Private/NiagaraSystemSimulation.cpp
+++ b/Engine/Plugins/FX/Niagara/Source/Niagara/Private/NiagaraSystemSimulation.cpp
@@ -765,6 +765,12 @@ void FNiagaraSystemSimulation::TransferInstance(FNiagaraSystemInstance* Instance
 			AddToInstanceList(Instance, ENiagaraSystemInstanceState::PendingSpawn);
 			break;
 
+		case ENiagaraSystemInstanceState::PendingSpawnPaused:
+			SourceSimulation->RemoveFromInstanceList(Instance);
+			Instance->SystemSimulation = this->AsShared();
+			AddToInstanceList(Instance, ENiagaraSystemInstanceState::PendingSpawnPaused);
+			break;
+
 		// Should never get here!
 		case ENiagaraSystemInstanceState::Spawning:
 			UE_LOG(LogNiagara, Fatal, TEXT("Attempting to transfer an instance that is in the spawning state"));
@@ -794,6 +800,9 @@ void FNiagaraSystemSimulation::TransferInstance(FNiagaraSystemInstance* Instance
 			check(Instance->SystemInstanceIndex == NewDataSetIndex);
 			break;
 		}
+
+		default:
+			checkNoEntry();
 	}
 
 	if (!bBindingsInitialized)
