StackObject now calls RefreshErrors on the DataInterface on PostChange, and no longer in RefreshChildrenInternal.
This fixes an issue where two pieces of UI code would constantly tell each other to refresh

diff --git a/Engine/Plugins/FX/Niagara/Source/NiagaraEditor/Private/ViewModels/Stack/NiagaraStackObject.cpp b/Engine/Plugins/FX/Niagara/Source/NiagaraEditor/Private/ViewModels/Stack/NiagaraStackObject.cpp
--- a/Engine/Plugins/FX/Niagara/Source/NiagaraEditor/Private/ViewModels/Stack/NiagaraStackObject.cpp
+++ b/Engine/Plugins/FX/Niagara/Source/NiagaraEditor/Private/ViewModels/Stack/NiagaraStackObject.cpp
@@ -33,7 +33,6 @@ void UNiagaraStackObject::Initialize(FRequiredEntryData InRequiredEntryData, UOb
 	bIsTopLevel = bInIsTopLevelObject;
 	bHideTopLevelCategories = bInHideTopLevelCategories;
 	OwningNiagaraNode = InOwningNiagaraNode;
-	bIsRefreshingDataInterfaceErrors = false;
 	FilterMode = EDetailNodeFilterMode::FilterRootNodesOnly;
 
 	MessageLogGuid = GetSystemViewModel()->GetMessageLogGuid();
@@ -64,7 +63,6 @@ void UNiagaraStackObject::Initialize(
 	bIsTopLevel = bInIsTopLevelStruct;
 	bHideTopLevelCategories = bInHideTopLevelCategories;
 	OwningNiagaraNode = InOwningNiagaraNode;
-	bIsRefreshingDataInterfaceErrors = false;
 }
 
 void UNiagaraStackObject::SetOnFilterDetailNodes(FNiagaraStackObjectShared::FOnFilterDetailNodes InOnFilterDetailNodes, EDetailNodeFilterMode InFilterMode)
@@ -98,6 +96,15 @@ void UNiagaraStackObject::NotifyPostChange(const FPropertyChangedEvent& Property
 {
 	if (GetObject() != nullptr)
 	{
+		UNiagaraDataInterface* DataInterfaceObject = Cast<UNiagaraDataInterface>(WeakObject.Get());
+		if (DataInterfaceObject != nullptr)
+		{
+			// First we tell the Details Panel to refresh its error UI
+			{
+				DataInterfaceObject->RefreshErrors();
+			}
+		}
+		
 		TArray<UObject*> ChangedObjects;
 		ChangedObjects.Add(GetObject());
 		OnDataObjectModified().Broadcast(ChangedObjects, ENiagaraDataObjectChange::Changed);
@@ -286,13 +293,6 @@ void UNiagaraStackObject::RefreshChildrenInternal(const TArray<UNiagaraStackEntr
 	UNiagaraDataInterface* DataInterfaceObject = Cast<UNiagaraDataInterface>(Object);
 	if (DataInterfaceObject != nullptr)
 	{
-		// First we need to refresh the errors on the data interface so that the rows in the property row generator 
-		// are correct.
-		{
-			bIsRefreshingDataInterfaceErrors = true;
-			DataInterfaceObject->RefreshErrors();
-		}
-
 		// Generate the summary stack issue for any errors which are generated.
 		TArray<FNiagaraDataInterfaceError> Errors;
 		TArray<FNiagaraDataInterfaceFeedback> Warnings, Info;
@@ -415,12 +415,7 @@ void UNiagaraStackObject::PostRefreshChildrenInternal()
 
 void UNiagaraStackObject::PropertyRowsRefreshed()
 {
-	if(bIsRefreshingDataInterfaceErrors == false)
-	{
-		RefreshChildren();
-	}
-	
-	bIsRefreshingDataInterfaceErrors = false;
+	RefreshChildren();
 }
 
 void UNiagaraStackObject::OnMessageManagerRefresh(const TArray<TSharedRef<const INiagaraMessage>>& NewMessages)
diff --git a/Engine/Plugins/FX/Niagara/Source/NiagaraEditor/Public/ViewModels/Stack/NiagaraStackObject.h b/Engine/Plugins/FX/Niagara/Source/NiagaraEditor/Public/ViewModels/Stack/NiagaraStackObject.h
--- a/Engine/Plugins/FX/Niagara/Source/NiagaraEditor/Public/ViewModels/Stack/NiagaraStackObject.h
+++ b/Engine/Plugins/FX/Niagara/Source/NiagaraEditor/Public/ViewModels/Stack/NiagaraStackObject.h
@@ -126,7 +126,6 @@ private:
 	TArray<FRegisteredClassCustomization> RegisteredClassCustomizations;
 	TArray<FRegisteredPropertyCustomization> RegisteredPropertyCustomizations;
 	TSharedPtr<IPropertyRowGenerator> PropertyRowGenerator;
-	bool bIsRefreshingDataInterfaceErrors;
 
 	/** An optional object guid that can be provided to identify the object this stack entry represents. This can be used for summary view purposes and is also given to the child property rows. */
 	TOptional<FGuid> ObjectGuid;
