Update - Sprite Renderers set to SubUV mode pops up dismissible note if Cutout Texture is not set

diff --git a/Engine/Plugins/FX/Niagara/Source/Niagara/Private/NiagaraSpriteRendererProperties.cpp b/Engine/Plugins/FX/Niagara/Source/Niagara/Private/NiagaraSpriteRendererProperties.cpp
--- a/Engine/Plugins/FX/Niagara/Source/Niagara/Private/NiagaraSpriteRendererProperties.cpp
+++ b/Engine/Plugins/FX/Niagara/Source/Niagara/Private/NiagaraSpriteRendererProperties.cpp
@@ -649,18 +649,18 @@ void UNiagaraSpriteRendererProperties::GetRendererFeedback(const FVersionedNiaga
 		MaterialParameters.GetFeedback(Materials, OutWarnings);
 	}
 
-	if (bSubImageBlend)
-	{
-		if (!bUseMaterialCutoutTexture && CutoutTexture == nullptr)
-		{
-			OutWarnings.Emplace(
-				LOCTEXT("MissingCutoutTexture", "Missing Cutout Texture for Sub UV Blending, which might lead to unnecessary overdraw and reduce performance.\nEnable Use Material Cutout Texture or set Cutout Texture"),
-				LOCTEXT("RequiredCutoutTexture", "Sub UV Blending enabled but Cutout Texture is not set."),
-				FText(),
-				FNiagaraRendererFeedbackFix(),
-				true
-			);
-		}
+	// If sprite material is transluscent and subimage size is not (1,1) and cutout texture is not used, display note
+	if (Material && Material->GetBlendMode() == BLEND_Translucent 
+		&& (SubImageSize.X != 1.0f || SubImageSize.Y != 1.0f)
+		&& (!bUseMaterialCutoutTexture && CutoutTexture == nullptr))
+	{
+		OutInfo.Emplace(
+			LOCTEXT("MissingCutoutTexture", "Missing Cutout Texture for Sub UV Image, which might lead to overdraw and reduce performance.\nEnable Use Material Cutout Texture or set Cutout Texture"),
+			LOCTEXT("RequiredCutoutTexture", "Cutout Texture is not set for Sub UV Image"),
+			FText(),
+			FNiagaraRendererFeedbackFix(),
+			true
+		);
 	}
 }
 
