Niagara data channel wizard improvements:
* Autowrap header text
* Added emitter spawn count scale multiplication
* Improved target namespaces for spawn module

#tests Only tested workflow with the new wizard, as this introduces no changes to existing "spawn from ndc" modules

#changelist validated
#virtualized

diff --git a/Engine/Plugins/FX/Niagara/Source/NiagaraEditor/Private/Widgets/DataChannel/NiagaraDataChannelWizard.cpp b/Engine/Plugins/FX/Niagara/Source/NiagaraEditor/Private/Widgets/DataChannel/NiagaraDataChannelWizard.cpp
--- a/Engine/Plugins/FX/Niagara/Source/NiagaraEditor/Private/Widgets/DataChannel/NiagaraDataChannelWizard.cpp
+++ b/Engine/Plugins/FX/Niagara/Source/NiagaraEditor/Private/Widgets/DataChannel/NiagaraDataChannelWizard.cpp
@@ -145,6 +145,7 @@ namespace UE::Niagara::Wizard::DataChannel
 				.AutoHeight()
 				[
 					SNew(STextBlock)
+					.AutoWrapText(true)
 					.Text(GetHeaderLabel())
 				]
 				+SVerticalBox::Slot()
@@ -407,6 +408,7 @@ namespace UE::Niagara::Wizard::DataChannel
 				.AutoHeight()
 				[
 					SNew(STextBlock)
+					.AutoWrapText(true)
 					.Text(this, &FSpawnConditionPage::GetHeaderLabel)
 				]
 				+SVerticalBox::Slot()
@@ -839,7 +841,15 @@ TSharedRef<FModuleWizardModel> DataChannel::CreateSpawnNDCModuleWizardModel()
 	struct FSelectVariablesPage : FSelectVariablesPageBase
 	{
 		explicit FSelectVariablesPage(FSelectSpawnAssetPage* InPreviousPage) : FSelectVariablesPageBase(InPreviousPage)
-		{}
+		{
+			SupportedNamespaces.Empty();
+			SupportedNamespaces.Add(MakeShared<FString>("Particles"));
+			SupportedNamespaces.Add(MakeShared<FString>("Particles.Module"));
+			SupportedNamespaces.Add(MakeShared<FString>("Output.Module"));
+			SupportedNamespaces.Add(MakeShared<FString>("Transient"));
+			TargetNamespace = SupportedNamespaces[0];
+		}
+		
 		virtual ~FSelectVariablesPage() override = default;
 
 		virtual FText GetHeaderLabel() override
@@ -949,6 +959,16 @@ TSharedRef<FModuleWizardModel> DataChannel::CreateSpawnNDCModuleWizardModel()
 					GraphSchema->TryCreateConnection(SpawnFunction->GetOutputPin(0), OutputNode->GetInputPin(0));
 					Utilities::SetDefaultBinding(Graph, DIPin->PinName, FName("Emitter.SpawnDataChannel"));
 
+					// add scalability inputs
+					UEdGraphPin* EmitterScalePin = nullptr;
+					UEdGraphPin* SpawnScalePin = nullptr;
+					if (AssetPage->Data->bModifySpawnCountByScalability)
+					{
+						EmitterScalePin = Utilities::AddReadParameterPin(FNiagaraTypeDefinition::GetFloatDef(), FName("Transient.ScalabilityEmitterSpawnCountScale"), MapGetNode);
+						Utilities::SetDefaultValue(Graph, EmitterScalePin->PinName, FNiagaraTypeDefinition::GetFloatDef(), 1.0f);
+						SpawnScalePin = Utilities::AddReadParameterPin(FNiagaraTypeDefinition::GetFloatDef(), FName("Engine.Emitter.SpawnCountScale"), MapGetNode);
+					}
+
 					// Add module inputs
 					if (UEdGraphPin* EnableInput = SpawnFunction->FindPin(FName("Enable"), EGPD_Input))
 					{
@@ -982,7 +1002,20 @@ TSharedRef<FModuleWizardModel> DataChannel::CreateSpawnNDCModuleWizardModel()
 						Utilities::SetTooltip(Graph, SpawnMinPin->PinName, SpawnFunction->Signature.InputDescriptions[FNiagaraVariable(FNiagaraTypeDefinition::GetIntDef(), MinInput->PinName)]);
 						int32 Default = SpawnMode == ENiagaraDataChanneSpawnModuleMode::ConditionalSpawn ? 1 : -1;//Default to -1 for spawn direct as this is a clamp. aka, no clamp by default.
 						Utilities::SetDefaultValue(Graph, SpawnMinPin->PinName, FNiagaraTypeDefinition::GetIntDef(), Default);
-						GraphSchema->TryCreateConnection(MinInput, SpawnMinPin);
+						if (AssetPage->Data->bModifySpawnCountByScalability)
+						{
+							// Multiply by emitter scalability
+							if (UNiagaraNodeFunctionCall* ScaleSpawnNode = Utilities::CreateFunctionCallNode(LoadObject<UNiagaraScript>(nullptr, TEXT("/Niagara/Functions/Spawn/ScaleBurstSpawnCount.ScaleBurstSpawnCount")), Graph))
+							{
+								GraphSchema->TryCreateConnection(InputNode->GetOutputPin(0), ScaleSpawnNode->FindPin(FName("ParamMap"), EGPD_Input));
+								GraphSchema->TryCreateConnection(SpawnMinPin, ScaleSpawnNode->FindPin(FName("SpawnCount"), EGPD_Input));
+								GraphSchema->TryCreateConnection(MinInput, ScaleSpawnNode->GetOutputPin(0));
+							}
+						}
+						else
+						{
+							GraphSchema->TryCreateConnection(MinInput, SpawnMinPin);
+						}
 					}
 					if (UEdGraphPin* MaxInput = SpawnFunction->FindPin(FName(SpawnMode == ENiagaraDataChanneSpawnModuleMode::ConditionalSpawn ? "Max Spawn Count" : "ClampMax"), EGPD_Input))
 					{
@@ -990,7 +1023,20 @@ TSharedRef<FModuleWizardModel> DataChannel::CreateSpawnNDCModuleWizardModel()
 						Utilities::SetTooltip(Graph, SpawnMaxPin->PinName, SpawnFunction->Signature.InputDescriptions[FNiagaraVariable(FNiagaraTypeDefinition::GetIntDef(), MaxInput->PinName)]);
 						int32 Default = SpawnMode == ENiagaraDataChanneSpawnModuleMode::ConditionalSpawn ? 1 : -1;//Default to -1 for spawn direct as this is a clamp. aka, no clamp by default.
 						Utilities::SetDefaultValue(Graph, SpawnMaxPin->PinName, FNiagaraTypeDefinition::GetIntDef(), Default);
-						GraphSchema->TryCreateConnection(MaxInput, SpawnMaxPin);
+						if (AssetPage->Data->bModifySpawnCountByScalability)
+						{
+							// Multiply by emitter scalability
+							if (UNiagaraNodeFunctionCall* ScaleSpawnNode = Utilities::CreateFunctionCallNode(LoadObject<UNiagaraScript>(nullptr, TEXT("/Niagara/Functions/Spawn/ScaleBurstSpawnCount.ScaleBurstSpawnCount")), Graph))
+							{
+								GraphSchema->TryCreateConnection(InputNode->GetOutputPin(0), ScaleSpawnNode->FindPin(FName("ParamMap"), EGPD_Input));
+								GraphSchema->TryCreateConnection(SpawnMaxPin, ScaleSpawnNode->FindPin(FName("SpawnCount"), EGPD_Input));
+								GraphSchema->TryCreateConnection(MaxInput, ScaleSpawnNode->GetOutputPin(0));
+							}
+						}
+						else
+						{
+							GraphSchema->TryCreateConnection(MaxInput, SpawnMaxPin);
+						}
 					}
 					if (SpawnMode == ENiagaraDataChanneSpawnModuleMode::DirectSpawn)
 					{
@@ -1106,13 +1152,7 @@ TSharedRef<FModuleWizardModel> DataChannel::CreateSpawnNDCModuleWizardModel()
 					}
 					
 					// create and connect exec index node
-					UNiagaraNodeOp* ExecIndexNode = NewObject<UNiagaraNodeOp>(Graph);
-					ExecIndexNode->OpName = FName("Util::ExecIndex");
-					ExecIndexNode->SetFlags(RF_Transactional);
-					Graph->AddNode(ExecIndexNode, false, false);
-					ExecIndexNode->CreateNewGuid();
-					ExecIndexNode->PostPlacedNewNode();
-					ExecIndexNode->AllocateDefaultPins();
+					UNiagaraNodeOp* ExecIndexNode = Utilities::CreateOpNode(FName("Util::ExecIndex"), Graph);
 					if (UEdGraphPin* EmitterIDInput = SpawnDataFunction->FindPin(FName("Spawned Particle Exec Index"), EGPD_Input))
 					{
 						GraphSchema->TryCreateConnection(EmitterIDInput, ExecIndexNode->Pins[0]);
diff --git a/Engine/Plugins/FX/Niagara/Source/NiagaraEditor/Private/Widgets/DataChannel/NiagaraDataChannelWizard.h b/Engine/Plugins/FX/Niagara/Source/NiagaraEditor/Private/Widgets/DataChannel/NiagaraDataChannelWizard.h
--- a/Engine/Plugins/FX/Niagara/Source/NiagaraEditor/Private/Widgets/DataChannel/NiagaraDataChannelWizard.h
+++ b/Engine/Plugins/FX/Niagara/Source/NiagaraEditor/Private/Widgets/DataChannel/NiagaraDataChannelWizard.h
@@ -108,6 +108,11 @@ public:
 	*/
 	UPROPERTY(EditAnywhere, Category = "Advanced")
 	bool bAutoTransformPositionData = true;
+
+	/**If true then min and max spawn counts are multiplied by emitter spawn count scale, similar to existing spawn modules like spawn rate or spawn burst.
+	*/
+	UPROPERTY(EditAnywhere, Category = "Advanced")
+	bool bModifySpawnCountByScalability = true;
 };
 
 UENUM()
diff --git a/Engine/Plugins/FX/Niagara/Source/NiagaraEditor/Private/Widgets/Wizard/SNiagaraModuleWizard.cpp b/Engine/Plugins/FX/Niagara/Source/NiagaraEditor/Private/Widgets/Wizard/SNiagaraModuleWizard.cpp
--- a/Engine/Plugins/FX/Niagara/Source/NiagaraEditor/Private/Widgets/Wizard/SNiagaraModuleWizard.cpp
+++ b/Engine/Plugins/FX/Niagara/Source/NiagaraEditor/Private/Widgets/Wizard/SNiagaraModuleWizard.cpp
@@ -3,6 +3,7 @@
 #include "Widgets/Wizard/SNiagaraModuleWizard.h"
 
 #include "NiagaraClipboard.h"
+#include "NiagaraNodeOp.h"
 #include "NiagaraNodeParameterMapGet.h"
 #include "NiagaraNodeParameterMapSet.h"
 #include "ViewModels/NiagaraScratchPadScriptViewModel.h"
@@ -232,6 +233,18 @@ UNiagaraNodeFunctionCall* Utilities::CreateFunctionCallNode(UNiagaraScript* Func
 	return nullptr;
 }
 
+UNiagaraNodeOp* Utilities::CreateOpNode(const FName& OpName, UNiagaraGraph* Graph)
+{
+	UNiagaraNodeOp* NewOpNode = NewObject<UNiagaraNodeOp>(Graph);
+	NewOpNode->OpName = OpName;
+	NewOpNode->SetFlags(RF_Transactional);
+	Graph->AddNode(NewOpNode, false, false);
+	NewOpNode->CreateNewGuid();
+	NewOpNode->PostPlacedNewNode();
+	NewOpNode->AllocateDefaultPins();
+	return NewOpNode;
+}
+
 void Utilities::SetDefaultBinding(UNiagaraGraph* Graph, const FName& VarName, const FName& DefaultBinding)
 {
 	if (UNiagaraScriptVariable* ScriptVariable = Graph->GetScriptVariable(VarName))
diff --git a/Engine/Plugins/FX/Niagara/Source/NiagaraEditor/Public/Widgets/Wizard/SNiagaraModuleWizard.h b/Engine/Plugins/FX/Niagara/Source/NiagaraEditor/Public/Widgets/Wizard/SNiagaraModuleWizard.h
--- a/Engine/Plugins/FX/Niagara/Source/NiagaraEditor/Public/Widgets/Wizard/SNiagaraModuleWizard.h
+++ b/Engine/Plugins/FX/Niagara/Source/NiagaraEditor/Public/Widgets/Wizard/SNiagaraModuleWizard.h
@@ -120,6 +120,7 @@ namespace UE::Niagara::Wizard
 		NIAGARAEDITOR_API TSharedRef<IDetailsView> CreateDetailsView();
 		NIAGARAEDITOR_API UNiagaraNodeFunctionCall* CreateDataInterfaceFunctionNode(const TSubclassOf<UNiagaraDataInterface>& DataInterfaceClass, const FName& FunctionName, UNiagaraGraph* Graph);
 		NIAGARAEDITOR_API UNiagaraNodeFunctionCall* CreateFunctionCallNode(UNiagaraScript* FunctionScript, UNiagaraGraph* Graph);
+		NIAGARAEDITOR_API UNiagaraNodeOp* CreateOpNode(const FName& OpName, UNiagaraGraph* Graph);
 
 		template <typename T>
 		NIAGARAEDITOR_API void SetDefaultValue(UNiagaraGraph* Graph, const FName& VarName, const FNiagaraTypeDefinition& TypeDef, T Value)
@@ -129,6 +130,7 @@ namespace UE::Niagara::Wizard
 				// create temp var for data storage
 				FNiagaraVariable Var(TypeDef, FName("Var"));
 				Var.SetValue(Value);
+				ScriptVariable->DefaultMode = ENiagaraDefaultMode::Value;
 				ScriptVariable->SetDefaultValueData(Var.GetData());
 				Graph->ScriptVariableChanged(ScriptVariable->Variable);
 			}
