Don't show position types as source attribute to scale spawn count, as it's not lwc compatible and a vector works just as well

diff --git a/Engine/Plugins/FX/Niagara/Source/NiagaraEditor/Private/NiagaraActions.cpp b/Engine/Plugins/FX/Niagara/Source/NiagaraEditor/Private/NiagaraActions.cpp
--- a/Engine/Plugins/FX/Niagara/Source/NiagaraEditor/Private/NiagaraActions.cpp
+++ b/Engine/Plugins/FX/Niagara/Source/NiagaraEditor/Private/NiagaraActions.cpp
@@ -1090,7 +1090,6 @@ TSharedPtr<SWidget> FNiagaraDataInterfaceNodeActionProvider_DataChannelRead::Get
 		AllowedTypes.Add(FNiagaraTypeHelper::GetVector2DDef());
 		AllowedTypes.Add(FNiagaraTypeHelper::GetVectorDef());
 		AllowedTypes.Add(FNiagaraTypeHelper::GetVector4Def());
-		AllowedTypes.Add(FNiagaraTypeDefinition::GetPositionDef());
 	}
 
 	return SNew(SNiagaraFunctionSpecifierNDCVariablesSelector)
