- Clear allocation estimate when changing scalability level

#tests Create a system with varying particle allocation counts on different scalabiltiy levels, use the high count observe memory usage, change to low count and observe memory is still high

diff --git a/Engine/Plugins/FX/Niagara/Source/Niagara/Private/NiagaraEmitter.cpp b/Engine/Plugins/FX/Niagara/Source/Niagara/Private/NiagaraEmitter.cpp
--- a/Engine/Plugins/FX/Niagara/Source/Niagara/Private/NiagaraEmitter.cpp
+++ b/Engine/Plugins/FX/Niagara/Source/Niagara/Private/NiagaraEmitter.cpp
@@ -3731,6 +3731,12 @@ void UNiagaraEmitter::ResolveScalabilitySettings()
 				}
 			}
 		}
+
+		// Clear out any allocation estimates as changing scalability could change the max spawn count
+		if (EmitterData.RuntimeEstimation.AllocationEstimate > 0)
+		{
+			EmitterData.ClearRuntimeAllocationEstimate(INDEX_NONE);
+		}
 	}
 }
 
