- Adjust tool tip to let users know that GPU scene already enables per particle culling

#rnx

diff --git a/Engine/Plugins/FX/Niagara/Source/Niagara/Public/NiagaraMeshRendererProperties.h b/Engine/Plugins/FX/Niagara/Source/Niagara/Public/NiagaraMeshRendererProperties.h
--- a/Engine/Plugins/FX/Niagara/Source/Niagara/Public/NiagaraMeshRendererProperties.h
+++ b/Engine/Plugins/FX/Niagara/Source/Niagara/Public/NiagaraMeshRendererProperties.h
@@ -228,7 +228,10 @@ public:
 	UPROPERTY(EditAnywhere, Category = "SubUV", meta = (DisplayName = "Sub UV Blending Enabled"))
 	uint32 bSubImageBlend : 1;
 
-	/** Enables frustum culling of individual mesh particles */
+	/**
+	Enables frustum culling of individual mesh particles.
+	GPU Scene will already perform per particle culling for meshes renderers, therefore enabling this may result in redundant work.
+	*/
 	UPROPERTY(EditAnywhere, Category = "Visibility")
 	uint32 bEnableFrustumCulling : 1;
 
