- Fix for explicit specialization of 'IsParameterCollectionParameter' after instantiation

#rnx

diff --git a/Engine/Plugins/FX/Niagara/Source/NiagaraEditor/Private/NiagaraCompilationDigestBridgeImpl.cpp b/Engine/Plugins/FX/Niagara/Source/NiagaraEditor/Private/NiagaraCompilationDigestBridgeImpl.cpp
--- a/Engine/Plugins/FX/Niagara/Source/NiagaraEditor/Private/NiagaraCompilationDigestBridgeImpl.cpp
+++ b/Engine/Plugins/FX/Niagara/Source/NiagaraEditor/Private/NiagaraCompilationDigestBridgeImpl.cpp
@@ -436,6 +436,29 @@ UObject* FNiagaraCompilationDigestBridge::GetCollectionUObject(FParameterCollect
 	return nullptr;
 }
 
+FNiagaraCompilationDigestBridge::FParameterCollection FNiagaraCompilationDigestBridge::IsParameterCollectionParameter(const FParamMapHistory& ParamMapHistory, const FNiagaraVariable& InVar, bool& bMissingParameter)
+{
+	if (ParamMapHistory.bParameterCollectionsSkipped)
+	{
+		UE_LOG(LogNiagaraEditor, Error, TEXT("NiagaraParameterCollection was skipped during history building.  History which require NPC data can not be generated during PostLoad()."));
+		return nullptr;
+	}
+
+	bMissingParameter = false;
+	FNameBuilder VarName(InVar.GetName());
+	for (const FNiagaraCompilationNPCHandle& CollectionHandle : ParamMapHistory.EncounteredParameterCollections.Handles)
+	{
+		FNameBuilder CollectionName(CollectionHandle.Namespace);
+		if (VarName.ToView().StartsWith(CollectionName.ToView()))
+		{
+			bMissingParameter = CollectionHandle.Resolve()->Variables.Contains(InVar);
+			return CollectionHandle;
+		}
+	}
+
+	return FNiagaraCompilationNPCHandle();
+}
+
 FNiagaraCompilationDigestBridge::FBuilderExtraData::FBuilderExtraData()
 {
 	TraversalStateContext = MakePimpl<FNiagaraTraversalStateContext>();
diff --git a/Engine/Plugins/FX/Niagara/Source/NiagaraEditor/Private/NiagaraCompilationGraphBridgeImpl.cpp b/Engine/Plugins/FX/Niagara/Source/NiagaraEditor/Private/NiagaraCompilationGraphBridgeImpl.cpp
--- a/Engine/Plugins/FX/Niagara/Source/NiagaraEditor/Private/NiagaraCompilationGraphBridgeImpl.cpp
+++ b/Engine/Plugins/FX/Niagara/Source/NiagaraEditor/Private/NiagaraCompilationGraphBridgeImpl.cpp
@@ -528,6 +528,29 @@ TObjectPtr<UObject> FNiagaraCompilationGraphBridge::GetCollectionUObject(FParame
 	return Collection->GetDefaultInstance()->GetParameterStore().GetUObject(Variable);
 }
 
+FNiagaraCompilationGraphBridge::FParameterCollection FNiagaraCompilationGraphBridge::IsParameterCollectionParameter(const FParamMapHistory& ParamMapHistory, const FNiagaraVariable& InVar, bool& bMissingParameter)
+{
+	if (ParamMapHistory.bParameterCollectionsSkipped)
+	{
+		UE_LOG(LogNiagaraEditor, Error, TEXT("NiagaraParameterCollection was skipped during history building.  History which require NPC data can not be generated during PostLoad()."));
+		return nullptr;
+	}
+
+	bMissingParameter = false;
+	FString VarName = InVar.GetName().ToString();
+	for (int32 i = 0; i < ParamMapHistory.EncounteredParameterCollections.Collections.Num(); ++i)
+	{
+		FNameBuilder CollectionNameBuilder(ParamMapHistory.EncounteredParameterCollections.CollectionNamespaces[i]);
+
+		if (VarName.StartsWith(CollectionNameBuilder))
+		{
+			bMissingParameter = !ParamMapHistory.EncounteredParameterCollections.CollectionVariables[i].Contains(InVar);
+			return ParamMapHistory.EncounteredParameterCollections.Collections[i];
+		}
+	}
+	return nullptr;
+}
+
 const UNiagaraNodeOutput* FNiagaraCompilationGraphBridge::AsOutputNode(const FNode* Node)
 {
 	return Cast<const FOutputNode>(Node);
diff --git a/Engine/Plugins/FX/Niagara/Source/NiagaraEditor/Private/NiagaraHlslTranslator.cpp b/Engine/Plugins/FX/Niagara/Source/NiagaraEditor/Private/NiagaraHlslTranslator.cpp
--- a/Engine/Plugins/FX/Niagara/Source/NiagaraEditor/Private/NiagaraHlslTranslator.cpp
+++ b/Engine/Plugins/FX/Niagara/Source/NiagaraEditor/Private/NiagaraHlslTranslator.cpp
@@ -6392,7 +6392,7 @@ bool TNiagaraHlslTranslator<GraphBridge>::ParameterMapRegisterExternalConstantNa
 	FString ParameterMapInstanceName = GetParameterMapInstanceName(InParamMapHistoryIdx);
 
 	bool bMissingParameter = false;
-	FParameterCollection Collection = ParamMapHistories[InParamMapHistoryIdx].IsParameterCollectionParameter(InVariable, bMissingParameter);
+	FParameterCollection Collection = GraphBridge::IsParameterCollectionParameter(ParamMapHistories[InParamMapHistoryIdx], InVariable, bMissingParameter);
 	if (GraphBridge::IsCollectionValid(Collection) && bMissingParameter)
 	{
 		Error(FText::Format(LOCTEXT("MissingNPCParameterError", "Parameter named {0} of type {1} was not found in Parameter Collection {2}"),
diff --git a/Engine/Plugins/FX/Niagara/Source/NiagaraEditor/Private/NiagaraParameterMapHistory.cpp b/Engine/Plugins/FX/Niagara/Source/NiagaraEditor/Private/NiagaraParameterMapHistory.cpp
--- a/Engine/Plugins/FX/Niagara/Source/NiagaraEditor/Private/NiagaraParameterMapHistory.cpp
+++ b/Engine/Plugins/FX/Niagara/Source/NiagaraEditor/Private/NiagaraParameterMapHistory.cpp
@@ -917,54 +917,6 @@ FNiagaraVariable FNiagaraParameterUtilities::ConvertVariableToRapidIterationCons
 	return FNiagaraUtilities::ConvertVariableToRapidIterationConstantName(InVar, InEmitterName, InUsage);
 }
 
-template<>
-FNiagaraCompilationGraphBridge::FParameterCollection TNiagaraParameterMapHistory<FNiagaraCompilationGraphBridge>::IsParameterCollectionParameter(FNiagaraVariable& InVar, bool& bMissingParameter)
-{
-	if (bParameterCollectionsSkipped)
-	{
-		UE_LOG(LogNiagaraEditor, Error, TEXT("NiagaraParameterCollection was skipped during history building.  History which require NPC data can not be generated during PostLoad()."));
-		return nullptr;
-	}
-
-	bMissingParameter = false;
-	FString VarName = InVar.GetName().ToString();
-	for (int32 i = 0; i < EncounteredParameterCollections.Collections.Num(); ++i)
-	{
-		FNameBuilder CollectionNameBuilder(EncounteredParameterCollections.CollectionNamespaces[i]);
-
-		if (VarName.StartsWith(CollectionNameBuilder))
-		{
-			bMissingParameter = !EncounteredParameterCollections.CollectionVariables[i].Contains(InVar);
-			return EncounteredParameterCollections.Collections[i];
-		}
-	}
-	return nullptr;
-}
-
-template<>
-FNiagaraCompilationDigestBridge::FParameterCollection TNiagaraParameterMapHistory<FNiagaraCompilationDigestBridge>::IsParameterCollectionParameter(FNiagaraVariable& InVar, bool& bMissingParameter)
-{
-	if (bParameterCollectionsSkipped)
-	{
-		UE_LOG(LogNiagaraEditor, Error, TEXT("NiagaraParameterCollection was skipped during history building.  History which require NPC data can not be generated during PostLoad()."));
-		return nullptr;
-	}
-
-	bMissingParameter = false;
-	FNameBuilder VarName(InVar.GetName());
-	for (const FNiagaraCompilationNPCHandle& CollectionHandle : EncounteredParameterCollections.Handles)
-	{
-		FNameBuilder CollectionName(CollectionHandle.Namespace);
-		if (VarName.ToView().StartsWith(CollectionName.ToView()))
-		{
-			bMissingParameter = CollectionHandle.Resolve()->Variables.Contains(InVar);
-			return CollectionHandle;
-		}
-	}
-
-	return FNiagaraCompilationNPCHandle();
-}
-
 template<typename GraphBridge>
 bool TNiagaraParameterMapHistory<GraphBridge>::ShouldIgnoreVariableDefault(const FNiagaraVariable& Var)const
 {
diff --git a/Engine/Plugins/FX/Niagara/Source/NiagaraEditor/Public/NiagaraCompilationBridge.h b/Engine/Plugins/FX/Niagara/Source/NiagaraEditor/Public/NiagaraCompilationBridge.h
--- a/Engine/Plugins/FX/Niagara/Source/NiagaraEditor/Public/NiagaraCompilationBridge.h
+++ b/Engine/Plugins/FX/Niagara/Source/NiagaraEditor/Public/NiagaraCompilationBridge.h
@@ -210,6 +210,7 @@ struct FNiagaraCompilationGraphBridge
 	static bool IsCollectionValid(FParameterCollection Collection);
 	static UNiagaraDataInterface* GetCollectionDataInterface(FParameterCollection Collection, const FNiagaraVariable& Variable);
 	static TObjectPtr<UObject> GetCollectionUObject(FParameterCollection Collection, const FNiagaraVariable& Variable);
+	static FParameterCollection IsParameterCollectionParameter(const FParamMapHistory& ParamMapHistory, const FNiagaraVariable& InVar, bool& bMissingParameter);
 
 	static const FOutputNode* AsOutputNode(const FNode* Node);
 	static bool IsParameterMapGet(const FNode* Node);
@@ -342,6 +343,7 @@ struct FNiagaraCompilationDigestBridge
 	static bool IsCollectionValid(FParameterCollection Collection);
 	static UNiagaraDataInterface* GetCollectionDataInterface(FParameterCollection Collection, const FNiagaraVariable& Variable);
 	static UObject* GetCollectionUObject(FParameterCollection Collection, const FNiagaraVariable& Variable);
+	static FParameterCollection IsParameterCollectionParameter(const FParamMapHistory& ParamMapHistory, const FNiagaraVariable& InVar, bool& bMissingParameter);
 
 	static const FOutputNode* AsOutputNode(const FNode* Node);
 	static bool IsParameterMapGet(const FNode* Node);
diff --git a/Engine/Plugins/FX/Niagara/Source/NiagaraEditor/Public/NiagaraParameterMapHistory.h b/Engine/Plugins/FX/Niagara/Source/NiagaraEditor/Public/NiagaraParameterMapHistory.h
--- a/Engine/Plugins/FX/Niagara/Source/NiagaraEditor/Public/NiagaraParameterMapHistory.h
+++ b/Engine/Plugins/FX/Niagara/Source/NiagaraEditor/Public/NiagaraParameterMapHistory.h
@@ -368,13 +368,6 @@ public:
 	/** Get the default value for this variable.*/
 	const FInputPin* GetDefaultValuePin(int32 VarIdx) const;
 
-	/**
-	If this is variable is a parameter in one of our tracked collections, return it.
-	@param InVar	Variable to test.
-	@param bMissingParameter	bool set to mark if this parameter was a collection parameter but is now missing from it's collection.
-	*/
-	FParameterCollection IsParameterCollectionParameter(FNiagaraVariable& InVar, bool& bMissingParameter);
-
 	bool ShouldIgnoreVariableDefault(const FNiagaraVariable& Var)const;
 
 	void ConditionalUpdateAliasedVariable(int32 VariableIndex, const FNiagaraVariableBase& InAliasedVar);
