Fixed padding of scalability elements in Niagara node header

diff --git a/Engine/Plugins/FX/Niagara/Source/NiagaraEditorWidgets/Private/SNiagaraOverviewStackNode.cpp b/Engine/Plugins/FX/Niagara/Source/NiagaraEditorWidgets/Private/SNiagaraOverviewStackNode.cpp
--- a/Engine/Plugins/FX/Niagara/Source/NiagaraEditorWidgets/Private/SNiagaraOverviewStackNode.cpp
+++ b/Engine/Plugins/FX/Niagara/Source/NiagaraEditorWidgets/Private/SNiagaraOverviewStackNode.cpp
@@ -1175,6 +1175,7 @@ TSharedRef<SWidget> SNiagaraOverviewStackNode::CreateScalabilityControls()
 		.AutoHeight()
 		.HAlign(HAlign_Center)
 		.VAlign(VAlign_Center)
+		.Padding(2.f, 4.f)
 		[
 			// Toggle Button to enter & exit scalability mode
 			SNew(SCheckBox)
