- Add transient to runtime spawned Cascade & Niagara components
- This fixes an issue where they can Modify the attached actor resulting in invalidations of DDC keys

diff --git a/Engine/Plugins/FX/Niagara/Source/Niagara/Private/NiagaraFunctionLibrary.cpp b/Engine/Plugins/FX/Niagara/Source/Niagara/Private/NiagaraFunctionLibrary.cpp
--- a/Engine/Plugins/FX/Niagara/Source/Niagara/Private/NiagaraFunctionLibrary.cpp
+++ b/Engine/Plugins/FX/Niagara/Source/Niagara/Private/NiagaraFunctionLibrary.cpp
@@ -72,7 +72,7 @@ UNiagaraComponent* CreateNiagaraSystem(UNiagaraSystem* SystemTemplate, UWorld* W
 	{
 		if (PoolingMethod == ENCPoolMethod::None)
 		{
-			NiagaraComponent = NewObject<UNiagaraComponent>((Actor ? Actor : (UObject*)World));
+			NiagaraComponent = NewObject<UNiagaraComponent>((Actor ? Actor : (UObject*)World), NAME_None, RF_Transient);
 			NiagaraComponent->SetAsset(SystemTemplate);
 			NiagaraComponent->bAutoActivate = false;
 			NiagaraComponent->SetAutoDestroy(bAutoDestroy);
diff --git a/Engine/Source/Runtime/Engine/Private/GameplayStatics.cpp b/Engine/Source/Runtime/Engine/Private/GameplayStatics.cpp
--- a/Engine/Source/Runtime/Engine/Private/GameplayStatics.cpp
+++ b/Engine/Source/Runtime/Engine/Private/GameplayStatics.cpp
@@ -1230,7 +1230,7 @@ UParticleSystemComponent* CreateParticleSystem(UParticleSystem* EmitterTemplate,
 		}
 		else
 		{
-			PSC = NewObject<UParticleSystemComponent>((Actor ? Actor : (UObject*)World));
+			PSC = NewObject<UParticleSystemComponent>((Actor ? Actor : (UObject*)World), NAME_None, RF_Transient);
 			PSC->bAutoDestroy = bAutoDestroy;
 			PSC->bAllowAnyoneToDestroyMe = true;
 			PSC->SecondsBeforeInactive = 0.0f;
